Copying //UE4/Dev-Physics to Dev-Main (//UE4/Dev-Main) [at] 7703071

diff --git a/Engine/Plugins/Experimental/HairStrands/Shaders/Private/NiagaraDataInterfaceHairStrands.ush b/Engine/Plugins/Experimental/HairStrands/Shaders/Private/NiagaraDataInterfaceHairStrands.ush
new file mode 100644
--- /dev/null
+++ b/Engine/Plugins/Experimental/HairStrands/Shaders/Private/NiagaraDataInterfaceHairStrands.ush
@@ -0,0 +1,238 @@
+// Copyright 1998-2019 Epic Games, Inc. All Rights Reserved.
+
+#pragma once
+
+#define WRAP_SIZE 32
+
+struct FDIHairStrandsContext
+{
+	RWBuffer<float4>	PointsPositionsBuffer;
+	Buffer<uint>		CurvesOffsetsBuffer;
+	Buffer<float4>		NodesPositionsBuffer;
+	Buffer<uint>		PointsNodesBuffer;
+	Buffer<float>		PointsWeightsBuffer;
+	float4x4			WorldTransform;
+	int					NumStrands;
+	int					StrandSize;
+	float				StrandDensity;
+	float				RootThickness;
+	float				TipThickness;
+}; 
+
+#define DIHAIRSTRANDS_DECLARE_CONSTANTS(NAME)\
+RWBuffer<float4> PointsPositionsBuffer_##NAME;\
+Buffer<uint> CurvesOffsetsBuffer_##NAME;\
+Buffer<float4> NodesPositionsBuffer_##NAME;\
+Buffer<uint> PointsNodesBuffer_##NAME;\
+Buffer<float> PointsWeightsBuffer_##NAME;\
+float4x4 WorldTransform_##NAME;\
+int NumStrands_##NAME;\
+int StrandSize_##NAME;\
+float StrandDensity_##NAME;\
+float RootThickness_##NAME;\
+float TipThickness_##NAME;
+
+#define DIHAIRSTRANDS_MAKE_CONTEXT(NAME)\
+FDIHairStrandsContext DIContext;\
+DIContext.PointsPositionsBuffer = PointsPositionsBuffer_##NAME;\
+DIContext.CurvesOffsetsBuffer = CurvesOffsetsBuffer_##NAME;\
+DIContext.NodesPositionsBuffer = NodesPositionsBuffer_##NAME;\
+DIContext.PointsNodesBuffer = PointsNodesBuffer_##NAME;\
+DIContext.PointsWeightsBuffer = PointsWeightsBuffer_##NAME;\
+DIContext.WorldTransform = WorldTransform_##NAME;\
+DIContext.NumStrands = NumStrands_##NAME;\
+DIContext.StrandSize = StrandSize_##NAME;\
+DIContext.StrandDensity = StrandDensity_##NAME;\
+DIContext.RootThickness = RootThickness_##NAME;\
+DIContext.TipThickness = TipThickness_##NAME;
+
+groupshared float3 NodesGlobalPosition[WRAP_SIZE];
+groupshared float4 NodesGlobalQuaternion[WRAP_SIZE];
+
+groupshared float NodesInverseMass[WRAP_SIZE];
+groupshared float NodesInverseInertia[WRAP_SIZE];
+
+// Given a node index return the strand index and the local index within the strand
+int DIHairStrands_GetNumNodes(in FDIHairStrandsContext DIContext)
+{
+	return DIContext.NumStrands * DIContext.StrandSize;
+}
+
+// Given a node index return the strand index and the local index within the strand
+bool DIHairStrands_IsRootNode(in FDIHairStrandsContext DIContext)
+{
+	return (GGroupThreadId.x % DIContext.StrandSize) != 0;
+}
+
+// Given a node index return the strand index and the local index within the strand
+void DIHairStrands_ComputeNodeIndices(in FDIHairStrandsContext DIContext, in int NodeIndex, out int StrandIndex, out int LocalIndex)
+{
+	if (NodeIndex < DIHairStrands_GetNumNodes(DIContext))
+	{
+		StrandIndex = (DIContext.StrandSize>0) ? NodeIndex / DIContext.StrandSize : 0;
+		LocalIndex = NodeIndex - StrandIndex * DIContext.StrandSize;
+	}
+	else
+	{
+		StrandIndex = 0;
+		LocalIndex = 0;
+	}
+}
+
+// Given a strand index return the corresponding strand length
+float DIHairStrands_ComputeStrandLength(in FDIHairStrandsContext DIContext, in int StrandIndex)
+{
+	float StrandLength = 0.0;
+	if( StrandIndex < DIContext.NumStrands )
+	{
+		int PointOffset = DIContext.CurvesOffsetsBuffer[StrandIndex];
+		const int EdgeCount = DIContext.CurvesOffsetsBuffer[StrandIndex+1]-PointOffset-1;
+
+		float3 PointNext = DIContext.PointsPositionsBuffer[PointOffset++].xyz;
+		float3 PointPrev = PointNext;
+
+		for (int EdgeIndex = 0; EdgeIndex < EdgeCount; ++EdgeIndex)
+		{
+			PointPrev = PointNext;
+			PointNext = DIContext.PointsPositionsBuffer[PointOffset++].xyz;
+			StrandLength += length(PointNext-PointPrev);
+		}
+	}
+	return StrandLength;
+}
+
+/*void DIHairStrands_ReportNodePosition(in FDIHairStrandsContext DIContext, in int NodeIndex, out float3 Out_NodePosition)
+{
+	int StrandIndex = 0;
+	int LocalIndex = 0;
+	DIHairStrands_ComputeNodeIndices(DIContext,NodeIndex,StrandIndex,LocalIndex);
+
+	const int PointBegin = DIContext.CurvesOffsetsBuffer[StrandIndex];
+	const int PointEnd = DIContext.CurvesOffsetsBuffer[StrandIndex+1];
+
+	for (int PointIndex = PointBegin; PointIndex < PointEnd; ++PointIndex)
+	{
+		const float Weight = (DIContext.PointsNodesBuffer[PointIndex] == NodeIndex) ? 1.0 : 0.0;
+		DIContext.PointsPositionsBuffer[PointIndex] += Weight * Out_NodePosition * DIContext.PointsWeightsBuffer[PointIndex];
+	}
+	DeviceMemoryBarrier();
+
+	for (int PointIndex = PointBegin; PointIndex < PointEnd; ++PointIndex)
+	{
+		const float Weight = (DIContext.PointsNodesBuffer[PointIndex] == NodeIndex-1) ? 1.0 : 0.0;
+		DIContext.PointsPositionsBuffer[PointIndex] -= Weight * Out_NodePosition * (1.0-DIContext.PointsWeightsBuffer[PointIndex]);
+	}
+	DeviceMemoryBarrier();
+}*/
+
+// Compute the strand position by linear interpolation over the points
+/*void DIHairStrands_ReportNodeVelocity(in FDIHairStrandsContext DIContext, in int NodeIndex, in float3 NodeVelocity, in float DeltaTime)
+{
+	NodesGlobalPosition[GGroupThreadId.x] = NodeVelocity * DeltaTime;
+	GroupMemoryBarrier(); 
+
+	if(DIHairStrands_IsRootNode(DIContext))
+	{
+		int StrandIndex = 0;
+		int LocalIndex = 0;
+		DIHairStrands_ComputeNodeIndices(DIContext,NodeIndex,StrandIndex,LocalIndex);
+
+		const int PointBegin = DIContext.CurvesOffsetsBuffer[StrandIndex];
+		const int PointEnd = DIContext.CurvesOffsetsBuffer[StrandIndex+1]; 
+		const int PointCount = PointEnd - PointBegin;
+
+		for (int PointIndex = 0; PointIndex < PointCount; ++PointIndex)
+		{
+			const float PointCoord = (float)(PointIndex) / (PointCount-1.0);
+			const float PointNode = PointCoord * (DIContext.StrandSize-1.0);
+
+			const int NodePrev = (PointIndex==0) ? 0 : (PointIndex==(PointCount-1.0)) ? DIContext.StrandSize-2 : floor(PointNode);
+			const int NodeNext = NodePrev+1;
+
+			const float PointAlpha = PointNode - (float)NodePrev;
+			DIContext.PointsPositionsBuffer[PointBegin+PointIndex].xyz = NodesGlobalPosition[GGroupThreadId.x]
+
+
+		}
+	}
+}*/
+
+// Compute the strand orientation with parallel transport along the curve
+void DIHairStrands_ComputeNodeOrientation(in FDIHairStrandsContext DIContext, 
+						in float3 NodePosition, out float4 OutNodeQuaternion)
+{
+	NodesGlobalPosition[GGroupThreadId.x] = NodePosition;
+	GroupMemoryBarrier(); 
+
+	if(DIHairStrands_IsRootNode(DIContext))
+	{
+		float4 Quaternion = QUATERNION_IDENTITY;
+		float3 TangentNext = float3(0.0,0.0,1.0);
+		float3 TangentPrev = TangentNext;
+
+		NodesGlobalQuaternion[GGroupThreadId.x] = Quaternion;
+
+		for (int EdgeIndex = GGroupThreadId.x+1, EdgeEnd = GGroupThreadId.x+DIContext.StrandSize; 
+				EdgeIndex < EdgeEnd; ++EdgeIndex)
+		{
+			TangentNext = normalize(NodesGlobalPosition[EdgeIndex] - NodesGlobalPosition[EdgeIndex-1]);
+			Quaternion = NormalizeQuat( MultiplyQuat( FindQuatBetweenNormals(TangentPrev,TangentNext), Quaternion ) );
+			TangentPrev = TangentNext;
+
+			NodesGlobalQuaternion[EdgeIndex] = Quaternion;
+		}
+	}
+	GroupMemoryBarrier();
+	OutNodeQuaternion = NodesGlobalQuaternion[GGroupThreadId.x];
+}
+
+// Compute the node mass
+void DIHairStrands_ComputeNodeMass(in FDIHairStrandsContext DIContext, in int NodeIndex, out float3 OutNodeMass)
+{
+	const float CoordScale = 1.0 / (DIContext.StrandSize-1.0);
+
+	int StrandIndex = 0, LocalIndex = 0;
+	DIHairStrands_ComputeNodeIndices(DIContext,NodeIndex,StrandIndex,LocalIndex);
+	const float EdgeLength = DIHairStrands_ComputeStrandLength(DIContext,StrandIndex) * CoordScale;
+
+	const bool IsFirst = (LocalIndex == 0);
+	const bool IsLast  = (LocalIndex == (DIContext.StrandSize-1));
+
+	const float EdgeScale = (IsFirst || IsLast) ? 0.5 : 1.0;
+	
+	const float PrevCoord = IsFirst ? 0.0 : (LocalIndex-0.5) * CoordScale;
+	const float NextCoord = IsLast ? 1.0  : (LocalIndex+0.5) * CoordScale;
+
+	const float PrevThick = DIContext.RootThickness * ( 1.0 - PrevCoord ) + DIContext.TipThickness * PrevCoord;
+	const float NextThick = DIContext.RootThickness * ( 1.0 - NextCoord ) + DIContext.TipThickness * NextCoord;
+
+	OutNodeMass = DIContext.StrandDensity * EdgeLength * EdgeScale * M_PI * (PrevThick * PrevThick + NextThick * PrevThick + NextThick * NextThick) / 12.0;
+}
+
+// Compute the node position by linear interpolation over the points
+void DIHairStrands_ComputeNodePosition(in FDIHairStrandsContext DIContext, in int NodeIndex, out float3 OutNodePosition)
+{
+	int StrandIndex = 0, LocalIndex = 0;
+	DIHairStrands_ComputeNodeIndices(DIContext,NodeIndex,StrandIndex,LocalIndex);
+
+	const int PointOffset = DIContext.CurvesOffsetsBuffer[StrandIndex];
+
+	const float NodeCoord = (float)(LocalIndex) / (DIContext.StrandSize-1.0);
+	const float PointCount = DIContext.CurvesOffsetsBuffer[StrandIndex+1]-PointOffset-1;
+
+	const float NodePoint = NodeCoord * PointCount;
+	const int PointPrev = (LocalIndex==0) ? 0 : (LocalIndex==(DIContext.StrandSize-1)) ? PointCount-1 : floor(NodePoint);
+	const int PointNext = PointPrev+1;
+
+	const float NodeAlpha = NodePoint - (float)PointPrev;
+	OutNodePosition =  DIContext.PointsPositionsBuffer[PointOffset+PointPrev].xyz * (1.0-NodeAlpha) + 
+					   DIContext.PointsPositionsBuffer[PointOffset+PointNext].xyz * NodeAlpha;
+	OutNodePosition = mul(float4( OutNodePosition, 1.0), DIContext.WorldTransform).xyz;
+}
+
+
+
+
+
+
+
diff --git a/Engine/Plugins/Experimental/HairStrands/Source/HairStrandsCore/HairStrandsCore.Build.cs b/Engine/Plugins/Experimental/HairStrands/Source/HairStrandsCore/HairStrandsCore.Build.cs
new file mode 100644
--- /dev/null
+++ b/Engine/Plugins/Experimental/HairStrands/Source/HairStrandsCore/HairStrandsCore.Build.cs
@@ -0,0 +1,26 @@
+// Copyright 1998-2019 Epic Games, Inc. All Rights Reserved.
+
+namespace UnrealBuildTool.Rules
+{
+	public class HairStrandsCore : ModuleRules
+	{
+        public HairStrandsCore(ReadOnlyTargetRules Target) : base(Target)
+		{
+			PrivateIncludePaths.Add(ModuleDirectory + "/Private");
+            PublicIncludePaths.Add(ModuleDirectory + "/Public");
+
+            PrivateDependencyModuleNames.AddRange(
+                new string[]
+                {
+					"Core",
+					"CoreUObject",
+					"Engine",
+                    "RenderCore",
+                    "Renderer",
+                    "RHI",
+					"ChaosCore",
+                    "Chaos"
+                });
+        }
+	}
+}
diff --git a/Engine/Plugins/Experimental/HairStrands/Source/HairStrandsCore/Private/HairStrandsAsset.cpp b/Engine/Plugins/Experimental/HairStrands/Source/HairStrandsCore/Private/HairStrandsAsset.cpp
new file mode 100644
--- /dev/null
+++ b/Engine/Plugins/Experimental/HairStrands/Source/HairStrandsCore/Private/HairStrandsAsset.cpp
@@ -0,0 +1,171 @@
+// Copyright 1998-2019 Epic Games, Inc. All Rights Reserved.
+
+#include "HairStrandsAsset.h"
+#include "RenderingThread.h"
+
+#include "Misc/Paths.h"
+
+#if WITH_EDITORONLY_DATA
+#include "EditorFramework/AssetImportData.h"
+#endif
+
+FHairStrandsInstance::~FHairStrandsInstance()
+{
+	if (RenderResource && bInstancedResource)
+	{
+		FHairStrandsResource* InResource = RenderResource;
+		ENQUEUE_RENDER_COMMAND(FDestroyHairStrandsResourceCommand)(
+			[InResource](FRHICommandList& RHICmdList)
+		{
+			InResource->ReleaseResource();
+		});
+	}
+	RenderResource = nullptr;
+}
+
+void FHairStrandsInstance::InitResource(FHairStrandsResource* InResource, bool bInstanced)
+{
+	check(!RenderResource);
+	RenderResource = InResource;
+	bInstancedResource = bInstanced;
+}
+
+void FHairStrandsInstance::UpdateTransforms(const FMatrix& LocalToWorld)
+{
+	LocalToGlobal = LocalToWorld;
+}
+
+FHairStrandsResource::FHairStrandsResource(FHairStrandsDatas* HairStrandsDatas) : 
+	PositionBuffer(), TangentBuffer(), StrandsDatas(HairStrandsDatas)
+{}
+
+void FHairStrandsResource::InitRHI()
+{
+	if (StrandsDatas != nullptr)
+	{
+		TArray<FHairStrandsPositionFormat::Type> RenderingPositions;
+		TArray<FHairStrandsTangentFormat::Type> RenderingTangents;
+		StrandsDatas->BuildRenderingDatas(RenderingPositions, RenderingTangents);
+		{
+			const uint32 PositionCount = RenderingPositions.Num();
+			const uint32 PositionBytes = FHairStrandsPositionFormat::SizeInByte*PositionCount;
+
+			PositionBuffer.Initialize(FHairStrandsPositionFormat::SizeInByte, PositionCount, FHairStrandsPositionFormat::Format, BUF_Static);
+			void* PositionBufferData = RHILockVertexBuffer(PositionBuffer.Buffer, 0, PositionBytes, RLM_WriteOnly);
+
+			FMemory::Memcpy(PositionBufferData, RenderingPositions.GetData(), PositionBytes);
+			RHIUnlockVertexBuffer(PositionBuffer.Buffer);
+		}
+		{
+			const uint32 TangentCount = RenderingTangents.Num();
+			const uint32 TangentBytes = FHairStrandsTangentFormat::SizeInByte*TangentCount;
+
+			TangentBuffer.Initialize(FHairStrandsTangentFormat::SizeInByte, TangentCount, FHairStrandsTangentFormat::Format, BUF_Static);
+			void* TangentBufferData = RHILockVertexBuffer(TangentBuffer.Buffer, 0, TangentBytes, RLM_WriteOnly);
+
+			FMemory::Memcpy(TangentBufferData, RenderingTangents.GetData(), TangentBytes);
+			RHIUnlockVertexBuffer(TangentBuffer.Buffer);
+		}
+	}
+}
+
+void FHairStrandsResource::ReleaseRHI()
+{
+	PositionBuffer.Release();
+	TangentBuffer.Release();
+}
+
+void UHairStrandsAsset::Serialize(FArchive& Ar)
+{
+	Super::Serialize(Ar);
+	StrandsDatas.Serialize(Ar);
+}
+
+void UHairStrandsAsset::InitResource()
+{
+	check(!HairStrandsResource);
+
+	HairStrandsResource = new FHairStrandsResource(&StrandsDatas);
+
+	BeginInitResource(HairStrandsResource);
+}
+
+void UHairStrandsAsset::UpdateResource()
+{
+	if (HairStrandsResource)
+	{
+		BeginUpdateResourceRHI(HairStrandsResource);
+	}
+}
+
+void UHairStrandsAsset::ReleaseResource()
+{
+	if (HairStrandsResource)
+	{
+		FHairStrandsResource* InResource = HairStrandsResource;
+		ENQUEUE_RENDER_COMMAND(ReleaseHairStrandsResourceCommand)(
+			[InResource](FRHICommandList& RHICmdList)
+		{
+			InResource->ReleaseResource();
+			delete InResource;
+		});
+		HairStrandsResource = nullptr;
+	}
+}
+
+void UHairStrandsAsset::PostLoad()
+{
+	Super::PostLoad();
+
+	if (!IsTemplate())
+	{
+		InitResource();
+	}
+
+#if WITH_EDITORONLY_DATA
+	if (!FilePath.IsEmpty() && AssetImportData)
+	{
+		FAssetImportInfo Info;
+		Info.Insert(FAssetImportInfo::FSourceFile(FilePath));
+		AssetImportData->SourceData = MoveTemp(Info);
+	}
+#endif
+}
+
+void UHairStrandsAsset::BeginDestroy()
+{
+	ReleaseResource();
+	Super::BeginDestroy();
+}
+
+#if WITH_EDITOR
+void UHairStrandsAsset::PostEditChangeProperty(FPropertyChangedEvent& PropertyChangedEvent)
+{
+	Super::PostEditChangeProperty(PropertyChangedEvent);
+	UpdateResource();
+}
+#endif // WITH_EDITOR
+
+#if WITH_EDITORONLY_DATA
+void UHairStrandsAsset::GetAssetRegistryTags(TArray<FAssetRegistryTag>& OutTags) const
+{
+	if (AssetImportData)
+	{
+		OutTags.Add(FAssetRegistryTag(SourceFileTagName(), AssetImportData->GetSourceData().ToJson(), FAssetRegistryTag::TT_Hidden));
+	}
+
+	Super::GetAssetRegistryTags(OutTags);
+}
+
+void UHairStrandsAsset::PostInitProperties()
+{
+	if (!HasAnyFlags(RF_ClassDefaultObject))
+	{
+		AssetImportData = NewObject<UAssetImportData>(this, TEXT("AssetImportData"));
+	}
+
+	Super::PostInitProperties();
+}
+#endif
+
+
diff --git a/Engine/Plugins/Experimental/HairStrands/Source/HairStrandsCore/Private/HairStrandsComponent.cpp b/Engine/Plugins/Experimental/HairStrands/Source/HairStrandsCore/Private/HairStrandsComponent.cpp
new file mode 100644
--- /dev/null
+++ b/Engine/Plugins/Experimental/HairStrands/Source/HairStrandsCore/Private/HairStrandsComponent.cpp
@@ -0,0 +1,66 @@
+// Copyright 1998-2019 Epic Games, Inc. All Rights Reserved. 
+
+#include "HairStrandsComponent.h"
+#include "Engine/CollisionProfile.h"
+#include "PrimitiveSceneProxy.h"
+
+UHairStrandsComponent::UHairStrandsComponent(const FObjectInitializer& ObjectInitializer)
+	: Super(ObjectInitializer)
+{
+	PrimaryComponentTick.bCanEverTick = true;
+	bTickInEditor = true;
+	bAutoActivate = true;
+	HairDensity = 1;
+	MergeThreshold = 0.1f;
+	bUsedForReference = false;
+
+	SetCollisionProfileName(UCollisionProfile::PhysicsActor_ProfileName);
+}
+
+FPrimitiveSceneProxy* UHairStrandsComponent::CreateSceneProxy()
+{
+	return nullptr;
+}
+
+FBoxSphereBounds UHairStrandsComponent::CalcBounds(const FTransform& LocalToWorld) const
+{
+	FBox HairBox(ForceInit);
+	if (HairStrandsAsset)
+	{
+		for (FVector& PointPosition : HairStrandsAsset->StrandsDatas.StrandsPoints.PointsPosition)
+		{
+			HairBox += PointPosition;
+		}
+		HairBox = HairBox.TransformBy(LocalToWorld);
+		return FBoxSphereBounds(HairBox);
+	}
+	else
+	{
+		return FBoxSphereBounds();
+	}
+}
+
+int32 UHairStrandsComponent::GetNumMaterials() const
+{
+	return 1;
+}
+
+void UHairStrandsComponent::PostLoad()
+{
+	Super::PostLoad();
+}
+
+void UHairStrandsComponent::OnRegister()
+{
+	Super::OnRegister();
+}
+
+void UHairStrandsComponent::OnUnregister()
+{
+	Super::OnUnregister();
+}
+
+void UHairStrandsComponent::TickComponent(float DeltaTime, enum ELevelTick TickType, FActorComponentTickFunction *ThisTickFunction)
+{
+	Super::TickComponent(DeltaTime, TickType, ThisTickFunction);
+}
\ No newline at end of file
diff --git a/Engine/Plugins/Experimental/HairStrands/Source/HairStrandsCore/Private/HairStrandsCore.cpp b/Engine/Plugins/Experimental/HairStrands/Source/HairStrandsCore/Private/HairStrandsCore.cpp
new file mode 100644
--- /dev/null
+++ b/Engine/Plugins/Experimental/HairStrands/Source/HairStrandsCore/Private/HairStrandsCore.cpp
@@ -0,0 +1,12 @@
+// Copyright 1998-2019 Epic Games, Inc. All Rights Reserved.
+
+#include "HairStrandsCore.h"
+
+void FHairStrandsCore::StartupModule()
+{
+}
+
+void FHairStrandsCore::ShutdownModule()
+{
+}
+
diff --git a/Engine/Plugins/Experimental/HairStrands/Source/HairStrandsCore/Private/HairStrandsDatas.cpp b/Engine/Plugins/Experimental/HairStrands/Source/HairStrandsCore/Private/HairStrandsDatas.cpp
new file mode 100644
--- /dev/null
+++ b/Engine/Plugins/Experimental/HairStrands/Source/HairStrandsCore/Private/HairStrandsDatas.cpp
@@ -0,0 +1,301 @@
+// Copyright 1998-2019 Epic Games, Inc. All Rights Reserved.
+
+#include "HairStrandsDatas.h"
+#include "RenderUtils.h"
+#include "Chaos/TriangleMeshImplicitObject.h"
+#include "Chaos/Particles.h"
+#include "Chaos/ErrorReporter.h"
+#include "Chaos/Levelset.h"
+#include "Chaos/Framework/Parallel.h"
+
+inline void CopyVectorToPosition(const FVector& InVector, FHairStrandsPositionFormat::Type& OutPosition)
+{
+	OutPosition.X = InVector.X;
+	OutPosition.Y = InVector.Y;
+	OutPosition.Z = InVector.Z;
+}
+
+inline void CopyPositionToVector(const FHairStrandsPositionFormat::Type& InPosition, FVector& OutVector )
+{
+	OutVector.X = InPosition.X;
+	OutVector.Y = InPosition.Y;
+	OutVector.Z = InPosition.Z;
+}
+
+void FHairStrandsDatas::BuildRenderingDatas(
+	TArray<FHairStrandsPositionFormat::Type>& OutPackedPositions,
+	TArray<FHairStrandsTangentFormat::Type>& OutPackedTangents) const
+{
+	if (GetNumCurves() > 0 && GetNumPoints() > 0)
+	{
+		OutPackedPositions.SetNum(GetNumPoints());
+		OutPackedTangents.SetNum(GetNumPoints()*2);
+
+		TArray<uint16>::TConstIterator CountIterator = StrandsCurves.CurvesCount.CreateConstIterator();
+		TArray<float>::TConstIterator LengthIterator = StrandsCurves.CurvesLength.CreateConstIterator();
+		TArray<float>::TConstIterator CoordUIterator = StrandsPoints.PointsCoordU.CreateConstIterator();
+		TArray<float>::TConstIterator RadiusIterator = StrandsPoints.PointsRadius.CreateConstIterator();
+		TArray<FVector>::TConstIterator PositionIterator = StrandsPoints.PointsPosition.CreateConstIterator();
+
+		TArray<FHairStrandsPositionFormat::Type>::TIterator PackedPositionIterator = OutPackedPositions.CreateIterator();
+		TArray<FHairStrandsTangentFormat::Type>::TIterator PackedTangentIterator = OutPackedTangents.CreateIterator();
+
+		FRandomStream Random;
+		for (uint32 CurveIndex = 0; CurveIndex < GetNumCurves(); ++CurveIndex, ++CountIterator, ++LengthIterator)
+		{
+			const uint16& PointCount = *CountIterator;
+			for (int32 PointIndex = 0; PointIndex < PointCount; ++PointIndex, 
+				++CoordUIterator, ++PositionIterator, ++RadiusIterator, ++PackedPositionIterator, ++PackedTangentIterator)
+			{
+				const uint32 PrevIndex = FMath::Max(0, PointIndex - 1);
+				const uint32 NextIndex = FMath::Min(PointCount + 1, PointCount - 1);
+				const FVector& PointPosition = *PositionIterator;
+
+				const FVector ForwardDir0 = PointIndex > 0 ? (PointPosition - StrandsPoints.PointsPosition[PrevIndex]).GetSafeNormal() : FVector::ZeroVector;
+				const FVector ForwardDir1 = PointIndex < PointCount - 1 ? (StrandsPoints.PointsPosition[NextIndex] - PointPosition).GetSafeNormal() : FVector::ZeroVector;
+				const FVector TangentZ = (ForwardDir0 + ForwardDir1).GetSafeNormal();
+
+				// Find quat from up (Z) vector to forward
+				// Apply quat orthogonal vectors
+				const FQuat DeltaQuat = FQuat::FindBetween(FVector(0, 0, -1), TangentZ);
+				const FVector RightDir = DeltaQuat.RotateVector(FVector(0, 1, 0));
+				const FVector UpDir = DeltaQuat.RotateVector(FVector(1, 0, 0));
+
+				const FVector TangentX = RightDir ^ TangentZ;
+				const FVector TangentY = RightDir;
+
+				FHairStrandsPositionFormat::Type& PackedPosition = *PackedPositionIterator;
+				CopyVectorToPosition(PointPosition, PackedPosition);
+				PackedPosition.ControlPointType = (PointIndex == 0) ? 1u : (PointIndex == (PointCount - 1) ? 2u : 0u);
+				PackedPosition.UCoord = FMath::Clamp((*CoordUIterator)*255.f, 0.f, 255.f);
+				PackedPosition.NormalizedRadius = FMath::Clamp((*RadiusIterator) * 63.f, 0.f, 63.f);
+				PackedPosition.NormalizedLength = FMath::Clamp((*LengthIterator) *255.f, 0.f, 255.f);
+				PackedPosition.Seed = Random.RandHelper(255);
+
+				FHairStrandsTangentFormat::Type& PackedTangentX = *PackedTangentIterator;
+				PackedTangentX = TangentX;
+
+				++PackedTangentIterator;
+
+				FHairStrandsTangentFormat::Type& PackedTangentZ = *PackedTangentIterator;
+				PackedTangentZ = TangentZ;
+				PackedTangentZ.Vector.W = GetBasisDeterminantSignByte(TangentX, TangentY, TangentZ);
+			}
+		}
+	}
+}
+
+
+void FHairStrandsDatas::BuildSimulationDatas(const uint32 StrandSize,
+	TArray<FHairStrandsPositionFormat::Type>& OutNodesPositions,
+	TArray<FHairStrandsWeightFormat::Type>& OutPointsWeights,
+	TArray<FHairStrandsIndexFormat::Type>& OutPointsNodes) const
+{
+	OutNodesPositions.SetNum(GetNumCurves()*StrandSize);
+	OutPointsWeights.SetNum(GetNumPoints());
+	OutPointsNodes.SetNum(GetNumPoints());
+
+	if (GetNumCurves() > 0 && GetNumPoints() > 0 && StrandSize > 1)
+	{
+		TArray<FVector>::TConstIterator PositionIterator = StrandsPoints.PointsPosition.CreateConstIterator();
+		TArray<uint16>::TConstIterator CountIterator = StrandsCurves.CurvesCount.CreateConstIterator();
+		TArray<float>::TConstIterator LengthIterator = StrandsCurves.CurvesLength.CreateConstIterator();
+
+		TArray<FHairStrandsIndexFormat::Type>::TIterator NodesIterator = OutPointsNodes.CreateIterator();
+
+		uint32 NumPoints = 0;
+		for (uint32 CurveIndex = 0; CurveIndex < GetNumCurves(); ++CurveIndex, ++LengthIterator, ++CountIterator)
+		{
+			const float SeparationLength = (*LengthIterator) * StrandsCurves.MaxLength / (StrandSize - 1);
+			uint32 SampleOffset = CurveIndex * StrandSize;
+			const uint16 EdgeCount = *CountIterator - 1;
+
+			CopyVectorToPosition(*PositionIterator, OutNodesPositions[SampleOffset]);
+			CopyVectorToPosition(*(PositionIterator + EdgeCount), OutNodesPositions[SampleOffset + StrandSize - 1]);
+
+			FVector VertexNext = *PositionIterator;
+			FVector VertexPrev = VertexNext;
+
+			++PositionIterator;
+
+			float StrandLength = 0.0;
+			uint32 LocalCount = 2;
+			for (uint16 EdgeIndex = 0; EdgeIndex < EdgeCount; ++EdgeIndex, ++PositionIterator, ++NodesIterator, ++NumPoints)
+			{
+				VertexPrev = VertexNext;
+				VertexNext = *PositionIterator;
+
+				*NodesIterator = SampleOffset;
+
+				FVector EdgeDirection = VertexNext - VertexPrev;
+				const float CurrentLength = StrandLength;
+				const float EdgeLength = EdgeDirection.Size();
+				StrandLength += EdgeLength;
+
+				if (StrandLength > SeparationLength)
+				{
+					EdgeDirection /= EdgeLength;
+					const uint32 NumNodes = StrandLength / SeparationLength;
+					FVector EdgePosition = VertexPrev + EdgeDirection * (SeparationLength - CurrentLength);
+					for (uint32 NodeIndex = 0; NodeIndex < NumNodes; ++NodeIndex, EdgePosition += EdgeDirection * SeparationLength)
+					{
+						if (LocalCount < StrandSize)
+						{
+							++SampleOffset;
+							CopyVectorToPosition(EdgePosition, OutNodesPositions[SampleOffset]);
+							++LocalCount;
+							StrandLength -= SeparationLength;
+						}
+					}
+				}
+			}
+			*NodesIterator = SampleOffset;
+			++NodesIterator;
+		}
+		PositionIterator.Reset();
+		NodesIterator.Reset();
+		LengthIterator.Reset();
+		CountIterator.Reset();
+
+		TArray<float>::TIterator WeightsIterator = OutPointsWeights.CreateIterator();
+
+		for (uint32 CurveIndex = 0; CurveIndex < GetNumCurves(); ++CurveIndex, ++LengthIterator, ++CountIterator)
+		{
+			const uint16& StrandCount = *CountIterator;
+			const float SeparationLength = (*LengthIterator) * StrandsCurves.MaxLength / (StrandSize - 1);
+			const float LengthScale = 1.0 / (SeparationLength * SeparationLength);
+
+			for (uint32 PointIndex = 0; PointIndex < StrandCount; ++PointIndex, ++PositionIterator, ++NodesIterator, ++WeightsIterator)
+			{
+				const uint32 NodeIndex = *NodesIterator;
+				FVector PrevPosition, NextPosition;
+				CopyPositionToVector(OutNodesPositions[NodeIndex], PrevPosition);
+				CopyPositionToVector(OutNodesPositions[NodeIndex+1], NextPosition);
+				const float PrevDist = (*PositionIterator - PrevPosition).SizeSquared() * LengthScale;
+				const float NextDist = (*PositionIterator - NextPosition).SizeSquared() * LengthScale;
+
+				float PrevWeight = (PrevDist < 1.0) ? (1.0 - PrevDist)* (1.0 - PrevDist)* (1.0 - PrevDist) : 0.0;
+				float NextWeight = (NextDist < 1.0) ? (1.0 - NextDist)* (1.0 - NextDist)* (1.0 - NextDist) : 0.0;
+
+				const float SumWeights = PrevWeight + NextWeight;
+				if (SumWeights != 0.0)
+				{
+					PrevWeight /= SumWeights;
+					NextWeight /= SumWeights;
+				}
+				*WeightsIterator = PrevWeight;
+			}
+		}
+	}
+}
+
+void FHairStrandsDatas::BuildInternalDatas()
+{
+	if (GetNumCurves() > 0 && GetNumPoints() > 0)
+	{
+		TArray<FVector>::TIterator PositionIterator = StrandsPoints.PointsPosition.CreateIterator();
+		TArray<float>::TIterator RadiusIterator = StrandsPoints.PointsRadius.CreateIterator();
+		TArray<float>::TIterator CoordUIterator = StrandsPoints.PointsCoordU.CreateIterator();
+
+		TArray<uint16>::TIterator CountIterator = StrandsCurves.CurvesCount.CreateIterator();
+		TArray<uint32>::TIterator OffsetIterator = StrandsCurves.CurvesOffset.CreateIterator();
+		TArray<float>::TIterator LengthIterator = StrandsCurves.CurvesLength.CreateIterator();
+
+		StrandsPoints.MaxRadius = 0.0;
+		StrandsCurves.MaxLength = 0.0;
+
+		uint32 StrandOffset = 0;
+		*OffsetIterator = StrandOffset; ++OffsetIterator;
+
+		for (uint32 CurveIndex = 0; CurveIndex < GetNumCurves(); ++CurveIndex, ++OffsetIterator, ++LengthIterator, ++CountIterator)
+		{
+			const uint16& StrandCount = *CountIterator;
+
+			StrandOffset += StrandCount;
+			*OffsetIterator = StrandOffset;
+
+			float StrandLength = 0.0;
+			FVector PreviousPosition(0.0, 0.0, 0.0);
+			for (uint32 PointIndex = 0; PointIndex < StrandCount; ++PointIndex, ++PositionIterator, ++RadiusIterator, ++CoordUIterator)
+			{
+				if (PointIndex > 0)
+				{
+					StrandLength += (*PositionIterator - PreviousPosition).Size();
+				}
+				*CoordUIterator = StrandLength;
+				PreviousPosition = *PositionIterator;
+
+				StrandsPoints.MaxRadius = FMath::Max(StrandsPoints.MaxRadius, *RadiusIterator);
+			}
+			*LengthIterator = StrandLength;
+			StrandsCurves.MaxLength = FMath::Max(StrandsCurves.MaxLength, StrandLength);
+		}
+
+		CountIterator.Reset();
+		LengthIterator.Reset();
+		RadiusIterator.Reset();
+		CoordUIterator.Reset();
+
+		for (uint32 CurveIndex = 0; CurveIndex < GetNumCurves(); ++CurveIndex, ++LengthIterator, ++CountIterator)
+		{
+			const uint16& StrandCount = *CountIterator;
+
+			for (uint32 PointIndex = 0; PointIndex < StrandCount; ++PointIndex, ++RadiusIterator, ++CoordUIterator)
+			{
+				*CoordUIterator /= *LengthIterator;
+				*RadiusIterator /= StrandsPoints.MaxRadius;
+			}
+			*LengthIterator /= StrandsCurves.MaxLength;
+		}
+	}
+}
+
+
+void FHairStrandsDatas::AttachStrandsRoots(UStaticMesh* StaticMesh, const FMatrix& TransformMatrix)
+{
+	/*if (StaticMesh)
+	{
+		FTriMeshCollisionData CollisionData;
+		StaticMesh->GetPhysicsTriMeshData(&CollisionData, true);
+
+		Chaos::TParticles<float, 3> Particles =  MoveTemp(*reinterpret_cast<TArray<Chaos::TVector<float, 3>>*>(&CollisionData.Vertices));
+		TArray<Chaos::TVector<int32, 3>> Triangles = MoveTemp(*reinterpret_cast<TArray<Chaos::TVector<int32, 3>>*>(&CollisionData.Indices)); 
+
+		TUniquePtr<Chaos::TTriangleMeshImplicitObject<float>> TriangleMesh = MakeUnique<Chaos::TTriangleMeshImplicitObject<float>>(MoveTemp(Particles), MoveTemp(Triangles));
+
+		TArray<int32> TriangleIndices;
+		TriangleIndices.Init(-1, GetNumCurves());
+
+		TArray<Chaos::TVector<float, 2>> BarycentricCoordinates;
+		BarycentricCoordinates.Init(Chaos::TVector<float, 2>(0, 0), GetNumCurves());
+
+		Chaos::PhysicsParallelFor(GetNumCurves(), [&](int32 CurveIndex) {
+
+			const FVector& RootPosition = StrandsPoints.PointsPosition[StrandsCurves.CurvesOffset[CurveIndex]];
+			const float Dist = TriangleMesh->FindClosestPoint(RootPosition, TriangleIndices[CurveIndex], BarycentricCoordinates[CurveIndex]);
+		});
+
+		// build cell domain
+		int32 MaxAxisSize = 50;
+		int MaxAxis = TriangleMesh->MLocalBoundingBox.LargestAxis();
+		Chaos::TVector<float, 3> Extents = TriangleMesh->MLocalBoundingBox.Extents();
+		Chaos::TVector<int32, 3> Counts(MaxAxisSize * Extents[0] / Extents[MaxAxis], MaxAxisSize * Extents[1] / Extents[MaxAxis], MaxAxisSize * Extents[2] / Extents[MaxAxis]);
+		Counts[0] = Counts[0] < 1 ? 1 : Counts[0];
+		Counts[1] = Counts[1] < 1 ? 1 : Counts[1]; 
+		Counts[2] = Counts[2] < 1 ? 1 : Counts[2];
+
+		Chaos::TUniformGrid<float, 3> Grid(TriangleMesh->MLocalBoundingBox.Min(), TriangleMesh->MLocalBoundingBox.Max(), Counts, 1);
+		Chaos::FErrorReporter ErrorReporter;
+		Chaos::TTriangleMesh<float> CollisionMesh(MoveTemp(TriangleMesh->MElements));
+		Chaos::TLevelSet<float, 3> LevelSet(ErrorReporter, Grid, TriangleMesh->MParticles, CollisionMesh, 0, false);
+
+		for (uint32 CurveIndex = 0; CurveIndex < 10; ++CurveIndex)
+		{
+			for (uint32 PointIndex = StrandsCurves.CurvesOffset[CurveIndex]; PointIndex < StrandsCurves.CurvesOffset[CurveIndex + 1]; ++PointIndex)
+			{
+				const FVector RootPosition = StrandsPoints.PointsPosition[PointIndex];
+				UE_LOG(LogHairStrands, Log, TEXT("Signed Distance = %d %d %f"), CurveIndex, PointIndex, LevelSet.SignedDistance(RootPosition));
+			}
+		}
+	}*/
+}
diff --git a/Engine/Plugins/Experimental/HairStrands/Source/HairStrandsCore/Public/HairStrandsAsset.h b/Engine/Plugins/Experimental/HairStrands/Source/HairStrandsCore/Public/HairStrandsAsset.h
new file mode 100644
--- /dev/null
+++ b/Engine/Plugins/Experimental/HairStrands/Source/HairStrandsCore/Public/HairStrandsAsset.h
@@ -0,0 +1,135 @@
+// Copyright 1998-2019 Epic Games, Inc. All Rights Reserved.
+
+
+#pragma once
+
+#include "UObject/Object.h"
+#include "UObject/ObjectMacros.h"
+#include "HairStrandsDatas.h"
+#include "RenderResource.h"
+
+#include "HairStrandsAsset.generated.h"
+
+
+/**
+ * An instance of a hair strands.
+ */
+class FHairStrandsInstance
+{
+public:
+
+	/** The hair strands resource. Kept an active ref count to manage release with UHairStrands. */
+	struct FHairStrandsResource* RenderResource;
+
+	/** World Transform of the hair strands */
+	FMatrix LocalToGlobal;
+
+	/** Instance destructor */
+	~FHairStrandsInstance();
+
+	/**
+	 * Initializes the instance for the given resource.
+	 * @param InResource - The resource to be used by this instance.
+	 * @param bInstanced - true if the resource is instanced and ownership is being transferred.
+	 */
+	void InitResource(struct FHairStrandsResource* InResource, bool bInstanced);
+
+	/**
+	 * Update the transforms for this hair strands instance.
+	 * @param LocalToWorld - Transform from local space to world space.
+	 */
+	void UpdateTransforms(const FMatrix& LocalToWorld);
+
+private:
+
+	/** true if the resource is instanced and owned by this instance. */
+	bool bInstancedResource;
+};
+
+/* Render buffers that will be used for rendering */
+struct FHairStrandsResource : public FRenderResource
+{
+	/** Build the hair strands resource */
+	FHairStrandsResource(FHairStrandsDatas* HairStrandsDatas);
+
+	/* Init the buffer */
+	virtual void InitRHI() override;
+
+	/* Release the buffer */
+	virtual void ReleaseRHI() override;
+
+	/* Get the resource name */
+	virtual FString GetFriendlyName() const override { return TEXT("FHairStrandsResource"); }
+
+	/* Strand hair position buffer */
+	FRWBuffer PositionBuffer;
+
+	/* Strand hair offset buffer */
+	FRWBuffer TangentBuffer;
+
+	/* Pointer to the hair strands datas */
+	FHairStrandsDatas* StrandsDatas;
+};
+
+/**
+ * Implements an asset that can be used to store hair strands
+ */
+UCLASS(BlueprintType, hidecategories = (Object))
+class HAIRSTRANDSCORE_API UHairStrandsAsset : public UObject
+{
+	GENERATED_BODY()
+
+public:
+
+	/** Holds the File Path. */
+	UPROPERTY(BlueprintReadOnly, EditAnywhere, Category = "HairStrands")
+	FString FilePath;
+
+	/** Hair strands datas */
+	FHairStrandsDatas StrandsDatas;
+
+	//~ Begin UObject Interface.
+	virtual void PostLoad() override;
+	virtual void BeginDestroy() override;
+	virtual void Serialize(FArchive& Ar) override;
+
+#if WITH_EDITOR
+
+	/**  Part of Uobject interface  */
+	virtual void PostEditChangeProperty(FPropertyChangedEvent& PropertyChangedEvent) override;
+
+#endif // WITH_EDITOR
+
+#if WITH_EDITORONLY_DATA
+
+	/** Asset data to be used when re-importing */
+	UPROPERTY(VisibleAnywhere, Instanced, Category = ImportSettings)
+	class UAssetImportData* AssetImportData;
+
+	/** Retrievde the asset tags*/
+	virtual void GetAssetRegistryTags(TArray<FAssetRegistryTag>& OutTags) const override;
+
+	/** Part of Uobject interface */
+	virtual void PostInitProperties() override;
+
+#endif // WITH_EDITORONLY_DATA
+
+	/** Initialize resources. */
+	void InitResource();
+
+	/** Update resources. */
+	void UpdateResource();
+
+	/** Release the hair strands resource. */
+	void ReleaseResource();
+
+	/**
+	 * Initializes an instance for use with this vector field.
+	 */
+	void InitInstance(class FVectorFieldInstance* Instance, bool bPreviewInstance);
+
+//private : 
+
+	/** Render resource to be allocated*/
+	FHairStrandsResource* HairStrandsResource;
+};
diff --git a/Engine/Plugins/Experimental/HairStrands/Source/HairStrandsCore/Public/HairStrandsComponent.h b/Engine/Plugins/Experimental/HairStrands/Source/HairStrandsCore/Public/HairStrandsComponent.h
new file mode 100644
--- /dev/null
+++ b/Engine/Plugins/Experimental/HairStrands/Source/HairStrandsCore/Public/HairStrandsComponent.h
@@ -0,0 +1,57 @@
+// Copyright 1998-2019 Epic Games, Inc. All Rights Reserved.
+
+#pragma once
+
+#include "CoreMinimal.h"
+#include "UObject/ObjectMacros.h"
+#include "Engine/EngineTypes.h"
+#include "Components/MeshComponent.h"
+#include "HairStrandsAsset.h"
+
+#include "HairStrandsComponent.generated.h"
+
+/** Component that allows you to specify custom triangle mesh geometry */
+UCLASS(hidecategories = (Object, Physics, Activation, "Components|Activation"), editinlinenew, meta = (BlueprintSpawnableComponent), ClassGroup = Rendering)
+class HAIRSTRANDSCORE_API UHairStrandsComponent : public UMeshComponent
+{
+	GENERATED_UCLASS_BODY()
+
+public:
+
+	UPROPERTY(EditAnywhere, BlueprintReadWrite, Category = "HairStrands")
+	UHairStrandsAsset* HairStrandsAsset;
+
+	/** Controls the hair density, to reduce or increase hair count during shadow rendering. This allows to increase/decrease the shadowing on hair when the number of strand is not realistic */
+	UPROPERTY(EditAnywhere, BlueprintReadWrite, Category = "HairStrands Rendering", meta = (ClampMin = "0.0001", UIMin = "0.01", UIMax = "10.0"))
+	float HairDensity;
+
+	/** Threshold for merging consecutive hair segments when loading asset */
+	UPROPERTY(EditAnywhere, BlueprintReadWrite, Category = "HairStrands Rendering", meta = (ClampMin = "0.01", UIMin = "0.01", UIMax = "1.0"))
+	float MergeThreshold;
+
+	/** Enable this asset to be rendered within the offline reference renderer */
+	UPROPERTY(EditAnywhere, BlueprintReadOnly, Category = Asset, meta = (MultiLine = false))
+	bool bUsedForReference;
+
+	/** The instance of the hair strands. */
+	FHairStrandsInstance* HairStrandsInstance;
+
+	//~ Begin UActorComponent Interface.
+	virtual void OnRegister() override;
+	virtual void OnUnregister() override;
+	virtual void TickComponent(float DeltaTime, enum ELevelTick TickType, FActorComponentTickFunction *ThisTickFunction) override;
+	//~ End UActorComponent Interface.
+
+	//~ Begin USceneComponent Interface.
+	virtual FBoxSphereBounds CalcBounds(const FTransform& LocalToWorld) const override;
+	//~ End USceneComponent Interface.
+
+	//~ Begin UPrimitiveComponent Interface.
+	virtual FPrimitiveSceneProxy* CreateSceneProxy() override;
+	//~ End UPrimitiveComponent Interface.
+
+	//~ Begin UMeshComponent Interface.
+	virtual void PostLoad() override;
+	virtual int32 GetNumMaterials() const override;
+	//~ End UMeshComponent Interface.
+};
diff --git a/Engine/Plugins/Experimental/HairStrands/Source/HairStrandsCore/Public/HairStrandsCore.h b/Engine/Plugins/Experimental/HairStrands/Source/HairStrandsCore/Public/HairStrandsCore.h
new file mode 100644
--- /dev/null
+++ b/Engine/Plugins/Experimental/HairStrands/Source/HairStrandsCore/Public/HairStrandsCore.h
@@ -0,0 +1,24 @@
+// Copyright 1998-2019 Epic Games, Inc. All Rights Reserved.
+
+#pragma once
+
+#include "Modules/ModuleInterface.h"
+#include "Modules/ModuleManager.h"
+
+/** Implements the HairStrands module  */
+class FHairStrandsCore : public IModuleInterface
+{
+public:
+
+	//~ IModuleInterface interface
+
+	virtual void StartupModule() override;
+	virtual void ShutdownModule() override;
+
+	virtual bool SupportsDynamicReloading() override
+	{
+		return true;
+	}
+};
+
+IMPLEMENT_MODULE(FHairStrandsCore, HairStrandsCore);
diff --git a/Engine/Plugins/Experimental/HairStrands/Source/HairStrandsCore/Public/HairStrandsDatas.h b/Engine/Plugins/Experimental/HairStrands/Source/HairStrandsCore/Public/HairStrandsDatas.h
new file mode 100644
--- /dev/null
+++ b/Engine/Plugins/Experimental/HairStrands/Source/HairStrandsCore/Public/HairStrandsDatas.h
@@ -0,0 +1,239 @@
+// Copyright 1998-2019 Epic Games, Inc. All Rights Reserved.
+
+#pragma once
+
+#include "CoreMinimal.h"
+#include "Engine/EngineTypes.h"
+#include "RHIDefinitions.h"
+#include "PackedNormal.h"
+
+DEFINE_LOG_CATEGORY_STATIC(LogHairStrands, Log, All);
+
+/**  A packed position for the hair strands datas .*/
+struct HAIRSTRANDSCORE_API FHairStrandsPackedPosition
+{
+	// #todo_hair: repack the position into relative 16 bits position. These position should be relative to the AABB center. 
+	// If Float point is not enough we could use fixed point value
+#if 0 
+	//FFloat16 X, Y, Z; 
+	//uint16 ControlPointType : 2; 
+	//uint16 NormalizedRadius : 6;
+	//uint16 UCoord : 8;
+#else
+	float X, Y, Z;
+	uint32 ControlPointType : 2;
+	uint32 NormalizedRadius : 6;
+	uint32 UCoord : 8;
+	uint32 NormalizedLength : 8;
+	uint32 Seed : 8;
+#endif
+};
+
+/** Hair strands position format */
+struct HAIRSTRANDSCORE_API FHairStrandsPositionFormat
+{
+	using Type = FHairStrandsPackedPosition;
+	static const uint32 ComponentCount = 1;
+	static const uint32 SizeInByte = sizeof(Type);
+	static const EVertexElementType VertexElementType = VET_Float4;	// VET_UShort4;
+	static const EPixelFormat Format = PF_A32B32G32R32F;			// PF_R16G16B16A16_UINT;
+};
+
+/** Hair Strands tangent format */
+struct HAIRSTRANDSCORE_API FHairStrandsTangentFormat
+{
+	// TangentX & tangentZ are packed into 2 * PF_R8G8B8A8_SNORM
+	using Type = FPackedNormal;
+	static const uint32 ComponentCount = 2;
+	static const uint32 SizeInByte = sizeof(Type);
+	static const EVertexElementType VertexElementType = VET_Float4;
+	static const EPixelFormat Format = PF_R8G8B8A8_SNORM;
+};
+
+/** Hair strands index format */
+struct HAIRSTRANDSCORE_API FHairStrandsIndexFormat
+{
+	using Type = uint32;
+	static const uint32 ComponentCount = 1;
+	static const uint32 SizeInByte = sizeof(Type);
+	static const EVertexElementType VertexElementType = VET_UInt;
+	static const EPixelFormat Format = PF_R32_UINT;		
+};
+
+/** Hair strands weights format */
+struct HAIRSTRANDSCORE_API FHairStrandsWeightFormat
+{
+	using Type = float;
+	static const uint32 ComponentCount = 1;
+	static const uint32 SizeInByte = sizeof(Type);
+	static const EVertexElementType VertexElementType = VET_Float1;	
+	static const EPixelFormat Format = PF_R32_FLOAT;			
+};
+
+/** Hair strands quaternion format */
+struct HAIRSTRANDSCORE_API FHairStrandsQuaternionFormat
+{
+	using Type = FQuat;
+	static const uint32 ComponentCount = 1;
+	static const uint32 SizeInByte = sizeof(Type);
+	static const EVertexElementType VertexElementType = VET_Float4;	
+	static const EPixelFormat Format = PF_A32B32G32R32F;			
+};
+
+/** Hair strands points attribute */
+struct HAIRSTRANDSCORE_API FHairStrandsPoints
+{
+	/** Serialize the points */
+	inline void Serialize(FArchive& Ar);
+
+	/** Set the number of points */
+	inline void SetNum(const uint32 NumPoints);
+
+	/** Reset the points to 0 */
+	inline void Reset();
+
+	/** Get the number of points */
+	inline uint32 Num() const { return PointsPosition.Num();  }
+
+	/** Points position in local space */
+	TArray<FVector> PointsPosition;
+
+	/** Normalized radius relative to the max one */
+	TArray<float> PointsRadius; // [0..1]
+
+	/** Normalized length relative to the strand one */
+	TArray<float> PointsCoordU; // [0..1]
+
+	/** Max strands points radius */
+	float MaxRadius = 0;
+};
+
+/** Hair strands Curves attribute */
+struct HAIRSTRANDSCORE_API FHairStrandsCurves
+{
+	/** Serialize the curves */
+	inline void Serialize(FArchive& Ar);
+
+	/** Set the number of Curves */
+	inline void SetNum(const uint32 NumPoints);
+
+	/** Reset the curves to 0 */
+	inline void Reset();
+
+	/** Get the number of Curves */
+	inline uint32 Num() const { return CurvesLength.Num(); }
+
+	/** Number of points per rod */
+	TArray<uint16> CurvesCount;
+
+	/** An offset represent the rod start in the point list */
+	TArray<uint32> CurvesOffset;
+
+	/** Normalized length relative to the max one */
+	TArray<float> CurvesLength; // [0..1]
+
+	/** Roots UV in the bounding sphere */
+	TArray<FVector2D> CurvesRootUV; // [0..1]
+
+	/** Max strands Curves length */
+	float MaxLength = 0;
+};
+
+/** Hair strands datas that are stored on CPU */
+struct HAIRSTRANDSCORE_API FHairStrandsDatas
+{
+	/** Serialize all the hair strands datas */
+	inline void Serialize(FArchive& Ar);
+
+	/* Get the total number of points */
+	inline uint32 GetNumPoints() const { return StrandsPoints.Num(); }
+
+	/* Get the total number of Curves */
+	inline uint32 GetNumCurves() const { return StrandsCurves.Num(); }
+
+	/** List of all the strands points */
+	FHairStrandsPoints StrandsPoints;
+
+	/** List of all the strands points */
+	FHairStrandsCurves StrandsCurves;
+
+	/** The Standard Hair Density */
+	float HairDensity = 1;
+
+	/* Strands bounding box */
+	FBox BoundingBox;
+
+	/** Build the packed datas for gpu rendering/simulation */
+	void BuildRenderingDatas(
+		TArray<FHairStrandsPositionFormat::Type>& OutPointsPositions,
+		TArray<FHairStrandsTangentFormat::Type>& OutPointsTangents) const;
+
+	/** Resample the internal points and curves simulation datas */
+	void BuildSimulationDatas(const uint32 StrandSize,
+		TArray<FHairStrandsPositionFormat::Type>& OutNodesPositions,
+		TArray<FHairStrandsWeightFormat::Type>& OutPointsWeights,
+		TArray<FHairStrandsIndexFormat::Type>& OutPointsNodes) const;
+
+	/** Build the internal points and curves datas */
+	void BuildInternalDatas();
+
+	/** Attach the roots to a static mesh */
+	void AttachStrandsRoots(UStaticMesh* StaticMesh, const FMatrix& TransformMartrix);
+};
+
+void FHairStrandsCurves::SetNum(const uint32 NumCurves)
+{
+	CurvesOffset.SetNum(NumCurves + 1);
+	CurvesCount.SetNum(NumCurves);
+	CurvesLength.SetNum(NumCurves);
+	CurvesRootUV.SetNum(NumCurves);
+}
+
+void FHairStrandsPoints::SetNum(const uint32 NumPoints)
+{
+	PointsPosition.SetNum(NumPoints);
+	PointsRadius.SetNum(NumPoints);
+	PointsCoordU.SetNum(NumPoints);
+}
+
+void FHairStrandsCurves::Reset()
+{
+	CurvesOffset.Reset();
+	CurvesCount.Reset();
+	CurvesLength.Reset();
+	CurvesRootUV.Reset();
+}
+
+void FHairStrandsPoints::Reset()
+{
+	PointsPosition.Reset();
+	PointsRadius.Reset();
+	PointsCoordU.Reset();
+}
+
+void FHairStrandsPoints::Serialize(FArchive& Ar)
+{
+	Ar << PointsPosition;
+	Ar << PointsRadius;
+	Ar << PointsCoordU;
+	Ar << MaxRadius;
+}
+
+void FHairStrandsCurves::Serialize(FArchive& Ar)
+{
+	Ar << CurvesCount;
+	Ar << CurvesOffset;
+	Ar << CurvesLength;
+	Ar << CurvesRootUV;
+	Ar << MaxLength;
+}
+
+void FHairStrandsDatas::Serialize(FArchive& Ar)
+{
+	StrandsPoints.Serialize(Ar);
+	StrandsCurves.Serialize(Ar);
+	Ar << HairDensity;
+	Ar << BoundingBox;
+}
+
+
diff --git a/Engine/Plugins/Experimental/HairStrands/Source/HairStrandsEditor/HairStrandsEditor.Build.cs b/Engine/Plugins/Experimental/HairStrands/Source/HairStrandsEditor/HairStrandsEditor.Build.cs
new file mode 100644
--- /dev/null
+++ b/Engine/Plugins/Experimental/HairStrands/Source/HairStrandsEditor/HairStrandsEditor.Build.cs
@@ -0,0 +1,27 @@
+// Copyright 1998-2019 Epic Games, Inc. All Rights Reserved.
+
+namespace UnrealBuildTool.Rules
+{
+	public class HairStrandsEditor : ModuleRules
+	{
+        public HairStrandsEditor(ReadOnlyTargetRules Target) : base(Target)
+		{
+			PrivateIncludePaths.Add(ModuleDirectory + "/Private");
+            PublicIncludePaths.Add(ModuleDirectory + "/Public");
+
+            PrivateDependencyModuleNames.AddRange(
+                new string[]
+                {
+					"Core",
+					"CoreUObject",
+					"HairStrandsCore",
+					"UnrealEd",
+					"AssetTools",
+                    "AlembicLib",
+                });
+            AddEngineThirdPartyPrivateStaticDependencies(Target,
+             "FBX"
+			);
+        }
+	}
+}
diff --git a/Engine/Plugins/Experimental/HairStrands/Source/HairStrandsEditor/Private/HairStrandsActions.cpp b/Engine/Plugins/Experimental/HairStrands/Source/HairStrandsEditor/Private/HairStrandsActions.cpp
new file mode 100644
--- /dev/null
+++ b/Engine/Plugins/Experimental/HairStrands/Source/HairStrandsEditor/Private/HairStrandsActions.cpp
@@ -0,0 +1,51 @@
+// Copyright 1998-2019 Epic Games, Inc. All Rights Reserved.
+
+#include "HairStrandsActions.h"
+#include "HairStrandsAsset.h"
+
+#define LOCTEXT_NAMESPACE "AssetTypeActions"
+
+FHairStrandsActions::FHairStrandsActions()
+{}
+
+bool FHairStrandsActions::CanFilter()
+{
+	return true;
+}
+
+void FHairStrandsActions::GetActions(const TArray<UObject*>& InObjects, FMenuBuilder& MenuBuilder)
+{
+	FAssetTypeActions_Base::GetActions(InObjects, MenuBuilder);
+}
+
+
+uint32 FHairStrandsActions::GetCategories()
+{
+	return EAssetTypeCategories::Misc;
+}
+
+FText FHairStrandsActions::GetName() const
+{
+	return NSLOCTEXT("AssetTypeActions", "AssetTypeActions_HairStrands", "Hair Strands");
+}
+
+UClass* FHairStrandsActions::GetSupportedClass() const
+{
+	return UHairStrandsAsset::StaticClass();
+}
+
+FColor FHairStrandsActions::GetTypeColor() const
+{
+	return FColor::White;
+}
+
+bool FHairStrandsActions::HasActions(const TArray<UObject*>& InObjects) const
+{
+	return true;
+}
+
+void FHairStrandsActions::OpenAssetEditor(const TArray<UObject*>& InObjects, TSharedPtr<IToolkitHost> EditWithinLevelEditor)
+{
+}
+
+#undef LOCTEXT_NAMESPACE
\ No newline at end of file
diff --git a/Engine/Plugins/Experimental/HairStrands/Source/HairStrandsEditor/Private/HairStrandsEditor.cpp b/Engine/Plugins/Experimental/HairStrands/Source/HairStrandsEditor/Private/HairStrandsEditor.cpp
new file mode 100644
--- /dev/null
+++ b/Engine/Plugins/Experimental/HairStrands/Source/HairStrandsEditor/Private/HairStrandsEditor.cpp
@@ -0,0 +1,29 @@
+// Copyright 1998-2019 Epic Games, Inc. All Rights Reserved.
+
+#include "HairStrandsEditor.h"
+#include "HairStrandsActions.h"
+
+void FHairStrandsEditor::StartupModule()
+{
+	IAssetTools& AssetTools = FModuleManager::LoadModuleChecked<FAssetToolsModule>("AssetTools").Get();
+	TSharedRef<IAssetTypeActions> AssetActions = MakeShareable(new FHairStrandsActions());
+
+	AssetTools.RegisterAssetTypeActions(AssetActions);
+	RegisteredAssetTypeActions.Add(AssetActions);
+}
+
+void FHairStrandsEditor::ShutdownModule()
+{
+	FAssetToolsModule* AssetToolsModule = FModuleManager::GetModulePtr<FAssetToolsModule>("AssetTools");
+
+	if (AssetToolsModule != nullptr)
+	{
+		IAssetTools& AssetTools = AssetToolsModule->Get();
+
+		for (auto Action : RegisteredAssetTypeActions)
+		{
+			AssetTools.UnregisterAssetTypeActions(Action);
+		}
+	}
+}
+
diff --git a/Engine/Plugins/Experimental/HairStrands/Source/HairStrandsEditor/Private/HairStrandsFactory.cpp b/Engine/Plugins/Experimental/HairStrands/Source/HairStrandsEditor/Private/HairStrandsFactory.cpp
new file mode 100644
--- /dev/null
+++ b/Engine/Plugins/Experimental/HairStrands/Source/HairStrandsEditor/Private/HairStrandsFactory.cpp
@@ -0,0 +1,157 @@
+// Copyright 1998-2019 Epic Games, Inc. All Rights Reserved.
+
+#include "HairStrandsFactory.h"
+#include "HairStrandsAsset.h"
+#include "HairStrandsLoader.h"
+
+#include "Developer/DesktopPlatform/Public/IDesktopPlatform.h"
+#include "Developer/DesktopPlatform/Public/DesktopPlatformModule.h"
+#include "UnrealEd/Public/EditorDirectories.h"
+
+#include "Misc/Paths.h"
+
+/**
+ * true if the extension is for the .hair, .fbx, .abc format.
+ */
+static bool IsHairStrandsFormat(const TCHAR* Extension)
+{
+	return (FCString::Stricmp(Extension, TEXT("abc")) == 0) || 
+		   (FCString::Stricmp(Extension, TEXT("fbx")) == 0) ||
+		   (FCString::Stricmp(Extension, TEXT("hair")) == 0);
+}
+
+static void BuildHairStrands( const FString& FilePath, FHairStrandsDatas& OutStrandsDatas)
+{
+	if (FCString::Stricmp(*FPaths::GetExtension(FilePath), TEXT("hair")) == 0)
+	{
+		THairStrandsLoader<FHairFormat>::LoadHairStrands(FilePath, OutStrandsDatas);
+	}
+	else if (FCString::Stricmp(*FPaths::GetExtension(FilePath), TEXT("fbx")) == 0)
+	{
+		THairStrandsLoader<FFbxFormat>::LoadHairStrands(FilePath, OutStrandsDatas);
+	}
+	else if (FCString::Stricmp(*FPaths::GetExtension(FilePath), TEXT("abc")) == 0)
+	{
+		THairStrandsLoader<FAbcFormat>::LoadHairStrands(FilePath, OutStrandsDatas);
+	}
+}
+
+UHairStrandsFactory::UHairStrandsFactory(const FObjectInitializer& ObjectInitializer)
+	: Super(ObjectInitializer)
+{
+	SupportedClass = UHairStrandsAsset::StaticClass();
+	bCreateNew = true;
+	bEditAfterNew = true;
+	bEditorImport = false;
+
+	Formats.Add(TEXT("hair;Hair hair strands file"));
+	Formats.Add(TEXT("fbx;Fbx hair strands file"));
+	Formats.Add(TEXT("abc;Alembic hair strands file"));
+}
+
+UObject* UHairStrandsFactory::FactoryCreateNew(UClass* InClass, UObject* InParent, FName InName, EObjectFlags Flags, UObject* Context, FFeedbackContext* Warn)
+{
+	UHairStrandsAsset* ExistingAsset = FindObject<UHairStrandsAsset>(InParent, *InName.ToString());
+	if (ExistingAsset)
+	{
+		ExistingAsset->ReleaseResource();
+	}
+	UHairStrandsAsset* CurrentAsset = NewObject<UHairStrandsAsset>(InParent, InClass, InName, Flags);
+
+	const FString Filter(TEXT("Hair Strands Files (*.hair,*.fbx,*.abc)|*.hair;*.fbx;*.abc"));
+
+	TArray<FString> OpenFilenames;
+	int32 FilterIndex = -1;
+	if (FDesktopPlatformModule::Get()->OpenFileDialog(
+		nullptr,
+		FString(TEXT("Choose a hair strands file")),
+		FEditorDirectories::Get().GetLastDirectory(ELastDirectory::GENERIC_IMPORT),
+		TEXT(""),
+		Filter,
+		EFileDialogFlags::None,
+		OpenFilenames,
+		FilterIndex))
+	{
+		CurrentAsset->FilePath = OpenFilenames[0];
+		BuildHairStrands(CurrentAsset->FilePath, CurrentAsset->StrandsDatas);
+	}
+
+	if (CurrentAsset)
+	{
+		CurrentAsset->InitResource();
+	}
+	return CurrentAsset;
+}
+
+UObject* UHairStrandsFactory::FactoryCreateFile(UClass * InClass, UObject * InParent, FName InName, EObjectFlags Flags, 
+	const FString & Filename, const TCHAR* Parms, FFeedbackContext * Warn, bool& bOutOperationCanceled) 
+{
+	UHairStrandsAsset* ExistingAsset = FindObject<UHairStrandsAsset>(InParent, *InName.ToString());
+	if (ExistingAsset)
+	{
+		ExistingAsset->ReleaseResource();
+	}
+
+	UHairStrandsAsset* CurrentAsset = NewObject<UHairStrandsAsset>(InParent, InClass, InName, Flags);
+
+	CurrentAsset->FilePath = Filename;
+	BuildHairStrands(CurrentAsset->FilePath, CurrentAsset->StrandsDatas);
+
+	if (CurrentAsset)
+	{
+		CurrentAsset->InitResource();
+	}
+
+	return CurrentAsset;
+}
+
+bool UHairStrandsFactory::FactoryCanImport(const FString& Filename)
+{
+	return  IsHairStrandsFormat(*FPaths::GetExtension(Filename));
+}
+
+bool UHairStrandsFactory::CanReimport(UObject* Obj, TArray<FString>& OutFilenames)
+{
+	UHairStrandsAsset* Asset = Cast<UHairStrandsAsset>(Obj);
+	if (Asset)
+	{
+		const FString& FileName = Asset->FilePath;
+		if (!FileName.IsEmpty())
+		{
+			OutFilenames.Add(FileName);
+		}
+
+		return true;
+	}
+	return false;
+}
+
+void UHairStrandsFactory::SetReimportPaths(UObject* Obj, const TArray<FString>& NewReimportPaths)
+{
+	UHairStrandsAsset* Asset = Cast<UHairStrandsAsset>(Obj);
+	if (Asset && ensure(NewReimportPaths.Num() == 1))
+	{
+		Asset->FilePath = NewReimportPaths[0];
+	}
+}
+
+EReimportResult::Type UHairStrandsFactory::Reimport(UObject* Obj)
+{
+	UHairStrandsAsset* CurrentAsset = Cast<UHairStrandsAsset>(Obj);
+	BuildHairStrands(CurrentAsset->FilePath, CurrentAsset->StrandsDatas);
+	// Try to find the outer package so we can dirty it up
+	if (Obj->GetOuter())
+	{
+		Obj->GetOuter()->MarkPackageDirty();
+	}
+	else
+	{
+		Obj->MarkPackageDirty();
+	}
+	return EReimportResult::Succeeded;
+}
+
+int32 UHairStrandsFactory::GetPriority() const
+{
+	return ImportPriority;
+}
diff --git a/Engine/Plugins/Experimental/HairStrands/Source/HairStrandsEditor/Private/HairStrandsLoader.cpp b/Engine/Plugins/Experimental/HairStrands/Source/HairStrandsEditor/Private/HairStrandsLoader.cpp
new file mode 100644
--- /dev/null
+++ b/Engine/Plugins/Experimental/HairStrands/Source/HairStrandsEditor/Private/HairStrandsLoader.cpp
@@ -0,0 +1,293 @@
+// Copyright 1998-2019 Epic Games, Inc. All Rights Reserved.
+
+#include "HairStrandsLoader.h"
+#include "HairStrandsDatas.h"
+#include "GenericPlatform/GenericPlatformFile.h"
+#include "HAL/PlatformFilemanager.h"
+
+#if PLATFORM_WINDOWS
+#include "Windows/WindowsHWrapper.h"
+#include "Windows/AllowWindowsPlatformTypes.h"
+#endif
+
+THIRD_PARTY_INCLUDES_START
+#include <fbxsdk.h>
+#include <Alembic/AbcGeom/All.h>
+#include <Alembic/AbcCoreFactory/IFactory.h>
+#include <Alembic/Abc/IArchive.h>
+#include <Alembic/Abc/IObject.h>
+THIRD_PARTY_INCLUDES_END
+
+#if PLATFORM_WINDOWS
+#include "Windows/HideWindowsPlatformTypes.h"
+#endif
+
+#define HAIR_FILE_SEGMENTS_BIT		1
+#define HAIR_FILE_POINTS_BIT		2
+#define HAIR_FILE_THICKNESS_BIT		4
+#define HAIR_FILE_TRANSPARENCY_BIT	8
+#define HAIR_FILE_COLORS_BIT		16  
+
+static const float RootRadius = 0.0001f; // m
+static const float TipRadius = 0.00005f; // m
+
+void FHairFormat::ParseFile(const FString& FileName, FHairStrandsDatas& HairStrands)  
+{
+	IPlatformFile& PlatformFile = FPlatformFileManager::Get().GetPlatformFile();
+	IFileHandle* FileHandle = PlatformFile.OpenRead(*FileName);
+
+	if (FileHandle)
+	{
+		FFileHeader FileHeader;
+		if (FileHandle->Read(reinterpret_cast<uint8*>(&FileHeader), sizeof(FFileHeader)))
+		{
+			if (strncmp(FileHeader.FileSignature, "HAIR", 4) == 0)
+			{
+				HairStrands.StrandsPoints.SetNum(FileHeader.NumPoints);
+				HairStrands.StrandsCurves.SetNum(FileHeader.NumStrands);
+
+				// Read strands counts
+				uint16* StrandsCounts = nullptr;
+				if (FileHeader.BitArrays & HAIR_FILE_SEGMENTS_BIT)
+				{
+					FileHandle->Read(reinterpret_cast<uint8*>(HairStrands.StrandsCurves.CurvesCount.GetData()), sizeof(uint16)*FileHeader.NumStrands);
+				}
+				else
+				{
+					HairStrands.StrandsCurves.CurvesCount.Init(FileHeader.StrandThickness, FileHeader.NumStrands);
+				}
+
+				// Read strands positions
+				float* StrandsPositions = nullptr;
+				if (FileHeader.BitArrays & HAIR_FILE_POINTS_BIT)
+				{
+					FileHandle->Read(reinterpret_cast<uint8*>(&HairStrands.StrandsPoints.PointsPosition.GetData()->X), sizeof(float)*FileHeader.NumPoints * 3);
+				}
+				else
+				{
+					HairStrands.StrandsPoints.PointsPosition.Init(FVector(0.0, 0.0, 0.0), FileHeader.NumPoints);
+				}
+
+				// Read strands positions thickness
+				float* StrandsThickness = nullptr;
+				if (FileHeader.BitArrays & HAIR_FILE_THICKNESS_BIT)
+				{
+					FileHandle->Read(reinterpret_cast<uint8*>(HairStrands.StrandsPoints.PointsRadius.GetData()), sizeof(float)*FileHeader.NumPoints);
+				}
+				else
+				{
+					HairStrands.StrandsPoints.PointsRadius.Init(FileHeader.StrandThickness, FileHeader.NumPoints);
+				}
+			}
+		}
+		delete FileHandle;
+	}
+}
+
+static void ParseNode(FbxNode* FileNode, FHairStrandsDatas& HairStrands)
+{
+	if (FileNode->GetNodeAttribute() != NULL && FileNode->GetNodeAttribute()->GetAttributeType() == FbxNodeAttribute::eNurbsCurve)
+	{
+		FbxNurbsCurve* NurbsCurve = static_cast<FbxNurbsCurve*>(FileNode->GetNodeAttribute());
+		FbxArray<FbxVector4> PointArray;
+		const uint32 PointCount = NurbsCurve->TessellateCurve(PointArray, 2);
+		
+		const uint32 NumPoints = HairStrands.GetNumPoints();
+		const uint32 NumCurves = HairStrands.GetNumCurves();
+
+		HairStrands.StrandsPoints.SetNum(NumPoints+PointCount);
+		HairStrands.StrandsCurves.SetNum(NumCurves + 1);
+		HairStrands.StrandsCurves.CurvesCount[NumCurves] = PointCount;
+
+		for (uint32 PointIndex = 0; PointIndex < PointCount; ++PointIndex)
+		{
+			const float CoordU = PointIndex / static_cast<float>(PointCount - 1);
+			const float Radius = FMath::Lerp(RootRadius, TipRadius, CoordU);
+			const FVector Position(PointArray[PointIndex][0], PointArray[PointIndex][1], PointArray[PointIndex][2]);
+
+			HairStrands.StrandsPoints.PointsPosition[NumPoints+ PointIndex] = Position;
+			HairStrands.StrandsPoints.PointsRadius[NumPoints + PointIndex] = Radius;
+		}
+	}
+	//	#todo_hair
+	// * Load user/custom properties (eccentricity, width/thickness, ...)
+	// * Load root/hierarchy transform
+
+	const uint32 NumChildren = FileNode->GetChildCount();
+	if (NumChildren > 0)
+	{
+		for (int32 ChildIndex = 0; ChildIndex < FileNode->GetChildCount(); ++ChildIndex)
+		{
+			ParseNode(FileNode->GetChild(ChildIndex), HairStrands);
+		}
+	}
+}
+
+void FFbxFormat::ParseFile(const FString& FileName, FHairStrandsDatas& HairStrands)
+{
+	FbxManager* FileManager = FbxManager::Create();
+	FbxIOSettings* FileSettings = FbxIOSettings::Create(FileManager, IOSROOT);
+
+	FileManager->SetIOSettings(FileSettings);
+
+	FbxScene* FileScene = FbxScene::Create(FileManager, "StrandAssetFbx");
+
+	// Create an importer.
+	FbxImporter* FileImporter = FbxImporter::Create(FileManager, "");
+
+	// Initialize the importer by providing a filename.
+	bool FileStatus = FileImporter->Initialize(TCHAR_TO_UTF8(*FileName), -1, FileSettings);
+
+	if (FileImporter->IsFBX())
+	{
+		// Only load geometry data
+		FileSettings->SetBoolProp(IMP_FBX_MATERIAL, false);
+		FileSettings->SetBoolProp(IMP_FBX_TEXTURE, false);
+		FileSettings->SetBoolProp(IMP_FBX_LINK, false);
+		FileSettings->SetBoolProp(IMP_FBX_SHAPE, true);
+		FileSettings->SetBoolProp(IMP_FBX_GOBO, false);
+		FileSettings->SetBoolProp(IMP_FBX_ANIMATION, false);
+		FileSettings->SetBoolProp(IMP_FBX_GLOBAL_SETTINGS, false);
+	}
+
+	// Import the scene and destroy the importer.
+	FileStatus = FileImporter->Import(FileScene);
+	FileImporter->Destroy();
+
+	if (FbxNode* FileNode = FileScene->GetRootNode())
+	{
+		for (int32 ChildIndex = 0; ChildIndex < FileNode->GetChildCount(); ++ChildIndex)
+		{
+			ParseNode(FileNode->GetChild(ChildIndex), HairStrands);
+		}
+	}
+}
+
+FMatrix ConvertAlembicMatrix(const Alembic::Abc::M44d& AbcMatrix)
+{
+	FMatrix Matrix;
+	for (uint32 i = 0; i < 16; ++i)
+	{
+		Matrix.M[i >> 2][i % 4] = (float)AbcMatrix.getValue()[i];
+	}
+
+	return Matrix;
+}
+
+static void ParseObject(const Alembic::Abc::IObject& InObject, FHairStrandsDatas& HairStrands, const FMatrix& ParentMatrix)
+{
+	// Get Header and MetaData info from current Alembic Object
+	Alembic::AbcCoreAbstract::ObjectHeader Header = InObject.getHeader();
+	const Alembic::Abc::MetaData ObjectMetaData = InObject.getMetaData();
+	const uint32 NumChildren = InObject.getNumChildren();
+
+	FMatrix LocalMatrix = ParentMatrix;
+
+	bool bHandled = false;
+	if (Alembic::AbcGeom::ICurves::matches(ObjectMetaData))
+	{
+		Alembic::AbcGeom::ICurves Curves = Alembic::AbcGeom::ICurves(InObject, Alembic::Abc::kWrapExisting);
+		Alembic::AbcGeom::ICurves::schema_type::Sample Sample = Curves.getSchema().getValue();
+
+		Alembic::Abc::FloatArraySamplePtr Widths = Curves.getSchema().getWidthsParam() ? Curves.getSchema().getWidthsParam().getExpandedValue().getVals() : nullptr;
+		Alembic::Abc::P3fArraySamplePtr Positions = Sample.getPositions();
+		Alembic::Abc::Int32ArraySamplePtr Counts = Sample.getCurvesNumVertices();
+
+		const uint32 PointsSize = Positions ? Positions->size() : 0;
+		const uint32 CurvesSize = Counts->size();
+
+		const uint32 NumPoints = HairStrands.GetNumPoints();
+		const uint32 NumCurves = HairStrands.GetNumCurves();
+
+		HairStrands.StrandsCurves.SetNum(NumCurves+CurvesSize);
+		HairStrands.StrandsPoints.SetNum(NumPoints+PointsSize);
+
+		uint32 GlobalIndex = 0;
+		for (uint32 CurveIndex = 0; CurveIndex < CurvesSize; ++CurveIndex)
+		{
+			const uint32 PointCount = (*Counts)[CurveIndex];
+			HairStrands.StrandsCurves.CurvesCount[NumCurves+CurveIndex] = PointCount;
+			
+			for (uint32 PointIndex = 0; PointIndex < PointCount; ++PointIndex, ++GlobalIndex)
+			{
+				const float CoordU = PointIndex / static_cast<float>(PointCount - 1);
+				const float Radius = (Widths) ? (*Widths)[GlobalIndex] : FMath::Lerp(RootRadius, TipRadius, CoordU);
+				Alembic::Abc::P3fArraySample::value_type Position = (*Positions)[GlobalIndex];
+
+				HairStrands.StrandsPoints.PointsPosition[NumPoints + GlobalIndex] = ParentMatrix.TransformPosition( FVector(Position.x, Position.y, Position.z) );
+				HairStrands.StrandsPoints.PointsRadius[NumPoints + GlobalIndex] = Radius;
+			}
+		}
+	}
+	else if (Alembic::AbcGeom::IXform::matches(ObjectMetaData))
+	{
+		Alembic::AbcGeom::IXform Xform = Alembic::AbcGeom::IXform(InObject, Alembic::Abc::kWrapExisting);
+		Alembic::AbcGeom::XformSample MatrixSample; 
+		Xform.getSchema().get(MatrixSample);
+
+		LocalMatrix =  ParentMatrix * ConvertAlembicMatrix(MatrixSample.getMatrix());
+	}
+
+	if (NumChildren > 0)
+	{
+		for (uint32 ChildIndex = 0; ChildIndex < NumChildren; ++ChildIndex)
+		{
+			ParseObject(InObject.getChild(ChildIndex), HairStrands, LocalMatrix);
+		}
+	}
+}
+
+void FAbcFormat::ParseFile(const FString& FileName, FHairStrandsDatas& HairStrands)
+{
+	/** Factory used to generate objects*/
+	Alembic::AbcCoreFactory::IFactory Factory;
+	Alembic::AbcCoreFactory::IFactory::CoreType CompressionType = Alembic::AbcCoreFactory::IFactory::kUnknown;
+	/** Archive-typed ABC file */
+	Alembic::Abc::IArchive Archive;
+	/** Alembic typed root (top) object*/
+	Alembic::Abc::IObject TopObject;
+
+	Factory.setPolicy(Alembic::Abc::ErrorHandler::kThrowPolicy);
+	Factory.setOgawaNumStreams(12);
+
+	// Extract Archive and compression type from file
+	Archive = Factory.getArchive(TCHAR_TO_UTF8(*FileName), CompressionType);
+	if (!Archive.valid())
+	{
+		return;
+	}
+
+	// Get Top/root object
+	TopObject = Alembic::Abc::IObject(Archive, Alembic::Abc::kTop);
+	if (!TopObject.valid())
+	{
+		return;
+	}
+
+	FMatrix ParentMatrix = FMatrix::Identity;
+	ParseObject(TopObject, HairStrands, ParentMatrix);
+}
+
+template<typename FileFormat>
+void THairStrandsLoader<FileFormat>::LoadHairStrands(const FString& FileName, FHairStrandsDatas& HairStrands)
+{
+	HairStrands.StrandsCurves.Reset();
+	HairStrands.StrandsPoints.Reset();
+
+	FileFormat::ParseFile(FileName,HairStrands);
+
+	TArray<FVector>::TIterator PositionIterator = HairStrands.StrandsPoints.PointsPosition.CreateIterator();
+	TArray<float>::TIterator RadiusIterator = HairStrands.StrandsPoints.PointsRadius.CreateIterator();
+
+	for (uint32 PointIndex = 0; PointIndex < HairStrands.GetNumPoints(); ++PointIndex, ++PositionIterator, ++RadiusIterator)
+	{
+		*PositionIterator *= FileFormat::UNIT_TO_CM;
+		*RadiusIterator *= FileFormat::UNIT_TO_CM;
+	}
+
+	HairStrands.BuildInternalDatas();
+}
+
+template struct THairStrandsLoader<FHairFormat>;
+template struct THairStrandsLoader<FFbxFormat>;
+template struct THairStrandsLoader<FAbcFormat>;
diff --git a/Engine/Plugins/Experimental/HairStrands/Source/HairStrandsEditor/Public/HairStrandsActions.h b/Engine/Plugins/Experimental/HairStrands/Source/HairStrandsEditor/Public/HairStrandsActions.h
new file mode 100644
--- /dev/null
+++ b/Engine/Plugins/Experimental/HairStrands/Source/HairStrandsEditor/Public/HairStrandsActions.h
@@ -0,0 +1,36 @@
+// Copyright 1998-2019 Epic Games, Inc. All Rights Reserved.
+
+#pragma once
+
+#include "AssetTypeActions_Base.h"
+#include "Templates/SharedPointer.h"
+
+class ISlateStyle;
+
+/**
+ * Implements an action for UHairStrands assets.
+ */
+class FHairStrandsActions : public FAssetTypeActions_Base
+{
+public:
+
+	/**
+	 * Creates and initializes a new instance.
+	 *
+	 * @param InStyle The style set to use for asset editor toolkits.
+	 */
+	FHairStrandsActions();
+
+public:
+
+	//~ FAssetTypeActions_Base overrides
+
+	virtual bool CanFilter() override;
+	virtual void GetActions(const TArray<UObject*>& InObjects, FMenuBuilder& MenuBuilder) override;
+	virtual uint32 GetCategories() override;
+	virtual FText GetName() const override;
+	virtual UClass* GetSupportedClass() const override;
+	virtual FColor GetTypeColor() const override;
+	virtual bool HasActions(const TArray<UObject*>& InObjects) const override;
+	virtual void OpenAssetEditor(const TArray<UObject*>& InObjects, TSharedPtr<IToolkitHost> EditWithinLevelEditor = TSharedPtr<IToolkitHost>()) override;
+};
diff --git a/Engine/Plugins/Experimental/HairStrands/Source/HairStrandsEditor/Public/HairStrandsEditor.h b/Engine/Plugins/Experimental/HairStrands/Source/HairStrandsEditor/Public/HairStrandsEditor.h
new file mode 100644
--- /dev/null
+++ b/Engine/Plugins/Experimental/HairStrands/Source/HairStrandsEditor/Public/HairStrandsEditor.h
@@ -0,0 +1,29 @@
+// Copyright 1998-2019 Epic Games, Inc. All Rights Reserved.
+
+#pragma once
+
+#include "Modules/ModuleInterface.h"
+#include "Modules/ModuleManager.h"
+
+/** Implements the HairStrands module  */
+class FHairStrandsEditor : public IModuleInterface
+{
+public:
+
+	//~ IModuleInterface interface
+
+	virtual void StartupModule() override;
+	virtual void ShutdownModule() override;
+
+	virtual bool SupportsDynamicReloading() override
+	{
+		return true;
+	}
+
+private:
+
+	/** The collection of registered asset type actions. */
+	TArray<TSharedRef<class IAssetTypeActions>> RegisteredAssetTypeActions;
+};
+
+IMPLEMENT_MODULE(FHairStrandsEditor, HairStrandsEditor);
diff --git a/Engine/Plugins/Experimental/HairStrands/Source/HairStrandsEditor/Public/HairStrandsFactory.h b/Engine/Plugins/Experimental/HairStrands/Source/HairStrandsEditor/Public/HairStrandsFactory.h
new file mode 100644
--- /dev/null
+++ b/Engine/Plugins/Experimental/HairStrands/Source/HairStrandsEditor/Public/HairStrandsFactory.h
@@ -0,0 +1,34 @@
+// Copyright 1998-2019 Epic Games, Inc. All Rights Reserved.
+
+#pragma once
+
+#include "Factories/Factory.h"
+#include "UObject/ObjectMacros.h"
+#include <UnrealEd/Public/EditorReimportHandler.h>
+
+#include "HairStrandsFactory.generated.h"
+
+
+/**
+ * Implements a factory for UHairStrands objects.
+ */
+UCLASS(hidecategories = Object)
+class UHairStrandsFactory : public UFactory, public FReimportHandler
+{
+	GENERATED_UCLASS_BODY()
+
+public:
+
+	//~ UFactory Interface
+	virtual UObject* FactoryCreateNew(UClass* InClass, UObject* InParent, FName InName, EObjectFlags Flags, UObject* Context, FFeedbackContext* Warn) override;
+	virtual UObject* FactoryCreateFile(UClass * InClass, UObject * InParent, FName InName, EObjectFlags Flags,
+		const FString & Filename, const TCHAR* Parms, FFeedbackContext * Warn, bool& bOutOperationCanceled) override;
+
+	//~ FReimportHandler interface
+	virtual bool FactoryCanImport(const FString& Filename) override;
+	virtual bool CanReimport(UObject* Obj, TArray<FString>& OutFilenames) override;
+	virtual void SetReimportPaths(UObject* Obj, const TArray<FString>& NewReimportPaths) override;
+	virtual EReimportResult::Type Reimport(UObject* Obj) override;
+	virtual int32 GetPriority() const override;
+};
+
diff --git a/Engine/Plugins/Experimental/HairStrands/Source/HairStrandsEditor/Public/HairStrandsLoader.h b/Engine/Plugins/Experimental/HairStrands/Source/HairStrandsEditor/Public/HairStrandsLoader.h
new file mode 100644
--- /dev/null
+++ b/Engine/Plugins/Experimental/HairStrands/Source/HairStrandsEditor/Public/HairStrandsLoader.h
@@ -0,0 +1,92 @@
+// Copyright 1998-2019 Epic Games, Inc. All Rights Reserved.
+
+#pragma once
+
+#include "CoreMinimal.h"
+#include "Engine/EngineTypes.h"
+#include "Containers/UnrealString.h"
+
+/** .hair file format */
+struct FHairFormat
+{
+  /** File unit to cm ratio */
+  static constexpr const float UNIT_TO_CM = 10;
+
+  /** Parse the .hair file */
+  static void ParseFile(const FString& FileName, struct FHairStrandsDatas& HairStrands);
+
+private :
+
+	///* Hair file header */
+	//struct HairFileHeader
+	//{
+	//	char	FileSignature[4];	/* This should be "HAIR" */
+	//	uint32	NumStrands = 0;		/* number of hair strands */
+	//	uint32	NumVertices = 0;	/* total number of points of all strands */
+	//	uint32	BitArrays = 0;			/* bit array of data in the file */
+
+	//	uint32	StrandSize = 0;		/* default number of segments of each strand */
+	//	float	StrandThickness = 1.f;	/* default thickness of hair strands */
+	//	float	StrandTransparency = 0.f;	/* default transparency of hair strands */
+	//	float	StrandColor[3] = { 0.f,0.f,0.f };		/* default color of hair strands */
+	//	char	FileInfo[HAIR_FILE_INFO_SIZE];	/* information about the file */
+	//};
+
+	/* Hair file header */
+	struct FFileHeader
+	{
+		/* This should be "HAIR" */
+		char	FileSignature[4];	
+
+		/* Number of hair strands */
+		uint32	NumStrands = 0;	
+
+		/* Total number of points of all strands */
+		uint32	NumPoints = 0;
+
+		/* Bit array of data in the file */
+		uint32	BitArrays = 0;
+
+		/* Default number of segments of each strand */
+		uint32	StrandCount = 0;
+
+		/* Default thickness of hair strands */
+		float	StrandThickness = 1.f;
+
+		/* Default transparency of hair strands */
+		float	StrandTransparency = 0.f;	
+
+		/* Default color of hair strands */
+		float	StrandColor[3] = { 0.f,0.f,0.f };
+
+		/* Information about the file */
+		char	FileInfo[88];	
+	};
+};
+
+/** .fbx file format */
+struct FFbxFormat
+{
+  /** File unit to cm ratio */
+  static constexpr const float UNIT_TO_CM = 100;
+
+  /** Parse the .fbx file */
+  static void ParseFile(const FString& FileName, struct FHairStrandsDatas& HairStrands);
+};
+
+/** .abc file format */
+struct FAbcFormat
+{
+   /** File unit to cm ratio */
+  static constexpr const float UNIT_TO_CM = 1;
+
+  /** Parse the .abc file */
+  static void ParseFile(const FString& FileName, struct FHairStrandsDatas& HairStrands);
+};
+
+/** Hair strands loader */
+template<typename FileFormat>
+struct THairStrandsLoader
+{
+    static void LoadHairStrands(const FString& FileName, struct FHairStrandsDatas& HairStrands); 
+};
diff --git a/Engine/Plugins/Experimental/HairStrands/Source/HairStrandsNiagara/HairStrandsNiagara.Build.cs b/Engine/Plugins/Experimental/HairStrands/Source/HairStrandsNiagara/HairStrandsNiagara.Build.cs
new file mode 100644
--- /dev/null
+++ b/Engine/Plugins/Experimental/HairStrands/Source/HairStrandsNiagara/HairStrandsNiagara.Build.cs
@@ -0,0 +1,32 @@
+// Copyright 1998-2019 Epic Games, Inc. All Rights Reserved.
+
+namespace UnrealBuildTool.Rules
+{
+	public class HairStrandsNiagara : ModuleRules
+	{
+        public HairStrandsNiagara(ReadOnlyTargetRules Target) : base(Target)
+		{
+			PrivateIncludePaths.Add(ModuleDirectory + "/Private");
+            PublicIncludePaths.Add(ModuleDirectory + "/Public");
+
+            PrivateDependencyModuleNames.AddRange(
+                new string[]
+                {
+                    "Core",
+                    "CoreUObject",
+					"Engine",
+					"Projects",
+                    "NiagaraCore",
+					"Niagara",
+					"NiagaraShader",
+                    "RenderCore",
+                    "Renderer",
+                    "UtilityShaders",
+                    "VectorVM",
+					"RHI",
+					"HairStrandsCore"
+                }
+			);
+        }
+	}
+}
diff --git a/Engine/Plugins/Experimental/HairStrands/Source/HairStrandsNiagara/Private/HairStrandsNiagara.cpp b/Engine/Plugins/Experimental/HairStrands/Source/HairStrandsNiagara/Private/HairStrandsNiagara.cpp
new file mode 100644
--- /dev/null
+++ b/Engine/Plugins/Experimental/HairStrands/Source/HairStrandsNiagara/Private/HairStrandsNiagara.cpp
@@ -0,0 +1,16 @@
+// Copyright 1998-2019 Epic Games, Inc. All Rights Reserved.
+
+#include "HairStrandsNiagara.h"
+#include "Interfaces/IPluginManager.h"
+#include "ShaderCore.h"
+
+void FHairStrandsNiagara::StartupModule()
+{
+	// Maps virtual shader source directory /Plugin/FX/Niagara to the plugin's actual Shaders directory.
+	FString PluginShaderDir = FPaths::Combine(IPluginManager::Get().FindPlugin(TEXT("HairStrands"))->GetBaseDir(), TEXT("Shaders"));
+	AddShaderSourceDirectoryMapping(TEXT("/Plugin/Experimental/HairStrands"), PluginShaderDir);
+}
+
+void FHairStrandsNiagara::ShutdownModule()
+{}
+
