Copying //UE4/Dev-Physics to Dev-Main (//UE4/Dev-Main) [at] 7703071

diff --git a/Engine/Plugins/Experimental/HairStrands/Source/HairStrandsNiagara/Private/NiagaraDataInterfaceHairStrands.cpp b/Engine/Plugins/Experimental/HairStrands/Source/HairStrandsNiagara/Private/NiagaraDataInterfaceHairStrands.cpp
new file mode 100644
--- /dev/null
+++ b/Engine/Plugins/Experimental/HairStrands/Source/HairStrandsNiagara/Private/NiagaraDataInterfaceHairStrands.cpp
@@ -0,0 +1,905 @@
+// Copyright 1998-2019 Epic Games, Inc. All Rights Reserved.
+
+#include "NiagaraDataInterfaceHairStrands.h"
+#include "NiagaraShader.h"
+#include "NiagaraComponent.h"
+#include "NiagaraRenderer.h"
+#include "NiagaraSystemInstance.h"
+#include "NiagaraEmitterInstanceBatcher.h"
+#include "ShaderParameterUtils.h"
+#include "HairStrandsAsset.h"  
+
+#define LOCTEXT_NAMESPACE "NiagaraDataInterfaceHairStrands"
+
+//------------------------------------------------------------------------------------------------------------
+
+static const FName GetStrandDensityName(TEXT("GetStrandDensity"));
+static const FName GetStrandSizeName(TEXT("GetStrandSize"));
+static const FName GetNumStrandsName(TEXT("GetNumStrands"));
+static const FName GetRootThicknessName(TEXT("GetRootThickness"));
+static const FName GetTipThicknessName(TEXT("GetTipThickness"));
+static const FName GetWorldTransformName(TEXT("GetWorldTransform"));
+static const FName GetVertexPositionName(TEXT("GetStrandPosition"));
+
+//------------------------------------------------------------------------------------------------------------
+
+static const FName ComputeNodePositionName(TEXT("ComputeNodePosition"));
+
+//------------------------------------------------------------------------------------------------------------
+
+const FString UNiagaraDataInterfaceHairStrands::NumStrandsName(TEXT("NumStrands_"));
+const FString UNiagaraDataInterfaceHairStrands::StrandSizeName(TEXT("StrandSize_"));
+const FString UNiagaraDataInterfaceHairStrands::StrandDensityName(TEXT("StrandDensity_"));
+const FString UNiagaraDataInterfaceHairStrands::RootThicknessName(TEXT("RootThickness_"));
+const FString UNiagaraDataInterfaceHairStrands::TipThicknessName(TEXT("TipThickness_"));
+const FString UNiagaraDataInterfaceHairStrands::WorldTransformName(TEXT("WorldTransform_"));
+const FString UNiagaraDataInterfaceHairStrands::PointsPositionsBufferName(TEXT("PointsPositionsBuffer_"));
+const FString UNiagaraDataInterfaceHairStrands::CurvesOffsetsBufferName(TEXT("CurvesOffsetsBuffer_"));
+const FString UNiagaraDataInterfaceHairStrands::NodesPositionsBufferName(TEXT("NodesPositionsBuffer_"));
+const FString UNiagaraDataInterfaceHairStrands::PointsNodesBufferName(TEXT("PointsNodesBuffer_"));
+const FString UNiagaraDataInterfaceHairStrands::PointsWeightsBufferName(TEXT("PointsWeightsBuffer_"));
+
+//------------------------------------------------------------------------------------------------------------
+
+struct FNDIHairStrandsParametersName
+{
+	FNDIHairStrandsParametersName(const FString& Suffix)
+	{
+		NumStrandsName = UNiagaraDataInterfaceHairStrands::NumStrandsName + Suffix;
+		StrandSizeName = UNiagaraDataInterfaceHairStrands::StrandSizeName + Suffix;
+		StrandDensityName = UNiagaraDataInterfaceHairStrands::StrandDensityName + Suffix;
+		RootThicknessName = UNiagaraDataInterfaceHairStrands::RootThicknessName + Suffix;
+		TipThicknessName = UNiagaraDataInterfaceHairStrands::TipThicknessName + Suffix;
+		WorldTransformName = UNiagaraDataInterfaceHairStrands::WorldTransformName + Suffix;
+		PointsPositionsBufferName = UNiagaraDataInterfaceHairStrands::PointsPositionsBufferName + Suffix;
+		CurvesOffsetsBufferName = UNiagaraDataInterfaceHairStrands::CurvesOffsetsBufferName + Suffix;
+		NodesPositionsBufferName = UNiagaraDataInterfaceHairStrands::NodesPositionsBufferName + Suffix;
+		PointsNodesBufferName = UNiagaraDataInterfaceHairStrands::PointsNodesBufferName + Suffix;
+		PointsWeightsBufferName = UNiagaraDataInterfaceHairStrands::PointsWeightsBufferName + Suffix;
+	}
+
+	FString NumStrandsName;
+	FString StrandSizeName;
+	FString StrandDensityName;
+	FString RootThicknessName;
+	FString TipThicknessName;
+	FString WorldTransformName;
+	FString PointsPositionsBufferName;
+	FString CurvesOffsetsBufferName;
+	FString NodesPositionsBufferName;
+	FString PointsNodesBufferName;
+	FString PointsWeightsBufferName;
+};
+
+//------------------------------------------------------------------------------------------------------------
+
+void FNDIHairStrandsBuffer::SetAsset(const UHairStrandsAsset* HairStrandsAsset)
+{
+	SourceAsset = HairStrandsAsset;
+}
+
+void FNDIHairStrandsBuffer::SetStrandSize(const uint32 HairStrandSize)
+{
+	StrandSize = HairStrandSize;
+}
+
+void FNDIHairStrandsBuffer::InitRHI()
+{
+	if (SourceAsset != nullptr)
+	{
+		if (SourceAsset->HairStrandsResource)
+		{
+			PointsPositionsBufferUav = SourceAsset->HairStrandsResource->PositionBuffer.UAV;
+		}
+		TArray<FHairStrandsPositionFormat::Type> NodesPositions;
+		TArray<FHairStrandsIndexFormat::Type> PointsNodes;
+		TArray<FHairStrandsWeightFormat::Type> PointsWeights;
+		SourceAsset->StrandsDatas.BuildSimulationDatas(StrandSize, NodesPositions, PointsWeights, PointsNodes);
+		{
+			const uint32 PositionCount = NodesPositions.Num();
+			const uint32 PositionBytes = FHairStrandsPositionFormat::SizeInByte*PositionCount;
+
+			NodesPositionsBuffer.Initialize(FHairStrandsPositionFormat::SizeInByte, PositionCount, FHairStrandsPositionFormat::Format, BUF_Static);
+			void* PositionBufferData = RHILockVertexBuffer(NodesPositionsBuffer.Buffer, 0, PositionBytes, RLM_WriteOnly);
+
+			FMemory::Memcpy(PositionBufferData, NodesPositions.GetData(), PositionBytes);
+			RHIUnlockVertexBuffer(NodesPositionsBuffer.Buffer);
+		}
+		{
+			const uint32 NodesCount = PointsNodes.Num();
+			const uint32 NodesBytes = FHairStrandsIndexFormat::SizeInByte*NodesCount;
+
+			PointsNodesBuffer.Initialize(FHairStrandsIndexFormat::SizeInByte, NodesCount, FHairStrandsIndexFormat::Format, BUF_Static);
+			void* NodesBufferData = RHILockVertexBuffer(PointsNodesBuffer.Buffer, 0, NodesBytes, RLM_WriteOnly);
+
+			FMemory::Memcpy(NodesBufferData, PointsNodes.GetData(), NodesBytes);
+			RHIUnlockVertexBuffer(PointsNodesBuffer.Buffer);
+		}
+		{
+			const uint32 WeightsCount = PointsWeights.Num();
+			const uint32 WeightsBytes = FHairStrandsWeightFormat::SizeInByte*WeightsCount;
+
+			PointsWeightsBuffer.Initialize(FHairStrandsWeightFormat::SizeInByte, WeightsCount, FHairStrandsWeightFormat::Format, BUF_Static);
+			void* WeightsBufferData = RHILockVertexBuffer(PointsWeightsBuffer.Buffer, 0, WeightsBytes, RLM_WriteOnly);
+
+			FMemory::Memcpy(WeightsBufferData, PointsWeights.GetData(), WeightsBytes);
+			RHIUnlockVertexBuffer(PointsWeightsBuffer.Buffer);
+		}
+		{
+			const uint32 OffsetCount = SourceAsset->StrandsDatas.GetNumCurves() + 1;
+			const uint32 OffsetBytes = sizeof(uint32)*OffsetCount;
+
+			CurvesOffsetsBuffer.Initialize(sizeof(uint32), OffsetCount, EPixelFormat::PF_R32_UINT, BUF_Static);
+			void* OffsetBufferData = RHILockVertexBuffer(CurvesOffsetsBuffer.Buffer, 0, OffsetBytes, RLM_WriteOnly);
+
+			FMemory::Memcpy(OffsetBufferData, SourceAsset->StrandsDatas.StrandsCurves.CurvesOffset.GetData(), OffsetBytes);
+			RHIUnlockVertexBuffer(CurvesOffsetsBuffer.Buffer);
+		}
+	}
+}
+
+void FNDIHairStrandsBuffer::ReleaseRHI()
+{
+	CurvesOffsetsBuffer.Release();
+	NodesPositionsBuffer.Release();
+	PointsNodesBuffer.Release();
+	PointsWeightsBuffer.Release();
+}
+
+//------------------------------------------------------------------------------------------------------------
+
+struct FNDIHairStrandsParametersCS : public FNiagaraDataInterfaceParametersCS
+{
+	virtual void Bind(const FNiagaraDataInterfaceParamRef& ParamRef, const class FShaderParameterMap& ParameterMap) override
+	{
+		FNDIHairStrandsParametersName ParamNames(ParamRef.ParameterInfo.DataInterfaceHLSLSymbol);
+
+		WorldTransform.Bind(ParameterMap, *ParamNames.WorldTransformName);
+		NumStrands.Bind(ParameterMap, *ParamNames.NumStrandsName);
+		StrandSize.Bind(ParameterMap, *ParamNames.StrandSizeName);
+		StrandDensity.Bind(ParameterMap, *ParamNames.StrandDensityName);
+		RootThickness.Bind(ParameterMap, *ParamNames.RootThicknessName);
+		TipThickness.Bind(ParameterMap, *ParamNames.TipThicknessName);
+		PointsPositionsBuffer.Bind(ParameterMap, *ParamNames.PointsPositionsBufferName);
+		CurvesOffsetsBuffer.Bind(ParameterMap, *ParamNames.CurvesOffsetsBufferName);
+		NodesPositionsBuffer.Bind(ParameterMap, *ParamNames.NodesPositionsBufferName);
+		PointsNodesBuffer.Bind(ParameterMap, *ParamNames.PointsNodesBufferName);
+		PointsWeightsBuffer.Bind(ParameterMap, *ParamNames.PointsWeightsBufferName);
+
+		if (!PointsPositionsBuffer.IsBound())
+		{
+			UE_LOG(LogHairStrands, Warning, TEXT("Binding failed for FNDIHairStrandsParametersCS %s. Was it optimized out?"), *ParamNames.PointsPositionsBufferName)
+		}
+		if (!CurvesOffsetsBuffer.IsBound())
+		{
+			UE_LOG(LogHairStrands, Warning, TEXT("Binding failed for FNDIHairStrandsParametersCS %s. Was it optimized out?"), *ParamNames.CurvesOffsetsBufferName)
+		}
+		if (!NodesPositionsBuffer.IsBound())
+		{
+			UE_LOG(LogHairStrands, Warning, TEXT("Binding failed for FNDIHairStrandsParametersCS %s. Was it optimized out?"), *ParamNames.NodesPositionsBufferName)
+		}
+		if (!PointsNodesBuffer.IsBound())
+		{
+			UE_LOG(LogHairStrands, Warning, TEXT("Binding failed for FNDIHairStrandsParametersCS %s. Was it optimized out?"), *ParamNames.PointsNodesBufferName)
+		}
+		if (!PointsWeightsBuffer.IsBound())
+		{
+			UE_LOG(LogHairStrands, Warning, TEXT("Binding failed for FNDIHairStrandsParametersCS %s. Was it optimized out?"), *ParamNames.PointsWeightsBufferName)
+		}
+	}
+
+	virtual void Serialize(FArchive& Ar) override
+	{
+		Ar << WorldTransform;
+		Ar << NumStrands;
+		Ar << StrandSize;
+		Ar << StrandDensity;
+		Ar << RootThickness;
+		Ar << TipThickness;
+		Ar << PointsPositionsBuffer;
+		Ar << CurvesOffsetsBuffer;
+		Ar << NodesPositionsBuffer;
+		Ar << PointsNodesBuffer;
+		Ar << PointsWeightsBuffer;
+	}
+
+	virtual void Set(FRHICommandList& RHICmdList, const FNiagaraDataInterfaceSetArgs& Context) const override
+	{
+		check(IsInRenderingThread());
+
+		FRHIComputeShader* ComputeShaderRHI = Context.Shader->GetComputeShader();
+
+		FNDIHairStrandsProxy* InterfaceProxy =
+			static_cast<FNDIHairStrandsProxy*>(Context.DataInterface);
+		FNDIHairStrandsData* ProxyData =
+			InterfaceProxy->SystemInstancesToProxyData.Find(Context.SystemInstance);
+		ensure(ProxyData);
+
+		if (ProxyData != nullptr)
+		{
+			FNDIHairStrandsBuffer* StrandBuffer = ProxyData->StrandBuffer;
+			SetUAVParameter(RHICmdList, ComputeShaderRHI, PointsPositionsBuffer, StrandBuffer->PointsPositionsBufferUav);
+			SetSRVParameter(RHICmdList, ComputeShaderRHI, CurvesOffsetsBuffer, StrandBuffer->CurvesOffsetsBuffer.SRV);
+			SetSRVParameter(RHICmdList, ComputeShaderRHI, NodesPositionsBuffer, StrandBuffer->NodesPositionsBuffer.SRV);
+			SetSRVParameter(RHICmdList, ComputeShaderRHI, PointsNodesBuffer, StrandBuffer->PointsNodesBuffer.SRV);
+			SetSRVParameter(RHICmdList, ComputeShaderRHI, PointsWeightsBuffer, StrandBuffer->PointsWeightsBuffer.SRV);
+
+			SetShaderValue(RHICmdList, ComputeShaderRHI, WorldTransform, ProxyData->WorldTransform);
+			SetShaderValue(RHICmdList, ComputeShaderRHI, NumStrands, ProxyData->NumStrands);
+			SetShaderValue(RHICmdList, ComputeShaderRHI, StrandSize, ProxyData->StrandSize);
+			SetShaderValue(RHICmdList, ComputeShaderRHI, StrandDensity, ProxyData->StrandDensity);
+			SetShaderValue(RHICmdList, ComputeShaderRHI, RootThickness, ProxyData->RootThickness);
+			SetShaderValue(RHICmdList, ComputeShaderRHI, TipThickness, ProxyData->TipThickness);
+		}
+		else
+		{
+			SetUAVParameter(RHICmdList, ComputeShaderRHI, PointsPositionsBuffer, FNiagaraRenderer::GetDummyFloatBuffer().UAV);
+			SetSRVParameter(RHICmdList, ComputeShaderRHI, CurvesOffsetsBuffer, FNiagaraRenderer::GetDummyUIntBuffer().SRV);
+			SetSRVParameter(RHICmdList, ComputeShaderRHI, NodesPositionsBuffer, FNiagaraRenderer::GetDummyFloatBuffer().SRV);
+			SetSRVParameter(RHICmdList, ComputeShaderRHI, PointsNodesBuffer, FNiagaraRenderer::GetDummyUIntBuffer().SRV);
+			SetSRVParameter(RHICmdList, ComputeShaderRHI, PointsWeightsBuffer, FNiagaraRenderer::GetDummyFloatBuffer().SRV);
+
+			SetShaderValue(RHICmdList, ComputeShaderRHI, WorldTransform, FMatrix::Identity);
+			SetShaderValue(RHICmdList, ComputeShaderRHI, NumStrands, 1);
+			SetShaderValue(RHICmdList, ComputeShaderRHI, StrandSize, 1);
+			SetShaderValue(RHICmdList, ComputeShaderRHI, StrandDensity, 1);
+			SetShaderValue(RHICmdList, ComputeShaderRHI, RootThickness, 0.1);
+			SetShaderValue(RHICmdList, ComputeShaderRHI, TipThickness, 0.1);
+		}
+	}
+
+private:
+
+	FShaderParameter WorldTransform;
+	FShaderParameter NumStrands;
+	FShaderParameter StrandSize;
+	FShaderParameter StrandDensity;
+	FShaderParameter RootThickness;
+	FShaderParameter TipThickness;
+
+	FShaderResourceParameter PointsPositionsBuffer;
+	FShaderResourceParameter CurvesOffsetsBuffer;
+	FShaderResourceParameter NodesPositionsBuffer;
+	FShaderResourceParameter PointsNodesBuffer;
+	FShaderResourceParameter PointsWeightsBuffer;
+};
+
+//------------------------------------------------------------------------------------------------------------
+
+void FNDIHairStrandsProxy::DeferredDestroy()
+{
+	for (const FGuid& Sys : DeferredDestroyList)
+	{
+		SystemInstancesToProxyData.Remove(Sys);
+	}
+
+	DeferredDestroyList.Empty();
+}
+
+void FNDIHairStrandsProxy::ConsumePerInstanceDataFromGameThread(void* PerInstanceData, const FGuid& Instance)
+{
+	FNDIHairStrandsData* SourceData = static_cast<FNDIHairStrandsData*>(PerInstanceData);
+	FNDIHairStrandsData* TargetData = SystemInstancesToProxyData.Find(Instance);
+
+	ensure(TargetData);
+	if (TargetData)
+	{
+		TargetData->WorldTransform = SourceData->WorldTransform;
+	}
+	else
+	{
+		UE_LOG(LogHairStrands, Log, TEXT("ConsumePerInstanceDataFromGameThread() ... could not find %s"), *Instance.ToString());
+	}
+}
+
+void FNDIHairStrandsProxy::InitializePerInstanceData(const FGuid& SystemInstance, FNDIHairStrandsBuffer* StrandBuffer, uint32 NumStrands, const uint8 StrandSize,
+	const float StrandDensity, const float RootThickness, const float TipThickness)
+{
+	check(IsInRenderingThread());
+
+	FNDIHairStrandsData* TargetData = SystemInstancesToProxyData.Find(SystemInstance);
+	if (TargetData != nullptr)
+	{
+		DeferredDestroyList.Remove(SystemInstance);
+	}
+	else
+	{
+		TargetData = &SystemInstancesToProxyData.Add(SystemInstance);
+	}
+	TargetData->StrandBuffer = StrandBuffer;
+	TargetData->NumStrands = NumStrands;
+	TargetData->StrandSize = StrandSize;
+	TargetData->StrandDensity = StrandDensity;
+	TargetData->RootThickness = RootThickness;
+	TargetData->TipThickness = TipThickness;
+}
+
+void FNDIHairStrandsProxy::DestroyPerInstanceData(NiagaraEmitterInstanceBatcher* Batcher, const FGuid& SystemInstance)
+{
+	check(IsInRenderingThread());
+
+	DeferredDestroyList.Add(SystemInstance);
+	Batcher->EnqueueDeferredDeletesForDI_RenderThread(this->AsShared());
+}
+
+//------------------------------------------------------------------------------------------------------------
+
+UNiagaraDataInterfaceHairStrands::UNiagaraDataInterfaceHairStrands(FObjectInitializer const& ObjectInitializer)
+	: Super(ObjectInitializer)
+	, StrandSize(EHairStrandsSize::Size4)
+	, StrandDensity(1.0)
+	, RootThickness(0.1)
+	, TipThickness(0.1)
+	, SourceAsset(nullptr)
+	, SourceTransform(FMatrix::Identity)
+	, StaticMesh(nullptr)
+{
+	Proxy = MakeShared<FNDIHairStrandsProxy, ESPMode::ThreadSafe>();
+}
+
+bool UNiagaraDataInterfaceHairStrands::InitPerInstanceData(void* PerInstanceData, FNiagaraSystemInstance* SystemInstance)
+{
+	FNDIHairStrandsData* InstanceData = new (PerInstanceData) FNDIHairStrandsData();
+
+	if (SourceAsset == nullptr)
+	{
+		UE_LOG(LogHairStrands, Log, TEXT("Hair Strands data interface has no valid asset. Failed InitPerInstanceData - %s"), *GetFullName());
+		return false;
+	}
+	else
+	{
+		FNDIHairStrandsBuffer* StrandBuffer = new FNDIHairStrandsBuffer;
+		StrandBuffer->SetAsset(SourceAsset);
+
+		// Push instance data to RT
+		{
+			const uint32 LocalNumStrands = SourceAsset->StrandsDatas.GetNumCurves();
+			const uint8 LocalStrandSize = static_cast<uint8>(StrandSize);
+			const FBox LocalStrandsBox;// = SourceAsset->ComputeStrandsBounds();
+			StrandBuffer->SetStrandSize(LocalStrandSize);
+
+			
+			/*SourceAsset->StrandsDatas.ResampleInternalDatas(static_cast<uint32>(StrandSize));
+			if (StaticMesh)
+			{
+				SourceAsset->StrandsDatas.AttachStrandsRoots(StaticMesh,SourceTransform);
+			}*/
+
+			UE_LOG(LogHairStrands, Log, TEXT("Num Strands = %d | Strand Size = %d | Num Vertices = %d | Min = %f %f %f | Max = %f %f %f"), LocalNumStrands, LocalStrandSize,
+				SourceAsset->StrandsDatas.GetNumPoints(), LocalStrandsBox.Min[0], LocalStrandsBox.Min[1], LocalStrandsBox.Min[2], LocalStrandsBox.Max[0], LocalStrandsBox.Max[1], LocalStrandsBox.Max[2] );
+
+			const float LocalStrandDensity = StrandDensity;
+			const float LocalRootThickness = RootThickness;
+			const float LocalTipThickness = TipThickness;
+
+			InstanceData->WorldTransform = SourceTransform;// SourceAsset->GetComponentToWorld().ToMatrixWithScale();
+			InstanceData->StrandSize = LocalStrandSize;
+			InstanceData->StrandDensity = LocalStrandDensity;
+			InstanceData->RootThickness = LocalRootThickness;
+			InstanceData->TipThickness = LocalTipThickness;
+			InstanceData->NumStrands = LocalNumStrands;
+			InstanceData->StrandBuffer = StrandBuffer;
+
+			FNDIHairStrandsProxy* ThisProxy = GetProxyAs<FNDIHairStrandsProxy>();
+			ENQUEUE_RENDER_COMMAND(FNiagaraDIPushInitialInstanceDataToRT) (
+				[ThisProxy, InstanceID = SystemInstance->GetId(), StrandBuffer, LocalNumStrands, LocalStrandSize, LocalStrandDensity, LocalRootThickness, LocalTipThickness](FRHICommandListImmediate& CmdList)
+			{
+				StrandBuffer->InitResource();
+				ThisProxy->InitializePerInstanceData(InstanceID, StrandBuffer, LocalNumStrands, LocalStrandSize, LocalStrandDensity, LocalRootThickness, LocalTipThickness);
+			}
+			);
+		}
+	}
+	return true;
+}
+
+void UNiagaraDataInterfaceHairStrands::DestroyPerInstanceData(void* PerInstanceData, FNiagaraSystemInstance* SystemInstance)
+{
+	FNDIHairStrandsData* InstanceData = static_cast<FNDIHairStrandsData*>(PerInstanceData);
+	InstanceData->~FNDIHairStrandsData();
+
+	if(InstanceData->StrandBuffer)
+	{
+		FNDIHairStrandsProxy* ThisProxy = GetProxyAs<FNDIHairStrandsProxy>();
+		FNDIHairStrandsBuffer* InBuffer = InstanceData->StrandBuffer;
+		ENQUEUE_RENDER_COMMAND(FNiagaraDIDestroyInstanceData) (
+			[ThisProxy, InBuffer, InstanceID = SystemInstance->GetId(), Batcher = SystemInstance->GetBatcher()](FRHICommandListImmediate& CmdList)
+		{
+			InBuffer->ReleaseResource();
+			ThisProxy->DestroyPerInstanceData(Batcher, InstanceID);
+			delete InBuffer;
+		}
+		);
+		InstanceData->StrandBuffer = nullptr;
+	}
+}
+
+bool UNiagaraDataInterfaceHairStrands::PerInstanceTick(void* PerInstanceData, FNiagaraSystemInstance* SystemInstance, float InDeltaSeconds)
+{
+	FNDIHairStrandsData* InstanceData = static_cast<FNDIHairStrandsData*>(PerInstanceData);
+
+	bool RequireReset = false;
+	if (SourceAsset != nullptr)
+	{
+		InstanceData->WorldTransform = SourceTransform; //SourceAsset->GetComponentToWorld().ToMatrixWithScale();
+	}
+	else
+	{
+		InstanceData->WorldTransform = SourceTransform;
+		RequireReset = true;
+	}
+	return RequireReset;
+}
+
+bool UNiagaraDataInterfaceHairStrands::CopyToInternal(UNiagaraDataInterface* Destination) const
+{
+	if (!Super::CopyToInternal(Destination))
+	{
+		return false;
+	}
+
+	UNiagaraDataInterfaceHairStrands* OtherTyped = CastChecked<UNiagaraDataInterfaceHairStrands>(Destination);
+	OtherTyped->StrandSize = StrandSize;
+	OtherTyped->StrandDensity = StrandDensity;
+	OtherTyped->RootThickness = RootThickness;
+	OtherTyped->TipThickness = TipThickness;
+	OtherTyped->SourceAsset = SourceAsset;
+	OtherTyped->SourceTransform = SourceTransform;
+
+	return true;
+}
+
+bool UNiagaraDataInterfaceHairStrands::Equals(const UNiagaraDataInterface* Other) const
+{
+	if (!Super::Equals(Other))
+	{
+		return false;
+	}
+	const UNiagaraDataInterfaceHairStrands* OtherTyped = CastChecked<const UNiagaraDataInterfaceHairStrands>(Other);
+
+	return  (OtherTyped->StrandSize == StrandSize)
+		&& (OtherTyped->StrandDensity == StrandDensity) && (OtherTyped->RootThickness == RootThickness)
+		&& (OtherTyped->TipThickness == TipThickness) && (OtherTyped->SourceAsset == SourceAsset) && (OtherTyped->SourceTransform == SourceTransform);
+}
+
+void UNiagaraDataInterfaceHairStrands::PostInitProperties()
+{
+	Super::PostInitProperties();
+
+	if (HasAnyFlags(RF_ClassDefaultObject))
+	{
+		FNiagaraTypeRegistry::Register(FNiagaraTypeDefinition(GetClass()), true, false, false);
+	}
+}
+
+void UNiagaraDataInterfaceHairStrands::GetFunctions(TArray<FNiagaraFunctionSignature>& OutFunctions)
+{
+	{
+		FNiagaraFunctionSignature Sig;
+		Sig.Name = GetNumStrandsName;
+		Sig.bMemberFunction = true;
+		Sig.bRequiresContext = false;
+		Sig.Inputs.Add(FNiagaraVariable(FNiagaraTypeDefinition(GetClass()), TEXT("Strand Hair")));
+		Sig.Outputs.Add(FNiagaraVariable(FNiagaraTypeDefinition::GetIntDef(), TEXT("Num Strands")));
+
+		OutFunctions.Add(Sig);
+	}
+	{
+		FNiagaraFunctionSignature Sig;
+		Sig.Name = GetStrandDensityName;
+		Sig.bMemberFunction = true;
+		Sig.bRequiresContext = false;
+		Sig.Inputs.Add(FNiagaraVariable(FNiagaraTypeDefinition(GetClass()), TEXT("Strand Hair")));
+		Sig.Outputs.Add(FNiagaraVariable(FNiagaraTypeDefinition::GetFloatDef(), TEXT("Strand Density")));
+
+		OutFunctions.Add(Sig);
+	}
+	{
+		FNiagaraFunctionSignature Sig;
+		Sig.Name = GetStrandSizeName;
+		Sig.bMemberFunction = true;
+		Sig.bRequiresContext = false;
+		Sig.Inputs.Add(FNiagaraVariable(FNiagaraTypeDefinition(GetClass()), TEXT("Strand Hair")));
+		Sig.Outputs.Add(FNiagaraVariable(FNiagaraTypeDefinition::GetIntDef(), TEXT("Strand Size")));
+
+		OutFunctions.Add(Sig);
+	}
+	{
+		FNiagaraFunctionSignature Sig;
+		Sig.Name = GetRootThicknessName;
+		Sig.bMemberFunction = true;
+		Sig.bRequiresContext = false;
+		Sig.Inputs.Add(FNiagaraVariable(FNiagaraTypeDefinition(GetClass()), TEXT("Strand Hair")));
+		Sig.Outputs.Add(FNiagaraVariable(FNiagaraTypeDefinition::GetFloatDef(), TEXT("Root Thickness")));
+
+		OutFunctions.Add(Sig);
+	}
+	{
+		FNiagaraFunctionSignature Sig;
+		Sig.Name = GetTipThicknessName;
+		Sig.bMemberFunction = true;
+		Sig.bRequiresContext = false;
+		Sig.Inputs.Add(FNiagaraVariable(FNiagaraTypeDefinition(GetClass()), TEXT("Strand Hair")));
+		Sig.Outputs.Add(FNiagaraVariable(FNiagaraTypeDefinition::GetFloatDef(), TEXT("Tip Thickness")));
+
+		OutFunctions.Add(Sig);
+	}
+	{
+		FNiagaraFunctionSignature Sig;
+		Sig.Name = GetWorldTransformName;
+		Sig.bMemberFunction = true;
+		Sig.bRequiresContext = false;
+		Sig.Inputs.Add(FNiagaraVariable(FNiagaraTypeDefinition(GetClass()), TEXT("Strand Hair")));
+		Sig.Outputs.Add(FNiagaraVariable(FNiagaraTypeDefinition::GetMatrix4Def(), TEXT("World Transform")));
+
+		OutFunctions.Add(Sig);
+	}
+	{
+		FNiagaraFunctionSignature Sig;
+		Sig.Name = GetVertexPositionName;
+		Sig.bMemberFunction = true;
+		Sig.bRequiresContext = false;
+		Sig.Inputs.Add(FNiagaraVariable(FNiagaraTypeDefinition(GetClass()), TEXT("Strand Hair")));
+		Sig.Inputs.Add(FNiagaraVariable(FNiagaraTypeDefinition::GetIntDef(), TEXT("Vertex Index")));
+		Sig.Outputs.Add(FNiagaraVariable(FNiagaraTypeDefinition::GetVec3Def(), TEXT("Vertex Position")));
+
+		OutFunctions.Add(Sig);
+	}
+	{
+		FNiagaraFunctionSignature Sig;
+		Sig.Name = ComputeNodePositionName;
+		Sig.bMemberFunction = true;
+		Sig.bRequiresContext = false;
+		Sig.Inputs.Add(FNiagaraVariable(FNiagaraTypeDefinition(GetClass()), TEXT("Strand Hair")));
+		Sig.Inputs.Add(FNiagaraVariable(FNiagaraTypeDefinition::GetIntDef(), TEXT("Node Index")));
+		Sig.Outputs.Add(FNiagaraVariable(FNiagaraTypeDefinition::GetVec3Def(), TEXT("Node Position")));
+
+		OutFunctions.Add(Sig);
+	}
+}
+
+DEFINE_NDI_DIRECT_FUNC_BINDER(UNiagaraDataInterfaceHairStrands, GetNumStrands);
+DEFINE_NDI_DIRECT_FUNC_BINDER(UNiagaraDataInterfaceHairStrands, GetStrandDensity);
+DEFINE_NDI_DIRECT_FUNC_BINDER(UNiagaraDataInterfaceHairStrands, GetStrandSize);
+DEFINE_NDI_DIRECT_FUNC_BINDER(UNiagaraDataInterfaceHairStrands, GetRootThickness);
+DEFINE_NDI_DIRECT_FUNC_BINDER(UNiagaraDataInterfaceHairStrands, GetTipThickness);
+DEFINE_NDI_DIRECT_FUNC_BINDER(UNiagaraDataInterfaceHairStrands, GetWorldTransform);
+DEFINE_NDI_DIRECT_FUNC_BINDER(UNiagaraDataInterfaceHairStrands, GetVertexPosition);
+DEFINE_NDI_DIRECT_FUNC_BINDER(UNiagaraDataInterfaceHairStrands, ComputeNodePosition);
+
+void UNiagaraDataInterfaceHairStrands::GetVMExternalFunction(const FVMExternalFunctionBindingInfo& BindingInfo, void* InstanceData, FVMExternalFunction &OutFunc)
+{
+	if (BindingInfo.Name == GetNumStrandsName)
+	{
+		check(BindingInfo.GetNumInputs() == 1 && BindingInfo.GetNumOutputs() == 1);
+		NDI_FUNC_BINDER(UNiagaraDataInterfaceHairStrands, GetNumStrands)::Bind(this, OutFunc);
+	}
+	else if (BindingInfo.Name == GetStrandDensityName)
+	{
+		check(BindingInfo.GetNumInputs() == 1 && BindingInfo.GetNumOutputs() == 1);
+		NDI_FUNC_BINDER(UNiagaraDataInterfaceHairStrands, GetStrandDensity)::Bind(this, OutFunc);
+	}
+	else if (BindingInfo.Name == GetStrandSizeName)
+	{
+		check(BindingInfo.GetNumInputs() == 1 && BindingInfo.GetNumOutputs() == 1);
+		NDI_FUNC_BINDER(UNiagaraDataInterfaceHairStrands, GetStrandSize)::Bind(this, OutFunc);
+	}
+	else if (BindingInfo.Name == GetRootThicknessName)
+	{
+		check(BindingInfo.GetNumInputs() == 1 && BindingInfo.GetNumOutputs() == 1);
+		NDI_FUNC_BINDER(UNiagaraDataInterfaceHairStrands, GetRootThickness)::Bind(this, OutFunc);
+	}
+	else if (BindingInfo.Name == GetTipThicknessName)
+	{
+		check(BindingInfo.GetNumInputs() == 1 && BindingInfo.GetNumOutputs() == 1);
+		NDI_FUNC_BINDER(UNiagaraDataInterfaceHairStrands, GetTipThickness)::Bind(this, OutFunc);
+	}
+	else if (BindingInfo.Name == GetWorldTransformName)
+	{
+		check(BindingInfo.GetNumInputs() == 1 && BindingInfo.GetNumOutputs() == 16);
+		NDI_FUNC_BINDER(UNiagaraDataInterfaceHairStrands, GetWorldTransform)::Bind(this, OutFunc);
+	}
+	else if (BindingInfo.Name == GetVertexPositionName)
+	{
+		check(BindingInfo.GetNumInputs() == 2 && BindingInfo.GetNumOutputs() == 3);
+		NDI_FUNC_BINDER(UNiagaraDataInterfaceHairStrands, GetVertexPosition)::Bind(this, OutFunc);
+	}
+	else if (BindingInfo.Name == ComputeNodePositionName)
+	{
+		check(BindingInfo.GetNumInputs() == 2 && BindingInfo.GetNumOutputs() == 3);
+		NDI_FUNC_BINDER(UNiagaraDataInterfaceHairStrands, ComputeNodePosition)::Bind(this, OutFunc);
+	}
+}
+
+void UNiagaraDataInterfaceHairStrands::GetNumStrands(FVectorVMContext& Context)
+{
+	VectorVM::FUserPtrHandler<FNDIHairStrandsData> InstData(Context);
+	VectorVM::FExternalFuncRegisterHandler<int32> OutNumStrands(Context);
+
+	for (int32 InstanceIdx = 0; InstanceIdx < Context.NumInstances; ++InstanceIdx)
+	{
+		*OutNumStrands.GetDestAndAdvance() = InstData->NumStrands;
+	}
+}
+
+void UNiagaraDataInterfaceHairStrands::GetStrandSize(FVectorVMContext& Context)
+{
+	VectorVM::FUserPtrHandler<FNDIHairStrandsData> InstData(Context);
+	VectorVM::FExternalFuncRegisterHandler<int32> OutStrandSize(Context);
+
+	for (int32 InstanceIdx = 0; InstanceIdx < Context.NumInstances; ++InstanceIdx)
+	{
+		*OutStrandSize.GetDestAndAdvance() = InstData->StrandSize;
+	}
+}
+
+void UNiagaraDataInterfaceHairStrands::GetStrandDensity(FVectorVMContext& Context)
+{
+	VectorVM::FUserPtrHandler<FNDIHairStrandsData> InstData(Context);
+	VectorVM::FExternalFuncRegisterHandler<float> OutStrandDensity(Context);
+
+	for (int32 InstanceIdx = 0; InstanceIdx < Context.NumInstances; ++InstanceIdx)
+	{
+		*OutStrandDensity.GetDestAndAdvance() = InstData->StrandDensity;
+	}
+}
+
+void UNiagaraDataInterfaceHairStrands::GetRootThickness(FVectorVMContext& Context)
+{
+	VectorVM::FUserPtrHandler<FNDIHairStrandsData> InstData(Context);
+	VectorVM::FExternalFuncRegisterHandler<float> OutRootThickness(Context);
+
+	for (int32 InstanceIdx = 0; InstanceIdx < Context.NumInstances; ++InstanceIdx)
+	{
+		*OutRootThickness.GetDestAndAdvance() = InstData->RootThickness;
+	}
+}
+
+void UNiagaraDataInterfaceHairStrands::GetTipThickness(FVectorVMContext& Context)
+{
+	VectorVM::FUserPtrHandler<FNDIHairStrandsData> InstData(Context);
+	VectorVM::FExternalFuncRegisterHandler<float> OutTipThickness(Context);
+
+	for (int32 InstanceIdx = 0; InstanceIdx < Context.NumInstances; ++InstanceIdx)
+	{
+		*OutTipThickness.GetDestAndAdvance() = InstData->TipThickness;
+	}
+}
+
+void UNiagaraDataInterfaceHairStrands::GetWorldTransform(FVectorVMContext& Context)
+{
+	VectorVM::FUserPtrHandler<FNDIHairStrandsData> InstData(Context);
+	VectorVM::FExternalFuncRegisterHandler<FMatrix> OutWorldTransform(Context);
+
+	for (int32 InstanceIdx = 0; InstanceIdx < Context.NumInstances; ++InstanceIdx)
+	{
+		*OutWorldTransform.GetDestAndAdvance() = InstData->WorldTransform;
+	}
+}
+
+void UNiagaraDataInterfaceHairStrands::GetVertexPosition(FVectorVMContext& Context)
+{
+	if (SourceAsset)
+	{
+		VectorVM::FUserPtrHandler<FNDIHairStrandsData> InstData(Context);
+		VectorVM::FExternalFuncInputHandler<int32> InVertexIndex(Context);
+
+		VectorVM::FExternalFuncRegisterHandler<float> OutVertexPosX(Context);
+		VectorVM::FExternalFuncRegisterHandler<float> OutVertexPosY(Context);
+		VectorVM::FExternalFuncRegisterHandler<float> OutVertexPosZ(Context);
+
+		FMatrix& WorldTransform = InstData->WorldTransform;
+
+		const uint32 NumVerts = InstData->NumStrands * InstData->StrandSize;
+		FVector Pos;
+		for (int32 i = 0; i < Context.NumInstances; i++)
+		{
+			const uint32 VertexIndex = InVertexIndex.Get() % NumVerts;
+			Pos = SourceAsset->StrandsDatas.StrandsPoints.PointsPosition[VertexIndex];
+			WorldTransform.TransformPosition(Pos);
+
+			*OutVertexPosX.GetDestAndAdvance() = Pos.X;
+			*OutVertexPosY.GetDestAndAdvance() = Pos.Y;
+			*OutVertexPosZ.GetDestAndAdvance() = Pos.Z;
+		}
+	}
+}
+
+void UNiagaraDataInterfaceHairStrands::ComputeNodePosition(FVectorVMContext& Context)
+{
+	if (SourceAsset)
+	{
+		VectorVM::FUserPtrHandler<FNDIHairStrandsData> InstData(Context);
+		VectorVM::FExternalFuncInputHandler<int32> InVertexIndex(Context);
+
+		VectorVM::FExternalFuncRegisterHandler<float> OutVertexPosX(Context);
+		VectorVM::FExternalFuncRegisterHandler<float> OutVertexPosY(Context);
+		VectorVM::FExternalFuncRegisterHandler<float> OutVertexPosZ(Context);
+
+		FMatrix& WorldTransform = InstData->WorldTransform;
+
+		const uint32 NumVerts = InstData->NumStrands * InstData->StrandSize;
+		FVector Pos;
+		for (int32 i = 0; i < Context.NumInstances; i++)
+		{
+			const uint32 VertexIndex = InVertexIndex.Get() % NumVerts;
+			Pos = SourceAsset->StrandsDatas.StrandsPoints.PointsPosition[VertexIndex];
+			WorldTransform.TransformPosition(Pos);
+
+			*OutVertexPosX.GetDestAndAdvance() = Pos.X;
+			*OutVertexPosY.GetDestAndAdvance() = Pos.Y;
+			*OutVertexPosZ.GetDestAndAdvance() = Pos.Z;
+		}
+	}
+}
+
+bool UNiagaraDataInterfaceHairStrands::GetFunctionHLSL(const FName& DefinitionFunctionName, FString InstanceFunctionName, FNiagaraDataInterfaceGPUParamInfo& ParamInfo, FString& OutHLSL)
+{
+	FNDIHairStrandsParametersName ParamNames(ParamInfo.DataInterfaceHLSLSymbol);
+
+	TMap<FString, FStringFormatArg> ArgsSample = {
+		{TEXT("InstanceFunctionName"), InstanceFunctionName},
+		{TEXT("NumStrandsName"), ParamNames.NumStrandsName},
+		{TEXT("StrandSizeName"), ParamNames.StrandSizeName},
+		{TEXT("StrandDensityName"), ParamNames.StrandDensityName},
+		{TEXT("RootThicknessName"), ParamNames.RootThicknessName},
+		{TEXT("TipThicknessName"), ParamNames.TipThicknessName},
+		{TEXT("WorldTransformName"), ParamNames.WorldTransformName},
+		{TEXT("PointsPositionsBufferName"), ParamNames.PointsPositionsBufferName},
+		{TEXT("CurvesOffsetsBufferName"), ParamNames.CurvesOffsetsBufferName},
+		{TEXT("NodesPositionsBufferName"), ParamNames.NodesPositionsBufferName},
+		{TEXT("PointsNodesBufferName"), ParamNames.PointsNodesBufferName},
+		{TEXT("PointsWeightsBufferName"), ParamNames.PointsWeightsBufferName},
+		{TEXT("HairStrandsContextName"), TEXT("DIHAIRSTRANDS_MAKE_CONTEXT(") + ParamInfo.DataInterfaceHLSLSymbol + TEXT(")")},
+	};
+
+	if (DefinitionFunctionName == GetStrandDensityName)
+	{
+		static const TCHAR *FormatSample = TEXT(R"(
+		void {InstanceFunctionName}(out float Out_StrandDensity)
+		{
+			Out_StrandDensity = {StrandDensitysName};
+		}
+		)");
+		OutHLSL += FString::Format(FormatSample, ArgsSample);
+		return true;
+	}
+	else if (DefinitionFunctionName == GetStrandSizeName)
+	{
+		static const TCHAR *FormatSample = TEXT(R"(
+		void {InstanceFunctionName}(out int Out_StrandSize)
+		{
+			Out_StrandSize = {StrandSizeName};
+		}
+		)");
+		OutHLSL += FString::Format(FormatSample, ArgsSample);
+		return true;
+	}
+	else if (DefinitionFunctionName == GetNumStrandsName)
+	{
+		static const TCHAR *FormatSample = TEXT(R"(
+		void {InstanceFunctionName}(out int Out_NumStrands)
+		{
+			Out_NumStrands = {NumStrandsName};
+		}
+		)");
+		OutHLSL += FString::Format(FormatSample, ArgsSample);
+		return true;
+	}
+	else if (DefinitionFunctionName == GetRootThicknessName)
+	{
+		static const TCHAR *FormatSample = TEXT(R"(
+		void {InstanceFunctionName}(out float Out_RootThickness)
+		{
+			Out_RootThickness = {RootThicknessName};
+		}
+		)");
+		OutHLSL += FString::Format(FormatSample, ArgsSample);
+		return true;
+	}
+	else if (DefinitionFunctionName == GetTipThicknessName)
+	{
+		static const TCHAR *FormatSample = TEXT(R"(
+		void {InstanceFunctionName}(out float Out_TipThickness)
+		{
+			Out_TipThickness = {TipThicknessName};
+		}
+		)");
+		OutHLSL += FString::Format(FormatSample, ArgsSample);
+		return true;
+	}
+	else if (DefinitionFunctionName == GetWorldTransformName)
+	{
+		static const TCHAR *FormatSample = TEXT(R"(
+		void {InstanceFunctionName}(out float4x4 Out_WorldTransform)
+		{
+			Out_WorldTransform = {WorldTransformName};
+		}
+		)");
+		OutHLSL += FString::Format(FormatSample, ArgsSample);
+		return true;
+	}
+	else if (DefinitionFunctionName == GetVertexPositionName)
+	{
+		static const TCHAR *FormatSample = TEXT(R"(
+			void {InstanceFunctionName} (in int VertexIndex, out float3 Out_VertexPosition)
+			{
+				VertexIndex *= 3 ;
+				Out_VertexPosition = float3({PointsPositionsBufferName}[VertexIndex], {PointsPositionsBufferName}[VertexIndex+1], {PointsPositionsBufferName}[VertexIndex+2]);
+				Out_VertexPosition = mul(float4(Out_VertexPosition, 1.0), {WorldTransformName}).xyz;
+			}
+			)");
+		OutHLSL += FString::Format(FormatSample, ArgsSample);
+		return true;
+	}
+	/*else if (DefinitionFunctionName == ComputeNodePositionName)
+	{
+		static const TCHAR *FormatSample = TEXT(R"(
+			void {InstanceFunctionName} (in int VertexIndex, out float3 Out_VertexPosition)
+			{
+				VertexIndex *= 3 ;
+				Out_VertexPosition = float3({PointsPositionsBufferName}[VertexIndex], {PointsPositionsBufferName}[VertexIndex+1], {PointsPositionsBufferName}[VertexIndex+2]);
+				Out_VertexPosition = mul(float4(Out_VertexPosition, 1.0), {WorldTransformName}).xyz;
+			}
+			)");
+		OutHLSL += FString::Format(FormatSample, ArgsSample);
+		return true;
+	}*/
+
+	//else if (DefinitionFunctionName == ComputeStrandIndexName)
+	//{
+	//	static const TCHAR *FormatSample = TEXT(R"(
+	//		void {InstanceFunctionName} (in int NodeIndex, out int2 Out_StrandIndex)
+	//		{
+	//			Out_StrandIndex.x = NodeIndex / {StrandSizeName};
+	//			Out_StrandIndex.y = NodeIndex - Out_StrandIndex.x * {StrandSizeName};
+	//		}
+	//		)");
+	//	OutHLSL += FString::Format(FormatSample, ArgsSample);
+	//	return true;
+	//}
+	else if (DefinitionFunctionName == ComputeNodePositionName)
+	{
+		static const TCHAR *FormatSample = TEXT(R"(
+			void {InstanceFunctionName} (in int NodeIndex, out float3 Out_NodePosition)
+			{
+				{HairStrandsContextName} DIHairStrands_ComputeNodePosition(DIContext,NodeIndex,Out_NodePosition);
+			}
+			)");
+		OutHLSL += FString::Format(FormatSample, ArgsSample);
+		return true;
+	}
+
+	OutHLSL += TEXT("\n");
+	return false;
+}
+
+void UNiagaraDataInterfaceHairStrands::GetCommonHLSL(FString& OutHLSL)
+{
+	OutHLSL += TEXT("#include \"/Plugin/Experimental/HairStrands/Private/NiagaraQuaternionUtils.ush\"\n");
+	OutHLSL += TEXT("#include \"/Plugin/Experimental/HairStrands/Private/NiagaraDataInterfaceHairStrands.ush\"\n");
+}
+
+void UNiagaraDataInterfaceHairStrands::GetParameterDefinitionHLSL(FNiagaraDataInterfaceGPUParamInfo& ParamInfo, FString& OutHLSL)
+{
+	OutHLSL += TEXT("DIHAIRSTRANDS_DECLARE_CONSTANTS(") + ParamInfo.DataInterfaceHLSLSymbol + TEXT(")\n");
+}
+
+void UNiagaraDataInterfaceHairStrands::ProvidePerInstanceDataForRenderThread(void* DataForRenderThread, void* PerInstanceData, const FGuid& SystemInstance)
+{
+	check(Proxy);
+
+	FNDIHairStrandsData* GameThreadData = static_cast<FNDIHairStrandsData*>(PerInstanceData);
+	FNDIHairStrandsData* RenderThreadData = static_cast<FNDIHairStrandsData*>(DataForRenderThread);
+
+	RenderThreadData->WorldTransform = GameThreadData->WorldTransform;
+}
+
+FNiagaraDataInterfaceParametersCS*
+UNiagaraDataInterfaceHairStrands::ConstructComputeParameters() const
+{
+	return new FNDIHairStrandsParametersCS();
+}
+
+
+#undef LOCTEXT_NAMESPACE
\ No newline at end of file
diff --git a/Engine/Plugins/Experimental/HairStrands/Source/HairStrandsNiagara/Public/HairStrandsNiagara.h b/Engine/Plugins/Experimental/HairStrands/Source/HairStrandsNiagara/Public/HairStrandsNiagara.h
new file mode 100644
--- /dev/null
+++ b/Engine/Plugins/Experimental/HairStrands/Source/HairStrandsNiagara/Public/HairStrandsNiagara.h
@@ -0,0 +1,19 @@
+// Copyright 1998-2019 Epic Games, Inc. All Rights Reserved.
+
+#pragma once
+
+#include "Modules/ModuleInterface.h"
+#include "Modules/ModuleManager.h"
+
+/** Implements the HairStrands module  */
+class FHairStrandsNiagara : public IModuleInterface
+{
+public:
+
+	//~ IModuleInterface interface
+
+	virtual void StartupModule() override;
+	virtual void ShutdownModule() override;
+};
+
+IMPLEMENT_MODULE(FHairStrandsNiagara, HairStrandsNiagara);
diff --git a/Engine/Plugins/Experimental/HairStrands/Source/HairStrandsNiagara/Public/NiagaraDataInterfaceHairStrands.h b/Engine/Plugins/Experimental/HairStrands/Source/HairStrandsNiagara/Public/NiagaraDataInterfaceHairStrands.h
new file mode 100644
--- /dev/null
+++ b/Engine/Plugins/Experimental/HairStrands/Source/HairStrandsNiagara/Public/NiagaraDataInterfaceHairStrands.h
@@ -0,0 +1,237 @@
+// Copyright 1998-2019 Epic Games, Inc. All Rights Reserved.
+
+#pragma once
+
+#include "NiagaraDataInterface.h"
+#include "NiagaraCommon.h"
+#include "VectorVM.h"
+#include "HairStrandsAsset.h"
+#include "NiagaraDataInterfaceHairStrands.generated.h"
+
+/** Size of each strands*/
+UENUM(BlueprintType)
+enum class EHairStrandsSize : uint8
+{
+	Size2 = 0x02 UMETA(DisplatName = "2"),
+	Size4 = 0x04 UMETA(DisplatName = "4"),
+	Size8 = 0x08 UMETA(DisplatName = "8"),
+	Size16 = 0x10 UMETA(DisplatName = "16"),
+	Size32 = 0x20 UMETA(DisplatName = "32")
+};
+
+/** Render buffers that will be used in hlsl functions */
+struct FNDIHairStrandsBuffer : public FRenderResource
+{
+	/** Set the asset that will be used to affect the buffer */
+	void SetAsset(const UHairStrandsAsset*  HairStrandsAsset);
+
+	/** Set the strand size */
+	void SetStrandSize(const uint32 StrandSize);
+
+	/** Init the buffer */
+	virtual void InitRHI() override;
+
+	/** Release the buffer */
+	virtual void ReleaseRHI() override;
+
+	/** Get the resource name */
+	virtual FString GetFriendlyName() const override { return TEXT("FNDIHairStrandsBuffer"); }
+
+	/** Position Buffer UAV */
+	FUnorderedAccessViewRHIRef PointsPositionsBufferUav;
+
+	/** Strand nodes positions buffer */
+	FRWBuffer NodesPositionsBuffer;
+
+	/** Strand points node index buffer */
+	FRWBuffer PointsNodesBuffer;
+
+	/** Strand points node weights buffer */
+	FRWBuffer PointsWeightsBuffer;
+
+	/** Strand curves point offset buffer */
+	FRWBuffer CurvesOffsetsBuffer;
+
+	/** The strand asset from which to sample */
+	const UHairStrandsAsset* SourceAsset;
+
+	/** Strand size that will be used for resampling*/
+	uint32 StrandSize;
+};
+
+/** Data stored per strand base instance*/
+struct FNDIHairStrandsData
+{
+	/** Cached World transform. */
+	FMatrix WorldTransform;
+
+	/** Number of strands*/
+	int32 NumStrands;
+
+	/** Strand size */
+	int32 StrandSize;
+
+	/** Strand Density */
+	float StrandDensity;
+
+	/** Root Thickness */
+	float RootThickness;
+
+	/** Tip Thickness */
+	float TipThickness;
+
+	/** Strands Gpu buffer */
+	FNDIHairStrandsBuffer* StrandBuffer;
+};
+
+/** Data Interface for the strand base */
+UCLASS(EditInlineNew, Category = "Strands", meta = (DisplayName = "Hair Strands"))
+class HAIRSTRANDSNIAGARA_API UNiagaraDataInterfaceHairStrands : public UNiagaraDataInterface
+{
+	GENERATED_UCLASS_BODY()
+
+public:
+
+	/** Size of each strand. */
+	UPROPERTY(EditAnywhere, Category = "Spawn")
+	EHairStrandsSize StrandSize;
+
+	/** Density of each strand. */
+	UPROPERTY(EditAnywhere, Category = "Spawn")
+	float StrandDensity;
+
+	/** Strand Root thickness. */
+	UPROPERTY(EditAnywhere, Category = "Spawn")
+	float RootThickness;
+
+	/** Strand Tip thickness. */
+	UPROPERTY(EditAnywhere, Category = "Spawn")
+	float TipThickness;
+
+	/** Mesh used to sample from when not overridden by a source actor from the scene. Also useful for previewing in the editor. */
+	UPROPERTY(EditAnywhere, Category = "Source")
+	UHairStrandsAsset* SourceAsset;
+
+	/** Source transform to be applied to the asset. */
+	UPROPERTY(EditAnywhere, Category = "Source")
+	FMatrix SourceTransform;
+
+	/** Static mesh on which the roots are attached. */
+	UPROPERTY(EditAnywhere, Category = "Attachment")
+	UStaticMesh* StaticMesh;
+
+	///** The source actor from which to sample */
+	//UPROPERTY(EditAnywhere, AdvancedDisplay, Category = "Source")
+	//AActor* SourceActor;
+
+	///** The source component from which to sample */
+	//class UHairStrandsComponent* SourceComponent;
+
+	/** UObject Interface */
+	virtual void PostInitProperties() override;
+
+	/** UNiagaraDataInterface Interface */
+	virtual void GetFunctions(TArray<FNiagaraFunctionSignature>& OutFunctions) override;
+	virtual void GetVMExternalFunction(const FVMExternalFunctionBindingInfo& BindingInfo, void* InstanceData, FVMExternalFunction &OutFunc) override;
+	virtual bool CanExecuteOnTarget(ENiagaraSimTarget Target) const override { return Target == ENiagaraSimTarget::GPUComputeSim; }
+	virtual bool InitPerInstanceData(void* PerInstanceData, FNiagaraSystemInstance* SystemInstance) override;
+	virtual void DestroyPerInstanceData(void* PerInstanceData, FNiagaraSystemInstance* SystemInstance)override;
+	virtual bool PerInstanceTick(void* PerInstanceData, FNiagaraSystemInstance* SystemInstance, float DeltaSeconds) override;
+	virtual int32 PerInstanceDataSize()const override { return sizeof(FNDIHairStrandsData); }
+	virtual bool Equals(const UNiagaraDataInterface* Other) const override;
+
+	/** GPU simulation  functionality */
+	virtual bool GetFunctionHLSL(const FName&  DefinitionFunctionName, FString InstanceFunctionName, FNiagaraDataInterfaceGPUParamInfo& ParamInfo, FString& OutHLSL) override;
+	virtual void GetParameterDefinitionHLSL(FNiagaraDataInterfaceGPUParamInfo& ParamInfo, FString& OutHLSL) override;
+	virtual FNiagaraDataInterfaceParametersCS* ConstructComputeParameters() const override;
+	virtual void ProvidePerInstanceDataForRenderThread(void* DataForRenderThread, void* PerInstanceData, const FGuid& SystemInstance) override;
+	virtual void GetCommonHLSL(FString& OutHLSL) override;
+
+	/** Get the number of strands */
+	void GetNumStrands(FVectorVMContext& Context);
+
+	/** Get the strand size  */
+	void GetStrandSize(FVectorVMContext& Context);
+
+	/** Get the strand density */
+	void GetStrandDensity(FVectorVMContext& Context);
+
+	/** Get the strand thickness at the root */
+	void GetRootThickness(FVectorVMContext& Context);
+
+	/** Get the strand thickness at the tip */
+	void GetTipThickness(FVectorVMContext& Context);
+
+	/** Get the world transform */
+	void GetWorldTransform(FVectorVMContext& Context);
+
+	/** Get the strand vertex position in world space*/
+	void GetVertexPosition(FVectorVMContext& Context);
+
+	/** Get the strand node position in world space*/
+	void ComputeNodePosition(FVectorVMContext& Context);
+
+	/** Name of the world transform */
+	static const FString WorldTransformName;
+
+	/** Name of the number of strands */
+	static const FString NumStrandsName;
+
+	/** Name of the strand size */
+	static const FString StrandSizeName;
+
+	/** Name of the strand density */
+	static const FString StrandDensityName;
+
+	/** Name of the root thickness */
+	static const FString RootThicknessName;
+
+	/** Name of the tip thickness */
+	static const FString TipThicknessName;
+
+	/** Name of the points positions buffer */
+	static const FString PointsPositionsBufferName;
+
+	/** Name of the curves offsets buffer */
+	static const FString CurvesOffsetsBufferName;
+
+	/** Name of the nodes positions buffer */
+	static const FString NodesPositionsBufferName;
+
+	/** Name of the points nodes buffer */
+	static const FString PointsNodesBufferName;
+
+	/** Name of the points weights buffer */
+	static const FString PointsWeightsBufferName;
+
+protected:
+	/** Copy one niagara DI to this */
+	virtual bool CopyToInternal(UNiagaraDataInterface* Destination) const override;
+};
+
+/** Proxy to send data to gpu */
+struct FNDIHairStrandsProxy : public FNiagaraDataInterfaceProxy
+{
+	/** Destroy internal data */
+	virtual void DeferredDestroy() override;
+
+	/** Get the size of the data that will be passed to render*/
+	virtual int32 PerInstanceDataPassedToRenderThreadSize() const override { return sizeof(FNDIHairStrandsData); }
+
+	/** Get the data that will be passed to render*/
+	virtual void ConsumePerInstanceDataFromGameThread(void* PerInstanceData, const FGuid& Instance) override;
+
+	/** Initialize the Proxy data strands buffer */
+	void InitializePerInstanceData(const FGuid& SystemInstance, FNDIHairStrandsBuffer* StrandsBuffer, const uint32 NumStrands, const uint8 StrandSize,
+		const float StrandDensity, const float RootThickness, const float TipThickness);
+
+	/** Destroy the proxy data if necessary */
+	void DestroyPerInstanceData(NiagaraEmitterInstanceBatcher* Batcher, const FGuid& SystemInstance);
+
+	/** List of proxy data for each system instances*/
+	TMap<FGuid, FNDIHairStrandsData> SystemInstancesToProxyData;
+
+	/** List of proxy data to destroy later */
+	TSet<FGuid> DeferredDestroyList;
+};
+
diff --git a/Engine/Plugins/FX/Niagara/Shaders/Private/NiagaraEmitterInstanceShader.usf b/Engine/Plugins/FX/Niagara/Shaders/Private/NiagaraEmitterInstanceShader.usf
--- a/Engine/Plugins/FX/Niagara/Shaders/Private/NiagaraEmitterInstanceShader.usf
+++ b/Engine/Plugins/FX/Niagara/Shaders/Private/NiagaraEmitterInstanceShader.usf
@@ -37,13 +37,25 @@ float Reciprocal(float x){ return 1.0f/x; }
  */
 #if GPU_SIMULATION
 
-	#include "/Engine/Private/SceneTexturesCommon.ush"
+	#define THREADGROUP_SIZE 64
 
 	static uint3 GDispatchThreadId;
+	static uint3 GGroupThreadId;
 	static uint GCurrentPhase;
+	static uint GEmitterTickCounter;
+	static uint GSimStart;
+
+	// To be removed as well but don't know who is using it
+	#include "/Engine/Private/SceneTexturesCommon.ush"
+
+	// Temporary include while everything is moved to plugins
+	#if USE_SHADER_STAGES == 1
+		#include "NiagaraStrandsPhysics.ush"
+	#endif
+
 	uint EmitterTickCounter;
 
-	#define THREADGROUP_SIZE 64
+	
 
 	float4 Modulo(float4 x, float4 y){ return fmod(x,y); }
 	float3 Modulo(float3 x, float3 y){ return fmod(x,y); }
@@ -324,6 +336,8 @@ int rand(int x, int Seed1, int Seed2, int Seed3)
 	/* Parameters for spawning and simulations
 	 */
 	uint UpdateStartInstance;
+	uint ShaderStageIndex;
+	int IterationInterfaceCount;
 	uint SpawnedInstances;
 
 	uint ComponentBufferSizeRead;
@@ -332,6 +346,7 @@ int rand(int x, int Seed1, int Seed2, int Seed3)
 	int NumParticlesPerEvent;
 	uint CopyInstancesBeforeStart;
 
+	uint SimStart;
 
 	/* Buffers for particle data and DrawIndirect calls
 	 */
@@ -373,20 +388,38 @@ int rand(int x, int Seed1, int Seed2, int Seed3)
 		MustBe0[DataSetID] = 0;
 		// End static assert
 
-		int PrevIdx = GSpawnStartInstance + SpawnedInstances;	// scratch instance as default; write to that for dead particles
-		if (bDoAcquire)
-		{
-			// Have to use uint's here to avoid PS4 compiler warnings about InterlockedAdd, cannot propagate uint due to CPU VM limitations...
-			uint RetPrevIdx;
-			// @TODO : add some TLS logic to avoid thread group for doing atomic for each thread. (gathering the actual required count)
-			InterlockedAdd(RWInstanceCounts[WriteInstanceCountOffset], (uint)1U, RetPrevIdx);
-			PrevIdx = (int)RetPrevIdx;
-		}
-		return PrevIdx;
+		#if USE_SHADER_STAGES
+			if(ShaderStageIndex == 0)
+			{
+				int PrevIdx = GSpawnStartInstance + SpawnedInstances;	// scratch instance as default; write to that for dead particles
+				if(bDoAcquire == true)
+				{
+					// Have to use uint's here to avoid PS4 compiler warnings about InterlockedAdd, cannot propagate uint due to CPU VM limitations...
+					uint RetPrevIdx;
+					// @TODO : add some TLS logic to avoid thread group for doing atomic for each thread. (gathering the actual required count)
+					InterlockedAdd(RWInstanceCounts[WriteInstanceCountOffset], (uint)1U, RetPrevIdx);
+					PrevIdx = (int)RetPrevIdx;
+				}
+				return PrevIdx;
+			}
+			else
+			{
+				return ExecIndex();
+			}
+		#else
+			int PrevIdx = GSpawnStartInstance + SpawnedInstances;	// scratch instance as default; write to that for dead particles
+			if(bDoAcquire == true)
+			{
+				// Have to use uint's here to avoid PS4 compiler warnings about InterlockedAdd, cannot propagate uint due to CPU VM limitations...
+				uint RetPrevIdx;
+				// @TODO : add some TLS logic to avoid thread group for doing atomic for each thread. (gathering the actual required count)
+				InterlockedAdd(RWInstanceCounts[WriteInstanceCountOffset], (uint)1U, RetPrevIdx);
+				PrevIdx = (int)RetPrevIdx;
+			}
+			return PrevIdx;
+		#endif
 	}
 
-
-
 	/* ---------------------------------------------------------------------
 	 * InputData operations 
 	 * ---------------------------------------------------------------------
@@ -430,6 +463,20 @@ int rand(int x, int Seed1, int Seed2, int Seed3)
 	void ExitStatScope()	{}
 #endif
 
+	/* 
+	 * Get the index to write onto the output buffer
+	 */
+	int OutputIndex(const int DataSetID, const bool bIsPersistent, const bool bIsValid)
+	{
+		#if USE_SHADER_STAGES
+			const int DenseIndex = AcquireIndex(DataSetID,bIsValid);
+			const int SparseIndex = ExecIndex();
+			return (bIsPersistent && (DenseIndex != SparseIndex) ) ? SparseIndex : DenseIndex;
+		#else
+			return AcquireIndex(DataSetID,bIsValid);
+		#endif
+	}
+
 //Include the simulation shader code generated by the node graph.
 #include "/Engine/Generated/NiagaraEmitterInstance.ush"
 
@@ -447,7 +494,9 @@ void SimulateMainComputeCS(
 	uint3 GroupThreadId : SV_GroupThreadID)
 {
 	GDispatchThreadId = DispatchThreadId;
-
+	GGroupThreadId = GroupThreadId;
+	GEmitterTickCounter = EmitterTickCounter;
+	GSimStart = SimStart;
 	/*
 	if(CopyInstancesBeforeStart == 1)
 	{
@@ -466,7 +515,25 @@ void SimulateMainComputeCS(
 	}
 	int MaxInstances = GSpawnStartInstance + SpawnedInstances;
 
-	if(InstanceID < UpdateStartInstance + MaxInstances )
+	bool UseInterface = false;
+	#if USE_SHADER_STAGES
+	if (IterationInterfaceCount > 0)
+	{
+		if (InstanceID < IterationInterfaceCount)
+		{
+			if (GSimStart == 1)
+			{
+				SimulateMain(InstanceID, 0, 0);
+			}
+			else
+			{
+				SimulateMain(InstanceID, 0, 1);
+			}
+		}
+		UseInterface = true;
+	}
+	#endif
+	if( !UseInterface && InstanceID < UpdateStartInstance + MaxInstances )
 	{
 		// calculate event index for the current instance; this is primarily for spawn events, where we do a single read from a secondary DataSet for a subset of particles
 		//  for all-particle events, this will always end up being 0 because NumParticlesPerEvent = total particles in the set
diff --git a/Engine/Plugins/FX/Niagara/Shaders/Private/NiagaraStrandsPhysics.ush b/Engine/Plugins/FX/Niagara/Shaders/Private/NiagaraStrandsPhysics.ush
new file mode 100644
--- /dev/null
+++ b/Engine/Plugins/FX/Niagara/Shaders/Private/NiagaraStrandsPhysics.ush
@@ -0,0 +1,456 @@
+// Copyright 1998-2019 Epic Games, Inc. All Rights Reserved.
+
+/*=============================================================================
+	NiagaraStrandsPhysics.ush
+=============================================================================*/
+
+#if GPU_SIMULATION
+
+/* -----------------------------------------------------------------
+ * Shared memory for datas that will be accessed within the constraints
+ * -----------------------------------------------------------------
+ */
+
+groupshared float3 SharedLocalPosition[THREADGROUP_SIZE];
+groupshared float4 SharedLocalQuaternion[THREADGROUP_SIZE];
+
+groupshared float SharedInverseMass[THREADGROUP_SIZE];
+groupshared float SharedInverseInertia[THREADGROUP_SIZE];
+
+groupshared float3 SharedLocalVelocity[THREADGROUP_SIZE];
+groupshared float3 SharedLocalOmega[THREADGROUP_SIZE];
+
+/* -----------------------------------------------------------------
+ * Get the strand parent position
+ * -----------------------------------------------------------------
+ */
+
+void GetParentPosition(const bool IsMobile, out float3 ParentPosition)
+{	
+	ParentPosition = SharedLocalPosition[IsMobile ? GGroupThreadId.x-1 : GGroupThreadId.x];
+}
+
+/* -----------------------------------------------------------------
+ * Fill the shared datas from the particles
+ * -----------------------------------------------------------------
+ */
+
+void ExtractPositionSharedDatas(const bool IsMobile, const float3 InParticlePosition, const float3 InParticleVelocity, const float InParticleMass)
+{	
+	SharedLocalPosition[GGroupThreadId.x] = InParticlePosition;
+	SharedLocalVelocity[GGroupThreadId.x] = InParticleVelocity;
+	SharedInverseMass[GGroupThreadId.x] = (InParticleMass != 0.0 && IsMobile) ? 1.0/InParticleMass : 0.0;
+
+	GroupMemoryBarrier(); 
+}
+
+void ExtractOrientationSharedDatas(const bool IsMobile, const float4 InParticleQuaternion, const float3 InParticleOmega, const float InParticleInertia)
+{
+	SharedLocalQuaternion[GGroupThreadId.x] = InParticleQuaternion;
+	SharedLocalOmega[GGroupThreadId.x] = InParticleOmega;
+	SharedInverseInertia[GGroupThreadId.x] = (InParticleInertia != 0.0 && IsMobile) ? 1.0/InParticleInertia : 0.0;
+
+	GroupMemoryBarrier(); 
+}
+
+/* -----------------------------------------------------------------
+ * Report the shared datas onto the particles
+ * -----------------------------------------------------------------
+ */
+
+void ReportPositionSharedDatas(const float InDeltaTime, inout float3 OutParticleVelocity, inout float3 OutParticlePosition)
+{
+	const float3 CurrentPosition = SharedLocalPosition[GGroupThreadId.x];
+	OutParticleVelocity += (CurrentPosition-OutParticlePosition) / InDeltaTime;
+	OutParticlePosition = CurrentPosition;
+}
+
+void ReportOrientationSharedDatas(const float InDeltaTime, inout float3 OutParticleOmega, inout float4 OutParticleQuaternion)
+{
+	const float4 CurrentQuaternion = SharedLocalQuaternion[GGroupThreadId.x];
+	const float4 DeltaQuat = float4(
+			-OutParticleQuaternion.xyz * CurrentQuaternion.w + CurrentQuaternion.xyz * OutParticleQuaternion.w + cross(CurrentQuaternion.xyz, -OutParticleQuaternion.xyz),
+			CurrentQuaternion.w * OutParticleQuaternion.w - dot(CurrentQuaternion.xyz, -OutParticleQuaternion.xyz));
+
+	OutParticleQuaternion = CurrentQuaternion;
+	OutParticleOmega += DeltaQuat.xyz * 2.0 / InDeltaTime;
+}
+
+/* -----------------------------------------------------------------
+ * Advect the particles shared datas with forces
+ * -----------------------------------------------------------------
+ */
+
+void AdvectPositionSharedDatas(const bool IsMobile, const float3 InParticleAcceleration, const float InDeltaTime, inout float3 OutParticleVelocity,  inout float3 OutParticlePosition)
+{	
+	const float3 ParticleAcceleration = IsMobile ? InParticleAcceleration : float3(0,0,0);
+	OutParticleVelocity += ParticleAcceleration * InDeltaTime;
+	OutParticlePosition += OutParticleVelocity * InDeltaTime;
+}
+
+void AdvectOrientationSharedDatas(const float3 InParticleOmega, const float InDeltaTime, inout float4 OutParticleQuaternion)
+{
+	OutParticleQuaternion = OutParticleQuaternion + 0.5 * InDeltaTime * float4(
+        InParticleOmega.xyz * OutParticleQuaternion.w + cross(InParticleOmega.xyz, OutParticleQuaternion.xyz),
+        - dot(InParticleOmega.xyz, OutParticleQuaternion.xyz) );
+
+	OutParticleQuaternion = normalize( OutParticleQuaternion );
+}
+
+/* -----------------------------------------------------------------
+ * Hooke spring material
+ * -----------------------------------------------------------------
+ */
+
+void StoreHookeSpringConfig(const bool IsMobile, const float InMaterialCompliance, out float OutRestLength)
+{
+	if (!IsMobile){
+		const float3 EdgeDirection = SharedLocalPosition[GGroupThreadId.x] - SharedLocalPosition[GGroupThreadId.x-1];
+		OutRestLength = length(EdgeDirection);
+	}
+	else
+	{
+		OutRestLength = 0.0;
+	}
+}
+
+void ComputeHookeSpringWeight(const float InMaterialCompliance, out float OutMaterialWeight)
+{
+	const float SchurDiagonal = ( SharedInverseMass[GGroupThreadId.x] + SharedInverseMass[GGroupThreadId.x-1] + InMaterialCompliance );
+	OutMaterialWeight = ( SchurDiagonal != 0.0 ) ? 1.0 / SchurDiagonal : 0.0;
+}
+
+void EvalHookeSpringCompliance(const float InDeltaTime, const float InYoungModulus, const float InRodThickness, const float InRestLength, out float OutMaterialCompliance)
+{
+	// Compliance = 1.0 / (k * dt * dt)
+	// with k = Y * A / L 
+	// A is the cross section area (Pi * R * R), L is the rest length and Y is the young modulus
+	OutMaterialCompliance = InRestLength*4.0/(InYoungModulus*PI*InRodThickness*InRodThickness*InDeltaTime*InDeltaTime);
+}
+
+void ResetHookeSpringMultiplier(out float OutMaterialMultiplier)
+{
+	OutMaterialMultiplier = 0;
+}
+
+void UpdateHookeSpringMultiplier(const float InMaterialCompliance, const float InMaterialWeight, const float InRestLength, inout float OutMaterialMultiplier)
+{
+	float3 EdgeDirection = SharedLocalPosition[GGroupThreadId.x] - SharedLocalPosition[GGroupThreadId.x-1];
+	const float EdgeLength = length(EdgeDirection);
+	EdgeDirection /= EdgeLength;
+
+	// XPBD lagrange multiplier update : dL = -(C+compliance*L) / (dC * invM * dCt + alpha)
+	const float DeltaLambda = -((EdgeLength-InRestLength) + OutMaterialMultiplier * InMaterialCompliance) * InMaterialWeight;
+
+	// L += dL
+	OutMaterialMultiplier += DeltaLambda;
+
+	// XPBD position update : dX = dL * dCt * invM
+	const float3 PositionDelta = EdgeDirection * DeltaLambda;
+
+	// dX += dX
+	SharedLocalPosition[GGroupThreadId.x] += PositionDelta * SharedInverseMass[GGroupThreadId.x];
+	SharedLocalPosition[GGroupThreadId.x-1] -= PositionDelta * SharedInverseMass[GGroupThreadId.x-1];
+}
+
+void SolveHookeSpringMaterial(const float InMaterialCompliance, const float InMaterialWeight, const float InRestLength, const bool InIsRed, inout float OutMaterialMultiplier)
+{
+	// Process all the red rods
+	if (InIsRed)
+	{	
+		UpdateHookeSpringMultiplier(InMaterialCompliance,InMaterialWeight,InRestLength,OutMaterialMultiplier);
+	}
+	// Process all the black rods
+	GroupMemoryBarrier(); 
+	if (!InIsRed)
+	{
+		UpdateHookeSpringMultiplier(InMaterialCompliance,InMaterialWeight,InRestLength,OutMaterialMultiplier);
+	}
+	GroupMemoryBarrier(); 
+} 
+
+/* -----------------------------------------------------------------
+ * Stretch/shear rod material
+ * -----------------------------------------------------------------
+ */
+
+ void ComputeStretchShearWeight(const float InMaterialCompliance, const float InRestLength, out float OutMaterialWeight)
+{
+	const float SchurDiagonal = ( SharedInverseMass[GGroupThreadId.x] + SharedInverseMass[GGroupThreadId.x-1] ) / (InRestLength*InRestLength) + 
+					SharedInverseInertia[GGroupThreadId.x-1]*4.0 + InMaterialCompliance;
+	OutMaterialWeight = ( SchurDiagonal != 0.0 ) ? 1.0 / SchurDiagonal : 0.0;
+}
+
+void EvalStretchShearCompliance(const float InDeltaTime, const float InYoungModulus, const float InRodThickness, const float InRestLength, out float OutMaterialCompliance)
+{
+	// Compliance = 1.0 / (k * dt * dt)
+	// with k = L * L * (Y * A / L) (the L*L term before is coming from the fact that our constraint is dL/L and not dL)
+	// A is the cross section area (Pi * R * R), L is the rest length and Y is the young modulus
+	OutMaterialCompliance = 4.0/(InYoungModulus*PI*InRodThickness*InRestLength*InRodThickness*InDeltaTime*InDeltaTime);
+}
+
+void ResetStretchShearMultiplier(out float3 OutMaterialMultiplier)
+{
+	OutMaterialMultiplier = float3(0,0,0);
+}
+
+void UpdateStretchShearMultiplier(const float InMaterialCompliance, const float InMaterialWeight, const float InRestLength,  inout float3 OutMaterialMultiplier)
+{
+	float3 EdgeDirection;
+	float4 q0 = SharedLocalQuaternion[GGroupThreadId.x-1];
+	EdgeDirection[0] = 2.0 * (q0.x * q0.z + q0.w * q0.y);
+	EdgeDirection[1] = 2.0 * (q0.y * q0.z - q0.w * q0.x);
+	EdgeDirection[2] = q0.w * q0.w - q0.x * q0.x - q0.y * q0.y + q0.z * q0.z;
+
+	// XPBD lagrange multiplier update : dL = -(C+compliance*L) / (dC * invM * dCt + alpha)
+	const float3 DeltaLambda = -((SharedLocalPosition[GGroupThreadId.x] - SharedLocalPosition[GGroupThreadId.x-1]) / InRestLength - EdgeDirection + 
+					OutMaterialMultiplier * InMaterialCompliance) * InMaterialWeight;
+
+	// L += dL
+	OutMaterialMultiplier += DeltaLambda;
+
+	// XPBD position update : dX = dL * dCt * invM
+	const float3 PositionDelta  = DeltaLambda/InRestLength;
+
+	// dX += dX
+	SharedLocalPosition[GGroupThreadId.x] += PositionDelta * SharedInverseMass[GGroupThreadId.x];
+	SharedLocalPosition[GGroupThreadId.x-1] -= PositionDelta * SharedInverseMass[GGroupThreadId.x-1];
+
+	const float4 qebar = float4(-q0.y, q0.x, -q0.w, q0.z);	
+	const float4 QuaternionDelta = -2.0 * float4( DeltaLambda.xyz * qebar.w + cross(DeltaLambda.xyz, qebar.xyz),
+					- dot(DeltaLambda.xyz, qebar.xyz) );
+
+	q0 += QuaternionDelta * SharedInverseInertia[GGroupThreadId.x-1];
+	SharedLocalQuaternion[GGroupThreadId.x-1] = normalize(q0);
+}
+
+void SolveStretchShearMaterial(const float InMaterialCompliance, const float InMaterialWeight, const float InRestLength, const bool InIsRed, inout float3 OutMaterialMultiplier)
+{
+	// Process all the red rods
+	if (InIsRed)
+	{	
+		UpdateStretchShearMultiplier(InMaterialCompliance,InMaterialWeight,InRestLength,OutMaterialMultiplier);
+	}
+	// Process all the black rods
+	GroupMemoryBarrier(); 
+	if (!InIsRed)
+	{
+		UpdateStretchShearMultiplier(InMaterialCompliance,InMaterialWeight,InRestLength,OutMaterialMultiplier);
+	}
+	GroupMemoryBarrier(); 
+} 
+
+/* -----------------------------------------------------------------
+ * Bbend/twist rod material
+ * -----------------------------------------------------------------
+ */
+
+ void ComputeBendTwistWeight(const float InMaterialCompliance, const float InRestLength, out float OutMaterialWeight)
+{
+	const float SchurDiagonal = ( SharedInverseInertia[GGroupThreadId.x] + SharedInverseInertia[GGroupThreadId.x-1] ) * 4.0 / (InRestLength*InRestLength) + InMaterialCompliance;
+	OutMaterialWeight = ( SchurDiagonal != 0.0 ) ? 1.0 / SchurDiagonal : 0.0;
+}
+
+void EvalBendTwistCompliance(const float InDeltaTime, const float InYoungModulus, const float InRodThickness, const float InRestLength, out float OutMaterialCompliance)
+{
+	// Compliance = 1.0 / (k * dt * dt)
+	// with k = (Y * I * L) 
+	// A is the polar moment of inertia (Pi * R * R * R * R / 4), L is the rest length and Y is the young modulus
+	OutMaterialCompliance = 64.0/(InYoungModulus*PI*InRestLength*InRodThickness*InRodThickness*InRodThickness*InRodThickness*InDeltaTime*InDeltaTime);
+}
+
+void ResetBendTwistMultiplier(out float3 OutMaterialMultiplier)
+{
+	OutMaterialMultiplier = float3(0,0,0);
+}
+
+void UpdateBendTwistMultiplier(const float InMaterialCompliance, const float InMaterialWeight, const float InRestLength, const float4 InRestDarboux, inout float3 OutMaterialMultiplier)
+{
+	float4 q0 = SharedLocalQuaternion[GGroupThreadId.x-1];
+	float4 q1 = SharedLocalQuaternion[GGroupThreadId.x];
+	const float4 CurrentDarboux = float4(
+			q1.xyz * q0.w - q0.xyz * q1.w + cross(-q0.xyz, q1.xyz),
+			q0.w * q1.w - dot(-q0.xyz, q1.xyz));
+	const float InverseRestLength = 2.0 / InRestLength;
+
+	const float4 OmegaPlus = (CurrentDarboux + InRestDarboux) * InverseRestLength;
+	const float4 OmegaMinus = (CurrentDarboux - InRestDarboux) * InverseRestLength;
+
+	float4 DeltaOmega = (dot(OmegaPlus.xyz,OmegaPlus.xyz) > dot(OmegaMinus.xyz,OmegaMinus.xyz) ) ? OmegaMinus : OmegaPlus;
+	DeltaOmega.xyz += OutMaterialMultiplier * InMaterialCompliance;  
+	DeltaOmega.w = 0;
+	DeltaOmega *= InMaterialWeight;
+
+	OutMaterialMultiplier -= DeltaOmega.xyz;
+
+	SharedLocalQuaternion[GGroupThreadId.x-1] += SharedInverseInertia[GGroupThreadId.x-1] *  float4(
+			DeltaOmega.xyz * q1.w + cross(q1.xyz, DeltaOmega.xyz), - dot(q1.xyz, DeltaOmega.xyz)) * InverseRestLength;
+	SharedLocalQuaternion[GGroupThreadId.x] -= SharedInverseInertia[GGroupThreadId.x] *  float4(
+			DeltaOmega.xyz * q0.w + cross(q0.xyz, DeltaOmega.xyz), - dot(q0.xyz, DeltaOmega.xyz)) * InverseRestLength;
+
+	SharedLocalQuaternion[GGroupThreadId.x-1] = normalize(SharedLocalQuaternion[GGroupThreadId.x-1]);
+	SharedLocalQuaternion[GGroupThreadId.x] = normalize(SharedLocalQuaternion[GGroupThreadId.x]);
+}
+
+void SolveBendTwistMaterial(const float InMaterialCompliance, const float InMaterialWeight, const float InRestLength, const float4 InRestDarboux, const bool InIsRed, inout float3 OutMaterialMultiplier)
+{
+	// Process all the red rods
+	if (InIsRed)
+	{	
+		UpdateBendTwistMultiplier(InMaterialCompliance,InMaterialWeight,InRestLength,InRestDarboux,OutMaterialMultiplier);
+	}
+	// Process all the black rods
+	GroupMemoryBarrier(); 
+	if (!InIsRed)
+	{
+		UpdateBendTwistMultiplier(InMaterialCompliance,InMaterialWeight,InRestLength,InRestDarboux,OutMaterialMultiplier);
+	}
+	GroupMemoryBarrier(); 
+} 
+
+/* -----------------------------------------------------------------
+ * Follow the leader constraint
+ * -----------------------------------------------------------------
+ */
+
+ void SolveFollowTheLeaderConstraint(const int InBlockSize, const float InRestLength, const float InDampening, const float InDeltaTime, const float InBending)
+{
+	const bool IsMobile = (GGroupThreadId.x % InBlockSize) != 0;
+	if (!IsMobile)
+	{
+		const float BendingLimit = 2*InBending*InRestLength;
+		for (int i = GGroupThreadId.x+1, end = GGroupThreadId.x+InBlockSize; i < end; ++i)
+		{
+			if (i > GGroupThreadId.x+1)
+			{
+				float3 EdgeDirection = SharedLocalPosition[i] - SharedLocalPosition[i-2];
+				const float EdgeLength = length(EdgeDirection);
+				EdgeDirection /= EdgeLength;
+				if (EdgeLength<=BendingLimit)
+				{
+					// Comutes dL
+					const float DeltaLambda = -(EdgeLength-BendingLimit);
+					const float3 DeltaPosition = EdgeDirection * DeltaLambda;
+
+					// dX += dX
+					SharedLocalPosition[i] += DeltaPosition;
+
+					const float3 DeltaVelocity = DeltaPosition / InDeltaTime;
+
+					// dV += dV
+					SharedLocalVelocity[i] += SharedInverseMass[i] * DeltaVelocity / ( SharedInverseMass[i] + InDampening * SharedInverseMass[i-2]);
+					SharedLocalVelocity[i-2] -= InDampening * SharedInverseMass[i-2] * DeltaVelocity / ( SharedInverseMass[i] + InDampening * SharedInverseMass[i-2]);
+				}
+				GroupMemoryBarrier(); 
+			}
+
+
+			float3 EdgeDirection = SharedLocalPosition[i] - SharedLocalPosition[i-1];
+			const float EdgeLength = length(EdgeDirection);
+			EdgeDirection /= EdgeLength;
+
+			// Comutes dL
+			const float DeltaLambda = -(EdgeLength-InRestLength);
+			const float3 DeltaPosition = EdgeDirection * DeltaLambda;
+
+			// dX += dX
+			SharedLocalPosition[i] += DeltaPosition;
+
+			//SharedLocalPosition[i] += SharedInverseMass[i] * DeltaPosition / ( SharedInverseMass[i] + InDampening * SharedInverseMass[i-1]);
+			//SharedLocalPosition[i-1] -= InDampening * SharedInverseMass[i-1] * DeltaPosition / ( SharedInverseMass[i] + InDampening * SharedInverseMass[i-1]);
+
+			const float3 DeltaVelocity = DeltaPosition / InDeltaTime;
+
+			// dV += dV
+			SharedLocalVelocity[i] += SharedInverseMass[i] * DeltaVelocity / ( SharedInverseMass[i] + InDampening * SharedInverseMass[i-1]);
+			SharedLocalVelocity[i-1] -= InDampening * SharedInverseMass[i-1] * DeltaVelocity / ( SharedInverseMass[i] + InDampening * SharedInverseMass[i-1]);
+
+			GroupMemoryBarrier(); 
+		}
+	}
+}
+
+/* -----------------------------------------------------------------
+ * Static volume collisions against sphere, box...
+ * -----------------------------------------------------------------
+ */
+
+ float3 ProcessCollisionData( const float3 InDeltaPosition, const float3 InDeltaVelocity, const float3 InCollisionNormal, const float InThicknessSum,
+							const float InStaticFriction, const float InKineticFriction )
+ {
+	const float DeltaVNormal = dot(InDeltaVelocity,InCollisionNormal);
+	const float DeltaPNormal = dot(InDeltaPosition,InCollisionNormal)-InThicknessSum;
+	const float NormDeltaPN = abs(DeltaPNormal);
+
+	if (DeltaPNormal < 0.0)
+	{
+		const float3 CollisionTangent = InDeltaVelocity - DeltaVNormal * InCollisionNormal;
+		const float TangentLength = length(CollisionTangent);
+
+		float3 CollisionTangentA = (TangentLength > 0.0) ? CollisionTangent/TangentLength : (abs(InCollisionNormal.x-1.0f) > 0.0) ? float3(1,0,0) : 
+																					        (abs(InCollisionNormal.y-1.0f) > 0.0) ? float3(0,1,0) :
+																						     float3(0,0,1);
+		float3 CollisionTangentB = cross(InCollisionNormal,CollisionTangentA);
+		CollisionTangentA = cross(CollisionTangentB,InCollisionNormal);
+
+		const float TangentLengthA = length(CollisionTangentA);
+		const float TangentLengthB = length(CollisionTangentB);
+
+		CollisionTangentA = (TangentLengthA > 0.0) ? CollisionTangentA/TangentLengthA : float3(0,0,0);
+		CollisionTangentB = (TangentLengthB > 0.0) ? CollisionTangentB/TangentLengthB : float3(0,0,0);
+
+		const float DeltaVTangentA = dot( InDeltaVelocity, CollisionTangentA);
+		const float DeltaVTangentB = dot( InDeltaVelocity, CollisionTangentB);
+
+		const float NormDeltaVTA = abs(DeltaVTangentA);
+		const float NormDeltaVTB = abs(DeltaVTangentB);
+
+		const float AlphaTangentA = ( NormDeltaVTA < InStaticFriction * NormDeltaPN ) ? 1.0 : (InKineticFriction > 0.0) ? min(InKineticFriction*NormDeltaPN/NormDeltaVTA, 1.0) : 0.0;
+		const float AlphaTangentB = ( NormDeltaVTB < InStaticFriction * NormDeltaPN ) ? 1.0 : (InKineticFriction > 0.0) ? min(InKineticFriction*NormDeltaPN/NormDeltaVTB, 1.0) : 0.0;
+
+		return DeltaPNormal * InCollisionNormal + AlphaTangentA * CollisionTangentA * DeltaVTangentA + AlphaTangentB * CollisionTangentB * DeltaVTangentB;
+	}
+	return float3(0,0,0);
+ }
+
+/* -----------------------------------------------------------------
+ * Static volume collisions against sphere, box...
+ * -----------------------------------------------------------------
+ */
+
+
+void DetectSphereVolumeCollision(const float3 InParticlePosition, const float3 InSphereCenter, const float InSphereRadius, 
+								 out float3 OutCollisionPosition, out float3 OutCollisionVelocity, inout float3 OutCollisionNormal)   
+{
+	OutCollisionNormal = InParticlePosition - InSphereCenter;
+	const float ParticleDistance = length(OutCollisionNormal);
+	const float CollisionDepth = ParticleDistance-InSphereRadius;
+	
+	OutCollisionNormal = (ParticleDistance != 0.0) ? OutCollisionNormal/ParticleDistance : float3(0,0,0);
+	OutCollisionPosition = InParticlePosition - OutCollisionNormal * CollisionDepth;
+
+	// Compute the sdf velocity based on the transform deriv
+	OutCollisionVelocity = float3(0,0,0);
+}  
+
+/* -----------------------------------------------------------------
+ * Static collision constraint
+ * -----------------------------------------------------------------
+ */
+
+void SolveStaticCollisionConstraint(const float InParticleRadius, const bool InIsMobile, inout float3 OutParticlePosition, inout float3 OutParticleVelocity, const float3 InCollisionPosition, const float3 InCollisionVelocity, const float3 InCollisionNormal, 
+				const float InStaticFriction, const float InKineticFriction, const float InDeltaTime )
+{
+	if(InIsMobile)
+	{
+		const float3 DeltaPosition = OutParticlePosition - InCollisionPosition;
+		const float3 DeltaVelocity = (OutParticleVelocity - InCollisionVelocity) * InDeltaTime;
+		const float ThicknessSum = InParticleRadius;
+
+		const float3 PositionDisplacement = ProcessCollisionData(DeltaPosition,DeltaVelocity,InCollisionNormal,ThicknessSum,InStaticFriction,InKineticFriction);
+		OutParticlePosition -= PositionDisplacement;
+		OutParticleVelocity -= PositionDisplacement / InDeltaTime;
+	}
+}
+
+#endif //GPU_SIMULATION
\ No newline at end of file
diff --git a/Engine/Plugins/FX/Niagara/Source/Niagara/Classes/NiagaraConstants.h b/Engine/Plugins/FX/Niagara/Source/Niagara/Classes/NiagaraConstants.h
--- a/Engine/Plugins/FX/Niagara/Source/Niagara/Classes/NiagaraConstants.h
+++ b/Engine/Plugins/FX/Niagara/Source/Niagara/Classes/NiagaraConstants.h
@@ -17,6 +17,7 @@
 #define PARAM_MAP_INITIAL_STR TEXT("Initial.")
 #define PARAM_MAP_INITIAL_BASE_STR TEXT("Initial")
 #define PARAM_MAP_RAPID_ITERATION_STR TEXT("Constants.")
+#define PARAM_MAP_INDICES_STR TEXT("Array.")
 #define PARAM_MAP_RAPID_ITERATION_BASE_STR TEXT("Constants")
 
 
diff --git a/Engine/Plugins/FX/Niagara/Source/Niagara/Classes/NiagaraDataInterface.h b/Engine/Plugins/FX/Niagara/Source/Niagara/Classes/NiagaraDataInterface.h
--- a/Engine/Plugins/FX/Niagara/Source/Niagara/Classes/NiagaraDataInterface.h
+++ b/Engine/Plugins/FX/Niagara/Source/Niagara/Classes/NiagaraDataInterface.h
@@ -157,12 +157,32 @@ private:
 
 struct FNiagaraDataInterfaceProxy : TSharedFromThis<FNiagaraDataInterfaceProxy, ESPMode::ThreadSafe>
 {
-	virtual ~FNiagaraDataInterfaceProxy() {check(IsInRenderingThread());}
+	virtual ~FNiagaraDataInterfaceProxy() {/*check(IsInRenderingThread());*/}
 
 	virtual int32 PerInstanceDataPassedToRenderThreadSize() const = 0;
 	virtual void ConsumePerInstanceDataFromGameThread(void* PerInstanceData, const FGuid& Instance) { check(false); }
 
 	virtual void DeferredDestroy() {}
+
+	// #todo(dmp): move all of this stuff to the RW interface to keep it out of here?
+
+	// a set of the shader stages that require the data interface for data output
+	TSet<int> OutputShaderStages;
+
+	// a set of the shader stages that require the data interface for setting number of output elements
+	TSet<int> IterationShaderStages;
+	
+	// number of elements to output to
+	uint32 ElementCount;
+
+	void SetElementCount(uint32 Count) { ElementCount = Count;  }
+	virtual bool IsOutputStage(uint32 CurrentStage) const { return OutputShaderStages.Contains(CurrentStage); }
+	virtual bool IsIterationStage(uint32 CurrentStage) const { return IterationShaderStages.Contains(CurrentStage); }
+
+	virtual void ResetData(FRHICommandList& RHICmdList, const FNiagaraDataInterfaceSetArgs& Context) { }
+
+	virtual void PreStage(FRHICommandList& RHICmdList, const FNiagaraDataInterfaceSetArgs& Context) {}	
+	virtual void PostStage(FRHICommandList& RHICmdList, const FNiagaraDataInterfaceSetArgs& Context) {}
 };
 
 struct FNiagaraDataInterfaceProxyCurveBase : public FNiagaraDataInterfaceProxy
diff --git a/Engine/Plugins/FX/Niagara/Source/Niagara/Classes/NiagaraDataInterfaceGrid2DCollection.h b/Engine/Plugins/FX/Niagara/Source/Niagara/Classes/NiagaraDataInterfaceGrid2DCollection.h
new file mode 100644
--- /dev/null
+++ b/Engine/Plugins/FX/Niagara/Source/Niagara/Classes/NiagaraDataInterfaceGrid2DCollection.h
@@ -0,0 +1,142 @@
+// Copyright 1998-2019 Epic Games, Inc. All Rights Reserved.
+#pragma once
+
+#include "NiagaraDataInterface.h"
+#include "NiagaraDataInterfaceRW.h"
+#include "ClearQuad.h"
+#include "NiagaraComponent.h"
+
+#include "NiagaraDataInterfaceGrid2DCollection.generated.h"
+
+class FNiagaraSystemInstance;
+class UTextureRenderTarget2D;
+
+class Grid2DBuffer
+{
+public:
+	Grid2DBuffer(int NumX, int NumY)
+	{
+		GridBuffer.Initialize(NumX, NumY);
+	}
+
+	FTextureRWBuffer GridBuffer;	
+};
+
+class Grid2DCollectionRWInstanceData
+{
+public:
+
+	int NumCellsX;
+	int NumCellsY;
+	float CellSize;
+	bool SetGridFromVoxelSize;
+
+	FVector WorldBBoxMin;
+	FVector2D WorldBBoxSize;
+
+	TArray<Grid2DBuffer*> Buffers;
+
+	Grid2DBuffer* CurrentData;
+	Grid2DBuffer* DestinationData;
+
+	Grid2DCollectionRWInstanceData() : CurrentData(nullptr), DestinationData(nullptr) {}
+
+	~Grid2DCollectionRWInstanceData()
+	{
+		for (Grid2DBuffer* Buffer : Buffers)
+		{
+			Buffer->GridBuffer.Release();
+			delete Buffer;
+		}
+	}
+
+	Grid2DBuffer* GetCurrentData() { return CurrentData; }
+	Grid2DBuffer* GetDestinationData() { return DestinationData; }
+	Grid2DBuffer &BeginSimulate()
+	{
+		for (Grid2DBuffer* Buffer : Buffers)
+		{
+			check(Buffer);
+			if (Buffer != CurrentData)
+			{
+				DestinationData = Buffer;
+				break;
+			}
+		}
+
+		if (DestinationData == nullptr)
+		{
+			DestinationData = new Grid2DBuffer(NumCellsX, NumCellsY);
+			Buffers.Add(DestinationData);
+		}
+
+		return *DestinationData;
+	}
+
+	void EndSimulate()
+	{
+		CurrentData = DestinationData;
+		DestinationData = nullptr;
+	}
+};
+
+struct FNiagaraDataInterfaceProxyGrid2DCollection : public FNiagaraDataInterfaceProxyRW
+{
+	FNiagaraDataInterfaceProxyGrid2DCollection() {}		
+
+	virtual void PreStage(FRHICommandList& RHICmdList, const FNiagaraDataInterfaceSetArgs& Context) override;
+	virtual void PostStage(FRHICommandList& RHICmdList, const FNiagaraDataInterfaceSetArgs& Context) override;
+	virtual void ResetData(FRHICommandList& RHICmdList, const FNiagaraDataInterfaceSetArgs& Context) override;
+
+	virtual void DeferredDestroy() override;		
+	void DestroyPerInstanceData(NiagaraEmitterInstanceBatcher* Batcher, const FGuid& SystemInstance);
+
+	/* List of proxy data for each system instances*/
+	// #todo(dmp): this should all be refactored to avoid duplicate code
+	TMap<FGuid, Grid2DCollectionRWInstanceData> SystemInstancesToProxyData;
+
+	/* List of proxy data to destroy later */
+	// #todo(dmp): this should all be refactored to avoid duplicate code
+	TSet<FGuid> DeferredDestroyList;
+};
+
+UCLASS(EditInlineNew, Category = "Grid", meta = (DisplayName = "Grid2D Collection"), Blueprintable, BlueprintType)
+class NIAGARA_API UNiagaraDataInterfaceGrid2DCollection : public UNiagaraDataInterfaceGrid2D
+{
+	GENERATED_UCLASS_BODY()
+
+public:
+	
+	virtual void PostInitProperties() override;
+	
+	//~ UNiagaraDataInterface interface
+	// VM functionality
+	virtual void GetFunctions(TArray<FNiagaraFunctionSignature>& OutFunctions) override;
+	virtual void GetVMExternalFunction(const FVMExternalFunctionBindingInfo& BindingInfo, void* InstanceData, FVMExternalFunction &OutFunc) override;
+
+	virtual bool Equals(const UNiagaraDataInterface* Other) const override;
+
+	// GPU sim functionality
+	virtual void GetParameterDefinitionHLSL(FNiagaraDataInterfaceGPUParamInfo& ParamInfo, FString& OutHLSL) override;
+	virtual bool GetFunctionHLSL(const FName&  DefinitionFunctionName, FString InstanceFunctionName, FNiagaraDataInterfaceGPUParamInfo& ParamInfo, FString& OutHLSL) override;
+	virtual FNiagaraDataInterfaceParametersCS* ConstructComputeParameters() const override;
+	
+	virtual void ProvidePerInstanceDataForRenderThread(void* DataForRenderThread, void* PerInstanceData, const FGuid& SystemInstance) override {}
+	virtual bool InitPerInstanceData(void* PerInstanceData, FNiagaraSystemInstance* SystemInstance) override;
+	virtual void DestroyPerInstanceData(void* PerInstanceData, FNiagaraSystemInstance* SystemInstance) override;
+	virtual bool PerInstanceTick(void* PerInstanceData, FNiagaraSystemInstance* SystemInstance, float DeltaSeconds) override { return false; }
+	virtual int32 PerInstanceDataSize()const override { return sizeof(Grid2DCollectionRWInstanceData); }
+	//~ UNiagaraDataInterface interface END
+
+	// Fills a texture render target 2d with the current data from the simulation
+	// #todo(dmp): this will eventually go away when we formalize how data makes it out of Niagara
+	UFUNCTION(BlueprintCallable, Category = Niagara)
+	virtual void FillTexture2D(const UNiagaraComponent *Component, UTextureRenderTarget2D *dest);
+
+protected:
+	//~ UNiagaraDataInterface interface
+	virtual bool CopyToInternal(UNiagaraDataInterface* Destination) const override;
+	//~ UNiagaraDataInterface interface END
+};
+
+
diff --git a/Engine/Plugins/FX/Niagara/Source/Niagara/Classes/NiagaraDataInterfaceNeighborGrid3D.h b/Engine/Plugins/FX/Niagara/Source/Niagara/Classes/NiagaraDataInterfaceNeighborGrid3D.h
new file mode 100644
--- /dev/null
+++ b/Engine/Plugins/FX/Niagara/Source/Niagara/Classes/NiagaraDataInterfaceNeighborGrid3D.h
@@ -0,0 +1,98 @@
+// Copyright 1998-2019 Epic Games, Inc. All Rights Reserved.
+#pragma once
+
+#include "NiagaraDataInterface.h"
+#include "NiagaraCommon.h"
+#include "NiagaraDataInterfaceRW.h"
+#include "ClearQuad.h"
+
+
+#include "NiagaraDataInterfaceNeighborGrid3D.generated.h"
+
+class FNiagaraSystemInstance;
+
+// store all data in in a class
+// move all data management to use per instance data
+// remove references to push data to render thread
+
+class NeighborGrid3DRWInstanceData
+{
+public:
+
+	FIntVector NumVoxels;
+	float VoxelSize;
+	bool SetGridFromVoxelSize;
+	uint32 MaxNeighborsPerVoxel;
+	FVector WorldBBoxMin;
+	FVector WorldBBoxSize;
+
+	FRWBuffer NeighborhoodBuffer;
+	FRWBuffer NeighborhoodCountBuffer;
+};
+
+struct FNiagaraDataInterfaceProxyNeighborGrid3D : public FNiagaraDataInterfaceProxyRW
+{	
+	virtual void PreStage(FRHICommandList& RHICmdList, const FNiagaraDataInterfaceSetArgs& Context) override;	
+	virtual void ConsumePerInstanceDataFromGameThread(void* PerInstanceData, const FGuid& Instance) override {}
+	virtual int32 PerInstanceDataPassedToRenderThreadSize() const override { return sizeof(NeighborGrid3DRWInstanceData); }	
+	void DestroyPerInstanceData(NiagaraEmitterInstanceBatcher* Batcher, const FGuid& SystemInstance);
+	void DeferredDestroy();
+
+	/* List of proxy data for each system instances*/
+	// #todo(dmp): this should all be refactored to avoid duplicate code
+	TMap<FGuid, NeighborGrid3DRWInstanceData> SystemInstancesToProxyData;
+
+	/* List of proxy data to destroy later */
+	// #todo(dmp): this should all be refactored to avoid duplicate code
+	TSet<FGuid> DeferredDestroyList;
+};
+
+UCLASS(EditInlineNew, Category = "Grid", meta = (DisplayName = "Neighbor Grid3D"))
+class NIAGARA_API UNiagaraDataInterfaceNeighborGrid3D : public UNiagaraDataInterfaceGrid3D
+{
+	GENERATED_UCLASS_BODY()
+
+public:	
+
+	UPROPERTY(EditAnywhere, Category = "Grid")
+		uint32 MaxNeighborsPerVoxel;
+
+public:
+
+	virtual void PostInitProperties() override
+	{
+		Super::PostInitProperties();
+
+		//Can we register data interfaces as regular types and fold them into the FNiagaraVariable framework for UI and function calls etc?
+		if (HasAnyFlags(RF_ClassDefaultObject))
+		{
+			FNiagaraTypeRegistry::Register(FNiagaraTypeDefinition(GetClass()), /*bCanBeParameter*/ true, /*bCanBePayload*/ false, /*bIsUserDefined*/ false);
+		}
+	}
+
+	//~ UNiagaraDataInterface interface
+	// VM functionality
+	virtual void GetFunctions(TArray<FNiagaraFunctionSignature>& OutFunctions) override;
+	virtual void GetVMExternalFunction(const FVMExternalFunctionBindingInfo& BindingInfo, void* InstanceData, FVMExternalFunction &OutFunc) override;
+
+	virtual bool Equals(const UNiagaraDataInterface* Other) const override;
+
+	// GPU sim functionality
+	virtual void GetParameterDefinitionHLSL(FNiagaraDataInterfaceGPUParamInfo& ParamInfo, FString& OutHLSL) override;
+	virtual bool GetFunctionHLSL(const FName&  DefinitionFunctionName, FString InstanceFunctionName, FNiagaraDataInterfaceGPUParamInfo& ParamInfo, FString& OutHLSL) override;
+	virtual FNiagaraDataInterfaceParametersCS* ConstructComputeParameters() const override;
+		
+	virtual void ProvidePerInstanceDataForRenderThread(void* DataForRenderThread, void* PerInstanceData, const FGuid& SystemInstance) override {}
+	virtual bool InitPerInstanceData(void* PerInstanceData, FNiagaraSystemInstance* SystemInstance) override;
+	virtual void DestroyPerInstanceData(void* PerInstanceData, FNiagaraSystemInstance* SystemInstance) override;
+	virtual bool PerInstanceTick(void* PerInstanceData, FNiagaraSystemInstance* SystemInstance, float DeltaSeconds) override { return false;  }
+	virtual int32 PerInstanceDataSize()const override { return sizeof(NeighborGrid3DRWInstanceData); }
+	//~ UNiagaraDataInterface interface END
+
+protected:
+	//~ UNiagaraDataInterface interface
+	virtual bool CopyToInternal(UNiagaraDataInterface* Destination) const override;
+	//~ UNiagaraDataInterface interface END	
+};
+
+
diff --git a/Engine/Plugins/FX/Niagara/Source/Niagara/Classes/NiagaraDataInterfaceRW.h b/Engine/Plugins/FX/Niagara/Source/Niagara/Classes/NiagaraDataInterfaceRW.h
new file mode 100644
--- /dev/null
+++ b/Engine/Plugins/FX/Niagara/Source/Niagara/Classes/NiagaraDataInterfaceRW.h
@@ -0,0 +1,215 @@
+// Copyright 1998-2019 Epic Games, Inc. All Rights Reserved.
+#pragma once
+
+#include "NiagaraDataInterface.h"
+#include "NiagaraDataInterfaceRW.generated.h"
+
+class FNiagaraSystemInstance;
+
+// Global HLSL variable base names, used by HLSL.
+static const FString NumVoxelsName(TEXT("NumVoxels_"));
+static const FString VoxelSizeName(TEXT("VoxelSize_"));
+static const FString WorldBBoxMinName(TEXT("WorldBBoxMin_"));
+static const FString WorldBBoxSizeName(TEXT("WorldBBoxSize_"));
+
+static const FString NumCellsName(TEXT("NumCells_"));
+static const FString CellSizeName(TEXT("CellSize_"));
+
+// Global VM function names, also used by the shaders code generation methods.
+static const FName NumVoxelsFunctionName("GetNumVoxels");
+static const FName VoxelSizeFunctionName("GetVoxelSize");
+
+static const FName NumCellsFunctionName("GetNumCells");
+static const FName CellSizeFunctionName("GetCellSize");
+
+static const FName WorldToUnitFunctionName("WorldToUnit");
+static const FName UnitToIndexFunctionName("UnitToIndex");
+static const FName IndexToUnitFunctionName("IndexToUnit");
+static const FName IndexToUnitStaggeredXFunctionName("IndexToUnitStaggeredX");
+static const FName IndexToUnitStaggeredYFunctionName("IndexToUnitStaggeredY");
+
+static const FName IndexToLinearFunctionName("IndexToLinear");
+static const FName LinearToIndexFunctionName("LinearToIndex");
+
+
+
+// #todo(dmp): some of the stuff we'd expect to see here is on FNiagaraDataInterfaceProxy - refactor?
+struct FNiagaraDataInterfaceProxyRW : public FNiagaraDataInterfaceProxy
+{
+public:
+
+	virtual void ConsumePerInstanceDataFromGameThread(void* PerInstanceData, const FGuid& Instance) override { check(false); }
+	virtual int32 PerInstanceDataPassedToRenderThreadSize() const override
+	{
+		return 0;
+	}	
+
+	virtual void ClearBuffers(FRHICommandList& RHICmdList) {}
+};
+
+
+UCLASS(abstract, EditInlineNew)
+class NIAGARA_API UNiagaraDataInterfaceRWBase : public UNiagaraDataInterface
+{
+	GENERATED_UCLASS_BODY()
+
+public:
+	UPROPERTY(EditAnywhere, Category = "RW")
+	TSet<int> OutputShaderStages;
+
+	UPROPERTY(EditAnywhere, Category = "RW")
+	TSet<int> IterationShaderStages;
+
+public:
+	//~ UObject interface
+
+	virtual void PostLoad() override
+	{
+		Super::PostLoad();
+
+		PushToRenderThread();
+	}
+
+#if WITH_EDITOR
+	virtual void PostEditChangeProperty(struct FPropertyChangedEvent& PropertyChangedEvent) override
+	{
+		Super::PostEditChangeProperty(PropertyChangedEvent);
+
+		PushToRenderThread();
+	}	
+	
+	virtual void PreEditChange(UProperty* PropertyAboutToChange) override
+	{
+		Super::PreEditChange(PropertyAboutToChange);
+
+		// Flush the rendering thread before making any changes to make sure the 
+		// data read by the compute shader isn't subject to a race condition.
+		// TODO(mv): Solve properly using something like a RT Proxy.
+		//FlushRenderingCommands();
+	}
+
+#endif
+//~ UObject interface END
+
+	virtual bool Equals(const UNiagaraDataInterface* Other) const override;
+
+	virtual bool CanExecuteOnTarget(ENiagaraSimTarget Target) const override
+	{
+		return true;
+	}
+
+#if WITH_EDITOR	
+	// Editor functionality
+	virtual TArray<FNiagaraDataInterfaceError> GetErrors() override
+	{
+		// TODO(mv): Improve error messages?
+		TArray<FNiagaraDataInterfaceError> Errors;
+
+		return Errors;
+	}
+#endif
+
+	void EmptyVMFunction(FVectorVMContext& Context) {}
+
+
+protected:
+	virtual bool CopyToInternal(UNiagaraDataInterface* Destination) const override;
+	virtual void PushToRenderThread() {};
+
+	FNiagaraDataInterfaceProxyRW* RWProxy;
+};
+
+
+UCLASS(abstract, EditInlineNew)
+class NIAGARA_API UNiagaraDataInterfaceGrid3D : public UNiagaraDataInterfaceRWBase
+{
+	GENERATED_UCLASS_BODY()
+
+public:
+	UPROPERTY(EditAnywhere, Category = "Grid", meta = (EditCondition = "!SetGridFromVoxelSize"))
+	FIntVector NumVoxels;
+
+	UPROPERTY(EditAnywhere, Category = "Grid", meta = (EditCondition = "SetGridFromVoxelSize"))
+	float VoxelSize;
+
+	UPROPERTY(EditAnywhere, Category = "Grid")
+	bool SetGridFromVoxelSize;	
+
+	UPROPERTY(EditAnywhere, Category = "Grid")
+	FVector WorldBBoxMin;
+
+	UPROPERTY(EditAnywhere, Category = "Grid")
+	FVector WorldBBoxSize;
+
+
+public:
+	//~ UNiagaraDataInterface interface
+	// VM functionality
+	virtual void GetFunctions(TArray<FNiagaraFunctionSignature>& OutFunctions) override;
+	virtual void GetVMExternalFunction(const FVMExternalFunctionBindingInfo& BindingInfo, void* InstanceData, FVMExternalFunction &OutFunc) override;
+
+	virtual bool Equals(const UNiagaraDataInterface* Other) const override;
+
+	// GPU sim functionality
+	virtual void GetParameterDefinitionHLSL(FNiagaraDataInterfaceGPUParamInfo& ParamInfo, FString& OutHLSL) override;
+	virtual bool GetFunctionHLSL(const FName&  DefinitionFunctionName, FString InstanceFunctionName, FNiagaraDataInterfaceGPUParamInfo& ParamInfo, FString& OutHLSL) override;
+	virtual FNiagaraDataInterfaceParametersCS* ConstructComputeParameters() const override;
+	//~ UNiagaraDataInterface interface END
+
+
+protected:
+	//~ UNiagaraDataInterface interface
+	virtual bool CopyToInternal(UNiagaraDataInterface* Destination) const override;
+	//~ UNiagaraDataInterface interface END
+
+	virtual void PushToRenderThread() override;
+};
+
+
+UCLASS(abstract, EditInlineNew)
+class NIAGARA_API UNiagaraDataInterfaceGrid2D : public UNiagaraDataInterfaceRWBase
+{
+	GENERATED_UCLASS_BODY()
+
+public:
+	UPROPERTY(EditAnywhere, Category = "Grid", meta = (EditCondition = "!SetGridFromCellSize"))
+	int32 NumCellsX;
+
+	UPROPERTY(EditAnywhere, Category = "Grid", meta = (EditCondition = "!SetGridFromCellSize"))
+	int32 NumCellsY;
+
+	UPROPERTY(EditAnywhere, Category = "Grid", meta = (EditCondition = "SetGridFromCellSize"))
+	float CellSize;
+
+	UPROPERTY(EditAnywhere, Category = "Grid")
+	bool SetGridFromCellSize;
+
+	UPROPERTY(EditAnywhere, Category = "Grid")
+	FVector WorldBBoxMin;
+
+	UPROPERTY(EditAnywhere, Category = "Grid")
+	FVector2D WorldBBoxSize;
+
+
+public:
+	//~ UNiagaraDataInterface interface
+	// VM functionality
+	virtual void GetFunctions(TArray<FNiagaraFunctionSignature>& OutFunctions) override;
+	virtual void GetVMExternalFunction(const FVMExternalFunctionBindingInfo& BindingInfo, void* InstanceData, FVMExternalFunction &OutFunc) override;
+
+	virtual bool Equals(const UNiagaraDataInterface* Other) const override;
+
+	// GPU sim functionality
+	virtual void GetParameterDefinitionHLSL(FNiagaraDataInterfaceGPUParamInfo& ParamInfo, FString& OutHLSL) override;
+	virtual bool GetFunctionHLSL(const FName&  DefinitionFunctionName, FString InstanceFunctionName, FNiagaraDataInterfaceGPUParamInfo& ParamInfo, FString& OutHLSL) override;
+	virtual FNiagaraDataInterfaceParametersCS* ConstructComputeParameters() const override;
+	//~ UNiagaraDataInterface interface END
+
+
+protected:
+	//~ UNiagaraDataInterface interface
+	virtual bool CopyToInternal(UNiagaraDataInterface* Destination) const override;
+	//~ UNiagaraDataInterface interface END
+
+	virtual void PushToRenderThread() override;
+};
\ No newline at end of file
diff --git a/Engine/Plugins/FX/Niagara/Source/Niagara/Classes/NiagaraDataSet.h b/Engine/Plugins/FX/Niagara/Source/Niagara/Classes/NiagaraDataSet.h
--- a/Engine/Plugins/FX/Niagara/Source/Niagara/Classes/NiagaraDataSet.h
+++ b/Engine/Plugins/FX/Niagara/Source/Niagara/Classes/NiagaraDataSet.h
@@ -222,7 +222,7 @@ public:
 	FNiagaraDataBuffer& BeginSimulate();
 
 	/** Ends a simulation pass and sets the current simulation state. */
-	void EndSimulate();
+	void EndSimulate(bool SetCurrentData = true);
 
 	/** Allocates space for NumInstances in the current destination buffer. */
 	void Allocate(int32 NumInstances, bool bMaintainExisting = false);
diff --git a/Engine/Plugins/FX/Niagara/Source/Niagara/Classes/NiagaraEmitterInstanceBatcher.h b/Engine/Plugins/FX/Niagara/Source/Niagara/Classes/NiagaraEmitterInstanceBatcher.h
--- a/Engine/Plugins/FX/Niagara/Source/Niagara/Classes/NiagaraEmitterInstanceBatcher.h
+++ b/Engine/Plugins/FX/Niagara/Source/Niagara/Classes/NiagaraEmitterInstanceBatcher.h
@@ -113,7 +113,7 @@ public:
 
 	void ProcessDebugInfo(FRHICommandList &RHICmdLis, const FNiagaraComputeExecutionContext* Context) const;
 
-	void SetDataInterfaceParameters(const TArray<FNiagaraDataInterfaceProxy*>& DataInterfaceProxies, FNiagaraShader* Shader, FRHICommandList &RHICmdList, const FNiagaraComputeInstanceData* Instance, const FNiagaraGPUSystemTick& Tick) const;
+	void SetDataInterfaceParameters(const TArray<FNiagaraDataInterfaceProxy*>& DataInterfaceProxies, FNiagaraShader* Shader, FRHICommandList &RHICmdList, const FNiagaraComputeInstanceData* Instance, const FNiagaraGPUSystemTick& Tick, uint32 ShaderStageIndex) const;
 	void UnsetDataInterfaceParameters(const TArray<FNiagaraDataInterfaceProxy*>& DataInterfaceProxies, FNiagaraShader* Shader, FRHICommandList& RHICmdList, const FNiagaraComputeInstanceData* Instance, const FNiagaraGPUSystemTick& Tick) const;
 
 	template<bool bDoResourceTransitions>
@@ -125,7 +125,9 @@ public:
 				FRHICommandList &RHICmdList, 
 				FRHIUniformBuffer* ViewUniformBuffer, 
 				const FNiagaraGpuSpawnInfo& SpawnInfo,
-				bool bCopyBeforeStart = false
+				bool bCopyBeforeStart = false,
+				uint32 ShaderStageIndex = 0,
+				FNiagaraDataInterfaceProxy *IterationInterface = nullptr
 			) const;
 
 	void ResizeCurrentBuffer(FRHICommandList &RHICmdList, FNiagaraComputeExecutionContext *Context, uint32 NewNumInstances, uint32 PrevNumInstances) const;
@@ -135,6 +137,22 @@ public:
 	FORCEINLINE EShaderPlatform GetShaderPlatform() const { return ShaderPlatform; }
 	FORCEINLINE ERHIFeatureLevel::Type GetFeatureLevel() const { return FeatureLevel; }
 
+	/** Reset the data interfaces and check if the spawn stages are valid */
+	bool ResetDataInterfaces(const FNiagaraGPUSystemTick& Tick, FNiagaraComputeInstanceData *Instance, FRHICommandList &RHICmdList, FNiagaraShader* ComputeShader ) const;
+
+	/** Given a shader stage index, find the corresponding data interface */
+	FNiagaraDataInterfaceProxy* FindIterationInterface(FNiagaraComputeInstanceData *Instance, const uint32 ShaderStageIndex) const;
+
+	/** Loop over all the data interfaces and call the prestage methods */
+	void PreStageInterface(const FNiagaraGPUSystemTick& Tick, FNiagaraComputeInstanceData *Instance, FRHICommandList &RHICmdList, FNiagaraShader* ComputeShader, const uint32 ShaderStageIndex) const;
+
+	/** Loop over all the data interfaces and call the poststage methods */
+	void PostStageInterface(const FNiagaraGPUSystemTick& Tick, FNiagaraComputeInstanceData *Instance, FRHICommandList &RHICmdList, FNiagaraShader* ComputeShader, const uint32 ShaderStageIndex) const;
+
+	/** Run the dispatch over multiple stages */
+	template<bool bDoResourceTransitions>
+	void DispatchMultipleStages(const FNiagaraGPUSystemTick& Tick, FNiagaraComputeInstanceData *Instance, FRHICommandList &RHICmdList, FRHIUniformBuffer* ViewUniformBuffer, FNiagaraShader* ComputeShader) const;
+
 private:
 
 	void ExecuteAll(FRHICommandList &RHICmdList, FRHIUniformBuffer* ViewUniformBuffer, bool bSetReadback);
diff --git a/Engine/Plugins/FX/Niagara/Source/Niagara/Classes/NiagaraScriptExecutionContext.h b/Engine/Plugins/FX/Niagara/Source/Niagara/Classes/NiagaraScriptExecutionContext.h
--- a/Engine/Plugins/FX/Niagara/Source/Niagara/Classes/NiagaraScriptExecutionContext.h
+++ b/Engine/Plugins/FX/Niagara/Source/Niagara/Classes/NiagaraScriptExecutionContext.h
@@ -162,7 +162,7 @@ struct FNiagaraComputeExecutionContext
 
 	void Reset(NiagaraEmitterInstanceBatcher* Batcher);
 
-	void InitParams(UNiagaraScript* InGPUComputeScript, ENiagaraSimTarget InSimTarget, const FString& InDebugSimName);
+	void InitParams(UNiagaraScript* InGPUComputeScript, ENiagaraSimTarget InSimTarget, const FString& InDebugSimName, const int32 InMaxUpdateIterations, const TSet<uint32> InSpawnStages);
 	void DirtyDataInterfaces();
 	bool Tick(FNiagaraSystemInstance* ParentSystemInstance);
 
@@ -207,6 +207,9 @@ public:
 	// Game thread spawn info will be sent to the render thread inside FNiagaraComputeInstanceData
 	FNiagaraGpuSpawnInfo GpuSpawnInfo_GT;
 
+	uint32 MaxUpdateIterations;
+	TSet<uint32> SpawnStages;
+
 #if WITH_EDITORONLY_DATA
 	mutable FRHIGPUMemoryReadback *GPUDebugDataReadbackFloat;
 	mutable FRHIGPUMemoryReadback *GPUDebugDataReadbackInt;
diff --git a/Engine/Plugins/FX/Niagara/Source/Niagara/Niagara.Build.cs b/Engine/Plugins/FX/Niagara/Source/Niagara/Niagara.Build.cs
--- a/Engine/Plugins/FX/Niagara/Source/Niagara/Niagara.Build.cs
+++ b/Engine/Plugins/FX/Niagara/Source/Niagara/Niagara.Build.cs
@@ -32,7 +32,7 @@ public class Niagara : ModuleRules
                 "RHI",
                 "UtilityShaders",
                 "NiagaraVertexFactories",
-                "RenderCore"
+                "RenderCore",
             }
         );
 
