CIS Fixes

diff --git a/Engine/Plugins/Experimental/HairStrands/Source/HairStrandsEditor/Private/HairStrandsFactory.cpp b/Engine/Plugins/Experimental/HairStrands/Source/HairStrandsEditor/Private/HairStrandsFactory.cpp
--- a/Engine/Plugins/Experimental/HairStrands/Source/HairStrandsEditor/Private/HairStrandsFactory.cpp
+++ b/Engine/Plugins/Experimental/HairStrands/Source/HairStrandsEditor/Private/HairStrandsFactory.cpp
@@ -57,28 +57,27 @@ UObject* UHairStrandsFactory::FactoryCreateNew(UClass* InClass, UObject* InParen
 		ExistingAsset->ReleaseResource();
 	}
 	UHairStrandsAsset* CurrentAsset = NewObject<UHairStrandsAsset>(InParent, InClass, InName, Flags);
-	check(CurrentAsset);
-
-	const FString Filter(TEXT("Hair Strands Files (*.hair,*.fbx,*.abc)|*.hair;*.fbx;*.abc"));
-
-	TArray<FString> OpenFilenames;
-	int32 FilterIndex = -1;
-	if (FDesktopPlatformModule::Get()->OpenFileDialog(
-		nullptr,
-		FString(TEXT("Choose a hair strands file")),
-		FEditorDirectories::Get().GetLastDirectory(ELastDirectory::GENERIC_IMPORT),
-		TEXT(""),
-		Filter,
-		EFileDialogFlags::None,
-		OpenFilenames,
-		FilterIndex))
-	{
-		CurrentAsset->FilePath = OpenFilenames[0];
-		BuildHairStrands(CurrentAsset->FilePath, CurrentAsset->StrandsDatas);
-	}
 
 	if (CurrentAsset)
 	{
+		const FString Filter(TEXT("Hair Strands Files (*.hair,*.fbx,*.abc)|*.hair;*.fbx;*.abc"));
+
+		TArray<FString> OpenFilenames;
+		int32 FilterIndex = -1;
+		if (FDesktopPlatformModule::Get()->OpenFileDialog(
+			nullptr,
+			FString(TEXT("Choose a hair strands file")),
+			FEditorDirectories::Get().GetLastDirectory(ELastDirectory::GENERIC_IMPORT),
+			TEXT(""),
+			Filter,
+			EFileDialogFlags::None,
+			OpenFilenames,
+			FilterIndex))
+		{			
+			CurrentAsset->FilePath = OpenFilenames[0];
+			BuildHairStrands(CurrentAsset->FilePath, CurrentAsset->StrandsDatas);
+		}
+
 		CurrentAsset->InitResource();
 	}
 	return CurrentAsset;
@@ -93,14 +92,12 @@ UObject* UHairStrandsFactory::FactoryCreateFile(UClass * InClass, UObject * InPa
 		ExistingAsset->ReleaseResource();
 	}
 
-	UHairStrandsAsset* CurrentAsset = NewObject<UHairStrandsAsset>(InParent, InClass, InName, Flags);
-	check(CurrentAsset);
-
-	CurrentAsset->FilePath = Filename;
-	BuildHairStrands(CurrentAsset->FilePath, CurrentAsset->StrandsDatas);
+	UHairStrandsAsset* CurrentAsset = NewObject<UHairStrandsAsset>(InParent, InClass, InName, Flags);	
 
 	if (CurrentAsset)
 	{
+		CurrentAsset->FilePath = Filename;
+		BuildHairStrands(CurrentAsset->FilePath, CurrentAsset->StrandsDatas);
 		CurrentAsset->InitResource();
 	}
 
