Add 0 entries to BP exposes enums

#rnx

diff --git a/Engine/Plugins/Experimental/HairStrands/Source/HairStrandsNiagara/Public/NiagaraDataInterfaceHairStrands.h b/Engine/Plugins/Experimental/HairStrands/Source/HairStrandsNiagara/Public/NiagaraDataInterfaceHairStrands.h
--- a/Engine/Plugins/Experimental/HairStrands/Source/HairStrandsNiagara/Public/NiagaraDataInterfaceHairStrands.h
+++ b/Engine/Plugins/Experimental/HairStrands/Source/HairStrandsNiagara/Public/NiagaraDataInterfaceHairStrands.h
@@ -12,6 +12,7 @@
 UENUM(BlueprintType)
 enum class EHairStrandsSize : uint8
 {
+	None = 0 UMETA(Hidden),
 	Size2 = 0x02 UMETA(DisplatName = "2"),
 	Size4 = 0x04 UMETA(DisplatName = "4"),
 	Size8 = 0x08 UMETA(DisplatName = "8"),
