Copying //UE4/Dev-Physics to Dev-Main (//UE4/Dev-Main) [at] 9098278

diff --git a/Engine/Plugins/Experimental/GeometryCollectionPlugin/Source/GeometryCollectionExampleCore/Private/GeometryCollection/GeometryCollectionExampleVisibility.cpp b/Engine/Plugins/Experimental/GeometryCollectionPlugin/Source/GeometryCollectionExampleCore/Private/GeometryCollection/GeometryCollectionExampleVisibility.cpp
--- a/Engine/Plugins/Experimental/GeometryCollectionPlugin/Source/GeometryCollectionExampleCore/Private/GeometryCollection/GeometryCollectionExampleVisibility.cpp
+++ b/Engine/Plugins/Experimental/GeometryCollectionPlugin/Source/GeometryCollectionExampleCore/Private/GeometryCollection/GeometryCollectionExampleVisibility.cpp
@@ -14,7 +14,7 @@
 namespace GeometryCollectionExample
 {
 	template<class T>
-	bool TestHideVertices(ExampleResponse&& R)
+	void TestHideVertices()
 	{
 		TSharedPtr<FGeometryCollection> Collection = GeometryCollection::MakeCubeElement(FTransform(FQuat::MakeFromEuler(FVector(0.f, 0.f, 0.f)), FVector(0.f, 0.f, 0.f)), FVector(1.0));
 		Collection->AppendGeometry(*GeometryCollection::MakeCubeElement(FTransform(FQuat::MakeFromEuler(FVector(0.f, 0.f, 0.f)), FVector(1.f, 0.f, 0.f)), FVector(1.0)));
@@ -40,7 +40,7 @@ namespace GeometryCollectionExample
 		for (int32 Idx = 0; Idx < NumFaces; Idx++)
 		{
 			check(VisibleArray[Idx]);
-			R.ExpectTrue(VisibleArray[Idx]);
+			EXPECT_TRUE(VisibleArray[Idx]);
 		}
 
 		// hide node 1
@@ -54,13 +54,13 @@ namespace GeometryCollectionExample
 			switch (BoneMapArray[IndicesArray[Idx][0]])
 			{
 			case 0:
-				R.ExpectTrue(VisibleArray[Idx]);
+				EXPECT_TRUE(VisibleArray[Idx]);
 			break;
 			case 1:
-				R.ExpectTrue(!VisibleArray[Idx]);
+				EXPECT_FALSE(VisibleArray[Idx]);
 			break;
 			case 2:
-				R.ExpectTrue(VisibleArray[Idx]);
+				EXPECT_TRUE(VisibleArray[Idx]);
 			break;
 			}			
 		}
@@ -80,20 +80,18 @@ namespace GeometryCollectionExample
 			switch (BoneMapArray[IndicesArray[Idx][0]])
 			{
 			case 0:
-				R.ExpectTrue(!VisibleArray[Idx]);
+				EXPECT_FALSE(VisibleArray[Idx]);
 				break;
 			case 1:
-				R.ExpectTrue(VisibleArray[Idx]);
+				EXPECT_TRUE(VisibleArray[Idx]);
 				break;
 			case 2:
-				R.ExpectTrue(!VisibleArray[Idx]);
+				EXPECT_FALSE(VisibleArray[Idx]);
 				break;
 			}
 		}
-
-		return !R.HasError();
 	}
-	template bool TestHideVertices<float>(ExampleResponse&& R);
+	template void TestHideVertices<float>();
 
 	
 
diff --git a/Engine/Plugins/Experimental/GeometryCollectionPlugin/Source/GeometryCollectionExampleCore/Public/GeometryCollection/GeometryCollectionExample.h b/Engine/Plugins/Experimental/GeometryCollectionPlugin/Source/GeometryCollectionExampleCore/Public/GeometryCollection/GeometryCollectionExample.h
--- a/Engine/Plugins/Experimental/GeometryCollectionPlugin/Source/GeometryCollectionExampleCore/Public/GeometryCollection/GeometryCollectionExample.h
+++ b/Engine/Plugins/Experimental/GeometryCollectionPlugin/Source/GeometryCollectionExampleCore/Public/GeometryCollection/GeometryCollectionExample.h
@@ -1,6 +1,15 @@
 // Copyright 1998-2019 Epic Games, Inc. All Rights Reserved.
 #pragma once
 
+#pragma warning(disable:4121)
+#pragma warning(disable:4668)
+#pragma warning(disable:4996)
+#include "gtest/gtest.h"
+#include "gmock/gmock.h"
+#pragma warning(default:4121)
+#pragma warning(default:4668)
+#pragma warning(default:4996)
+
 #include "GeometryCollection/GeometryCollectionExampleResponse.h"
 
 template<class RESPONSE> void ExecuteExamples();
diff --git a/Engine/Plugins/Experimental/GeometryCollectionPlugin/Source/GeometryCollectionExampleCore/Public/GeometryCollection/GeometryCollectionExampleBoneHierarchy.h b/Engine/Plugins/Experimental/GeometryCollectionPlugin/Source/GeometryCollectionExampleCore/Public/GeometryCollection/GeometryCollectionExampleBoneHierarchy.h
--- a/Engine/Plugins/Experimental/GeometryCollectionPlugin/Source/GeometryCollectionExampleCore/Public/GeometryCollection/GeometryCollectionExampleBoneHierarchy.h
+++ b/Engine/Plugins/Experimental/GeometryCollectionPlugin/Source/GeometryCollectionExampleCore/Public/GeometryCollection/GeometryCollectionExampleBoneHierarchy.h
@@ -7,6 +7,6 @@ namespace GeometryCollectionExample
 {	
 
 	template<class T>
-	bool TestImplicitBoneHierarchy(ExampleResponse&& R);
+	void TestImplicitBoneHierarchy();
 	
 }
diff --git a/Engine/Plugins/Experimental/GeometryCollectionPlugin/Source/GeometryCollectionExampleCore/Public/GeometryCollection/GeometryCollectionExampleClean.h b/Engine/Plugins/Experimental/GeometryCollectionPlugin/Source/GeometryCollectionExampleCore/Public/GeometryCollection/GeometryCollectionExampleClean.h
--- a/Engine/Plugins/Experimental/GeometryCollectionPlugin/Source/GeometryCollectionExampleCore/Public/GeometryCollection/GeometryCollectionExampleClean.h
+++ b/Engine/Plugins/Experimental/GeometryCollectionPlugin/Source/GeometryCollectionExampleCore/Public/GeometryCollection/GeometryCollectionExampleClean.h
@@ -6,17 +6,17 @@
 namespace GeometryCollectionExample
 {	
 	template<class T>
-	bool GEOMETRYCOLLECTIONEXAMPLECORE_API TestDeleteCoincidentVertices(ExampleResponse&& R);
+	void GEOMETRYCOLLECTIONEXAMPLECORE_API TestDeleteCoincidentVertices();
 
 	template<class T>
-	bool GEOMETRYCOLLECTIONEXAMPLECORE_API TestDeleteCoincidentVertices2(ExampleResponse&& R);
+	void GEOMETRYCOLLECTIONEXAMPLECORE_API TestDeleteCoincidentVertices2();
 
 	template<class T>
-	bool GEOMETRYCOLLECTIONEXAMPLECORE_API TestDeleteZeroAreaFaces(ExampleResponse&& R);
+	void GEOMETRYCOLLECTIONEXAMPLECORE_API TestDeleteZeroAreaFaces();
 
 	template<class T>
-	bool GEOMETRYCOLLECTIONEXAMPLECORE_API TestDeleteHiddenFaces(ExampleResponse&& R);
+	void GEOMETRYCOLLECTIONEXAMPLECORE_API TestDeleteHiddenFaces();
 
 	template<class T>
-	bool GEOMETRYCOLLECTIONEXAMPLECORE_API TestFillHoles(ExampleResponse&& R);
+	void GEOMETRYCOLLECTIONEXAMPLECORE_API TestFillHoles();
 }
diff --git a/Engine/Plugins/Experimental/GeometryCollectionPlugin/Source/GeometryCollectionExampleCore/Public/GeometryCollection/GeometryCollectionExampleClustering.h b/Engine/Plugins/Experimental/GeometryCollectionPlugin/Source/GeometryCollectionExampleCore/Public/GeometryCollection/GeometryCollectionExampleClustering.h
--- a/Engine/Plugins/Experimental/GeometryCollectionPlugin/Source/GeometryCollectionExampleCore/Public/GeometryCollection/GeometryCollectionExampleClustering.h
+++ b/Engine/Plugins/Experimental/GeometryCollectionPlugin/Source/GeometryCollectionExampleCore/Public/GeometryCollection/GeometryCollectionExampleClustering.h
@@ -2,58 +2,52 @@
 #pragma once
 
 #include "GeometryCollection/GeometryCollectionExample.h"
-#include "GeometryCollection/GeometryCollectionExampleResponse.h"
 #include "Chaos/Declares.h"
 
 namespace GeometryCollectionExample
 {
-
-#if TODO_REIMPLEMENT_RIGID_CLUSTERING
 	template<class T>
-	bool RigidBodies_ClusterTest_SingleLevelNonBreaking(ExampleResponse&& R);
+	void RigidBodies_ClusterTest_SingleLevelNonBreaking();
 
 	template<class T>
-	bool RigidBodies_ClusterTest_DeactivateClusterParticle(ExampleResponse&& R);
+	void RigidBodies_ClusterTest_DeactivateClusterParticle();
 
 	template<class T>
-	bool RigidBodies_ClusterTest_SingleLevelBreaking(ExampleResponse&& R);
+	void RigidBodies_ClusterTest_SingleLevelBreaking();
 
 	template<class T>
-	bool RigidBodies_ClusterTest_NestedCluster(ExampleResponse&& R);
+	void RigidBodies_ClusterTest_NestedCluster();
 
 	template<class T>
-	bool RigidBodies_ClusterTest_NestedCluster_MultiStrain(ExampleResponse&& R);
+	void RigidBodies_ClusterTest_NestedCluster_MultiStrain();
 
 	template<class T>
-	bool RigidBodies_ClusterTest_NestedCluster_Halt(ExampleResponse&& R);
+	void RigidBodies_ClusterTest_NestedCluster_Halt();
 
 	template<class T>
-	bool RigidBodies_ClusterTest_KinematicAnchor(ExampleResponse&& R);
+	void RigidBodies_ClusterTest_KinematicAnchor();
 
 	template<class T>
-	bool RigidBodies_ClusterTest_StaticAnchor(ExampleResponse&& R);
+	void RigidBodies_ClusterTest_StaticAnchor();
 
 	template<class T>
-	bool RigidBodies_ClusterTest_UnionClusters(ExampleResponse&& R);
+	void RigidBodies_ClusterTest_UnionClusters();
 
 	template<class T>
-	bool RigidBodies_ClusterTest_ReleaseClusterParticle_ClusteredNode(ExampleResponse&& R);
+	void RigidBodies_ClusterTest_ReleaseClusterParticle_ClusteredNode();
 
 	template<class T>
-	bool RigidBodies_ClusterTest_ReleaseClusterParticle_ClusteredKinematicNode(ExampleResponse&& R);
+	void RigidBodies_ClusterTest_ReleaseClusterParticle_ClusteredKinematicNode();
 
 	template<class T>
-	bool RigidBodies_ClusterTest_ReleaseClusterParticles_AllLeafNodes(ExampleResponse&& R);
+	void RigidBodies_ClusterTest_ReleaseClusterParticles_AllLeafNodes();
 
 	template<class T>
-	bool RigidBodies_ClusterTest_ReleaseClusterParticles_ClusterNodeAndSubClusterNode(ExampleResponse&& R);
+	void RigidBodies_ClusterTest_ReleaseClusterParticles_ClusterNodeAndSubClusterNode();
 
 	template<class T>
-	bool RigidBodies_ClusterTest_RemoveOnFracture(ExampleResponse&& R);
+	void RigidBodies_ClusterTest_RemoveOnFracture();
 
 	template<class T>
-	bool RigidBodiess_ClusterTest_ParticleImplicitCollisionGeometry(ExampleResponse&& R);
-
-#endif
-
+	void RigidBodiess_ClusterTest_ParticleImplicitCollisionGeometry();
 }
diff --git a/Engine/Plugins/Experimental/GeometryCollectionPlugin/Source/GeometryCollectionExampleCore/Public/GeometryCollection/GeometryCollectionExampleCreation.h b/Engine/Plugins/Experimental/GeometryCollectionPlugin/Source/GeometryCollectionExampleCore/Public/GeometryCollection/GeometryCollectionExampleCreation.h
--- a/Engine/Plugins/Experimental/GeometryCollectionPlugin/Source/GeometryCollectionExampleCore/Public/GeometryCollection/GeometryCollectionExampleCreation.h
+++ b/Engine/Plugins/Experimental/GeometryCollectionPlugin/Source/GeometryCollectionExampleCore/Public/GeometryCollection/GeometryCollectionExampleCreation.h
@@ -6,45 +6,45 @@
 namespace GeometryCollectionExample
 {	
 	template<class T>
-	bool GEOMETRYCOLLECTIONEXAMPLECORE_API CheckIncrementMask(ExampleResponse&& R);
+	void GEOMETRYCOLLECTIONEXAMPLECORE_API CheckIncrementMask();
 
 	template<class T>
-	bool GEOMETRYCOLLECTIONEXAMPLECORE_API Creation(ExampleResponse&& R);
+	void GEOMETRYCOLLECTIONEXAMPLECORE_API Creation();
 
 	template<class T>
-	bool GEOMETRYCOLLECTIONEXAMPLECORE_API AppendTransformHierarchy(ExampleResponse&& R);
+	void GEOMETRYCOLLECTIONEXAMPLECORE_API AppendTransformHierarchy();
 
 	template<class T>
-	bool GEOMETRYCOLLECTIONEXAMPLECORE_API ContiguousElementsTest(ExampleResponse&& R);
+	void GEOMETRYCOLLECTIONEXAMPLECORE_API ContiguousElementsTest();
 
 	template<class T>
-	bool GEOMETRYCOLLECTIONEXAMPLECORE_API DeleteFromEnd(ExampleResponse&& R);
+	void GEOMETRYCOLLECTIONEXAMPLECORE_API DeleteFromEnd();
 
 	template<class T>
-	bool GEOMETRYCOLLECTIONEXAMPLECORE_API DeleteFromStart(ExampleResponse&& R);
+	void GEOMETRYCOLLECTIONEXAMPLECORE_API DeleteFromStart();
 
 	template<class T>
-	bool GEOMETRYCOLLECTIONEXAMPLECORE_API DeleteFromMiddle(ExampleResponse&& R);
+	void GEOMETRYCOLLECTIONEXAMPLECORE_API DeleteFromMiddle();
 
 	template<class T>
-	bool GEOMETRYCOLLECTIONEXAMPLECORE_API DeleteBranch(ExampleResponse&& R);
+	void GEOMETRYCOLLECTIONEXAMPLECORE_API DeleteBranch();
 
 	template<class T>
-	bool GEOMETRYCOLLECTIONEXAMPLECORE_API DeleteRootLeafMiddle(ExampleResponse&& R);
+	void GEOMETRYCOLLECTIONEXAMPLECORE_API DeleteRootLeafMiddle();
 
 	template<class T>
-	bool GEOMETRYCOLLECTIONEXAMPLECORE_API DeleteEverything(ExampleResponse&& R);
+	void GEOMETRYCOLLECTIONEXAMPLECORE_API DeleteEverything();
 
 	template<class T>
-	bool GEOMETRYCOLLECTIONEXAMPLECORE_API ParentTransformTest(ExampleResponse&& R);
+	void GEOMETRYCOLLECTIONEXAMPLECORE_API ParentTransformTest();
 
 	template<class T>
-	bool GEOMETRYCOLLECTIONEXAMPLECORE_API ReindexMaterialsTest(ExampleResponse&& R);
+	void GEOMETRYCOLLECTIONEXAMPLECORE_API ReindexMaterialsTest();
 
 	template<class T>
-	bool GEOMETRYCOLLECTIONEXAMPLECORE_API AttributeTransferTest(ExampleResponse&& R);
+	void GEOMETRYCOLLECTIONEXAMPLECORE_API AttributeTransferTest();
 
 	template<class T>
-	bool GEOMETRYCOLLECTIONEXAMPLECORE_API AttributeDependencyTest(ExampleResponse&& R);
+	void GEOMETRYCOLLECTIONEXAMPLECORE_API AttributeDependencyTest();
 
 }
diff --git a/Engine/Plugins/Experimental/GeometryCollectionPlugin/Source/GeometryCollectionExampleCore/Public/GeometryCollection/GeometryCollectionExampleDecimation.h b/Engine/Plugins/Experimental/GeometryCollectionPlugin/Source/GeometryCollectionExampleCore/Public/GeometryCollection/GeometryCollectionExampleDecimation.h
--- a/Engine/Plugins/Experimental/GeometryCollectionPlugin/Source/GeometryCollectionExampleCore/Public/GeometryCollection/GeometryCollectionExampleDecimation.h
+++ b/Engine/Plugins/Experimental/GeometryCollectionPlugin/Source/GeometryCollectionExampleCore/Public/GeometryCollection/GeometryCollectionExampleDecimation.h
@@ -7,6 +7,6 @@ namespace GeometryCollectionExample
 {	
 
 	template<class T>
-	bool TestGeometryDecimation(ExampleResponse&& R);
+	void TestGeometryDecimation();
 	
 }
diff --git a/Engine/Plugins/Experimental/GeometryCollectionPlugin/Source/GeometryCollectionExampleCore/Public/GeometryCollection/GeometryCollectionExampleEvents.h b/Engine/Plugins/Experimental/GeometryCollectionPlugin/Source/GeometryCollectionExampleCore/Public/GeometryCollection/GeometryCollectionExampleEvents.h
--- a/Engine/Plugins/Experimental/GeometryCollectionPlugin/Source/GeometryCollectionExampleCore/Public/GeometryCollection/GeometryCollectionExampleEvents.h
+++ b/Engine/Plugins/Experimental/GeometryCollectionPlugin/Source/GeometryCollectionExampleCore/Public/GeometryCollection/GeometryCollectionExampleEvents.h
@@ -7,6 +7,6 @@
 namespace GeometryCollectionExample
 {
 	template<class T>
-	bool Event_Handler(ExampleResponse&& R);
+	void Event_Handler();
 
 }
diff --git a/Engine/Plugins/Experimental/GeometryCollectionPlugin/Source/GeometryCollectionExampleCore/Public/GeometryCollection/GeometryCollectionExampleFields.h b/Engine/Plugins/Experimental/GeometryCollectionPlugin/Source/GeometryCollectionExampleCore/Public/GeometryCollection/GeometryCollectionExampleFields.h
--- a/Engine/Plugins/Experimental/GeometryCollectionPlugin/Source/GeometryCollectionExampleCore/Public/GeometryCollection/GeometryCollectionExampleFields.h
+++ b/Engine/Plugins/Experimental/GeometryCollectionPlugin/Source/GeometryCollectionExampleCore/Public/GeometryCollection/GeometryCollectionExampleFields.h
@@ -7,55 +7,55 @@ namespace GeometryCollectionExample
 {
 	
 	template<class T>
-	bool GEOMETRYCOLLECTIONEXAMPLECORE_API Fields_NoiseSample(ExampleResponse&& R);
+	void GEOMETRYCOLLECTIONEXAMPLECORE_API Fields_NoiseSample();
 
 	template<class T>
-	bool GEOMETRYCOLLECTIONEXAMPLECORE_API Fields_RadialIntMask(ExampleResponse&& R);
+	void GEOMETRYCOLLECTIONEXAMPLECORE_API Fields_RadialIntMask();
 
 	template<class T>
-	bool GEOMETRYCOLLECTIONEXAMPLECORE_API Fields_RadialFalloff(ExampleResponse&& R);
+	void GEOMETRYCOLLECTIONEXAMPLECORE_API Fields_RadialFalloff();
 
 	template<class T>
-	bool GEOMETRYCOLLECTIONEXAMPLECORE_API Fields_PlaneFalloff(ExampleResponse&& R);
+	void GEOMETRYCOLLECTIONEXAMPLECORE_API Fields_PlaneFalloff();
 
 	template<class T>
-	bool GEOMETRYCOLLECTIONEXAMPLECORE_API Fields_UniformVector(ExampleResponse&& R);
+	void GEOMETRYCOLLECTIONEXAMPLECORE_API Fields_UniformVector();
 
 	template<class T>
-	bool GEOMETRYCOLLECTIONEXAMPLECORE_API Fields_RaidalVector(ExampleResponse&& R);
+	void GEOMETRYCOLLECTIONEXAMPLECORE_API Fields_RaidalVector();
 
 	template<class T>
-	bool GEOMETRYCOLLECTIONEXAMPLECORE_API Fields_SumVectorFullMult(ExampleResponse&& R);
+	void GEOMETRYCOLLECTIONEXAMPLECORE_API Fields_SumVectorFullMult();
 
 	template<class T>
-	bool GEOMETRYCOLLECTIONEXAMPLECORE_API Fields_SumVectorFullDiv(ExampleResponse&& R);
+	void GEOMETRYCOLLECTIONEXAMPLECORE_API Fields_SumVectorFullDiv();
 
 	template<class T>
-	bool GEOMETRYCOLLECTIONEXAMPLECORE_API Fields_SumVectorFullAdd(ExampleResponse&& R);
+	void GEOMETRYCOLLECTIONEXAMPLECORE_API Fields_SumVectorFullAdd();
 
 	template<class T>
-	bool GEOMETRYCOLLECTIONEXAMPLECORE_API Fields_SumVectorFullSub(ExampleResponse&& R);
+	void GEOMETRYCOLLECTIONEXAMPLECORE_API Fields_SumVectorFullSub();
 
 	template<class T>
-	bool GEOMETRYCOLLECTIONEXAMPLECORE_API Fields_SumVectorLeftSide(ExampleResponse&& R);
+	void GEOMETRYCOLLECTIONEXAMPLECORE_API Fields_SumVectorLeftSide();
 
 	template<class T>
-	bool GEOMETRYCOLLECTIONEXAMPLECORE_API Fields_SumVectorRightSide(ExampleResponse&& R);
+	void GEOMETRYCOLLECTIONEXAMPLECORE_API Fields_SumVectorRightSide();
 
 	template<class T>
-	bool GEOMETRYCOLLECTIONEXAMPLECORE_API Fields_SumScalar(ExampleResponse&& R);
+	void GEOMETRYCOLLECTIONEXAMPLECORE_API Fields_SumScalar();
 
 	template<class T>
-	bool GEOMETRYCOLLECTIONEXAMPLECORE_API Fields_SumScalarRightSide(ExampleResponse&& R);
+	void GEOMETRYCOLLECTIONEXAMPLECORE_API Fields_SumScalarRightSide();
 
 	template<class T>
-	bool GEOMETRYCOLLECTIONEXAMPLECORE_API Fields_SumScalarLeftSide(ExampleResponse&& R);
+	void GEOMETRYCOLLECTIONEXAMPLECORE_API Fields_SumScalarLeftSide();
 
 	template<class T>
-	bool GEOMETRYCOLLECTIONEXAMPLECORE_API Fields_Culling(ExampleResponse&& R);
+	void GEOMETRYCOLLECTIONEXAMPLECORE_API Fields_Culling();
 
 	template<class T>
-	bool GEOMETRYCOLLECTIONEXAMPLECORE_API Fields_SerializeAPI(ExampleResponse&& R);
+	void GEOMETRYCOLLECTIONEXAMPLECORE_API Fields_SerializeAPI();
 
 
 }
diff --git a/Engine/Plugins/Experimental/GeometryCollectionPlugin/Source/GeometryCollectionExampleCore/Public/GeometryCollection/GeometryCollectionExampleImplicitCapsule.h b/Engine/Plugins/Experimental/GeometryCollectionPlugin/Source/GeometryCollectionExampleCore/Public/GeometryCollection/GeometryCollectionExampleImplicitCapsule.h
--- a/Engine/Plugins/Experimental/GeometryCollectionPlugin/Source/GeometryCollectionExampleCore/Public/GeometryCollection/GeometryCollectionExampleImplicitCapsule.h
+++ b/Engine/Plugins/Experimental/GeometryCollectionPlugin/Source/GeometryCollectionExampleCore/Public/GeometryCollection/GeometryCollectionExampleImplicitCapsule.h
@@ -7,6 +7,6 @@ namespace GeometryCollectionExample
 {	
 
 	template<class T>
-	bool TestImplicitCapsule(ExampleResponse&& R);
+	void TestImplicitCapsule();
 	
 }
diff --git a/Engine/Plugins/Experimental/GeometryCollectionPlugin/Source/GeometryCollectionExampleCore/Public/GeometryCollection/GeometryCollectionExampleImplicitCylinder.h b/Engine/Plugins/Experimental/GeometryCollectionPlugin/Source/GeometryCollectionExampleCore/Public/GeometryCollection/GeometryCollectionExampleImplicitCylinder.h
--- a/Engine/Plugins/Experimental/GeometryCollectionPlugin/Source/GeometryCollectionExampleCore/Public/GeometryCollection/GeometryCollectionExampleImplicitCylinder.h
+++ b/Engine/Plugins/Experimental/GeometryCollectionPlugin/Source/GeometryCollectionExampleCore/Public/GeometryCollection/GeometryCollectionExampleImplicitCylinder.h
@@ -7,6 +7,6 @@ namespace GeometryCollectionExample
 {	
 
 	template<class T>
-	bool TestImplicitCylinder(ExampleResponse&& R);
+	void TestImplicitCylinder();
 	
 }
diff --git a/Engine/Plugins/Experimental/GeometryCollectionPlugin/Source/GeometryCollectionExampleCore/Public/GeometryCollection/GeometryCollectionExampleImplicitSphere.h b/Engine/Plugins/Experimental/GeometryCollectionPlugin/Source/GeometryCollectionExampleCore/Public/GeometryCollection/GeometryCollectionExampleImplicitSphere.h
--- a/Engine/Plugins/Experimental/GeometryCollectionPlugin/Source/GeometryCollectionExampleCore/Public/GeometryCollection/GeometryCollectionExampleImplicitSphere.h
+++ b/Engine/Plugins/Experimental/GeometryCollectionPlugin/Source/GeometryCollectionExampleCore/Public/GeometryCollection/GeometryCollectionExampleImplicitSphere.h
@@ -7,6 +7,6 @@ namespace GeometryCollectionExample
 {	
 
 	template<class T>
-	bool TestImplicitSphere(ExampleResponse&& R);
+	void TestImplicitSphere();
 	
 }
diff --git a/Engine/Plugins/Experimental/GeometryCollectionPlugin/Source/GeometryCollectionExampleCore/Public/GeometryCollection/GeometryCollectionExampleMatrices.h b/Engine/Plugins/Experimental/GeometryCollectionPlugin/Source/GeometryCollectionExampleCore/Public/GeometryCollection/GeometryCollectionExampleMatrices.h
--- a/Engine/Plugins/Experimental/GeometryCollectionPlugin/Source/GeometryCollectionExampleCore/Public/GeometryCollection/GeometryCollectionExampleMatrices.h
+++ b/Engine/Plugins/Experimental/GeometryCollectionPlugin/Source/GeometryCollectionExampleCore/Public/GeometryCollection/GeometryCollectionExampleMatrices.h
@@ -7,12 +7,12 @@ namespace GeometryCollectionExample
 {	
 
 	template<class T>
-	bool GEOMETRYCOLLECTIONEXAMPLECORE_API BasicGlobalMatrices(ExampleResponse&& R);
+	void GEOMETRYCOLLECTIONEXAMPLECORE_API BasicGlobalMatrices();
 
 	template<class T>
-	bool GEOMETRYCOLLECTIONEXAMPLECORE_API ReparentingMatrices(ExampleResponse&& R);
+	void GEOMETRYCOLLECTIONEXAMPLECORE_API ReparentingMatrices();
 
 	template<class T>
-	bool GEOMETRYCOLLECTIONEXAMPLECORE_API TransformMatrixElement(ExampleResponse&& R);
+	void GEOMETRYCOLLECTIONEXAMPLECORE_API TransformMatrixElement();
 
 }
diff --git a/Engine/Plugins/Experimental/GeometryCollectionPlugin/Source/GeometryCollectionExampleCore/Public/GeometryCollection/GeometryCollectionExampleProximity.h b/Engine/Plugins/Experimental/GeometryCollectionPlugin/Source/GeometryCollectionExampleCore/Public/GeometryCollection/GeometryCollectionExampleProximity.h
--- a/Engine/Plugins/Experimental/GeometryCollectionPlugin/Source/GeometryCollectionExampleCore/Public/GeometryCollection/GeometryCollectionExampleProximity.h
+++ b/Engine/Plugins/Experimental/GeometryCollectionPlugin/Source/GeometryCollectionExampleCore/Public/GeometryCollection/GeometryCollectionExampleProximity.h
@@ -7,33 +7,33 @@ namespace GeometryCollectionExample
 {	
 	
 	template<class T>
-	bool GEOMETRYCOLLECTIONEXAMPLECORE_API BuildProximity(ExampleResponse&& R);
+	void GEOMETRYCOLLECTIONEXAMPLECORE_API BuildProximity();
 	
 	template<class T>
-	bool GEOMETRYCOLLECTIONEXAMPLECORE_API GeometryDeleteFromStart(ExampleResponse&& R);
+	void GEOMETRYCOLLECTIONEXAMPLECORE_API GeometryDeleteFromStart();
 
 	template<class T>
-	bool GEOMETRYCOLLECTIONEXAMPLECORE_API GeometryDeleteFromEnd(ExampleResponse&& R);
+	void GEOMETRYCOLLECTIONEXAMPLECORE_API GeometryDeleteFromEnd();
 
 	template<class T>
-	bool GEOMETRYCOLLECTIONEXAMPLECORE_API GeometryDeleteFromMiddle(ExampleResponse&& R);
+	void GEOMETRYCOLLECTIONEXAMPLECORE_API GeometryDeleteFromMiddle();
 
 	template<class T>
-	bool GEOMETRYCOLLECTIONEXAMPLECORE_API GeometryDeleteMultipleFromMiddle(ExampleResponse&& R);
+	void GEOMETRYCOLLECTIONEXAMPLECORE_API GeometryDeleteMultipleFromMiddle();
 
 	template<class T>
-	bool GEOMETRYCOLLECTIONEXAMPLECORE_API GeometryDeleteRandom(ExampleResponse&& R);
+	void GEOMETRYCOLLECTIONEXAMPLECORE_API GeometryDeleteRandom();
 
 	template<class T>
-	bool GEOMETRYCOLLECTIONEXAMPLECORE_API GeometryDeleteRandom2(ExampleResponse&& R);
+	void GEOMETRYCOLLECTIONEXAMPLECORE_API GeometryDeleteRandom2();
 
 	template<class T>
-	bool GEOMETRYCOLLECTIONEXAMPLECORE_API GeometryDeleteAll(ExampleResponse&& R);
+	void GEOMETRYCOLLECTIONEXAMPLECORE_API GeometryDeleteAll();
 
 	template<class T>
-	bool GEOMETRYCOLLECTIONEXAMPLECORE_API GeometrySwapFlat(ExampleResponse&& R);
+	void GEOMETRYCOLLECTIONEXAMPLECORE_API GeometrySwapFlat();
 
 	template<class T>
-	bool GEOMETRYCOLLECTIONEXAMPLECORE_API TestFracturedGeometry(ExampleResponse&& R);
+	void GEOMETRYCOLLECTIONEXAMPLECORE_API TestFracturedGeometry();
 	
 }
diff --git a/Engine/Plugins/Experimental/GeometryCollectionPlugin/Source/GeometryCollectionExampleCore/Public/GeometryCollection/GeometryCollectionExampleSimulation.h b/Engine/Plugins/Experimental/GeometryCollectionPlugin/Source/GeometryCollectionExampleCore/Public/GeometryCollection/GeometryCollectionExampleSimulation.h
--- a/Engine/Plugins/Experimental/GeometryCollectionPlugin/Source/GeometryCollectionExampleCore/Public/GeometryCollection/GeometryCollectionExampleSimulation.h
+++ b/Engine/Plugins/Experimental/GeometryCollectionPlugin/Source/GeometryCollectionExampleCore/Public/GeometryCollection/GeometryCollectionExampleSimulation.h
@@ -8,28 +8,28 @@ namespace GeometryCollectionExample
 {	
 	
 	template<class T>
-	bool RigidBodiesFallingUnderGravity(ExampleResponse&& R);
+	void RigidBodiesFallingUnderGravity();
 
 	template<class T>
-	bool RigidBodiesCollidingWithSolverFloor(ExampleResponse&& R);
+	void RigidBodiesCollidingWithSolverFloor();
 
 	template<class T>
-	bool RigidBodiesSingleSphereCollidingWithSolverFloor(ExampleResponse&& R);
+	void RigidBodiesSingleSphereCollidingWithSolverFloor();
 
 	template<class T>
-	bool RigidBodiesSingleSphereIntersectingWithSolverFloor(ExampleResponse&& R);
+	void RigidBodiesSingleSphereIntersectingWithSolverFloor();
 
 	template<class T>
-	bool RigidBodiesKinematic(ExampleResponse&& R);
+	void RigidBodiesKinematic();
 
 	template<class T>
-	bool RigidBodiesSleepingActivation(ExampleResponse&& R);
+	void RigidBodiesSleepingActivation();
 
 	template<class T>
-	bool RigidBodies_CollisionGroup(ExampleResponse&& R);
+	void RigidBodies_CollisionGroup();
 
 	template<class T>
-	bool RigidBodies_Initialize_ParticleImplicitCollisionGeometry(ExampleResponse&& R);
+	void RigidBodies_Initialize_ParticleImplicitCollisionGeometry();
 
 	
 }
diff --git a/Engine/Plugins/Experimental/GeometryCollectionPlugin/Source/GeometryCollectionExampleCore/Public/GeometryCollection/GeometryCollectionExampleSimulationField.h b/Engine/Plugins/Experimental/GeometryCollectionPlugin/Source/GeometryCollectionExampleCore/Public/GeometryCollection/GeometryCollectionExampleSimulationField.h
--- a/Engine/Plugins/Experimental/GeometryCollectionPlugin/Source/GeometryCollectionExampleCore/Public/GeometryCollection/GeometryCollectionExampleSimulationField.h
+++ b/Engine/Plugins/Experimental/GeometryCollectionPlugin/Source/GeometryCollectionExampleCore/Public/GeometryCollection/GeometryCollectionExampleSimulationField.h
@@ -8,38 +8,38 @@ namespace GeometryCollectionExample
 {	
 	
 	template<class T>
-	bool RigidBodies_Field_KinematicActivation(ExampleResponse&& R);
+	void RigidBodies_Field_KinematicActivation();
 
 	template<class T>
-	bool RigidBodies_Field_InitialLinearVelocity(ExampleResponse&& R);
+	void RigidBodies_Field_InitialLinearVelocity();
 
 	template<class T>
-	bool RigidBodies_Field_StayDynamic(ExampleResponse&& R);
+	void RigidBodies_Field_StayDynamic();
 
 	template<class T>
-	bool RigidBodies_Field_LinearForce(ExampleResponse&& R);
+	void RigidBodies_Field_LinearForce();
 
 	template<class T>
-	bool RigidBodies_Field_Torque(ExampleResponse&& R);
+	void RigidBodies_Field_Torque();
 
 	template<class T>
-	bool RigidBodies_Field_Kill(ExampleResponse&& R);
+	void RigidBodies_Field_Kill();
 
 	template<class T>
-	bool RigidBodies_Field_LinearVelocity(ExampleResponse&& R);
+	void RigidBodies_Field_LinearVelocity();
 
 	template<class T>
-	bool RigidBodies_Field_CollisionGroup(ExampleResponse&& R);
+	void RigidBodies_Field_CollisionGroup();
 	
 	template<class T>
-	bool RigidBodies_Field_ClusterBreak_StrainModel_Test1(ExampleResponse&& R);
+	void RigidBodies_Field_ClusterBreak_StrainModel_Test1();
 
 	template<class T>
-	bool RigidBodies_Field_ClusterBreak_StrainModel_Test2(ExampleResponse&& R);
+	void RigidBodies_Field_ClusterBreak_StrainModel_Test2();
 
 	template<class T>
-	bool RigidBodies_Field_ClusterBreak_StrainModel_Test3(ExampleResponse&& R);
+	void RigidBodies_Field_ClusterBreak_StrainModel_Test3();
 
 	template<class T>
-	bool RigidBodies_Field_ClusterBreak_StrainModel_Test4(ExampleResponse&& R);
+	void RigidBodies_Field_ClusterBreak_StrainModel_Test4();
 }
diff --git a/Engine/Plugins/Experimental/GeometryCollectionPlugin/Source/GeometryCollectionExampleCore/Public/GeometryCollection/GeometryCollectionExampleSimulationSolver.h b/Engine/Plugins/Experimental/GeometryCollectionPlugin/Source/GeometryCollectionExampleCore/Public/GeometryCollection/GeometryCollectionExampleSimulationSolver.h
--- a/Engine/Plugins/Experimental/GeometryCollectionPlugin/Source/GeometryCollectionExampleCore/Public/GeometryCollection/GeometryCollectionExampleSimulationSolver.h
+++ b/Engine/Plugins/Experimental/GeometryCollectionPlugin/Source/GeometryCollectionExampleCore/Public/GeometryCollection/GeometryCollectionExampleSimulationSolver.h
@@ -8,21 +8,21 @@ namespace GeometryCollectionExample
 {	
 	
 	template<class T>
-	bool Solver_AdvanceNoObjects(ExampleResponse&& R);
+	void Solver_AdvanceNoObjects();
 
 	template<class T>
-	bool Solver_AdvanceDisabledObjects(ExampleResponse&& R);
+	void Solver_AdvanceDisabledObjects();
 
 	template<class T>
-	bool Solver_AdvanceDisabledClusteredObjects(ExampleResponse&& R);
+	void Solver_AdvanceDisabledClusteredObjects();
 
 	template<class T>
-	bool Solver_ValidateReverseMapping(ExampleResponse&& R);
+	void Solver_ValidateReverseMapping();
 
 	template<class T>
-	bool Solver_CollisionEventFilter(ExampleResponse&& R);
+	void Solver_CollisionEventFilter();
 
 	template<class T>
-	bool Solver_BreakingEventFilter(ExampleResponse&& R);
+	void Solver_BreakingEventFilter();
 
 }
diff --git a/Engine/Plugins/Experimental/GeometryCollectionPlugin/Source/GeometryCollectionExampleCore/Public/GeometryCollection/GeometryCollectionExampleSimulationStreaming.h b/Engine/Plugins/Experimental/GeometryCollectionPlugin/Source/GeometryCollectionExampleCore/Public/GeometryCollection/GeometryCollectionExampleSimulationStreaming.h
--- a/Engine/Plugins/Experimental/GeometryCollectionPlugin/Source/GeometryCollectionExampleCore/Public/GeometryCollection/GeometryCollectionExampleSimulationStreaming.h
+++ b/Engine/Plugins/Experimental/GeometryCollectionPlugin/Source/GeometryCollectionExampleCore/Public/GeometryCollection/GeometryCollectionExampleSimulationStreaming.h
@@ -7,12 +7,12 @@
 namespace GeometryCollectionExample
 {	
 	template<class T>
-	bool RigidBodies_Streaming_StartSolverEmpty(ExampleResponse&& R);
+	void RigidBodies_Streaming_StartSolverEmpty();
 
 	template<class T>
-	bool RigidBodies_Streaming_BulkInitialization(ExampleResponse&& R);
+	void RigidBodies_Streaming_BulkInitialization();
 	
 	template<class T>
-	bool RigidBodies_Streaming_DeferedClusteringInitialization(ExampleResponse&& R);
+	void RigidBodies_Streaming_DeferedClusteringInitialization();
 
 }
diff --git a/Engine/Plugins/Experimental/GeometryCollectionPlugin/Source/GeometryCollectionExampleCore/Public/GeometryCollection/GeometryCollectionExampleSpatialHash.h b/Engine/Plugins/Experimental/GeometryCollectionPlugin/Source/GeometryCollectionExampleCore/Public/GeometryCollection/GeometryCollectionExampleSpatialHash.h
--- a/Engine/Plugins/Experimental/GeometryCollectionPlugin/Source/GeometryCollectionExampleCore/Public/GeometryCollection/GeometryCollectionExampleSpatialHash.h
+++ b/Engine/Plugins/Experimental/GeometryCollectionPlugin/Source/GeometryCollectionExampleCore/Public/GeometryCollection/GeometryCollectionExampleSpatialHash.h
@@ -7,21 +7,21 @@ namespace GeometryCollectionExample
 {	
 
 	template<class T>
-	bool GEOMETRYCOLLECTIONEXAMPLECORE_API GetClosestPointsTest1(ExampleResponse&& R);
+	void GEOMETRYCOLLECTIONEXAMPLECORE_API GetClosestPointsTest1();
 	
 	template<class T>
-	bool GEOMETRYCOLLECTIONEXAMPLECORE_API GetClosestPointsTest2(ExampleResponse&& R);
+	void GEOMETRYCOLLECTIONEXAMPLECORE_API GetClosestPointsTest2();
 
 	template<class T>
-	bool GEOMETRYCOLLECTIONEXAMPLECORE_API GetClosestPointsTest3(ExampleResponse&& R);
+	void GEOMETRYCOLLECTIONEXAMPLECORE_API GetClosestPointsTest3();
 
 	template<class T>
-	bool GEOMETRYCOLLECTIONEXAMPLECORE_API GetClosestPointTest(ExampleResponse&& R);
+	void GEOMETRYCOLLECTIONEXAMPLECORE_API GetClosestPointTest();
 
 	template<class T>
-	bool GEOMETRYCOLLECTIONEXAMPLECORE_API HashTableUpdateTest(ExampleResponse&& R);
+	void GEOMETRYCOLLECTIONEXAMPLECORE_API HashTableUpdateTest();
 
 	template<class T>
-	bool GEOMETRYCOLLECTIONEXAMPLECORE_API HashTablePressureTest(ExampleResponse&& R);
+	void GEOMETRYCOLLECTIONEXAMPLECORE_API HashTablePressureTest();
 
 }
diff --git a/Engine/Plugins/Experimental/GeometryCollectionPlugin/Source/GeometryCollectionExampleCore/Public/GeometryCollection/GeometryCollectionExampleUtility.h b/Engine/Plugins/Experimental/GeometryCollectionPlugin/Source/GeometryCollectionExampleCore/Public/GeometryCollection/GeometryCollectionExampleUtility.h
--- a/Engine/Plugins/Experimental/GeometryCollectionPlugin/Source/GeometryCollectionExampleCore/Public/GeometryCollection/GeometryCollectionExampleUtility.h
+++ b/Engine/Plugins/Experimental/GeometryCollectionPlugin/Source/GeometryCollectionExampleCore/Public/GeometryCollection/GeometryCollectionExampleUtility.h
@@ -4,11 +4,14 @@
 #include "ChaosSolversModule.h"
 #include "Chaos/Serializable.h"
 #include "GeometryCollection/GeometryCollectionSimulationTypes.h"
+//#include "GeometryCollection/GeometryCollection.h"
 #include "PhysicsProxy/GeometryCollectionPhysicsProxy.h"
 #include "Templates/SharedPointer.h"
 
 class FGeometryCollection;
+class FGeometryCollectionPhysicsProxy;
 class FGeometryDynamicCollection;
+struct FSimulationParameters;
 
 namespace Chaos
 {
@@ -29,5 +32,32 @@ namespace GeometryCollectionExample {
 	TSharedPtr<FGeometryCollection>	CreateClusteredBody_ThreeByTwo_ThreeTransform(FVector Position);
 	TSharedPtr<FGeometryCollection>	CreateClusteredBody_FracturedGeometry(FVector Position = FVector(0));
 
-
+	template<class T>
+	void InitMaterialToZero(TUniquePtr<Chaos::TChaosPhysicsMaterial<T>> const &PhysicalMaterial);
+
+	using FCollectionInitFunc = TFunction<void(TSharedPtr<FGeometryCollection>&)>;
+	struct InitCollectionsParameters
+	{
+		const FTransform& RestCenter;
+		const FVector RestScale;
+		FCollectionInitFunc RestInitFunc;
+		int DynamicStateDefault;
+	};
+
+	template<class T>
+	void InitCollections(
+		TUniquePtr<Chaos::TChaosPhysicsMaterial<T>> &PhysicalMaterial,
+		TSharedPtr<FGeometryCollection>& RestCollection,
+		TSharedPtr<FGeometryDynamicCollection>& DynamicCollection,
+		InitCollectionsParameters& InitParams
+	);
+
+	using FInitFunc = TFunction<void(FSimulationParameters&)>;
+	template<class T>
+	FGeometryCollectionPhysicsProxy* RigidBodySetup(
+		TUniquePtr<Chaos::TChaosPhysicsMaterial<T>> &PhysicalMaterial,
+		TSharedPtr<FGeometryCollection> &RestCollection,
+		TSharedPtr<FGeometryDynamicCollection> &DynamicCollection,
+		FInitFunc CustomFunc = nullptr
+	);
 }
diff --git a/Engine/Plugins/Experimental/GeometryCollectionPlugin/Source/GeometryCollectionExampleCore/Public/GeometryCollection/GeometryCollectionExampleVisibility.h b/Engine/Plugins/Experimental/GeometryCollectionPlugin/Source/GeometryCollectionExampleCore/Public/GeometryCollection/GeometryCollectionExampleVisibility.h
--- a/Engine/Plugins/Experimental/GeometryCollectionPlugin/Source/GeometryCollectionExampleCore/Public/GeometryCollection/GeometryCollectionExampleVisibility.h
+++ b/Engine/Plugins/Experimental/GeometryCollectionPlugin/Source/GeometryCollectionExampleCore/Public/GeometryCollection/GeometryCollectionExampleVisibility.h
@@ -6,5 +6,5 @@
 namespace GeometryCollectionExample
 {	
 	template<class T>
-	bool GEOMETRYCOLLECTIONEXAMPLECORE_API TestHideVertices(ExampleResponse&& R);
+	void GEOMETRYCOLLECTIONEXAMPLECORE_API TestHideVertices();
 }
diff --git a/Engine/Plugins/Experimental/HairStrands/Config/BaseHairStrands.ini b/Engine/Plugins/Experimental/HairStrands/Config/BaseHairStrands.ini
new file mode 100644
--- /dev/null
+++ b/Engine/Plugins/Experimental/HairStrands/Config/BaseHairStrands.ini
@@ -0,0 +1,5 @@
+[CoreRedirects]
++PropertyRedirects=(OldName="HairStrandsAsset",NewName="GroomAsset")
++ClassRedirects=(OldName="HairStrandsActor",NewName="GroomActor")
++ClassRedirects=(OldName="HairStrandsAsset",NewName="GroomAsset")
++ClassRedirects=(OldName="HairStrandsComponent",NewName="GroomComponent",ValueChanges=(("HairStrandsComponent0","GroomComponent0")))
diff --git a/Engine/Plugins/Experimental/HairStrands/HairStrands.uplugin b/Engine/Plugins/Experimental/HairStrands/HairStrands.uplugin
--- a/Engine/Plugins/Experimental/HairStrands/HairStrands.uplugin
+++ b/Engine/Plugins/Experimental/HairStrands/HairStrands.uplugin
@@ -18,11 +18,11 @@
 		{
 			"Name": "HairStrandsCore",
 			"Type": "Runtime",
-			"LoadingPhase": "Default",
+			"LoadingPhase": "EarliestPossible",
 			"WhitelistPlatforms": [
 				"Win64",
 				"Win32"
-			] 
+			]
 		},
 		{
 			"Name": "HairStrandsEditor",
diff --git a/Engine/Plugins/Experimental/HairStrands/Shaders/Private/NiagaraDataInterfaceHairStrands.ush b/Engine/Plugins/Experimental/HairStrands/Shaders/Private/NiagaraDataInterfaceHairStrands.ush
--- a/Engine/Plugins/Experimental/HairStrands/Shaders/Private/NiagaraDataInterfaceHairStrands.ush
+++ b/Engine/Plugins/Experimental/HairStrands/Shaders/Private/NiagaraDataInterfaceHairStrands.ush
@@ -1,56 +1,368 @@
 // Copyright 1998-2019 Epic Games, Inc. All Rights Reserved.
 
+/*=============================================================================
+NiagaraDataInterfaceHairStrands.ush 
+=============================================================================*/
+
 #pragma once
 
+/* -----------------------------------------------------------------
+ * Hair Strands constants and context
+ * -----------------------------------------------------------------
+ */
+
 #define WRAP_SIZE 32
+static const float M_PI = 3.14159265;
+static const float SCALE = 1e+6;
+
+static const int MASS_OFFSET = 0;
+static const int TANGENT_OFFSET = 3;
+static const int VELOCITY_OFFSET = 6;
 
 struct FDIHairStrandsContext
 {
-	RWBuffer<float4>	PointsPositionsBuffer;
+	RWBuffer<uint4>		PointsPositionsBuffer;
 	Buffer<uint>		CurvesOffsetsBuffer;
-	Buffer<float4>		NodesPositionsBuffer;
-	Buffer<uint>		PointsNodesBuffer;
-	Buffer<float>		PointsWeightsBuffer;
+	Buffer<uint4>		RestPositionsBuffer;
+	Texture3D<int>		GridCurrentBuffer;
+	RWTexture3D<int>	GridDestinationBuffer;
 	float4x4			WorldTransform;
+	float4x4			WorldInverse;
+	float4				WorldRotation;
 	int					NumStrands;
 	int					StrandSize;
 	float				StrandDensity;
 	float				RootThickness;
 	float				TipThickness;
+	float4				GridOrigin;
+	int4				GridSize;
 }; 
 
 #define DIHAIRSTRANDS_DECLARE_CONSTANTS(NAME)\
-RWBuffer<float4> PointsPositionsBuffer_##NAME;\
+RWBuffer<uint4> PointsPositionsBuffer_##NAME;\
 Buffer<uint> CurvesOffsetsBuffer_##NAME;\
-Buffer<float4> NodesPositionsBuffer_##NAME;\
-Buffer<uint> PointsNodesBuffer_##NAME;\
-Buffer<float> PointsWeightsBuffer_##NAME;\
+Buffer<uint4> RestPositionsBuffer_##NAME;\
+Texture3D<int> GridCurrentBuffer_##NAME;\
+RWTexture3D<int> GridDestinationBuffer_##NAME;\
 float4x4 WorldTransform_##NAME;\
+float4x4 WorldInverse_##NAME;\
+float4 WorldRotation_##NAME;\
 int NumStrands_##NAME;\
 int StrandSize_##NAME;\
 float StrandDensity_##NAME;\
 float RootThickness_##NAME;\
-float TipThickness_##NAME;
+float TipThickness_##NAME;\
+float4 GridOrigin_##NAME;\
+int4 GridSize_##NAME;\
 
 #define DIHAIRSTRANDS_MAKE_CONTEXT(NAME)\
 FDIHairStrandsContext DIContext;\
 DIContext.PointsPositionsBuffer = PointsPositionsBuffer_##NAME;\
 DIContext.CurvesOffsetsBuffer = CurvesOffsetsBuffer_##NAME;\
-DIContext.NodesPositionsBuffer = NodesPositionsBuffer_##NAME;\
-DIContext.PointsNodesBuffer = PointsNodesBuffer_##NAME;\
-DIContext.PointsWeightsBuffer = PointsWeightsBuffer_##NAME;\
+DIContext.RestPositionsBuffer = RestPositionsBuffer_##NAME;\
 DIContext.WorldTransform = WorldTransform_##NAME;\
+DIContext.WorldInverse = WorldInverse_##NAME;\
+DIContext.WorldRotation = WorldRotation_##NAME;\
 DIContext.NumStrands = NumStrands_##NAME;\
 DIContext.StrandSize = StrandSize_##NAME;\
 DIContext.StrandDensity = StrandDensity_##NAME;\
 DIContext.RootThickness = RootThickness_##NAME;\
-DIContext.TipThickness = TipThickness_##NAME;
+DIContext.TipThickness = TipThickness_##NAME;\
+DIContext.GridOrigin = GridOrigin_##NAME;\
+DIContext.GridSize = GridSize_##NAME;\
+DIContext.GridCurrentBuffer = GridCurrentBuffer_##NAME;\
+DIContext.GridDestinationBuffer = GridDestinationBuffer_##NAME;\
+
+/* -----------------------------------------------------------------
+ * Grid utilities
+ * -----------------------------------------------------------------
+ */
+
+// Return true if the grid index is within the domain size
+bool DIHairStrands_IsIndexValid(in FDIHairStrandsContext DIContext, in int3 GridIndex)
+{
+	return ((GridIndex.x < (DIContext.GridSize.x-1)) && (GridIndex.y < (DIContext.GridSize.y-1)) && (GridIndex.z < (DIContext.GridSize.z-1)));
+}
+
+// Transform the node world position in grid space
+float3 DIHairStrands_GetGridPosition(in FDIHairStrandsContext DIContext, in float3 NodePosition)
+{
+	return ( mul(float4( NodePosition, 1.0), DIContext.WorldInverse).xyz - DIContext.GridOrigin.xyz );
+}
+
+// Return the grid hash given a grid index
+int DIHairStrands_GetGridHash(in FDIHairStrandsContext DIContext, in int3 GridIndex)
+{
+	return GridIndex.z * DIContext.GridSize.y * DIContext.GridSize.x + GridIndex.y * DIContext.GridSize.x + GridIndex.x;
+}
+
+// Compute the grid index and fraction
+void DIHairStrands_ComputeGridIndex(in FDIHairStrandsContext DIContext, in float3 GridPosition, in float3 GridOffset, out int3 OutGridIndex, out float3 OutGridFraction)
+{
+	const float3 GridCell = GridPosition / DIContext.GridOrigin.w - GridOffset;
+
+	OutGridIndex = floor(GridCell);
+	OutGridFraction = GridCell - OutGridIndex;
+}
+
+int3 DIHairStrands_GridNodeIndex( in FDIHairStrandsContext DIContext, in int3 GridIndex, in int DataOffset)
+{
+	return int3( GridIndex.x + DataOffset * DIContext.GridSize.x, GridIndex.yz);
+}
+
+void DIHairStrands_AddNodeScalar( in FDIHairStrandsContext DIContext, in int3 GridIndex, in int DataOffset, in float NodeScalar, in bool EnableLocking)
+{
+	const int ConvertScalar = NodeScalar * SCALE;
+	if(EnableLocking)
+	{
+		InterlockedAdd(DIContext.GridDestinationBuffer[DIHairStrands_GridNodeIndex(DIContext,GridIndex,DataOffset)], ConvertScalar);
+	}
+	else
+	{
+		DIContext.GridDestinationBuffer[DIHairStrands_GridNodeIndex(DIContext,GridIndex,DataOffset)] += ConvertScalar;
+	}
+	DeviceMemoryBarrier();
+}
+
+void DIHairStrands_AddNodeVector( in FDIHairStrandsContext DIContext, in int3 GridIndex, in int DataOffset, in float3 NodeVector, in bool EnableLocking)
+{
+	DIHairStrands_AddNodeScalar(DIContext,GridIndex,DataOffset,NodeVector.x,EnableLocking);
+	DIHairStrands_AddNodeScalar(DIContext,GridIndex,DataOffset+1,NodeVector.y,EnableLocking);
+	DIHairStrands_AddNodeScalar(DIContext,GridIndex,DataOffset+2,NodeVector.z,EnableLocking);
+}
+
+float DIHairStrands_GetGridScalar( in FDIHairStrandsContext DIContext, in int3 GridIndex, in int DataOffset)
+{
+	return DIContext.GridCurrentBuffer[DIHairStrands_GridNodeIndex(DIContext,GridIndex,DataOffset)] / SCALE;
+}
+
+float3 DIHairStrands_GetGridVector( in FDIHairStrandsContext DIContext, in int3 GridIndex, in int DataOffset)
+{
+	return float3(DIHairStrands_GetGridScalar(DIContext,GridIndex,DataOffset),
+				  DIHairStrands_GetGridScalar(DIContext,GridIndex,DataOffset+1),
+				  DIHairStrands_GetGridScalar(DIContext,GridIndex,DataOffset+2));
+}
+
+// Transfer the velocity from the node to the grid
+void DIHairStrands_BuildCollisionGrid(in FDIHairStrandsContext DIContext, in float3 NodePosition, in float3 NodeVelocity, in float NodeMass, 
+			in float3 NodeGradientX, in float3 NodeGradientY, in float3 NodeGradientZ, out bool OutBuildStatus)
+{
+	SharedLocalPosition[GGroupThreadId.x] = NodePosition;
+	SharedLocalVelocity[GGroupThreadId.x] = NodeVelocity;
+	SharedInverseMass[GGroupThreadId.x] = NodeMass;
+	GroupMemoryBarrier(); 
+
+	OutBuildStatus = false;
+
+	const int LocalIndex = GGroupThreadId.x % DIContext.StrandSize;
+	if(LocalIndex != 0)
+	{
+		const float3 DeltaVelocity = SharedLocalVelocity[GGroupThreadId.x] - SharedLocalVelocity[GGroupThreadId.x-1];
+		const float3 DeltaPosition = SharedLocalPosition[GGroupThreadId.x] - SharedLocalPosition[GGroupThreadId.x-1];
+
+		const float2 SegmentWeight = (LocalIndex == 1) ? float2(1.0,0.5) : (LocalIndex == (DIContext.StrandSize-1) ) ? float2(0.5,1.0) : float2(0.5,0.5);
+		const float SegmentMass =  (SharedInverseMass[GGroupThreadId.x-1] * SegmentWeight.x + SharedInverseMass[GGroupThreadId.x] * SegmentWeight.y);
+		const float SegmentLength = length(DeltaPosition);
+
+		const float3 GridPositionA = DIHairStrands_GetGridPosition(DIContext,SharedLocalPosition[GGroupThreadId.x-1]);
+		const float3 GridPositionB = DIHairStrands_GetGridPosition(DIContext,SharedLocalPosition[GGroupThreadId.x]);
+		const float3 GridSegment = GridPositionB - GridPositionA;
+
+		const int NumSamples = ceil(SegmentLength / DIContext.GridOrigin.w);
+
+		const float SampleMass = SegmentMass / NumSamples;
+		const float3 SampleTangent = DeltaPosition / SegmentLength;
+
+		for( int SampleIndex = 0; SampleIndex < NumSamples; ++SampleIndex)
+		{
+			const float SampleCoord = (0.5+SampleIndex) / NumSamples;
+			const float3 SamplePosition = GridPositionA + SampleCoord * GridSegment;
+			const float3 SampleMomentum = (SharedLocalVelocity[GGroupThreadId.x-1] + SampleCoord * DeltaVelocity) * SampleMass;
+
+			for( int ComponentIndex = 0; ComponentIndex < 3; ++ComponentIndex)
+			{
+				const float3 NodeGradient = (ComponentIndex == 0) ? NodeGradientX : (ComponentIndex == 1) ? NodeGradientY : NodeGradientZ;
+				float3 GridOffset = float3(0.5,0.5,0.5); GridOffset[ComponentIndex] = 0.0;
+				int3 GridIndex = int3(0,0,0);
+				float3 GridFraction = float3(0,0,0);
+				DIHairStrands_ComputeGridIndex(DIContext,SamplePosition,GridOffset,GridIndex,GridFraction);
+
+				if(DIHairStrands_IsIndexValid(DIContext,GridIndex))
+				{
+					OutBuildStatus = true;
+					for(int i = 0; i < 2; ++i)
+					{
+						const float GridWeightX = (i == 0) ? 1.0 - GridFraction.x : GridFraction.x;
+						for(int j = 0; j < 2; ++j)
+						{
+							const float GridWeightY = (j == 0) ? 1.0 - GridFraction.y : GridFraction.y;
+							for(int k = 0; k < 2; ++k)
+							{
+								const float GridWeightZ = (k == 0) ? 1.0 - GridFraction.z : GridFraction.z;
+
+								const float GridWeightXYZ = GridWeightX*GridWeightY*GridWeightZ;
+
+								const int3 GridLocation = int3(GridIndex.x+i,GridIndex.y+j,GridIndex.z+k);
+								const float3 DeltaPosition = ((GridLocation+GridOffset)* DIContext.GridOrigin.w) - SamplePosition;
+								const float ExtendedMomentum = SampleMomentum[ComponentIndex] + dot(NodeGradient, DeltaPosition) * SampleMass;
+
+								DIHairStrands_AddNodeScalar(DIContext,GridLocation,MASS_OFFSET+ComponentIndex,SampleMass * GridWeightXYZ,true);
+								DIHairStrands_AddNodeScalar(DIContext,GridLocation,VELOCITY_OFFSET+ComponentIndex,ExtendedMomentum * GridWeightXYZ,true);
+								//DIHairStrands_AddNodeScalar(DIContext,GridLocation,TANGENT_OFFSET+ComponentIndex,SampleTangent[ComponentIndex] * GridWeightXYZ,true);
+							}
+						}
+					}
+				}
+			}
+		}
+	}
+}
+
+
+// Return the velocity from the grid to the node
+void DIHairStrands_QueryCollisionGrid(in FDIHairStrandsContext DIContext, in float3 NodePosition, out float3 OutGridVelocity,
+			out float3 OutGridTangent, out float OutGridMass, out float3 OutGridGradientX, out float3 OutGridGradientY, out float3 OutGridGradientZ)
+{
+	const float3 GridPosition = DIHairStrands_GetGridPosition(DIContext,NodePosition);
+	
+	OutGridVelocity = float3(0,0,0);
+	OutGridMass = 0.0;
+	OutGridTangent = float3(0,0,0);
+	OutGridGradientX = float3(0,0,0);
+	OutGridGradientY = float3(0,0,0);
+	OutGridGradientZ = float3(0,0,0);
+
+	for( int ComponentIndex = 0; ComponentIndex < 3; ++ComponentIndex)
+	{
+		float3 GridOffset = float3(0.5,0.5,0.5); GridOffset[ComponentIndex] = 0.0;
+		int3 GridIndex = int3(0,0,0);
+		float3 GridFraction = float3(0,0,0);
+		DIHairStrands_ComputeGridIndex(DIContext,GridPosition,GridOffset,GridIndex,GridFraction);
+		if(DIHairStrands_IsIndexValid(DIContext,GridIndex))
+		{
+			for(int i = 0; i < 2; ++i)
+			{
+				const float GridWeightX = (i == 0) ? 1.0 - GridFraction.x : GridFraction.x;
+				const float GridGradientX = (i == 0) ? -1.0: 1.0;
+				for(int j = 0; j < 2; ++j)
+				{
+					const float GridWeightY = (j == 0) ? 1.0 - GridFraction.y : GridFraction.y;
+					const float GridGradientY = (j == 0) ? -1.0: 1.0;
+					for(int k = 0; k < 2; ++k)
+					{
+						const float GridWeightZ = (k == 0) ? 1.0 - GridFraction.z : GridFraction.z;
+						const float GridGradientZ = (k == 0) ? -1.0: 1.0;
+
+						const float GridWeightXYZ = GridWeightX*GridWeightY*GridWeightZ;
+						const float3 GridGradientXYZ = float3(GridGradientX*GridWeightY*GridWeightZ,
+															  GridWeightX*GridGradientY*GridWeightZ,
+															  GridWeightX*GridWeightY*GridGradientZ) / DIContext.GridOrigin.w;
+
+						const int3 GridLocation = int3(GridIndex.x+i,GridIndex.y+j,GridIndex.z+k);
+					
+						const float LocalMass = DIHairStrands_GetGridScalar(DIContext,GridLocation,MASS_OFFSET+ComponentIndex);
+						if( LocalMass != 0.0 )
+						{
+							const float LocalVelocity = DIHairStrands_GetGridScalar(DIContext,GridLocation,VELOCITY_OFFSET+ComponentIndex) / LocalMass;
+							OutGridMass += LocalMass * GridWeightXYZ;
+							OutGridVelocity[ComponentIndex] += LocalVelocity * GridWeightXYZ;
+							OutGridTangent[ComponentIndex] += (DIHairStrands_GetGridScalar(DIContext,GridLocation,TANGENT_OFFSET+ComponentIndex) / LocalMass) * GridWeightXYZ;
+							//OutGridTangent[ComponentIndex] += DIHairStrands_GetGridScalar(DIContext,GridLocation,TANGENT_OFFSET+ComponentIndex) * GridWeightXYZ;
+							if( ComponentIndex == 0)
+							{
+								OutGridGradientX +=  GridGradientXYZ * LocalVelocity * GridWeightXYZ;
+							}
+							else if( ComponentIndex == 1)
+							{
+								OutGridGradientY +=  GridGradientXYZ * LocalVelocity * GridWeightXYZ;
+							}
+							else
+							{
+								OutGridGradientZ +=  GridGradientXYZ * LocalVelocity * GridWeightXYZ;
+							}
+						}
+					}
+				}
+			}
+		}
+	}
+	OutGridMass /= (3*DIContext.GridOrigin.w * DIContext.GridOrigin.w * DIContext.GridOrigin.w);
+	//if(OutGridMass!=0.0)
+	//{
+	//	OutGridTangent = normalize(OutGridTangent);
+	//}
+}
 
-groupshared float3 NodesGlobalPosition[WRAP_SIZE];
-groupshared float4 NodesGlobalQuaternion[WRAP_SIZE];
 
-groupshared float NodesInverseMass[WRAP_SIZE];
-groupshared float NodesInverseInertia[WRAP_SIZE];
+// Return the velocity from the grid to the node
+void DIHairStrands_ProjectCollisionGrid(in FDIHairStrandsContext DIContext, in int GridHash, out bool OutProjectStatus)
+{	
+	OutProjectStatus = false;
+	const int GridIndexX = GridHash / ( DIContext.GridSize.y * DIContext.GridSize.z);
+	const int GridDummy = GridHash - GridIndexX * DIContext.GridSize.y * DIContext.GridSize.z;
+	const int GridIndexY = GridDummy / (DIContext.GridSize.z);
+	const int GridIndexZ = GridDummy % (DIContext.GridSize.z);
+
+	const bool OddStage = (ShaderStageIndex % 2)  == 0;
+	const bool OddVoxel = ((GridIndexX+GridIndexY+GridIndexZ) % 2) == 0;
+
+	if( (OddStage && OddVoxel) || (!OddStage && !OddVoxel) )
+	{
+		const float VoxelVolume = DIContext.GridOrigin.w * DIContext.GridOrigin.w * DIContext.GridOrigin.w;
+
+		const int3 PrevIndex = int3(GridIndexX,GridIndexY,GridIndexZ);
+		if(DIHairStrands_IsIndexValid(DIContext,PrevIndex))
+		{
+			OutProjectStatus = true;
+			float GridDivergence = 0.0;
+			float GridFraction = 0.0;
+			for( int ComponentIndex = 0; ComponentIndex < 3; ++ComponentIndex)
+			{
+				int3 NextIndex = PrevIndex; NextIndex[ComponentIndex] += 1;
+
+				const float PrevMass = DIHairStrands_GetGridScalar(DIContext,PrevIndex,MASS_OFFSET+ComponentIndex);
+				const float PrevDensity = min(1,max(0,PrevMass / VoxelVolume));
+				if( PrevMass > 0.)
+				{
+					GridDivergence -= DIHairStrands_GetGridScalar(DIContext,PrevIndex,VELOCITY_OFFSET+ComponentIndex) * PrevDensity / PrevMass;
+					GridFraction += PrevDensity;
+				}
+
+				const float NextMass = DIHairStrands_GetGridScalar(DIContext,NextIndex,MASS_OFFSET+ComponentIndex);
+				const float NextDensity = min(1,max(0,NextMass / VoxelVolume));
+				if( NextMass > 0.)
+				{
+					GridDivergence += DIHairStrands_GetGridScalar(DIContext,NextIndex,VELOCITY_OFFSET+ComponentIndex) * NextDensity / NextMass;
+					GridFraction += NextDensity;
+				}
+			}
+			if( GridFraction > 0.)
+			{
+				GridDivergence *= DIContext.GridOrigin.w;
+				const float GridPressure = -GridDivergence / GridFraction;
+
+				const float PressureGradient = GridPressure / DIContext.GridOrigin.w;
+				for( int ComponentIndexB = 0; ComponentIndexB < 3; ++ComponentIndexB)
+				{
+					int3 NextIndex = PrevIndex; NextIndex[ComponentIndexB] += 1;
+
+					const float PrevMass = DIHairStrands_GetGridScalar(DIContext,PrevIndex,MASS_OFFSET+ComponentIndexB);
+					DIHairStrands_AddNodeScalar(DIContext,PrevIndex,VELOCITY_OFFSET+ComponentIndexB,-PressureGradient*PrevMass,false);
+					DIHairStrands_AddNodeScalar(DIContext,PrevIndex,TANGENT_OFFSET+ComponentIndexB,-PressureGradient*PrevMass,false);
+
+					const float NextMass = DIHairStrands_GetGridScalar(DIContext,NextIndex,MASS_OFFSET+ComponentIndexB);
+					DIHairStrands_AddNodeScalar(DIContext,NextIndex,VELOCITY_OFFSET+ComponentIndexB,PressureGradient*NextMass,false);
+					DIHairStrands_AddNodeScalar(DIContext,NextIndex,TANGENT_OFFSET+ComponentIndexB,PressureGradient*NextMass,false);
+				}
+			}
+		}	
+	}
+}
+
+/* -----------------------------------------------------------------
+ * Utilities regarding strands indexing
+ * -----------------------------------------------------------------
+ */
 
 // Given a node index return the strand index and the local index within the strand
 int DIHairStrands_GetNumNodes(in FDIHairStrandsContext DIContext)
@@ -59,9 +371,9 @@ int DIHairStrands_GetNumNodes(in FDIHairStrandsContext DIContext)
 }
 
 // Given a node index return the strand index and the local index within the strand
-bool DIHairStrands_IsRootNode(in FDIHairStrandsContext DIContext)
+bool DIHairStrands_IsRootNode(in FDIHairStrandsContext DIContext, in int NodeIndex)
 {
-	return (GGroupThreadId.x % DIContext.StrandSize) != 0;
+	return (NodeIndex % DIContext.StrandSize) != 0;
 }
 
 // Given a node index return the strand index and the local index within the strand
@@ -88,148 +400,342 @@ float DIHairStrands_ComputeStrandLength(in FDIHairStrandsContext DIContext, in i
 		int PointOffset = DIContext.CurvesOffsetsBuffer[StrandIndex];
 		const int EdgeCount = DIContext.CurvesOffsetsBuffer[StrandIndex+1]-PointOffset-1;
 
-		float3 PointNext = DIContext.PointsPositionsBuffer[PointOffset++].xyz;
+		float3 PointNext = f16tof32(DIContext.RestPositionsBuffer[PointOffset++].xyz);
 		float3 PointPrev = PointNext;
 
 		for (int EdgeIndex = 0; EdgeIndex < EdgeCount; ++EdgeIndex)
 		{
 			PointPrev = PointNext;
-			PointNext = DIContext.PointsPositionsBuffer[PointOffset++].xyz;
+			PointNext = f16tof32(DIContext.RestPositionsBuffer[PointOffset++].xyz);
 			StrandLength += length(PointNext-PointPrev);
 		}
 	}
 	return StrandLength;
 }
 
-/*void DIHairStrands_ReportNodePosition(in FDIHairStrandsContext DIContext, in int NodeIndex, out float3 Out_NodePosition)
+/* -----------------------------------------------------------------
+ * Node Mass/Inertia/Position/Orientation computation
+ * -----------------------------------------------------------------
+ */
+
+// Compute the node mass
+void DIHairStrands_ComputeNodeMass(in FDIHairStrandsContext DIContext, in int NodeIndex, out float OutNodeMass)
 {
-	int StrandIndex = 0;
-	int LocalIndex = 0;
+	const int ReducedSize =  DIContext.StrandSize-1;
+	const float CoordScale = 1.0 / (ReducedSize-1.0);
+
+	int StrandIndex = 0, LocalIndex = 0;
 	DIHairStrands_ComputeNodeIndices(DIContext,NodeIndex,StrandIndex,LocalIndex);
+	const float EdgeLength = DIHairStrands_ComputeStrandLength(DIContext,StrandIndex) * CoordScale;
 
-	const int PointBegin = DIContext.CurvesOffsetsBuffer[StrandIndex];
-	const int PointEnd = DIContext.CurvesOffsetsBuffer[StrandIndex+1];
+	LocalIndex = max(0,LocalIndex - 1);
+
+	const bool IsFirst = (LocalIndex == 0);
+	const bool IsLast  = (LocalIndex == (ReducedSize-1));
+
+	const float EdgeScale = (IsFirst || IsLast) ? 0.5 : 1.0;
+	
+	const float PrevCoord = IsFirst ? 0.0 : (LocalIndex-0.5) * CoordScale;
+	const float NextCoord = IsLast  ? 1.0 : (LocalIndex+0.5) * CoordScale;
+
+	const float PrevThick = DIContext.RootThickness * ( 1.0 - PrevCoord ) + DIContext.TipThickness * PrevCoord;
+	const float NextThick = DIContext.RootThickness * ( 1.0 - NextCoord ) + DIContext.TipThickness * NextCoord;
+	const float MeanRadius = 0.25 * ( PrevThick + NextThick );
+
+	//OutNodeMass = DIContext.StrandDensity * EdgeLength * EdgeScale * M_PI * MeanRadius * MeanRadius;
+	OutNodeMass = DIContext.StrandDensity * EdgeLength * EdgeScale * M_PI * (PrevThick * PrevThick + NextThick * PrevThick + NextThick * NextThick) / 12.0;
+}
+
+// Compute the node inertia
+void DIHairStrands_ComputeNodeInertia(in FDIHairStrandsContext DIContext, in int NodeIndex, out float3 OutNodeInertia)
+{
+	const int ReducedSize =  DIContext.StrandSize-1;
+	const float CoordScale = 1.0 / (ReducedSize-1.0);
+
+	int StrandIndex = 0, LocalIndex = 0;
+	DIHairStrands_ComputeNodeIndices(DIContext,NodeIndex,StrandIndex,LocalIndex);
+	const float EdgeLength = DIHairStrands_ComputeStrandLength(DIContext,StrandIndex) * CoordScale;
+
+	LocalIndex = min(max(0,LocalIndex - 1),ReducedSize-2);
+
+	const float PrevCoord = LocalIndex * CoordScale;
+	const float NextCoord = (LocalIndex+1) * CoordScale;
+
+	const float PrevThick = DIContext.RootThickness * ( 1.0 - PrevCoord ) + DIContext.TipThickness * PrevCoord;
+	const float NextThick = DIContext.RootThickness * ( 1.0 - NextCoord ) + DIContext.TipThickness * NextCoord;
+	const float MeanRadius = 0.25 * ( PrevThick + NextThick );
+
+	const float EdgeMass = DIContext.StrandDensity * EdgeLength * M_PI * (PrevThick * PrevThick + NextThick * PrevThick + NextThick * NextThick) / 12.0;
+	//const float RadialInertia = EdgeMass * ( 3.0 * MeanRadius * MeanRadius + EdgeLength * EdgeLength) / 12.0;
+	const float RadialInertia = EdgeMass;// * ( MeanRadius * MeanRadius) / 4.0;
+	OutNodeInertia = float3( RadialInertia, RadialInertia, 2.0 * RadialInertia);//EdgeMass * MeanRadius * MeanRadius * 0.5 );
+
+	//const float EdgeMass = DIContext.StrandDensity * EdgeLength * M_PI * MeanRadius * MeanRadius;
+	//const float RadialInertia = EdgeMass * ( MeanRadius * MeanRadius) / 4.0;
+	//OutNodeInertia = float3( RadialInertia, RadialInertia, 2.0 * RadialInertia );
+}
+
+// Compute the node position by linear interpolation over the points
+void DIHairStrands_ComputeNodePosition(in FDIHairStrandsContext DIContext, in int NodeIndex, out float3 OutNodePosition)
+{
+	const int ReducedSize =  DIContext.StrandSize-1;
+	const float CoordScale = 1.0 / (ReducedSize-1.0);
 
-	for (int PointIndex = PointBegin; PointIndex < PointEnd; ++PointIndex)
+	int StrandIndex = 0, LocalIndex = 0;
+	DIHairStrands_ComputeNodeIndices(DIContext,NodeIndex,StrandIndex,LocalIndex);
+
+	const bool RootNode = (LocalIndex == 0);
+	LocalIndex = max(0,LocalIndex - 1);
+
+	const int PointOffset = DIContext.CurvesOffsetsBuffer[StrandIndex];
+
+	const float NodeCoord = (float)(LocalIndex) * CoordScale;
+	const float PointCount = DIContext.CurvesOffsetsBuffer[StrandIndex+1]-PointOffset-1;
+
+	const float NodePoint = NodeCoord * PointCount;
+	const int PointPrev = (LocalIndex==0) ? 0 : (LocalIndex==(ReducedSize-1)) ? PointCount-1 : floor(NodePoint);
+	const int PointNext = PointPrev+1;
+
+	if(RootNode)
 	{
-		const float Weight = (DIContext.PointsNodesBuffer[PointIndex] == NodeIndex) ? 1.0 : 0.0;
-		DIContext.PointsPositionsBuffer[PointIndex] += Weight * Out_NodePosition * DIContext.PointsWeightsBuffer[PointIndex];
+		const float EdgeLength = DIHairStrands_ComputeStrandLength(DIContext,StrandIndex) * CoordScale;
+		float3 EdgeDirection = f16tof32(DIContext.RestPositionsBuffer[PointOffset+PointNext].xyz) - 
+									 f16tof32(DIContext.RestPositionsBuffer[PointOffset+PointPrev].xyz);
+		OutNodePosition = 	f16tof32(DIContext.RestPositionsBuffer[PointOffset+PointPrev].xyz) - EdgeLength * normalize(EdgeDirection);					 
 	}
-	DeviceMemoryBarrier();
+	else
+	{
+		const float NodeAlpha = NodePoint - (float)PointPrev;
+		OutNodePosition =  f16tof32(DIContext.RestPositionsBuffer[PointOffset+PointPrev].xyz) * (1.0-NodeAlpha) + 
+						   f16tof32(DIContext.RestPositionsBuffer[PointOffset+PointNext].xyz) * NodeAlpha;
+	}
+	OutNodePosition = mul(float4( OutNodePosition, 1.0), DIContext.WorldTransform).xyz;
+}
+
+// Compute the node orientation
+void DIHairStrands_ComputeNodeOrientation(in FDIHairStrandsContext DIContext, in int NodeIndex, out float4 OutNodeOrientation)
+{
+	int StrandIndex = 0, LocalIndex = 0;
+	DIHairStrands_ComputeNodeIndices(DIContext,NodeIndex,StrandIndex,LocalIndex);
+	const int RootIndex = StrandIndex * DIContext.StrandSize;
 
-	for (int PointIndex = PointBegin; PointIndex < PointEnd; ++PointIndex)
+	float4 NodeQuaternion = QUATERNION_IDENTITY;
+	float3 TangentNext = float3(0,0,1);
+	float3 TangentPrev = TangentNext;
+
+	float3 PositionNext = float3(0.0,0.0,0.0);
+	DIHairStrands_ComputeNodePosition(DIContext,RootIndex,PositionNext);
+	float3 PositionPrev = PositionNext;
+
+	for (int EdgeIndex = RootIndex, EdgeEnd = min(RootIndex+DIContext.StrandSize-1,NodeIndex+1); EdgeIndex < EdgeEnd; ++EdgeIndex)
 	{
-		const float Weight = (DIContext.PointsNodesBuffer[PointIndex] == NodeIndex-1) ? 1.0 : 0.0;
-		DIContext.PointsPositionsBuffer[PointIndex] -= Weight * Out_NodePosition * (1.0-DIContext.PointsWeightsBuffer[PointIndex]);
+		PositionPrev = PositionNext;
+		DIHairStrands_ComputeNodePosition(DIContext,EdgeIndex+1,PositionNext);
+
+		TangentPrev = TangentNext;
+		TangentNext = normalize(PositionNext - PositionPrev);
+		NodeQuaternion = NormalizeQuat( MultiplyQuat( FindQuatBetweenNormals(TangentPrev,TangentNext), NodeQuaternion ) );
 	}
-	DeviceMemoryBarrier();
-}*/
+	OutNodeOrientation = NodeQuaternion;
+}
+
+/* -----------------------------------------------------------------
+ * Edge delta position and delta orientation
+ * -----------------------------------------------------------------
+ */
 
-// Compute the strand position by linear interpolation over the points
-/*void DIHairStrands_ReportNodeVelocity(in FDIHairStrandsContext DIContext, in int NodeIndex, in float3 NodeVelocity, in float DeltaTime)
+// Compute the edge length value
+void DIHairStrands_ComputeEdgeLength(in FDIHairStrandsContext DIContext, in int NodeIndex, in float3 NodePosition, out float OutEdgeLength)
 {
-	NodesGlobalPosition[GGroupThreadId.x] = NodeVelocity * DeltaTime;
+	SharedLocalPosition[GGroupThreadId.x] = NodePosition;
 	GroupMemoryBarrier(); 
 
-	if(DIHairStrands_IsRootNode(DIContext))
+	int StrandIndex = 0, LocalIndex = 0;
+	DIHairStrands_ComputeNodeIndices(DIContext,NodeIndex,StrandIndex,LocalIndex);
+
+	// L = ||P1-P0||
+	OutEdgeLength = (LocalIndex>0) ? length(SharedLocalPosition[GGroupThreadId.x]- SharedLocalPosition[GGroupThreadId.x-1]) : 0.0;
+}
+
+// Compute the edge darboux vector (diff between consecutive edge orientations)
+void DIHairStrands_ComputeEdgeRotation(in FDIHairStrandsContext DIContext, in int NodeIndex, in float4 NodeOrientation, out float4 OutEdgeRotation)
+{
+	SharedLocalQuaternion[GGroupThreadId.x] = NodeOrientation;
+	GroupMemoryBarrier(); 
+
+	int StrandIndex = 0, LocalIndex = 0;
+	DIHairStrands_ComputeNodeIndices(DIContext,NodeIndex,StrandIndex,LocalIndex);
+
+	// D = Q0^-1 * Q1
+	if(LocalIndex>0)
+	{
+		float4 q0 = SharedLocalQuaternion[GGroupThreadId.x-1];
+		float4 q1 = SharedLocalQuaternion[GGroupThreadId.x];
+		OutEdgeRotation = float4(
+				q1.xyz * q0.w - q0.xyz * q1.w + cross(-q0.xyz, q1.xyz),
+				q0.w * q1.w - dot(-q0.xyz, q1.xyz));
+		float4 OmegaPlus = OutEdgeRotation + float4(0,0,0,1);
+		float4 OmegaMinus = OutEdgeRotation - float4(0,0,0,1);
+		if( dot(OmegaMinus,OmegaMinus) > dot(OmegaPlus,OmegaPlus) ) OutEdgeRotation = -OutEdgeRotation;
+		//OutEdgeRotation = MultiplyQuat(InverseQuat(SharedLocalQuaternion[GGroupThreadId.x-1]),SharedLocalQuaternion[GGroupThreadId.x]);
+	}
+	else
 	{
-		int StrandIndex = 0;
-		int LocalIndex = 0;
-		DIHairStrands_ComputeNodeIndices(DIContext,NodeIndex,StrandIndex,LocalIndex);
+		OutEdgeRotation = QUATERNION_IDENTITY;
+	}
+}
+
+/* -----------------------------------------------------------------
+ * Points position update
+ * -----------------------------------------------------------------
+ */
 
+// Reset the deformed points position to the rest ones
+void DIHairStrands_ResetPointPosition(in FDIHairStrandsContext DIContext, in int NodeIndex, out bool ReportStatus) 
+{
+	int StrandIndex = 0;
+	int LocalIndex = 0;
+	DIHairStrands_ComputeNodeIndices(DIContext,NodeIndex,StrandIndex,LocalIndex);
+	ReportStatus = false;
+
+	if(LocalIndex == 0)
+	{
 		const int PointBegin = DIContext.CurvesOffsetsBuffer[StrandIndex];
 		const int PointEnd = DIContext.CurvesOffsetsBuffer[StrandIndex+1]; 
 		const int PointCount = PointEnd - PointBegin;
 
+		ReportStatus = true;
 		for (int PointIndex = 0; PointIndex < PointCount; ++PointIndex)
 		{
-			const float PointCoord = (float)(PointIndex) / (PointCount-1.0);
-			const float PointNode = PointCoord * (DIContext.StrandSize-1.0);
-
-			const int NodePrev = (PointIndex==0) ? 0 : (PointIndex==(PointCount-1.0)) ? DIContext.StrandSize-2 : floor(PointNode);
-			const int NodeNext = NodePrev+1;
-
-			const float PointAlpha = PointNode - (float)NodePrev;
-			DIContext.PointsPositionsBuffer[PointBegin+PointIndex].xyz = NodesGlobalPosition[GGroupThreadId.x]
-
-
+			DIContext.PointsPositionsBuffer[PointBegin+PointIndex] = DIContext.RestPositionsBuffer[PointBegin+PointIndex];
 		}
 	}
-}*/
+}
 
-// Compute the strand orientation with parallel transport along the curve
-void DIHairStrands_ComputeNodeOrientation(in FDIHairStrandsContext DIContext, 
-						in float3 NodePosition, out float4 OutNodeQuaternion)
+// Report interpolated nodes displacements onto the points positions 
+void DIHairStrands_UpdatePointPosition(in FDIHairStrandsContext DIContext, in int NodeIndex, in float3 NodeDisplace, out bool ReportStatus) 
 {
-	NodesGlobalPosition[GGroupThreadId.x] = NodePosition;
+	const int ReducedSize =  DIContext.StrandSize-1;
+	SharedLocalPosition[GGroupThreadId.x] = NodeDisplace;
 	GroupMemoryBarrier(); 
 
-	if(DIHairStrands_IsRootNode(DIContext))
+	int StrandIndex = 0, LocalIndex = 0;
+	DIHairStrands_ComputeNodeIndices(DIContext,NodeIndex,StrandIndex,LocalIndex);
+	LocalIndex = LocalIndex - 1;
+
+	ReportStatus = false;
+
+	const int PointBegin = DIContext.CurvesOffsetsBuffer[StrandIndex];
+	const int PointEnd = DIContext.CurvesOffsetsBuffer[StrandIndex+1]; 
+	const int PointCount = PointEnd - PointBegin;
+
+	for (int PointIndex = 0; PointIndex < PointCount; ++PointIndex)
 	{
-		float4 Quaternion = QUATERNION_IDENTITY;
-		float3 TangentNext = float3(0.0,0.0,1.0);
-		float3 TangentPrev = TangentNext;
+		const float PointCoord = (float)(PointIndex) / (PointCount-1.0);
+		const float PointNode = PointCoord * (ReducedSize-1.0);
 
-		NodesGlobalQuaternion[GGroupThreadId.x] = Quaternion;
+		const int NodePrev = (PointIndex==0) ? 0 : (PointIndex==(PointCount-1.0)) ? ReducedSize-2 : floor(PointNode);
+		const int NodeNext = NodePrev+1;
 
-		for (int EdgeIndex = GGroupThreadId.x+1, EdgeEnd = GGroupThreadId.x+DIContext.StrandSize; 
-				EdgeIndex < EdgeEnd; ++EdgeIndex)
+		if(NodePrev == LocalIndex)
 		{
-			TangentNext = normalize(NodesGlobalPosition[EdgeIndex] - NodesGlobalPosition[EdgeIndex-1]);
-			Quaternion = NormalizeQuat( MultiplyQuat( FindQuatBetweenNormals(TangentPrev,TangentNext), Quaternion ) );
-			TangentPrev = TangentNext;
-
-			NodesGlobalQuaternion[EdgeIndex] = Quaternion;
+			ReportStatus = true;
+			const float PointAlpha = PointNode - (float)NodePrev;
+			const float3 LocalDisplace = SharedLocalPosition[GGroupThreadId.x] * (1.0-PointAlpha) + SharedLocalPosition[GGroupThreadId.x+1] * PointAlpha;
+			const float3 DeformedPosition = f16tof32(DIContext.RestPositionsBuffer[PointBegin+PointIndex].xyz) + LocalDisplace;
+			DIContext.PointsPositionsBuffer[PointBegin+PointIndex] = uint4( f32tof16(DeformedPosition), DIContext.RestPositionsBuffer[PointBegin+PointIndex].w);
 		}
 	}
-	GroupMemoryBarrier();
-	OutNodeQuaternion = NodesGlobalQuaternion[GGroupThreadId.x];
 }
 
-// Compute the node mass
-void DIHairStrands_ComputeNodeMass(in FDIHairStrandsContext DIContext, in int NodeIndex, out float3 OutNodeMass)
+// Get the point position
+void DIHairStrands_GetPointPosition(in FDIHairStrandsContext DIContext, in int PointIndex, out float3 OutPointPosition)
 {
-	const float CoordScale = 1.0 / (DIContext.StrandSize-1.0);
+	OutPointPosition = mul(float4( f16tof32(DIContext.PointsPositionsBuffer[PointIndex].xyz), 1.0), DIContext.WorldTransform).xyz;
+}
 
-	int StrandIndex = 0, LocalIndex = 0;
-	DIHairStrands_ComputeNodeIndices(DIContext,NodeIndex,StrandIndex,LocalIndex);
-	const float EdgeLength = DIHairStrands_ComputeStrandLength(DIContext,StrandIndex) * CoordScale;
+/* -----------------------------------------------------------------
+ * Nodes time integration
+ * -----------------------------------------------------------------
+ */
 
-	const bool IsFirst = (LocalIndex == 0);
-	const bool IsLast  = (LocalIndex == (DIContext.StrandSize-1));
+// Add external force to the linear velocity and advect node position
+void DIHairStrands_AdvectNodePosition(in FDIHairStrandsContext DIContext, in float NodeMass, in float InverssMass, in float3 ExternalForce, in float DeltaTime, inout float3 OutLinearVelocity, inout float3 OutNodePosition)
+{	
+	OutLinearVelocity += DeltaTime * InverssMass * ExternalForce;
 
-	const float EdgeScale = (IsFirst || IsLast) ? 0.5 : 1.0;
-	
-	const float PrevCoord = IsFirst ? 0.0 : (LocalIndex-0.5) * CoordScale;
-	const float NextCoord = IsLast ? 1.0  : (LocalIndex+0.5) * CoordScale;
+	OutNodePosition += OutLinearVelocity * DeltaTime;
+}
 
-	const float PrevThick = DIContext.RootThickness * ( 1.0 - PrevCoord ) + DIContext.TipThickness * PrevCoord;
-	const float NextThick = DIContext.RootThickness * ( 1.0 - NextCoord ) + DIContext.TipThickness * NextCoord;
+// Add external torque to the angular velocity and advect node orientation
+void DIHairStrands_AdvectNodeOrientation(in FDIHairStrandsContext DIContext, in float3 NodeInertia, in float3 InverseInertia, in float3 ExternalTorque, in float DeltaTime, inout float3 OutAngularVelocity, inout float4 OutNodeOrientation)
+{	
+	OutAngularVelocity += DeltaTime * InverseInertia * (ExternalTorque - cross(OutAngularVelocity, NodeInertia * OutAngularVelocity));
 
-	OutNodeMass = DIContext.StrandDensity * EdgeLength * EdgeScale * M_PI * (PrevThick * PrevThick + NextThick * PrevThick + NextThick * NextThick) / 12.0;
+	OutNodeOrientation = OutNodeOrientation + 0.5 * DeltaTime * float4(
+        OutAngularVelocity.xyz * OutNodeOrientation.w + cross(OutAngularVelocity.xyz, OutNodeOrientation.xyz),
+        - dot(OutAngularVelocity.xyz, OutNodeOrientation.xyz));
+
+	OutNodeOrientation = NormalizeQuat(OutNodeOrientation);
 }
 
-// Compute the node position by linear interpolation over the points
-void DIHairStrands_ComputeNodePosition(in FDIHairStrandsContext DIContext, in int NodeIndex, out float3 OutNodePosition)
+// Update the node linear velocity based on the node position difference
+void DIHairStrands_UpdateLinearVelocity(in FDIHairStrandsContext DIContext, in float3 PreviousPosition, in float3 NodePosition, in float DeltaTime, inout float3 OutLinearVelocity)
+{	
+	OutLinearVelocity = (NodePosition-PreviousPosition) / DeltaTime;
+}
+
+// Update the node angular velocity based on the node orientation difference
+void DIHairStrands_UpdateAngularVelocity(in FDIHairStrandsContext DIContext, in float4 PreviousOrientation, in float4 NodeOrientation, in float DeltaTime, inout float3 OutAngularVelocity)
 {
-	int StrandIndex = 0, LocalIndex = 0;
-	DIHairStrands_ComputeNodeIndices(DIContext,NodeIndex,StrandIndex,LocalIndex);
+	const float4 DeltaQuat = float4(
+				NodeOrientation.xyz * PreviousOrientation.w - PreviousOrientation.xyz * NodeOrientation.w - cross(PreviousOrientation.xyz, NodeOrientation.xyz),
+				PreviousOrientation.w * NodeOrientation.w + dot(PreviousOrientation.xyz, NodeOrientation.xyz));
+	const float AxisLength = length( DeltaQuat.xyz );
 
-	const int PointOffset = DIContext.CurvesOffsetsBuffer[StrandIndex];
+	if (AxisLength<SMALL_NUMBER)
+	{
+		 OutAngularVelocity =  DeltaQuat.xyz * 2.0 / DeltaTime;
+	}
+	else
+	{
+		const float QuatAngle = 2.0 * atan2(AxisLength,DeltaQuat.w );
+		OutAngularVelocity = DeltaQuat.xyz * QuatAngle / (AxisLength*DeltaTime);
+	}
+}
 
-	const float NodeCoord = (float)(LocalIndex) / (DIContext.StrandSize-1.0);
-	const float PointCount = DIContext.CurvesOffsetsBuffer[StrandIndex+1]-PointOffset-1;
+/* -----------------------------------------------------------------
+ * Update node position orientation
+ * -----------------------------------------------------------------
+ */
 
-	const float NodePoint = NodeCoord * PointCount;
-	const int PointPrev = (LocalIndex==0) ? 0 : (LocalIndex==(DIContext.StrandSize-1)) ? PointCount-1 : floor(NodePoint);
-	const int PointNext = PointPrev+1;
+// Compute rest position
+void DIHairStrands_ComputeRestPosition(in FDIHairStrandsContext DIContext, in float3 NodePosition, out float3 OutRestPosition)
+{	
+	OutRestPosition = mul(float4( NodePosition, 1.0), DIContext.WorldInverse).xyz;
+}
 
-	const float NodeAlpha = NodePoint - (float)PointPrev;
-	OutNodePosition =  DIContext.PointsPositionsBuffer[PointOffset+PointPrev].xyz * (1.0-NodeAlpha) + 
-					   DIContext.PointsPositionsBuffer[PointOffset+PointNext].xyz * NodeAlpha;
-	OutNodePosition = mul(float4( OutNodePosition, 1.0), DIContext.WorldTransform).xyz;
+// Compute rest orientation
+void DIHairStrands_ComputeRestOrientation(in FDIHairStrandsContext DIContext, in float4 NodeOrientation, out float4 OutRestOrientation)
+{	
+	OutRestOrientation = NormalizeQuat( MultiplyQuat(InverseQuat(DIContext.WorldRotation),NodeOrientation) );
+}
+
+// Update Node Position
+void DIHairStrands_AttachNodePosition(in FDIHairStrandsContext DIContext, in float3 RestPosition, out float3 OutNodePosition)
+{
+	OutNodePosition = mul(float4( RestPosition, 1.0), DIContext.WorldTransform).xyz;
 }
 
+// Update Node Orientation
+void DIHairStrands_AttachNodeOrientation(in FDIHairStrandsContext DIContext, in float4 RestOrientation, out float4 OutNodeOrientation)
+{
+	OutNodeOrientation = NormalizeQuat( MultiplyQuat(DIContext.WorldRotation,RestOrientation) );
+}
+
+
+
+
 
 
 
diff --git a/Engine/Plugins/Experimental/HairStrands/Source/HairStrandsCore/HairStrandsCore.Build.cs b/Engine/Plugins/Experimental/HairStrands/Source/HairStrandsCore/HairStrandsCore.Build.cs
--- a/Engine/Plugins/Experimental/HairStrands/Source/HairStrandsCore/HairStrandsCore.Build.cs
+++ b/Engine/Plugins/Experimental/HairStrands/Source/HairStrandsCore/HairStrandsCore.Build.cs
@@ -4,23 +4,25 @@ namespace UnrealBuildTool.Rules
 {
 	public class HairStrandsCore : ModuleRules
 	{
-        public HairStrandsCore(ReadOnlyTargetRules Target) : base(Target)
+		public HairStrandsCore(ReadOnlyTargetRules Target) : base(Target)
 		{
 			PrivateIncludePaths.Add(ModuleDirectory + "/Private");
-            PublicIncludePaths.Add(ModuleDirectory + "/Public");
+			PublicIncludePaths.Add(ModuleDirectory + "/Public");
 
-            PrivateDependencyModuleNames.AddRange(
-                new string[]
-                {
+			PrivateDependencyModuleNames.AddRange(
+				new string[]
+				{
 					"Core",
 					"CoreUObject",
 					"Engine",
-                    "RenderCore",
-                    "Renderer",
-                    "RHI",
+					"MeshDescription",
+					"Niagara",
+					"RenderCore",
+					"Renderer",
+					"RHI",
 					"ChaosCore",
-                    "Chaos"
-                });
-        }
+					"Chaos"
+				});
+		}
 	}
 }
diff --git a/Engine/Plugins/Experimental/HairStrands/Source/HairStrandsCore/Private/GroomAsset.cpp b/Engine/Plugins/Experimental/HairStrands/Source/HairStrandsCore/Private/GroomAsset.cpp
new file mode 100644
--- /dev/null
+++ b/Engine/Plugins/Experimental/HairStrands/Source/HairStrandsCore/Private/GroomAsset.cpp
@@ -0,0 +1,289 @@
+// Copyright 1998-2019 Epic Games, Inc. All Rights Reserved.
+
+#include "GroomAsset.h"
+#include "RenderingThread.h"
+#include "HairStrandsVertexFactory.h"
+#include "Misc/Paths.h"
+#include "UObject/PhysicsObjectVersion.h"
+
+#if WITH_EDITORONLY_DATA
+#include "EditorFramework/AssetImportData.h"
+#endif
+
+template<typename FormatType>
+void CreateBuffer(const TArray<typename FormatType::Type>& InData, FRWBuffer& OutBuffer)
+{
+	const uint32 DataCount = InData.Num();
+	const uint32 DataSizeInBytes = FormatType::SizeInByte*DataCount;
+
+	if (DataSizeInBytes == 0) return;
+
+	OutBuffer.Initialize(FormatType::SizeInByte, DataCount, FormatType::Format, BUF_Static);
+	void* BufferData = RHILockVertexBuffer(OutBuffer.Buffer, 0, DataSizeInBytes, RLM_WriteOnly);
+
+	FMemory::Memcpy(BufferData, InData.GetData(), DataSizeInBytes);
+	RHIUnlockVertexBuffer(OutBuffer.Buffer);
+}
+
+template<typename FormatType>
+void CreateBuffer(uint32 InVertexCount , FRWBuffer& OutBuffer)
+{
+	const uint32 DataCount = InVertexCount;
+	const uint32 DataSizeInBytes = FormatType::SizeInByte*DataCount;
+
+	if (DataSizeInBytes == 0) return;
+
+	OutBuffer.Initialize(FormatType::SizeInByte, DataCount, FormatType::Format, BUF_Static);
+	void* BufferData = RHILockVertexBuffer(OutBuffer.Buffer, 0, DataSizeInBytes, RLM_WriteOnly);
+	FMemory::Memset(BufferData, 0, DataSizeInBytes);
+	RHIUnlockVertexBuffer(OutBuffer.Buffer);
+}
+
+/////////////////////////////////////////////////////////////////////////////////////////
+
+FHairStrandsResource::FHairStrandsResource(FHairStrandsDatas* HairStrandsDatas) : 
+	RestPositionBuffer(), TangentBuffer(), AttributeBuffer(), StrandsDatas(HairStrandsDatas)
+{}
+
+void FHairStrandsResource::InitRHI()
+{
+	if (StrandsDatas != nullptr)
+	{
+		TArray<FHairStrandsPositionFormat::Type> RenderingPositions;
+		TArray<FHairStrandsAttributeFormat::Type> RenderingAttributes;
+		StrandsDatas->BuildRenderingDatas(RenderingPositions, RenderingAttributes);
+
+		CreateBuffer<FHairStrandsPositionFormat>(RenderingPositions, RestPositionBuffer);
+		CreateBuffer<FHairStrandsPositionFormat>(RenderingPositions, DeformedPositionBuffer[0]);
+		CreateBuffer<FHairStrandsPositionFormat>(RenderingPositions, DeformedPositionBuffer[1]);
+		CreateBuffer<FHairStrandsTangentFormat>(RenderingPositions.Num() * FHairStrandsTangentFormat::ComponentCount, TangentBuffer);
+		CreateBuffer<FHairStrandsAttributeFormat>(RenderingAttributes, AttributeBuffer);
+	}
+}
+
+void FHairStrandsResource::ReleaseRHI()
+{
+	RestPositionBuffer.Release();
+	DeformedPositionBuffer[0].Release();
+	DeformedPositionBuffer[1].Release();
+	TangentBuffer.Release();
+	AttributeBuffer.Release();
+}
+
+/////////////////////////////////////////////////////////////////////////////////////////
+
+FHairStrandsInterpolationResource::FHairStrandsInterpolationResource(const FHairStrandsInterpolationDatas& InData, const FHairStrandsDatas& SimDatas) :
+	Interpolation0Buffer(), Interpolation1Buffer()
+{
+	check(Interpolation0.Num()==0);
+	check(Interpolation1.Num()==0);
+	InData.BuildRenderingDatas(Interpolation0, Interpolation1);
+
+	const uint32 RootCount = SimDatas.GetNumCurves();
+	SimRootPointIndex.SetNum(SimDatas.GetNumPoints());
+	for (uint32 CurveIndex = 0; CurveIndex < RootCount; ++CurveIndex)
+	{
+		const uint16 PointCount = SimDatas.StrandsCurves.CurvesCount[CurveIndex];
+		const uint32 PointOffset = SimDatas.StrandsCurves.CurvesOffset[CurveIndex];
+		for (uint32 PointIndex = 0; PointIndex < PointCount; ++PointIndex)
+		{
+			SimRootPointIndex[PointIndex + PointOffset] = PointOffset;
+		}
+	}
+}
+
+void FHairStrandsInterpolationResource::InitRHI()
+{
+	check(Interpolation0.Num() == Interpolation1.Num());
+	if (Interpolation0.Num() > 0)
+	{
+		CreateBuffer<FHairStrandsInterpolation0Format>(Interpolation0, Interpolation0Buffer);
+		CreateBuffer<FHairStrandsInterpolation1Format>(Interpolation1, Interpolation1Buffer);
+		CreateBuffer<FHairStrandsRootIndexFormat>(SimRootPointIndex, SimRootPointIndexBuffer);
+		Interpolation0.SetNum(0);
+		Interpolation1.SetNum(0);
+		SimRootPointIndex.SetNum(0);
+	}
+}
+
+void FHairStrandsInterpolationResource::ReleaseRHI()
+{
+	Interpolation0Buffer.Release();
+	Interpolation1Buffer.Release();
+}
+
+/////////////////////////////////////////////////////////////////////////////////////////
+
+#if RHI_RAYTRACING
+// RT geometry is built to for a cross around the fiber.
+// 4 triangles per hair vertex => 12 vertices per hair vertex
+FHairStrandsRaytracingResource::FHairStrandsRaytracingResource(FHairStrandsDatas* InData) :
+	PositionBuffer(), VertexCount(InData->GetNumPoints()*12)  
+{}
+
+void FHairStrandsRaytracingResource::InitRHI()
+{
+	check(IsInRenderingThread());
+	CreateBuffer<FHairStrandsRaytracingFormat>(VertexCount, PositionBuffer);
+}
+
+void FHairStrandsRaytracingResource::ReleaseRHI()
+{
+	PositionBuffer.Release();
+	RayTracingGeometry.ReleaseResource();
+}
+#endif
+
+/////////////////////////////////////////////////////////////////////////////////////////
+
+void UGroomAsset::Serialize(FArchive& Ar)
+{
+	Super::Serialize(Ar);
+
+	Ar.UsingCustomVersion(FPhysicsObjectVersion::GUID);
+
+	// #ueent_todo: Serialize HairDescription
+
+	HairRenderData.Serialize(Ar);
+
+	if (Ar.CustomVer(FPhysicsObjectVersion::GUID) >= FPhysicsObjectVersion::HairAssetSerialization)
+	{
+		HairSimulationData.Serialize(Ar);
+	}
+}
+
+void UGroomAsset::InitResource()
+{
+	check(!HairStrandsResource);
+
+	HairStrandsResource = new FHairStrandsResource(&HairRenderData);
+
+	BeginInitResource(HairStrandsResource);
+
+	HairSimulationResource = new FHairStrandsResource(&HairSimulationData);
+
+	BeginInitResource(HairSimulationResource);
+}
+
+void UGroomAsset::UpdateResource()
+{
+	if (HairStrandsResource)
+	{
+		BeginUpdateResourceRHI(HairStrandsResource);
+	}
+
+	if (HairSimulationResource)
+	{
+		BeginUpdateResourceRHI(HairSimulationResource);
+	}
+}
+
+void UGroomAsset::ReleaseResource()
+{
+	if (HairStrandsResource)
+	{
+		FHairStrandsResource* InResource = HairStrandsResource;
+		ENQUEUE_RENDER_COMMAND(ReleaseHairStrandsResourceCommand)(
+			[InResource](FRHICommandList& RHICmdList)
+		{
+			InResource->ReleaseResource();
+			delete InResource;
+		});
+		HairStrandsResource = nullptr;
+	}
+
+	if (HairSimulationResource)
+	{
+		FHairStrandsResource* InResource = HairSimulationResource;
+		ENQUEUE_RENDER_COMMAND(ReleaseHairStrandsResourceCommand)(
+			[InResource](FRHICommandList& RHICmdList)
+		{
+			InResource->ReleaseResource();
+			delete InResource;
+		});
+		HairSimulationResource = nullptr;
+	}
+}
+
+void UGroomAsset::Reset()
+{
+	ReleaseResource();
+
+	HairRenderData.Reset();
+	HairSimulationData.Reset();
+
+	RenderHairGroups.Reset();
+	SimulationHairGroups.Reset();
+}
+
+void UGroomAsset::PostLoad()
+{
+	Super::PostLoad();
+
+	// If the asset was previously serialized without hair sim data, generate it here
+	bool bAutoGeneratedGuides = false;
+	if (HairSimulationData.StrandsPoints.Num() == 0 || HairSimulationData.StrandsCurves.Num() == 0)
+	{
+		DecimateStrandData(HairRenderData, FMath::Clamp(HairToGuideDensity, 0.f, 1.f), HairSimulationData);
+		bAutoGeneratedGuides = true;
+	}
+
+	// If the asset was previously serialized without group settings, fill them out here
+	if (RenderHairGroups.Num() == 0)
+	{
+		FHairGroupRenderSettings& GroupSettings = RenderHairGroups.AddDefaulted_GetRef();
+		GroupSettings.GroupID = 0;
+		GroupSettings.NumCurves = HairRenderData.GetNumCurves();
+	}
+
+	if (SimulationHairGroups.Num() == 0)
+	{
+		FHairGroupSimulationSettings& GroupSettings = SimulationHairGroups.AddDefaulted_GetRef();
+		GroupSettings.GroupID = 0;
+		GroupSettings.NumCurves = HairSimulationData.GetNumCurves();
+		GroupSettings.bIsAutoGenerated = bAutoGeneratedGuides;
+	}
+
+	if (!IsTemplate())
+	{
+		InitResource();
+	}
+}
+
+void UGroomAsset::BeginDestroy()
+{
+	ReleaseResource();
+	Super::BeginDestroy();
+}
+
+#if WITH_EDITOR
+void UGroomAsset::PostEditChangeProperty(FPropertyChangedEvent& PropertyChangedEvent)
+{
+	Super::PostEditChangeProperty(PropertyChangedEvent);
+	UpdateResource();
+}
+#endif // WITH_EDITOR
+
+#if WITH_EDITORONLY_DATA
+void UGroomAsset::GetAssetRegistryTags(TArray<FAssetRegistryTag>& OutTags) const
+{
+	if (AssetImportData)
+	{
+		OutTags.Add(FAssetRegistryTag(SourceFileTagName(), AssetImportData->GetSourceData().ToJson(), FAssetRegistryTag::TT_Hidden));
+	}
+
+	Super::GetAssetRegistryTags(OutTags);
+}
+
+void UGroomAsset::PostInitProperties()
+{
+	if (!HasAnyFlags(RF_ClassDefaultObject))
+	{
+		AssetImportData = NewObject<UAssetImportData>(this, TEXT("AssetImportData"));
+	}
+
+	Super::PostInitProperties();
+}
+#endif
+
+
