Copying //UE4/Dev-Physics to Dev-Main (//UE4/Dev-Main) [at] 9098278

diff --git a/Engine/Plugins/Experimental/HairStrands/Source/HairStrandsCore/Private/GroomComponent.cpp b/Engine/Plugins/Experimental/HairStrands/Source/HairStrandsCore/Private/GroomComponent.cpp
new file mode 100644
--- /dev/null
+++ b/Engine/Plugins/Experimental/HairStrands/Source/HairStrandsCore/Private/GroomComponent.cpp
@@ -0,0 +1,639 @@
+// Copyright 1998-2019 Epic Games, Inc. All Rights Reserved. 
+
+#include "GroomComponent.h"
+#include "Materials/Material.h"
+#include "MaterialShared.h"
+#include "Engine/CollisionProfile.h"
+#include "Engine/Engine.h"
+#include "PrimitiveSceneProxy.h"
+#include "HairStrandsRendering.h"
+#include "RayTracingInstanceUtils.h"
+#include "HairStrandsInterface.h"
+#include "UObject/UObjectIterator.h"
+
+static float GHairClipLength = -1;
+static FAutoConsoleVariableRef CVarHairClipLength(TEXT("r.HairStrands.DebugClipLength"), GHairClipLength, TEXT("Clip hair strands which have a lenth larger than this value. (default is -1, no effect)"));
+float GetHairClipLength() { return GHairClipLength > 0 ? GHairClipLength : 100000;  }
+
+/**
+ * An material render proxy which overrides the debug mode parameter.
+ */
+class ENGINE_VTABLE FHairDebugModeMaterialRenderProxy : public FMaterialRenderProxy
+{
+public:
+	const FMaterialRenderProxy* const Parent;
+	const float DebugMode;
+	const float HairMinRadius;
+	const float HairMaxRadius;
+	const float HairClipLength;
+
+	FName DebugModeParamName;
+	FName MinHairRadiusParamName;
+	FName MaxHairRadiusParamName;
+	FName HairClipLengthParamName;
+
+	/** Initialization constructor. */
+	FHairDebugModeMaterialRenderProxy(const FMaterialRenderProxy* InParent, float InMode, float InMinRadius, float InMaxRadius, float InHairClipLength) :
+		Parent(InParent),
+		DebugMode(InMode),
+		HairMinRadius(InMinRadius),
+		HairMaxRadius(InMaxRadius),
+		HairClipLength(InHairClipLength),
+		DebugModeParamName(NAME_FloatProperty),
+		MinHairRadiusParamName(NAME_ByteProperty),
+		MaxHairRadiusParamName(NAME_IntProperty),
+		HairClipLengthParamName(NAME_BoolProperty)
+	{}
+
+	virtual const FMaterial& GetMaterialWithFallback(ERHIFeatureLevel::Type InFeatureLevel, const FMaterialRenderProxy*& OutFallbackMaterialRenderProxy) const
+	{
+		return Parent->GetMaterialWithFallback(InFeatureLevel, OutFallbackMaterialRenderProxy);
+	}
+
+	virtual bool GetVectorValue(const FMaterialParameterInfo& ParameterInfo, FLinearColor* OutValue, const FMaterialRenderContext& Context) const
+	{
+		return Parent->GetVectorValue(ParameterInfo, OutValue, Context);
+	}
+
+	virtual bool GetScalarValue(const FMaterialParameterInfo& ParameterInfo, float* OutValue, const FMaterialRenderContext& Context) const
+	{
+		if (ParameterInfo.Name == DebugModeParamName)
+		{
+			*OutValue = DebugMode;
+			return true;
+		}
+		else if (ParameterInfo.Name == MinHairRadiusParamName)
+		{
+			*OutValue = HairMinRadius;
+			return true;
+		}
+		else if (ParameterInfo.Name == MaxHairRadiusParamName)
+		{
+			*OutValue = HairMaxRadius;
+			return true;
+		}
+		else if (ParameterInfo.Name == HairClipLengthParamName)
+		{
+			*OutValue = HairClipLength;
+			return true;
+		}
+		else
+		{
+			return Parent->GetScalarValue(ParameterInfo, OutValue, Context);
+		}
+	}
+
+	virtual bool GetTextureValue(const FMaterialParameterInfo& ParameterInfo, const URuntimeVirtualTexture** OutValue, const FMaterialRenderContext& Context) const
+	{
+		return Parent->GetTextureValue(ParameterInfo, OutValue, Context);
+	}
+
+	virtual bool GetTextureValue(const FMaterialParameterInfo& ParameterInfo, const UTexture** OutValue, const FMaterialRenderContext& Context) const
+	{
+		return Parent->GetTextureValue(ParameterInfo, OutValue, Context);
+	}
+};
+
+
+/////////////////////////////////////////////////////////////////////////////////////////
+/////////////////////////////////////////////////////////////////////////////////////////
+//  FStrandHairSceneProxy
+/////////////////////////////////////////////////////////////////////////////////////////
+/////////////////////////////////////////////////////////////////////////////////////////
+
+class FHairStrandsSceneProxy final : public FPrimitiveSceneProxy
+{
+public:
+	SIZE_T GetTypeHash() const override
+	{
+		static size_t UniquePointer;
+		return reinterpret_cast<size_t>(&UniquePointer);
+	}
+
+	FHairStrandsSceneProxy(UGroomComponent* Component)
+		: FPrimitiveSceneProxy(Component)
+		, VertexFactory(GetScene().GetFeatureLevel(), "FStrandHairSceneProxy")
+		, MaterialRelevance(Component->GetMaterialRelevance(GetScene().GetFeatureLevel()))
+	
+	{
+		check(Component);
+		check(Component->GroomAsset);
+		
+		const uint32 VertexCount = Component->GroomAsset->HairRenderData.GetNumPoints();
+		const float MinHairRadius = 0; // Todo: Component->GroomAsset->HairRenderData.StrandsCurves.MinRadius;
+		const float MaxHairRadius = Component->GroomAsset->HairRenderData.StrandsCurves.MaxRadius;
+		const float MaxHairLength = Component->GroomAsset->HairRenderData.StrandsCurves.MaxLength;
+		const float HairDensity = Component->GroomAsset->HairRenderData.HairDensity;
+		const FVector& HairWorldOffset = Component->GroomAsset->HairRenderData.BoundingBox.GetCenter();
+
+		check(Component->InterpolationOutput);
+		FHairStrandsInterpolationOutput* LocalOutput = Component->InterpolationOutput;
+		FHairStrandsVertexFactory* LocaVertexFactor = &VertexFactory;
+		ENQUEUE_RENDER_COMMAND(StaticMeshVertexBuffersLegacyInit)(
+		[LocalOutput, LocaVertexFactor, HairWorldOffset, MinHairRadius, MaxHairRadius, MaxHairLength, HairDensity](FRHICommandListImmediate& RHICmdList)
+		{
+			const uint32 Offset = 0;
+			FHairStrandsVertexFactory::FDataType Data;
+			Data.MinStrandRadius = MinHairRadius;
+			Data.MaxStrandRadius = MaxHairRadius;
+			Data.MaxStrandLength = MaxHairLength;
+			Data.HairDensity = HairDensity;
+			Data.HairWorldOffset = HairWorldOffset;
+			Data.InterpolationOutput = LocalOutput;
+
+			LocaVertexFactor->SetData(Data);
+			LocaVertexFactor->InitResource();
+		});
+
+
+		Material = Component->GetMaterial(0);
+		if (Material == nullptr || !Material->GetMaterialResource(GetScene().GetFeatureLevel())->IsUsedWithHairStrands())
+		{
+			Material = GEngine->HairDefaultMaterial;
+		}
+
+		#if RHI_RAYTRACING
+		if (IsRayTracingEnabled())
+			RayTracingGeometry = Component->RaytracingResources ? &Component->RaytracingResources->RayTracingGeometry : nullptr;
+		#endif
+	}
+
+	virtual ~FHairStrandsSceneProxy()
+	{
+		VertexFactory.ReleaseResource();
+	}
+	
+#if RHI_RAYTRACING
+	virtual bool IsRayTracingRelevant() const { return true; }
+	virtual bool IsRayTracingStaticRelevant() const { return false; }
+
+	virtual void GetDynamicRayTracingInstances(FRayTracingMaterialGatheringContext & Context, TArray<struct FRayTracingInstance>& OutRayTracingInstances) override
+	{
+		if (!IsRayTracingEnabled() || !RayTracingGeometry || !RayTracingGeometry->RayTracingGeometryRHI.IsValid())
+			return;
+
+		check(RayTracingGeometry->Initializer.PositionVertexBuffer.IsValid());
+		AddOpaqueRaytracingInstance(GetLocalToWorld(), RayTracingGeometry, RaytracingInstanceMask_ThinShadow, OutRayTracingInstances);
+	}
+#endif
+
+	virtual void GetDynamicMeshElements(const TArray<const FSceneView*>& Views, const FSceneViewFamily& ViewFamily, uint32 VisibilityMap, FMeshElementCollector& Collector) const override
+	{
+		const uint32 HairVertexCount = VertexFactory.GetData().InterpolationOutput ? VertexFactory.GetData().InterpolationOutput->VFInput.VertexCount : 0;
+		if (HairVertexCount == 0) return;
+
+		QUICK_SCOPE_CYCLE_COUNTER(STAT_HairStrandsSceneProxy_GetDynamicMeshElements);
+
+		FMaterialRenderProxy* MaterialProxy = nullptr;
+		{
+			const EHairStrandsDebugMode DebugMode = GetHairStrandsDebugMode();
+			if (DebugMode != EHairStrandsDebugMode::None)
+			{
+				float DebugModeScalar = 0;
+				switch(DebugMode)
+				{
+				case EHairStrandsDebugMode::None						: DebugModeScalar =99.f; break;
+				case EHairStrandsDebugMode::SimHairStrands				: DebugModeScalar = 0.f; break;
+				case EHairStrandsDebugMode::RenderHairStrands			: DebugModeScalar = 0.f; break;
+				case EHairStrandsDebugMode::RenderHairRootUV			: DebugModeScalar = 1.f; break;
+				case EHairStrandsDebugMode::RenderHairUV				: DebugModeScalar = 2.f; break;
+				case EHairStrandsDebugMode::RenderHairSeed				: DebugModeScalar = 3.f; break;
+				case EHairStrandsDebugMode::RenderHairDimension			: DebugModeScalar = 4.f; break;
+				case EHairStrandsDebugMode::RenderHairRadiusVariation	: DebugModeScalar = 5.f; break;
+				};
+
+				const float HairMinRadius = VertexFactory.GetMinStrandRadius();
+				const float HairMaxRadius = VertexFactory.GetMaxStrandRadius();
+				const float HairClipLength = GetHairClipLength();
+				auto DebugMaterial = new FHairDebugModeMaterialRenderProxy(
+					GEngine->HairDebugMaterial ? GEngine->HairDebugMaterial->GetRenderProxy() : nullptr,
+					DebugModeScalar, HairMinRadius, HairMaxRadius, HairClipLength);
+				Collector.RegisterOneFrameMaterialProxy(DebugMaterial);
+				MaterialProxy = DebugMaterial;
+			}
+			else
+			{
+				MaterialProxy = Material->GetRenderProxy();
+			}
+		}
+
+		for (int32 ViewIndex = 0; ViewIndex < Views.Num(); ViewIndex++)
+		{
+			if (VisibilityMap & (1 << ViewIndex))
+			{
+				const FSceneView* View = Views[ViewIndex];
+				// Draw the mesh.
+				FMeshBatch& Mesh = Collector.AllocateMesh();
+				FMeshBatchElement& BatchElement = Mesh.Elements[0];
+				BatchElement.IndexBuffer = nullptr;
+				Mesh.bWireframe = false;
+				Mesh.VertexFactory = &VertexFactory;
+				Mesh.MaterialRenderProxy = MaterialProxy;
+				bool bHasPrecomputedVolumetricLightmap;
+				FMatrix PreviousLocalToWorld;
+				int32 SingleCaptureIndex;
+				bool bOutputVelocity = false;
+				bool bDrawVelocity = false; // Velocity vector is done in a custom fashion
+				GetScene().GetPrimitiveUniformShaderParameters_RenderThread(GetPrimitiveSceneInfo(), bHasPrecomputedVolumetricLightmap, PreviousLocalToWorld, SingleCaptureIndex, bOutputVelocity);
+				bOutputVelocity = false;
+				FDynamicPrimitiveUniformBuffer& DynamicPrimitiveUniformBuffer = Collector.AllocateOneFrameResource<FDynamicPrimitiveUniformBuffer>();
+				DynamicPrimitiveUniformBuffer.Set(GetLocalToWorld(), PreviousLocalToWorld, GetBounds(), GetLocalBounds(), true, false, bDrawVelocity, bOutputVelocity);
+				BatchElement.PrimitiveUniformBufferResource = &DynamicPrimitiveUniformBuffer.UniformBuffer;
+				BatchElement.FirstIndex = 0;
+				BatchElement.NumPrimitives = HairVertexCount * 2;
+				BatchElement.MinVertexIndex = 0;
+				BatchElement.MaxVertexIndex = HairVertexCount * 6;
+				Mesh.ReverseCulling = IsLocalToWorldDeterminantNegative();
+				Mesh.Type = PT_TriangleList;
+				Mesh.DepthPriorityGroup = SDPG_World;
+				Mesh.bCanApplyViewModeOverrides = false;
+				Collector.AddMesh(ViewIndex, Mesh);
+
+			#if !(UE_BUILD_SHIPPING || UE_BUILD_TEST)
+				// Render bounds
+				RenderBounds(Collector.GetPDI(ViewIndex), ViewFamily.EngineShowFlags, GetBounds(), IsSelected());
+			#endif
+			}
+		}
+	}
+
+	virtual FPrimitiveViewRelevance GetViewRelevance(const FSceneView* View) const override
+	{
+		FPrimitiveViewRelevance Result;
+		Result.bHairStrandsRelevance = true;
+
+		// Special pass for hair strands geometry (not part of the base pass, and shadowing is handlded in a custom fashion)
+		Result.bDrawRelevance = false;		
+		Result.bShadowRelevance = false;
+		Result.bRenderInMainPass = false;
+		Result.bDynamicRelevance = true;
+
+		// Selection only
+		#if WITH_EDITOR
+		{
+			const bool bIsSelected = (IsSelected() || IsHovered());
+			Result.bEditorStaticSelectionRelevance = bIsSelected;
+			Result.bDrawRelevance = bIsSelected;
+		}
+		#endif
+		MaterialRelevance.SetPrimitiveViewRelevance(Result);
+		return Result;
+	}
+
+	virtual uint32 GetMemoryFootprint(void) const override { return(sizeof(*this) + GetAllocatedSize()); }
+
+	uint32 GetAllocatedSize(void) const { return(FPrimitiveSceneProxy::GetAllocatedSize()); }
+private:
+	UMaterialInterface* Material = nullptr;
+	FHairStrandsVertexFactory VertexFactory;
+	FMaterialRelevance MaterialRelevance;
+#if RHI_RAYTRACING
+	FRayTracingGeometry* RayTracingGeometry = nullptr;
+#endif
+};
+
+
+
+/////////////////////////////////////////////////////////////////////////////////////////
+/////////////////////////////////////////////////////////////////////////////////////////
+// UComponent
+//////////////////////////////////////////////////////////////////////////
+//////////////////////////////////////////////////////////////////////////
+
+UGroomComponent::UGroomComponent(const FObjectInitializer& ObjectInitializer)
+	: Super(ObjectInitializer)
+{
+	PrimaryComponentTick.bCanEverTick = true;
+	bTickInEditor = true;
+	bAutoActivate = true;
+	bSelectable = true;
+	HairDensity = 1;
+	MergeThreshold = 0.1f;
+	InitializedResources = nullptr;
+	Mobility = EComponentMobility::Movable;
+	if (GEngine)
+	{
+		SetMaterial(0, GEngine->HairDefaultMaterial);
+	}
+
+	SetCollisionProfileName(UCollisionProfile::PhysicsActor_ProfileName);
+}
+
+FPrimitiveSceneProxy* UGroomComponent::CreateSceneProxy()
+{
+	if (!GroomAsset || GroomAsset->HairRenderData.GetNumCurves() == 0 || !InterpolationOutput || !InterpolationInput)
+		return nullptr;
+
+	return new FHairStrandsSceneProxy(this);
+}
+
+
+FBoxSphereBounds UGroomComponent::CalcBounds(const FTransform& LocalToWorld) const
+{
+	FBox HairBox(ForceInit);
+	if (GroomAsset)
+	{
+		return FBoxSphereBounds(GroomAsset->HairRenderData.BoundingBox.TransformBy(LocalToWorld));
+	}
+	else
+	{
+		return FBoxSphereBounds();
+	}
+}
+
+int32 UGroomComponent::GetNumMaterials() const
+{
+	if (GroomAsset)
+	{
+		int32 NumGroups = GroomAsset->RenderHairGroups.Num();
+		return NumGroups > 0 ? NumGroups : 1;
+	}
+	return 1;
+}
+
+UMaterialInterface* UGroomComponent::GetMaterial(int32 ElementIndex) const
+{
+	UMaterialInterface* OverrideMaterial = Super::GetMaterial(ElementIndex);
+
+	const ERHIFeatureLevel::Type FeatureLevel = GetScene() ? GetScene()->GetFeatureLevel() : ERHIFeatureLevel::Num;
+	if (FeatureLevel == ERHIFeatureLevel::Num)
+	{ 
+		return GEngine->HairDefaultMaterial;
+	}
+
+	if (!OverrideMaterial && GroomAsset && ElementIndex < GroomAsset->RenderHairGroups.Num())
+	{
+		if (UMaterialInterface* Material = GroomAsset->RenderHairGroups[ElementIndex].Material)
+		{
+			if (!Material->GetMaterialResource(FeatureLevel)->IsUsedWithHairStrands())
+			{
+				Material = GEngine->HairDefaultMaterial;
+			}
+			return Material;
+		}
+		else
+		{
+			return GEngine->HairDefaultMaterial;
+		}
+	}
+
+	if (OverrideMaterial == nullptr || !OverrideMaterial->GetMaterialResource(FeatureLevel)->IsUsedWithHairStrands())
+	{
+		OverrideMaterial = GEngine->HairDefaultMaterial;
+	}
+
+	return OverrideMaterial;
+}
+
+template<typename T> void SafeDelete(T*& Data) 
+{ 
+	if (Data)
+	{
+		T* LocalData = Data;
+		ENQUEUE_RENDER_COMMAND(FHairStrandsBuffers)(
+			[LocalData](FRHICommandListImmediate& RHICmdList)
+		{
+			delete LocalData;
+		});
+		Data = nullptr;
+	}
+
+}
+template<typename T> void SafeRelease(T*& Data) 
+{ 
+	if (Data)
+	{
+		T* LocalData = Data;
+		ENQUEUE_RENDER_COMMAND(FHairStrandsBuffers)(
+			[LocalData](FRHICommandListImmediate& RHICmdList)
+		{
+			LocalData->ReleaseResource();
+			delete LocalData;
+		});
+		Data = nullptr;
+	}
+}
+
+void UGroomComponent::InitResources()
+{
+	ReleaseResources();
+
+	if (GroomAsset && GroomAsset->HairStrandsResource)
+	{
+		check(GroomAsset->HairStrandsResource);
+
+		InitializedResources = GroomAsset;
+
+		FHairStrandsDatas* SimStrandDatas = &GroomAsset->HairSimulationData;
+
+		FHairStrandsInterpolationDatas InterpolationDatas;
+		InterpolationDatas.BuildInterpolationDatas(*SimStrandDatas, GroomAsset->HairRenderData);
+
+		InterpolationResource = new FHairStrandsInterpolationResource(InterpolationDatas, *SimStrandDatas);
+		BeginInitResource(InterpolationResource);
+
+		#if RHI_RAYTRACING
+		FHairStrandsRaytracingResource* LocalRaytracingResources = nullptr;
+		if (IsRayTracingEnabled())
+		{
+			RaytracingResources = new FHairStrandsRaytracingResource(&GroomAsset->HairRenderData);
+			BeginInitResource(RaytracingResources);
+			LocalRaytracingResources = RaytracingResources;
+		}
+		#endif
+
+		InterpolationOutput = new FHairStrandsInterpolationOutput();
+		InterpolationInput = new FHairStrandsInterpolationInput();
+
+		FHairStrandsInterpolationInput* Input = InterpolationInput;
+		FHairStrandsInterpolationOutput* Output = InterpolationOutput;
+		FHairStrandsResource* RenderResources = GroomAsset->HairStrandsResource;
+		FHairStrandsResource* SimResources = GroomAsset->HairSimulationResource;
+		FHairStrandsInterpolationResource* LocalInterpolationResource = InterpolationResource;
+
+		check(Input);
+		Input->HairRadius = GroomAsset->HairRenderData.StrandsCurves.MaxRadius;
+		Input->HairWorldOffset = GroomAsset->HairRenderData.BoundingBox.GetCenter();
+
+		const uint64 Id = (uint64)this;
+		ENQUEUE_RENDER_COMMAND(FHairStrandsBuffers)(
+			[Id, Input, Output, RenderResources, SimResources, LocalInterpolationResource
+			#if RHI_RAYTRACING
+			, LocalRaytracingResources
+			#endif
+			](FRHICommandListImmediate& RHICmdList)
+		{
+			Input->RenderRestPosePositionBuffer = &RenderResources->RestPositionBuffer;
+			Input->RenderAttributeBuffer		= &RenderResources->AttributeBuffer;
+			Input->RenderVertexCount			= RenderResources->StrandsDatas->GetNumPoints();
+
+			Input->SimRestPosePositionBuffer	= &SimResources->RestPositionBuffer;
+			Input->SimAttributeBuffer			= &SimResources->AttributeBuffer;
+			Input->SimVertexCount				= SimResources->StrandsDatas->GetNumPoints();
+			Input->SimRootPointIndexBuffer		= &LocalInterpolationResource->SimRootPointIndexBuffer;
+
+			Input->Interpolation0Buffer			= &LocalInterpolationResource->Interpolation0Buffer;
+			Input->Interpolation1Buffer			= &LocalInterpolationResource->Interpolation1Buffer;
+
+			#if RHI_RAYTRACING
+			if (IsRayTracingEnabled())
+			{
+				Input->RaytracingGeometry		= &LocalRaytracingResources->RayTracingGeometry;
+				Input->RaytracingPositionBuffer	= &LocalRaytracingResources->PositionBuffer;
+				Input->RaytracingVertexCount	= LocalRaytracingResources->VertexCount;
+			}
+			#endif
+
+			Output->SimDeformedPositionBuffer[0]	= &SimResources->DeformedPositionBuffer[0];
+			Output->SimDeformedPositionBuffer[1]	= &SimResources->DeformedPositionBuffer[1];
+			Output->RenderDeformedPositionBuffer[0] = &RenderResources->DeformedPositionBuffer[0];
+			Output->RenderDeformedPositionBuffer[1] = &RenderResources->DeformedPositionBuffer[1];
+			Output->RenderAttributeBuffer			= &RenderResources->AttributeBuffer;
+			Output->RenderTangentBuffer				= &RenderResources->TangentBuffer;
+			Output->SimTangentBuffer				= &SimResources->TangentBuffer;
+
+			FHairStrandsInterpolation Interpolation;
+			Interpolation.Input = Input;
+			Interpolation.Output = Output;
+			Interpolation.Function = ComputeHairStrandsInterpolation;
+			RegisterHairStrands(Id, Interpolation);
+		});
+
+	}
+}
+
+void UGroomComponent::ReleaseResources()
+{
+	// Unregister component interpolation resources
+	const uint64 Id = (uint64)this;
+	ENQUEUE_RENDER_COMMAND(StaticMeshVertexBuffersLegacyInit)(
+		[Id](FRHICommandListImmediate& RHICmdList)
+	{
+		UnregisterHairStrands(Id);
+	});
+
+	SafeRelease(InterpolationResource);
+
+	// Delay destruction as resources reference by the interpolation 
+	// structs are used on the rendering thread, 
+	FHairStrandsInterpolationInput* Input = InterpolationInput;
+	FHairStrandsInterpolationOutput* Output = InterpolationOutput;
+	ENQUEUE_RENDER_COMMAND(FHairStrandsBuffers)(
+		[Input, Output](FRHICommandListImmediate& RHICmdList)
+	{
+		delete Input;
+		delete Output;
+	});
+	InterpolationInput = nullptr;
+	InterpolationOutput = nullptr;
+	InitializedResources = nullptr;
+
+#if RHI_RAYTRACING
+	SafeRelease(RaytracingResources);
+#endif
+}
+
+void UGroomComponent::PostLoad()
+{
+	Super::PostLoad();
+
+	if (GroomAsset)
+	{
+		// Make sure that the asset initialized its resources first since the component needs them to initialize its own resources
+		GroomAsset->ConditionalPostLoad();
+	}
+
+	if (!IsTemplate())
+	{
+		InitResources();
+	}
+}
+
+void UGroomComponent::OnRegister()
+{
+	Super::OnRegister();
+}
+
+void UGroomComponent::OnUnregister()
+{
+	Super::OnUnregister();
+}
+
+void UGroomComponent::TickComponent(float DeltaTime, enum ELevelTick TickType, FActorComponentTickFunction *ThisTickFunction)
+{
+	Super::TickComponent(DeltaTime, TickType, ThisTickFunction);	
+}
+void UGroomComponent::GetUsedMaterials(TArray<UMaterialInterface*>& OutMaterials, bool bGetDebugMaterials) const
+{
+#if WITH_EDITOR
+	UMeshComponent::GetUsedMaterials(OutMaterials, bGetDebugMaterials);
+	if (bGetDebugMaterials)
+		OutMaterials.Add(GEngine->HairDebugMaterial);
+#endif
+	OutMaterials.Add(GEngine->HairDefaultMaterial);
+}
+
+#if WITH_EDITOR
+void UGroomComponent::PostEditChangeProperty(FPropertyChangedEvent& PropertyChangedEvent) 
+{
+	Super::PostEditChangeProperty(PropertyChangedEvent);
+
+	UProperty* PropertyThatChanged = PropertyChangedEvent.Property;
+	const FName PropertyName = PropertyThatChanged ? PropertyThatChanged->GetFName() : NAME_None;
+
+	//  Init/release resources when setting the GroomAsset (or undoing)
+	if (PropertyName == GET_MEMBER_NAME_CHECKED(UGroomComponent, GroomAsset) ||
+		PropertyThatChanged == nullptr)
+	{
+		if (GroomAsset != InitializedResources)
+		{
+			if (GroomAsset)
+			{
+				InitResources();
+			}
+			else
+			{
+				ReleaseResources();
+			}
+		}
+	}
+}
+#endif
+
+FGroomComponentRecreateRenderStateContext::FGroomComponentRecreateRenderStateContext(UGroomAsset* GroomAsset)
+{
+	if (!GroomAsset)
+	{
+		return;
+	}
+
+	for (TObjectIterator<UGroomComponent> HairStrandsComponentIt; HairStrandsComponentIt; ++HairStrandsComponentIt)
+	{
+		if (HairStrandsComponentIt->GroomAsset == GroomAsset)
+		{
+			if (HairStrandsComponentIt->IsRenderStateCreated())
+			{
+				HairStrandsComponentIt->DestroyRenderState_Concurrent();
+				GroomComponents.Add(*HairStrandsComponentIt);
+			}
+		}
+	}
+
+	// Flush the rendering commands generated by the detachments.
+	FlushRenderingCommands();
+}
+
+FGroomComponentRecreateRenderStateContext::~FGroomComponentRecreateRenderStateContext()
+{
+	const int32 ComponentCount = GroomComponents.Num();
+	for (int32 ComponentIndex = 0; ComponentIndex < ComponentCount; ++ComponentIndex)
+	{
+		UGroomComponent* GroomComponent = GroomComponents[ComponentIndex];
+
+		if (GroomComponent->IsRegistered() && !GroomComponent->IsRenderStateCreated())
+		{
+			GroomComponent->InitResources();
+			GroomComponent->CreateRenderState_Concurrent();
+		}
+	}
+}
diff --git a/Engine/Plugins/Experimental/HairStrands/Source/HairStrandsCore/Private/HairStrandsCore.cpp b/Engine/Plugins/Experimental/HairStrands/Source/HairStrandsCore/Private/HairStrandsCore.cpp
--- a/Engine/Plugins/Experimental/HairStrands/Source/HairStrandsCore/Private/HairStrandsCore.cpp
+++ b/Engine/Plugins/Experimental/HairStrands/Source/HairStrandsCore/Private/HairStrandsCore.cpp
@@ -2,6 +2,8 @@
 
 #include "HairStrandsCore.h"
 
+IMPLEMENT_MODULE(FHairStrandsCore, HairStrandsCore);
+
 void FHairStrandsCore::StartupModule()
 {
 }
diff --git a/Engine/Plugins/Experimental/HairStrands/Source/HairStrandsCore/Private/HairStrandsDatas.cpp b/Engine/Plugins/Experimental/HairStrands/Source/HairStrandsCore/Private/HairStrandsDatas.cpp
--- a/Engine/Plugins/Experimental/HairStrands/Source/HairStrandsCore/Private/HairStrandsDatas.cpp
+++ b/Engine/Plugins/Experimental/HairStrands/Source/HairStrandsCore/Private/HairStrandsDatas.cpp
@@ -7,6 +7,18 @@
 #include "Chaos/ErrorReporter.h"
 #include "Chaos/Levelset.h"
 #include "Chaos/Framework/Parallel.h"
+#include "UObject/PhysicsObjectVersion.h"
+#include "Async/ParallelFor.h"
+
+// For debug purpose
+static float GHairInterpolationMetric_Distance = 1;
+static float GHairInterpolationMetric_Angle = 0;
+static float GHairInterpolationMetric_Length = 0;
+static float GHairInterpolationMetric_AngleAttenuation = 5;
+static FAutoConsoleVariableRef CVarHairInterpolationMetric_Distance(TEXT("r.HairStrands.InterpolationMetric.Distance"), GHairInterpolationMetric_Distance, TEXT("Hair strands interpolation metric weights for distance"));
+static FAutoConsoleVariableRef CVarHairInterpolationMetric_Angle(TEXT("r.HairStrands.InterpolationMetric.Angle"), GHairInterpolationMetric_Angle, TEXT("Hair strands interpolation metric weights for angle"));
+static FAutoConsoleVariableRef CVarHairInterpolationMetric_Length(TEXT("r.HairStrands.InterpolationMetric.Length"), GHairInterpolationMetric_Length, TEXT("Hair strands interpolation metric weights for length"));
+static FAutoConsoleVariableRef CVarHairInterpolationMetric_AngleAttenuation(TEXT("r.HairStrands.InterpolationMetric.AngleAttenuation"), GHairInterpolationMetric_AngleAttenuation, TEXT("Hair strands interpolation angle attenuation"));
 
 inline void CopyVectorToPosition(const FVector& InVector, FHairStrandsPositionFormat::Type& OutPosition)
 {
@@ -22,69 +34,608 @@ inline void CopyPositionToVector(const FHairStrandsPositionFormat::Type& InPosit
 	OutVector.Z = InPosition.Z;
 }
 
-void FHairStrandsDatas::BuildRenderingDatas(
-	TArray<FHairStrandsPositionFormat::Type>& OutPackedPositions,
-	TArray<FHairStrandsTangentFormat::Type>& OutPackedTangents) const
+void DecimateStrandData(const FHairStrandsDatas& InData, float DecimationPercentage, FHairStrandsDatas& OutData)
 {
-	if (GetNumCurves() > 0 && GetNumPoints() > 0)
+	// Pick randomly strand as guide 
+	// Divide strands in buckets and pick randomly one stand per bucket
+	const uint32 CurveCount = InData.StrandsCurves.Num();
+	const uint32 OutCurveCount = FMath::Clamp(uint32(CurveCount * DecimationPercentage), 1u, CurveCount);
+
+	const uint32 BucketSize = CurveCount / OutCurveCount;
+
+	TArray<uint32> CurveIndices;
+	CurveIndices.SetNum(OutCurveCount);
+
+	uint32 OutTotalPointCount = 0;
+	FRandomStream Random;
+	for (uint32 BucketIndex = 0; BucketIndex < OutCurveCount; BucketIndex++)
 	{
-		OutPackedPositions.SetNum(GetNumPoints());
-		OutPackedTangents.SetNum(GetNumPoints()*2);
+		const uint32 CurveIndex = BucketIndex * BucketSize;// +BucketSize * Random.FRand();
+		CurveIndices[BucketIndex] = CurveIndex;
+		OutTotalPointCount += InData.StrandsCurves.CurvesCount[CurveIndex];
+	}
 
-		TArray<uint16>::TConstIterator CountIterator = StrandsCurves.CurvesCount.CreateConstIterator();
-		TArray<float>::TConstIterator LengthIterator = StrandsCurves.CurvesLength.CreateConstIterator();
-		TArray<float>::TConstIterator CoordUIterator = StrandsPoints.PointsCoordU.CreateConstIterator();
-		TArray<float>::TConstIterator RadiusIterator = StrandsPoints.PointsRadius.CreateConstIterator();
-		TArray<FVector>::TConstIterator PositionIterator = StrandsPoints.PointsPosition.CreateConstIterator();
+	OutData.StrandsCurves.SetNum(OutCurveCount);
+	OutData.StrandsPoints.SetNum(OutTotalPointCount);
+	OutData.HairDensity = InData.HairDensity;
+
+	uint32 OutPointOffset = 0; 
+	for (uint32 OutCurveIndex = 0; OutCurveIndex < OutCurveCount; OutCurveIndex++)
+	{
+		const uint32 InCurveIndex = CurveIndices[OutCurveIndex];
+		OutData.StrandsCurves.CurvesCount[OutCurveIndex] = InData.StrandsCurves.CurvesCount[InCurveIndex];
+		OutData.StrandsCurves.CurvesRootUV[OutCurveIndex] = InData.StrandsCurves.CurvesRootUV[InCurveIndex];
+		OutData.StrandsCurves.CurvesOffset[OutCurveIndex] = OutPointOffset;
+
+		const uint32 InPointOffset = InData.StrandsCurves.CurvesOffset[InCurveIndex];
+		const uint32 OutPointCount = OutData.StrandsCurves.CurvesCount[OutCurveIndex];
+		for (uint32 OutPointIndex = 0; OutPointIndex < OutPointCount; ++OutPointIndex, ++OutPointOffset)
+		{
+			OutData.StrandsPoints.PointsPosition[OutPointOffset] = InData.StrandsPoints.PointsPosition[OutPointIndex + InPointOffset];
+			OutData.StrandsPoints.PointsCoordU[OutPointOffset] = InData.StrandsPoints.PointsCoordU[OutPointIndex + InPointOffset];
+			OutData.StrandsPoints.PointsRadius[OutPointOffset] = InData.StrandsPoints.PointsRadius[OutPointIndex + InPointOffset] * InData.StrandsCurves.MaxRadius;
+		}
+	}
+
+	OutData.BuildInternalDatas();
+}
+
+template<typename T>
+void SwapValue(T& A, T& B)
+{
+	T Temp = A;
+	A = B;
+	B = Temp;
+}
+
+void FHairStrandsInterpolationDatas::SetNum(const uint32 NumCurves)
+{
+	PointsSimCurvesVertexWeights.SetNum(NumCurves);
+	PointsSimCurvesVertexIndex.SetNum(NumCurves);
+	PointsSimCurvesIndex.SetNum(NumCurves);
+}
+
+void FHairStrandsCurves::SetNum(const uint32 NumCurves)
+{
+	CurvesOffset.SetNum(NumCurves + 1);
+	CurvesCount.SetNum(NumCurves);
+	CurvesLength.SetNum(NumCurves);
+	CurvesRootUV.SetNum(NumCurves);
+}
+
+void FHairStrandsPoints::SetNum(const uint32 NumPoints)
+{
+	PointsPosition.SetNum(NumPoints);
+	PointsRadius.SetNum(NumPoints);
+	PointsCoordU.SetNum(NumPoints);
+}
+
+void FHairStrandsInterpolationDatas::Reset()
+{
+	PointsSimCurvesVertexWeights.Reset();
+	PointsSimCurvesVertexIndex.Reset();
+	PointsSimCurvesIndex.Reset();
+}
+
+void FHairStrandsCurves::Reset()
+{
+	CurvesOffset.Reset();
+	CurvesCount.Reset();
+	CurvesLength.Reset();
+	CurvesRootUV.Reset();
+}
+
+void FHairStrandsPoints::Reset()
+{
+	PointsPosition.Reset();
+	PointsRadius.Reset();
+	PointsCoordU.Reset();
+}
+
+void FHairStrandsInterpolationDatas::Serialize(FArchive& Ar)
+{
+	Ar << PointsSimCurvesVertexWeights;
+	Ar << PointsSimCurvesVertexIndex;
+	Ar << PointsSimCurvesIndex;
+}
+
+void FHairStrandsPoints::Serialize(FArchive& Ar)
+{
+	Ar << PointsPosition;
+	Ar << PointsRadius;
+	Ar << PointsCoordU;
+}
+
+void FHairStrandsCurves::Serialize(FArchive& Ar)
+{
+	Ar.UsingCustomVersion(FPhysicsObjectVersion::GUID);
+
+	Ar << CurvesCount;
+	Ar << CurvesOffset;
+	Ar << CurvesLength;
+	Ar << CurvesRootUV;
+	Ar << MaxLength;
+	Ar << MaxRadius;
+
+	if (Ar.CustomVer(FPhysicsObjectVersion::GUID) >= FPhysicsObjectVersion::HairAssetSerialization_V2)
+	{
+		Ar << CurvesGroupID;
+	}
+}
+
+void FHairStrandsDatas::Serialize(FArchive& Ar)
+{
+	StrandsPoints.Serialize(Ar);
+	StrandsCurves.Serialize(Ar);
+	Ar << HairDensity;
+	Ar << BoundingBox;
+}
+
+void FHairStrandsDatas::Reset()
+{
+	StrandsCurves.Reset();
+	StrandsPoints.Reset();
+}
+
+struct FHairRoot
+{
+	FVector Position;
+	uint32  VertexCount;
+	FVector Normal;
+	uint32  Index;
+	float   Length;
+};
+
+struct FHairInterpolationMetric
+{	
+	// Total/combined metrics
+	float Metric;
+
+	// Debug info
+	float DistanceMetric;
+	float AngularMetric;
+	float LengthMetric;
+
+	float CosAngle;
+	float Distance;
+	float GuideLength;
+	float RenderLength;
+
+};
+
+static inline FHairInterpolationMetric ComputeInterpolationMetric(const FHairRoot& RenderRoot, const FHairRoot& GuideRoot)
+{
+	FHairInterpolationMetric Out;
+	Out.Distance = FVector::Dist(RenderRoot.Position, GuideRoot.Position);
+	Out.CosAngle = FVector::DotProduct(RenderRoot.Normal, GuideRoot.Normal);
+	Out.GuideLength = GuideRoot.Length;
+	Out.RenderLength = RenderRoot.Length;
+
+	// Metric takes into account the following properties to find guides which are close, share similar orientation, and 
+	// have similar length for better interpolation
+	// * distance
+	// * orientation 
+	// * length 
+	const float AngularAttenuation = GHairInterpolationMetric_AngleAttenuation > 1 ? GHairInterpolationMetric_AngleAttenuation : 0;
+	Out.DistanceMetric	= Out.Distance * GHairInterpolationMetric_Distance;
+	Out.AngularMetric	= AngularAttenuation == 0 ? 0 : (FMath::Clamp((1 - FMath::Pow(Out.CosAngle, AngularAttenuation)), 0.f, 1.f) * GHairInterpolationMetric_Angle);
+	Out.LengthMetric	= FMath::Abs(FMath::Max(Out.GuideLength / float(Out.RenderLength), Out.RenderLength / float(Out.GuideLength)) - 1) * GHairInterpolationMetric_Length; // Ratio
+	Out.Metric			= Out.DistanceMetric + Out.AngularMetric + Out.LengthMetric;
+
+	return Out;
+}
+
+template<uint32 NumSamples>
+static inline FVector GetCurvePosition(const FHairStrandsDatas& CurvesDatas, const uint32 CurveIndex, const uint32 SampleIndex)
+{
+	const float PointCount = CurvesDatas.StrandsCurves.CurvesCount[CurveIndex]-1.0;
+	const uint32 PointOffset = CurvesDatas.StrandsCurves.CurvesOffset[CurveIndex];
+
+	const float CurvePoint = static_cast<float>(SampleIndex) * PointCount / (static_cast<float>(NumSamples)-1.0f);
+	const uint32 PointPrev = (SampleIndex == 0) ? 0 : (SampleIndex == (NumSamples-1)) ? PointCount - 1 : floor(CurvePoint);
+	const uint32 PointNext = PointPrev + 1;
+
+	const float PointAlpha = CurvePoint - static_cast<float>(PointPrev);
+	return CurvesDatas.StrandsPoints.PointsPosition[PointOffset+PointPrev]*(1.0f-PointAlpha) + 
+		   CurvesDatas.StrandsPoints.PointsPosition[PointOffset+PointNext]*PointAlpha;
+}
+
+template<uint32 NumSamples>
+static inline float ComputeCurvesMetric(const FHairStrandsDatas& RenderCurvesDatas, const uint32 RenderCurveIndex, 
+										const FHairStrandsDatas& GuideCurvesDatas, const uint32 GuideCurveIndex, const float RootImportance, 
+										const float ShapeImportance, const float ProximityImportance)
+{
+	const float AverageLength = FMath::Max(0.5f * (RenderCurvesDatas.StrandsCurves.CurvesLength[RenderCurveIndex] * RenderCurvesDatas.StrandsCurves.MaxLength +
+											       GuideCurvesDatas.StrandsCurves.CurvesLength[GuideCurveIndex] * GuideCurvesDatas.StrandsCurves.MaxLength), SMALL_NUMBER);
+
+	static const float DeltaCoord = 1.0f / static_cast<float>(NumSamples-1);
+
+	const FVector& RenderRoot = RenderCurvesDatas.StrandsPoints.PointsPosition[RenderCurvesDatas.StrandsCurves.CurvesOffset[RenderCurveIndex]];
+	const FVector& GuideRoot = GuideCurvesDatas.StrandsPoints.PointsPosition[GuideCurvesDatas.StrandsCurves.CurvesOffset[GuideCurveIndex]];
+
+	float CurveProximityMetric = 0.0;
+	float CurveShapeMetric = 0.0;
+	for (uint32 SampleIndex = 0; SampleIndex < NumSamples; ++SampleIndex)
+	{
+		const FVector GuidePosition = GetCurvePosition<NumSamples>(GuideCurvesDatas, GuideCurveIndex, SampleIndex);
+		const FVector RenderPosition = GetCurvePosition<NumSamples>(RenderCurvesDatas, RenderCurveIndex, SampleIndex);
+		const float RootWeight = FMath::Exp(-RootImportance*SampleIndex*DeltaCoord);
+
+		CurveProximityMetric += (GuidePosition - RenderPosition).Size() * RootWeight;
+		CurveShapeMetric += (GuidePosition - GuideRoot - RenderPosition + RenderRoot).Size() * RootWeight;
+	}
+	CurveShapeMetric *= DeltaCoord / AverageLength;
+	CurveProximityMetric *= DeltaCoord / AverageLength;
+
+	return FMath::Exp(-ShapeImportance*CurveShapeMetric) * FMath::Exp(-ProximityImportance * CurveProximityMetric);
+}
+
+static inline void PrintInterpolationMetric(const FHairInterpolationMetric& In)
+{
+	UE_LOG(LogHairStrands, Log, TEXT("------------------------------------------------------------------------------------------"));
+	UE_LOG(LogHairStrands, Log, TEXT("Total Metric = %f"), In.Metric);
+	UE_LOG(LogHairStrands, Log, TEXT("Distance     = %f (%f)"), In.Distance, In.DistanceMetric);
+	UE_LOG(LogHairStrands, Log, TEXT("Angle        = %f (%f)"), FMath::RadiansToDegrees(FMath::Acos(In.CosAngle)), In.AngularMetric);
+	UE_LOG(LogHairStrands, Log, TEXT("Length       = %f/%f (%f)"), In.RenderLength, In.GuideLength, In.LengthMetric);
+}
+
+void FHairStrandsInterpolationDatas::BuildInterpolationDatas(
+	const FHairStrandsDatas& SimStrandsData,
+	const FHairStrandsDatas& RenStrandsData)
+{
+	SetNum(RenStrandsData.GetNumPoints());
+
+	const bool bRandomizeInterpolation = true;
+	const bool bUseUniqueGuide = false;
+	const bool bPrintDebugMetric = false;
+
+	typedef TArray<FHairRoot> FRoots;
+
+	// Extract strand roots
+	auto ExtractRoot = [](const FHairStrandsDatas& InData, TArray<FHairRoot>& OutRoots)
+	{
+		const uint32 CurveCount = InData.StrandsCurves.Num();
+		OutRoots.Reserve(CurveCount);
+		for (uint32 CurveIndex = 0; CurveIndex < CurveCount; ++CurveIndex)
+		{
+			const uint32 PointOffset = InData.StrandsCurves.CurvesOffset[CurveIndex];
+			const uint32 PointCount  = InData.StrandsCurves.CurvesCount[CurveIndex];
+			const float  CurveLength = InData.StrandsCurves.CurvesLength[CurveIndex] * InData.StrandsCurves.MaxLength;
+			check(PointCount > 1);
+			const FVector& P0 = InData.StrandsPoints.PointsPosition[PointOffset];
+			const FVector& P1 = InData.StrandsPoints.PointsPosition[PointOffset+1];
+			FVector N = (P1 - P0).GetSafeNormal();
+
+			// Fallback in case the initial points are too close (this happens on certain assets)
+			if (FVector::DotProduct(N, N) == 0)
+			{
+				N = FVector(0, 0, 1);
+			}
+			OutRoots.Add({ P0, PointCount, N, PointOffset, CurveLength });
+		}
+	};
+
+	FRoots RenRoots, SimRoots;
+	ExtractRoot(RenStrandsData, RenRoots);
+	ExtractRoot(SimStrandsData, SimRoots);
+
+	// Find k-closest guide:
+	// N hairs, M guides. Complexity = MxN
+	// #hair_todo: build an acceleration structure for fast(er) look up or at least make it run on parallel (e.g. with ParallelFor)
+	uint32 TotalInvalidInterpolationCount = 0;
+	const float MinWeightDistance = 0.0001f;
+	const static uint32 GuideCount = 3;
+	const static uint32 KGuideCount = GuideCount * 2;
+	FRandomStream Random;
+	const uint32 RenCurveCount = RenStrandsData.GetNumCurves();
+	const uint32 SimCurveCount = SimStrandsData.GetNumCurves();
+	uint32 RenGlobaPointIndex = 0;
+	//ParallelFor(RenCurveCount, [this, &RenRoots, &SimRoots, &SimCurveCount, &RenStrandsData, &SimStrandsData](uint32 RenCurveIndex) {
+	for (uint32 RenCurveIndex = 0; RenCurveIndex < RenCurveCount; ++RenCurveIndex)
+	{
+		const FHairRoot& StrandRoot = RenRoots[RenCurveIndex];
+
+		int32 ClosestGuideIndices[GuideCount];
+		{
+			float KMinMetrics[KGuideCount];
+			int32 KClosestGuideIndices[KGuideCount];
+			for (uint32 ClearIt = 0; ClearIt < KGuideCount; ++ClearIt)
+			{
+				KMinMetrics[ClearIt] = FLT_MAX;
+				KClosestGuideIndices[ClearIt] = -1;
+			}
+
+			ParallelFor(SimCurveCount, [this, &KMinMetrics, &KClosestGuideIndices, &RenCurveIndex, &RenStrandsData, &SimStrandsData, &StrandRoot, &SimRoots](uint32 SimCurveIndex) 
+			//for (uint32 SimCurveIndex = 0; SimCurveIndex < SimCurveCount; ++SimCurveIndex)
+			{
+				const FHairRoot& GuideRoot = SimRoots[SimCurveIndex];
+				//if ((StrandRoot.Position - GuideRoot.Position).Size() < 5.0f)
+				{
+					const float Metric = 1.0 - ComputeCurvesMetric<16>(RenStrandsData, RenCurveIndex, SimStrandsData, SimCurveIndex, 0.0f, 1.0f, 1.0f);
+					if (Metric < KMinMetrics[KGuideCount - 1])
+					{
+						int32 LastGuideIndex = SimCurveIndex;
+						float LastMetric = Metric;
+						for (uint32 Index = 0; Index < KGuideCount; ++Index)
+						{
+							if (Metric < KMinMetrics[Index])
+							{
+								SwapValue(KClosestGuideIndices[Index], LastGuideIndex);
+								SwapValue(KMinMetrics[Index], LastMetric);
+							}
+						}
+					}
+				}
+
+				//const FHairRoot& GuideRoot = SimRoots[SimCurveIndex];
+				//const FHairInterpolationMetric Metric = ComputeInterpolationMetric(StrandRoot, GuideRoot);
+
+				//const bool bIsSameHemisphere = GHairInterpolationMetric_AngleAttenuation > 1 ? (Metric.CosAngle > 0) : true;
+				//if (bIsSameHemisphere)
+				//{
+				//	// If the current guide is closer than the furtherest apart guide, we need to update 
+				//	// the closest guide set
+				//	if (Metric.Metric < KMinMetrics[KGuideCount - 1])
+				//	{
+				//		int32 LastGuideIndex = SimCurveIndex;
+				//		float LastMetric = Metric.Metric;
+				//		for (uint32 Index = 0; Index < KGuideCount; ++Index)
+				//		{
+				//			if (Metric.Metric < KMinMetrics[Index])
+				//			{
+				//				SwapValue(KClosestGuideIndices[Index], LastGuideIndex);
+				//				SwapValue(KMinMetrics[Index], LastMetric);
+				//			}
+				//		}
+				//	}
+				//}
+			});
+
+			// Debug
+			if (bPrintDebugMetric)
+			{
+				const FHairInterpolationMetric ClosestMetric = ComputeInterpolationMetric(StrandRoot, SimRoots[KClosestGuideIndices[0]]);
+				const float Threshold = 20;
+				if (ClosestMetric.Metric > Threshold)
+				{
+					PrintInterpolationMetric(ClosestMetric);
+					++TotalInvalidInterpolationCount;
+				}
+			}
+
+			// Randomize influence guide to break interpolation coherence, and create a more random/natural pattern
+			{
+				uint32 RandIndex0 = 0;
+				uint32 RandIndex1 = 1;
+				uint32 RandIndex2 = 2;
+				/*if (bRandomizeInterpolation)
+				{
+					do
+					{
+						RandIndex0 = Random.RandRange(0, KGuideCount - 1);
+						RandIndex1 = Random.RandRange(0, KGuideCount - 1);
+						RandIndex2 = Random.RandRange(0, KGuideCount - 1);
 
-		TArray<FHairStrandsPositionFormat::Type>::TIterator PackedPositionIterator = OutPackedPositions.CreateIterator();
-		TArray<FHairStrandsTangentFormat::Type>::TIterator PackedTangentIterator = OutPackedTangents.CreateIterator();
+					} while (RandIndex0 == RandIndex1 || RandIndex0 == RandIndex2 || RandIndex1 == RandIndex2);
+				}*/
+				ClosestGuideIndices[0] = KClosestGuideIndices[RandIndex0];
+				ClosestGuideIndices[1] = KClosestGuideIndices[RandIndex1];
+				ClosestGuideIndices[2] = KClosestGuideIndices[RandIndex2];
 
-		FRandomStream Random;
-		for (uint32 CurveIndex = 0; CurveIndex < GetNumCurves(); ++CurveIndex, ++CountIterator, ++LengthIterator)
+				if (bUseUniqueGuide)
+				{
+					ClosestGuideIndices[1] = KClosestGuideIndices[RandIndex0];
+					ClosestGuideIndices[2] = KClosestGuideIndices[RandIndex0];
+				}
+
+				float MinMetrics[GuideCount];
+				MinMetrics[0] = KMinMetrics[RandIndex0];
+				MinMetrics[1] = KMinMetrics[RandIndex1];
+				MinMetrics[2] = KMinMetrics[RandIndex2];
+
+
+				while (!(MinMetrics[0] <= MinMetrics[1] && MinMetrics[1] <= MinMetrics[2]))
+				{
+					if (MinMetrics[0] > MinMetrics[1])
+					{
+						SwapValue(MinMetrics[0], MinMetrics[1]);
+						SwapValue(ClosestGuideIndices[0], ClosestGuideIndices[1]);
+					}
+
+					if (MinMetrics[1] > MinMetrics[2])
+					{
+						SwapValue(MinMetrics[1], MinMetrics[2]);
+						SwapValue(ClosestGuideIndices[1], ClosestGuideIndices[2]);
+					}
+				}
+
+				// If there less than 3 valid guides, fill the rest with existing valid guides
+				// This can happen due to the normal-orientation based rejection above
+				if (ClosestGuideIndices[1] < 0)
+				{
+					ClosestGuideIndices[1] = ClosestGuideIndices[0];
+					MinMetrics[1] = MinMetrics[0];
+				}
+				if (ClosestGuideIndices[2] < 0)
+				{
+					ClosestGuideIndices[2] = ClosestGuideIndices[1];
+					MinMetrics[2] = MinMetrics[1];
+				}
+
+				check(ClosestGuideIndices[0] >= 0);
+				check(ClosestGuideIndices[1] >= 0);
+				check(ClosestGuideIndices[2] >= 0);
+				check(MinMetrics[0] <= MinMetrics[1]);
+				check(MinMetrics[1] <= MinMetrics[2]);
+			}
+		}
+
+
+		const uint32 RendPointCount	= RenStrandsData.StrandsCurves.CurvesCount[RenCurveIndex];
+		const uint32 RenOffset		= RenStrandsData.StrandsCurves.CurvesOffset[RenCurveIndex];
+		const FVector& RenPointPositionRoot = RenStrandsData.StrandsPoints.PointsPosition[RenOffset];
+		for (uint32 RenPointIndex = 0; RenPointIndex < RendPointCount; ++RenPointIndex, ++RenGlobaPointIndex)
 		{
-			const uint16& PointCount = *CountIterator;
-			for (int32 PointIndex = 0; PointIndex < PointCount; ++PointIndex, 
-				++CoordUIterator, ++PositionIterator, ++RadiusIterator, ++PackedPositionIterator, ++PackedTangentIterator)
+			const FVector& RenPointPosition = RenStrandsData.StrandsPoints.PointsPosition[RenPointIndex + RenOffset];
+			const float RenPointDistance = RenStrandsData.StrandsPoints.PointsCoordU[RenPointIndex + RenOffset] * RenStrandsData.StrandsCurves.CurvesLength[RenCurveIndex] * RenStrandsData.StrandsCurves.MaxLength;
+
+			float TotalWeight = 0;
+			for (uint32 KIndex = 0; KIndex < GuideCount; ++KIndex)
 			{
-				const uint32 PrevIndex = FMath::Max(0, PointIndex - 1);
-				const uint32 NextIndex = FMath::Min(PointCount + 1, PointCount - 1);
-				const FVector& PointPosition = *PositionIterator;
-
-				const FVector ForwardDir0 = PointIndex > 0 ? (PointPosition - StrandsPoints.PointsPosition[PrevIndex]).GetSafeNormal() : FVector::ZeroVector;
-				const FVector ForwardDir1 = PointIndex < PointCount - 1 ? (StrandsPoints.PointsPosition[NextIndex] - PointPosition).GetSafeNormal() : FVector::ZeroVector;
-				const FVector TangentZ = (ForwardDir0 + ForwardDir1).GetSafeNormal();
-
-				// Find quat from up (Z) vector to forward
-				// Apply quat orthogonal vectors
-				const FQuat DeltaQuat = FQuat::FindBetween(FVector(0, 0, -1), TangentZ);
-				const FVector RightDir = DeltaQuat.RotateVector(FVector(0, 1, 0));
-				const FVector UpDir = DeltaQuat.RotateVector(FVector(1, 0, 0));
-
-				const FVector TangentX = RightDir ^ TangentZ;
-				const FVector TangentY = RightDir;
-
-				FHairStrandsPositionFormat::Type& PackedPosition = *PackedPositionIterator;
-				CopyVectorToPosition(PointPosition, PackedPosition);
-				PackedPosition.ControlPointType = (PointIndex == 0) ? 1u : (PointIndex == (PointCount - 1) ? 2u : 0u);
-				PackedPosition.UCoord = FMath::Clamp((*CoordUIterator)*255.f, 0.f, 255.f);
-				PackedPosition.NormalizedRadius = FMath::Clamp((*RadiusIterator) * 63.f, 0.f, 63.f);
-				PackedPosition.NormalizedLength = FMath::Clamp((*LengthIterator) *255.f, 0.f, 255.f);
-				PackedPosition.Seed = Random.RandHelper(255);
-
-				FHairStrandsTangentFormat::Type& PackedTangentX = *PackedTangentIterator;
-				PackedTangentX = TangentX;
-
-				++PackedTangentIterator;
-
-				FHairStrandsTangentFormat::Type& PackedTangentZ = *PackedTangentIterator;
-				PackedTangentZ = TangentZ;
-				PackedTangentZ.Vector.W = GetBasisDeterminantSignByte(TangentX, TangentY, TangentZ);
+
+			#define WEIGHT_METHOD 0
+
+				// Find the closest vertex on the guide which matches the strand vertex distance along its curve
+			#if WEIGHT_METHOD == 0
+				const uint32 SimCurveIndex = ClosestGuideIndices[KIndex];
+				const uint32 SimOffset = SimStrandsData.StrandsCurves.CurvesOffset[SimCurveIndex];
+				float PrevSimPointDistance = 0;
+				uint32 ClosestSimPointIndex = 0;
+				const uint32 SimPointCount = SimStrandsData.StrandsCurves.CurvesCount[SimCurveIndex];
+				for (uint32 SimPointIndex = 0; SimPointIndex < SimPointCount; ++SimPointIndex, ++ClosestSimPointIndex)
+				{
+					const float SimPointDistance = SimStrandsData.StrandsPoints.PointsCoordU[SimPointIndex + SimOffset] * SimStrandsData.StrandsCurves.CurvesLength[SimCurveIndex] * SimStrandsData.StrandsCurves.MaxLength;
+					if (RenPointDistance >= PrevSimPointDistance && RenPointDistance <= SimPointDistance)
+					{
+						//const float GuideDistance = Guide.Points[SimPointIndex].Distance;
+						const float D0 = FMath::Abs(PrevSimPointDistance - RenPointDistance);
+						const float D1 = FMath::Abs(SimPointDistance - RenPointDistance);
+						if (D0 > D1)
+						{
+							ClosestSimPointIndex = SimPointIndex;
+						}
+						break;
+					}
+					PrevSimPointDistance = SimPointDistance;					
+				}
+				check(SimPointCount > 0);
+				ClosestSimPointIndex = FMath::Clamp(ClosestSimPointIndex, 0u, uint32(SimPointCount - 1));
+
+				
+				
+				const FVector& SimPointPosition = SimStrandsData.StrandsPoints.PointsPosition[ClosestSimPointIndex + SimOffset];
+				const float Weight = 1.0f / FMath::Max(MinWeightDistance, FVector::Distance(RenPointPosition, SimPointPosition));
+				PointsSimCurvesIndex[RenGlobaPointIndex][KIndex] = SimCurveIndex;
+				PointsSimCurvesVertexIndex[RenGlobaPointIndex][KIndex] = ClosestSimPointIndex + SimOffset;
+				PointsSimCurvesVertexWeights[RenGlobaPointIndex][KIndex] = Weight;
+			#endif
+
+				// Use only the root as a *constant* weight for deformation along each vertex
+			#if WEIGHT_METHOD == 1
+				const uint32 SimCurveIndex = ClosestGuideIndices[KIndex];
+				const uint32 SimOffset = SimStrandsData.StrandsCurves.CurvesOffset[SimCurveIndex];
+				const FVector& SimRootPointPosition = SimStrandsData.StrandsPoints[SimOffset];
+				const float Weight = 1.0f / FMath::Max(MinWeightDistance, FVector::Distance(RenPointPosition, SimRootPointPosition));
+				PointsSimCurvesIndex[RenGlobaPointIndex][KIndex] = SimCurveIndex;
+				PointsSimCurvesVertexIndex[RenGlobaPointIndex][KIndex] = SimOffset;
+				PointsSimCurvesVertexWeights[RenGlobaPointIndex][KIndex] = Weight;
+
+			#endif
+
+				// Use the *same vertex index* to match guide vertex with strand vertex
+			#if WEIGHT_METHOD == 2
+				check(SimPointCount > 0);
+				const uint32 SimCurveIndex = ClosestGuideIndices[KIndex];
+				const uint32 SimPointIndex = FMath::Clamp(RenPointIndex, 0, SimPointCount - 1);
+				const FVector& SimPointPosition = SimStrandsData.StrandsPoints[SimPointIndex + SimOffset];
+				const float Weight = 1.0f / FMath::Max(MinWeightDistance, FVector::Distance(RenPointPosition, SimPointPosition));
+				PointsSimCurvesIndex[RenGlobaPointIndex][KIndex] = SimCurveIndex;
+				PointsSimCurvesVertexIndex[RenGlobaPointIndex][KIndex] = SimPointIndex + SimOffset;
+				PointsSimCurvesVertexWeights[RenGlobaPointIndex][KIndex] = Weight;
+			#endif
+
+				TotalWeight += Weight;
+			}
+
+			for (int32 KIndex = 0; KIndex < GuideCount; ++KIndex)
+			{
+				PointsSimCurvesVertexWeights[RenGlobaPointIndex][KIndex] /= TotalWeight;
 			}
 		}
 	}
+
+	if (bPrintDebugMetric)
+	{
+		UE_LOG(LogHairStrands, Log, TEXT("Invalid interpolation count: %d/%d)"), TotalInvalidInterpolationCount, RenCurveCount);
+	}
 }
 
+void FHairStrandsInterpolationDatas::BuildRenderingDatas(
+	TArray<FHairStrandsInterpolation0Format::Type>& OutPointsInterpolation0,
+	TArray<FHairStrandsInterpolation1Format::Type>& OutPointsInterpolation1) const
+{
+	const uint32 PointCount = Num();
+	if (PointCount == 0)
+		return;
+	auto LowerPart = [](uint32 Index) { return uint16(Index & 0xFFFF); };
+	auto UpperPart = [](uint32 Index) { return uint8((Index >> 16) & 0xFF); };
+
+	OutPointsInterpolation0.SetNum(PointCount*FHairStrandsInterpolation0Format::ComponentCount);
+	OutPointsInterpolation1.SetNum(PointCount*FHairStrandsInterpolation1Format::ComponentCount);
+
+	for (uint32 PointIndex = 0; PointIndex < PointCount; ++PointIndex)
+	{
+		const FIntVector& Indices = PointsSimCurvesVertexIndex[PointIndex];
+		const FVector& Weights = PointsSimCurvesVertexWeights[PointIndex];
+
+		FHairStrandsInterpolation0Format::Type& OutInterp0 = OutPointsInterpolation0[PointIndex];
+		OutInterp0.Index0 = LowerPart(Indices[0]);
+		OutInterp0.Index1 = LowerPart(Indices[1]);
+		OutInterp0.Index2 = LowerPart(Indices[2]);
+		OutInterp0.VertexWeight0 = Weights[0] * 255.f;
+		OutInterp0.VertexWeight1 = Weights[1] * 255.f;
+
+		FHairStrandsInterpolation1Format::Type& OutInterp1 = OutPointsInterpolation1[PointIndex];
+		OutInterp1.VertexIndex0 = UpperPart(Indices[0]);
+		OutInterp1.VertexIndex1 = UpperPart(Indices[1]);
+		OutInterp1.VertexIndex2 = UpperPart(Indices[2]);
+		OutInterp1.Pad0 = 0;
+	}	
+}
+
+void FHairStrandsDatas::BuildRenderingDatas(
+	TArray<FHairStrandsPositionFormat::Type>& OutPackedPositions,
+	TArray<FHairStrandsAttributeFormat::Type>& OutPackedAttributes) const
+{
+	if (!(GetNumCurves() > 0 && GetNumPoints() > 0))
+		return;
+
+	OutPackedPositions.SetNum(GetNumPoints()*FHairStrandsPositionFormat::ComponentCount);
+	OutPackedAttributes.SetNum(GetNumPoints()*FHairStrandsAttributeFormat::ComponentCount);
+	
+	const FVector HairBoxCenter = BoundingBox.GetCenter();
+
+	FRandomStream Random;
+	for (uint32 CurveIndex = 0; CurveIndex < GetNumCurves(); ++CurveIndex)
+	{
+		const float CurveSeed = Random.RandHelper(255);
+		const int32 IndexOffset = StrandsCurves.CurvesOffset[CurveIndex];
+		const uint16& PointCount = StrandsCurves.CurvesCount[CurveIndex];
+		for (int32 PointIndex = 0; PointIndex < PointCount; ++PointIndex)
+		{
+			const uint32 PrevIndex = FMath::Max(0, PointIndex - 1);
+			const uint32 NextIndex = FMath::Min(PointCount + 1, PointCount - 1);
+			const FVector& PointPosition = StrandsPoints.PointsPosition[PointIndex + IndexOffset];
+
+			const float CoordU = StrandsPoints.PointsCoordU[PointIndex + IndexOffset];
+			const float NormalizedRadius = StrandsPoints.PointsRadius[PointIndex + IndexOffset];
+			const float NormalizedLength = CoordU * StrandsCurves.CurvesLength[CurveIndex];
+
+			FHairStrandsPositionFormat::Type& PackedPosition = OutPackedPositions[PointIndex + IndexOffset];
+			CopyVectorToPosition(PointPosition - HairBoxCenter, PackedPosition);
+			PackedPosition.ControlPointType = (PointIndex == 0) ? 1u : (PointIndex == (PointCount - 1) ? 2u : 0u);
+			PackedPosition.NormalizedRadius = uint8(FMath::Clamp(NormalizedRadius * 63.f, 0.f, 63.f));
+			PackedPosition.NormalizedLength = uint8(FMath::Clamp(NormalizedLength *255.f, 0.f, 255.f));
+			
+			const FVector2D RootUV = StrandsCurves.CurvesRootUV[CurveIndex];
+			FHairStrandsAttributeFormat::Type& PackedAttributes = OutPackedAttributes[PointIndex + IndexOffset];
+			PackedAttributes.RootU = uint32(FMath::Clamp(RootUV.X, 0.f, 1.f) * 0xFF);
+			PackedAttributes.RootV = uint32(FMath::Clamp(RootUV.Y, 0.f, 1.f) * 0xFF);
+			PackedAttributes.UCoord = FMath::Clamp(CoordU * 255.f, 0.f, 255.f);
+			PackedAttributes.Seed = CurveSeed;
+		}
+	}
+}
 
 void FHairStrandsDatas::BuildSimulationDatas(const uint32 StrandSize,
 	TArray<FHairStrandsPositionFormat::Type>& OutNodesPositions,
@@ -189,8 +740,92 @@ void FHairStrandsDatas::BuildSimulationDatas(const uint32 StrandSize,
 	}
 }
 
-void FHairStrandsDatas::BuildInternalDatas()
+///////////////////////////////////////////////////////////////////////////////////////////////////
+// Auto-generate Root UV data if not loaded
+static FVector2D SignNotZero(const FVector2D& v)
+{
+	return FVector2D((v.X >= 0.0) ? +1.0 : -1.0, (v.Y >= 0.0) ? +1.0 : -1.0);
+}
+
+// A Survey of Efficient Representations for Independent Unit Vectors
+// Reference: http://jcgt.org/published/0003/02/01/paper.pdf
+// Assume normalized input. Output is on [-1, 1] for each component.
+static FVector2D SphericalToOctahedron(const FVector& v)
+{
+	// Project the sphere onto the octahedron, and then onto the xy plane
+	FVector2D p = FVector2D(v.X, v.Y) * (1.0 / (abs(v.X) + abs(v.Y) + abs(v.Z)));
+	// Reflect the folds of the lower hemisphere over the diagonals
+	return (v.Z <= 0.0) ? ((FVector2D(1, 1) - FVector2D(abs(p.Y), abs(p.X))) * SignNotZero(p)) : p;
+}
+
+static void ComputeRootUV(FHairStrandsCurves& Curves, FHairStrandsPoints& Points)//; FStrands& Strands)
 {
+	TArray<FVector> RootPoints;
+	const uint32 CurveCount = Curves.Num();
+	RootPoints.Reserve(Curves.Num());
+	FVector MinAABB(MAX_FLT, MAX_FLT, MAX_FLT);
+	FVector MaxAABB(-MAX_FLT, -MAX_FLT, -MAX_FLT);
+	FMatrix Rotation = FRotationMatrix::Make(FRotator(0, 0, -90));
+	for (uint32 CurveIndex=0; CurveIndex< CurveCount; ++CurveIndex)
+	{
+		const uint32 Offset = Curves.CurvesOffset[CurveIndex];
+		check(Offset < uint32(Points.PointsPosition.Num()));
+		const FVector P = Rotation.TransformPosition(Points.PointsPosition[Offset]);
+
+		RootPoints.Add(P);
+		MinAABB.X = FMath::Min(P.X, MinAABB.X);
+		MinAABB.Y = FMath::Min(P.Y, MinAABB.Y);
+		MinAABB.Z = FMath::Min(P.Z, MinAABB.Z);
+
+		MaxAABB.X = FMath::Max(P.X, MaxAABB.X);
+		MaxAABB.Y = FMath::Max(P.Y, MaxAABB.Y);
+		MaxAABB.Z = FMath::Max(P.Z, MaxAABB.Z);
+	}
+
+	// Compute sphere bound
+	const FVector Extent = MaxAABB - MinAABB;
+	FSphere SBound;
+	SBound.Center = (MaxAABB + MinAABB) * 0.5f;
+	SBound.W = FMath::Max(Extent.X, FMath::Max(Extent.Y, Extent.Z));
+
+	// Project root point onto the bounding sphere and map it onto 
+	// an octahedron, which is unfold onto the unit space [0,1]^2
+	TArray<FVector2D> RootUVs;
+	RootUVs.Reserve(Curves.Num());
+	FVector2D MinUV(MAX_FLT, MAX_FLT);
+	FVector2D MaxUV(-MAX_FLT, -MAX_FLT);
+	for (const FVector& RootP : RootPoints)
+	{
+		FVector D = RootP - SBound.Center;
+		D.Normalize();
+		FVector2D UV = SphericalToOctahedron(D);
+		UV += FVector2D(1, 1);
+		UV *= 0.5f;
+		RootUVs.Add(UV);
+
+		MinUV.X = FMath::Min(UV.X, MinUV.X);
+		MinUV.Y = FMath::Min(UV.Y, MinUV.Y);
+
+		MaxUV.X = FMath::Max(UV.X, MaxUV.X);
+		MaxUV.Y = FMath::Max(UV.Y, MaxUV.Y);
+	}
+
+	// Find the minimal UV space cover by root point, and 
+	// offsets/scales it to maximize UV space
+	const FVector2D UVScale(1 / (MaxUV.X - MinUV.X), 1 / (MaxUV.Y - MinUV.Y));
+	const FVector2D UVOffset(-MinUV.X, -MinUV.Y);
+	uint32 Index = 0;
+	for (FVector2D& RootUV : Curves.CurvesRootUV)
+	{
+		RootUV = (RootUVs[Index++] + UVOffset) * UVScale;
+	}
+}
+
+void FHairStrandsDatas::BuildInternalDatas(bool bComputeRootUV)
+{
+	BoundingBox.Min = {  FLT_MAX,  FLT_MAX ,  FLT_MAX };
+	BoundingBox.Max = { -FLT_MAX, -FLT_MAX , -FLT_MAX };
+
 	if (GetNumCurves() > 0 && GetNumPoints() > 0)
 	{
 		TArray<FVector>::TIterator PositionIterator = StrandsPoints.PointsPosition.CreateIterator();
@@ -201,7 +836,7 @@ void FHairStrandsDatas::BuildInternalDatas()
 		TArray<uint32>::TIterator OffsetIterator = StrandsCurves.CurvesOffset.CreateIterator();
 		TArray<float>::TIterator LengthIterator = StrandsCurves.CurvesLength.CreateIterator();
 
-		StrandsPoints.MaxRadius = 0.0;
+		StrandsCurves.MaxRadius = 0.0;
 		StrandsCurves.MaxLength = 0.0;
 
 		uint32 StrandOffset = 0;
@@ -218,6 +853,8 @@ void FHairStrandsDatas::BuildInternalDatas()
 			FVector PreviousPosition(0.0, 0.0, 0.0);
 			for (uint32 PointIndex = 0; PointIndex < StrandCount; ++PointIndex, ++PositionIterator, ++RadiusIterator, ++CoordUIterator)
 			{
+				BoundingBox += *PositionIterator;
+
 				if (PointIndex > 0)
 				{
 					StrandLength += (*PositionIterator - PreviousPosition).Size();
@@ -225,7 +862,7 @@ void FHairStrandsDatas::BuildInternalDatas()
 				*CoordUIterator = StrandLength;
 				PreviousPosition = *PositionIterator;
 
-				StrandsPoints.MaxRadius = FMath::Max(StrandsPoints.MaxRadius, *RadiusIterator);
+				StrandsCurves.MaxRadius = FMath::Max(StrandsCurves.MaxRadius, *RadiusIterator);
 			}
 			*LengthIterator = StrandLength;
 			StrandsCurves.MaxLength = FMath::Max(StrandsCurves.MaxLength, StrandLength);
@@ -243,10 +880,15 @@ void FHairStrandsDatas::BuildInternalDatas()
 			for (uint32 PointIndex = 0; PointIndex < StrandCount; ++PointIndex, ++RadiusIterator, ++CoordUIterator)
 			{
 				*CoordUIterator /= *LengthIterator;
-				*RadiusIterator /= StrandsPoints.MaxRadius;
+				*RadiusIterator /= StrandsCurves.MaxRadius;
 			}
 			*LengthIterator /= StrandsCurves.MaxLength;
 		}
+
+		if (bComputeRootUV)
+		{
+			ComputeRootUV(StrandsCurves, StrandsPoints);
+		}
 	}
 }
 
diff --git a/Engine/Plugins/Experimental/HairStrands/Source/HairStrandsCore/Private/HairStrandsRendering.cpp b/Engine/Plugins/Experimental/HairStrands/Source/HairStrandsCore/Private/HairStrandsRendering.cpp
new file mode 100644
--- /dev/null
+++ b/Engine/Plugins/Experimental/HairStrands/Source/HairStrandsCore/Private/HairStrandsRendering.cpp
@@ -0,0 +1,454 @@
+// Copyright 1998-2019 Epic Games, Inc. All Rights Reserved. 
+
+#include "HairStrandsRendering.h"
+#include "HairStrandsDatas.h"
+#include "RenderGraphBuilder.h"
+#include "RenderGraphUtils.h"
+#include "ShaderParameterStruct.h"
+#include "GlobalShader.h"
+#include "HairStrandsInterface.h"
+
+static int32 GHairDeformationType = 0;
+static FAutoConsoleVariableRef CVarHairDeformationType(TEXT("r.HairStrands.DeformationType"), GHairDeformationType, TEXT("Type of procedural deformation applied on hair strands (0:bypass, 1:wave, 2:normal)"));
+
+static float GHairRaytracingRadiusScale = 1;
+static FAutoConsoleVariableRef CVarHairRaytracingRadiusScale(TEXT("r.HairStrands.RaytracingRadiusScale"), GHairRaytracingRadiusScale, TEXT("Scale factor for raytracing hair strands geometry"));
+
+
+static FIntVector ComputeDispatchCount(uint32 ItemCount, uint32 GroupSize)
+{
+	const uint32 BatchCount = FMath::DivideAndRoundUp(ItemCount, GroupSize);
+	const uint32 DispatchCountX = FMath::FloorToInt(FMath::Sqrt(BatchCount));
+	const uint32 DispatchCountY = DispatchCountX + FMath::DivideAndRoundUp(BatchCount - DispatchCountX * DispatchCountX, DispatchCountX);
+
+	check(DispatchCountX <= 65535);
+	check(DispatchCountY <= 65535);
+	check(BatchCount <= DispatchCountX * DispatchCountY);
+	return FIntVector(DispatchCountX, DispatchCountY, 1);
+}
+
+inline uint32 ComputeGroupSize()
+{
+	const uint32 GroupSize = IsRHIDeviceAMD() ? 64 : (IsRHIDeviceNVIDIA() ? 32 : 64);
+	check(GroupSize == 64 || GroupSize == 32);
+	return GroupSize;
+}
+
+inline uint32 GetGroupSizePermutation(uint32 GroupSize)
+{
+	return GroupSize == 64 ? 0 : (GroupSize == 32 ? 1 : 2);
+}
+
+class FDeformGuideCS : public FGlobalShader
+{
+	DECLARE_GLOBAL_SHADER(FDeformGuideCS);
+	SHADER_USE_PARAMETER_STRUCT(FDeformGuideCS, FGlobalShader);
+
+	class FGroupSize : SHADER_PERMUTATION_INT("PERMUTATION_GROUP_SIZE", 2);
+	class FDeformationType : SHADER_PERMUTATION_INT("PERMUTATION_DEFORMATION", 3);
+	using FPermutationDomain = TShaderPermutationDomain<FGroupSize, FDeformationType>;
+
+	BEGIN_SHADER_PARAMETER_STRUCT(FParameters, )
+		SHADER_PARAMETER(uint32, VertexCount)
+		SHADER_PARAMETER(uint32, IterationCount)
+		SHADER_PARAMETER_SRV(Buffer, SimRestPosePositionBuffer)
+		SHADER_PARAMETER_SRV(Buffer, SimRootIndexBuffer)
+		SHADER_PARAMETER_UAV(RWBuffer, OutSimDeformedPositionBuffer)
+	END_SHADER_PARAMETER_STRUCT()
+
+public:
+	static bool ShouldCompilePermutation(const FGlobalShaderPermutationParameters& Parameters) { return IsHairStrandsSupported(Parameters.Platform); }
+};
+
+IMPLEMENT_GLOBAL_SHADER(FDeformGuideCS, "/Engine/Private/HairStrands/HairStrandsGuideDeform.usf", "MainCS", SF_Compute);
+
+static void AddDeformSimHairStrandsPass(
+	FRDGBuilder& GraphBuilder,
+	uint32 DeformationType,
+	uint32 VertexCount,
+	FRHIShaderResourceView* SimRestPosePositionBuffer,
+	FRHIShaderResourceView* SimRootIndexBuffer,
+	FRHIUnorderedAccessView* OutSimDeformedPositionBuffer)
+{
+	static uint32 IterationCount = 0;
+	++IterationCount;
+
+	FDeformGuideCS::FParameters* Parameters = GraphBuilder.AllocParameters<FDeformGuideCS::FParameters>();
+	Parameters->SimRestPosePositionBuffer = SimRestPosePositionBuffer;
+	Parameters->SimRootIndexBuffer = SimRootIndexBuffer;
+	Parameters->OutSimDeformedPositionBuffer = OutSimDeformedPositionBuffer;
+	Parameters->VertexCount = VertexCount;
+	Parameters->IterationCount = IterationCount % 10000;
+	const uint32 GroupSize = ComputeGroupSize();
+
+	FDeformGuideCS::FPermutationDomain PermutationVector;
+	PermutationVector.Set<FDeformGuideCS::FGroupSize>(GetGroupSizePermutation(GroupSize));
+	PermutationVector.Set<FDeformGuideCS::FDeformationType>(DeformationType);
+
+	TShaderMap<FGlobalShaderType>* ShaderMap = GetGlobalShaderMap(ERHIFeatureLevel::SM5);
+
+	const uint32 DispatchCount = FMath::DivideAndRoundUp(VertexCount, GroupSize);
+	check(DispatchCount <= 65535);
+
+	TShaderMapRef<FDeformGuideCS> ComputeShader(ShaderMap, PermutationVector);
+	FComputeShaderUtils::AddPass(
+		GraphBuilder,
+		RDG_EVENT_NAME("DeformSimHairStrands"),
+		*ComputeShader,
+		Parameters,
+		FIntVector(DispatchCount, 1, 1));
+}
+
+///////////////////////////////////////////////////////////////////////////////////////////////////
+
+class FHairInterpolationCS : public FGlobalShader
+{
+	DECLARE_GLOBAL_SHADER(FHairInterpolationCS);
+	SHADER_USE_PARAMETER_STRUCT(FHairInterpolationCS, FGlobalShader);
+
+	class FGroupSize : SHADER_PERMUTATION_INT("PERMUTATION_GROUP_SIZE", 2);
+	class FDebug : SHADER_PERMUTATION_INT("PERMUTATION_DEBUG", 2);
+	using FPermutationDomain = TShaderPermutationDomain<FGroupSize, FDebug>;
+
+	BEGIN_SHADER_PARAMETER_STRUCT(FParameters, )
+		SHADER_PARAMETER(uint32, VertexCount)
+		SHADER_PARAMETER(uint32, DispatchCountX)
+		SHADER_PARAMETER_SRV(Buffer, RenderRestPosePositionBuffer)
+		SHADER_PARAMETER_UAV(RWBuffer, OutRenderDeformedPositionBuffer)
+
+		SHADER_PARAMETER_SRV(Buffer, SimRestPosePositionBuffer)
+		SHADER_PARAMETER_SRV(Buffer, DeformedSimPositionBuffer)
+
+		SHADER_PARAMETER_SRV(Buffer, Interpolation0Buffer)
+		SHADER_PARAMETER_SRV(Buffer, Interpolation1Buffer)
+
+		SHADER_PARAMETER_SRV(Buffer, SimAttributeBuffer)
+		SHADER_PARAMETER_UAV(RWBuffer, OutRenderAttributeBuffer)
+
+		END_SHADER_PARAMETER_STRUCT()
+
+public:
+	static bool ShouldCompilePermutation(const FGlobalShaderPermutationParameters& Parameters) { return IsHairStrandsSupported(Parameters.Platform); }
+};
+
+IMPLEMENT_GLOBAL_SHADER(FHairInterpolationCS, "/Engine/Private/HairStrands/HairStrandsInterpolation.usf", "MainCS", SF_Compute);
+
+static void AddHairStrandsInterpolationPass(
+	FRDGBuilder& GraphBuilder,
+	const uint32 VertexCount,
+	const FShaderResourceViewRHIRef& RenderRestPosePositionBuffer,
+	const FShaderResourceViewRHIRef& Interpolation0Buffer,
+	const FShaderResourceViewRHIRef& Interpolation1Buffer,
+	const FShaderResourceViewRHIRef& SimRestPosePositionBuffer,
+	const FShaderResourceViewRHIRef& SimDeformedPositionBuffer,
+	const FShaderResourceViewRHIRef& SimAttributeBuffer,
+	const FUnorderedAccessViewRHIRef& OutRenderPositionBuffer,
+	const FUnorderedAccessViewRHIRef& OutRenderAttributeBuffer)
+{
+	const bool bCopySimAttributesToRenderAttributes = SimAttributeBuffer != nullptr && OutRenderAttributeBuffer != nullptr;
+	const uint32 GroupSize = ComputeGroupSize();
+	const FIntVector DispatchCount = ComputeDispatchCount(VertexCount, GroupSize);
+
+	FHairInterpolationCS::FParameters* Parameters = GraphBuilder.AllocParameters<FHairInterpolationCS::FParameters>();
+	Parameters->RenderRestPosePositionBuffer = RenderRestPosePositionBuffer;
+	Parameters->SimRestPosePositionBuffer = SimRestPosePositionBuffer;
+	Parameters->DeformedSimPositionBuffer = SimDeformedPositionBuffer;
+	Parameters->Interpolation0Buffer = Interpolation0Buffer;
+	Parameters->Interpolation1Buffer = Interpolation1Buffer;
+	Parameters->OutRenderDeformedPositionBuffer = OutRenderPositionBuffer;
+	if (bCopySimAttributesToRenderAttributes)
+	{
+		Parameters->SimAttributeBuffer = SimAttributeBuffer;
+		Parameters->OutRenderAttributeBuffer = OutRenderAttributeBuffer;
+	}
+	Parameters->VertexCount = VertexCount;
+	Parameters->DispatchCountX = DispatchCount.X;
+
+	FHairInterpolationCS::FPermutationDomain PermutationVector;
+	PermutationVector.Set<FHairInterpolationCS::FGroupSize>(GetGroupSizePermutation(GroupSize));
+	PermutationVector.Set<FHairInterpolationCS::FDebug>(bCopySimAttributesToRenderAttributes ? 1 : 0);
+
+	TShaderMap<FGlobalShaderType>* ShaderMap = GetGlobalShaderMap(ERHIFeatureLevel::SM5);
+
+	TShaderMapRef<FHairInterpolationCS> ComputeShader(ShaderMap, PermutationVector);
+	FComputeShaderUtils::AddPass(
+		GraphBuilder,
+		RDG_EVENT_NAME("HairStrandsInterpolation"),
+		*ComputeShader,
+		Parameters,
+		DispatchCount);
+}
+
+///////////////////////////////////////////////////////////////////////////////////////////////////
+
+class FHairTangentCS : public FGlobalShader
+{
+	DECLARE_GLOBAL_SHADER(FHairTangentCS);
+	SHADER_USE_PARAMETER_STRUCT(FHairTangentCS, FGlobalShader);
+
+	class FGroupSize : SHADER_PERMUTATION_INT("PERMUTATION_GROUP_SIZE", 2);
+	using FPermutationDomain = TShaderPermutationDomain<FGroupSize>;
+
+	BEGIN_SHADER_PARAMETER_STRUCT(FParameters, )
+		SHADER_PARAMETER(uint32, VertexCount)
+		SHADER_PARAMETER(uint32, DispatchCountX)
+		SHADER_PARAMETER_SRV(Buffer, PositionBuffer)
+		SHADER_PARAMETER_UAV(RWBuffer, OutputTangentBuffer)
+	END_SHADER_PARAMETER_STRUCT()
+
+public:
+	static bool ShouldCompilePermutation(const FGlobalShaderPermutationParameters& Parameters) { return IsHairStrandsSupported(Parameters.Platform); }
+};
+
+IMPLEMENT_GLOBAL_SHADER(FHairTangentCS, "/Engine/Private/HairStrands/HairStrandsTangent.usf", "MainCS", SF_Compute);
+
+static void AddHairTangentPass(
+	FRDGBuilder& GraphBuilder,
+	uint32 VertexCount,
+	const FShaderResourceViewRHIRef& PositionBuffer,
+	const FUnorderedAccessViewRHIRef& OutTangentBuffer)
+{
+	const uint32 GroupSize = ComputeGroupSize();
+	const FIntVector DispatchCount = ComputeDispatchCount(VertexCount, GroupSize);
+
+	FHairTangentCS::FParameters* Parameters = GraphBuilder.AllocParameters<FHairTangentCS::FParameters>();
+	Parameters->PositionBuffer = PositionBuffer;
+	Parameters->OutputTangentBuffer = OutTangentBuffer;
+	Parameters->VertexCount = VertexCount;
+	Parameters->DispatchCountX = DispatchCount.X;
+
+	FHairTangentCS::FPermutationDomain PermutationVector;
+	PermutationVector.Set<FHairTangentCS::FGroupSize>(GetGroupSizePermutation(GroupSize));
+
+	TShaderMap<FGlobalShaderType>* ShaderMap = GetGlobalShaderMap(ERHIFeatureLevel::SM5);
+
+	TShaderMapRef<FHairTangentCS> ComputeShader(ShaderMap, PermutationVector);
+	FComputeShaderUtils::AddPass(
+		GraphBuilder,
+		RDG_EVENT_NAME("HairStrandsTangent"),
+		*ComputeShader,
+		Parameters,
+		DispatchCount);
+}
+
+///////////////////////////////////////////////////////////////////////////////////////////////////
+
+class FHairRaytracingGeometryCS : public FGlobalShader
+{
+	DECLARE_GLOBAL_SHADER(FHairRaytracingGeometryCS);
+	SHADER_USE_PARAMETER_STRUCT(FHairRaytracingGeometryCS, FGlobalShader);
+
+	class FGroupSize : SHADER_PERMUTATION_INT("PERMUTATION_GROUP_SIZE", 2);
+	using FPermutationDomain = TShaderPermutationDomain<FGroupSize>;
+
+	BEGIN_SHADER_PARAMETER_STRUCT(FParameters, )
+		SHADER_PARAMETER(uint32, VertexCount)
+		SHADER_PARAMETER(uint32, DispatchCountX)
+		SHADER_PARAMETER(FVector, StrandHairWorldOffset)
+		SHADER_PARAMETER(float, StrandHairRadius)
+		SHADER_PARAMETER_SRV(Buffer, PositionBuffer)
+		SHADER_PARAMETER_UAV(RWBuffer, OutputPositionBuffer)
+	END_SHADER_PARAMETER_STRUCT()
+
+public:
+	static bool ShouldCompilePermutation(const FGlobalShaderPermutationParameters& Parameters) { return IsHairStrandsSupported(Parameters.Platform); }
+};
+
+IMPLEMENT_GLOBAL_SHADER(FHairRaytracingGeometryCS, "/Engine/Private/HairStrands/HairStrandsRaytracingGeometry.usf", "MainCS", SF_Compute);
+
+static void AddGenerateRaytracingGeometryPass(
+	FRDGBuilder& GraphBuilder,
+	uint32 VertexCount,
+	float HairRadius,
+	const FVector& HairWorldOffset,
+	const FShaderResourceViewRHIRef& PositionBuffer,
+	const FUnorderedAccessViewRHIRef& OutPositionBuffer)
+{
+	const uint32 GroupSize = ComputeGroupSize();
+	const FIntVector DispatchCount = ComputeDispatchCount(VertexCount, GroupSize);
+
+	FHairRaytracingGeometryCS::FParameters* Parameters = GraphBuilder.AllocParameters<FHairRaytracingGeometryCS::FParameters>();
+	Parameters->VertexCount = VertexCount;
+	Parameters->DispatchCountX = DispatchCount.X;
+	Parameters->StrandHairWorldOffset = HairWorldOffset;
+	Parameters->StrandHairRadius = HairRadius;
+	Parameters->PositionBuffer = PositionBuffer;
+	Parameters->OutputPositionBuffer = OutPositionBuffer;
+
+	FHairRaytracingGeometryCS::FPermutationDomain PermutationVector;
+	PermutationVector.Set<FHairRaytracingGeometryCS::FGroupSize>(GetGroupSizePermutation(GroupSize));
+
+	TShaderMap<FGlobalShaderType>* ShaderMap = GetGlobalShaderMap(ERHIFeatureLevel::SM5);
+
+	TShaderMapRef<FHairRaytracingGeometryCS> ComputeShader(ShaderMap, PermutationVector);
+	FComputeShaderUtils::AddPass(
+		GraphBuilder,
+		RDG_EVENT_NAME("HairStrandsRaytracingGeometry"),
+		*ComputeShader,
+		Parameters,
+		DispatchCount);
+}
+
+///////////////////////////////////////////////////////////////////////////////////////////////////
+
+#if RHI_RAYTRACING
+static void UpdateHairAccelerationStructure(FRHICommandList& RHICmdList, FRayTracingGeometry* RayTracingGeometry)
+{
+	SCOPED_DRAW_EVENT(RHICmdList, CommitHairRayTracingGeometryUpdates);
+
+	TArray<FAccelerationStructureUpdateParams> Updates;
+	FAccelerationStructureUpdateParams Params;
+	Params.Geometry = RayTracingGeometry->RayTracingGeometryRHI;
+	Params.VertexBuffer = RayTracingGeometry->Initializer.PositionVertexBuffer;
+	Updates.Add(Params);
+
+	RHICmdList.UpdateAccelerationStructures(Updates);
+}
+
+static void BuildHairAccelerationStructure(FRHICommandList& RHICmdList, uint32 RaytracingVertexCount, FVertexBufferRHIRef& PositionBuffer, FRayTracingGeometry* OutRayTracingGeometry)
+{
+	FRayTracingGeometryInitializer Initializer;
+	Initializer.PositionVertexBuffer = PositionBuffer;
+	Initializer.VertexBufferByteOffset = 0;
+	Initializer.VertexBufferStride = FHairStrandsRaytracingFormat::SizeInByte;
+	Initializer.VertexBufferElementType = FHairStrandsRaytracingFormat::VertexElementType;
+	Initializer.IndexBuffer = nullptr;
+	Initializer.IndexBufferByteOffset = 0;
+	Initializer.GeometryType = RTGT_Triangles;
+	Initializer.TotalPrimitiveCount = RaytracingVertexCount;
+	Initializer.bFastBuild = true;
+	Initializer.bAllowUpdate = true;
+
+	OutRayTracingGeometry->SetInitializer(Initializer);
+	OutRayTracingGeometry->RayTracingGeometryRHI = RHICreateRayTracingGeometry(Initializer);
+	RHICmdList.BuildAccelerationStructure(OutRayTracingGeometry->RayTracingGeometryRHI);
+}
+#endif // RHI_RAYTRACING
+
+void ComputeHairStrandsInterpolation(
+	FRHICommandListImmediate& RHICmdList,
+	FHairStrandsInterpolationInput* InInput,
+	FHairStrandsInterpolationOutput* InOutput)
+{
+	if (!InInput || !InOutput) return;
+	FHairStrandsInterpolationInput& Input = *InInput;
+	FHairStrandsInterpolationOutput& Output = *InOutput;
+	Output.VFInput.Reset();
+
+	if (!Input.IsValid() || !Output.IsValid()) return;
+
+	DECLARE_GPU_STAT(HairStrandsInterpolation);
+	SCOPED_DRAW_EVENT(RHICmdList, HairStrandsInterpolation);
+	SCOPED_GPU_STAT(RHICmdList, HairStrandsInterpolation);
+
+	FRDGBuilder GraphBuilder(RHICmdList);
+
+	const uint32 CurrIndex = Output.CurrentIndex;
+	const uint32 PrevIndex = (Output.CurrentIndex+1)%2;
+
+	// Procedural deformers in place of actual physics solver
+	// 0: Simulation
+	// 1: Bypass
+	// 2: Wave
+	// 3: Straighten hair in direction of the root's normal
+	const int32 DeformationType = FMath::Clamp(GHairDeformationType, 0, Input.SimRootPointIndexBuffer ? 3 : 2) - 1;
+	if (GHairDeformationType > 0)
+	{
+		AddDeformSimHairStrandsPass(
+			GraphBuilder,
+			DeformationType,
+			Input.SimVertexCount,
+			Input.SimRestPosePositionBuffer->SRV,
+			Input.SimRootPointIndexBuffer ? Input.SimRootPointIndexBuffer->SRV : nullptr,
+			Output.SimDeformedPositionBuffer[CurrIndex]->UAV);
+	}
+
+	// If the deformation is driven by the physics simulation, then the output is always the 0 index
+	const uint32 SimIndex = GHairDeformationType == 0 ? 0 : CurrIndex;
+
+	// Debug mode:
+	// * None	: Display hair normally
+	// * Sim	: Show sim strands
+	// * Render : Show rendering strands with sim color influence
+	const EHairStrandsDebugMode DebugMode = GetHairStrandsDebugMode();
+	if (DebugMode == EHairStrandsDebugMode::SimHairStrands)
+	{
+		AddHairTangentPass(
+			GraphBuilder,
+			Input.SimVertexCount,
+			Output.SimDeformedPositionBuffer[SimIndex]->SRV,
+			Output.SimTangentBuffer->UAV);
+
+		GraphBuilder.Execute();
+	
+		Output.VFInput.HairPositionBuffer = Output.SimDeformedPositionBuffer[SimIndex]->SRV;
+		Output.VFInput.HairPreviousPositionBuffer = Output.SimDeformedPositionBuffer[SimIndex]->SRV;
+		Output.VFInput.HairTangentBuffer = Output.SimTangentBuffer->SRV;
+		Output.VFInput.HairAttributeBuffer = Input.SimAttributeBuffer->SRV;
+		Output.VFInput.VertexCount = Input.SimVertexCount;
+	}
+	else
+	{
+		const uint32 BufferSizeInBytes = Input.RenderVertexCount * FHairStrandsAttributeFormat::SizeInByte;
+		if (DebugMode == EHairStrandsDebugMode::RenderHairStrands && Output.RenderPatchedAttributeBuffer.NumBytes != BufferSizeInBytes)
+		{
+			Output.RenderPatchedAttributeBuffer.Release();
+			Output.RenderPatchedAttributeBuffer.Initialize(FHairStrandsAttributeFormat::SizeInByte, Input.RenderVertexCount, FHairStrandsAttributeFormat::Format, BUF_Static);
+		}
+
+		AddHairStrandsInterpolationPass(
+			GraphBuilder,
+			Input.RenderVertexCount,
+			Input.RenderRestPosePositionBuffer->SRV,
+			Input.Interpolation0Buffer->SRV,
+			Input.Interpolation1Buffer->SRV,
+			Input.SimRestPosePositionBuffer->SRV,
+			Output.SimDeformedPositionBuffer[SimIndex]->SRV,
+			DebugMode == EHairStrandsDebugMode::RenderHairStrands ? Input.SimAttributeBuffer->SRV : nullptr,
+			Output.RenderDeformedPositionBuffer[CurrIndex]->UAV,
+			DebugMode == EHairStrandsDebugMode::RenderHairStrands ? Output.RenderPatchedAttributeBuffer.UAV : nullptr);
+
+		AddHairTangentPass(
+			GraphBuilder,
+			Input.RenderVertexCount,
+			Output.RenderDeformedPositionBuffer[CurrIndex]->SRV,
+			Output.RenderTangentBuffer->UAV);
+
+#if RHI_RAYTRACING
+		if (IsRayTracingEnabled())
+		{
+			AddGenerateRaytracingGeometryPass(
+				GraphBuilder,
+				Input.RenderVertexCount,
+				Input.HairRadius * GHairRaytracingRadiusScale,
+				Input.HairWorldOffset,
+				Output.RenderDeformedPositionBuffer[CurrIndex]->SRV,
+				Input.RaytracingPositionBuffer->UAV);
+		}
+#endif
+		GraphBuilder.Execute();
+
+		Output.VFInput.HairPositionBuffer = Output.RenderDeformedPositionBuffer[CurrIndex]->SRV;
+		Output.VFInput.HairPreviousPositionBuffer = Output.RenderDeformedPositionBuffer[PrevIndex]->SRV;
+		Output.VFInput.HairTangentBuffer = Output.RenderTangentBuffer->SRV;
+		Output.VFInput.HairAttributeBuffer = DebugMode == EHairStrandsDebugMode::RenderHairStrands ? Output.RenderPatchedAttributeBuffer.SRV : Input.RenderAttributeBuffer->SRV;
+		Output.VFInput.VertexCount = Input.RenderVertexCount;
+
+#if RHI_RAYTRACING
+		if (IsRayTracingEnabled())
+		{
+			FRHIUnorderedAccessView* UAV = Input.RaytracingPositionBuffer->UAV;
+			RHICmdList.TransitionResources(EResourceTransitionAccess::EReadable, EResourceTransitionPipeline::EComputeToGfx, &UAV, 1);
+
+			const bool bNeedFullBuild = !Input.bIsRTGeometryInitialized;
+			if (bNeedFullBuild)
+				BuildHairAccelerationStructure(RHICmdList, Input.RaytracingVertexCount, Input.RaytracingPositionBuffer->Buffer, Input.RaytracingGeometry);
+			else
+				UpdateHairAccelerationStructure(RHICmdList, Input.RaytracingGeometry);
+			Input.bIsRTGeometryInitialized = true;
+		}
+#endif
+	}
+
+	Output.CurrentIndex = PrevIndex;
+}
diff --git a/Engine/Plugins/Experimental/HairStrands/Source/HairStrandsCore/Private/HairStrandsVertexFactory.cpp b/Engine/Plugins/Experimental/HairStrands/Source/HairStrandsCore/Private/HairStrandsVertexFactory.cpp
new file mode 100644
--- /dev/null
+++ b/Engine/Plugins/Experimental/HairStrands/Source/HairStrandsCore/Private/HairStrandsVertexFactory.cpp
@@ -0,0 +1,215 @@
+// Copyright 1998-2019 Epic Games, Inc. All Rights Reserved.
+
+/*=============================================================================
+	StrandHairVertexFactory.cpp: Strand hair vertex factory implementation
+=============================================================================*/
+
+#include "HairStrandsVertexFactory.h"
+#include "SceneView.h"
+#include "MeshBatch.h"
+#include "ShaderParameterUtils.h"
+#include "Rendering/ColorVertexBuffer.h"
+#include "MeshMaterialShader.h"
+#include "HairStrandsInterface.h"
+
+static float GStrandHairWidth = 0.0f;
+static FAutoConsoleVariableRef CVarStrandHairWidth(TEXT("r.HairStrands.StrandWidth"), GStrandHairWidth, TEXT("Width of hair strand"));
+
+float FHairStrandsVertexFactory::GetMaxStrandRadius() const
+{
+	return GStrandHairWidth > 0 ? GStrandHairWidth * 0.5f : Data.MaxStrandRadius;
+}
+
+#define OPTIMIZE_OFF 0
+
+#if OPTIMIZE_OFF
+	#pragma optimize("", off)
+#endif
+
+/////////////////////////////////////////////////////////////////////////////////////////
+
+template<typename T> inline void BindParam(FMeshDrawSingleShaderBindings& ShaderBindings, const FShaderResourceParameter& Param, T* Value)	{ if (Param.IsBound() && Value) ShaderBindings.Add(Param, Value); }
+template<typename T> inline void BindParam(FMeshDrawSingleShaderBindings& ShaderBindings, const FShaderParameter& Param, const T& Value)	{ if (Param.IsBound()) ShaderBindings.Add(Param, Value); }
+
+class FHairStrandsVertexFactoryShaderParameters : public FVertexFactoryShaderParameters
+{
+public:
+
+	FShaderParameter Radius;
+	FShaderParameter Length;
+	FShaderParameter RadiusAtDepth1_Primary;
+	FShaderParameter RadiusAtDepth1_Velocity;
+	FShaderParameter WorldOffset;
+	FShaderParameter Density;
+
+	FShaderResourceParameter PositionBuffer;
+	FShaderResourceParameter PreviousPositionBuffer;
+	FShaderResourceParameter AttributeBuffer;
+	FShaderResourceParameter TangentBuffer;
+
+	virtual void Bind(const FShaderParameterMap& ParameterMap) override
+	{
+		Radius.Bind(ParameterMap, TEXT("HairStrandsVF_Radius"));
+		Length.Bind(ParameterMap, TEXT("HairStrandsVF_Length"));
+		RadiusAtDepth1_Primary.Bind(ParameterMap, TEXT("HairStrandsVF_RadiusAtDepth1_Primary"));
+		RadiusAtDepth1_Velocity.Bind(ParameterMap, TEXT("HairStrandsVF_RadiusAtDepth1_Velocity"));
+		WorldOffset.Bind(ParameterMap, TEXT("HairStrandsVF_WorldOffset"));
+		Density.Bind(ParameterMap, TEXT("HairStrandsVF_Density"));	
+
+		PositionBuffer.Bind(ParameterMap, TEXT("HairStrandsVF_PositionBuffer"));
+		PreviousPositionBuffer.Bind(ParameterMap, TEXT("HairStrandsVF_PreviousPositionBuffer"));
+		AttributeBuffer.Bind(ParameterMap, TEXT("HairStrandsVF_AttributeBuffer"));
+		TangentBuffer.Bind(ParameterMap, TEXT("HairStrandsVF_TangentBuffer"));
+	}
+
+	virtual void Serialize(FArchive& Ar) override
+	{
+		Ar << Radius;
+		Ar << Length;
+		Ar << RadiusAtDepth1_Primary;
+		Ar << RadiusAtDepth1_Velocity;
+		Ar << WorldOffset;
+		Ar << Density;
+
+		Ar << PositionBuffer;
+		Ar << PreviousPositionBuffer;
+		Ar << AttributeBuffer;
+		Ar << TangentBuffer;
+	}
+
+	virtual void GetElementShaderBindings(
+		const FSceneInterface* Scene,
+		const FSceneView* View,
+		const FMeshMaterialShader* Shader,
+		const EVertexInputStreamType InputStreamType,
+		ERHIFeatureLevel::Type FeatureLevel,
+		const FVertexFactory* VertexFactory,
+		const FMeshBatchElement& BatchElement,
+		FMeshDrawSingleShaderBindings& ShaderBindings,
+		FVertexInputStreamArray& VertexStreams
+	) const override
+	{
+		const FHairStrandsVertexFactory* VF = static_cast<const FHairStrandsVertexFactory*>(VertexFactory);
+		
+		const FMinHairRadiusAtDepth1 MinRadiusAtDepth1 = ComputeMinStrandRadiusAtDepth1(
+			FIntPoint(View->UnconstrainedViewRect.Width(), View->UnconstrainedViewRect.Height()),
+			View->FOV,
+			GetHairVisibilitySampleCount(),
+			0.f);
+
+		BindParam(ShaderBindings, PositionBuffer, VF->GetPositionSRV());
+		BindParam(ShaderBindings, PreviousPositionBuffer, VF->GetPreviousPositionSRV());
+		BindParam(ShaderBindings, AttributeBuffer, VF->GetAttributeSRV());
+		BindParam(ShaderBindings, TangentBuffer, VF->GetTangentSRV());
+		BindParam(ShaderBindings, Radius, VF->GetMaxStrandRadius());
+		BindParam(ShaderBindings, Length, VF->GetMaxStrandLength());
+		BindParam(ShaderBindings, WorldOffset, VF->GetWorldOffset());
+		BindParam(ShaderBindings, Density, VF->GetHairDensity());
+		BindParam(ShaderBindings, RadiusAtDepth1_Primary, MinRadiusAtDepth1.Primary);
+		BindParam(ShaderBindings, RadiusAtDepth1_Velocity, MinRadiusAtDepth1.Velocity);		
+	}
+};
+
+/**
+ * Should we cache the material's shadertype on this platform with this vertex factory? 
+ */
+bool FHairStrandsVertexFactory::ShouldCompilePermutation(EShaderPlatform Platform, const class FMaterial* Material, const class FShaderType* ShaderType)
+{
+	return (Material->GetMaterialDomain() == MD_Surface && Material->IsUsedWithHairStrands() && Platform == EShaderPlatform::SP_PCD3D_SM5) || Material->IsSpecialEngineMaterial();
+}
+
+void FHairStrandsVertexFactory::ModifyCompilationEnvironment(const FVertexFactoryType* Type, EShaderPlatform Platform, const FMaterial* Material, FShaderCompilerEnvironment& OutEnvironment)
+{
+	OutEnvironment.SetDefine(TEXT("VF_SUPPORTS_PRIMITIVE_SCENE_DATA"), Type->SupportsPrimitiveIdStream() && UseGPUScene(Platform, GetMaxSupportedFeatureLevel(Platform)));
+	OutEnvironment.SetDefine(TEXT("VF_STRAND_HAIR"), TEXT("1"));
+}
+
+void FHairStrandsVertexFactory::ValidateCompiledResult(const FVertexFactoryType* Type, EShaderPlatform Platform, const FShaderParameterMap& ParameterMap, TArray<FString>& OutErrors)
+{
+	if (Type->SupportsPrimitiveIdStream() 
+		&& UseGPUScene(Platform, GetMaxSupportedFeatureLevel(Platform)) 
+		&& ParameterMap.ContainsParameterAllocation(FPrimitiveUniformShaderParameters::StaticStructMetadata.GetShaderVariableName()))
+	{
+		OutErrors.AddUnique(*FString::Printf(TEXT("Shader attempted to bind the Primitive uniform buffer even though Vertex Factory %s computes a PrimitiveId per-instance.  This will break auto-instancing.  Shaders should use GetPrimitiveData(PrimitiveId).Member instead of Primitive.Member."), Type->GetName()));
+	}
+}
+
+void FHairStrandsVertexFactory::SetData(const FDataType& InData)
+{
+	check(IsInRenderingThread());
+	Data = InData;
+	UpdateRHI();
+}
+
+/**
+* Copy the data from another vertex factory
+* @param Other - factory to copy from
+*/
+void FHairStrandsVertexFactory::Copy(const FHairStrandsVertexFactory& Other)
+{
+	FHairStrandsVertexFactory* VertexFactory = this;
+	const FDataType* DataCopy = &Other.Data;
+	ENQUEUE_RENDER_COMMAND(FHairStrandsVertexFactoryCopyData)(
+		[VertexFactory, DataCopy](FRHICommandListImmediate& RHICmdList)
+		{
+			VertexFactory->Data = *DataCopy;
+		});
+	BeginUpdateResourceRHI(this);
+}
+
+void FHairStrandsVertexFactory::InitRHI()
+{
+	bNeedsDeclaration = false;
+	bSupportsManualVertexFetch = true;
+
+	// We create different streams based on feature level
+	check(HasValidFeatureLevel());
+
+	// VertexFactory needs to be able to support max possible shader platform and feature level
+	// in case if we switch feature level at runtime.
+	const bool bCanUseGPUScene = UseGPUScene(GMaxRHIShaderPlatform, GMaxRHIFeatureLevel);
+
+	FVertexDeclarationElementList Elements;
+	SetPrimitiveIdStreamIndex(EVertexInputStreamType::Default, -1);
+	if (GetType()->SupportsPrimitiveIdStream() && bCanUseGPUScene)
+	{
+		// When the VF is used for rendering in normal mesh passes, this vertex buffer and offset will be overridden
+		Elements.Add(AccessStreamComponent(FVertexStreamComponent(&GPrimitiveIdDummy, 0, 0, sizeof(uint32), VET_UInt, EVertexStreamUsage::Instancing), 13));
+		SetPrimitiveIdStreamIndex(EVertexInputStreamType::Default, Elements.Last().StreamIndex);
+		bNeedsDeclaration = true;
+	}
+
+	check(Streams.Num() > 0);
+
+	InitDeclaration(Elements);
+	check(IsValidRef(GetDeclaration()));
+}
+
+FVertexFactoryShaderParameters* FHairStrandsVertexFactory::ConstructShaderParameters(EShaderFrequency ShaderFrequency)
+{
+	if (ShaderFrequency == SF_Vertex)
+	{
+		return new FHairStrandsVertexFactoryShaderParameters();
+	}
+
+#if RHI_RAYTRACING
+	if (ShaderFrequency == SF_RayHitGroup)
+	{
+		return new FHairStrandsVertexFactoryShaderParameters();
+	}
+#endif // RHI_RAYTRACING
+
+	return NULL;
+}
+
+void FHairStrandsVertexFactory::ReleaseRHI()
+{
+	FVertexFactory::ReleaseRHI();
+}
+
+
+IMPLEMENT_VERTEX_FACTORY_TYPE_EX(FHairStrandsVertexFactory,"/Engine/Private/HairStrands/HairStrandsVertexFactory.ush",true,false,true,true,true,true,true);
+
+#if OPTIMIZE_OFF
+	#pragma optimize("", on)
+#endif
\ No newline at end of file
diff --git a/Engine/Plugins/Experimental/HairStrands/Source/HairStrandsCore/Public/GroomActor.h b/Engine/Plugins/Experimental/HairStrands/Source/HairStrandsCore/Public/GroomActor.h
new file mode 100644
--- /dev/null
+++ b/Engine/Plugins/Experimental/HairStrands/Source/HairStrandsCore/Public/GroomActor.h
@@ -0,0 +1,24 @@
+// Copyright 1998-2019 Epic Games, Inc. All Rights Reserved.
+
+#pragma once
+
+#include "CoreMinimal.h"
+#include "UObject/ObjectMacros.h"
+#include "GameFramework/Actor.h"
+#include "GroomActor.generated.h"
+
+
+/** An actor that renders a simulated hair */
+UCLASS(HideCategories = (Input, Replication, Mobility), showcategories = ("Input|MouseInput", "Input|TouchInput"))
+class HAIRSTRANDSCORE_API AGroomActor : public AActor
+{
+	GENERATED_UCLASS_BODY()
+
+	/** Strand hair component that performs simulation and rendering */
+	UPROPERTY(Category = StrandHair, VisibleAnywhere, BlueprintReadOnly)
+	class UGroomComponent* GroomComponent;
+
+public:
+	/** Returns GroomComponent subobject **/
+	class UGroomComponent* GetGroomComponent() const { return GroomComponent; }
+};
\ No newline at end of file
diff --git a/Engine/Plugins/Experimental/HairStrands/Source/HairStrandsCore/Public/GroomAsset.h b/Engine/Plugins/Experimental/HairStrands/Source/HairStrandsCore/Public/GroomAsset.h
new file mode 100644
--- /dev/null
+++ b/Engine/Plugins/Experimental/HairStrands/Source/HairStrandsCore/Public/GroomAsset.h
@@ -0,0 +1,208 @@
+// Copyright 1998-2019 Epic Games, Inc. All Rights Reserved.
+
+
+#pragma once
+
+#include "UObject/Object.h"
+#include "UObject/ObjectMacros.h"
+#include "HairDescription.h"
+#include "HairStrandsDatas.h"
+#include "RenderResource.h"
+
+#include "GroomAsset.generated.h"
+
+class UMaterialInterface;
+class UNiagaraSystem;
+
+/* Render buffers that will be used for rendering */
+struct FHairStrandsResource : public FRenderResource
+{
+	/** Build the hair strands resource */
+	FHairStrandsResource(FHairStrandsDatas* HairStrandsDatas);
+
+	/* Init the buffer */
+	virtual void InitRHI() override;
+
+	/* Release the buffer */
+	virtual void ReleaseRHI() override;
+
+	/* Get the resource name */
+	virtual FString GetFriendlyName() const override { return TEXT("FHairStrandsResource"); }
+
+	/* Strand hair rest position buffer */
+	FRWBuffer RestPositionBuffer;
+
+	/* Strand hair (previous) position buffer */
+	FRWBuffer DeformedPositionBuffer[2];
+
+	/* Strand hair offset buffer */
+	FRWBuffer TangentBuffer;
+
+	/* Strand hair offset buffer */
+	FRWBuffer AttributeBuffer;
+
+	/* Pointer to the hair strands datas */
+	FHairStrandsDatas* StrandsDatas;
+};
+
+struct FHairStrandsInterpolationResource : public FRenderResource
+{
+	/** Build the hair strands resource */
+	FHairStrandsInterpolationResource(const FHairStrandsInterpolationDatas& InterpolationData, const FHairStrandsDatas& SimDatas);
+
+	/* Init the buffer */
+	virtual void InitRHI() override;
+
+	/* Release the buffer */
+	virtual void ReleaseRHI() override;
+
+	/* Get the resource name */
+	virtual FString GetFriendlyName() const override { return TEXT("FHairStrandsInterplationResource"); }
+
+	FRWBuffer Interpolation0Buffer;
+	FRWBuffer Interpolation1Buffer;
+
+	/* Store interpolation data until the resource is created */
+	TArray<FHairStrandsInterpolation0Format::Type> Interpolation0;
+	TArray<FHairStrandsInterpolation1Format::Type> Interpolation1;
+
+	// For debug purpose only (should be remove once all hair simulation is correctly handled)
+	FRWBuffer SimRootPointIndexBuffer;
+	TArray<FHairStrandsRootIndexFormat::Type> SimRootPointIndex;
+};
+
+#if RHI_RAYTRACING
+struct FHairStrandsRaytracingResource : public FRenderResource
+{
+	/** Build the hair strands resource */
+	FHairStrandsRaytracingResource(FHairStrandsDatas* HairStrandsDatas);
+
+	/* Init the buffer */
+	virtual void InitRHI() override;
+
+	/* Release the buffer */
+	virtual void ReleaseRHI() override;
+
+	/* Get the resource name */
+	virtual FString GetFriendlyName() const override { return TEXT("FHairStrandsRaytracingResource"); }
+
+	FRWBuffer PositionBuffer;
+	FRayTracingGeometry RayTracingGeometry;
+	uint32 VertexCount;
+};
+#endif
+
+USTRUCT(BlueprintType)
+struct HAIRSTRANDSCORE_API FHairGroupRenderSettings
+{
+	GENERATED_BODY()
+
+	UPROPERTY(VisibleAnywhere, Category = "Info")
+	int32 GroupID = 0;
+
+	UPROPERTY(VisibleAnywhere, Category = "Info", meta = (DisplayName = "Curve Count"))
+	int32 NumCurves = 0;
+
+	UPROPERTY(EditAnywhere, Category="Rendering")
+	UMaterialInterface* Material = nullptr;
+};
+
+USTRUCT(BlueprintType)
+struct HAIRSTRANDSCORE_API FHairGroupSimulationSettings
+{
+	GENERATED_BODY()
+
+	UPROPERTY(VisibleAnywhere, Category = "Info")
+	int32 GroupID = 0;
+
+	UPROPERTY(VisibleAnywhere, Category = "Info", meta = (DisplayName = "Curve Count"))
+	int32 NumCurves = 0;
+
+	UPROPERTY(VisibleAnywhere, Category = "Info", meta = (DisplayName = "Auto-generated Guides", ToolTip = "Checked when guides were auto-generated"))
+	bool bIsAutoGenerated = false;
+
+	UPROPERTY(EditAnywhere, Category="Simulation", meta = (DisplayName = "Niagara System Asset"))
+	UNiagaraSystem* Asset = nullptr;
+
+	// Add other physics/simulation settings
+};
+
+/**
+ * Implements an asset that can be used to store hair strands
+ */
+UCLASS(BlueprintType, hidecategories = (Object))
+class HAIRSTRANDSCORE_API UGroomAsset : public UObject
+{
+	GENERATED_BODY()
+
+public:
+
+	// #ueent_todo: Hair strands data should be computed and saved in the DDC
+	/** Hair strands data for rendering */
+	FHairStrandsDatas HairRenderData;
+
+	/** Guide strands data for simulation */
+	FHairStrandsDatas HairSimulationData;
+
+	UPROPERTY(EditAnywhere, EditFixedSize, BlueprintReadWrite, Category = "HairGroupSettings", meta = (DisplayName = "Rendering"))
+	TArray<FHairGroupRenderSettings> RenderHairGroups;
+
+	UPROPERTY(EditAnywhere, EditFixedSize, BlueprintReadWrite, Category = "HairGroupSettings", meta = (DisplayName = "Simulation"))
+	TArray<FHairGroupSimulationSettings> SimulationHairGroups;
+
+	//~ Begin UObject Interface.
+	virtual void PostLoad() override;
+	virtual void BeginDestroy() override;
+	virtual void Serialize(FArchive& Ar) override;
+
+	/** Density factor for converting hair into guide curve if no guides are provided. */
+	UPROPERTY(EditAnywhere, BlueprintReadWrite, Category = "BuildSettings", meta = (ClampMin = "0.01", UIMin = "0.01", UIMax = "1.0"))
+	float HairToGuideDensity = 0.1f;
+
+#if WITH_EDITOR
+
+	/**  Part of Uobject interface  */
+	virtual void PostEditChangeProperty(FPropertyChangedEvent& PropertyChangedEvent) override;
+
+#endif // WITH_EDITOR
+
+#if WITH_EDITORONLY_DATA
+
+	/** Asset data to be used when re-importing */
+	UPROPERTY(VisibleAnywhere, Instanced, Category = ImportSettings)
+	class UAssetImportData* AssetImportData;
+
+	/** Retrievde the asset tags*/
+	virtual void GetAssetRegistryTags(TArray<FAssetRegistryTag>& OutTags) const override;
+
+	/** Part of Uobject interface */
+	virtual void PostInitProperties() override;
+
+#endif // WITH_EDITORONLY_DATA
+
+	/** Initialize resources. */
+	void InitResource();
+
+	/** Update resources. */
+	void UpdateResource();
+
+	/** Release the hair strands resource. */
+	void ReleaseResource();
+
+	/**
+	 * Initializes an instance for use with this vector field.
+	 */
+	void InitInstance(class FVectorFieldInstance* Instance, bool bPreviewInstance);
+
+	void Reset();
+
+//private : 
+
+	TUniquePtr<FHairDescription> HairDescription;
+
+	/** Interpolated hair render resource to be allocated */
+	FHairStrandsResource* HairStrandsResource;
+
+	/** Guide render resource to be allocated */
+	FHairStrandsResource* HairSimulationResource;
+};
diff --git a/Engine/Plugins/Experimental/HairStrands/Source/HairStrandsCore/Public/GroomComponent.h b/Engine/Plugins/Experimental/HairStrands/Source/HairStrandsCore/Public/GroomComponent.h
new file mode 100644
--- /dev/null
+++ b/Engine/Plugins/Experimental/HairStrands/Source/HairStrandsCore/Public/GroomComponent.h
@@ -0,0 +1,92 @@
+// Copyright 1998-2019 Epic Games, Inc. All Rights Reserved.
+
+#pragma once
+
+#include "CoreMinimal.h"
+#include "UObject/ObjectMacros.h"
+#include "Engine/EngineTypes.h"
+#include "Components/MeshComponent.h"
+#include "GroomAsset.h"
+#include "RHIDefinitions.h"
+
+#include "GroomComponent.generated.h"
+
+/** Component that allows you to specify custom triangle mesh geometry */
+UCLASS(HideCategories = (Object, Physics, Activation, Mobility, "Components|Activation"), editinlinenew, meta = (BlueprintSpawnableComponent), ClassGroup = Rendering)
+class HAIRSTRANDSCORE_API UGroomComponent : public UMeshComponent
+{
+	GENERATED_UCLASS_BODY()
+
+public:
+
+	/** Hair strand asset used for rendering. */
+	UPROPERTY(EditAnywhere, BlueprintReadWrite, Category = "Groom")
+	UGroomAsset* GroomAsset;
+
+	/** Controls the hair density, to reduce or increase hair count during shadow rendering. This allows to increase/decrease the shadowing on hair when the number of strand is not realistic */
+	UPROPERTY(EditAnywhere, BlueprintReadWrite, Category = "HairStrands Rendering", meta = (ClampMin = "0.0001", UIMin = "0.01", UIMax = "10.0"))
+	float HairDensity;
+
+	/** Threshold for merging consecutive hair segments when loading asset */
+	UPROPERTY(EditAnywhere, BlueprintReadWrite, Category = "HairStrands Rendering", meta = (ClampMin = "0.01", UIMin = "0.01", UIMax = "1.0"))
+	float MergeThreshold;
+
+	//~ Begin UActorComponent Interface.
+	virtual void OnRegister() override;
+	virtual void OnUnregister() override;
+	virtual void TickComponent(float DeltaTime, enum ELevelTick TickType, FActorComponentTickFunction *ThisTickFunction) override;
+	//~ End UActorComponent Interface.
+
+	//~ Begin USceneComponent Interface.
+	virtual FBoxSphereBounds CalcBounds(const FTransform& LocalToWorld) const override;
+	//~ End USceneComponent Interface.
+
+	//~ Begin UPrimitiveComponent Interface.
+	virtual FPrimitiveSceneProxy* CreateSceneProxy() override;
+	//~ End UPrimitiveComponent Interface.
+
+	//~ Begin UMeshComponent Interface.
+	virtual void PostLoad() override;
+	virtual int32 GetNumMaterials() const override;
+	virtual UMaterialInterface* GetMaterial(int32 ElementIndex) const override;
+	//~ End UMeshComponent Interface.
+
+	/** Return the guide hairs datas */
+	FHairStrandsDatas* GetGuideStrandsDatas() { return GroomAsset ? &GroomAsset->HairSimulationData : nullptr; }
+
+	/** Return the guide hairs resources*/
+	FHairStrandsResource* GetGuideStrandsResource() { return GroomAsset ? GroomAsset->HairSimulationResource : nullptr; }
+
+#if WITH_EDITOR
+	virtual void PostEditChangeProperty(FPropertyChangedEvent& PropertyChangedEvent) override;
+#endif
+	// #hair_todo: cached/serialize interpolation data
+	FHairStrandsInterpolationResource* InterpolationResource = nullptr;
+	struct FHairStrandsInterpolationOutput* InterpolationOutput = nullptr;
+	struct FHairStrandsInterpolationInput* InterpolationInput = nullptr;
+
+#if RHI_RAYTRACING
+	FHairStrandsRaytracingResource* RaytracingResources = nullptr;
+#endif
+
+private:
+	void* InitializedResources;
+
+	void InitResources();
+	void ReleaseResources();
+
+	virtual void GetUsedMaterials(TArray<UMaterialInterface*>& OutMaterials, bool bGetDebugMaterials) const override;
+
+	friend class FGroomComponentRecreateRenderStateContext;
+};
+
+/** Used to recreate render context for all GroomComponents that use a given GroomAsset */
+class HAIRSTRANDSCORE_API FGroomComponentRecreateRenderStateContext
+{
+public:
+	FGroomComponentRecreateRenderStateContext(UGroomAsset* GroomAsset);
+	~FGroomComponentRecreateRenderStateContext();
+
+private:
+	TArray<UGroomComponent*> GroomComponents;
+};
diff --git a/Engine/Plugins/Experimental/HairStrands/Source/HairStrandsCore/Public/HairStrandsCore.h b/Engine/Plugins/Experimental/HairStrands/Source/HairStrandsCore/Public/HairStrandsCore.h
--- a/Engine/Plugins/Experimental/HairStrands/Source/HairStrandsCore/Public/HairStrandsCore.h
+++ b/Engine/Plugins/Experimental/HairStrands/Source/HairStrandsCore/Public/HairStrandsCore.h
@@ -20,5 +20,3 @@ public:
 		return true;
 	}
 };
-
-IMPLEMENT_MODULE(FHairStrandsCore, HairStrandsCore);
