Copying //UE4/Dev-Physics to Dev-Main (//UE4/Dev-Main) [at] 9098278

diff --git a/Engine/Plugins/Experimental/HairStrands/Source/HairStrandsCore/Public/HairStrandsDatas.h b/Engine/Plugins/Experimental/HairStrands/Source/HairStrandsCore/Public/HairStrandsDatas.h
--- a/Engine/Plugins/Experimental/HairStrands/Source/HairStrandsCore/Public/HairStrandsDatas.h
+++ b/Engine/Plugins/Experimental/HairStrands/Source/HairStrandsCore/Public/HairStrandsDatas.h
@@ -6,50 +6,113 @@
 #include "Engine/EngineTypes.h"
 #include "RHIDefinitions.h"
 #include "PackedNormal.h"
+#include "HairStrandsVertexFactory.h"
 
 DEFINE_LOG_CATEGORY_STATIC(LogHairStrands, Log, All);
 
-/**  A packed position for the hair strands datas .*/
-struct HAIRSTRANDSCORE_API FHairStrandsPackedPosition
-{
-	// #todo_hair: repack the position into relative 16 bits position. These position should be relative to the AABB center. 
-	// If Float point is not enough we could use fixed point value
-#if 0 
-	//FFloat16 X, Y, Z; 
-	//uint16 ControlPointType : 2; 
-	//uint16 NormalizedRadius : 6;
-	//uint16 UCoord : 8;
-#else
-	float X, Y, Z;
-	uint32 ControlPointType : 2;
-	uint32 NormalizedRadius : 6;
-	uint32 UCoord : 8;
-	uint32 NormalizedLength : 8;
-	uint32 Seed : 8;
-#endif
-};
-
-/** Hair strands position format */
-struct HAIRSTRANDSCORE_API FHairStrandsPositionFormat
-{
-	using Type = FHairStrandsPackedPosition;
+struct FPackedHairVertex
+{
+	FFloat16 X, Y, Z;
+	uint8 ControlPointType : 2;
+	uint8 NormalizedRadius : 6;
+	uint8 NormalizedLength;
+};
+
+struct FPackedHairAttributeVertex
+{
+	uint8 RootU;
+	uint8 RootV;
+	uint8 UCoord;
+	uint8 Seed;
+};
+
+struct FHairInterpolation0Vertex
+{
+	uint16 Index0;
+	uint16 Index1;
+	uint16 Index2;
+
+	uint8 VertexWeight0;
+	uint8 VertexWeight1;
+};
+
+struct FHairInterpolation1Vertex
+{
+	uint8 VertexIndex0;
+	uint8 VertexIndex1;
+	uint8 VertexIndex2;
+	uint8 Pad0;
+};
+
+struct FHairStrandsPositionFormat
+{
+	typedef FPackedHairVertex Type;
+	static const uint32 ComponentCount = 1;
+	static const uint32 SizeInByte = sizeof(Type);
+	static const EVertexElementType VertexElementType = VET_UShort4;
+	static const EPixelFormat Format = PF_R16G16B16A16_UINT;
+};
+
+struct FHairStrandsAttributeFormat
+{
+	typedef FPackedHairAttributeVertex Type;
 	static const uint32 ComponentCount = 1;
 	static const uint32 SizeInByte = sizeof(Type);
-	static const EVertexElementType VertexElementType = VET_Float4;	// VET_UShort4;
-	static const EPixelFormat Format = PF_A32B32G32R32F;			// PF_R16G16B16A16_UINT;
+	static const EVertexElementType VertexElementType = VET_UByte4;
+	static const EPixelFormat Format = PF_R8G8B8A8;
 };
 
-/** Hair Strands tangent format */
-struct HAIRSTRANDSCORE_API FHairStrandsTangentFormat
+struct FHairStrandsTangentFormat
 {
+	typedef FPackedNormal Type;
+
 	// TangentX & tangentZ are packed into 2 * PF_R8G8B8A8_SNORM
-	using Type = FPackedNormal;
 	static const uint32 ComponentCount = 2;
 	static const uint32 SizeInByte = sizeof(Type);
 	static const EVertexElementType VertexElementType = VET_Float4;
 	static const EPixelFormat Format = PF_R8G8B8A8_SNORM;
 };
 
+struct FHairStrandsInterpolation0Format
+{
+	typedef FHairInterpolation0Vertex Type;
+
+	static const uint32 ComponentCount = 1;
+	static const uint32 SizeInByte = sizeof(Type);
+	static const EVertexElementType VertexElementType = VET_UShort4;
+	static const EPixelFormat Format = PF_R16G16B16A16_UINT;
+};
+
+struct FHairStrandsInterpolation1Format
+{
+	typedef FHairInterpolation1Vertex Type;
+
+	static const uint32 ComponentCount = 1;
+	static const uint32 SizeInByte = sizeof(Type);
+	static const EVertexElementType VertexElementType = VET_UByte4;
+	static const EPixelFormat Format = PF_R8G8B8A8_UINT;
+};
+
+struct FHairStrandsRootIndexFormat
+{
+	typedef uint32 Type;
+
+	static const uint32 ComponentCount = 1;
+	static const uint32 SizeInByte = sizeof(Type);
+	static const EVertexElementType VertexElementType = VET_UInt;
+	static const EPixelFormat Format = PF_R32_UINT;
+};
+
+struct FHairStrandsRaytracingFormat
+{
+	typedef FVector4 Type;
+
+	static const uint32 ComponentCount = 1;
+	static const uint32 SizeInByte = sizeof(Type);
+	static const EVertexElementType VertexElementType = VET_Float4;
+	static const EPixelFormat Format = PF_A32B32G32R32F;
+};
+
 /** Hair strands index format */
 struct HAIRSTRANDSCORE_API FHairStrandsIndexFormat
 {
@@ -57,7 +120,7 @@ struct HAIRSTRANDSCORE_API FHairStrandsIndexFormat
 	static const uint32 ComponentCount = 1;
 	static const uint32 SizeInByte = sizeof(Type);
 	static const EVertexElementType VertexElementType = VET_UInt;
-	static const EPixelFormat Format = PF_R32_UINT;		
+	static const EPixelFormat Format = PF_R32_UINT;
 };
 
 /** Hair strands weights format */
@@ -66,18 +129,42 @@ struct HAIRSTRANDSCORE_API FHairStrandsWeightFormat
 	using Type = float;
 	static const uint32 ComponentCount = 1;
 	static const uint32 SizeInByte = sizeof(Type);
-	static const EVertexElementType VertexElementType = VET_Float1;	
-	static const EPixelFormat Format = PF_R32_FLOAT;			
+	static const EVertexElementType VertexElementType = VET_Float1;
+	static const EPixelFormat Format = PF_R32_FLOAT;
 };
 
-/** Hair strands quaternion format */
-struct HAIRSTRANDSCORE_API FHairStrandsQuaternionFormat
+/** Hair strands points interpolation attributes */
+struct HAIRSTRANDSCORE_API FHairStrandsInterpolationDatas
 {
-	using Type = FQuat;
-	static const uint32 ComponentCount = 1;
-	static const uint32 SizeInByte = sizeof(Type);
-	static const EVertexElementType VertexElementType = VET_Float4;	
-	static const EPixelFormat Format = PF_A32B32G32R32F;			
+	/** Serialize the interpolated points */
+	inline void Serialize(FArchive& Ar);
+
+	/** Set the number of interpolated points */
+	inline void SetNum(const uint32 NumPoints);
+
+	/** Reset the interpolated points to 0 */
+	inline void Reset();
+
+	/** Get the number of interpolated points */
+	inline uint32 Num() const { return PointsSimCurvesVertexIndex.Num(); }
+
+	/** Simulation curve indices, ordered by closest influence */
+	TArray<FIntVector> PointsSimCurvesIndex;
+
+	/** Closest vertex indices on simulation curve, ordered by closest influence */
+	TArray<FIntVector> PointsSimCurvesVertexIndex;
+
+	/** Weight of vertex indices on simulation curve, ordered by closest influence */
+	TArray<FVector>	PointsSimCurvesVertexWeights;
+
+	void BuildInterpolationDatas(
+		const struct FHairStrandsDatas& SimStrandsData,
+		const struct FHairStrandsDatas& RenderStrandsData);
+
+	/** Build data for interpolation between simulation and rendering */
+	void BuildRenderingDatas(
+		TArray<FHairStrandsInterpolation0Format::Type>& OutPointsInterpolation0,
+		TArray<FHairStrandsInterpolation1Format::Type>& OutPointsInterpolation1) const;
 };
 
 /** Hair strands points attribute */
@@ -101,11 +188,8 @@ struct HAIRSTRANDSCORE_API FHairStrandsPoints
 	/** Normalized radius relative to the max one */
 	TArray<float> PointsRadius; // [0..1]
 
-	/** Normalized length relative to the strand one */
+	/** Normalized length */
 	TArray<float> PointsCoordU; // [0..1]
-
-	/** Max strands points radius */
-	float MaxRadius = 0;
 };
 
 /** Hair strands Curves attribute */
@@ -135,8 +219,13 @@ struct HAIRSTRANDSCORE_API FHairStrandsCurves
 	/** Roots UV in the bounding sphere */
 	TArray<FVector2D> CurvesRootUV; // [0..1]
 
+	TArray<uint32> CurvesGroupID;
+
 	/** Max strands Curves length */
 	float MaxLength = 0;
+
+	/** Max strands Curves radius */
+	float MaxRadius = 0;
 };
 
 /** Hair strands datas that are stored on CPU */
@@ -151,10 +240,12 @@ struct HAIRSTRANDSCORE_API FHairStrandsDatas
 	/* Get the total number of Curves */
 	inline uint32 GetNumCurves() const { return StrandsCurves.Num(); }
 
+	void Reset();
+
 	/** List of all the strands points */
 	FHairStrandsPoints StrandsPoints;
 
-	/** List of all the strands points */
+	/** List of all the strands curves */
 	FHairStrandsCurves StrandsCurves;
 
 	/** The Standard Hair Density */
@@ -166,7 +257,7 @@ struct HAIRSTRANDSCORE_API FHairStrandsDatas
 	/** Build the packed datas for gpu rendering/simulation */
 	void BuildRenderingDatas(
 		TArray<FHairStrandsPositionFormat::Type>& OutPointsPositions,
-		TArray<FHairStrandsTangentFormat::Type>& OutPointsTangents) const;
+		TArray<FHairStrandsAttributeFormat::Type>& OutPackedAttributes) const;
 
 	/** Resample the internal points and curves simulation datas */
 	void BuildSimulationDatas(const uint32 StrandSize,
@@ -175,65 +266,11 @@ struct HAIRSTRANDSCORE_API FHairStrandsDatas
 		TArray<FHairStrandsIndexFormat::Type>& OutPointsNodes) const;
 
 	/** Build the internal points and curves datas */
-	void BuildInternalDatas();
+	void BuildInternalDatas(bool bComputeRootUV = false);
 
 	/** Attach the roots to a static mesh */
 	void AttachStrandsRoots(UStaticMesh* StaticMesh, const FMatrix& TransformMartrix);
 };
 
-void FHairStrandsCurves::SetNum(const uint32 NumCurves)
-{
-	CurvesOffset.SetNum(NumCurves + 1);
-	CurvesCount.SetNum(NumCurves);
-	CurvesLength.SetNum(NumCurves);
-	CurvesRootUV.SetNum(NumCurves);
-}
-
-void FHairStrandsPoints::SetNum(const uint32 NumPoints)
-{
-	PointsPosition.SetNum(NumPoints);
-	PointsRadius.SetNum(NumPoints);
-	PointsCoordU.SetNum(NumPoints);
-}
-
-void FHairStrandsCurves::Reset()
-{
-	CurvesOffset.Reset();
-	CurvesCount.Reset();
-	CurvesLength.Reset();
-	CurvesRootUV.Reset();
-}
-
-void FHairStrandsPoints::Reset()
-{
-	PointsPosition.Reset();
-	PointsRadius.Reset();
-	PointsCoordU.Reset();
-}
-
-void FHairStrandsPoints::Serialize(FArchive& Ar)
-{
-	Ar << PointsPosition;
-	Ar << PointsRadius;
-	Ar << PointsCoordU;
-	Ar << MaxRadius;
-}
-
-void FHairStrandsCurves::Serialize(FArchive& Ar)
-{
-	Ar << CurvesCount;
-	Ar << CurvesOffset;
-	Ar << CurvesLength;
-	Ar << CurvesRootUV;
-	Ar << MaxLength;
-}
-
-void FHairStrandsDatas::Serialize(FArchive& Ar)
-{
-	StrandsPoints.Serialize(Ar);
-	StrandsCurves.Serialize(Ar);
-	Ar << HairDensity;
-	Ar << BoundingBox;
-}
-
-
+// #ueent_todo: Move the data building and decimation code to a HairBuilder
+void HAIRSTRANDSCORE_API DecimateStrandData(const FHairStrandsDatas& InData, float DecimationPercentage, FHairStrandsDatas& OutData);
diff --git a/Engine/Plugins/Experimental/HairStrands/Source/HairStrandsCore/Public/HairStrandsRendering.h b/Engine/Plugins/Experimental/HairStrands/Source/HairStrandsCore/Public/HairStrandsRendering.h
new file mode 100644
--- /dev/null
+++ b/Engine/Plugins/Experimental/HairStrands/Source/HairStrandsCore/Public/HairStrandsRendering.h
@@ -0,0 +1,101 @@
+// Copyright 1998-2019 Epic Games, Inc. All Rights Reserved.
+
+#pragma once
+
+#include "CoreMinimal.h"
+#include "UObject/ObjectMacros.h"
+#include "Engine/EngineTypes.h"
+#include "RenderGraphResources.h"
+
+struct FHairStrandsInterpolationInput
+{
+	FRWBuffer*	RenderRestPosePositionBuffer = nullptr;
+	FRWBuffer*	RenderAttributeBuffer = nullptr;
+	uint32		RenderVertexCount = 0;
+
+	FRWBuffer*	SimRestPosePositionBuffer = nullptr;
+	FRWBuffer*	SimAttributeBuffer = nullptr;
+	uint32		SimVertexCount = 0;
+
+	FRWBuffer*	Interpolation0Buffer = nullptr;
+	FRWBuffer*	Interpolation1Buffer = nullptr;
+
+	// For debug purpose only
+	FRWBuffer*	SimRootPointIndexBuffer = nullptr;
+
+	#if RHI_RAYTRACING
+	FRayTracingGeometry*	RaytracingGeometry = nullptr;
+	FRWBuffer*				RaytracingPositionBuffer = nullptr;
+	uint32					RaytracingVertexCount = 0;
+	bool					bIsRTGeometryInitialized = false;
+	#endif
+
+	float HairRadius = 0;
+	FVector HairWorldOffset = FVector::ZeroVector;
+
+	inline bool IsValid() const
+	{
+		return
+			(RenderRestPosePositionBuffer && RenderRestPosePositionBuffer->SRV) &&
+			(RenderAttributeBuffer && RenderAttributeBuffer->SRV) &&
+			(SimRestPosePositionBuffer && SimRestPosePositionBuffer->SRV) &&
+			(SimAttributeBuffer && SimAttributeBuffer->SRV) &&
+			(Interpolation0Buffer && Interpolation0Buffer->SRV) &&
+			(Interpolation1Buffer && Interpolation1Buffer->SRV) &&
+			RenderVertexCount != 0 &&
+			SimVertexCount != 0;
+	}
+};
+
+struct FHairStrandsInterpolationOutput
+{
+	FRWBuffer* SimDeformedPositionBuffer[2] = { nullptr, nullptr };
+	FRWBuffer* RenderDeformedPositionBuffer[2] = { nullptr, nullptr };
+	uint32 CurrentIndex = 0;
+
+	FRWBuffer* RenderTangentBuffer = nullptr;
+	FRWBuffer* RenderAttributeBuffer = nullptr;
+
+	FRWBuffer* SimTangentBuffer = nullptr;
+
+	// Debug buffers (allocated on-the-fly if used)
+	FRWBuffer RenderPatchedAttributeBuffer;
+
+	// Input to the strand vertex factory. This allows to 
+	// abstract resources generation when debug mode are used
+	struct VertexFactoryInput
+	{
+		FShaderResourceViewRHIRef HairPositionBuffer = nullptr;
+		FShaderResourceViewRHIRef HairPreviousPositionBuffer = nullptr;
+		FShaderResourceViewRHIRef HairTangentBuffer = nullptr;
+		FShaderResourceViewRHIRef HairAttributeBuffer = nullptr;
+		uint32 VertexCount = 0;
+
+		inline void Reset()
+		{
+			HairPositionBuffer = nullptr;
+			HairPreviousPositionBuffer = nullptr;
+			HairTangentBuffer = nullptr;
+			HairAttributeBuffer = nullptr;
+			VertexCount = 0;
+		}
+	};
+	VertexFactoryInput VFInput;
+
+	inline bool IsValid() const
+	{
+		return
+			(SimDeformedPositionBuffer[0] && SimDeformedPositionBuffer[0]->SRV) &&
+			(SimDeformedPositionBuffer[1] && SimDeformedPositionBuffer[1]->SRV) &&
+			(RenderDeformedPositionBuffer[0] && RenderDeformedPositionBuffer[0]->SRV) &&
+			(RenderDeformedPositionBuffer[1] && RenderDeformedPositionBuffer[1]->SRV) &&
+			(RenderTangentBuffer && RenderTangentBuffer->SRV) &&
+			(RenderAttributeBuffer && RenderAttributeBuffer->SRV);
+	}
+};
+
+void ComputeHairStrandsInterpolation(
+	FRHICommandListImmediate& RHICmdList,
+	FHairStrandsInterpolationInput* Input,
+	FHairStrandsInterpolationOutput* Output);
+
diff --git a/Engine/Plugins/Experimental/HairStrands/Source/HairStrandsCore/Public/HairStrandsVertexFactory.h b/Engine/Plugins/Experimental/HairStrands/Source/HairStrandsCore/Public/HairStrandsVertexFactory.h
new file mode 100644
--- /dev/null
+++ b/Engine/Plugins/Experimental/HairStrands/Source/HairStrandsCore/Public/HairStrandsVertexFactory.h
@@ -0,0 +1,97 @@
+// Copyright 1998-2019 Epic Games, Inc. All Rights Reserved.
+
+#pragma once
+
+#include "CoreMinimal.h"
+#include "ShaderParameters.h"
+#include "Components.h"
+#include "VertexFactory.h"
+#include "RenderGraphResources.h"
+#include "HairStrandsDatas.h"
+#include "HairStrandsRendering.h"
+
+class FMaterial;
+class FSceneView;
+struct FMeshBatchElement;
+
+/**
+ * A vertex factory which simply transforms explicit vertex attributes from local to world space.
+ */
+class HAIRSTRANDSCORE_API FHairStrandsVertexFactory : public FVertexFactory
+{
+	DECLARE_VERTEX_FACTORY_TYPE(FHairStrandsVertexFactory);
+public:
+
+	FHairStrandsVertexFactory(ERHIFeatureLevel::Type InFeatureLevel, const char* InDebugName)
+		: FVertexFactory(InFeatureLevel)
+		, DebugName(InDebugName)
+	{
+		bSupportsManualVertexFetch = true;
+	}
+	
+	struct FDataType
+	{
+		FHairStrandsInterpolationOutput* InterpolationOutput = nullptr;
+		
+		float MinStrandRadius = 0;
+		float MaxStrandRadius = 0;
+		float MaxStrandLength = 0;
+		float HairDensity = 1;
+		FVector HairWorldOffset = FVector::ZeroVector;
+	};
+
+	/**
+	 * Should we cache the material's shadertype on this platform with this vertex factory? 
+	 */
+	static bool ShouldCompilePermutation(EShaderPlatform Platform, const class FMaterial* Material, const class FShaderType* ShaderType);
+	static void ModifyCompilationEnvironment(const FVertexFactoryType* Type, EShaderPlatform Platform, const FMaterial* Material, FShaderCompilerEnvironment& OutEnvironment);
+	static void ValidateCompiledResult(const FVertexFactoryType* Type, EShaderPlatform Platform, const FShaderParameterMap& ParameterMap, TArray<FString>& OutErrors);
+
+	/**
+	 * An implementation of the interface used by TSynchronizedResource to update the resource with new data from the game thread.
+	 */
+	void SetData(const FDataType& InData);
+
+	/**
+	* Copy the data from another vertex factory
+	* @param Other - factory to copy from
+	*/
+	void Copy(const FHairStrandsVertexFactory& Other);
+
+	// FRenderResource interface.
+	virtual void InitRHI() override;
+	virtual void ReleaseRHI() override;
+
+	static bool SupportsTessellationShaders() { return true; }
+
+	static FVertexFactoryShaderParameters* ConstructShaderParameters(EShaderFrequency ShaderFrequency);
+
+	inline FRHIShaderResourceView* GetPositionSRV() const			{ check(Data.InterpolationOutput); return Data.InterpolationOutput->VFInput.HairPositionBuffer; };
+	inline FRHIShaderResourceView* GetPreviousPositionSRV() const	{ check(Data.InterpolationOutput); return Data.InterpolationOutput->VFInput.HairPreviousPositionBuffer; }
+	inline FRHIShaderResourceView* GetAttributeSRV() const			{ check(Data.InterpolationOutput); return Data.InterpolationOutput->VFInput.HairAttributeBuffer; }
+	inline FRHIShaderResourceView* GetTangentSRV() const			{ check(Data.InterpolationOutput); return Data.InterpolationOutput->VFInput.HairTangentBuffer; }
+
+	float GetMaxStrandRadius() const;
+	inline float GetMinStrandRadius() const							{ return Data.MinStrandRadius; }
+	inline float GetMaxStrandLength() const							{ return Data.MaxStrandLength; }
+	inline float GetHairDensity() const								{ return Data.HairDensity;  }
+	inline const FVector& GetWorldOffset() const					{ return Data.HairWorldOffset; }
+
+	const FDataType& GetData() const { return Data; }
+protected:
+
+	FDataType Data;
+
+	struct FDebugName
+	{
+		FDebugName(const char* InDebugName)
+#if !UE_BUILD_SHIPPING
+			: DebugName(InDebugName)
+#endif
+		{}
+	private:
+#if !UE_BUILD_SHIPPING
+		const char* DebugName;
+#endif
+	} DebugName;
+};
\ No newline at end of file
diff --git a/Engine/Plugins/Experimental/HairStrands/Source/HairStrandsEditor/HairStrandsEditor.Build.cs b/Engine/Plugins/Experimental/HairStrands/Source/HairStrandsEditor/HairStrandsEditor.Build.cs
--- a/Engine/Plugins/Experimental/HairStrands/Source/HairStrandsEditor/HairStrandsEditor.Build.cs
+++ b/Engine/Plugins/Experimental/HairStrands/Source/HairStrandsEditor/HairStrandsEditor.Build.cs
@@ -4,24 +4,24 @@ namespace UnrealBuildTool.Rules
 {
 	public class HairStrandsEditor : ModuleRules
 	{
-        public HairStrandsEditor(ReadOnlyTargetRules Target) : base(Target)
+		public HairStrandsEditor(ReadOnlyTargetRules Target) : base(Target)
 		{
 			PrivateIncludePaths.Add(ModuleDirectory + "/Private");
-            PublicIncludePaths.Add(ModuleDirectory + "/Public");
+			PublicIncludePaths.Add(ModuleDirectory + "/Public");
 
-            PrivateDependencyModuleNames.AddRange(
-                new string[]
-                {
+			PrivateDependencyModuleNames.AddRange(
+				new string[]
+				{
 					"Core",
 					"CoreUObject",
+					"Engine",
 					"HairStrandsCore",
 					"UnrealEd",
 					"AssetTools",
-                    "AlembicLib",
-                });
-            AddEngineThirdPartyPrivateStaticDependencies(Target,
-             "FBX"
+				});
+			AddEngineThirdPartyPrivateStaticDependencies(Target,
+			 "FBX"
 			);
-        }
+		}
 	}
 }
diff --git a/Engine/Plugins/Experimental/HairStrands/Source/HairStrandsEditor/Private/FbxHairTranslator.cpp b/Engine/Plugins/Experimental/HairStrands/Source/HairStrandsEditor/Private/FbxHairTranslator.cpp
new file mode 100644
--- /dev/null
+++ b/Engine/Plugins/Experimental/HairStrands/Source/HairStrandsEditor/Private/FbxHairTranslator.cpp
@@ -0,0 +1,284 @@
+// Copyright 1998-2019 Epic Games, Inc. All Rights Reserved.
+
+#include "FbxHairTranslator.h"
+
+#include "Misc/Paths.h"
+
+#if PLATFORM_WINDOWS
+#include "Windows/WindowsHWrapper.h"
+#include "Windows/AllowWindowsPlatformTypes.h"
+#endif
+
+THIRD_PARTY_INCLUDES_START
+#include <fbxsdk.h>
+THIRD_PARTY_INCLUDES_END
+
+#if PLATFORM_WINDOWS
+#include "Windows/HideWindowsPlatformTypes.h"
+#endif
+
+namespace FbxHairFormat
+{
+	static const float RootRadius = 0.0001f; // m
+	static const float TipRadius = 0.00005f; // m
+
+	/** File unit to cm ratio */
+	static constexpr const float UNIT_TO_CM = 100;
+}
+
+class FFbxHairImporter
+{
+public:
+	FFbxHairImporter()
+		: FileManager(nullptr)
+		, FileScene(nullptr)
+		, FileImporter(nullptr)
+	{
+		FileManager = FbxManager::Create();
+
+		if (FileManager)
+		{
+			FbxIOSettings* FileSettings = FbxIOSettings::Create(FileManager, IOSROOT);
+
+			FileManager->SetIOSettings(FileSettings);
+
+			FileScene = FbxScene::Create(FileManager, "StrandAssetFbx");
+
+			FileImporter = FbxImporter::Create(FileManager, "");
+		}
+	}
+
+	~FFbxHairImporter()
+	{
+		if (FileImporter)
+		{
+			FileImporter->Destroy();
+		}
+
+		if (FileScene)
+		{
+			FileScene->Destroy();
+		}
+
+		if (FileManager)
+		{
+			FileManager->Destroy();
+		}
+	}
+
+	bool ImportFile(const FString& FileName)
+	{
+		if (!FileManager)
+		{
+			return false;
+		}
+
+		FbxIOSettings* FileSettings = FileManager->GetIOSettings();
+
+		// Initialize the importer by providing a filename.
+		bool FileStatus = FileImporter->Initialize(TCHAR_TO_UTF8(*FileName), -1, FileSettings);
+
+		if (FileStatus && FileImporter->IsFBX())
+		{
+			// Only load geometry data
+			FileSettings->SetBoolProp(IMP_FBX_MATERIAL, false);
+			FileSettings->SetBoolProp(IMP_FBX_TEXTURE, false);
+			FileSettings->SetBoolProp(IMP_FBX_LINK, false);
+			FileSettings->SetBoolProp(IMP_FBX_SHAPE, true);
+			FileSettings->SetBoolProp(IMP_FBX_GOBO, false);
+			FileSettings->SetBoolProp(IMP_FBX_ANIMATION, false);
+			FileSettings->SetBoolProp(IMP_FBX_GLOBAL_SETTINGS, false);
+
+			// Import the scene and destroy the importer.
+			FileStatus = FileImporter->Import(FileScene);
+
+			if (FileStatus)
+			{
+				ImportedFileName = FileName;
+			}
+		}
+
+		return FileStatus;
+	}
+
+	const FString& GetImportedFileName() const
+	{
+		return ImportedFileName;
+	}
+
+	FbxScene* GetFbxScene() const
+	{
+		return FileScene;
+	}
+
+private:
+	FString ImportedFileName;
+
+	FbxManager* FileManager;
+	FbxScene* FileScene;
+	FbxImporter* FileImporter;
+};
+
+static void ParseFbxNode(FbxNode* FileNode, FHairDescription& HairDescription)
+{
+	if (FileNode->GetNodeAttribute() != NULL && FileNode->GetNodeAttribute()->GetAttributeType() == FbxNodeAttribute::eNurbsCurve)
+	{
+		FbxNurbsCurve* NurbsCurve = static_cast<FbxNurbsCurve*>(FileNode->GetNodeAttribute());
+		FbxArray<FbxVector4> PointArray;
+		const uint32 PointCount = NurbsCurve->TessellateCurve(PointArray, 2);
+
+		FStrandID StrandID = HairDescription.AddStrand();
+
+		TStrandAttributesRef<int> StrandNumVertices = HairDescription.StrandAttributes().GetAttributesRef<int>(HairAttribute::Strand::VertexCount);
+		StrandNumVertices[StrandID] = PointCount;
+
+		TVertexAttributesRef<FVector> VertexPositions = HairDescription.VertexAttributes().GetAttributesRef<FVector>(HairAttribute::Vertex::Position);
+		TVertexAttributesRef<float> VertexWidths = HairDescription.VertexAttributes().GetAttributesRef<float>(HairAttribute::Vertex::Width);
+
+		for (uint32 PointIndex = 0; PointIndex < PointCount; ++PointIndex)
+		{
+			const float CoordU = PointIndex / static_cast<float>(PointCount - 1);
+			const float Radius = FMath::Lerp(FbxHairFormat::RootRadius, FbxHairFormat::TipRadius, CoordU);
+			const FVector Position(PointArray[PointIndex][0], PointArray[PointIndex][1], PointArray[PointIndex][2]);
+
+			FVertexID VertexID = HairDescription.AddVertex();
+
+			VertexPositions[VertexID] = Position * FbxHairFormat::UNIT_TO_CM;
+			VertexWidths[VertexID] = Radius * FbxHairFormat::UNIT_TO_CM;
+		}
+	}
+	//	#todo_hair
+	// * Load user/custom properties (eccentricity, width/thickness, ...)
+	// * Load root/hierarchy transform
+
+	const uint32 NumChildren = FileNode->GetChildCount();
+	if (NumChildren > 0)
+	{
+		for (int32 ChildIndex = 0; ChildIndex < FileNode->GetChildCount(); ++ChildIndex)
+		{
+			ParseFbxNode(FileNode->GetChild(ChildIndex), HairDescription);
+		}
+	}
+}
+
+bool FFbxHairTranslator::Translate(const FString& FileName, FHairDescription& HairDescription)
+{
+	// Reuse the FbxHairImporter if there was one created previously by CanTranslate
+	// There could be none if the translator is used for a re-import
+	if (!FbxHairImporter.IsValid() || FileName != FbxHairImporter->GetImportedFileName())
+	{
+		FbxHairImporter = MakeShared<FFbxHairImporter>();
+		if (!FbxHairImporter.IsValid())
+		{
+			return false;
+		}
+
+		bool bFileImported = FbxHairImporter->ImportFile(FileName);
+		if (!bFileImported)
+		{
+			FbxHairImporter.Reset();
+			return false;
+		}
+	}
+
+	// Add required version attributes, first version 0.1
+	FGroomID GroomID(0);
+	HairDescription.GroomAttributes().RegisterAttribute<int>(HairAttribute::Groom::MajorVersion);
+	TGroomAttributesRef<int> MajorVersion = HairDescription.GroomAttributes().GetAttributesRef<int>(HairAttribute::Groom::MajorVersion);
+	MajorVersion.Set(GroomID, 0);
+
+	HairDescription.GroomAttributes().RegisterAttribute<int>(HairAttribute::Groom::MinorVersion);
+	TGroomAttributesRef<int> MinorVersion = HairDescription.GroomAttributes().GetAttributesRef<int>(HairAttribute::Groom::MinorVersion);
+	MinorVersion.Set(GroomID, 1);
+
+	// Handle width as a per-vertex attribute
+	HairDescription.VertexAttributes().RegisterAttribute<float>(HairAttribute::Vertex::Width);
+
+	if (FbxNode* FileNode = FbxHairImporter->GetFbxScene()->GetRootNode())
+	{
+		for (int32 ChildIndex = 0; ChildIndex < FileNode->GetChildCount(); ++ChildIndex)
+		{
+			ParseFbxNode(FileNode->GetChild(ChildIndex), HairDescription);
+		}
+	}
+
+	// Can now release the FbxHairImporter
+	FbxHairImporter.Reset();
+
+	return HairDescription.IsValid();
+}
+
+static void ValidateFbxNode(FbxNode* FileNode, int32& NumCurves, bool& bHasUntranslatableData)
+{
+	// Validate that the FBX has curves only
+	// Reject any attribute type that could be translated by another translator/factory
+	if (FileNode->GetNodeAttribute())
+	{
+		switch (FileNode->GetNodeAttribute()->GetAttributeType())
+		{
+		case FbxNodeAttribute::eSkeleton:
+		case FbxNodeAttribute::eMesh:
+			bHasUntranslatableData = true;
+			break;
+		case FbxNodeAttribute::eNurbsCurve:
+			++NumCurves;
+			break;
+		}
+	}
+
+	for (int32 ChildIndex = 0; ChildIndex < FileNode->GetChildCount() && !bHasUntranslatableData; ++ChildIndex)
+	{
+		ValidateFbxNode(FileNode->GetChild(ChildIndex), NumCurves, bHasUntranslatableData);
+	}
+}
+
+bool FFbxHairTranslator::CanTranslate(const FString& FilePath)
+{
+	if (!IsFileExtensionSupported(FPaths::GetExtension(FilePath)))
+	{
+		return false;
+	}
+
+	// Import the FBX file and check if it contains curves only
+	FbxHairImporter = MakeShared<FFbxHairImporter>();
+	if (!FbxHairImporter.IsValid())
+	{
+		return false;
+	}
+
+	bool bFileImported = FbxHairImporter->ImportFile(FilePath);
+	bool bCanTranslate = false;
+
+	if (bFileImported)
+	{
+		int32 NumCurves = 0;
+		bool bHasUntranslatableData = false;
+		if (FbxNode* FileNode = FbxHairImporter->GetFbxScene()->GetRootNode())
+		{
+			for (int32 ChildIndex = 0; ChildIndex < FileNode->GetChildCount() && !bHasUntranslatableData; ++ChildIndex)
+			{
+				ValidateFbxNode(FileNode->GetChild(ChildIndex), NumCurves, bHasUntranslatableData);
+			}
+		}
+
+		bCanTranslate = !bHasUntranslatableData && NumCurves > 0;
+	}
+
+	// Keep the FbxHairImporter until the file is translated to prevent importing the file again by the actual Translate
+	if (!bFileImported || !bCanTranslate)
+	{
+		FbxHairImporter.Reset();
+	}
+
+	return bCanTranslate;
+}
+
+bool FFbxHairTranslator::IsFileExtensionSupported(const FString& FileExtension) const
+{
+	return GetSupportedFormat().StartsWith(FileExtension);
+}
+
+FString FFbxHairTranslator::GetSupportedFormat() const
+{
+	return TEXT("fbx;Fbx hair strands file");
+}
diff --git a/Engine/Plugins/Experimental/HairStrands/Source/HairStrandsEditor/Private/FbxHairTranslator.h b/Engine/Plugins/Experimental/HairStrands/Source/HairStrandsEditor/Private/FbxHairTranslator.h
new file mode 100644
--- /dev/null
+++ b/Engine/Plugins/Experimental/HairStrands/Source/HairStrandsEditor/Private/FbxHairTranslator.h
@@ -0,0 +1,19 @@
+// Copyright 1998-2019 Epic Games, Inc. All Rights Reserved.
+
+#pragma once
+
+#include "HairStrandsTranslator.h"
+
+class FFbxHairImporter;
+
+class FFbxHairTranslator : public IHairStrandsTranslator
+{
+public:
+	virtual bool Translate(const FString& FilePath, FHairDescription& OutHairDescription) override;
+	virtual bool CanTranslate(const FString& FilePath) override;
+	virtual bool IsFileExtensionSupported(const FString& FileExtension) const override;
+	virtual FString GetSupportedFormat() const override;
+
+private:
+	TSharedPtr<FFbxHairImporter> FbxHairImporter;
+};
diff --git a/Engine/Plugins/Experimental/HairStrands/Source/HairStrandsEditor/Private/GroomActions.cpp b/Engine/Plugins/Experimental/HairStrands/Source/HairStrandsEditor/Private/GroomActions.cpp
new file mode 100644
--- /dev/null
+++ b/Engine/Plugins/Experimental/HairStrands/Source/HairStrandsEditor/Private/GroomActions.cpp
@@ -0,0 +1,65 @@
+// Copyright 1998-2019 Epic Games, Inc. All Rights Reserved.
+
+#include "GroomActions.h"
+#include "GroomAsset.h"
+
+#include "EditorFramework/AssetImportData.h"
+#include "Toolkits/SimpleAssetEditor.h"
+
+#define LOCTEXT_NAMESPACE "AssetTypeActions"
+
+FGroomActions::FGroomActions()
+{}
+
+bool FGroomActions::CanFilter()
+{
+	return true;
+}
+
+void FGroomActions::GetActions(const TArray<UObject*>& InObjects, FMenuBuilder& MenuBuilder)
+{
+	FAssetTypeActions_Base::GetActions(InObjects, MenuBuilder);
+}
+
+
+uint32 FGroomActions::GetCategories()
+{
+	return EAssetTypeCategories::Misc;
+}
+
+FText FGroomActions::GetName() const
+{
+	return NSLOCTEXT("AssetTypeActions", "AssetTypeActions_Groom", "Groom");
+}
+
+void FGroomActions::GetResolvedSourceFilePaths(const TArray<UObject*>& TypeAssets, TArray<FString>& OutSourceFilePaths) const
+{
+	for (UObject* Asset : TypeAssets)
+	{
+		const UGroomAsset* GroomAsset = CastChecked<UGroomAsset>(Asset);
+		GroomAsset->AssetImportData->ExtractFilenames(OutSourceFilePaths);
+	}
+}
+
+UClass* FGroomActions::GetSupportedClass() const
+{
+	return UGroomAsset::StaticClass();
+}
+
+FColor FGroomActions::GetTypeColor() const
+{
+	return FColor::White;
+}
+
+bool FGroomActions::HasActions(const TArray<UObject*>& InObjects) const
+{
+	return true;
+}
+
+void FGroomActions::OpenAssetEditor(const TArray<UObject*>& InObjects, TSharedPtr<IToolkitHost> EditWithinLevelEditor)
+{
+	// #ueent_todo: Will need a custom editor at some point, for now just use the Properties editor
+	FSimpleAssetEditor::CreateEditor(EToolkitMode::Standalone, EditWithinLevelEditor, InObjects);
+}
+
+#undef LOCTEXT_NAMESPACE
\ No newline at end of file
diff --git a/Engine/Plugins/Experimental/HairStrands/Source/HairStrandsEditor/Private/HairFormatTranslator.cpp b/Engine/Plugins/Experimental/HairStrands/Source/HairStrandsEditor/Private/HairFormatTranslator.cpp
new file mode 100644
--- /dev/null
+++ b/Engine/Plugins/Experimental/HairStrands/Source/HairStrandsEditor/Private/HairFormatTranslator.cpp
@@ -0,0 +1,163 @@
+// Copyright 1998-2019 Epic Games, Inc. All Rights Reserved.
+
+#include "HairFormatTranslator.h"
+
+#include "GenericPlatform/GenericPlatformFile.h"
+#include "HAL/PlatformFilemanager.h"
+#include "Misc/Paths.h"
+
+#define HAIR_FILE_SEGMENTS_BIT		1
+#define HAIR_FILE_POINTS_BIT		2
+#define HAIR_FILE_THICKNESS_BIT		4
+#define HAIR_FILE_TRANSPARENCY_BIT	8
+#define HAIR_FILE_COLORS_BIT		16  
+
+namespace HairFormat
+{
+	static const float RootRadius = 0.0001f; // m
+	static const float TipRadius = 0.00005f; // m
+
+	/** File unit to cm ratio */
+	static constexpr const float UNIT_TO_CM = 10;
+
+	/* Hair file header */
+	struct FFileHeader
+	{
+		/* This should be "HAIR" */
+		char	FileSignature[4];	
+
+		/* Number of hair strands */
+		uint32	NumStrands = 0;	
+
+		/* Total number of points of all strands */
+		uint32	NumPoints = 0;
+
+		/* Bit array of data in the file */
+		uint32	BitArrays = 0;
+
+		/* Default number of segments of each strand */
+		uint32	StrandCount = 0;
+
+		/* Default thickness of hair strands */
+		float	StrandThickness = 1.f;
+
+		/* Default transparency of hair strands */
+		float	StrandTransparency = 0.f;	
+
+		/* Default color of hair strands */
+		float	StrandColor[3] = { 0.f,0.f,0.f };
+
+		/* Information about the file */
+		char	FileInfo[88];	
+	};
+}
+
+bool FHairFormatTranslator::Translate(const FString& FileName, FHairDescription& HairDescription)
+{
+	IPlatformFile& PlatformFile = FPlatformFileManager::Get().GetPlatformFile();
+	IFileHandle* FileHandle = PlatformFile.OpenRead(*FileName);
+
+	bool bIsValid = false;
+
+	if (FileHandle)
+	{
+		HairFormat::FFileHeader FileHeader;
+		if (FileHandle->Read(reinterpret_cast<uint8*>(&FileHeader), sizeof(HairFormat::FFileHeader)))
+		{
+			if (strncmp(FileHeader.FileSignature, "HAIR", 4) == 0)
+			{
+				// Add required version attributes, first version 0.1
+				FGroomID GroomID(0);
+				HairDescription.GroomAttributes().RegisterAttribute<int>(HairAttribute::Groom::MajorVersion);
+				TGroomAttributesRef<int> MajorVersion = HairDescription.GroomAttributes().GetAttributesRef<int>(HairAttribute::Groom::MajorVersion);
+				MajorVersion.Set(GroomID, 0);
+
+				HairDescription.GroomAttributes().RegisterAttribute<int>(HairAttribute::Groom::MinorVersion);
+				TGroomAttributesRef<int> MinorVersion = HairDescription.GroomAttributes().GetAttributesRef<int>(HairAttribute::Groom::MinorVersion);
+				MinorVersion.Set(GroomID, 1);
+
+				int32 NumCurves = FileHeader.NumStrands;
+				int32 NumVertices = FileHeader.NumPoints;
+
+				HairDescription.InitializeStrands(NumCurves);
+				HairDescription.InitializeVertices(NumVertices);
+
+				// Read strand segment counts
+				TStrandAttributesRef<int> StrandNumVertices = HairDescription.StrandAttributes().GetAttributesRef<int>(HairAttribute::Strand::VertexCount);
+				if (FileHeader.BitArrays & HAIR_FILE_SEGMENTS_BIT)
+				{
+					TArray<uint16> NumSegmentsPerCurve;
+					NumSegmentsPerCurve.SetNum(NumCurves);
+
+					FileHandle->Read(reinterpret_cast<uint8*>(NumSegmentsPerCurve.GetData()), sizeof(uint16) * NumCurves);
+
+					for (int32 CurveIndex = 0; CurveIndex < NumCurves; ++CurveIndex)
+					{
+						FStrandID StrandID(CurveIndex);
+						StrandNumVertices[StrandID] = NumSegmentsPerCurve[CurveIndex] + 1; // a segment has a start and end vertex 
+					}
+				}
+				else
+				{
+					for (int32 CurveIndex = 0; CurveIndex < NumCurves; ++CurveIndex)
+					{
+						FStrandID StrandID(CurveIndex);
+						StrandNumVertices[StrandID] = FileHeader.StrandCount;
+					}
+				}
+
+				// Read strand vertex positions
+				TVertexAttributesRef<FVector> VertexPositions = HairDescription.VertexAttributes().GetAttributesRef<FVector>(HairAttribute::Vertex::Position);
+				if (FileHeader.BitArrays & HAIR_FILE_POINTS_BIT)
+				{
+					TArray<FVector> Positions;
+					Positions.SetNum(NumVertices);
+
+					FileHandle->Read(reinterpret_cast<uint8*>(&Positions.GetData()->X), sizeof(float) * NumVertices * 3);
+
+					for (int32 VertexIndex = 0; VertexIndex < NumVertices; ++VertexIndex)
+					{
+						FVertexID VertexID(VertexIndex);
+						VertexPositions[VertexID] = Positions[VertexIndex];
+					}
+				}
+
+				// Read strand vertex thickness, with default value of StrandThickness if this bit is not set
+				HairDescription.VertexAttributes().RegisterAttribute<float>(HairAttribute::Vertex::Width, 1, FileHeader.StrandThickness);
+				TVertexAttributesRef<float> VertexWidths = HairDescription.VertexAttributes().GetAttributesRef<float>(HairAttribute::Vertex::Width);
+				if (FileHeader.BitArrays & HAIR_FILE_THICKNESS_BIT)
+				{
+					TArray<float> Widths;
+					Widths.SetNum(NumVertices);
+
+					FileHandle->Read(reinterpret_cast<uint8*>(Widths.GetData()), sizeof(float) * NumVertices);
+
+					for (int32 VertexIndex = 0; VertexIndex < NumVertices; ++VertexIndex)
+					{
+						FVertexID VertexID(VertexIndex);
+						VertexWidths[VertexID] = Widths[VertexIndex];
+					}
+				}
+				
+				bIsValid = HairDescription.IsValid();
+			}
+		}
+		delete FileHandle;
+	}
+	return bIsValid;
+}
+
+bool FHairFormatTranslator::CanTranslate(const FString& FilePath)
+{
+	return IsFileExtensionSupported(FPaths::GetExtension(FilePath));
+}
+
+bool FHairFormatTranslator::IsFileExtensionSupported(const FString& FileExtension) const
+{
+	return GetSupportedFormat().StartsWith(FileExtension);
+}
+
+FString FHairFormatTranslator::GetSupportedFormat() const
+{
+	return TEXT("hair;Hair format hair strands file");
+}
diff --git a/Engine/Plugins/Experimental/HairStrands/Source/HairStrandsEditor/Private/HairFormatTranslator.h b/Engine/Plugins/Experimental/HairStrands/Source/HairStrandsEditor/Private/HairFormatTranslator.h
new file mode 100644
--- /dev/null
+++ b/Engine/Plugins/Experimental/HairStrands/Source/HairStrandsEditor/Private/HairFormatTranslator.h
@@ -0,0 +1,14 @@
+// Copyright 1998-2019 Epic Games, Inc. All Rights Reserved.
+
+#pragma once
+
+#include "HairStrandsTranslator.h"
+
+class FHairFormatTranslator : public IHairStrandsTranslator
+{
+public:
+	virtual bool Translate(const FString& FilePath, FHairDescription& OutHairDescription) override;
+	virtual bool CanTranslate(const FString& FilePath) override;
+	virtual bool IsFileExtensionSupported(const FString& FileExtension) const override;
+	virtual FString GetSupportedFormat() const override;
+};
diff --git a/Engine/Plugins/Experimental/HairStrands/Source/HairStrandsEditor/Private/HairStrandsEditor.cpp b/Engine/Plugins/Experimental/HairStrands/Source/HairStrandsEditor/Private/HairStrandsEditor.cpp
--- a/Engine/Plugins/Experimental/HairStrands/Source/HairStrandsEditor/Private/HairStrandsEditor.cpp
+++ b/Engine/Plugins/Experimental/HairStrands/Source/HairStrandsEditor/Private/HairStrandsEditor.cpp
@@ -1,19 +1,30 @@
 // Copyright 1998-2019 Epic Games, Inc. All Rights Reserved.
 
 #include "HairStrandsEditor.h"
-#include "HairStrandsActions.h"
+#include "GroomActions.h"
+
+#include "FbxHairTranslator.h"
+#include "HairFormatTranslator.h"
+#include "PbrtHairTranslator.h"
+
+IMPLEMENT_MODULE(FHairStrandsEditor, HairStrandsEditor);
 
 void FHairStrandsEditor::StartupModule()
 {
 	IAssetTools& AssetTools = FModuleManager::LoadModuleChecked<FAssetToolsModule>("AssetTools").Get();
-	TSharedRef<IAssetTypeActions> AssetActions = MakeShareable(new FHairStrandsActions());
+	TSharedRef<IAssetTypeActions> AssetActions = MakeShareable(new FGroomActions());
 
 	AssetTools.RegisterAssetTypeActions(AssetActions);
 	RegisteredAssetTypeActions.Add(AssetActions);
+
+	RegisterHairTranslator<FFbxHairTranslator>();
+	RegisterHairTranslator<FHairFormatTranslator>();
+	RegisterHairTranslator<FPbrtHairTranslator>();
 }
 
 void FHairStrandsEditor::ShutdownModule()
 {
+	// #ueent_todo: Unregister the translators
 	FAssetToolsModule* AssetToolsModule = FModuleManager::GetModulePtr<FAssetToolsModule>("AssetTools");
 
 	if (AssetToolsModule != nullptr)
@@ -27,3 +38,13 @@ void FHairStrandsEditor::ShutdownModule()
 	}
 }
 
+TArray<TSharedPtr<IHairStrandsTranslator>> FHairStrandsEditor::GetHairTranslators()
+{
+	TArray<TSharedPtr<IHairStrandsTranslator>> Translators;
+	for (TFunction<TSharedPtr<IHairStrandsTranslator>()>& SpawnTranslator : TranslatorSpawners)
+	{
+		Translators.Add(SpawnTranslator());
+	}
+
+	return Translators;
+}
diff --git a/Engine/Plugins/Experimental/HairStrands/Source/HairStrandsEditor/Private/HairStrandsFactory.cpp b/Engine/Plugins/Experimental/HairStrands/Source/HairStrandsEditor/Private/HairStrandsFactory.cpp
--- a/Engine/Plugins/Experimental/HairStrands/Source/HairStrandsEditor/Private/HairStrandsFactory.cpp
+++ b/Engine/Plugins/Experimental/HairStrands/Source/HairStrandsEditor/Private/HairStrandsFactory.cpp
@@ -1,104 +1,93 @@
 // Copyright 1998-2019 Epic Games, Inc. All Rights Reserved.
 
 #include "HairStrandsFactory.h"
-#include "HairStrandsAsset.h"
-#include "HairStrandsLoader.h"
-
-#include "Developer/DesktopPlatform/Public/IDesktopPlatform.h"
-#include "Developer/DesktopPlatform/Public/DesktopPlatformModule.h"
-#include "UnrealEd/Public/EditorDirectories.h"
 
+#include "EditorFramework/AssetImportData.h"
+#include "HairDescription.h"
+#include "GroomAsset.h"
+#include "HairStrandsEditor.h"
+#include "HairStrandsImporter.h"
+#include "HairStrandsTranslator.h"
 #include "Misc/Paths.h"
+#include "Misc/ScopedSlowTask.h"
 
-/**
- * true if the extension is for the .hair, .fbx, .abc format.
- */
-static bool IsHairStrandsFormat(const TCHAR* Extension)
-{
-	return (FCString::Stricmp(Extension, TEXT("abc")) == 0) || 
-		   (FCString::Stricmp(Extension, TEXT("fbx")) == 0) ||
-		   (FCString::Stricmp(Extension, TEXT("hair")) == 0);
-}
+#define LOCTEXT_NAMESPACE "HairStrandsFactory"
 
-static void BuildHairStrands( const FString& FilePath, FHairStrandsDatas& OutStrandsDatas)
+UHairStrandsFactory::UHairStrandsFactory(const FObjectInitializer& ObjectInitializer)
+	: Super(ObjectInitializer)
 {
-	if (FCString::Stricmp(*FPaths::GetExtension(FilePath), TEXT("hair")) == 0)
-	{
-		THairStrandsLoader<FHairFormat>::LoadHairStrands(FilePath, OutStrandsDatas);
-	}
-	else if (FCString::Stricmp(*FPaths::GetExtension(FilePath), TEXT("fbx")) == 0)
+	SupportedClass = UGroomAsset::StaticClass();
+	bCreateNew = false;		// manual creation not allow
+	bEditAfterNew = false;
+	bEditorImport = true;	// only allow import
+
+	// Slightly increased priority to allow its translators to check if they can translate the file
+	ImportPriority += 1;
+
+	// Lazy init the translators to let them register themselves before the CDO is used
+	if (!HasAnyFlags(RF_ClassDefaultObject))
 	{
-		THairStrandsLoader<FFbxFormat>::LoadHairStrands(FilePath, OutStrandsDatas);
+		InitTranslators();
 	}
-	else if (FCString::Stricmp(*FPaths::GetExtension(FilePath), TEXT("abc")) == 0)
+}
+
+void UHairStrandsFactory::InitTranslators()
+{
+	Formats.Reset();
+
+	Translators = FHairStrandsEditor::Get().GetHairTranslators();
+	for (TSharedPtr<IHairStrandsTranslator> Translator : Translators)
 	{
-		THairStrandsLoader<FAbcFormat>::LoadHairStrands(FilePath, OutStrandsDatas);
+		Formats.Add(Translator->GetSupportedFormat());
 	}
 }
 
-UHairStrandsFactory::UHairStrandsFactory(const FObjectInitializer& ObjectInitializer)
-	: Super(ObjectInitializer)
+void UHairStrandsFactory::GetSupportedFileExtensions(TArray<FString>& OutExtensions) const
 {
-	SupportedClass = UHairStrandsAsset::StaticClass();
-	bCreateNew = true;
-	bEditAfterNew = true;
-	bEditorImport = false;
-
-	Formats.Add(TEXT("hair;Hair hair strands file"));
-	Formats.Add(TEXT("fbx;Fbx hair strands file"));
-	Formats.Add(TEXT("abc;Alembic hair strands file"));
+	if (HasAnyFlags(RF_ClassDefaultObject) && Formats.Num() == 0)
+	{
+		// Init the translators the first time the CDO is used
+		UHairStrandsFactory* Factory = const_cast<UHairStrandsFactory*>(this);
+		Factory->InitTranslators();
+	}
+
+	Super::GetSupportedFileExtensions(OutExtensions);
 }
 
-UObject* UHairStrandsFactory::FactoryCreateNew(UClass* InClass, UObject* InParent, FName InName, EObjectFlags Flags, UObject* Context, FFeedbackContext* Warn)
+UObject* UHairStrandsFactory::FactoryCreateFile(UClass* InClass, UObject* InParent, FName InName, EObjectFlags Flags, 
+	const FString& Filename, const TCHAR* Parms, FFeedbackContext* Warn, bool& bOutOperationCanceled) 
 {
-	UHairStrandsAsset* ExistingAsset = FindObject<UHairStrandsAsset>(InParent, *InName.ToString());
-	if (ExistingAsset)
+	// Translate the hair data from the file
+	TSharedPtr<IHairStrandsTranslator> SelectedTranslator = GetTranslator(Filename);
+	if (!SelectedTranslator.IsValid())
 	{
-		ExistingAsset->ReleaseResource();
+		return nullptr;
 	}
-	UHairStrandsAsset* CurrentAsset = NewObject<UHairStrandsAsset>(InParent, InClass, InName, Flags);
 
-	if (CurrentAsset)
-	{
-		const FString Filter(TEXT("Hair Strands Files (*.hair,*.fbx,*.abc)|*.hair;*.fbx;*.abc"));
-
-		TArray<FString> OpenFilenames;
-		int32 FilterIndex = -1;
-		if (FDesktopPlatformModule::Get()->OpenFileDialog(
-			nullptr,
-			FString(TEXT("Choose a hair strands file")),
-			FEditorDirectories::Get().GetLastDirectory(ELastDirectory::GENERIC_IMPORT),
-			TEXT(""),
-			Filter,
-			EFileDialogFlags::None,
-			OpenFilenames,
-			FilterIndex))
-		{			
-			CurrentAsset->FilePath = OpenFilenames[0];
-			BuildHairStrands(CurrentAsset->FilePath, CurrentAsset->StrandsDatas);
-		}
+	//GEditor->GetEditorSubsystem<UImportSubsystem>()->BroadcastAssetPreImport(this, InClass, InParent, InName, FileType);
 
-		CurrentAsset->InitResource();
-	}
-	return CurrentAsset;
-}
+	FScopedSlowTask Progress((float) 1, LOCTEXT("ImportHairAsset", "Importing hair asset..."), true);
+	Progress.MakeDialog(true);
 
-UObject* UHairStrandsFactory::FactoryCreateFile(UClass * InClass, UObject * InParent, FName InName, EObjectFlags Flags, 
-	const FString & Filename, const TCHAR* Parms, FFeedbackContext * Warn, bool& bOutOperationCanceled) 
-{
-	UHairStrandsAsset* ExistingAsset = FindObject<UHairStrandsAsset>(InParent, *InName.ToString());
-	if (ExistingAsset)
+	FHairDescription HairDescription;
+	if (!SelectedTranslator->Translate(Filename, HairDescription))
 	{
-		ExistingAsset->ReleaseResource();
+		return nullptr;
 	}
 
-	UHairStrandsAsset* CurrentAsset = NewObject<UHairStrandsAsset>(InParent, InClass, InName, Flags);	
+	// Might try to import the same file in the same folder, so if an asset already exists there, reuse and update it
+	UGroomAsset* ExistingAsset = FindObject<UGroomAsset>(InParent, *InName.ToString());
 
+	FHairImportContext HairImportContext(InParent, InClass, InName, Flags);
+	UGroomAsset* CurrentAsset = FHairStrandsImporter::ImportHair(HairImportContext, HairDescription, ExistingAsset);
 	if (CurrentAsset)
 	{
-		CurrentAsset->FilePath = Filename;
-		BuildHairStrands(CurrentAsset->FilePath, CurrentAsset->StrandsDatas);
-		CurrentAsset->InitResource();
+		// Setup asset import data
+		if (!CurrentAsset->AssetImportData)
+		{
+			CurrentAsset->AssetImportData = NewObject<UAssetImportData>(CurrentAsset);
+		}
+		CurrentAsset->AssetImportData->Update(Filename);
 	}
 
 	return CurrentAsset;
@@ -106,51 +95,25 @@ UObject* UHairStrandsFactory::FactoryCreateFile(UClass * InClass, UObject * InPa
 
 bool UHairStrandsFactory::FactoryCanImport(const FString& Filename)
 {
-	return  IsHairStrandsFormat(*FPaths::GetExtension(Filename));
-}
-
-bool UHairStrandsFactory::CanReimport(UObject* Obj, TArray<FString>& OutFilenames)
-{
-	UHairStrandsAsset* Asset = Cast<UHairStrandsAsset>(Obj);
-	if (Asset)
+	for (TSharedPtr<IHairStrandsTranslator> Translator : Translators)
 	{
-		const FString& FileName = Asset->FilePath;
-		if (!FileName.IsEmpty())
+		if (Translator->CanTranslate(Filename))
 		{
-			OutFilenames.Add(FileName);
+			return true;
 		}
-
-		return true;
 	}
 	return false;
 }
 
-void UHairStrandsFactory::SetReimportPaths(UObject* Obj, const TArray<FString>& NewReimportPaths)
-{
-	UHairStrandsAsset* Asset = Cast<UHairStrandsAsset>(Obj);
-	if (Asset && ensure(NewReimportPaths.Num() == 1))
-	{
-		Asset->FilePath = NewReimportPaths[0];
-	}
-}
-
-EReimportResult::Type UHairStrandsFactory::Reimport(UObject* Obj)
+TSharedPtr<IHairStrandsTranslator> UHairStrandsFactory::GetTranslator(const FString& Filename)
 {
-	UHairStrandsAsset* CurrentAsset = Cast<UHairStrandsAsset>(Obj);
-	BuildHairStrands(CurrentAsset->FilePath, CurrentAsset->StrandsDatas);
-	// Try to find the outer package so we can dirty it up
-	if (Obj->GetOuter())
+	FString Extension = FPaths::GetExtension(Filename);
+	for (TSharedPtr<IHairStrandsTranslator> Translator : Translators)
 	{
-		Obj->GetOuter()->MarkPackageDirty();
-	}
-	else
-	{
-		Obj->MarkPackageDirty();
+		if (Translator->IsFileExtensionSupported(Extension))
+		{
+			return Translator;
+		}
 	}
-	return EReimportResult::Succeeded;
-}
-
-int32 UHairStrandsFactory::GetPriority() const
-{
-	return ImportPriority;
-}
+	return {};
+}
\ No newline at end of file
diff --git a/Engine/Plugins/Experimental/HairStrands/Source/HairStrandsEditor/Private/HairStrandsImporter.cpp b/Engine/Plugins/Experimental/HairStrands/Source/HairStrandsEditor/Private/HairStrandsImporter.cpp
new file mode 100644
--- /dev/null
+++ b/Engine/Plugins/Experimental/HairStrands/Source/HairStrandsEditor/Private/HairStrandsImporter.cpp
@@ -0,0 +1,228 @@
+// Copyright 1998-2019 Epic Games, Inc. All Rights Reserved.
+
+#include "HairStrandsImporter.h"
+
+#include "HairDescription.h"
+#include "GroomAsset.h"
+#include "GroomComponent.h"
+
+DEFINE_LOG_CATEGORY_STATIC(LogHairImporter, Log, All);
+
+FHairImportContext::FHairImportContext(UObject* InParent, UClass* InClass, FName InName, EObjectFlags InFlags)
+	: Parent(InParent)
+	, Class(InClass)
+	, Name(InName)
+	, Flags(InFlags)
+{
+}
+
+UGroomAsset* FHairStrandsImporter::ImportHair(const FHairImportContext& ImportContext, FHairDescription& HairDescription, UGroomAsset* ExistingHair)
+{
+	// For now, just convert HairDescription to HairStrandsDatas
+	int32 NumCurves = HairDescription.GetNumStrands();
+	int32 NumVertices = HairDescription.GetNumVertices();
+
+	// Check for required attributes
+	TGroomAttributesRef<int> MajorVersion = HairDescription.GroomAttributes().GetAttributesRef<int>(HairAttribute::Groom::MajorVersion);
+	TGroomAttributesRef<int> MinorVersion = HairDescription.GroomAttributes().GetAttributesRef<int>(HairAttribute::Groom::MinorVersion);
+
+	if (!MajorVersion.IsValid() || !MinorVersion.IsValid())
+	{
+		UE_LOG(LogHairImporter, Warning, TEXT("Failed to import hair: No version number attributes found. Please re-export the input file."));
+		return nullptr;
+	}
+
+	FGroomID GroomID(0);
+	
+	TGroomAttributesRef<float> GroomHairWidthAttribute = HairDescription.GroomAttributes().GetAttributesRef<float>(HairAttribute::Groom::Width);
+	TOptional<float> GroomHairWidth;
+	if (GroomHairWidthAttribute.IsValid())
+	{
+		GroomHairWidth = GroomHairWidthAttribute[GroomID];
+	}
+
+	TGroomAttributesRef<FVector> GroomHairColorAttribute = HairDescription.GroomAttributes().GetAttributesRef<FVector>(HairAttribute::Groom::Color);
+	TOptional<FVector> GroomHairColor;
+	if (GroomHairColorAttribute.IsValid())
+	{
+		GroomHairColor = GroomHairColorAttribute[GroomID];
+	}
+
+	TVertexAttributesRef<FVector> VertexPositions = HairDescription.VertexAttributes().GetAttributesRef<FVector>(HairAttribute::Vertex::Position);
+	TStrandAttributesRef<int> StrandNumVertices = HairDescription.StrandAttributes().GetAttributesRef<int>(HairAttribute::Strand::VertexCount);
+
+	if (!VertexPositions.IsValid() || !StrandNumVertices.IsValid())
+	{
+		UE_LOG(LogHairImporter, Warning, TEXT("Failed to import hair: No vertices or curves data found."));
+		return nullptr;
+	}
+
+	TVertexAttributesRef<float> VertexWidths = HairDescription.VertexAttributes().GetAttributesRef<float>(HairAttribute::Vertex::Width);
+	TStrandAttributesRef<float> StrandWidths = HairDescription.StrandAttributes().GetAttributesRef<float>(HairAttribute::Strand::Width);
+
+	TStrandAttributesRef<FVector2D> StrandRootUV = HairDescription.StrandAttributes().GetAttributesRef<FVector2D>(HairAttribute::Strand::RootUV);
+	bool bHasUVData = StrandRootUV.IsValid();
+
+	TStrandAttributesRef<bool> StrandGuides = HairDescription.StrandAttributes().GetAttributesRef<bool>(HairAttribute::Strand::Guide);
+	TStrandAttributesRef<int> GroupIDs = HairDescription.StrandAttributes().GetAttributesRef<int>(HairAttribute::Strand::GroupID);
+
+	FHairStrandsDatas HairRenderData;
+	FHairStrandsDatas HairSimulationData;
+	FHairStrandsDatas* CurrentHairStrandsDatas = &HairRenderData;
+
+	TMap<int32, FHairGroupRenderSettings> RenderHairGroups;
+	TMap<int32, FHairGroupSimulationSettings> SimHairGroups;
+
+	int32 GlobalVertexIndex = 0;
+	int32 NumHairCurves = 0;
+	int32 NumGuideCurves = 0;
+	int32 NumHairPoints = 0;
+	int32 NumGuidePoints = 0;
+	for (int32 CurveIndex = 0; CurveIndex < NumCurves; ++CurveIndex)
+	{
+		FStrandID StrandID(CurveIndex);
+
+		bool bIsGuide = false;
+		if (StrandGuides.IsValid())
+		{
+			bIsGuide = StrandGuides[StrandID];
+		}
+
+		int32 CurveNumVertices = StrandNumVertices[StrandID];
+
+		int32 GroupID = 0;
+		if (GroupIDs.IsValid())
+		{
+			GroupID = GroupIDs[StrandID];
+		}
+
+		if (!bIsGuide)
+		{
+			++NumHairCurves;
+			NumHairPoints += CurveNumVertices;
+			CurrentHairStrandsDatas = &HairRenderData;
+
+			FHairGroupRenderSettings& GroupSettings = RenderHairGroups.FindOrAdd(GroupID);
+			GroupSettings.GroupID = GroupID;
+			++GroupSettings.NumCurves;
+		}
+		else
+		{
+			++NumGuideCurves;
+			NumGuidePoints += CurveNumVertices;
+			CurrentHairStrandsDatas = &HairSimulationData;
+
+			FHairGroupSimulationSettings& GroupSettings = SimHairGroups.FindOrAdd(GroupID);
+			GroupSettings.GroupID = GroupID;
+			++GroupSettings.NumCurves;
+		}
+
+		CurrentHairStrandsDatas->StrandsCurves.CurvesCount.Add(CurveNumVertices);
+		CurrentHairStrandsDatas->StrandsCurves.CurvesGroupID.Add(GroupID);
+
+		if (bHasUVData)
+		{
+			CurrentHairStrandsDatas->StrandsCurves.CurvesRootUV.Add(StrandRootUV[StrandID]);
+		}
+
+		float StrandWidth = GroomHairWidth ? GroomHairWidth.GetValue() : 0.f;
+		if (StrandWidths.IsValid())
+		{
+			StrandWidth = StrandWidths[StrandID];
+		}
+
+		for (int32 VertexIndex = 0; VertexIndex < CurveNumVertices; ++VertexIndex, ++GlobalVertexIndex)
+		{
+			FVertexID VertexID(GlobalVertexIndex);
+
+			CurrentHairStrandsDatas->StrandsPoints.PointsPosition.Add(VertexPositions[VertexID]);
+
+			float VertexWidth = 0.f;
+			if (VertexWidths.IsValid())
+			{
+				VertexWidth = VertexWidths[VertexID];
+			}
+
+			// Fall back to strand width if there was no vertex width
+			if (VertexWidth == 0.f && StrandWidth != 0.f)
+			{
+				VertexWidth = StrandWidth;
+			}
+
+			CurrentHairStrandsDatas->StrandsPoints.PointsRadius.Add(VertexWidth * 0.5f);
+		}
+	}
+
+	FGroomComponentRecreateRenderStateContext RecreateRenderContext(ExistingHair);
+
+	UGroomAsset* HairAsset = nullptr;
+	if (ExistingHair)
+	{
+		HairAsset = ExistingHair;
+
+		ExistingHair->Reset();
+	}
+
+	if (!HairAsset)
+	{
+		HairAsset = NewObject<UGroomAsset>(ImportContext.Parent, ImportContext.Class, ImportContext.Name, ImportContext.Flags);
+
+		if (!HairAsset)
+		{
+			UE_LOG(LogHairImporter, Warning, TEXT("Failed to import hair: Could not allocate memory to create asset."));
+			return nullptr;
+		}
+	}
+
+	HairRenderData.StrandsCurves.SetNum(NumHairCurves);
+	HairRenderData.StrandsPoints.SetNum(NumHairPoints);
+	HairRenderData.BuildInternalDatas(!bHasUVData);
+	HairAsset->HairRenderData = MoveTemp(HairRenderData);
+
+	if (NumGuideCurves > 0)
+	{
+		HairSimulationData.StrandsCurves.SetNum(NumGuideCurves);
+		HairSimulationData.StrandsPoints.SetNum(NumGuidePoints);
+		HairSimulationData.BuildInternalDatas(true); // Imported guides don't currently have root UVs so force computing them
+		HairAsset->HairSimulationData = MoveTemp(HairSimulationData);
+	}
+	else
+	{
+		DecimateStrandData(HairAsset->HairRenderData, FMath::Clamp(HairAsset->HairToGuideDensity, 0.f, 1.f), HairAsset->HairSimulationData);
+	}
+
+	HairAsset->HairDescription = MakeUnique<FHairDescription>(MoveTemp(HairDescription));
+
+	if (RenderHairGroups.Num() > 0)
+	{
+		for (TPair<int32, FHairGroupRenderSettings>& Pair : RenderHairGroups)
+		{
+			HairAsset->RenderHairGroups.Add(MoveTemp(Pair.Value));
+		}
+	}
+	else
+	{
+		FHairGroupRenderSettings& GroupSettings = HairAsset->RenderHairGroups.AddDefaulted_GetRef();
+		GroupSettings.GroupID = 0;
+		GroupSettings.NumCurves = HairRenderData.GetNumCurves();
+	}
+
+	if (SimHairGroups.Num() > 0)
+	{
+		for (TPair<int32, FHairGroupSimulationSettings>& Pair : SimHairGroups)
+		{
+			HairAsset->SimulationHairGroups.Add(MoveTemp(Pair.Value));
+		}
+	}
+	else
+	{
+		FHairGroupSimulationSettings& GroupSettings = HairAsset->SimulationHairGroups.AddDefaulted_GetRef();
+		GroupSettings.GroupID = 0;
+		GroupSettings.NumCurves = HairAsset->HairSimulationData.GetNumCurves();
+		GroupSettings.bIsAutoGenerated = true;
+	}
+
+	HairAsset->InitResource();
+
+	return HairAsset;
+}
diff --git a/Engine/Plugins/Experimental/HairStrands/Source/HairStrandsEditor/Private/PbrtHairTranslator.cpp b/Engine/Plugins/Experimental/HairStrands/Source/HairStrandsEditor/Private/PbrtHairTranslator.cpp
new file mode 100644
--- /dev/null
+++ b/Engine/Plugins/Experimental/HairStrands/Source/HairStrandsEditor/Private/PbrtHairTranslator.cpp
@@ -0,0 +1,36 @@
+// Copyright 1998-2019 Epic Games, Inc. All Rights Reserved.
+
+#include "PbrtHairTranslator.h"
+#include "HairStrandsDatas.h"
+
+#include "Misc/Paths.h"
+
+namespace PbrtHairFormat
+{
+	static const float RootRadius = 0.0001f; // m
+	static const float TipRadius = 0.00005f; // m
+}
+
+bool FPbrtHairTranslator::Translate(const FString& FileName, FHairDescription& HairDescription)
+{
+	bool bSuccess = false;
+
+	// #ueent_todo: Pbrt loader from HairStrandsLoader
+
+	return bSuccess;
+}
+
+bool FPbrtHairTranslator::CanTranslate(const FString& FilePath)
+{
+	return IsFileExtensionSupported(FPaths::GetExtension(FilePath));
+}
+
+bool FPbrtHairTranslator::IsFileExtensionSupported(const FString& FileExtension) const
+{
+	return GetSupportedFormat().StartsWith(FileExtension);
+}
+
+FString FPbrtHairTranslator::GetSupportedFormat() const
+{
+	return TEXT("pbrt;pbrt file");
+}
diff --git a/Engine/Plugins/Experimental/HairStrands/Source/HairStrandsEditor/Private/PbrtHairTranslator.h b/Engine/Plugins/Experimental/HairStrands/Source/HairStrandsEditor/Private/PbrtHairTranslator.h
new file mode 100644
--- /dev/null
+++ b/Engine/Plugins/Experimental/HairStrands/Source/HairStrandsEditor/Private/PbrtHairTranslator.h
@@ -0,0 +1,14 @@
+// Copyright 1998-2019 Epic Games, Inc. All Rights Reserved.
+
+#pragma once
+
+#include "HairStrandsTranslator.h"
+
+class FPbrtHairTranslator : public IHairStrandsTranslator
+{
+public:
+	virtual bool Translate(const FString& FilePath, FHairDescription& OutHairDescription) override;
+	virtual bool CanTranslate(const FString& FilePath) override;
+	virtual bool IsFileExtensionSupported(const FString& FileExtension) const override;
+	virtual FString GetSupportedFormat() const override;
+};
diff --git a/Engine/Plugins/Experimental/HairStrands/Source/HairStrandsEditor/Public/GroomActions.h b/Engine/Plugins/Experimental/HairStrands/Source/HairStrandsEditor/Public/GroomActions.h
new file mode 100644
--- /dev/null
+++ b/Engine/Plugins/Experimental/HairStrands/Source/HairStrandsEditor/Public/GroomActions.h
@@ -0,0 +1,38 @@
+// Copyright 1998-2019 Epic Games, Inc. All Rights Reserved.
+
+#pragma once
+
+#include "AssetTypeActions_Base.h"
+#include "Templates/SharedPointer.h"
+
+class ISlateStyle;
+
+/**
+ * Implements an action for UHairStrands assets.
+ */
+class FGroomActions : public FAssetTypeActions_Base
+{
+public:
+
+	/**
+	 * Creates and initializes a new instance.
+	 *
+	 * @param InStyle The style set to use for asset editor toolkits.
+	 */
+	FGroomActions();
+
+public:
+
+	//~ FAssetTypeActions_Base overrides
+
+	virtual bool CanFilter() override;
+	virtual void GetActions(const TArray<UObject*>& InObjects, FMenuBuilder& MenuBuilder) override;
+	virtual uint32 GetCategories() override;
+	virtual FText GetName() const override;
+	virtual void GetResolvedSourceFilePaths(const TArray<UObject*>& TypeAssets, TArray<FString>& OutSourceFilePaths) const override;
+	virtual UClass* GetSupportedClass() const override;
+	virtual FColor GetTypeColor() const override;
+	virtual bool HasActions(const TArray<UObject*>& InObjects) const override;
+	virtual bool IsImportedAsset() const override { return true; }
+	virtual void OpenAssetEditor(const TArray<UObject*>& InObjects, TSharedPtr<IToolkitHost> EditWithinLevelEditor = TSharedPtr<IToolkitHost>()) override;
+};
diff --git a/Engine/Plugins/Experimental/HairStrands/Source/HairStrandsEditor/Public/HairStrandsEditor.h b/Engine/Plugins/Experimental/HairStrands/Source/HairStrandsEditor/Public/HairStrandsEditor.h
--- a/Engine/Plugins/Experimental/HairStrands/Source/HairStrandsEditor/Public/HairStrandsEditor.h
+++ b/Engine/Plugins/Experimental/HairStrands/Source/HairStrandsEditor/Public/HairStrandsEditor.h
@@ -5,8 +5,12 @@
 #include "Modules/ModuleInterface.h"
 #include "Modules/ModuleManager.h"
 
+#define HAIRSTRANDSEDITOR_MODULE_NAME TEXT("HairStrandsEditor")
+
+class IHairStrandsTranslator;
+
 /** Implements the HairStrands module  */
-class FHairStrandsEditor : public IModuleInterface
+class HAIRSTRANDSEDITOR_API FHairStrandsEditor : public IModuleInterface
 {
 public:
 
@@ -17,13 +21,31 @@ public:
 
 	virtual bool SupportsDynamicReloading() override
 	{
-		return true;
+		return false;
+	}
+
+	static inline FHairStrandsEditor& Get()
+	{
+		return FModuleManager::LoadModuleChecked<FHairStrandsEditor>(HAIRSTRANDSEDITOR_MODULE_NAME);
 	}
 
+	/** Register HairStrandsTranslator to add support for import by the HairStandsFactory */
+	template <typename TranslatorType>
+	void RegisterHairTranslator()
+	{
+		TranslatorSpawners.Add([]
+		{
+			return MakeShared<TranslatorType>();
+		});
+	}
+
+	/** Get new instances of HairStrandsTranslators */
+	TArray<TSharedPtr<IHairStrandsTranslator>> GetHairTranslators();
+
 private:
 
 	/** The collection of registered asset type actions. */
 	TArray<TSharedRef<class IAssetTypeActions>> RegisteredAssetTypeActions;
-};
 
-IMPLEMENT_MODULE(FHairStrandsEditor, HairStrandsEditor);
+	TArray<TFunction<TSharedPtr<IHairStrandsTranslator>()>> TranslatorSpawners;
+};
diff --git a/Engine/Plugins/Experimental/HairStrands/Source/HairStrandsEditor/Public/HairStrandsFactory.h b/Engine/Plugins/Experimental/HairStrands/Source/HairStrandsEditor/Public/HairStrandsFactory.h
--- a/Engine/Plugins/Experimental/HairStrands/Source/HairStrandsEditor/Public/HairStrandsFactory.h
+++ b/Engine/Plugins/Experimental/HairStrands/Source/HairStrandsEditor/Public/HairStrandsFactory.h
@@ -4,31 +4,33 @@
 
 #include "Factories/Factory.h"
 #include "UObject/ObjectMacros.h"
-#include <UnrealEd/Public/EditorReimportHandler.h>
 
 #include "HairStrandsFactory.generated.h"
 
+class IHairStrandsTranslator;
 
 /**
  * Implements a factory for UHairStrands objects.
  */
 UCLASS(hidecategories = Object)
-class UHairStrandsFactory : public UFactory, public FReimportHandler
+class UHairStrandsFactory : public UFactory
 {
 	GENERATED_UCLASS_BODY()
 
 public:
 
 	//~ UFactory Interface
-	virtual UObject* FactoryCreateNew(UClass* InClass, UObject* InParent, FName InName, EObjectFlags Flags, UObject* Context, FFeedbackContext* Warn) override;
-	virtual UObject* FactoryCreateFile(UClass * InClass, UObject * InParent, FName InName, EObjectFlags Flags,
-		const FString & Filename, const TCHAR* Parms, FFeedbackContext * Warn, bool& bOutOperationCanceled) override;
-
-	//~ FReimportHandler interface
+	virtual UObject* FactoryCreateFile(UClass* InClass, UObject* InParent, FName InName, EObjectFlags Flags,
+		const FString& Filename, const TCHAR* Parms, FFeedbackContext* Warn, bool& bOutOperationCanceled) override;
 	virtual bool FactoryCanImport(const FString& Filename) override;
-	virtual bool CanReimport(UObject* Obj, TArray<FString>& OutFilenames) override;
-	virtual void SetReimportPaths(UObject* Obj, const TArray<FString>& NewReimportPaths) override;
-	virtual EReimportResult::Type Reimport(UObject* Obj) override;
-	virtual int32 GetPriority() const override;
+	virtual void GetSupportedFileExtensions(TArray<FString>& OutExtensions) const override;
+
+protected:
+	void InitTranslators();
+
+	TSharedPtr<IHairStrandsTranslator> GetTranslator(const FString& Filename);
+
+private:
+	TArray<TSharedPtr<IHairStrandsTranslator>> Translators;
 };
 
diff --git a/Engine/Plugins/Experimental/HairStrands/Source/HairStrandsEditor/Public/HairStrandsImporter.h b/Engine/Plugins/Experimental/HairStrands/Source/HairStrandsEditor/Public/HairStrandsImporter.h
new file mode 100644
--- /dev/null
+++ b/Engine/Plugins/Experimental/HairStrands/Source/HairStrandsEditor/Public/HairStrandsImporter.h
@@ -0,0 +1,24 @@
+// Copyright 1998-2019 Epic Games, Inc. All Rights Reserved.
+
+#pragma once
+
+#include "CoreMinimal.h"
+#include "UObject/ObjectMacros.h"
+
+class FHairDescription;
+class UGroomAsset;
+
+struct HAIRSTRANDSEDITOR_API FHairImportContext
+{
+	FHairImportContext(UObject* InParent = nullptr, UClass* InClass = nullptr, FName InName = FName(), EObjectFlags InFlags = EObjectFlags::RF_NoFlags);
+
+	UObject* Parent;
+	UClass* Class;
+	FName Name;
+	EObjectFlags Flags;
+};
+
+struct HAIRSTRANDSEDITOR_API FHairStrandsImporter
+{
+	static UGroomAsset* ImportHair(const FHairImportContext& ImportContext, FHairDescription& HairDescription, UGroomAsset* ExistingHair = nullptr);
+};
diff --git a/Engine/Plugins/Experimental/HairStrands/Source/HairStrandsNiagara/Private/HairStrandsNiagara.cpp b/Engine/Plugins/Experimental/HairStrands/Source/HairStrandsNiagara/Private/HairStrandsNiagara.cpp
--- a/Engine/Plugins/Experimental/HairStrands/Source/HairStrandsNiagara/Private/HairStrandsNiagara.cpp
+++ b/Engine/Plugins/Experimental/HairStrands/Source/HairStrandsNiagara/Private/HairStrandsNiagara.cpp
@@ -4,6 +4,8 @@
 #include "Interfaces/IPluginManager.h"
 #include "ShaderCore.h"
 
+IMPLEMENT_MODULE(FHairStrandsNiagara, HairStrandsNiagara);
+
 void FHairStrandsNiagara::StartupModule()
 {
 	// Maps virtual shader source directory /Plugin/FX/Niagara to the plugin's actual Shaders directory.
