Copying //UE4/Dev-Physics to Dev-Main (//UE4/Dev-Main) [at] 9098278

diff --git a/Engine/Plugins/Experimental/HairStrands/Source/HairStrandsNiagara/Private/NiagaraDataInterfaceHairStrands.cpp b/Engine/Plugins/Experimental/HairStrands/Source/HairStrandsNiagara/Private/NiagaraDataInterfaceHairStrands.cpp
--- a/Engine/Plugins/Experimental/HairStrands/Source/HairStrandsNiagara/Private/NiagaraDataInterfaceHairStrands.cpp
+++ b/Engine/Plugins/Experimental/HairStrands/Source/HairStrandsNiagara/Private/NiagaraDataInterfaceHairStrands.cpp
@@ -6,8 +6,12 @@
 #include "NiagaraRenderer.h"
 #include "NiagaraSystemInstance.h"
 #include "NiagaraEmitterInstanceBatcher.h"
+
 #include "ShaderParameterUtils.h"
-#include "HairStrandsAsset.h"  
+#include "ClearQuad.h"
+
+#include "GroomComponent.h"
+#include "GroomAsset.h"  
 
 #define LOCTEXT_NAMESPACE "NiagaraDataInterfaceHairStrands"
 
@@ -19,11 +23,42 @@ static const FName GetNumStrandsName(TEXT("GetNumStrands"));
 static const FName GetRootThicknessName(TEXT("GetRootThickness"));
 static const FName GetTipThicknessName(TEXT("GetTipThickness"));
 static const FName GetWorldTransformName(TEXT("GetWorldTransform"));
-static const FName GetVertexPositionName(TEXT("GetStrandPosition"));
+static const FName GetWorldInverseName(TEXT("GetWorldInverse"));
+static const FName GetPointPositionName(TEXT("GetPointPosition"));
 
 //------------------------------------------------------------------------------------------------------------
 
 static const FName ComputeNodePositionName(TEXT("ComputeNodePosition"));
+static const FName ComputeNodeOrientationName(TEXT("ComputeNodeOrientation"));
+static const FName ComputeNodeMassName(TEXT("ComputeNodeMass"));
+static const FName ComputeNodeInertiaName(TEXT("ComputeNodeInertia"));
+static const FName ComputeEdgeLengthName(TEXT("ComputeEdgeLength"));
+static const FName ComputeEdgeRotationName(TEXT("ComputeEdgeRotation"));
+static const FName ComputeRestPositionName(TEXT("ComputeRestPosition"));
+static const FName ComputeRestOrientationName(TEXT("ComputeRestOrientation"));
+
+//------------------------------------------------------------------------------------------------------------
+
+static const FName AdvectNodePositionName(TEXT("AdvectNodePosition"));
+static const FName AdvectNodeOrientationName(TEXT("AdvectNodeOrientation"));
+static const FName UpdateLinearVelocityName(TEXT("UpdateLinearVelocity"));
+static const FName UpdateAngularVelocityName(TEXT("UpdateAngularVelocity"));
+
+//------------------------------------------------------------------------------------------------------------
+
+static const FName AttachNodePositionName(TEXT("AttachNodePosition"));
+static const FName AttachNodeOrientationName(TEXT("AttachNodeOrientation"));
+
+//------------------------------------------------------------------------------------------------------------
+
+static const FName UpdatePointPositionName(TEXT("UpdatePointPosition"));
+static const FName ResetPointPositionName(TEXT("ResetPointPosition"));
+
+//------------------------------------------------------------------------------------------------------------
+
+static const FName BuildCollisionGridName(TEXT("BuildCollisionGrid"));
+static const FName QueryCollisionGridName(TEXT("QueryCollisionGrid"));
+static const FName ProjectCollisionGridName(TEXT("ProjectCollisionGrid"));
 
 //------------------------------------------------------------------------------------------------------------
 
@@ -33,11 +68,18 @@ const FString UNiagaraDataInterfaceHairStrands::StrandDensityName(TEXT("StrandDe
 const FString UNiagaraDataInterfaceHairStrands::RootThicknessName(TEXT("RootThickness_"));
 const FString UNiagaraDataInterfaceHairStrands::TipThicknessName(TEXT("TipThickness_"));
 const FString UNiagaraDataInterfaceHairStrands::WorldTransformName(TEXT("WorldTransform_"));
+const FString UNiagaraDataInterfaceHairStrands::WorldInverseName(TEXT("WorldInverse_"));
+const FString UNiagaraDataInterfaceHairStrands::WorldRotationName(TEXT("WorldRotation_"));
+
 const FString UNiagaraDataInterfaceHairStrands::PointsPositionsBufferName(TEXT("PointsPositionsBuffer_"));
 const FString UNiagaraDataInterfaceHairStrands::CurvesOffsetsBufferName(TEXT("CurvesOffsetsBuffer_"));
-const FString UNiagaraDataInterfaceHairStrands::NodesPositionsBufferName(TEXT("NodesPositionsBuffer_"));
-const FString UNiagaraDataInterfaceHairStrands::PointsNodesBufferName(TEXT("PointsNodesBuffer_"));
-const FString UNiagaraDataInterfaceHairStrands::PointsWeightsBufferName(TEXT("PointsWeightsBuffer_"));
+const FString UNiagaraDataInterfaceHairStrands::RestPositionsBufferName(TEXT("RestPositionsBuffer_"));
+
+const FString UNiagaraDataInterfaceHairStrands::GridCurrentBufferName(TEXT("GridCurrentBuffer_"));
+const FString UNiagaraDataInterfaceHairStrands::GridDestinationBufferName(TEXT("GridDestinationBuffer_"));
+
+const FString UNiagaraDataInterfaceHairStrands::GridSizeName(TEXT("GridSize_"));
+const FString UNiagaraDataInterfaceHairStrands::GridOriginName(TEXT("GridOrigin_"));
 
 //------------------------------------------------------------------------------------------------------------
 
@@ -51,11 +93,18 @@ struct FNDIHairStrandsParametersName
 		RootThicknessName = UNiagaraDataInterfaceHairStrands::RootThicknessName + Suffix;
 		TipThicknessName = UNiagaraDataInterfaceHairStrands::TipThicknessName + Suffix;
 		WorldTransformName = UNiagaraDataInterfaceHairStrands::WorldTransformName + Suffix;
+		WorldInverseName = UNiagaraDataInterfaceHairStrands::WorldInverseName + Suffix;
+		WorldRotationName = UNiagaraDataInterfaceHairStrands::WorldRotationName + Suffix;
+
 		PointsPositionsBufferName = UNiagaraDataInterfaceHairStrands::PointsPositionsBufferName + Suffix;
 		CurvesOffsetsBufferName = UNiagaraDataInterfaceHairStrands::CurvesOffsetsBufferName + Suffix;
-		NodesPositionsBufferName = UNiagaraDataInterfaceHairStrands::NodesPositionsBufferName + Suffix;
-		PointsNodesBufferName = UNiagaraDataInterfaceHairStrands::PointsNodesBufferName + Suffix;
-		PointsWeightsBufferName = UNiagaraDataInterfaceHairStrands::PointsWeightsBufferName + Suffix;
+		RestPositionsBufferName = UNiagaraDataInterfaceHairStrands::RestPositionsBufferName + Suffix;
+
+		GridCurrentBufferName = UNiagaraDataInterfaceHairStrands::GridCurrentBufferName + Suffix;
+		GridDestinationBufferName = UNiagaraDataInterfaceHairStrands::GridDestinationBufferName + Suffix;
+
+		GridSizeName = UNiagaraDataInterfaceHairStrands::GridSizeName + Suffix;
+		GridOriginName = UNiagaraDataInterfaceHairStrands::GridOriginName + Suffix;
 	}
 
 	FString NumStrandsName;
@@ -64,75 +113,67 @@ struct FNDIHairStrandsParametersName
 	FString RootThicknessName;
 	FString TipThicknessName;
 	FString WorldTransformName;
+	FString WorldInverseName;
+	FString WorldRotationName;
+
 	FString PointsPositionsBufferName;
 	FString CurvesOffsetsBufferName;
-	FString NodesPositionsBufferName;
-	FString PointsNodesBufferName;
-	FString PointsWeightsBufferName;
+	FString RestPositionsBufferName;
+
+	FString GridCurrentBufferName;
+	FString GridDestinationBufferName;
+
+	FString GridSizeName;
+	FString GridOriginName;
 };
 
 //------------------------------------------------------------------------------------------------------------
 
-void FNDIHairStrandsBuffer::SetAsset(const UHairStrandsAsset* HairStrandsAsset)
+void FNDICollisionGridBuffer::SetGridSize(const FGridSize InGridSize)
 {
-	SourceAsset = HairStrandsAsset;
+	GridSize = InGridSize;
 }
 
-void FNDIHairStrandsBuffer::SetStrandSize(const uint32 HairStrandSize)
+void FNDICollisionGridBuffer::InitRHI()
 {
-	StrandSize = HairStrandSize;
+	if(GridSize.GridSizeX != 0 && GridSize.GridSizeY != 0 && GridSize.GridSizeZ != 0)
+	{
+		static const uint32 NumComponents = 9;
+		GridDataBuffer.Initialize(sizeof(int32), GridSize.GridSizeX*NumComponents, GridSize.GridSizeY,
+			GridSize.GridSizeZ, EPixelFormat::PF_R32_SINT);
+	}
 }
 
-void FNDIHairStrandsBuffer::InitRHI()
+void FNDICollisionGridBuffer::ReleaseRHI()
 {
-	if (SourceAsset != nullptr)
-	{
-		if (SourceAsset->HairStrandsResource)
-		{
-			PointsPositionsBufferUav = SourceAsset->HairStrandsResource->PositionBuffer.UAV;
-		}
-		TArray<FHairStrandsPositionFormat::Type> NodesPositions;
-		TArray<FHairStrandsIndexFormat::Type> PointsNodes;
-		TArray<FHairStrandsWeightFormat::Type> PointsWeights;
-		SourceAsset->StrandsDatas.BuildSimulationDatas(StrandSize, NodesPositions, PointsWeights, PointsNodes);
-		{
-			const uint32 PositionCount = NodesPositions.Num();
-			const uint32 PositionBytes = FHairStrandsPositionFormat::SizeInByte*PositionCount;
-
-			NodesPositionsBuffer.Initialize(FHairStrandsPositionFormat::SizeInByte, PositionCount, FHairStrandsPositionFormat::Format, BUF_Static);
-			void* PositionBufferData = RHILockVertexBuffer(NodesPositionsBuffer.Buffer, 0, PositionBytes, RLM_WriteOnly);
-
-			FMemory::Memcpy(PositionBufferData, NodesPositions.GetData(), PositionBytes);
-			RHIUnlockVertexBuffer(NodesPositionsBuffer.Buffer);
-		}
-		{
-			const uint32 NodesCount = PointsNodes.Num();
-			const uint32 NodesBytes = FHairStrandsIndexFormat::SizeInByte*NodesCount;
+	GridDataBuffer.Release();
+}
 
-			PointsNodesBuffer.Initialize(FHairStrandsIndexFormat::SizeInByte, NodesCount, FHairStrandsIndexFormat::Format, BUF_Static);
-			void* NodesBufferData = RHILockVertexBuffer(PointsNodesBuffer.Buffer, 0, NodesBytes, RLM_WriteOnly);
+void FNDICollisionGridBuffer::ClearBuffers(FRHICommandList& RHICmdList)
+{
+	ClearUAV(RHICmdList, GridDataBuffer, FLinearColor(0, 0, 0, 0));
+}
 
-			FMemory::Memcpy(NodesBufferData, PointsNodes.GetData(), NodesBytes);
-			RHIUnlockVertexBuffer(PointsNodesBuffer.Buffer);
-		}
-		{
-			const uint32 WeightsCount = PointsWeights.Num();
-			const uint32 WeightsBytes = FHairStrandsWeightFormat::SizeInByte*WeightsCount;
+//------------------------------------------------------------------------------------------------------------
 
-			PointsWeightsBuffer.Initialize(FHairStrandsWeightFormat::SizeInByte, WeightsCount, FHairStrandsWeightFormat::Format, BUF_Static);
-			void* WeightsBufferData = RHILockVertexBuffer(PointsWeightsBuffer.Buffer, 0, WeightsBytes, RLM_WriteOnly);
+void FNDIHairStrandsBuffer::SetHairAsset(const FHairStrandsDatas*  HairStrandsDatas, const FHairStrandsResource*  HairStrandsResource)
+{
+	SourceDatas = HairStrandsDatas;
+	SourceResources = HairStrandsResource;
+}
 
-			FMemory::Memcpy(WeightsBufferData, PointsWeights.GetData(), WeightsBytes);
-			RHIUnlockVertexBuffer(PointsWeightsBuffer.Buffer);
-		}
+void FNDIHairStrandsBuffer::InitRHI()
+{
+	if (SourceDatas != nullptr && SourceResources != nullptr)
+	{
 		{
-			const uint32 OffsetCount = SourceAsset->StrandsDatas.GetNumCurves() + 1;
+			const uint32 OffsetCount = SourceDatas->GetNumCurves() + 1;
 			const uint32 OffsetBytes = sizeof(uint32)*OffsetCount;
 
 			CurvesOffsetsBuffer.Initialize(sizeof(uint32), OffsetCount, EPixelFormat::PF_R32_UINT, BUF_Static);
 			void* OffsetBufferData = RHILockVertexBuffer(CurvesOffsetsBuffer.Buffer, 0, OffsetBytes, RLM_WriteOnly);
 
-			FMemory::Memcpy(OffsetBufferData, SourceAsset->StrandsDatas.StrandsCurves.CurvesOffset.GetData(), OffsetBytes);
+			FMemory::Memcpy(OffsetBufferData, SourceDatas->StrandsCurves.CurvesOffset.GetData(), OffsetBytes);
 			RHIUnlockVertexBuffer(CurvesOffsetsBuffer.Buffer);
 		}
 	}
@@ -141,9 +182,15 @@ void FNDIHairStrandsBuffer::InitRHI()
 void FNDIHairStrandsBuffer::ReleaseRHI()
 {
 	CurvesOffsetsBuffer.Release();
-	NodesPositionsBuffer.Release();
-	PointsNodesBuffer.Release();
-	PointsWeightsBuffer.Release();
+}
+
+//------------------------------------------------------------------------------------------------------------
+
+void FNDIHairStrandsData::SwapBuffers()
+{
+	FNDICollisionGridBuffer* StoredBufferPointer = CurrentGridBuffer;
+	CurrentGridBuffer = DestinationGridBuffer;
+	DestinationGridBuffer = StoredBufferPointer;
 }
 
 //------------------------------------------------------------------------------------------------------------
@@ -155,16 +202,23 @@ struct FNDIHairStrandsParametersCS : public FNiagaraDataInterfaceParametersCS
 		FNDIHairStrandsParametersName ParamNames(ParamRef.ParameterInfo.DataInterfaceHLSLSymbol);
 
 		WorldTransform.Bind(ParameterMap, *ParamNames.WorldTransformName);
+		WorldInverse.Bind(ParameterMap, *ParamNames.WorldInverseName);
+		WorldRotation.Bind(ParameterMap, *ParamNames.WorldRotationName);
 		NumStrands.Bind(ParameterMap, *ParamNames.NumStrandsName);
 		StrandSize.Bind(ParameterMap, *ParamNames.StrandSizeName);
 		StrandDensity.Bind(ParameterMap, *ParamNames.StrandDensityName);
 		RootThickness.Bind(ParameterMap, *ParamNames.RootThicknessName);
 		TipThickness.Bind(ParameterMap, *ParamNames.TipThicknessName);
+
 		PointsPositionsBuffer.Bind(ParameterMap, *ParamNames.PointsPositionsBufferName);
 		CurvesOffsetsBuffer.Bind(ParameterMap, *ParamNames.CurvesOffsetsBufferName);
-		NodesPositionsBuffer.Bind(ParameterMap, *ParamNames.NodesPositionsBufferName);
-		PointsNodesBuffer.Bind(ParameterMap, *ParamNames.PointsNodesBufferName);
-		PointsWeightsBuffer.Bind(ParameterMap, *ParamNames.PointsWeightsBufferName);
+		RestPositionsBuffer.Bind(ParameterMap, *ParamNames.RestPositionsBufferName);
+
+		GridCurrentBuffer.Bind(ParameterMap, *ParamNames.GridCurrentBufferName);
+		GridDestinationBuffer.Bind(ParameterMap, *ParamNames.GridDestinationBufferName);
+
+		GridOrigin.Bind(ParameterMap, *ParamNames.GridOriginName);
+		GridSize.Bind(ParameterMap, *ParamNames.GridSizeName);
 
 		if (!PointsPositionsBuffer.IsBound())
 		{
@@ -174,33 +228,42 @@ struct FNDIHairStrandsParametersCS : public FNiagaraDataInterfaceParametersCS
 		{
 			UE_LOG(LogHairStrands, Warning, TEXT("Binding failed for FNDIHairStrandsParametersCS %s. Was it optimized out?"), *ParamNames.CurvesOffsetsBufferName)
 		}
-		if (!NodesPositionsBuffer.IsBound())
+		if (!RestPositionsBuffer.IsBound())
 		{
-			UE_LOG(LogHairStrands, Warning, TEXT("Binding failed for FNDIHairStrandsParametersCS %s. Was it optimized out?"), *ParamNames.NodesPositionsBufferName)
+			UE_LOG(LogHairStrands, Warning, TEXT("Binding failed for FNDIHairStrandsParametersCS %s. Was it optimized out?"), *ParamNames.RestPositionsBufferName)
 		}
-		if (!PointsNodesBuffer.IsBound())
+
+		if (!GridCurrentBuffer.IsBound())
 		{
-			UE_LOG(LogHairStrands, Warning, TEXT("Binding failed for FNDIHairStrandsParametersCS %s. Was it optimized out?"), *ParamNames.PointsNodesBufferName)
+			UE_LOG(LogHairStrands, Warning, TEXT("Binding failed for FNDIHairStrandsParametersCS %s. Was it optimized out?"), *ParamNames.GridCurrentBufferName)
 		}
-		if (!PointsWeightsBuffer.IsBound())
+
+		if (!GridDestinationBuffer.IsBound())
 		{
-			UE_LOG(LogHairStrands, Warning, TEXT("Binding failed for FNDIHairStrandsParametersCS %s. Was it optimized out?"), *ParamNames.PointsWeightsBufferName)
+			UE_LOG(LogHairStrands, Warning, TEXT("Binding failed for FNDIHairStrandsParametersCS %s. Was it optimized out?"), *ParamNames.GridDestinationBufferName)
 		}
 	}
 
 	virtual void Serialize(FArchive& Ar) override
 	{
 		Ar << WorldTransform;
+		Ar << WorldInverse;
+		Ar << WorldRotation;
 		Ar << NumStrands;
 		Ar << StrandSize;
 		Ar << StrandDensity;
 		Ar << RootThickness;
 		Ar << TipThickness;
+
 		Ar << PointsPositionsBuffer;
 		Ar << CurvesOffsetsBuffer;
-		Ar << NodesPositionsBuffer;
-		Ar << PointsNodesBuffer;
-		Ar << PointsWeightsBuffer;
+		Ar << RestPositionsBuffer;
+
+		Ar << GridCurrentBuffer;
+		Ar << GridDestinationBuffer;
+
+		Ar << GridOrigin;
+		Ar << GridSize;
 	}
 
 	virtual void Set(FRHICommandList& RHICmdList, const FNiagaraDataInterfaceSetArgs& Context) const override
@@ -217,29 +280,50 @@ struct FNDIHairStrandsParametersCS : public FNiagaraDataInterfaceParametersCS
 
 		if (ProxyData != nullptr)
 		{
-			FNDIHairStrandsBuffer* StrandBuffer = ProxyData->StrandBuffer;
-			SetUAVParameter(RHICmdList, ComputeShaderRHI, PointsPositionsBuffer, StrandBuffer->PointsPositionsBufferUav);
-			SetSRVParameter(RHICmdList, ComputeShaderRHI, CurvesOffsetsBuffer, StrandBuffer->CurvesOffsetsBuffer.SRV);
-			SetSRVParameter(RHICmdList, ComputeShaderRHI, NodesPositionsBuffer, StrandBuffer->NodesPositionsBuffer.SRV);
-			SetSRVParameter(RHICmdList, ComputeShaderRHI, PointsNodesBuffer, StrandBuffer->PointsNodesBuffer.SRV);
-			SetSRVParameter(RHICmdList, ComputeShaderRHI, PointsWeightsBuffer, StrandBuffer->PointsWeightsBuffer.SRV);
+			FNDIHairStrandsBuffer* HairStrandsBuffer = ProxyData->HairStrandsBuffer;
+
+			FNDICollisionGridBuffer* CurrentGridBuffer = ProxyData->CurrentGridBuffer;
+			FNDICollisionGridBuffer* DestinationGridBuffer = ProxyData->DestinationGridBuffer;
+
+			SetUAVParameter(RHICmdList, ComputeShaderRHI, PointsPositionsBuffer, HairStrandsBuffer->SourceResources->DeformedPositionBuffer[0].UAV);
+			SetSRVParameter(RHICmdList, ComputeShaderRHI, CurvesOffsetsBuffer, HairStrandsBuffer->CurvesOffsetsBuffer.SRV);
+			SetSRVParameter(RHICmdList, ComputeShaderRHI, RestPositionsBuffer, HairStrandsBuffer->SourceResources->RestPositionBuffer.SRV);
+
+			RHICmdList.TransitionResource(EResourceTransitionAccess::EWritable, EResourceTransitionPipeline::EComputeToCompute, DestinationGridBuffer->GridDataBuffer.UAV);
+			SetUAVParameter(RHICmdList, ComputeShaderRHI, GridDestinationBuffer, DestinationGridBuffer->GridDataBuffer.UAV);
+
+			RHICmdList.TransitionResource(EResourceTransitionAccess::EReadable, EResourceTransitionPipeline::EComputeToCompute, CurrentGridBuffer->GridDataBuffer.UAV);
+			SetSRVParameter(RHICmdList, ComputeShaderRHI, GridCurrentBuffer, CurrentGridBuffer->GridDataBuffer.SRV);
 
 			SetShaderValue(RHICmdList, ComputeShaderRHI, WorldTransform, ProxyData->WorldTransform);
+			SetShaderValue(RHICmdList, ComputeShaderRHI, WorldInverse, ProxyData->WorldTransform.Inverse());
+			SetShaderValue(RHICmdList, ComputeShaderRHI, WorldRotation, ProxyData->WorldTransform.ToQuat());
 			SetShaderValue(RHICmdList, ComputeShaderRHI, NumStrands, ProxyData->NumStrands);
 			SetShaderValue(RHICmdList, ComputeShaderRHI, StrandSize, ProxyData->StrandSize);
 			SetShaderValue(RHICmdList, ComputeShaderRHI, StrandDensity, ProxyData->StrandDensity);
 			SetShaderValue(RHICmdList, ComputeShaderRHI, RootThickness, ProxyData->RootThickness);
 			SetShaderValue(RHICmdList, ComputeShaderRHI, TipThickness, ProxyData->TipThickness);
+
+			SetShaderValue(RHICmdList, ComputeShaderRHI, GridOrigin, ProxyData->GridOrigin);
+			SetShaderValue(RHICmdList, ComputeShaderRHI, GridSize, ProxyData->GridSize);
 		}
 		else
 		{
 			SetUAVParameter(RHICmdList, ComputeShaderRHI, PointsPositionsBuffer, FNiagaraRenderer::GetDummyFloatBuffer().UAV);
 			SetSRVParameter(RHICmdList, ComputeShaderRHI, CurvesOffsetsBuffer, FNiagaraRenderer::GetDummyUIntBuffer().SRV);
-			SetSRVParameter(RHICmdList, ComputeShaderRHI, NodesPositionsBuffer, FNiagaraRenderer::GetDummyFloatBuffer().SRV);
-			SetSRVParameter(RHICmdList, ComputeShaderRHI, PointsNodesBuffer, FNiagaraRenderer::GetDummyUIntBuffer().SRV);
-			SetSRVParameter(RHICmdList, ComputeShaderRHI, PointsWeightsBuffer, FNiagaraRenderer::GetDummyFloatBuffer().SRV);
+			SetSRVParameter(RHICmdList, ComputeShaderRHI, RestPositionsBuffer, FNiagaraRenderer::GetDummyFloatBuffer().SRV);
+
+		
+
+			SetUAVParameter(RHICmdList, ComputeShaderRHI, GridDestinationBuffer, FNiagaraRenderer::GetDummyUIntBuffer().UAV);
+			SetSRVParameter(RHICmdList, ComputeShaderRHI, GridCurrentBuffer, FNiagaraRenderer::GetDummyUIntBuffer().SRV);
+
+			SetShaderValue(RHICmdList, ComputeShaderRHI, GridOrigin, FVector4(0,0,0,0));
+			SetShaderValue(RHICmdList, ComputeShaderRHI, GridSize, FGridSize());
 
 			SetShaderValue(RHICmdList, ComputeShaderRHI, WorldTransform, FMatrix::Identity);
+			SetShaderValue(RHICmdList, ComputeShaderRHI, WorldInverse, FMatrix::Identity);
+			SetShaderValue(RHICmdList, ComputeShaderRHI, WorldRotation, FQuat::Identity);
 			SetShaderValue(RHICmdList, ComputeShaderRHI, NumStrands, 1);
 			SetShaderValue(RHICmdList, ComputeShaderRHI, StrandSize, 1);
 			SetShaderValue(RHICmdList, ComputeShaderRHI, StrandDensity, 1);
@@ -248,9 +332,17 @@ struct FNDIHairStrandsParametersCS : public FNiagaraDataInterfaceParametersCS
 		}
 	}
 
+	virtual void Unset(FRHICommandList& RHICmdList, const FNiagaraDataInterfaceSetArgs& Context) const override
+	{
+		SetUAVParameter(RHICmdList, Context.Shader->GetComputeShader(), PointsPositionsBuffer, nullptr);
+		SetUAVParameter(RHICmdList, Context.Shader->GetComputeShader(), GridDestinationBuffer, nullptr);
+	}
+
 private:
 
 	FShaderParameter WorldTransform;
+	FShaderParameter WorldInverse;
+	FShaderParameter WorldRotation;
 	FShaderParameter NumStrands;
 	FShaderParameter StrandSize;
 	FShaderParameter StrandDensity;
@@ -259,9 +351,13 @@ private:
 
 	FShaderResourceParameter PointsPositionsBuffer;
 	FShaderResourceParameter CurvesOffsetsBuffer;
-	FShaderResourceParameter NodesPositionsBuffer;
-	FShaderResourceParameter PointsNodesBuffer;
-	FShaderResourceParameter PointsWeightsBuffer;
+	FShaderResourceParameter RestPositionsBuffer;
+
+	FShaderResourceParameter GridCurrentBuffer;
+	FShaderResourceParameter GridDestinationBuffer;
+
+	FShaderParameter GridSize;
+	FShaderParameter GridOrigin;
 };
 
 //------------------------------------------------------------------------------------------------------------
@@ -292,8 +388,8 @@ void FNDIHairStrandsProxy::ConsumePerInstanceDataFromGameThread(void* PerInstanc
 	}
 }
 
-void FNDIHairStrandsProxy::InitializePerInstanceData(const FNiagaraSystemInstanceID& SystemInstance, FNDIHairStrandsBuffer* StrandBuffer, uint32 NumStrands, const uint8 StrandSize,
-	const float StrandDensity, const float RootThickness, const float TipThickness)
+void FNDIHairStrandsProxy::InitializePerInstanceData(const FNiagaraSystemInstanceID& SystemInstance, FNDIHairStrandsBuffer* HairStrandsBuffer, FNDICollisionGridBuffer* CurrentGridBuffer, FNDICollisionGridBuffer* DestinationGridBuffer, uint32 NumStrands, const uint8 StrandSize,
+	const float StrandDensity, const float RootThickness, const float TipThickness, const FVector4& GridOrigin, const FGridSize& GridSize)
 {
 	check(IsInRenderingThread());
 
@@ -306,12 +402,16 @@ void FNDIHairStrandsProxy::InitializePerInstanceData(const FNiagaraSystemInstanc
 	{
 		TargetData = &SystemInstancesToProxyData.Add(SystemInstance);
 	}
-	TargetData->StrandBuffer = StrandBuffer;
+	TargetData->HairStrandsBuffer = HairStrandsBuffer;
+	TargetData->CurrentGridBuffer = CurrentGridBuffer;
+	TargetData->DestinationGridBuffer = DestinationGridBuffer;
 	TargetData->NumStrands = NumStrands;
 	TargetData->StrandSize = StrandSize;
 	TargetData->StrandDensity = StrandDensity;
 	TargetData->RootThickness = RootThickness;
 	TargetData->TipThickness = TipThickness;
+	TargetData->GridOrigin = GridOrigin;
+	TargetData->GridSize = GridSize;
 }
 
 void FNDIHairStrandsProxy::DestroyPerInstanceData(NiagaraEmitterInstanceBatcher* Batcher, const FNiagaraSystemInstanceID& SystemInstance)
@@ -330,62 +430,149 @@ UNiagaraDataInterfaceHairStrands::UNiagaraDataInterfaceHairStrands(FObjectInitia
 	, StrandDensity(1.0)
 	, RootThickness(0.1)
 	, TipThickness(0.1)
-	, SourceAsset(nullptr)
+	, DefaultSource(nullptr)
+	, SourceActor(nullptr)
 	, SourceTransform(FMatrix::Identity)
-	, StaticMesh(nullptr)
+	, GridSizeX(10)
+	, GridSizeY(10)
+	, GridSizeZ(10)
+	, SourceComponent(nullptr)
+	, GroomAsset(nullptr)
 {
+
 	Proxy = MakeShared<FNDIHairStrandsProxy, ESPMode::ThreadSafe>();
+
+	/*static int32 RefCounter = 0;
+	if (RefCounter == 21)
+	{
+		UE_LOG(LogHairStrands, Log, TEXT("Bad HairStrands : %d %d"), this, RefCounter);
+	}
+	++RefCounter;*/
+}
+
+bool UNiagaraDataInterfaceHairStrands::IsComponentValid() const
+{
+	return (SourceComponent.IsValid() && SourceComponent != nullptr);
 }
 
 bool UNiagaraDataInterfaceHairStrands::InitPerInstanceData(void* PerInstanceData, FNiagaraSystemInstance* SystemInstance)
 {
 	FNDIHairStrandsData* InstanceData = new (PerInstanceData) FNDIHairStrandsData();
+	check(InstanceData);
+	
+	SourceComponent = nullptr;
+	if (SourceActor)
+	{
+		AGroomActor* HairStrandsActor = Cast<AGroomActor>(SourceActor);
+		if (HairStrandsActor != nullptr)
+		{
+			SourceComponent = HairStrandsActor->GetGroomComponent();
+		}
+		else
+		{
+			SourceComponent = SourceActor->FindComponentByClass<UGroomComponent>();
+		}
+	}
+	else
+	{
+		if (UNiagaraComponent* SimComp = SystemInstance->GetComponent())
+		{
+			if (UGroomComponent* ParentComp = Cast<UGroomComponent>(SimComp->GetAttachParent()))
+			{
+				SourceComponent = ParentComp;
+			}
+			else if (UGroomComponent* OuterComp = SimComp->GetTypedOuter<UGroomComponent>())
+			{
+				SourceComponent = OuterComp;
+			}
+			else if (AActor* Owner = SimComp->GetAttachmentRootActor())
+			{
+				for (UActorComponent* ActorComp : Owner->GetComponents())
+				{
+					UGroomComponent* SourceComp = Cast<UGroomComponent>(ActorComp);
+					if (SourceComp && SourceComp->GroomAsset)
+					{
+						SourceComponent = SourceComp;
+						break;
+					}
+				}
+			}
+		}
+	}
+	GroomAsset = IsComponentValid() ? SourceComponent->GroomAsset : DefaultSource;
+	FHairStrandsDatas* StrandsDatas = IsComponentValid() ? SourceComponent->GetGuideStrandsDatas() : (DefaultSource != nullptr) ? &DefaultSource->HairRenderData : nullptr;
+	FHairStrandsResource* StrandsResource = IsComponentValid() ? SourceComponent->GetGuideStrandsResource() : (DefaultSource != nullptr) ? DefaultSource->HairStrandsResource : nullptr;
 
-	if (SourceAsset == nullptr)
+	if (StrandsDatas == nullptr || StrandsResource == nullptr)
 	{
-		UE_LOG(LogHairStrands, Log, TEXT("Hair Strands data interface has no valid asset. Failed InitPerInstanceData - %s"), *GetFullName());
+		UE_LOG(LogHairStrands, Log, TEXT("Hair Strands data interface has no valid asset. Failed InitPerInstanceData - %s %d"), *GetFullName(), DefaultSource);
 		return false;
 	}
 	else
-	{
-		FNDIHairStrandsBuffer* StrandBuffer = new FNDIHairStrandsBuffer;
-		StrandBuffer->SetAsset(SourceAsset);
+	{ 
+		FNDIHairStrandsBuffer* HairStrandsBuffer = new FNDIHairStrandsBuffer;
+		FNDICollisionGridBuffer* CurrentGridBuffer = new FNDICollisionGridBuffer;
+		FNDICollisionGridBuffer* DestinationGridBuffer = new FNDICollisionGridBuffer;
+		HairStrandsBuffer->SetHairAsset(StrandsDatas, StrandsResource);
+		
 
 		// Push instance data to RT
 		{
-			const uint32 LocalNumStrands = SourceAsset->StrandsDatas.GetNumCurves();
+			const uint32 LocalNumStrands = StrandsDatas->GetNumCurves();
 			const uint8 LocalStrandSize = static_cast<uint8>(StrandSize);
-			const FBox LocalStrandsBox;// = SourceAsset->ComputeStrandsBounds();
-			StrandBuffer->SetStrandSize(LocalStrandSize);
+			const FBox& LocalStrandsBox = StrandsDatas->BoundingBox;
 
-			
-			/*SourceAsset->StrandsDatas.ResampleInternalDatas(static_cast<uint32>(StrandSize));
-			if (StaticMesh)
-			{
-				SourceAsset->StrandsDatas.AttachStrandsRoots(StaticMesh,SourceTransform);
-			}*/
+			const FVector GridExtent = LocalStrandsBox.GetExtent();
+			const FVector GridLengths((LocalStrandsBox.Max.X - LocalStrandsBox.Min.X) / (GridSizeX-1), 
+									  (LocalStrandsBox.Max.Y - LocalStrandsBox.Min.Y) / (GridSizeY-1), 
+									  (LocalStrandsBox.Max.Z - LocalStrandsBox.Min.Z) / (GridSizeZ-1));
+			const float GridLength = GridLengths.GetMax();
 
-			UE_LOG(LogHairStrands, Log, TEXT("Num Strands = %d | Strand Size = %d | Num Vertices = %d | Min = %f %f %f | Max = %f %f %f"), LocalNumStrands, LocalStrandSize,
-				SourceAsset->StrandsDatas.GetNumPoints(), LocalStrandsBox.Min[0], LocalStrandsBox.Min[1], LocalStrandsBox.Min[2], LocalStrandsBox.Max[0], LocalStrandsBox.Max[1], LocalStrandsBox.Max[2] );
+			const FGridSize GridSize(GridSizeX, GridSizeY, GridSizeZ);
+			const FVector4 GridOrigin(-GridExtent.X, -GridExtent.Y, -GridExtent.Z, GridLength);
+			CurrentGridBuffer->SetGridSize(GridSize);
+			DestinationGridBuffer->SetGridSize(GridSize);
 
 			const float LocalStrandDensity = StrandDensity;
 			const float LocalRootThickness = RootThickness;
 			const float LocalTipThickness = TipThickness;
 
-			InstanceData->WorldTransform = SourceTransform;// SourceAsset->GetComponentToWorld().ToMatrixWithScale();
+			const FTransform WorldOffset(-LocalStrandsBox.GetCenter());
+			const FMatrix WorldTransform = WorldOffset.ToInverseMatrixWithScale() * SourceTransform;
+
+			InstanceData->WorldTransform = IsComponentValid() ? WorldTransform * SourceComponent->GetComponentToWorld().ToMatrixWithScale() : 
+				WorldTransform * SystemInstance->GetComponent()->GetComponentToWorld().ToMatrixWithScale();
+
+		//UE_LOG(LogHairStrands, Log, TEXT("Num Strands = %d | Strand Size = %d | Num Vertices = %d | Min = %f %f %f | Max = %f %f %f | Transform = %s"), LocalNumStrands, LocalStrandSize,
+				//StrandsDatas->GetNumPoints(), LocalStrandsBox.Min[0], LocalStrandsBox.Min[1], LocalStrandsBox.Min[2], LocalStrandsBox.Max[0], LocalStrandsBox.Max[1], LocalStrandsBox.Max[2], *InstanceData->WorldTransform.ToString());
+
+			TSet<int> RT_OutputShaderStages = OutputShaderStages;
+			TSet<int> RT_IterationShaderStages = IterationShaderStages;
+
 			InstanceData->StrandSize = LocalStrandSize;
 			InstanceData->StrandDensity = LocalStrandDensity;
 			InstanceData->RootThickness = LocalRootThickness;
 			InstanceData->TipThickness = LocalTipThickness;
 			InstanceData->NumStrands = LocalNumStrands;
-			InstanceData->StrandBuffer = StrandBuffer;
+			InstanceData->HairStrandsBuffer = HairStrandsBuffer;
+			InstanceData->CurrentGridBuffer = CurrentGridBuffer;
+			InstanceData->DestinationGridBuffer = DestinationGridBuffer;
+			InstanceData->GridSize = GridSize;
+			InstanceData->GridOrigin = GridOrigin;
 
 			FNDIHairStrandsProxy* ThisProxy = GetProxyAs<FNDIHairStrandsProxy>();
 			ENQUEUE_RENDER_COMMAND(FNiagaraDIPushInitialInstanceDataToRT) (
-				[ThisProxy, InstanceID = SystemInstance->GetId(), StrandBuffer, LocalNumStrands, LocalStrandSize, LocalStrandDensity, LocalRootThickness, LocalTipThickness](FRHICommandListImmediate& CmdList)
+				[ThisProxy, RT_OutputShaderStages, RT_IterationShaderStages, InstanceID = SystemInstance->GetId(), HairStrandsBuffer, CurrentGridBuffer, DestinationGridBuffer, LocalNumStrands, LocalStrandSize, LocalStrandDensity, LocalRootThickness, LocalTipThickness, GridSize, GridOrigin](FRHICommandListImmediate& CmdList)
 			{
-				StrandBuffer->InitResource();
-				ThisProxy->InitializePerInstanceData(InstanceID, StrandBuffer, LocalNumStrands, LocalStrandSize, LocalStrandDensity, LocalRootThickness, LocalTipThickness);
+				ThisProxy->OutputShaderStages = RT_OutputShaderStages;
+				ThisProxy->IterationShaderStages = RT_IterationShaderStages;
+				ThisProxy->SetElementCount(GridSize.GridSizeX*GridSize.GridSizeY*GridSize.GridSizeZ);
+
+				HairStrandsBuffer->InitResource();
+				CurrentGridBuffer->InitResource();
+				DestinationGridBuffer->InitResource();
+
+				ThisProxy->InitializePerInstanceData(InstanceID, HairStrandsBuffer, CurrentGridBuffer, DestinationGridBuffer, LocalNumStrands, LocalStrandSize, LocalStrandDensity, LocalRootThickness, LocalTipThickness, GridOrigin,GridSize);
 			}
 			);
 		}
@@ -398,19 +585,27 @@ void UNiagaraDataInterfaceHairStrands::DestroyPerInstanceData(void* PerInstanceD
 	FNDIHairStrandsData* InstanceData = static_cast<FNDIHairStrandsData*>(PerInstanceData);
 	InstanceData->~FNDIHairStrandsData();
 
-	if(InstanceData->StrandBuffer)
+	if (InstanceData->HairStrandsBuffer && InstanceData->CurrentGridBuffer && InstanceData->DestinationGridBuffer)
 	{
 		FNDIHairStrandsProxy* ThisProxy = GetProxyAs<FNDIHairStrandsProxy>();
-		FNDIHairStrandsBuffer* InBuffer = InstanceData->StrandBuffer;
+		FNDIHairStrandsBuffer* HairStrandsBuffer = InstanceData->HairStrandsBuffer;
+		FNDICollisionGridBuffer* CurrentGridBuffer = InstanceData->CurrentGridBuffer;
+		FNDICollisionGridBuffer* DestinationGridBuffer = InstanceData->DestinationGridBuffer;
 		ENQUEUE_RENDER_COMMAND(FNiagaraDIDestroyInstanceData) (
-			[ThisProxy, InBuffer, InstanceID = SystemInstance->GetId(), Batcher = SystemInstance->GetBatcher()](FRHICommandListImmediate& CmdList)
+			[ThisProxy, HairStrandsBuffer, CurrentGridBuffer, DestinationGridBuffer, InstanceID = SystemInstance->GetId(), Batcher = SystemInstance->GetBatcher()](FRHICommandListImmediate& CmdList)
 		{
-			InBuffer->ReleaseResource();
+			HairStrandsBuffer->ReleaseResource();
+			CurrentGridBuffer->ReleaseResource();
+			DestinationGridBuffer->ReleaseResource();
 			ThisProxy->DestroyPerInstanceData(Batcher, InstanceID);
-			delete InBuffer;
+			delete HairStrandsBuffer;
+			delete CurrentGridBuffer;
+			delete DestinationGridBuffer;
 		}
 		);
-		InstanceData->StrandBuffer = nullptr;
+		InstanceData->HairStrandsBuffer = nullptr;
+		InstanceData->CurrentGridBuffer = nullptr;
+		InstanceData->DestinationGridBuffer = nullptr;
 	}
 }
 
@@ -418,15 +613,22 @@ bool UNiagaraDataInterfaceHairStrands::PerInstanceTick(void* PerInstanceData, FN
 {
 	FNDIHairStrandsData* InstanceData = static_cast<FNDIHairStrandsData*>(PerInstanceData);
 
+	FHairStrandsDatas* StrandsDatas = IsComponentValid() ? 
+		SourceComponent->GetGuideStrandsDatas() : (DefaultSource != nullptr) ? &DefaultSource->HairRenderData : nullptr;
 	bool RequireReset = false;
-	if (SourceAsset != nullptr)
-	{
-		InstanceData->WorldTransform = SourceTransform; //SourceAsset->GetComponentToWorld().ToMatrixWithScale();
-	}
-	else
+	if (StrandsDatas)
 	{
-		InstanceData->WorldTransform = SourceTransform;
-		RequireReset = true;
+		const FTransform WorldOffset(-StrandsDatas->BoundingBox.GetCenter());
+		const FMatrix WorldTransform = WorldOffset.ToInverseMatrixWithScale() * SourceTransform;
+
+		if (IsComponentValid())
+		{
+			InstanceData->WorldTransform = WorldTransform * SourceComponent->GetComponentToWorld().ToMatrixWithScale();
+		}
+		else
+		{
+			InstanceData->WorldTransform = WorldTransform * SystemInstance->GetComponent()->GetComponentToWorld().ToMatrixWithScale();
+		}
 	}
 	return RequireReset;
 }
@@ -443,8 +645,14 @@ bool UNiagaraDataInterfaceHairStrands::CopyToInternal(UNiagaraDataInterface* Des
 	OtherTyped->StrandDensity = StrandDensity;
 	OtherTyped->RootThickness = RootThickness;
 	OtherTyped->TipThickness = TipThickness;
-	OtherTyped->SourceAsset = SourceAsset;
+	OtherTyped->GroomAsset = GroomAsset;
+	OtherTyped->SourceActor= SourceActor;
+	OtherTyped->SourceComponent = SourceComponent;
+	OtherTyped->DefaultSource = DefaultSource;
 	OtherTyped->SourceTransform = SourceTransform;
+	OtherTyped->GridSizeX = GridSizeX;
+	OtherTyped->GridSizeY = GridSizeY;
+	OtherTyped->GridSizeZ = GridSizeZ;
 
 	return true;
 }
@@ -459,7 +667,9 @@ bool UNiagaraDataInterfaceHairStrands::Equals(const UNiagaraDataInterface* Other
 
 	return  (OtherTyped->StrandSize == StrandSize)
 		&& (OtherTyped->StrandDensity == StrandDensity) && (OtherTyped->RootThickness == RootThickness)
-		&& (OtherTyped->TipThickness == TipThickness) && (OtherTyped->SourceAsset == SourceAsset) && (OtherTyped->SourceTransform == SourceTransform);
+		&& (OtherTyped->TipThickness == TipThickness) && (OtherTyped->GroomAsset == GroomAsset) && (OtherTyped->SourceTransform == SourceTransform) 
+		&& (OtherTyped->SourceActor == SourceActor) && (OtherTyped->SourceComponent == SourceComponent) && (OtherTyped->DefaultSource == DefaultSource)
+		&& (OtherTyped->GridSizeX == GridSizeX) && (OtherTyped->GridSizeY == GridSizeY) && (OtherTyped->GridSizeZ == GridSizeZ);
 }
 
 void UNiagaraDataInterfaceHairStrands::PostInitProperties()
@@ -479,7 +689,7 @@ void UNiagaraDataInterfaceHairStrands::GetFunctions(TArray<FNiagaraFunctionSigna
 		Sig.Name = GetNumStrandsName;
 		Sig.bMemberFunction = true;
 		Sig.bRequiresContext = false;
-		Sig.Inputs.Add(FNiagaraVariable(FNiagaraTypeDefinition(GetClass()), TEXT("Strand Hair")));
+		Sig.Inputs.Add(FNiagaraVariable(FNiagaraTypeDefinition(GetClass()), TEXT("Hair Strands")));
 		Sig.Outputs.Add(FNiagaraVariable(FNiagaraTypeDefinition::GetIntDef(), TEXT("Num Strands")));
 
 		OutFunctions.Add(Sig);
@@ -489,7 +699,7 @@ void UNiagaraDataInterfaceHairStrands::GetFunctions(TArray<FNiagaraFunctionSigna
 		Sig.Name = GetStrandDensityName;
 		Sig.bMemberFunction = true;
 		Sig.bRequiresContext = false;
-		Sig.Inputs.Add(FNiagaraVariable(FNiagaraTypeDefinition(GetClass()), TEXT("Strand Hair")));
+		Sig.Inputs.Add(FNiagaraVariable(FNiagaraTypeDefinition(GetClass()), TEXT("Hair Strands")));
 		Sig.Outputs.Add(FNiagaraVariable(FNiagaraTypeDefinition::GetFloatDef(), TEXT("Strand Density")));
 
 		OutFunctions.Add(Sig);
@@ -499,7 +709,7 @@ void UNiagaraDataInterfaceHairStrands::GetFunctions(TArray<FNiagaraFunctionSigna
 		Sig.Name = GetStrandSizeName;
 		Sig.bMemberFunction = true;
 		Sig.bRequiresContext = false;
-		Sig.Inputs.Add(FNiagaraVariable(FNiagaraTypeDefinition(GetClass()), TEXT("Strand Hair")));
+		Sig.Inputs.Add(FNiagaraVariable(FNiagaraTypeDefinition(GetClass()), TEXT("Hair Strands")));
 		Sig.Outputs.Add(FNiagaraVariable(FNiagaraTypeDefinition::GetIntDef(), TEXT("Strand Size")));
 
 		OutFunctions.Add(Sig);
@@ -509,7 +719,7 @@ void UNiagaraDataInterfaceHairStrands::GetFunctions(TArray<FNiagaraFunctionSigna
 		Sig.Name = GetRootThicknessName;
 		Sig.bMemberFunction = true;
 		Sig.bRequiresContext = false;
-		Sig.Inputs.Add(FNiagaraVariable(FNiagaraTypeDefinition(GetClass()), TEXT("Strand Hair")));
+		Sig.Inputs.Add(FNiagaraVariable(FNiagaraTypeDefinition(GetClass()), TEXT("Hair Strands")));
 		Sig.Outputs.Add(FNiagaraVariable(FNiagaraTypeDefinition::GetFloatDef(), TEXT("Root Thickness")));
 
 		OutFunctions.Add(Sig);
@@ -519,7 +729,7 @@ void UNiagaraDataInterfaceHairStrands::GetFunctions(TArray<FNiagaraFunctionSigna
 		Sig.Name = GetTipThicknessName;
 		Sig.bMemberFunction = true;
 		Sig.bRequiresContext = false;
-		Sig.Inputs.Add(FNiagaraVariable(FNiagaraTypeDefinition(GetClass()), TEXT("Strand Hair")));
+		Sig.Inputs.Add(FNiagaraVariable(FNiagaraTypeDefinition(GetClass()), TEXT("Hair Strands")));
 		Sig.Outputs.Add(FNiagaraVariable(FNiagaraTypeDefinition::GetFloatDef(), TEXT("Tip Thickness")));
 
 		OutFunctions.Add(Sig);
@@ -529,17 +739,27 @@ void UNiagaraDataInterfaceHairStrands::GetFunctions(TArray<FNiagaraFunctionSigna
 		Sig.Name = GetWorldTransformName;
 		Sig.bMemberFunction = true;
 		Sig.bRequiresContext = false;
-		Sig.Inputs.Add(FNiagaraVariable(FNiagaraTypeDefinition(GetClass()), TEXT("Strand Hair")));
+		Sig.Inputs.Add(FNiagaraVariable(FNiagaraTypeDefinition(GetClass()), TEXT("Hair Strands")));
 		Sig.Outputs.Add(FNiagaraVariable(FNiagaraTypeDefinition::GetMatrix4Def(), TEXT("World Transform")));
 
 		OutFunctions.Add(Sig);
 	}
 	{
 		FNiagaraFunctionSignature Sig;
-		Sig.Name = GetVertexPositionName;
+		Sig.Name = GetWorldInverseName;
+		Sig.bMemberFunction = true;
+		Sig.bRequiresContext = false;
+		Sig.Inputs.Add(FNiagaraVariable(FNiagaraTypeDefinition(GetClass()), TEXT("Hair Strands")));
+		Sig.Outputs.Add(FNiagaraVariable(FNiagaraTypeDefinition::GetMatrix4Def(), TEXT("World Inverse")));
+
+		OutFunctions.Add(Sig);
+	}
+	{
+		FNiagaraFunctionSignature Sig;
+		Sig.Name = GetPointPositionName;
 		Sig.bMemberFunction = true;
 		Sig.bRequiresContext = false;
-		Sig.Inputs.Add(FNiagaraVariable(FNiagaraTypeDefinition(GetClass()), TEXT("Strand Hair")));
+		Sig.Inputs.Add(FNiagaraVariable(FNiagaraTypeDefinition(GetClass()), TEXT("Hair Strands")));
 		Sig.Inputs.Add(FNiagaraVariable(FNiagaraTypeDefinition::GetIntDef(), TEXT("Vertex Index")));
 		Sig.Outputs.Add(FNiagaraVariable(FNiagaraTypeDefinition::GetVec3Def(), TEXT("Vertex Position")));
 
@@ -550,10 +770,237 @@ void UNiagaraDataInterfaceHairStrands::GetFunctions(TArray<FNiagaraFunctionSigna
 		Sig.Name = ComputeNodePositionName;
 		Sig.bMemberFunction = true;
 		Sig.bRequiresContext = false;
-		Sig.Inputs.Add(FNiagaraVariable(FNiagaraTypeDefinition(GetClass()), TEXT("Strand Hair")));
+		Sig.Inputs.Add(FNiagaraVariable(FNiagaraTypeDefinition(GetClass()), TEXT("Hair Strands")));
+		Sig.Inputs.Add(FNiagaraVariable(FNiagaraTypeDefinition::GetIntDef(), TEXT("Node Index")));
+		Sig.Outputs.Add(FNiagaraVariable(FNiagaraTypeDefinition::GetVec3Def(), TEXT("Node Position")));
+
+		OutFunctions.Add(Sig);
+	}
+	{
+		FNiagaraFunctionSignature Sig;
+		Sig.Name = ComputeNodeOrientationName;
+		Sig.bMemberFunction = true;
+		Sig.bRequiresContext = false;
+		Sig.Inputs.Add(FNiagaraVariable(FNiagaraTypeDefinition(GetClass()), TEXT("Hair Strands")));
+		Sig.Inputs.Add(FNiagaraVariable(FNiagaraTypeDefinition::GetIntDef(), TEXT("Node Index")));
+		Sig.Outputs.Add(FNiagaraVariable(FNiagaraTypeDefinition::GetQuatDef(), TEXT("Node Orientation")));
+
+		OutFunctions.Add(Sig);
+	}
+	{
+		FNiagaraFunctionSignature Sig;
+		Sig.Name = ComputeNodeMassName;
+		Sig.bMemberFunction = true;
+		Sig.bRequiresContext = false;
+		Sig.Inputs.Add(FNiagaraVariable(FNiagaraTypeDefinition(GetClass()), TEXT("Hair Strands")));
+		Sig.Inputs.Add(FNiagaraVariable(FNiagaraTypeDefinition::GetIntDef(), TEXT("Node Index")));
+		Sig.Outputs.Add(FNiagaraVariable(FNiagaraTypeDefinition::GetFloatDef(), TEXT("Node Mass")));
+
+		OutFunctions.Add(Sig);
+	}
+	{
+		FNiagaraFunctionSignature Sig;
+		Sig.Name = ComputeNodeInertiaName;
+		Sig.bMemberFunction = true;
+		Sig.bRequiresContext = false;
+		Sig.Inputs.Add(FNiagaraVariable(FNiagaraTypeDefinition(GetClass()), TEXT("Hair Strands")));
+		Sig.Inputs.Add(FNiagaraVariable(FNiagaraTypeDefinition::GetIntDef(), TEXT("Node Index")));
+		Sig.Outputs.Add(FNiagaraVariable(FNiagaraTypeDefinition::GetVec3Def(), TEXT("Node Inertia")));
+
+		OutFunctions.Add(Sig);
+	}
+	{
+		FNiagaraFunctionSignature Sig;
+		Sig.Name = ComputeEdgeLengthName;
+		Sig.bMemberFunction = true;
+		Sig.bRequiresContext = false;
+		Sig.Inputs.Add(FNiagaraVariable(FNiagaraTypeDefinition(GetClass()), TEXT("Hair Strands")));
+		Sig.Inputs.Add(FNiagaraVariable(FNiagaraTypeDefinition::GetIntDef(), TEXT("Node Index")));
+		Sig.Inputs.Add(FNiagaraVariable(FNiagaraTypeDefinition::GetVec3Def(), TEXT("Node Position")));
+		Sig.Outputs.Add(FNiagaraVariable(FNiagaraTypeDefinition::GetFloatDef(), TEXT("Edge Length")));
+
+		OutFunctions.Add(Sig);
+	}
+	{
+		FNiagaraFunctionSignature Sig;
+		Sig.Name = ComputeEdgeRotationName;
+		Sig.bMemberFunction = true;
+		Sig.bRequiresContext = false;
+		Sig.Inputs.Add(FNiagaraVariable(FNiagaraTypeDefinition(GetClass()), TEXT("Hair Strands")));
 		Sig.Inputs.Add(FNiagaraVariable(FNiagaraTypeDefinition::GetIntDef(), TEXT("Node Index")));
+		Sig.Inputs.Add(FNiagaraVariable(FNiagaraTypeDefinition::GetQuatDef(), TEXT("Node Orientation")));
+		Sig.Outputs.Add(FNiagaraVariable(FNiagaraTypeDefinition::GetQuatDef(), TEXT("Edge Rotation")));
+
+		OutFunctions.Add(Sig);
+	}
+	{
+		FNiagaraFunctionSignature Sig;
+		Sig.Name = ComputeRestPositionName;
+		Sig.bMemberFunction = true;
+		Sig.bRequiresContext = false;
+		Sig.Inputs.Add(FNiagaraVariable(FNiagaraTypeDefinition(GetClass()), TEXT("Hair Strands")));
+		Sig.Inputs.Add(FNiagaraVariable(FNiagaraTypeDefinition::GetVec3Def(), TEXT("Node Position")));
+		Sig.Outputs.Add(FNiagaraVariable(FNiagaraTypeDefinition::GetVec3Def(), TEXT("Rest Position")));
+
+		OutFunctions.Add(Sig);
+	}
+	{
+		FNiagaraFunctionSignature Sig;
+		Sig.Name = ComputeRestOrientationName;
+		Sig.bMemberFunction = true;
+		Sig.bRequiresContext = false;
+		Sig.Inputs.Add(FNiagaraVariable(FNiagaraTypeDefinition(GetClass()), TEXT("Hair Strands")));
+		Sig.Inputs.Add(FNiagaraVariable(FNiagaraTypeDefinition::GetQuatDef(), TEXT("Node Orientation")));
+		Sig.Outputs.Add(FNiagaraVariable(FNiagaraTypeDefinition::GetQuatDef(), TEXT("Rest Orientation")));
+
+		OutFunctions.Add(Sig);
+	}
+	{
+		FNiagaraFunctionSignature Sig;
+		Sig.Name = AdvectNodePositionName;
+		Sig.bMemberFunction = true;
+		Sig.bRequiresContext = false;
+		Sig.Inputs.Add(FNiagaraVariable(FNiagaraTypeDefinition(GetClass()), TEXT("Hair Strands")));
+		Sig.Inputs.Add(FNiagaraVariable(FNiagaraTypeDefinition::GetFloatDef(), TEXT("Node Mass")));
+		Sig.Inputs.Add(FNiagaraVariable(FNiagaraTypeDefinition::GetFloatDef(), TEXT("Inverse Mass")));
+		Sig.Inputs.Add(FNiagaraVariable(FNiagaraTypeDefinition::GetVec3Def(), TEXT("External Force")));
+		Sig.Inputs.Add(FNiagaraVariable(FNiagaraTypeDefinition::GetFloatDef(), TEXT("Delta Time")));
+		Sig.Inputs.Add(FNiagaraVariable(FNiagaraTypeDefinition::GetVec3Def(), TEXT("Linear Velocity")));
+		Sig.Inputs.Add(FNiagaraVariable(FNiagaraTypeDefinition::GetVec3Def(), TEXT("Node Position")));
+		Sig.Outputs.Add(FNiagaraVariable(FNiagaraTypeDefinition::GetVec3Def(), TEXT("Linear Velocity")));
+		Sig.Outputs.Add(FNiagaraVariable(FNiagaraTypeDefinition::GetVec3Def(), TEXT("Node Position")));
+
+		OutFunctions.Add(Sig);
+	}
+	{
+		FNiagaraFunctionSignature Sig;
+		Sig.Name = AdvectNodeOrientationName;
+		Sig.bMemberFunction = true;
+		Sig.bRequiresContext = false;
+		Sig.Inputs.Add(FNiagaraVariable(FNiagaraTypeDefinition(GetClass()), TEXT("Hair Strands")));
+		Sig.Inputs.Add(FNiagaraVariable(FNiagaraTypeDefinition::GetVec3Def(), TEXT("Node Inertia")));
+		Sig.Inputs.Add(FNiagaraVariable(FNiagaraTypeDefinition::GetVec3Def(), TEXT("Inverse Inertia")));
+		Sig.Inputs.Add(FNiagaraVariable(FNiagaraTypeDefinition::GetVec3Def(), TEXT("External Torque")));
+		Sig.Inputs.Add(FNiagaraVariable(FNiagaraTypeDefinition::GetFloatDef(), TEXT("Delta Time")));
+		Sig.Inputs.Add(FNiagaraVariable(FNiagaraTypeDefinition::GetVec3Def(), TEXT("Angular Velocity")));
+		Sig.Inputs.Add(FNiagaraVariable(FNiagaraTypeDefinition::GetQuatDef(), TEXT("Node Orientation")));
+		Sig.Outputs.Add(FNiagaraVariable(FNiagaraTypeDefinition::GetVec3Def(), TEXT("Angular Velocity")));
+		Sig.Outputs.Add(FNiagaraVariable(FNiagaraTypeDefinition::GetQuatDef(), TEXT("Node Orientation"))); 
+
+		OutFunctions.Add(Sig);
+	}
+	{
+		FNiagaraFunctionSignature Sig;
+		Sig.Name = UpdateLinearVelocityName;
+		Sig.bMemberFunction = true;
+		Sig.bRequiresContext = false;
+		Sig.Inputs.Add(FNiagaraVariable(FNiagaraTypeDefinition(GetClass()), TEXT("Hair Strands")));
+		Sig.Inputs.Add(FNiagaraVariable(FNiagaraTypeDefinition::GetVec3Def(), TEXT("Previous Position")));
+		Sig.Inputs.Add(FNiagaraVariable(FNiagaraTypeDefinition::GetVec3Def(), TEXT("Node Position")));
+		Sig.Inputs.Add(FNiagaraVariable(FNiagaraTypeDefinition::GetFloatDef(), TEXT("Delta Time")));
+		Sig.Outputs.Add(FNiagaraVariable(FNiagaraTypeDefinition::GetVec3Def(), TEXT("Linear Velocity")));
+
+		OutFunctions.Add(Sig);
+	}
+	{
+		FNiagaraFunctionSignature Sig;
+		Sig.Name = UpdateAngularVelocityName;
+		Sig.bMemberFunction = true;
+		Sig.bRequiresContext = false;
+		Sig.Inputs.Add(FNiagaraVariable(FNiagaraTypeDefinition(GetClass()), TEXT("Hair Strands")));
+		Sig.Inputs.Add(FNiagaraVariable(FNiagaraTypeDefinition::GetQuatDef(), TEXT("Previous Orientation")));
+		Sig.Inputs.Add(FNiagaraVariable(FNiagaraTypeDefinition::GetQuatDef(), TEXT("Node Orientation")));
+		Sig.Inputs.Add(FNiagaraVariable(FNiagaraTypeDefinition::GetFloatDef(), TEXT("Delta Time")));
+		Sig.Outputs.Add(FNiagaraVariable(FNiagaraTypeDefinition::GetVec3Def(), TEXT("Angular Velocity")));
+
+		OutFunctions.Add(Sig);
+	}
+	{
+		FNiagaraFunctionSignature Sig;
+		Sig.Name = AttachNodePositionName;
+		Sig.bMemberFunction = true;
+		Sig.bRequiresContext = false;
+		Sig.Inputs.Add(FNiagaraVariable(FNiagaraTypeDefinition(GetClass()), TEXT("Hair Strands")));
+		Sig.Inputs.Add(FNiagaraVariable(FNiagaraTypeDefinition::GetVec3Def(), TEXT("Rest Position")));
 		Sig.Outputs.Add(FNiagaraVariable(FNiagaraTypeDefinition::GetVec3Def(), TEXT("Node Position")));
 
+		OutFunctions.Add(Sig);
+	}
+	{
+		FNiagaraFunctionSignature Sig;
+		Sig.Name = AttachNodeOrientationName;
+		Sig.bMemberFunction = true;
+		Sig.bRequiresContext = false;
+		Sig.Inputs.Add(FNiagaraVariable(FNiagaraTypeDefinition(GetClass()), TEXT("Hair Strands")));
+		Sig.Inputs.Add(FNiagaraVariable(FNiagaraTypeDefinition::GetQuatDef(), TEXT("Rest Orientation")));
+		Sig.Outputs.Add(FNiagaraVariable(FNiagaraTypeDefinition::GetQuatDef(), TEXT("Node Orientation")));
+
+		OutFunctions.Add(Sig);
+	}
+	{
+		FNiagaraFunctionSignature Sig;
+		Sig.Name = UpdatePointPositionName;
+		Sig.bMemberFunction = true;
+		Sig.bRequiresContext = false;
+		Sig.Inputs.Add(FNiagaraVariable(FNiagaraTypeDefinition(GetClass()), TEXT("Hair Strands")));
+		Sig.Inputs.Add(FNiagaraVariable(FNiagaraTypeDefinition::GetIntDef(), TEXT("Node Index")));
+		Sig.Inputs.Add(FNiagaraVariable(FNiagaraTypeDefinition::GetVec3Def(), TEXT("Node Displace")));
+		Sig.Outputs.Add(FNiagaraVariable(FNiagaraTypeDefinition::GetBoolDef(), TEXT("Report Status")));
+
+		OutFunctions.Add(Sig);
+	}
+	{
+		FNiagaraFunctionSignature Sig;
+		Sig.Name = ResetPointPositionName;
+		Sig.bMemberFunction = true;
+		Sig.bRequiresContext = false;
+		Sig.Inputs.Add(FNiagaraVariable(FNiagaraTypeDefinition(GetClass()), TEXT("Hair Strands")));
+		Sig.Inputs.Add(FNiagaraVariable(FNiagaraTypeDefinition::GetIntDef(), TEXT("Node Index")));
+		Sig.Outputs.Add(FNiagaraVariable(FNiagaraTypeDefinition::GetBoolDef(), TEXT("Report Status")));
+
+		OutFunctions.Add(Sig);
+	}
+	{
+		FNiagaraFunctionSignature Sig;
+		Sig.Name = BuildCollisionGridName;
+		Sig.bMemberFunction = true;
+		Sig.bRequiresContext = false;
+		Sig.Inputs.Add(FNiagaraVariable(FNiagaraTypeDefinition(GetClass()), TEXT("Hair Strands")));
+		Sig.Inputs.Add(FNiagaraVariable(FNiagaraTypeDefinition::GetVec3Def(), TEXT("Node Position")));
+		Sig.Inputs.Add(FNiagaraVariable(FNiagaraTypeDefinition::GetVec3Def(), TEXT("Node Velocity")));
+		Sig.Inputs.Add(FNiagaraVariable(FNiagaraTypeDefinition::GetFloatDef(), TEXT("Node Mass")));
+		Sig.Inputs.Add(FNiagaraVariable(FNiagaraTypeDefinition::GetVec3Def(), TEXT("Node GradientX")));
+		Sig.Inputs.Add(FNiagaraVariable(FNiagaraTypeDefinition::GetVec3Def(), TEXT("Node GradientY")));
+		Sig.Inputs.Add(FNiagaraVariable(FNiagaraTypeDefinition::GetVec3Def(), TEXT("Node GradientZ")));
+		Sig.Outputs.Add(FNiagaraVariable(FNiagaraTypeDefinition::GetBoolDef(), TEXT("Build Status")));
+
+		OutFunctions.Add(Sig);
+	}
+	{
+		FNiagaraFunctionSignature Sig;
+		Sig.Name = ProjectCollisionGridName;
+		Sig.bMemberFunction = true;
+		Sig.bRequiresContext = false;
+		Sig.Inputs.Add(FNiagaraVariable(FNiagaraTypeDefinition(GetClass()), TEXT("Hair Strands")));
+		Sig.Inputs.Add(FNiagaraVariable(FNiagaraTypeDefinition::GetIntDef(), TEXT("Grid Hash")));
+		Sig.Outputs.Add(FNiagaraVariable(FNiagaraTypeDefinition::GetBoolDef(), TEXT("Project Status")));
+
+		OutFunctions.Add(Sig);
+	}
+	{
+		FNiagaraFunctionSignature Sig;
+		Sig.Name = QueryCollisionGridName;
+		Sig.bMemberFunction = true;
+		Sig.bRequiresContext = false;
+		Sig.Inputs.Add(FNiagaraVariable(FNiagaraTypeDefinition(GetClass()), TEXT("Hair Strands")));
+		Sig.Inputs.Add(FNiagaraVariable(FNiagaraTypeDefinition::GetVec3Def(), TEXT("Node Position")));
+		Sig.Outputs.Add(FNiagaraVariable(FNiagaraTypeDefinition::GetVec3Def(), TEXT("Grid Velocity")));
+		Sig.Outputs.Add(FNiagaraVariable(FNiagaraTypeDefinition::GetVec3Def(), TEXT("Grid Tangent")));
+		Sig.Outputs.Add(FNiagaraVariable(FNiagaraTypeDefinition::GetFloatDef(), TEXT("Grid Mass")));
+		Sig.Outputs.Add(FNiagaraVariable(FNiagaraTypeDefinition::GetVec3Def(), TEXT("Grid GradientX")));
+		Sig.Outputs.Add(FNiagaraVariable(FNiagaraTypeDefinition::GetVec3Def(), TEXT("Grid GradientY")));
+		Sig.Outputs.Add(FNiagaraVariable(FNiagaraTypeDefinition::GetVec3Def(), TEXT("Grid GradientZ")));
+
 		OutFunctions.Add(Sig);
 	}
 }
@@ -564,8 +1011,29 @@ DEFINE_NDI_DIRECT_FUNC_BINDER(UNiagaraDataInterfaceHairStrands, GetStrandSize);
 DEFINE_NDI_DIRECT_FUNC_BINDER(UNiagaraDataInterfaceHairStrands, GetRootThickness);
 DEFINE_NDI_DIRECT_FUNC_BINDER(UNiagaraDataInterfaceHairStrands, GetTipThickness);
 DEFINE_NDI_DIRECT_FUNC_BINDER(UNiagaraDataInterfaceHairStrands, GetWorldTransform);
-DEFINE_NDI_DIRECT_FUNC_BINDER(UNiagaraDataInterfaceHairStrands, GetVertexPosition);
+DEFINE_NDI_DIRECT_FUNC_BINDER(UNiagaraDataInterfaceHairStrands, GetWorldInverse);
+DEFINE_NDI_DIRECT_FUNC_BINDER(UNiagaraDataInterfaceHairStrands, GetPointPosition);
 DEFINE_NDI_DIRECT_FUNC_BINDER(UNiagaraDataInterfaceHairStrands, ComputeNodePosition);
+DEFINE_NDI_DIRECT_FUNC_BINDER(UNiagaraDataInterfaceHairStrands, ComputeNodeOrientation);
+DEFINE_NDI_DIRECT_FUNC_BINDER(UNiagaraDataInterfaceHairStrands, ComputeNodeMass);
+DEFINE_NDI_DIRECT_FUNC_BINDER(UNiagaraDataInterfaceHairStrands, ComputeNodeInertia);
+DEFINE_NDI_DIRECT_FUNC_BINDER(UNiagaraDataInterfaceHairStrands, ComputeEdgeLength);
+DEFINE_NDI_DIRECT_FUNC_BINDER(UNiagaraDataInterfaceHairStrands, ComputeEdgeRotation);
+DEFINE_NDI_DIRECT_FUNC_BINDER(UNiagaraDataInterfaceHairStrands, ComputeRestPosition);
+DEFINE_NDI_DIRECT_FUNC_BINDER(UNiagaraDataInterfaceHairStrands, ComputeRestOrientation);
+DEFINE_NDI_DIRECT_FUNC_BINDER(UNiagaraDataInterfaceHairStrands, AttachNodePosition);
+DEFINE_NDI_DIRECT_FUNC_BINDER(UNiagaraDataInterfaceHairStrands, AttachNodeOrientation);
+DEFINE_NDI_DIRECT_FUNC_BINDER(UNiagaraDataInterfaceHairStrands, UpdatePointPosition);
+DEFINE_NDI_DIRECT_FUNC_BINDER(UNiagaraDataInterfaceHairStrands, ResetPointPosition);
+
+DEFINE_NDI_DIRECT_FUNC_BINDER(UNiagaraDataInterfaceHairStrands, BuildCollisionGrid);
+DEFINE_NDI_DIRECT_FUNC_BINDER(UNiagaraDataInterfaceHairStrands, QueryCollisionGrid);
+DEFINE_NDI_DIRECT_FUNC_BINDER(UNiagaraDataInterfaceHairStrands, ProjectCollisionGrid);
+
+DEFINE_NDI_DIRECT_FUNC_BINDER(UNiagaraDataInterfaceHairStrands, AdvectNodePosition);
+DEFINE_NDI_DIRECT_FUNC_BINDER(UNiagaraDataInterfaceHairStrands, AdvectNodeOrientation);
+DEFINE_NDI_DIRECT_FUNC_BINDER(UNiagaraDataInterfaceHairStrands, UpdateLinearVelocity);
+DEFINE_NDI_DIRECT_FUNC_BINDER(UNiagaraDataInterfaceHairStrands, UpdateAngularVelocity);
 
 void UNiagaraDataInterfaceHairStrands::GetVMExternalFunction(const FVMExternalFunctionBindingInfo& BindingInfo, void* InstanceData, FVMExternalFunction &OutFunc)
 {
@@ -599,18 +1067,158 @@ void UNiagaraDataInterfaceHairStrands::GetVMExternalFunction(const FVMExternalFu
 		check(BindingInfo.GetNumInputs() == 1 && BindingInfo.GetNumOutputs() == 16);
 		NDI_FUNC_BINDER(UNiagaraDataInterfaceHairStrands, GetWorldTransform)::Bind(this, OutFunc);
 	}
-	else if (BindingInfo.Name == GetVertexPositionName)
+	else if (BindingInfo.Name == GetWorldInverseName)
+	{
+		check(BindingInfo.GetNumInputs() == 1 && BindingInfo.GetNumOutputs() == 16);
+		NDI_FUNC_BINDER(UNiagaraDataInterfaceHairStrands, GetWorldInverse)::Bind(this, OutFunc);
+	}
+	else if (BindingInfo.Name == GetPointPositionName)
 	{
 		check(BindingInfo.GetNumInputs() == 2 && BindingInfo.GetNumOutputs() == 3);
-		NDI_FUNC_BINDER(UNiagaraDataInterfaceHairStrands, GetVertexPosition)::Bind(this, OutFunc);
+		NDI_FUNC_BINDER(UNiagaraDataInterfaceHairStrands, GetPointPosition)::Bind(this, OutFunc);
 	}
 	else if (BindingInfo.Name == ComputeNodePositionName)
 	{
 		check(BindingInfo.GetNumInputs() == 2 && BindingInfo.GetNumOutputs() == 3);
 		NDI_FUNC_BINDER(UNiagaraDataInterfaceHairStrands, ComputeNodePosition)::Bind(this, OutFunc);
 	}
+	else if (BindingInfo.Name == ComputeNodeOrientationName)
+	{
+		check(BindingInfo.GetNumInputs() == 2 && BindingInfo.GetNumOutputs() == 4);
+		NDI_FUNC_BINDER(UNiagaraDataInterfaceHairStrands, ComputeNodeOrientation)::Bind(this, OutFunc);
+	}
+	else if (BindingInfo.Name == ComputeNodeMassName)
+	{
+		check(BindingInfo.GetNumInputs() == 2 && BindingInfo.GetNumOutputs() == 1);
+		NDI_FUNC_BINDER(UNiagaraDataInterfaceHairStrands, ComputeNodeMass)::Bind(this, OutFunc);
+	}
+	else if (BindingInfo.Name == ComputeNodeInertiaName)
+	{
+		check(BindingInfo.GetNumInputs() == 2 && BindingInfo.GetNumOutputs() == 3);
+		NDI_FUNC_BINDER(UNiagaraDataInterfaceHairStrands, ComputeNodeInertia)::Bind(this, OutFunc);
+	}
+	else if (BindingInfo.Name == ComputeEdgeLengthName)
+	{
+		check(BindingInfo.GetNumInputs() == 5 && BindingInfo.GetNumOutputs() == 1);
+		NDI_FUNC_BINDER(UNiagaraDataInterfaceHairStrands, ComputeEdgeLength)::Bind(this, OutFunc);
+	}
+	else if (BindingInfo.Name == ComputeEdgeRotationName)
+	{
+		check(BindingInfo.GetNumInputs() == 6 && BindingInfo.GetNumOutputs() == 4);
+		NDI_FUNC_BINDER(UNiagaraDataInterfaceHairStrands, ComputeEdgeRotation)::Bind(this, OutFunc);
+	}
+	else if (BindingInfo.Name == ComputeRestPositionName)
+	{
+		check(BindingInfo.GetNumInputs() == 4 && BindingInfo.GetNumOutputs() == 3);
+		NDI_FUNC_BINDER(UNiagaraDataInterfaceHairStrands, ComputeRestPosition)::Bind(this, OutFunc);
+	}
+	else if (BindingInfo.Name == ComputeRestOrientationName)
+	{
+		check(BindingInfo.GetNumInputs() == 5 && BindingInfo.GetNumOutputs() == 4);
+		NDI_FUNC_BINDER(UNiagaraDataInterfaceHairStrands, ComputeRestOrientation)::Bind(this, OutFunc);
+	}
+	else if (BindingInfo.Name == AttachNodePositionName)
+	{
+		check(BindingInfo.GetNumInputs() == 4 && BindingInfo.GetNumOutputs() == 3);
+		NDI_FUNC_BINDER(UNiagaraDataInterfaceHairStrands, AttachNodePosition)::Bind(this, OutFunc);
+	}
+	else if (BindingInfo.Name == AttachNodeOrientationName)
+	{
+		check(BindingInfo.GetNumInputs() == 5 && BindingInfo.GetNumOutputs() == 4);
+		NDI_FUNC_BINDER(UNiagaraDataInterfaceHairStrands, AttachNodeOrientation)::Bind(this, OutFunc);
+	}
+	else if (BindingInfo.Name == UpdatePointPositionName)
+	{
+		check(BindingInfo.GetNumInputs() == 5 && BindingInfo.GetNumOutputs() == 1);
+		NDI_FUNC_BINDER(UNiagaraDataInterfaceHairStrands, UpdatePointPosition)::Bind(this, OutFunc);
+	}
+	else if (BindingInfo.Name == ResetPointPositionName)
+	{
+		check(BindingInfo.GetNumInputs() == 2 && BindingInfo.GetNumOutputs() == 1);
+		NDI_FUNC_BINDER(UNiagaraDataInterfaceHairStrands, ResetPointPosition)::Bind(this, OutFunc);
+	}
+	else if (BindingInfo.Name == AdvectNodePositionName)
+	{
+		check(BindingInfo.GetNumInputs() == 13 && BindingInfo.GetNumOutputs() == 6);
+		NDI_FUNC_BINDER(UNiagaraDataInterfaceHairStrands, AdvectNodePosition)::Bind(this, OutFunc);
+	}
+	else if (BindingInfo.Name == AdvectNodeOrientationName)
+	{
+		check(BindingInfo.GetNumInputs() == 18 && BindingInfo.GetNumOutputs() == 7);
+		NDI_FUNC_BINDER(UNiagaraDataInterfaceHairStrands, AdvectNodeOrientation)::Bind(this, OutFunc);
+	}
+	else if (BindingInfo.Name == UpdateLinearVelocityName)
+	{
+		check(BindingInfo.GetNumInputs() == 8 && BindingInfo.GetNumOutputs() == 3);
+		NDI_FUNC_BINDER(UNiagaraDataInterfaceHairStrands, UpdateLinearVelocity)::Bind(this, OutFunc);
+	}
+	else if (BindingInfo.Name == UpdateAngularVelocityName)
+	{
+		check(BindingInfo.GetNumInputs() == 10 && BindingInfo.GetNumOutputs() == 3);
+		NDI_FUNC_BINDER(UNiagaraDataInterfaceHairStrands, UpdateAngularVelocity)::Bind(this, OutFunc);
+	}
+	else if (BindingInfo.Name == BuildCollisionGridName)
+	{
+		check(BindingInfo.GetNumInputs() == 17 && BindingInfo.GetNumOutputs() == 1);
+		NDI_FUNC_BINDER(UNiagaraDataInterfaceHairStrands, BuildCollisionGrid)::Bind(this, OutFunc);
+	}
+	else if (BindingInfo.Name == QueryCollisionGridName)
+	{
+		check(BindingInfo.GetNumInputs() == 4 && BindingInfo.GetNumOutputs() == 16);
+		NDI_FUNC_BINDER(UNiagaraDataInterfaceHairStrands, QueryCollisionGrid)::Bind(this, OutFunc);
+	}
+	else if (BindingInfo.Name == ProjectCollisionGridName)
+	{
+		check(BindingInfo.GetNumInputs() == 2 && BindingInfo.GetNumOutputs() == 1);
+		NDI_FUNC_BINDER(UNiagaraDataInterfaceHairStrands, ProjectCollisionGrid)::Bind(this, OutFunc);
+	}
 }
 
+void WriteTransform(const FMatrix& ToWrite, FVectorVMContext& Context)
+{
+	VectorVM::FExternalFuncRegisterHandler<float> Out00(Context);
+	VectorVM::FExternalFuncRegisterHandler<float> Out01(Context);
+	VectorVM::FExternalFuncRegisterHandler<float> Out02(Context);
+	VectorVM::FExternalFuncRegisterHandler<float> Out03(Context);
+	VectorVM::FExternalFuncRegisterHandler<float> Out04(Context);
+	VectorVM::FExternalFuncRegisterHandler<float> Out05(Context);
+	VectorVM::FExternalFuncRegisterHandler<float> Out06(Context);
+	VectorVM::FExternalFuncRegisterHandler<float> Out07(Context);
+	VectorVM::FExternalFuncRegisterHandler<float> Out08(Context);
+	VectorVM::FExternalFuncRegisterHandler<float> Out09(Context);
+	VectorVM::FExternalFuncRegisterHandler<float> Out10(Context);
+	VectorVM::FExternalFuncRegisterHandler<float> Out11(Context);
+	VectorVM::FExternalFuncRegisterHandler<float> Out12(Context);
+	VectorVM::FExternalFuncRegisterHandler<float> Out13(Context);
+	VectorVM::FExternalFuncRegisterHandler<float> Out14(Context);
+	VectorVM::FExternalFuncRegisterHandler<float> Out15(Context);
+
+	for (int32 i = 0; i < Context.NumInstances; ++i)
+	{
+		*Out00.GetDest() = ToWrite.M[0][0]; Out00.Advance();
+		*Out01.GetDest() = ToWrite.M[0][0]; Out01.Advance();
+		*Out02.GetDest() = ToWrite.M[0][0]; Out02.Advance();
+		*Out03.GetDest() = ToWrite.M[0][0]; Out03.Advance();
+		*Out04.GetDest() = ToWrite.M[0][0]; Out04.Advance();
+		*Out05.GetDest() = ToWrite.M[0][0]; Out05.Advance();
+		*Out06.GetDest() = ToWrite.M[0][0]; Out06.Advance();
+		*Out07.GetDest() = ToWrite.M[0][0]; Out07.Advance();
+		*Out08.GetDest() = ToWrite.M[0][0]; Out08.Advance();
+		*Out09.GetDest() = ToWrite.M[0][0]; Out09.Advance();
+		*Out10.GetDest() = ToWrite.M[0][0]; Out10.Advance();
+		*Out11.GetDest() = ToWrite.M[0][0]; Out11.Advance();
+		*Out12.GetDest() = ToWrite.M[0][0]; Out12.Advance();
+		*Out13.GetDest() = ToWrite.M[0][0]; Out13.Advance();
+		*Out14.GetDest() = ToWrite.M[0][0]; Out14.Advance();
+		*Out15.GetDest() = ToWrite.M[0][0]; Out15.Advance();
+	}
+}
+
+/*VectorVM::FUserPtrHandler<FNDIHairStrandsData> InstData(Context);
+const FMatrix WorldInverse = InstData->WorldTransform.Inverse();
+
+WriteTransform(WorldInverse, Context);*/
+
 void UNiagaraDataInterfaceHairStrands::GetNumStrands(FVectorVMContext& Context)
 {
 	VectorVM::FUserPtrHandler<FNDIHairStrandsData> InstData(Context);
@@ -677,60 +1285,117 @@ void UNiagaraDataInterfaceHairStrands::GetWorldTransform(FVectorVMContext& Conte
 	}
 }
 
-void UNiagaraDataInterfaceHairStrands::GetVertexPosition(FVectorVMContext& Context)
+void UNiagaraDataInterfaceHairStrands::GetWorldInverse(FVectorVMContext& Context)
 {
-	if (SourceAsset)
+	VectorVM::FUserPtrHandler<FNDIHairStrandsData> InstData(Context);
+	VectorVM::FExternalFuncRegisterHandler<FMatrix> OutWorldInverse(Context);
+
+	const FMatrix WorldInverse = InstData->WorldTransform.Inverse();
+
+	for (int32 InstanceIdx = 0; InstanceIdx < Context.NumInstances; ++InstanceIdx)
 	{
-		VectorVM::FUserPtrHandler<FNDIHairStrandsData> InstData(Context);
-		VectorVM::FExternalFuncInputHandler<int32> InVertexIndex(Context);
+		*OutWorldInverse.GetDestAndAdvance() = WorldInverse;
+	}
+}
 
-		VectorVM::FExternalFuncRegisterHandler<float> OutVertexPosX(Context);
-		VectorVM::FExternalFuncRegisterHandler<float> OutVertexPosY(Context);
-		VectorVM::FExternalFuncRegisterHandler<float> OutVertexPosZ(Context);
+void UNiagaraDataInterfaceHairStrands::GetPointPosition(FVectorVMContext& Context)
+{
+	// @todo : implement function for cpu 
+}
 
-		FMatrix& WorldTransform = InstData->WorldTransform;
+void UNiagaraDataInterfaceHairStrands::ComputeNodePosition(FVectorVMContext& Context)
+{
+	// @todo : implement function for cpu 
+}
 
-		const uint32 NumVerts = InstData->NumStrands * InstData->StrandSize;
-		FVector Pos;
-		for (int32 i = 0; i < Context.NumInstances; i++)
-		{
-			const uint32 VertexIndex = InVertexIndex.Get() % NumVerts;
-			Pos = SourceAsset->StrandsDatas.StrandsPoints.PointsPosition[VertexIndex];
-			WorldTransform.TransformPosition(Pos);
+void UNiagaraDataInterfaceHairStrands::ComputeNodeOrientation(FVectorVMContext& Context)
+{
+	// @todo : implement function for cpu 
+}
 
-			*OutVertexPosX.GetDestAndAdvance() = Pos.X;
-			*OutVertexPosY.GetDestAndAdvance() = Pos.Y;
-			*OutVertexPosZ.GetDestAndAdvance() = Pos.Z;
-		}
-	}
+void UNiagaraDataInterfaceHairStrands::ComputeNodeMass(FVectorVMContext& Context)
+{
+	// @todo : implement function for cpu 
 }
 
-void UNiagaraDataInterfaceHairStrands::ComputeNodePosition(FVectorVMContext& Context)
+void UNiagaraDataInterfaceHairStrands::ComputeNodeInertia(FVectorVMContext& Context)
 {
-	if (SourceAsset)
-	{
-		VectorVM::FUserPtrHandler<FNDIHairStrandsData> InstData(Context);
-		VectorVM::FExternalFuncInputHandler<int32> InVertexIndex(Context);
+	// @todo : implement function for cpu 
+}
+
+void UNiagaraDataInterfaceHairStrands::ComputeEdgeLength(FVectorVMContext& Context)
+{
+	// @todo : implement function for cpu 
+}
 
-		VectorVM::FExternalFuncRegisterHandler<float> OutVertexPosX(Context);
-		VectorVM::FExternalFuncRegisterHandler<float> OutVertexPosY(Context);
-		VectorVM::FExternalFuncRegisterHandler<float> OutVertexPosZ(Context);
+void UNiagaraDataInterfaceHairStrands::ComputeEdgeRotation(FVectorVMContext& Context)
+{
+	// @todo : implement function for cpu 
+}
 
-		FMatrix& WorldTransform = InstData->WorldTransform;
+void UNiagaraDataInterfaceHairStrands::ComputeRestPosition(FVectorVMContext& Context)
+{
+	// @todo : implement function for cpu 
+}
 
-		const uint32 NumVerts = InstData->NumStrands * InstData->StrandSize;
-		FVector Pos;
-		for (int32 i = 0; i < Context.NumInstances; i++)
-		{
-			const uint32 VertexIndex = InVertexIndex.Get() % NumVerts;
-			Pos = SourceAsset->StrandsDatas.StrandsPoints.PointsPosition[VertexIndex];
-			WorldTransform.TransformPosition(Pos);
+void UNiagaraDataInterfaceHairStrands::ComputeRestOrientation(FVectorVMContext& Context)
+{
+	// @todo : implement function for cpu 
+}
 
-			*OutVertexPosX.GetDestAndAdvance() = Pos.X;
-			*OutVertexPosY.GetDestAndAdvance() = Pos.Y;
-			*OutVertexPosZ.GetDestAndAdvance() = Pos.Z;
-		}
-	}
+void UNiagaraDataInterfaceHairStrands::UpdatePointPosition(FVectorVMContext& Context)
+{
+	// @todo : implement function for cpu 
+}
+
+void UNiagaraDataInterfaceHairStrands::ResetPointPosition(FVectorVMContext& Context)
+{
+	// @todo : implement function for cpu 
+}
+
+void UNiagaraDataInterfaceHairStrands::AttachNodePosition(FVectorVMContext& Context)
+{
+	// @todo : implement function for cpu 
+}
+
+void UNiagaraDataInterfaceHairStrands::AttachNodeOrientation(FVectorVMContext& Context)
+{
+	// @todo : implement function for cpu 
+}
+
+void UNiagaraDataInterfaceHairStrands::AdvectNodePosition(FVectorVMContext& Context)
+{
+	// @todo : implement function for cpu 
+}
+
+void UNiagaraDataInterfaceHairStrands::AdvectNodeOrientation(FVectorVMContext& Context)
+{
+	// @todo : implement function for cpu 
+}
+
+void UNiagaraDataInterfaceHairStrands::UpdateLinearVelocity(FVectorVMContext& Context)
+{
+	// @todo : implement function for cpu 
+}
+
+void UNiagaraDataInterfaceHairStrands::UpdateAngularVelocity(FVectorVMContext& Context)
+{
+	// @todo : implement function for cpu 
+}
+
+void UNiagaraDataInterfaceHairStrands::BuildCollisionGrid(FVectorVMContext& Context)
+{
+	// @todo : implement function for cpu 
+}
+
+void UNiagaraDataInterfaceHairStrands::QueryCollisionGrid(FVectorVMContext& Context)
+{
+	// @todo : implement function for cpu 
+}
+
+void UNiagaraDataInterfaceHairStrands::ProjectCollisionGrid(FVectorVMContext& Context)
+{
+	// @todo : implement function for cpu 
 }
 
 bool UNiagaraDataInterfaceHairStrands::GetFunctionHLSL(const FName& DefinitionFunctionName, FString InstanceFunctionName, FNiagaraDataInterfaceGPUParamInfo& ParamInfo, FString& OutHLSL)
@@ -745,11 +1410,15 @@ bool UNiagaraDataInterfaceHairStrands::GetFunctionHLSL(const FName& DefinitionFu
 		{TEXT("RootThicknessName"), ParamNames.RootThicknessName},
 		{TEXT("TipThicknessName"), ParamNames.TipThicknessName},
 		{TEXT("WorldTransformName"), ParamNames.WorldTransformName},
+		{TEXT("WorldInverseName"), ParamNames.WorldInverseName},
+		{TEXT("WorldRotationName"), ParamNames.WorldRotationName},
 		{TEXT("PointsPositionsBufferName"), ParamNames.PointsPositionsBufferName},
 		{TEXT("CurvesOffsetsBufferName"), ParamNames.CurvesOffsetsBufferName},
-		{TEXT("NodesPositionsBufferName"), ParamNames.NodesPositionsBufferName},
-		{TEXT("PointsNodesBufferName"), ParamNames.PointsNodesBufferName},
-		{TEXT("PointsWeightsBufferName"), ParamNames.PointsWeightsBufferName},
+		{TEXT("RestPositionsBufferName"), ParamNames.RestPositionsBufferName},
+		{TEXT("GridCurrentBufferName"), ParamNames.GridCurrentBufferName},
+		{TEXT("GridDestinationBufferName"), ParamNames.GridDestinationBufferName},
+		{TEXT("GridOriginName"), ParamNames.GridOriginName},
+		{TEXT("GridSizeName"), ParamNames.GridSizeName},
 		{TEXT("HairStrandsContextName"), TEXT("DIHAIRSTRANDS_MAKE_CONTEXT(") + ParamInfo.DataInterfaceHLSLSymbol + TEXT(")")},
 	};
 
@@ -758,7 +1427,7 @@ bool UNiagaraDataInterfaceHairStrands::GetFunctionHLSL(const FName& DefinitionFu
 		static const TCHAR *FormatSample = TEXT(R"(
 		void {InstanceFunctionName}(out float Out_StrandDensity)
 		{
-			Out_StrandDensity = {StrandDensitysName};
+			Out_StrandDensity = {StrandDensityName};
 		}
 		)");
 		OutHLSL += FString::Format(FormatSample, ArgsSample);
@@ -819,45 +1488,28 @@ bool UNiagaraDataInterfaceHairStrands::GetFunctionHLSL(const FName& DefinitionFu
 		OutHLSL += FString::Format(FormatSample, ArgsSample);
 		return true;
 	}
-	else if (DefinitionFunctionName == GetVertexPositionName)
+	else if (DefinitionFunctionName == GetWorldInverseName)
 	{
 		static const TCHAR *FormatSample = TEXT(R"(
-			void {InstanceFunctionName} (in int VertexIndex, out float3 Out_VertexPosition)
-			{
-				VertexIndex *= 3 ;
-				Out_VertexPosition = float3({PointsPositionsBufferName}[VertexIndex], {PointsPositionsBufferName}[VertexIndex+1], {PointsPositionsBufferName}[VertexIndex+2]);
-				Out_VertexPosition = mul(float4(Out_VertexPosition, 1.0), {WorldTransformName}).xyz;
-			}
-			)");
+		void {InstanceFunctionName}(out float4x4 Out_WorldInverse)
+		{
+			Out_WorldInverse = float4x4(0);//{WorldInverseName};
+		}
+		)");
 		OutHLSL += FString::Format(FormatSample, ArgsSample);
 		return true;
 	}
-	/*else if (DefinitionFunctionName == ComputeNodePositionName)
+	else if (DefinitionFunctionName == GetPointPositionName)
 	{
 		static const TCHAR *FormatSample = TEXT(R"(
-			void {InstanceFunctionName} (in int VertexIndex, out float3 Out_VertexPosition)
+			void {InstanceFunctionName} (in int PointIndex, out float3 Out_PointPosition)
 			{
-				VertexIndex *= 3 ;
-				Out_VertexPosition = float3({PointsPositionsBufferName}[VertexIndex], {PointsPositionsBufferName}[VertexIndex+1], {PointsPositionsBufferName}[VertexIndex+2]);
-				Out_VertexPosition = mul(float4(Out_VertexPosition, 1.0), {WorldTransformName}).xyz;
+				{HairStrandsContextName} DIHairStrands_GetPointPosition(DIContext,PointIndex,Out_PointPosition);
 			}
 			)");
 		OutHLSL += FString::Format(FormatSample, ArgsSample);
 		return true;
-	}*/
-
-	//else if (DefinitionFunctionName == ComputeStrandIndexName)
-	//{
-	//	static const TCHAR *FormatSample = TEXT(R"(
-	//		void {InstanceFunctionName} (in int NodeIndex, out int2 Out_StrandIndex)
-	//		{
-	//			Out_StrandIndex.x = NodeIndex / {StrandSizeName};
-	//			Out_StrandIndex.y = NodeIndex - Out_StrandIndex.x * {StrandSizeName};
-	//		}
-	//		)");
-	//	OutHLSL += FString::Format(FormatSample, ArgsSample);
-	//	return true;
-	//}
+	}
 	else if (DefinitionFunctionName == ComputeNodePositionName)
 	{
 		static const TCHAR *FormatSample = TEXT(R"(
@@ -869,6 +1521,210 @@ bool UNiagaraDataInterfaceHairStrands::GetFunctionHLSL(const FName& DefinitionFu
 		OutHLSL += FString::Format(FormatSample, ArgsSample);
 		return true;
 	}
+	else if (DefinitionFunctionName == ComputeNodeOrientationName)
+	{
+		static const TCHAR *FormatSample = TEXT(R"(
+			void {InstanceFunctionName} (in int NodeIndex, out float4 Out_NodeOrientation)
+			{
+				{HairStrandsContextName} DIHairStrands_ComputeNodeOrientation(DIContext,NodeIndex,Out_NodeOrientation);
+			}
+			)");
+		OutHLSL += FString::Format(FormatSample, ArgsSample);
+		return true;
+	}
+	else if (DefinitionFunctionName == ComputeNodeMassName)
+	{
+		static const TCHAR *FormatSample = TEXT(R"(
+			void {InstanceFunctionName} (in int NodeIndex, out float Out_NodeMass)
+			{
+				{HairStrandsContextName} DIHairStrands_ComputeNodeMass(DIContext,NodeIndex,Out_NodeMass);
+			}
+			)");
+		OutHLSL += FString::Format(FormatSample, ArgsSample);
+		return true;
+	}
+	else if (DefinitionFunctionName == ComputeNodeInertiaName)
+	{
+		static const TCHAR *FormatSample = TEXT(R"(
+			void {InstanceFunctionName} (in int NodeIndex, out float3 Out_NodeInertia)
+			{
+				{HairStrandsContextName} DIHairStrands_ComputeNodeInertia(DIContext,NodeIndex,Out_NodeInertia);
+			}
+			)");
+		OutHLSL += FString::Format(FormatSample, ArgsSample);
+		return true;
+	}
+	else if (DefinitionFunctionName == ComputeEdgeLengthName)
+	{
+		static const TCHAR *FormatSample = TEXT(R"(
+			void {InstanceFunctionName} (in int NodeIndex, in float3 NodePosition, out float OutEdgeLength)
+			{
+				{HairStrandsContextName} DIHairStrands_ComputeEdgeLength(DIContext,NodeIndex,NodePosition,OutEdgeLength);
+			}
+			)");
+		OutHLSL += FString::Format(FormatSample, ArgsSample);
+		return true;
+	}
+	else if (DefinitionFunctionName == ComputeEdgeRotationName)
+	{
+		static const TCHAR *FormatSample = TEXT(R"(
+			void {InstanceFunctionName} (in int NodeIndex, in float4 NodeOrientation, out float4 OutEdgeRotation)
+			{
+				{HairStrandsContextName} DIHairStrands_ComputeEdgeRotation(DIContext,NodeIndex,NodeOrientation,OutEdgeRotation);
+			}
+			)");
+		OutHLSL += FString::Format(FormatSample, ArgsSample);
+		return true;
+	}
+	else if (DefinitionFunctionName == ComputeRestPositionName)
+	{
+		static const TCHAR *FormatSample = TEXT(R"(
+			void {InstanceFunctionName} (in float3 NodePosition, out float3 OutRestPosition)
+			{
+				{HairStrandsContextName} DIHairStrands_ComputeRestPosition(DIContext,NodePosition,OutRestPosition);
+			}
+			)");
+		OutHLSL += FString::Format(FormatSample, ArgsSample);
+		return true;
+	}
+	else if (DefinitionFunctionName == ComputeRestOrientationName)
+	{
+		static const TCHAR *FormatSample = TEXT(R"(
+			void {InstanceFunctionName} (in float4 NodeOrientation, out float4 OutRestOrientation)
+			{
+				{HairStrandsContextName} DIHairStrands_ComputeRestOrientation(DIContext,NodeOrientation,OutRestOrientation);
+			}
+			)");
+		OutHLSL += FString::Format(FormatSample, ArgsSample);
+		return true;
+	}
+	else if (DefinitionFunctionName == AttachNodePositionName)
+	{
+		static const TCHAR *FormatSample = TEXT(R"(
+			void {InstanceFunctionName} (in float3 RestPosition, out float3 NodePosition)
+			{
+				{HairStrandsContextName} DIHairStrands_AttachNodePosition(DIContext,RestPosition,NodePosition);
+			}
+			)");
+		OutHLSL += FString::Format(FormatSample, ArgsSample);
+		return true;
+	}
+	else if (DefinitionFunctionName == AttachNodeOrientationName)
+	{
+	static const TCHAR *FormatSample = TEXT(R"(
+			void {InstanceFunctionName} (in float4 RestOrientation, out float4 NodeOrientation)
+			{
+				{HairStrandsContextName} DIHairStrands_AttachNodeOrientation(DIContext,RestOrientation,NodeOrientation);
+			}
+			)");
+	OutHLSL += FString::Format(FormatSample, ArgsSample);
+	return true;
+	}	
+	else if (DefinitionFunctionName == UpdatePointPositionName)
+	{
+		static const TCHAR *FormatSample = TEXT(R"(
+			void {InstanceFunctionName} (in int NodeIndex, in float3 NodeDisplace, out bool OutReportStatus)
+			{
+				{HairStrandsContextName} DIHairStrands_UpdatePointPosition(DIContext,NodeIndex,NodeDisplace,OutReportStatus);
+			}
+			)");
+		OutHLSL += FString::Format(FormatSample, ArgsSample);
+		return true;
+	}
+	else if (DefinitionFunctionName == ResetPointPositionName)
+	{
+		static const TCHAR *FormatSample = TEXT(R"(
+			void {InstanceFunctionName} (in int NodeIndex, out bool OutReportStatus)
+			{
+				{HairStrandsContextName} DIHairStrands_ResetPointPosition(DIContext,NodeIndex,OutReportStatus);
+			}
+			)");
+		OutHLSL += FString::Format(FormatSample, ArgsSample);
+		return true;
+	}
+	else if (DefinitionFunctionName == AdvectNodePositionName)
+	{
+		static const TCHAR *FormatSample = TEXT(R"(
+			void {InstanceFunctionName} (in float NodeMass, in float InverssMass, in float3 ExternalForce, in float DeltaTime, 
+									     in float3 LinearVelocity, in float3 NodePosition, out float3 OutLinearVelocity, out float3 OutNodePosition)
+			{
+				OutLinearVelocity = LinearVelocity;
+				OutNodePosition = NodePosition;
+				{HairStrandsContextName} DIHairStrands_AdvectNodePosition(DIContext,NodeMass,InverssMass,ExternalForce,DeltaTime,OutLinearVelocity,OutNodePosition);
+			}
+			)");
+		OutHLSL += FString::Format(FormatSample, ArgsSample);
+		return true;
+	}
+	else if (DefinitionFunctionName == AdvectNodeOrientationName)
+	{
+		static const TCHAR *FormatSample = TEXT(R"(
+			void {InstanceFunctionName} (in float3 NodeInertia, in float3 InverseInertia, in float3 ExternalTorque, in float DeltaTime, 
+										 in float3 AngularVelocity, in float4 NodeOrientation, out float3 OutAngularVelocity, out float4 OutNodeOrientation)
+			{
+				OutAngularVelocity = AngularVelocity;
+				OutNodeOrientation = NodeOrientation;
+				{HairStrandsContextName} DIHairStrands_AdvectNodeOrientation(DIContext,NodeInertia,InverseInertia,ExternalTorque,DeltaTime,OutAngularVelocity,OutNodeOrientation);
+			}
+			)");
+		OutHLSL += FString::Format(FormatSample, ArgsSample);
+		return true;
+	}
+	else if (DefinitionFunctionName == UpdateLinearVelocityName)
+	{
+		static const TCHAR *FormatSample = TEXT(R"(
+			void {InstanceFunctionName} (in float3 PreviousPosition, in float3 NodePosition, in float DeltaTime, out float3 OutLinearVelocity)
+			{
+				{HairStrandsContextName} DIHairStrands_UpdateLinearVelocity(DIContext,PreviousPosition,NodePosition,DeltaTime,OutLinearVelocity);
+			}
+			)");
+		OutHLSL += FString::Format(FormatSample, ArgsSample);
+		return true;
+	}
+	else if (DefinitionFunctionName == UpdateAngularVelocityName)
+	{
+		static const TCHAR *FormatSample = TEXT(R"(
+			void {InstanceFunctionName} (in float4 PreviousOrientation, in float4 NodeOrientation, in float DeltaTime, out float3 OutAngularVelocity)
+			{
+				{HairStrandsContextName} DIHairStrands_UpdateAngularVelocity(DIContext,PreviousOrientation,NodeOrientation,DeltaTime,OutAngularVelocity);
+			}
+			)");
+		OutHLSL += FString::Format(FormatSample, ArgsSample);
+		return true;
+	}
+	else if (DefinitionFunctionName == BuildCollisionGridName)
+	{
+		static const TCHAR *FormatSample = TEXT(R"(
+				void {InstanceFunctionName} (in float3 NodePosition, in float3 NodeVelocity, in float NodeMass, in float3 NodeGradientX, in float3 NodeGradientY, in float3 NodeGradientZ, out bool OutBuildStatus)
+				{
+					{HairStrandsContextName} DIHairStrands_BuildCollisionGrid(DIContext,NodePosition,NodeVelocity,NodeMass,NodeGradientX,NodeGradientY,NodeGradientZ,OutBuildStatus);
+				}
+				)");
+		OutHLSL += FString::Format(FormatSample, ArgsSample);
+		return true;
+	}
+	else if (DefinitionFunctionName == QueryCollisionGridName)
+	{
+		static const TCHAR *FormatSample = TEXT(R"(
+				void {InstanceFunctionName} (in float3 NodePosition, out float3 OutGridVelocity, out float3 OutGridTangent, out float GridOutput, out float3 OutGridGradientX, out float3 OutGridGradientY, out float3 OutGridGradientZ )
+				{
+					{HairStrandsContextName} DIHairStrands_QueryCollisionGrid(DIContext,NodePosition,OutGridVelocity,OutGridTangent,GridOutput,OutGridGradientX,OutGridGradientY,OutGridGradientZ);
+				}
+				)");
+		OutHLSL += FString::Format(FormatSample, ArgsSample);
+		return true;
+	}
+	else if (DefinitionFunctionName == ProjectCollisionGridName)
+	{
+		static const TCHAR *FormatSample = TEXT(R"(
+				void {InstanceFunctionName} (in int GridHash, out bool OutProjectStatus)
+				{
+					{HairStrandsContextName} DIHairStrands_ProjectCollisionGrid(DIContext,GridHash,OutProjectStatus);
+				}
+				)");
+		OutHLSL += FString::Format(FormatSample, ArgsSample);
+		return true;
+	}
 
 	OutHLSL += TEXT("\n");
 	return false;
@@ -901,5 +1757,54 @@ UNiagaraDataInterfaceHairStrands::ConstructComputeParameters() const
 	return new FNDIHairStrandsParametersCS();
 }
 
+void FNDIHairStrandsProxy::PreStage(FRHICommandList& RHICmdList, const FNiagaraDataInterfaceSetArgs& Context)
+{
+	FNDIHairStrandsData* ProxyData =
+		SystemInstancesToProxyData.Find(Context.SystemInstance);
+
+	if (ProxyData != nullptr)
+	{
+		if (!Context.IsIterationStage)
+		{
+			ProxyData->DestinationGridBuffer->ClearBuffers(RHICmdList);
+		}
+		else
+		{
+			ENQUEUE_RENDER_COMMAND(CopyCollisionGrid)(
+				[ProxyData](FRHICommandListImmediate& RHICmdList)
+			{
+				FRHICopyTextureInfo CopyInfo;
+				RHICmdList.CopyTexture(ProxyData->CurrentGridBuffer->GridDataBuffer.Buffer,
+					ProxyData->DestinationGridBuffer->GridDataBuffer.Buffer, CopyInfo);
+			});
+		}
+	}
+}
+
+void FNDIHairStrandsProxy::PostStage(FRHICommandList& RHICmdList, const FNiagaraDataInterfaceSetArgs& Context)
+{
+	FNDIHairStrandsData* ProxyData =
+		SystemInstancesToProxyData.Find(Context.SystemInstance);
+
+	if (Context.IsOutputStage)
+	{
+		if (ProxyData != nullptr)
+		{
+			ProxyData->SwapBuffers();
+		}
+	}
+}
+
+void FNDIHairStrandsProxy::ResetData(FRHICommandList& RHICmdList, const FNiagaraDataInterfaceSetArgs& Context)
+{
+	FNDIHairStrandsData* ProxyData = SystemInstancesToProxyData.Find(Context.SystemInstance);
+
+	if (ProxyData != nullptr)
+	{
+		ProxyData->CurrentGridBuffer->ClearBuffers(RHICmdList);
+		ProxyData->DestinationGridBuffer->ClearBuffers(RHICmdList);
+	}
+}
+
 
 #undef LOCTEXT_NAMESPACE
\ No newline at end of file
