Copying //UE4/Dev-Physics to Dev-Main (//UE4/Dev-Main) [at] 9098278

diff --git a/Engine/Plugins/Experimental/HairStrands/Source/HairStrandsNiagara/Private/NiagaraDataInterfacePhysicsAsset.cpp b/Engine/Plugins/Experimental/HairStrands/Source/HairStrandsNiagara/Private/NiagaraDataInterfacePhysicsAsset.cpp
new file mode 100644
--- /dev/null
+++ b/Engine/Plugins/Experimental/HairStrands/Source/HairStrandsNiagara/Private/NiagaraDataInterfacePhysicsAsset.cpp
@@ -0,0 +1,792 @@
+// Copyright 1998-2019 Epic Games, Inc. All Rights Reserved.
+
+#include "NiagaraDataInterfacePhysicsAsset.h"
+#include "Animation/SkeletalMeshActor.h"
+#include "Components/SkeletalMeshComponent.h"
+#include "SkeletalMeshTypes.h"
+#include "AnimationRuntime.h"
+#include "NiagaraShader.h"
+#include "NiagaraComponent.h"
+#include "NiagaraRenderer.h"
+#include "NiagaraSystemInstance.h"
+#include "NiagaraEmitterInstanceBatcher.h"
+#include "ShaderParameterUtils.h"
+
+#define LOCTEXT_NAMESPACE "NiagaraDataInterfacePhysicsAsset"
+DEFINE_LOG_CATEGORY_STATIC(LogPhysicsAsset, Log, All);
+
+//------------------------------------------------------------------------------------------------------------
+
+static const FName GetNumBoxesName(TEXT("GetNumBoxes"));
+static const FName GetNumSpheresName(TEXT("GetNumSpheres"));
+static const FName GetNumCapsulesName(TEXT("GetNumCapsules"));
+
+//------------------------------------------------------------------------------------------------------------
+
+static const FName GetClosestPointName(TEXT("GetClosestPoint"));
+
+//------------------------------------------------------------------------------------------------------------
+
+const FString UNiagaraDataInterfacePhysicsAsset::ElementOffsetsName(TEXT("ElementOffsets_"));
+
+const FString UNiagaraDataInterfacePhysicsAsset::CurrentTransformBufferName(TEXT("CurrentTransformBuffer_"));
+const FString UNiagaraDataInterfacePhysicsAsset::PreviousTransformBufferName(TEXT("PreviousTransformBuffer_"));
+const FString UNiagaraDataInterfacePhysicsAsset::InverseTransformBufferName(TEXT("InverseTransformBuffer_"));
+const FString UNiagaraDataInterfacePhysicsAsset::ElementExtentBufferName(TEXT("ElementExtentBuffer_"));
+
+//------------------------------------------------------------------------------------------------------------
+
+struct FNDIPhysicsAssetParametersName
+{
+	FNDIPhysicsAssetParametersName(const FString& Suffix)
+	{
+		ElementOffsetsName = UNiagaraDataInterfacePhysicsAsset::ElementOffsetsName + Suffix;
+
+		CurrentTransformBufferName = UNiagaraDataInterfacePhysicsAsset::CurrentTransformBufferName + Suffix;
+		PreviousTransformBufferName = UNiagaraDataInterfacePhysicsAsset::PreviousTransformBufferName + Suffix;
+		InverseTransformBufferName = UNiagaraDataInterfacePhysicsAsset::InverseTransformBufferName + Suffix;
+		ElementExtentBufferName = UNiagaraDataInterfacePhysicsAsset::ElementExtentBufferName + Suffix;
+	}
+
+	FString ElementOffsetsName;
+	
+	FString CurrentTransformBufferName;
+	FString PreviousTransformBufferName;
+	FString InverseTransformBufferName;
+	FString ElementExtentBufferName;
+};
+
+//------------------------------------------------------------------------------------------------------------
+
+template<typename BufferType, typename DataType, int ElementSize, EPixelFormat PixelFormat, bool InitBuffer>
+void CreateInternalBuffer(const uint32 ElementCount, const TArray<DataType>& InputData, FRWBuffer& OutputBuffer)
+{
+	if (ElementCount > 0)
+	{
+		const uint32 BufferCount = ElementCount * ElementSize;
+		const uint32 BufferBytes = sizeof(BufferType)*BufferCount;
+
+		if (InitBuffer)
+		{
+			OutputBuffer.Initialize(sizeof(BufferType), BufferCount, PixelFormat, BUF_Static);
+		}
+		void* OutputData = RHILockVertexBuffer(OutputBuffer.Buffer, 0, BufferBytes, RLM_WriteOnly);
+
+		FMemory::Memcpy(OutputData, InputData.GetData(), BufferBytes);
+		RHIUnlockVertexBuffer(OutputBuffer.Buffer);
+	}
+}
+
+void FillCurrentTransforms(const FTransform& ElementTransform, uint32& ElementCount,
+	TArray<FVector4>& OutCurrentTransform, TArray<FVector4>& OutInverseTransform)
+{
+	const uint32 ElementOffset = 3 * ElementCount;
+	const FMatrix ElementMatrix = ElementTransform.ToMatrixWithScale();
+	const FMatrix ElementInverse = ElementMatrix.Inverse();
+
+	ElementMatrix.To3x4MatrixTranspose(&OutCurrentTransform[ElementOffset].X);
+	ElementInverse.To3x4MatrixTranspose(&OutInverseTransform[ElementOffset].X);
+	++ElementCount;
+}
+
+void CreateInternalArrays(const TWeakObjectPtr<UPhysicsAsset> PhysicsAsset, const TWeakObjectPtr<USkeletalMeshComponent> SkeletalMesh,
+	FNDIPhysicsAssetArrays* OutAssetArrays, const FTransform& WorldTransform)
+{
+	if (OutAssetArrays != nullptr)
+	{
+		OutAssetArrays->ElementOffsets.BoxOffset = 0;
+		OutAssetArrays->ElementOffsets.SphereOffset = 0;
+		OutAssetArrays->ElementOffsets.CapsuleOffset = 0;
+		OutAssetArrays->ElementOffsets.NumElements = 0;
+		if (PhysicsAsset != nullptr && OutAssetArrays != nullptr)
+		{
+			const FReferenceSkeleton* RefSkeleton = &PhysicsAsset->GetPreviewMesh()->RefSkeleton;
+			if (RefSkeleton != nullptr)
+			{
+				TArray<FTransform> BoneTransforms;
+				if (SkeletalMesh != nullptr)
+				{
+					BoneTransforms = SkeletalMesh->GetComponentSpaceTransforms();
+				}
+				else
+				{
+					FAnimationRuntime::FillUpComponentSpaceTransforms(*RefSkeleton, RefSkeleton->GetRefBonePose(), BoneTransforms);
+				}
+				const bool bHasMasterPoseComponent = (SkeletalMesh != nullptr) && SkeletalMesh->MasterPoseComponent.IsValid();
+
+				uint32 NumBoxes = 0;
+				uint32 NumSpheres = 0;
+				uint32 NumCapsules = 0;
+				for (const UBodySetup* BodySetup : PhysicsAsset->SkeletalBodySetups)
+				{
+					const FName BoneName = BodySetup->BoneName;
+					const int32 BoneIndex = RefSkeleton->FindBoneIndex(BoneName);
+					if (BoneIndex != INDEX_NONE)
+					{
+						NumBoxes += BodySetup->AggGeom.BoxElems.Num();
+						NumSpheres += BodySetup->AggGeom.SphereElems.Num();
+						NumCapsules += BodySetup->AggGeom.SphylElems.Num();
+					}
+				}
+
+				OutAssetArrays->ElementOffsets.BoxOffset = 0;
+				OutAssetArrays->ElementOffsets.SphereOffset = OutAssetArrays->ElementOffsets.BoxOffset + NumBoxes;
+				OutAssetArrays->ElementOffsets.CapsuleOffset = OutAssetArrays->ElementOffsets.SphereOffset + NumSpheres;
+				OutAssetArrays->ElementOffsets.NumElements = OutAssetArrays->ElementOffsets.CapsuleOffset + NumCapsules;
+
+				const uint32 NumTransforms = OutAssetArrays->ElementOffsets.NumElements * 3;
+				const uint32 NumExtents = OutAssetArrays->ElementOffsets.NumElements;
+
+				OutAssetArrays->CurrentTransform.SetNum(NumTransforms);
+				OutAssetArrays->InverseTransform.SetNum(NumTransforms);
+				OutAssetArrays->PreviousTransform.SetNum(NumTransforms);
+				OutAssetArrays->ElementExtent.SetNum(NumExtents);
+
+				uint32 ElementCount = 0;
+				for (const UBodySetup* BodySetup : PhysicsAsset->SkeletalBodySetups)
+				{
+					const FName BoneName = BodySetup->BoneName;
+					const int32 BoneIndex = RefSkeleton->FindBoneIndex(BoneName);
+					if (BoneIndex != INDEX_NONE)
+					{
+						const FTransform BoneTransform = bHasMasterPoseComponent ? SkeletalMesh->GetBoneTransform(BoneIndex) : BoneTransforms[BoneIndex] * WorldTransform;
+
+						for (const FKBoxElem& BoxElem : BodySetup->AggGeom.BoxElems)
+						{
+							const FTransform ElementTransform = FTransform(BoxElem.Rotation, BoxElem.Center) * BoneTransform;
+							OutAssetArrays->ElementExtent[ElementCount] = FVector4(BoxElem.X, BoxElem.Y, BoxElem.Z, 0);
+							FillCurrentTransforms(ElementTransform, ElementCount, OutAssetArrays->CurrentTransform, OutAssetArrays->InverseTransform);
+						}
+
+						for (const FKSphereElem& SphereElem : BodySetup->AggGeom.SphereElems)
+						{
+							const FTransform ElementTransform = FTransform(SphereElem.Center) * BoneTransform;
+							OutAssetArrays->ElementExtent[ElementCount] = FVector4(SphereElem.Radius, 0, 0, 0);
+							FillCurrentTransforms(ElementTransform, ElementCount, OutAssetArrays->CurrentTransform, OutAssetArrays->InverseTransform);
+						}
+
+						for (const FKSphylElem& CapsuleElem : BodySetup->AggGeom.SphylElems)
+						{
+							const FTransform ElementTransform = FTransform(CapsuleElem.Rotation, CapsuleElem.Center) * BoneTransform;
+							OutAssetArrays->ElementExtent[ElementCount] = FVector4(CapsuleElem.Radius, CapsuleElem.Length, 0, 0);
+							FillCurrentTransforms(ElementTransform, ElementCount, OutAssetArrays->CurrentTransform, OutAssetArrays->InverseTransform);
+						}
+					}
+				}
+				OutAssetArrays->PreviousTransform = OutAssetArrays->CurrentTransform;
+			}
+		}
+	}
+}
+
+void UpdateInternalArrays(const TWeakObjectPtr<UPhysicsAsset> PhysicsAsset, const TWeakObjectPtr<USkeletalMeshComponent> SkeletalMesh,
+	FNDIPhysicsAssetArrays* OutAssetArrays, const FTransform& WorldTransform)
+{
+	if (PhysicsAsset != nullptr && OutAssetArrays != nullptr)
+	{
+		const FReferenceSkeleton* RefSkeleton = &PhysicsAsset->GetPreviewMesh()->RefSkeleton;
+		if (RefSkeleton != nullptr)
+		{
+			TArray<FTransform> BoneTransforms;
+			if (SkeletalMesh != nullptr)
+			{
+				BoneTransforms = SkeletalMesh->GetComponentSpaceTransforms();
+			}
+			else
+			{
+				FAnimationRuntime::FillUpComponentSpaceTransforms(*RefSkeleton, RefSkeleton->GetRefBonePose(), BoneTransforms);
+			}
+			const bool bHasMasterPoseComponent = (SkeletalMesh != nullptr) && SkeletalMesh->MasterPoseComponent.IsValid();
+
+			OutAssetArrays->PreviousTransform = OutAssetArrays->CurrentTransform;
+
+			uint32 ElementCount = 0;
+			for (const UBodySetup* BodySetup : PhysicsAsset->SkeletalBodySetups)
+			{
+				const FName BoneName = BodySetup->BoneName;
+				const int32 BoneIndex = RefSkeleton->FindBoneIndex(BoneName);
+				if (BoneIndex != INDEX_NONE)
+				{
+					const FTransform BoneTransform = bHasMasterPoseComponent ? SkeletalMesh->GetBoneTransform(BoneIndex) : BoneTransforms[BoneIndex] * WorldTransform;
+
+					for (const FKBoxElem& BoxElem : BodySetup->AggGeom.BoxElems)
+					{
+						const FTransform ElementTransform = FTransform(BoxElem.Rotation, BoxElem.Center) * BoneTransform;
+						FillCurrentTransforms(ElementTransform, ElementCount, OutAssetArrays->CurrentTransform, OutAssetArrays->InverseTransform);
+					}
+
+					for (const FKSphereElem& SphereElem : BodySetup->AggGeom.SphereElems)
+					{
+						const FTransform ElementTransform = FTransform(SphereElem.Center) * BoneTransform;
+						FillCurrentTransforms(ElementTransform, ElementCount, OutAssetArrays->CurrentTransform, OutAssetArrays->InverseTransform);
+					}
+
+					for (const FKSphylElem& CapsuleElem : BodySetup->AggGeom.SphylElems)
+					{
+						const FTransform ElementTransform = FTransform(CapsuleElem.Rotation, CapsuleElem.Center) * BoneTransform;
+						FillCurrentTransforms(ElementTransform, ElementCount, OutAssetArrays->CurrentTransform, OutAssetArrays->InverseTransform);
+					}
+				}
+			}
+		}
+	}
+}
+//------------------------------------------------------------------------------------------------------------
+
+
+bool FNDIPhysicsAssetBuffer::IsValid() const
+{
+	return (PhysicsAsset.IsValid() && PhysicsAsset.Get() != nullptr) && (AssetArrays.IsValid() && AssetArrays.Get() != nullptr);
+}
+
+void FNDIPhysicsAssetBuffer::SetupArrays(const TWeakObjectPtr<UPhysicsAsset> InPhysicsAsset, const TWeakObjectPtr<USkeletalMeshComponent> InSkeletalMesh, const FTransform& InWorldTransform)
+{
+	PhysicsAsset = InPhysicsAsset;
+	SkeletalMesh = InSkeletalMesh;
+	WorldTransform = InWorldTransform;
+
+	AssetArrays = MakeUnique<FNDIPhysicsAssetArrays>();
+
+	if (IsValid())
+	{
+		CreateInternalArrays(PhysicsAsset, SkeletalMesh, AssetArrays.Get(), WorldTransform);
+	}
+}
+
+void FNDIPhysicsAssetBuffer::UpdateBuffers()
+{
+	if (IsValid())
+	{
+		UpdateInternalArrays(PhysicsAsset, SkeletalMesh, AssetArrays.Get(), WorldTransform);
+
+		FNDIPhysicsAssetBuffer* ThisBuffer = this;
+		ENQUEUE_RENDER_COMMAND(UpdatePhysicsAsset)(
+			[ThisBuffer](FRHICommandListImmediate& RHICmdList) mutable
+		{
+			CreateInternalBuffer<FVector4, FVector4, 1, EPixelFormat::PF_A32B32G32R32F, false>(ThisBuffer->AssetArrays->CurrentTransform.Num(), ThisBuffer->AssetArrays->CurrentTransform, ThisBuffer->CurrentTransformBuffer);
+			CreateInternalBuffer<FVector4, FVector4, 1, EPixelFormat::PF_A32B32G32R32F, false>(ThisBuffer->AssetArrays->PreviousTransform.Num(), ThisBuffer->AssetArrays->PreviousTransform, ThisBuffer->PreviousTransformBuffer);
+			CreateInternalBuffer<FVector4, FVector4, 1, EPixelFormat::PF_A32B32G32R32F, false>(ThisBuffer->AssetArrays->InverseTransform.Num(), ThisBuffer->AssetArrays->InverseTransform, ThisBuffer->InverseTransformBuffer);
+		}
+		);
+	}
+}
+
+void FNDIPhysicsAssetBuffer::InitRHI()
+{
+	if (IsValid())
+	{
+		CreateInternalBuffer<FVector4, FVector4, 1, EPixelFormat::PF_A32B32G32R32F, true>(AssetArrays->CurrentTransform.Num(), AssetArrays->CurrentTransform, CurrentTransformBuffer);
+		CreateInternalBuffer<FVector4, FVector4, 1, EPixelFormat::PF_A32B32G32R32F, true>(AssetArrays->PreviousTransform.Num(), AssetArrays->PreviousTransform, PreviousTransformBuffer);
+		CreateInternalBuffer<FVector4, FVector4, 1, EPixelFormat::PF_A32B32G32R32F, true>(AssetArrays->InverseTransform.Num(), AssetArrays->InverseTransform, InverseTransformBuffer);
+		CreateInternalBuffer<FVector4, FVector4, 1, EPixelFormat::PF_A32B32G32R32F, true>(AssetArrays->ElementExtent.Num(), AssetArrays->ElementExtent, ElementExtentBuffer);
+
+		UE_LOG(LogPhysicsAsset, Warning, TEXT("Num Capsules = %d | Num Spheres = %d | Num Boxes = %d"), AssetArrays->ElementOffsets.NumElements - AssetArrays->ElementOffsets.CapsuleOffset,
+			AssetArrays->ElementOffsets.CapsuleOffset - AssetArrays->ElementOffsets.SphereOffset, AssetArrays->ElementOffsets.SphereOffset - AssetArrays->ElementOffsets.BoxOffset);
+	}
+}
+
+void FNDIPhysicsAssetBuffer::ReleaseRHI()
+{
+	CurrentTransformBuffer.Release();
+	PreviousTransformBuffer.Release();
+	InverseTransformBuffer.Release();
+	ElementExtentBuffer.Release();
+}
+
+//------------------------------------------------------------------------------------------------------------
+
+struct FNDIPhysicsAssetParametersCS : public FNiagaraDataInterfaceParametersCS
+{
+	virtual void Bind(const FNiagaraDataInterfaceParamRef& ParamRef, const class FShaderParameterMap& ParameterMap) override
+	{
+		FNDIPhysicsAssetParametersName ParamNames(ParamRef.ParameterInfo.DataInterfaceHLSLSymbol);
+
+		ElementOffsets.Bind(ParameterMap, *ParamNames.ElementOffsetsName);
+
+		CurrentTransformBuffer.Bind(ParameterMap, *ParamNames.CurrentTransformBufferName);
+		PreviousTransformBuffer.Bind(ParameterMap, *ParamNames.PreviousTransformBufferName);
+		InverseTransformBuffer.Bind(ParameterMap, *ParamNames.InverseTransformBufferName);
+		ElementExtentBuffer.Bind(ParameterMap, *ParamNames.ElementExtentBufferName);
+
+		if (!CurrentTransformBuffer.IsBound())
+		{
+			UE_LOG(LogPhysicsAsset, Warning, TEXT("Binding failed for FNDIPhysicsAssetParametersCS %s. Was it optimized out?"), *ParamNames.CurrentTransformBufferName)
+		}
+		if (!PreviousTransformBuffer.IsBound())
+		{
+			UE_LOG(LogPhysicsAsset, Warning, TEXT("Binding failed for FNDIPhysicsAssetParametersCS %s. Was it optimized out?"), *ParamNames.PreviousTransformBufferName)
+		}
+		if (!InverseTransformBuffer.IsBound())
+		{
+			UE_LOG(LogPhysicsAsset, Warning, TEXT("Binding failed for FNDIPhysicsAssetParametersCS %s. Was it optimized out?"), *ParamNames.InverseTransformBufferName)
+		}
+		if (!ElementExtentBuffer.IsBound())
+		{
+			UE_LOG(LogPhysicsAsset, Warning, TEXT("Binding failed for FNDIPhysicsAssetParametersCS %s. Was it optimized out?"), *ParamNames.ElementExtentBufferName)
+		}
+	}
+
+	virtual void Serialize(FArchive& Ar) override
+	{
+		Ar << ElementOffsets;
+		Ar << CurrentTransformBuffer;
+		Ar << PreviousTransformBuffer;
+		Ar << InverseTransformBuffer;
+		Ar << ElementExtentBuffer;
+	}
+
+	virtual void Set(FRHICommandList& RHICmdList, const FNiagaraDataInterfaceSetArgs& Context) const override
+	{
+		check(IsInRenderingThread());
+
+		FRHIComputeShader* ComputeShaderRHI = Context.Shader->GetComputeShader();
+
+		FNDIPhysicsAssetProxy* InterfaceProxy =
+			static_cast<FNDIPhysicsAssetProxy*>(Context.DataInterface);
+		FNDIPhysicsAssetData* ProxyData =
+			InterfaceProxy->SystemInstancesToProxyData.Find(Context.SystemInstance);
+		ensure(ProxyData);
+
+		if (ProxyData != nullptr)
+		{
+			FNDIPhysicsAssetBuffer* AssetBuffer = ProxyData->AssetBuffer;
+			SetSRVParameter(RHICmdList, ComputeShaderRHI, CurrentTransformBuffer, AssetBuffer->CurrentTransformBuffer.SRV);
+			SetSRVParameter(RHICmdList, ComputeShaderRHI, PreviousTransformBuffer, AssetBuffer->PreviousTransformBuffer.SRV);
+			SetSRVParameter(RHICmdList, ComputeShaderRHI, InverseTransformBuffer, AssetBuffer->InverseTransformBuffer.SRV);
+			SetSRVParameter(RHICmdList, ComputeShaderRHI, ElementExtentBuffer, AssetBuffer->ElementExtentBuffer.SRV);
+
+			SetShaderValue(RHICmdList, ComputeShaderRHI, ElementOffsets, AssetBuffer->AssetArrays->ElementOffsets);
+			
+		}
+		else
+		{
+			SetSRVParameter(RHICmdList, ComputeShaderRHI, CurrentTransformBuffer, FNiagaraRenderer::GetDummyFloatBuffer().SRV);
+			SetSRVParameter(RHICmdList, ComputeShaderRHI, PreviousTransformBuffer, FNiagaraRenderer::GetDummyFloatBuffer().SRV);
+			SetSRVParameter(RHICmdList, ComputeShaderRHI, InverseTransformBuffer, FNiagaraRenderer::GetDummyFloatBuffer().SRV);
+			SetSRVParameter(RHICmdList, ComputeShaderRHI, ElementExtentBuffer, FNiagaraRenderer::GetDummyFloatBuffer().SRV);
+
+			static const FElementOffset DummyOffsets(0,0,0,0);
+			SetShaderValue(RHICmdList, ComputeShaderRHI, ElementOffsets, DummyOffsets);
+		}
+	}
+
+	virtual void Unset(FRHICommandList& RHICmdList, const FNiagaraDataInterfaceSetArgs& Context) const override
+	{
+	}
+
+private:
+
+	FShaderParameter ElementOffsets;
+
+	FShaderResourceParameter CurrentTransformBuffer;
+	FShaderResourceParameter PreviousTransformBuffer;
+	FShaderResourceParameter InverseTransformBuffer;
+	FShaderResourceParameter ElementExtentBuffer;
+};
+
+//------------------------------------------------------------------------------------------------------------
+
+void FNDIPhysicsAssetProxy::DeferredDestroy()
+{
+	for (const FNiagaraSystemInstanceID& Sys : DeferredDestroyList)
+	{
+		SystemInstancesToProxyData.Remove(Sys);
+	}
+
+	DeferredDestroyList.Empty();
+}
+
+void FNDIPhysicsAssetProxy::ConsumePerInstanceDataFromGameThread(void* PerInstanceData, const FNiagaraSystemInstanceID& Instance)
+{
+	FNDIPhysicsAssetData* SourceData = static_cast<FNDIPhysicsAssetData*>(PerInstanceData);
+	FNDIPhysicsAssetData* TargetData = SystemInstancesToProxyData.Find(Instance);
+
+	ensure(TargetData);
+	if (TargetData)
+	{
+		TargetData->AssetBuffer = SourceData->AssetBuffer;
+	}
+	else
+	{
+		UE_LOG(LogPhysicsAsset, Log, TEXT("ConsumePerInstanceDataFromGameThread() ... could not find %d"), Instance);
+	}
+}
+
+void FNDIPhysicsAssetProxy::InitializePerInstanceData(const FNiagaraSystemInstanceID& SystemInstance, FNDIPhysicsAssetBuffer* AssetBuffer)
+{
+	check(IsInRenderingThread());
+
+	FNDIPhysicsAssetData* TargetData = SystemInstancesToProxyData.Find(SystemInstance);
+	if (TargetData != nullptr)
+	{
+		DeferredDestroyList.Remove(SystemInstance);
+	}
+	else
+	{
+		TargetData = &SystemInstancesToProxyData.Add(SystemInstance);
+	}
+	TargetData->AssetBuffer = AssetBuffer;
+}
+
+void FNDIPhysicsAssetProxy::DestroyPerInstanceData(NiagaraEmitterInstanceBatcher* Batcher, const FNiagaraSystemInstanceID& SystemInstance)
+{
+	check(IsInRenderingThread());
+
+	DeferredDestroyList.Add(SystemInstance);
+	Batcher->EnqueueDeferredDeletesForDI_RenderThread(this->AsShared());
+}
+
+//------------------------------------------------------------------------------------------------------------
+
+UNiagaraDataInterfacePhysicsAsset::UNiagaraDataInterfacePhysicsAsset(FObjectInitializer const& ObjectInitializer)
+	: Super(ObjectInitializer)
+	, DefaultSource(nullptr)
+	, SourceActor(nullptr)
+	, SourceComponent(nullptr)
+	, PhysicsAsset(nullptr)
+{
+	Proxy = MakeShared<FNDIPhysicsAssetProxy, ESPMode::ThreadSafe>();
+}
+
+bool UNiagaraDataInterfacePhysicsAsset::InitPerInstanceData(void* PerInstanceData, FNiagaraSystemInstance* SystemInstance)
+{
+	FNDIPhysicsAssetData* InstanceData = new (PerInstanceData) FNDIPhysicsAssetData();
+
+	SourceComponent = nullptr;
+	if (SourceActor)
+	{
+		ASkeletalMeshActor* SkeletalMeshActor = Cast<ASkeletalMeshActor>(SourceActor);
+		if (SkeletalMeshActor != nullptr)
+		{
+			SourceComponent = SkeletalMeshActor->GetSkeletalMeshComponent();
+		}
+		else
+		{
+			SourceComponent = SourceActor->FindComponentByClass<USkeletalMeshComponent>();
+		}
+	}
+	else
+	{
+		if (UNiagaraComponent* SimComp = SystemInstance->GetComponent())
+		{
+			if (USkeletalMeshComponent* ParentComp = Cast<USkeletalMeshComponent>(SimComp->GetAttachParent()))
+			{
+				SourceComponent = ParentComp;
+			}
+			else if (USkeletalMeshComponent* OuterComp = SimComp->GetTypedOuter<USkeletalMeshComponent>())
+			{
+				SourceComponent = OuterComp;
+			}
+			else if (AActor* Owner = SimComp->GetAttachmentRootActor())
+			{
+				for (UActorComponent* ActorComp : Owner->GetComponents())
+				{
+					USkeletalMeshComponent* SourceComp = Cast<USkeletalMeshComponent>(ActorComp);
+					if (SourceComp && SourceComp->SkeletalMesh)
+					{
+						SourceComponent = SourceComp;
+						break;
+					}
+				}
+			}
+		}
+	}
+
+	if (SourceComponent != nullptr)
+	{
+		UE_LOG(LogPhysicsAsset, Log, TEXT("Found Skeletal Mesh %d %d %d"), SourceComponent.Get(), SourceComponent->GetPhysicsAsset(), DefaultSource);
+
+	}
+
+	PhysicsAsset = (SourceComponent != nullptr) ? SourceComponent->GetPhysicsAsset() : ( DefaultSource != nullptr ) ? DefaultSource : nullptr;
+
+	if (PhysicsAsset == nullptr )
+	{
+		UE_LOG(LogPhysicsAsset, Log, TEXT("Physics Asset data interface has no valid asset. Failed InitPerInstanceData - %s %d %d"), *GetFullName(), PhysicsAsset.Get(), DefaultSource);
+		return false;
+	}
+	else
+	{
+		const FTransform WorldTransform = (SourceComponent != nullptr) ? SourceComponent->GetComponentTransform() : SystemInstance->GetComponent()->GetComponentTransform();
+		FNDIPhysicsAssetBuffer* AssetBuffer = new FNDIPhysicsAssetBuffer;
+		AssetBuffer->SetupArrays(PhysicsAsset, SourceComponent, WorldTransform);
+
+		// Push instance data to RT
+		{
+			InstanceData->AssetBuffer = AssetBuffer;
+
+			FNDIPhysicsAssetProxy* ThisProxy = GetProxyAs<FNDIPhysicsAssetProxy>();
+			ENQUEUE_RENDER_COMMAND(FNiagaraDIPushInitialInstanceDataToRT) (
+				[ThisProxy, InstanceID = SystemInstance->GetId(), AssetBuffer](FRHICommandListImmediate& CmdList)
+			{
+				AssetBuffer->InitResource();
+				ThisProxy->InitializePerInstanceData(InstanceID, AssetBuffer);
+			}
+			);
+		}
+	}
+	return true;
+}
+
+void UNiagaraDataInterfacePhysicsAsset::DestroyPerInstanceData(void* PerInstanceData, FNiagaraSystemInstance* SystemInstance)
+{
+	FNDIPhysicsAssetData* InstanceData = static_cast<FNDIPhysicsAssetData*>(PerInstanceData);
+	InstanceData->~FNDIPhysicsAssetData();
+
+	if (InstanceData->AssetBuffer)
+	{
+		FNDIPhysicsAssetProxy* ThisProxy = GetProxyAs<FNDIPhysicsAssetProxy>();
+		FNDIPhysicsAssetBuffer* InBuffer = InstanceData->AssetBuffer;
+		ENQUEUE_RENDER_COMMAND(FNiagaraDIDestroyInstanceData) (
+			[ThisProxy, InBuffer,  InstanceID = SystemInstance->GetId(), Batcher = SystemInstance->GetBatcher()](FRHICommandListImmediate& CmdList)
+		{
+			InBuffer->ReleaseResource();
+			ThisProxy->DestroyPerInstanceData(Batcher, InstanceID);
+			delete InBuffer;
+		}
+		);
+		InstanceData->AssetBuffer = nullptr;
+	}
+}
+
+bool UNiagaraDataInterfacePhysicsAsset::PerInstanceTick(void* PerInstanceData, FNiagaraSystemInstance* SystemInstance, float InDeltaSeconds)
+{
+	FNDIPhysicsAssetData* InstanceData = static_cast<FNDIPhysicsAssetData*>(PerInstanceData);
+	if (InstanceData->AssetBuffer)
+	{
+		InstanceData->AssetBuffer->WorldTransform = (InstanceData->AssetBuffer->SkeletalMesh != nullptr) ? InstanceData->AssetBuffer->SkeletalMesh->GetComponentTransform() :
+			SystemInstance->GetComponent()->GetComponentTransform();
+		InstanceData->AssetBuffer->UpdateBuffers();
+	}
+	return false;
+}
+
+bool UNiagaraDataInterfacePhysicsAsset::CopyToInternal(UNiagaraDataInterface* Destination) const
+{
+	if (!Super::CopyToInternal(Destination))
+	{
+		return false;
+	}
+
+	UNiagaraDataInterfacePhysicsAsset* OtherTyped = CastChecked<UNiagaraDataInterfacePhysicsAsset>(Destination);
+	OtherTyped->PhysicsAsset = PhysicsAsset;
+	OtherTyped->SourceActor = SourceActor;
+	OtherTyped->SourceComponent = SourceComponent;
+	OtherTyped->DefaultSource = DefaultSource;
+
+	return true;
+}
+
+bool UNiagaraDataInterfacePhysicsAsset::Equals(const UNiagaraDataInterface* Other) const
+{
+	if (!Super::Equals(Other))
+	{
+		return false;
+	}
+	const UNiagaraDataInterfacePhysicsAsset* OtherTyped = CastChecked<const UNiagaraDataInterfacePhysicsAsset>(Other);
+
+	return  (OtherTyped->PhysicsAsset == PhysicsAsset) && (OtherTyped->SourceActor == SourceActor) && (OtherTyped->SourceComponent == SourceComponent) && (OtherTyped->DefaultSource == DefaultSource);
+}
+
+void UNiagaraDataInterfacePhysicsAsset::PostInitProperties()
+{
+	Super::PostInitProperties();
+
+	if (HasAnyFlags(RF_ClassDefaultObject))
+	{
+		FNiagaraTypeRegistry::Register(FNiagaraTypeDefinition(GetClass()), true, false, false);
+	}
+}
+
+void UNiagaraDataInterfacePhysicsAsset::GetFunctions(TArray<FNiagaraFunctionSignature>& OutFunctions)
+{
+	{
+		FNiagaraFunctionSignature Sig;
+		Sig.Name = GetNumBoxesName;
+		Sig.bMemberFunction = true;
+		Sig.bRequiresContext = false;
+		Sig.Inputs.Add(FNiagaraVariable(FNiagaraTypeDefinition(GetClass()), TEXT("Physics Asset")));
+		Sig.Outputs.Add(FNiagaraVariable(FNiagaraTypeDefinition::GetIntDef(), TEXT("Num Boxes")));
+
+		OutFunctions.Add(Sig);
+	}
+	{
+		FNiagaraFunctionSignature Sig;
+		Sig.Name = GetNumSpheresName;
+		Sig.bMemberFunction = true;
+		Sig.bRequiresContext = false;
+		Sig.Inputs.Add(FNiagaraVariable(FNiagaraTypeDefinition(GetClass()), TEXT("Physics Asset")));
+		Sig.Outputs.Add(FNiagaraVariable(FNiagaraTypeDefinition::GetIntDef(), TEXT("Num Spheres")));
+
+		OutFunctions.Add(Sig);
+	}
+	{
+		FNiagaraFunctionSignature Sig;
+		Sig.Name = GetNumCapsulesName;
+		Sig.bMemberFunction = true;
+		Sig.bRequiresContext = false;
+		Sig.Inputs.Add(FNiagaraVariable(FNiagaraTypeDefinition(GetClass()), TEXT("Physics Asset")));
+		Sig.Outputs.Add(FNiagaraVariable(FNiagaraTypeDefinition::GetIntDef(), TEXT("Num Capsules")));
+
+		OutFunctions.Add(Sig);
+	}
+	{
+		FNiagaraFunctionSignature Sig;
+		Sig.Name = GetClosestPointName;
+		Sig.bMemberFunction = true;
+		Sig.bRequiresContext = false;
+		Sig.Inputs.Add(FNiagaraVariable(FNiagaraTypeDefinition(GetClass()), TEXT("Physics Asset")));
+		Sig.Inputs.Add(FNiagaraVariable(FNiagaraTypeDefinition::GetVec3Def(), TEXT("World Position")));
+		Sig.Inputs.Add(FNiagaraVariable(FNiagaraTypeDefinition::GetFloatDef(), TEXT("Delta Time")));
+		Sig.Outputs.Add(FNiagaraVariable(FNiagaraTypeDefinition::GetVec3Def(), TEXT("Closest Position")));
+		Sig.Outputs.Add(FNiagaraVariable(FNiagaraTypeDefinition::GetVec3Def(), TEXT("Closest Normal")));
+		Sig.Outputs.Add(FNiagaraVariable(FNiagaraTypeDefinition::GetVec3Def(), TEXT("Closest Velocity")));
+		Sig.Outputs.Add(FNiagaraVariable(FNiagaraTypeDefinition::GetFloatDef(), TEXT("Closest Distance")));
+
+		OutFunctions.Add(Sig);
+	}
+}
+
+DEFINE_NDI_DIRECT_FUNC_BINDER(UNiagaraDataInterfacePhysicsAsset, GetNumBoxes);
+DEFINE_NDI_DIRECT_FUNC_BINDER(UNiagaraDataInterfacePhysicsAsset, GetNumSpheres);
+DEFINE_NDI_DIRECT_FUNC_BINDER(UNiagaraDataInterfacePhysicsAsset, GetNumCapsules);
+DEFINE_NDI_DIRECT_FUNC_BINDER(UNiagaraDataInterfacePhysicsAsset, GetClosestPoint);
+
+void UNiagaraDataInterfacePhysicsAsset::GetVMExternalFunction(const FVMExternalFunctionBindingInfo& BindingInfo, void* InstanceData, FVMExternalFunction &OutFunc)
+{
+	if (BindingInfo.Name == GetNumBoxesName)
+	{
+		check(BindingInfo.GetNumInputs() == 1 && BindingInfo.GetNumOutputs() == 1);
+		NDI_FUNC_BINDER(UNiagaraDataInterfacePhysicsAsset, GetNumBoxes)::Bind(this, OutFunc);
+	}
+	else if (BindingInfo.Name == GetNumSpheresName)
+	{
+		check(BindingInfo.GetNumInputs() == 1 && BindingInfo.GetNumOutputs() == 1);
+		NDI_FUNC_BINDER(UNiagaraDataInterfacePhysicsAsset, GetNumSpheres)::Bind(this, OutFunc);
+	}
+	else if (BindingInfo.Name == GetNumCapsulesName)
+	{
+		check(BindingInfo.GetNumInputs() == 1 && BindingInfo.GetNumOutputs() == 1);
+		NDI_FUNC_BINDER(UNiagaraDataInterfacePhysicsAsset, GetNumCapsules)::Bind(this, OutFunc);
+	}
+	else if (BindingInfo.Name == GetClosestPointName)
+	{
+		check(BindingInfo.GetNumInputs() == 5 && BindingInfo.GetNumOutputs() == 10);
+		NDI_FUNC_BINDER(UNiagaraDataInterfacePhysicsAsset, GetClosestPoint)::Bind(this, OutFunc);
+	}
+}
+
+void UNiagaraDataInterfacePhysicsAsset::GetNumBoxes(FVectorVMContext& Context)
+{
+}
+
+void UNiagaraDataInterfacePhysicsAsset::GetNumSpheres(FVectorVMContext& Context)
+{
+}
+
+void UNiagaraDataInterfacePhysicsAsset::GetNumCapsules(FVectorVMContext& Context)
+{
+}
+
+void UNiagaraDataInterfacePhysicsAsset::GetClosestPoint(FVectorVMContext& Context)
+{
+}
+
+bool UNiagaraDataInterfacePhysicsAsset::GetFunctionHLSL(const FName& DefinitionFunctionName, FString InstanceFunctionName, FNiagaraDataInterfaceGPUParamInfo& ParamInfo, FString& OutHLSL)
+{
+	FNDIPhysicsAssetParametersName ParamNames(ParamInfo.DataInterfaceHLSLSymbol);
+
+	TMap<FString, FStringFormatArg> ArgsSample = {
+		{TEXT("InstanceFunctionName"), InstanceFunctionName},
+		{TEXT("ElementOffsetsName"), ParamNames.ElementOffsetsName},
+		{TEXT("CurrentTransformBufferName"), ParamNames.CurrentTransformBufferName},
+		{TEXT("PreviousTransformBufferName"), ParamNames.PreviousTransformBufferName},
+		{TEXT("InverseTransformBufferName"), ParamNames.InverseTransformBufferName},
+		{TEXT("ElementExtentBufferName"), ParamNames.ElementExtentBufferName},
+		{TEXT("PhysicsAssetContextName"), TEXT("DIPhysicsAsset_MAKE_CONTEXT(") + ParamInfo.DataInterfaceHLSLSymbol + TEXT(")")},
+	};
+
+	if (DefinitionFunctionName == GetNumBoxesName)
+	{
+		static const TCHAR *FormatSample = TEXT(R"(
+		void {InstanceFunctionName}(out int OutNumBoxes)
+		{
+			{PhysicsAssetContextName}
+			OutNumBoxes = DIPhysicsAsset_GetNumBoxes(DIContext);
+		}
+		)");
+		OutHLSL += FString::Format(FormatSample, ArgsSample);
+		return true;
+	}
+	else if (DefinitionFunctionName == GetNumCapsulesName)
+	{
+		static const TCHAR *FormatSample = TEXT(R"(
+		void {InstanceFunctionName}(out int OutNumCapsules)
+		{
+			{PhysicsAssetContextName}
+			OutNumCapsules = DIPhysicsAsset_GetNumCapsules(DIContext);
+		}
+		)");
+		OutHLSL += FString::Format(FormatSample, ArgsSample);
+		return true;
+	}
+	else if (DefinitionFunctionName == GetNumSpheresName)
+	{
+		static const TCHAR *FormatSample = TEXT(R"(
+		void {InstanceFunctionName}(out int OutNumSpheres)
+		{
+			{PhysicsAssetContextName}
+			OutNumSpheres = DIPhysicsAsset_GetNumSpheres(DIContext);
+		}
+		)");
+		OutHLSL += FString::Format(FormatSample, ArgsSample);
+		return true;
+	}
+	else if (DefinitionFunctionName == GetClosestPointName)
+	{
+		static const TCHAR *FormatSample = TEXT(R"(
+		void {InstanceFunctionName}(in float3 WorldPosition, in float DeltaTime, out float3 OutClosestPosition, 
+							out float3 OutClosestNormal, out float3 OutClosestVelocity, out float OutClosestDistance)
+		{
+			{PhysicsAssetContextName} DIPhysicsAsset_GetClosestPoint(DIContext,WorldPosition,DeltaTime,
+				OutClosestPosition,OutClosestNormal,OutClosestVelocity,OutClosestDistance);
+		}
+		)");
+		OutHLSL += FString::Format(FormatSample, ArgsSample);
+		return true;
+	}
+
+	OutHLSL += TEXT("\n");
+	return false;
+}
+
+void UNiagaraDataInterfacePhysicsAsset::GetCommonHLSL(FString& OutHLSL)
+{
+	OutHLSL += TEXT("#include \"/Plugin/Experimental/HairStrands/Private/NiagaraDataInterfacePhysicsAsset.ush\"\n");
+}
+
+void UNiagaraDataInterfacePhysicsAsset::GetParameterDefinitionHLSL(FNiagaraDataInterfaceGPUParamInfo& ParamInfo, FString& OutHLSL)
+{
+	OutHLSL += TEXT("DIPhysicsAsset_DECLARE_CONSTANTS(") + ParamInfo.DataInterfaceHLSLSymbol + TEXT(")\n");
+}
+
+void UNiagaraDataInterfacePhysicsAsset::ProvidePerInstanceDataForRenderThread(void* DataForRenderThread, void* PerInstanceData, const FNiagaraSystemInstanceID& SystemInstance)
+{
+	FNDIPhysicsAssetData* GameThreadData = static_cast<FNDIPhysicsAssetData*>(PerInstanceData);
+	FNDIPhysicsAssetData* RenderThreadData = static_cast<FNDIPhysicsAssetData*>(DataForRenderThread);
+
+	if (GameThreadData != nullptr && RenderThreadData != nullptr)
+	{
+		RenderThreadData->AssetBuffer = GameThreadData->AssetBuffer;
+	}
+	check(Proxy);
+}
+
+FNiagaraDataInterfaceParametersCS*
+UNiagaraDataInterfacePhysicsAsset::ConstructComputeParameters() const
+{
+	return new FNDIPhysicsAssetParametersCS();
+}
+
+
+#undef LOCTEXT_NAMESPACE
\ No newline at end of file
diff --git a/Engine/Plugins/Experimental/HairStrands/Source/HairStrandsNiagara/Public/HairStrandsNiagara.h b/Engine/Plugins/Experimental/HairStrands/Source/HairStrandsNiagara/Public/HairStrandsNiagara.h
--- a/Engine/Plugins/Experimental/HairStrands/Source/HairStrandsNiagara/Public/HairStrandsNiagara.h
+++ b/Engine/Plugins/Experimental/HairStrands/Source/HairStrandsNiagara/Public/HairStrandsNiagara.h
@@ -15,5 +15,3 @@ public:
 	virtual void StartupModule() override;
 	virtual void ShutdownModule() override;
 };
-
-IMPLEMENT_MODULE(FHairStrandsNiagara, HairStrandsNiagara);
diff --git a/Engine/Plugins/Experimental/HairStrands/Source/HairStrandsNiagara/Public/NiagaraDataInterfaceHairStrands.h b/Engine/Plugins/Experimental/HairStrands/Source/HairStrandsNiagara/Public/NiagaraDataInterfaceHairStrands.h
--- a/Engine/Plugins/Experimental/HairStrands/Source/HairStrandsNiagara/Public/NiagaraDataInterfaceHairStrands.h
+++ b/Engine/Plugins/Experimental/HairStrands/Source/HairStrandsNiagara/Public/NiagaraDataInterfaceHairStrands.h
@@ -3,9 +3,11 @@
 #pragma once
 
 #include "NiagaraDataInterface.h"
+#include "NiagaraDataInterfaceRW.h"
 #include "NiagaraCommon.h"
 #include "VectorVM.h"
-#include "HairStrandsAsset.h"
+#include "GroomAsset.h"
+#include "GroomActor.h"
 #include "NiagaraDataInterfaceHairStrands.generated.h"
 
 /** Size of each strands*/
@@ -20,14 +22,25 @@ enum class EHairStrandsSize : uint8
 	Size32 = 0x20 UMETA(DisplatName = "32")
 };
 
-/** Render buffers that will be used in hlsl functions */
-struct FNDIHairStrandsBuffer : public FRenderResource
+/** Structure storing the grid size*/
+struct FGridSize
 {
-	/** Set the asset that will be used to affect the buffer */
-	void SetAsset(const UHairStrandsAsset*  HairStrandsAsset);
+	FGridSize(const uint32 InGridSizeX = 0, const uint32 InGridSizeY = 0, 
+		      const uint32 InGridSizeZ = 0, const uint32 InGridSizeW = 0) : GridSizeX(InGridSizeX), 
+		GridSizeY(InGridSizeY), GridSizeZ(InGridSizeZ), GridSizeW(InGridSizeW) {}
+
+	uint32 GridSizeX;
+	uint32 GridSizeY;
+	uint32 GridSizeZ;
+	uint32 GridSizeW;
+};
 
-	/** Set the strand size */
-	void SetStrandSize(const uint32 StrandSize);
+
+/** Render buffers that will be used in hlsl functions */
+struct FNDICollisionGridBuffer : public FRenderResource
+{
+	/** Set the grid size */
+	void SetGridSize(const FGridSize GridSize);
 
 	/** Init the buffer */
 	virtual void InitRHI() override;
@@ -35,34 +48,53 @@ struct FNDIHairStrandsBuffer : public FRenderResource
 	/** Release the buffer */
 	virtual void ReleaseRHI() override;
 
+	/** Clear all UAV*/
+	void ClearBuffers(FRHICommandList& RHICmdList);
+
 	/** Get the resource name */
-	virtual FString GetFriendlyName() const override { return TEXT("FNDIHairStrandsBuffer"); }
+	virtual FString GetFriendlyName() const override { return TEXT("FNDICollisionGridBuffer"); }
+
+	/** Grid data texture */
+	FTextureRWBuffer3D GridDataBuffer;
+
+	/** Grid size that will be used for the collision*/
+	FGridSize GridSize;
+};
 
-	/** Position Buffer UAV */
-	FUnorderedAccessViewRHIRef PointsPositionsBufferUav;
+/** Render buffers that will be used in hlsl functions */
+struct FNDIHairStrandsBuffer : public FRenderResource
+{
+	/** Set the asset that will be used to affect the buffer */
+	void SetHairAsset(const FHairStrandsDatas*  HairStrandsDatas, const FHairStrandsResource*  HairStrandsResource );
 
-	/** Strand nodes positions buffer */
-	FRWBuffer NodesPositionsBuffer;
+	/** Init the buffer */
+	virtual void InitRHI() override;
 
-	/** Strand points node index buffer */
-	FRWBuffer PointsNodesBuffer;
+	/** Release the buffer */
+	virtual void ReleaseRHI() override;
 
-	/** Strand points node weights buffer */
-	FRWBuffer PointsWeightsBuffer;
+	/** Get the resource name */
+	virtual FString GetFriendlyName() const override { return TEXT("FNDIHairStrandsBuffer"); }
 
 	/** Strand curves point offset buffer */
 	FRWBuffer CurvesOffsetsBuffer;
 
-	/** The strand asset from which to sample */
-	const UHairStrandsAsset* SourceAsset;
+	/** The strand asset datas from which to sample */
+	const FHairStrandsDatas* SourceDatas;
 
-	/** Strand size that will be used for resampling*/
-	uint32 StrandSize;
+	/** The strand asset resource from which to sample */
+	const FHairStrandsResource* SourceResources;
 };
 
 /** Data stored per strand base instance*/
 struct FNDIHairStrandsData
 {
+	/** Swap the current and the destination data */
+	void SwapBuffers();
+
+	/** Init the current and the destination data */
+	void InitBuffers();
+
 	/** Cached World transform. */
 	FMatrix WorldTransform;
 
@@ -81,13 +113,34 @@ struct FNDIHairStrandsData
 	/** Tip Thickness */
 	float TipThickness;
 
+	/** Grid Origin */
+	FVector4 GridOrigin;
+
+	/** Grid Size */
+	FGridSize GridSize;
+
 	/** Strands Gpu buffer */
-	FNDIHairStrandsBuffer* StrandBuffer;
+	FNDIHairStrandsBuffer* HairStrandsBuffer;
+
+	/** Collision Grid Buffer A */
+	FNDICollisionGridBuffer* CollisionGridBufferA;
+
+	/** Collision Grid Buffer B */
+	FNDICollisionGridBuffer* CollisionGridBufferB;
+
+	/** Pointer to the current buffer */
+	FNDICollisionGridBuffer* CurrentGridBuffer;
+
+	/** Pointer to the destination buffer */
+	FNDICollisionGridBuffer* DestinationGridBuffer;
+
+	/** Sorted index particle buffers*/
+	//FParticleSortBuffers SortedParticleBuffers;
 };
 
 /** Data Interface for the strand base */
 UCLASS(EditInlineNew, Category = "Strands", meta = (DisplayName = "Hair Strands"))
-class HAIRSTRANDSNIAGARA_API UNiagaraDataInterfaceHairStrands : public UNiagaraDataInterface
+class HAIRSTRANDSNIAGARA_API UNiagaraDataInterfaceHairStrands : public UNiagaraDataInterfaceRWBase
 {
 	GENERATED_UCLASS_BODY()
 
@@ -109,24 +162,35 @@ public:
 	UPROPERTY(EditAnywhere, Category = "Spawn")
 	float TipThickness;
 
-	/** Mesh used to sample from when not overridden by a source actor from the scene. Also useful for previewing in the editor. */
+	/** Hair Strands Asset used to sample from when not overridden by a source actor from the scene. Also useful for previewing in the editor. */
 	UPROPERTY(EditAnywhere, Category = "Source")
-	UHairStrandsAsset* SourceAsset;
+	UGroomAsset* DefaultSource;
+
+	/** The source actor from which to sample */
+	UPROPERTY(EditAnywhere, Category = "Source")
+	AActor* SourceActor;
 
 	/** Source transform to be applied to the asset. */
 	UPROPERTY(EditAnywhere, Category = "Source")
 	FMatrix SourceTransform;
 
-	/** Static mesh on which the roots are attached. */
-	UPROPERTY(EditAnywhere, Category = "Attachment")
-	UStaticMesh* StaticMesh;
+	/** Grid size along the X axis. */
+	UPROPERTY(EditAnywhere, Category = "Grid")
+	uint32 GridSizeX;
+
+	/** Grid size along the Y axis. */
+	UPROPERTY(EditAnywhere, Category = "Grid")
+	uint32 GridSizeY;
 
-	///** The source actor from which to sample */
-	//UPROPERTY(EditAnywhere, AdvancedDisplay, Category = "Source")
-	//AActor* SourceActor;
+	/** Grid size along the Z axis. */
+	UPROPERTY(EditAnywhere, Category = "Grid")
+	uint32 GridSizeZ;
 
-	///** The source component from which to sample */
-	//class UHairStrandsComponent* SourceComponent;
+	/** The source component from which to sample */
+	TWeakObjectPtr<class UGroomComponent> SourceComponent;
+
+	/** The source asset from which to sample */
+	TWeakObjectPtr<class UGroomAsset> GroomAsset;
 
 	/** UObject Interface */
 	virtual void PostInitProperties() override;
@@ -148,6 +212,9 @@ public:
 	virtual void ProvidePerInstanceDataForRenderThread(void* DataForRenderThread, void* PerInstanceData, const FNiagaraSystemInstanceID& SystemInstance) override;
 	virtual void GetCommonHLSL(FString& OutHLSL) override;
 
+	/** Check if the component is Valid */
+	bool IsComponentValid() const;
+
 	/** Get the number of strands */
 	void GetNumStrands(FVectorVMContext& Context);
 
@@ -166,15 +233,78 @@ public:
 	/** Get the world transform */
 	void GetWorldTransform(FVectorVMContext& Context);
 
+	/** Get the world inverse */
+	void GetWorldInverse(FVectorVMContext& Context);
+
 	/** Get the strand vertex position in world space*/
-	void GetVertexPosition(FVectorVMContext& Context);
+	void GetPointPosition(FVectorVMContext& Context);
 
 	/** Get the strand node position in world space*/
 	void ComputeNodePosition(FVectorVMContext& Context);
 
+	/** Get the strand node orientation in world space*/
+	void ComputeNodeOrientation(FVectorVMContext& Context);
+
+	/** Get the strand node mass */
+	void ComputeNodeMass(FVectorVMContext& Context);
+
+	/** Get the strand node inertia */
+	void ComputeNodeInertia(FVectorVMContext& Context);
+
+	/** Compute the edge length (diff between 2 nodes positions)*/
+	void ComputeEdgeLength(FVectorVMContext& Context);
+
+	/** Compute the edge orientation (diff between 2 nodes orientations) */
+	void ComputeEdgeRotation(FVectorVMContext& Context);
+
+	/** Compute the rest local position */
+	void ComputeRestPosition(FVectorVMContext& Context);
+
+	/** Compute the rest local orientation */
+	void ComputeRestOrientation(FVectorVMContext& Context);
+
+	/** Update the root node orientation based on the current transform */
+	void AttachNodePosition(FVectorVMContext& Context);
+
+	/** Update the root node position based on the current transform */
+	void AttachNodeOrientation(FVectorVMContext& Context);
+
+	/** Report the node displacement onto the points position*/
+	void UpdatePointPosition(FVectorVMContext& Context);
+
+	/** Reset the point position to be the rest one */
+	void ResetPointPosition(FVectorVMContext& Context);
+
+	/** Add external force to the linear velocity and advect node position */
+	void AdvectNodePosition(FVectorVMContext& Context);
+
+	/** Add external torque to the angular velocity and advect node orientation*/
+	void AdvectNodeOrientation(FVectorVMContext& Context);
+
+	/** Update the node linear velocity based on the node position difference */
+	void UpdateLinearVelocity(FVectorVMContext& Context);
+
+	/** Update the node angular velocity based on the node orientation difference */
+	void UpdateAngularVelocity(FVectorVMContext& Context);
+
+	/** Build the collision grid */
+	void BuildCollisionGrid(FVectorVMContext& Context);
+
+	/** Query the collision grid */
+	void QueryCollisionGrid(FVectorVMContext& Context);
+
+	/** Project the collision grid */
+	void ProjectCollisionGrid(FVectorVMContext& Context);
+
 	/** Name of the world transform */
 	static const FString WorldTransformName;
 
+	/** Name of the world transform */
+	static const FString WorldInverseName;
+
+	/** Name of the world rotation */
+	static const FString WorldRotationName;
+
 	/** Name of the number of strands */
 	static const FString NumStrandsName;
 
@@ -197,13 +327,19 @@ public:
 	static const FString CurvesOffsetsBufferName;
 
 	/** Name of the nodes positions buffer */
-	static const FString NodesPositionsBufferName;
+	static const FString RestPositionsBufferName;
+
+	/** Name of the grid current buffer */
+	static const FString GridCurrentBufferName;
+
+	/** Name of the grid X velocity buffer */
+	static const FString GridDestinationBufferName;
 
-	/** Name of the points nodes buffer */
-	static const FString PointsNodesBufferName;
+	/** Name of the grid origin  */
+	static const FString GridOriginName;
 
-	/** Name of the points weights buffer */
-	static const FString PointsWeightsBufferName;
+	/** Name of the grid size */
+	static const FString GridSizeName;
 
 protected:
 	/** Copy one niagara DI to this */
@@ -223,12 +359,21 @@ struct FNDIHairStrandsProxy : public FNiagaraDataInterfaceProxy
 	virtual void ConsumePerInstanceDataFromGameThread(void* PerInstanceData, const FNiagaraSystemInstanceID& Instance) override;
 
 	/** Initialize the Proxy data strands buffer */
-	void InitializePerInstanceData(const FNiagaraSystemInstanceID& SystemInstance, FNDIHairStrandsBuffer* StrandsBuffer, const uint32 NumStrands, const uint8 StrandSize,
-		const float StrandDensity, const float RootThickness, const float TipThickness);
+	void InitializePerInstanceData(const FNiagaraSystemInstanceID& SystemInstance, FNDIHairStrandsBuffer* StrandsBuffer, FNDICollisionGridBuffer* CollisionGridBufferA, FNDICollisionGridBuffer* CollisionGridBufferB, const uint32 NumStrands, const uint8 StrandSize,
+		const float StrandDensity, const float RootThickness, const float TipThickness, const FVector4& GridOrigin, const FGridSize& GridSize);
 
 	/** Destroy the proxy data if necessary */
 	void DestroyPerInstanceData(NiagaraEmitterInstanceBatcher* Batcher, const FNiagaraSystemInstanceID& SystemInstance);
 
+	/** Launch all pre stage functions */
+	virtual void PreStage(FRHICommandList& RHICmdList, const FNiagaraDataInterfaceSetArgs& Context) override;
+
+	/** Launch all post stage functions */
+	virtual void PostStage(FRHICommandList& RHICmdList, const FNiagaraDataInterfaceSetArgs& Context) override;
+
+	/** Reset the buffers  */
+	virtual void ResetData(FRHICommandList& RHICmdList, const FNiagaraDataInterfaceSetArgs& Context) override;
+
 	/** List of proxy data for each system instances*/
 	TMap<FNiagaraSystemInstanceID, FNDIHairStrandsData> SystemInstancesToProxyData;
 
diff --git a/Engine/Plugins/Experimental/HairStrands/Source/HairStrandsNiagara/Public/NiagaraDataInterfacePhysicsAsset.h b/Engine/Plugins/Experimental/HairStrands/Source/HairStrandsNiagara/Public/NiagaraDataInterfacePhysicsAsset.h
new file mode 100644
--- /dev/null
+++ b/Engine/Plugins/Experimental/HairStrands/Source/HairStrandsNiagara/Public/NiagaraDataInterfacePhysicsAsset.h
@@ -0,0 +1,188 @@
+// Copyright 1998-2019 Epic Games, Inc. All Rights Reserved.
+
+#pragma once
+
+#include "NiagaraDataInterface.h"
+#include "NiagaraCommon.h"
+#include "VectorVM.h"
+#include "PhysicsEngine/PhysicsAsset.h"
+#include "NiagaraDataInterfacePhysicsAsset.generated.h"
+
+/** Element offsets in the array list */
+struct FElementOffset
+{
+	FElementOffset(const uint32 InBoxOffset, const uint32 InSphereOffset, const uint32 InCapsuleOffset, const uint32 InNumElements) :
+		BoxOffset(InBoxOffset), SphereOffset(InSphereOffset), CapsuleOffset(InCapsuleOffset), NumElements(InNumElements)
+	{}
+
+	FElementOffset() :
+		BoxOffset(0), SphereOffset(0), CapsuleOffset(0), NumElements(0)
+	{}
+	uint32 BoxOffset;
+	uint32 SphereOffset;
+	uint32 CapsuleOffset;
+	uint32 NumElements;
+};
+
+/** Arrays in which the cpu datas will be str */
+struct FNDIPhysicsAssetArrays 
+{
+	FElementOffset ElementOffsets;
+	TArray<FVector4> CurrentTransform;
+	TArray<FVector4> InverseTransform;
+	TArray<FVector4> PreviousTransform;
+	TArray<FVector4> ElementExtent;
+};
+
+/** Render buffers that will be used in hlsl functions */
+struct FNDIPhysicsAssetBuffer : public FRenderResource
+{
+	/** Check if all the assets are valid */
+	bool IsValid() const;
+
+	/** Set the assets that will be used to affect the buffer */
+	void SetupArrays(const TWeakObjectPtr<class UPhysicsAsset>  PhysicsAsset, const TWeakObjectPtr<class USkeletalMeshComponent>  SkeletalMesh, const FTransform& InWorldTransform);
+
+	/** Update the buffers */
+	void UpdateBuffers();
+
+	/** Init the buffer */
+	virtual void InitRHI() override;
+
+	/** Release the buffer */
+	virtual void ReleaseRHI() override;
+
+	/** Get the resource name */
+	virtual FString GetFriendlyName() const override { return TEXT("FNDIPhysicsAssetBuffer"); }
+
+	/** Current transform buffer */
+	FRWBuffer CurrentTransformBuffer;
+
+	/** Previous transform buffer */
+	FRWBuffer PreviousTransformBuffer;
+
+	/** Inverse transform buffer*/
+	FRWBuffer InverseTransformBuffer;
+
+	/** Element extent buffer */
+	FRWBuffer ElementExtentBuffer;
+
+	/** The physics asset datas from which the buffers will be constructed */
+	TWeakObjectPtr<class UPhysicsAsset> PhysicsAsset;
+
+	/** The skeletal mesh from which the transform will be extracted*/
+	TWeakObjectPtr<class USkeletalMeshComponent> SkeletalMesh;
+
+	/** Physics Asset arrays */
+	TUniquePtr<FNDIPhysicsAssetArrays> AssetArrays;
+
+	/** World transform */
+	FTransform WorldTransform;
+};
+
+/** Data stored per physics asset instance*/
+struct FNDIPhysicsAssetData
+{
+	/** Physics asset Gpu buffer */
+	FNDIPhysicsAssetBuffer* AssetBuffer;
+};
+
+/** Data Interface for the strand base */
+UCLASS(EditInlineNew, Category = "Strands", meta = (DisplayName = "Physics Asset"))
+class HAIRSTRANDSNIAGARA_API UNiagaraDataInterfacePhysicsAsset : public UNiagaraDataInterface
+{
+	GENERATED_UCLASS_BODY()
+
+public:
+
+	/** Skeletal Mesh from which the Physics Asset will be found. */
+	UPROPERTY(EditAnywhere, Category = "Source")
+	UPhysicsAsset* DefaultSource;
+
+	/** The source actor from which to sample */
+	UPROPERTY(EditAnywhere, Category = "Source")
+	AActor* SourceActor;
+
+	/** The source component from which to sample */
+	TWeakObjectPtr<class USkeletalMeshComponent> SourceComponent;
+
+	/** The source asset from which to sample */
+	TWeakObjectPtr<class UPhysicsAsset> PhysicsAsset;
+
+	/** UObject Interface */
+	virtual void PostInitProperties() override;
+
+	/** UNiagaraDataInterface Interface */
+	virtual void GetFunctions(TArray<FNiagaraFunctionSignature>& OutFunctions) override;
+	virtual void GetVMExternalFunction(const FVMExternalFunctionBindingInfo& BindingInfo, void* InstanceData, FVMExternalFunction &OutFunc) override;
+	virtual bool CanExecuteOnTarget(ENiagaraSimTarget Target) const override { return Target == ENiagaraSimTarget::GPUComputeSim; }
+	virtual bool InitPerInstanceData(void* PerInstanceData, FNiagaraSystemInstance* SystemInstance) override;
+	virtual void DestroyPerInstanceData(void* PerInstanceData, FNiagaraSystemInstance* SystemInstance)override;
+	virtual bool PerInstanceTick(void* PerInstanceData, FNiagaraSystemInstance* SystemInstance, float DeltaSeconds) override;
+	virtual int32 PerInstanceDataSize()const override { return sizeof(FNDIPhysicsAssetData); }
+	virtual bool Equals(const UNiagaraDataInterface* Other) const override;
+
+	/** GPU simulation  functionality */
+	virtual bool GetFunctionHLSL(const FName&  DefinitionFunctionName, FString InstanceFunctionName, FNiagaraDataInterfaceGPUParamInfo& ParamInfo, FString& OutHLSL) override;
+	virtual void GetParameterDefinitionHLSL(FNiagaraDataInterfaceGPUParamInfo& ParamInfo, FString& OutHLSL) override;
+	virtual FNiagaraDataInterfaceParametersCS* ConstructComputeParameters() const override;
+	virtual void ProvidePerInstanceDataForRenderThread(void* DataForRenderThread, void* PerInstanceData, const FNiagaraSystemInstanceID& SystemInstance) override;
+	virtual void GetCommonHLSL(FString& OutHLSL) override;
+
+	/** Get the number of boxes*/
+	void GetNumBoxes(FVectorVMContext& Context);
+
+	/** Get the number of spheres  */
+	void GetNumSpheres(FVectorVMContext& Context);
+
+	/** Get the number of capsules */
+	void GetNumCapsules(FVectorVMContext& Context);
+
+	/** Get the closest point */
+	void GetClosestPoint(FVectorVMContext& Context);
+
+	/** Name of element offsets */
+	static const FString ElementOffsetsName;
+
+	/** Name of the current transform buffer */
+	static const FString CurrentTransformBufferName;
+
+	/** Name of the previous transform buffer */
+	static const FString PreviousTransformBufferName;
+
+	/** Name of the inverse transform buffer */
+	static const FString InverseTransformBufferName;
+
+	/** Name of the element extent buffer */
+	static const FString ElementExtentBufferName;
+
+protected:
+	/** Copy one niagara DI to this */
+	virtual bool CopyToInternal(UNiagaraDataInterface* Destination) const override;
+};
+
+/** Proxy to send data to gpu */
+struct FNDIPhysicsAssetProxy : public FNiagaraDataInterfaceProxy
+{
+	/** Destroy internal data */
+	virtual void DeferredDestroy() override;
+
+	/** Get the size of the data that will be passed to render*/
+	virtual int32 PerInstanceDataPassedToRenderThreadSize() const override { return sizeof(FNDIPhysicsAssetData); }
+
+	/** Get the data that will be passed to render*/
+	virtual void ConsumePerInstanceDataFromGameThread(void* PerInstanceData, const FNiagaraSystemInstanceID& Instance) override;
+
+	/** Initialize the Proxy data strands buffer */
+	void InitializePerInstanceData(const FNiagaraSystemInstanceID& SystemInstance, FNDIPhysicsAssetBuffer* AssetBuffer);
+
+	/** Destroy the proxy data if necessary */
+	void DestroyPerInstanceData(NiagaraEmitterInstanceBatcher* Batcher, const FNiagaraSystemInstanceID& SystemInstance);
+
+	/** List of proxy data for each system instances*/
+	TMap<FNiagaraSystemInstanceID, FNDIPhysicsAssetData> SystemInstancesToProxyData;
+
+	/** List of proxy data to destroy later */
+	TSet<FNiagaraSystemInstanceID> DeferredDestroyList;
+};
+
diff --git a/Engine/Plugins/Experimental/SkeletalReduction/Source/SkeletalMeshReduction.Build.cs b/Engine/Plugins/Experimental/SkeletalReduction/Source/SkeletalMeshReduction.Build.cs
--- a/Engine/Plugins/Experimental/SkeletalReduction/Source/SkeletalMeshReduction.Build.cs
+++ b/Engine/Plugins/Experimental/SkeletalReduction/Source/SkeletalMeshReduction.Build.cs
@@ -56,7 +56,7 @@ namespace UnrealBuildTool.Rules
                     "AnimationModifiers",
                     "MeshBoneReduction",
                     "QuadricMeshReduction",
-                    "ClothingSystemRuntime",
+                    "ClothingSystemRuntimeNv",
                     "SkeletalMeshUtilitiesCommon",
 					// ... add private dependencies that you statically link with here ...
                     // QuadricMeshReduction is only for testing
diff --git a/Engine/Plugins/FX/Niagara/Shaders/Private/NiagaraStrandsPhysics.ush b/Engine/Plugins/FX/Niagara/Shaders/Private/NiagaraStrandsPhysics.ush
--- a/Engine/Plugins/FX/Niagara/Shaders/Private/NiagaraStrandsPhysics.ush
+++ b/Engine/Plugins/FX/Niagara/Shaders/Private/NiagaraStrandsPhysics.ush
@@ -73,7 +73,10 @@ void ReportOrientationSharedDatas(const float InDeltaTime, inout float3 OutParti
 			CurrentQuaternion.w * OutParticleQuaternion.w - dot(CurrentQuaternion.xyz, -OutParticleQuaternion.xyz));
 
 	OutParticleQuaternion = CurrentQuaternion;
-	OutParticleOmega += DeltaQuat.xyz * 2.0 / InDeltaTime;
+
+	const float AxisLength = length( DeltaQuat.xyz );
+	const float QuatAngle = 2.0 * atan2(AxisLength,DeltaQuat.w );
+	OutParticleOmega += DeltaQuat.xyz * QuatAngle / (AxisLength*InDeltaTime);
 }
 
 /* -----------------------------------------------------------------
@@ -178,6 +181,8 @@ void SolveHookeSpringMaterial(const float InMaterialCompliance, const float InMa
 {
 	const float SchurDiagonal = ( SharedInverseMass[GGroupThreadId.x] + SharedInverseMass[GGroupThreadId.x-1] ) / (InRestLength*InRestLength) + 
 					SharedInverseInertia[GGroupThreadId.x-1]*4.0 + InMaterialCompliance;
+	//const float SchurDiagonal = ( SharedInverseMass[GGroupThreadId.x] + SharedInverseMass[GGroupThreadId.x-1] ) / (InRestLength*InRestLength) + 
+	//				 + InMaterialCompliance;
 	OutMaterialWeight = ( SchurDiagonal != 0.0 ) ? 1.0 / SchurDiagonal : 0.0;
 }
 
@@ -240,8 +245,162 @@ void SolveStretchShearMaterial(const float InMaterialCompliance, const float InM
 	GroupMemoryBarrier(); 
 } 
 
+
+
+/* -----------------------------------------------------------------
+ * Angular Spring material
+ * -----------------------------------------------------------------
+ */
+
+ void ComputeAngularSpringWeight(const float InMaterialCompliance, const float InRestLength, out float OutMaterialWeight)
+{
+	const float SchurDiagonal = SharedInverseInertia[GGroupThreadId.x-1]*4.0 + InMaterialCompliance;
+	OutMaterialWeight = ( SchurDiagonal != 0.0 ) ? 1.0 / SchurDiagonal : 0.0;
+}
+
+void EvalAngularSpringCompliance(const float InDeltaTime, const float InYoungModulus, const float InRodThickness, const float InRestLength, out float OutMaterialCompliance)
+{
+	// Compliance = 1.0 / (k * dt * dt)
+	// with k = L * L * (Y * A / L) (the L*L term before is coming from the fact that our constraint is dL/L and not dL)
+	// A is the cross section area (Pi * R * R), L is the rest length and Y is the young modulus
+	OutMaterialCompliance = 4.0/(InYoungModulus*PI*InRodThickness*InRestLength*InRodThickness*InDeltaTime*InDeltaTime);
+}
+
+void ResetAngularSpringMultiplier(out float3 OutMaterialMultiplier)
+{
+	OutMaterialMultiplier = float3(0,0,0);
+}
+
+void UpdateAngularSpringMultiplier(const float InMaterialCompliance, const float InMaterialWeight, const float InRestLength,  inout float3 OutMaterialMultiplier)
+{
+	float3 EdgeDirection;
+	float4 q0 = SharedLocalQuaternion[GGroupThreadId.x-1];
+	EdgeDirection[0] = 2.0 * (q0.x * q0.z + q0.w * q0.y);
+	EdgeDirection[1] = 2.0 * (q0.y * q0.z - q0.w * q0.x);
+	EdgeDirection[2] = q0.w * q0.w - q0.x * q0.x - q0.y * q0.y + q0.z * q0.z;
+
+	// XPBD lagrange multiplier update : dL = -(C+compliance*L) / (dC * invM * dCt + alpha)
+	const float3 DeltaLambda = -((SharedLocalPosition[GGroupThreadId.x] - SharedLocalPosition[GGroupThreadId.x-1]) / InRestLength - EdgeDirection + 
+					OutMaterialMultiplier * InMaterialCompliance) * InMaterialWeight;
+
+	// L += dL
+	OutMaterialMultiplier += DeltaLambda;
+
+	const float4 qebar = float4(-q0.y, q0.x, -q0.w, q0.z);	
+	const float4 QuaternionDelta = -2.0 * float4( DeltaLambda.xyz * qebar.w + cross(DeltaLambda.xyz, qebar.xyz),
+					- dot(DeltaLambda.xyz, qebar.xyz) );
+
+	q0 += QuaternionDelta * SharedInverseInertia[GGroupThreadId.x-1];
+	SharedLocalQuaternion[GGroupThreadId.x-1] = normalize(q0);
+}
+
+void SolveAngularSpringMaterial(const float InMaterialCompliance, const float InMaterialWeight, const float InRestLength, const bool InIsRed, inout float3 OutMaterialMultiplier)
+{
+	// Process all the red rods
+	if (InIsRed)
+	{	
+		UpdateAngularSpringMultiplier(InMaterialCompliance,InMaterialWeight,InRestLength,OutMaterialMultiplier);
+	}
+	// Process all the black rods
+	GroupMemoryBarrier(); 
+	if (!InIsRed)
+	{
+		UpdateAngularSpringMultiplier(InMaterialCompliance,InMaterialWeight,InRestLength,OutMaterialMultiplier);
+	}
+	GroupMemoryBarrier(); 
+} 
+
 /* -----------------------------------------------------------------
- * Bbend/twist rod material
+ * Bend spring rod material
+ * -----------------------------------------------------------------
+ */
+
+void UpdateBendSpringMultiplier(const float InMaterialCompliance, const float InMaterialWeight, const float InRestLength, const float3 InRestDirection, inout float3 OutMaterialMultiplier, const bool IsFirst)
+{
+	float4 ParentFrame = (IsFirst) ? SharedLocalQuaternion[GGroupThreadId.x-1] : SharedLocalQuaternion[GGroupThreadId.x-2];
+	float3 CrossProduct = 2.0 * cross(ParentFrame.xyz,InRestDirection);
+	float3 EdgeDirection = InRestDirection + ParentFrame.w * CrossProduct + cross(ParentFrame.xyz,CrossProduct);
+
+	// XPBD lagrange multiplier update : dL = -(C+compliance*L) / (dC * invM * dCt + alpha)
+	const float3 DeltaLambda = -((SharedLocalPosition[GGroupThreadId.x] - SharedLocalPosition[GGroupThreadId.x-1]) / InRestLength - EdgeDirection / InRestLength + 
+					OutMaterialMultiplier * InMaterialCompliance) * InMaterialWeight;
+
+	// L += dL
+	OutMaterialMultiplier += DeltaLambda;
+
+	// XPBD position update : dX = dL * dCt * invM
+	const float3 PositionDelta  = DeltaLambda/InRestLength;
+
+	// dX += dX
+	SharedLocalPosition[GGroupThreadId.x] += PositionDelta * SharedInverseMass[GGroupThreadId.x];
+	SharedLocalPosition[GGroupThreadId.x-1] -= PositionDelta * SharedInverseMass[GGroupThreadId.x-1];
+}
+
+void SolveBendSpringMaterial(const float InMaterialCompliance, const float InMaterialWeight, const float InRestLength, const float3 InRestDirection, const bool InIsRed, inout float3 OutMaterialMultiplier, const int InStrandSize)
+{
+	const bool IsFirst = (GGroupThreadId.x % InStrandSize) == 1;
+	// Process all the red rods
+	if (InIsRed)
+	{	
+		UpdateBendSpringMultiplier(InMaterialCompliance,InMaterialWeight,InRestLength,InRestDirection,OutMaterialMultiplier,IsFirst);
+	}
+	// Process all the black rods
+	GroupMemoryBarrier(); 
+	if (!InIsRed)
+	{
+		UpdateBendSpringMultiplier(InMaterialCompliance,InMaterialWeight,InRestLength,InRestDirection,OutMaterialMultiplier,IsFirst);
+	}
+	GroupMemoryBarrier(); 
+} 
+
+
+/* -----------------------------------------------------------------
+ * Parallel transport
+ * -----------------------------------------------------------------
+ */
+
+float4 MultiplyQuat_tmp(in float4 QuatA, in float4 QuatB)
+{
+	return float4(
+        QuatB.xyz * QuatA.w + QuatA.xyz * QuatB.w + cross(QuatA.xyz, QuatB.xyz),
+        QuatA.w * QuatB.w - dot(QuatA.xyz, QuatB.xyz));
+}
+
+float4 FindQuatBetweenHelper_tmp(in float3 A, in float3 B, in float NormAB)
+{
+	float W = NormAB + dot(A,B);
+	float4 Quat = (W>1e-6*NormAB) ? float4(A.y*B.z-A.z*B.y,A.z*B.x-A.x*B.z,A.x*B.y-A.y*B.x,W) :
+			(abs(A.x) > abs(A.y)) ? float4(-A.z,0.0,A.x,0.0) : float4(0.0,-A.z,A.y,0.0);
+	return normalize(Quat);
+}
+
+ void ComputeBishopFrame(const int StrandSize)
+{
+	const bool IsMobile = (GGroupThreadId.x % StrandSize) != 0;
+	if (!IsMobile)
+	{
+		float4 NodeQuaternion = float4(0,0,0,1);//SharedLocalQuaternion[GGroupThreadId.x];
+		float3 TangentNext = float3(0,0,1);
+		float3 TangentPrev = TangentNext;
+
+		float3 PositionNext = SharedLocalPosition[GGroupThreadId.x];
+		float3 PositionPrev = PositionNext;
+
+		for (int EdgeIndex = 0, EdgeEnd = StrandSize-1; EdgeIndex < EdgeEnd; ++EdgeIndex)
+		{
+			PositionPrev = PositionNext;
+			PositionNext = SharedLocalPosition[GGroupThreadId.x+EdgeIndex+1];
+
+			TangentPrev = TangentNext;
+			TangentNext = normalize(PositionNext - PositionPrev);
+			NodeQuaternion = normalize( MultiplyQuat_tmp( FindQuatBetweenHelper_tmp(TangentPrev,TangentNext,1.0), NodeQuaternion ) );
+			SharedLocalQuaternion[GGroupThreadId.x+EdgeIndex] = NodeQuaternion;
+		}
+	}
+	GroupMemoryBarrier(); 
+}
+/* -----------------------------------------------------------------
+ * Bend/twist rod material
  * -----------------------------------------------------------------
  */
 
@@ -276,7 +435,7 @@ void UpdateBendTwistMultiplier(const float InMaterialCompliance, const float InM
 	const float4 OmegaPlus = (CurrentDarboux + InRestDarboux) * InverseRestLength;
 	const float4 OmegaMinus = (CurrentDarboux - InRestDarboux) * InverseRestLength;
 
-	float4 DeltaOmega = (dot(OmegaPlus.xyz,OmegaPlus.xyz) > dot(OmegaMinus.xyz,OmegaMinus.xyz) ) ? OmegaMinus : OmegaPlus;
+	float4 DeltaOmega = (dot(OmegaPlus,OmegaPlus) > dot(OmegaMinus,OmegaMinus) ) ? OmegaMinus : OmegaPlus;
 	DeltaOmega.xyz += OutMaterialMultiplier * InMaterialCompliance;  
 	DeltaOmega.w = 0;
 	DeltaOmega *= InMaterialWeight;
@@ -292,6 +451,40 @@ void UpdateBendTwistMultiplier(const float InMaterialCompliance, const float InM
 	SharedLocalQuaternion[GGroupThreadId.x] = normalize(SharedLocalQuaternion[GGroupThreadId.x]);
 }
 
+/*void UpdateBendTwistMultiplier(const float InMaterialCompliance, const float InMaterialWeight, const float InRestLength, const float4 InRestDarboux, inout float3 OutMaterialMultiplier)
+{
+	float4 q0 = SharedLocalQuaternion[GGroupThreadId.x-1];
+	float4 q1 = SharedLocalQuaternion[GGroupThreadId.x];
+	const float4 CurrentDarboux = float4(
+			q1.xyz * q0.w - q0.xyz * q1.w + cross(-q0.xyz, q1.xyz),
+			q0.w * q1.w - dot(-q0.xyz, q1.xyz));
+	const float InverseRestLength = 2.0 / InRestLength;
+
+	float3 DeltaOmega = (CurrentDarboux.xyz / CurrentDarboux.w - InRestDarboux.xyz / InRestDarboux.w) * InverseRestLength;
+	DeltaOmega.xyz += OutMaterialMultiplier * InMaterialCompliance;  
+	DeltaOmega *= InMaterialWeight;
+
+	OutMaterialMultiplier -= DeltaOmega.xyz;
+
+	DeltaOmega = DeltaOmega - CurrentDarboux.xyz * dot(CurrentDarboux.xyz,DeltaOmega); 
+
+	const float qtu = dot(q0,q1);
+	const float dto = dot(CurrentDarboux.xyz, DeltaOmega.xyz);
+
+	const float4 q1tmp = qtu * q1; 
+	const float4 q0tmp = qtu * q0; 
+
+	const float afac = dot(CurrentDarboux.xyz,DeltaOmega);
+
+	SharedLocalQuaternion[GGroupThreadId.x-1] += SharedInverseInertia[GGroupThreadId.x-1] * ( float4(
+			DeltaOmega.xyz * q1tmp.w + cross(q1tmp.xyz, DeltaOmega.xyz), - dot(q1tmp.xyz, DeltaOmega.xyz)) + dto * q1) * InverseRestLength;
+	SharedLocalQuaternion[GGroupThreadId.x] -= SharedInverseInertia[GGroupThreadId.x] *  ( float4(
+			DeltaOmega.xyz * q0tmp.w + cross(q0tmp.xyz, DeltaOmega.xyz), - dot(q0tmp.xyz, DeltaOmega.xyz)) - dto * q0) * InverseRestLength;
+
+	SharedLocalQuaternion[GGroupThreadId.x-1] = normalize(SharedLocalQuaternion[GGroupThreadId.x-1]);
+	SharedLocalQuaternion[GGroupThreadId.x] = normalize(SharedLocalQuaternion[GGroupThreadId.x]);
+}*/
+
 void SolveBendTwistMaterial(const float InMaterialCompliance, const float InMaterialWeight, const float InRestLength, const float4 InRestDarboux, const bool InIsRed, inout float3 OutMaterialMultiplier)
 {
 	// Process all the red rods
@@ -376,14 +569,16 @@ void SolveBendTwistMaterial(const float InMaterialCompliance, const float InMate
  */
 
  float3 ProcessCollisionData( const float3 InDeltaPosition, const float3 InDeltaVelocity, const float3 InCollisionNormal, const float InThicknessSum,
-							const float InStaticFriction, const float InKineticFriction )
+							const float InStaticFriction, const float InKineticFriction, const bool IsProjection )
  {
 	const float DeltaVNormal = dot(InDeltaVelocity,InCollisionNormal);
 	const float DeltaPNormal = dot(InDeltaPosition,InCollisionNormal)-InThicknessSum;
-	const float NormDeltaPN = abs(DeltaPNormal);
 
 	if (DeltaPNormal < 0.0)
 	{
+		const float DeltaFNormal = IsProjection ? DeltaPNormal : DeltaVNormal;
+		const float NormDeltaFN = abs(DeltaFNormal);
+
 		const float3 CollisionTangent = InDeltaVelocity - DeltaVNormal * InCollisionNormal;
 		const float TangentLength = length(CollisionTangent);
 
@@ -405,10 +600,10 @@ void SolveBendTwistMaterial(const float InMaterialCompliance, const float InMate
 		const float NormDeltaVTA = abs(DeltaVTangentA);
 		const float NormDeltaVTB = abs(DeltaVTangentB);
 
-		const float AlphaTangentA = ( NormDeltaVTA < InStaticFriction * NormDeltaPN ) ? 1.0 : (InKineticFriction > 0.0) ? min(InKineticFriction*NormDeltaPN/NormDeltaVTA, 1.0) : 0.0;
-		const float AlphaTangentB = ( NormDeltaVTB < InStaticFriction * NormDeltaPN ) ? 1.0 : (InKineticFriction > 0.0) ? min(InKineticFriction*NormDeltaPN/NormDeltaVTB, 1.0) : 0.0;
+		const float AlphaTangentA = ( NormDeltaVTA < InStaticFriction * NormDeltaFN ) ? 1.0 : (InKineticFriction > 0.0) ? min(InKineticFriction*NormDeltaFN/NormDeltaVTA, 1.0) : 0.0;
+		const float AlphaTangentB = ( NormDeltaVTB < InStaticFriction * NormDeltaFN ) ? 1.0 : (InKineticFriction > 0.0) ? min(InKineticFriction*NormDeltaFN/NormDeltaVTB, 1.0) : 0.0;
 
-		return DeltaPNormal * InCollisionNormal + AlphaTangentA * CollisionTangentA * DeltaVTangentA + AlphaTangentB * CollisionTangentB * DeltaVTangentB;
+		return DeltaFNormal * InCollisionNormal + AlphaTangentA * CollisionTangentA * DeltaVTangentA + AlphaTangentB * CollisionTangentB * DeltaVTangentB;
 	}
 	return float3(0,0,0);
  }
@@ -418,7 +613,6 @@ void SolveBendTwistMaterial(const float InMaterialCompliance, const float InMate
  * -----------------------------------------------------------------
  */
 
-
 void DetectSphereVolumeCollision(const float3 InParticlePosition, const float3 InSphereCenter, const float InSphereRadius, 
 								 out float3 OutCollisionPosition, out float3 OutCollisionVelocity, inout float3 OutCollisionNormal)   
 {
@@ -439,7 +633,7 @@ void DetectSphereVolumeCollision(const float3 InParticlePosition, const float3 I
  */
 
 void SolveStaticCollisionConstraint(const float InParticleRadius, const bool InIsMobile, inout float3 OutParticlePosition, inout float3 OutParticleVelocity, const float3 InCollisionPosition, const float3 InCollisionVelocity, const float3 InCollisionNormal, 
-				const float InStaticFriction, const float InKineticFriction, const float InDeltaTime )
+				const float InStaticFriction, const float InKineticFriction, const float InDeltaTime, const bool IsProjection )
 {
 	if(InIsMobile)
 	{
@@ -447,7 +641,7 @@ void SolveStaticCollisionConstraint(const float InParticleRadius, const bool InI
 		const float3 DeltaVelocity = (OutParticleVelocity - InCollisionVelocity) * InDeltaTime;
 		const float ThicknessSum = InParticleRadius;
 
-		const float3 PositionDisplacement = ProcessCollisionData(DeltaPosition,DeltaVelocity,InCollisionNormal,ThicknessSum,InStaticFriction,InKineticFriction);
+		const float3 PositionDisplacement = ProcessCollisionData(DeltaPosition,DeltaVelocity,InCollisionNormal,ThicknessSum,InStaticFriction,InKineticFriction,IsProjection);
 		OutParticlePosition -= PositionDisplacement;
 		OutParticleVelocity -= PositionDisplacement / InDeltaTime;
 	}
diff --git a/Engine/Plugins/FX/Niagara/Source/Niagara/Classes/NiagaraDataInterfaceGrid2DCollection.h b/Engine/Plugins/FX/Niagara/Source/Niagara/Classes/NiagaraDataInterfaceGrid2DCollection.h
--- a/Engine/Plugins/FX/Niagara/Source/Niagara/Classes/NiagaraDataInterfaceGrid2DCollection.h
+++ b/Engine/Plugins/FX/Niagara/Source/Niagara/Classes/NiagaraDataInterfaceGrid2DCollection.h
@@ -16,10 +16,10 @@ class Grid2DBuffer
 public:
 	Grid2DBuffer(int NumX, int NumY)
 	{		
-		GridBuffer.Initialize(NumX, NumY);
+		GridBuffer.Initialize(16, NumX, NumY, EPixelFormat::PF_A32B32G32R32F);
 	}
 
-	FTextureRWBuffer GridBuffer;	
+	FTextureRWBuffer2D GridBuffer;	
 };
 
 class Grid2DCollectionRWInstanceData
diff --git a/Engine/Plugins/FX/Niagara/Source/Niagara/Private/NiagaraDataInterfaceSkeletalMesh.cpp b/Engine/Plugins/FX/Niagara/Source/Niagara/Private/NiagaraDataInterfaceSkeletalMesh.cpp
--- a/Engine/Plugins/FX/Niagara/Source/Niagara/Private/NiagaraDataInterfaceSkeletalMesh.cpp
+++ b/Engine/Plugins/FX/Niagara/Source/Niagara/Private/NiagaraDataInterfaceSkeletalMesh.cpp
@@ -344,7 +344,7 @@ FSkeletalMeshGpuSpawnStaticBuffers::~FSkeletalMeshGpuSpawnStaticBuffers()
 void FSkeletalMeshGpuSpawnStaticBuffers::Initialise(FNDISkeletalMesh_InstanceData* InstData, const FSkeletalMeshLODRenderData& SkeletalMeshLODRenderData, const FSkeletalMeshSamplingLODBuiltData& MeshSamplingLODBuiltData)
 {
 	if (!InstData)
-	{
+    {
 		SkeletalMeshSamplingLODBuiltData = nullptr;
 		bUseGpuUniformlyDistributedSampling = false;
 
@@ -358,38 +358,40 @@ void FSkeletalMeshGpuSpawnStaticBuffers::Initialise(FNDISkeletalMesh_InstanceDat
 		SpecificSocketBoneOffset = 0;
 		return;
 	}
+	else
+	{
+		SkeletalMeshSamplingLODBuiltData = &MeshSamplingLODBuiltData;
+		bUseGpuUniformlyDistributedSampling = InstData->bIsGpuUniformlyDistributedSampling;
 
-	SkeletalMeshSamplingLODBuiltData = &MeshSamplingLODBuiltData;
-	bUseGpuUniformlyDistributedSampling = InstData->bIsGpuUniformlyDistributedSampling;
+		LODRenderData = &SkeletalMeshLODRenderData;
+		TriangleCount = SkeletalMeshLODRenderData.MultiSizeIndexContainer.GetIndexBuffer()->Num() / 3;
+		VertexCount = SkeletalMeshLODRenderData.GetNumVertices();
 
-	LODRenderData = &SkeletalMeshLODRenderData;
-	TriangleCount = SkeletalMeshLODRenderData.MultiSizeIndexContainer.GetIndexBuffer()->Num() / 3;
-	VertexCount = SkeletalMeshLODRenderData.GetNumVertices();
-	
-	if (TriangleCount == 0)
-	{
-		UE_LOG(LogNiagara, Warning, TEXT("FSkeletalMeshGpuSpawnStaticBuffers> Triangle count is invalid %d"), TriangleCount, (InstData && InstData->Mesh) ? *InstData->Mesh->GetFullName() : TEXT("Unknown Mesh"));
-	}
-	if (VertexCount == 0)
-	{
-		UE_LOG(LogNiagara, Warning, TEXT("FSkeletalMeshGpuSpawnStaticBuffers> Vertex count is invalid %d"), VertexCount, (InstData && InstData->Mesh) ? *InstData->Mesh->GetFullName() : TEXT("Unknown Mesh"));
-	}
+		if (TriangleCount == 0)
+		{
+			UE_LOG(LogNiagara, Warning, TEXT("FSkeletalMeshGpuSpawnStaticBuffers> Triangle count is invalid %d"), TriangleCount, (InstData && InstData->Mesh) ? *InstData->Mesh->GetFullName() : TEXT("Unknown Mesh"));
+		}
+		if (VertexCount == 0)
+		{
+			UE_LOG(LogNiagara, Warning, TEXT("FSkeletalMeshGpuSpawnStaticBuffers> Vertex count is invalid %d"), VertexCount, (InstData && InstData->Mesh) ? *InstData->Mesh->GetFullName() : TEXT("Unknown Mesh"));
+		}
 
-	// Copy Specific Bones / Socket data into arrays that the renderer will use to create read buffers
-	//-TODO: Exclude setting up these arrays if we don't sample from them
-	NumSpecificBones = InstData->SpecificBones.Num();
-	if (NumSpecificBones > 0)
-	{
-		SpecificBonesArray.Reserve(NumSpecificBones);
-		for ( int32 v : InstData->SpecificBones )
+		// Copy Specific Bones / Socket data into arrays that the renderer will use to create read buffers
+		//-TODO: Exclude setting up these arrays if we don't sample from them
+		NumSpecificBones = InstData->SpecificBones.Num();
+		if (NumSpecificBones > 0)
 		{
-			check(v <= 65535);
-			SpecificBonesArray.Add(v);
+			SpecificBonesArray.Reserve(NumSpecificBones);
+			for (int32 v : InstData->SpecificBones)
+			{
+				check(v <= 65535);
+				SpecificBonesArray.Add(v);
+			}
 		}
-	}
 
-	NumSpecificSockets = InstData->SpecificSockets.Num();
-	SpecificSocketBoneOffset = InstData->SpecificSocketBoneOffset;
+		NumSpecificSockets = InstData->SpecificSockets.Num();
+		SpecificSocketBoneOffset = InstData->SpecificSocketBoneOffset;
+	}
 }
 
 void FSkeletalMeshGpuSpawnStaticBuffers::InitRHI()
diff --git a/Engine/Shaders/Private/AmbientCubemapComposite.usf b/Engine/Shaders/Private/AmbientCubemapComposite.usf
--- a/Engine/Shaders/Private/AmbientCubemapComposite.usf
+++ b/Engine/Shaders/Private/AmbientCubemapComposite.usf
@@ -1,5 +1,9 @@
 // Copyright 1998-2019 Epic Games, Inc. All Rights Reserved.
 
+/*=============================================================================
+	PostprocessAmbient.usf: To apply a ambient cubemap as a postprocess
+=============================================================================*/
+
 #include "Common.ush"
 #include "PostProcessCommon.ush"
 #include "DeferredShadingCommon.ush"
@@ -155,7 +159,7 @@ float3 ImageBasedLightingMIS( FGBufferData GBuffer, float3 V, float3 N, uint2 Ra
 				}
 				float Weight = rcp( InvWeight );
 
-				FShadowTerms Shadow = { 1, 1, 1 };
+				FShadowTerms Shadow = { 1, 1, 1, InitHairTransmittanceData() };
 				FDirectLighting LightingSample = EvaluateBxDF( GBuffer, N, V, L, NoL, Shadow );
 
 				Lighting += SampleColor * Weight * ( LightingSample.Diffuse + LightingSample.Specular + LightingSample.Transmission );
@@ -295,6 +299,8 @@ float3 ImageBasedLightingHair( FGBufferData GBuffer, float3 V, float3 N, uint2 R
 
 	uint NumSamples = 32;
 	
+	FHairTransmittanceData HairTransmittance = InitHairTransmittanceData();
+
 	LOOP
 	for( uint i = 0; i < NumSamples; i++ )
 	{
@@ -311,7 +317,7 @@ float3 ImageBasedLightingHair( FGBufferData GBuffer, float3 V, float3 N, uint2 R
 
 			float3 Shading = 0;
 			bool bEvalMultiScatter = true;
-			Shading = HairShading( GBuffer, L, V, N, 1, 0, 0, Random, bEvalMultiScatter );
+			Shading = HairShading( GBuffer, L, V, N, 1, HairTransmittance, 0, 0, Random, bEvalMultiScatter );
 			
 			Lighting += SampleColor * Shading * Weight;
 		}
@@ -422,11 +428,12 @@ void MainPS(in noperspective float4 UVAndScreenPos : TEXCOORD0, float4 SvPositio
 	
 	BRANCH if( GBuffer.ShadingModelID == SHADINGMODELID_HAIR )
 	{
+		FHairTransmittanceData HairTransmittance = InitHairTransmittanceData();
 		float3 FakeNormal = normalize( V - N * dot(V,N) );
 		FakeNormal = normalize( FakeNormal + 0 * float3(0,0,1) );
 		SpecularContribution = TextureCubeSampleLevel(AmbientCubemap, AmbientCubemapSampler, FakeNormal, AbsoluteDiffuseMip).rgb;
 		bool bEvalMultiScatter = true;
-		SpecularContribution *= PI * HairShading( GBuffer, FakeNormal, V, N, 1, 0, 0.2, Random, bEvalMultiScatter );
+		SpecularContribution *= PI * HairShading( GBuffer, FakeNormal, V, N, 1, HairTransmittance, 0, 0.2, Random, bEvalMultiScatter);
 		NonSpecularContribution = 0;
 	}
 
diff --git a/Engine/Shaders/Private/BasePassPixelShader.usf b/Engine/Shaders/Private/BasePassPixelShader.usf
--- a/Engine/Shaders/Private/BasePassPixelShader.usf
+++ b/Engine/Shaders/Private/BasePassPixelShader.usf
@@ -539,7 +539,7 @@ float3 GetSimpleForwardLightingDirectionalLight(FGBufferData GBuffer, float3 Dif
 
 	float3 LightColor = ResolvedView.DirectionalLightColor.rgb * PI;
 	
-	FShadowTerms Shadow = { 1, 1, 1 };
+	FShadowTerms Shadow = { 1, 1, 1, InitHairTransmittanceData() };
 	FDirectLighting Lighting = EvaluateBxDF( GBuffer, N, V, L, NoL, Shadow );
 
 	// Not computing specular, material was forced fully rough
@@ -1040,12 +1040,13 @@ ShadingModel);
 		#if MATERIAL_SHADINGMODEL_HAIR
 		if (GBuffer.ShadingModelID == SHADINGMODELID_HAIR)
 		{
+			FHairTransmittanceData TransmittanceData = InitHairTransmittanceData();
 			float3 N = MaterialParameters.WorldNormal;
 			float3 V = MaterialParameters.CameraVector;
 			float3 L = normalize( V - N * dot(V,N) );
 			DiffuseDir = L;
 			bool bEvalMultiScatter = true;
-			DiffuseColorForIndirect = 2*PI * HairShading( GBuffer, L, V, N, 1, 0, 0.2, uint2(0,0), bEvalMultiScatter );
+			DiffuseColorForIndirect = 2*PI * HairShading( GBuffer, L, V, N, 1, TransmittanceData, 0, 0.2, uint2(0,0), bEvalMultiScatter);
 		}
 		#endif
 
diff --git a/Engine/Shaders/Private/DeferredLightPixelShaders.usf b/Engine/Shaders/Private/DeferredLightPixelShaders.usf
--- a/Engine/Shaders/Private/DeferredLightPixelShaders.usf
+++ b/Engine/Shaders/Private/DeferredLightPixelShaders.usf
@@ -7,9 +7,53 @@
 #define SUPPORT_CONTACT_SHADOWS 1
 
 #include "Common.ush"
+#if USE_HAIR_LIGHTING
+#include "HairStrands/HairStrandsCommon.ush"
+#endif
+
 #include "DeferredShadingCommon.ush"
 #include "DeferredLightingCommon.ush"
 
+#if USE_HAIR_LIGHTING
+	#include "HairStrands/HairStrandsDeepShadowCommon.ush"
+	#include "HairStrands/HairStrandsDeepTransmittanceCommon.ush"
+	#include "HairStrands/HairStrandsDeepTransmittanceDualScattering.ush"
+	#include "HairStrands/HairStrandsVisibilityCommon.ush"
+	#if HAIR_CUSTOM_BSDF == 1
+	#include "HairBsdf.ush"
+	#endif
+#endif
+
+struct FInputParams
+{
+	float2 PixelPos;
+	float4 ScreenPosition;
+	float2 ScreenUV;
+	float3 ScreenVector;
+};
+
+struct FDerivedParams
+{
+	float3 CameraVector;
+	float3 WorldPosition;
+};
+
+FDerivedParams GetDerivedParams(in FInputParams Input, in float SceneDepth)
+{
+	FDerivedParams Out;
+#if LIGHT_SOURCE_SHAPE > 0
+	// With a perspective projection, the clip space position is NDC * Clip.w
+	// With an orthographic projection, clip space is the same as NDC
+	float2 ClipPosition = Input.ScreenPosition.xy / Input.ScreenPosition.w * (View.ViewToClip[3][3] < 1.0f ? SceneDepth : 1.0f);
+	Out.WorldPosition = mul(float4(ClipPosition, SceneDepth, 1), View.ScreenToWorld).xyz;
+	Out.CameraVector = normalize(Out.WorldPosition - View.WorldCameraOrigin);
+#else
+	Out.WorldPosition = Input.ScreenVector * SceneDepth + View.WorldCameraOrigin;
+	Out.CameraVector = normalize(Input.ScreenVector);
+#endif
+	return Out;
+}
+
 FDeferredLightData SetupLightDataForStandardDeferred()
 {
 	// Build the light data struct using the DeferredLightUniforms and light defines
@@ -41,6 +85,7 @@ FDeferredLightData SetupLightDataForStandardDeferred()
 	LightData.RectLightBarnCosAngle = DeferredLightUniforms.RectLightBarnCosAngle;
 	LightData.RectLightBarnLength = DeferredLightUniforms.RectLightBarnLength;
 
+	LightData.HairTransmittance = InitHairTransmittanceData();
 	return LightData;
 }
 
@@ -52,6 +97,144 @@ uint GetLightingChannelMask(float2 UV)
 	return LightingChannelsTexture.Load(uint3(IntegerUV, 0)).x;
 }
 
+#if USE_HAIR_LIGHTING
+
+Texture2D<uint> HairVisibilityNodeOffsetAndCount;
+StructuredBuffer<FNodeData> HairVisibilityNodeData;
+
+Texture2D<float4> ScreenShadowMaskSubPixelTexture;
+
+uint HairTransmittanceBufferMaxCount;
+StructuredBuffer<float4> HairTransmittanceBuffer;
+Texture3D<float4> HairLUTTexture;
+SamplerState HairLUTSampler;
+uint HairTransmittanceEnable;
+float4 HairComponents; // R, TT, TRT, LocaScattering
+
+float Luma(float3 Color)
+{
+	return 0.25f * ((Color.g * 2.0) + (Color.r + Color.b));
+}
+
+#define RESOLVE_IN_PERCEPTUAL_SPACE 0
+	
+float4 LinearToPerceptual(float4 Color)
+{
+#if RESOLVE_IN_PERCEPTUAL_SPACE
+	return float4(Color.xyz / (1 + Luma(Color.xyz)), Color.a);
+#else
+	return Color;
+#endif
+}
+float4 PerceptualToLinear(float4 Color) 
+{
+#if RESOLVE_IN_PERCEPTUAL_SPACE
+	return float4(Color.xyz / (1 - Luma(Color.xyz)), Color.a);
+#else
+	return Color;
+#endif
+}
+
+void ComputeHairLighting(
+	in FInputParams InputParams,
+	in FNodeDesc NodeDesc, 
+	out float OutOpaqueWeight,
+	out float4 OutColor, 
+	out float4 OutColorSubPixel)
+{
+	float Coverage = 0;
+	float4 LocalColor = 0;
+	LOOP
+	for (uint SampleIt = 0; SampleIt < NodeDesc.Count; SampleIt++)
+	{
+		const uint LocalOffset = NodeDesc.Offset + SampleIt;
+		const FNodeData NodeData = HairVisibilityNodeData[LocalOffset];
+		float			SampleDepth;
+		FVisibilityData SampleVisData;
+		float3			SampleBaseColor;
+		float			SampleRoughness;
+		float			SampleSpecular;
+		DecodeNodeData(NodeData, SampleDepth, SampleVisData, SampleBaseColor, SampleRoughness, SampleSpecular);
+
+		// Inject material data from the visibility/mini-gbuffer for hair
+		const float SceneDepth = ConvertFromDeviceZ(SampleDepth);
+		FScreenSpaceData HairScreenSpaceData;
+		HairScreenSpaceData.AmbientOcclusion = 1;
+		HairScreenSpaceData.GBuffer.WorldNormal = SampleVisData.Tangent;
+		HairScreenSpaceData.GBuffer.BaseColor = SampleBaseColor;
+		HairScreenSpaceData.GBuffer.Roughness = SampleRoughness;
+		HairScreenSpaceData.GBuffer.ShadingModelID = SHADINGMODELID_HAIR;
+		HairScreenSpaceData.GBuffer.DiffuseColor = 0;
+		HairScreenSpaceData.GBuffer.SpecularColor = 0;
+		HairScreenSpaceData.GBuffer.Specular = SampleSpecular;
+		HairScreenSpaceData.GBuffer.Metallic = 0;
+		HairScreenSpaceData.GBuffer.Depth = SampleDepth;
+		HairScreenSpaceData.GBuffer.GBufferAO = 1;
+		HairScreenSpaceData.GBuffer.CustomData = 0;
+		HairScreenSpaceData.GBuffer.IndirectIrradiance = 0;
+		HairScreenSpaceData.GBuffer.PrecomputedShadowFactors = 0;
+
+		const FDerivedParams DerivedParams = GetDerivedParams(InputParams, SceneDepth);
+		FDeferredLightData LightData = SetupLightDataForStandardDeferred();
+		LightData.HairTransmittance.ScatteringComponent = HairComponents.xyz;
+
+		// Compute Hair transmittance
+		BRANCH if (HairTransmittanceEnable > 0)
+		{
+			const float3 LightPosition = LightData.Position;
+			float3 L = float3(0, 0, 1);
+			if (LightData.bRadialLight)
+				L = normalize(LightPosition - DerivedParams.WorldPosition);
+			else
+				L = LightData.Direction;
+
+			const float3 V = normalize(-DerivedParams.CameraVector);
+			LightData.HairTransmittance = GetHairTransmittance(
+				V, 
+				L,
+				HairScreenSpaceData.GBuffer, 
+				LocalOffset,
+				HairTransmittanceBufferMaxCount,
+				HairTransmittanceBuffer,
+				HairLUTTexture,
+				HairLUTSampler,
+				HairTransmittanceEnable,
+				HairComponents);
+		}
+
+		const float Dither = InterleavedGradientNoise(InputParams.PixelPos, View.StateFrameIndexMod8);
+		const uint3 PixelCoord = uint3(floor(InputParams.PixelPos), 0);
+		const float4 LightAttenuation = ScreenShadowMaskSubPixelTexture.Load(PixelCoord);
+
+		FRectTexture RectTexture = InitRectTexture(DeferredLightUniforms.SourceTexture);
+		const float4 Radiance = GetDynamicLighting(DerivedParams.WorldPosition, DerivedParams.CameraVector, HairScreenSpaceData.GBuffer, HairScreenSpaceData.AmbientOcclusion, HairScreenSpaceData.GBuffer.ShadingModelID, LightData, LightAttenuation, Dither, uint2(InputParams.PixelPos), RectTexture);
+		const float  Attenuation = ComputeLightProfileMultiplier(DerivedParams.WorldPosition, DeferredLightUniforms.Position, -DeferredLightUniforms.Direction, DeferredLightUniforms.Tangent);
+
+		const float4 LinearColor = Radiance * Attenuation;
+		const float4 PercepualColor = LinearToPerceptual(LinearColor);
+		const float LocalCoverage = From8bitCoverage(SampleVisData.Coverage8bit);
+		LocalColor += PercepualColor * LocalCoverage;
+		Coverage += LocalCoverage;
+	}
+
+	OutColor = 0;
+	OutColorSubPixel = 0;
+	if (Coverage < 1)
+	{
+		OutColorSubPixel = LocalColor;
+		OutColorSubPixel.a = Coverage;
+		Coverage = 0;
+	}
+	else
+	{
+		OutColor = LocalColor;
+		Coverage = 1;
+	}
+	OutOpaqueWeight = 1 - Coverage;
+}
+
+#endif // USE_HAIR_LIGHTING
+
 
 void DeferredLightPixelMain(
 #if LIGHT_SOURCE_SHAPE > 0
@@ -62,47 +245,77 @@ void DeferredLightPixelMain(
 #endif
 	float4 SVPos			: SV_POSITION,
 	out float4 OutColor		: SV_Target0
+#if USE_HAIR_LIGHTING
+,	out float4 OutColorSubPixel	: SV_Target1
+#endif
 	)
 {
+	const float2 PixelPos = SVPos.xy;
 	OutColor = 0;
 
+	// Convert input data (directional/local light)
+	FInputParams InputParams = (FInputParams)0;
+	InputParams.PixelPos		= SVPos.xy;
 #if LIGHT_SOURCE_SHAPE > 0
-	float2 ScreenUV = InScreenPosition.xy / InScreenPosition.w * View.ScreenPositionScaleBias.xy + View.ScreenPositionScaleBias.wz;
+	InputParams.ScreenPosition	= InScreenPosition;
+	InputParams.ScreenUV		= InScreenPosition.xy / InScreenPosition.w * View.ScreenPositionScaleBias.xy + View.ScreenPositionScaleBias.wz;
+	InputParams.ScreenVector	= 0;
 #else
-	float3 CameraVector = normalize(ScreenVector);
+	InputParams.ScreenPosition	= 0;
+	InputParams.ScreenUV		= ScreenUV;
+	InputParams.ScreenVector	= ScreenVector;
 #endif
+
+#if USE_HAIR_LIGHTING
+	const uint3 PixelCoord = uint3(floor(InputParams.PixelPos), 0);
+	const FNodeDesc NodeDesc = DecodeNodeDesc(HairVisibilityNodeOffsetAndCount.Load(PixelCoord));
 	
-	FScreenSpaceData ScreenSpaceData = GetScreenSpaceData(ScreenUV);
+	// Only light pixels marked as using deferred shading
+	OutColorSubPixel = 0;
+	float OpaqueWeight = 1;
+	BRANCH 
+	if (NodeDesc.Count > 0
+	#if USE_LIGHTING_CHANNELS
+		&& (GetLightingChannelMask(InputParams.ScreenUV) & DeferredLightUniforms.LightingChannelMask)
+	#endif
+		)
+	{
+		ComputeHairLighting(InputParams, NodeDesc, OpaqueWeight, OutColor, OutColorSubPixel);
+	#if USE_PREEXPOSURE
+		OutColorSubPixel.rgb *= View.PreExposure;
+	#endif
+	}
+#endif
+
+	FScreenSpaceData ScreenSpaceData = GetScreenSpaceData(InputParams.ScreenUV);
 
 	// Only light pixels marked as using deferred shading
 	BRANCH if( ScreenSpaceData.GBuffer.ShadingModelID > 0 
 #if USE_LIGHTING_CHANNELS
-		&& (GetLightingChannelMask(ScreenUV) & DeferredLightUniforms.LightingChannelMask)
+		&& (GetLightingChannelMask(InputParams.ScreenUV) & DeferredLightUniforms.LightingChannelMask)
 #endif
 		)
 	{
-		float SceneDepth = CalcSceneDepth(ScreenUV);
-
-#if LIGHT_SOURCE_SHAPE > 0
-		// With a perspective projection, the clip space position is NDC * Clip.w
-		// With an orthographic projection, clip space is the same as NDC
-		float2 ClipPosition = InScreenPosition.xy / InScreenPosition.w * (View.ViewToClip[3][3] < 1.0f ? SceneDepth : 1.0f);
-		float3 WorldPosition = mul(float4(ClipPosition, SceneDepth, 1), View.ScreenToWorld).xyz;
-		float3 CameraVector = normalize(WorldPosition - View.WorldCameraOrigin);
-#else
-		float3 WorldPosition = ScreenVector * SceneDepth + View.WorldCameraOrigin;
-#endif
+		const float SceneDepth = CalcSceneDepth(InputParams.ScreenUV);
+		const FDerivedParams DerivedParams = GetDerivedParams(InputParams, SceneDepth);
 
 		FDeferredLightData LightData = SetupLightDataForStandardDeferred();
 
-		float Dither = InterleavedGradientNoise( SVPos.xy, View.StateFrameIndexMod8 );
+		float Dither = InterleavedGradientNoise(InputParams.PixelPos, View.StateFrameIndexMod8 );
 
 		FRectTexture RectTexture = InitRectTexture(DeferredLightUniforms.SourceTexture);
-		OutColor = GetDynamicLighting(WorldPosition, CameraVector, ScreenSpaceData.GBuffer, ScreenSpaceData.AmbientOcclusion, ScreenSpaceData.GBuffer.ShadingModelID, LightData, GetPerPixelLightAttenuation(ScreenUV), Dither, uint2( SVPos.xy ), RectTexture);
-		OutColor *= ComputeLightProfileMultiplier(WorldPosition, DeferredLightUniforms.Position, -DeferredLightUniforms.Direction, DeferredLightUniforms.Tangent);
+		const float4 Radiance = GetDynamicLighting(DerivedParams.WorldPosition, DerivedParams.CameraVector, ScreenSpaceData.GBuffer, ScreenSpaceData.AmbientOcclusion, ScreenSpaceData.GBuffer.ShadingModelID, LightData, GetPerPixelLightAttenuation(InputParams.ScreenUV), Dither, uint2(InputParams.PixelPos), RectTexture);
+		const float  Attenuation = ComputeLightProfileMultiplier(DerivedParams.WorldPosition, DeferredLightUniforms.Position, -DeferredLightUniforms.Direction, DeferredLightUniforms.Tangent);
 
+#if USE_HAIR_LIGHTING
+		OutColor += LinearToPerceptual(Radiance * Attenuation) * OpaqueWeight;
+		OutColor  = PerceptualToLinear(OutColor);
+#else
+		OutColor += (Radiance * Attenuation);
+#endif
+	}
+		
 #if USE_PREEXPOSURE
 		OutColor.rgb *= View.PreExposure;
 #endif
-	}
-}
\ No newline at end of file
+}
diff --git a/Engine/Shaders/Private/DeferredLightingCommon.ush b/Engine/Shaders/Private/DeferredLightingCommon.ush
--- a/Engine/Shaders/Private/DeferredLightingCommon.ush
+++ b/Engine/Shaders/Private/DeferredLightingCommon.ush
@@ -47,6 +47,8 @@ struct FDeferredLightData
 	uint ShadowedBits;
 	float RectLightBarnCosAngle;
 	float RectLightBarnLength;
+
+	FHairTransmittanceData HairTransmittance;
 };
 
 /** Data about a single light to be shaded with the simple shading model, designed for speed and limited feature set. */
@@ -200,11 +202,21 @@ void GetShadowTerms(FGBufferData GBuffer, FDeferredLightData LightData, float3 W
 
 		FLATTEN
 		if( GBuffer.ShadingModelID == SHADINGMODELID_HAIR ||  GBuffer.ShadingModelID == SHADINGMODELID_EYE  )
-			Shadow.TransmissionShadow *= ContactShadow;
+		{
+			// If hair transmittance is enabled, sharp shadowing should already be handled by 
+			// the dedicated deep shadow maps. Thus no need for contact shadow
+			if (!LightData.HairTransmittance.Enable)
+			{
+				Shadow.TransmissionShadow *= ContactShadow;
+			}
+		}
 		else
 			Shadow.TransmissionShadow *= ContactShadow * 0.5 + 0.5;
 	}
 #endif
+
+	Shadow.HairTransmittance = LightData.HairTransmittance;
+	Shadow.HairTransmittance.OpaqueVisibility = Shadow.SurfaceShadow;
 }
 
 float GetLocalLightAttenuation(
@@ -300,6 +312,8 @@ float4 GetDynamicLighting(float3 WorldPosition, float3 CameraVector, FGBufferDat
 		Shadow.SurfaceShadow = AmbientOcclusion;
 		Shadow.TransmissionShadow = 1;
 		Shadow.TransmissionThickness = 1;
+		Shadow.HairTransmittance.Transmittance = 1;
+		Shadow.HairTransmittance.OpaqueVisibility = 1;
 		GetShadowTerms(GBuffer, LightData, WorldPosition, L, LightAttenuation, Dither, Shadow);
 
 		LightAccumulator.EstimatedCost += 0.3f;		// add the cost of getting the shadow terms
diff --git a/Engine/Shaders/Private/GpuDebugRendering.ush b/Engine/Shaders/Private/GpuDebugRendering.ush
--- a/Engine/Shaders/Private/GpuDebugRendering.ush
+++ b/Engine/Shaders/Private/GpuDebugRendering.ush
@@ -1,8 +1,10 @@
 #pragma once
 
-#define DEBUG_GPU_RENDERING 1
+#ifndef GPU_DEBUG_RENDERING
+#define GPU_DEBUG_RENDERING 0
+#endif
 
-#if DEBUG_GPU_RENDERING
+#if GPU_DEBUG_RENDERING
 
 #define COLOR_RED		0
 #define COLOR_GREEN		1
