Copying //UE4/Dev-Physics to Dev-Main (//UE4/Dev-Main) [at] 9420715

#rnx

diff --git a/Engine/Plugins/Experimental/AlembicHairImporter/Source/AlembicHairTranslator/Private/AlembicHairTranslator.cpp b/Engine/Plugins/Experimental/AlembicHairImporter/Source/AlembicHairTranslator/Private/AlembicHairTranslator.cpp
--- a/Engine/Plugins/Experimental/AlembicHairImporter/Source/AlembicHairTranslator/Private/AlembicHairTranslator.cpp
+++ b/Engine/Plugins/Experimental/AlembicHairImporter/Source/AlembicHairTranslator/Private/AlembicHairTranslator.cpp
@@ -3,6 +3,7 @@
 #include "AlembicHairTranslator.h"
 
 #include "HairDescription.h"
+#include "GroomImportOptions.h"
 #include "Misc/Paths.h"
 
 #if PLATFORM_WINDOWS
@@ -254,6 +255,16 @@ namespace AlembicHairTranslatorUtils
 					ConvertAlembicAttribute<Alembic::AbcGeom::IBoolGeomParam, Alembic::Abc::BoolArraySamplePtr, bool>(HairDescription, StrandID, StartVertexID, NumVertices, Parameters, PropName);
 				}
 				break;
+				case Alembic::Util::kInt8POD:
+				{
+					ConvertAlembicAttribute<Alembic::AbcGeom::ICharGeomParam, Alembic::Abc::CharArraySamplePtr, int>(HairDescription, StrandID, StartVertexID, NumVertices, Parameters, PropName);
+				}
+				break;
+				case Alembic::Util::kInt16POD:
+				{
+					ConvertAlembicAttribute<Alembic::AbcGeom::IInt16GeomParam, Alembic::Abc::Int16ArraySamplePtr, int>(HairDescription, StrandID, StartVertexID, NumVertices, Parameters, PropName);
+				}
+				break;
 				case Alembic::Util::kInt32POD:
 				{
 					ConvertAlembicAttribute<Alembic::AbcGeom::IInt32GeomParam, Alembic::Abc::Int32ArraySamplePtr, int>(HairDescription, StrandID, StartVertexID, NumVertices, Parameters, PropName);
@@ -308,7 +319,7 @@ FMatrix ConvertAlembicMatrix(const Alembic::Abc::M44d& AbcMatrix)
 	return Matrix;
 }
 
-static void ParseObject(const Alembic::Abc::IObject& InObject, FHairDescription& HairDescription, const FMatrix& ParentMatrix, bool bCheckGroomAttributes)
+static void ParseObject(const Alembic::Abc::IObject& InObject, FHairDescription& HairDescription, const FMatrix& ParentMatrix, const FMatrix& ConversionMatrix, float Scale, bool bCheckGroomAttributes)
 {
 	// Get MetaData info from current Alembic Object
 	const Alembic::Abc::MetaData ObjectMetaData = InObject.getMetaData();
@@ -350,6 +361,7 @@ static void ParseObject(const Alembic::Abc::IObject& InObject, FHairDescription&
 			}
 		}
 
+		FMatrix ConvertedMatrix = ParentMatrix * ConversionMatrix;
 		uint32 GlobalIndex = 0;
 		for (uint32 CurveIndex = 0; CurveIndex < NumCurves; ++CurveIndex)
 		{
@@ -366,7 +378,7 @@ static void ParseObject(const Alembic::Abc::IObject& InObject, FHairDescription&
 
 				Alembic::Abc::P3fArraySample::value_type Position = (*Positions)[GlobalIndex];
 
-				FVector ConvertedPosition = ParentMatrix.TransformPosition(FVector(Position.x, Position.y, Position.z)) * AlembicHairFormat::UNIT_TO_CM;
+				FVector ConvertedPosition = ConvertedMatrix.TransformPosition(FVector(Position.x, Position.y, Position.z));
 				SetHairVertexAttribute(HairDescription, VertexID, HairAttribute::Vertex::Position, ConvertedPosition);
 
 				float Width = 0;
@@ -377,7 +389,7 @@ static void ParseObject(const Alembic::Abc::IObject& InObject, FHairDescription&
 					const float CoordU = PointIndex / static_cast<float>(CurveNumVertices - 1);
 					Width = FMath::Lerp(AlembicHairFormat::RootRadius, AlembicHairFormat::TipRadius, CoordU);
 				}
-					break;
+				break;
 				case EAttributeFrequency::CV:
 					Width = (*Widths)[GlobalIndex];
 					break;
@@ -386,7 +398,7 @@ static void ParseObject(const Alembic::Abc::IObject& InObject, FHairDescription&
 				// Per-vertex widths
 				if ((WidthFrequency == EAttributeFrequency::CV || WidthFrequency == EAttributeFrequency::None))
 				{
-					SetHairVertexAttribute(HairDescription, VertexID, HairAttribute::Vertex::Width, Width);
+					SetHairVertexAttribute(HairDescription, VertexID, HairAttribute::Vertex::Width, Width * Scale);
 				}
 			}
 
@@ -405,12 +417,12 @@ static void ParseObject(const Alembic::Abc::IObject& InObject, FHairDescription&
 				if (!StrandWidths.IsValid() && !VertexWidths.IsValid())
 				{
 					const float Width = (*Widths)[CurveIndex];
-					SetHairStrandAttribute(HairDescription, StrandID, HairAttribute::Strand::Width, Width);
+					SetHairStrandAttribute(HairDescription, StrandID, HairAttribute::Strand::Width, Width * Scale);
 				}
 				if (StrandWidths.IsValid() && StrandWidths[StrandID] == 0.f)
 				{
 					const float Width = (*Widths)[CurveIndex];
-					SetHairStrandAttribute(HairDescription, StrandID, HairAttribute::Strand::Width, Width);
+					SetHairStrandAttribute(HairDescription, StrandID, HairAttribute::Strand::Width, Width * Scale);
 				}
 			}
 		}
@@ -442,12 +454,12 @@ static void ParseObject(const Alembic::Abc::IObject& InObject, FHairDescription&
 	{
 		for (uint32 ChildIndex = 0; ChildIndex < NumChildren; ++ChildIndex)
 		{
-			ParseObject(InObject.getChild(ChildIndex), HairDescription, LocalMatrix, bCheckGroomAttributes);
+			ParseObject(InObject.getChild(ChildIndex), HairDescription, LocalMatrix, ConversionMatrix, Scale, bCheckGroomAttributes);
 		}
 	}
 }
 
-bool FAlembicHairTranslator::Translate(const FString& FileName, FHairDescription& HairDescription)
+bool FAlembicHairTranslator::Translate(const FString& FileName, FHairDescription& HairDescription, const FGroomConversionSettings& ConversionSettings)
 {
 	/** Factory used to generate objects*/
 	Alembic::AbcCoreFactory::IFactory Factory;
@@ -476,8 +488,9 @@ bool FAlembicHairTranslator::Translate(const FString& FileName, FHairDescription
 		return false;
 	}
 
+	FMatrix ConversionMatrix = FScaleMatrix::Make(ConversionSettings.Scale) * FRotationMatrix::Make(FQuat::MakeFromEuler(ConversionSettings.Rotation));
 	FMatrix ParentMatrix = FMatrix::Identity;
-	ParseObject(TopObject, HairDescription, ParentMatrix, true);
+	ParseObject(TopObject, HairDescription, ParentMatrix, ConversionMatrix, ConversionSettings.Scale.X, true);
 
 	return HairDescription.IsValid();
 }
diff --git a/Engine/Plugins/Experimental/AlembicHairImporter/Source/AlembicHairTranslator/Private/AlembicHairTranslator.h b/Engine/Plugins/Experimental/AlembicHairImporter/Source/AlembicHairTranslator/Private/AlembicHairTranslator.h
--- a/Engine/Plugins/Experimental/AlembicHairImporter/Source/AlembicHairTranslator/Private/AlembicHairTranslator.h
+++ b/Engine/Plugins/Experimental/AlembicHairImporter/Source/AlembicHairTranslator/Private/AlembicHairTranslator.h
@@ -7,7 +7,7 @@
 class FAlembicHairTranslator : public IHairStrandsTranslator
 {
 public:
-	virtual bool Translate(const FString& FilePath, FHairDescription& OutHairDescription) override;
+	virtual bool Translate(const FString& FilePath, FHairDescription& OutHairDescription, const FGroomConversionSettings& ConversionSettings) override;
 	virtual bool CanTranslate(const FString& FilePath) override;
 	virtual bool IsFileExtensionSupported(const FString& FileExtension) const override;
 	virtual FString GetSupportedFormat() const override;
diff --git a/Engine/Plugins/Experimental/HairStrands/Config/BaseHairStrands.ini b/Engine/Plugins/Experimental/HairStrands/Config/BaseHairStrands.ini
--- a/Engine/Plugins/Experimental/HairStrands/Config/BaseHairStrands.ini
+++ b/Engine/Plugins/Experimental/HairStrands/Config/BaseHairStrands.ini
@@ -1,3 +1,7 @@
+[Startup]
+;Enable Shader stages in Niagara
+fx.UseShaderStages=1
+
 [CoreRedirects]
 +PropertyRedirects=(OldName="HairStrandsAsset",NewName="GroomAsset")
 +ClassRedirects=(OldName="HairStrandsActor",NewName="GroomActor")
diff --git a/Engine/Plugins/Experimental/HairStrands/Shaders/Private/NiagaraDataInterfaceHairStrands.ush b/Engine/Plugins/Experimental/HairStrands/Shaders/Private/NiagaraDataInterfaceHairStrands.ush
--- a/Engine/Plugins/Experimental/HairStrands/Shaders/Private/NiagaraDataInterfaceHairStrands.ush
+++ b/Engine/Plugins/Experimental/HairStrands/Shaders/Private/NiagaraDataInterfaceHairStrands.ush
@@ -73,6 +73,23 @@ DIContext.GridSize = GridSize_##NAME;\
 DIContext.GridCurrentBuffer = GridCurrentBuffer_##NAME;\
 DIContext.GridDestinationBuffer = GridDestinationBuffer_##NAME;\
 
+/* -----------------------------------------------------------------
+ * Box utilities
+ * -----------------------------------------------------------------
+ */
+
+// Return the bounding box center
+void DIHairStrands_GetBoxCenter(in FDIHairStrandsContext DIContext, out float3 OutBoxCenter)
+{
+	OutBoxCenter = float3(0,0,0);
+}
+
+// Return the bounding box extent
+void DIHairStrands_GetBoxExtent(in FDIHairStrandsContext DIContext, out float3 OutBoxExtent)
+{
+	OutBoxExtent = -DIContext.GridOrigin.xyz;
+}
+
 /* -----------------------------------------------------------------
  * Grid utilities
  * -----------------------------------------------------------------
@@ -81,7 +98,8 @@ DIContext.GridDestinationBuffer = GridDestinationBuffer_##NAME;\
 // Return true if the grid index is within the domain size
 bool DIHairStrands_IsIndexValid(in FDIHairStrandsContext DIContext, in int3 GridIndex)
 {
-	return ((GridIndex.x < (DIContext.GridSize.x-1)) && (GridIndex.y < (DIContext.GridSize.y-1)) && (GridIndex.z < (DIContext.GridSize.z-1)));
+	return ((GridIndex.x < (DIContext.GridSize.x-1)) && (GridIndex.y < (DIContext.GridSize.y-1)) && 
+			(GridIndex.z < (DIContext.GridSize.z-1)) && (GridIndex.x >= 0) && (GridIndex.y >= 0) && (GridIndex.z >= 0));
 }
 
 // Transform the node world position in grid space
@@ -110,16 +128,38 @@ int3 DIHairStrands_GridNodeIndex( in FDIHairStrandsContext DIContext, in int3 Gr
 	return int3( GridIndex.x + DataOffset * DIContext.GridSize.x, GridIndex.yz);
 }
 
+void InterlockedAddFloat(in FDIHairStrandsContext DIContext, in int3 GridIndex, in int DataOffset, in float FloatValue)
+{
+	int IntValue = asint(FloatValue);
+	int CompareValue = 0;
+	int StoredValue = 0;
+
+	[allow_uav_condition] while(true)
+	{
+		InterlockedCompareExchange(DIContext.GridDestinationBuffer[
+				DIHairStrands_GridNodeIndex(DIContext,GridIndex,DataOffset)],CompareValue,IntValue,StoredValue);
+		if(StoredValue == CompareValue)
+		{
+			break;
+		}
+		CompareValue = StoredValue;
+		IntValue = asint(FloatValue+asfloat(StoredValue));
+	}
+}
+
 void DIHairStrands_AddNodeScalar( in FDIHairStrandsContext DIContext, in int3 GridIndex, in int DataOffset, in float NodeScalar, in bool EnableLocking)
 {
 	const int ConvertScalar = NodeScalar * SCALE;
 	if(EnableLocking)
 	{
 		InterlockedAdd(DIContext.GridDestinationBuffer[DIHairStrands_GridNodeIndex(DIContext,GridIndex,DataOffset)], ConvertScalar);
+		//InterlockedAddFloat(DIContext,GridIndex,DataOffset,NodeScalar);
 	}
 	else
 	{
 		DIContext.GridDestinationBuffer[DIHairStrands_GridNodeIndex(DIContext,GridIndex,DataOffset)] += ConvertScalar;
+		//const float FloatValue = asfloat(DIContext.GridDestinationBuffer[DIHairStrands_GridNodeIndex(DIContext,GridIndex,DataOffset)]) + NodeScalar;
+		//DIContext.GridDestinationBuffer[DIHairStrands_GridNodeIndex(DIContext,GridIndex,DataOffset)] = asint(FloatValue);
 	}
 	DeviceMemoryBarrier();
 }
@@ -133,6 +173,7 @@ void DIHairStrands_AddNodeVector( in FDIHairStrandsContext DIContext, in int3 Gr
 
 float DIHairStrands_GetGridScalar( in FDIHairStrandsContext DIContext, in int3 GridIndex, in int DataOffset)
 {
+	//return asfloat(DIContext.GridCurrentBuffer[DIHairStrands_GridNodeIndex(DIContext,GridIndex,DataOffset)]);
 	return DIContext.GridCurrentBuffer[DIHairStrands_GridNodeIndex(DIContext,GridIndex,DataOffset)] / SCALE;
 }
 
@@ -147,8 +188,8 @@ float3 DIHairStrands_GetGridVector( in FDIHairStrandsContext DIContext, in int3
 void DIHairStrands_BuildCollisionGrid(in FDIHairStrandsContext DIContext, in float3 NodePosition, in float3 NodeVelocity, in float NodeMass, 
 			in float3 NodeGradientX, in float3 NodeGradientY, in float3 NodeGradientZ, out bool OutBuildStatus)
 {
-	SharedLocalPosition[GGroupThreadId.x] = NodePosition;
-	SharedLocalVelocity[GGroupThreadId.x] = NodeVelocity;
+	SharedNodePosition[GGroupThreadId.x] = NodePosition;
+	SharedPreviousPosition[GGroupThreadId.x] = NodeVelocity;
 	SharedInverseMass[GGroupThreadId.x] = NodeMass;
 	GroupMemoryBarrier(); 
 
@@ -157,15 +198,15 @@ void DIHairStrands_BuildCollisionGrid(in FDIHairStrandsContext DIContext, in flo
 	const int LocalIndex = GGroupThreadId.x % DIContext.StrandSize;
 	if(LocalIndex != 0)
 	{
-		const float3 DeltaVelocity = SharedLocalVelocity[GGroupThreadId.x] - SharedLocalVelocity[GGroupThreadId.x-1];
-		const float3 DeltaPosition = SharedLocalPosition[GGroupThreadId.x] - SharedLocalPosition[GGroupThreadId.x-1];
+		const float3 DeltaVelocity = SharedPreviousPosition[GGroupThreadId.x] - SharedPreviousPosition[GGroupThreadId.x-1];
+		const float3 DeltaPosition = SharedNodePosition[GGroupThreadId.x] - SharedNodePosition[GGroupThreadId.x-1];
 
 		const float2 SegmentWeight = (LocalIndex == 1) ? float2(1.0,0.5) : (LocalIndex == (DIContext.StrandSize-1) ) ? float2(0.5,1.0) : float2(0.5,0.5);
 		const float SegmentMass =  (SharedInverseMass[GGroupThreadId.x-1] * SegmentWeight.x + SharedInverseMass[GGroupThreadId.x] * SegmentWeight.y);
 		const float SegmentLength = length(DeltaPosition);
 
-		const float3 GridPositionA = DIHairStrands_GetGridPosition(DIContext,SharedLocalPosition[GGroupThreadId.x-1]);
-		const float3 GridPositionB = DIHairStrands_GetGridPosition(DIContext,SharedLocalPosition[GGroupThreadId.x]);
+		const float3 GridPositionA = DIHairStrands_GetGridPosition(DIContext,SharedNodePosition[GGroupThreadId.x-1]);
+		const float3 GridPositionB = DIHairStrands_GetGridPosition(DIContext,SharedNodePosition[GGroupThreadId.x]);
 		const float3 GridSegment = GridPositionB - GridPositionA;
 
 		const int NumSamples = ceil(SegmentLength / DIContext.GridOrigin.w);
@@ -177,7 +218,9 @@ void DIHairStrands_BuildCollisionGrid(in FDIHairStrandsContext DIContext, in flo
 		{
 			const float SampleCoord = (0.5+SampleIndex) / NumSamples;
 			const float3 SamplePosition = GridPositionA + SampleCoord * GridSegment;
-			const float3 SampleMomentum = (SharedLocalVelocity[GGroupThreadId.x-1] + SampleCoord * DeltaVelocity) * SampleMass;
+			//const float3 SampleMomentum = NodeVelocity * SampleMass;
+			const float3 SampleMomentum = (SharedPreviousPosition[GGroupThreadId.x-1] + SampleCoord * DeltaVelocity) * SampleMass;
+
 
 			for( int ComponentIndex = 0; ComponentIndex < 3; ++ComponentIndex)
 			{
@@ -293,11 +336,86 @@ void DIHairStrands_QueryCollisionGrid(in FDIHairStrandsContext DIContext, in flo
 	//}
 }
 
+// Return the grid index given a grid hash
+int3 DIHairStrands_ComputeGridIndex(in FDIHairStrandsContext DIContext, in int GridHash)
+{
+	const int GridIndexX = GridHash / ( DIContext.GridSize.y * DIContext.GridSize.z);
+	const int GridDummy = GridHash - GridIndexX * DIContext.GridSize.y * DIContext.GridSize.z;
+	const int GridIndexY = GridDummy / (DIContext.GridSize.z);
+	const int GridIndexZ = GridDummy % (DIContext.GridSize.z);
+
+	return int3(GridIndexX,GridIndexY,GridIndexZ);
+}
+
+// Compute the grid pressure
+float DIHairStrands_ComputeGridPressure(in FDIHairStrandsContext DIContext, in int3 PrevIndex)
+{	
+	float GridPressure = 0.0;
+
+	if(DIHairStrands_IsIndexValid(DIContext,PrevIndex))
+	{
+		const float VoxelVolume = DIContext.GridOrigin.w * DIContext.GridOrigin.w * DIContext.GridOrigin.w;
+		float GridDivergence = 0.0;
+		float GridFraction = 0.0;
+
+		[unroll]
+		for(int ComponentIndex = 0; ComponentIndex < 3; ++ComponentIndex)
+		{
+			int3 NextIndex = PrevIndex; NextIndex[ComponentIndex] += 1;
+
+			const float PrevMass = DIHairStrands_GetGridScalar(DIContext,PrevIndex,MASS_OFFSET+ComponentIndex);
+			const float PrevDensity = min(1,max(0,PrevMass / VoxelVolume));
+			
+			if( PrevMass > 0.)
+			{
+				GridDivergence -= DIHairStrands_GetGridScalar(DIContext,PrevIndex,VELOCITY_OFFSET+ComponentIndex) / PrevMass;
+				GridFraction += 1.0 / PrevDensity;
+			}
+
+			const float NextMass = DIHairStrands_GetGridScalar(DIContext,NextIndex,MASS_OFFSET+ComponentIndex);
+			const float NextDensity = min(1,max(0,NextMass / VoxelVolume));
+
+			if( NextMass > 0.)
+			{
+				GridDivergence += DIHairStrands_GetGridScalar(DIContext,NextIndex,VELOCITY_OFFSET+ComponentIndex) / NextMass;
+				GridFraction += 1.0 / NextDensity;
+			}
+		}
+		if( GridFraction > 0.)
+		{
+			GridDivergence *= DIContext.GridOrigin.w;
+			GridPressure = -GridDivergence / GridFraction;
+		}
+	}
+	return GridPressure;
+}
 
 // Return the velocity from the grid to the node
 void DIHairStrands_ProjectCollisionGrid(in FDIHairStrandsContext DIContext, in int GridHash, out bool OutProjectStatus)
 {	
 	OutProjectStatus = false;
+
+	const int3 NextIndex = DIHairStrands_ComputeGridIndex(DIContext,GridHash);
+	const float NextPressure = DIHairStrands_ComputeGridPressure(DIContext,NextIndex);
+	const float VoxelVolume = DIContext.GridOrigin.w * DIContext.GridOrigin.w * DIContext.GridOrigin.w;
+
+	[unroll]
+	for(int ComponentIndex = 0; ComponentIndex < 3; ++ComponentIndex)
+	{
+		int3 PrevIndex = NextIndex; PrevIndex[ComponentIndex] -= 1;
+		const float PrevPressure = DIHairStrands_ComputeGridPressure(DIContext,PrevIndex);
+
+		const float NextMass = DIHairStrands_GetGridScalar(DIContext,NextIndex,MASS_OFFSET+ComponentIndex);
+		const float NextDensity = min(1,max(0,NextMass / VoxelVolume));
+
+		const float PressureGradient = (NextPressure - PrevPressure) / ( DIContext.GridOrigin.w * NextDensity );
+
+		DIHairStrands_AddNodeScalar(DIContext,NextIndex,VELOCITY_OFFSET+ComponentIndex,-PressureGradient*NextMass,false);
+		DIHairStrands_AddNodeScalar(DIContext,NextIndex,TANGENT_OFFSET+ComponentIndex,-PressureGradient*NextMass,false);
+	}
+	return;
+
+
 	const int GridIndexX = GridHash / ( DIContext.GridSize.y * DIContext.GridSize.z);
 	const int GridDummy = GridHash - GridIndexX * DIContext.GridSize.y * DIContext.GridSize.z;
 	const int GridIndexY = GridDummy / (DIContext.GridSize.z);
@@ -306,7 +424,7 @@ void DIHairStrands_ProjectCollisionGrid(in FDIHairStrandsContext DIContext, in i
 	const bool OddStage = (ShaderStageIndex % 2)  == 0;
 	const bool OddVoxel = ((GridIndexX+GridIndexY+GridIndexZ) % 2) == 0;
 
-	if( (OddStage && OddVoxel) || (!OddStage && !OddVoxel) )
+	//if( (OddStage && OddVoxel) || (!OddStage && !OddVoxel) )
 	{
 		const float VoxelVolume = DIContext.GridOrigin.w * DIContext.GridOrigin.w * DIContext.GridOrigin.w;
 
@@ -347,12 +465,12 @@ void DIHairStrands_ProjectCollisionGrid(in FDIHairStrandsContext DIContext, in i
 					int3 NextIndex = PrevIndex; NextIndex[ComponentIndexB] += 1;
 
 					const float PrevMass = DIHairStrands_GetGridScalar(DIContext,PrevIndex,MASS_OFFSET+ComponentIndexB);
-					DIHairStrands_AddNodeScalar(DIContext,PrevIndex,VELOCITY_OFFSET+ComponentIndexB,-PressureGradient*PrevMass,false);
-					DIHairStrands_AddNodeScalar(DIContext,PrevIndex,TANGENT_OFFSET+ComponentIndexB,-PressureGradient*PrevMass,false);
+					DIHairStrands_AddNodeScalar(DIContext,PrevIndex,VELOCITY_OFFSET+ComponentIndexB,-PressureGradient*PrevMass,true);
+					DIHairStrands_AddNodeScalar(DIContext,PrevIndex,TANGENT_OFFSET+ComponentIndexB,-PressureGradient*PrevMass,true);
 
 					const float NextMass = DIHairStrands_GetGridScalar(DIContext,NextIndex,MASS_OFFSET+ComponentIndexB);
-					DIHairStrands_AddNodeScalar(DIContext,NextIndex,VELOCITY_OFFSET+ComponentIndexB,PressureGradient*NextMass,false);
-					DIHairStrands_AddNodeScalar(DIContext,NextIndex,TANGENT_OFFSET+ComponentIndexB,PressureGradient*NextMass,false);
+					DIHairStrands_AddNodeScalar(DIContext,NextIndex,VELOCITY_OFFSET+ComponentIndexB,PressureGradient*NextMass,true);
+					DIHairStrands_AddNodeScalar(DIContext,NextIndex,TANGENT_OFFSET+ComponentIndexB,PressureGradient*NextMass,true);
 				}
 			}
 		}	
@@ -419,7 +537,7 @@ float DIHairStrands_ComputeStrandLength(in FDIHairStrandsContext DIContext, in i
  */
 
 // Compute the node mass
-void DIHairStrands_ComputeNodeMass(in FDIHairStrandsContext DIContext, in int NodeIndex, out float OutNodeMass)
+void DIHairStrands_ComputeNodeMass(in FDIHairStrandsContext DIContext, in int NodeIndex, in float StrandsDensity, in float RootThickness, in float TipThickness, out float OutNodeMass)
 {
 	const int ReducedSize =  DIContext.StrandSize-1;
 	const float CoordScale = 1.0 / (ReducedSize-1.0);
@@ -438,16 +556,16 @@ void DIHairStrands_ComputeNodeMass(in FDIHairStrandsContext DIContext, in int No
 	const float PrevCoord = IsFirst ? 0.0 : (LocalIndex-0.5) * CoordScale;
 	const float NextCoord = IsLast  ? 1.0 : (LocalIndex+0.5) * CoordScale;
 
-	const float PrevThick = DIContext.RootThickness * ( 1.0 - PrevCoord ) + DIContext.TipThickness * PrevCoord;
-	const float NextThick = DIContext.RootThickness * ( 1.0 - NextCoord ) + DIContext.TipThickness * NextCoord;
+	const float PrevThick = RootThickness * ( 1.0 - PrevCoord ) + TipThickness * PrevCoord;
+	const float NextThick = RootThickness * ( 1.0 - NextCoord ) + TipThickness * NextCoord;
 	const float MeanRadius = 0.25 * ( PrevThick + NextThick );
 
-	//OutNodeMass = DIContext.StrandDensity * EdgeLength * EdgeScale * M_PI * MeanRadius * MeanRadius;
-	OutNodeMass = DIContext.StrandDensity * EdgeLength * EdgeScale * M_PI * (PrevThick * PrevThick + NextThick * PrevThick + NextThick * NextThick) / 12.0;
+	//OutNodeMass = StrandsDensity * EdgeLength * EdgeScale * M_PI * MeanRadius * MeanRadius;
+	OutNodeMass = StrandsDensity * EdgeLength * EdgeScale * M_PI * (PrevThick * PrevThick + NextThick * PrevThick + NextThick * NextThick) / 12.0;
 }
 
 // Compute the node inertia
-void DIHairStrands_ComputeNodeInertia(in FDIHairStrandsContext DIContext, in int NodeIndex, out float3 OutNodeInertia)
+void DIHairStrands_ComputeNodeInertia(in FDIHairStrandsContext DIContext, in int NodeIndex, in float StrandsDensity, in float RootThickness, in float TipThickness,  out float3 OutNodeInertia)
 {
 	const int ReducedSize =  DIContext.StrandSize-1;
 	const float CoordScale = 1.0 / (ReducedSize-1.0);
@@ -461,16 +579,16 @@ void DIHairStrands_ComputeNodeInertia(in FDIHairStrandsContext DIContext, in int
 	const float PrevCoord = LocalIndex * CoordScale;
 	const float NextCoord = (LocalIndex+1) * CoordScale;
 
-	const float PrevThick = DIContext.RootThickness * ( 1.0 - PrevCoord ) + DIContext.TipThickness * PrevCoord;
-	const float NextThick = DIContext.RootThickness * ( 1.0 - NextCoord ) + DIContext.TipThickness * NextCoord;
+	const float PrevThick = RootThickness * ( 1.0 - PrevCoord ) + TipThickness * PrevCoord;
+	const float NextThick = RootThickness * ( 1.0 - NextCoord ) + TipThickness * NextCoord;
 	const float MeanRadius = 0.25 * ( PrevThick + NextThick );
 
-	const float EdgeMass = DIContext.StrandDensity * EdgeLength * M_PI * (PrevThick * PrevThick + NextThick * PrevThick + NextThick * NextThick) / 12.0;
-	//const float RadialInertia = EdgeMass * ( 3.0 * MeanRadius * MeanRadius + EdgeLength * EdgeLength) / 12.0;
-	const float RadialInertia = EdgeMass;// * ( MeanRadius * MeanRadius) / 4.0;
-	OutNodeInertia = float3( RadialInertia, RadialInertia, 2.0 * RadialInertia);//EdgeMass * MeanRadius * MeanRadius * 0.5 );
+	const float EdgeMass = StrandsDensity * EdgeLength * M_PI * (PrevThick * PrevThick + NextThick * PrevThick + NextThick * NextThick) / 12.0;
+	const float RadialInertia = EdgeMass * ( 3.0 * MeanRadius * MeanRadius + EdgeLength * EdgeLength) / 12.0;
+	//const float RadialInertia = EdgeMass * ( MeanRadius * MeanRadius) / 4.0;
+	OutNodeInertia = float3( RadialInertia, RadialInertia, EdgeMass * MeanRadius * MeanRadius * 0.5 );
 
-	//const float EdgeMass = DIContext.StrandDensity * EdgeLength * M_PI * MeanRadius * MeanRadius;
+	//const float EdgeMass = StrandsDensity * EdgeLength * M_PI * MeanRadius * MeanRadius;
 	//const float RadialInertia = EdgeMass * ( MeanRadius * MeanRadius) / 4.0;
 	//OutNodeInertia = float3( RadialInertia, RadialInertia, 2.0 * RadialInertia );
 }
@@ -498,9 +616,19 @@ void DIHairStrands_ComputeNodePosition(in FDIHairStrandsContext DIContext, in in
 
 	if(RootNode)
 	{
+		const float FirstPoint = CoordScale * PointCount;
+		const int FirstPrev = floor(FirstPoint);
+		const int FirstNext = PointPrev+1;
+
+		const float FirstAlpha = FirstPoint - (float)FirstPrev;
+		const float3 FirstPosition =  f16tof32(DIContext.RestPositionsBuffer[PointOffset+FirstPrev].xyz) * (1.0-FirstAlpha) + 
+						   f16tof32(DIContext.RestPositionsBuffer[PointOffset+FirstNext].xyz) * FirstAlpha;
+		const float3 DiffPosition = FirstPosition - f16tof32(DIContext.RestPositionsBuffer[PointOffset].xyz);
+		const float3 EdgeDirection =  normalize(DiffPosition);
+
 		const float EdgeLength = DIHairStrands_ComputeStrandLength(DIContext,StrandIndex) * CoordScale;
-		float3 EdgeDirection = f16tof32(DIContext.RestPositionsBuffer[PointOffset+PointNext].xyz) - 
-									 f16tof32(DIContext.RestPositionsBuffer[PointOffset+PointPrev].xyz);
+		//float3 EdgeDirection = f16tof32(DIContext.RestPositionsBuffer[PointOffset+PointNext].xyz) - 
+		//							 f16tof32(DIContext.RestPositionsBuffer[PointOffset+PointPrev].xyz);
 		OutNodePosition = 	f16tof32(DIContext.RestPositionsBuffer[PointOffset+PointPrev].xyz) - EdgeLength * normalize(EdgeDirection);					 
 	}
 	else
@@ -547,20 +675,20 @@ void DIHairStrands_ComputeNodeOrientation(in FDIHairStrandsContext DIContext, in
 // Compute the edge length value
 void DIHairStrands_ComputeEdgeLength(in FDIHairStrandsContext DIContext, in int NodeIndex, in float3 NodePosition, out float OutEdgeLength)
 {
-	SharedLocalPosition[GGroupThreadId.x] = NodePosition;
+	SharedNodePosition[GGroupThreadId.x] = NodePosition;
 	GroupMemoryBarrier(); 
 
 	int StrandIndex = 0, LocalIndex = 0;
 	DIHairStrands_ComputeNodeIndices(DIContext,NodeIndex,StrandIndex,LocalIndex);
 
 	// L = ||P1-P0||
-	OutEdgeLength = (LocalIndex>0) ? length(SharedLocalPosition[GGroupThreadId.x]- SharedLocalPosition[GGroupThreadId.x-1]) : 0.0;
+	OutEdgeLength = (LocalIndex>0) ? length(SharedNodePosition[GGroupThreadId.x]- SharedNodePosition[GGroupThreadId.x-1]) : 0.0;
 }
 
 // Compute the edge darboux vector (diff between consecutive edge orientations)
 void DIHairStrands_ComputeEdgeRotation(in FDIHairStrandsContext DIContext, in int NodeIndex, in float4 NodeOrientation, out float4 OutEdgeRotation)
 {
-	SharedLocalQuaternion[GGroupThreadId.x] = NodeOrientation;
+	SharedNodeOrientation[GGroupThreadId.x] = NodeOrientation;
 	GroupMemoryBarrier(); 
 
 	int StrandIndex = 0, LocalIndex = 0;
@@ -569,15 +697,15 @@ void DIHairStrands_ComputeEdgeRotation(in FDIHairStrandsContext DIContext, in in
 	// D = Q0^-1 * Q1
 	if(LocalIndex>0)
 	{
-		float4 q0 = SharedLocalQuaternion[GGroupThreadId.x-1];
-		float4 q1 = SharedLocalQuaternion[GGroupThreadId.x];
+		float4 q0 = SharedNodeOrientation[GGroupThreadId.x-1];
+		float4 q1 = SharedNodeOrientation[GGroupThreadId.x];
 		OutEdgeRotation = float4(
 				q1.xyz * q0.w - q0.xyz * q1.w + cross(-q0.xyz, q1.xyz),
 				q0.w * q1.w - dot(-q0.xyz, q1.xyz));
 		float4 OmegaPlus = OutEdgeRotation + float4(0,0,0,1);
 		float4 OmegaMinus = OutEdgeRotation - float4(0,0,0,1);
 		if( dot(OmegaMinus,OmegaMinus) > dot(OmegaPlus,OmegaPlus) ) OutEdgeRotation = -OutEdgeRotation;
-		//OutEdgeRotation = MultiplyQuat(InverseQuat(SharedLocalQuaternion[GGroupThreadId.x-1]),SharedLocalQuaternion[GGroupThreadId.x]);
+		//OutEdgeRotation = MultiplyQuat(InverseQuat(SharedNodeOrientation[GGroupThreadId.x-1]),SharedNodeOrientation[GGroupThreadId.x]);
 	}
 	else
 	{
@@ -616,7 +744,7 @@ void DIHairStrands_ResetPointPosition(in FDIHairStrandsContext DIContext, in int
 void DIHairStrands_UpdatePointPosition(in FDIHairStrandsContext DIContext, in int NodeIndex, in float3 NodeDisplace, out bool ReportStatus) 
 {
 	const int ReducedSize =  DIContext.StrandSize-1;
-	SharedLocalPosition[GGroupThreadId.x] = NodeDisplace;
+	SharedNodePosition[GGroupThreadId.x] = NodeDisplace;
 	GroupMemoryBarrier(); 
 
 	int StrandIndex = 0, LocalIndex = 0;
@@ -641,7 +769,7 @@ void DIHairStrands_UpdatePointPosition(in FDIHairStrandsContext DIContext, in in
 		{
 			ReportStatus = true;
 			const float PointAlpha = PointNode - (float)NodePrev;
-			const float3 LocalDisplace = SharedLocalPosition[GGroupThreadId.x] * (1.0-PointAlpha) + SharedLocalPosition[GGroupThreadId.x+1] * PointAlpha;
+			const float3 LocalDisplace = SharedNodePosition[GGroupThreadId.x] * (1.0-PointAlpha) + SharedNodePosition[GGroupThreadId.x+1] * PointAlpha;
 			const float3 DeformedPosition = f16tof32(DIContext.RestPositionsBuffer[PointBegin+PointIndex].xyz) + LocalDisplace;
 			DIContext.PointsPositionsBuffer[PointBegin+PointIndex] = uint4( f32tof16(DeformedPosition), DIContext.RestPositionsBuffer[PointBegin+PointIndex].w);
 		}
@@ -680,17 +808,16 @@ void DIHairStrands_AdvectNodeOrientation(in FDIHairStrandsContext DIContext, in
 }
 
 // Update the node linear velocity based on the node position difference
-void DIHairStrands_UpdateLinearVelocity(in FDIHairStrandsContext DIContext, in float3 PreviousPosition, in float3 NodePosition, in float DeltaTime, inout float3 OutLinearVelocity)
+void DIHairStrands_UpdateLinearVelocity(in FDIHairStrandsContext DIContext, in float3 PreviousPosition, in float3 NodePosition, in float DeltaTime, out float3 OutLinearVelocity)
 {	
 	OutLinearVelocity = (NodePosition-PreviousPosition) / DeltaTime;
 }
 
+
 // Update the node angular velocity based on the node orientation difference
-void DIHairStrands_UpdateAngularVelocity(in FDIHairStrandsContext DIContext, in float4 PreviousOrientation, in float4 NodeOrientation, in float DeltaTime, inout float3 OutAngularVelocity)
+void DIHairStrands_UpdateAngularVelocity(in FDIHairStrandsContext DIContext, in float4 PreviousOrientation, in float4 NodeOrientation, in float DeltaTime, out float3 OutAngularVelocity)
 {
-	const float4 DeltaQuat = float4(
-				NodeOrientation.xyz * PreviousOrientation.w - PreviousOrientation.xyz * NodeOrientation.w - cross(PreviousOrientation.xyz, NodeOrientation.xyz),
-				PreviousOrientation.w * NodeOrientation.w + dot(PreviousOrientation.xyz, NodeOrientation.xyz));
+	const float4 DeltaQuat = MultiplyQuat(NodeOrientation,InverseQuat(PreviousOrientation));
 	const float AxisLength = length( DeltaQuat.xyz );
 
 	if (AxisLength<SMALL_NUMBER)
diff --git a/Engine/Plugins/Experimental/HairStrands/Shaders/Private/NiagaraDataInterfacePhysicsAsset.ush b/Engine/Plugins/Experimental/HairStrands/Shaders/Private/NiagaraDataInterfacePhysicsAsset.ush
--- a/Engine/Plugins/Experimental/HairStrands/Shaders/Private/NiagaraDataInterfacePhysicsAsset.ush
+++ b/Engine/Plugins/Experimental/HairStrands/Shaders/Private/NiagaraDataInterfacePhysicsAsset.ush
@@ -25,14 +25,14 @@ struct FDIPhysicsAssetContext
 	Buffer<float4>		ElementExtentBuffer;
 }; 
 
-#define DIPhysicsAsset_DECLARE_CONSTANTS(NAME)\
+#define DIPHYSICSASSET_DECLARE_CONSTANTS(NAME)\
 uint4 ElementOffsets_##NAME;\
 Buffer<float4>		CurrentTransformBuffer_##NAME;\
 Buffer<float4>		InverseTransformBuffer_##NAME;\
 Buffer<float4>		PreviousTransformBuffer_##NAME;\
 Buffer<float4>		ElementExtentBuffer_##NAME;\
 
-#define DIPhysicsAsset_MAKE_CONTEXT(NAME)\
+#define DIPHYSICSASSET_MAKE_CONTEXT(NAME)\
 FDIPhysicsAssetContext DIContext; \
 DIContext.ElementOffsets = ElementOffsets_##NAME;\
 DIContext.CurrentTransformBuffer = CurrentTransformBuffer_##NAME;\
diff --git a/Engine/Plugins/Experimental/HairStrands/Shaders/Private/NiagaraDataInterfacePressureGrid.ush b/Engine/Plugins/Experimental/HairStrands/Shaders/Private/NiagaraDataInterfacePressureGrid.ush
new file mode 100644
--- /dev/null
+++ b/Engine/Plugins/Experimental/HairStrands/Shaders/Private/NiagaraDataInterfacePressureGrid.ush
@@ -0,0 +1,594 @@
+// Copyright 1998-2019 Epic Games, Inc. All Rights Reserved.
+
+/*=============================================================================
+NiagaraDataInterfacePressureGrid.ush 
+=============================================================================*/
+
+#pragma once
+
+/* -----------------------------------------------------------------
+ * Pressure Grid constants and context
+ * -----------------------------------------------------------------
+ */
+
+static const float FLOAT_PRECISION = 1e+6;
+static const bool FAST_MODE = true;
+
+static const int FLUID_MASS_OFFSET = 0;
+static const int FLUID_VELOCITY_OFFSET = 3;
+static const int FLUID_PRESSURE_OFFSET = 6;
+static const int SOLID_WEIGHT_OFFSET = 7;
+static const int SOLID_VELOCITY_OFFSET = 10;
+static const int SOLID_DISTANCE_OFFSET = 13;
+
+struct FDIPressureGridContext
+{
+	float4x4			WorldTransform;
+	float4x4			WorldInverse;
+	float4				GridOrigin;
+	int4				GridSize;
+	Texture3D<int>		GridCurrentBuffer;
+	RWTexture3D<int>	GridDestinationBuffer;
+}; 
+
+#define DIPRESSUREGRID_DECLARE_CONSTANTS(NAME)\
+float4x4 WorldTransform_##NAME;\
+float4x4 WorldInverse_##NAME;\
+float4 GridOrigin_##NAME;\
+int4 GridSize_##NAME;\
+Texture3D<int> GridCurrentBuffer_##NAME;\
+RWTexture3D<int> GridDestinationBuffer_##NAME;\
+
+#define DIPRESSUREGRID_MAKE_CONTEXT(NAME)\
+FDIPressureGridContext DIContext; \
+DIContext.WorldTransform = WorldTransform_##NAME;\
+DIContext.WorldInverse = WorldInverse_##NAME;\
+DIContext.GridOrigin = GridOrigin_##NAME;\
+DIContext.GridSize = GridSize_##NAME;\
+DIContext.GridCurrentBuffer = GridCurrentBuffer_##NAME;\
+DIContext.GridDestinationBuffer = GridDestinationBuffer_##NAME;\
+
+/* -----------------------------------------------------------------
+ * Grid utilities
+ * -----------------------------------------------------------------
+ */
+
+// Return true if the grid index is in the valid space
+bool DIPressureGrid_IsIndexValid(in int3 GridIndex, in int4 ValidSpace)
+{
+	return ((GridIndex.x < (ValidSpace.x-1)) && (GridIndex.y < (ValidSpace.y-1)) && 
+			(GridIndex.z < (ValidSpace.z-1)) && (GridIndex.x >= 0) && (GridIndex.y >= 0) && (GridIndex.z >= 0));
+}
+
+// Return true if the grid cell index is in any voxel center (0,N-2)
+bool DIPressureGrid_IsCellValid(in FDIPressureGridContext DIContext, in int3 GridIndex)
+{
+	return DIPressureGrid_IsIndexValid(GridIndex,DIContext.GridSize);
+}
+
+// Return true if the grid index is in any voxel face for the component one (0,N-1) and in the voxel center for the others (0,N-2)
+bool DIPressureGrid_IsFaceValid(in FDIPressureGridContext DIContext, in int3 GridIndex, in int ComponentIndex)
+{
+	int4 ValidSpace = DIContext.GridSize; //ValidSpace[ComponentIndex] += 1;
+	if( ComponentIndex == 0 )
+	{
+		ValidSpace.x += 1;
+	}
+	else if( ComponentIndex == 1 )
+	{
+		ValidSpace.y += 1;
+	}
+	else if( ComponentIndex == 2 )
+	{
+		ValidSpace.z += 1;
+	}
+	return DIPressureGrid_IsIndexValid(GridIndex,ValidSpace);
+}
+
+// Transform the node world position in grid space
+float3 DIPressureGrid_GetGridPosition(in FDIPressureGridContext DIContext, in float3 NodePosition)
+{
+	return ( mul(float4( NodePosition, 1.0), DIContext.WorldInverse).xyz - DIContext.GridOrigin.xyz );
+}
+
+// Compute the grid index and fraction given a grid position
+void DIPressureGrid_GetGridIndex(in FDIPressureGridContext DIContext, in float3 GridPosition, in float3 GridOffset, out int3 OutGridIndex, out float3 OutGridFraction)
+{
+	const float3 GridCell = GridPosition / DIContext.GridOrigin.w - GridOffset;
+
+	OutGridIndex = floor(GridCell);
+	OutGridFraction = GridCell - OutGridIndex;
+}
+
+// Compute the data index from agrid index and a data offset
+int3 DIPressureGrid_GetDataIndex( in FDIPressureGridContext DIContext, in int3 GridIndex, in int DataOffset)
+{
+	return int3( GridIndex.x + DataOffset * DIContext.GridSize.x, GridIndex.yz);
+}
+
+/* -----------------------------------------------------------------
+ * Get grid information from hash value
+ * -----------------------------------------------------------------
+ */
+
+// Compute the grid index given a grid hash
+int3 DIPressureGrid_ComputeGridIndex(in FDIPressureGridContext DIContext, in int GridCell)
+{
+	const int GridSizeXY = DIContext.GridSize.y * DIContext.GridSize.z;
+	const int GridCellYZ = GridCell % GridSizeXY;
+
+	return int3(GridCell / GridSizeXY, GridCellYZ / DIContext.GridSize.z,GridCellYZ % DIContext.GridSize.z);
+}
+
+// Compute the grid hash given a grid index
+int DIPressureGrid_ComputeGridCell(in FDIPressureGridContext DIContext, in int3 GridIndex)
+{
+	return GridIndex.x * DIContext.GridSize.y * DIContext.GridSize.z + GridIndex.y * DIContext.GridSize.z + GridIndex.z;
+}
+
+// Compute the world grid position
+void DIPressureGrid_GetCellPosition(in FDIPressureGridContext DIContext, in int GridCell, out float3 OutCellPosition)
+{
+	const int3 GridIndex = DIPressureGrid_ComputeGridIndex(DIContext,GridCell);
+	OutCellPosition = mul(float4((GridIndex+float3(0.5,0.5,0.5))*DIContext.GridOrigin.w+DIContext.GridOrigin.xyz,1.0),DIContext.WorldTransform).xyz;
+}
+
+/* -----------------------------------------------------------------
+ * Transfer the velocity and mass buffer from the current to the destination buffer
+ * -----------------------------------------------------------------
+ */
+
+void DIPressureGrid_TransferGridComponent(in FDIPressureGridContext DIContext, in int3 GridIndex, in int ComponentIndex)
+{
+	const int GridMomentum = DIContext.GridCurrentBuffer[DIPressureGrid_GetDataIndex(DIContext,GridIndex,FLUID_VELOCITY_OFFSET+ComponentIndex)];
+	const int GridMass = DIContext.GridCurrentBuffer[DIPressureGrid_GetDataIndex(DIContext,GridIndex,FLUID_MASS_OFFSET+ComponentIndex)];
+
+	if(GridMass != 0)
+	{
+		if(FAST_MODE)
+		{
+			const float GridMassFloat = float(GridMass) / FLOAT_PRECISION;
+			const float GridVelocityFloat = ( float(GridMomentum) / FLOAT_PRECISION ) / GridMassFloat;
+
+			DIContext.GridDestinationBuffer[DIPressureGrid_GetDataIndex(DIContext,GridIndex,FLUID_VELOCITY_OFFSET+ComponentIndex)] = asint(GridVelocityFloat);
+			DIContext.GridDestinationBuffer[DIPressureGrid_GetDataIndex(DIContext,GridIndex,FLUID_MASS_OFFSET+ComponentIndex)] = asint(GridMassFloat);
+		}
+		else
+		{
+			DIContext.GridDestinationBuffer[DIPressureGrid_GetDataIndex(DIContext,GridIndex,FLUID_VELOCITY_OFFSET+ComponentIndex)] = asint(asfloat(GridMomentum)/asfloat(GridMass));
+			DIContext.GridDestinationBuffer[DIPressureGrid_GetDataIndex(DIContext,GridIndex,FLUID_MASS_OFFSET+ComponentIndex)] = GridMass;
+		}
+	}
+	else
+	{
+		DIContext.GridDestinationBuffer[DIPressureGrid_GetDataIndex(DIContext,GridIndex,FLUID_VELOCITY_OFFSET+ComponentIndex)] = 0;
+		DIContext.GridDestinationBuffer[DIPressureGrid_GetDataIndex(DIContext,GridIndex,FLUID_MASS_OFFSET+ComponentIndex)] = 0;
+	}
+	DeviceMemoryBarrier();
+}
+
+void DIPressureGrid_TransferCellVelocity(in FDIPressureGridContext DIContext, in int GridCell, out bool OutTransferStatus)
+{
+	const int3 GridIndex = DIPressureGrid_ComputeGridIndex(DIContext,GridCell);
+	OutTransferStatus = false;
+	//if(DIPressureGrid_IsFaceValid(DIContext,GridIndex,0))
+	{
+		OutTransferStatus = true;
+		DIPressureGrid_TransferGridComponent(DIContext,GridIndex,0);
+	}
+	//if(DIPressureGrid_IsFaceValid(DIContext,GridIndex,1))
+	{
+		OutTransferStatus = true;
+		DIPressureGrid_TransferGridComponent(DIContext,GridIndex,1);
+	}
+	//if(DIPressureGrid_IsFaceValid(DIContext,GridIndex,2))
+	{
+		OutTransferStatus = true;
+		DIPressureGrid_TransferGridComponent(DIContext,GridIndex,2);
+	}
+}
+
+/* -----------------------------------------------------------------
+ * Interlocked Add Float (slow but accurate version) and Int (fast but approximate version)
+ * -----------------------------------------------------------------
+ */
+
+ void DIPressureGrid_InterlockedAddFloat(in FDIPressureGridContext DIContext, in int3 GridIndex, in int DataOffset, in float FloatValue)
+{
+	int IntValue = asint(FloatValue);
+	int CompareValue = 0;
+	int StoredValue = 0;
+
+	[allow_uav_condition] while(true)
+	{
+		InterlockedCompareExchange(DIContext.GridDestinationBuffer[
+				DIPressureGrid_GetDataIndex(DIContext,GridIndex,DataOffset)],CompareValue,IntValue,StoredValue);
+		if(StoredValue == CompareValue)
+		{
+			break;
+		}
+		CompareValue = StoredValue;
+		IntValue = asint(FloatValue+asfloat(StoredValue));
+	}
+}
+
+void DIPressureGrid_InterlockedAddInt(in FDIPressureGridContext DIContext, in int3 GridIndex, in int DataOffset, in int IntValue)
+{
+	InterlockedAdd(DIContext.GridDestinationBuffer[DIPressureGrid_GetDataIndex(DIContext,GridIndex,DataOffset)], IntValue);
+}
+
+/* -----------------------------------------------------------------
+ * Update the velocity and mass buffer
+ * -----------------------------------------------------------------
+ */
+
+void DIPressureGrid_InterlockedAddScalar(in FDIPressureGridContext DIContext, in int3 GridIndex, in int DataOffset, in float FloatValue)
+{
+	if(FAST_MODE)
+	{
+		const int IntValue = FloatValue * FLOAT_PRECISION;
+		DIPressureGrid_InterlockedAddInt(DIContext,GridIndex,DataOffset,IntValue);
+	}
+	else
+	{
+		DIPressureGrid_InterlockedAddFloat(DIContext,GridIndex,DataOffset,FloatValue);
+	}
+}
+
+void DIPressureGrid_AddGridScalar(in FDIPressureGridContext DIContext, in int3 GridIndex, in int DataOffset, in float FloatValue)
+{
+	const float SumValue = asfloat(DIContext.GridCurrentBuffer[DIPressureGrid_GetDataIndex(DIContext,GridIndex,DataOffset)])+FloatValue;
+	DIContext.GridDestinationBuffer[DIPressureGrid_GetDataIndex(DIContext,GridIndex,DataOffset)] = asint(SumValue);
+}
+
+float DIPressureGrid_GetGridScalar(in FDIPressureGridContext DIContext, in int3 GridIndex, in int DataOffset)
+{
+	return asfloat(DIContext.GridCurrentBuffer[DIPressureGrid_GetDataIndex(DIContext,GridIndex,DataOffset)]);
+}
+
+void DIPressureGrid_SetGridScalar(in FDIPressureGridContext DIContext, in int3 GridIndex, in int DataOffset, in float FloatValue)
+{
+	DIContext.GridDestinationBuffer[DIPressureGrid_GetDataIndex(DIContext,GridIndex,DataOffset)] = asint(FloatValue);
+}
+
+void DIPressureGrid_SetGridVector( in FDIPressureGridContext DIContext, in int3 GridIndex, in int DataOffset, in float3 NodeVector)
+{
+	DIPressureGrid_SetGridScalar(DIContext,GridIndex,DataOffset,NodeVector.x);
+	DIPressureGrid_SetGridScalar(DIContext,GridIndex,DataOffset+1,NodeVector.y);
+	DIPressureGrid_SetGridScalar(DIContext,GridIndex,DataOffset+2,NodeVector.z);
+}
+
+void DIPressureGrid_AddGridVector( in FDIPressureGridContext DIContext, in int3 GridIndex, in int DataOffset, in float3 NodeVector)
+{
+	DIPressureGrid_AddGridScalar(DIContext,GridIndex,DataOffset,NodeVector.x);
+	DIPressureGrid_AddGridScalar(DIContext,GridIndex,DataOffset+1,NodeVector.y);
+	DIPressureGrid_AddGridScalar(DIContext,GridIndex,DataOffset+2,NodeVector.z);
+}
+
+float3 DIPressureGrid_GetGridVector( in FDIPressureGridContext DIContext, in int3 GridIndex, in int DataOffset)
+{
+	return float3(DIPressureGrid_GetGridScalar(DIContext,GridIndex,DataOffset),
+				  DIPressureGrid_GetGridScalar(DIContext,GridIndex,DataOffset+1),
+				  DIPressureGrid_GetGridScalar(DIContext,GridIndex,DataOffset+2));
+}
+
+/* -----------------------------------------------------------------
+ * Build Velocity Field from nodes
+ * -----------------------------------------------------------------
+ */
+
+// Transfer the velocity component from the node to the grid
+void DIPressureGrid_BuildVelocityComponent(in FDIPressureGridContext DIContext, in float3 GridPosition, in float SampleMass, in float SampleMomentum,
+			in float3 MomentumGradient, in float3 ComponentOffset, in int ComponentIndex )
+{
+	int3 GridIndex = int3(0,0,0); float3 GridFraction = float3(0,0,0);
+	DIPressureGrid_GetGridIndex(DIContext,GridPosition,ComponentOffset,GridIndex,GridFraction);
+
+	if(DIPressureGrid_IsFaceValid(DIContext,GridIndex,ComponentIndex) && 
+	   DIPressureGrid_IsFaceValid(DIContext,GridIndex+int3(1,1,1),ComponentIndex))
+	{
+		[unroll]
+		for(int i = 0; i < 2; ++i)
+		{
+			const float GridWeightX = (i == 0) ? 1.0 - GridFraction.x : GridFraction.x;
+			[unroll]
+			for(int j = 0; j < 2; ++j)
+			{
+				const float GridWeightY = (j == 0) ? 1.0 - GridFraction.y : GridFraction.y;
+				[unroll]
+				for(int k = 0; k < 2; ++k)
+				{
+					const float GridWeightZ = (k == 0) ? 1.0 - GridFraction.z : GridFraction.z;
+					const float GridWeightXYZ = GridWeightX*GridWeightY*GridWeightZ;
+
+					const int3 GridLocation = int3(GridIndex.x+i,GridIndex.y+j,GridIndex.z+k);
+					const float3 DeltaPosition = ((GridLocation+ComponentOffset)*DIContext.GridOrigin.w) - GridPosition;
+					const float GridMomentum = SampleMomentum + dot(MomentumGradient, DeltaPosition);
+
+					DIPressureGrid_InterlockedAddScalar(DIContext,GridLocation,FLUID_MASS_OFFSET+ComponentIndex,SampleMass*GridWeightXYZ);
+					DIPressureGrid_InterlockedAddScalar(DIContext,GridLocation,FLUID_VELOCITY_OFFSET+ComponentIndex,GridMomentum*GridWeightXYZ);
+				}
+			}
+		}
+	}
+}
+
+// Transfer the velocity from the node to the grid
+void DIPressureGrid_BuildVelocityField(in FDIPressureGridContext DIContext, in int StrandsSize, in float3 NodePosition, in float NodeMass, in float3 NodeVelocity, 
+			in float3 VelocityGradientX, in float3 VelocityGradientY, in float3 VelocityGradientZ, out bool OutBuildStatus)
+{
+	SharedNodePosition[GGroupThreadId.x] = NodePosition;
+	SharedPreviousPosition[GGroupThreadId.x] = NodeVelocity;
+	SharedInverseMass[GGroupThreadId.x] = NodeMass;
+	GroupMemoryBarrier();
+
+	const int LocalIndex = GGroupThreadId.x % StrandsSize;
+
+	if( LocalIndex > 0 )
+	{
+		const float3 DeltaVelocity = SharedPreviousPosition[GGroupThreadId.x] - SharedPreviousPosition[GGroupThreadId.x-1];
+		const float3 DeltaPosition = SharedNodePosition[GGroupThreadId.x] - SharedNodePosition[GGroupThreadId.x-1];
+
+		const float2 SegmentWeight = (LocalIndex == 1) ? float2(1.0,0.5) : (LocalIndex == (StrandsSize-1) ) ? float2(0.5,1.0) : float2(0.5,0.5);
+		const float SegmentMass =  (SharedInverseMass[GGroupThreadId.x-1] * SegmentWeight.x + SharedInverseMass[GGroupThreadId.x] * SegmentWeight.y);
+		const float SegmentLength = length(DeltaPosition);
+
+		const int NumSamples = ceil(SegmentLength / DIContext.GridOrigin.w);
+
+		const float SampleMass = SegmentMass / NumSamples;
+
+		const float3 MomentumGradientX = VelocityGradientX * SampleMass;
+		const float3 MomentumGradientY = VelocityGradientY * SampleMass;
+		const float3 MomentumGradientZ = VelocityGradientZ * SampleMass;
+
+		for( int SampleIndex = 0; SampleIndex < NumSamples; ++SampleIndex)
+		{
+			const float SampleCoord = (0.5+SampleIndex) / NumSamples;
+			const float3 SamplePosition = SharedNodePosition[GGroupThreadId.x] + SampleCoord * DeltaPosition;
+			const float3 SampleMomentum = (SharedPreviousPosition[GGroupThreadId.x] + SampleCoord * DeltaVelocity) * SampleMass;
+
+			const float3 GridPosition = DIPressureGrid_GetGridPosition(DIContext,SamplePosition);
+
+			DIPressureGrid_BuildVelocityComponent(DIContext,GridPosition,SampleMass,SampleMomentum.x,MomentumGradientX,float3(0.0,0.5,0.5),0);
+			DIPressureGrid_BuildVelocityComponent(DIContext,GridPosition,SampleMass,SampleMomentum.y,MomentumGradientY,float3(0.5,0.0,0.5),1);
+			DIPressureGrid_BuildVelocityComponent(DIContext,GridPosition,SampleMass,SampleMomentum.z,MomentumGradientZ,float3(0.5,0.5,0.0),2);
+		}
+	}
+	OutBuildStatus = true;
+}
+
+/* -----------------------------------------------------------------
+ * Sample Velocity field from grid
+ * -----------------------------------------------------------------
+ */
+
+// Sample the grid velocity field and gradient component at the node local position 
+void DIPressureGrid_SampleVelocityComponent(in FDIPressureGridContext DIContext, in float3 GridPosition, in float3 ComponentOffset, in int ComponentIndex, 
+						out float OutGridVelocity, out float OutGridMass, out float3 OutGridGradient)
+{
+	int3 GridIndex = int3(0,0,0); float3 GridFraction = float3(0,0,0);
+	DIPressureGrid_GetGridIndex(DIContext,GridPosition,ComponentOffset,GridIndex,GridFraction);
+
+	OutGridVelocity = 0.0;
+	OutGridGradient = float3(0,0,0);
+	OutGridMass = 0.0;
+
+	if(DIPressureGrid_IsFaceValid(DIContext,GridIndex,ComponentIndex) && 
+	   DIPressureGrid_IsFaceValid(DIContext,GridIndex+int3(1,1,1),ComponentIndex))
+	{
+		[unroll]
+		for(int i = 0; i < 2; ++i)
+		{
+			const float GridWeightX = (i == 0) ? 1.0 - GridFraction.x : GridFraction.x;
+			const float GridGradientX = (i == 0) ? -1.0: 1.0;
+			[unroll]
+			for(int j = 0; j < 2; ++j)
+			{
+				const float GridWeightY = (j == 0) ? 1.0 - GridFraction.y : GridFraction.y;
+				const float GridGradientY = (j == 0) ? -1.0: 1.0;
+				[unroll]
+				for(int k = 0; k < 2; ++k)
+				{
+					const float GridWeightZ = (k == 0) ? 1.0 - GridFraction.z : GridFraction.z;
+					const float GridGradientZ = (k == 0) ? -1.0: 1.0;
+
+					const float GridWeightXYZ = GridWeightX*GridWeightY*GridWeightZ;
+					const float3 GridGradientXYZ = float3(GridGradientX*GridWeightY*GridWeightZ,
+														  GridWeightX*GridGradientY*GridWeightZ,
+														  GridWeightX*GridWeightY*GridGradientZ) / DIContext.GridOrigin.w;
+
+					const int3 GridLocation = int3(GridIndex.x+i,GridIndex.y+j,GridIndex.z+k);
+					const float GridVelocity = DIPressureGrid_GetGridScalar(DIContext,GridLocation,FLUID_VELOCITY_OFFSET+ComponentIndex);
+					const float GridMass = DIPressureGrid_GetGridScalar(DIContext,GridLocation,FLUID_MASS_OFFSET+ComponentIndex);
+
+					OutGridMass += GridMass * GridWeightXYZ;
+					OutGridVelocity += GridVelocity * GridWeightXYZ;
+					OutGridGradient += GridVelocity * GridGradientXYZ;
+				}
+			}
+		}
+	}
+}
+
+// Sample the grid velocity field and gradient at the node world position 
+void DIPressureGrid_SampleVelocityField(in FDIPressureGridContext DIContext, in float3 NodePosition, out float3 OutGridVelocity, out float OutNodeDensity,
+			out float3 OutGridGradientX, out float3 OutGridGradientY, out float3 OutGridGradientZ)
+{
+	OutGridVelocity = float3(0,0,0);
+	OutGridGradientX = float3(0,0,0);
+	OutGridGradientY = float3(0,0,0);
+	OutGridGradientZ = float3(0,0,0);
+
+	float3 GridMass = float3(0,0,0);
+
+	const float3 GridPosition = DIPressureGrid_GetGridPosition(DIContext,NodePosition);
+	DIPressureGrid_SampleVelocityComponent(DIContext,GridPosition,float3(0.0,0.5,0.5),0,OutGridVelocity.x,GridMass.x,OutGridGradientX);
+	DIPressureGrid_SampleVelocityComponent(DIContext,GridPosition,float3(0.5,0.0,0.5),1,OutGridVelocity.y,GridMass.y,OutGridGradientY);
+	DIPressureGrid_SampleVelocityComponent(DIContext,GridPosition,float3(0.5,0.5,0.0),2,OutGridVelocity.z,GridMass.z,OutGridGradientZ);
+
+	OutNodeDensity = (GridMass.x + GridMass.y + GridMass.z)/(3.0*DIContext.GridOrigin.w*DIContext.GridOrigin.w*DIContext.GridOrigin.w);
+}
+
+/* -----------------------------------------------------------------
+ * Update boundary weights
+ * -----------------------------------------------------------------
+ */
+
+ float DIPressureGrid_GetFractionInside(in float SignedDistanceA, in float SignedDistanceB)
+ {
+	const bool IsInsideA = SignedDistanceA<0.0;
+	const bool IsInsideB = SignedDistanceB<0.0;
+
+	return (IsInsideA && IsInsideB) ? 1.0 : (IsInsideA && !IsInsideB) ? SignedDistanceA / (SignedDistanceA - SignedDistanceB) : 
+		   (!IsInsideA && IsInsideB) ? SignedDistanceB / (SignedDistanceB - SignedDistanceA)  : 0.0;
+ }
+
+ // Update the solid weights and velocities
+void DIPressureGrid_SetSolidBoundary(in FDIPressureGridContext DIContext, in int GridCell, in float SolidDistance,
+																in float3 SolidVelocity, out bool OutBoundaryStatus)
+{
+	int3 GridIndex = DIPressureGrid_ComputeGridIndex(DIContext,GridCell);
+	OutBoundaryStatus = false;
+	if( DIPressureGrid_IsCellValid(DIContext,GridIndex))
+	{
+		OutBoundaryStatus = true;
+		DIPressureGrid_SetGridScalar(DIContext,GridIndex,SOLID_DISTANCE_OFFSET,SolidDistance);
+		DIPressureGrid_SetGridVector(DIContext,GridIndex,SOLID_VELOCITY_OFFSET,SolidVelocity);
+	}
+}
+
+ // Update the solid weights and velocities
+void DIPressureGrid_ComputeBoundaryWeights(in FDIPressureGridContext DIContext, in int GridCell, out bool OutWeightsStatus)
+{
+	int3 GridIndex = DIPressureGrid_ComputeGridIndex(DIContext,GridCell);
+	OutWeightsStatus = false;
+
+	for(int ComponentIndex = 0; ComponentIndex < 3; ++ComponentIndex)
+	{
+		if( DIPressureGrid_IsFaceValid(DIContext,GridIndex,ComponentIndex))
+		{
+			float SolidWeight = 0.0, SolidVelocity = 0.0;
+			if( GridIndex[ComponentIndex] > 0 && GridIndex[ComponentIndex] < (DIContext.GridSize[ComponentIndex]-1) )
+			{
+				OutWeightsStatus = true;
+				const float RightDistance = DIPressureGrid_GetGridScalar(DIContext,GridIndex,SOLID_DISTANCE_OFFSET);
+				const float RightVelocity = DIPressureGrid_GetGridScalar(DIContext,GridIndex,SOLID_VELOCITY_OFFSET+ComponentIndex);
+
+				GridIndex[ComponentIndex] -= 1;
+				const float LeftDistance = DIPressureGrid_GetGridScalar(DIContext,GridIndex,SOLID_DISTANCE_OFFSET);
+				const float LeftVelocity = DIPressureGrid_GetGridScalar(DIContext,GridIndex,SOLID_VELOCITY_OFFSET+ComponentIndex);
+
+				SolidWeight = max(0.0,min(1.0 - DIPressureGrid_GetFractionInside(LeftDistance,RightDistance),1.0));
+				SolidVelocity = 0.5 * (RightVelocity+LeftVelocity);
+			}
+			DIPressureGrid_SetGridScalar(DIContext,GridIndex,SOLID_WEIGHT_OFFSET+ComponentIndex,SolidWeight);
+			DIPressureGrid_SetGridScalar(DIContext,GridIndex,SOLID_VELOCITY_OFFSET+ComponentIndex,SolidVelocity);
+		}
+	}
+}
+
+/* -----------------------------------------------------------------
+ * Sample Velocity field from grid
+ * -----------------------------------------------------------------
+ */
+
+// Compute the grid pressure
+// Mdv/dt = Jt * P with C = Volume * div(V) = 0.0 = ( vnx * wnx * dx * dx - vpx * wpx * dx * dx ) + ... (divergence theorem)
+// -> dv = inv(M) * Jt * P * dt;
+// Jnx = dC/dvnx = wnx * dx * dx, Jpx = dC/dvpx = -wpx * dx * dx
+// since C = C0 + dC/dv * dv = C0 + J * dv = 0 = C0 + J * inv(M) * Jt * P * dt; and M = wnx * mnx
+// the schur complement is J * inv(M) * Jt * P * dt = -C0 -> P * dt = -C0 / (wnx * dx * dx * wnx * dx * dx / ( wnx * mnx) + wpx * ....)
+// P * dt = - sum( vnx * wnx * dx * dx - vpx * wpx * dx * dx ) / sum (wnx * dx * dx * wnx * dx * dx / ( wnx * mnx);
+//        = - sum( vnx * wnx ) / sum( wnx * dx * dx / mnx ) = -sum( vnx * wnx ) / ( dx * dx * sum( wnx / mnx ) )
+float DIPressureGrid_ComputeGridPressure(in FDIPressureGridContext DIContext, in int3 PrevIndex, out float OutCellDivergence)
+{	
+	float GridPressure = 0.0;
+	OutCellDivergence = 5.0;
+
+	if(DIPressureGrid_IsCellValid(DIContext,PrevIndex))
+	{
+		float GridDivergence = 0.0;
+		float GridFraction = 0.0;
+
+		[unroll]
+		for(int ComponentIndex = 0; ComponentIndex < 3; ++ComponentIndex)
+		{
+			int3 NextIndex = PrevIndex; NextIndex[ComponentIndex] += 1;
+
+			const float PrevMass = DIPressureGrid_GetGridScalar(DIContext,PrevIndex,FLUID_MASS_OFFSET+ComponentIndex);
+			const float PrevWeight = DIPressureGrid_GetGridScalar(DIContext,PrevIndex,SOLID_WEIGHT_OFFSET+ComponentIndex);
+			
+			GridDivergence -= DIPressureGrid_GetGridScalar(DIContext,PrevIndex,SOLID_VELOCITY_OFFSET+ComponentIndex) * (1.0-PrevWeight);
+
+			if( PrevMass > 0.)
+			{
+				GridDivergence -= DIPressureGrid_GetGridScalar(DIContext,PrevIndex,FLUID_VELOCITY_OFFSET+ComponentIndex) * PrevWeight;
+				GridFraction += PrevWeight / PrevMass;
+			}
+
+			const float NextMass = DIPressureGrid_GetGridScalar(DIContext,NextIndex,FLUID_MASS_OFFSET+ComponentIndex);
+			const float NextWeight = DIPressureGrid_GetGridScalar(DIContext,NextIndex,SOLID_WEIGHT_OFFSET+ComponentIndex);
+			
+			GridDivergence += DIPressureGrid_GetGridScalar(DIContext,NextIndex,SOLID_VELOCITY_OFFSET+ComponentIndex) * (1.0-NextWeight);
+			if( NextMass > 0.)
+			{
+				GridDivergence += DIPressureGrid_GetGridScalar(DIContext,NextIndex,FLUID_VELOCITY_OFFSET+ComponentIndex) * NextWeight;
+				GridFraction += NextWeight / NextMass;
+			}
+		}
+		if( GridFraction > 0.)
+		{
+			GridPressure = -GridDivergence / ( DIContext.GridOrigin.w * DIContext.GridOrigin.w * GridFraction);
+		}
+		OutCellDivergence = 3.0;//GridDivergence;
+	}
+	return GridPressure;
+}
+
+// Project the pressure grid to be divergence free
+void DIPressureGrid_ProjectVelocityField(in FDIPressureGridContext DIContext, in int GridCell, out bool OutProjectStatus, out float OutCellDivergence)
+{	
+	OutProjectStatus = false;
+	OutCellDivergence = 10.0;
+
+	const int3 PrevIndex = DIPressureGrid_ComputeGridIndex(DIContext,GridCell);
+	if( DIPressureGrid_IsCellValid(DIContext,PrevIndex))
+	{
+		if( (PrevIndex.x+PrevIndex.y+PrevIndex.z+ShaderStageIndex) % 2 == 0)
+		{
+			OutProjectStatus = true;
+			const float GridPressure = DIPressureGrid_ComputeGridPressure(DIContext,PrevIndex,OutCellDivergence) * DIContext.GridOrigin.w * DIContext.GridOrigin.w;
+		
+			for( int ComponentIndex = 0; ComponentIndex < 3; ++ComponentIndex)
+			{
+				int3 NextIndex = PrevIndex; NextIndex[ComponentIndex] += 1;
+
+				const float PrevMass = DIPressureGrid_GetGridScalar(DIContext,PrevIndex,FLUID_MASS_OFFSET+ComponentIndex);
+				const float PrevWeight = DIPressureGrid_GetGridScalar(DIContext,PrevIndex,SOLID_WEIGHT_OFFSET+ComponentIndex);
+
+				if( PrevMass > 0.0)
+				{
+					DIPressureGrid_AddGridScalar(DIContext,PrevIndex,FLUID_VELOCITY_OFFSET+ComponentIndex,-GridPressure*PrevWeight/PrevMass);
+				}
+
+				const float NextMass = DIPressureGrid_GetGridScalar(DIContext,NextIndex,FLUID_MASS_OFFSET+ComponentIndex);
+				const float NextWeight = DIPressureGrid_GetGridScalar(DIContext,NextIndex,SOLID_WEIGHT_OFFSET+ComponentIndex);
+
+				if( NextMass > 0.0)
+				{
+					DIPressureGrid_AddGridScalar(DIContext,NextIndex,FLUID_VELOCITY_OFFSET+ComponentIndex,GridPressure*NextWeight/NextMass);
+				}
+			}
+		}
+	}
+}
+
+
+
+
+
+
+
+
+
+
diff --git a/Engine/Plugins/Experimental/HairStrands/Shaders/Private/NiagaraHookeSpringMaterial.ush b/Engine/Plugins/Experimental/HairStrands/Shaders/Private/NiagaraHookeSpringMaterial.ush
new file mode 100644
--- /dev/null
+++ b/Engine/Plugins/Experimental/HairStrands/Shaders/Private/NiagaraHookeSpringMaterial.ush
@@ -0,0 +1,238 @@
+// Copyright 1998-2019 Epic Games, Inc. All Rights Reserved.
+
+#pragma once
+
+/* -----------------------------------------------------------------
+ * Stretch spring material
+ * -----------------------------------------------------------------
+ */
+
+ void SetupStretchSpringMaterial(in int StrandsSize, in float YoungModulus, in float RodThickness, 
+	in float RestLength, in float DeltaTime, in bool ProjectConstraint, in float MaterialDamping, out float OutMaterialCompliance, out float OutMaterialWeight, out float OutMaterialMultiplier)
+{
+	const int LocalIndex = (GGroupThreadId.x % StrandsSize);
+	if( LocalIndex > 1 )
+	{
+		// Compliance = 1.0 / (k * dt * dt)
+		// with k = Y * A / L 
+		// A is the cross section area (Pi * R * R), L is the rest length and Y is the young modulus
+		OutMaterialCompliance = RestLength*4.0/(YoungModulus*PI*RodThickness*RodThickness*DeltaTime*DeltaTime);
+
+		const float SumInverseMass = !ProjectConstraint ? ( SharedInverseMass[GGroupThreadId.x] + SharedInverseMass[GGroupThreadId.x-1] ) : SharedInverseMass[GGroupThreadId.x];
+		const float SchurDiagonal = ( (1.0  + MaterialDamping / DeltaTime ) * SumInverseMass + OutMaterialCompliance );
+		OutMaterialWeight = ( SchurDiagonal != 0.0 ) ? 1.0 / SchurDiagonal : 0.0;
+
+		OutMaterialMultiplier = 0;
+	}
+}
+
+void UpdateStretchSpringMultiplier(in float RestLength, in float DeltaTime, in bool ProjectConstraint, in float MaterialDamping, in float MaterialCompliance, in float MaterialWeight, inout float OutMaterialMultiplier)
+{
+	float3 EdgeDirection = SharedNodePosition[GGroupThreadId.x] - SharedNodePosition[GGroupThreadId.x-1];
+	const float3 DeltaVelocity = (EdgeDirection - ( SharedPreviousPosition[GGroupThreadId.x] - SharedPreviousPosition[GGroupThreadId.x-1] ) ) / DeltaTime; 
+
+	const float EdgeLength = length(EdgeDirection);
+	EdgeDirection /= EdgeLength;
+
+	// XPBD lagrange multiplier update : dL = -(C+compliance*L) / (dC * invM * dCt + alpha)
+	const float DeltaLambda = -((EdgeLength-RestLength) + OutMaterialMultiplier * MaterialCompliance + MaterialDamping * dot(EdgeDirection,DeltaVelocity) ) * MaterialWeight;
+
+	// L += dL
+	OutMaterialMultiplier += DeltaLambda;
+
+	// XPBD position update : dX = dL * dCt * invM
+	const float3 PositionDelta = EdgeDirection * DeltaLambda;
+
+	// dX += dX
+	SharedNodePosition[GGroupThreadId.x] += PositionDelta * SharedInverseMass[GGroupThreadId.x];
+	if(!ProjectConstraint)
+	{
+		SharedNodePosition[GGroupThreadId.x-1] -= PositionDelta * SharedInverseMass[GGroupThreadId.x-1];
+	}
+}
+
+void SolveStretchSpringMaterial(in int StrandsSize, in float RestLength, in float DeltaTime, in float MaterialDamping, 
+			in float MaterialCompliance, in float MaterialWeight, in float MaterialMultiplier, out float OutMaterialMultiplier)
+{
+	OutMaterialMultiplier = MaterialMultiplier;
+	const int LocalIndex = (GGroupThreadId.x % StrandsSize);
+	if( LocalIndex > 1)
+	{
+		const int IsRed = (GGroupThreadId.x % 2) == 0;
+		// Process all the red rods
+		if (IsRed)
+		{	
+			UpdateStretchSpringMultiplier(RestLength,DeltaTime,false,MaterialDamping,MaterialCompliance,MaterialWeight,OutMaterialMultiplier);
+		}
+		// Process all the black rods
+		GroupMemoryBarrier(); 
+		if (!IsRed)
+		{
+			UpdateStretchSpringMultiplier(RestLength,DeltaTime,false,MaterialDamping,MaterialCompliance,MaterialWeight,OutMaterialMultiplier);
+		}
+		GroupMemoryBarrier(); 
+	}
+} 
+
+void ProjectStretchSpringMaterial(in int StrandsSize, in float YoungModulus, in float RodThickness, in float RestLength, in float DeltaTime, out float OutMaterialMultiplier)
+{
+	float MaterialCompliance = 0.0;
+	float MaterialWeight = 0.0;
+	OutMaterialMultiplier = 0.0;
+
+	SetupStretchSpringMaterial(StrandsSize, YoungModulus, RodThickness, RestLength, DeltaTime, true, 0.0, MaterialCompliance, MaterialWeight, OutMaterialMultiplier);
+
+	const int LocalIndex = (GGroupThreadId.x % StrandsSize);
+	if( LocalIndex > 1 )
+	{
+		for(int i = 2; i < StrandsSize; ++i)
+		{
+			if( LocalIndex == i )
+			{
+				UpdateStretchSpringMultiplier(RestLength,DeltaTime,true,0.0,MaterialCompliance,MaterialWeight,OutMaterialMultiplier);
+			}
+			GroupMemoryBarrier(); 
+		}
+	}
+}
+
+/* -----------------------------------------------------------------
+ * Compute Bishop frame
+ * -----------------------------------------------------------------
+ */
+
+void ComputeMaterialFrame(in int StrandsSize)
+{
+	const int LocalIndex = (GGroupThreadId.x % StrandsSize);
+	if( LocalIndex == 0 )
+	{
+		float4 NodeQuaternion = SharedNodeOrientation[GGroupThreadId.x];
+		float3 TangentPrev = normalize(SharedNodePosition[GGroupThreadId.x+1] - SharedNodePosition[GGroupThreadId.x]);
+		float3 TangentNext = TangentPrev;
+
+		for (int EdgeIndex = 1, EdgeEnd = StrandsSize-1; EdgeIndex < EdgeEnd; ++EdgeIndex)
+		{
+			TangentPrev = TangentNext;
+			TangentNext = normalize(SharedNodePosition[GGroupThreadId.x+EdgeIndex+1] - SharedNodePosition[GGroupThreadId.x+EdgeIndex]);
+			NodeQuaternion = NormalizeQuat( MultiplyQuat( FindQuatBetweenNormals(TangentPrev,TangentNext), NodeQuaternion ) );
+			SharedNodeOrientation[GGroupThreadId.x+EdgeIndex] = NodeQuaternion;
+			GroupMemoryBarrier();
+		}
+	}
+	GroupMemoryBarrier();
+}
+
+/* -----------------------------------------------------------------
+ * Bend spring material
+ * -----------------------------------------------------------------
+ */
+
+ void SetupBendSpringMaterial(in int StrandsSize, in float YoungModulus, in float RodThickness, 
+	in float RestLength, in float DeltaTime, in bool ProjectConstraint, in float MaterialDamping, out float OutMaterialCompliance, out float OutMaterialWeight, out float3 OutMaterialMultiplier)
+{
+	float MaterialMultiplier = 0;
+	SetupStretchSpringMaterial(StrandsSize,YoungModulus,RodThickness,RestLength,DeltaTime,ProjectConstraint,MaterialDamping,OutMaterialCompliance,OutMaterialWeight,MaterialMultiplier);
+	OutMaterialMultiplier = float3(0,0,0);
+	//ComputeMaterialFrame(StrandsSize);
+}
+
+void UpdateBendSpringMultiplier( in float3 RestDirection, in float DeltaTime, in bool ProjectConstraint, in float MaterialDamping, in float InMaterialCompliance, in float InMaterialWeight, inout float3 OutMaterialMultiplier)
+{
+	const float3 EdgeDirection = SharedNodePosition[GGroupThreadId.x] - SharedNodePosition[GGroupThreadId.x-1];
+	const float3 DeltaVelocity = (EdgeDirection - ( SharedPreviousPosition[GGroupThreadId.x] - SharedPreviousPosition[GGroupThreadId.x-1] ) ) / DeltaTime; 
+
+	const float4 ParentFrame = SharedNodeOrientation[GGroupThreadId.x-2];
+	const float3 TargetDirection = RotateVectorByQuat(RestDirection,ParentFrame);
+
+	// XPBD lagrange multiplier update : dL = -(C+compliance*L) / (dC * invM * dCt + alpha)
+	const float3 DeltaLambda = -(EdgeDirection - TargetDirection + 
+					OutMaterialMultiplier * InMaterialCompliance  + MaterialDamping * DeltaVelocity ) * InMaterialWeight;
+
+	// L += dL
+	OutMaterialMultiplier += DeltaLambda;
+
+	// XPBD position update : dX = dL * dCt * invM
+	const float3 PositionDelta  = DeltaLambda;
+
+	// dX += dX
+	SharedNodePosition[GGroupThreadId.x] += PositionDelta * SharedInverseMass[GGroupThreadId.x];
+	if(!ProjectConstraint)
+	{
+		SharedNodePosition[GGroupThreadId.x-1] -= PositionDelta * SharedInverseMass[GGroupThreadId.x-1];
+	}
+}
+
+void SolveBendSpringMaterial(in int StrandsSize, in float3 RestDirection, in float DeltaTime, in float MaterialDamping,
+			in float MaterialCompliance, in float MaterialWeight, in float3 MaterialMultiplier, out float3 OutMaterialMultiplier)
+{
+	//ComputeMaterialFrame(StrandsSize);
+
+	OutMaterialMultiplier = MaterialMultiplier;
+	const int LocalIndex = (GGroupThreadId.x % StrandsSize);
+	if( LocalIndex > 1)
+	{
+		const int IsRed = (GGroupThreadId.x % 2) == 0;
+		// Process all the red rods
+		if (IsRed)
+		{	
+			UpdateBendSpringMultiplier(RestDirection,DeltaTime,false,MaterialDamping,MaterialCompliance,MaterialWeight,OutMaterialMultiplier);
+		}
+		GroupMemoryBarrier(); 
+
+		if (!IsRed)
+		{
+			UpdateBendSpringMultiplier(RestDirection,DeltaTime,false,MaterialDamping,MaterialCompliance,MaterialWeight,OutMaterialMultiplier);
+		}
+		GroupMemoryBarrier(); 
+	}
+} 
+
+void ProjectBendSpringMaterial(in int StrandsSize, in float YoungModulus, in float RodThickness, in float RestLength, in float3 RestDirection, in float DeltaTime,
+			out float3 OutMaterialMultiplier)
+{
+	float MaterialCompliance = 0.0;
+	float MaterialWeight = 0.0;
+	OutMaterialMultiplier = float3(0,0,0);
+
+	ComputeMaterialFrame(StrandsSize);
+	SetupBendSpringMaterial(StrandsSize, YoungModulus, RodThickness, RestLength, DeltaTime, true, 0.0, MaterialCompliance, MaterialWeight, OutMaterialMultiplier);
+
+	float3 MaterialMultiplier = float3(0,0,0);
+	const int LocalIndex = (GGroupThreadId.x % StrandsSize);
+	if( LocalIndex > 1 )
+	{
+		for(int i = 2; i < StrandsSize; ++i)
+		{
+			if( LocalIndex == i )
+			{
+				UpdateBendSpringMultiplier(RestDirection,DeltaTime,true,0.0,MaterialCompliance,MaterialWeight,OutMaterialMultiplier);
+			}
+			GroupMemoryBarrier(); 
+		}
+	}
+}
+
+/* -----------------------------------------------------------------
+ * Compute Rest Direction
+ * -----------------------------------------------------------------
+ */
+
+void ComputeRestDirection(in int StrandsSize, in float3 NodePosition, in float4 NodeOrientation, out float3 OutRestDirection)
+{
+	SharedNodePosition[GGroupThreadId.x] = NodePosition;
+	SharedNodeOrientation[GGroupThreadId.x] = NodeOrientation;
+	GroupMemoryBarrier(); 
+
+	ComputeMaterialFrame(StrandsSize);
+
+	OutRestDirection = float3(0,0,0);
+	const int LocalIndex = (GGroupThreadId.x % StrandsSize);
+	if( LocalIndex > 1 )
+	{
+		const float3 EdgeDirection = SharedNodePosition[GGroupThreadId.x] - SharedNodePosition[GGroupThreadId.x-1];
+
+		const float4 ParentFrame = SharedNodeOrientation[GGroupThreadId.x-2];
+		OutRestDirection = UnrotateVectorByQuat(EdgeDirection,ParentFrame);
+	}
+	GroupMemoryBarrier(); 
+}
\ No newline at end of file
diff --git a/Engine/Plugins/Experimental/HairStrands/Shaders/Private/NiagaraQuaternionUtils.ush b/Engine/Plugins/Experimental/HairStrands/Shaders/Private/NiagaraQuaternionUtils.ush
--- a/Engine/Plugins/Experimental/HairStrands/Shaders/Private/NiagaraQuaternionUtils.ush
+++ b/Engine/Plugins/Experimental/HairStrands/Shaders/Private/NiagaraQuaternionUtils.ush
@@ -40,6 +40,9 @@ float4 FindQuatBetweenHelper(in float3 A, in float3 B, in float NormAB)
 	float W = NormAB + dot(A,B);
 	float4 Quat = (W>1e-6*NormAB) ? float4(A.y*B.z-A.z*B.y,A.z*B.x-A.x*B.z,A.x*B.y-A.y*B.x,W) :
 			(abs(A.x) > abs(A.y)) ? float4(-A.z,0.0,A.x,0.0) : float4(0.0,-A.z,A.y,0.0);
+
+	//float S = sqrt((NormAB + dot(A,B))*2.0);
+	//float4 Quat = float4(cross(A,B)/S,0.5*S);
 	return NormalizeQuat(Quat);
 }
 
diff --git a/Engine/Plugins/Experimental/HairStrands/Shaders/Private/NiagaraStaticCollisionConstraint.ush b/Engine/Plugins/Experimental/HairStrands/Shaders/Private/NiagaraStaticCollisionConstraint.ush
new file mode 100644
--- /dev/null
+++ b/Engine/Plugins/Experimental/HairStrands/Shaders/Private/NiagaraStaticCollisionConstraint.ush
@@ -0,0 +1,61 @@
+// Copyright 1998-2019 Epic Games, Inc. All Rights Reserved.
+
+#pragma once
+
+/* -----------------------------------------------------------------
+ * Static collision constraint
+ * -----------------------------------------------------------------
+ */
+ float3 ComputeStaticCollisionDisplacement( in float3 DeltaPosition, in float3 DeltaVelocity, in float3 CollisionNormal, in float PenetrationDepth,
+							  in float StaticFriction, in float KineticFriction, in bool IsProjection)
+ {
+	const float DeltaVNormal = dot(DeltaVelocity,CollisionNormal);
+	const float DeltaPNormal = dot(DeltaPosition,CollisionNormal)-PenetrationDepth;
+
+	if (DeltaPNormal < 0.0)
+	{
+		const float DeltaFNormal = IsProjection ? DeltaPNormal : DeltaVNormal;
+		const float NormDeltaFN = abs(DeltaFNormal);
+
+		const float3 CollisionTangent = DeltaVelocity - DeltaVNormal * CollisionNormal;
+		const float TangentLength = length(CollisionTangent);
+
+		float3 CollisionTangentA = (TangentLength > 0.0) ? CollisionTangent/TangentLength : (abs(CollisionNormal.x-1.0f) > 0.0) ? float3(1,0,0) : 
+																					        (abs(CollisionNormal.y-1.0f) > 0.0) ? float3(0,1,0) :
+																						     float3(0,0,1);
+		float3 CollisionTangentB = cross(CollisionNormal,CollisionTangentA);
+		CollisionTangentA = cross(CollisionTangentB,CollisionNormal);
+
+		const float TangentLengthA = length(CollisionTangentA);
+		const float TangentLengthB = length(CollisionTangentB);
+
+		CollisionTangentA = (TangentLengthA > 0.0) ? CollisionTangentA/TangentLengthA : float3(0,0,0);
+		CollisionTangentB = (TangentLengthB > 0.0) ? CollisionTangentB/TangentLengthB : float3(0,0,0);
+
+		const float DeltaVTangentA = dot( DeltaVelocity, CollisionTangentA);
+		const float DeltaVTangentB = dot( DeltaVelocity, CollisionTangentB);
+
+		const float NormDeltaVTA = abs(DeltaVTangentA);
+		const float NormDeltaVTB = abs(DeltaVTangentB);
+
+		const float AlphaTangentA = ( NormDeltaVTA < StaticFriction * NormDeltaFN ) ? 1.0 : (KineticFriction > 0.0) ? min(KineticFriction*NormDeltaFN/NormDeltaVTA, 1.0) : 0.0;
+		const float AlphaTangentB = ( NormDeltaVTB < StaticFriction * NormDeltaFN ) ? 1.0 : (KineticFriction > 0.0) ? min(KineticFriction*NormDeltaFN/NormDeltaVTB, 1.0) : 0.0;
+
+		return DeltaFNormal * CollisionNormal + AlphaTangentA * CollisionTangentA * DeltaVTangentA + AlphaTangentB * CollisionTangentB * DeltaVTangentB;
+	}
+	return float3(0,0,0);
+ }
+
+void SolveStaticCollisionConstraint(in int StrandsSize, in float PenetrationDepth, in float3 CollisionPosition, in float3 CollisionVelocity, in float3 CollisionNormal, 
+				in float StaticFriction, in float KineticFriction, in float DeltaTime, in bool IsProjection, in float3 PreviousPosition, inout float3 OutNodePosition )
+{
+	if( GGroupThreadId.x % StrandsSize > 1)
+	{
+		const float3 NodeVelocity = (OutNodePosition - PreviousPosition)/DeltaTime;
+		const float3 DeltaPosition = OutNodePosition - CollisionPosition;
+		const float3 DeltaVelocity = (NodeVelocity - CollisionVelocity) * DeltaTime;
+
+		const float3 PositionDisplacement = ComputeStaticCollisionDisplacement(DeltaPosition,DeltaVelocity,CollisionNormal,PenetrationDepth,StaticFriction,KineticFriction,IsProjection);
+		OutNodePosition -= PositionDisplacement;
+	}
+}
\ No newline at end of file
diff --git a/Engine/Plugins/Experimental/HairStrands/Source/HairStrandsCore/Private/GroomAsset.cpp b/Engine/Plugins/Experimental/HairStrands/Source/HairStrandsCore/Private/GroomAsset.cpp
--- a/Engine/Plugins/Experimental/HairStrands/Source/HairStrandsCore/Private/GroomAsset.cpp
+++ b/Engine/Plugins/Experimental/HairStrands/Source/HairStrandsCore/Private/GroomAsset.cpp
@@ -1,10 +1,10 @@
 // Copyright 1998-2019 Epic Games, Inc. All Rights Reserved.
 
 #include "GroomAsset.h"
+#include "GroomBuilder.h"
 #include "RenderingThread.h"
 #include "HairStrandsVertexFactory.h"
 #include "Misc/Paths.h"
-#include "UObject/PhysicsObjectVersion.h"
 
 #if WITH_EDITORONLY_DATA
 #include "EditorFramework/AssetImportData.h"
@@ -41,24 +41,20 @@ void CreateBuffer(uint32 InVertexCount , FRWBuffer& OutBuffer)
 
 /////////////////////////////////////////////////////////////////////////////////////////
 
-FHairStrandsResource::FHairStrandsResource(FHairStrandsDatas* HairStrandsDatas) : 
-	RestPositionBuffer(), TangentBuffer(), AttributeBuffer(), StrandsDatas(HairStrandsDatas)
+FHairStrandsResource::FHairStrandsResource(const FHairStrandsDatas::FRenderData& HairStrandRenderData) : 
+	RestPositionBuffer(), TangentBuffer(), AttributeBuffer(), RenderData(HairStrandRenderData)
 {}
 
 void FHairStrandsResource::InitRHI()
 {
-	if (StrandsDatas != nullptr)
-	{
-		TArray<FHairStrandsPositionFormat::Type> RenderingPositions;
-		TArray<FHairStrandsAttributeFormat::Type> RenderingAttributes;
-		StrandsDatas->BuildRenderingDatas(RenderingPositions, RenderingAttributes);
-
-		CreateBuffer<FHairStrandsPositionFormat>(RenderingPositions, RestPositionBuffer);
-		CreateBuffer<FHairStrandsPositionFormat>(RenderingPositions, DeformedPositionBuffer[0]);
-		CreateBuffer<FHairStrandsPositionFormat>(RenderingPositions, DeformedPositionBuffer[1]);
-		CreateBuffer<FHairStrandsTangentFormat>(RenderingPositions.Num() * FHairStrandsTangentFormat::ComponentCount, TangentBuffer);
-		CreateBuffer<FHairStrandsAttributeFormat>(RenderingAttributes, AttributeBuffer);
-	}
+	const TArray<FHairStrandsPositionFormat::Type>& RenderingPositions = RenderData.RenderingPositions;
+	const TArray<FHairStrandsAttributeFormat::Type>& RenderingAttributes = RenderData.RenderingAttributes;
+
+	CreateBuffer<FHairStrandsPositionFormat>(RenderingPositions, RestPositionBuffer);
+	CreateBuffer<FHairStrandsPositionFormat>(RenderingPositions, DeformedPositionBuffer[0]);
+	CreateBuffer<FHairStrandsPositionFormat>(RenderingPositions, DeformedPositionBuffer[1]);
+	CreateBuffer<FHairStrandsTangentFormat>(RenderingPositions.Num() * FHairStrandsTangentFormat::ComponentCount, TangentBuffer);
+	CreateBuffer<FHairStrandsAttributeFormat>(RenderingAttributes, AttributeBuffer);
 }
 
 void FHairStrandsResource::ReleaseRHI()
@@ -72,13 +68,116 @@ void FHairStrandsResource::ReleaseRHI()
 
 /////////////////////////////////////////////////////////////////////////////////////////
 
-FHairStrandsInterpolationResource::FHairStrandsInterpolationResource(const FHairStrandsInterpolationDatas& InData, const FHairStrandsDatas& SimDatas) :
-	Interpolation0Buffer(), Interpolation1Buffer()
+FHairStrandsRootResource::FHairStrandsRootResource(const FHairStrandsDatas* HairStrandsDatas, uint32 LODCount):
+	RootCount(HairStrandsDatas ? HairStrandsDatas->GetNumCurves() : 0)
 {
-	check(Interpolation0.Num()==0);
-	check(Interpolation1.Num()==0);
-	InData.BuildRenderingDatas(Interpolation0, Interpolation1);
+	if (!HairStrandsDatas)
+		return;
+
+	const uint32 CurveCount = HairStrandsDatas->GetNumCurves();
+	CurveIndices.SetNum(HairStrandsDatas->GetNumPoints());
+	RootPositions.SetNum(RootCount);
+	RootNormals.SetNum(RootCount);
+
+
+	for (uint32 CurveIndex = 0; CurveIndex < CurveCount; ++CurveIndex)
+	{
+		const uint32 RootIndex = HairStrandsDatas->StrandsCurves.CurvesOffset[CurveIndex];
+		const uint32 PointCount = HairStrandsDatas->StrandsCurves.CurvesCount[CurveIndex];
+		for (uint32 PointIndex = 0; PointIndex < PointCount; ++PointIndex)
+		{
+			CurveIndices[RootIndex + PointIndex] = CurveIndex; // RootIndex;
+		}
 
+		check(PointCount > 1);
+
+		const FVector P0 = HairStrandsDatas->StrandsPoints.PointsPosition[RootIndex];
+		const FVector P1 = HairStrandsDatas->StrandsPoints.PointsPosition[RootIndex+1];
+		FVector N0 = (P1 - P0).GetSafeNormal();
+
+		// Fallback in case the initial points are too close (this happens on certain assets)
+		if (FVector::DotProduct(N0, N0) == 0)
+		{
+			N0 = FVector(0, 0, 1);
+		}
+
+		FVector4_16 P;
+		P.X = P0.X;
+		P.Y = P0.Y;
+		P.Z = P0.Z;
+		P.W = 1;
+
+		FVector4_16 N;
+		N.X = N0.X;
+		N.Y = N0.Y;
+		N.Z = N0.Z;
+		N.W = 0;
+
+		RootPositions[CurveIndex] = P;
+		RootNormals[CurveIndex] = N;
+	}
+	
+	MeshProjectionLODs.SetNum(LODCount);
+	uint32 LODIndex = 0;
+	for (FMeshProjectionLOD& MeshProjectionLOD : MeshProjectionLODs)
+	{
+		MeshProjectionLOD.Status = FMeshProjectionLOD::EStatus::Invalid;
+		MeshProjectionLOD.LODIndex = LODIndex++;
+	}
+}
+
+void FHairStrandsRootResource::InitRHI()
+{
+	if (CurveIndices.Num() > 0)
+	{
+		CreateBuffer<FHairStrandsIndexFormat>(CurveIndices, VertexToCurveIndexBuffer);
+		CreateBuffer<FHairStrandsRootPositionFormat>(RootPositions, RootPositionBuffer);
+		CreateBuffer<FHairStrandsRootNormalFormat>(RootNormals, RootNormalBuffer);	
+		for (FMeshProjectionLOD& LOD : MeshProjectionLODs)
+		{
+			LOD.Status = FHairStrandsRootResource::FMeshProjectionLOD::EStatus::Initialized;
+
+			// Create buffers. Initialization will be done by render passes
+			CreateBuffer<FHairStrandsCurveTriangleIndexFormat>(RootCount, LOD.RootTriangleIndexBuffer);
+			CreateBuffer<FHairStrandsCurveTriangleBarycentricFormat>(RootCount, LOD.RootTriangleBarycentricBuffer);
+			CreateBuffer<FHairStrandsCurveTranslationFormat>(RootCount, LOD.RestRootTrianglePosition0Buffer);
+			CreateBuffer<FHairStrandsCurveTranslationFormat>(RootCount, LOD.RestRootTrianglePosition1Buffer);
+			CreateBuffer<FHairStrandsCurveTranslationFormat>(RootCount, LOD.RestRootTrianglePosition2Buffer);
+
+			// Strand hair roots translation and rotation in triangle-deformed position relative to the bound triangle 
+			CreateBuffer<FHairStrandsCurveTranslationFormat>(RootCount, LOD.DeformedRootTrianglePosition0Buffer);
+			CreateBuffer<FHairStrandsCurveTranslationFormat>(RootCount, LOD.DeformedRootTrianglePosition1Buffer);
+			CreateBuffer<FHairStrandsCurveTranslationFormat>(RootCount, LOD.DeformedRootTrianglePosition2Buffer);
+		}
+	}
+}
+
+void FHairStrandsRootResource::ReleaseRHI()
+{
+	RootPositionBuffer.Release();
+	RootNormalBuffer.Release();
+	VertexToCurveIndexBuffer.Release();
+
+	for (FMeshProjectionLOD& LOD : MeshProjectionLODs)
+	{
+		LOD.Status = FHairStrandsRootResource::FMeshProjectionLOD::EStatus::Invalid;
+		LOD.RootTriangleIndexBuffer.Release();
+		LOD.RootTriangleBarycentricBuffer.Release();
+		LOD.RestRootTrianglePosition0Buffer.Release();
+		LOD.RestRootTrianglePosition1Buffer.Release();
+		LOD.RestRootTrianglePosition2Buffer.Release();
+		LOD.DeformedRootTrianglePosition0Buffer.Release();
+		LOD.DeformedRootTrianglePosition1Buffer.Release();
+		LOD.DeformedRootTrianglePosition2Buffer.Release();
+	}
+	MeshProjectionLODs.Empty();
+}
+
+/////////////////////////////////////////////////////////////////////////////////////////
+
+FHairStrandsInterpolationResource::FHairStrandsInterpolationResource(const FHairStrandsInterpolationDatas::FRenderData& InterpolationRenderData, const FHairStrandsDatas& SimDatas) :
+	Interpolation0Buffer(), Interpolation1Buffer(), RenderData(InterpolationRenderData)
+{
 	const uint32 RootCount = SimDatas.GetNumCurves();
 	SimRootPointIndex.SetNum(SimDatas.GetNumPoints());
 	for (uint32 CurveIndex = 0; CurveIndex < RootCount; ++CurveIndex)
@@ -94,22 +193,17 @@ FHairStrandsInterpolationResource::FHairStrandsInterpolationResource(const FHair
 
 void FHairStrandsInterpolationResource::InitRHI()
 {
-	check(Interpolation0.Num() == Interpolation1.Num());
-	if (Interpolation0.Num() > 0)
-	{
-		CreateBuffer<FHairStrandsInterpolation0Format>(Interpolation0, Interpolation0Buffer);
-		CreateBuffer<FHairStrandsInterpolation1Format>(Interpolation1, Interpolation1Buffer);
-		CreateBuffer<FHairStrandsRootIndexFormat>(SimRootPointIndex, SimRootPointIndexBuffer);
-		Interpolation0.SetNum(0);
-		Interpolation1.SetNum(0);
-		SimRootPointIndex.SetNum(0);
-	}
+	CreateBuffer<FHairStrandsInterpolation0Format>(RenderData.Interpolation0, Interpolation0Buffer);
+	CreateBuffer<FHairStrandsInterpolation1Format>(RenderData.Interpolation1, Interpolation1Buffer);
+	CreateBuffer<FHairStrandsRootIndexFormat>(SimRootPointIndex, SimRootPointIndexBuffer);
+	SimRootPointIndex.SetNum(0);
 }
 
 void FHairStrandsInterpolationResource::ReleaseRHI()
 {
 	Interpolation0Buffer.Release();
 	Interpolation1Buffer.Release();
+	SimRootPointIndexBuffer.Release();
 }
 
 /////////////////////////////////////////////////////////////////////////////////////////
@@ -117,8 +211,8 @@ void FHairStrandsInterpolationResource::ReleaseRHI()
 #if RHI_RAYTRACING
 // RT geometry is built to for a cross around the fiber.
 // 4 triangles per hair vertex => 12 vertices per hair vertex
-FHairStrandsRaytracingResource::FHairStrandsRaytracingResource(FHairStrandsDatas* InData) :
-	PositionBuffer(), VertexCount(InData->GetNumPoints()*12)  
+FHairStrandsRaytracingResource::FHairStrandsRaytracingResource(const FHairStrandsDatas& InData) :
+	PositionBuffer(), VertexCount(InData.GetNumPoints()*12)  
 {}
 
 void FHairStrandsRaytracingResource::InitRHI()
@@ -140,68 +234,66 @@ void UGroomAsset::Serialize(FArchive& Ar)
 {
 	Super::Serialize(Ar);
 
-	Ar.UsingCustomVersion(FPhysicsObjectVersion::GUID);
-
-	// #ueent_todo: Serialize HairDescription
-
-	HairRenderData.Serialize(Ar);
-
-	if (Ar.CustomVer(FPhysicsObjectVersion::GUID) >= FPhysicsObjectVersion::HairAssetSerialization)
-	{
-		HairSimulationData.Serialize(Ar);
-	}
+	Ar << HairGroupsData;
 }
 
 void UGroomAsset::InitResource()
 {
-	check(!HairStrandsResource);
-
-	HairStrandsResource = new FHairStrandsResource(&HairRenderData);
+	for (FHairGroupData& GroupData : HairGroupsData)
+	{
+		GroupData.HairStrandsResource = new FHairStrandsResource(GroupData.HairRenderData.RenderData);
 
-	BeginInitResource(HairStrandsResource);
+		BeginInitResource(GroupData.HairStrandsResource);
 
-	HairSimulationResource = new FHairStrandsResource(&HairSimulationData);
+		GroupData.HairSimulationResource = new FHairStrandsResource(GroupData.HairSimulationData.RenderData);
 
-	BeginInitResource(HairSimulationResource);
+		BeginInitResource(GroupData.HairSimulationResource);
+	}
 }
 
 void UGroomAsset::UpdateResource()
 {
-	if (HairStrandsResource)
+	for (FHairGroupData& GroupData : HairGroupsData)
 	{
-		BeginUpdateResourceRHI(HairStrandsResource);
-	}
+		if (GroupData.HairStrandsResource)
+		{
+			BeginUpdateResourceRHI(GroupData.HairStrandsResource);
+		}
 
-	if (HairSimulationResource)
-	{
-		BeginUpdateResourceRHI(HairSimulationResource);
+		if (GroupData.HairSimulationResource)
+		{
+			BeginUpdateResourceRHI(GroupData.HairSimulationResource);
+		}
 	}
 }
 
 void UGroomAsset::ReleaseResource()
 {
-	if (HairStrandsResource)
+	for (FHairGroupData& GroupData : HairGroupsData)
 	{
-		FHairStrandsResource* InResource = HairStrandsResource;
-		ENQUEUE_RENDER_COMMAND(ReleaseHairStrandsResourceCommand)(
-			[InResource](FRHICommandList& RHICmdList)
+		if (GroupData.HairStrandsResource)
 		{
-			InResource->ReleaseResource();
-			delete InResource;
-		});
-		HairStrandsResource = nullptr;
-	}
+			FHairStrandsResource* InResource = GroupData.HairStrandsResource;
+			ENQUEUE_RENDER_COMMAND(ReleaseHairStrandsResourceCommand)(
+				[InResource](FRHICommandList& RHICmdList)
+			{
+				InResource->ReleaseResource();
+				delete InResource;
+			});
+			GroupData.HairStrandsResource = nullptr;
+		}
 
-	if (HairSimulationResource)
-	{
-		FHairStrandsResource* InResource = HairSimulationResource;
-		ENQUEUE_RENDER_COMMAND(ReleaseHairStrandsResourceCommand)(
-			[InResource](FRHICommandList& RHICmdList)
+		if (GroupData.HairSimulationResource)
 		{
-			InResource->ReleaseResource();
-			delete InResource;
-		});
-		HairSimulationResource = nullptr;
+			FHairStrandsResource* InResource = GroupData.HairSimulationResource;
+			ENQUEUE_RENDER_COMMAND(ReleaseHairStrandsResourceCommand)(
+				[InResource](FRHICommandList& RHICmdList)
+			{
+				InResource->ReleaseResource();
+				delete InResource;
+			});
+			GroupData.HairSimulationResource = nullptr;
+		}
 	}
 }
 
@@ -209,39 +301,18 @@ void UGroomAsset::Reset()
 {
 	ReleaseResource();
 
-	HairRenderData.Reset();
-	HairSimulationData.Reset();
-
-	RenderHairGroups.Reset();
-	SimulationHairGroups.Reset();
+	HairGroupsInfo.Reset();
+	HairGroupsData.Reset();
 }
 
 void UGroomAsset::PostLoad()
 {
 	Super::PostLoad();
 
-	// If the asset was previously serialized without hair sim data, generate it here
-	bool bAutoGeneratedGuides = false;
-	if (HairSimulationData.StrandsPoints.Num() == 0 || HairSimulationData.StrandsCurves.Num() == 0)
-	{
-		DecimateStrandData(HairRenderData, FMath::Clamp(HairToGuideDensity, 0.f, 1.f), HairSimulationData);
-		bAutoGeneratedGuides = true;
-	}
-
-	// If the asset was previously serialized without group settings, fill them out here
-	if (RenderHairGroups.Num() == 0)
+	check(HairGroupsData.Num() > 0);
+	if (HairGroupsData[0].HairSimulationData.GetNumCurves() == 0 || HairGroupsData[0].HairInterpolationData.Num() == 0)
 	{
-		FHairGroupRenderSettings& GroupSettings = RenderHairGroups.AddDefaulted_GetRef();
-		GroupSettings.GroupID = 0;
-		GroupSettings.NumCurves = HairRenderData.GetNumCurves();
-	}
-
-	if (SimulationHairGroups.Num() == 0)
-	{
-		FHairGroupSimulationSettings& GroupSettings = SimulationHairGroups.AddDefaulted_GetRef();
-		GroupSettings.GroupID = 0;
-		GroupSettings.NumCurves = HairSimulationData.GetNumCurves();
-		GroupSettings.bIsAutoGenerated = bAutoGeneratedGuides;
+		FGroomBuilder::BuildData(this);
 	}
 
 	if (!IsTemplate())
@@ -286,4 +357,16 @@ void UGroomAsset::PostInitProperties()
 }
 #endif
 
+int32 UGroomAsset::GetNumHairGroups() const
+{
+	return HairGroupsData.Num();
+}
+
+FArchive& operator<<(FArchive& Ar, FHairGroupData& GroupData)
+{
+	GroupData.HairRenderData.Serialize(Ar);
+	GroupData.HairSimulationData.Serialize(Ar);
+	GroupData.HairInterpolationData.Serialize(Ar);
 
+	return Ar;
+}
