Copying //UE4/Dev-Physics to Dev-Main (//UE4/Dev-Main) [at] 9420715

diff --git a/Engine/Plugins/Experimental/HairStrands/Source/HairStrandsCore/Private/GroomBuilder.cpp b/Engine/Plugins/Experimental/HairStrands/Source/HairStrandsCore/Private/GroomBuilder.cpp
new file mode 100644
--- /dev/null
+++ b/Engine/Plugins/Experimental/HairStrands/Source/HairStrandsCore/Private/GroomBuilder.cpp
@@ -0,0 +1,957 @@
+// Copyright 1998-2019 Epic Games, Inc. All Rights Reserved.
+
+#pragma once
+
+#include "GroomBuilder.h"
+#include "GroomAsset.h"
+#include "GroomComponent.h"
+#include "GroomSettings.h"
+#include "HairDescription.h"
+
+#include "Async/ParallelFor.h"
+#include "HAL/IConsoleManager.h"
+#include "Math/UnrealMathUtility.h"
+#include "Misc/ScopedSlowTask.h"
+
+DEFINE_LOG_CATEGORY_STATIC(LogGroomBuilder, Log, All);
+
+#define LOCTEXT_NAMESPACE "GroomBuilder"
+
+// For debug purpose
+static float GHairInterpolationMetric_Distance = 1;
+static float GHairInterpolationMetric_Angle = 0;
+static float GHairInterpolationMetric_Length = 0;
+static float GHairInterpolationMetric_AngleAttenuation = 5;
+static FAutoConsoleVariableRef CVarHairInterpolationMetric_Distance(TEXT("r.HairStrands.InterpolationMetric.Distance"), GHairInterpolationMetric_Distance, TEXT("Hair strands interpolation metric weights for distance"));
+static FAutoConsoleVariableRef CVarHairInterpolationMetric_Angle(TEXT("r.HairStrands.InterpolationMetric.Angle"), GHairInterpolationMetric_Angle, TEXT("Hair strands interpolation metric weights for angle"));
+static FAutoConsoleVariableRef CVarHairInterpolationMetric_Length(TEXT("r.HairStrands.InterpolationMetric.Length"), GHairInterpolationMetric_Length, TEXT("Hair strands interpolation metric weights for length"));
+static FAutoConsoleVariableRef CVarHairInterpolationMetric_AngleAttenuation(TEXT("r.HairStrands.InterpolationMetric.AngleAttenuation"), GHairInterpolationMetric_AngleAttenuation, TEXT("Hair strands interpolation angle attenuation"));
+
+namespace HairStrandsBuilder
+{
+	FVector2D SignNotZero(const FVector2D& v)
+	{
+		return FVector2D((v.X >= 0.0) ? +1.0 : -1.0, (v.Y >= 0.0) ? +1.0 : -1.0);
+	}
+
+	// A Survey of Efficient Representations for Independent Unit Vectors
+	// Reference: http://jcgt.org/published/0003/02/01/paper.pdf
+	// Assume normalized input. Output is on [-1, 1] for each component.
+	FVector2D SphericalToOctahedron(const FVector& v)
+	{
+		// Project the sphere onto the octahedron, and then onto the xy plane
+		FVector2D p = FVector2D(v.X, v.Y) * (1.0 / (abs(v.X) + abs(v.Y) + abs(v.Z)));
+		// Reflect the folds of the lower hemisphere over the diagonals
+		return (v.Z <= 0.0) ? ((FVector2D(1, 1) - FVector2D(abs(p.Y), abs(p.X))) * SignNotZero(p)) : p;
+	}
+
+	// Auto-generate Root UV data if not loaded
+	void ComputeRootUV(FHairStrandsCurves& Curves, FHairStrandsPoints& Points)
+	{
+		TRACE_CPUPROFILER_EVENT_SCOPE(HairStrandsBuilder::ComputeRootUV);
+
+		TArray<FVector> RootPoints;
+		const uint32 CurveCount = Curves.Num();
+		RootPoints.Reserve(Curves.Num());
+		FVector MinAABB(MAX_FLT, MAX_FLT, MAX_FLT);
+		FVector MaxAABB(-MAX_FLT, -MAX_FLT, -MAX_FLT);
+		FMatrix Rotation = FRotationMatrix::Make(FRotator(0, 0, -90));
+		for (uint32 CurveIndex=0; CurveIndex< CurveCount; ++CurveIndex)
+		{
+			const uint32 Offset = Curves.CurvesOffset[CurveIndex];
+			check(Offset < uint32(Points.PointsPosition.Num()));
+			const FVector P = Rotation.TransformPosition(Points.PointsPosition[Offset]);
+
+			RootPoints.Add(P);
+			MinAABB.X = FMath::Min(P.X, MinAABB.X);
+			MinAABB.Y = FMath::Min(P.Y, MinAABB.Y);
+			MinAABB.Z = FMath::Min(P.Z, MinAABB.Z);
+
+			MaxAABB.X = FMath::Max(P.X, MaxAABB.X);
+			MaxAABB.Y = FMath::Max(P.Y, MaxAABB.Y);
+			MaxAABB.Z = FMath::Max(P.Z, MaxAABB.Z);
+		}
+
+		// Compute sphere bound
+		const FVector Extent = MaxAABB - MinAABB;
+		FSphere SBound;
+		SBound.Center = (MaxAABB + MinAABB) * 0.5f;
+		SBound.W = FMath::Max(Extent.X, FMath::Max(Extent.Y, Extent.Z));
+
+		// Project root point onto the bounding sphere and map it onto 
+		// an octahedron, which is unfold onto the unit space [0,1]^2
+		TArray<FVector2D> RootUVs;
+		RootUVs.Reserve(Curves.Num());
+		FVector2D MinUV(MAX_FLT, MAX_FLT);
+		FVector2D MaxUV(-MAX_FLT, -MAX_FLT);
+		for (const FVector& RootP : RootPoints)
+		{
+			FVector D = RootP - SBound.Center;
+			D.Normalize();
+			FVector2D UV = SphericalToOctahedron(D);
+			UV += FVector2D(1, 1);
+			UV *= 0.5f;
+			RootUVs.Add(UV);
+
+			MinUV.X = FMath::Min(UV.X, MinUV.X);
+			MinUV.Y = FMath::Min(UV.Y, MinUV.Y);
+
+			MaxUV.X = FMath::Max(UV.X, MaxUV.X);
+			MaxUV.Y = FMath::Max(UV.Y, MaxUV.Y);
+		}
+
+		// Find the minimal UV space cover by root point, and 
+		// offsets/scales it to maximize UV space
+		const FVector2D UVScale(1 / (MaxUV.X - MinUV.X), 1 / (MaxUV.Y - MinUV.Y));
+		const FVector2D UVOffset(-MinUV.X, -MinUV.Y);
+		uint32 Index = 0;
+		for (FVector2D& RootUV : Curves.CurvesRootUV)
+		{
+			RootUV = (RootUVs[Index++] + UVOffset) * UVScale;
+		}
+	}
+
+	/** Build the internal points and curves data */
+	void BuildInternalData(FHairStrandsDatas& HairStrands, bool bComputeRootUV)
+	{
+		TRACE_CPUPROFILER_EVENT_SCOPE(HairStrandsBuilder::BuildInternalData);
+
+		FHairStrandsCurves& Curves = HairStrands.StrandsCurves;
+		FHairStrandsPoints& Points = HairStrands.StrandsPoints;
+
+		HairStrands.BoundingBox.Min = {  FLT_MAX,  FLT_MAX ,  FLT_MAX };
+		HairStrands.BoundingBox.Max = { -FLT_MAX, -FLT_MAX , -FLT_MAX };
+
+		if (HairStrands.GetNumCurves() > 0 && HairStrands.GetNumPoints() > 0)
+		{
+			TArray<FVector>::TIterator PositionIterator = Points.PointsPosition.CreateIterator();
+			TArray<float>::TIterator RadiusIterator = Points.PointsRadius.CreateIterator();
+			TArray<float>::TIterator CoordUIterator = Points.PointsCoordU.CreateIterator();
+
+			TArray<uint16>::TIterator CountIterator = Curves.CurvesCount.CreateIterator();
+			TArray<uint32>::TIterator OffsetIterator = Curves.CurvesOffset.CreateIterator();
+			TArray<float>::TIterator LengthIterator = Curves.CurvesLength.CreateIterator();
+
+			Curves.MaxRadius = 0.0;
+			Curves.MaxLength = 0.0;
+
+			uint32 StrandOffset = 0;
+			*OffsetIterator = StrandOffset; ++OffsetIterator;
+
+			for (uint32 CurveIndex = 0; CurveIndex < HairStrands.GetNumCurves(); ++CurveIndex, ++OffsetIterator, ++LengthIterator, ++CountIterator)
+			{
+				const uint16& StrandCount = *CountIterator;
+
+				StrandOffset += StrandCount;
+				*OffsetIterator = StrandOffset;
+
+				float StrandLength = 0.0;
+				FVector PreviousPosition(0.0, 0.0, 0.0);
+				for (uint32 PointIndex = 0; PointIndex < StrandCount; ++PointIndex, ++PositionIterator, ++RadiusIterator, ++CoordUIterator)
+				{
+					HairStrands.BoundingBox += *PositionIterator;
+
+					if (PointIndex > 0)
+					{
+						StrandLength += (*PositionIterator - PreviousPosition).Size();
+					}
+					*CoordUIterator = StrandLength;
+					PreviousPosition = *PositionIterator;
+
+					Curves.MaxRadius = FMath::Max(Curves.MaxRadius, *RadiusIterator);
+				}
+				*LengthIterator = StrandLength;
+				Curves.MaxLength = FMath::Max(Curves.MaxLength, StrandLength);
+			}
+
+			CountIterator.Reset();
+			LengthIterator.Reset();
+			RadiusIterator.Reset();
+			CoordUIterator.Reset();
+
+			for (uint32 CurveIndex = 0; CurveIndex < HairStrands.GetNumCurves(); ++CurveIndex, ++LengthIterator, ++CountIterator)
+			{
+				const uint16& StrandCount = *CountIterator;
+
+				for (uint32 PointIndex = 0; PointIndex < StrandCount; ++PointIndex, ++RadiusIterator, ++CoordUIterator)
+				{
+					*CoordUIterator /= *LengthIterator;
+					*RadiusIterator /= Curves.MaxRadius;
+				}
+				*LengthIterator /= Curves.MaxLength;
+			}
+
+			if (bComputeRootUV)
+			{
+				ComputeRootUV(Curves, Points);
+			}
+		}
+	}
+
+	inline void CopyVectorToPosition(const FVector& InVector, FHairStrandsPositionFormat::Type& OutPosition)
+	{
+		OutPosition.X = InVector.X;
+		OutPosition.Y = InVector.Y;
+		OutPosition.Z = InVector.Z;
+	}
+
+	/** Build the packed datas for gpu rendering/simulation */
+	void BuildRenderData(FHairStrandsDatas& HairStrands)
+	{
+		TRACE_CPUPROFILER_EVENT_SCOPE(HairStrandsBuilder::BuildRenderData);
+
+		uint32 NumCurves = HairStrands.GetNumCurves();
+		uint32 NumPoints = HairStrands.GetNumPoints();
+		if (!(NumCurves > 0 && NumPoints > 0))
+			return;
+
+		TArray<FHairStrandsPositionFormat::Type>& OutPackedPositions = HairStrands.RenderData.RenderingPositions;
+		TArray<FHairStrandsAttributeFormat::Type>& OutPackedAttributes = HairStrands.RenderData.RenderingAttributes;
+
+		OutPackedPositions.SetNum(NumPoints * FHairStrandsPositionFormat::ComponentCount);
+		OutPackedAttributes.SetNum(NumPoints * FHairStrandsAttributeFormat::ComponentCount);
+
+		const FVector HairBoxCenter = HairStrands.BoundingBox.GetCenter();
+
+		FHairStrandsCurves& Curves = HairStrands.StrandsCurves;
+		FHairStrandsPoints& Points = HairStrands.StrandsPoints;
+
+		FRandomStream Random;
+		for (uint32 CurveIndex = 0; CurveIndex < NumCurves; ++CurveIndex)
+		{
+			const float CurveSeed = Random.RandHelper(255);
+			const int32 IndexOffset = Curves.CurvesOffset[CurveIndex];
+			const uint16& PointCount = Curves.CurvesCount[CurveIndex];
+			for (int32 PointIndex = 0; PointIndex < PointCount; ++PointIndex)
+			{
+				const uint32 PrevIndex = FMath::Max(0, PointIndex - 1);
+				const uint32 NextIndex = FMath::Min(PointCount + 1, PointCount - 1);
+				const FVector& PointPosition = Points.PointsPosition[PointIndex + IndexOffset];
+
+				const float CoordU = Points.PointsCoordU[PointIndex + IndexOffset];
+				const float NormalizedRadius = Points.PointsRadius[PointIndex + IndexOffset];
+				const float NormalizedLength = CoordU * Curves.CurvesLength[CurveIndex];
+
+				FHairStrandsPositionFormat::Type& PackedPosition = OutPackedPositions[PointIndex + IndexOffset];
+				CopyVectorToPosition(PointPosition - HairBoxCenter, PackedPosition);
+				PackedPosition.ControlPointType = (PointIndex == 0) ? 1u : (PointIndex == (PointCount - 1) ? 2u : 0u);
+				PackedPosition.NormalizedRadius = uint8(FMath::Clamp(NormalizedRadius * 63.f, 0.f, 63.f));
+				PackedPosition.NormalizedLength = uint8(FMath::Clamp(NormalizedLength *255.f, 0.f, 255.f));
+
+				const FVector2D RootUV = Curves.CurvesRootUV[CurveIndex];
+				FHairStrandsAttributeFormat::Type& PackedAttributes = OutPackedAttributes[PointIndex + IndexOffset];
+				PackedAttributes.RootU = uint32(FMath::Clamp(RootUV.X, 0.f, 1.f) * 0xFF);
+				PackedAttributes.RootV = uint32(FMath::Clamp(RootUV.Y, 0.f, 1.f) * 0xFF);
+				PackedAttributes.UCoord = FMath::Clamp(CoordU * 255.f, 0.f, 255.f);
+				PackedAttributes.Seed = CurveSeed;
+			}
+		}
+	}
+}
+
+namespace HairInterpolationBuilder
+{
+	struct FHairRoot
+	{
+		FVector Position;
+		uint32  VertexCount;
+		FVector Normal;
+		uint32  Index;
+		float   Length;
+	};
+
+	struct FHairInterpolationMetric
+	{	
+		// Total/combined metrics
+		float Metric;
+
+		// Debug info
+		float DistanceMetric;
+		float AngularMetric;
+		float LengthMetric;
+
+		float CosAngle;
+		float Distance;
+		float GuideLength;
+		float RenderLength;
+
+	};
+
+	inline FHairInterpolationMetric ComputeInterpolationMetric(const FHairRoot& RenderRoot, const FHairRoot& GuideRoot)
+	{
+		FHairInterpolationMetric Out;
+		Out.Distance = FVector::Dist(RenderRoot.Position, GuideRoot.Position);
+		Out.CosAngle = FVector::DotProduct(RenderRoot.Normal, GuideRoot.Normal);
+		Out.GuideLength = GuideRoot.Length;
+		Out.RenderLength = RenderRoot.Length;
+
+		// Metric takes into account the following properties to find guides which are close, share similar orientation, and 
+		// have similar length for better interpolation
+		// * distance
+		// * orientation 
+		// * length 
+		const float AngularAttenuation = GHairInterpolationMetric_AngleAttenuation > 1 ? GHairInterpolationMetric_AngleAttenuation : 0;
+		Out.DistanceMetric	= Out.Distance * GHairInterpolationMetric_Distance;
+		Out.AngularMetric	= AngularAttenuation == 0 ? 0 : (FMath::Clamp((1 - FMath::Pow(Out.CosAngle, AngularAttenuation)), 0.f, 1.f) * GHairInterpolationMetric_Angle);
+		Out.LengthMetric	= FMath::Abs(FMath::Max(Out.GuideLength / float(Out.RenderLength), Out.RenderLength / float(Out.GuideLength)) - 1) * GHairInterpolationMetric_Length; // Ratio
+		Out.Metric			= Out.DistanceMetric + Out.AngularMetric + Out.LengthMetric;
+
+		return Out;
+	}
+
+	template<uint32 NumSamples>
+	inline FVector GetCurvePosition(const FHairStrandsDatas& CurvesDatas, const uint32 CurveIndex, const uint32 SampleIndex)
+	{
+		const float PointCount = CurvesDatas.StrandsCurves.CurvesCount[CurveIndex]-1.0;
+		const uint32 PointOffset = CurvesDatas.StrandsCurves.CurvesOffset[CurveIndex];
+
+		const float CurvePoint = static_cast<float>(SampleIndex) * PointCount / (static_cast<float>(NumSamples)-1.0f);
+		const uint32 PointPrev = (SampleIndex == 0) ? 0 : (SampleIndex == (NumSamples-1)) ? PointCount - 1 : floor(CurvePoint);
+		const uint32 PointNext = PointPrev + 1;
+
+		const float PointAlpha = CurvePoint - static_cast<float>(PointPrev);
+		return CurvesDatas.StrandsPoints.PointsPosition[PointOffset+PointPrev]*(1.0f-PointAlpha) + 
+			CurvesDatas.StrandsPoints.PointsPosition[PointOffset+PointNext]*PointAlpha;
+	}
+
+	template<uint32 NumSamples>
+	inline float ComputeCurvesMetric(const FHairStrandsDatas& RenderCurvesDatas, const uint32 RenderCurveIndex, 
+		const FHairStrandsDatas& GuideCurvesDatas, const uint32 GuideCurveIndex, const float RootImportance, 
+		const float ShapeImportance, const float ProximityImportance)
+	{
+		const float AverageLength = FMath::Max(0.5f * (RenderCurvesDatas.StrandsCurves.CurvesLength[RenderCurveIndex] * RenderCurvesDatas.StrandsCurves.MaxLength +
+			GuideCurvesDatas.StrandsCurves.CurvesLength[GuideCurveIndex] * GuideCurvesDatas.StrandsCurves.MaxLength), SMALL_NUMBER);
+
+		static const float DeltaCoord = 1.0f / static_cast<float>(NumSamples-1);
+
+		const FVector& RenderRoot = RenderCurvesDatas.StrandsPoints.PointsPosition[RenderCurvesDatas.StrandsCurves.CurvesOffset[RenderCurveIndex]];
+		const FVector& GuideRoot = GuideCurvesDatas.StrandsPoints.PointsPosition[GuideCurvesDatas.StrandsCurves.CurvesOffset[GuideCurveIndex]];
+
+		float CurveProximityMetric = 0.0;
+		float CurveShapeMetric = 0.0;
+		for (uint32 SampleIndex = 0; SampleIndex < NumSamples; ++SampleIndex)
+		{
+			const FVector GuidePosition = GetCurvePosition<NumSamples>(GuideCurvesDatas, GuideCurveIndex, SampleIndex);
+			const FVector RenderPosition = GetCurvePosition<NumSamples>(RenderCurvesDatas, RenderCurveIndex, SampleIndex);
+			const float RootWeight = FMath::Exp(-RootImportance*SampleIndex*DeltaCoord);
+
+			CurveProximityMetric += (GuidePosition - RenderPosition).Size() * RootWeight;
+			CurveShapeMetric += (GuidePosition - GuideRoot - RenderPosition + RenderRoot).Size() * RootWeight;
+		}
+		CurveShapeMetric *= DeltaCoord / AverageLength;
+		CurveProximityMetric *= DeltaCoord / AverageLength;
+
+		return FMath::Exp(-ShapeImportance*CurveShapeMetric) * FMath::Exp(-ProximityImportance * CurveProximityMetric);
+	}
+
+	inline void PrintInterpolationMetric(const FHairInterpolationMetric& In)
+	{
+		UE_LOG(LogGroomBuilder, Log, TEXT("------------------------------------------------------------------------------------------"));
+		UE_LOG(LogGroomBuilder, Log, TEXT("Total Metric = %f"), In.Metric);
+		UE_LOG(LogGroomBuilder, Log, TEXT("Distance     = %f (%f)"), In.Distance, In.DistanceMetric);
+		UE_LOG(LogGroomBuilder, Log, TEXT("Angle        = %f (%f)"), FMath::RadiansToDegrees(FMath::Acos(In.CosAngle)), In.AngularMetric);
+		UE_LOG(LogGroomBuilder, Log, TEXT("Length       = %f/%f (%f)"), In.RenderLength, In.GuideLength, In.LengthMetric);
+	}
+
+	template<typename T>
+	void SwapValue(T& A, T& B)
+	{
+		T Temp = A;
+		A = B;
+		B = Temp;
+	}
+
+	void BuildInterpolationData(FHairStrandsInterpolationDatas& InterpolationData,
+		const FHairStrandsDatas& SimStrandsData,
+		const FHairStrandsDatas& RenStrandsData)
+	{
+		TRACE_CPUPROFILER_EVENT_SCOPE(HairInterpolationBuilder::BuildInterpolationData);
+
+		InterpolationData.SetNum(RenStrandsData.GetNumPoints());
+
+		typedef TArray<FHairRoot> FRoots;
+
+		// Extract strand roots
+		auto ExtractRoot = [](const FHairStrandsDatas& InData, TArray<FHairRoot>& OutRoots)
+		{
+			const uint32 CurveCount = InData.StrandsCurves.Num();
+			OutRoots.Reserve(CurveCount);
+			for (uint32 CurveIndex = 0; CurveIndex < CurveCount; ++CurveIndex)
+			{
+				const uint32 PointOffset = InData.StrandsCurves.CurvesOffset[CurveIndex];
+				const uint32 PointCount  = InData.StrandsCurves.CurvesCount[CurveIndex];
+				const float  CurveLength = InData.StrandsCurves.CurvesLength[CurveIndex] * InData.StrandsCurves.MaxLength;
+				check(PointCount > 1);
+				const FVector& P0 = InData.StrandsPoints.PointsPosition[PointOffset];
+				const FVector& P1 = InData.StrandsPoints.PointsPosition[PointOffset+1];
+				FVector N = (P1 - P0).GetSafeNormal();
+
+				// Fallback in case the initial points are too close (this happens on certain assets)
+				if (FVector::DotProduct(N, N) == 0)
+				{
+					N = FVector(0, 0, 1);
+				}
+				OutRoots.Add({ P0, PointCount, N, PointOffset, CurveLength });
+			}
+		};
+
+		FRoots RenRoots, SimRoots;
+		ExtractRoot(RenStrandsData, RenRoots);
+		ExtractRoot(SimStrandsData, SimRoots);
+
+		// Find k-closest guide:
+		// N hairs, M guides. Complexity = MxN
+		// #hair_todo: build an acceleration structure for fast(er) look up or at least make it run on parallel (e.g. with ParallelFor)
+		uint32 TotalInvalidInterpolationCount = 0;
+		const static float MinWeightDistance = 0.0001f;
+		const static uint32 GuideCount = 3;
+		const static uint32 KGuideCount = GuideCount * 2;
+
+		const static bool bRandomizeInterpolation = false;
+		const static bool bUseUniqueGuide = false;
+		const static bool bPrintDebugMetric = false;
+
+		FRandomStream Random;
+		const uint32 RenCurveCount = RenStrandsData.GetNumCurves();
+		const uint32 SimCurveCount = SimStrandsData.GetNumCurves();
+
+		ParallelFor(RenCurveCount, 
+		[
+			RenCurveCount, &RenRoots, &RenStrandsData,
+			SimCurveCount, &SimRoots, &SimStrandsData, 
+			&TotalInvalidInterpolationCount,  
+			&InterpolationData, 
+			&Random
+		] (uint32 RenCurveIndex) 
+		//for (uint32 RenCurveIndex = 0; RenCurveIndex < RenCurveCount; ++RenCurveIndex)
+		{
+			const FHairRoot& StrandRoot = RenRoots[RenCurveIndex];
+
+			int32 ClosestGuideIndices[GuideCount];
+			{
+				float KMinMetrics[KGuideCount];
+				int32 KClosestGuideIndices[KGuideCount];
+				for (uint32 ClearIt = 0; ClearIt < KGuideCount; ++ClearIt)
+				{
+					KMinMetrics[ClearIt] = FLT_MAX;
+					KClosestGuideIndices[ClearIt] = -1;
+				}
+
+				for (uint32 SimCurveIndex = 0; SimCurveIndex < SimCurveCount; ++SimCurveIndex)
+				{
+					const FHairRoot& GuideRoot = SimRoots[SimCurveIndex];
+					//if ((StrandRoot.Position - GuideRoot.Position).Size() < 5.0f)
+					{
+						const float Metric = 1.0 - ComputeCurvesMetric<16>(RenStrandsData, RenCurveIndex, SimStrandsData, SimCurveIndex, 0.0f, 1.0f, 1.0f);
+						if (Metric < KMinMetrics[KGuideCount - 1])
+						{
+							int32 LastGuideIndex = SimCurveIndex;
+							float LastMetric = Metric;
+							for (uint32 Index = 0; Index < KGuideCount; ++Index)
+							{
+								if (Metric < KMinMetrics[Index])
+								{
+									SwapValue(KClosestGuideIndices[Index], LastGuideIndex);
+									SwapValue(KMinMetrics[Index], LastMetric);
+								}
+							}
+						}
+					}
+				}
+
+				// Debug
+				//if (bPrintDebugMetric)
+				//{
+				//	const FHairInterpolationMetric ClosestMetric = ComputeInterpolationMetric(StrandRoot, SimRoots[KClosestGuideIndices[0]]);
+				//	const float Threshold = 20;
+				//	if (ClosestMetric.Metric > Threshold)
+				//	{
+				//		PrintInterpolationMetric(ClosestMetric);
+				//		++TotalInvalidInterpolationCount;
+				//	}
+				//}
+
+				// Randomize influence guide to break interpolation coherence, and create a more random/natural pattern
+				{
+					uint32 RandIndex0 = 0;
+					uint32 RandIndex1 = 1;
+					uint32 RandIndex2 = 2;
+					if (bRandomizeInterpolation)
+					{
+						do
+						{
+							RandIndex0 = Random.RandRange(0, KGuideCount - 1);
+							RandIndex1 = Random.RandRange(0, KGuideCount - 1);
+							RandIndex2 = Random.RandRange(0, KGuideCount - 1);
+
+						} while (RandIndex0 == RandIndex1 || RandIndex0 == RandIndex2 || RandIndex1 == RandIndex2);
+					}
+					ClosestGuideIndices[0] = KClosestGuideIndices[RandIndex0];
+					ClosestGuideIndices[1] = KClosestGuideIndices[RandIndex1];
+					ClosestGuideIndices[2] = KClosestGuideIndices[RandIndex2];
+
+					if (bUseUniqueGuide)
+					{
+						ClosestGuideIndices[1] = KClosestGuideIndices[RandIndex0];
+						ClosestGuideIndices[2] = KClosestGuideIndices[RandIndex0];
+					}
+
+					float MinMetrics[GuideCount];
+					MinMetrics[0] = KMinMetrics[RandIndex0];
+					MinMetrics[1] = KMinMetrics[RandIndex1];
+					MinMetrics[2] = KMinMetrics[RandIndex2];
+
+
+					while (!(MinMetrics[0] <= MinMetrics[1] && MinMetrics[1] <= MinMetrics[2]))
+					{
+						if (MinMetrics[0] > MinMetrics[1])
+						{
+							SwapValue(MinMetrics[0], MinMetrics[1]);
+							SwapValue(ClosestGuideIndices[0], ClosestGuideIndices[1]);
+						}
+
+						if (MinMetrics[1] > MinMetrics[2])
+						{
+							SwapValue(MinMetrics[1], MinMetrics[2]);
+							SwapValue(ClosestGuideIndices[1], ClosestGuideIndices[2]);
+						}
+					}
+
+					// If there less than 3 valid guides, fill the rest with existing valid guides
+					// This can happen due to the normal-orientation based rejection above
+					if (ClosestGuideIndices[1] < 0)
+					{
+						ClosestGuideIndices[1] = ClosestGuideIndices[0];
+						MinMetrics[1] = MinMetrics[0];
+					}
+					if (ClosestGuideIndices[2] < 0)
+					{
+						ClosestGuideIndices[2] = ClosestGuideIndices[1];
+						MinMetrics[2] = MinMetrics[1];
+					}
+
+					check(ClosestGuideIndices[0] >= 0);
+					check(ClosestGuideIndices[1] >= 0);
+					check(ClosestGuideIndices[2] >= 0);
+					check(MinMetrics[0] <= MinMetrics[1]);
+					check(MinMetrics[1] <= MinMetrics[2]);
+				}
+			}
+
+
+			const uint32 RendPointCount	= RenStrandsData.StrandsCurves.CurvesCount[RenCurveIndex];
+			const uint32 RenOffset		= RenStrandsData.StrandsCurves.CurvesOffset[RenCurveIndex];
+			const FVector& RenPointPositionRoot = RenStrandsData.StrandsPoints.PointsPosition[RenOffset];
+			for (uint32 RenPointIndex = 0; RenPointIndex < RendPointCount; ++RenPointIndex)
+			{
+				const uint32 PointGlobalIndex = RenPointIndex + RenOffset;
+				const FVector& RenPointPosition = RenStrandsData.StrandsPoints.PointsPosition[PointGlobalIndex];
+				const float RenPointDistance = RenStrandsData.StrandsPoints.PointsCoordU[PointGlobalIndex] * RenStrandsData.StrandsCurves.CurvesLength[RenCurveIndex] * RenStrandsData.StrandsCurves.MaxLength;
+
+				float TotalWeight = 0;
+				for (uint32 KIndex = 0; KIndex < GuideCount; ++KIndex)
+				{
+
+				#define WEIGHT_METHOD 0
+
+					// Find the closest vertex on the guide which matches the strand vertex distance along its curve
+				#if WEIGHT_METHOD == 0
+					const uint32 SimCurveIndex = ClosestGuideIndices[KIndex];
+					const uint32 SimOffset = SimStrandsData.StrandsCurves.CurvesOffset[SimCurveIndex];
+					float PrevSimPointDistance = 0;
+					uint32 ClosestSimPointIndex = 0;
+					const uint32 SimPointCount = SimStrandsData.StrandsCurves.CurvesCount[SimCurveIndex];
+					for (uint32 SimPointIndex = 0; SimPointIndex < SimPointCount; ++SimPointIndex, ++ClosestSimPointIndex)
+					{
+						const float SimPointDistance = SimStrandsData.StrandsPoints.PointsCoordU[SimPointIndex + SimOffset] * SimStrandsData.StrandsCurves.CurvesLength[SimCurveIndex] * SimStrandsData.StrandsCurves.MaxLength;
+						if (RenPointDistance >= PrevSimPointDistance && RenPointDistance <= SimPointDistance)
+						{
+							//const float GuideDistance = Guide.Points[SimPointIndex].Distance;
+							const float D0 = FMath::Abs(PrevSimPointDistance - RenPointDistance);
+							const float D1 = FMath::Abs(SimPointDistance - RenPointDistance);
+							if (D0 > D1)
+							{
+								ClosestSimPointIndex = SimPointIndex;
+							}
+							break;
+						}
+						PrevSimPointDistance = SimPointDistance;					
+					}
+					check(SimPointCount > 0);
+					ClosestSimPointIndex = FMath::Clamp(ClosestSimPointIndex, 0u, uint32(SimPointCount - 1));
+
+					const FVector& SimPointPosition = SimStrandsData.StrandsPoints.PointsPosition[ClosestSimPointIndex + SimOffset];
+					const float Weight = 1.0f / FMath::Max(MinWeightDistance, FVector::Distance(RenPointPosition, SimPointPosition));
+					InterpolationData.PointsSimCurvesIndex[PointGlobalIndex][KIndex] = SimCurveIndex;
+					InterpolationData.PointsSimCurvesVertexIndex[PointGlobalIndex][KIndex] = ClosestSimPointIndex + SimOffset;
+					InterpolationData.PointsSimCurvesVertexWeights[PointGlobalIndex][KIndex] = Weight;
+				#endif
+
+					// Use only the root as a *constant* weight for deformation along each vertex
+				#if WEIGHT_METHOD == 1
+					const uint32 SimCurveIndex = ClosestGuideIndices[KIndex];
+					const uint32 SimOffset = SimStrandsData.StrandsCurves.CurvesOffset[SimCurveIndex];
+					const FVector& SimRootPointPosition = SimStrandsData.StrandsPoints[SimOffset];
+					const float Weight = 1.0f / FMath::Max(MinWeightDistance, FVector::Distance(RenPointPosition, SimRootPointPosition));
+					InterpolationData.PointsSimCurvesIndex[PointGlobalIndex][KIndex] = SimCurveIndex;
+					InterpolationData.PointsSimCurvesVertexIndex[PointGlobalIndex][KIndex] = SimOffset;
+					InterpolationData.PointsSimCurvesVertexWeights[PointGlobalIndex][KIndex] = Weight;
+				#endif
+
+					// Use the *same vertex index* to match guide vertex with strand vertex
+				#if WEIGHT_METHOD == 2
+					check(SimPointCount > 0);
+					const uint32 SimCurveIndex = ClosestGuideIndices[KIndex];
+					const uint32 SimPointIndex = FMath::Clamp(RenPointIndex, 0, SimPointCount - 1);
+					const FVector& SimPointPosition = SimStrandsData.StrandsPoints[SimPointIndex + SimOffset];
+					const float Weight = 1.0f / FMath::Max(MinWeightDistance, FVector::Distance(RenPointPosition, SimPointPosition));
+					InterpolationData.PointsSimCurvesIndex[PointGlobalIndex][KIndex] = SimCurveIndex;
+					InterpolationData.PointsSimCurvesVertexIndex[PointGlobalIndex][KIndex] = SimPointIndex + SimOffset;
+					InterpolationData.PointsSimCurvesVertexWeights[PointGlobalIndex][KIndex] = Weight;
+				#endif
+
+					TotalWeight += Weight;
+				}
+
+				for (int32 KIndex = 0; KIndex < GuideCount; ++KIndex)
+				{
+					InterpolationData.PointsSimCurvesVertexWeights[PointGlobalIndex][KIndex] /= TotalWeight;
+				}
+			}
+		});
+
+		if (bPrintDebugMetric)
+		{
+			UE_LOG(LogGroomBuilder, Log, TEXT("Invalid interpolation count: %d/%d)"), TotalInvalidInterpolationCount, RenCurveCount);
+		}
+	}
+
+	/** Build data for interpolation between simulation and rendering */
+	void BuildRenderData(FHairStrandsInterpolationDatas& HairInterpolation)
+	{
+		TRACE_CPUPROFILER_EVENT_SCOPE(HairInterpolationBuilder::BuildRenderData);
+
+		const uint32 PointCount = HairInterpolation.Num();
+		if (PointCount == 0)
+			return;
+
+		auto LowerPart = [](uint32 Index) { return uint16(Index & 0xFFFF); };
+		auto UpperPart = [](uint32 Index) { return uint8((Index >> 16) & 0xFF); };
+
+		TArray<FHairStrandsInterpolation0Format::Type>& OutPointsInterpolation0 = HairInterpolation.RenderData.Interpolation0;
+		TArray<FHairStrandsInterpolation1Format::Type>& OutPointsInterpolation1 = HairInterpolation.RenderData.Interpolation1;
+
+		OutPointsInterpolation0.SetNum(PointCount * FHairStrandsInterpolation0Format::ComponentCount);
+		OutPointsInterpolation1.SetNum(PointCount * FHairStrandsInterpolation1Format::ComponentCount);
+
+		for (uint32 PointIndex = 0; PointIndex < PointCount; ++PointIndex)
+		{
+			const FIntVector& Indices = HairInterpolation.PointsSimCurvesVertexIndex[PointIndex];
+			const FVector& Weights = HairInterpolation.PointsSimCurvesVertexWeights[PointIndex];
+
+			FHairStrandsInterpolation0Format::Type& OutInterp0 = OutPointsInterpolation0[PointIndex];
+			OutInterp0.Index0 = LowerPart(Indices[0]);
+			OutInterp0.Index1 = LowerPart(Indices[1]);
+			OutInterp0.Index2 = LowerPart(Indices[2]);
+			OutInterp0.VertexWeight0 = Weights[0] * 255.f;
+			OutInterp0.VertexWeight1 = Weights[1] * 255.f;
+
+			FHairStrandsInterpolation1Format::Type& OutInterp1 = OutPointsInterpolation1[PointIndex];
+			OutInterp1.VertexIndex0 = UpperPart(Indices[0]);
+			OutInterp1.VertexIndex1 = UpperPart(Indices[1]);
+			OutInterp1.VertexIndex2 = UpperPart(Indices[2]);
+			OutInterp1.Pad0 = 0;
+		}	
+	}
+}
+
+bool FGroomBuilder::BuildGroom(const FHairDescription& HairDescription, const FGroomBuildSettings& BuildSettings, UGroomAsset* GroomAsset)
+{
+	if (!GroomAsset)
+	{
+		return false;
+	}
+
+	// Convert HairDescription to HairStrandsDatas
+	// For now, just convert HairDescription to HairStrandsDatas
+	int32 NumCurves = HairDescription.GetNumStrands();
+	int32 NumVertices = HairDescription.GetNumVertices();
+
+	// Check for required attributes
+	TGroomAttributesConstRef<int> MajorVersion = HairDescription.GroomAttributes().GetAttributesRef<int>(HairAttribute::Groom::MajorVersion);
+	TGroomAttributesConstRef<int> MinorVersion = HairDescription.GroomAttributes().GetAttributesRef<int>(HairAttribute::Groom::MinorVersion);
+
+	if (!MajorVersion.IsValid() || !MinorVersion.IsValid())
+	{
+		UE_LOG(LogGroomBuilder, Warning, TEXT("No version number attributes found. The groom may be missing attributes even if it imports."));
+	}
+
+	FGroomID GroomID(0);
+
+	TGroomAttributesConstRef<float> GroomHairWidthAttribute = HairDescription.GroomAttributes().GetAttributesRef<float>(HairAttribute::Groom::Width);
+	TOptional<float> GroomHairWidth;
+	if (GroomHairWidthAttribute.IsValid())
+	{
+		GroomHairWidth = GroomHairWidthAttribute[GroomID];
+	}
+
+	TGroomAttributesConstRef<FVector> GroomHairColorAttribute = HairDescription.GroomAttributes().GetAttributesRef<FVector>(HairAttribute::Groom::Color);
+	TOptional<FVector> GroomHairColor;
+	if (GroomHairColorAttribute.IsValid())
+	{
+		GroomHairColor = GroomHairColorAttribute[GroomID];
+	}
+
+	TVertexAttributesConstRef<FVector> VertexPositions = HairDescription.VertexAttributes().GetAttributesRef<FVector>(HairAttribute::Vertex::Position);
+	TStrandAttributesConstRef<int> StrandNumVertices = HairDescription.StrandAttributes().GetAttributesRef<int>(HairAttribute::Strand::VertexCount);
+
+	if (!VertexPositions.IsValid() || !StrandNumVertices.IsValid())
+	{
+		UE_LOG(LogGroomBuilder, Warning, TEXT("Failed to import hair: No vertices or curves data found."));
+		return false;
+	}
+
+	TVertexAttributesConstRef<float> VertexWidths = HairDescription.VertexAttributes().GetAttributesRef<float>(HairAttribute::Vertex::Width);
+	TStrandAttributesConstRef<float> StrandWidths = HairDescription.StrandAttributes().GetAttributesRef<float>(HairAttribute::Strand::Width);
+
+	TStrandAttributesConstRef<FVector2D> StrandRootUV = HairDescription.StrandAttributes().GetAttributesRef<FVector2D>(HairAttribute::Strand::RootUV);
+	bool bHasUVData = StrandRootUV.IsValid();
+
+	TStrandAttributesConstRef<int> StrandGuides = HairDescription.StrandAttributes().GetAttributesRef<int>(HairAttribute::Strand::Guide);
+	TStrandAttributesConstRef<int> GroupIDs = HairDescription.StrandAttributes().GetAttributesRef<int>(HairAttribute::Strand::GroupID);
+
+	bool bImportGuides = !BuildSettings.bOverrideGuides;
+
+	FHairStrandsDatas* CurrentHairStrandsDatas = nullptr;
+
+	typedef TPair<FHairGroupInfo, FHairGroupData> FHairGroup;
+	TMap<int32, FHairGroup> HairGroups;
+
+	int32 GlobalVertexIndex = 0;
+	int32 NumHairCurves = 0;
+	int32 NumGuideCurves = 0;
+	int32 NumHairPoints = 0;
+	int32 NumGuidePoints = 0;
+	for (int32 CurveIndex = 0; CurveIndex < NumCurves; ++CurveIndex)
+	{
+		FStrandID StrandID(CurveIndex);
+
+		bool bIsGuide = false;
+		if (StrandGuides.IsValid())
+		{
+			// Version 0.1 defines 1 as being guide
+			bIsGuide = StrandGuides[StrandID] == 1;
+		}
+
+		int32 CurveNumVertices = StrandNumVertices[StrandID];
+
+		int32 GroupID = 0;
+		if (GroupIDs.IsValid())
+		{
+			GroupID = GroupIDs[StrandID];
+		}
+
+		FHairGroup& Group = HairGroups.FindOrAdd(GroupID);
+		FHairGroupInfo& GroupInfo = Group.Key;
+		FHairGroupData& GroupData = Group.Value;
+		GroupInfo.GroupID = GroupID;
+		if (!bIsGuide)
+		{
+			++NumHairCurves;
+			NumHairPoints += CurveNumVertices;
+			CurrentHairStrandsDatas = &GroupData.HairRenderData;
+
+			++GroupInfo.NumCurves;
+		}
+		else if (bImportGuides)
+		{
+			++NumGuideCurves;
+			NumGuidePoints += CurveNumVertices;
+			CurrentHairStrandsDatas = &GroupData.HairSimulationData;
+
+			++GroupInfo.NumGuides;
+		}
+
+		CurrentHairStrandsDatas->StrandsCurves.CurvesCount.Add(CurveNumVertices);
+
+		if (bHasUVData)
+		{
+			CurrentHairStrandsDatas->StrandsCurves.CurvesRootUV.Add(StrandRootUV[StrandID]);
+		}
+
+		float StrandWidth = GroomHairWidth ? GroomHairWidth.GetValue() : 0.01f;
+		if (StrandWidths.IsValid())
+		{
+			StrandWidth = StrandWidths[StrandID];
+		}
+
+		for (int32 VertexIndex = 0; VertexIndex < CurveNumVertices; ++VertexIndex, ++GlobalVertexIndex)
+		{
+			FVertexID VertexID(GlobalVertexIndex);
+
+			CurrentHairStrandsDatas->StrandsPoints.PointsPosition.Add(VertexPositions[VertexID]);
+
+			float VertexWidth = 0.f;
+			if (VertexWidths.IsValid())
+			{
+				VertexWidth = VertexWidths[VertexID];
+			}
+
+			// Fall back to strand width if there was no vertex width
+			if (VertexWidth == 0.f && StrandWidth != 0.f)
+			{
+				VertexWidth = StrandWidth;
+			}
+
+			CurrentHairStrandsDatas->StrandsPoints.PointsRadius.Add(VertexWidth * 0.5f);
+		}
+	}
+
+	FGroomComponentRecreateRenderStateContext RecreateRenderContext(GroomAsset);
+
+	for (TPair<int32, FHairGroup>& HairGroupIt : HairGroups)
+	{
+		int32 GroupID = HairGroupIt.Key;
+		FHairGroup& Group = HairGroupIt.Value;
+		FHairGroupInfo& GroupInfo = Group.Key;
+		FHairGroupData& GroupData = Group.Value;
+
+		FHairStrandsDatas& HairRenderData = GroupData.HairRenderData;
+		int32 GroupNumCurves = HairRenderData.StrandsCurves.Num();
+		HairRenderData.StrandsCurves.SetNum(GroupNumCurves);
+		GroupInfo.NumCurves = GroupNumCurves;
+
+		int32 GroupNumPoints = HairRenderData.StrandsPoints.Num();
+		HairRenderData.StrandsPoints.SetNum(GroupNumPoints);
+
+		HairStrandsBuilder::BuildInternalData(HairRenderData, !bHasUVData);
+
+		FHairStrandsDatas& HairSimulationData = GroupData.HairSimulationData;
+		GroupNumCurves = GroupData.HairSimulationData.StrandsCurves.Num();
+
+		if (GroupNumCurves > 0)
+		{
+			GroupInfo.NumGuides = GroupNumCurves;
+			HairSimulationData.StrandsCurves.SetNum(GroupNumCurves);
+
+			GroupNumPoints = HairSimulationData.StrandsPoints.Num();
+			HairSimulationData.StrandsPoints.SetNum(GroupNumPoints);
+
+			HairStrandsBuilder::BuildInternalData(HairSimulationData, true); // Imported guides don't currently have root UVs so force computing them
+		}
+		else
+		{
+			GroomAsset->HairToGuideDensity = FMath::Clamp(BuildSettings.HairToGuideDensity, 0.01f, 1.f);
+		}
+	}
+
+	for (TPair<int32, FHairGroup> HairGroupIt : HairGroups)
+	{
+		FHairGroup& Group = HairGroupIt.Value;
+		FHairGroupInfo& GroupInfo = Group.Key;
+		FHairGroupData& GroupData = Group.Value;
+
+		GroomAsset->HairGroupsInfo.Add(MoveTemp(GroupInfo));
+		GroomAsset->HairGroupsData.Add(MoveTemp(GroupData));
+	}
+
+	BuildData(GroomAsset);
+
+	GroomAsset->InitResource();
+
+	return true;
+}
+
+void FGroomBuilder::BuildData(UGroomAsset* GroomAsset)
+{
+	if (!GroomAsset)
+	{
+		return;
+	}
+
+	float NumTasks = 4.0f * GroomAsset->GetNumHairGroups();
+
+	FScopedSlowTask SlowTask(NumTasks, LOCTEXT("BuildData", "Building groom render and simulation data"));
+	SlowTask.MakeDialog();
+
+	for (int32 Index = 0; Index < GroomAsset->GetNumHairGroups(); ++Index)
+	{
+		FHairGroupData& GroupData = GroomAsset->HairGroupsData[Index];
+
+		FHairStrandsDatas& HairRenderData = GroupData.HairRenderData;
+		FHairStrandsDatas& HairSimulationData = GroupData.HairSimulationData;
+
+		if (HairSimulationData.GetNumCurves() == 0)
+		{
+			FHairGroupInfo& GroupInfo = GroomAsset->HairGroupsInfo[Index];
+
+			float GuideDensity = FMath::Clamp(GroomAsset->HairToGuideDensity, 0.01f, 1.f);
+			GenerateGuides(HairRenderData, GuideDensity, HairSimulationData);
+			GroupInfo.bIsAutoGenerated = true;
+			GroupInfo.NumGuides = HairSimulationData.GetNumCurves();
+		}
+
+		// Build RenderData for HairStrandsDatas
+		SlowTask.EnterProgressFrame(1.0f, LOCTEXT("BuildHairRenderData", "Building render data for hair"));
+		HairStrandsBuilder::BuildRenderData(HairRenderData);
+
+		SlowTask.EnterProgressFrame(1.0f, LOCTEXT("BuildSimRenderData", "Building render data for guides"));
+		HairStrandsBuilder::BuildRenderData(HairSimulationData);
+
+		FHairStrandsInterpolationDatas& HairInterpolationData = GroupData.HairInterpolationData;
+
+		// Build InterpolationData from render and simulation HairStrandsDatas
+		SlowTask.EnterProgressFrame(1.0f, LOCTEXT("BuildInterpolationData", "Building hair interpolation data"));
+		HairInterpolationBuilder::BuildInterpolationData(HairInterpolationData, HairSimulationData, HairRenderData);
+
+		// Build Rendering data for InterpolationData
+		SlowTask.EnterProgressFrame(1.0f, LOCTEXT("BuildInterpolationRenderData", "Building render data for hair interpolation"));
+		HairInterpolationBuilder::BuildRenderData(HairInterpolationData);
+	}
+}
+
+void FGroomBuilder::GenerateGuides(const FHairStrandsDatas& InData, float DecimationPercentage, FHairStrandsDatas& OutData)
+{
+	// Pick randomly strand as guide 
+	// Divide strands in buckets and pick randomly one stand per bucket
+	const uint32 CurveCount = InData.StrandsCurves.Num();
+	const uint32 OutCurveCount = FMath::Clamp(uint32(CurveCount * DecimationPercentage), 1u, CurveCount);
+
+	const uint32 BucketSize = CurveCount / OutCurveCount;
+
+	TArray<uint32> CurveIndices;
+	CurveIndices.SetNum(OutCurveCount);
+
+	uint32 OutTotalPointCount = 0;
+	FRandomStream Random;
+	for (uint32 BucketIndex = 0; BucketIndex < OutCurveCount; BucketIndex++)
+	{
+		const uint32 CurveIndex = BucketIndex * BucketSize;// +BucketSize * Random.FRand();
+		CurveIndices[BucketIndex] = CurveIndex;
+		OutTotalPointCount += InData.StrandsCurves.CurvesCount[CurveIndex];
+	}
+
+	OutData.StrandsCurves.SetNum(OutCurveCount);
+	OutData.StrandsPoints.SetNum(OutTotalPointCount);
+	OutData.HairDensity = InData.HairDensity;
+
+	uint32 OutPointOffset = 0; 
+	for (uint32 OutCurveIndex = 0; OutCurveIndex < OutCurveCount; OutCurveIndex++)
+	{
+		const uint32 InCurveIndex = CurveIndices[OutCurveIndex];
+		OutData.StrandsCurves.CurvesCount[OutCurveIndex] = InData.StrandsCurves.CurvesCount[InCurveIndex];
+		OutData.StrandsCurves.CurvesRootUV[OutCurveIndex] = InData.StrandsCurves.CurvesRootUV[InCurveIndex];
+		OutData.StrandsCurves.CurvesOffset[OutCurveIndex] = OutPointOffset;
+
+		const uint32 InPointOffset = InData.StrandsCurves.CurvesOffset[InCurveIndex];
+		const uint32 OutPointCount = OutData.StrandsCurves.CurvesCount[OutCurveIndex];
+		for (uint32 OutPointIndex = 0; OutPointIndex < OutPointCount; ++OutPointIndex, ++OutPointOffset)
+		{
+			OutData.StrandsPoints.PointsPosition[OutPointOffset] = InData.StrandsPoints.PointsPosition[OutPointIndex + InPointOffset];
+			OutData.StrandsPoints.PointsCoordU[OutPointOffset] = InData.StrandsPoints.PointsCoordU[OutPointIndex + InPointOffset];
+			OutData.StrandsPoints.PointsRadius[OutPointOffset] = InData.StrandsPoints.PointsRadius[OutPointIndex + InPointOffset] * InData.StrandsCurves.MaxRadius;
+		}
+	}
+
+	HairStrandsBuilder::BuildInternalData(OutData, false);
+}
diff --git a/Engine/Plugins/Experimental/HairStrands/Source/HairStrandsCore/Private/GroomComponent.cpp b/Engine/Plugins/Experimental/HairStrands/Source/HairStrandsCore/Private/GroomComponent.cpp
--- a/Engine/Plugins/Experimental/HairStrands/Source/HairStrandsCore/Private/GroomComponent.cpp
+++ b/Engine/Plugins/Experimental/HairStrands/Source/HairStrandsCore/Private/GroomComponent.cpp
@@ -10,11 +10,20 @@
 #include "RayTracingInstanceUtils.h"
 #include "HairStrandsInterface.h"
 #include "UObject/UObjectIterator.h"
+#include "GlobalShader.h"
+#include "Components/SkeletalMeshComponent.h"
+#include "Engine/RendererSettings.h"
+#include "Animation/AnimationSettings.h"
 
 static float GHairClipLength = -1;
 static FAutoConsoleVariableRef CVarHairClipLength(TEXT("r.HairStrands.DebugClipLength"), GHairClipLength, TEXT("Clip hair strands which have a lenth larger than this value. (default is -1, no effect)"));
 float GetHairClipLength() { return GHairClipLength > 0 ? GHairClipLength : 100000;  }
 
+static int32 GHairStrandsMeshProjectionForceRefPoseEnable = 0;
+static int32 GHairStrandsMeshProjectionForceLOD = -1;
+static FAutoConsoleVariableRef CVarHairStrandsMeshProjectionForceRefPoseEnable(TEXT("r.HairStrands.MeshProjection.RefPose"), GHairStrandsMeshProjectionForceRefPoseEnable, TEXT("Enable/Disable reference pose"));
+static FAutoConsoleVariableRef CVarHairStrandsMeshProjectionForceLOD(TEXT("r.HairStrands.MeshProjection.LOD"), GHairStrandsMeshProjectionForceLOD, TEXT("Force a specific LOD"));
+
 /**
  * An material render proxy which overrides the debug mode parameter.
  */
@@ -118,13 +127,15 @@ public:
 	{
 		check(Component);
 		check(Component->GroomAsset);
-		
-		const uint32 VertexCount = Component->GroomAsset->HairRenderData.GetNumPoints();
+		check(Component->GroomAsset->GetNumHairGroups() > 0);
+
+		const FHairGroupData& GroupData = Component->GroomAsset->HairGroupsData[0];
+		const uint32 VertexCount = GroupData.HairRenderData.GetNumPoints();
 		const float MinHairRadius = 0; // Todo: Component->GroomAsset->HairRenderData.StrandsCurves.MinRadius;
-		const float MaxHairRadius = Component->GroomAsset->HairRenderData.StrandsCurves.MaxRadius;
-		const float MaxHairLength = Component->GroomAsset->HairRenderData.StrandsCurves.MaxLength;
-		const float HairDensity = Component->GroomAsset->HairRenderData.HairDensity;
-		const FVector& HairWorldOffset = Component->GroomAsset->HairRenderData.BoundingBox.GetCenter();
+		const float MaxHairRadius = GroupData.HairRenderData.StrandsCurves.MaxRadius;
+		const float MaxHairLength = GroupData.HairRenderData.StrandsCurves.MaxLength;
+		const float HairDensity = GroupData.HairRenderData.HairDensity;
+		const FVector& HairWorldOffset = GroupData.HairRenderData.BoundingBox.GetCenter();
 
 		check(Component->InterpolationOutput);
 		FHairStrandsInterpolationOutput* LocalOutput = Component->InterpolationOutput;
@@ -186,7 +197,7 @@ public:
 
 		FMaterialRenderProxy* MaterialProxy = nullptr;
 		{
-			const EHairStrandsDebugMode DebugMode = GetHairStrandsDebugMode();
+			const EHairStrandsDebugMode DebugMode = GetHairStrandsDebugStrandsMode();
 			if (DebugMode != EHairStrandsDebugMode::None)
 			{
 				float DebugModeScalar = 0;
@@ -259,6 +270,17 @@ public:
 
 	virtual FPrimitiveViewRelevance GetViewRelevance(const FSceneView* View) const override
 	{
+		const EHairStrandsDebugMode DebugMode = GetHairStrandsDebugStrandsMode();
+		if (DebugMode != EHairStrandsDebugMode::None)
+		{
+			FPrimitiveViewRelevance Result;
+			Result.bDrawRelevance = IsShown(View);
+			Result.bShadowRelevance = IsShadowCast(View);
+			Result.bDynamicRelevance = true;
+			MaterialRelevance.SetPrimitiveViewRelevance(Result);
+			return Result;
+		}
+
 		FPrimitiveViewRelevance Result;
 		Result.bHairStrandsRelevance = true;
 
@@ -307,10 +329,16 @@ UGroomComponent::UGroomComponent(const FObjectInitializer& ObjectInitializer)
 	bTickInEditor = true;
 	bAutoActivate = true;
 	bSelectable = true;
+	RegisteredSkeletalMeshComponent = nullptr;
 	HairDensity = 1;
-	MergeThreshold = 0.1f;
+	bSkinGroom = false;
 	InitializedResources = nullptr;
+	RenRootResources = nullptr;
+	SimRootResources = nullptr;
 	Mobility = EComponentMobility::Movable;
+	MeshProjectionTickDelay = 0;
+	MeshProjectionLODIndex = -1;
+	MeshProjectionState = EMeshProjectionState::Invalid;
 	if (GEngine)
 	{
 		SetMaterial(0, GEngine->HairDefaultMaterial);
@@ -321,19 +349,34 @@ UGroomComponent::UGroomComponent(const FObjectInitializer& ObjectInitializer)
 
 FPrimitiveSceneProxy* UGroomComponent::CreateSceneProxy()
 {
-	if (!GroomAsset || GroomAsset->HairRenderData.GetNumCurves() == 0 || !InterpolationOutput || !InterpolationInput)
+	if (!GroomAsset || GroomAsset->GetNumHairGroups() == 0 || GroomAsset->HairGroupsData[0].HairRenderData.GetNumCurves() == 0 || !InterpolationOutput || !InterpolationInput)
 		return nullptr;
 
 	return new FHairStrandsSceneProxy(this);
 }
 
 
-FBoxSphereBounds UGroomComponent::CalcBounds(const FTransform& LocalToWorld) const
+FBoxSphereBounds UGroomComponent::CalcBounds(const FTransform& InLocalToWorld) const
 {
 	FBox HairBox(ForceInit);
-	if (GroomAsset)
+	if (GroomAsset && GroomAsset->GetNumHairGroups() > 0)
 	{
-		return FBoxSphereBounds(GroomAsset->HairRenderData.BoundingBox.TransformBy(LocalToWorld));
+		if (RegisteredSkeletalMeshComponent)
+		{
+			// Transform the bounding box with an extra offset coming from the skin animation
+			const FBoxSphereBounds MeshLocalBound = RegisteredSkeletalMeshComponent->CalcBounds(FTransform::Identity);
+			const FVector HairTranslation = GroomAsset->HairGroupsData[0].HairRenderData.BoundingBox.GetCenter();
+			const FVector MeshTranslation = MeshLocalBound.GetSphere().Center;
+			const FVector LocalTranslation = MeshTranslation - HairTranslation;
+
+			FTransform LocalToWorld = InLocalToWorld;
+			LocalToWorld.SetLocation(LocalTranslation + InLocalToWorld.GetLocation());
+			return FBoxSphereBounds(GroomAsset->HairGroupsData[0].HairRenderData.BoundingBox.TransformBy(LocalToWorld));
+		}
+		else
+		{
+			return FBoxSphereBounds(GroomAsset->HairGroupsData[0].HairRenderData.BoundingBox.TransformBy(InLocalToWorld));
+		}
 	}
 	else
 	{
@@ -345,8 +388,7 @@ int32 UGroomComponent::GetNumMaterials() const
 {
 	if (GroomAsset)
 	{
-		int32 NumGroups = GroomAsset->RenderHairGroups.Num();
-		return NumGroups > 0 ? NumGroups : 1;
+		return FMath::Max(GroomAsset->GetNumHairGroups(), 1);
 	}
 	return 1;
 }
@@ -356,14 +398,9 @@ UMaterialInterface* UGroomComponent::GetMaterial(int32 ElementIndex) const
 	UMaterialInterface* OverrideMaterial = Super::GetMaterial(ElementIndex);
 
 	const ERHIFeatureLevel::Type FeatureLevel = GetScene() ? GetScene()->GetFeatureLevel() : ERHIFeatureLevel::Num;
-	if (FeatureLevel == ERHIFeatureLevel::Num)
-	{ 
-		return GEngine->HairDefaultMaterial;
-	}
-
-	if (!OverrideMaterial && GroomAsset && ElementIndex < GroomAsset->RenderHairGroups.Num())
+	if (!OverrideMaterial && GroomAsset && ElementIndex < GroomAsset->GetNumHairGroups() && FeatureLevel != ERHIFeatureLevel::Num)
 	{
-		if (UMaterialInterface* Material = GroomAsset->RenderHairGroups[ElementIndex].Material)
+		if (UMaterialInterface* Material = GroomAsset->HairGroupsInfo[ElementIndex].Material)
 		{
 			if (!Material->GetMaterialResource(FeatureLevel)->IsUsedWithHairStrands())
 			{
@@ -385,6 +422,61 @@ UMaterialInterface* UGroomComponent::GetMaterial(int32 ElementIndex) const
 	return OverrideMaterial;
 }
 
+static FHairStrandsProjectionHairData ToProjectionHairData(FHairStrandsRootResource* In)
+{
+	FHairStrandsProjectionHairData Out = {};
+	if (!In)
+		return Out;
+
+	Out.RootCount = In->RootCount;
+	Out.RootPositionBuffer = In->RootPositionBuffer.SRV;
+	Out.RootNormalBuffer = In->RootNormalBuffer.SRV;
+	Out.VertexToCurveIndexBuffer = &In->VertexToCurveIndexBuffer;
+
+	for (FHairStrandsRootResource::FMeshProjectionLOD& MeshLODData : In->MeshProjectionLODs)
+	{
+		FHairStrandsProjectionHairData::LODData& LODData = Out.LODDatas.AddDefaulted_GetRef();
+		LODData.LODIndex = MeshLODData.LODIndex;
+		LODData.RootTriangleIndexBuffer = &MeshLODData.RootTriangleIndexBuffer;
+		LODData.RootTriangleBarycentricBuffer = &MeshLODData.RootTriangleBarycentricBuffer;
+		
+		LODData.RestRootCenter					= MeshLODData.RestRootCenter;
+		LODData.RestRootTrianglePosition0Buffer = &MeshLODData.RestRootTrianglePosition0Buffer;
+		LODData.RestRootTrianglePosition1Buffer = &MeshLODData.RestRootTrianglePosition1Buffer;
+		LODData.RestRootTrianglePosition2Buffer = &MeshLODData.RestRootTrianglePosition2Buffer;
+
+		LODData.DeformedRootCenter					= MeshLODData.DeformedRootCenter;
+		LODData.DeformedRootTrianglePosition0Buffer = &MeshLODData.DeformedRootTrianglePosition0Buffer;
+		LODData.DeformedRootTrianglePosition1Buffer = &MeshLODData.DeformedRootTrianglePosition1Buffer;
+		LODData.DeformedRootTrianglePosition2Buffer = &MeshLODData.DeformedRootTrianglePosition2Buffer;
+		LODData.bIsValid = 
+			MeshLODData.Status == FHairStrandsRootResource::FMeshProjectionLOD::EStatus::Projected ||
+			MeshLODData.Status == FHairStrandsRootResource::FMeshProjectionLOD::EStatus::Initialized;
+	}
+
+	return Out;
+}
+
+FHairStrandsDatas* UGroomComponent::GetGuideStrandsDatas()
+{
+	if (!GroomAsset || GroomAsset->GetNumHairGroups() == 0)
+	{
+		return nullptr;
+	}
+
+	return &GroomAsset->HairGroupsData[0].HairSimulationData;
+}
+
+FHairStrandsResource* UGroomComponent::GetGuideStrandsResource()
+{
+	if (!GroomAsset || GroomAsset->GetNumHairGroups() == 0)
+	{
+		return nullptr;
+	}
+
+	return GroomAsset->HairGroupsData[0].HairSimulationResource;
+}
+
 template<typename T> void SafeDelete(T*& Data) 
 { 
 	if (Data)
@@ -414,50 +506,115 @@ template<typename T> void SafeRelease(T*& Data)
 	}
 }
 
+void CallbackMeshObjectCallback(
+	FSkeletalMeshObjectCallbackData::EEventType Event,
+	FSkeletalMeshObject* MeshObject,
+	uint64 UserData) 
+{
+	ENQUEUE_RENDER_COMMAND(FHairStrandsMeshObjectUpdate)(
+		[Event, MeshObject, UserData](FRHICommandListImmediate& RHICmdList)
+	{
+		const uint64 ComponentId = uint64(UserData & 0xFFFFFFFF);
+		const EWorldType::Type WorldType = EWorldType::Type((UserData>>32) & 0xFFFFFFFF);
+		if (Event == FSkeletalMeshObjectCallbackData::EEventType::Register || Event == FSkeletalMeshObjectCallbackData::EEventType::Update)
+		{
+			UpdateHairStrands(ComponentId, WorldType, MeshObject);
+		}
+		else
+		{
+			UpdateHairStrands(ComponentId, WorldType, nullptr);
+		}
+	});
+}
+
 void UGroomComponent::InitResources()
 {
 	ReleaseResources();
 
-	if (GroomAsset && GroomAsset->HairStrandsResource)
+	if (GroomAsset && GroomAsset->GetNumHairGroups() > 0 && GroomAsset->HairGroupsData[0].HairStrandsResource)
 	{
-		check(GroomAsset->HairStrandsResource);
+		check(GroomAsset->HairGroupsData[0].HairStrandsResource);
 
 		InitializedResources = GroomAsset;
 
-		FHairStrandsDatas* SimStrandDatas = &GroomAsset->HairSimulationData;
+		FHairStrandsDebugInfo::FGroupInfos DebugGroupInfos;
+		for (const FHairGroupData GroupData : GroomAsset->HairGroupsData)
+		{
+			FHairStrandsDebugInfo::FGroupInfo& DebugGroupInfo = DebugGroupInfos.AddDefaulted_GetRef();
+			DebugGroupInfo.MaxLength	= GroupData.HairRenderData.StrandsCurves.MaxLength;
+			DebugGroupInfo.MaxRadius	= GroupData.HairRenderData.StrandsCurves.MaxRadius;
+			DebugGroupInfo.VertexCount	= GroupData.HairRenderData.GetNumPoints();
+			DebugGroupInfo.CurveCount	= GroupData.HairRenderData.GetNumCurves();
+		}
 
-		FHairStrandsInterpolationDatas InterpolationDatas;
-		InterpolationDatas.BuildInterpolationDatas(*SimStrandDatas, GroomAsset->HairRenderData);
+		const FHairGroupData& GroupData = GroomAsset->HairGroupsData[0];
 
-		InterpolationResource = new FHairStrandsInterpolationResource(InterpolationDatas, *SimStrandDatas);
+		InterpolationResource = new FHairStrandsInterpolationResource(GroupData.HairInterpolationData.RenderData, GroupData.HairSimulationData);
 		BeginInitResource(InterpolationResource);
 
 		#if RHI_RAYTRACING
 		FHairStrandsRaytracingResource* LocalRaytracingResources = nullptr;
 		if (IsRayTracingEnabled())
 		{
-			RaytracingResources = new FHairStrandsRaytracingResource(&GroomAsset->HairRenderData);
+			RaytracingResources = new FHairStrandsRaytracingResource(GroupData.HairRenderData);
 			BeginInitResource(RaytracingResources);
 			LocalRaytracingResources = RaytracingResources;
 		}
 		#endif
 
+
+		const FPrimitiveComponentId LocalComponentId = ComponentId;
+		EWorldType::Type WorldType = GetWorld() ? EWorldType::Type(GetWorld()->WorldType) : EWorldType::None;
+		WorldType = WorldType == EWorldType::Inactive ? EWorldType::Editor : WorldType;
+
+		// Insure the ticking of the Groom component always happens after the skeletalMeshComponent.
+		USkeletalMeshComponent* SkeletalMeshComponent = bSkinGroom && GetAttachParent() ? Cast<USkeletalMeshComponent>(GetAttachParent()) : nullptr;
+		if (SkeletalMeshComponent)
+		{
+			AddTickPrerequisiteComponent(SkeletalMeshComponent);
+			RegisteredSkeletalMeshComponent = SkeletalMeshComponent;
+
+			FSkeletalMeshObjectCallbackData CallbackData;
+			CallbackData.Run = CallbackMeshObjectCallback;
+			CallbackData.UserData = (uint64(LocalComponentId.PrimIDValue) & 0xFFFFFFFF) | (uint64(WorldType) << 32);
+			SkeletalMeshComponent->MeshObjectCallbackData = CallbackData;
+
+			const uint32 LODCount = SkeletalMeshComponent->GetNumLODs();
+			if (LODCount > 0)
+			{
+				RenRootResources = new FHairStrandsRootResource(&GroupData.HairRenderData, LODCount);
+				SimRootResources = new FHairStrandsRootResource(&GroupData.HairSimulationData, LODCount);
+				BeginInitResource(RenRootResources);
+				BeginInitResource(SimRootResources);
+			}
+		}
+		
 		InterpolationOutput = new FHairStrandsInterpolationOutput();
 		InterpolationInput = new FHairStrandsInterpolationInput();
 
 		FHairStrandsInterpolationInput* Input = InterpolationInput;
 		FHairStrandsInterpolationOutput* Output = InterpolationOutput;
-		FHairStrandsResource* RenderResources = GroomAsset->HairStrandsResource;
-		FHairStrandsResource* SimResources = GroomAsset->HairSimulationResource;
+		FHairStrandsResource* RenderResources = GroupData.HairStrandsResource;
+		FHairStrandsResource* SimResources = GroupData.HairSimulationResource;
 		FHairStrandsInterpolationResource* LocalInterpolationResource = InterpolationResource;
 
 		check(Input);
-		Input->HairRadius = GroomAsset->HairRenderData.StrandsCurves.MaxRadius;
-		Input->HairWorldOffset = GroomAsset->HairRenderData.BoundingBox.GetCenter();
+		Input->HairRadius = GroupData.HairRenderData.StrandsCurves.MaxRadius;
+		Input->HairWorldOffset = GroupData.HairRenderData.BoundingBox.GetCenter();
+
 
-		const uint64 Id = (uint64)this;
+		FHairStrandsRootResource* LocalRenRootResources = RenRootResources;
+		FHairStrandsRootResource* LocalSimRootResources = SimRootResources;
+		FTransform HairLocalToWorld = GetComponentTransform();
+		FTransform SkinLocalToWorld = bSkinGroom && SkeletalMeshComponent ? SkeletalMeshComponent->GetComponentTransform() : FTransform::Identity;
+		
+		const uint64 Id = LocalComponentId.PrimIDValue;
 		ENQUEUE_RENDER_COMMAND(FHairStrandsBuffers)(
-			[Id, Input, Output, RenderResources, SimResources, LocalInterpolationResource
+			[Id, Input, Output, 
+			RenderResources, SimResources, LocalRenRootResources, LocalSimRootResources, LocalInterpolationResource,
+			HairLocalToWorld, SkinLocalToWorld,
+			WorldType,
+			DebugGroupInfos
 			#if RHI_RAYTRACING
 			, LocalRaytracingResources
 			#endif
@@ -465,11 +622,11 @@ void UGroomComponent::InitResources()
 		{
 			Input->RenderRestPosePositionBuffer = &RenderResources->RestPositionBuffer;
 			Input->RenderAttributeBuffer		= &RenderResources->AttributeBuffer;
-			Input->RenderVertexCount			= RenderResources->StrandsDatas->GetNumPoints();
+			Input->RenderVertexCount			= RenderResources->RenderData.RenderingPositions.Num() / FHairStrandsPositionFormat::ComponentCount;
 
 			Input->SimRestPosePositionBuffer	= &SimResources->RestPositionBuffer;
 			Input->SimAttributeBuffer			= &SimResources->AttributeBuffer;
-			Input->SimVertexCount				= SimResources->StrandsDatas->GetNumPoints();
+			Input->SimVertexCount				= SimResources->RenderData.RenderingPositions.Num() / FHairStrandsPositionFormat::ComponentCount;
 			Input->SimRootPointIndexBuffer		= &LocalInterpolationResource->SimRootPointIndexBuffer;
 
 			Input->Interpolation0Buffer			= &LocalInterpolationResource->Interpolation0Buffer;
@@ -492,11 +649,21 @@ void UGroomComponent::InitResources()
 			Output->RenderTangentBuffer				= &RenderResources->TangentBuffer;
 			Output->SimTangentBuffer				= &SimResources->TangentBuffer;
 
-			FHairStrandsInterpolation Interpolation;
+			FHairStrandsProjectionHairData RenProjectionData = ToProjectionHairData(LocalRenRootResources);
+			FHairStrandsProjectionHairData SimProjectionData = ToProjectionHairData(LocalSimRootResources);
+
+			FHairStrandsInterpolationData Interpolation;
 			Interpolation.Input = Input;
 			Interpolation.Output = Output;
 			Interpolation.Function = ComputeHairStrandsInterpolation;
-			RegisterHairStrands(Id, Interpolation);
+
+			RegisterHairStrands(
+				Id,
+				WorldType,
+				Interpolation,
+				RenProjectionData, 
+				SimProjectionData,
+				DebugGroupInfos);
 		});
 
 	}
@@ -505,14 +672,18 @@ void UGroomComponent::InitResources()
 void UGroomComponent::ReleaseResources()
 {
 	// Unregister component interpolation resources
-	const uint64 Id = (uint64)this;
+	const FPrimitiveComponentId LocalComponentId = ComponentId;
+	const EWorldType::Type WorldType = GetWorld() ? EWorldType::Type(GetWorld()->WorldType) : EWorldType::None;
+	const uint64 Id = LocalComponentId.PrimIDValue;
 	ENQUEUE_RENDER_COMMAND(StaticMeshVertexBuffersLegacyInit)(
-		[Id](FRHICommandListImmediate& RHICmdList)
+		[Id, WorldType](FRHICommandListImmediate& RHICmdList)
 	{
-		UnregisterHairStrands(Id);
+		UnregisterHairStrands(Id, WorldType);
 	});
 
 	SafeRelease(InterpolationResource);
+	SafeRelease(RenRootResources);
+	SafeRelease(SimRootResources);
 
 	// Delay destruction as resources reference by the interpolation 
 	// structs are used on the rendering thread, 
@@ -528,9 +699,20 @@ void UGroomComponent::ReleaseResources()
 	InterpolationOutput = nullptr;
 	InitializedResources = nullptr;
 
+	MeshProjectionLODIndex = -1;
+	MeshProjectionTickDelay = 0;
+	MeshProjectionState = EMeshProjectionState::Invalid;
+
 #if RHI_RAYTRACING
 	SafeRelease(RaytracingResources);
 #endif
+
+	// Insure the ticking of the Groom component always happens after the skeletalMeshComponent.
+	if (RegisteredSkeletalMeshComponent)
+	{
+		RemoveTickPrerequisiteComponent(RegisteredSkeletalMeshComponent);
+		RegisteredSkeletalMeshComponent = nullptr;
+	}
 }
 
 void UGroomComponent::PostLoad()
@@ -552,6 +734,15 @@ void UGroomComponent::PostLoad()
 void UGroomComponent::OnRegister()
 {
 	Super::OnRegister();
+
+	const EWorldType::Type WorldType = GetWorld() ? EWorldType::Type(GetWorld()->WorldType) : EWorldType::None;
+	const uint64 Id = ComponentId.PrimIDValue;
+
+	ENQUEUE_RENDER_COMMAND(FHairStrandsRegister)(
+		[Id, WorldType](FRHICommandListImmediate& RHICmdList)
+	{
+		UpdateHairStrands(Id, WorldType);
+	});
 }
 
 void UGroomComponent::OnUnregister()
@@ -562,7 +753,120 @@ void UGroomComponent::OnUnregister()
 void UGroomComponent::TickComponent(float DeltaTime, enum ELevelTick TickType, FActorComponentTickFunction *ThisTickFunction)
 {
 	Super::TickComponent(DeltaTime, TickType, ThisTickFunction);	
+	
+	const EWorldType::Type WorldType = GetWorld() ? EWorldType::Type(GetWorld()->WorldType) : EWorldType::None;
+	const uint64 Id = ComponentId.PrimIDValue;
+	const ERHIFeatureLevel::Type FeatureLevel = GetWorld() ? ERHIFeatureLevel::Type(GetWorld()->FeatureLevel) : ERHIFeatureLevel::Num;
+
+	USkeletalMeshComponent* SkeletalMeshComponent = RegisteredSkeletalMeshComponent;	
+	if (SkeletalMeshComponent)
+	{
+		const uint32 SkeletalLODCount = SkeletalMeshComponent->GetNumLODs();
+		const uint32 ResourceLODCount = RenRootResources ? RenRootResources->MeshProjectionLODs.Num() : 0;
+		if (SkeletalLODCount != ResourceLODCount)
+		{
+
+			FHairStrandsRootResource* LocalRenRootResources = SkeletalLODCount > 0 ? new FHairStrandsRootResource((GroomAsset && GroomAsset->GetNumHairGroups() > 0) ? &GroomAsset->HairGroupsData[0].HairRenderData : nullptr, SkeletalLODCount) : nullptr;
+			FHairStrandsRootResource* LocalSimRootResources = SkeletalLODCount > 0 ? new FHairStrandsRootResource((GroomAsset && GroomAsset->GetNumHairGroups() > 0) ? &GroomAsset->HairGroupsData[0].HairSimulationData : nullptr, SkeletalLODCount) : nullptr;
+
+			auto InitRootResource = [] (FHairStrandsRootResource*& PersistentResources, FHairStrandsRootResource* LocalResources)
+			{
+				SafeRelease(PersistentResources);
+				PersistentResources = LocalResources;
+				if (PersistentResources)
+				{
+					BeginInitResource(PersistentResources);
+				}
+			};
+			InitRootResource(RenRootResources, LocalRenRootResources);
+			InitRootResource(SimRootResources, LocalSimRootResources);
+
+			FTransform HairLocalToWorld = GetComponentTransform();
+			ENQUEUE_RENDER_COMMAND(FHairStrandsTick_LODUpdate)(
+				[Id, WorldType, LocalRenRootResources, LocalSimRootResources, HairLocalToWorld](FRHICommandListImmediate& RHICmdList)
+			{
+				FHairStrandsProjectionHairData RenProjectionData = ToProjectionHairData(LocalRenRootResources);
+				FHairStrandsProjectionHairData SimProjectionData = ToProjectionHairData(LocalSimRootResources);
+				UpdateHairStrands(Id, WorldType, HairLocalToWorld, RenProjectionData, SimProjectionData);
+			});
+		}
+
+		if (MeshProjectionState != EMeshProjectionState::Completed && MeshProjectionTickDelay == 0)
+		{
+
+			if (MeshProjectionState == EMeshProjectionState::Invalid)
+			{
+				MeshProjectionLODIndex = 0;
+				MeshProjectionState = EMeshProjectionState::WaitForData;
+			}
+
+			if (MeshProjectionLODIndex == SkeletalMeshComponent->GetNumLODs())
+			{
+				SkeletalMeshComponent->SetForceRefPose(false);
+				SkeletalMeshComponent->SetForcedLOD(0);
+				MeshProjectionState = EMeshProjectionState::Completed;
+			}
+
+			if (MeshProjectionLODIndex < SkeletalMeshComponent->GetNumLODs())
+			{
+				SkeletalMeshComponent->SetForceRefPose(true);
+				SkeletalMeshComponent->SetForcedLOD(MeshProjectionLODIndex + 1);
+
+				const FVector RestRootCenter = SkeletalMeshComponent->Bounds.GetSphere().Center;
+				const uint32 LODIndex = MeshProjectionLODIndex;
+				ENQUEUE_RENDER_COMMAND(FHairStrandsTick_Projection)(
+					[Id, WorldType, FeatureLevel, LODIndex, RestRootCenter](FRHICommandListImmediate& RHICmdList)
+				{
+					AddHairStrandsProjectionQuery(RHICmdList, Id, WorldType, LODIndex, RestRootCenter);
+				});
+				MeshProjectionTickDelay += 5;
+				++MeshProjectionLODIndex;
+			}
+		}
+		else
+		{
+			// For debug purpose
+			const int32 ForceLOD = GHairStrandsMeshProjectionForceLOD >= 0 ? FMath::Clamp(GHairStrandsMeshProjectionForceLOD, 0, SkeletalMeshComponent->GetNumLODs() - 1) : -1;
+			SkeletalMeshComponent->SetForcedLOD(ForceLOD + 1);
+			SkeletalMeshComponent->SetForceRefPose(GHairStrandsMeshProjectionForceRefPoseEnable > 0);
+		}
+
+		// When a skeletal object with projection is enabled, activate the refresh of the bounding box to 
+		// insure the component/proxy bounding box alwaws lies onto the actual skinned mesh
+		MarkRenderTransformDirty();
+	}
+
+	if (MeshProjectionTickDelay > 0)
+	{
+		MeshProjectionTickDelay--;
+	}
+
+	const FVector WorldsBoundsCenter  = SkeletalMeshComponent ? SkeletalMeshComponent->Bounds.GetSphere().Center : FVector::ZeroVector;
+	const FTransform SkinLocalToWorld = SkeletalMeshComponent ? SkeletalMeshComponent->GetComponentTransform() : FTransform();
+	const FTransform HairLocalToWorld = GetComponentTransform();
+	ENQUEUE_RENDER_COMMAND(FHairStrandsTick_TransformUpdate)(
+		[Id, WorldType, HairLocalToWorld, SkinLocalToWorld, WorldsBoundsCenter, FeatureLevel](FRHICommandListImmediate& RHICmdList)
+	{		
+		if (ERHIFeatureLevel::Num == FeatureLevel)
+			return;
+
+		UpdateHairStrands(Id, WorldType, HairLocalToWorld, SkinLocalToWorld, WorldsBoundsCenter);
+	});
+}
+
+void UGroomComponent::SendRenderTransform_Concurrent()
+{
+	if (RegisteredSkeletalMeshComponent)
+	{
+		if (ShouldComponentAddToScene() && ShouldRender())
+		{
+			GetWorld()->Scene->UpdatePrimitiveTransform(this);
+		}
+	}
+
+	Super::SendRenderTransform_Concurrent();
 }
+
 void UGroomComponent::GetUsedMaterials(TArray<UMaterialInterface*>& OutMaterials, bool bGetDebugMaterials) const
 {
 #if WITH_EDITOR
@@ -582,8 +886,17 @@ void UGroomComponent::PostEditChangeProperty(FPropertyChangedEvent& PropertyChan
 	const FName PropertyName = PropertyThatChanged ? PropertyThatChanged->GetFName() : NAME_None;
 
 	//  Init/release resources when setting the GroomAsset (or undoing)
-	if (PropertyName == GET_MEMBER_NAME_CHECKED(UGroomComponent, GroomAsset) ||
-		PropertyThatChanged == nullptr)
+	const bool bRecreateResources =
+		(PropertyName == GET_MEMBER_NAME_CHECKED(UGroomComponent, GroomAsset) || PropertyThatChanged == nullptr) ||
+		PropertyName == GET_MEMBER_NAME_CHECKED(UGroomComponent, bSkinGroom);
+
+	const bool bSupportSkinProjection = GetDefault<URendererSettings>()->bSupportSkinCacheShaders && !UAnimationSettings::Get()->bTickAnimationOnSkeletalMeshInit;
+	if (!bSupportSkinProjection)
+	{
+		bSkinGroom = false;
+	}
+
+	if (bRecreateResources)
 	{
 		if (GroomAsset != InitializedResources)
 		{
