Copying //UE4/Dev-Physics to Dev-Main (//UE4/Dev-Main) [at] 9420715

diff --git a/Engine/Plugins/Experimental/HairStrands/Source/HairStrandsCore/Private/HairStrandsDatas.cpp b/Engine/Plugins/Experimental/HairStrands/Source/HairStrandsCore/Private/HairStrandsDatas.cpp
--- a/Engine/Plugins/Experimental/HairStrands/Source/HairStrandsCore/Private/HairStrandsDatas.cpp
+++ b/Engine/Plugins/Experimental/HairStrands/Source/HairStrandsCore/Private/HairStrandsDatas.cpp
@@ -1,92 +1,6 @@
 // Copyright 1998-2019 Epic Games, Inc. All Rights Reserved.
 
 #include "HairStrandsDatas.h"
-#include "RenderUtils.h"
-#include "Chaos/TriangleMeshImplicitObject.h"
-#include "Chaos/Particles.h"
-#include "Chaos/ErrorReporter.h"
-#include "Chaos/Levelset.h"
-#include "Chaos/Framework/Parallel.h"
-#include "UObject/PhysicsObjectVersion.h"
-#include "Async/ParallelFor.h"
-
-// For debug purpose
-static float GHairInterpolationMetric_Distance = 1;
-static float GHairInterpolationMetric_Angle = 0;
-static float GHairInterpolationMetric_Length = 0;
-static float GHairInterpolationMetric_AngleAttenuation = 5;
-static FAutoConsoleVariableRef CVarHairInterpolationMetric_Distance(TEXT("r.HairStrands.InterpolationMetric.Distance"), GHairInterpolationMetric_Distance, TEXT("Hair strands interpolation metric weights for distance"));
-static FAutoConsoleVariableRef CVarHairInterpolationMetric_Angle(TEXT("r.HairStrands.InterpolationMetric.Angle"), GHairInterpolationMetric_Angle, TEXT("Hair strands interpolation metric weights for angle"));
-static FAutoConsoleVariableRef CVarHairInterpolationMetric_Length(TEXT("r.HairStrands.InterpolationMetric.Length"), GHairInterpolationMetric_Length, TEXT("Hair strands interpolation metric weights for length"));
-static FAutoConsoleVariableRef CVarHairInterpolationMetric_AngleAttenuation(TEXT("r.HairStrands.InterpolationMetric.AngleAttenuation"), GHairInterpolationMetric_AngleAttenuation, TEXT("Hair strands interpolation angle attenuation"));
-
-inline void CopyVectorToPosition(const FVector& InVector, FHairStrandsPositionFormat::Type& OutPosition)
-{
-	OutPosition.X = InVector.X;
-	OutPosition.Y = InVector.Y;
-	OutPosition.Z = InVector.Z;
-}
-
-inline void CopyPositionToVector(const FHairStrandsPositionFormat::Type& InPosition, FVector& OutVector )
-{
-	OutVector.X = InPosition.X;
-	OutVector.Y = InPosition.Y;
-	OutVector.Z = InPosition.Z;
-}
-
-void DecimateStrandData(const FHairStrandsDatas& InData, float DecimationPercentage, FHairStrandsDatas& OutData)
-{
-	// Pick randomly strand as guide 
-	// Divide strands in buckets and pick randomly one stand per bucket
-	const uint32 CurveCount = InData.StrandsCurves.Num();
-	const uint32 OutCurveCount = FMath::Clamp(uint32(CurveCount * DecimationPercentage), 1u, CurveCount);
-
-	const uint32 BucketSize = CurveCount / OutCurveCount;
-
-	TArray<uint32> CurveIndices;
-	CurveIndices.SetNum(OutCurveCount);
-
-	uint32 OutTotalPointCount = 0;
-	FRandomStream Random;
-	for (uint32 BucketIndex = 0; BucketIndex < OutCurveCount; BucketIndex++)
-	{
-		const uint32 CurveIndex = BucketIndex * BucketSize;// +BucketSize * Random.FRand();
-		CurveIndices[BucketIndex] = CurveIndex;
-		OutTotalPointCount += InData.StrandsCurves.CurvesCount[CurveIndex];
-	}
-
-	OutData.StrandsCurves.SetNum(OutCurveCount);
-	OutData.StrandsPoints.SetNum(OutTotalPointCount);
-	OutData.HairDensity = InData.HairDensity;
-
-	uint32 OutPointOffset = 0; 
-	for (uint32 OutCurveIndex = 0; OutCurveIndex < OutCurveCount; OutCurveIndex++)
-	{
-		const uint32 InCurveIndex = CurveIndices[OutCurveIndex];
-		OutData.StrandsCurves.CurvesCount[OutCurveIndex] = InData.StrandsCurves.CurvesCount[InCurveIndex];
-		OutData.StrandsCurves.CurvesRootUV[OutCurveIndex] = InData.StrandsCurves.CurvesRootUV[InCurveIndex];
-		OutData.StrandsCurves.CurvesOffset[OutCurveIndex] = OutPointOffset;
-
-		const uint32 InPointOffset = InData.StrandsCurves.CurvesOffset[InCurveIndex];
-		const uint32 OutPointCount = OutData.StrandsCurves.CurvesCount[OutCurveIndex];
-		for (uint32 OutPointIndex = 0; OutPointIndex < OutPointCount; ++OutPointIndex, ++OutPointOffset)
-		{
-			OutData.StrandsPoints.PointsPosition[OutPointOffset] = InData.StrandsPoints.PointsPosition[OutPointIndex + InPointOffset];
-			OutData.StrandsPoints.PointsCoordU[OutPointOffset] = InData.StrandsPoints.PointsCoordU[OutPointIndex + InPointOffset];
-			OutData.StrandsPoints.PointsRadius[OutPointOffset] = InData.StrandsPoints.PointsRadius[OutPointIndex + InPointOffset] * InData.StrandsCurves.MaxRadius;
-		}
-	}
-
-	OutData.BuildInternalDatas();
-}
-
-template<typename T>
-void SwapValue(T& A, T& B)
-{
-	T Temp = A;
-	A = B;
-	B = Temp;
-}
 
 void FHairStrandsInterpolationDatas::SetNum(const uint32 NumCurves)
 {
@@ -132,812 +46,110 @@ void FHairStrandsPoints::Reset()
 	PointsCoordU.Reset();
 }
 
-void FHairStrandsInterpolationDatas::Serialize(FArchive& Ar)
-{
-	Ar << PointsSimCurvesVertexWeights;
-	Ar << PointsSimCurvesVertexIndex;
-	Ar << PointsSimCurvesIndex;
-}
-
-void FHairStrandsPoints::Serialize(FArchive& Ar)
-{
-	Ar << PointsPosition;
-	Ar << PointsRadius;
-	Ar << PointsCoordU;
-}
-
-void FHairStrandsCurves::Serialize(FArchive& Ar)
+FArchive& operator<<(FArchive& Ar, FPackedHairVertex& Vertex)
 {
-	Ar.UsingCustomVersion(FPhysicsObjectVersion::GUID);
+	Ar << Vertex.X;
+	Ar << Vertex.Y;
+	Ar << Vertex.Z;
 
-	Ar << CurvesCount;
-	Ar << CurvesOffset;
-	Ar << CurvesLength;
-	Ar << CurvesRootUV;
-	Ar << MaxLength;
-	Ar << MaxRadius;
+	Ar << Vertex.NormalizedLength;
 
-	if (Ar.CustomVer(FPhysicsObjectVersion::GUID) >= FPhysicsObjectVersion::HairAssetSerialization_V2)
+	if (Ar.IsLoading())
 	{
-		Ar << CurvesGroupID;
+		uint8 Value;
+		Ar << Value;
+		Vertex.ControlPointType = Value & 0x03; // first 2 bits
+		Vertex.NormalizedRadius = Value >> 2;
 	}
-}
-
-void FHairStrandsDatas::Serialize(FArchive& Ar)
-{
-	StrandsPoints.Serialize(Ar);
-	StrandsCurves.Serialize(Ar);
-	Ar << HairDensity;
-	Ar << BoundingBox;
-}
-
-void FHairStrandsDatas::Reset()
-{
-	StrandsCurves.Reset();
-	StrandsPoints.Reset();
-}
-
-struct FHairRoot
-{
-	FVector Position;
-	uint32  VertexCount;
-	FVector Normal;
-	uint32  Index;
-	float   Length;
-};
-
-struct FHairInterpolationMetric
-{	
-	// Total/combined metrics
-	float Metric;
-
-	// Debug info
-	float DistanceMetric;
-	float AngularMetric;
-	float LengthMetric;
-
-	float CosAngle;
-	float Distance;
-	float GuideLength;
-	float RenderLength;
-
-};
-
-static inline FHairInterpolationMetric ComputeInterpolationMetric(const FHairRoot& RenderRoot, const FHairRoot& GuideRoot)
-{
-	FHairInterpolationMetric Out;
-	Out.Distance = FVector::Dist(RenderRoot.Position, GuideRoot.Position);
-	Out.CosAngle = FVector::DotProduct(RenderRoot.Normal, GuideRoot.Normal);
-	Out.GuideLength = GuideRoot.Length;
-	Out.RenderLength = RenderRoot.Length;
-
-	// Metric takes into account the following properties to find guides which are close, share similar orientation, and 
-	// have similar length for better interpolation
-	// * distance
-	// * orientation 
-	// * length 
-	const float AngularAttenuation = GHairInterpolationMetric_AngleAttenuation > 1 ? GHairInterpolationMetric_AngleAttenuation : 0;
-	Out.DistanceMetric	= Out.Distance * GHairInterpolationMetric_Distance;
-	Out.AngularMetric	= AngularAttenuation == 0 ? 0 : (FMath::Clamp((1 - FMath::Pow(Out.CosAngle, AngularAttenuation)), 0.f, 1.f) * GHairInterpolationMetric_Angle);
-	Out.LengthMetric	= FMath::Abs(FMath::Max(Out.GuideLength / float(Out.RenderLength), Out.RenderLength / float(Out.GuideLength)) - 1) * GHairInterpolationMetric_Length; // Ratio
-	Out.Metric			= Out.DistanceMetric + Out.AngularMetric + Out.LengthMetric;
-
-	return Out;
-}
-
-template<uint32 NumSamples>
-static inline FVector GetCurvePosition(const FHairStrandsDatas& CurvesDatas, const uint32 CurveIndex, const uint32 SampleIndex)
-{
-	const float PointCount = CurvesDatas.StrandsCurves.CurvesCount[CurveIndex]-1.0;
-	const uint32 PointOffset = CurvesDatas.StrandsCurves.CurvesOffset[CurveIndex];
-
-	const float CurvePoint = static_cast<float>(SampleIndex) * PointCount / (static_cast<float>(NumSamples)-1.0f);
-	const uint32 PointPrev = (SampleIndex == 0) ? 0 : (SampleIndex == (NumSamples-1)) ? PointCount - 1 : floor(CurvePoint);
-	const uint32 PointNext = PointPrev + 1;
-
-	const float PointAlpha = CurvePoint - static_cast<float>(PointPrev);
-	return CurvesDatas.StrandsPoints.PointsPosition[PointOffset+PointPrev]*(1.0f-PointAlpha) + 
-		   CurvesDatas.StrandsPoints.PointsPosition[PointOffset+PointNext]*PointAlpha;
-}
-
-template<uint32 NumSamples>
-static inline float ComputeCurvesMetric(const FHairStrandsDatas& RenderCurvesDatas, const uint32 RenderCurveIndex, 
-										const FHairStrandsDatas& GuideCurvesDatas, const uint32 GuideCurveIndex, const float RootImportance, 
-										const float ShapeImportance, const float ProximityImportance)
-{
-	const float AverageLength = FMath::Max(0.5f * (RenderCurvesDatas.StrandsCurves.CurvesLength[RenderCurveIndex] * RenderCurvesDatas.StrandsCurves.MaxLength +
-											       GuideCurvesDatas.StrandsCurves.CurvesLength[GuideCurveIndex] * GuideCurvesDatas.StrandsCurves.MaxLength), SMALL_NUMBER);
-
-	static const float DeltaCoord = 1.0f / static_cast<float>(NumSamples-1);
-
-	const FVector& RenderRoot = RenderCurvesDatas.StrandsPoints.PointsPosition[RenderCurvesDatas.StrandsCurves.CurvesOffset[RenderCurveIndex]];
-	const FVector& GuideRoot = GuideCurvesDatas.StrandsPoints.PointsPosition[GuideCurvesDatas.StrandsCurves.CurvesOffset[GuideCurveIndex]];
-
-	float CurveProximityMetric = 0.0;
-	float CurveShapeMetric = 0.0;
-	for (uint32 SampleIndex = 0; SampleIndex < NumSamples; ++SampleIndex)
+	else
 	{
-		const FVector GuidePosition = GetCurvePosition<NumSamples>(GuideCurvesDatas, GuideCurveIndex, SampleIndex);
-		const FVector RenderPosition = GetCurvePosition<NumSamples>(RenderCurvesDatas, RenderCurveIndex, SampleIndex);
-		const float RootWeight = FMath::Exp(-RootImportance*SampleIndex*DeltaCoord);
-
-		CurveProximityMetric += (GuidePosition - RenderPosition).Size() * RootWeight;
-		CurveShapeMetric += (GuidePosition - GuideRoot - RenderPosition + RenderRoot).Size() * RootWeight;
+		uint8 Value = Vertex.ControlPointType | (Vertex.NormalizedRadius << 2);
+		Ar << Value;
 	}
-	CurveShapeMetric *= DeltaCoord / AverageLength;
-	CurveProximityMetric *= DeltaCoord / AverageLength;
-
-	return FMath::Exp(-ShapeImportance*CurveShapeMetric) * FMath::Exp(-ProximityImportance * CurveProximityMetric);
+	return Ar;
 }
 
-static inline void PrintInterpolationMetric(const FHairInterpolationMetric& In)
+FArchive& operator<<(FArchive& Ar, FPackedHairAttributeVertex& Vertex)
 {
-	UE_LOG(LogHairStrands, Log, TEXT("------------------------------------------------------------------------------------------"));
-	UE_LOG(LogHairStrands, Log, TEXT("Total Metric = %f"), In.Metric);
-	UE_LOG(LogHairStrands, Log, TEXT("Distance     = %f (%f)"), In.Distance, In.DistanceMetric);
-	UE_LOG(LogHairStrands, Log, TEXT("Angle        = %f (%f)"), FMath::RadiansToDegrees(FMath::Acos(In.CosAngle)), In.AngularMetric);
-	UE_LOG(LogHairStrands, Log, TEXT("Length       = %f/%f (%f)"), In.RenderLength, In.GuideLength, In.LengthMetric);
+	Ar << Vertex.RootU;
+	Ar << Vertex.RootV;
+	Ar << Vertex.UCoord;
+	Ar << Vertex.Seed;
+
+	return Ar;
 }
 
-void FHairStrandsInterpolationDatas::BuildInterpolationDatas(
-	const FHairStrandsDatas& SimStrandsData,
-	const FHairStrandsDatas& RenStrandsData)
+FArchive& operator<<(FArchive& Ar, FHairInterpolation0Vertex& Vertex)
 {
-	SetNum(RenStrandsData.GetNumPoints());
-
-	const bool bRandomizeInterpolation = true;
-	const bool bUseUniqueGuide = false;
-	const bool bPrintDebugMetric = false;
-
-	typedef TArray<FHairRoot> FRoots;
-
-	// Extract strand roots
-	auto ExtractRoot = [](const FHairStrandsDatas& InData, TArray<FHairRoot>& OutRoots)
-	{
-		const uint32 CurveCount = InData.StrandsCurves.Num();
-		OutRoots.Reserve(CurveCount);
-		for (uint32 CurveIndex = 0; CurveIndex < CurveCount; ++CurveIndex)
-		{
-			const uint32 PointOffset = InData.StrandsCurves.CurvesOffset[CurveIndex];
-			const uint32 PointCount  = InData.StrandsCurves.CurvesCount[CurveIndex];
-			const float  CurveLength = InData.StrandsCurves.CurvesLength[CurveIndex] * InData.StrandsCurves.MaxLength;
-			check(PointCount > 1);
-			const FVector& P0 = InData.StrandsPoints.PointsPosition[PointOffset];
-			const FVector& P1 = InData.StrandsPoints.PointsPosition[PointOffset+1];
-			FVector N = (P1 - P0).GetSafeNormal();
-
-			// Fallback in case the initial points are too close (this happens on certain assets)
-			if (FVector::DotProduct(N, N) == 0)
-			{
-				N = FVector(0, 0, 1);
-			}
-			OutRoots.Add({ P0, PointCount, N, PointOffset, CurveLength });
-		}
-	};
-
-	FRoots RenRoots, SimRoots;
-	ExtractRoot(RenStrandsData, RenRoots);
-	ExtractRoot(SimStrandsData, SimRoots);
-
-	// Find k-closest guide:
-	// N hairs, M guides. Complexity = MxN
-	// #hair_todo: build an acceleration structure for fast(er) look up or at least make it run on parallel (e.g. with ParallelFor)
-	uint32 TotalInvalidInterpolationCount = 0;
-	const float MinWeightDistance = 0.0001f;
-	const static uint32 GuideCount = 3;
-	const static uint32 KGuideCount = GuideCount * 2;
-	FRandomStream Random;
-	const uint32 RenCurveCount = RenStrandsData.GetNumCurves();
-	const uint32 SimCurveCount = SimStrandsData.GetNumCurves();
-	uint32 RenGlobaPointIndex = 0;
-	//ParallelFor(RenCurveCount, [this, &RenRoots, &SimRoots, &SimCurveCount, &RenStrandsData, &SimStrandsData](uint32 RenCurveIndex) {
-	for (uint32 RenCurveIndex = 0; RenCurveIndex < RenCurveCount; ++RenCurveIndex)
-	{
-		const FHairRoot& StrandRoot = RenRoots[RenCurveIndex];
-
-		int32 ClosestGuideIndices[GuideCount];
-		{
-			float KMinMetrics[KGuideCount];
-			int32 KClosestGuideIndices[KGuideCount];
-			for (uint32 ClearIt = 0; ClearIt < KGuideCount; ++ClearIt)
-			{
-				KMinMetrics[ClearIt] = FLT_MAX;
-				KClosestGuideIndices[ClearIt] = -1;
-			}
-
-			ParallelFor(SimCurveCount, [this, &KMinMetrics, &KClosestGuideIndices, &RenCurveIndex, &RenStrandsData, &SimStrandsData, &StrandRoot, &SimRoots](uint32 SimCurveIndex) 
-			//for (uint32 SimCurveIndex = 0; SimCurveIndex < SimCurveCount; ++SimCurveIndex)
-			{
-				const FHairRoot& GuideRoot = SimRoots[SimCurveIndex];
-				//if ((StrandRoot.Position - GuideRoot.Position).Size() < 5.0f)
-				{
-					const float Metric = 1.0 - ComputeCurvesMetric<16>(RenStrandsData, RenCurveIndex, SimStrandsData, SimCurveIndex, 0.0f, 1.0f, 1.0f);
-					if (Metric < KMinMetrics[KGuideCount - 1])
-					{
-						int32 LastGuideIndex = SimCurveIndex;
-						float LastMetric = Metric;
-						for (uint32 Index = 0; Index < KGuideCount; ++Index)
-						{
-							if (Metric < KMinMetrics[Index])
-							{
-								SwapValue(KClosestGuideIndices[Index], LastGuideIndex);
-								SwapValue(KMinMetrics[Index], LastMetric);
-							}
-						}
-					}
-				}
-
-				//const FHairRoot& GuideRoot = SimRoots[SimCurveIndex];
-				//const FHairInterpolationMetric Metric = ComputeInterpolationMetric(StrandRoot, GuideRoot);
-
-				//const bool bIsSameHemisphere = GHairInterpolationMetric_AngleAttenuation > 1 ? (Metric.CosAngle > 0) : true;
-				//if (bIsSameHemisphere)
-				//{
-				//	// If the current guide is closer than the furtherest apart guide, we need to update 
-				//	// the closest guide set
-				//	if (Metric.Metric < KMinMetrics[KGuideCount - 1])
-				//	{
-				//		int32 LastGuideIndex = SimCurveIndex;
-				//		float LastMetric = Metric.Metric;
-				//		for (uint32 Index = 0; Index < KGuideCount; ++Index)
-				//		{
-				//			if (Metric.Metric < KMinMetrics[Index])
-				//			{
-				//				SwapValue(KClosestGuideIndices[Index], LastGuideIndex);
-				//				SwapValue(KMinMetrics[Index], LastMetric);
-				//			}
-				//		}
-				//	}
-				//}
-			});
-
-			// Debug
-			if (bPrintDebugMetric)
-			{
-				const FHairInterpolationMetric ClosestMetric = ComputeInterpolationMetric(StrandRoot, SimRoots[KClosestGuideIndices[0]]);
-				const float Threshold = 20;
-				if (ClosestMetric.Metric > Threshold)
-				{
-					PrintInterpolationMetric(ClosestMetric);
-					++TotalInvalidInterpolationCount;
-				}
-			}
-
-			// Randomize influence guide to break interpolation coherence, and create a more random/natural pattern
-			{
-				uint32 RandIndex0 = 0;
-				uint32 RandIndex1 = 1;
-				uint32 RandIndex2 = 2;
-				/*if (bRandomizeInterpolation)
-				{
-					do
-					{
-						RandIndex0 = Random.RandRange(0, KGuideCount - 1);
-						RandIndex1 = Random.RandRange(0, KGuideCount - 1);
-						RandIndex2 = Random.RandRange(0, KGuideCount - 1);
-
-					} while (RandIndex0 == RandIndex1 || RandIndex0 == RandIndex2 || RandIndex1 == RandIndex2);
-				}*/
-				ClosestGuideIndices[0] = KClosestGuideIndices[RandIndex0];
-				ClosestGuideIndices[1] = KClosestGuideIndices[RandIndex1];
-				ClosestGuideIndices[2] = KClosestGuideIndices[RandIndex2];
-
-				if (bUseUniqueGuide)
-				{
-					ClosestGuideIndices[1] = KClosestGuideIndices[RandIndex0];
-					ClosestGuideIndices[2] = KClosestGuideIndices[RandIndex0];
-				}
-
-				float MinMetrics[GuideCount];
-				MinMetrics[0] = KMinMetrics[RandIndex0];
-				MinMetrics[1] = KMinMetrics[RandIndex1];
-				MinMetrics[2] = KMinMetrics[RandIndex2];
-
-
-				while (!(MinMetrics[0] <= MinMetrics[1] && MinMetrics[1] <= MinMetrics[2]))
-				{
-					if (MinMetrics[0] > MinMetrics[1])
-					{
-						SwapValue(MinMetrics[0], MinMetrics[1]);
-						SwapValue(ClosestGuideIndices[0], ClosestGuideIndices[1]);
-					}
-
-					if (MinMetrics[1] > MinMetrics[2])
-					{
-						SwapValue(MinMetrics[1], MinMetrics[2]);
-						SwapValue(ClosestGuideIndices[1], ClosestGuideIndices[2]);
-					}
-				}
+	Ar << Vertex.Index0;
+	Ar << Vertex.Index1;
+	Ar << Vertex.Index2;
+	Ar << Vertex.VertexWeight0;
+	Ar << Vertex.VertexWeight1;
 
-				// If there less than 3 valid guides, fill the rest with existing valid guides
-				// This can happen due to the normal-orientation based rejection above
-				if (ClosestGuideIndices[1] < 0)
-				{
-					ClosestGuideIndices[1] = ClosestGuideIndices[0];
-					MinMetrics[1] = MinMetrics[0];
-				}
-				if (ClosestGuideIndices[2] < 0)
-				{
-					ClosestGuideIndices[2] = ClosestGuideIndices[1];
-					MinMetrics[2] = MinMetrics[1];
-				}
-
-				check(ClosestGuideIndices[0] >= 0);
-				check(ClosestGuideIndices[1] >= 0);
-				check(ClosestGuideIndices[2] >= 0);
-				check(MinMetrics[0] <= MinMetrics[1]);
-				check(MinMetrics[1] <= MinMetrics[2]);
-			}
-		}
-
-
-		const uint32 RendPointCount	= RenStrandsData.StrandsCurves.CurvesCount[RenCurveIndex];
-		const uint32 RenOffset		= RenStrandsData.StrandsCurves.CurvesOffset[RenCurveIndex];
-		const FVector& RenPointPositionRoot = RenStrandsData.StrandsPoints.PointsPosition[RenOffset];
-		for (uint32 RenPointIndex = 0; RenPointIndex < RendPointCount; ++RenPointIndex, ++RenGlobaPointIndex)
-		{
-			const FVector& RenPointPosition = RenStrandsData.StrandsPoints.PointsPosition[RenPointIndex + RenOffset];
-			const float RenPointDistance = RenStrandsData.StrandsPoints.PointsCoordU[RenPointIndex + RenOffset] * RenStrandsData.StrandsCurves.CurvesLength[RenCurveIndex] * RenStrandsData.StrandsCurves.MaxLength;
-
-			float TotalWeight = 0;
-			for (uint32 KIndex = 0; KIndex < GuideCount; ++KIndex)
-			{
-
-			#define WEIGHT_METHOD 0
-
-				// Find the closest vertex on the guide which matches the strand vertex distance along its curve
-			#if WEIGHT_METHOD == 0
-				const uint32 SimCurveIndex = ClosestGuideIndices[KIndex];
-				const uint32 SimOffset = SimStrandsData.StrandsCurves.CurvesOffset[SimCurveIndex];
-				float PrevSimPointDistance = 0;
-				uint32 ClosestSimPointIndex = 0;
-				const uint32 SimPointCount = SimStrandsData.StrandsCurves.CurvesCount[SimCurveIndex];
-				for (uint32 SimPointIndex = 0; SimPointIndex < SimPointCount; ++SimPointIndex, ++ClosestSimPointIndex)
-				{
-					const float SimPointDistance = SimStrandsData.StrandsPoints.PointsCoordU[SimPointIndex + SimOffset] * SimStrandsData.StrandsCurves.CurvesLength[SimCurveIndex] * SimStrandsData.StrandsCurves.MaxLength;
-					if (RenPointDistance >= PrevSimPointDistance && RenPointDistance <= SimPointDistance)
-					{
-						//const float GuideDistance = Guide.Points[SimPointIndex].Distance;
-						const float D0 = FMath::Abs(PrevSimPointDistance - RenPointDistance);
-						const float D1 = FMath::Abs(SimPointDistance - RenPointDistance);
-						if (D0 > D1)
-						{
-							ClosestSimPointIndex = SimPointIndex;
-						}
-						break;
-					}
-					PrevSimPointDistance = SimPointDistance;					
-				}
-				check(SimPointCount > 0);
-				ClosestSimPointIndex = FMath::Clamp(ClosestSimPointIndex, 0u, uint32(SimPointCount - 1));
-
-				
-				
-				const FVector& SimPointPosition = SimStrandsData.StrandsPoints.PointsPosition[ClosestSimPointIndex + SimOffset];
-				const float Weight = 1.0f / FMath::Max(MinWeightDistance, FVector::Distance(RenPointPosition, SimPointPosition));
-				PointsSimCurvesIndex[RenGlobaPointIndex][KIndex] = SimCurveIndex;
-				PointsSimCurvesVertexIndex[RenGlobaPointIndex][KIndex] = ClosestSimPointIndex + SimOffset;
-				PointsSimCurvesVertexWeights[RenGlobaPointIndex][KIndex] = Weight;
-			#endif
-
-				// Use only the root as a *constant* weight for deformation along each vertex
-			#if WEIGHT_METHOD == 1
-				const uint32 SimCurveIndex = ClosestGuideIndices[KIndex];
-				const uint32 SimOffset = SimStrandsData.StrandsCurves.CurvesOffset[SimCurveIndex];
-				const FVector& SimRootPointPosition = SimStrandsData.StrandsPoints[SimOffset];
-				const float Weight = 1.0f / FMath::Max(MinWeightDistance, FVector::Distance(RenPointPosition, SimRootPointPosition));
-				PointsSimCurvesIndex[RenGlobaPointIndex][KIndex] = SimCurveIndex;
-				PointsSimCurvesVertexIndex[RenGlobaPointIndex][KIndex] = SimOffset;
-				PointsSimCurvesVertexWeights[RenGlobaPointIndex][KIndex] = Weight;
-
-			#endif
-
-				// Use the *same vertex index* to match guide vertex with strand vertex
-			#if WEIGHT_METHOD == 2
-				check(SimPointCount > 0);
-				const uint32 SimCurveIndex = ClosestGuideIndices[KIndex];
-				const uint32 SimPointIndex = FMath::Clamp(RenPointIndex, 0, SimPointCount - 1);
-				const FVector& SimPointPosition = SimStrandsData.StrandsPoints[SimPointIndex + SimOffset];
-				const float Weight = 1.0f / FMath::Max(MinWeightDistance, FVector::Distance(RenPointPosition, SimPointPosition));
-				PointsSimCurvesIndex[RenGlobaPointIndex][KIndex] = SimCurveIndex;
-				PointsSimCurvesVertexIndex[RenGlobaPointIndex][KIndex] = SimPointIndex + SimOffset;
-				PointsSimCurvesVertexWeights[RenGlobaPointIndex][KIndex] = Weight;
-			#endif
-
-				TotalWeight += Weight;
-			}
-
-			for (int32 KIndex = 0; KIndex < GuideCount; ++KIndex)
-			{
-				PointsSimCurvesVertexWeights[RenGlobaPointIndex][KIndex] /= TotalWeight;
-			}
-		}
-	}
-
-	if (bPrintDebugMetric)
-	{
-		UE_LOG(LogHairStrands, Log, TEXT("Invalid interpolation count: %d/%d)"), TotalInvalidInterpolationCount, RenCurveCount);
-	}
+	return Ar;
 }
 
-void FHairStrandsInterpolationDatas::BuildRenderingDatas(
-	TArray<FHairStrandsInterpolation0Format::Type>& OutPointsInterpolation0,
-	TArray<FHairStrandsInterpolation1Format::Type>& OutPointsInterpolation1) const
+FArchive& operator<<(FArchive& Ar, FHairInterpolation1Vertex& Vertex)
 {
-	const uint32 PointCount = Num();
-	if (PointCount == 0)
-		return;
-	auto LowerPart = [](uint32 Index) { return uint16(Index & 0xFFFF); };
-	auto UpperPart = [](uint32 Index) { return uint8((Index >> 16) & 0xFF); };
-
-	OutPointsInterpolation0.SetNum(PointCount*FHairStrandsInterpolation0Format::ComponentCount);
-	OutPointsInterpolation1.SetNum(PointCount*FHairStrandsInterpolation1Format::ComponentCount);
-
-	for (uint32 PointIndex = 0; PointIndex < PointCount; ++PointIndex)
-	{
-		const FIntVector& Indices = PointsSimCurvesVertexIndex[PointIndex];
-		const FVector& Weights = PointsSimCurvesVertexWeights[PointIndex];
-
-		FHairStrandsInterpolation0Format::Type& OutInterp0 = OutPointsInterpolation0[PointIndex];
-		OutInterp0.Index0 = LowerPart(Indices[0]);
-		OutInterp0.Index1 = LowerPart(Indices[1]);
-		OutInterp0.Index2 = LowerPart(Indices[2]);
-		OutInterp0.VertexWeight0 = Weights[0] * 255.f;
-		OutInterp0.VertexWeight1 = Weights[1] * 255.f;
+	Ar << Vertex.VertexIndex0;
+	Ar << Vertex.VertexIndex1;
+	Ar << Vertex.VertexIndex2;
+	Ar << Vertex.Pad0;
 
-		FHairStrandsInterpolation1Format::Type& OutInterp1 = OutPointsInterpolation1[PointIndex];
-		OutInterp1.VertexIndex0 = UpperPart(Indices[0]);
-		OutInterp1.VertexIndex1 = UpperPart(Indices[1]);
-		OutInterp1.VertexIndex2 = UpperPart(Indices[2]);
-		OutInterp1.Pad0 = 0;
-	}	
+	return Ar;
 }
 
-void FHairStrandsDatas::BuildRenderingDatas(
-	TArray<FHairStrandsPositionFormat::Type>& OutPackedPositions,
-	TArray<FHairStrandsAttributeFormat::Type>& OutPackedAttributes) const
+void FHairStrandsInterpolationDatas::FRenderData::Serialize(FArchive& Ar)
 {
-	if (!(GetNumCurves() > 0 && GetNumPoints() > 0))
-		return;
-
-	OutPackedPositions.SetNum(GetNumPoints()*FHairStrandsPositionFormat::ComponentCount);
-	OutPackedAttributes.SetNum(GetNumPoints()*FHairStrandsAttributeFormat::ComponentCount);
-	
-	const FVector HairBoxCenter = BoundingBox.GetCenter();
-
-	FRandomStream Random;
-	for (uint32 CurveIndex = 0; CurveIndex < GetNumCurves(); ++CurveIndex)
-	{
-		const float CurveSeed = Random.RandHelper(255);
-		const int32 IndexOffset = StrandsCurves.CurvesOffset[CurveIndex];
-		const uint16& PointCount = StrandsCurves.CurvesCount[CurveIndex];
-		for (int32 PointIndex = 0; PointIndex < PointCount; ++PointIndex)
-		{
-			const uint32 PrevIndex = FMath::Max(0, PointIndex - 1);
-			const uint32 NextIndex = FMath::Min(PointCount + 1, PointCount - 1);
-			const FVector& PointPosition = StrandsPoints.PointsPosition[PointIndex + IndexOffset];
-
-			const float CoordU = StrandsPoints.PointsCoordU[PointIndex + IndexOffset];
-			const float NormalizedRadius = StrandsPoints.PointsRadius[PointIndex + IndexOffset];
-			const float NormalizedLength = CoordU * StrandsCurves.CurvesLength[CurveIndex];
-
-			FHairStrandsPositionFormat::Type& PackedPosition = OutPackedPositions[PointIndex + IndexOffset];
-			CopyVectorToPosition(PointPosition - HairBoxCenter, PackedPosition);
-			PackedPosition.ControlPointType = (PointIndex == 0) ? 1u : (PointIndex == (PointCount - 1) ? 2u : 0u);
-			PackedPosition.NormalizedRadius = uint8(FMath::Clamp(NormalizedRadius * 63.f, 0.f, 63.f));
-			PackedPosition.NormalizedLength = uint8(FMath::Clamp(NormalizedLength *255.f, 0.f, 255.f));
-			
-			const FVector2D RootUV = StrandsCurves.CurvesRootUV[CurveIndex];
-			FHairStrandsAttributeFormat::Type& PackedAttributes = OutPackedAttributes[PointIndex + IndexOffset];
-			PackedAttributes.RootU = uint32(FMath::Clamp(RootUV.X, 0.f, 1.f) * 0xFF);
-			PackedAttributes.RootV = uint32(FMath::Clamp(RootUV.Y, 0.f, 1.f) * 0xFF);
-			PackedAttributes.UCoord = FMath::Clamp(CoordU * 255.f, 0.f, 255.f);
-			PackedAttributes.Seed = CurveSeed;
-		}
-	}
+	Ar << Interpolation0;
+	Ar << Interpolation1;
 }
 
-void FHairStrandsDatas::BuildSimulationDatas(const uint32 StrandSize,
-	TArray<FHairStrandsPositionFormat::Type>& OutNodesPositions,
-	TArray<FHairStrandsWeightFormat::Type>& OutPointsWeights,
-	TArray<FHairStrandsIndexFormat::Type>& OutPointsNodes) const
+void FHairStrandsInterpolationDatas::Serialize(FArchive& Ar)
 {
-	OutNodesPositions.SetNum(GetNumCurves()*StrandSize);
-	OutPointsWeights.SetNum(GetNumPoints());
-	OutPointsNodes.SetNum(GetNumPoints());
-
-	if (GetNumCurves() > 0 && GetNumPoints() > 0 && StrandSize > 1)
-	{
-		TArray<FVector>::TConstIterator PositionIterator = StrandsPoints.PointsPosition.CreateConstIterator();
-		TArray<uint16>::TConstIterator CountIterator = StrandsCurves.CurvesCount.CreateConstIterator();
-		TArray<float>::TConstIterator LengthIterator = StrandsCurves.CurvesLength.CreateConstIterator();
-
-		TArray<FHairStrandsIndexFormat::Type>::TIterator NodesIterator = OutPointsNodes.CreateIterator();
-
-		uint32 NumPoints = 0;
-		for (uint32 CurveIndex = 0; CurveIndex < GetNumCurves(); ++CurveIndex, ++LengthIterator, ++CountIterator)
-		{
-			const float SeparationLength = (*LengthIterator) * StrandsCurves.MaxLength / (StrandSize - 1);
-			uint32 SampleOffset = CurveIndex * StrandSize;
-			const uint16 EdgeCount = *CountIterator - 1;
-
-			CopyVectorToPosition(*PositionIterator, OutNodesPositions[SampleOffset]);
-			CopyVectorToPosition(*(PositionIterator + EdgeCount), OutNodesPositions[SampleOffset + StrandSize - 1]);
-
-			FVector VertexNext = *PositionIterator;
-			FVector VertexPrev = VertexNext;
-
-			++PositionIterator;
-
-			float StrandLength = 0.0;
-			uint32 LocalCount = 2;
-			for (uint16 EdgeIndex = 0; EdgeIndex < EdgeCount; ++EdgeIndex, ++PositionIterator, ++NodesIterator, ++NumPoints)
-			{
-				VertexPrev = VertexNext;
-				VertexNext = *PositionIterator;
-
-				*NodesIterator = SampleOffset;
-
-				FVector EdgeDirection = VertexNext - VertexPrev;
-				const float CurrentLength = StrandLength;
-				const float EdgeLength = EdgeDirection.Size();
-				StrandLength += EdgeLength;
-
-				if (StrandLength > SeparationLength)
-				{
-					EdgeDirection /= EdgeLength;
-					const uint32 NumNodes = StrandLength / SeparationLength;
-					FVector EdgePosition = VertexPrev + EdgeDirection * (SeparationLength - CurrentLength);
-					for (uint32 NodeIndex = 0; NodeIndex < NumNodes; ++NodeIndex, EdgePosition += EdgeDirection * SeparationLength)
-					{
-						if (LocalCount < StrandSize)
-						{
-							++SampleOffset;
-							CopyVectorToPosition(EdgePosition, OutNodesPositions[SampleOffset]);
-							++LocalCount;
-							StrandLength -= SeparationLength;
-						}
-					}
-				}
-			}
-			*NodesIterator = SampleOffset;
-			++NodesIterator;
-		}
-		PositionIterator.Reset();
-		NodesIterator.Reset();
-		LengthIterator.Reset();
-		CountIterator.Reset();
-
-		TArray<float>::TIterator WeightsIterator = OutPointsWeights.CreateIterator();
-
-		for (uint32 CurveIndex = 0; CurveIndex < GetNumCurves(); ++CurveIndex, ++LengthIterator, ++CountIterator)
-		{
-			const uint16& StrandCount = *CountIterator;
-			const float SeparationLength = (*LengthIterator) * StrandsCurves.MaxLength / (StrandSize - 1);
-			const float LengthScale = 1.0 / (SeparationLength * SeparationLength);
-
-			for (uint32 PointIndex = 0; PointIndex < StrandCount; ++PointIndex, ++PositionIterator, ++NodesIterator, ++WeightsIterator)
-			{
-				const uint32 NodeIndex = *NodesIterator;
-				FVector PrevPosition, NextPosition;
-				CopyPositionToVector(OutNodesPositions[NodeIndex], PrevPosition);
-				CopyPositionToVector(OutNodesPositions[NodeIndex+1], NextPosition);
-				const float PrevDist = (*PositionIterator - PrevPosition).SizeSquared() * LengthScale;
-				const float NextDist = (*PositionIterator - NextPosition).SizeSquared() * LengthScale;
-
-				float PrevWeight = (PrevDist < 1.0) ? (1.0 - PrevDist)* (1.0 - PrevDist)* (1.0 - PrevDist) : 0.0;
-				float NextWeight = (NextDist < 1.0) ? (1.0 - NextDist)* (1.0 - NextDist)* (1.0 - NextDist) : 0.0;
+	Ar << PointsSimCurvesVertexWeights;
+	Ar << PointsSimCurvesVertexIndex;
+	Ar << PointsSimCurvesIndex;
 
-				const float SumWeights = PrevWeight + NextWeight;
-				if (SumWeights != 0.0)
-				{
-					PrevWeight /= SumWeights;
-					NextWeight /= SumWeights;
-				}
-				*WeightsIterator = PrevWeight;
-			}
-		}
-	}
+	RenderData.Serialize(Ar);
 }
 
-///////////////////////////////////////////////////////////////////////////////////////////////////
-// Auto-generate Root UV data if not loaded
-static FVector2D SignNotZero(const FVector2D& v)
+void FHairStrandsPoints::Serialize(FArchive& Ar)
 {
-	return FVector2D((v.X >= 0.0) ? +1.0 : -1.0, (v.Y >= 0.0) ? +1.0 : -1.0);
+	Ar << PointsPosition;
+	Ar << PointsRadius;
+	Ar << PointsCoordU;
 }
 
-// A Survey of Efficient Representations for Independent Unit Vectors
-// Reference: http://jcgt.org/published/0003/02/01/paper.pdf
-// Assume normalized input. Output is on [-1, 1] for each component.
-static FVector2D SphericalToOctahedron(const FVector& v)
+void FHairStrandsCurves::Serialize(FArchive& Ar)
 {
-	// Project the sphere onto the octahedron, and then onto the xy plane
-	FVector2D p = FVector2D(v.X, v.Y) * (1.0 / (abs(v.X) + abs(v.Y) + abs(v.Z)));
-	// Reflect the folds of the lower hemisphere over the diagonals
-	return (v.Z <= 0.0) ? ((FVector2D(1, 1) - FVector2D(abs(p.Y), abs(p.X))) * SignNotZero(p)) : p;
+	Ar << CurvesCount;
+	Ar << CurvesOffset;
+	Ar << CurvesLength;
+	Ar << CurvesRootUV;
+	Ar << MaxLength;
+	Ar << MaxRadius;
 }
 
-static void ComputeRootUV(FHairStrandsCurves& Curves, FHairStrandsPoints& Points)//; FStrands& Strands)
+void FHairStrandsDatas::FRenderData::Serialize(FArchive& Ar)
 {
-	TArray<FVector> RootPoints;
-	const uint32 CurveCount = Curves.Num();
-	RootPoints.Reserve(Curves.Num());
-	FVector MinAABB(MAX_FLT, MAX_FLT, MAX_FLT);
-	FVector MaxAABB(-MAX_FLT, -MAX_FLT, -MAX_FLT);
-	FMatrix Rotation = FRotationMatrix::Make(FRotator(0, 0, -90));
-	for (uint32 CurveIndex=0; CurveIndex< CurveCount; ++CurveIndex)
-	{
-		const uint32 Offset = Curves.CurvesOffset[CurveIndex];
-		check(Offset < uint32(Points.PointsPosition.Num()));
-		const FVector P = Rotation.TransformPosition(Points.PointsPosition[Offset]);
-
-		RootPoints.Add(P);
-		MinAABB.X = FMath::Min(P.X, MinAABB.X);
-		MinAABB.Y = FMath::Min(P.Y, MinAABB.Y);
-		MinAABB.Z = FMath::Min(P.Z, MinAABB.Z);
-
-		MaxAABB.X = FMath::Max(P.X, MaxAABB.X);
-		MaxAABB.Y = FMath::Max(P.Y, MaxAABB.Y);
-		MaxAABB.Z = FMath::Max(P.Z, MaxAABB.Z);
-	}
-
-	// Compute sphere bound
-	const FVector Extent = MaxAABB - MinAABB;
-	FSphere SBound;
-	SBound.Center = (MaxAABB + MinAABB) * 0.5f;
-	SBound.W = FMath::Max(Extent.X, FMath::Max(Extent.Y, Extent.Z));
-
-	// Project root point onto the bounding sphere and map it onto 
-	// an octahedron, which is unfold onto the unit space [0,1]^2
-	TArray<FVector2D> RootUVs;
-	RootUVs.Reserve(Curves.Num());
-	FVector2D MinUV(MAX_FLT, MAX_FLT);
-	FVector2D MaxUV(-MAX_FLT, -MAX_FLT);
-	for (const FVector& RootP : RootPoints)
-	{
-		FVector D = RootP - SBound.Center;
-		D.Normalize();
-		FVector2D UV = SphericalToOctahedron(D);
-		UV += FVector2D(1, 1);
-		UV *= 0.5f;
-		RootUVs.Add(UV);
-
-		MinUV.X = FMath::Min(UV.X, MinUV.X);
-		MinUV.Y = FMath::Min(UV.Y, MinUV.Y);
-
-		MaxUV.X = FMath::Max(UV.X, MaxUV.X);
-		MaxUV.Y = FMath::Max(UV.Y, MaxUV.Y);
-	}
-
-	// Find the minimal UV space cover by root point, and 
-	// offsets/scales it to maximize UV space
-	const FVector2D UVScale(1 / (MaxUV.X - MinUV.X), 1 / (MaxUV.Y - MinUV.Y));
-	const FVector2D UVOffset(-MinUV.X, -MinUV.Y);
-	uint32 Index = 0;
-	for (FVector2D& RootUV : Curves.CurvesRootUV)
-	{
-		RootUV = (RootUVs[Index++] + UVOffset) * UVScale;
-	}
+	Ar << RenderingPositions;
+	Ar << RenderingAttributes;
 }
 
-void FHairStrandsDatas::BuildInternalDatas(bool bComputeRootUV)
+void FHairStrandsDatas::Serialize(FArchive& Ar)
 {
-	BoundingBox.Min = {  FLT_MAX,  FLT_MAX ,  FLT_MAX };
-	BoundingBox.Max = { -FLT_MAX, -FLT_MAX , -FLT_MAX };
-
-	if (GetNumCurves() > 0 && GetNumPoints() > 0)
-	{
-		TArray<FVector>::TIterator PositionIterator = StrandsPoints.PointsPosition.CreateIterator();
-		TArray<float>::TIterator RadiusIterator = StrandsPoints.PointsRadius.CreateIterator();
-		TArray<float>::TIterator CoordUIterator = StrandsPoints.PointsCoordU.CreateIterator();
-
-		TArray<uint16>::TIterator CountIterator = StrandsCurves.CurvesCount.CreateIterator();
-		TArray<uint32>::TIterator OffsetIterator = StrandsCurves.CurvesOffset.CreateIterator();
-		TArray<float>::TIterator LengthIterator = StrandsCurves.CurvesLength.CreateIterator();
-
-		StrandsCurves.MaxRadius = 0.0;
-		StrandsCurves.MaxLength = 0.0;
-
-		uint32 StrandOffset = 0;
-		*OffsetIterator = StrandOffset; ++OffsetIterator;
-
-		for (uint32 CurveIndex = 0; CurveIndex < GetNumCurves(); ++CurveIndex, ++OffsetIterator, ++LengthIterator, ++CountIterator)
-		{
-			const uint16& StrandCount = *CountIterator;
-
-			StrandOffset += StrandCount;
-			*OffsetIterator = StrandOffset;
-
-			float StrandLength = 0.0;
-			FVector PreviousPosition(0.0, 0.0, 0.0);
-			for (uint32 PointIndex = 0; PointIndex < StrandCount; ++PointIndex, ++PositionIterator, ++RadiusIterator, ++CoordUIterator)
-			{
-				BoundingBox += *PositionIterator;
-
-				if (PointIndex > 0)
-				{
-					StrandLength += (*PositionIterator - PreviousPosition).Size();
-				}
-				*CoordUIterator = StrandLength;
-				PreviousPosition = *PositionIterator;
-
-				StrandsCurves.MaxRadius = FMath::Max(StrandsCurves.MaxRadius, *RadiusIterator);
-			}
-			*LengthIterator = StrandLength;
-			StrandsCurves.MaxLength = FMath::Max(StrandsCurves.MaxLength, StrandLength);
-		}
-
-		CountIterator.Reset();
-		LengthIterator.Reset();
-		RadiusIterator.Reset();
-		CoordUIterator.Reset();
-
-		for (uint32 CurveIndex = 0; CurveIndex < GetNumCurves(); ++CurveIndex, ++LengthIterator, ++CountIterator)
-		{
-			const uint16& StrandCount = *CountIterator;
-
-			for (uint32 PointIndex = 0; PointIndex < StrandCount; ++PointIndex, ++RadiusIterator, ++CoordUIterator)
-			{
-				*CoordUIterator /= *LengthIterator;
-				*RadiusIterator /= StrandsCurves.MaxRadius;
-			}
-			*LengthIterator /= StrandsCurves.MaxLength;
-		}
+	StrandsPoints.Serialize(Ar);
+	StrandsCurves.Serialize(Ar);
+	Ar << HairDensity;
+	Ar << BoundingBox;
 
-		if (bComputeRootUV)
-		{
-			ComputeRootUV(StrandsCurves, StrandsPoints);
-		}
-	}
+	RenderData.Serialize(Ar);
 }
 
-
-void FHairStrandsDatas::AttachStrandsRoots(UStaticMesh* StaticMesh, const FMatrix& TransformMatrix)
+void FHairStrandsDatas::Reset()
 {
-	/*if (StaticMesh)
-	{
-		FTriMeshCollisionData CollisionData;
-		StaticMesh->GetPhysicsTriMeshData(&CollisionData, true);
-
-		Chaos::TParticles<float, 3> Particles =  MoveTemp(*reinterpret_cast<TArray<Chaos::TVector<float, 3>>*>(&CollisionData.Vertices));
-		TArray<Chaos::TVector<int32, 3>> Triangles = MoveTemp(*reinterpret_cast<TArray<Chaos::TVector<int32, 3>>*>(&CollisionData.Indices)); 
-
-		TUniquePtr<Chaos::TTriangleMeshImplicitObject<float>> TriangleMesh = MakeUnique<Chaos::TTriangleMeshImplicitObject<float>>(MoveTemp(Particles), MoveTemp(Triangles));
-
-		TArray<int32> TriangleIndices;
-		TriangleIndices.Init(-1, GetNumCurves());
-
-		TArray<Chaos::TVector<float, 2>> BarycentricCoordinates;
-		BarycentricCoordinates.Init(Chaos::TVector<float, 2>(0, 0), GetNumCurves());
-
-		Chaos::PhysicsParallelFor(GetNumCurves(), [&](int32 CurveIndex) {
-
-			const FVector& RootPosition = StrandsPoints.PointsPosition[StrandsCurves.CurvesOffset[CurveIndex]];
-			const float Dist = TriangleMesh->FindClosestPoint(RootPosition, TriangleIndices[CurveIndex], BarycentricCoordinates[CurveIndex]);
-		});
-
-		// build cell domain
-		int32 MaxAxisSize = 50;
-		int MaxAxis = TriangleMesh->MLocalBoundingBox.LargestAxis();
-		Chaos::TVector<float, 3> Extents = TriangleMesh->MLocalBoundingBox.Extents();
-		Chaos::TVector<int32, 3> Counts(MaxAxisSize * Extents[0] / Extents[MaxAxis], MaxAxisSize * Extents[1] / Extents[MaxAxis], MaxAxisSize * Extents[2] / Extents[MaxAxis]);
-		Counts[0] = Counts[0] < 1 ? 1 : Counts[0];
-		Counts[1] = Counts[1] < 1 ? 1 : Counts[1]; 
-		Counts[2] = Counts[2] < 1 ? 1 : Counts[2];
-
-		Chaos::TUniformGrid<float, 3> Grid(TriangleMesh->MLocalBoundingBox.Min(), TriangleMesh->MLocalBoundingBox.Max(), Counts, 1);
-		Chaos::FErrorReporter ErrorReporter;
-		Chaos::TTriangleMesh<float> CollisionMesh(MoveTemp(TriangleMesh->MElements));
-		Chaos::TLevelSet<float, 3> LevelSet(ErrorReporter, Grid, TriangleMesh->MParticles, CollisionMesh, 0, false);
-
-		for (uint32 CurveIndex = 0; CurveIndex < 10; ++CurveIndex)
-		{
-			for (uint32 PointIndex = StrandsCurves.CurvesOffset[CurveIndex]; PointIndex < StrandsCurves.CurvesOffset[CurveIndex + 1]; ++PointIndex)
-			{
-				const FVector RootPosition = StrandsPoints.PointsPosition[PointIndex];
-				UE_LOG(LogHairStrands, Log, TEXT("Signed Distance = %d %d %f"), CurveIndex, PointIndex, LevelSet.SignedDistance(RootPosition));
-			}
-		}
-	}*/
+	StrandsCurves.Reset();
+	StrandsPoints.Reset();
 }
diff --git a/Engine/Plugins/Experimental/HairStrands/Source/HairStrandsCore/Private/HairStrandsRendering.cpp b/Engine/Plugins/Experimental/HairStrands/Source/HairStrandsCore/Private/HairStrandsRendering.cpp
--- a/Engine/Plugins/Experimental/HairStrands/Source/HairStrandsCore/Private/HairStrandsRendering.cpp
+++ b/Engine/Plugins/Experimental/HairStrands/Source/HairStrandsCore/Private/HairStrandsRendering.cpp
@@ -108,11 +108,17 @@ class FHairInterpolationCS : public FGlobalShader
 
 	class FGroupSize : SHADER_PERMUTATION_INT("PERMUTATION_GROUP_SIZE", 2);
 	class FDebug : SHADER_PERMUTATION_INT("PERMUTATION_DEBUG", 2);
-	using FPermutationDomain = TShaderPermutationDomain<FGroupSize, FDebug>;
+	class FDynamicGeometry : SHADER_PERMUTATION_INT("PERMUTATION_DYNAMIC_GEOMETRY", 2);
+	using FPermutationDomain = TShaderPermutationDomain<FGroupSize, FDebug, FDynamicGeometry>;
 
 	BEGIN_SHADER_PARAMETER_STRUCT(FParameters, )
 		SHADER_PARAMETER(uint32, VertexCount)
 		SHADER_PARAMETER(uint32, DispatchCountX)
+
+		SHADER_PARAMETER(FVector, RestPositionWorldCenter)
+		SHADER_PARAMETER(FVector, DeformedPositionWorldCenter)
+		SHADER_PARAMETER(FVector, OutRenderDeformedPositionCenter)
+
 		SHADER_PARAMETER_SRV(Buffer, RenderRestPosePositionBuffer)
 		SHADER_PARAMETER_UAV(RWBuffer, OutRenderDeformedPositionBuffer)
 
@@ -124,6 +130,16 @@ class FHairInterpolationCS : public FGlobalShader
 
 		SHADER_PARAMETER_SRV(Buffer, SimAttributeBuffer)
 		SHADER_PARAMETER_UAV(RWBuffer, OutRenderAttributeBuffer)
+		SHADER_PARAMETER_SRV(Buffer<float4>, RestPosition0Buffer)
+		SHADER_PARAMETER_SRV(Buffer<float4>, RestPosition1Buffer)
+		SHADER_PARAMETER_SRV(Buffer<float4>, RestPosition2Buffer)
+
+		SHADER_PARAMETER_SRV(Buffer<float4>, DeformedPosition0Buffer)
+		SHADER_PARAMETER_SRV(Buffer<float4>, DeformedPosition1Buffer)
+		SHADER_PARAMETER_SRV(Buffer<float4>, DeformedPosition2Buffer)
+
+		SHADER_PARAMETER_SRV(Buffer<uint>,	RootToTriangleIndex)
+		SHADER_PARAMETER_SRV(Buffer<uint>, VertexToRootIndexBuffer)
 
 		END_SHADER_PARAMETER_STRUCT()
 
@@ -135,6 +151,8 @@ IMPLEMENT_GLOBAL_SHADER(FHairInterpolationCS, "/Engine/Private/HairStrands/HairS
 
 static void AddHairStrandsInterpolationPass(
 	FRDGBuilder& GraphBuilder,
+	const FHairStrandsProjectionHairData& InHairData,
+	const int32 LODIndex,
 	const uint32 VertexCount,
 	const FShaderResourceViewRHIRef& RenderRestPosePositionBuffer,
 	const FShaderResourceViewRHIRef& Interpolation0Buffer,
@@ -163,11 +181,32 @@ static void AddHairStrandsInterpolationPass(
 	}
 	Parameters->VertexCount = VertexCount;
 	Parameters->DispatchCountX = DispatchCount.X;
+	Parameters->OutRenderDeformedPositionCenter = FVector::ZeroVector;
+
+	const bool bSupportDynamicMesh = InHairData.RootCount > 0 && LODIndex >= 0 && LODIndex < InHairData.LODDatas.Num() && InHairData.LODDatas[LODIndex].bIsValid;
+	if (bSupportDynamicMesh)
+	{
+		Parameters->OutRenderDeformedPositionCenter = FVector::ZeroVector; // TODO
+
+		Parameters->RestPositionWorldCenter = FVector::ZeroVector; // TODO: InHairData.LODDatas[LODIndex].RestRootCenter;
+		Parameters->RestPosition0Buffer = InHairData.LODDatas[LODIndex].RestRootTrianglePosition0Buffer->SRV;
+		Parameters->RestPosition1Buffer = InHairData.LODDatas[LODIndex].RestRootTrianglePosition1Buffer->SRV;
+		Parameters->RestPosition2Buffer = InHairData.LODDatas[LODIndex].RestRootTrianglePosition2Buffer->SRV;
+
+		Parameters->DeformedPositionWorldCenter = FVector::ZeroVector; // TODO: InHairData.LODDatas[LODIndex].DeformedRootCenter;
+		Parameters->DeformedPosition0Buffer = InHairData.LODDatas[LODIndex].DeformedRootTrianglePosition0Buffer->SRV;
+		Parameters->DeformedPosition1Buffer = InHairData.LODDatas[LODIndex].DeformedRootTrianglePosition1Buffer->SRV;
+		Parameters->DeformedPosition2Buffer = InHairData.LODDatas[LODIndex].DeformedRootTrianglePosition2Buffer->SRV;
+
+		Parameters->RootToTriangleIndex = InHairData.LODDatas[LODIndex].RootTriangleIndexBuffer->SRV;
+		Parameters->VertexToRootIndexBuffer = InHairData.VertexToCurveIndexBuffer->SRV;
+	}
 
 	FHairInterpolationCS::FPermutationDomain PermutationVector;
 	PermutationVector.Set<FHairInterpolationCS::FGroupSize>(GetGroupSizePermutation(GroupSize));
 	PermutationVector.Set<FHairInterpolationCS::FDebug>(bCopySimAttributesToRenderAttributes ? 1 : 0);
-
+	PermutationVector.Set<FHairInterpolationCS::FDynamicGeometry>(bSupportDynamicMesh ? 1 : 0);
+	
 	TShaderMap<FGlobalShaderType>* ShaderMap = GetGlobalShaderMap(ERHIFeatureLevel::SM5);
 
 	TShaderMapRef<FHairInterpolationCS> ComputeShader(ShaderMap, PermutationVector);
@@ -328,7 +367,9 @@ static void BuildHairAccelerationStructure(FRHICommandList& RHICmdList, uint32 R
 void ComputeHairStrandsInterpolation(
 	FRHICommandListImmediate& RHICmdList,
 	FHairStrandsInterpolationInput* InInput,
-	FHairStrandsInterpolationOutput* InOutput)
+	FHairStrandsInterpolationOutput* InOutput,
+	FHairStrandsProjectionHairData& InHairData,
+	int32 LODIndex)
 {
 	if (!InInput || !InOutput) return;
 	FHairStrandsInterpolationInput& Input = *InInput;
@@ -370,7 +411,7 @@ void ComputeHairStrandsInterpolation(
 	// * None	: Display hair normally
 	// * Sim	: Show sim strands
 	// * Render : Show rendering strands with sim color influence
-	const EHairStrandsDebugMode DebugMode = GetHairStrandsDebugMode();
+	const EHairStrandsDebugMode DebugMode = GetHairStrandsDebugStrandsMode();
 	if (DebugMode == EHairStrandsDebugMode::SimHairStrands)
 	{
 		AddHairTangentPass(
@@ -398,6 +439,8 @@ void ComputeHairStrandsInterpolation(
 
 		AddHairStrandsInterpolationPass(
 			GraphBuilder,
+			InHairData,
+			LODIndex,
 			Input.RenderVertexCount,
 			Input.RenderRestPosePositionBuffer->SRV,
 			Input.Interpolation0Buffer->SRV,
diff --git a/Engine/Plugins/Experimental/HairStrands/Source/HairStrandsCore/Public/GroomAsset.h b/Engine/Plugins/Experimental/HairStrands/Source/HairStrandsCore/Public/GroomAsset.h
--- a/Engine/Plugins/Experimental/HairStrands/Source/HairStrandsCore/Public/GroomAsset.h
+++ b/Engine/Plugins/Experimental/HairStrands/Source/HairStrandsCore/Public/GroomAsset.h
@@ -14,11 +14,69 @@
 class UMaterialInterface;
 class UNiagaraSystem;
 
+/* Render buffers for root deformation for dynamic meshes */
+struct FHairStrandsRootResource : public FRenderResource
+{
+	/** Build the hair strands resource */
+	FHairStrandsRootResource(const FHairStrandsDatas* HairStrandsDatas, uint32 LODCount);
+
+	/* Init the buffer */
+	virtual void InitRHI() override;
+
+	/* Release the buffer */
+	virtual void ReleaseRHI() override;
+
+	/* Get the resource name */
+	virtual FString GetFriendlyName() const override { return TEXT("FHairStrandsRootResource"); }
+	
+	FRWBuffer RootPositionBuffer;
+	FRWBuffer RootNormalBuffer;
+
+	struct FMeshProjectionLOD
+	{
+		enum class EStatus { Invalid, Initialized, Projected };
+		EStatus Status = EStatus::Invalid;
+		int32 LODIndex = -1;
+
+		/* Triangle on which a root is attached */
+		FRWBuffer RootTriangleIndexBuffer;
+		FRWBuffer RootTriangleBarycentricBuffer;
+	
+		/* Strand hair roots translation and rotation in rest position relative to the bound triangle. Positions are relative to the rest root center */
+		FVector	  RestRootCenter = FVector::ZeroVector;
+		FRWBuffer RestRootTrianglePosition0Buffer;
+		FRWBuffer RestRootTrianglePosition1Buffer;
+		FRWBuffer RestRootTrianglePosition2Buffer;
+
+		/* Strand hair roots translation and rotation in triangle-deformed position relative to the bound triangle. Positions are relative the deformed root center*/
+		FVector   DeformedRootCenter = FVector::ZeroVector;
+		FRWBuffer DeformedRootTrianglePosition0Buffer;
+		FRWBuffer DeformedRootTrianglePosition1Buffer;
+		FRWBuffer DeformedRootTrianglePosition2Buffer;
+	};
+
+	/* Store the hair projection information for each mesh LOD */
+	TArray<FMeshProjectionLOD> MeshProjectionLODs;
+
+	/* Strand hair vertex to curve index */
+	FRWBuffer VertexToCurveIndexBuffer;
+
+	const uint32 RootCount;
+	/* Curve index for every vertices */
+	TArray<FHairStrandsIndexFormat::Type> CurveIndices;
+
+	/* Curve root's positions */
+	TArray<FHairStrandsRootPositionFormat::Type> RootPositions;
+
+	/* Curve root's normal orientation */
+	TArray<FHairStrandsRootNormalFormat::Type> RootNormals;
+};
+
 /* Render buffers that will be used for rendering */
 struct FHairStrandsResource : public FRenderResource
 {
 	/** Build the hair strands resource */
-	FHairStrandsResource(FHairStrandsDatas* HairStrandsDatas);
+	FHairStrandsResource(const FHairStrandsDatas::FRenderData& HairStrandRenderData);
 
 	/* Init the buffer */
 	virtual void InitRHI() override;
@@ -41,14 +99,14 @@ struct FHairStrandsResource : public FRenderResource
 	/* Strand hair offset buffer */
 	FRWBuffer AttributeBuffer;
 
-	/* Pointer to the hair strands datas */
-	FHairStrandsDatas* StrandsDatas;
+	/* Reference to the hair strands render data */
+	const FHairStrandsDatas::FRenderData& RenderData;
 };
 
 struct FHairStrandsInterpolationResource : public FRenderResource
 {
 	/** Build the hair strands resource */
-	FHairStrandsInterpolationResource(const FHairStrandsInterpolationDatas& InterpolationData, const FHairStrandsDatas& SimDatas);
+	FHairStrandsInterpolationResource(const FHairStrandsInterpolationDatas::FRenderData& InterpolationRenderData, const FHairStrandsDatas& SimDatas);
 
 	/* Init the buffer */
 	virtual void InitRHI() override;
@@ -62,20 +120,19 @@ struct FHairStrandsInterpolationResource : public FRenderResource
 	FRWBuffer Interpolation0Buffer;
 	FRWBuffer Interpolation1Buffer;
 
-	/* Store interpolation data until the resource is created */
-	TArray<FHairStrandsInterpolation0Format::Type> Interpolation0;
-	TArray<FHairStrandsInterpolation1Format::Type> Interpolation1;
-
 	// For debug purpose only (should be remove once all hair simulation is correctly handled)
 	FRWBuffer SimRootPointIndexBuffer;
 	TArray<FHairStrandsRootIndexFormat::Type> SimRootPointIndex;
+
+	/* Reference to the hair strands interpolation render data */
+	const FHairStrandsInterpolationDatas::FRenderData& RenderData;
 };
 
 #if RHI_RAYTRACING
 struct FHairStrandsRaytracingResource : public FRenderResource
 {
 	/** Build the hair strands resource */
-	FHairStrandsRaytracingResource(FHairStrandsDatas* HairStrandsDatas);
+	FHairStrandsRaytracingResource(const FHairStrandsDatas& HairStrandsDatas);
 
 	/* Init the buffer */
 	virtual void InitRHI() override;
@@ -93,7 +150,7 @@ struct FHairStrandsRaytracingResource : public FRenderResource
 #endif
 
 USTRUCT(BlueprintType)
-struct HAIRSTRANDSCORE_API FHairGroupRenderSettings
+struct HAIRSTRANDSCORE_API FHairGroupInfo
 {
 	GENERATED_BODY()
 
@@ -103,28 +160,32 @@ struct HAIRSTRANDSCORE_API FHairGroupRenderSettings
 	UPROPERTY(VisibleAnywhere, Category = "Info", meta = (DisplayName = "Curve Count"))
 	int32 NumCurves = 0;
 
+	UPROPERTY(VisibleAnywhere, Category = "Info", meta = (DisplayName = "Guide Count"))
+	int32 NumGuides = 0;
+
+	UPROPERTY(VisibleAnywhere, Category = "Info", meta = (DisplayName = "Auto-generated Guides", ToolTip = "Checked when guides were auto-generated"))
+	bool bIsAutoGenerated = false;
+
 	UPROPERTY(EditAnywhere, Category="Rendering")
 	UMaterialInterface* Material = nullptr;
+
+	UPROPERTY(EditAnywhere, Category="Simulation", meta = (DisplayName = "Niagara System Asset"))
+	UNiagaraSystem* NiagaraAsset = nullptr;
 };
 
-USTRUCT(BlueprintType)
-struct HAIRSTRANDSCORE_API FHairGroupSimulationSettings
+struct HAIRSTRANDSCORE_API FHairGroupData
 {
-	GENERATED_BODY()
-
-	UPROPERTY(VisibleAnywhere, Category = "Info")
-	int32 GroupID = 0;
-
-	UPROPERTY(VisibleAnywhere, Category = "Info", meta = (DisplayName = "Curve Count"))
-	int32 NumCurves = 0;
+	FHairStrandsDatas HairRenderData;
+	FHairStrandsDatas HairSimulationData;
+	FHairStrandsInterpolationDatas HairInterpolationData;
 
-	UPROPERTY(VisibleAnywhere, Category = "Info", meta = (DisplayName = "Auto-generated Guides", ToolTip = "Checked when guides were auto-generated"))
-	bool bIsAutoGenerated = false;
+	/** Interpolated hair render resource to be allocated */
+	FHairStrandsResource* HairStrandsResource = nullptr;
 
-	UPROPERTY(EditAnywhere, Category="Simulation", meta = (DisplayName = "Niagara System Asset"))
-	UNiagaraSystem* Asset = nullptr;
+	/** Guide render resource to be allocated */
+	FHairStrandsResource* HairSimulationResource = nullptr;
 
-	// Add other physics/simulation settings
+	friend FArchive& operator<<(FArchive& Ar, FHairGroupData& GroupData);
 };
 
 /**
@@ -137,18 +198,10 @@ class HAIRSTRANDSCORE_API UGroomAsset : public UObject
 
 public:
 
-	// #ueent_todo: Hair strands data should be computed and saved in the DDC
-	/** Hair strands data for rendering */
-	FHairStrandsDatas HairRenderData;
-
-	/** Guide strands data for simulation */
-	FHairStrandsDatas HairSimulationData;
-
-	UPROPERTY(EditAnywhere, EditFixedSize, BlueprintReadWrite, Category = "HairGroupSettings", meta = (DisplayName = "Rendering"))
-	TArray<FHairGroupRenderSettings> RenderHairGroups;
+	UPROPERTY(EditAnywhere, EditFixedSize, BlueprintReadWrite, Category = "HairGroups", meta = (DisplayName = "Group"))
+	TArray<FHairGroupInfo> HairGroupsInfo;
 
-	UPROPERTY(EditAnywhere, EditFixedSize, BlueprintReadWrite, Category = "HairGroupSettings", meta = (DisplayName = "Simulation"))
-	TArray<FHairGroupSimulationSettings> SimulationHairGroups;
+	TArray<FHairGroupData> HairGroupsData;
 
 	//~ Begin UObject Interface.
 	virtual void PostLoad() override;
@@ -156,7 +209,7 @@ public:
 	virtual void Serialize(FArchive& Ar) override;
 
 	/** Density factor for converting hair into guide curve if no guides are provided. */
-	UPROPERTY(EditAnywhere, BlueprintReadWrite, Category = "BuildSettings", meta = (ClampMin = "0.01", UIMin = "0.01", UIMax = "1.0"))
+	UPROPERTY(VisibleAnywhere, BlueprintReadWrite, Category = "BuildSettings", meta = (ClampMin = "0.01", UIMin = "0.01", UIMax = "1.0"))
 	float HairToGuideDensity = 0.1f;
 
 #if WITH_EDITOR
@@ -196,13 +249,9 @@ public:
 
 	void Reset();
 
+	int32 GetNumHairGroups() const;
+
 //private : 
 
 	TUniquePtr<FHairDescription> HairDescription;
-
-	/** Interpolated hair render resource to be allocated */
-	FHairStrandsResource* HairStrandsResource;
-
-	/** Guide render resource to be allocated */
-	FHairStrandsResource* HairSimulationResource;
 };
diff --git a/Engine/Plugins/Experimental/HairStrands/Source/HairStrandsCore/Public/GroomBuilder.h b/Engine/Plugins/Experimental/HairStrands/Source/HairStrandsCore/Public/GroomBuilder.h
new file mode 100644
--- /dev/null
+++ b/Engine/Plugins/Experimental/HairStrands/Source/HairStrandsCore/Public/GroomBuilder.h
@@ -0,0 +1,15 @@
+// Copyright 1998-2019 Epic Games, Inc. All Rights Reserved.
+
+#pragma once
+
+#include "CoreMinimal.h"
+
+struct FHairStrandsDatas;
+class UGroomAsset;
+
+struct HAIRSTRANDSCORE_API FGroomBuilder
+{
+	static bool BuildGroom(const class FHairDescription& HairDescription, const struct FGroomBuildSettings& BuildSettings, UGroomAsset* GroomAsset);
+	static void GenerateGuides(const FHairStrandsDatas& InData, float DecimationPercentage, FHairStrandsDatas& OutData);
+	static void BuildData(UGroomAsset* GroomAsset);
+};
\ No newline at end of file
diff --git a/Engine/Plugins/Experimental/HairStrands/Source/HairStrandsCore/Public/GroomComponent.h b/Engine/Plugins/Experimental/HairStrands/Source/HairStrandsCore/Public/GroomComponent.h
--- a/Engine/Plugins/Experimental/HairStrands/Source/HairStrandsCore/Public/GroomComponent.h
+++ b/Engine/Plugins/Experimental/HairStrands/Source/HairStrandsCore/Public/GroomComponent.h
@@ -27,14 +27,20 @@ public:
 	UPROPERTY(EditAnywhere, BlueprintReadWrite, Category = "HairStrands Rendering", meta = (ClampMin = "0.0001", UIMin = "0.01", UIMax = "10.0"))
 	float HairDensity;
 
-	/** Threshold for merging consecutive hair segments when loading asset */
-	UPROPERTY(EditAnywhere, BlueprintReadWrite, Category = "HairStrands Rendering", meta = (ClampMin = "0.01", UIMin = "0.01", UIMax = "1.0"))
-	float MergeThreshold;
+	/** 
+	 * When activated, the hair groom will be attached and skinned onto the mesh, if the groom component is a child of a skeletal/skinned component.
+	 * This requires the following projection settings: 
+	 * - Rendering settings: 'Skin cache' enabled
+	 * - Animation settings: 'Tick Animation On Skeletal Mesh Init' disabled
+	 */
+	UPROPERTY(EditAnywhere, BlueprintReadOnly, Category = "HairStrands Rendering")
+	bool bSkinGroom;
 
 	//~ Begin UActorComponent Interface.
 	virtual void OnRegister() override;
 	virtual void OnUnregister() override;
 	virtual void TickComponent(float DeltaTime, enum ELevelTick TickType, FActorComponentTickFunction *ThisTickFunction) override;
+	virtual void SendRenderTransform_Concurrent() override;
 	//~ End UActorComponent Interface.
 
 	//~ Begin USceneComponent Interface.
@@ -52,10 +58,10 @@ public:
 	//~ End UMeshComponent Interface.
 
 	/** Return the guide hairs datas */
-	FHairStrandsDatas* GetGuideStrandsDatas() { return GroomAsset ? &GroomAsset->HairSimulationData : nullptr; }
+	FHairStrandsDatas* GetGuideStrandsDatas();
 
 	/** Return the guide hairs resources*/
-	FHairStrandsResource* GetGuideStrandsResource() { return GroomAsset ? GroomAsset->HairSimulationResource : nullptr; }
+	FHairStrandsResource* GetGuideStrandsResource();
 
 #if WITH_EDITOR
 	virtual void PostEditChangeProperty(FPropertyChangedEvent& PropertyChangedEvent) override;
@@ -64,6 +70,8 @@ public:
 	FHairStrandsInterpolationResource* InterpolationResource = nullptr;
 	struct FHairStrandsInterpolationOutput* InterpolationOutput = nullptr;
 	struct FHairStrandsInterpolationInput* InterpolationInput = nullptr;
+	struct FHairStrandsRootResource* RenRootResources = nullptr;
+	struct FHairStrandsRootResource* SimRootResources = nullptr;
 
 #if RHI_RAYTRACING
 	FHairStrandsRaytracingResource* RaytracingResources = nullptr;
@@ -72,6 +80,17 @@ public:
 private:
 	void* InitializedResources;
 
+	enum class EMeshProjectionState
+	{
+		Invalid,
+		WaitForData,
+		Completed
+	};
+	class USkeletalMeshComponent* RegisteredSkeletalMeshComponent;
+	int32 MeshProjectionLODIndex;
+	uint32 MeshProjectionTickDelay;
+	EMeshProjectionState MeshProjectionState;
+
 	void InitResources();
 	void ReleaseResources();
 
diff --git a/Engine/Plugins/Experimental/HairStrands/Source/HairStrandsCore/Public/HairStrandsDatas.h b/Engine/Plugins/Experimental/HairStrands/Source/HairStrandsCore/Public/HairStrandsDatas.h
--- a/Engine/Plugins/Experimental/HairStrands/Source/HairStrandsCore/Public/HairStrandsDatas.h
+++ b/Engine/Plugins/Experimental/HairStrands/Source/HairStrandsCore/Public/HairStrandsDatas.h
@@ -16,6 +16,8 @@ struct FPackedHairVertex
 	uint8 ControlPointType : 2;
 	uint8 NormalizedRadius : 6;
 	uint8 NormalizedLength;
+
+	friend FArchive& operator<<(FArchive& Ar, FPackedHairVertex& Vertex);
 };
 
 struct FPackedHairAttributeVertex
@@ -24,6 +26,8 @@ struct FPackedHairAttributeVertex
 	uint8 RootV;
 	uint8 UCoord;
 	uint8 Seed;
+
+	friend FArchive& operator<<(FArchive& Ar, FPackedHairAttributeVertex& Vertex);
 };
 
 struct FHairInterpolation0Vertex
@@ -34,6 +38,8 @@ struct FHairInterpolation0Vertex
 
 	uint8 VertexWeight0;
 	uint8 VertexWeight1;
+
+	friend FArchive& operator<<(FArchive& Ar, FPackedHairVertex& FHairInterpolation0Vertex);
 };
 
 struct FHairInterpolation1Vertex
@@ -42,8 +48,18 @@ struct FHairInterpolation1Vertex
 	uint8 VertexIndex1;
 	uint8 VertexIndex2;
 	uint8 Pad0;
+
+	friend FArchive& operator<<(FArchive& Ar, FPackedHairVertex& FHairInterpolation1Vertex);
 };
 
+struct FVector4_16
+{
+	FFloat16 X;
+	FFloat16 Y;
+	FFloat16 Z;
+	FFloat16 W;
+}; 
+
 struct FHairStrandsPositionFormat
 {
 	typedef FPackedHairVertex Type;
@@ -124,7 +140,7 @@ struct HAIRSTRANDSCORE_API FHairStrandsIndexFormat
 };
 
 /** Hair strands weights format */
-struct HAIRSTRANDSCORE_API FHairStrandsWeightFormat
+struct FHairStrandsWeightFormat
 {
 	using Type = float;
 	static const uint32 ComponentCount = 1;
@@ -133,6 +149,60 @@ struct HAIRSTRANDSCORE_API FHairStrandsWeightFormat
 	static const EPixelFormat Format = PF_R32_FLOAT;
 };
 
+struct FHairStrandsCurveTranslationFormat
+{
+	using Type = FVector4_16;
+	static const uint32 ComponentCount = 1;
+	static const uint32 SizeInByte = sizeof(Type);
+	static const EVertexElementType VertexElementType = VET_Float4;
+	static const EPixelFormat Format = PF_FloatRGBA;
+};
+
+struct FHairStrandsCurveRotationFormat
+{
+	using Type = FVector4_16;
+	static const uint32 ComponentCount = 1;
+	static const uint32 SizeInByte = sizeof(Type);
+	static const EVertexElementType VertexElementType = VET_Float4;
+	static const EPixelFormat Format = PF_FloatRGBA;
+};
+
+struct FHairStrandsCurveTriangleIndexFormat
+{
+	using Type = uint32;
+	static const uint32 ComponentCount = 1;
+	static const uint32 SizeInByte = sizeof(Type);
+	static const EVertexElementType VertexElementType = VET_UInt;
+	static const EPixelFormat Format = PF_R32_UINT;
+};
+
+struct FHairStrandsCurveTriangleBarycentricFormat
+{
+	using Type = uint32;
+	static const uint32 ComponentCount = 1;
+	static const uint32 SizeInByte = sizeof(Type);
+	static const EVertexElementType VertexElementType = VET_UInt;
+	static const EPixelFormat Format = PF_R32_UINT;
+};
+
+struct FHairStrandsRootPositionFormat
+{
+	typedef FVector4_16 Type;
+	static const uint32 ComponentCount = 1;
+	static const uint32 SizeInByte = sizeof(Type);
+	static const EVertexElementType VertexElementType = VET_UShort4;
+	static const EPixelFormat Format = PF_FloatRGBA;
+};
+
+struct FHairStrandsRootNormalFormat
+{
+	typedef FVector4_16 Type;
+	static const uint32 ComponentCount = 1;
+	static const uint32 SizeInByte = sizeof(Type);
+	static const EVertexElementType VertexElementType = VET_Float4;
+	static const EPixelFormat Format = PF_FloatRGBA;
+};
+
 /** Hair strands points interpolation attributes */
 struct HAIRSTRANDSCORE_API FHairStrandsInterpolationDatas
 {
@@ -157,14 +227,13 @@ struct HAIRSTRANDSCORE_API FHairStrandsInterpolationDatas
 	/** Weight of vertex indices on simulation curve, ordered by closest influence */
 	TArray<FVector>	PointsSimCurvesVertexWeights;
 
-	void BuildInterpolationDatas(
-		const struct FHairStrandsDatas& SimStrandsData,
-		const struct FHairStrandsDatas& RenderStrandsData);
+	struct FRenderData
+	{
+		TArray<FHairStrandsInterpolation0Format::Type> Interpolation0;
+		TArray<FHairStrandsInterpolation1Format::Type> Interpolation1;
 
-	/** Build data for interpolation between simulation and rendering */
-	void BuildRenderingDatas(
-		TArray<FHairStrandsInterpolation0Format::Type>& OutPointsInterpolation0,
-		TArray<FHairStrandsInterpolation1Format::Type>& OutPointsInterpolation1) const;
+		inline void Serialize(FArchive& Ar);
+	} RenderData;
 };
 
 /** Hair strands points attribute */
@@ -205,7 +274,7 @@ struct HAIRSTRANDSCORE_API FHairStrandsCurves
 	inline void Reset();
 
 	/** Get the number of Curves */
-	inline uint32 Num() const { return CurvesLength.Num(); }
+	inline uint32 Num() const { return CurvesCount.Num(); }
 
 	/** Number of points per rod */
 	TArray<uint16> CurvesCount;
@@ -219,8 +288,6 @@ struct HAIRSTRANDSCORE_API FHairStrandsCurves
 	/** Roots UV in the bounding sphere */
 	TArray<FVector2D> CurvesRootUV; // [0..1]
 
-	TArray<uint32> CurvesGroupID;
-
 	/** Max strands Curves length */
 	float MaxLength = 0;
 
@@ -254,23 +321,11 @@ struct HAIRSTRANDSCORE_API FHairStrandsDatas
 	/* Strands bounding box */
 	FBox BoundingBox;
 
-	/** Build the packed datas for gpu rendering/simulation */
-	void BuildRenderingDatas(
-		TArray<FHairStrandsPositionFormat::Type>& OutPointsPositions,
-		TArray<FHairStrandsAttributeFormat::Type>& OutPackedAttributes) const;
+	struct FRenderData
+	{
+		TArray<FHairStrandsPositionFormat::Type> RenderingPositions;
+		TArray<FHairStrandsAttributeFormat::Type> RenderingAttributes;
 
-	/** Resample the internal points and curves simulation datas */
-	void BuildSimulationDatas(const uint32 StrandSize,
-		TArray<FHairStrandsPositionFormat::Type>& OutNodesPositions,
-		TArray<FHairStrandsWeightFormat::Type>& OutPointsWeights,
-		TArray<FHairStrandsIndexFormat::Type>& OutPointsNodes) const;
-
-	/** Build the internal points and curves datas */
-	void BuildInternalDatas(bool bComputeRootUV = false);
-
-	/** Attach the roots to a static mesh */
-	void AttachStrandsRoots(UStaticMesh* StaticMesh, const FMatrix& TransformMartrix);
+		inline void Serialize(FArchive& Ar);
+	} RenderData;
 };
-
-// #ueent_todo: Move the data building and decimation code to a HairBuilder
-void HAIRSTRANDSCORE_API DecimateStrandData(const FHairStrandsDatas& InData, float DecimationPercentage, FHairStrandsDatas& OutData);
diff --git a/Engine/Plugins/Experimental/HairStrands/Source/HairStrandsCore/Public/HairStrandsRendering.h b/Engine/Plugins/Experimental/HairStrands/Source/HairStrandsCore/Public/HairStrandsRendering.h
--- a/Engine/Plugins/Experimental/HairStrands/Source/HairStrandsCore/Public/HairStrandsRendering.h
+++ b/Engine/Plugins/Experimental/HairStrands/Source/HairStrandsCore/Public/HairStrandsRendering.h
@@ -97,5 +97,7 @@ struct FHairStrandsInterpolationOutput
 void ComputeHairStrandsInterpolation(
 	FRHICommandListImmediate& RHICmdList,
 	FHairStrandsInterpolationInput* Input,
-	FHairStrandsInterpolationOutput* Output);
+	FHairStrandsInterpolationOutput* Output, 
+	struct FHairStrandsProjectionHairData& HairData,
+	int32 LODIndex);
 
diff --git a/Engine/Plugins/Experimental/HairStrands/Source/HairStrandsEditor/HairStrandsEditor.Build.cs b/Engine/Plugins/Experimental/HairStrands/Source/HairStrandsEditor/HairStrandsEditor.Build.cs
--- a/Engine/Plugins/Experimental/HairStrands/Source/HairStrandsEditor/HairStrandsEditor.Build.cs
+++ b/Engine/Plugins/Experimental/HairStrands/Source/HairStrandsEditor/HairStrandsEditor.Build.cs
@@ -14,8 +14,13 @@ namespace UnrealBuildTool.Rules
 				{
 					"Core",
 					"CoreUObject",
+					"EditorStyle",
 					"Engine",
 					"HairStrandsCore",
+					"InputCore",
+					"MainFrame",
+					"Slate",
+					"SlateCore",
 					"UnrealEd",
 					"AssetTools",
 				});
diff --git a/Engine/Plugins/Experimental/HairStrands/Source/HairStrandsEditor/Private/FbxHairTranslator.cpp b/Engine/Plugins/Experimental/HairStrands/Source/HairStrandsEditor/Private/FbxHairTranslator.cpp
--- a/Engine/Plugins/Experimental/HairStrands/Source/HairStrandsEditor/Private/FbxHairTranslator.cpp
+++ b/Engine/Plugins/Experimental/HairStrands/Source/HairStrandsEditor/Private/FbxHairTranslator.cpp
@@ -2,6 +2,7 @@
 
 #include "FbxHairTranslator.h"
 
+#include "GroomImportOptions.h"
 #include "Misc/Paths.h"
 
 #if PLATFORM_WINDOWS
@@ -119,7 +120,7 @@ private:
 	FbxImporter* FileImporter;
 };
 
-static void ParseFbxNode(FbxNode* FileNode, FHairDescription& HairDescription)
+static void ParseFbxNode(FbxNode* FileNode, FHairDescription& HairDescription, const FMatrix& ConversionMatrix, float Scale)
 {
 	if (FileNode->GetNodeAttribute() != NULL && FileNode->GetNodeAttribute()->GetAttributeType() == FbxNodeAttribute::eNurbsCurve)
 	{
@@ -143,8 +144,8 @@ static void ParseFbxNode(FbxNode* FileNode, FHairDescription& HairDescription)
 
 			FVertexID VertexID = HairDescription.AddVertex();
 
-			VertexPositions[VertexID] = Position * FbxHairFormat::UNIT_TO_CM;
-			VertexWidths[VertexID] = Radius * FbxHairFormat::UNIT_TO_CM;
+			VertexPositions[VertexID] = ConversionMatrix.TransformPosition(Position);
+			VertexWidths[VertexID] = Radius * Scale;
 		}
 	}
 	//	#todo_hair
@@ -156,12 +157,12 @@ static void ParseFbxNode(FbxNode* FileNode, FHairDescription& HairDescription)
 	{
 		for (int32 ChildIndex = 0; ChildIndex < FileNode->GetChildCount(); ++ChildIndex)
 		{
-			ParseFbxNode(FileNode->GetChild(ChildIndex), HairDescription);
+			ParseFbxNode(FileNode->GetChild(ChildIndex), HairDescription, ConversionMatrix, Scale);
 		}
 	}
 }
 
-bool FFbxHairTranslator::Translate(const FString& FileName, FHairDescription& HairDescription)
+bool FFbxHairTranslator::Translate(const FString& FileName, FHairDescription& HairDescription, const FGroomConversionSettings& ConversionSettings)
 {
 	// Reuse the FbxHairImporter if there was one created previously by CanTranslate
 	// There could be none if the translator is used for a re-import
@@ -196,9 +197,10 @@ bool FFbxHairTranslator::Translate(const FString& FileName, FHairDescription& Ha
 
 	if (FbxNode* FileNode = FbxHairImporter->GetFbxScene()->GetRootNode())
 	{
+		FMatrix ConversionMatrix = FScaleMatrix::Make(ConversionSettings.Scale) * FRotationMatrix::Make(FQuat::MakeFromEuler(ConversionSettings.Rotation));
 		for (int32 ChildIndex = 0; ChildIndex < FileNode->GetChildCount(); ++ChildIndex)
 		{
-			ParseFbxNode(FileNode->GetChild(ChildIndex), HairDescription);
+			ParseFbxNode(FileNode->GetChild(ChildIndex), HairDescription, ConversionMatrix, ConversionSettings.Scale.X);
 		}
 	}
 
diff --git a/Engine/Plugins/Experimental/HairStrands/Source/HairStrandsEditor/Private/FbxHairTranslator.h b/Engine/Plugins/Experimental/HairStrands/Source/HairStrandsEditor/Private/FbxHairTranslator.h
--- a/Engine/Plugins/Experimental/HairStrands/Source/HairStrandsEditor/Private/FbxHairTranslator.h
+++ b/Engine/Plugins/Experimental/HairStrands/Source/HairStrandsEditor/Private/FbxHairTranslator.h
@@ -9,7 +9,7 @@ class FFbxHairImporter;
 class FFbxHairTranslator : public IHairStrandsTranslator
 {
 public:
-	virtual bool Translate(const FString& FilePath, FHairDescription& OutHairDescription) override;
+	virtual bool Translate(const FString& FilePath, FHairDescription& OutHairDescription, const FGroomConversionSettings& ConversionSettings) override;
 	virtual bool CanTranslate(const FString& FilePath) override;
 	virtual bool IsFileExtensionSupported(const FString& FileExtension) const override;
 	virtual FString GetSupportedFormat() const override;
diff --git a/Engine/Plugins/Experimental/HairStrands/Source/HairStrandsEditor/Private/GroomAssetImportData.h b/Engine/Plugins/Experimental/HairStrands/Source/HairStrandsEditor/Private/GroomAssetImportData.h
new file mode 100644
--- /dev/null
+++ b/Engine/Plugins/Experimental/HairStrands/Source/HairStrandsEditor/Private/GroomAssetImportData.h
@@ -0,0 +1,19 @@
+// Copyright 1998-2019 Epic Games, Inc. All Rights Reserved.
+
+#pragma once
+
+#include "CoreMinimal.h"
+
+#include "EditorFramework/AssetImportData.h"
+#include "GroomAssetImportData.generated.h"
+
+UCLASS()
+class UGroomAssetImportData : public UAssetImportData
+{
+	GENERATED_UCLASS_BODY()
+
+public:
+
+	UPROPERTY()
+	class UGroomImportOptions* ImportOptions;
+};
diff --git a/Engine/Plugins/Experimental/HairStrands/Source/HairStrandsEditor/Private/HairFormatTranslator.cpp b/Engine/Plugins/Experimental/HairStrands/Source/HairStrandsEditor/Private/HairFormatTranslator.cpp
--- a/Engine/Plugins/Experimental/HairStrands/Source/HairStrandsEditor/Private/HairFormatTranslator.cpp
+++ b/Engine/Plugins/Experimental/HairStrands/Source/HairStrandsEditor/Private/HairFormatTranslator.cpp
@@ -3,6 +3,7 @@
 #include "HairFormatTranslator.h"
 
 #include "GenericPlatform/GenericPlatformFile.h"
+#include "GroomImportOptions.h"
 #include "HAL/PlatformFilemanager.h"
 #include "Misc/Paths.h"
 
@@ -52,7 +53,7 @@ namespace HairFormat
 	};
 }
 
-bool FHairFormatTranslator::Translate(const FString& FileName, FHairDescription& HairDescription)
+bool FHairFormatTranslator::Translate(const FString& FileName, FHairDescription& HairDescription, const FGroomConversionSettings& ConversionSettings)
 {
 	IPlatformFile& PlatformFile = FPlatformFileManager::Get().GetPlatformFile();
 	IFileHandle* FileHandle = PlatformFile.OpenRead(*FileName);
@@ -110,6 +111,8 @@ bool FHairFormatTranslator::Translate(const FString& FileName, FHairDescription&
 				TVertexAttributesRef<FVector> VertexPositions = HairDescription.VertexAttributes().GetAttributesRef<FVector>(HairAttribute::Vertex::Position);
 				if (FileHeader.BitArrays & HAIR_FILE_POINTS_BIT)
 				{
+					FMatrix ConversionMatrix = FScaleMatrix::Make(ConversionSettings.Scale) * FRotationMatrix::Make(FQuat::MakeFromEuler(ConversionSettings.Rotation));
+
 					TArray<FVector> Positions;
 					Positions.SetNum(NumVertices);
 
@@ -118,7 +121,7 @@ bool FHairFormatTranslator::Translate(const FString& FileName, FHairDescription&
 					for (int32 VertexIndex = 0; VertexIndex < NumVertices; ++VertexIndex)
 					{
 						FVertexID VertexID(VertexIndex);
-						VertexPositions[VertexID] = Positions[VertexIndex];
+						VertexPositions[VertexID] = ConversionMatrix.TransformPosition(Positions[VertexIndex]);
 					}
 				}
 
@@ -127,6 +130,8 @@ bool FHairFormatTranslator::Translate(const FString& FileName, FHairDescription&
 				TVertexAttributesRef<float> VertexWidths = HairDescription.VertexAttributes().GetAttributesRef<float>(HairAttribute::Vertex::Width);
 				if (FileHeader.BitArrays & HAIR_FILE_THICKNESS_BIT)
 				{
+					float Scale = ConversionSettings.Scale.X;
+
 					TArray<float> Widths;
 					Widths.SetNum(NumVertices);
 
@@ -135,7 +140,7 @@ bool FHairFormatTranslator::Translate(const FString& FileName, FHairDescription&
 					for (int32 VertexIndex = 0; VertexIndex < NumVertices; ++VertexIndex)
 					{
 						FVertexID VertexID(VertexIndex);
-						VertexWidths[VertexID] = Widths[VertexIndex];
+						VertexWidths[VertexID] = Widths[VertexIndex] * Scale;
 					}
 				}
 				
diff --git a/Engine/Plugins/Experimental/HairStrands/Source/HairStrandsEditor/Private/HairStrandsFactory.cpp b/Engine/Plugins/Experimental/HairStrands/Source/HairStrandsEditor/Private/HairStrandsFactory.cpp
--- a/Engine/Plugins/Experimental/HairStrands/Source/HairStrandsEditor/Private/HairStrandsFactory.cpp
+++ b/Engine/Plugins/Experimental/HairStrands/Source/HairStrandsEditor/Private/HairStrandsFactory.cpp
@@ -3,13 +3,17 @@
 #include "HairStrandsFactory.h"
 
 #include "EditorFramework/AssetImportData.h"
-#include "HairDescription.h"
 #include "GroomAsset.h"
+#include "GroomAssetImportData.h"
+#include "GroomImportOptions.h"
+#include "GroomImportOptionsWindow.h"
+#include "HairDescription.h"
 #include "HairStrandsEditor.h"
 #include "HairStrandsImporter.h"
 #include "HairStrandsTranslator.h"
 #include "Misc/Paths.h"
 #include "Misc/ScopedSlowTask.h"
+#include "Modules/ModuleManager.h"
 
 #define LOCTEXT_NAMESPACE "HairStrandsFactory"
 
@@ -27,6 +31,8 @@ UHairStrandsFactory::UHairStrandsFactory(const FObjectInitializer& ObjectInitial
 	// Lazy init the translators to let them register themselves before the CDO is used
 	if (!HasAnyFlags(RF_ClassDefaultObject))
 	{
+		ImportOptions = NewObject<UGroomImportOptions>();
+
 		InitTranslators();
 	}
 }
@@ -57,6 +63,8 @@ void UHairStrandsFactory::GetSupportedFileExtensions(TArray<FString>& OutExtensi
 UObject* UHairStrandsFactory::FactoryCreateFile(UClass* InClass, UObject* InParent, FName InName, EObjectFlags Flags, 
 	const FString& Filename, const TCHAR* Parms, FFeedbackContext* Warn, bool& bOutOperationCanceled) 
 {
+	bOutOperationCanceled = false;
+
 	// Translate the hair data from the file
 	TSharedPtr<IHairStrandsTranslator> SelectedTranslator = GetTranslator(Filename);
 	if (!SelectedTranslator.IsValid())
@@ -64,13 +72,27 @@ UObject* UHairStrandsFactory::FactoryCreateFile(UClass* InClass, UObject* InPare
 		return nullptr;
 	}
 
+	if (!GIsRunningUnattendedScript && !IsAutomatedImport())
+	{
+		// Display import options and handle user cancellation
+		TSharedPtr<SGroomImportOptionsWindow> GroomOptionWindow = SGroomImportOptionsWindow::DisplayOptions(ImportOptions, Filename);
+		if (!GroomOptionWindow->ShouldImport())
+		{
+			bOutOperationCanceled = true;
+			return nullptr;
+		}
+
+		// Save the options as the new default
+		ImportOptions->SaveConfig();
+	}
+
 	//GEditor->GetEditorSubsystem<UImportSubsystem>()->BroadcastAssetPreImport(this, InClass, InParent, InName, FileType);
 
 	FScopedSlowTask Progress((float) 1, LOCTEXT("ImportHairAsset", "Importing hair asset..."), true);
 	Progress.MakeDialog(true);
 
 	FHairDescription HairDescription;
-	if (!SelectedTranslator->Translate(Filename, HairDescription))
+	if (!SelectedTranslator->Translate(Filename, HairDescription, ImportOptions->ConversionSettings))
 	{
 		return nullptr;
 	}
@@ -78,16 +100,19 @@ UObject* UHairStrandsFactory::FactoryCreateFile(UClass* InClass, UObject* InPare
 	// Might try to import the same file in the same folder, so if an asset already exists there, reuse and update it
 	UGroomAsset* ExistingAsset = FindObject<UGroomAsset>(InParent, *InName.ToString());
 
-	FHairImportContext HairImportContext(InParent, InClass, InName, Flags);
+	FHairImportContext HairImportContext(ImportOptions, InParent, InClass, InName, Flags);
 	UGroomAsset* CurrentAsset = FHairStrandsImporter::ImportHair(HairImportContext, HairDescription, ExistingAsset);
 	if (CurrentAsset)
 	{
 		// Setup asset import data
-		if (!CurrentAsset->AssetImportData)
+		if (!CurrentAsset->AssetImportData || !CurrentAsset->AssetImportData->IsA<UGroomAssetImportData>())
 		{
-			CurrentAsset->AssetImportData = NewObject<UAssetImportData>(CurrentAsset);
+			CurrentAsset->AssetImportData = NewObject<UGroomAssetImportData>(CurrentAsset);
 		}
 		CurrentAsset->AssetImportData->Update(Filename);
+
+		UGroomAssetImportData* GroomAssetImportData = Cast<UGroomAssetImportData>(CurrentAsset->AssetImportData);
+		GroomAssetImportData->ImportOptions = DuplicateObject<UGroomImportOptions>(ImportOptions, GroomAssetImportData);
 	}
 
 	return CurrentAsset;
diff --git a/Engine/Plugins/Experimental/HairStrands/Source/HairStrandsEditor/Private/HairStrandsImporter.cpp b/Engine/Plugins/Experimental/HairStrands/Source/HairStrandsEditor/Private/HairStrandsImporter.cpp
--- a/Engine/Plugins/Experimental/HairStrands/Source/HairStrandsEditor/Private/HairStrandsImporter.cpp
+++ b/Engine/Plugins/Experimental/HairStrands/Source/HairStrandsEditor/Private/HairStrandsImporter.cpp
@@ -4,12 +4,14 @@
 
 #include "HairDescription.h"
 #include "GroomAsset.h"
-#include "GroomComponent.h"
+#include "GroomBuilder.h"
+#include "GroomImportOptions.h"
 
 DEFINE_LOG_CATEGORY_STATIC(LogHairImporter, Log, All);
 
-FHairImportContext::FHairImportContext(UObject* InParent, UClass* InClass, FName InName, EObjectFlags InFlags)
-	: Parent(InParent)
+FHairImportContext::FHairImportContext(UGroomImportOptions* InImportOptions, UObject* InParent, UClass* InClass, FName InName, EObjectFlags InFlags)
+	: ImportOptions(InImportOptions)
+	, Parent(InParent)
 	, Class(InClass)
 	, Name(InName)
 	, Flags(InFlags)
@@ -18,143 +20,6 @@ FHairImportContext::FHairImportContext(UObject* InParent, UClass* InClass, FName
 
 UGroomAsset* FHairStrandsImporter::ImportHair(const FHairImportContext& ImportContext, FHairDescription& HairDescription, UGroomAsset* ExistingHair)
 {
-	// For now, just convert HairDescription to HairStrandsDatas
-	int32 NumCurves = HairDescription.GetNumStrands();
-	int32 NumVertices = HairDescription.GetNumVertices();
-
-	// Check for required attributes
-	TGroomAttributesRef<int> MajorVersion = HairDescription.GroomAttributes().GetAttributesRef<int>(HairAttribute::Groom::MajorVersion);
-	TGroomAttributesRef<int> MinorVersion = HairDescription.GroomAttributes().GetAttributesRef<int>(HairAttribute::Groom::MinorVersion);
-
-	if (!MajorVersion.IsValid() || !MinorVersion.IsValid())
-	{
-		UE_LOG(LogHairImporter, Warning, TEXT("Failed to import hair: No version number attributes found. Please re-export the input file."));
-		return nullptr;
-	}
-
-	FGroomID GroomID(0);
-	
-	TGroomAttributesRef<float> GroomHairWidthAttribute = HairDescription.GroomAttributes().GetAttributesRef<float>(HairAttribute::Groom::Width);
-	TOptional<float> GroomHairWidth;
-	if (GroomHairWidthAttribute.IsValid())
-	{
-		GroomHairWidth = GroomHairWidthAttribute[GroomID];
-	}
-
-	TGroomAttributesRef<FVector> GroomHairColorAttribute = HairDescription.GroomAttributes().GetAttributesRef<FVector>(HairAttribute::Groom::Color);
-	TOptional<FVector> GroomHairColor;
-	if (GroomHairColorAttribute.IsValid())
-	{
-		GroomHairColor = GroomHairColorAttribute[GroomID];
-	}
-
-	TVertexAttributesRef<FVector> VertexPositions = HairDescription.VertexAttributes().GetAttributesRef<FVector>(HairAttribute::Vertex::Position);
-	TStrandAttributesRef<int> StrandNumVertices = HairDescription.StrandAttributes().GetAttributesRef<int>(HairAttribute::Strand::VertexCount);
-
-	if (!VertexPositions.IsValid() || !StrandNumVertices.IsValid())
-	{
-		UE_LOG(LogHairImporter, Warning, TEXT("Failed to import hair: No vertices or curves data found."));
-		return nullptr;
-	}
-
-	TVertexAttributesRef<float> VertexWidths = HairDescription.VertexAttributes().GetAttributesRef<float>(HairAttribute::Vertex::Width);
-	TStrandAttributesRef<float> StrandWidths = HairDescription.StrandAttributes().GetAttributesRef<float>(HairAttribute::Strand::Width);
-
-	TStrandAttributesRef<FVector2D> StrandRootUV = HairDescription.StrandAttributes().GetAttributesRef<FVector2D>(HairAttribute::Strand::RootUV);
-	bool bHasUVData = StrandRootUV.IsValid();
-
-	TStrandAttributesRef<bool> StrandGuides = HairDescription.StrandAttributes().GetAttributesRef<bool>(HairAttribute::Strand::Guide);
-	TStrandAttributesRef<int> GroupIDs = HairDescription.StrandAttributes().GetAttributesRef<int>(HairAttribute::Strand::GroupID);
-
-	FHairStrandsDatas HairRenderData;
-	FHairStrandsDatas HairSimulationData;
-	FHairStrandsDatas* CurrentHairStrandsDatas = &HairRenderData;
-
-	TMap<int32, FHairGroupRenderSettings> RenderHairGroups;
-	TMap<int32, FHairGroupSimulationSettings> SimHairGroups;
-
-	int32 GlobalVertexIndex = 0;
-	int32 NumHairCurves = 0;
-	int32 NumGuideCurves = 0;
-	int32 NumHairPoints = 0;
-	int32 NumGuidePoints = 0;
-	for (int32 CurveIndex = 0; CurveIndex < NumCurves; ++CurveIndex)
-	{
-		FStrandID StrandID(CurveIndex);
-
-		bool bIsGuide = false;
-		if (StrandGuides.IsValid())
-		{
-			bIsGuide = StrandGuides[StrandID];
-		}
-
-		int32 CurveNumVertices = StrandNumVertices[StrandID];
-
-		int32 GroupID = 0;
-		if (GroupIDs.IsValid())
-		{
-			GroupID = GroupIDs[StrandID];
-		}
-
-		if (!bIsGuide)
-		{
-			++NumHairCurves;
-			NumHairPoints += CurveNumVertices;
-			CurrentHairStrandsDatas = &HairRenderData;
-
-			FHairGroupRenderSettings& GroupSettings = RenderHairGroups.FindOrAdd(GroupID);
-			GroupSettings.GroupID = GroupID;
-			++GroupSettings.NumCurves;
-		}
-		else
-		{
-			++NumGuideCurves;
-			NumGuidePoints += CurveNumVertices;
-			CurrentHairStrandsDatas = &HairSimulationData;
-
-			FHairGroupSimulationSettings& GroupSettings = SimHairGroups.FindOrAdd(GroupID);
-			GroupSettings.GroupID = GroupID;
-			++GroupSettings.NumCurves;
-		}
-
-		CurrentHairStrandsDatas->StrandsCurves.CurvesCount.Add(CurveNumVertices);
-		CurrentHairStrandsDatas->StrandsCurves.CurvesGroupID.Add(GroupID);
-
-		if (bHasUVData)
-		{
-			CurrentHairStrandsDatas->StrandsCurves.CurvesRootUV.Add(StrandRootUV[StrandID]);
-		}
-
-		float StrandWidth = GroomHairWidth ? GroomHairWidth.GetValue() : 0.f;
-		if (StrandWidths.IsValid())
-		{
-			StrandWidth = StrandWidths[StrandID];
-		}
-
-		for (int32 VertexIndex = 0; VertexIndex < CurveNumVertices; ++VertexIndex, ++GlobalVertexIndex)
-		{
-			FVertexID VertexID(GlobalVertexIndex);
-
-			CurrentHairStrandsDatas->StrandsPoints.PointsPosition.Add(VertexPositions[VertexID]);
-
-			float VertexWidth = 0.f;
-			if (VertexWidths.IsValid())
-			{
-				VertexWidth = VertexWidths[VertexID];
-			}
-
-			// Fall back to strand width if there was no vertex width
-			if (VertexWidth == 0.f && StrandWidth != 0.f)
-			{
-				VertexWidth = StrandWidth;
-			}
-
-			CurrentHairStrandsDatas->StrandsPoints.PointsRadius.Add(VertexWidth * 0.5f);
-		}
-	}
-
-	FGroomComponentRecreateRenderStateContext RecreateRenderContext(ExistingHair);
-
 	UGroomAsset* HairAsset = nullptr;
 	if (ExistingHair)
 	{
@@ -174,55 +39,19 @@ UGroomAsset* FHairStrandsImporter::ImportHair(const FHairImportContext& ImportCo
 		}
 	}
 
-	HairRenderData.StrandsCurves.SetNum(NumHairCurves);
-	HairRenderData.StrandsPoints.SetNum(NumHairPoints);
-	HairRenderData.BuildInternalDatas(!bHasUVData);
-	HairAsset->HairRenderData = MoveTemp(HairRenderData);
-
-	if (NumGuideCurves > 0)
-	{
-		HairSimulationData.StrandsCurves.SetNum(NumGuideCurves);
-		HairSimulationData.StrandsPoints.SetNum(NumGuidePoints);
-		HairSimulationData.BuildInternalDatas(true); // Imported guides don't currently have root UVs so force computing them
-		HairAsset->HairSimulationData = MoveTemp(HairSimulationData);
-	}
-	else
-	{
-		DecimateStrandData(HairAsset->HairRenderData, FMath::Clamp(HairAsset->HairToGuideDensity, 0.f, 1.f), HairAsset->HairSimulationData);
-	}
-
-	HairAsset->HairDescription = MakeUnique<FHairDescription>(MoveTemp(HairDescription));
-
-	if (RenderHairGroups.Num() > 0)
-	{
-		for (TPair<int32, FHairGroupRenderSettings>& Pair : RenderHairGroups)
-		{
-			HairAsset->RenderHairGroups.Add(MoveTemp(Pair.Value));
-		}
-	}
-	else
-	{
-		FHairGroupRenderSettings& GroupSettings = HairAsset->RenderHairGroups.AddDefaulted_GetRef();
-		GroupSettings.GroupID = 0;
-		GroupSettings.NumCurves = HairRenderData.GetNumCurves();
-	}
-
-	if (SimHairGroups.Num() > 0)
+	bool bSucceeded = FGroomBuilder::BuildGroom(HairDescription, ImportContext.ImportOptions->BuildSettings, HairAsset);
+	if (!bSucceeded)
 	{
-		for (TPair<int32, FHairGroupSimulationSettings>& Pair : SimHairGroups)
+		// Purge the newly created asset that failed to import
+		if (HairAsset != ExistingHair)
 		{
-			HairAsset->SimulationHairGroups.Add(MoveTemp(Pair.Value));
+			HairAsset->ClearFlags(RF_Standalone);
+			CollectGarbage(GARBAGE_COLLECTION_KEEPFLAGS);
 		}
-	}
-	else
-	{
-		FHairGroupSimulationSettings& GroupSettings = HairAsset->SimulationHairGroups.AddDefaulted_GetRef();
-		GroupSettings.GroupID = 0;
-		GroupSettings.NumCurves = HairAsset->HairSimulationData.GetNumCurves();
-		GroupSettings.bIsAutoGenerated = true;
+		return nullptr;
 	}
 
-	HairAsset->InitResource();
+	HairAsset->HairDescription = MakeUnique<FHairDescription>(MoveTemp(HairDescription));
 
 	return HairAsset;
 }
diff --git a/Engine/Plugins/Experimental/HairStrands/Source/HairStrandsEditor/Private/PbrtHairTranslator.cpp b/Engine/Plugins/Experimental/HairStrands/Source/HairStrandsEditor/Private/PbrtHairTranslator.cpp
--- a/Engine/Plugins/Experimental/HairStrands/Source/HairStrandsEditor/Private/PbrtHairTranslator.cpp
+++ b/Engine/Plugins/Experimental/HairStrands/Source/HairStrandsEditor/Private/PbrtHairTranslator.cpp
@@ -1,8 +1,8 @@
 // Copyright 1998-2019 Epic Games, Inc. All Rights Reserved.
 
 #include "PbrtHairTranslator.h"
-#include "HairStrandsDatas.h"
 
+#include "GroomImportOptions.h"
 #include "Misc/Paths.h"
 
 namespace PbrtHairFormat
@@ -11,7 +11,7 @@ namespace PbrtHairFormat
 	static const float TipRadius = 0.00005f; // m
 }
 
-bool FPbrtHairTranslator::Translate(const FString& FileName, FHairDescription& HairDescription)
+bool FPbrtHairTranslator::Translate(const FString& FileName, FHairDescription& HairDescription, const FGroomConversionSettings& ConversionSettings)
 {
 	bool bSuccess = false;
 
diff --git a/Engine/Plugins/Experimental/HairStrands/Source/HairStrandsEditor/Private/ReimportHairStrandsFactory.cpp b/Engine/Plugins/Experimental/HairStrands/Source/HairStrandsEditor/Private/ReimportHairStrandsFactory.cpp
--- a/Engine/Plugins/Experimental/HairStrands/Source/HairStrandsEditor/Private/ReimportHairStrandsFactory.cpp
+++ b/Engine/Plugins/Experimental/HairStrands/Source/HairStrandsEditor/Private/ReimportHairStrandsFactory.cpp
@@ -5,6 +5,9 @@
 #include "EditorFramework/AssetImportData.h"
 #include "HairDescription.h"
 #include "GroomAsset.h"
+#include "GroomAssetImportData.h"
+#include "GroomImportOptions.h"
+#include "GroomImportOptionsWindow.h"
 #include "HairStrandsImporter.h"
 #include "HairStrandsTranslator.h"
 #include "Misc/ScopedSlowTask.h"
@@ -71,6 +74,41 @@ EReimportResult::Type UReimportHairStrandsFactory::Reimport(UObject* Obj)
 
 		CurrentFilename = HairAsset->AssetImportData->GetFirstFilename();
 
+		UGroomAssetImportData* GroomAssetImportData = Cast<UGroomAssetImportData>(HairAsset->AssetImportData);
+		UGroomImportOptions* CurrentOptions = nullptr;
+		if (GroomAssetImportData)
+		{
+			// Duplicate the options to prevent dirtying the asset when they are modified but the re-import is cancelled
+			CurrentOptions = DuplicateObject<UGroomImportOptions>(GroomAssetImportData->ImportOptions, nullptr);
+		}
+		else
+		{
+			// Convert the AssetImportData to a GroomAssetImportData
+			GroomAssetImportData = NewObject<UGroomAssetImportData>(HairAsset);
+			GroomAssetImportData->Update(CurrentFilename);
+			HairAsset->AssetImportData = GroomAssetImportData;
+		}
+
+		if (!CurrentOptions)
+		{
+			// Make sure to have ImportOptions. Could happen if we just converted the AssetImportData
+			CurrentOptions = NewObject<UGroomImportOptions>();
+		}
+
+		if (!GIsRunningUnattendedScript && !IsAutomatedImport())
+		{
+			TSharedPtr<SGroomImportOptionsWindow> GroomOptionWindow = SGroomImportOptionsWindow::DisplayOptions(CurrentOptions, CurrentFilename);
+
+			if (!GroomOptionWindow->ShouldImport())
+			{
+				return EReimportResult::Cancelled;
+			}
+
+			// Move the transient ImportOptions to the asset package and set it on the GroomAssetImportData for serialization
+			CurrentOptions->Rename(nullptr, GroomAssetImportData);
+			GroomAssetImportData->ImportOptions = CurrentOptions;
+		}
+
 		TSharedPtr<IHairStrandsTranslator> SelectedTranslator = GetTranslator(CurrentFilename);
 		if (!SelectedTranslator.IsValid())
 		{
@@ -81,12 +119,12 @@ EReimportResult::Type UReimportHairStrandsFactory::Reimport(UObject* Obj)
 		Progress.MakeDialog(true);
 
 		FHairDescription HairDescription;
-		if (!SelectedTranslator->Translate(CurrentFilename, HairDescription))
+		if (!SelectedTranslator->Translate(CurrentFilename, HairDescription, CurrentOptions->ConversionSettings))
 		{
 			return EReimportResult::Failed;
 		}
 
-		UGroomAsset* ReimportedHair = FHairStrandsImporter::ImportHair(FHairImportContext(), HairDescription, HairAsset);
+		UGroomAsset* ReimportedHair = FHairStrandsImporter::ImportHair(FHairImportContext(CurrentOptions), HairDescription, HairAsset);
 		if (!ReimportedHair)
 		{
 			return EReimportResult::Failed;
diff --git a/Engine/Plugins/Experimental/HairStrands/Source/HairStrandsEditor/Public/GroomImportOptions.h b/Engine/Plugins/Experimental/HairStrands/Source/HairStrandsEditor/Public/GroomImportOptions.h
new file mode 100644
--- /dev/null
+++ b/Engine/Plugins/Experimental/HairStrands/Source/HairStrandsEditor/Public/GroomImportOptions.h
@@ -0,0 +1,21 @@
+// Copyright 1998-2019 Epic Games, Inc. All Rights Reserved.
+
+#pragma once
+
+#include "CoreMinimal.h"
+#include "GroomSettings.h"
+#include "UObject/Object.h"
+#include "GroomImportOptions.generated.h"
+
+UCLASS(BlueprintType, config = EditorPerProjectUserSettings)
+class HAIRSTRANDSEDITOR_API UGroomImportOptions : public UObject
+{
+	GENERATED_UCLASS_BODY()
+
+public:
+	UPROPERTY(config, EditAnywhere, BlueprintReadWrite, meta = (ShowOnlyInnerProperties), Category = Conversion)
+	FGroomConversionSettings ConversionSettings;
+
+	UPROPERTY(config, EditAnywhere, BlueprintReadWrite, meta = (ShowOnlyInnerProperties), Category = BuildSettings)
+	FGroomBuildSettings BuildSettings;
+};
diff --git a/Engine/Plugins/Experimental/HairStrands/Source/HairStrandsEditor/Public/HairStrandsFactory.h b/Engine/Plugins/Experimental/HairStrands/Source/HairStrandsEditor/Public/HairStrandsFactory.h
--- a/Engine/Plugins/Experimental/HairStrands/Source/HairStrandsEditor/Public/HairStrandsFactory.h
+++ b/Engine/Plugins/Experimental/HairStrands/Source/HairStrandsEditor/Public/HairStrandsFactory.h
@@ -30,6 +30,8 @@ protected:
 
 	TSharedPtr<IHairStrandsTranslator> GetTranslator(const FString& Filename);
 
+	class UGroomImportOptions* ImportOptions;
+
 private:
 	TArray<TSharedPtr<IHairStrandsTranslator>> Translators;
 };
diff --git a/Engine/Plugins/Experimental/HairStrands/Source/HairStrandsEditor/Public/HairStrandsImporter.h b/Engine/Plugins/Experimental/HairStrands/Source/HairStrandsEditor/Public/HairStrandsImporter.h
--- a/Engine/Plugins/Experimental/HairStrands/Source/HairStrandsEditor/Public/HairStrandsImporter.h
+++ b/Engine/Plugins/Experimental/HairStrands/Source/HairStrandsEditor/Public/HairStrandsImporter.h
@@ -7,11 +7,13 @@
 
 class FHairDescription;
 class UGroomAsset;
+class UGroomImportOptions;
 
 struct HAIRSTRANDSEDITOR_API FHairImportContext
 {
-	FHairImportContext(UObject* InParent = nullptr, UClass* InClass = nullptr, FName InName = FName(), EObjectFlags InFlags = EObjectFlags::RF_NoFlags);
+	FHairImportContext(UGroomImportOptions* InImportOptions, UObject* InParent = nullptr, UClass* InClass = nullptr, FName InName = FName(), EObjectFlags InFlags = EObjectFlags::RF_NoFlags);
 
+	UGroomImportOptions* ImportOptions;
 	UObject* Parent;
 	UClass* Class;
 	FName Name;
diff --git a/Engine/Plugins/Experimental/HairStrands/Source/HairStrandsEditor/Public/HairStrandsTranslator.h b/Engine/Plugins/Experimental/HairStrands/Source/HairStrandsEditor/Public/HairStrandsTranslator.h
--- a/Engine/Plugins/Experimental/HairStrands/Source/HairStrandsEditor/Public/HairStrandsTranslator.h
+++ b/Engine/Plugins/Experimental/HairStrands/Source/HairStrandsEditor/Public/HairStrandsTranslator.h
@@ -10,7 +10,7 @@ public:
 	virtual ~IHairStrandsTranslator() {}
 
 	/** Translate a given file into a HairDescription; return true if successful */
-	virtual bool Translate(const FString& FilePath, FHairDescription& OutHairDescription) = 0;
+	virtual bool Translate(const FString& FilePath, FHairDescription& OutHairDescription, const struct FGroomConversionSettings& ConversionSettings) = 0;
 
 	/** Return true if a given file can be translated (by checking its content if necessary) */
 	virtual bool CanTranslate(const FString& FilePath) = 0;
diff --git a/Engine/Plugins/Experimental/HairStrands/Source/HairStrandsNiagara/Private/HairStrandsNiagara.cpp b/Engine/Plugins/Experimental/HairStrands/Source/HairStrandsNiagara/Private/HairStrandsNiagara.cpp
--- a/Engine/Plugins/Experimental/HairStrands/Source/HairStrandsNiagara/Private/HairStrandsNiagara.cpp
+++ b/Engine/Plugins/Experimental/HairStrands/Source/HairStrandsNiagara/Private/HairStrandsNiagara.cpp
@@ -8,6 +8,9 @@ IMPLEMENT_MODULE(FHairStrandsNiagara, HairStrandsNiagara);
 
 void FHairStrandsNiagara::StartupModule()
 {
+	static const auto UseShaderStagesCVar = IConsoleManager::Get().FindConsoleVariable(TEXT("fx.UseShaderStages"));
+	UseShaderStagesCVar->Set(1);
+
 	// Maps virtual shader source directory /Plugin/FX/Niagara to the plugin's actual Shaders directory.
 	FString PluginShaderDir = FPaths::Combine(IPluginManager::Get().FindPlugin(TEXT("HairStrands"))->GetBaseDir(), TEXT("Shaders"));
 	AddShaderSourceDirectoryMapping(TEXT("/Plugin/Experimental/HairStrands"), PluginShaderDir);
