Copying //UE4/Dev-Physics to Dev-Main (//UE4/Dev-Main) [at] 9420715

diff --git a/Engine/Plugins/Experimental/HairStrands/Source/HairStrandsNiagara/Private/NiagaraDataInterfaceHairStrands.cpp b/Engine/Plugins/Experimental/HairStrands/Source/HairStrandsNiagara/Private/NiagaraDataInterfaceHairStrands.cpp
--- a/Engine/Plugins/Experimental/HairStrands/Source/HairStrandsNiagara/Private/NiagaraDataInterfaceHairStrands.cpp
+++ b/Engine/Plugins/Experimental/HairStrands/Source/HairStrandsNiagara/Private/NiagaraDataInterfaceHairStrands.cpp
@@ -62,6 +62,44 @@ static const FName ProjectCollisionGridName(TEXT("ProjectCollisionGrid"));
 
 //------------------------------------------------------------------------------------------------------------
 
+static const FName GetBoxCenterName(TEXT("GetBoxCenter"));
+static const FName GetBoxExtentName(TEXT("GetBoxExtent"));
+
+//------------------------------------------------------------------------------------------------------------
+
+static const FName SetupStretchSpringMaterialName(TEXT("SetupStretchSpringMaterial"));
+static const FName SolveStretchSpringMaterialName(TEXT("SolveStretchSpringMaterial"));
+static const FName ProjectStretchSpringMaterialName(TEXT("ProjectStretchSpringMaterial"));
+
+//------------------------------------------------------------------------------------------------------------
+
+static const FName SetupBendSpringMaterialName(TEXT("SetupBendSpringMaterial"));
+static const FName SolveBendSpringMaterialName(TEXT("SolveBendSpringMaterial"));
+static const FName ProjectBendSpringMaterialName(TEXT("ProjectBendSpringMaterial"));
+
+//------------------------------------------------------------------------------------------------------------
+
+static const FName SetupStretchRodMaterialName(TEXT("SetupStretchRodMaterial"));
+static const FName SolveStretchRodMaterialName(TEXT("SolveStretchRodMaterial"));
+static const FName ProjectStretchRodMaterialName(TEXT("ProjectStretchRodMaterial"));
+
+//------------------------------------------------------------------------------------------------------------
+
+static const FName SetupBendRodMaterialName(TEXT("SetupBendRodMaterial"));
+static const FName SolveBendRodMaterialName(TEXT("SolveBendRodMaterial"));
+static const FName ProjectBendRodMaterialName(TEXT("ProjectBendRodMaterial"));
+//------------------------------------------------------------------------------------------------------------
+
+static const FName SolveStaticCollisionConstraintName(TEXT("SolveStaticCollisionConstraint"));
+static const FName ProjectStaticCollisionConstraintName(TEXT("ProjectStaticCollisionConstraint"));
+
+//------------------------------------------------------------------------------------------------------------
+
+static const FName ComputeRestDirectionName(TEXT("ComputeRestDirection"));
+static const FName UpdateNodeOrientationName(TEXT("UpdateNodeOrientation"));
+
+//------------------------------------------------------------------------------------------------------------
+
 const FString UNiagaraDataInterfaceHairStrands::NumStrandsName(TEXT("NumStrands_"));
 const FString UNiagaraDataInterfaceHairStrands::StrandSizeName(TEXT("StrandSize_"));
 const FString UNiagaraDataInterfaceHairStrands::StrandDensityName(TEXT("StrandDensity_"));
@@ -129,18 +167,18 @@ struct FNDIHairStrandsParametersName
 
 //------------------------------------------------------------------------------------------------------------
 
-void FNDICollisionGridBuffer::SetGridSize(const FGridSize InGridSize)
+void FNDICollisionGridBuffer::SetGridSize(const FUintVector4 InGridSize)
 {
 	GridSize = InGridSize;
 }
 
 void FNDICollisionGridBuffer::InitRHI()
 {
-	if(GridSize.GridSizeX != 0 && GridSize.GridSizeY != 0 && GridSize.GridSizeZ != 0)
+	if(GridSize.X != 0 && GridSize.Y != 0 && GridSize.Z != 0)
 	{
 		static const uint32 NumComponents = 9;
-		GridDataBuffer.Initialize(sizeof(int32), GridSize.GridSizeX*NumComponents, GridSize.GridSizeY,
-			GridSize.GridSizeZ, EPixelFormat::PF_R32_SINT);
+		GridDataBuffer.Initialize(sizeof(int32), GridSize.X*NumComponents, GridSize.Y,
+			GridSize.Z, EPixelFormat::PF_R32_SINT);
 	}
 }
 
@@ -313,13 +351,11 @@ struct FNDIHairStrandsParametersCS : public FNiagaraDataInterfaceParametersCS
 			SetSRVParameter(RHICmdList, ComputeShaderRHI, CurvesOffsetsBuffer, FNiagaraRenderer::GetDummyUIntBuffer().SRV);
 			SetSRVParameter(RHICmdList, ComputeShaderRHI, RestPositionsBuffer, FNiagaraRenderer::GetDummyFloatBuffer().SRV);
 
-		
-
 			SetUAVParameter(RHICmdList, ComputeShaderRHI, GridDestinationBuffer, FNiagaraRenderer::GetDummyUIntBuffer().UAV);
 			SetSRVParameter(RHICmdList, ComputeShaderRHI, GridCurrentBuffer, FNiagaraRenderer::GetDummyUIntBuffer().SRV);
 
 			SetShaderValue(RHICmdList, ComputeShaderRHI, GridOrigin, FVector4(0,0,0,0));
-			SetShaderValue(RHICmdList, ComputeShaderRHI, GridSize, FGridSize());
+			SetShaderValue(RHICmdList, ComputeShaderRHI, GridSize, FUintVector4());
 
 			SetShaderValue(RHICmdList, ComputeShaderRHI, WorldTransform, FMatrix::Identity);
 			SetShaderValue(RHICmdList, ComputeShaderRHI, WorldInverse, FMatrix::Identity);
@@ -389,7 +425,7 @@ void FNDIHairStrandsProxy::ConsumePerInstanceDataFromGameThread(void* PerInstanc
 }
 
 void FNDIHairStrandsProxy::InitializePerInstanceData(const FNiagaraSystemInstanceID& SystemInstance, FNDIHairStrandsBuffer* HairStrandsBuffer, FNDICollisionGridBuffer* CurrentGridBuffer, FNDICollisionGridBuffer* DestinationGridBuffer, uint32 NumStrands, const uint8 StrandSize,
-	const float StrandDensity, const float RootThickness, const float TipThickness, const FVector4& GridOrigin, const FGridSize& GridSize)
+	const float StrandDensity, const float RootThickness, const float TipThickness, const FVector4& GridOrigin, const FUintVector4& GridSize)
 {
 	check(IsInRenderingThread());
 
@@ -441,13 +477,6 @@ UNiagaraDataInterfaceHairStrands::UNiagaraDataInterfaceHairStrands(FObjectInitia
 {
 
 	Proxy = MakeShared<FNDIHairStrandsProxy, ESPMode::ThreadSafe>();
-
-	/*static int32 RefCounter = 0;
-	if (RefCounter == 21)
-	{
-		UE_LOG(LogHairStrands, Log, TEXT("Bad HairStrands : %d %d"), this, RefCounter);
-	}
-	++RefCounter;*/
 }
 
 bool UNiagaraDataInterfaceHairStrands::IsComponentValid() const
@@ -500,8 +529,8 @@ bool UNiagaraDataInterfaceHairStrands::InitPerInstanceData(void* PerInstanceData
 		}
 	}
 	GroomAsset = IsComponentValid() ? SourceComponent->GroomAsset : DefaultSource;
-	FHairStrandsDatas* StrandsDatas = IsComponentValid() ? SourceComponent->GetGuideStrandsDatas() : (DefaultSource != nullptr) ? &DefaultSource->HairRenderData : nullptr;
-	FHairStrandsResource* StrandsResource = IsComponentValid() ? SourceComponent->GetGuideStrandsResource() : (DefaultSource != nullptr) ? DefaultSource->HairStrandsResource : nullptr;
+	FHairStrandsDatas* StrandsDatas = IsComponentValid() ? SourceComponent->GetGuideStrandsDatas() : (DefaultSource != nullptr && DefaultSource->GetNumHairGroups() > 0) ? &DefaultSource->HairGroupsData[0].HairRenderData : nullptr;
+	FHairStrandsResource* StrandsResource = IsComponentValid() ? SourceComponent->GetGuideStrandsResource() : (DefaultSource != nullptr && DefaultSource->GetNumHairGroups() > 0) ? DefaultSource->HairGroupsData[0].HairStrandsResource : nullptr;
 
 	if (StrandsDatas == nullptr || StrandsResource == nullptr)
 	{
@@ -528,7 +557,7 @@ bool UNiagaraDataInterfaceHairStrands::InitPerInstanceData(void* PerInstanceData
 									  (LocalStrandsBox.Max.Z - LocalStrandsBox.Min.Z) / (GridSizeZ-1));
 			const float GridLength = GridLengths.GetMax();
 
-			const FGridSize GridSize(GridSizeX, GridSizeY, GridSizeZ);
+			const FUintVector4 GridSize(GridSizeX, GridSizeY, GridSizeZ,0);
 			const FVector4 GridOrigin(-GridExtent.X, -GridExtent.Y, -GridExtent.Z, GridLength);
 			CurrentGridBuffer->SetGridSize(GridSize);
 			DestinationGridBuffer->SetGridSize(GridSize);
@@ -566,7 +595,7 @@ bool UNiagaraDataInterfaceHairStrands::InitPerInstanceData(void* PerInstanceData
 			{
 				ThisProxy->OutputShaderStages = RT_OutputShaderStages;
 				ThisProxy->IterationShaderStages = RT_IterationShaderStages;
-				ThisProxy->SetElementCount(GridSize.GridSizeX*GridSize.GridSizeY*GridSize.GridSizeZ);
+				ThisProxy->SetElementCount(GridSize.X*GridSize.Y*GridSize.Z);
 
 				HairStrandsBuffer->InitResource();
 				CurrentGridBuffer->InitResource();
@@ -614,7 +643,7 @@ bool UNiagaraDataInterfaceHairStrands::PerInstanceTick(void* PerInstanceData, FN
 	FNDIHairStrandsData* InstanceData = static_cast<FNDIHairStrandsData*>(PerInstanceData);
 
 	FHairStrandsDatas* StrandsDatas = IsComponentValid() ? 
-		SourceComponent->GetGuideStrandsDatas() : (DefaultSource != nullptr) ? &DefaultSource->HairRenderData : nullptr;
+		SourceComponent->GetGuideStrandsDatas() : (DefaultSource != nullptr && DefaultSource->GetNumHairGroups() > 0) ? &DefaultSource->HairGroupsData[0].HairRenderData : nullptr;
 	bool RequireReset = false;
 	if (StrandsDatas)
 	{
@@ -624,10 +653,12 @@ bool UNiagaraDataInterfaceHairStrands::PerInstanceTick(void* PerInstanceData, FN
 		if (IsComponentValid())
 		{
 			InstanceData->WorldTransform = WorldTransform * SourceComponent->GetComponentToWorld().ToMatrixWithScale();
+			//UE_LOG(LogHairStrands, Warning, TEXT("Get Component Transform : %s"), *InstanceData->WorldTransform.ToString());
 		}
 		else
 		{
 			InstanceData->WorldTransform = WorldTransform * SystemInstance->GetComponent()->GetComponentToWorld().ToMatrixWithScale();
+			//UE_LOG(LogHairStrands, Warning, TEXT("Get Instance Transform : %s"), *InstanceData->WorldTransform.ToString());
 		}
 	}
 	return RequireReset;
@@ -794,6 +825,9 @@ void UNiagaraDataInterfaceHairStrands::GetFunctions(TArray<FNiagaraFunctionSigna
 		Sig.bRequiresContext = false;
 		Sig.Inputs.Add(FNiagaraVariable(FNiagaraTypeDefinition(GetClass()), TEXT("Hair Strands")));
 		Sig.Inputs.Add(FNiagaraVariable(FNiagaraTypeDefinition::GetIntDef(), TEXT("Node Index")));
+		Sig.Inputs.Add(FNiagaraVariable(FNiagaraTypeDefinition::GetFloatDef(), TEXT("Strands Density")));
+		Sig.Inputs.Add(FNiagaraVariable(FNiagaraTypeDefinition::GetFloatDef(), TEXT("Root Thickness")));
+		Sig.Inputs.Add(FNiagaraVariable(FNiagaraTypeDefinition::GetFloatDef(), TEXT("Tip Thickness")));
 		Sig.Outputs.Add(FNiagaraVariable(FNiagaraTypeDefinition::GetFloatDef(), TEXT("Node Mass")));
 
 		OutFunctions.Add(Sig);
@@ -805,6 +839,9 @@ void UNiagaraDataInterfaceHairStrands::GetFunctions(TArray<FNiagaraFunctionSigna
 		Sig.bRequiresContext = false;
 		Sig.Inputs.Add(FNiagaraVariable(FNiagaraTypeDefinition(GetClass()), TEXT("Hair Strands")));
 		Sig.Inputs.Add(FNiagaraVariable(FNiagaraTypeDefinition::GetIntDef(), TEXT("Node Index")));
+		Sig.Inputs.Add(FNiagaraVariable(FNiagaraTypeDefinition::GetFloatDef(), TEXT("Strands Density")));
+		Sig.Inputs.Add(FNiagaraVariable(FNiagaraTypeDefinition::GetFloatDef(), TEXT("Root Thickness")));
+		Sig.Inputs.Add(FNiagaraVariable(FNiagaraTypeDefinition::GetFloatDef(), TEXT("Tip Thickness")));
 		Sig.Outputs.Add(FNiagaraVariable(FNiagaraTypeDefinition::GetVec3Def(), TEXT("Node Inertia")));
 
 		OutFunctions.Add(Sig);
@@ -1001,6 +1038,271 @@ void UNiagaraDataInterfaceHairStrands::GetFunctions(TArray<FNiagaraFunctionSigna
 		Sig.Outputs.Add(FNiagaraVariable(FNiagaraTypeDefinition::GetVec3Def(), TEXT("Grid GradientY")));
 		Sig.Outputs.Add(FNiagaraVariable(FNiagaraTypeDefinition::GetVec3Def(), TEXT("Grid GradientZ")));
 
+		OutFunctions.Add(Sig);
+	}
+	{
+		FNiagaraFunctionSignature Sig;
+		Sig.Name = GetBoxCenterName;
+		Sig.bMemberFunction = true;
+		Sig.bRequiresContext = false;
+		Sig.Inputs.Add(FNiagaraVariable(FNiagaraTypeDefinition(GetClass()), TEXT("Hair Strands")));
+		Sig.Outputs.Add(FNiagaraVariable(FNiagaraTypeDefinition::GetVec3Def(), TEXT("Grid Center")));
+
+		OutFunctions.Add(Sig);
+	}
+	{
+		FNiagaraFunctionSignature Sig;
+		Sig.Name = GetBoxExtentName;
+		Sig.bMemberFunction = true;
+		Sig.bRequiresContext = false;
+		Sig.Inputs.Add(FNiagaraVariable(FNiagaraTypeDefinition(GetClass()), TEXT("Hair Strands")));
+		Sig.Outputs.Add(FNiagaraVariable(FNiagaraTypeDefinition::GetVec3Def(), TEXT("Grid Extent")));
+
+		OutFunctions.Add(Sig);
+	}
+	{
+		FNiagaraFunctionSignature Sig;
+		Sig.Name = SetupStretchSpringMaterialName;
+		Sig.bMemberFunction = true;
+		Sig.bRequiresContext = false;
+		Sig.Inputs.Add(FNiagaraVariable(FNiagaraTypeDefinition(GetClass()), TEXT("Hair Strands")));
+		Sig.Inputs.Add(FNiagaraVariable(FNiagaraTypeDefinition::GetFloatDef(), TEXT("Stretch Stiffness")));
+		Sig.Inputs.Add(FNiagaraVariable(FNiagaraTypeDefinition::GetFloatDef(), TEXT("Node Thickness")));
+		Sig.Inputs.Add(FNiagaraVariable(FNiagaraTypeDefinition::GetFloatDef(), TEXT("Rest Length")));
+		Sig.Inputs.Add(FNiagaraVariable(FNiagaraTypeDefinition::GetFloatDef(), TEXT("Delta Time")));
+		Sig.Inputs.Add(FNiagaraVariable(FNiagaraTypeDefinition::GetFloatDef(), TEXT("Material Damping")));
+		Sig.Outputs.Add(FNiagaraVariable(FNiagaraTypeDefinition::GetFloatDef(), TEXT("Material Compliance")));
+		Sig.Outputs.Add(FNiagaraVariable(FNiagaraTypeDefinition::GetFloatDef(), TEXT("Material Weight")));
+		Sig.Outputs.Add(FNiagaraVariable(FNiagaraTypeDefinition::GetFloatDef(), TEXT("Material Multiplier")));
+		OutFunctions.Add(Sig);
+	}
+	{
+		FNiagaraFunctionSignature Sig;
+		Sig.Name = SolveStretchSpringMaterialName;
+		Sig.bMemberFunction = true;
+		Sig.bRequiresContext = false;
+		Sig.Inputs.Add(FNiagaraVariable(FNiagaraTypeDefinition(GetClass()), TEXT("Hair Strands")));
+		Sig.Inputs.Add(FNiagaraVariable(FNiagaraTypeDefinition::GetFloatDef(), TEXT("Rest Length")));
+		Sig.Inputs.Add(FNiagaraVariable(FNiagaraTypeDefinition::GetFloatDef(), TEXT("Delta Time")));
+		Sig.Inputs.Add(FNiagaraVariable(FNiagaraTypeDefinition::GetFloatDef(), TEXT("Material Damping")));
+		Sig.Inputs.Add(FNiagaraVariable(FNiagaraTypeDefinition::GetFloatDef(), TEXT("Material Compliance")));
+		Sig.Inputs.Add(FNiagaraVariable(FNiagaraTypeDefinition::GetFloatDef(), TEXT("Material Weight")));
+		Sig.Inputs.Add(FNiagaraVariable(FNiagaraTypeDefinition::GetFloatDef(), TEXT("Material Multiplier")));
+		Sig.Outputs.Add(FNiagaraVariable(FNiagaraTypeDefinition::GetFloatDef(), TEXT("Material Multiplier")));
+		OutFunctions.Add(Sig);
+	}
+	{
+		FNiagaraFunctionSignature Sig;
+		Sig.Name = ProjectStretchSpringMaterialName;
+		Sig.bMemberFunction = true;
+		Sig.bRequiresContext = false;
+		Sig.Inputs.Add(FNiagaraVariable(FNiagaraTypeDefinition(GetClass()), TEXT("Hair Strands")));
+		Sig.Inputs.Add(FNiagaraVariable(FNiagaraTypeDefinition::GetFloatDef(), TEXT("Stretch Stiffness")));
+		Sig.Inputs.Add(FNiagaraVariable(FNiagaraTypeDefinition::GetFloatDef(), TEXT("Node Thickness")));
+		Sig.Inputs.Add(FNiagaraVariable(FNiagaraTypeDefinition::GetFloatDef(), TEXT("Rest Length")));
+		Sig.Inputs.Add(FNiagaraVariable(FNiagaraTypeDefinition::GetFloatDef(), TEXT("Delta Time")));
+		Sig.Outputs.Add(FNiagaraVariable(FNiagaraTypeDefinition::GetFloatDef(), TEXT("Material Multiplier")));
+
+		OutFunctions.Add(Sig);
+	}
+
+	{
+		FNiagaraFunctionSignature Sig;
+		Sig.Name = SetupBendSpringMaterialName;
+		Sig.bMemberFunction = true;
+		Sig.bRequiresContext = false;
+		Sig.Inputs.Add(FNiagaraVariable(FNiagaraTypeDefinition(GetClass()), TEXT("Hair Strands")));
+		Sig.Inputs.Add(FNiagaraVariable(FNiagaraTypeDefinition::GetFloatDef(), TEXT("Bend Stiffness")));
+		Sig.Inputs.Add(FNiagaraVariable(FNiagaraTypeDefinition::GetFloatDef(), TEXT("Node Thickness")));
+		Sig.Inputs.Add(FNiagaraVariable(FNiagaraTypeDefinition::GetFloatDef(), TEXT("Rest Length")));
+		Sig.Inputs.Add(FNiagaraVariable(FNiagaraTypeDefinition::GetFloatDef(), TEXT("Delta Time")));
+		Sig.Inputs.Add(FNiagaraVariable(FNiagaraTypeDefinition::GetFloatDef(), TEXT("Material Damping")));
+		Sig.Outputs.Add(FNiagaraVariable(FNiagaraTypeDefinition::GetFloatDef(), TEXT("Material Compliance")));
+		Sig.Outputs.Add(FNiagaraVariable(FNiagaraTypeDefinition::GetFloatDef(), TEXT("Material Weight")));
+		Sig.Outputs.Add(FNiagaraVariable(FNiagaraTypeDefinition::GetVec3Def(), TEXT("Material Multiplier")));
+		OutFunctions.Add(Sig);
+	}
+	{
+		FNiagaraFunctionSignature Sig;
+		Sig.Name = SolveBendSpringMaterialName;
+		Sig.bMemberFunction = true;
+		Sig.bRequiresContext = false;
+		Sig.Inputs.Add(FNiagaraVariable(FNiagaraTypeDefinition(GetClass()), TEXT("Hair Strands")));
+		Sig.Inputs.Add(FNiagaraVariable(FNiagaraTypeDefinition::GetVec3Def(), TEXT("Rest Direction")));
+		Sig.Inputs.Add(FNiagaraVariable(FNiagaraTypeDefinition::GetFloatDef(), TEXT("Delta Time")));
+		Sig.Inputs.Add(FNiagaraVariable(FNiagaraTypeDefinition::GetFloatDef(), TEXT("Material Damping")));
+		Sig.Inputs.Add(FNiagaraVariable(FNiagaraTypeDefinition::GetFloatDef(), TEXT("Material Compliance")));
+		Sig.Inputs.Add(FNiagaraVariable(FNiagaraTypeDefinition::GetFloatDef(), TEXT("Material Weight")));
+		Sig.Inputs.Add(FNiagaraVariable(FNiagaraTypeDefinition::GetVec3Def(), TEXT("Material Multiplier")));
+		Sig.Outputs.Add(FNiagaraVariable(FNiagaraTypeDefinition::GetVec3Def(), TEXT("Material Multiplier")));
+		OutFunctions.Add(Sig);
+	}
+	{
+		FNiagaraFunctionSignature Sig;
+		Sig.Name = ProjectBendSpringMaterialName;
+		Sig.bMemberFunction = true;
+		Sig.bRequiresContext = false;
+		Sig.Inputs.Add(FNiagaraVariable(FNiagaraTypeDefinition(GetClass()), TEXT("Hair Strands")));
+		Sig.Inputs.Add(FNiagaraVariable(FNiagaraTypeDefinition::GetFloatDef(), TEXT("Bend Stiffness")));
+		Sig.Inputs.Add(FNiagaraVariable(FNiagaraTypeDefinition::GetFloatDef(), TEXT("Node Thickness")));
+		Sig.Inputs.Add(FNiagaraVariable(FNiagaraTypeDefinition::GetFloatDef(), TEXT("Rest Length")));
+		Sig.Inputs.Add(FNiagaraVariable(FNiagaraTypeDefinition::GetVec3Def(), TEXT("Rest Direction")));
+		Sig.Inputs.Add(FNiagaraVariable(FNiagaraTypeDefinition::GetFloatDef(), TEXT("Delta Time")));
+		Sig.Outputs.Add(FNiagaraVariable(FNiagaraTypeDefinition::GetVec3Def(), TEXT("Material Multiplier")));
+
+		OutFunctions.Add(Sig);
+	}
+
+	{
+		FNiagaraFunctionSignature Sig;
+		Sig.Name = SetupStretchRodMaterialName;
+		Sig.bMemberFunction = true;
+		Sig.bRequiresContext = false;
+		Sig.Inputs.Add(FNiagaraVariable(FNiagaraTypeDefinition(GetClass()), TEXT("Hair Strands")));
+		Sig.Inputs.Add(FNiagaraVariable(FNiagaraTypeDefinition::GetFloatDef(), TEXT("Stretch Stiffness")));
+		Sig.Inputs.Add(FNiagaraVariable(FNiagaraTypeDefinition::GetFloatDef(), TEXT("Node Thickness")));
+		Sig.Inputs.Add(FNiagaraVariable(FNiagaraTypeDefinition::GetFloatDef(), TEXT("Rest Length")));
+		Sig.Inputs.Add(FNiagaraVariable(FNiagaraTypeDefinition::GetFloatDef(), TEXT("Delta Time")));
+		Sig.Inputs.Add(FNiagaraVariable(FNiagaraTypeDefinition::GetFloatDef(), TEXT("Material Damping")));
+		Sig.Outputs.Add(FNiagaraVariable(FNiagaraTypeDefinition::GetFloatDef(), TEXT("Material Compliance")));
+		Sig.Outputs.Add(FNiagaraVariable(FNiagaraTypeDefinition::GetFloatDef(), TEXT("Material Weight")));
+		Sig.Outputs.Add(FNiagaraVariable(FNiagaraTypeDefinition::GetVec3Def(), TEXT("Material Multiplier")));
+		OutFunctions.Add(Sig);
+	}
+	{
+		FNiagaraFunctionSignature Sig;
+		Sig.Name = SolveStretchRodMaterialName;
+		Sig.bMemberFunction = true;
+		Sig.bRequiresContext = false;
+		Sig.Inputs.Add(FNiagaraVariable(FNiagaraTypeDefinition(GetClass()), TEXT("Hair Strands")));
+		Sig.Inputs.Add(FNiagaraVariable(FNiagaraTypeDefinition::GetFloatDef(), TEXT("Rest Length")));
+		Sig.Inputs.Add(FNiagaraVariable(FNiagaraTypeDefinition::GetFloatDef(), TEXT("Delta Time")));
+		Sig.Inputs.Add(FNiagaraVariable(FNiagaraTypeDefinition::GetFloatDef(), TEXT("Material Damping")));
+		Sig.Inputs.Add(FNiagaraVariable(FNiagaraTypeDefinition::GetFloatDef(), TEXT("Material Compliance")));
+		Sig.Inputs.Add(FNiagaraVariable(FNiagaraTypeDefinition::GetFloatDef(), TEXT("Material Weight")));
+		Sig.Inputs.Add(FNiagaraVariable(FNiagaraTypeDefinition::GetVec3Def(), TEXT("Material Multiplier")));
+		Sig.Outputs.Add(FNiagaraVariable(FNiagaraTypeDefinition::GetVec3Def(), TEXT("Material Multiplier")));
+		OutFunctions.Add(Sig);
+	}
+	{
+		FNiagaraFunctionSignature Sig;
+		Sig.Name = ProjectStretchRodMaterialName;
+		Sig.bMemberFunction = true;
+		Sig.bRequiresContext = false;
+		Sig.Inputs.Add(FNiagaraVariable(FNiagaraTypeDefinition(GetClass()), TEXT("Hair Strands")));
+		Sig.Inputs.Add(FNiagaraVariable(FNiagaraTypeDefinition::GetFloatDef(), TEXT("Stretch Stiffness")));
+		Sig.Inputs.Add(FNiagaraVariable(FNiagaraTypeDefinition::GetFloatDef(), TEXT("Node Thickness")));
+		Sig.Inputs.Add(FNiagaraVariable(FNiagaraTypeDefinition::GetFloatDef(), TEXT("Rest Length")));
+		Sig.Inputs.Add(FNiagaraVariable(FNiagaraTypeDefinition::GetFloatDef(), TEXT("Delta Time")));
+		Sig.Outputs.Add(FNiagaraVariable(FNiagaraTypeDefinition::GetVec3Def(), TEXT("Material Multiplier")));
+
+		OutFunctions.Add(Sig);
+	}
+
+	{
+		FNiagaraFunctionSignature Sig;
+		Sig.Name = SetupBendRodMaterialName;
+		Sig.bMemberFunction = true;
+		Sig.bRequiresContext = false;
+		Sig.Inputs.Add(FNiagaraVariable(FNiagaraTypeDefinition(GetClass()), TEXT("Hair Strands")));
+		Sig.Inputs.Add(FNiagaraVariable(FNiagaraTypeDefinition::GetFloatDef(), TEXT("Bend Stiffness")));
+		Sig.Inputs.Add(FNiagaraVariable(FNiagaraTypeDefinition::GetFloatDef(), TEXT("Node Thickness")));
+		Sig.Inputs.Add(FNiagaraVariable(FNiagaraTypeDefinition::GetFloatDef(), TEXT("Rest Length")));
+		Sig.Inputs.Add(FNiagaraVariable(FNiagaraTypeDefinition::GetFloatDef(), TEXT("Delta Time")));
+		Sig.Inputs.Add(FNiagaraVariable(FNiagaraTypeDefinition::GetFloatDef(), TEXT("Material Damping")));
+		Sig.Outputs.Add(FNiagaraVariable(FNiagaraTypeDefinition::GetFloatDef(), TEXT("Material Compliance")));
+		Sig.Outputs.Add(FNiagaraVariable(FNiagaraTypeDefinition::GetFloatDef(), TEXT("Material Weight")));
+		Sig.Outputs.Add(FNiagaraVariable(FNiagaraTypeDefinition::GetVec3Def(), TEXT("Material Multiplier")));
+		OutFunctions.Add(Sig);
+	}
+	{
+		FNiagaraFunctionSignature Sig;
+		Sig.Name = SolveBendRodMaterialName;
+		Sig.bMemberFunction = true;
+		Sig.bRequiresContext = false;
+		Sig.Inputs.Add(FNiagaraVariable(FNiagaraTypeDefinition(GetClass()), TEXT("Hair Strands")));
+		Sig.Inputs.Add(FNiagaraVariable(FNiagaraTypeDefinition::GetFloatDef(), TEXT("Rest Length")));
+		Sig.Inputs.Add(FNiagaraVariable(FNiagaraTypeDefinition::GetQuatDef(), TEXT("Rest Darboux")));
+		Sig.Inputs.Add(FNiagaraVariable(FNiagaraTypeDefinition::GetFloatDef(), TEXT("Delta Time")));
+		Sig.Inputs.Add(FNiagaraVariable(FNiagaraTypeDefinition::GetFloatDef(), TEXT("Material Damping")));
+		Sig.Inputs.Add(FNiagaraVariable(FNiagaraTypeDefinition::GetFloatDef(), TEXT("Material Compliance")));
+		Sig.Inputs.Add(FNiagaraVariable(FNiagaraTypeDefinition::GetFloatDef(), TEXT("Material Weight")));
+		Sig.Inputs.Add(FNiagaraVariable(FNiagaraTypeDefinition::GetVec3Def(), TEXT("Material Multiplier")));
+		Sig.Outputs.Add(FNiagaraVariable(FNiagaraTypeDefinition::GetVec3Def(), TEXT("Material Multiplier")));
+		OutFunctions.Add(Sig);
+	}
+	{
+		FNiagaraFunctionSignature Sig;
+		Sig.Name = ProjectBendRodMaterialName;
+		Sig.bMemberFunction = true;
+		Sig.bRequiresContext = false;
+		Sig.Inputs.Add(FNiagaraVariable(FNiagaraTypeDefinition(GetClass()), TEXT("Hair Strands")));
+		Sig.Inputs.Add(FNiagaraVariable(FNiagaraTypeDefinition::GetFloatDef(), TEXT("Stretch Stiffness")));
+		Sig.Inputs.Add(FNiagaraVariable(FNiagaraTypeDefinition::GetFloatDef(), TEXT("Node Thickness")));
+		Sig.Inputs.Add(FNiagaraVariable(FNiagaraTypeDefinition::GetFloatDef(), TEXT("Rest Length")));
+		Sig.Inputs.Add(FNiagaraVariable(FNiagaraTypeDefinition::GetQuatDef(), TEXT("Rest Darboux")));
+		Sig.Inputs.Add(FNiagaraVariable(FNiagaraTypeDefinition::GetFloatDef(), TEXT("Delta Time")));
+		Sig.Outputs.Add(FNiagaraVariable(FNiagaraTypeDefinition::GetVec3Def(), TEXT("Material Multiplier")));
+
+		OutFunctions.Add(Sig);
+	}
+	{
+		FNiagaraFunctionSignature Sig;
+		Sig.Name = SolveStaticCollisionConstraintName;
+		Sig.bMemberFunction = true;
+		Sig.bRequiresContext = false;
+		Sig.Inputs.Add(FNiagaraVariable(FNiagaraTypeDefinition(GetClass()), TEXT("Hair Strands")));
+		Sig.Inputs.Add(FNiagaraVariable(FNiagaraTypeDefinition::GetFloatDef(), TEXT("Penetration Depth")));
+		Sig.Inputs.Add(FNiagaraVariable(FNiagaraTypeDefinition::GetVec3Def(), TEXT("Collision Position")));
+		Sig.Inputs.Add(FNiagaraVariable(FNiagaraTypeDefinition::GetVec3Def(), TEXT("Collision Velocity")));
+		Sig.Inputs.Add(FNiagaraVariable(FNiagaraTypeDefinition::GetVec3Def(), TEXT("Collision Normal")));
+		Sig.Inputs.Add(FNiagaraVariable(FNiagaraTypeDefinition::GetFloatDef(), TEXT("Static Friction")));
+		Sig.Inputs.Add(FNiagaraVariable(FNiagaraTypeDefinition::GetFloatDef(), TEXT("Kinetic Friction")));
+		Sig.Inputs.Add(FNiagaraVariable(FNiagaraTypeDefinition::GetFloatDef(), TEXT("Delta Time")));
+		Sig.Inputs.Add(FNiagaraVariable(FNiagaraTypeDefinition::GetVec3Def(), TEXT("Previous Position")));
+		Sig.Outputs.Add(FNiagaraVariable(FNiagaraTypeDefinition::GetVec3Def(), TEXT("Constraint Multiplier")));
+
+		OutFunctions.Add(Sig);
+	}
+	{
+		FNiagaraFunctionSignature Sig;
+		Sig.Name = ProjectStaticCollisionConstraintName;
+		Sig.bMemberFunction = true;
+		Sig.bRequiresContext = false;
+		Sig.Inputs.Add(FNiagaraVariable(FNiagaraTypeDefinition(GetClass()), TEXT("Hair Strands")));
+		Sig.Inputs.Add(FNiagaraVariable(FNiagaraTypeDefinition::GetFloatDef(), TEXT("Penetration Depth")));
+		Sig.Inputs.Add(FNiagaraVariable(FNiagaraTypeDefinition::GetVec3Def(), TEXT("Collision Position")));
+		Sig.Inputs.Add(FNiagaraVariable(FNiagaraTypeDefinition::GetVec3Def(), TEXT("Collision Velocity")));
+		Sig.Inputs.Add(FNiagaraVariable(FNiagaraTypeDefinition::GetVec3Def(), TEXT("Collision Normal")));
+		Sig.Inputs.Add(FNiagaraVariable(FNiagaraTypeDefinition::GetFloatDef(), TEXT("Static Friction")));
+		Sig.Inputs.Add(FNiagaraVariable(FNiagaraTypeDefinition::GetFloatDef(), TEXT("Kinetic Friction")));
+		Sig.Inputs.Add(FNiagaraVariable(FNiagaraTypeDefinition::GetFloatDef(), TEXT("Delta Time")));
+		Sig.Inputs.Add(FNiagaraVariable(FNiagaraTypeDefinition::GetVec3Def(), TEXT("Previous Position")));
+		Sig.Inputs.Add(FNiagaraVariable(FNiagaraTypeDefinition::GetVec3Def(), TEXT("Node Position")));
+		Sig.Outputs.Add(FNiagaraVariable(FNiagaraTypeDefinition::GetVec3Def(), TEXT("Node Position")));
+
+		OutFunctions.Add(Sig);
+	}
+	{
+		FNiagaraFunctionSignature Sig;
+		Sig.Name = ComputeRestDirectionName;
+		Sig.bMemberFunction = true;
+		Sig.bRequiresContext = false;
+		Sig.Inputs.Add(FNiagaraVariable(FNiagaraTypeDefinition(GetClass()), TEXT("Hair Strands")));
+		Sig.Inputs.Add(FNiagaraVariable(FNiagaraTypeDefinition::GetVec3Def(), TEXT("Node Position")));
+		Sig.Inputs.Add(FNiagaraVariable(FNiagaraTypeDefinition::GetQuatDef(), TEXT("Node Orientation")));
+		Sig.Outputs.Add(FNiagaraVariable(FNiagaraTypeDefinition::GetVec3Def(), TEXT("Rest Direction")));
+
+		OutFunctions.Add(Sig);
+	}
+	{
+		FNiagaraFunctionSignature Sig;
+		Sig.Name = UpdateNodeOrientationName;
+		Sig.bMemberFunction = true;
+		Sig.bRequiresContext = false;
+		Sig.Inputs.Add(FNiagaraVariable(FNiagaraTypeDefinition(GetClass()), TEXT("Hair Strands")));
+		Sig.Outputs.Add(FNiagaraVariable(FNiagaraTypeDefinition::GetBoolDef(), TEXT("Update Status")));
+
 		OutFunctions.Add(Sig);
 	}
 }
@@ -1030,11 +1332,36 @@ DEFINE_NDI_DIRECT_FUNC_BINDER(UNiagaraDataInterfaceHairStrands, BuildCollisionGr
 DEFINE_NDI_DIRECT_FUNC_BINDER(UNiagaraDataInterfaceHairStrands, QueryCollisionGrid);
 DEFINE_NDI_DIRECT_FUNC_BINDER(UNiagaraDataInterfaceHairStrands, ProjectCollisionGrid);
 
+DEFINE_NDI_DIRECT_FUNC_BINDER(UNiagaraDataInterfaceHairStrands, GetBoxExtent);
+DEFINE_NDI_DIRECT_FUNC_BINDER(UNiagaraDataInterfaceHairStrands, GetBoxCenter);
+
 DEFINE_NDI_DIRECT_FUNC_BINDER(UNiagaraDataInterfaceHairStrands, AdvectNodePosition);
 DEFINE_NDI_DIRECT_FUNC_BINDER(UNiagaraDataInterfaceHairStrands, AdvectNodeOrientation);
 DEFINE_NDI_DIRECT_FUNC_BINDER(UNiagaraDataInterfaceHairStrands, UpdateLinearVelocity);
 DEFINE_NDI_DIRECT_FUNC_BINDER(UNiagaraDataInterfaceHairStrands, UpdateAngularVelocity);
 
+DEFINE_NDI_DIRECT_FUNC_BINDER(UNiagaraDataInterfaceHairStrands, SetupStretchSpringMaterial);
+DEFINE_NDI_DIRECT_FUNC_BINDER(UNiagaraDataInterfaceHairStrands, SolveStretchSpringMaterial);
+DEFINE_NDI_DIRECT_FUNC_BINDER(UNiagaraDataInterfaceHairStrands, ProjectStretchSpringMaterial);
+
+DEFINE_NDI_DIRECT_FUNC_BINDER(UNiagaraDataInterfaceHairStrands, SetupBendSpringMaterial);
+DEFINE_NDI_DIRECT_FUNC_BINDER(UNiagaraDataInterfaceHairStrands, SolveBendSpringMaterial);
+DEFINE_NDI_DIRECT_FUNC_BINDER(UNiagaraDataInterfaceHairStrands, ProjectBendSpringMaterial);
+
+DEFINE_NDI_DIRECT_FUNC_BINDER(UNiagaraDataInterfaceHairStrands, SetupStretchRodMaterial);
+DEFINE_NDI_DIRECT_FUNC_BINDER(UNiagaraDataInterfaceHairStrands, SolveStretchRodMaterial);
+DEFINE_NDI_DIRECT_FUNC_BINDER(UNiagaraDataInterfaceHairStrands, ProjectStretchRodMaterial);
+
+DEFINE_NDI_DIRECT_FUNC_BINDER(UNiagaraDataInterfaceHairStrands, SetupBendRodMaterial);
+DEFINE_NDI_DIRECT_FUNC_BINDER(UNiagaraDataInterfaceHairStrands, SolveBendRodMaterial);
+DEFINE_NDI_DIRECT_FUNC_BINDER(UNiagaraDataInterfaceHairStrands, ProjectBendRodMaterial);
+
+DEFINE_NDI_DIRECT_FUNC_BINDER(UNiagaraDataInterfaceHairStrands, SolveStaticCollisionConstraint);
+DEFINE_NDI_DIRECT_FUNC_BINDER(UNiagaraDataInterfaceHairStrands, ProjectStaticCollisionConstraint);
+
+DEFINE_NDI_DIRECT_FUNC_BINDER(UNiagaraDataInterfaceHairStrands, ComputeRestDirection);
+DEFINE_NDI_DIRECT_FUNC_BINDER(UNiagaraDataInterfaceHairStrands, UpdateNodeOrientation);
+
 void UNiagaraDataInterfaceHairStrands::GetVMExternalFunction(const FVMExternalFunctionBindingInfo& BindingInfo, void* InstanceData, FVMExternalFunction &OutFunc)
 {
 	if (BindingInfo.Name == GetNumStrandsName)
@@ -1089,12 +1416,12 @@ void UNiagaraDataInterfaceHairStrands::GetVMExternalFunction(const FVMExternalFu
 	}
 	else if (BindingInfo.Name == ComputeNodeMassName)
 	{
-		check(BindingInfo.GetNumInputs() == 2 && BindingInfo.GetNumOutputs() == 1);
+		check(BindingInfo.GetNumInputs() == 5 && BindingInfo.GetNumOutputs() == 1);
 		NDI_FUNC_BINDER(UNiagaraDataInterfaceHairStrands, ComputeNodeMass)::Bind(this, OutFunc);
 	}
 	else if (BindingInfo.Name == ComputeNodeInertiaName)
 	{
-		check(BindingInfo.GetNumInputs() == 2 && BindingInfo.GetNumOutputs() == 3);
+		check(BindingInfo.GetNumInputs() == 5 && BindingInfo.GetNumOutputs() == 3);
 		NDI_FUNC_BINDER(UNiagaraDataInterfaceHairStrands, ComputeNodeInertia)::Bind(this, OutFunc);
 	}
 	else if (BindingInfo.Name == ComputeEdgeLengthName)
@@ -1172,6 +1499,96 @@ void UNiagaraDataInterfaceHairStrands::GetVMExternalFunction(const FVMExternalFu
 		check(BindingInfo.GetNumInputs() == 2 && BindingInfo.GetNumOutputs() == 1);
 		NDI_FUNC_BINDER(UNiagaraDataInterfaceHairStrands, ProjectCollisionGrid)::Bind(this, OutFunc);
 	}
+	else if (BindingInfo.Name == GetBoxExtentName)
+	{
+		check(BindingInfo.GetNumInputs() == 1 && BindingInfo.GetNumOutputs() == 3);
+		NDI_FUNC_BINDER(UNiagaraDataInterfaceHairStrands, GetBoxExtent)::Bind(this, OutFunc);
+	}
+	else if (BindingInfo.Name == GetBoxCenterName)
+	{
+		check(BindingInfo.GetNumInputs() == 1 && BindingInfo.GetNumOutputs() == 3);
+		NDI_FUNC_BINDER(UNiagaraDataInterfaceHairStrands, GetBoxCenter)::Bind(this, OutFunc);
+	}
+	else if (BindingInfo.Name == SetupStretchSpringMaterialName)
+	{
+		check(BindingInfo.GetNumInputs() == 6 && BindingInfo.GetNumOutputs() == 3);
+		NDI_FUNC_BINDER(UNiagaraDataInterfaceHairStrands, SetupStretchSpringMaterial)::Bind(this, OutFunc);
+	}
+	else if (BindingInfo.Name == SolveStretchSpringMaterialName)
+	{
+		check(BindingInfo.GetNumInputs() == 7 && BindingInfo.GetNumOutputs() == 1);
+		NDI_FUNC_BINDER(UNiagaraDataInterfaceHairStrands, SolveStretchSpringMaterial)::Bind(this, OutFunc);
+	}
+	else if (BindingInfo.Name == ProjectStretchSpringMaterialName)
+	{
+		check(BindingInfo.GetNumInputs() == 5 && BindingInfo.GetNumOutputs() == 1);
+		NDI_FUNC_BINDER(UNiagaraDataInterfaceHairStrands, ProjectStretchSpringMaterial)::Bind(this, OutFunc);
+	}
+	else if (BindingInfo.Name == SetupBendSpringMaterialName)
+	{
+		check(BindingInfo.GetNumInputs() == 6 && BindingInfo.GetNumOutputs() == 5);
+		NDI_FUNC_BINDER(UNiagaraDataInterfaceHairStrands, SetupBendSpringMaterial)::Bind(this, OutFunc);
+	}
+	else if (BindingInfo.Name == SolveBendSpringMaterialName)
+	{
+		check(BindingInfo.GetNumInputs() == 11 && BindingInfo.GetNumOutputs() == 3);
+		NDI_FUNC_BINDER(UNiagaraDataInterfaceHairStrands, SolveBendSpringMaterial)::Bind(this, OutFunc);
+	}
+	else if (BindingInfo.Name == ProjectBendSpringMaterialName)
+	{
+		check(BindingInfo.GetNumInputs() == 8 && BindingInfo.GetNumOutputs() == 3);
+		NDI_FUNC_BINDER(UNiagaraDataInterfaceHairStrands, ProjectBendSpringMaterial)::Bind(this, OutFunc);
+	}
+	else if (BindingInfo.Name == SetupStretchRodMaterialName)
+	{
+		check(BindingInfo.GetNumInputs() == 6 && BindingInfo.GetNumOutputs() == 5);
+		NDI_FUNC_BINDER(UNiagaraDataInterfaceHairStrands, SetupStretchRodMaterial)::Bind(this, OutFunc);
+	}
+	else if (BindingInfo.Name == SolveStretchRodMaterialName)
+	{
+		check(BindingInfo.GetNumInputs() == 9 && BindingInfo.GetNumOutputs() == 3);
+		NDI_FUNC_BINDER(UNiagaraDataInterfaceHairStrands, SolveStretchRodMaterial)::Bind(this, OutFunc);
+	}
+	else if (BindingInfo.Name == ProjectStretchRodMaterialName)
+	{
+		check(BindingInfo.GetNumInputs() == 5 && BindingInfo.GetNumOutputs() == 3);
+		NDI_FUNC_BINDER(UNiagaraDataInterfaceHairStrands, ProjectStretchRodMaterial)::Bind(this, OutFunc);
+	}
+	else if (BindingInfo.Name == SetupBendRodMaterialName)
+	{
+		check(BindingInfo.GetNumInputs() == 6 && BindingInfo.GetNumOutputs() == 5);
+		NDI_FUNC_BINDER(UNiagaraDataInterfaceHairStrands, SetupBendRodMaterial)::Bind(this, OutFunc);
+	}
+	else if (BindingInfo.Name == SolveBendRodMaterialName)
+	{
+		check(BindingInfo.GetNumInputs() == 13 && BindingInfo.GetNumOutputs() == 3);
+		NDI_FUNC_BINDER(UNiagaraDataInterfaceHairStrands, SolveBendRodMaterial)::Bind(this, OutFunc);
+	}
+	else if (BindingInfo.Name == ProjectBendRodMaterialName)
+	{
+		check(BindingInfo.GetNumInputs() == 9 && BindingInfo.GetNumOutputs() == 3);
+		NDI_FUNC_BINDER(UNiagaraDataInterfaceHairStrands, ProjectBendRodMaterial)::Bind(this, OutFunc);
+	}
+	else if (BindingInfo.Name == SolveStaticCollisionConstraintName)
+	{
+		check(BindingInfo.GetNumInputs() == 17 && BindingInfo.GetNumOutputs() == 3);
+		NDI_FUNC_BINDER(UNiagaraDataInterfaceHairStrands, SolveStaticCollisionConstraint)::Bind(this, OutFunc);
+	}
+	else if (BindingInfo.Name == ProjectStaticCollisionConstraintName)
+	{
+		check(BindingInfo.GetNumInputs() == 20 && BindingInfo.GetNumOutputs() == 3);
+		NDI_FUNC_BINDER(UNiagaraDataInterfaceHairStrands, ProjectStaticCollisionConstraint)::Bind(this, OutFunc);
+	}
+	else if (BindingInfo.Name == ComputeRestDirectionName)
+	{
+		check(BindingInfo.GetNumInputs() == 8 && BindingInfo.GetNumOutputs() == 3);
+		NDI_FUNC_BINDER(UNiagaraDataInterfaceHairStrands, ComputeRestDirection)::Bind(this, OutFunc);
+	}
+	else if (BindingInfo.Name == UpdateNodeOrientationName)
+	{
+		check(BindingInfo.GetNumInputs() == 1 && BindingInfo.GetNumOutputs() == 1);
+		NDI_FUNC_BINDER(UNiagaraDataInterfaceHairStrands, UpdateNodeOrientation)::Bind(this, OutFunc);
+	}
 }
 
 void WriteTransform(const FMatrix& ToWrite, FVectorVMContext& Context)
@@ -1196,29 +1613,24 @@ void WriteTransform(const FMatrix& ToWrite, FVectorVMContext& Context)
 	for (int32 i = 0; i < Context.NumInstances; ++i)
 	{
 		*Out00.GetDest() = ToWrite.M[0][0]; Out00.Advance();
-		*Out01.GetDest() = ToWrite.M[0][0]; Out01.Advance();
-		*Out02.GetDest() = ToWrite.M[0][0]; Out02.Advance();
-		*Out03.GetDest() = ToWrite.M[0][0]; Out03.Advance();
-		*Out04.GetDest() = ToWrite.M[0][0]; Out04.Advance();
-		*Out05.GetDest() = ToWrite.M[0][0]; Out05.Advance();
-		*Out06.GetDest() = ToWrite.M[0][0]; Out06.Advance();
-		*Out07.GetDest() = ToWrite.M[0][0]; Out07.Advance();
-		*Out08.GetDest() = ToWrite.M[0][0]; Out08.Advance();
-		*Out09.GetDest() = ToWrite.M[0][0]; Out09.Advance();
-		*Out10.GetDest() = ToWrite.M[0][0]; Out10.Advance();
-		*Out11.GetDest() = ToWrite.M[0][0]; Out11.Advance();
-		*Out12.GetDest() = ToWrite.M[0][0]; Out12.Advance();
-		*Out13.GetDest() = ToWrite.M[0][0]; Out13.Advance();
-		*Out14.GetDest() = ToWrite.M[0][0]; Out14.Advance();
-		*Out15.GetDest() = ToWrite.M[0][0]; Out15.Advance();
+		*Out01.GetDest() = ToWrite.M[0][1]; Out01.Advance();
+		*Out02.GetDest() = ToWrite.M[0][2]; Out02.Advance();
+		*Out03.GetDest() = ToWrite.M[0][3]; Out03.Advance();
+		*Out04.GetDest() = ToWrite.M[1][0]; Out04.Advance();
+		*Out05.GetDest() = ToWrite.M[1][1]; Out05.Advance();
+		*Out06.GetDest() = ToWrite.M[1][2]; Out06.Advance();
+		*Out07.GetDest() = ToWrite.M[1][3]; Out07.Advance();
+		*Out08.GetDest() = ToWrite.M[2][0]; Out08.Advance();
+		*Out09.GetDest() = ToWrite.M[2][1]; Out09.Advance();
+		*Out10.GetDest() = ToWrite.M[2][2]; Out10.Advance();
+		*Out11.GetDest() = ToWrite.M[2][3]; Out11.Advance();
+		*Out12.GetDest() = ToWrite.M[3][0]; Out12.Advance();
+		*Out13.GetDest() = ToWrite.M[3][1]; Out13.Advance();
+		*Out14.GetDest() = ToWrite.M[3][2]; Out14.Advance();
+		*Out15.GetDest() = ToWrite.M[3][3]; Out15.Advance();
 	}
 }
 
-/*VectorVM::FUserPtrHandler<FNDIHairStrandsData> InstData(Context);
-const FMatrix WorldInverse = InstData->WorldTransform.Inverse();
-
-WriteTransform(WorldInverse, Context);*/
-
 void UNiagaraDataInterfaceHairStrands::GetNumStrands(FVectorVMContext& Context)
 {
 	VectorVM::FUserPtrHandler<FNDIHairStrandsData> InstData(Context);
@@ -1277,24 +1689,50 @@ void UNiagaraDataInterfaceHairStrands::GetTipThickness(FVectorVMContext& Context
 void UNiagaraDataInterfaceHairStrands::GetWorldTransform(FVectorVMContext& Context)
 {
 	VectorVM::FUserPtrHandler<FNDIHairStrandsData> InstData(Context);
-	VectorVM::FExternalFuncRegisterHandler<FMatrix> OutWorldTransform(Context);
+	FMatrix WorldTransform = InstData->WorldTransform;
+
+	//UE_LOG(LogHairStrands, Warning, TEXT("Get Hair Transform : %s"), *InstData->WorldTransform.ToString());
+
+	WriteTransform(WorldTransform, Context);
+}
+
+void UNiagaraDataInterfaceHairStrands::GetWorldInverse(FVectorVMContext& Context)
+{
+	VectorVM::FUserPtrHandler<FNDIHairStrandsData> InstData(Context);
+	FMatrix WorldInverse = InstData->WorldTransform.Inverse();
+
+	WriteTransform(WorldInverse, Context);
+}
+
+void UNiagaraDataInterfaceHairStrands::GetBoxCenter(FVectorVMContext& Context)
+{
+	VectorVM::FUserPtrHandler<FNDIHairStrandsData> InstData(Context);
+	VectorVM::FExternalFuncRegisterHandler<float> OutBoxCenterX(Context);
+	VectorVM::FExternalFuncRegisterHandler<float> OutBoxCenterY(Context);
+	VectorVM::FExternalFuncRegisterHandler<float> OutBoxCenterZ(Context);
 
 	for (int32 InstanceIdx = 0; InstanceIdx < Context.NumInstances; ++InstanceIdx)
 	{
-		*OutWorldTransform.GetDestAndAdvance() = InstData->WorldTransform;
+		*OutBoxCenterX.GetDestAndAdvance() = 0.0;
+		*OutBoxCenterY.GetDestAndAdvance() = 0.0;
+		*OutBoxCenterZ.GetDestAndAdvance() = 0.0;
 	}
 }
 
-void UNiagaraDataInterfaceHairStrands::GetWorldInverse(FVectorVMContext& Context)
+void UNiagaraDataInterfaceHairStrands::GetBoxExtent(FVectorVMContext& Context)
 {
 	VectorVM::FUserPtrHandler<FNDIHairStrandsData> InstData(Context);
-	VectorVM::FExternalFuncRegisterHandler<FMatrix> OutWorldInverse(Context);
-
-	const FMatrix WorldInverse = InstData->WorldTransform.Inverse();
+	VectorVM::FExternalFuncRegisterHandler<float> OutBoxExtentX(Context);
+	VectorVM::FExternalFuncRegisterHandler<float> OutBoxExtentY(Context);
+	VectorVM::FExternalFuncRegisterHandler<float> OutBoxExtentZ(Context);
 
 	for (int32 InstanceIdx = 0; InstanceIdx < Context.NumInstances; ++InstanceIdx)
 	{
-		*OutWorldInverse.GetDestAndAdvance() = WorldInverse;
+		*OutBoxExtentX.GetDestAndAdvance() = -InstData->GridOrigin.X;
+		*OutBoxExtentY.GetDestAndAdvance() = -InstData->GridOrigin.Y;
+		*OutBoxExtentZ.GetDestAndAdvance() = -InstData->GridOrigin.Z;
+
+		//UE_LOG(LogHairStrands, Warning, TEXT("Get Box Extent : %f %f %f"), InstData->GridOrigin.X, InstData->GridOrigin.Y, InstData->GridOrigin.Z);
 	}
 }
 
@@ -1398,6 +1836,87 @@ void UNiagaraDataInterfaceHairStrands::ProjectCollisionGrid(FVectorVMContext& Co
 	// @todo : implement function for cpu 
 }
 
+void UNiagaraDataInterfaceHairStrands::SetupStretchSpringMaterial(FVectorVMContext& Context)
+{
+	// @todo : implement function for cpu 
+}
+
+void UNiagaraDataInterfaceHairStrands::SolveStretchSpringMaterial(FVectorVMContext& Context)
+{
+	// @todo : implement function for cpu 
+}
+
+void UNiagaraDataInterfaceHairStrands::ProjectStretchSpringMaterial(FVectorVMContext& Context)
+{
+	// @todo : implement function for cpu 
+}
+
+void UNiagaraDataInterfaceHairStrands::SetupBendSpringMaterial(FVectorVMContext& Context)
+{
+	// @todo : implement function for cpu 
+}
+
+void UNiagaraDataInterfaceHairStrands::SolveBendSpringMaterial(FVectorVMContext& Context)
+{
+	// @todo : implement function for cpu 
+}
+
+void UNiagaraDataInterfaceHairStrands::ProjectBendSpringMaterial(FVectorVMContext& Context)
+{
+	// @todo : implement function for cpu 
+}
+
+
+void UNiagaraDataInterfaceHairStrands::SetupStretchRodMaterial(FVectorVMContext& Context)
+{
+	// @todo : implement function for cpu 
+}
+
+void UNiagaraDataInterfaceHairStrands::SolveStretchRodMaterial(FVectorVMContext& Context)
+{
+	// @todo : implement function for cpu 
+}
+
+void UNiagaraDataInterfaceHairStrands::ProjectStretchRodMaterial(FVectorVMContext& Context)
+{
+	// @todo : implement function for cpu 
+}
+
+void UNiagaraDataInterfaceHairStrands::SetupBendRodMaterial(FVectorVMContext& Context)
+{
+	// @todo : implement function for cpu 
+}
+
+void UNiagaraDataInterfaceHairStrands::SolveBendRodMaterial(FVectorVMContext& Context)
+{
+	// @todo : implement function for cpu 
+}
+
+void UNiagaraDataInterfaceHairStrands::ProjectBendRodMaterial(FVectorVMContext& Context)
+{
+	// @todo : implement function for cpu 
+}
+
+void UNiagaraDataInterfaceHairStrands::ComputeRestDirection(FVectorVMContext& Context)
+{
+	// @todo : implement function for cpu 
+}
+
+void UNiagaraDataInterfaceHairStrands::UpdateNodeOrientation(FVectorVMContext& Context)
+{
+	// @todo : implement function for cpu 
+}
+
+void UNiagaraDataInterfaceHairStrands::SolveStaticCollisionConstraint(FVectorVMContext& Context)
+{
+	// @todo : implement function for cpu 
+}
+
+void UNiagaraDataInterfaceHairStrands::ProjectStaticCollisionConstraint(FVectorVMContext& Context)
+{
+	// @todo : implement function for cpu 
+}
+
 bool UNiagaraDataInterfaceHairStrands::GetFunctionHLSL(const FName& DefinitionFunctionName, FString InstanceFunctionName, FNiagaraDataInterfaceGPUParamInfo& ParamInfo, FString& OutHLSL)
 {
 	FNDIHairStrandsParametersName ParamNames(ParamInfo.DataInterfaceHLSLSymbol);
@@ -1425,9 +1944,9 @@ bool UNiagaraDataInterfaceHairStrands::GetFunctionHLSL(const FName& DefinitionFu
 	if (DefinitionFunctionName == GetStrandDensityName)
 	{
 		static const TCHAR *FormatSample = TEXT(R"(
-		void {InstanceFunctionName}(out float Out_StrandDensity)
+		void {InstanceFunctionName}(out float OutStrandDensity)
 		{
-			Out_StrandDensity = {StrandDensityName};
+			OutStrandDensity = {StrandDensityName};
 		}
 		)");
 		OutHLSL += FString::Format(FormatSample, ArgsSample);
@@ -1436,9 +1955,9 @@ bool UNiagaraDataInterfaceHairStrands::GetFunctionHLSL(const FName& DefinitionFu
 	else if (DefinitionFunctionName == GetStrandSizeName)
 	{
 		static const TCHAR *FormatSample = TEXT(R"(
-		void {InstanceFunctionName}(out int Out_StrandSize)
+		void {InstanceFunctionName}(out int OutStrandSize)
 		{
-			Out_StrandSize = {StrandSizeName};
+			OutStrandSize = {StrandSizeName};
 		}
 		)");
 		OutHLSL += FString::Format(FormatSample, ArgsSample);
@@ -1447,9 +1966,9 @@ bool UNiagaraDataInterfaceHairStrands::GetFunctionHLSL(const FName& DefinitionFu
 	else if (DefinitionFunctionName == GetNumStrandsName)
 	{
 		static const TCHAR *FormatSample = TEXT(R"(
-		void {InstanceFunctionName}(out int Out_NumStrands)
+		void {InstanceFunctionName}(out int OutNumStrands)
 		{
-			Out_NumStrands = {NumStrandsName};
+			OutNumStrands = {NumStrandsName};
 		}
 		)");
 		OutHLSL += FString::Format(FormatSample, ArgsSample);
@@ -1458,9 +1977,9 @@ bool UNiagaraDataInterfaceHairStrands::GetFunctionHLSL(const FName& DefinitionFu
 	else if (DefinitionFunctionName == GetRootThicknessName)
 	{
 		static const TCHAR *FormatSample = TEXT(R"(
-		void {InstanceFunctionName}(out float Out_RootThickness)
+		void {InstanceFunctionName}(out float OutRootThickness)
 		{
-			Out_RootThickness = {RootThicknessName};
+			OutRootThickness = {RootThicknessName};
 		}
 		)");
 		OutHLSL += FString::Format(FormatSample, ArgsSample);
@@ -1469,9 +1988,9 @@ bool UNiagaraDataInterfaceHairStrands::GetFunctionHLSL(const FName& DefinitionFu
 	else if (DefinitionFunctionName == GetTipThicknessName)
 	{
 		static const TCHAR *FormatSample = TEXT(R"(
-		void {InstanceFunctionName}(out float Out_TipThickness)
+		void {InstanceFunctionName}(out float OutTipThickness)
 		{
-			Out_TipThickness = {TipThicknessName};
+			OutTipThickness = {TipThicknessName};
 		}
 		)");
 		OutHLSL += FString::Format(FormatSample, ArgsSample);
@@ -1480,9 +1999,9 @@ bool UNiagaraDataInterfaceHairStrands::GetFunctionHLSL(const FName& DefinitionFu
 	else if (DefinitionFunctionName == GetWorldTransformName)
 	{
 		static const TCHAR *FormatSample = TEXT(R"(
-		void {InstanceFunctionName}(out float4x4 Out_WorldTransform)
+		void {InstanceFunctionName}(out float4x4 OutWorldTransform)
 		{
-			Out_WorldTransform = {WorldTransformName};
+			OutWorldTransform = {WorldTransformName};
 		}
 		)");
 		OutHLSL += FString::Format(FormatSample, ArgsSample);
@@ -1491,9 +2010,9 @@ bool UNiagaraDataInterfaceHairStrands::GetFunctionHLSL(const FName& DefinitionFu
 	else if (DefinitionFunctionName == GetWorldInverseName)
 	{
 		static const TCHAR *FormatSample = TEXT(R"(
-		void {InstanceFunctionName}(out float4x4 Out_WorldInverse)
+		void {InstanceFunctionName}(out float4x4 OutWorldInverse)
 		{
-			Out_WorldInverse = float4x4(0);//{WorldInverseName};
+			OutWorldInverse = {WorldInverseName};
 		}
 		)");
 		OutHLSL += FString::Format(FormatSample, ArgsSample);
@@ -1502,9 +2021,9 @@ bool UNiagaraDataInterfaceHairStrands::GetFunctionHLSL(const FName& DefinitionFu
 	else if (DefinitionFunctionName == GetPointPositionName)
 	{
 		static const TCHAR *FormatSample = TEXT(R"(
-			void {InstanceFunctionName} (in int PointIndex, out float3 Out_PointPosition)
+			void {InstanceFunctionName} (in int PointIndex, out float3 OutPointPosition)
 			{
-				{HairStrandsContextName} DIHairStrands_GetPointPosition(DIContext,PointIndex,Out_PointPosition);
+				{HairStrandsContextName} DIHairStrands_GetPointPosition(DIContext,PointIndex,OutPointPosition);
 			}
 			)");
 		OutHLSL += FString::Format(FormatSample, ArgsSample);
@@ -1513,9 +2032,9 @@ bool UNiagaraDataInterfaceHairStrands::GetFunctionHLSL(const FName& DefinitionFu
 	else if (DefinitionFunctionName == ComputeNodePositionName)
 	{
 		static const TCHAR *FormatSample = TEXT(R"(
-			void {InstanceFunctionName} (in int NodeIndex, out float3 Out_NodePosition)
+			void {InstanceFunctionName} (in int NodeIndex, out float3 OutNodePosition)
 			{
-				{HairStrandsContextName} DIHairStrands_ComputeNodePosition(DIContext,NodeIndex,Out_NodePosition);
+				{HairStrandsContextName} DIHairStrands_ComputeNodePosition(DIContext,NodeIndex,OutNodePosition);
 			}
 			)");
 		OutHLSL += FString::Format(FormatSample, ArgsSample);
@@ -1524,9 +2043,9 @@ bool UNiagaraDataInterfaceHairStrands::GetFunctionHLSL(const FName& DefinitionFu
 	else if (DefinitionFunctionName == ComputeNodeOrientationName)
 	{
 		static const TCHAR *FormatSample = TEXT(R"(
-			void {InstanceFunctionName} (in int NodeIndex, out float4 Out_NodeOrientation)
+			void {InstanceFunctionName} (in int NodeIndex, out float4 OutNodeOrientation)
 			{
-				{HairStrandsContextName} DIHairStrands_ComputeNodeOrientation(DIContext,NodeIndex,Out_NodeOrientation);
+				{HairStrandsContextName} DIHairStrands_ComputeNodeOrientation(DIContext,NodeIndex,OutNodeOrientation);
 			}
 			)");
 		OutHLSL += FString::Format(FormatSample, ArgsSample);
@@ -1535,9 +2054,9 @@ bool UNiagaraDataInterfaceHairStrands::GetFunctionHLSL(const FName& DefinitionFu
 	else if (DefinitionFunctionName == ComputeNodeMassName)
 	{
 		static const TCHAR *FormatSample = TEXT(R"(
-			void {InstanceFunctionName} (in int NodeIndex, out float Out_NodeMass)
+			void {InstanceFunctionName} (in int NodeIndex, in float StrandsDensity, in float RootThickness, in float TipThickness, out float OutNodeMass)
 			{
-				{HairStrandsContextName} DIHairStrands_ComputeNodeMass(DIContext,NodeIndex,Out_NodeMass);
+				{HairStrandsContextName} DIHairStrands_ComputeNodeMass(DIContext,NodeIndex,StrandsDensity,RootThickness,TipThickness,OutNodeMass);
 			}
 			)");
 		OutHLSL += FString::Format(FormatSample, ArgsSample);
@@ -1546,9 +2065,9 @@ bool UNiagaraDataInterfaceHairStrands::GetFunctionHLSL(const FName& DefinitionFu
 	else if (DefinitionFunctionName == ComputeNodeInertiaName)
 	{
 		static const TCHAR *FormatSample = TEXT(R"(
-			void {InstanceFunctionName} (in int NodeIndex, out float3 Out_NodeInertia)
+			void {InstanceFunctionName} (in int NodeIndex, in float StrandsDensity, in float RootThickness, in float TipThickness, out float3 OutNodeInertia)
 			{
-				{HairStrandsContextName} DIHairStrands_ComputeNodeInertia(DIContext,NodeIndex,Out_NodeInertia);
+				{HairStrandsContextName} DIHairStrands_ComputeNodeInertia(DIContext,NodeIndex,StrandsDensity,RootThickness,TipThickness,OutNodeInertia);
 			}
 			)");
 		OutHLSL += FString::Format(FormatSample, ArgsSample);
@@ -1725,6 +2244,219 @@ bool UNiagaraDataInterfaceHairStrands::GetFunctionHLSL(const FName& DefinitionFu
 		OutHLSL += FString::Format(FormatSample, ArgsSample);
 		return true;
 	}
+	else if (DefinitionFunctionName == GetBoxExtentName)
+	{
+		static const TCHAR *FormatSample = TEXT(R"(
+				void {InstanceFunctionName} (out float3 OutBoxExtent)
+				{
+					{HairStrandsContextName} DIHairStrands_GetBoxExtent(DIContext,OutBoxExtent);
+				}
+				)");
+		OutHLSL += FString::Format(FormatSample, ArgsSample);
+		return true;
+	}
+	else if (DefinitionFunctionName == GetBoxCenterName)
+	{
+		static const TCHAR *FormatSample = TEXT(R"(
+				void {InstanceFunctionName} (out float3 OutBoxCenter)
+				{
+					{HairStrandsContextName} DIHairStrands_GetBoxCenter(DIContext,OutBoxCenter);
+				}
+				)");
+		OutHLSL += FString::Format(FormatSample, ArgsSample);
+		return true;
+	}
+	else if (DefinitionFunctionName == SetupStretchSpringMaterialName)
+	{
+		static const TCHAR *FormatSample = TEXT(R"(
+				void {InstanceFunctionName} (in float YoungModulus, in float RodThickness, 
+in float RestLength, in float DeltaTime, in float MaterialDamping, out float OutMaterialCompliance, out float OutMaterialWeight, out float OutMaterialMultiplier)
+				{
+					{HairStrandsContextName} SetupStretchSpringMaterial(DIContext.StrandSize,YoungModulus,RodThickness,RestLength,DeltaTime,false,MaterialDamping,OutMaterialCompliance,OutMaterialWeight,OutMaterialMultiplier);
+				}
+				)");
+		OutHLSL += FString::Format(FormatSample, ArgsSample);
+		return true;
+	}
+	else if (DefinitionFunctionName == SolveStretchSpringMaterialName)
+	{
+		static const TCHAR *FormatSample = TEXT(R"(
+				void {InstanceFunctionName} (in float RestLength, in float DeltaTime, in float MaterialDamping, 
+		in float MaterialCompliance, in float MaterialWeight, in float MaterialMultiplier, out float OutMaterialMultiplier)
+				{
+					{HairStrandsContextName} SolveStretchSpringMaterial(DIContext.StrandSize,RestLength,DeltaTime,MaterialDamping,MaterialCompliance,MaterialWeight,MaterialMultiplier,OutMaterialMultiplier);
+				}
+				)");
+		OutHLSL += FString::Format(FormatSample, ArgsSample);
+		return true;
+	}
+	else if (DefinitionFunctionName == ProjectStretchSpringMaterialName)
+	{
+		static const TCHAR *FormatSample = TEXT(R"(
+				void {InstanceFunctionName} (in float YoungModulus, in float RodThickness, in float RestLength, in float DeltaTime, out float OutMaterialMultiplier)
+				{
+					{HairStrandsContextName} ProjectStretchSpringMaterial(DIContext.StrandSize,YoungModulus,RodThickness,RestLength,DeltaTime,OutMaterialMultiplier);
+				}
+				)");
+		OutHLSL += FString::Format(FormatSample, ArgsSample);
+		return true;
+	}
+	else if (DefinitionFunctionName == SetupBendSpringMaterialName)
+	{
+		static const TCHAR *FormatSample = TEXT(R"(
+					void {InstanceFunctionName} (in float YoungModulus, in float RodThickness, 
+	in float RestLength, in float DeltaTime, in float MaterialDamping, out float OutMaterialCompliance, out float OutMaterialWeight, out float3 OutMaterialMultiplier)
+					{
+						{HairStrandsContextName} SetupBendSpringMaterial(DIContext.StrandSize,YoungModulus,RodThickness,RestLength,DeltaTime,false,MaterialDamping,OutMaterialCompliance,OutMaterialWeight,OutMaterialMultiplier);
+					}
+					)");
+		OutHLSL += FString::Format(FormatSample, ArgsSample);
+		return true;
+	}
+	else if (DefinitionFunctionName == SolveBendSpringMaterialName)
+	{
+		static const TCHAR *FormatSample = TEXT(R"(
+					void {InstanceFunctionName} (in float3 RestDirection, in float DeltaTime, in float MaterialDamping, 
+			in float MaterialCompliance, in float MaterialWeight, in float3 MaterialMultiplier, out float3 OutMaterialMultiplier)
+					{
+						{HairStrandsContextName} SolveBendSpringMaterial(DIContext.StrandSize,RestDirection,DeltaTime,MaterialDamping,MaterialCompliance,MaterialWeight,MaterialMultiplier,OutMaterialMultiplier);
+					}
+					)");
+		OutHLSL += FString::Format(FormatSample, ArgsSample);
+		return true;
+	}
+	else if (DefinitionFunctionName == ProjectBendSpringMaterialName)
+	{
+		static const TCHAR *FormatSample = TEXT(R"(
+					void {InstanceFunctionName} (in float YoungModulus, in float RodThickness, in float RestLength, in float3 RestDirection, in float DeltaTime, out float OutMaterialMultiplier)
+					{
+						{HairStrandsContextName} ProjectBendSpringMaterial(DIContext.StrandSize,YoungModulus,RodThickness,RestLength,RestDirection,DeltaTime,OutMaterialMultiplier);
+					}
+					)");
+		OutHLSL += FString::Format(FormatSample, ArgsSample);
+		return true;
+	}
+	else if (DefinitionFunctionName == SetupStretchRodMaterialName)
+	{
+		static const TCHAR *FormatSample = TEXT(R"(
+					void {InstanceFunctionName} (in float YoungModulus, in float RodThickness, 
+	in float RestLength, in float DeltaTime, in float MaterialDamping, out float OutMaterialCompliance, out float OutMaterialWeight, out float3 OutMaterialMultiplier)
+					{
+						{HairStrandsContextName} SetupStretchRodMaterial(DIContext.StrandSize,YoungModulus,RodThickness,RestLength,DeltaTime,false,MaterialDamping,OutMaterialCompliance,OutMaterialWeight,OutMaterialMultiplier);
+					}
+					)");
+		OutHLSL += FString::Format(FormatSample, ArgsSample);
+		return true;
+	}
+	else if (DefinitionFunctionName == SolveStretchRodMaterialName)
+	{
+		static const TCHAR *FormatSample = TEXT(R"(
+					void {InstanceFunctionName} (in float RestLength, in float DeltaTime, in float MaterialDamping, 
+			in float MaterialCompliance, in float MaterialWeight, in float3 MaterialMultiplier, out float3 OutMaterialMultiplier)
+					{
+						{HairStrandsContextName} SolveStretchRodMaterial(DIContext.StrandSize,RestLength,DeltaTime,MaterialDamping,MaterialCompliance,MaterialWeight,MaterialMultiplier,OutMaterialMultiplier);
+					}
+					)");
+		OutHLSL += FString::Format(FormatSample, ArgsSample);
+		return true;
+	}
+	else if (DefinitionFunctionName == ProjectStretchRodMaterialName)
+	{
+		static const TCHAR *FormatSample = TEXT(R"(
+					void {InstanceFunctionName} (in float YoungModulus, in float RodThickness, in float RestLength, in float DeltaTime, out float3 OutMaterialMultiplier)
+					{
+						{HairStrandsContextName} ProjectStretchRodMaterial(DIContext.StrandSize,YoungModulus,RodThickness,RestLength,DeltaTime,OutMaterialMultiplier);
+					}
+					)");
+		OutHLSL += FString::Format(FormatSample, ArgsSample);
+		return true;
+	}
+	else if (DefinitionFunctionName == SetupBendRodMaterialName)
+	{
+		static const TCHAR *FormatSample = TEXT(R"(
+						void {InstanceFunctionName} (in float YoungModulus, in float RodThickness, 
+		in float RestLength, in float DeltaTime, in float MaterialDamping, out float OutMaterialCompliance, out float OutMaterialWeight, out float3 OutMaterialMultiplier)
+						{
+							{HairStrandsContextName} SetupBendRodMaterial(DIContext.StrandSize,YoungModulus,RodThickness,RestLength,DeltaTime,false,MaterialDamping,OutMaterialCompliance,OutMaterialWeight,OutMaterialMultiplier);
+						}
+						)");
+		OutHLSL += FString::Format(FormatSample, ArgsSample);
+		return true;
+	}
+	else if (DefinitionFunctionName == SolveBendRodMaterialName)
+	{
+		static const TCHAR *FormatSample = TEXT(R"(
+						void {InstanceFunctionName} (in float RestLength, in float4 RestDarboux, in float DeltaTime, in float MaterialDamping, 
+				in float MaterialCompliance, in float MaterialWeight, in float3 MaterialMultiplier, out float3 OutMaterialMultiplier)
+						{
+							{HairStrandsContextName} SolveBendRodMaterial(DIContext.StrandSize,RestLength,RestDarboux,DeltaTime,MaterialDamping,MaterialCompliance,MaterialWeight,MaterialMultiplier,OutMaterialMultiplier);
+						}
+						)");
+		OutHLSL += FString::Format(FormatSample, ArgsSample);
+		return true;
+	}
+	else if (DefinitionFunctionName == ProjectBendRodMaterialName)
+	{
+		static const TCHAR *FormatSample = TEXT(R"(
+						void {InstanceFunctionName} (in float YoungModulus, in float RodThickness, in float RestLength, in float4 RestDarboux, in float DeltaTime, out float3 OutMaterialMultiplier)
+						{
+							{HairStrandsContextName} ProjectBendRodMaterial(DIContext.StrandSize,YoungModulus,RodThickness,RestLength,RestDarboux,DeltaTime,OutMaterialMultiplier);
+						}
+						)");
+		OutHLSL += FString::Format(FormatSample, ArgsSample);
+		return true;
+	}
+	else if (DefinitionFunctionName == SolveStaticCollisionConstraintName)
+	{
+		static const TCHAR *FormatSample = TEXT(R"(
+							void {InstanceFunctionName} (in float PenetrationDepth, in float3 CollisionPosition, in float3 CollisionVelocity, in float3 CollisionNormal, 
+				in float StaticFriction, in float KineticFriction, in float DeltaTime, in float3 PreviousPosition, out float3 OutMaterialMultiplier )
+							{
+								OutMaterialMultiplier = float3(0,0,0);
+								{HairStrandsContextName} SolveStaticCollisionConstraint(DIContext.StrandSize,PenetrationDepth,
+									CollisionPosition,CollisionVelocity,CollisionNormal,StaticFriction,KineticFriction,DeltaTime,false,PreviousPosition,SharedNodePosition[GGroupThreadId.x]);
+							}
+							)");
+		OutHLSL += FString::Format(FormatSample, ArgsSample);
+		return true;
+	}
+	else if (DefinitionFunctionName == ProjectStaticCollisionConstraintName)
+	{
+			static const TCHAR *FormatSample = TEXT(R"(
+						void {InstanceFunctionName} (in float PenetrationDepth, in float3 CollisionPosition, in float3 CollisionVelocity, in float3 CollisionNormal, 
+			in float StaticFriction, in float KineticFriction, in float DeltaTime, in float3 PreviousPosition, in float3 NodePosition, out float3 OutNodePosition )
+						{
+							OutNodePosition = NodePosition;
+							{HairStrandsContextName} SolveStaticCollisionConstraint(DIContext.StrandSize,PenetrationDepth,
+								CollisionPosition,CollisionVelocity,CollisionNormal,StaticFriction,KineticFriction,DeltaTime,true,PreviousPosition,OutNodePosition);
+						}
+						)");
+		OutHLSL += FString::Format(FormatSample, ArgsSample);
+		return true;
+	}
+	else if (DefinitionFunctionName == ComputeRestDirectionName)
+	{
+		static const TCHAR *FormatSample = TEXT(R"(
+							void {InstanceFunctionName} (in float3 NodePosition, in float4 NodeOrientation, out float3 OutRestDirection)
+							{
+								{HairStrandsContextName} ComputeRestDirection(DIContext.StrandSize,NodePosition,NodeOrientation,OutRestDirection);
+							}
+							)");
+		OutHLSL += FString::Format(FormatSample, ArgsSample);
+		return true;
+	}
+	else if (DefinitionFunctionName == UpdateNodeOrientationName)
+	{
+		static const TCHAR *FormatSample = TEXT(R"(
+					void {InstanceFunctionName} ( out bool OutUpdateStatus)
+					{
+						{HairStrandsContextName} ComputeMaterialFrame(DIContext.StrandSize);
+						OutUpdateStatus = true;
+					}
+					)");
+		OutHLSL += FString::Format(FormatSample, ArgsSample);
+		return true;
+	}
 
 	OutHLSL += TEXT("\n");
 	return false;
@@ -1733,6 +2465,9 @@ bool UNiagaraDataInterfaceHairStrands::GetFunctionHLSL(const FName& DefinitionFu
 void UNiagaraDataInterfaceHairStrands::GetCommonHLSL(FString& OutHLSL)
 {
 	OutHLSL += TEXT("#include \"/Plugin/Experimental/HairStrands/Private/NiagaraQuaternionUtils.ush\"\n");
+	OutHLSL += TEXT("#include \"/Plugin/Experimental/HairStrands/Private/NiagaraHookeSpringMaterial.ush\"\n");
+	OutHLSL += TEXT("#include \"/Plugin/Experimental/HairStrands/Private/NiagaraCosseratRodMaterial.ush\"\n");
+	OutHLSL += TEXT("#include \"/Plugin/Experimental/HairStrands/Private/NiagaraStaticCollisionConstraint.ush\"\n");
 	OutHLSL += TEXT("#include \"/Plugin/Experimental/HairStrands/Private/NiagaraDataInterfaceHairStrands.ush\"\n");
 }
 
diff --git a/Engine/Plugins/Experimental/HairStrands/Source/HairStrandsNiagara/Private/NiagaraDataInterfacePhysicsAsset.cpp b/Engine/Plugins/Experimental/HairStrands/Source/HairStrandsNiagara/Private/NiagaraDataInterfacePhysicsAsset.cpp
--- a/Engine/Plugins/Experimental/HairStrands/Source/HairStrandsNiagara/Private/NiagaraDataInterfacePhysicsAsset.cpp
+++ b/Engine/Plugins/Experimental/HairStrands/Source/HairStrandsNiagara/Private/NiagaraDataInterfacePhysicsAsset.cpp
@@ -703,7 +703,7 @@ bool UNiagaraDataInterfacePhysicsAsset::GetFunctionHLSL(const FName& DefinitionF
 		{TEXT("PreviousTransformBufferName"), ParamNames.PreviousTransformBufferName},
 		{TEXT("InverseTransformBufferName"), ParamNames.InverseTransformBufferName},
 		{TEXT("ElementExtentBufferName"), ParamNames.ElementExtentBufferName},
-		{TEXT("PhysicsAssetContextName"), TEXT("DIPhysicsAsset_MAKE_CONTEXT(") + ParamInfo.DataInterfaceHLSLSymbol + TEXT(")")},
+		{TEXT("PhysicsAssetContextName"), TEXT("DIPHYSICSASSET_MAKE_CONTEXT(") + ParamInfo.DataInterfaceHLSLSymbol + TEXT(")")},
 	};
 
 	if (DefinitionFunctionName == GetNumBoxesName)
@@ -767,7 +767,7 @@ void UNiagaraDataInterfacePhysicsAsset::GetCommonHLSL(FString& OutHLSL)
 
 void UNiagaraDataInterfacePhysicsAsset::GetParameterDefinitionHLSL(FNiagaraDataInterfaceGPUParamInfo& ParamInfo, FString& OutHLSL)
 {
-	OutHLSL += TEXT("DIPhysicsAsset_DECLARE_CONSTANTS(") + ParamInfo.DataInterfaceHLSLSymbol + TEXT(")\n");
+	OutHLSL += TEXT("DIPHYSICSASSET_DECLARE_CONSTANTS(") + ParamInfo.DataInterfaceHLSLSymbol + TEXT(")\n");
 }
 
 void UNiagaraDataInterfacePhysicsAsset::ProvidePerInstanceDataForRenderThread(void* DataForRenderThread, void* PerInstanceData, const FNiagaraSystemInstanceID& SystemInstance)
