Copying //UE4/Dev-Physics to Dev-Main (//UE4/Dev-Main) [at] 9420715

diff --git a/Engine/Plugins/Experimental/HairStrands/Source/HairStrandsNiagara/Private/NiagaraDataInterfacePressureGrid.cpp b/Engine/Plugins/Experimental/HairStrands/Source/HairStrandsNiagara/Private/NiagaraDataInterfacePressureGrid.cpp
new file mode 100644
--- /dev/null
+++ b/Engine/Plugins/Experimental/HairStrands/Source/HairStrandsNiagara/Private/NiagaraDataInterfacePressureGrid.cpp
@@ -0,0 +1,873 @@
+// Copyright 1998-2019 Epic Games, Inc. All Rights Reserved.
+
+#include "NiagaraDataInterfacePressureGrid.h"
+#include "NiagaraShader.h"
+#include "NiagaraComponent.h"
+#include "NiagaraRenderer.h"
+#include "NiagaraSystemInstance.h"
+#include "NiagaraEmitterInstanceBatcher.h"
+
+#include "ShaderParameterUtils.h"
+#include "ClearQuad.h"
+
+#define LOCTEXT_NAMESPACE "NiagaraDataInterfacePressureGrid"
+DEFINE_LOG_CATEGORY_STATIC(LogPressureGrid, Log, All);
+
+//------------------------------------------------------------------------------------------------------------
+
+static const FName BuildVelocityFieldName(TEXT("BuildVelocityField"));
+static const FName ProjectVelocityFieldName(TEXT("ProjectVelocityField"));
+static const FName SampleVelocityFieldName(TEXT("SampleVelocityField"));
+static const FName GetCellPositionName(TEXT("GetCellPosition"));
+static const FName TransferCellVelocityName(TEXT("TransferCellVelocity"));
+static const FName SetSolidBoundaryName(TEXT("SetSolidBoundary"));
+static const FName ComputeBoundaryWeightsName(TEXT("ComputeBoundaryWeights"));
+static const FName BuildGridTopologyName(TEXT("BuildGridTopology"));
+static const FName UpdateGridTransformName(TEXT("UpdateGridTransform"));
+
+//------------------------------------------------------------------------------------------------------------
+
+const FString UNiagaraDataInterfacePressureGrid::GridCurrentBufferName(TEXT("GridCurrentBuffer_"));
+const FString UNiagaraDataInterfacePressureGrid::GridDestinationBufferName(TEXT("GridDestinationBuffer_"));
+
+const FString UNiagaraDataInterfacePressureGrid::GridSizeName(TEXT("GridSize_"));
+const FString UNiagaraDataInterfacePressureGrid::GridOriginName(TEXT("GridOrigin_"));
+
+const FString UNiagaraDataInterfacePressureGrid::WorldTransformName(TEXT("WorldTransform_"));
+const FString UNiagaraDataInterfacePressureGrid::WorldInverseName(TEXT("WorldInverse_"));
+
+//------------------------------------------------------------------------------------------------------------
+
+struct FNDIPressureGridParametersName
+{
+	FNDIPressureGridParametersName(const FString& Suffix)
+	{
+		GridCurrentBufferName = UNiagaraDataInterfacePressureGrid::GridCurrentBufferName + Suffix;
+		GridDestinationBufferName = UNiagaraDataInterfacePressureGrid::GridDestinationBufferName + Suffix;
+
+		GridSizeName = UNiagaraDataInterfacePressureGrid::GridSizeName + Suffix;
+		GridOriginName = UNiagaraDataInterfacePressureGrid::GridOriginName + Suffix;
+
+		WorldTransformName = UNiagaraDataInterfacePressureGrid::WorldTransformName + Suffix;
+		WorldInverseName = UNiagaraDataInterfacePressureGrid::WorldInverseName + Suffix;
+	}
+
+	FString GridCurrentBufferName;
+	FString GridDestinationBufferName;
+
+	FString GridSizeName;
+	FString GridOriginName;
+
+	FString WorldTransformName;
+	FString WorldInverseName;
+};
+
+//------------------------------------------------------------------------------------------------------------
+
+void FNDIPressureGridBuffer::SetGridSize(const FUintVector4 InGridSize)
+{
+	GridSize = InGridSize;
+}
+
+void FNDIPressureGridBuffer::InitRHI()
+{
+	if(GridSize.X != 0 && GridSize.Y != 0 && GridSize.Z != 0)
+	{
+		static const uint32 NumComponents = 14;
+		GridDataBuffer.Initialize(sizeof(int32), GridSize.X*NumComponents, GridSize.Y,
+			GridSize.Z, EPixelFormat::PF_R32_SINT);
+	}
+}
+
+void FNDIPressureGridBuffer::ReleaseRHI()
+{
+	GridDataBuffer.Release();
+}
+
+void FNDIPressureGridBuffer::ClearBuffers(FRHICommandList& RHICmdList)
+{
+	ClearUAV(RHICmdList, GridDataBuffer, FLinearColor(0, 0, 0, 0));
+}
+
+//------------------------------------------------------------------------------------------------------------
+
+void FNDIPressureGridData::SwapBuffers()
+{
+	FNDIPressureGridBuffer* StoredBufferPointer = CurrentGridBuffer;
+	CurrentGridBuffer = DestinationGridBuffer;
+	DestinationGridBuffer = StoredBufferPointer;
+}
+
+//------------------------------------------------------------------------------------------------------------
+
+struct FNDIPressureGridParametersCS : public FNiagaraDataInterfaceParametersCS
+{
+	virtual void Bind(const FNiagaraDataInterfaceParamRef& ParamRef, const class FShaderParameterMap& ParameterMap) override
+	{
+		FNDIPressureGridParametersName ParamNames(ParamRef.ParameterInfo.DataInterfaceHLSLSymbol);
+
+		GridCurrentBuffer.Bind(ParameterMap, *ParamNames.GridCurrentBufferName);
+		GridDestinationBuffer.Bind(ParameterMap, *ParamNames.GridDestinationBufferName);
+
+		GridOrigin.Bind(ParameterMap, *ParamNames.GridOriginName);
+		GridSize.Bind(ParameterMap, *ParamNames.GridSizeName);
+
+		WorldTransform.Bind(ParameterMap, *ParamNames.WorldTransformName);
+		WorldInverse.Bind(ParameterMap, *ParamNames.WorldInverseName);
+
+		if (!GridCurrentBuffer.IsBound())
+		{
+			UE_LOG(LogPressureGrid, Warning, TEXT("Binding failed for FNDIPressureGridParametersCS %s. Was it optimized out?"), *ParamNames.GridCurrentBufferName)
+		}
+
+		if (!GridDestinationBuffer.IsBound())
+		{
+			UE_LOG(LogPressureGrid, Warning, TEXT("Binding failed for FNDIPressureGridParametersCS %s. Was it optimized out?"), *ParamNames.GridDestinationBufferName)
+		}
+	}
+
+	virtual void Serialize(FArchive& Ar) override
+	{
+		Ar << GridCurrentBuffer;
+		Ar << GridDestinationBuffer;
+
+		Ar << GridOrigin;
+		Ar << GridSize;
+
+		Ar << WorldTransform;
+		Ar << WorldInverse;
+	}
+
+	virtual void Set(FRHICommandList& RHICmdList, const FNiagaraDataInterfaceSetArgs& Context) const override
+	{
+		check(IsInRenderingThread());
+
+		FRHIComputeShader* ComputeShaderRHI = Context.Shader->GetComputeShader();
+
+		FNDIPressureGridProxy* InterfaceProxy =
+			static_cast<FNDIPressureGridProxy*>(Context.DataInterface);
+		FNDIPressureGridData* ProxyData =
+			InterfaceProxy->SystemInstancesToProxyData.Find(Context.SystemInstance);
+		ensure(ProxyData);
+
+		if (ProxyData != nullptr)
+		{
+			FNDIPressureGridBuffer* CurrentGridBuffer = ProxyData->CurrentGridBuffer;
+			FNDIPressureGridBuffer* DestinationGridBuffer = ProxyData->DestinationGridBuffer;
+
+			RHICmdList.TransitionResource(EResourceTransitionAccess::EWritable, EResourceTransitionPipeline::EComputeToCompute, DestinationGridBuffer->GridDataBuffer.UAV);
+			SetUAVParameter(RHICmdList, ComputeShaderRHI, GridDestinationBuffer, DestinationGridBuffer->GridDataBuffer.UAV);
+
+			RHICmdList.TransitionResource(EResourceTransitionAccess::EReadable, EResourceTransitionPipeline::EComputeToCompute, CurrentGridBuffer->GridDataBuffer.UAV);
+			SetSRVParameter(RHICmdList, ComputeShaderRHI, GridCurrentBuffer, CurrentGridBuffer->GridDataBuffer.SRV);
+
+			//UE_LOG(LogPressureGrid, Warning, TEXT("Binding Grid Origin : %s"), *ProxyData->GridOrigin.ToString());
+			SetShaderValue(RHICmdList, ComputeShaderRHI, GridOrigin, ProxyData->GridOrigin);
+
+			//UE_LOG(LogPressureGrid, Warning, TEXT("Binding Grid Size : %d %d %d"), ProxyData->GridSize.X, ProxyData->GridSize.Y, ProxyData->GridSize.Z);
+			SetShaderValue(RHICmdList, ComputeShaderRHI, GridSize, ProxyData->GridSize);
+
+			//UE_LOG(LogPressureGrid, Warning, TEXT("Binding Grid Transform : %s"), *ProxyData->WorldTransform.ToString());
+			SetShaderValue(RHICmdList, ComputeShaderRHI, WorldTransform, ProxyData->WorldTransform);
+			SetShaderValue(RHICmdList, ComputeShaderRHI, WorldInverse, ProxyData->WorldTransform.Inverse());
+		}
+		else
+		{
+			SetUAVParameter(RHICmdList, ComputeShaderRHI, GridDestinationBuffer, FNiagaraRenderer::GetDummyUIntBuffer().UAV);
+			SetSRVParameter(RHICmdList, ComputeShaderRHI, GridCurrentBuffer, FNiagaraRenderer::GetDummyUIntBuffer().SRV);
+
+			SetShaderValue(RHICmdList, ComputeShaderRHI, GridOrigin, FVector4(0,0,0,0));
+			SetShaderValue(RHICmdList, ComputeShaderRHI, GridSize, FUintVector4());
+
+			SetShaderValue(RHICmdList, ComputeShaderRHI, WorldTransform, FMatrix::Identity);
+			SetShaderValue(RHICmdList, ComputeShaderRHI, WorldInverse, FMatrix::Identity);
+		}
+	}
+
+	virtual void Unset(FRHICommandList& RHICmdList, const FNiagaraDataInterfaceSetArgs& Context) const override
+	{
+		SetUAVParameter(RHICmdList, Context.Shader->GetComputeShader(), GridDestinationBuffer, nullptr);
+	}
+
+private:
+
+	FShaderResourceParameter GridCurrentBuffer;
+	FShaderResourceParameter GridDestinationBuffer;
+
+	FShaderParameter GridSize;
+	FShaderParameter GridOrigin;
+
+	FShaderParameter WorldTransform;
+	FShaderParameter WorldInverse;
+};
+
+//------------------------------------------------------------------------------------------------------------
+
+void FNDIPressureGridProxy::DeferredDestroy()
+{
+	for (const FNiagaraSystemInstanceID& Sys : DeferredDestroyList)
+	{
+		SystemInstancesToProxyData.Remove(Sys);
+	}
+
+	DeferredDestroyList.Empty();
+}
+
+void FNDIPressureGridProxy::ConsumePerInstanceDataFromGameThread(void* PerInstanceData, const FNiagaraSystemInstanceID& Instance)
+{
+	FNDIPressureGridData* SourceData = static_cast<FNDIPressureGridData*>(PerInstanceData);
+	FNDIPressureGridData* TargetData = SystemInstancesToProxyData.Find(Instance);
+
+	ensure(TargetData);
+	if (TargetData)
+	{
+		TargetData->WorldTransform = SourceData->WorldTransform;
+		TargetData->GridOrigin = SourceData->GridOrigin;
+	}
+	else
+	{
+		UE_LOG(LogPressureGrid, Log, TEXT("ConsumePerInstanceDataFromGameThread() ... could not find %s"), *FNiagaraUtilities::SystemInstanceIDToString(Instance));
+	}
+}
+
+void FNDIPressureGridProxy::InitializePerInstanceData(const FNiagaraSystemInstanceID& SystemInstance, FNDIPressureGridBuffer* CurrentGridBuffer,
+FNDIPressureGridBuffer* DestinationGridBuffer, const FVector4& GridOrigin, const FUintVector4& GridSize)
+{
+	check(IsInRenderingThread());
+
+	FNDIPressureGridData* TargetData = SystemInstancesToProxyData.Find(SystemInstance);
+	if (TargetData != nullptr)
+	{
+		DeferredDestroyList.Remove(SystemInstance);
+	}
+	else
+	{
+		TargetData = &SystemInstancesToProxyData.Add(SystemInstance);
+	}
+	TargetData->CurrentGridBuffer = CurrentGridBuffer;
+	TargetData->DestinationGridBuffer = DestinationGridBuffer;
+	TargetData->GridOrigin = GridOrigin;
+	TargetData->GridSize = GridSize;
+}
+
+void FNDIPressureGridProxy::DestroyPerInstanceData(NiagaraEmitterInstanceBatcher* Batcher, const FNiagaraSystemInstanceID& SystemInstance)
+{
+	check(IsInRenderingThread());
+
+	DeferredDestroyList.Add(SystemInstance);
+	Batcher->EnqueueDeferredDeletesForDI_RenderThread(this->AsShared());
+}
+
+//------------------------------------------------------------------------------------------------------------
+
+UNiagaraDataInterfacePressureGrid::UNiagaraDataInterfacePressureGrid(FObjectInitializer const& ObjectInitializer)
+	: Super(ObjectInitializer)
+	, GridSize(10)
+	, GridOrigin(0,0,0)
+	, GridLength(1)
+{
+
+	Proxy = MakeShared<FNDIPressureGridProxy, ESPMode::ThreadSafe>();
+}
+
+bool UNiagaraDataInterfacePressureGrid::InitPerInstanceData(void* PerInstanceData, FNiagaraSystemInstance* SystemInstance)
+{
+	FNDIPressureGridData* InstanceData = new (PerInstanceData) FNDIPressureGridData();
+	check(InstanceData);
+	
+	FNDIPressureGridBuffer* CurrentGridBuffer = new FNDIPressureGridBuffer;
+	FNDIPressureGridBuffer* DestinationGridBuffer = new FNDIPressureGridBuffer;
+		
+
+	// Push instance data to RT
+	{
+		const FUintVector4 LocalGridSize(GridSize.X, GridSize.Y, GridSize.Z,0);
+		const FVector4 LocalGridOrigin(GridOrigin.X,GridOrigin.Y,GridOrigin.Z,GridLength);
+		CurrentGridBuffer->SetGridSize(LocalGridSize);
+		DestinationGridBuffer->SetGridSize(LocalGridSize);
+
+		TSet<int> RT_OutputShaderStages = OutputShaderStages;
+		TSet<int> RT_IterationShaderStages = IterationShaderStages;
+		
+		InstanceData->CurrentGridBuffer = CurrentGridBuffer;
+		InstanceData->DestinationGridBuffer = DestinationGridBuffer;
+		InstanceData->GridSize = LocalGridSize;
+		InstanceData->GridOrigin = LocalGridOrigin;
+
+		InstanceData->WorldTransform = FMatrix::Identity;// SystemInstance->GetComponent()->GetComponentToWorld().ToMatrixWithScale();
+
+		FNDIPressureGridProxy* ThisProxy = GetProxyAs<FNDIPressureGridProxy>();
+		ENQUEUE_RENDER_COMMAND(FNiagaraDIPushInitialInstanceDataToRT) (
+			[ThisProxy, RT_OutputShaderStages, RT_IterationShaderStages, InstanceID = SystemInstance->GetId(), CurrentGridBuffer, DestinationGridBuffer, LocalGridSize, LocalGridOrigin](FRHICommandListImmediate& CmdList)
+		{
+			ThisProxy->OutputShaderStages = RT_OutputShaderStages;
+			ThisProxy->IterationShaderStages = RT_IterationShaderStages;
+			ThisProxy->SetElementCount(LocalGridSize.X*LocalGridSize.Y*LocalGridSize.Z);
+
+			CurrentGridBuffer->InitResource();
+			DestinationGridBuffer->InitResource();
+
+			ThisProxy->InitializePerInstanceData(InstanceID, CurrentGridBuffer, DestinationGridBuffer, LocalGridOrigin, LocalGridSize);
+		}
+		);
+	}
+	return true;
+}
+
+void UNiagaraDataInterfacePressureGrid::DestroyPerInstanceData(void* PerInstanceData, FNiagaraSystemInstance* SystemInstance)
+{
+	FNDIPressureGridData* InstanceData = static_cast<FNDIPressureGridData*>(PerInstanceData);
+	InstanceData->~FNDIPressureGridData();
+
+	if (InstanceData->CurrentGridBuffer && InstanceData->DestinationGridBuffer)
+	{
+		FNDIPressureGridProxy* ThisProxy = GetProxyAs<FNDIPressureGridProxy>();
+		FNDIPressureGridBuffer* CurrentGridBuffer = InstanceData->CurrentGridBuffer;
+		FNDIPressureGridBuffer* DestinationGridBuffer = InstanceData->DestinationGridBuffer;
+		ENQUEUE_RENDER_COMMAND(FNiagaraDIDestroyInstanceData) (
+			[ThisProxy, CurrentGridBuffer, DestinationGridBuffer, InstanceID = SystemInstance->GetId(), Batcher = SystemInstance->GetBatcher()](FRHICommandListImmediate& CmdList)
+		{
+			CurrentGridBuffer->ReleaseResource();
+			DestinationGridBuffer->ReleaseResource();
+			ThisProxy->DestroyPerInstanceData(Batcher, InstanceID);
+			delete CurrentGridBuffer;
+			delete DestinationGridBuffer;
+		}
+		);
+		InstanceData->CurrentGridBuffer = nullptr;
+		InstanceData->DestinationGridBuffer = nullptr;
+	}
+}
+
+bool UNiagaraDataInterfacePressureGrid::PerInstanceTick(void* PerInstanceData, FNiagaraSystemInstance* SystemInstance, float InDeltaSeconds)
+{
+	FNDIPressureGridData* InstanceData = static_cast<FNDIPressureGridData*>(PerInstanceData);
+
+	bool RequireReset = false;
+	if (InstanceData)
+	{
+		InstanceData->WorldTransform = SystemInstance->GetComponent()->GetComponentToWorld().ToMatrixWithScale();
+	}
+	
+	return RequireReset;
+}
+
+bool UNiagaraDataInterfacePressureGrid::CopyToInternal(UNiagaraDataInterface* Destination) const
+{
+	if (!Super::CopyToInternal(Destination))
+	{
+		return false;
+	}
+
+	UNiagaraDataInterfacePressureGrid* OtherTyped = CastChecked<UNiagaraDataInterfacePressureGrid>(Destination);
+	OtherTyped->GridSize = GridSize;
+	OtherTyped->GridOrigin = GridOrigin;
+	OtherTyped->GridLength= GridLength;
+
+	return true;
+}
+
+bool UNiagaraDataInterfacePressureGrid::Equals(const UNiagaraDataInterface* Other) const
+{
+	if (!Super::Equals(Other))
+	{
+		return false;
+	}
+	const UNiagaraDataInterfacePressureGrid* OtherTyped = CastChecked<const UNiagaraDataInterfacePressureGrid>(Other);
+
+	return (OtherTyped->GridSize == GridSize) && (OtherTyped->GridOrigin == GridOrigin) && (OtherTyped->GridLength == GridLength);
+}
+
+void UNiagaraDataInterfacePressureGrid::PostInitProperties()
+{
+	Super::PostInitProperties();
+
+	if (HasAnyFlags(RF_ClassDefaultObject))
+	{
+		FNiagaraTypeRegistry::Register(FNiagaraTypeDefinition(GetClass()), true, false, false);
+	}
+}
+
+void UNiagaraDataInterfacePressureGrid::GetFunctions(TArray<FNiagaraFunctionSignature>& OutFunctions)
+{
+	{
+		FNiagaraFunctionSignature Sig;
+		Sig.Name = BuildVelocityFieldName;
+		Sig.bMemberFunction = true;
+		Sig.bRequiresContext = false;
+		Sig.Inputs.Add(FNiagaraVariable(FNiagaraTypeDefinition(GetClass()), TEXT("Pressure Grid")));
+		Sig.Inputs.Add(FNiagaraVariable(FNiagaraTypeDefinition::GetIntDef(), TEXT("Strands Size")));
+		Sig.Inputs.Add(FNiagaraVariable(FNiagaraTypeDefinition::GetVec3Def(), TEXT("Node Position")));
+		Sig.Inputs.Add(FNiagaraVariable(FNiagaraTypeDefinition::GetFloatDef(), TEXT("Node Mass")));
+		Sig.Inputs.Add(FNiagaraVariable(FNiagaraTypeDefinition::GetVec3Def(), TEXT("Node Velocity")));
+		Sig.Inputs.Add(FNiagaraVariable(FNiagaraTypeDefinition::GetVec3Def(), TEXT("Velocity GradientX")));
+		Sig.Inputs.Add(FNiagaraVariable(FNiagaraTypeDefinition::GetVec3Def(), TEXT("Velocity GradientY")));
+		Sig.Inputs.Add(FNiagaraVariable(FNiagaraTypeDefinition::GetVec3Def(), TEXT("Velocity GradientZ")));
+		Sig.Outputs.Add(FNiagaraVariable(FNiagaraTypeDefinition::GetBoolDef(), TEXT("Build Status")));
+
+		OutFunctions.Add(Sig);
+	}
+	{
+		FNiagaraFunctionSignature Sig;
+		Sig.Name = SampleVelocityFieldName;
+		Sig.bMemberFunction = true;
+		Sig.bRequiresContext = false;
+		Sig.Inputs.Add(FNiagaraVariable(FNiagaraTypeDefinition(GetClass()), TEXT("Pressure Grid")));
+		Sig.Inputs.Add(FNiagaraVariable(FNiagaraTypeDefinition::GetVec3Def(), TEXT("Node Position")));
+		Sig.Outputs.Add(FNiagaraVariable(FNiagaraTypeDefinition::GetVec3Def(), TEXT("Node Velocity")));
+		Sig.Outputs.Add(FNiagaraVariable(FNiagaraTypeDefinition::GetFloatDef(), TEXT("Node Density")));
+		Sig.Outputs.Add(FNiagaraVariable(FNiagaraTypeDefinition::GetVec3Def(), TEXT("Velocity GradientX")));
+		Sig.Outputs.Add(FNiagaraVariable(FNiagaraTypeDefinition::GetVec3Def(), TEXT("Velocity GradientY")));
+		Sig.Outputs.Add(FNiagaraVariable(FNiagaraTypeDefinition::GetVec3Def(), TEXT("Velocity GradientZ")));
+
+		OutFunctions.Add(Sig);
+	}
+	{
+		FNiagaraFunctionSignature Sig;
+		Sig.Name = ProjectVelocityFieldName;
+		Sig.bMemberFunction = true;
+		Sig.bRequiresContext = false;
+		Sig.Inputs.Add(FNiagaraVariable(FNiagaraTypeDefinition(GetClass()), TEXT("Pressure Grid")));
+		Sig.Inputs.Add(FNiagaraVariable(FNiagaraTypeDefinition::GetIntDef(), TEXT("Grid Cell")));
+		Sig.Outputs.Add(FNiagaraVariable(FNiagaraTypeDefinition::GetBoolDef(), TEXT("Project Status")));
+		Sig.Outputs.Add(FNiagaraVariable(FNiagaraTypeDefinition::GetFloatDef(), TEXT("Cell Divergence")));
+		OutFunctions.Add(Sig);
+	}
+	{
+		FNiagaraFunctionSignature Sig;
+		Sig.Name = GetCellPositionName;
+		Sig.bMemberFunction = true;
+		Sig.bRequiresContext = false;
+		Sig.Inputs.Add(FNiagaraVariable(FNiagaraTypeDefinition(GetClass()), TEXT("Pressure Grid")));
+		Sig.Inputs.Add(FNiagaraVariable(FNiagaraTypeDefinition::GetIntDef(), TEXT("Grid Cell")));
+		Sig.Outputs.Add(FNiagaraVariable(FNiagaraTypeDefinition::GetVec3Def(), TEXT("Grid Position")));
+
+		OutFunctions.Add(Sig);
+	}
+	{
+		FNiagaraFunctionSignature Sig;
+		Sig.Name = SetSolidBoundaryName;
+		Sig.bMemberFunction = true;
+		Sig.bRequiresContext = false;
+		Sig.Inputs.Add(FNiagaraVariable(FNiagaraTypeDefinition(GetClass()), TEXT("Pressure Grid")));
+		Sig.Inputs.Add(FNiagaraVariable(FNiagaraTypeDefinition::GetIntDef(), TEXT("Grid Cell")));
+		Sig.Inputs.Add(FNiagaraVariable(FNiagaraTypeDefinition::GetFloatDef(), TEXT("Cell Distance")));
+		Sig.Inputs.Add(FNiagaraVariable(FNiagaraTypeDefinition::GetVec3Def(), TEXT("Cell Velocity")));
+		Sig.Outputs.Add(FNiagaraVariable(FNiagaraTypeDefinition::GetBoolDef(), TEXT("Boundary Status")));
+
+		OutFunctions.Add(Sig);
+	}
+	{
+		FNiagaraFunctionSignature Sig;
+		Sig.Name = ComputeBoundaryWeightsName;
+		Sig.bMemberFunction = true;
+		Sig.bRequiresContext = false;
+		Sig.Inputs.Add(FNiagaraVariable(FNiagaraTypeDefinition(GetClass()), TEXT("Pressure Grid")));
+		Sig.Inputs.Add(FNiagaraVariable(FNiagaraTypeDefinition::GetIntDef(), TEXT("Grid Cell")));
+		Sig.Outputs.Add(FNiagaraVariable(FNiagaraTypeDefinition::GetBoolDef(), TEXT("Weights Status")));
+
+		OutFunctions.Add(Sig);
+	}
+	{
+		FNiagaraFunctionSignature Sig;
+		Sig.Name = TransferCellVelocityName;
+		Sig.bMemberFunction = true;
+		Sig.bRequiresContext = false;
+		Sig.Inputs.Add(FNiagaraVariable(FNiagaraTypeDefinition(GetClass()), TEXT("Pressure Grid")));
+		Sig.Inputs.Add(FNiagaraVariable(FNiagaraTypeDefinition::GetIntDef(), TEXT("Grid Cell")));
+		Sig.Outputs.Add(FNiagaraVariable(FNiagaraTypeDefinition::GetBoolDef(), TEXT("Transfer Status")));
+
+		OutFunctions.Add(Sig);
+	}
+	{
+		FNiagaraFunctionSignature Sig;
+		Sig.Name = BuildGridTopologyName;
+		Sig.bMemberFunction = true;
+		Sig.bRequiresContext = false;
+		Sig.Inputs.Add(FNiagaraVariable(FNiagaraTypeDefinition(GetClass()), TEXT("Pressure Grid")));
+		Sig.Inputs.Add(FNiagaraVariable(FNiagaraTypeDefinition::GetVec3Def(), TEXT("Grid Center")));
+		Sig.Inputs.Add(FNiagaraVariable(FNiagaraTypeDefinition::GetVec3Def(), TEXT("Grid Extent")));
+		Sig.Outputs.Add(FNiagaraVariable(FNiagaraTypeDefinition::GetVec3Def(), TEXT("Grid Origin")));
+		Sig.Outputs.Add(FNiagaraVariable(FNiagaraTypeDefinition::GetFloatDef(), TEXT("Grid Length")));
+
+		OutFunctions.Add(Sig);
+	}
+	{
+		FNiagaraFunctionSignature Sig;
+		Sig.Name = UpdateGridTransformName;
+		Sig.bMemberFunction = true;
+		Sig.bRequiresContext = false;
+		Sig.Inputs.Add(FNiagaraVariable(FNiagaraTypeDefinition(GetClass()), TEXT("Pressure Grid")));
+		Sig.Inputs.Add(FNiagaraVariable(FNiagaraTypeDefinition::GetMatrix4Def(), TEXT("Grid Transform")));
+		Sig.Outputs.Add(FNiagaraVariable(FNiagaraTypeDefinition::GetBoolDef(), TEXT("Transform Status")));
+
+		OutFunctions.Add(Sig);
+	}
+}
+
+DEFINE_NDI_DIRECT_FUNC_BINDER(UNiagaraDataInterfacePressureGrid, BuildVelocityField);
+DEFINE_NDI_DIRECT_FUNC_BINDER(UNiagaraDataInterfacePressureGrid, ProjectVelocityField);
+DEFINE_NDI_DIRECT_FUNC_BINDER(UNiagaraDataInterfacePressureGrid, GetCellPosition);
+DEFINE_NDI_DIRECT_FUNC_BINDER(UNiagaraDataInterfacePressureGrid, SetSolidBoundary);
+DEFINE_NDI_DIRECT_FUNC_BINDER(UNiagaraDataInterfacePressureGrid, TransferCellVelocity);
+DEFINE_NDI_DIRECT_FUNC_BINDER(UNiagaraDataInterfacePressureGrid, ComputeBoundaryWeights);
+DEFINE_NDI_DIRECT_FUNC_BINDER(UNiagaraDataInterfacePressureGrid, SampleVelocityField);
+DEFINE_NDI_DIRECT_FUNC_BINDER(UNiagaraDataInterfacePressureGrid, BuildGridTopology);
+DEFINE_NDI_DIRECT_FUNC_BINDER(UNiagaraDataInterfacePressureGrid, UpdateGridTransform);
+
+void UNiagaraDataInterfacePressureGrid::GetVMExternalFunction(const FVMExternalFunctionBindingInfo& BindingInfo, void* InstanceData, FVMExternalFunction &OutFunc)
+{
+	if (BindingInfo.Name == BuildVelocityFieldName)
+	{
+		check(BindingInfo.GetNumInputs() == 19 && BindingInfo.GetNumOutputs() == 1);
+		NDI_FUNC_BINDER(UNiagaraDataInterfacePressureGrid, BuildVelocityField)::Bind(this, OutFunc);
+	}
+	else if (BindingInfo.Name == ProjectVelocityFieldName)
+	{
+		check(BindingInfo.GetNumInputs() == 2 && BindingInfo.GetNumOutputs() == 2);
+		NDI_FUNC_BINDER(UNiagaraDataInterfacePressureGrid, ProjectVelocityField)::Bind(this, OutFunc);
+	}
+	else if (BindingInfo.Name == GetCellPositionName)
+	{
+		check(BindingInfo.GetNumInputs() == 2 && BindingInfo.GetNumOutputs() == 3);
+		NDI_FUNC_BINDER(UNiagaraDataInterfacePressureGrid, GetCellPosition)::Bind(this, OutFunc);
+	}
+	else if (BindingInfo.Name == SampleVelocityFieldName)
+	{
+		check(BindingInfo.GetNumInputs() == 4 && BindingInfo.GetNumOutputs() == 13);
+		NDI_FUNC_BINDER(UNiagaraDataInterfacePressureGrid, SampleVelocityField)::Bind(this, OutFunc);
+	}
+	else if (BindingInfo.Name == SetSolidBoundaryName)
+	{
+		check(BindingInfo.GetNumInputs() == 6 && BindingInfo.GetNumOutputs() == 1);
+		NDI_FUNC_BINDER(UNiagaraDataInterfacePressureGrid, SetSolidBoundary)::Bind(this, OutFunc);
+	}
+	else if (BindingInfo.Name == ComputeBoundaryWeightsName)
+	{
+		check(BindingInfo.GetNumInputs() == 2 && BindingInfo.GetNumOutputs() == 1);
+		NDI_FUNC_BINDER(UNiagaraDataInterfacePressureGrid, ComputeBoundaryWeights)::Bind(this, OutFunc);
+	}
+	else if (BindingInfo.Name ==TransferCellVelocityName)
+	{
+		check(BindingInfo.GetNumInputs() == 2 && BindingInfo.GetNumOutputs() == 1);
+		NDI_FUNC_BINDER(UNiagaraDataInterfacePressureGrid, TransferCellVelocity)::Bind(this, OutFunc);
+	}
+	else if (BindingInfo.Name == BuildGridTopologyName)
+	{
+		check(BindingInfo.GetNumInputs() == 7 && BindingInfo.GetNumOutputs() == 4);
+		NDI_FUNC_BINDER(UNiagaraDataInterfacePressureGrid, BuildGridTopology)::Bind(this, OutFunc);
+	}
+	else if (BindingInfo.Name == UpdateGridTransformName)
+	{
+		check(BindingInfo.GetNumInputs() == 17 && BindingInfo.GetNumOutputs() == 1);
+		NDI_FUNC_BINDER(UNiagaraDataInterfacePressureGrid, UpdateGridTransform)::Bind(this, OutFunc);
+	}
+}
+
+void UNiagaraDataInterfacePressureGrid::BuildVelocityField(FVectorVMContext& Context)
+{
+	// @todo : implement function for cpu 
+}
+
+void UNiagaraDataInterfacePressureGrid::ProjectVelocityField(FVectorVMContext& Context)
+{
+	// @todo : implement function for cpu 
+}
+
+void UNiagaraDataInterfacePressureGrid::GetCellPosition(FVectorVMContext& Context)
+{
+	// @todo : implement function for cpu 
+}
+
+void UNiagaraDataInterfacePressureGrid::SampleVelocityField(FVectorVMContext& Context)
+{
+	// @todo : implement function for cpu 
+}
+
+void UNiagaraDataInterfacePressureGrid::ComputeBoundaryWeights(FVectorVMContext& Context)
+{
+	// @todo : implement function for cpu 
+}
+
+void UNiagaraDataInterfacePressureGrid::SetSolidBoundary(FVectorVMContext& Context)
+{
+	// @todo : implement function for cpu 
+}
+
+void UNiagaraDataInterfacePressureGrid::TransferCellVelocity(FVectorVMContext& Context)
+{
+	// @todo : implement function for cpu 
+}
+
+void UNiagaraDataInterfacePressureGrid::BuildGridTopology(FVectorVMContext& Context)
+{
+	VectorVM::FExternalFuncInputHandler<float> CenterXParam(Context);
+	VectorVM::FExternalFuncInputHandler<float> CenterYParam(Context);
+	VectorVM::FExternalFuncInputHandler<float> CenterZParam(Context);
+	VectorVM::FExternalFuncInputHandler<float> ExtentXParam(Context);
+	VectorVM::FExternalFuncInputHandler<float> ExtentYParam(Context);
+	VectorVM::FExternalFuncInputHandler<float> ExtentZParam(Context);
+
+	VectorVM::FUserPtrHandler<FNDIPressureGridData> InstData(Context);
+	VectorVM::FExternalFuncRegisterHandler<float> OutGridOriginX(Context);
+	VectorVM::FExternalFuncRegisterHandler<float> OutGridOriginY(Context);
+	VectorVM::FExternalFuncRegisterHandler<float> OutGridOriginZ(Context);
+	VectorVM::FExternalFuncRegisterHandler<float> OutGridLength(Context);
+
+	for (int32 InstanceIdx = 0; InstanceIdx < Context.NumInstances; ++InstanceIdx)
+	{
+		const FVector GridCenter(*CenterXParam.GetDestAndAdvance(), *CenterYParam.GetDestAndAdvance(), *CenterZParam.GetDestAndAdvance());
+		const FVector GridExtent(*ExtentXParam.GetDestAndAdvance(), *ExtentYParam.GetDestAndAdvance(), *ExtentZParam.GetDestAndAdvance());
+
+		const FVector GridLengths(2.0 * GridExtent.X / (InstData->GridSize.X - 1),
+								  2.0 * GridExtent.Y / (InstData->GridSize.Y - 1),
+								  2.0 * GridExtent.Z / (InstData->GridSize.Z - 1));
+		const float MaxLength = GridLengths.GetMax();
+
+		const FVector RegularExtent((InstData->GridSize.X - 1)*MaxLength, 
+								    (InstData->GridSize.Y - 1)*MaxLength, 
+									(InstData->GridSize.Z - 1)*MaxLength);
+		const FVector BoxOrigin = GridCenter - 0.5 * RegularExtent;
+		InstData->GridOrigin = FVector4(BoxOrigin.X, BoxOrigin.Y, BoxOrigin.Z, MaxLength);
+
+		*OutGridOriginX.GetDestAndAdvance() = BoxOrigin.X;
+		*OutGridOriginY.GetDestAndAdvance() = BoxOrigin.Y;
+		*OutGridOriginZ.GetDestAndAdvance() = BoxOrigin.Z;
+		*OutGridLength.GetDestAndAdvance() = MaxLength;
+		//UE_LOG(LogPressureGrid, Warning, TEXT("Get Grid Extent : %f %f %f | Grid Lengths : %f %f %f | Regular Extent : %f %f %f | Box Origin : %f %f %f"), 
+		//				GridExtent.X, GridExtent.Y, GridExtent.Z, GridLengths.X, GridLengths.Y, GridLengths.Z, RegularExtent.X, RegularExtent.Y, RegularExtent.Z, InstData->GridOrigin.X, InstData->GridOrigin.Y, InstData->GridOrigin.Z);
+	}
+
+	
+}
+
+void UNiagaraDataInterfacePressureGrid::UpdateGridTransform(FVectorVMContext& Context)
+{
+	VectorVM::FExternalFuncInputHandler<float> Out00(Context);
+	VectorVM::FExternalFuncInputHandler<float> Out01(Context);
+	VectorVM::FExternalFuncInputHandler<float> Out02(Context);
+	VectorVM::FExternalFuncInputHandler<float> Out03(Context);
+
+	VectorVM::FExternalFuncInputHandler<float> Out10(Context);
+	VectorVM::FExternalFuncInputHandler<float> Out11(Context);
+	VectorVM::FExternalFuncInputHandler<float> Out12(Context);
+	VectorVM::FExternalFuncInputHandler<float> Out13(Context);
+
+	VectorVM::FExternalFuncInputHandler<float> Out20(Context);
+	VectorVM::FExternalFuncInputHandler<float> Out21(Context);
+	VectorVM::FExternalFuncInputHandler<float> Out22(Context);
+	VectorVM::FExternalFuncInputHandler<float> Out23(Context);
+
+	VectorVM::FExternalFuncInputHandler<float> Out30(Context);
+	VectorVM::FExternalFuncInputHandler<float> Out31(Context);
+	VectorVM::FExternalFuncInputHandler<float> Out32(Context);
+	VectorVM::FExternalFuncInputHandler<float> Out33(Context);
+
+	VectorVM::FUserPtrHandler<FNDIPressureGridData> InstData(Context);
+	VectorVM::FExternalFuncRegisterHandler<bool> OutTransformStatus(Context);
+
+	for (int32 InstanceIdx = 0; InstanceIdx < Context.NumInstances; ++InstanceIdx)
+	{
+		FMatrix Transform;
+		Transform.M[0][0] = *Out00.GetDestAndAdvance();
+		Transform.M[0][1] = *Out01.GetDestAndAdvance();
+		Transform.M[0][2] = *Out02.GetDestAndAdvance();
+		Transform.M[0][3] = *Out03.GetDestAndAdvance();
+
+		Transform.M[1][0] = *Out10.GetDestAndAdvance();
+		Transform.M[1][1] = *Out11.GetDestAndAdvance();
+		Transform.M[1][2] = *Out12.GetDestAndAdvance();
+		Transform.M[1][3] = *Out13.GetDestAndAdvance();
+
+		Transform.M[2][0] = *Out20.GetDestAndAdvance();
+		Transform.M[2][1] = *Out21.GetDestAndAdvance();
+		Transform.M[2][2] = *Out22.GetDestAndAdvance();
+		Transform.M[2][3] = *Out23.GetDestAndAdvance();
+
+		Transform.M[3][0] = *Out30.GetDestAndAdvance();
+		Transform.M[3][1] = *Out31.GetDestAndAdvance();
+		Transform.M[3][2] = *Out32.GetDestAndAdvance();
+		Transform.M[3][3] = *Out33.GetDestAndAdvance();
+
+		InstData->WorldTransform = Transform;
+		InstData->WorldInverse = Transform.Inverse();
+
+		*OutTransformStatus.GetDestAndAdvance() = true;
+	}
+	//UE_LOG(LogPressureGrid, Warning, TEXT("Get Grid Transform : %s"), *InstData->WorldTransform.ToString() );
+}
+
+bool UNiagaraDataInterfacePressureGrid::GetFunctionHLSL(const FName& DefinitionFunctionName, FString InstanceFunctionName, FNiagaraDataInterfaceGPUParamInfo& ParamInfo, FString& OutHLSL)
+{
+	FNDIPressureGridParametersName ParamNames(ParamInfo.DataInterfaceHLSLSymbol);
+
+	TMap<FString, FStringFormatArg> ArgsSample = {
+		{TEXT("InstanceFunctionName"), InstanceFunctionName},
+		{TEXT("GridCurrentBufferName"), ParamNames.GridCurrentBufferName},
+		{TEXT("GridDestinationBufferName"), ParamNames.GridDestinationBufferName},
+		{TEXT("GridOriginName"), ParamNames.GridOriginName},
+		{TEXT("GridSizeName"), ParamNames.GridSizeName},
+		{TEXT("WorldTransformName"), ParamNames.WorldTransformName},
+		{TEXT("WorldInverseName"), ParamNames.WorldInverseName},
+		{TEXT("PressureGridContextName"), TEXT("DIPRESSUREGRID_MAKE_CONTEXT(") + ParamInfo.DataInterfaceHLSLSymbol + TEXT(")")},
+	};
+
+	if (DefinitionFunctionName == BuildVelocityFieldName)
+	{
+		static const TCHAR *FormatSample = TEXT(R"(
+				void {InstanceFunctionName} (in int StrandsSize, in float3 NodePosition, in float NodeMass, in float3 NodeVelocity, in float3 VelocityGradientX, in float3 VelocityGradientY, in float3 VelocityGradientZ, out bool OutBuildStatus)
+				{
+					{PressureGridContextName} DIPressureGrid_BuildVelocityField(DIContext,StrandsSize,NodePosition,NodeMass,NodeVelocity,VelocityGradientX,VelocityGradientY,VelocityGradientZ,OutBuildStatus);
+				}
+				)");
+		OutHLSL += FString::Format(FormatSample, ArgsSample);
+		return true;
+	}
+	else if (DefinitionFunctionName == SampleVelocityFieldName)
+	{
+		static const TCHAR *FormatSample = TEXT(R"(
+				void {InstanceFunctionName} (in float3 NodePosition, out float3 OutGridVelocity, out float OutGridDensity, out float3 OutGridGradientX, out float3 OutGridGradientY, out float3 OutGridGradientZ )
+				{
+					{PressureGridContextName} DIPressureGrid_SampleVelocityField(DIContext,NodePosition,OutGridVelocity,OutGridDensity,OutGridGradientX,OutGridGradientY,OutGridGradientZ);
+				}
+				)");
+		OutHLSL += FString::Format(FormatSample, ArgsSample);
+		return true;
+	}
+	else if (DefinitionFunctionName == ProjectVelocityFieldName)
+	{
+		static const TCHAR *FormatSample = TEXT(R"(
+				void {InstanceFunctionName} (in int GridCell, out bool OutProjectStatus, out float OutCellDivergence)
+				{
+					{PressureGridContextName} DIPressureGrid_ProjectVelocityField(DIContext,GridCell,OutProjectStatus,OutCellDivergence);
+				}
+				)");
+		OutHLSL += FString::Format(FormatSample, ArgsSample);
+		return true;
+	}
+	else if (DefinitionFunctionName == GetCellPositionName)
+	{
+		static const TCHAR *FormatSample = TEXT(R"(
+				void {InstanceFunctionName} (in int GridCell, out float3 OutGridPosition)
+				{
+					{PressureGridContextName} DIPressureGrid_GetCellPosition(DIContext,GridCell,OutGridPosition);
+				}
+				)");
+		OutHLSL += FString::Format(FormatSample, ArgsSample);
+		return true;
+	}
+	else if (DefinitionFunctionName == SetSolidBoundaryName)
+	{
+		static const TCHAR *FormatSample = TEXT(R"(
+				void {InstanceFunctionName} (in int GridCell, in float SolidDistance, in float3 SolidVelocity, out bool OutBoundaryStatus)
+				{
+					{PressureGridContextName} DIPressureGrid_SetSolidBoundary(DIContext,GridCell,SolidDistance,SolidVelocity,OutBoundaryStatus);
+				}
+				)");
+		OutHLSL += FString::Format(FormatSample, ArgsSample);
+		return true;
+	}
+	else if (DefinitionFunctionName == ComputeBoundaryWeightsName)
+	{
+		static const TCHAR *FormatSample = TEXT(R"(
+				void {InstanceFunctionName} (in int GridCell, out bool OutWeightsStatus)
+				{
+					{PressureGridContextName} DIPressureGrid_ComputeBoundaryWeights(DIContext,GridCell,OutWeightsStatus);
+				}
+				)");
+		OutHLSL += FString::Format(FormatSample, ArgsSample);
+		return true;
+	}
+	else if (DefinitionFunctionName == TransferCellVelocityName)
+	{
+		static const TCHAR *FormatSample = TEXT(R"(
+				void {InstanceFunctionName} (in int GridCell, out bool OutTransferStatus)
+				{
+					{PressureGridContextName} DIPressureGrid_TransferCellVelocity(DIContext,GridCell,OutTransferStatus);
+				}
+				)");
+		OutHLSL += FString::Format(FormatSample, ArgsSample);
+		return true;
+	}
+	
+
+	OutHLSL += TEXT("\n");
+	return false;
+}
+
+void UNiagaraDataInterfacePressureGrid::GetCommonHLSL(FString& OutHLSL)
+{
+	OutHLSL += TEXT("#include \"/Plugin/Experimental/HairStrands/Private/NiagaraDataInterfacePressureGrid.ush\"\n");
+}
+
+void UNiagaraDataInterfacePressureGrid::GetParameterDefinitionHLSL(FNiagaraDataInterfaceGPUParamInfo& ParamInfo, FString& OutHLSL)
+{
+	OutHLSL += TEXT("DIPRESSUREGRID_DECLARE_CONSTANTS(") + ParamInfo.DataInterfaceHLSLSymbol + TEXT(")\n");
+}
+
+void UNiagaraDataInterfacePressureGrid::ProvidePerInstanceDataForRenderThread(void* DataForRenderThread, void* PerInstanceData, const FNiagaraSystemInstanceID& SystemInstance)
+{
+	FNDIPressureGridData* GameThreadData = static_cast<FNDIPressureGridData*>(PerInstanceData);
+	FNDIPressureGridData* RenderThreadData = static_cast<FNDIPressureGridData*>(DataForRenderThread);
+
+	RenderThreadData->WorldTransform = GameThreadData->WorldTransform;
+	RenderThreadData->GridOrigin = GameThreadData->GridOrigin;
+}
+
+FNiagaraDataInterfaceParametersCS*
+UNiagaraDataInterfacePressureGrid::ConstructComputeParameters() const
+{
+	return new FNDIPressureGridParametersCS();
+}
+
+void FNDIPressureGridProxy::PreStage(FRHICommandList& RHICmdList, const FNiagaraDataInterfaceSetArgs& Context)
+{
+	FNDIPressureGridData* ProxyData =
+		SystemInstancesToProxyData.Find(Context.SystemInstance);
+
+	if (ProxyData != nullptr)
+	{
+		if (!Context.IsIterationStage)
+		{
+			ProxyData->DestinationGridBuffer->ClearBuffers(RHICmdList);
+		}
+		else
+		{
+			ENQUEUE_RENDER_COMMAND(CopyPressureGrid)(
+				[ProxyData](FRHICommandListImmediate& RHICmdList)
+			{
+				FRHICopyTextureInfo CopyInfo;
+				RHICmdList.CopyTexture(ProxyData->CurrentGridBuffer->GridDataBuffer.Buffer,
+					ProxyData->DestinationGridBuffer->GridDataBuffer.Buffer, CopyInfo);
+			});
+		}
+	}
+}
+
+void FNDIPressureGridProxy::PostStage(FRHICommandList& RHICmdList, const FNiagaraDataInterfaceSetArgs& Context)
+{
+	FNDIPressureGridData* ProxyData =
+		SystemInstancesToProxyData.Find(Context.SystemInstance);
+
+	if (Context.IsOutputStage)
+	{
+		if (ProxyData != nullptr)
+		{
+			ProxyData->SwapBuffers();
+		}
+	}
+}
+
+void FNDIPressureGridProxy::ResetData(FRHICommandList& RHICmdList, const FNiagaraDataInterfaceSetArgs& Context)
+{
+	FNDIPressureGridData* ProxyData = SystemInstancesToProxyData.Find(Context.SystemInstance);
+
+	if (ProxyData != nullptr)
+	{
+		ProxyData->CurrentGridBuffer->ClearBuffers(RHICmdList);
+		ProxyData->DestinationGridBuffer->ClearBuffers(RHICmdList);
+	}
+}
+
+
+#undef LOCTEXT_NAMESPACE
\ No newline at end of file
diff --git a/Engine/Plugins/Experimental/HairStrands/Source/HairStrandsNiagara/Public/NiagaraDataInterfaceHairStrands.h b/Engine/Plugins/Experimental/HairStrands/Source/HairStrandsNiagara/Public/NiagaraDataInterfaceHairStrands.h
--- a/Engine/Plugins/Experimental/HairStrands/Source/HairStrandsNiagara/Public/NiagaraDataInterfaceHairStrands.h
+++ b/Engine/Plugins/Experimental/HairStrands/Source/HairStrandsNiagara/Public/NiagaraDataInterfaceHairStrands.h
@@ -22,25 +22,11 @@ enum class EHairStrandsSize : uint8
 	Size32 = 0x20 UMETA(DisplatName = "32")
 };
 
-/** Structure storing the grid size*/
-struct FGridSize
-{
-	FGridSize(const uint32 InGridSizeX = 0, const uint32 InGridSizeY = 0, 
-		      const uint32 InGridSizeZ = 0, const uint32 InGridSizeW = 0) : GridSizeX(InGridSizeX), 
-		GridSizeY(InGridSizeY), GridSizeZ(InGridSizeZ), GridSizeW(InGridSizeW) {}
-
-	uint32 GridSizeX;
-	uint32 GridSizeY;
-	uint32 GridSizeZ;
-	uint32 GridSizeW;
-};
-
-
 /** Render buffers that will be used in hlsl functions */
 struct FNDICollisionGridBuffer : public FRenderResource
 {
 	/** Set the grid size */
-	void SetGridSize(const FGridSize GridSize);
+	void SetGridSize(const FUintVector4 GridSize);
 
 	/** Init the buffer */
 	virtual void InitRHI() override;
@@ -58,7 +44,7 @@ struct FNDICollisionGridBuffer : public FRenderResource
 	FTextureRWBuffer3D GridDataBuffer;
 
 	/** Grid size that will be used for the collision*/
-	FGridSize GridSize;
+	FUintVector4 GridSize;
 };
 
 /** Render buffers that will be used in hlsl functions */
@@ -117,7 +103,7 @@ struct FNDIHairStrandsData
 	FVector4 GridOrigin;
 
 	/** Grid Size */
-	FGridSize GridSize;
+	FUintVector4 GridSize;
 
 	/** Strands Gpu buffer */
 	FNDIHairStrandsBuffer* HairStrandsBuffer;
@@ -133,9 +119,6 @@ struct FNDIHairStrandsData
 
 	/** Pointer to the destination buffer */
 	FNDICollisionGridBuffer* DestinationGridBuffer;
-
-	/** Sorted index particle buffers*/
-	//FParticleSortBuffers SortedParticleBuffers;
 };
 
 /** Data Interface for the strand base */
@@ -296,6 +279,60 @@ public:
 	/** Project the collision grid */
 	void ProjectCollisionGrid(FVectorVMContext& Context);
 
+	/** Get the bounding box center */
+	void GetBoxCenter(FVectorVMContext& Context);
+
+	/** Get the bounding box extent */
+	void GetBoxExtent(FVectorVMContext& Context);
+
+	/** Setup the stretch spring material */
+	void SetupStretchSpringMaterial(FVectorVMContext& Context);
+
+	/** Solve the stretch spring material */
+	void SolveStretchSpringMaterial(FVectorVMContext& Context);
+
+	/** Project the stretch spring material */
+	void ProjectStretchSpringMaterial(FVectorVMContext& Context);
+
+	/** Setup the bend spring material */
+	void SetupBendSpringMaterial(FVectorVMContext& Context);
+
+	/** Solve the bend spring material */
+	void SolveBendSpringMaterial(FVectorVMContext& Context);
+
+	/** Project the bend spring material */
+	void ProjectBendSpringMaterial(FVectorVMContext& Context);
+
+	/** Setup the stretch rod material */
+	void SetupStretchRodMaterial(FVectorVMContext& Context);
+
+	/** Solve the stretch rod material */
+	void SolveStretchRodMaterial(FVectorVMContext& Context);
+
+	/** Project the stretch rod material */
+	void ProjectStretchRodMaterial(FVectorVMContext& Context);
+
+	/** Setup the bend rod material */
+	void SetupBendRodMaterial(FVectorVMContext& Context);
+
+	/** Solve the bend rod material */
+	void SolveBendRodMaterial(FVectorVMContext& Context);
+
+	/** Project the bend rod material */
+	void ProjectBendRodMaterial(FVectorVMContext& Context);
+
+	/** Solve the static collision constraint */
+	void SolveStaticCollisionConstraint(FVectorVMContext& Context);
+
+	/** Project the static collision constraint */
+	void ProjectStaticCollisionConstraint(FVectorVMContext& Context);
+
+	/** Compute the rest direction*/
+	void ComputeRestDirection(FVectorVMContext& Context);
+
+	/** Update the node orientation to match the bishop frame*/
+	void UpdateNodeOrientation(FVectorVMContext& Context);
+
 	/** Name of the world transform */
 	static const FString WorldTransformName;
 
@@ -360,7 +397,7 @@ struct FNDIHairStrandsProxy : public FNiagaraDataInterfaceProxy
 
 	/** Initialize the Proxy data strands buffer */
 	void InitializePerInstanceData(const FNiagaraSystemInstanceID& SystemInstance, FNDIHairStrandsBuffer* StrandsBuffer, FNDICollisionGridBuffer* CollisionGridBufferA, FNDICollisionGridBuffer* CollisionGridBufferB, const uint32 NumStrands, const uint8 StrandSize,
-		const float StrandDensity, const float RootThickness, const float TipThickness, const FVector4& GridOrigin, const FGridSize& GridSize);
+		const float StrandDensity, const float RootThickness, const float TipThickness, const FVector4& GridOrigin, const FUintVector4& GridSize);
 
 	/** Destroy the proxy data if necessary */
 	void DestroyPerInstanceData(NiagaraEmitterInstanceBatcher* Batcher, const FNiagaraSystemInstanceID& SystemInstance);
diff --git a/Engine/Plugins/Experimental/HairStrands/Source/HairStrandsNiagara/Public/NiagaraDataInterfacePressureGrid.h b/Engine/Plugins/Experimental/HairStrands/Source/HairStrandsNiagara/Public/NiagaraDataInterfacePressureGrid.h
new file mode 100644
--- /dev/null
+++ b/Engine/Plugins/Experimental/HairStrands/Source/HairStrandsNiagara/Public/NiagaraDataInterfacePressureGrid.h
@@ -0,0 +1,185 @@
+// Copyright 1998-2019 Epic Games, Inc. All Rights Reserved.
+
+#pragma once
+
+#include "NiagaraDataInterface.h"
+#include "NiagaraDataInterfaceRW.h"
+#include "NiagaraCommon.h"
+#include "VectorVM.h"
+#include "NiagaraDataInterfacePressureGrid.generated.h"
+
+
+/** Render buffers that will be used in hlsl functions */
+struct FNDIPressureGridBuffer : public FRenderResource
+{
+	/** Set the grid size */
+	void SetGridSize(const FUintVector4 GridSize);
+
+	/** Init the buffer */
+	virtual void InitRHI() override;
+
+	/** Release the buffer */
+	virtual void ReleaseRHI() override;
+
+	/** Clear all UAV*/
+	void ClearBuffers(FRHICommandList& RHICmdList);
+
+	/** Get the resource name */
+	virtual FString GetFriendlyName() const override { return TEXT("FNDIPressureGridBuffer"); }
+
+	/** Grid data texture */
+	FTextureRWBuffer3D GridDataBuffer;
+
+	/** Grid size that will be used for the collision*/
+	FUintVector4 GridSize;
+};
+
+/** Data stored per strand base instance*/
+struct FNDIPressureGridData
+{
+	/** Swap the current and the destination data */
+	void SwapBuffers();
+
+	/** Grid Origin */
+	FVector4 GridOrigin;
+
+	/** Grid Size */
+	FUintVector4 GridSize;
+	
+	/** World Transform */
+	FMatrix WorldTransform;
+
+	/** Inverse world transform */
+	FMatrix WorldInverse;
+
+	/** Pointer to the current buffer */
+	FNDIPressureGridBuffer* CurrentGridBuffer;
+
+	/** Pointer to the destination buffer */
+	FNDIPressureGridBuffer* DestinationGridBuffer;
+};
+
+/** Data Interface for the strand base */
+UCLASS(EditInlineNew, Category = "Grid", meta = (DisplayName = "Pressure Grid"))
+class HAIRSTRANDSNIAGARA_API UNiagaraDataInterfacePressureGrid : public UNiagaraDataInterfaceRWBase
+{
+	GENERATED_UCLASS_BODY()
+
+public:
+
+	/** Grid size along the X axis. */
+	UPROPERTY(EditAnywhere, Category = "Spawn")
+	FIntVector GridSize;
+
+	/** Grid size along the Y axis. */
+	UPROPERTY(EditAnywhere, Category = "Spawn")
+	FVector GridOrigin;
+
+	/** Grid size along the Z axis. */
+	UPROPERTY(EditAnywhere, Category = "Spawn")
+	float GridLength;
+
+	/** UObject Interface */
+	virtual void PostInitProperties() override;
+
+	/** UNiagaraDataInterface Interface */
+	virtual void GetFunctions(TArray<FNiagaraFunctionSignature>& OutFunctions) override;
+	virtual void GetVMExternalFunction(const FVMExternalFunctionBindingInfo& BindingInfo, void* InstanceData, FVMExternalFunction &OutFunc) override;
+	virtual bool CanExecuteOnTarget(ENiagaraSimTarget Target) const override { return Target == ENiagaraSimTarget::GPUComputeSim; }
+	virtual bool InitPerInstanceData(void* PerInstanceData, FNiagaraSystemInstance* SystemInstance) override;
+	virtual void DestroyPerInstanceData(void* PerInstanceData, FNiagaraSystemInstance* SystemInstance)override;
+	virtual bool PerInstanceTick(void* PerInstanceData, FNiagaraSystemInstance* SystemInstance, float DeltaSeconds) override;
+	virtual int32 PerInstanceDataSize()const override { return sizeof(FNDIPressureGridData); }
+	virtual bool Equals(const UNiagaraDataInterface* Other) const override;
+
+	/** GPU simulation  functionality */
+	virtual bool GetFunctionHLSL(const FName&  DefinitionFunctionName, FString InstanceFunctionName, FNiagaraDataInterfaceGPUParamInfo& ParamInfo, FString& OutHLSL) override;
+	virtual void GetParameterDefinitionHLSL(FNiagaraDataInterfaceGPUParamInfo& ParamInfo, FString& OutHLSL) override;
+	virtual FNiagaraDataInterfaceParametersCS* ConstructComputeParameters() const override;
+	virtual void ProvidePerInstanceDataForRenderThread(void* DataForRenderThread, void* PerInstanceData, const FNiagaraSystemInstanceID& SystemInstance) override;
+	virtual void GetCommonHLSL(FString& OutHLSL) override;
+
+	/** Build the velocity field */
+	void BuildVelocityField(FVectorVMContext& Context);
+
+	/** Sample the grid */
+	void SampleVelocityField(FVectorVMContext& Context);
+
+	/** Project the velocity field to be divergence free */
+	void ProjectVelocityField(FVectorVMContext& Context);
+
+	/** Compute the cell position*/
+	void GetCellPosition(FVectorVMContext& Context);
+
+	/** Transfer the cell velocity */
+	void TransferCellVelocity(FVectorVMContext& Context);
+
+	/** Set the solid boundary */
+	void SetSolidBoundary(FVectorVMContext& Context);
+
+	/** Compute the solid weights */
+	void ComputeBoundaryWeights(FVectorVMContext& Context);
+
+	/** Build the grid topology */
+	void BuildGridTopology(FVectorVMContext& Context);
+
+	/** Update the grid transform */
+	void UpdateGridTransform(FVectorVMContext& Context);
+
+	/** Name of the grid current buffer */
+	static const FString GridCurrentBufferName;
+
+	/** Name of the grid X velocity buffer */
+	static const FString GridDestinationBufferName;
+
+	/** Name of the grid origin  */
+	static const FString GridOriginName;
+
+	/** Name of the grid size */
+	static const FString GridSizeName;
+
+	/** Name of the world transform  */
+	static const FString WorldTransformName;
+
+	/** Name of the World transform inverse */
+	static const FString WorldInverseName;
+
+protected:
+	/** Copy one niagara DI to this */
+	virtual bool CopyToInternal(UNiagaraDataInterface* Destination) const override;
+};
+
+/** Proxy to send data to gpu */
+struct FNDIPressureGridProxy : public FNiagaraDataInterfaceProxy
+{
+	/** Destroy internal data */
+	virtual void DeferredDestroy() override;
+
+	/** Get the size of the data that will be passed to render*/
+	virtual int32 PerInstanceDataPassedToRenderThreadSize() const override { return sizeof(FNDIPressureGridData); }
+
+	/** Get the data that will be passed to render*/
+	virtual void ConsumePerInstanceDataFromGameThread(void* PerInstanceData, const FNiagaraSystemInstanceID& Instance) override;
+
+	/** Initialize the Proxy data strands buffer */
+	void InitializePerInstanceData(const FNiagaraSystemInstanceID& SystemInstance, FNDIPressureGridBuffer* CurrentGridBuffer, FNDIPressureGridBuffer* DestinationGridBuffer, const FVector4& GridOrigin, const FUintVector4& GridSize);
+
+	/** Destroy the proxy data if necessary */
+	void DestroyPerInstanceData(NiagaraEmitterInstanceBatcher* Batcher, const FNiagaraSystemInstanceID& SystemInstance);
+
+	/** Launch all pre stage functions */
+	virtual void PreStage(FRHICommandList& RHICmdList, const FNiagaraDataInterfaceSetArgs& Context) override;
+
+	/** Launch all post stage functions */
+	virtual void PostStage(FRHICommandList& RHICmdList, const FNiagaraDataInterfaceSetArgs& Context) override;
+
+	/** Reset the buffers  */
+	virtual void ResetData(FRHICommandList& RHICmdList, const FNiagaraDataInterfaceSetArgs& Context) override;
+
+	/** List of proxy data for each system instances*/
+	TMap<FNiagaraSystemInstanceID, FNDIPressureGridData> SystemInstancesToProxyData;
+
+	/** List of proxy data to destroy later */
+	TSet<FNiagaraSystemInstanceID> DeferredDestroyList;
+};
+
diff --git a/Engine/Plugins/FX/Niagara/Shaders/Private/NiagaraStrandsPhysics.ush b/Engine/Plugins/FX/Niagara/Shaders/Private/NiagaraStrandsPhysics.ush
--- a/Engine/Plugins/FX/Niagara/Shaders/Private/NiagaraStrandsPhysics.ush
+++ b/Engine/Plugins/FX/Niagara/Shaders/Private/NiagaraStrandsPhysics.ush
@@ -11,14 +11,14 @@
  * -----------------------------------------------------------------
  */
 
-groupshared float3 SharedLocalPosition[THREADGROUP_SIZE];
-groupshared float4 SharedLocalQuaternion[THREADGROUP_SIZE];
+groupshared float3 SharedNodePosition[THREADGROUP_SIZE];
+groupshared float4 SharedNodeOrientation[THREADGROUP_SIZE];
 
 groupshared float SharedInverseMass[THREADGROUP_SIZE];
 groupshared float SharedInverseInertia[THREADGROUP_SIZE];
 
-groupshared float3 SharedLocalVelocity[THREADGROUP_SIZE];
-groupshared float3 SharedLocalOmega[THREADGROUP_SIZE];
+groupshared float3 SharedPreviousPosition[THREADGROUP_SIZE];
+groupshared float4 SharedPreviousOrientation[THREADGROUP_SIZE];
 
 /* -----------------------------------------------------------------
  * Get the strand parent position
@@ -27,7 +27,7 @@ groupshared float3 SharedLocalOmega[THREADGROUP_SIZE];
 
 void GetParentPosition(const bool IsMobile, out float3 ParentPosition)
 {	
-	ParentPosition = SharedLocalPosition[IsMobile ? GGroupThreadId.x-1 : GGroupThreadId.x];
+	ParentPosition = SharedNodePosition[IsMobile ? GGroupThreadId.x-1 : GGroupThreadId.x];
 }
 
 /* -----------------------------------------------------------------
@@ -37,8 +37,8 @@ void GetParentPosition(const bool IsMobile, out float3 ParentPosition)
 
 void ExtractPositionSharedDatas(const bool IsMobile, const float3 InParticlePosition, const float3 InParticleVelocity, const float InParticleMass)
 {	
-	SharedLocalPosition[GGroupThreadId.x] = InParticlePosition;
-	SharedLocalVelocity[GGroupThreadId.x] = InParticleVelocity;
+	SharedNodePosition[GGroupThreadId.x] = InParticlePosition;
+	//SharedLocalVelocity[GGroupThreadId.x] = InParticleVelocity;
 	SharedInverseMass[GGroupThreadId.x] = (InParticleMass != 0.0 && IsMobile) ? 1.0/InParticleMass : 0.0;
 
 	GroupMemoryBarrier(); 
@@ -46,8 +46,8 @@ void ExtractPositionSharedDatas(const bool IsMobile, const float3 InParticlePosi
 
 void ExtractOrientationSharedDatas(const bool IsMobile, const float4 InParticleQuaternion, const float3 InParticleOmega, const float InParticleInertia)
 {
-	SharedLocalQuaternion[GGroupThreadId.x] = InParticleQuaternion;
-	SharedLocalOmega[GGroupThreadId.x] = InParticleOmega;
+	SharedNodeOrientation[GGroupThreadId.x] = InParticleQuaternion;
+	//SharedLocalOmega[GGroupThreadId.x] = InParticleOmega;
 	SharedInverseInertia[GGroupThreadId.x] = (InParticleInertia != 0.0 && IsMobile) ? 1.0/InParticleInertia : 0.0;
 
 	GroupMemoryBarrier(); 
@@ -58,25 +58,25 @@ void ExtractOrientationSharedDatas(const bool IsMobile, const float4 InParticleQ
  * -----------------------------------------------------------------
  */
 
-void ReportPositionSharedDatas(const float InDeltaTime, inout float3 OutParticleVelocity, inout float3 OutParticlePosition)
+void ReportPositionSharedDatas(const float InDeltaTime, inout float3 OutParticleVelocity, out float3 OutParticlePosition)
 {
-	const float3 CurrentPosition = SharedLocalPosition[GGroupThreadId.x];
-	OutParticleVelocity += (CurrentPosition-OutParticlePosition) / InDeltaTime;
+	const float3 CurrentPosition = SharedNodePosition[GGroupThreadId.x];
+	//OutParticleVelocity += (CurrentPosition-OutParticlePosition) / InDeltaTime;
 	OutParticlePosition = CurrentPosition;
 }
 
-void ReportOrientationSharedDatas(const float InDeltaTime, inout float3 OutParticleOmega, inout float4 OutParticleQuaternion)
+void ReportOrientationSharedDatas(const float InDeltaTime, inout float3 OutParticleOmega, out float4 OutParticleQuaternion)
 {
-	const float4 CurrentQuaternion = SharedLocalQuaternion[GGroupThreadId.x];
-	const float4 DeltaQuat = float4(
-			-OutParticleQuaternion.xyz * CurrentQuaternion.w + CurrentQuaternion.xyz * OutParticleQuaternion.w + cross(CurrentQuaternion.xyz, -OutParticleQuaternion.xyz),
-			CurrentQuaternion.w * OutParticleQuaternion.w - dot(CurrentQuaternion.xyz, -OutParticleQuaternion.xyz));
+	const float4 CurrentQuaternion = SharedNodeOrientation[GGroupThreadId.x];
+	//const float4 DeltaQuat = float4(
+	//		-OutParticleQuaternion.xyz * CurrentQuaternion.w + CurrentQuaternion.xyz * OutParticleQuaternion.w + cross(CurrentQuaternion.xyz, -OutParticleQuaternion.xyz),
+	//		CurrentQuaternion.w * OutParticleQuaternion.w - dot(CurrentQuaternion.xyz, -OutParticleQuaternion.xyz));
 
 	OutParticleQuaternion = CurrentQuaternion;
 
-	const float AxisLength = length( DeltaQuat.xyz );
-	const float QuatAngle = 2.0 * atan2(AxisLength,DeltaQuat.w );
-	OutParticleOmega += DeltaQuat.xyz * QuatAngle / (AxisLength*InDeltaTime);
+	//const float AxisLength = length( DeltaQuat.xyz );
+	//const float QuatAngle = 2.0 * atan2(AxisLength,DeltaQuat.w );
+	//OutParticleOmega += DeltaQuat.xyz * QuatAngle / (AxisLength*InDeltaTime);
 }
 
 /* -----------------------------------------------------------------
@@ -108,7 +108,7 @@ void AdvectOrientationSharedDatas(const float3 InParticleOmega, const float InDe
 void StoreHookeSpringConfig(const bool IsMobile, const float InMaterialCompliance, out float OutRestLength)
 {
 	if (!IsMobile){
-		const float3 EdgeDirection = SharedLocalPosition[GGroupThreadId.x] - SharedLocalPosition[GGroupThreadId.x-1];
+		const float3 EdgeDirection = SharedNodePosition[GGroupThreadId.x] - SharedNodePosition[GGroupThreadId.x-1];
 		OutRestLength = length(EdgeDirection);
 	}
 	else
@@ -138,7 +138,7 @@ void ResetHookeSpringMultiplier(out float OutMaterialMultiplier)
 
 void UpdateHookeSpringMultiplier(const float InMaterialCompliance, const float InMaterialWeight, const float InRestLength, inout float OutMaterialMultiplier)
 {
-	float3 EdgeDirection = SharedLocalPosition[GGroupThreadId.x] - SharedLocalPosition[GGroupThreadId.x-1];
+	float3 EdgeDirection = SharedNodePosition[GGroupThreadId.x] - SharedNodePosition[GGroupThreadId.x-1];
 	const float EdgeLength = length(EdgeDirection);
 	EdgeDirection /= EdgeLength;
 
@@ -152,8 +152,8 @@ void UpdateHookeSpringMultiplier(const float InMaterialCompliance, const float I
 	const float3 PositionDelta = EdgeDirection * DeltaLambda;
 
 	// dX += dX
-	SharedLocalPosition[GGroupThreadId.x] += PositionDelta * SharedInverseMass[GGroupThreadId.x];
-	SharedLocalPosition[GGroupThreadId.x-1] -= PositionDelta * SharedInverseMass[GGroupThreadId.x-1];
+	SharedNodePosition[GGroupThreadId.x] += PositionDelta * SharedInverseMass[GGroupThreadId.x];
+	SharedNodePosition[GGroupThreadId.x-1] -= PositionDelta * SharedInverseMass[GGroupThreadId.x-1];
 }
 
 void SolveHookeSpringMaterial(const float InMaterialCompliance, const float InMaterialWeight, const float InRestLength, const bool InIsRed, inout float OutMaterialMultiplier)
@@ -202,13 +202,13 @@ void ResetStretchShearMultiplier(out float3 OutMaterialMultiplier)
 void UpdateStretchShearMultiplier(const float InMaterialCompliance, const float InMaterialWeight, const float InRestLength,  inout float3 OutMaterialMultiplier)
 {
 	float3 EdgeDirection;
-	float4 q0 = SharedLocalQuaternion[GGroupThreadId.x-1];
+	float4 q0 = SharedNodeOrientation[GGroupThreadId.x-1];
 	EdgeDirection[0] = 2.0 * (q0.x * q0.z + q0.w * q0.y);
 	EdgeDirection[1] = 2.0 * (q0.y * q0.z - q0.w * q0.x);
 	EdgeDirection[2] = q0.w * q0.w - q0.x * q0.x - q0.y * q0.y + q0.z * q0.z;
 
 	// XPBD lagrange multiplier update : dL = -(C+compliance*L) / (dC * invM * dCt + alpha)
-	const float3 DeltaLambda = -((SharedLocalPosition[GGroupThreadId.x] - SharedLocalPosition[GGroupThreadId.x-1]) / InRestLength - EdgeDirection + 
+	const float3 DeltaLambda = -((SharedNodePosition[GGroupThreadId.x] - SharedNodePosition[GGroupThreadId.x-1]) / InRestLength - EdgeDirection + 
 					OutMaterialMultiplier * InMaterialCompliance) * InMaterialWeight;
 
 	// L += dL
@@ -218,15 +218,15 @@ void UpdateStretchShearMultiplier(const float InMaterialCompliance, const float
 	const float3 PositionDelta  = DeltaLambda/InRestLength;
 
 	// dX += dX
-	SharedLocalPosition[GGroupThreadId.x] += PositionDelta * SharedInverseMass[GGroupThreadId.x];
-	SharedLocalPosition[GGroupThreadId.x-1] -= PositionDelta * SharedInverseMass[GGroupThreadId.x-1];
+	SharedNodePosition[GGroupThreadId.x] += PositionDelta * SharedInverseMass[GGroupThreadId.x];
+	SharedNodePosition[GGroupThreadId.x-1] -= PositionDelta * SharedInverseMass[GGroupThreadId.x-1];
 
 	const float4 qebar = float4(-q0.y, q0.x, -q0.w, q0.z);	
 	const float4 QuaternionDelta = -2.0 * float4( DeltaLambda.xyz * qebar.w + cross(DeltaLambda.xyz, qebar.xyz),
 					- dot(DeltaLambda.xyz, qebar.xyz) );
 
 	q0 += QuaternionDelta * SharedInverseInertia[GGroupThreadId.x-1];
-	SharedLocalQuaternion[GGroupThreadId.x-1] = normalize(q0);
+	SharedNodeOrientation[GGroupThreadId.x-1] = normalize(q0);
 }
 
 void SolveStretchShearMaterial(const float InMaterialCompliance, const float InMaterialWeight, const float InRestLength, const bool InIsRed, inout float3 OutMaterialMultiplier)
@@ -274,13 +274,13 @@ void ResetAngularSpringMultiplier(out float3 OutMaterialMultiplier)
 void UpdateAngularSpringMultiplier(const float InMaterialCompliance, const float InMaterialWeight, const float InRestLength,  inout float3 OutMaterialMultiplier)
 {
 	float3 EdgeDirection;
-	float4 q0 = SharedLocalQuaternion[GGroupThreadId.x-1];
+	float4 q0 = SharedNodeOrientation[GGroupThreadId.x-1];
 	EdgeDirection[0] = 2.0 * (q0.x * q0.z + q0.w * q0.y);
 	EdgeDirection[1] = 2.0 * (q0.y * q0.z - q0.w * q0.x);
 	EdgeDirection[2] = q0.w * q0.w - q0.x * q0.x - q0.y * q0.y + q0.z * q0.z;
 
 	// XPBD lagrange multiplier update : dL = -(C+compliance*L) / (dC * invM * dCt + alpha)
-	const float3 DeltaLambda = -((SharedLocalPosition[GGroupThreadId.x] - SharedLocalPosition[GGroupThreadId.x-1]) / InRestLength - EdgeDirection + 
+	const float3 DeltaLambda = -((SharedNodePosition[GGroupThreadId.x] - SharedNodePosition[GGroupThreadId.x-1]) / InRestLength - EdgeDirection + 
 					OutMaterialMultiplier * InMaterialCompliance) * InMaterialWeight;
 
 	// L += dL
@@ -291,7 +291,7 @@ void UpdateAngularSpringMultiplier(const float InMaterialCompliance, const float
 					- dot(DeltaLambda.xyz, qebar.xyz) );
 
 	q0 += QuaternionDelta * SharedInverseInertia[GGroupThreadId.x-1];
-	SharedLocalQuaternion[GGroupThreadId.x-1] = normalize(q0);
+	SharedNodeOrientation[GGroupThreadId.x-1] = normalize(q0);
 }
 
 void SolveAngularSpringMaterial(const float InMaterialCompliance, const float InMaterialWeight, const float InRestLength, const bool InIsRed, inout float3 OutMaterialMultiplier)
@@ -317,12 +317,12 @@ void SolveAngularSpringMaterial(const float InMaterialCompliance, const float In
 
 void UpdateBendSpringMultiplier(const float InMaterialCompliance, const float InMaterialWeight, const float InRestLength, const float3 InRestDirection, inout float3 OutMaterialMultiplier, const bool IsFirst)
 {
-	float4 ParentFrame = (IsFirst) ? SharedLocalQuaternion[GGroupThreadId.x-1] : SharedLocalQuaternion[GGroupThreadId.x-2];
+	float4 ParentFrame = (IsFirst) ? SharedNodeOrientation[GGroupThreadId.x-1] : SharedNodeOrientation[GGroupThreadId.x-2];
 	float3 CrossProduct = 2.0 * cross(ParentFrame.xyz,InRestDirection);
 	float3 EdgeDirection = InRestDirection + ParentFrame.w * CrossProduct + cross(ParentFrame.xyz,CrossProduct);
 
 	// XPBD lagrange multiplier update : dL = -(C+compliance*L) / (dC * invM * dCt + alpha)
-	const float3 DeltaLambda = -((SharedLocalPosition[GGroupThreadId.x] - SharedLocalPosition[GGroupThreadId.x-1]) / InRestLength - EdgeDirection / InRestLength + 
+	const float3 DeltaLambda = -((SharedNodePosition[GGroupThreadId.x] - SharedNodePosition[GGroupThreadId.x-1]) / InRestLength - EdgeDirection / InRestLength + 
 					OutMaterialMultiplier * InMaterialCompliance) * InMaterialWeight;
 
 	// L += dL
@@ -332,8 +332,8 @@ void UpdateBendSpringMultiplier(const float InMaterialCompliance, const float In
 	const float3 PositionDelta  = DeltaLambda/InRestLength;
 
 	// dX += dX
-	SharedLocalPosition[GGroupThreadId.x] += PositionDelta * SharedInverseMass[GGroupThreadId.x];
-	SharedLocalPosition[GGroupThreadId.x-1] -= PositionDelta * SharedInverseMass[GGroupThreadId.x-1];
+	SharedNodePosition[GGroupThreadId.x] += PositionDelta * SharedInverseMass[GGroupThreadId.x];
+	SharedNodePosition[GGroupThreadId.x-1] -= PositionDelta * SharedInverseMass[GGroupThreadId.x-1];
 }
 
 void SolveBendSpringMaterial(const float InMaterialCompliance, const float InMaterialWeight, const float InRestLength, const float3 InRestDirection, const bool InIsRed, inout float3 OutMaterialMultiplier, const int InStrandSize)
@@ -379,22 +379,22 @@ float4 FindQuatBetweenHelper_tmp(in float3 A, in float3 B, in float NormAB)
 	const bool IsMobile = (GGroupThreadId.x % StrandSize) != 0;
 	if (!IsMobile)
 	{
-		float4 NodeQuaternion = float4(0,0,0,1);//SharedLocalQuaternion[GGroupThreadId.x];
+		float4 NodeQuaternion = float4(0,0,0,1);//SharedNodeOrientation[GGroupThreadId.x];
 		float3 TangentNext = float3(0,0,1);
 		float3 TangentPrev = TangentNext;
 
-		float3 PositionNext = SharedLocalPosition[GGroupThreadId.x];
+		float3 PositionNext = SharedNodePosition[GGroupThreadId.x];
 		float3 PositionPrev = PositionNext;
 
 		for (int EdgeIndex = 0, EdgeEnd = StrandSize-1; EdgeIndex < EdgeEnd; ++EdgeIndex)
 		{
 			PositionPrev = PositionNext;
-			PositionNext = SharedLocalPosition[GGroupThreadId.x+EdgeIndex+1];
+			PositionNext = SharedNodePosition[GGroupThreadId.x+EdgeIndex+1];
 
 			TangentPrev = TangentNext;
 			TangentNext = normalize(PositionNext - PositionPrev);
 			NodeQuaternion = normalize( MultiplyQuat_tmp( FindQuatBetweenHelper_tmp(TangentPrev,TangentNext,1.0), NodeQuaternion ) );
-			SharedLocalQuaternion[GGroupThreadId.x+EdgeIndex] = NodeQuaternion;
+			SharedNodeOrientation[GGroupThreadId.x+EdgeIndex] = NodeQuaternion;
 		}
 	}
 	GroupMemoryBarrier(); 
@@ -425,8 +425,8 @@ void ResetBendTwistMultiplier(out float3 OutMaterialMultiplier)
 
 void UpdateBendTwistMultiplier(const float InMaterialCompliance, const float InMaterialWeight, const float InRestLength, const float4 InRestDarboux, inout float3 OutMaterialMultiplier)
 {
-	float4 q0 = SharedLocalQuaternion[GGroupThreadId.x-1];
-	float4 q1 = SharedLocalQuaternion[GGroupThreadId.x];
+	float4 q0 = SharedNodeOrientation[GGroupThreadId.x-1];
+	float4 q1 = SharedNodeOrientation[GGroupThreadId.x];
 	const float4 CurrentDarboux = float4(
 			q1.xyz * q0.w - q0.xyz * q1.w + cross(-q0.xyz, q1.xyz),
 			q0.w * q1.w - dot(-q0.xyz, q1.xyz));
@@ -442,19 +442,19 @@ void UpdateBendTwistMultiplier(const float InMaterialCompliance, const float InM
 
 	OutMaterialMultiplier -= DeltaOmega.xyz;
 
-	SharedLocalQuaternion[GGroupThreadId.x-1] += SharedInverseInertia[GGroupThreadId.x-1] *  float4(
+	SharedNodeOrientation[GGroupThreadId.x-1] += SharedInverseInertia[GGroupThreadId.x-1] *  float4(
 			DeltaOmega.xyz * q1.w + cross(q1.xyz, DeltaOmega.xyz), - dot(q1.xyz, DeltaOmega.xyz)) * InverseRestLength;
-	SharedLocalQuaternion[GGroupThreadId.x] -= SharedInverseInertia[GGroupThreadId.x] *  float4(
+	SharedNodeOrientation[GGroupThreadId.x] -= SharedInverseInertia[GGroupThreadId.x] *  float4(
 			DeltaOmega.xyz * q0.w + cross(q0.xyz, DeltaOmega.xyz), - dot(q0.xyz, DeltaOmega.xyz)) * InverseRestLength;
 
-	SharedLocalQuaternion[GGroupThreadId.x-1] = normalize(SharedLocalQuaternion[GGroupThreadId.x-1]);
-	SharedLocalQuaternion[GGroupThreadId.x] = normalize(SharedLocalQuaternion[GGroupThreadId.x]);
+	SharedNodeOrientation[GGroupThreadId.x-1] = normalize(SharedNodeOrientation[GGroupThreadId.x-1]);
+	SharedNodeOrientation[GGroupThreadId.x] = normalize(SharedNodeOrientation[GGroupThreadId.x]);
 }
 
 /*void UpdateBendTwistMultiplier(const float InMaterialCompliance, const float InMaterialWeight, const float InRestLength, const float4 InRestDarboux, inout float3 OutMaterialMultiplier)
 {
-	float4 q0 = SharedLocalQuaternion[GGroupThreadId.x-1];
-	float4 q1 = SharedLocalQuaternion[GGroupThreadId.x];
+	float4 q0 = SharedNodeOrientation[GGroupThreadId.x-1];
+	float4 q1 = SharedNodeOrientation[GGroupThreadId.x];
 	const float4 CurrentDarboux = float4(
 			q1.xyz * q0.w - q0.xyz * q1.w + cross(-q0.xyz, q1.xyz),
 			q0.w * q1.w - dot(-q0.xyz, q1.xyz));
@@ -476,13 +476,13 @@ void UpdateBendTwistMultiplier(const float InMaterialCompliance, const float InM
 
 	const float afac = dot(CurrentDarboux.xyz,DeltaOmega);
 
-	SharedLocalQuaternion[GGroupThreadId.x-1] += SharedInverseInertia[GGroupThreadId.x-1] * ( float4(
+	SharedNodeOrientation[GGroupThreadId.x-1] += SharedInverseInertia[GGroupThreadId.x-1] * ( float4(
 			DeltaOmega.xyz * q1tmp.w + cross(q1tmp.xyz, DeltaOmega.xyz), - dot(q1tmp.xyz, DeltaOmega.xyz)) + dto * q1) * InverseRestLength;
-	SharedLocalQuaternion[GGroupThreadId.x] -= SharedInverseInertia[GGroupThreadId.x] *  ( float4(
+	SharedNodeOrientation[GGroupThreadId.x] -= SharedInverseInertia[GGroupThreadId.x] *  ( float4(
 			DeltaOmega.xyz * q0tmp.w + cross(q0tmp.xyz, DeltaOmega.xyz), - dot(q0tmp.xyz, DeltaOmega.xyz)) - dto * q0) * InverseRestLength;
 
-	SharedLocalQuaternion[GGroupThreadId.x-1] = normalize(SharedLocalQuaternion[GGroupThreadId.x-1]);
-	SharedLocalQuaternion[GGroupThreadId.x] = normalize(SharedLocalQuaternion[GGroupThreadId.x]);
+	SharedNodeOrientation[GGroupThreadId.x-1] = normalize(SharedNodeOrientation[GGroupThreadId.x-1]);
+	SharedNodeOrientation[GGroupThreadId.x] = normalize(SharedNodeOrientation[GGroupThreadId.x]);
 }*/
 
 void SolveBendTwistMaterial(const float InMaterialCompliance, const float InMaterialWeight, const float InRestLength, const float4 InRestDarboux, const bool InIsRed, inout float3 OutMaterialMultiplier)
@@ -516,7 +516,7 @@ void SolveBendTwistMaterial(const float InMaterialCompliance, const float InMate
 		{
 			if (i > GGroupThreadId.x+1)
 			{
-				float3 EdgeDirection = SharedLocalPosition[i] - SharedLocalPosition[i-2];
+				float3 EdgeDirection = SharedNodePosition[i] - SharedNodePosition[i-2];
 				const float EdgeLength = length(EdgeDirection);
 				EdgeDirection /= EdgeLength;
 				if (EdgeLength<=BendingLimit)
@@ -526,19 +526,19 @@ void SolveBendTwistMaterial(const float InMaterialCompliance, const float InMate
 					const float3 DeltaPosition = EdgeDirection * DeltaLambda;
 
 					// dX += dX
-					SharedLocalPosition[i] += DeltaPosition;
+					SharedNodePosition[i] += DeltaPosition;
 
 					const float3 DeltaVelocity = DeltaPosition / InDeltaTime;
 
 					// dV += dV
-					SharedLocalVelocity[i] += SharedInverseMass[i] * DeltaVelocity / ( SharedInverseMass[i] + InDampening * SharedInverseMass[i-2]);
-					SharedLocalVelocity[i-2] -= InDampening * SharedInverseMass[i-2] * DeltaVelocity / ( SharedInverseMass[i] + InDampening * SharedInverseMass[i-2]);
+					//SharedLocalVelocity[i] += SharedInverseMass[i] * DeltaVelocity / ( SharedInverseMass[i] + InDampening * SharedInverseMass[i-2]);
+					//SharedLocalVelocity[i-2] -= InDampening * SharedInverseMass[i-2] * DeltaVelocity / ( SharedInverseMass[i] + InDampening * SharedInverseMass[i-2]);
 				}
 				GroupMemoryBarrier(); 
 			}
 
 
-			float3 EdgeDirection = SharedLocalPosition[i] - SharedLocalPosition[i-1];
+			float3 EdgeDirection = SharedNodePosition[i] - SharedNodePosition[i-1];
 			const float EdgeLength = length(EdgeDirection);
 			EdgeDirection /= EdgeLength;
 
@@ -547,16 +547,16 @@ void SolveBendTwistMaterial(const float InMaterialCompliance, const float InMate
 			const float3 DeltaPosition = EdgeDirection * DeltaLambda;
 
 			// dX += dX
-			SharedLocalPosition[i] += DeltaPosition;
+			SharedNodePosition[i] += DeltaPosition;
 
-			//SharedLocalPosition[i] += SharedInverseMass[i] * DeltaPosition / ( SharedInverseMass[i] + InDampening * SharedInverseMass[i-1]);
-			//SharedLocalPosition[i-1] -= InDampening * SharedInverseMass[i-1] * DeltaPosition / ( SharedInverseMass[i] + InDampening * SharedInverseMass[i-1]);
+			//SharedNodePosition[i] += SharedInverseMass[i] * DeltaPosition / ( SharedInverseMass[i] + InDampening * SharedInverseMass[i-1]);
+			//SharedNodePosition[i-1] -= InDampening * SharedInverseMass[i-1] * DeltaPosition / ( SharedInverseMass[i] + InDampening * SharedInverseMass[i-1]);
 
 			const float3 DeltaVelocity = DeltaPosition / InDeltaTime;
 
 			// dV += dV
-			SharedLocalVelocity[i] += SharedInverseMass[i] * DeltaVelocity / ( SharedInverseMass[i] + InDampening * SharedInverseMass[i-1]);
-			SharedLocalVelocity[i-1] -= InDampening * SharedInverseMass[i-1] * DeltaVelocity / ( SharedInverseMass[i] + InDampening * SharedInverseMass[i-1]);
+			//SharedLocalVelocity[i] += SharedInverseMass[i] * DeltaVelocity / ( SharedInverseMass[i] + InDampening * SharedInverseMass[i-1]);
+			//SharedLocalVelocity[i-1] -= InDampening * SharedInverseMass[i-1] * DeltaVelocity / ( SharedInverseMass[i] + InDampening * SharedInverseMass[i-1]);
 
 			GroupMemoryBarrier(); 
 		}
diff --git a/Engine/Shaders/Private/HairStrands/HairStrandsInterpolation.usf b/Engine/Shaders/Private/HairStrands/HairStrandsInterpolation.usf
--- a/Engine/Shaders/Private/HairStrands/HairStrandsInterpolation.usf
+++ b/Engine/Shaders/Private/HairStrands/HairStrandsInterpolation.usf
@@ -11,7 +11,95 @@ Buffer<uint4>  DeformedSimPositionBuffer;
 Buffer<uint4>  Interpolation0Buffer;
 Buffer<uint4>  Interpolation1Buffer;
 
-RWBuffer<uint4>  OutRenderDeformedPositionBuffer;
+float3			OutRenderDeformedPositionCenter; 
+RWBuffer<uint4> OutRenderDeformedPositionBuffer;
+
+#if PERMUTATION_DYNAMIC_GEOMETRY == 1
+
+// Compact all these buffers into 2 buffers: translation + quaternion
+float3			RestPositionWorldCenter;
+Buffer<float4>  RestPosition0Buffer;
+Buffer<float4>  RestPosition1Buffer;
+Buffer<float4>  RestPosition2Buffer;
+
+float3			DeformedPositionWorldCenter;
+Buffer<float4>  DeformedPosition0Buffer;
+Buffer<float4>  DeformedPosition1Buffer;
+Buffer<float4>  DeformedPosition2Buffer;
+
+Buffer<uint>	RootToTriangleIndex;
+Buffer<uint>	VertexToRootIndexBuffer;
+
+#define QUATERNION_IDENTITY float4(0,0,0,1)
+#define SMALL_NUMBER 1e-8
+
+float4 NormalizeQuat(in float4 Quat)
+{
+	float SquaredNorm = dot(Quat, Quat);
+	return (SquaredNorm >= SMALL_NUMBER) ? Quat / sqrt(SquaredNorm) : QUATERNION_IDENTITY;
+}
+
+float4 InverseQuat(in float4 Quat)
+{
+	return float4(-Quat.x, -Quat.y, -Quat.z, Quat.w);
+}
+
+float4 MultiplyQuat(in float4 QuatA, in float4 QuatB)
+{
+	return float4(
+		QuatB.xyz * QuatA.w + QuatA.xyz * QuatB.w + cross(QuatA.xyz, QuatB.xyz),
+		QuatA.w * QuatB.w - dot(QuatA.xyz, QuatB.xyz));
+}
+
+float3 RotateVectorByQuat(in float3 Vector, in float4 Quat)
+{
+	float3 T = 2.0 * cross(Quat.xyz, Vector);
+	return Vector + Quat.w * T + cross(Quat.xyz, T);
+}
+
+float3 UnrotateVectorByQuat(in float3 Vector, in float4 Quat)
+{
+	float3 T = 2.0 * cross(Quat.xyz, Vector);
+	return Vector - Quat.w * T + cross(Quat.xyz, T);
+}
+
+float4 FindQuatBetweenHelper(in float3 A, in float3 B, in float NormAB)
+{
+	float W = NormAB + dot(A, B);
+	float4 Quat = (W > 1e-6*NormAB) ? float4(A.y*B.z - A.z*B.y, A.z*B.x - A.x*B.z, A.x*B.y - A.y*B.x, W) :
+		(abs(A.x) > abs(A.y)) ? float4(-A.z, 0.0, A.x, 0.0) : float4(0.0, -A.z, A.y, 0.0);
+	return NormalizeQuat(Quat);
+}
+
+float4 FindQuatBetweenNormals(in float3 NormalA, in float3 NormalB)
+{
+	float NormAB = 1.0;
+	return FindQuatBetweenHelper(NormalA, NormalB, NormAB);
+}
+
+void GetTriangleTransformation(uint VertexIndex, out float3 T, out float4 Q)
+{
+	const uint RootIndex = VertexToRootIndexBuffer[VertexIndex];
+//	const uint TriangleIndex = RootToTriangleIndex[RootIndex];
+
+	float3 Rest_P1 = RestPosition1Buffer[RootIndex].xyz + RestPositionWorldCenter;
+	float3 Rest_P0 = RestPosition0Buffer[RootIndex].xyz + RestPositionWorldCenter;
+	float3 Rest_P2 = RestPosition2Buffer[RootIndex].xyz + RestPositionWorldCenter;
+	float3 Rest_E0 = Rest_P1 - Rest_P0;
+	float3 Rest_E1 = Rest_P2 - Rest_P0;
+	float3 Rest_N = normalize(cross(Rest_E0, Rest_E1));
+
+	float3 Deformed_P0	= DeformedPosition0Buffer[RootIndex].xyz + DeformedPositionWorldCenter;
+	float3 Deformed_P1	= DeformedPosition1Buffer[RootIndex].xyz + DeformedPositionWorldCenter;
+	float3 Deformed_P2	= DeformedPosition2Buffer[RootIndex].xyz + DeformedPositionWorldCenter;
+	float3 Deformed_E0	= Deformed_P1 - Deformed_P0;
+	float3 Deformed_E1	= Deformed_P2 - Deformed_P0;
+	float3 Deformed_N	= normalize(cross(Deformed_E0, Deformed_E1));
+
+	T = Deformed_P0 - Rest_P0;
+	Q = FindQuatBetweenNormals(Rest_N, Deformed_N);
+}
+#endif
 
 #if PERMUTATION_DEBUG == 1
 Buffer<float4>	 SimAttributeBuffer;
@@ -57,7 +145,18 @@ void MainCS(uint2 DispatchThreadId : SV_DispatchThreadID)
 
 		}
 
-		OutRenderDeformedPositionBuffer[VertexIndex] = float4(f32tof16(ControlPoint + CurrOffset), RenderRestPosePositionBuffer[VertexIndex].w);
+		// Apply dynamic mesh deformation (translation / rotation)
+		// hair_todo: applied the transformation in local space and then rotation everything in global space
+		#if PERMUTATION_DYNAMIC_GEOMETRY == 1
+		{
+			float3 DynTranslation = 0;
+			float4 DynRotation = 0;
+			GetTriangleTransformation(VertexIndex, DynTranslation, DynRotation);
+			CurrOffset += DynTranslation;
+		}
+		#endif	
+
+		OutRenderDeformedPositionBuffer[VertexIndex] = float4(f32tof16(ControlPoint + CurrOffset - OutRenderDeformedPositionCenter), RenderRestPosePositionBuffer[VertexIndex].w);
 
 		#if PERMUTATION_DEBUG == 1
 		{
diff --git a/Engine/Shaders/Private/HairStrands/HairStrandsMeshProjection.usf b/Engine/Shaders/Private/HairStrands/HairStrandsMeshProjection.usf
new file mode 100644
--- /dev/null
+++ b/Engine/Shaders/Private/HairStrands/HairStrandsMeshProjection.usf
@@ -0,0 +1,132 @@
+// Copyright 1998-2019 Epic Games, Inc. All Rights Reserved.
+
+#include "../Common.ush"
+#include "HairStrandsMeshProjectionCommon.ush"
+
+uint bClear;
+uint MaxRootCount;
+float DistanceThreshold;
+float4x4 RootLocalToWorld;
+
+float4x4 MeshLocalToWorld;
+uint MeshNumPrimitives;
+uint MeshSectionIndex;
+uint MeshMaxIndexCount;
+uint MeshMaxVertexCount;
+uint MeshIndexOffset;
+uint MaxSpinLockCount;
+Buffer<uint> MeshIndexBuffer;
+Buffer<float> MeshPositionBuffer;
+
+Buffer<float4> RootPositionBuffer;
+Buffer<float4> RootNormalBuffer;
+
+RWBuffer<uint> OutRootTriangleIndex;
+//RWBuffer<uint> OutRootTriangleBarycentrics;
+RWStructuredBuffer<float> OutRootTriangleDistance;
+
+////////////////////////////////////////////////////////////////////////////////////////////////////
+// Curve projection onto A mesh
+// The projection is executed by running on curve's root per thread, which test all triangles
+////////////////////////////////////////////////////////////////////////////////////////////////////
+
+
+// Closest point on A triangle from another point
+// Code from the book "Real-Time Collision Detection" by Christer Ericson
+float3 ClosestPoint(const FHairMeshTriangle Tri, const float3 P)
+{
+	const float3 A = Tri.P0;
+	const float3 B = Tri.P1;
+	const float3 C = Tri.P2;
+
+	// Check if P is in vertex region outside A.
+	float3 AB = B - A;
+	float3 AC = C - A;
+	float3 AP = P - A;
+	float D1 = dot(AB, AP);
+	float D2 = dot(AC, AP);
+	if (D1 <= 0.f && D2 <= 0.f)
+		return A; // Barycentric coordinates are (1,0,0).
+
+	// Check if P is in vertex region outside B.
+	float3 BP = P - B;
+	float D3 = dot(AB, BP);
+	float D4 = dot(AC, BP);
+	if (D3 >= 0.f && D4 <= D3)
+		return B; // Barycentric coordinates are (0,1,0).
+
+	// Check if P is in edge region of AB, and if so, return the projection of P onto AB.
+	float VC = D1 * D4 - D3 * D2;
+	if (VC <= 0.f && D1 >= 0.f && D3 <= 0.f)
+	{
+		float V = D1 / (D1 - D3);
+		return A + V * AB; // The barycentric coordinates are (1-V, V, 0).
+	}
+
+	// Check if P is in vertex region outside C.
+	float3 CP = P - C;
+	float D5 = dot(AB, CP);
+	float D6 = dot(AC, CP);
+	if (D6 >= 0.f && D5 <= D6)
+		return C; // The barycentric coordinates are (0,0,1).
+
+	// Check if P is in edge region of AC, and if so, return the projection of P onto AC.
+	float VB = D5 * D2 - D1 * D6;
+	if (VB <= 0.f && D2 >= 0.f && D6 <= 0.f)
+	{
+		float W = D2 / (D2 - D6);
+		return A + W * AC; // The barycentric coordinates are (1-W, 0, W).
+	}
+
+	// Check if P is in edge region of BC, and if so, return the projection of P onto BC.
+	float VA = D3 * D6 - D5 * D4;
+	if (VA <= 0.f && D4 - D3 >= 0.f && D5 - D6 >= 0.f)
+	{
+		float W = (D4 - D3) / (D4 - D3 + D5 - D6);
+		return B + W * (C - B); // The barycentric coordinates are (0, 1-W, W).
+	}
+
+	// P must be inside the face region. Compute the closest point through its barycentric coordinates (u,V,W).
+	float Denom = 1.f / (VA + VB + VC);
+	float V = VB * Denom;
+	float W = VC * Denom;
+	return A + AB * V + AC * W;
+}
+
+[numthreads(32, 1, 1)]
+void MainCS(uint3 DispatchThreadId : SV_DispatchThreadID)
+{
+	// Shift root index with thread index to reduce/avoid collisions
+	const uint RootIndex = DispatchThreadId.x;
+
+	if (RootIndex >= MaxRootCount) return;
+
+	float3 RootPosition = RootPositionBuffer[RootIndex].xyz;
+	RootPosition = mul(float4(RootPosition, 1), RootLocalToWorld).xyz;
+
+	const float3 RootNormal = RootNormalBuffer[RootIndex].xyz;
+
+	float MinDistance = bClear ? 999999999 : OutRootTriangleDistance[RootIndex];
+	int TriangleIndex = -1;
+	for (uint PrimitiveIt = 0; PrimitiveIt < MeshNumPrimitives; ++PrimitiveIt)
+	{
+		const FHairMeshTriangle Tri = GetTriangleData(PrimitiveIt, MeshLocalToWorld, MeshPositionBuffer, MeshIndexBuffer, MeshIndexOffset, MeshMaxIndexCount, MeshMaxVertexCount);
+		
+		const float3 B = 0.5f;
+		const float3 C = ClosestPoint(Tri, RootPosition);
+		const float Distance = distance(C, RootPosition);
+		if (abs(Distance) < MinDistance)
+		{			
+			MinDistance = abs(Distance);
+			TriangleIndex = PrimitiveIt;
+		}
+	}
+
+	if (TriangleIndex >= 0)
+	{
+		const uint EncodeIndex = EncodeTriangleIndex(TriangleIndex, MeshSectionIndex);
+		OutRootTriangleIndex[RootIndex] = EncodeIndex;
+		OutRootTriangleDistance[RootIndex] = MinDistance;
+	}
+
+}
\ No newline at end of file
diff --git a/Engine/Shaders/Private/HairStrands/HairStrandsMeshProjectionCommon.ush b/Engine/Shaders/Private/HairStrands/HairStrandsMeshProjectionCommon.ush
new file mode 100644
--- /dev/null
+++ b/Engine/Shaders/Private/HairStrands/HairStrandsMeshProjectionCommon.ush
@@ -0,0 +1,106 @@
+// Copyright 1998-2019 Epic Games, Inc. All Rights Reserved.
+
+struct FHairMeshTriangle
+{
+	float3 P0;
+	float3 P1;
+	float3 P2;
+
+	float3 E0;
+	float3 E1;
+
+	float3 N;
+};
+
+float3 GetVertexPosition(uint VId, Buffer<float> InPositionBuffer, Buffer<uint> InIndexBuffer, uint InIndexOffset, uint InMaxIndexCount, uint InMaxVertexCount)
+{
+	const uint Index = min(InIndexOffset + VId, InMaxIndexCount-1);
+	uint VertexIndex = InIndexBuffer.Load(Index);
+	VertexIndex = min(VertexIndex, InMaxVertexCount-1);
+
+	float3 P = 0;
+	P.x = InPositionBuffer.Load(VertexIndex*3 + 0);
+	P.y = InPositionBuffer.Load(VertexIndex*3 + 1);
+	P.z = InPositionBuffer.Load(VertexIndex*3 + 2);
+	return P;
+}
+
+FHairMeshTriangle GetTriangleData(uint TriangleIndex, float4x4 LocalToWorld, Buffer<float> InPositionBuffer, Buffer<uint> InIndexBuffer, uint InIndexOffset, uint InMaxIndexCount, uint InMaxVertexCount)
+{
+	const uint BaseVertexId = TriangleIndex*3;
+	const uint VId0 = BaseVertexId;
+	const uint VId1 = BaseVertexId+1;
+	const uint VId2 = BaseVertexId+2;
+
+	FHairMeshTriangle Out;
+	Out.P0 = GetVertexPosition(VId0, InPositionBuffer, InIndexBuffer, InIndexOffset, InMaxIndexCount, InMaxVertexCount);
+	Out.P1 = GetVertexPosition(VId1, InPositionBuffer, InIndexBuffer, InIndexOffset, InMaxIndexCount, InMaxVertexCount);
+	Out.P2 = GetVertexPosition(VId2, InPositionBuffer, InIndexBuffer, InIndexOffset, InMaxIndexCount, InMaxVertexCount);
+
+	Out.P0 = mul(float4(Out.P0,1), LocalToWorld).xyz;
+	Out.P1 = mul(float4(Out.P1,1), LocalToWorld).xyz;
+	Out.P2 = mul(float4(Out.P2,1), LocalToWorld).xyz;
+
+	Out.E0 = Out.P1 - Out.P0;
+	Out.E1 = Out.P2 - Out.P0;
+	Out.N = normalize(cross(Out.E0, Out.E1));
+
+	return Out;
+}
+
+struct FHairProjectionData
+{
+	uint TriangleIndex;
+	float2 Barycentrics;
+};
+
+FHairProjectionData DecodeHairProjectionData(uint4 Data)
+{
+	const float B0 = f16tof32(Data.z);
+	const float B1 = f16tof32(Data.w);
+
+	FHairProjectionData Out;
+	Out.Barycentrics = float2(B0, B1);
+	Out.TriangleIndex = (Data.x & 0xFFFF) | ((Data.y & 0xFFFF) << 16);
+	return Out;
+}
+
+uint4 EncodeHairProjectionData(FHairProjectionData Data)
+{
+	uint4 Out;
+	Out.x = Data.TriangleIndex & 0xFFFF;
+	Out.y = (Data.TriangleIndex>>16) & 0xFFFF;
+	Out.z = f32tof16(saturate(Data.Barycentrics.x));
+	Out.w = f32tof16(saturate(Data.Barycentrics.y));
+	return Out;
+}
+
+uint EncodeBarycentrics(float2 B)
+{
+	return f32tof16(B.x) | (f32tof16(B.y) << 16);
+}
+
+float2 DecodeBarycentrics(uint E)
+{
+	return float2(f16tof32(E & 0xFFFF), f16tof32((E>>16) & 0xFFFF));
+}
+
+
+struct FHairTriangleIndex
+{
+	uint TriangleIndex;
+	uint SectionIndex;
+};
+
+uint EncodeTriangleIndex(uint TriangleIndex, uint SectionIndex)
+{
+	return ((SectionIndex & 0x3) <<28) | (TriangleIndex & 0xCFFFFFFF);
+}
+
+FHairTriangleIndex DecodeTriangleIndex(uint Encoded)
+{
+	FHairTriangleIndex Out;
+	Out.SectionIndex = (Encoded>>28) & 0x7;
+	Out.TriangleIndex = Encoded & 0xCFFFFFFF;
+	return Out;
+}
\ No newline at end of file
diff --git a/Engine/Shaders/Private/HairStrands/HairStrandsMeshProjectionHairDebug.usf b/Engine/Shaders/Private/HairStrands/HairStrandsMeshProjectionHairDebug.usf
new file mode 100644
--- /dev/null
+++ b/Engine/Shaders/Private/HairStrands/HairStrandsMeshProjectionHairDebug.usf
@@ -0,0 +1,139 @@
+// Copyright 1998-2019 Epic Games, Inc. All Rights Reserved.
+
+//#define SCENE_TEXTURES_DISABLED 1
+
+#include "../Common.ush"
+#include "HairStrandsMeshProjectionCommon.ush"
+
+#define INPUT_ROOT_FRAME 0
+#define INPUT_ROOT_TRIANGLE 1
+
+uint MaxRootCount;
+
+#if PERMUTATION_INPUT_TYPE == INPUT_ROOT_FRAME
+float4x4 RootLocalToWorld;
+Buffer<float4> RootPositionBuffer;
+Buffer<float4> RootNormalBuffer;
+
+uint			DeformedFrameEnable;
+Buffer<float4>  RestPosition0Buffer;
+Buffer<float4>  RestPosition1Buffer;
+Buffer<float4>  RestPosition2Buffer;
+
+Buffer<float4>  DeformedPosition0Buffer;
+Buffer<float4>  DeformedPosition1Buffer;
+Buffer<float4>  DeformedPosition2Buffer;
+#endif
+
+#if PERMUTATION_INPUT_TYPE == INPUT_ROOT_TRIANGLE
+Buffer<float4> RootTrianglePosition0;
+Buffer<float4> RootTrianglePosition1;
+Buffer<float4> RootTrianglePosition2;
+#endif
+
+#if PERMUTATION_INPUT_TYPE == INPUT_ROOT_FRAME
+void GetTriangleTransformation(uint RootIndex, out float3 T, out float4 Q)
+{
+	float3 Rest_P1 = RestPosition1Buffer[RootIndex].xyz;
+	float3 Rest_P0 = RestPosition0Buffer[RootIndex].xyz;
+	float3 Rest_P2 = RestPosition2Buffer[RootIndex].xyz;
+	float3 Rest_E0 = Rest_P1 - Rest_P0;
+	float3 Rest_E1 = Rest_P2 - Rest_P0;
+	float3 Rest_N = normalize(cross(Rest_E0, Rest_E1));
+
+	float3 Deformed_P0 = DeformedPosition0Buffer[RootIndex].xyz;
+	float3 Deformed_P1 = DeformedPosition1Buffer[RootIndex].xyz;
+	float3 Deformed_P2 = DeformedPosition2Buffer[RootIndex].xyz;
+	float3 Deformed_E0 = Deformed_P1 - Deformed_P0;
+	float3 Deformed_E1 = Deformed_P2 - Deformed_P0;
+	float3 Deformed_N = normalize(cross(Deformed_E0, Deformed_E1));
+
+	T = Deformed_P0 - Rest_P0;
+	Q = 0;// FindQuatBetweenNormals(Rest_N, Deformed_N);
+}
+#endif
+
+void MainVS(
+	uint VertexId : SV_VertexID,
+	out float3 OutFaceNormal : FACE_NORMAL,
+	out float4 OutPosition : SV_POSITION)
+{
+	ResolvedView = ResolveView();
+
+#if PERMUTATION_INPUT_TYPE == INPUT_ROOT_FRAME  
+	const uint LineVertex = (VertexId%2);
+	const uint LineIndex = VertexId/2;
+
+	OutFaceNormal 	= float3(0,0,1);
+	OutPosition 	= float4(0,0,0,0);
+	if (LineIndex > MaxRootCount) return;
+
+	float3 N;
+	float3 P;
+
+	if (DeformedFrameEnable)
+	{
+		float3 DynTranslation = 0;
+		float4 DynRotation = 0;
+		GetTriangleTransformation(LineIndex, DynTranslation, DynRotation);
+
+		float3 P0 = RootPositionBuffer[LineIndex];
+		P0 = mul(float4(P0, 1), RootLocalToWorld).xyz;
+		P0 += DynTranslation;
+
+		N = RootNormalBuffer[LineIndex];
+		N = mul(float4(N, 0), RootLocalToWorld).xyz;
+
+		const float3 P1 = P0 + N;
+		P = LineVertex == 0 ? P0 : P1;
+	}
+	else
+	{
+		float3 P0 = RootPositionBuffer[LineIndex];
+		P0 = mul(float4(P0,1), RootLocalToWorld).xyz;
+	
+		N  = RootNormalBuffer[LineIndex];
+		N = mul(float4(N, 0), RootLocalToWorld).xyz;
+
+		const float3 P1 = P0 + N;
+
+		P = LineVertex == 0 ? P0 : P1;
+	}
+
+#endif
+
+#if PERMUTATION_INPUT_TYPE == INPUT_ROOT_TRIANGLE
+	const uint TriangleVertex = (VertexId % 3);
+	const uint TriangleIndex = VertexId / 3;
+
+	if (TriangleIndex > MaxRootCount) return;
+
+	const float3 P0 = RootTrianglePosition0[TriangleIndex];
+	const float3 P1 = RootTrianglePosition1[TriangleIndex];
+	const float3 P2 = RootTrianglePosition2[TriangleIndex];
+
+	const float3 E0 = P1 - P0;
+	const float3 E1 = P2 - P0;
+	const float3 N = normalize(cross(E0, E1));
+
+	float3 P = 0;
+	if (TriangleVertex == 0)
+		P = P0;
+	if (TriangleVertex == 1)
+		P = P1;
+	if (TriangleVertex == 2)
+		P = P2;
+#endif
+
+	OutFaceNormal = N;
+	OutPosition = mul(float4(P,1), View.WorldToClip);
+}
+
+void MainPS(
+	in float3 FaceNormal : FACE_NORMAL,
+	out float4 OutColor : SV_Target0)
+{
+	FaceNormal += float3(1,1,1);
+	FaceNormal *= 0.5f;
+	OutColor = float4(FaceNormal,1);
+}
\ No newline at end of file
diff --git a/Engine/Shaders/Private/HairStrands/HairStrandsMeshProjectionMeshDebug.usf b/Engine/Shaders/Private/HairStrands/HairStrandsMeshProjectionMeshDebug.usf
new file mode 100644
--- /dev/null
+++ b/Engine/Shaders/Private/HairStrands/HairStrandsMeshProjectionMeshDebug.usf
@@ -0,0 +1,58 @@
+// Copyright 1998-2019 Epic Games, Inc. All Rights Reserved.
+
+//#define SCENE_TEXTURES_DISABLED 1
+
+#include "../Common.ush"
+#include "HairStrandsMeshProjectionCommon.ush"
+
+float4x4 LocalToWorld;
+uint VertexOffset;
+uint IndexOffset;
+uint MaxIndexCount;
+uint MaxVertexCount;
+
+Buffer<float> InputVertexBuffer;
+#if PERMUTATION_INPUT_TYPE == 1
+Buffer<uint> InputIndexBuffer;
+#endif
+
+
+void MainVS(
+	uint VertexId : SV_VertexID,
+	out float3 OutFaceNormal : FACE_NORMAL,
+	out float4 OutPosition : SV_POSITION)
+{
+	ResolvedView = ResolveView();
+
+#if PERMUTATION_INPUT_TYPE == 0
+	const uint VertexIndex = min(VertexOffset + VertexId, MaxVertexCount-1);
+	float3 P = InputVertexBuffer[VertexIndex];
+	OutFaceNormal = float3(0,0,1);
+#endif
+
+#if PERMUTATION_INPUT_TYPE == 1
+	const uint TriangleVertexIndex = (VertexId%3);
+	const uint TriangleIndex = VertexId / 3;
+	const FHairMeshTriangle Tri = GetTriangleData(TriangleIndex, LocalToWorld, InputVertexBuffer, InputIndexBuffer, IndexOffset, MaxIndexCount, MaxVertexCount);
+	OutFaceNormal = Tri.N;
+
+	float3 P = 0;
+	if (TriangleVertexIndex == 0)
+		P = Tri.P0;
+	if (TriangleVertexIndex == 1)
+		P = Tri.P1;
+	if (TriangleVertexIndex == 2)
+		P = Tri.P2;
+#endif
+
+	OutPosition = mul(float4(P,1), View.WorldToClip);
+}
+
+void MainPS(
+	in float3 FaceNormal : FACE_NORMAL,
+	out float4 OutColor : SV_Target0)
+{
+	FaceNormal += float3(1,1,1);
+	FaceNormal *= 0.5f;
+	OutColor = float4(FaceNormal,1);
+}
\ No newline at end of file
diff --git a/Engine/Shaders/Private/HairStrands/HairStrandsMeshUpdate.usf b/Engine/Shaders/Private/HairStrands/HairStrandsMeshUpdate.usf
new file mode 100644
--- /dev/null
+++ b/Engine/Shaders/Private/HairStrands/HairStrandsMeshUpdate.usf
@@ -0,0 +1,39 @@
+// Copyright 1998-2019 Epic Games, Inc. All Rights Reserved.
+
+#include "../Common.ush"
+#include "HairStrandsMeshProjectionCommon.ush"
+
+uint RootTriangleWorldCenter;
+float4x4 MeshLocalToWorld;
+uint MeshSectionIndex;
+uint MeshMaxIndexCount;
+uint MeshMaxVertexCount;
+uint MeshIndexOffset;
+Buffer<uint> MeshIndexBuffer;
+Buffer<float> MeshPositionBuffer;
+
+float4x4 RootLocalToWorld;
+uint MaxRootCount;
+Buffer<uint> RootTriangleIndex;
+
+RWBuffer<float4> OutRootTrianglePosition0;
+RWBuffer<float4> OutRootTrianglePosition1;
+RWBuffer<float4> OutRootTrianglePosition2;
+
+[numthreads(32, 1, 1)]
+void MainCS(uint3 DispatchThreadId : SV_DispatchThreadID)
+{ 
+	const uint RootIndex = DispatchThreadId.x;
+	if (RootIndex >= MaxRootCount) 
+		return;
+	
+	const FHairTriangleIndex TriangleIndex = DecodeTriangleIndex(RootTriangleIndex[RootIndex]);
+	if (TriangleIndex.SectionIndex != MeshSectionIndex)
+		return;
+
+	const FHairMeshTriangle Tri = GetTriangleData(TriangleIndex.TriangleIndex, MeshLocalToWorld, MeshPositionBuffer, MeshIndexBuffer, MeshIndexOffset, MeshMaxIndexCount, MeshMaxVertexCount);
+
+	OutRootTrianglePosition0[RootIndex] = (float4(Tri.P0 - RootTriangleWorldCenter,1));
+	OutRootTrianglePosition1[RootIndex] = (float4(Tri.P1 - RootTriangleWorldCenter,1));
+	OutRootTrianglePosition2[RootIndex] = (float4(Tri.P2 - RootTriangleWorldCenter,1));
+}
\ No newline at end of file
diff --git a/Engine/Source/Runtime/Core/Public/UObject/PhysicsObjectVersion.h b/Engine/Source/Runtime/Core/Public/UObject/PhysicsObjectVersion.h
--- a/Engine/Source/Runtime/Core/Public/UObject/PhysicsObjectVersion.h
+++ b/Engine/Source/Runtime/Core/Public/UObject/PhysicsObjectVersion.h
@@ -15,10 +15,6 @@ struct CORE_API FPhysicsObjectVersion
 		PerShapeData,
 		// Add serialization from handle back to particle
 		SerializeGTGeometryParticles,
-		// Temporary versioning to serialize hair guide data in addition to hair render data
-		HairAssetSerialization,
-		HairAssetSerialization_V2,
-
 		// -----<new versions can be added above this line>-------------------------------------------------
 		VersionPlusOne,
 		LatestVersion = VersionPlusOne - 1
diff --git a/Engine/Source/Runtime/Engine/Classes/Components/SkinnedMeshComponent.h b/Engine/Source/Runtime/Engine/Classes/Components/SkinnedMeshComponent.h
--- a/Engine/Source/Runtime/Engine/Classes/Components/SkinnedMeshComponent.h
+++ b/Engine/Source/Runtime/Engine/Classes/Components/SkinnedMeshComponent.h
@@ -129,6 +129,18 @@ struct FActiveMorphTarget
 	}
 };
 
+struct FSkeletalMeshObjectCallbackData
+{
+	enum class EEventType { Register, Unregister, Update };
+	typedef void (*TCallbackMeshObjectCallback)(
+		EEventType Event, 
+		class FSkeletalMeshObject* MeshObject,
+		uint64 UserData);
+
+	uint64 UserData = 0;
+	TCallbackMeshObjectCallback Run = nullptr;
+};
+
 /** Vertex skin weight info supplied for a component override. */
 USTRUCT(BlueprintType, meta = (HasNativeMake = "Engine.KismetRenderingLibrary.MakeSkinWeightInfo", HasNativeBreak = "Engine.KismetRenderingLibrary.BreakSkinWeightInfo"))
 struct FSkelMeshSkinWeightInfo
@@ -621,6 +633,7 @@ public:
 
 	/** Object responsible for sending bone transforms, morph target state etc. to render thread. */
 	class FSkeletalMeshObject*	MeshObject;
+	FSkeletalMeshObjectCallbackData MeshObjectCallbackData;
 
 	/** Gets the skeletal mesh resource used for rendering the component. */
 	FSkeletalMeshRenderData* GetSkeletalMeshRenderData() const;
diff --git a/Engine/Source/Runtime/Engine/Private/Components/SkinnedMeshComponent.cpp b/Engine/Source/Runtime/Engine/Private/Components/SkinnedMeshComponent.cpp
--- a/Engine/Source/Runtime/Engine/Private/Components/SkinnedMeshComponent.cpp
+++ b/Engine/Source/Runtime/Engine/Private/Components/SkinnedMeshComponent.cpp
@@ -782,6 +782,12 @@ void USkinnedMeshComponent::TickComponent(float DeltaTime, enum ELevelTick TickT
 	// Tick ActorComponent first.
 	Super::TickComponent(DeltaTime, TickType, ThisTickFunction);
 
+	// Refresh callback in case MeshObjectCallbackData was updated
+	if (MeshObject)
+	{
+		MeshObject->SetCallbackData(MeshObjectCallbackData);
+	}
+
 	// See if this mesh was rendered recently. This has to happen first because other data will rely on this
 	bRecentlyRendered = (GetLastRenderTime() > GetWorld()->TimeSeconds - 1.0f);
 
diff --git a/Engine/Source/Runtime/Engine/Private/GPUSkinCache.cpp b/Engine/Source/Runtime/Engine/Private/GPUSkinCache.cpp
--- a/Engine/Source/Runtime/Engine/Private/GPUSkinCache.cpp
+++ b/Engine/Source/Runtime/Engine/Private/GPUSkinCache.cpp
@@ -255,6 +255,22 @@ public:
 		DispatchData[Section].UpdateVertexFactoryDeclaration();
 	}
 
+	inline FCachedGeometrySection GetCachedGeometry(int32 SectionIndex) const
+	{
+		FCachedGeometrySection MeshSection;
+		const FSkelMeshRenderSection& Section = *DispatchData[SectionIndex].Section;
+		MeshSection.PositionBuffer	= DispatchData[SectionIndex].PositionBuffer->SRV;
+		MeshSection.TotalVertexCount= DispatchData[SectionIndex].PositionBuffer->NumBytes / (sizeof(float)*3);
+		MeshSection.NumPrimitives	= Section.NumTriangles;
+		MeshSection.IndexBaseIndex	= Section.BaseIndex;
+		MeshSection.VertexBaseIndex = Section.BaseVertexIndex;
+		MeshSection.IndexBuffer		= nullptr;
+		MeshSection.TotalIndexCount	= 0;
+		MeshSection.LODIndex		= 0;
+		MeshSection.SectionIndex	= SectionIndex;
+		return MeshSection;
+	}
+
 	bool IsSectionValid(int32 Section) const
 	{
 		const FSectionDispatchData& SectionData = DispatchData[Section];
@@ -1392,6 +1408,11 @@ void FGPUSkinCache::InvalidateAllEntries()
 	SET_MEMORY_STAT(STAT_GPUSkinCache_TangentsIntermediateMemUsed, 0);
 }
 
+FCachedGeometrySection FGPUSkinCache::GetCachedGeometry(FGPUSkinCacheEntry* InOutEntry, uint32 sectionIndex)
+{
+	return InOutEntry ? InOutEntry->GetCachedGeometry(sectionIndex) : FCachedGeometrySection();
+}
+
 void FGPUSkinCache::CVarSinkFunction()
 {
 	int32 NewGPUSkinCacheValue = CVarEnableGPUSkinCache.GetValueOnAnyThread() != 0;
diff --git a/Engine/Source/Runtime/Engine/Private/SkeletalRenderGPUSkin.cpp b/Engine/Source/Runtime/Engine/Private/SkeletalRenderGPUSkin.cpp
--- a/Engine/Source/Runtime/Engine/Private/SkeletalRenderGPUSkin.cpp
+++ b/Engine/Source/Runtime/Engine/Private/SkeletalRenderGPUSkin.cpp
@@ -123,6 +123,10 @@ void FMorphVertexBuffer::ReleaseDynamicRHI()
 /*-----------------------------------------------------------------------------
 FSkeletalMeshObjectGPUSkin
 -----------------------------------------------------------------------------*/
+void FSkeletalMeshObjectGPUSkin::SetCallbackData(FSkeletalMeshObjectCallbackData& InMeshObjectCallbackData)
+{
+	CallbackData = InMeshObjectCallbackData;
+}
 
 FSkeletalMeshObjectGPUSkin::FSkeletalMeshObjectGPUSkin(USkinnedMeshComponent* InMeshComponent, FSkeletalMeshRenderData* InSkelMeshRenderData, ERHIFeatureLevel::Type InFeatureLevel)
 	: FSkeletalMeshObject(InMeshComponent, InSkelMeshRenderData, InFeatureLevel)
@@ -140,11 +144,21 @@ FSkeletalMeshObjectGPUSkin::FSkeletalMeshObjectGPUSkin(USkinnedMeshComponent* In
 	}
 
 	InitResources(InMeshComponent);
+	CallbackData = InMeshComponent->MeshObjectCallbackData;
+	if (CallbackData.Run)
+	{
+		CallbackData.Run(FSkeletalMeshObjectCallbackData::EEventType::Register, this, CallbackData.UserData);
+	}
 }
 
 
 FSkeletalMeshObjectGPUSkin::~FSkeletalMeshObjectGPUSkin()
 {
+	if (CallbackData.Run)
+	{
+		CallbackData.Run(FSkeletalMeshObjectCallbackData::EEventType::Unregister, this, CallbackData.UserData);
+	}
+
 	check(!RHIThreadFenceForDynamicData.GetReference());
 	if (DynamicData)
 	{
@@ -290,6 +304,11 @@ void FSkeletalMeshObjectGPUSkin::Update(int32 LODIndex,USkinnedMeshComponent* In
 			MeshObject->UpdateDynamicData_RenderThread(GPUSkinCache, RHICmdList, NewDynamicData, nullptr, FrameNumberToPrepare, RevisionNumber);
 		}
 	);
+
+	if (CallbackData.Run)
+	{
+		CallbackData.Run(FSkeletalMeshObjectCallbackData::EEventType::Update, this, CallbackData.UserData);
+	}
 }
 
 void FSkeletalMeshObjectGPUSkin::UpdateSkinWeightBuffer(USkinnedMeshComponent* InMeshComponent)
@@ -513,6 +532,7 @@ void FSkeletalMeshObjectGPUSkin::ProcessUpdatedDynamicData(FGPUSkinCache* GPUSki
 		bDataPresent = VertexFactoryData.VertexFactories.Num() > 0;
 	}
 
+	CachedGeometry.Reset();
 	if (bDataPresent)
 	{
 		for (int32 SectionIdx = 0; SectionIdx < Sections.Num(); SectionIdx++)
@@ -600,6 +620,12 @@ void FSkeletalMeshObjectGPUSkin::ProcessUpdatedDynamicData(FGPUSkinCache* GPUSki
 				GPUSkinCache->ProcessEntry(RHICmdList, VertexFactory,
 					VertexFactoryData.PassthroughVertexFactories[SectionIdx].Get(), Section, this, bMorph ? &LOD.MorphVertexBuffer : 0, bClothFactory ? &LODData.ClothVertexBuffer : 0,
 					bClothFactory ? DynamicData->ClothingSimData.Find(Section.CorrespondClothAssetIndex) : 0, LocalToCloth, DynamicData->ClothBlendWeight, RevisionNumber, SectionIdx, SkinCacheEntry);
+
+				FCachedGeometrySection CachedSection = GPUSkinCache->GetCachedGeometry(SkinCacheEntry, SectionIdx);
+				CachedSection.IndexBuffer = LODData.MultiSizeIndexContainer.GetIndexBuffer()->GetSRV();
+				CachedSection.TotalIndexCount = LODData.MultiSizeIndexContainer.GetIndexBuffer()->Num();
+				CachedSection.LODIndex = DynamicData->LODIndex;
+				CachedGeometry.Sections.Add(CachedSection);
 			}
 			if (bNeedFence)
 			{
diff --git a/Engine/Source/Runtime/Engine/Private/SkeletalRenderGPUSkin.h b/Engine/Source/Runtime/Engine/Private/SkeletalRenderGPUSkin.h
--- a/Engine/Source/Runtime/Engine/Private/SkeletalRenderGPUSkin.h
+++ b/Engine/Source/Runtime/Engine/Private/SkeletalRenderGPUSkin.h
@@ -278,6 +278,8 @@ public:
 	virtual bool IsCPUSkinned() const override { return false; }
 	virtual TArray<FTransform>* GetComponentSpaceTransforms() const override;
 	virtual const TArray<FMatrix>& GetReferenceToLocalMatrices() const override;
+	virtual void SetCallbackData(FSkeletalMeshObjectCallbackData& CallbackData) override;
+	virtual FCachedGeometry GetCachedGeometry() const override { return CachedGeometry; };
 
 #if RHI_RAYTRACING
 	/** Geometry for ray tracing. */
@@ -548,6 +550,12 @@ private:
 
 	/** last updated bone transform revision number */
 	uint32 LastBoneTransformRevisionNumber;
+
+	/** Last cached geometry description */
+	FCachedGeometry CachedGeometry;
+
+	/** Callback for register/unregister/update events */
+	FSkeletalMeshObjectCallbackData CallbackData;
 };
 
 
