Add support group support for hair rendering.

(ushell-p4-cherrypick of 9439735 by Charles.deRousiers)

diff --git a/Engine/Plugins/Experimental/HairStrands/Source/HairStrandsCore/Private/GroomComponent.cpp b/Engine/Plugins/Experimental/HairStrands/Source/HairStrandsCore/Private/GroomComponent.cpp
--- a/Engine/Plugins/Experimental/HairStrands/Source/HairStrandsCore/Private/GroomComponent.cpp
+++ b/Engine/Plugins/Experimental/HairStrands/Source/HairStrandsCore/Private/GroomComponent.cpp
@@ -122,51 +122,62 @@ public:
 	FHairStrandsSceneProxy(UGroomComponent* Component)
 		: FPrimitiveSceneProxy(Component)
 		, VertexFactory(GetScene().GetFeatureLevel(), "FStrandHairSceneProxy")
-		, MaterialRelevance(Component->GetMaterialRelevance(GetScene().GetFeatureLevel()))
-	
+		, MaterialRelevance(Component->GetMaterialRelevance(GetScene().GetFeatureLevel()))	
 	{
 		check(Component);
 		check(Component->GroomAsset);
 		check(Component->GroomAsset->GetNumHairGroups() > 0);
+		check(Component->InterpolationOutput);
 
-		const FHairGroupData& GroupData = Component->GroomAsset->HairGroupsData[0];
-		const uint32 VertexCount = GroupData.HairRenderData.GetNumPoints();
-		const float MinHairRadius = 0; // Todo: Component->GroomAsset->HairRenderData.StrandsCurves.MinRadius;
-		const float MaxHairRadius = GroupData.HairRenderData.StrandsCurves.MaxRadius;
-		const float MaxHairLength = GroupData.HairRenderData.StrandsCurves.MaxLength;
-		const float HairDensity = GroupData.HairRenderData.HairDensity;
-		const FVector& HairWorldOffset = GroupData.HairRenderData.BoundingBox.GetCenter();
+		FHairStrandsVertexFactory::FDataType VFData;
+		VFData.InterpolationOutput = Component->InterpolationOutput;
 
-		check(Component->InterpolationOutput);
-		FHairStrandsInterpolationOutput* LocalOutput = Component->InterpolationOutput;
-		FHairStrandsVertexFactory* LocaVertexFactor = &VertexFactory;
-		ENQUEUE_RENDER_COMMAND(StaticMeshVertexBuffersLegacyInit)(
-		[LocalOutput, LocaVertexFactor, HairWorldOffset, MinHairRadius, MaxHairRadius, MaxHairLength, HairDensity](FRHICommandListImmediate& RHICmdList)
-		{
-			const uint32 Offset = 0;
-			FHairStrandsVertexFactory::FDataType Data;
-			Data.MinStrandRadius = MinHairRadius;
-			Data.MaxStrandRadius = MaxHairRadius;
-			Data.MaxStrandLength = MaxHairLength;
-			Data.HairDensity = HairDensity;
-			Data.HairWorldOffset = HairWorldOffset;
-			Data.InterpolationOutput = LocalOutput;
-
-			LocaVertexFactor->SetData(Data);
-			LocaVertexFactor->InitResource();
-		});
+		const uint32 GroupCount = Component->GroomAsset->GetNumHairGroups();
+		check(Component->GroomAsset->HairGroupsData.Num() == Component->HairGroupResources.Num());
+		for (uint32 GroupIt=0;GroupIt<GroupCount; GroupIt++)
+		{		
+			const FHairGroupData& InGroupData = Component->GroomAsset->HairGroupsData[GroupIt];
+			const UGroomComponent::FHairGroupResource& GroupResources = Component->HairGroupResources[GroupIt];
 
+			UMaterialInterface* Material = Component->GetMaterial(GroupIt);
+			if (Material == nullptr || !Material->GetMaterialResource(GetScene().GetFeatureLevel())->IsUsedWithHairStrands())
+			{
+				Material = GEngine->HairDefaultMaterial;
+			}
+
+			FHairStrandsVertexFactory::FDataType::HairGroup& VFGroupData = VFData.HairGroups.AddDefaulted_GetRef();
+			VFGroupData.MinStrandRadius		= 0;
+			VFGroupData.MaxStrandRadius		= InGroupData.HairRenderData.StrandsCurves.MaxRadius;
+			VFGroupData.MaxStrandLength		= InGroupData.HairRenderData.StrandsCurves.MaxLength;
+			VFGroupData.HairDensity			= InGroupData.HairRenderData.HairDensity;
+			VFGroupData.HairWorldOffset		= InGroupData.HairRenderData.BoundingBox.GetCenter();
+
+			#if RHI_RAYTRACING
+			FRayTracingGeometry* RayTracingGeometry = nullptr;
+			if (IsRayTracingEnabled() && GroupResources.RaytracingResources)
+			{
+				RayTracingGeometry = &GroupResources.RaytracingResources->RayTracingGeometry;
+			}
+			#endif
+
+			HairGroup& OutGroupData = HairGroups.Add_GetRef(
+			{
+				GroupIt,
+				Material
+				#if RHI_RAYTRACING
+				, RayTracingGeometry
+				#endif
+			});
 
-		Material = Component->GetMaterial(0);
-		if (Material == nullptr || !Material->GetMaterialResource(GetScene().GetFeatureLevel())->IsUsedWithHairStrands())
-		{
-			Material = GEngine->HairDefaultMaterial;
 		}
 
-		#if RHI_RAYTRACING
-		if (IsRayTracingEnabled())
-			RayTracingGeometry = Component->RaytracingResources ? &Component->RaytracingResources->RayTracingGeometry : nullptr;
-		#endif
+		FHairStrandsVertexFactory* LocalVertexFactory = &VertexFactory;
+		ENQUEUE_RENDER_COMMAND(InitHairStrandsVertexFactory)(
+			[LocalVertexFactory, VFData](FRHICommandListImmediate& RHICmdList)
+		{
+			LocalVertexFactory->SetData(VFData);
+			LocalVertexFactory->InitResource();
+		});
 	}
 
 	virtual ~FHairStrandsSceneProxy()
@@ -180,94 +191,121 @@ public:
 
 	virtual void GetDynamicRayTracingInstances(FRayTracingMaterialGatheringContext & Context, TArray<struct FRayTracingInstance>& OutRayTracingInstances) override
 	{
-		if (!IsRayTracingEnabled() || !RayTracingGeometry || !RayTracingGeometry->RayTracingGeometryRHI.IsValid())
+		if (!IsRayTracingEnabled() || HairGroups.Num() == 0)
 			return;
 
-		for (const FRayTracingGeometrySegment& Segment : RayTracingGeometry->Initializer.Segments)
+		for (const HairGroup& GroupData : HairGroups)
 		{
-			check(Segment.VertexBuffer.IsValid());
+			if (GroupData.RayTracingGeometry && GroupData.RayTracingGeometry->RayTracingGeometryRHI.IsValid())
+			{
+				for (const FRayTracingGeometrySegment& Segment : GroupData.RayTracingGeometry->Initializer.Segments)
+				{
+					check(Segment.VertexBuffer.IsValid());
+				}
+				AddOpaqueRaytracingInstance(GetLocalToWorld(), GroupData.RayTracingGeometry, RaytracingInstanceMask_ThinShadow, OutRayTracingInstances);
+			}
 		}
-
-		AddOpaqueRaytracingInstance(GetLocalToWorld(), RayTracingGeometry, RaytracingInstanceMask_ThinShadow, OutRayTracingInstances);
 	}
 #endif
 
 	virtual void GetDynamicMeshElements(const TArray<const FSceneView*>& Views, const FSceneViewFamily& ViewFamily, uint32 VisibilityMap, FMeshElementCollector& Collector) const override
 	{
-		const uint32 HairVertexCount = VertexFactory.GetData().InterpolationOutput ? VertexFactory.GetData().InterpolationOutput->VFInput.VertexCount : 0;
-		if (HairVertexCount == 0) return;
+		const uint32 GroupCount = HairGroups.Num();
+		check(GroupCount == VertexFactory.GetData().HairGroups.Num());
+		if (GroupCount == 0)
+			return;
+
+		bool bHasOneElementValid = false;
+		for (FHairStrandsInterpolationOutput::HairGroup& HairGroup : VertexFactory.GetData().InterpolationOutput->HairGroups)
+		{
+			if (HairGroup.VFInput.VertexCount > 0)
+			{
+				bHasOneElementValid = true;
+				break;
+			}
+		}
+		if (!bHasOneElementValid) 
+			return;
 
 		QUICK_SCOPE_CYCLE_COUNTER(STAT_HairStrandsSceneProxy_GetDynamicMeshElements);
 
+
 		FMaterialRenderProxy* MaterialProxy = nullptr;
+		const EHairStrandsDebugMode DebugMode = GetHairStrandsDebugStrandsMode();
+		if (DebugMode != EHairStrandsDebugMode::None)
 		{
-			const EHairStrandsDebugMode DebugMode = GetHairStrandsDebugStrandsMode();
-			if (DebugMode != EHairStrandsDebugMode::None)
+			float DebugModeScalar = 0;
+			switch(DebugMode)
 			{
-				float DebugModeScalar = 0;
-				switch(DebugMode)
-				{
-				case EHairStrandsDebugMode::None						: DebugModeScalar =99.f; break;
-				case EHairStrandsDebugMode::SimHairStrands				: DebugModeScalar = 0.f; break;
-				case EHairStrandsDebugMode::RenderHairStrands			: DebugModeScalar = 0.f; break;
-				case EHairStrandsDebugMode::RenderHairRootUV			: DebugModeScalar = 1.f; break;
-				case EHairStrandsDebugMode::RenderHairUV				: DebugModeScalar = 2.f; break;
-				case EHairStrandsDebugMode::RenderHairSeed				: DebugModeScalar = 3.f; break;
-				case EHairStrandsDebugMode::RenderHairDimension			: DebugModeScalar = 4.f; break;
-				case EHairStrandsDebugMode::RenderHairRadiusVariation	: DebugModeScalar = 5.f; break;
-				};
-
-				const float HairMinRadius = VertexFactory.GetMinStrandRadius();
-				const float HairMaxRadius = VertexFactory.GetMaxStrandRadius();
-				const float HairClipLength = GetHairClipLength();
-				auto DebugMaterial = new FHairDebugModeMaterialRenderProxy(
-					GEngine->HairDebugMaterial ? GEngine->HairDebugMaterial->GetRenderProxy() : nullptr,
-					DebugModeScalar, HairMinRadius, HairMaxRadius, HairClipLength);
-				Collector.RegisterOneFrameMaterialProxy(DebugMaterial);
-				MaterialProxy = DebugMaterial;
-			}
-			else
+			case EHairStrandsDebugMode::None						: DebugModeScalar =99.f; break;
+			case EHairStrandsDebugMode::SimHairStrands				: DebugModeScalar = 0.f; break;
+			case EHairStrandsDebugMode::RenderHairStrands			: DebugModeScalar = 0.f; break;
+			case EHairStrandsDebugMode::RenderHairRootUV			: DebugModeScalar = 1.f; break;
+			case EHairStrandsDebugMode::RenderHairUV				: DebugModeScalar = 2.f; break;
+			case EHairStrandsDebugMode::RenderHairSeed				: DebugModeScalar = 3.f; break;
+			case EHairStrandsDebugMode::RenderHairDimension			: DebugModeScalar = 4.f; break;
+			case EHairStrandsDebugMode::RenderHairRadiusVariation	: DebugModeScalar = 5.f; break;
+			};
+
+			float HairMinRadius = FLT_MAX;
+			float HairMaxRadius = 0;
+			for (uint32 GroupIt = 0; GroupIt < GroupCount; ++GroupIt)
 			{
-				MaterialProxy = Material->GetRenderProxy();
+				HairMinRadius = FMath::Min(HairMinRadius, VertexFactory.GetMinStrandRadius(GroupIt));
+				HairMaxRadius = FMath::Max(HairMaxRadius, VertexFactory.GetMaxStrandRadius(GroupIt));
 			}
+			const float HairClipLength = GetHairClipLength();
+			auto DebugMaterial = new FHairDebugModeMaterialRenderProxy(
+				GEngine->HairDebugMaterial ? GEngine->HairDebugMaterial->GetRenderProxy() : nullptr,
+				DebugModeScalar, HairMinRadius, HairMaxRadius, HairClipLength);
+			Collector.RegisterOneFrameMaterialProxy(DebugMaterial);
+			MaterialProxy = DebugMaterial;
 		}
-
+		
 		for (int32 ViewIndex = 0; ViewIndex < Views.Num(); ViewIndex++)
 		{
 			if (VisibilityMap & (1 << ViewIndex))
 			{
-				const FSceneView* View = Views[ViewIndex];
-				// Draw the mesh.
-				FMeshBatch& Mesh = Collector.AllocateMesh();
-				FMeshBatchElement& BatchElement = Mesh.Elements[0];
-				BatchElement.IndexBuffer = nullptr;
-				Mesh.bWireframe = false;
-				Mesh.VertexFactory = &VertexFactory;
-				Mesh.MaterialRenderProxy = MaterialProxy;
-				bool bHasPrecomputedVolumetricLightmap;
-				FMatrix PreviousLocalToWorld;
-				int32 SingleCaptureIndex;
-				bool bOutputVelocity = false;
-				bool bDrawVelocity = false; // Velocity vector is done in a custom fashion
-				GetScene().GetPrimitiveUniformShaderParameters_RenderThread(GetPrimitiveSceneInfo(), bHasPrecomputedVolumetricLightmap, PreviousLocalToWorld, SingleCaptureIndex, bOutputVelocity);
-				bOutputVelocity = false;
-				FDynamicPrimitiveUniformBuffer& DynamicPrimitiveUniformBuffer = Collector.AllocateOneFrameResource<FDynamicPrimitiveUniformBuffer>();
-				DynamicPrimitiveUniformBuffer.Set(GetLocalToWorld(), PreviousLocalToWorld, GetBounds(), GetLocalBounds(), true, false, bDrawVelocity, bOutputVelocity);
-				BatchElement.PrimitiveUniformBufferResource = &DynamicPrimitiveUniformBuffer.UniformBuffer;
-				BatchElement.FirstIndex = 0;
-				BatchElement.NumPrimitives = HairVertexCount * 2;
-				BatchElement.MinVertexIndex = 0;
-				BatchElement.MaxVertexIndex = HairVertexCount * 6;
-				Mesh.ReverseCulling = IsLocalToWorldDeterminantNegative();
-				Mesh.Type = PT_TriangleList;
-				Mesh.DepthPriorityGroup = SDPG_World;
-				Mesh.bCanApplyViewModeOverrides = false;
-				Collector.AddMesh(ViewIndex, Mesh);
-
-			#if !(UE_BUILD_SHIPPING || UE_BUILD_TEST)
-				// Render bounds
-				RenderBounds(Collector.GetPDI(ViewIndex), ViewFamily.EngineShowFlags, GetBounds(), IsSelected());
-			#endif
+				for (uint32 GroupIt = 0; GroupIt < GroupCount; ++GroupIt)
+				{
+					const HairGroup& GroupData = HairGroups[GroupIt];
+					const uint32 HairVertexCount = VertexFactory.GetData().InterpolationOutput->HairGroups[GroupIt].VFInput.VertexCount;
+					const FHairStrandsVertexFactory::FDataType::HairGroup& HairGroup = VertexFactory.GetData().HairGroups[GroupIt];
+
+					const FSceneView* View = Views[ViewIndex];
+					// Draw the mesh.
+					FMeshBatch& Mesh = Collector.AllocateMesh();
+					FMeshBatchElement& BatchElement = Mesh.Elements[0];
+					BatchElement.IndexBuffer = nullptr;
+					Mesh.bWireframe = false;
+					Mesh.VertexFactory = &VertexFactory;
+					Mesh.MaterialRenderProxy = MaterialProxy == nullptr ? GroupData.Material->GetRenderProxy() : MaterialProxy;
+					bool bHasPrecomputedVolumetricLightmap;
+					FMatrix PreviousLocalToWorld;
+					int32 SingleCaptureIndex;
+					bool bOutputVelocity = false;
+					bool bDrawVelocity = false; // Velocity vector is done in a custom fashion
+					GetScene().GetPrimitiveUniformShaderParameters_RenderThread(GetPrimitiveSceneInfo(), bHasPrecomputedVolumetricLightmap, PreviousLocalToWorld, SingleCaptureIndex, bOutputVelocity);
+					bOutputVelocity = false;
+					FDynamicPrimitiveUniformBuffer& DynamicPrimitiveUniformBuffer = Collector.AllocateOneFrameResource<FDynamicPrimitiveUniformBuffer>();
+					DynamicPrimitiveUniformBuffer.Set(GetLocalToWorld(), PreviousLocalToWorld, GetBounds(), GetLocalBounds(), true, false, bDrawVelocity, bOutputVelocity);
+					BatchElement.PrimitiveUniformBufferResource = &DynamicPrimitiveUniformBuffer.UniformBuffer;
+					BatchElement.FirstIndex = 0;
+					BatchElement.NumPrimitives = HairVertexCount * 2;
+					BatchElement.MinVertexIndex = 0;
+					BatchElement.MaxVertexIndex = HairVertexCount * 6;
+					BatchElement.UserData = reinterpret_cast<void*>(uint64(GroupIt));
+					Mesh.ReverseCulling = IsLocalToWorldDeterminantNegative();
+					Mesh.Type = PT_TriangleList;
+					Mesh.DepthPriorityGroup = SDPG_World;
+					Mesh.bCanApplyViewModeOverrides = false;
+					Collector.AddMesh(ViewIndex, Mesh);
+
+				#if !(UE_BUILD_SHIPPING || UE_BUILD_TEST)
+					// Render bounds
+					RenderBounds(Collector.GetPDI(ViewIndex), ViewFamily.EngineShowFlags, GetBounds(), IsSelected());
+				#endif
+				}
 			}
 		}
 	}
@@ -310,12 +348,18 @@ public:
 
 	uint32 GetAllocatedSize(void) const { return(FPrimitiveSceneProxy::GetAllocatedSize()); }
 private:
-	UMaterialInterface* Material = nullptr;
+
 	FHairStrandsVertexFactory VertexFactory;
 	FMaterialRelevance MaterialRelevance;
-#if RHI_RAYTRACING
-	FRayTracingGeometry* RayTracingGeometry = nullptr;
-#endif
+	struct HairGroup
+	{
+		uint32 GroupIndex = 0;
+		UMaterialInterface* Material = nullptr;
+	#if RHI_RAYTRACING
+		FRayTracingGeometry* RayTracingGeometry = nullptr;
+	#endif
+	};
+	TArray<HairGroup> HairGroups;
 };
 
 
@@ -337,8 +381,6 @@ UGroomComponent::UGroomComponent(const FObjectInitializer& ObjectInitializer)
 	HairDensity = 1;
 	bSkinGroom = false;
 	InitializedResources = nullptr;
-	RenRootResources = nullptr;
-	SimRootResources = nullptr;
 	Mobility = EComponentMobility::Movable;
 	MeshProjectionTickDelay = 0;
 	MeshProjectionLODIndex = -1;
@@ -367,19 +409,32 @@ FBoxSphereBounds UGroomComponent::CalcBounds(const FTransform& InLocalToWorld) c
 	{
 		if (RegisteredSkeletalMeshComponent)
 		{
-			// Transform the bounding box with an extra offset coming from the skin animation
-			const FBoxSphereBounds MeshLocalBound = RegisteredSkeletalMeshComponent->CalcBounds(FTransform::Identity);
-			const FVector HairTranslation = GroomAsset->HairGroupsData[0].HairRenderData.BoundingBox.GetCenter();
-			const FVector MeshTranslation = MeshLocalBound.GetSphere().Center;
-			const FVector LocalTranslation = MeshTranslation - HairTranslation;
-
-			FTransform LocalToWorld = InLocalToWorld;
-			LocalToWorld.SetLocation(LocalTranslation + InLocalToWorld.GetLocation());
-			return FBoxSphereBounds(GroomAsset->HairGroupsData[0].HairRenderData.BoundingBox.TransformBy(LocalToWorld));
+			FBoxSphereBounds WorldBounds(EForceInit::ForceInit);
+			for (const FHairGroupData& GroupData : GroomAsset->HairGroupsData)
+			{
+				// Transform the bounding box with an extra offset coming from the skin animation
+				const FBoxSphereBounds MeshLocalBound = RegisteredSkeletalMeshComponent->CalcBounds(FTransform::Identity);
+				const FVector HairTranslation = GroupData.HairRenderData.BoundingBox.GetCenter();
+				const FVector MeshTranslation = MeshLocalBound.GetSphere().Center;
+				const FVector LocalTranslation = MeshTranslation - HairTranslation;
+
+				FTransform LocalToWorld = InLocalToWorld;
+				LocalToWorld.SetLocation(LocalTranslation + InLocalToWorld.GetLocation());
+				FBoxSphereBounds GroupWorldBound(GroupData.HairRenderData.BoundingBox.TransformBy(LocalToWorld));
+			
+				WorldBounds = Union(GroupWorldBound, WorldBounds);
+			}
+			return WorldBounds;
+
 		}
 		else
 		{
-			return FBoxSphereBounds(GroomAsset->HairGroupsData[0].HairRenderData.BoundingBox.TransformBy(InLocalToWorld));
+			FBox LocalBounds(EForceInit::ForceInit);
+			for (const FHairGroupData& GroupData : GroomAsset->HairGroupsData)
+			{
+				LocalBounds += GroupData.HairRenderData.BoundingBox;
+			}
+			return FBoxSphereBounds(LocalBounds.TransformBy(InLocalToWorld));
 		}
 	}
 	else
@@ -426,9 +481,11 @@ UMaterialInterface* UGroomComponent::GetMaterial(int32 ElementIndex) const
 	return OverrideMaterial;
 }
 
-static FHairStrandsProjectionHairData ToProjectionHairData(FHairStrandsRootResource* In)
+static FHairStrandsProjectionHairData::HairGroup ToProjectionHairData(FHairStrandsRootResource* In)
 {
-	FHairStrandsProjectionHairData Out = {};
+	check(IsInRenderingThread());
+
+	FHairStrandsProjectionHairData::HairGroup Out = {};
 	if (!In)
 		return Out;
 
@@ -461,24 +518,24 @@ static FHairStrandsProjectionHairData ToProjectionHairData(FHairStrandsRootResou
 	return Out;
 }
 
-FHairStrandsDatas* UGroomComponent::GetGuideStrandsDatas()
+FHairStrandsDatas* UGroomComponent::GetGuideStrandsDatas(uint32 GroupIndex)
 {
-	if (!GroomAsset || GroomAsset->GetNumHairGroups() == 0)
+	if (!GroomAsset || GroupIndex >= uint32(GroomAsset->GetNumHairGroups()))
 	{
 		return nullptr;
 	}
 
-	return &GroomAsset->HairGroupsData[0].HairSimulationData;
+	return &GroomAsset->HairGroupsData[GroupIndex].HairSimulationData;
 }
 
-FHairStrandsResource* UGroomComponent::GetGuideStrandsResource()
+FHairStrandsResource* UGroomComponent::GetGuideStrandsResource(uint32 GroupIndex)
 {
-	if (!GroomAsset || GroomAsset->GetNumHairGroups() == 0)
+	if (!GroomAsset || GroupIndex >= uint32(GroomAsset->GetNumHairGroups()))
 	{
 		return nullptr;
 	}
 
-	return GroomAsset->HairGroupsData[0].HairSimulationResource;
+	return GroomAsset->HairGroupsData[GroupIndex].HairSimulationResource;
 }
 
 template<typename T> void SafeDelete(T*& Data) 
@@ -535,142 +592,150 @@ void UGroomComponent::InitResources()
 {
 	ReleaseResources();
 
-	if (GroomAsset && GroomAsset->GetNumHairGroups() > 0 && GroomAsset->HairGroupsData[0].HairStrandsResource)
+	if (!GroomAsset || GroomAsset->GetNumHairGroups() == 0)
+		return;
+
+	InitializedResources = GroomAsset;
+
+	const FPrimitiveComponentId LocalComponentId = ComponentId;
+	EWorldType::Type WorldType = GetWorld() ? EWorldType::Type(GetWorld()->WorldType) : EWorldType::None;
+	WorldType = WorldType == EWorldType::Inactive ? EWorldType::Editor : WorldType;
+
+	// Insure the ticking of the Groom component always happens after the skeletalMeshComponent.
+	USkeletalMeshComponent* SkeletalMeshComponent = bSkinGroom && GetAttachParent() ? Cast<USkeletalMeshComponent>(GetAttachParent()) : nullptr;
+	if (SkeletalMeshComponent)
 	{
-		check(GroomAsset->HairGroupsData[0].HairStrandsResource);
+		RegisteredSkeletalMeshComponent = SkeletalMeshComponent;
+		AddTickPrerequisiteComponent(SkeletalMeshComponent);
+
+		FSkeletalMeshObjectCallbackData CallbackData;
+		CallbackData.Run = CallbackMeshObjectCallback;
+		CallbackData.UserData = (uint64(LocalComponentId.PrimIDValue) & 0xFFFFFFFF) | (uint64(WorldType) << 32);
+		SkeletalMeshComponent->MeshObjectCallbackData = CallbackData;
+	}
+
+	FTransform HairLocalToWorld = GetComponentTransform();
+	FTransform SkinLocalToWorld = bSkinGroom && SkeletalMeshComponent ? SkeletalMeshComponent->GetComponentTransform() : FTransform::Identity;
+	
+	InterpolationOutput = new FHairStrandsInterpolationOutput();
+	InterpolationInput = new FHairStrandsInterpolationInput();
 
-		InitializedResources = GroomAsset;
+	FHairStrandsDebugInfo DebugGroupInfo;
+	for (FHairGroupData& GroupData : GroomAsset->HairGroupsData)
+	{
+		if (!GroupData.HairStrandsResource)
+			return;
 
-		FHairStrandsDebugInfo::FGroupInfos DebugGroupInfos;
-		for (const FHairGroupData GroupData : GroomAsset->HairGroupsData)
+		FHairStrandsDebugInfo::HairGroup& DebugHairGroup = DebugGroupInfo.HairGroups.AddDefaulted_GetRef();
 		{
-			FHairStrandsDebugInfo::FGroupInfo& DebugGroupInfo = DebugGroupInfos.AddDefaulted_GetRef();
-			DebugGroupInfo.MaxLength	= GroupData.HairRenderData.StrandsCurves.MaxLength;
-			DebugGroupInfo.MaxRadius	= GroupData.HairRenderData.StrandsCurves.MaxRadius;
-			DebugGroupInfo.VertexCount	= GroupData.HairRenderData.GetNumPoints();
-			DebugGroupInfo.CurveCount	= GroupData.HairRenderData.GetNumCurves();
+			DebugHairGroup.MaxLength	= GroupData.HairRenderData.StrandsCurves.MaxLength;
+			DebugHairGroup.MaxRadius	= GroupData.HairRenderData.StrandsCurves.MaxRadius;
+			DebugHairGroup.VertexCount	= GroupData.HairRenderData.GetNumPoints();
+			DebugHairGroup.CurveCount	= GroupData.HairRenderData.GetNumCurves();
 		}
 
-		const FHairGroupData& GroupData = GroomAsset->HairGroupsData[0];
-
-		InterpolationResource = new FHairStrandsInterpolationResource(GroupData.HairInterpolationData.RenderData, GroupData.HairSimulationData);
-		BeginInitResource(InterpolationResource);
+		FHairGroupResource& Res = HairGroupResources.AddDefaulted_GetRef();
+		Res.InterpolationResource = new FHairStrandsInterpolationResource(GroupData.HairInterpolationData.RenderData, GroupData.HairSimulationData);
+		BeginInitResource(Res.InterpolationResource);
 
 		#if RHI_RAYTRACING
-		FHairStrandsRaytracingResource* LocalRaytracingResources = nullptr;
 		if (IsRayTracingEnabled())
 		{
-			RaytracingResources = new FHairStrandsRaytracingResource(GroupData.HairRenderData);
-			BeginInitResource(RaytracingResources);
-			LocalRaytracingResources = RaytracingResources;
+			Res.RaytracingResources = new FHairStrandsRaytracingResource(GroupData.HairRenderData);
+			BeginInitResource(Res.RaytracingResources);
 		}
 		#endif
 
-
-		const FPrimitiveComponentId LocalComponentId = ComponentId;
-		EWorldType::Type WorldType = GetWorld() ? EWorldType::Type(GetWorld()->WorldType) : EWorldType::None;
-		WorldType = WorldType == EWorldType::Inactive ? EWorldType::Editor : WorldType;
-
-		// Insure the ticking of the Groom component always happens after the skeletalMeshComponent.
-		USkeletalMeshComponent* SkeletalMeshComponent = bSkinGroom && GetAttachParent() ? Cast<USkeletalMeshComponent>(GetAttachParent()) : nullptr;
 		if (SkeletalMeshComponent)
 		{
-			AddTickPrerequisiteComponent(SkeletalMeshComponent);
-			RegisteredSkeletalMeshComponent = SkeletalMeshComponent;
-
-			FSkeletalMeshObjectCallbackData CallbackData;
-			CallbackData.Run = CallbackMeshObjectCallback;
-			CallbackData.UserData = (uint64(LocalComponentId.PrimIDValue) & 0xFFFFFFFF) | (uint64(WorldType) << 32);
-			SkeletalMeshComponent->MeshObjectCallbackData = CallbackData;
-
 			const uint32 LODCount = SkeletalMeshComponent->GetNumLODs();
 			if (LODCount > 0)
 			{
-				RenRootResources = new FHairStrandsRootResource(&GroupData.HairRenderData, LODCount);
-				SimRootResources = new FHairStrandsRootResource(&GroupData.HairSimulationData, LODCount);
-				BeginInitResource(RenRootResources);
-				BeginInitResource(SimRootResources);
+				Res.RenRootResources = new FHairStrandsRootResource(&GroupData.HairRenderData, LODCount);
+				Res.SimRootResources = new FHairStrandsRootResource(&GroupData.HairSimulationData, LODCount);
+				BeginInitResource(Res.RenRootResources);
+				BeginInitResource(Res.SimRootResources);
 			}
 		}
 		
-		InterpolationOutput = new FHairStrandsInterpolationOutput();
-		InterpolationInput = new FHairStrandsInterpolationInput();
-
-		FHairStrandsInterpolationInput* Input = InterpolationInput;
-		FHairStrandsInterpolationOutput* Output = InterpolationOutput;
-		FHairStrandsResource* RenderResources = GroupData.HairStrandsResource;
-		FHairStrandsResource* SimResources = GroupData.HairSimulationResource;
-		FHairStrandsInterpolationResource* LocalInterpolationResource = InterpolationResource;
+		Res.RenderResources = GroupData.HairStrandsResource;
+		Res.SimResources = GroupData.HairSimulationResource;
 
-		check(Input);
-		Input->HairRadius = GroupData.HairRenderData.StrandsCurves.MaxRadius;
-		Input->HairWorldOffset = GroupData.HairRenderData.BoundingBox.GetCenter();
+		FHairStrandsInterpolationOutput::HairGroup& InterpolationOutputGroup = InterpolationOutput->HairGroups.AddDefaulted_GetRef();
+		FHairStrandsInterpolationInput::FHairGroup& InterpolationInputGroup = InterpolationInput->HairGroups.AddDefaulted_GetRef();
+		InterpolationInputGroup.HairRadius		= GroupData.HairRenderData.StrandsCurves.MaxRadius;
+		InterpolationInputGroup.HairWorldOffset = GroupData.HairRenderData.BoundingBox.GetCenter();
+	}
 
+	FHairStrandsInterpolationData Interpolation;
+	Interpolation.Input  = InterpolationInput;
+	Interpolation.Output = InterpolationOutput;
+	Interpolation.Function = ComputeHairStrandsInterpolation;
 
-		FHairStrandsRootResource* LocalRenRootResources = RenRootResources;
-		FHairStrandsRootResource* LocalSimRootResources = SimRootResources;
-		FTransform HairLocalToWorld = GetComponentTransform();
-		FTransform SkinLocalToWorld = bSkinGroom && SkeletalMeshComponent ? SkeletalMeshComponent->GetComponentTransform() : FTransform::Identity;
-		
-		const uint64 Id = LocalComponentId.PrimIDValue;
-		ENQUEUE_RENDER_COMMAND(FHairStrandsBuffers)(
-			[Id, Input, Output, 
-			RenderResources, SimResources, LocalRenRootResources, LocalSimRootResources, LocalInterpolationResource,
+	FHairGroupResources* LocalResources = &HairGroupResources;
+	const uint64 Id = LocalComponentId.PrimIDValue;
+	ENQUEUE_RENDER_COMMAND(FHairStrandsBuffers)(
+		[
+			Id,
+			Interpolation,
+			LocalResources,
 			HairLocalToWorld, SkinLocalToWorld,
 			WorldType,
-			DebugGroupInfos
-			#if RHI_RAYTRACING
-			, LocalRaytracingResources
-			#endif
-			](FRHICommandListImmediate& RHICmdList)
+			DebugGroupInfo
+		]
+		(FRHICommandListImmediate& RHICmdList)
+	{
+		FHairStrandsProjectionHairData RenProjectionDatas;
+		FHairStrandsProjectionHairData SimProjectionDatas;
+		const uint32 GroupCount = LocalResources->Num();
+		for (uint32 GroupIt=0;GroupIt<GroupCount; ++GroupIt)
 		{
-			Input->RenderRestPosePositionBuffer = &RenderResources->RestPositionBuffer;
-			Input->RenderAttributeBuffer		= &RenderResources->AttributeBuffer;
-			Input->RenderVertexCount			= RenderResources->RenderData.RenderingPositions.Num() / FHairStrandsPositionFormat::ComponentCount;
+			FHairGroupResource& Res = (*LocalResources)[GroupIt];
+
+			FHairStrandsInterpolationInput::FHairGroup& InputGroup		= Interpolation.Input->HairGroups[GroupIt];
+			FHairStrandsInterpolationOutput::HairGroup& OutputGroup	= Interpolation.Output->HairGroups[GroupIt];
 
-			Input->SimRestPosePositionBuffer	= &SimResources->RestPositionBuffer;
-			Input->SimAttributeBuffer			= &SimResources->AttributeBuffer;
-			Input->SimVertexCount				= SimResources->RenderData.RenderingPositions.Num() / FHairStrandsPositionFormat::ComponentCount;
-			Input->SimRootPointIndexBuffer		= &LocalInterpolationResource->SimRootPointIndexBuffer;
+			InputGroup.RenderRestPosePositionBuffer	= &Res.RenderResources->RestPositionBuffer;
+			InputGroup.RenderAttributeBuffer		= &Res.RenderResources->AttributeBuffer;
+			InputGroup.RenderVertexCount			=  Res.RenderResources->RenderData.RenderingPositions.Num() / FHairStrandsPositionFormat::ComponentCount;
 
-			Input->Interpolation0Buffer			= &LocalInterpolationResource->Interpolation0Buffer;
-			Input->Interpolation1Buffer			= &LocalInterpolationResource->Interpolation1Buffer;
+			InputGroup.SimRestPosePositionBuffer	= &Res.SimResources->RestPositionBuffer;
+			InputGroup.SimAttributeBuffer			= &Res.SimResources->AttributeBuffer;
+			InputGroup.SimVertexCount				=  Res.SimResources->RenderData.RenderingPositions.Num() / FHairStrandsPositionFormat::ComponentCount;
+			InputGroup.SimRootPointIndexBuffer		= &Res.InterpolationResource->SimRootPointIndexBuffer;
+
+			InputGroup.Interpolation0Buffer			= &Res.InterpolationResource->Interpolation0Buffer;
+			InputGroup.Interpolation1Buffer			= &Res.InterpolationResource->Interpolation1Buffer;
 
 			#if RHI_RAYTRACING
 			if (IsRayTracingEnabled())
 			{
-				Input->RaytracingGeometry		= &LocalRaytracingResources->RayTracingGeometry;
-				Input->RaytracingPositionBuffer	= &LocalRaytracingResources->PositionBuffer;
-				Input->RaytracingVertexCount	= LocalRaytracingResources->VertexCount;
+				InputGroup.RaytracingGeometry		= &Res.RaytracingResources->RayTracingGeometry;
+				InputGroup.RaytracingPositionBuffer	= &Res.RaytracingResources->PositionBuffer;
+				InputGroup.RaytracingVertexCount	=  Res.RaytracingResources->VertexCount;
 			}
 			#endif
 
-			Output->SimDeformedPositionBuffer[0]	= &SimResources->DeformedPositionBuffer[0];
-			Output->SimDeformedPositionBuffer[1]	= &SimResources->DeformedPositionBuffer[1];
-			Output->RenderDeformedPositionBuffer[0] = &RenderResources->DeformedPositionBuffer[0];
-			Output->RenderDeformedPositionBuffer[1] = &RenderResources->DeformedPositionBuffer[1];
-			Output->RenderAttributeBuffer			= &RenderResources->AttributeBuffer;
-			Output->RenderTangentBuffer				= &RenderResources->TangentBuffer;
-			Output->SimTangentBuffer				= &SimResources->TangentBuffer;
-
-			FHairStrandsProjectionHairData RenProjectionData = ToProjectionHairData(LocalRenRootResources);
-			FHairStrandsProjectionHairData SimProjectionData = ToProjectionHairData(LocalSimRootResources);
-
-			FHairStrandsInterpolationData Interpolation;
-			Interpolation.Input = Input;
-			Interpolation.Output = Output;
-			Interpolation.Function = ComputeHairStrandsInterpolation;
-
-			RegisterHairStrands(
-				Id,
-				WorldType,
-				Interpolation,
-				RenProjectionData, 
-				SimProjectionData,
-				DebugGroupInfos);
-		});
+			OutputGroup.SimDeformedPositionBuffer[0]	= &Res.SimResources->DeformedPositionBuffer[0];
+			OutputGroup.SimDeformedPositionBuffer[1]	= &Res.SimResources->DeformedPositionBuffer[1];
+			OutputGroup.RenderDeformedPositionBuffer[0] = &Res.RenderResources->DeformedPositionBuffer[0];
+			OutputGroup.RenderDeformedPositionBuffer[1] = &Res.RenderResources->DeformedPositionBuffer[1];
+			OutputGroup.RenderAttributeBuffer			= &Res.RenderResources->AttributeBuffer;
+			OutputGroup.RenderTangentBuffer				= &Res.RenderResources->TangentBuffer;
+			OutputGroup.SimTangentBuffer				= &Res.SimResources->TangentBuffer;
 
-	}
+			RenProjectionDatas.HairGroups.Add(ToProjectionHairData(Res.RenRootResources));
+			SimProjectionDatas.HairGroups.Add(ToProjectionHairData(Res.SimRootResources));
+		}
+
+		RegisterHairStrands(
+			Id,
+			WorldType,
+			Interpolation,
+			RenProjectionDatas, 
+			SimProjectionDatas,
+			DebugGroupInfo);
+	});
 }
 
 void UGroomComponent::ReleaseResources()
@@ -685,9 +750,15 @@ void UGroomComponent::ReleaseResources()
 		UnregisterHairStrands(Id, WorldType);
 	});
 
-	SafeRelease(InterpolationResource);
-	SafeRelease(RenRootResources);
-	SafeRelease(SimRootResources);
+	for (FHairGroupResource& Res : HairGroupResources)
+	{
+		SafeRelease(Res.InterpolationResource);
+		SafeRelease(Res.RenRootResources);
+		SafeRelease(Res.SimRootResources);
+	#if RHI_RAYTRACING
+		SafeRelease(Res.RaytracingResources);
+	#endif
+	}
 
 	// Delay destruction as resources reference by the interpolation 
 	// structs are used on the rendering thread, 
@@ -707,9 +778,6 @@ void UGroomComponent::ReleaseResources()
 	MeshProjectionTickDelay = 0;
 	MeshProjectionState = EMeshProjectionState::Invalid;
 
-#if RHI_RAYTRACING
-	SafeRelease(RaytracingResources);
-#endif
 
 	// Insure the ticking of the Groom component always happens after the skeletalMeshComponent.
 	if (RegisteredSkeletalMeshComponent)
@@ -766,14 +834,10 @@ void UGroomComponent::TickComponent(float DeltaTime, enum ELevelTick TickType, F
 	if (SkeletalMeshComponent)
 	{
 		const uint32 SkeletalLODCount = SkeletalMeshComponent->GetNumLODs();
-		const uint32 ResourceLODCount = RenRootResources ? RenRootResources->MeshProjectionLODs.Num() : 0;
+		const uint32 ResourceLODCount = HairGroupResources.Num() > 0 && HairGroupResources[0].RenRootResources ? HairGroupResources[0].RenRootResources->MeshProjectionLODs.Num() : 0;
 		if (SkeletalLODCount != ResourceLODCount)
 		{
-
-			FHairStrandsRootResource* LocalRenRootResources = SkeletalLODCount > 0 ? new FHairStrandsRootResource((GroomAsset && GroomAsset->GetNumHairGroups() > 0) ? &GroomAsset->HairGroupsData[0].HairRenderData : nullptr, SkeletalLODCount) : nullptr;
-			FHairStrandsRootResource* LocalSimRootResources = SkeletalLODCount > 0 ? new FHairStrandsRootResource((GroomAsset && GroomAsset->GetNumHairGroups() > 0) ? &GroomAsset->HairGroupsData[0].HairSimulationData : nullptr, SkeletalLODCount) : nullptr;
-
-			auto InitRootResource = [] (FHairStrandsRootResource*& PersistentResources, FHairStrandsRootResource* LocalResources)
+			auto InitRootResource = [](FHairStrandsRootResource*& PersistentResources, FHairStrandsRootResource* LocalResources)
 			{
 				SafeRelease(PersistentResources);
 				PersistentResources = LocalResources;
@@ -782,16 +846,42 @@ void UGroomComponent::TickComponent(float DeltaTime, enum ELevelTick TickType, F
 					BeginInitResource(PersistentResources);
 				}
 			};
-			InitRootResource(RenRootResources, LocalRenRootResources);
-			InitRootResource(SimRootResources, LocalSimRootResources);
+
+			TArray<FHairStrandsRootResource*> GroupRenRootResources;
+			TArray<FHairStrandsRootResource*> GroupSimRootResources;
+			const uint32 GroupCount = GroomAsset ? GroomAsset->GetNumHairGroups() : 0;
+			check(GroupCount == HairGroupResources.Num());
+			for (uint32 GroupIndex=0; GroupIndex < GroupCount; ++GroupIndex)
+			{
+				const FHairGroupData* HairGroupsData = &GroomAsset->HairGroupsData[GroupIndex];
+				FHairGroupResource& Res = HairGroupResources[GroupIndex];
+				FHairStrandsRootResource* LocalRenRootResources = SkeletalLODCount > 0 ? new FHairStrandsRootResource(&HairGroupsData->HairRenderData, SkeletalLODCount) : nullptr;
+				FHairStrandsRootResource* LocalSimRootResources = SkeletalLODCount > 0 ? new FHairStrandsRootResource(&HairGroupsData->HairSimulationData, SkeletalLODCount) : nullptr;
+
+				InitRootResource(Res.RenRootResources, LocalRenRootResources);
+				InitRootResource(Res.SimRootResources, LocalSimRootResources);
+
+				GroupRenRootResources.Add(LocalRenRootResources);
+				GroupSimRootResources.Add(LocalSimRootResources);
+			}
 
 			FTransform HairLocalToWorld = GetComponentTransform();
 			ENQUEUE_RENDER_COMMAND(FHairStrandsTick_LODUpdate)(
-				[Id, WorldType, LocalRenRootResources, LocalSimRootResources, HairLocalToWorld](FRHICommandListImmediate& RHICmdList)
+				[Id, WorldType, GroupRenRootResources, GroupSimRootResources, HairLocalToWorld](FRHICommandListImmediate& RHICmdList)
 			{
-				FHairStrandsProjectionHairData RenProjectionData = ToProjectionHairData(LocalRenRootResources);
-				FHairStrandsProjectionHairData SimProjectionData = ToProjectionHairData(LocalSimRootResources);
-				UpdateHairStrands(Id, WorldType, HairLocalToWorld, RenProjectionData, SimProjectionData);
+				FHairStrandsProjectionHairData RenProjectionDatas;
+				for (FHairStrandsRootResource* Res : GroupRenRootResources)
+				{
+					RenProjectionDatas.HairGroups.Add(ToProjectionHairData(Res));
+				}
+
+				FHairStrandsProjectionHairData SimProjectionDatas;
+				for (FHairStrandsRootResource* Res : GroupSimRootResources)
+				{
+					SimProjectionDatas.HairGroups.Add(ToProjectionHairData(Res));
+				}
+
+				UpdateHairStrands(Id, WorldType, HairLocalToWorld, RenProjectionDatas, SimProjectionDatas);
 			});
 		}
 
@@ -836,7 +926,7 @@ void UGroomComponent::TickComponent(float DeltaTime, enum ELevelTick TickType, F
 		}
 
 		// When a skeletal object with projection is enabled, activate the refresh of the bounding box to 
-		// insure the component/proxy bounding box alwaws lies onto the actual skinned mesh
+		// insure the component/proxy bounding box always lies onto the actual skinned mesh
 		MarkRenderTransformDirty();
 	}
 
diff --git a/Engine/Plugins/Experimental/HairStrands/Source/HairStrandsCore/Private/HairStrandsRendering.cpp b/Engine/Plugins/Experimental/HairStrands/Source/HairStrandsCore/Private/HairStrandsRendering.cpp
--- a/Engine/Plugins/Experimental/HairStrands/Source/HairStrandsCore/Private/HairStrandsRendering.cpp
+++ b/Engine/Plugins/Experimental/HairStrands/Source/HairStrandsCore/Private/HairStrandsRendering.cpp
@@ -151,7 +151,7 @@ IMPLEMENT_GLOBAL_SHADER(FHairInterpolationCS, "/Engine/Private/HairStrands/HairS
 
 static void AddHairStrandsInterpolationPass(
 	FRDGBuilder& GraphBuilder,
-	const FHairStrandsProjectionHairData& InHairData,
+	const FHairStrandsProjectionHairData::HairGroup& InHairData,
 	const int32 LODIndex,
 	const uint32 VertexCount,
 	const FShaderResourceViewRHIRef& RenderRestPosePositionBuffer,
@@ -370,130 +370,136 @@ void ComputeHairStrandsInterpolation(
 	FRHICommandListImmediate& RHICmdList,
 	FHairStrandsInterpolationInput* InInput,
 	FHairStrandsInterpolationOutput* InOutput,
-	FHairStrandsProjectionHairData& InHairData,
+	FHairStrandsProjectionHairData& InHairDatas,
 	int32 LODIndex)
 {
 	if (!InInput || !InOutput) return;
-	FHairStrandsInterpolationInput& Input = *InInput;
-	FHairStrandsInterpolationOutput& Output = *InOutput;
-	Output.VFInput.Reset();
 
-	if (!Input.IsValid() || !Output.IsValid()) return;
+	const uint32 GroupCount = InOutput->HairGroups.Num();
+	for (uint32 GroupIndex = 0; GroupIndex < GroupCount; ++GroupIndex)
+	{
+		FHairStrandsInterpolationInput::FHairGroup& Input   = InInput->HairGroups[GroupIndex];
+		FHairStrandsInterpolationOutput::HairGroup& Output = InOutput->HairGroups[GroupIndex];
+		Output.VFInput.Reset();
 
-	DECLARE_GPU_STAT(HairStrandsInterpolation);
-	SCOPED_DRAW_EVENT(RHICmdList, HairStrandsInterpolation);
-	SCOPED_GPU_STAT(RHICmdList, HairStrandsInterpolation);
+		if (!Input.IsValid() || !Output.IsValid()) return;
 
-	FRDGBuilder GraphBuilder(RHICmdList);
+		DECLARE_GPU_STAT(HairStrandsInterpolation);
+		SCOPED_DRAW_EVENT(RHICmdList, HairStrandsInterpolation);
+		SCOPED_GPU_STAT(RHICmdList, HairStrandsInterpolation);
 
-	const uint32 CurrIndex = Output.CurrentIndex;
-	const uint32 PrevIndex = (Output.CurrentIndex+1)%2;
+		FRDGBuilder GraphBuilder(RHICmdList);
 
-	// Procedural deformers in place of actual physics solver
-	// 0: Simulation
-	// 1: Bypass
-	// 2: Wave
-	// 3: Straighten hair in direction of the root's normal
-	const int32 DeformationType = FMath::Clamp(GHairDeformationType, 0, Input.SimRootPointIndexBuffer ? 3 : 2) - 1;
-	if (GHairDeformationType > 0)
-	{
-		AddDeformSimHairStrandsPass(
-			GraphBuilder,
-			DeformationType,
-			Input.SimVertexCount,
-			Input.SimRestPosePositionBuffer->SRV,
-			Input.SimRootPointIndexBuffer ? Input.SimRootPointIndexBuffer->SRV : nullptr,
-			Output.SimDeformedPositionBuffer[CurrIndex]->UAV);
-	}
+		const uint32 CurrIndex = Output.CurrentIndex;
+		const uint32 PrevIndex = (Output.CurrentIndex + 1) % 2;
 
-	// If the deformation is driven by the physics simulation, then the output is always the 0 index
-	const uint32 SimIndex = GHairDeformationType == 0 ? 0 : CurrIndex;
+		// Procedural deformers in place of actual physics solver
+		// 0: Simulation
+		// 1: Bypass
+		// 2: Wave
+		// 3: Straighten hair in direction of the root's normal
+		const int32 DeformationType = FMath::Clamp(GHairDeformationType, 0, Input.SimRootPointIndexBuffer ? 3 : 2) - 1;
+		if (GHairDeformationType > 0)
+		{
+			AddDeformSimHairStrandsPass(
+				GraphBuilder,
+				DeformationType,
+				Input.SimVertexCount,
+				Input.SimRestPosePositionBuffer->SRV,
+				Input.SimRootPointIndexBuffer ? Input.SimRootPointIndexBuffer->SRV : nullptr,
+				Output.SimDeformedPositionBuffer[CurrIndex]->UAV);
+		}
 
-	// Debug mode:
-	// * None	: Display hair normally
-	// * Sim	: Show sim strands
-	// * Render : Show rendering strands with sim color influence
-	const EHairStrandsDebugMode DebugMode = GetHairStrandsDebugStrandsMode();
-	if (DebugMode == EHairStrandsDebugMode::SimHairStrands)
-	{
-		AddHairTangentPass(
-			GraphBuilder,
-			Input.SimVertexCount,
-			Output.SimDeformedPositionBuffer[SimIndex]->SRV,
-			Output.SimTangentBuffer->UAV);
+		// If the deformation is driven by the physics simulation, then the output is always the 0 index
+		const uint32 SimIndex = GHairDeformationType == 0 ? 0 : CurrIndex;
 
-		GraphBuilder.Execute();
-	
-		Output.VFInput.HairPositionBuffer = Output.SimDeformedPositionBuffer[SimIndex]->SRV;
-		Output.VFInput.HairPreviousPositionBuffer = Output.SimDeformedPositionBuffer[SimIndex]->SRV;
-		Output.VFInput.HairTangentBuffer = Output.SimTangentBuffer->SRV;
-		Output.VFInput.HairAttributeBuffer = Input.SimAttributeBuffer->SRV;
-		Output.VFInput.VertexCount = Input.SimVertexCount;
-	}
-	else
-	{
-		const uint32 BufferSizeInBytes = Input.RenderVertexCount * FHairStrandsAttributeFormat::SizeInByte;
-		if (DebugMode == EHairStrandsDebugMode::RenderHairStrands && Output.RenderPatchedAttributeBuffer.NumBytes != BufferSizeInBytes)
+		// Debug mode:
+		// * None	: Display hair normally
+		// * Sim	: Show sim strands
+		// * Render : Show rendering strands with sim color influence
+		const EHairStrandsDebugMode DebugMode = GetHairStrandsDebugStrandsMode();
+		if (DebugMode == EHairStrandsDebugMode::SimHairStrands)
 		{
-			Output.RenderPatchedAttributeBuffer.Release();
-			Output.RenderPatchedAttributeBuffer.Initialize(FHairStrandsAttributeFormat::SizeInByte, Input.RenderVertexCount, FHairStrandsAttributeFormat::Format, BUF_Static);
-		}
+			AddHairTangentPass(
+				GraphBuilder,
+				Input.SimVertexCount,
+				Output.SimDeformedPositionBuffer[SimIndex]->SRV,
+				Output.SimTangentBuffer->UAV);
 
-		AddHairStrandsInterpolationPass(
-			GraphBuilder,
-			InHairData,
-			LODIndex,
-			Input.RenderVertexCount,
-			Input.RenderRestPosePositionBuffer->SRV,
-			Input.Interpolation0Buffer->SRV,
-			Input.Interpolation1Buffer->SRV,
-			Input.SimRestPosePositionBuffer->SRV,
-			Output.SimDeformedPositionBuffer[SimIndex]->SRV,
-			DebugMode == EHairStrandsDebugMode::RenderHairStrands ? Input.SimAttributeBuffer->SRV : nullptr,
-			Output.RenderDeformedPositionBuffer[CurrIndex]->UAV,
-			DebugMode == EHairStrandsDebugMode::RenderHairStrands ? Output.RenderPatchedAttributeBuffer.UAV : nullptr);
-
-		AddHairTangentPass(
-			GraphBuilder,
-			Input.RenderVertexCount,
-			Output.RenderDeformedPositionBuffer[CurrIndex]->SRV,
-			Output.RenderTangentBuffer->UAV);
+			GraphBuilder.Execute();
 
-#if RHI_RAYTRACING
-		if (IsRayTracingEnabled())
+			Output.VFInput.HairPositionBuffer = Output.SimDeformedPositionBuffer[SimIndex]->SRV;
+			Output.VFInput.HairPreviousPositionBuffer = Output.SimDeformedPositionBuffer[SimIndex]->SRV;
+			Output.VFInput.HairTangentBuffer = Output.SimTangentBuffer->SRV;
+			Output.VFInput.HairAttributeBuffer = Input.SimAttributeBuffer->SRV;
+			Output.VFInput.VertexCount = Input.SimVertexCount;
+		}
+		else
 		{
-			AddGenerateRaytracingGeometryPass(
+			const uint32 BufferSizeInBytes = Input.RenderVertexCount * FHairStrandsAttributeFormat::SizeInByte;
+			if (DebugMode == EHairStrandsDebugMode::RenderHairStrands && Output.RenderPatchedAttributeBuffer.NumBytes != BufferSizeInBytes)
+			{
+				Output.RenderPatchedAttributeBuffer.Release();
+				Output.RenderPatchedAttributeBuffer.Initialize(FHairStrandsAttributeFormat::SizeInByte, Input.RenderVertexCount, FHairStrandsAttributeFormat::Format, BUF_Static);
+			}
+
+			check(GroupIndex < uint32(InHairDatas.HairGroups.Num()));
+			AddHairStrandsInterpolationPass(
+				GraphBuilder,
+				InHairDatas.HairGroups[GroupIndex],
+				LODIndex,
+				Input.RenderVertexCount,
+				Input.RenderRestPosePositionBuffer->SRV,
+				Input.Interpolation0Buffer->SRV,
+				Input.Interpolation1Buffer->SRV,
+				Input.SimRestPosePositionBuffer->SRV,
+				Output.SimDeformedPositionBuffer[SimIndex]->SRV,
+				DebugMode == EHairStrandsDebugMode::RenderHairStrands ? Input.SimAttributeBuffer->SRV : nullptr,
+				Output.RenderDeformedPositionBuffer[CurrIndex]->UAV,
+				DebugMode == EHairStrandsDebugMode::RenderHairStrands ? Output.RenderPatchedAttributeBuffer.UAV : nullptr);
+
+			AddHairTangentPass(
 				GraphBuilder,
 				Input.RenderVertexCount,
-				Input.HairRadius * GHairRaytracingRadiusScale,
-				Input.HairWorldOffset,
 				Output.RenderDeformedPositionBuffer[CurrIndex]->SRV,
-				Input.RaytracingPositionBuffer->UAV);
-		}
+				Output.RenderTangentBuffer->UAV);
+
+#if RHI_RAYTRACING
+			if (IsRayTracingEnabled())
+			{
+				AddGenerateRaytracingGeometryPass(
+					GraphBuilder,
+					Input.RenderVertexCount,
+					Input.HairRadius * GHairRaytracingRadiusScale,
+					Input.HairWorldOffset,
+					Output.RenderDeformedPositionBuffer[CurrIndex]->SRV,
+					Input.RaytracingPositionBuffer->UAV);
+			}
 #endif
-		GraphBuilder.Execute();
+			GraphBuilder.Execute();
 
-		Output.VFInput.HairPositionBuffer = Output.RenderDeformedPositionBuffer[CurrIndex]->SRV;
-		Output.VFInput.HairPreviousPositionBuffer = Output.RenderDeformedPositionBuffer[PrevIndex]->SRV;
-		Output.VFInput.HairTangentBuffer = Output.RenderTangentBuffer->SRV;
-		Output.VFInput.HairAttributeBuffer = DebugMode == EHairStrandsDebugMode::RenderHairStrands ? Output.RenderPatchedAttributeBuffer.SRV : Input.RenderAttributeBuffer->SRV;
-		Output.VFInput.VertexCount = Input.RenderVertexCount;
+			Output.VFInput.HairPositionBuffer = Output.RenderDeformedPositionBuffer[CurrIndex]->SRV;
+			Output.VFInput.HairPreviousPositionBuffer = Output.RenderDeformedPositionBuffer[PrevIndex]->SRV;
+			Output.VFInput.HairTangentBuffer = Output.RenderTangentBuffer->SRV;
+			Output.VFInput.HairAttributeBuffer = DebugMode == EHairStrandsDebugMode::RenderHairStrands ? Output.RenderPatchedAttributeBuffer.SRV : Input.RenderAttributeBuffer->SRV;
+			Output.VFInput.VertexCount = Input.RenderVertexCount;
 
 #if RHI_RAYTRACING
-		if (IsRayTracingEnabled())
-		{
-			FRHIUnorderedAccessView* UAV = Input.RaytracingPositionBuffer->UAV;
-			RHICmdList.TransitionResources(EResourceTransitionAccess::EReadable, EResourceTransitionPipeline::EComputeToGfx, &UAV, 1);
-
-			const bool bNeedFullBuild = !Input.bIsRTGeometryInitialized;
-			if (bNeedFullBuild)
-				BuildHairAccelerationStructure(RHICmdList, Input.RaytracingVertexCount, Input.RaytracingPositionBuffer->Buffer, Input.RaytracingGeometry);
-			else
-				UpdateHairAccelerationStructure(RHICmdList, Input.RaytracingGeometry);
-			Input.bIsRTGeometryInitialized = true;
-		}
+			if (IsRayTracingEnabled())
+			{
+				FRHIUnorderedAccessView* UAV = Input.RaytracingPositionBuffer->UAV;
+				RHICmdList.TransitionResources(EResourceTransitionAccess::EReadable, EResourceTransitionPipeline::EComputeToGfx, &UAV, 1);
+
+				const bool bNeedFullBuild = !Input.bIsRTGeometryInitialized;
+				if (bNeedFullBuild)
+					BuildHairAccelerationStructure(RHICmdList, Input.RaytracingVertexCount, Input.RaytracingPositionBuffer->Buffer, Input.RaytracingGeometry);
+				else
+					UpdateHairAccelerationStructure(RHICmdList, Input.RaytracingGeometry);
+				Input.bIsRTGeometryInitialized = true;
+			}
 #endif
-	}
+		}
 
-	Output.CurrentIndex = PrevIndex;
+		Output.CurrentIndex = PrevIndex;
+	}
 }
diff --git a/Engine/Plugins/Experimental/HairStrands/Source/HairStrandsCore/Private/HairStrandsVertexFactory.cpp b/Engine/Plugins/Experimental/HairStrands/Source/HairStrandsCore/Private/HairStrandsVertexFactory.cpp
--- a/Engine/Plugins/Experimental/HairStrands/Source/HairStrandsCore/Private/HairStrandsVertexFactory.cpp
+++ b/Engine/Plugins/Experimental/HairStrands/Source/HairStrandsCore/Private/HairStrandsVertexFactory.cpp
@@ -15,9 +15,9 @@
 static float GStrandHairWidth = 0.0f;
 static FAutoConsoleVariableRef CVarStrandHairWidth(TEXT("r.HairStrands.StrandWidth"), GStrandHairWidth, TEXT("Width of hair strand"));
 
-float FHairStrandsVertexFactory::GetMaxStrandRadius() const
+float FHairStrandsVertexFactory::GetMaxStrandRadius(uint32 GroupIndex) const
 {
-	return GStrandHairWidth > 0 ? GStrandHairWidth * 0.5f : Data.MaxStrandRadius;
+	return GStrandHairWidth > 0 ? GStrandHairWidth * 0.5f : Data.HairGroups[GroupIndex].MaxStrandRadius;
 }
 
 #define OPTIMIZE_OFF 0
@@ -97,14 +97,16 @@ public:
 			GetHairVisibilitySampleCount(),
 			0.f);
 
-		BindParam(ShaderBindings, PositionBuffer, VF->GetPositionSRV());
-		BindParam(ShaderBindings, PreviousPositionBuffer, VF->GetPreviousPositionSRV());
-		BindParam(ShaderBindings, AttributeBuffer, VF->GetAttributeSRV());
-		BindParam(ShaderBindings, TangentBuffer, VF->GetTangentSRV());
-		BindParam(ShaderBindings, Radius, VF->GetMaxStrandRadius());
-		BindParam(ShaderBindings, Length, VF->GetMaxStrandLength());
-		BindParam(ShaderBindings, WorldOffset, VF->GetWorldOffset());
-		BindParam(ShaderBindings, Density, VF->GetHairDensity());
+		const uint64 GroupIndex = reinterpret_cast<uint64>(BatchElement.UserData);
+
+		BindParam(ShaderBindings, PositionBuffer, VF->GetPositionSRV(GroupIndex));
+		BindParam(ShaderBindings, PreviousPositionBuffer, VF->GetPreviousPositionSRV(GroupIndex));
+		BindParam(ShaderBindings, AttributeBuffer, VF->GetAttributeSRV(GroupIndex));
+		BindParam(ShaderBindings, TangentBuffer, VF->GetTangentSRV(GroupIndex));
+		BindParam(ShaderBindings, Radius, VF->GetMaxStrandRadius(GroupIndex));
+		BindParam(ShaderBindings, Length, VF->GetMaxStrandLength(GroupIndex));
+		BindParam(ShaderBindings, WorldOffset, VF->GetWorldOffset(GroupIndex));
+		BindParam(ShaderBindings, Density, VF->GetHairDensity(GroupIndex));
 		BindParam(ShaderBindings, RadiusAtDepth1_Primary, MinRadiusAtDepth1.Primary);
 		BindParam(ShaderBindings, RadiusAtDepth1_Velocity, MinRadiusAtDepth1.Velocity);		
 	}
diff --git a/Engine/Plugins/Experimental/HairStrands/Source/HairStrandsCore/Public/GroomComponent.h b/Engine/Plugins/Experimental/HairStrands/Source/HairStrandsCore/Public/GroomComponent.h
--- a/Engine/Plugins/Experimental/HairStrands/Source/HairStrandsCore/Public/GroomComponent.h
+++ b/Engine/Plugins/Experimental/HairStrands/Source/HairStrandsCore/Public/GroomComponent.h
@@ -58,25 +58,33 @@ public:
 	//~ End UMeshComponent Interface.
 
 	/** Return the guide hairs datas */
-	FHairStrandsDatas* GetGuideStrandsDatas();
+	FHairStrandsDatas* GetGuideStrandsDatas(uint32 GroupIndex);
 
 	/** Return the guide hairs resources*/
-	FHairStrandsResource* GetGuideStrandsResource();
+	FHairStrandsResource* GetGuideStrandsResource(uint32 GroupIndex);
 
 #if WITH_EDITOR
 	virtual void PostEditChangeProperty(FPropertyChangedEvent& PropertyChangedEvent) override;
 #endif
-	// #hair_todo: cached/serialize interpolation data
-	FHairStrandsInterpolationResource* InterpolationResource = nullptr;
-	struct FHairStrandsInterpolationOutput* InterpolationOutput = nullptr;
-	struct FHairStrandsInterpolationInput* InterpolationInput = nullptr;
-	struct FHairStrandsRootResource* RenRootResources = nullptr;
-	struct FHairStrandsRootResource* SimRootResources = nullptr;
 
-#if RHI_RAYTRACING
-	FHairStrandsRaytracingResource* RaytracingResources = nullptr;
-#endif
+	struct FHairGroupResource
+	{
+		FHairStrandsInterpolationResource* InterpolationResource = nullptr;
+		struct FHairStrandsRootResource* RenRootResources = nullptr;
+		struct FHairStrandsRootResource* SimRootResources = nullptr;
+	#if RHI_RAYTRACING
+		FHairStrandsRaytracingResource* RaytracingResources = nullptr;
+	#endif
+
+		// Resources own by the asset
+		FHairStrandsResource* RenderResources = nullptr;
+		FHairStrandsResource* SimResources = nullptr;
+	};
+	typedef TArray<FHairGroupResource> FHairGroupResources;
 
+	FHairGroupResources HairGroupResources;
+	struct FHairStrandsInterpolationOutput* InterpolationOutput = nullptr;
+	struct FHairStrandsInterpolationInput* InterpolationInput = nullptr;
 private:
 	void* InitializedResources;
 
diff --git a/Engine/Plugins/Experimental/HairStrands/Source/HairStrandsCore/Public/HairStrandsRendering.h b/Engine/Plugins/Experimental/HairStrands/Source/HairStrandsCore/Public/HairStrandsRendering.h
--- a/Engine/Plugins/Experimental/HairStrands/Source/HairStrandsCore/Public/HairStrandsRendering.h
+++ b/Engine/Plugins/Experimental/HairStrands/Source/HairStrandsCore/Public/HairStrandsRendering.h
@@ -9,58 +9,50 @@
 
 struct FHairStrandsInterpolationInput
 {
-	FRWBuffer*	RenderRestPosePositionBuffer = nullptr;
-	FRWBuffer*	RenderAttributeBuffer = nullptr;
-	uint32		RenderVertexCount = 0;
+	struct FHairGroup
+	{
+		FRWBuffer*	RenderRestPosePositionBuffer = nullptr;
+		FRWBuffer*	RenderAttributeBuffer = nullptr;
+		uint32		RenderVertexCount = 0;
 
-	FRWBuffer*	SimRestPosePositionBuffer = nullptr;
-	FRWBuffer*	SimAttributeBuffer = nullptr;
-	uint32		SimVertexCount = 0;
+		FRWBuffer*	SimRestPosePositionBuffer = nullptr;
+		FRWBuffer*	SimAttributeBuffer = nullptr;
+		uint32		SimVertexCount = 0;
 
-	FRWBuffer*	Interpolation0Buffer = nullptr;
-	FRWBuffer*	Interpolation1Buffer = nullptr;
+		FRWBuffer*	Interpolation0Buffer = nullptr;
+		FRWBuffer*	Interpolation1Buffer = nullptr;
 
-	// For debug purpose only
-	FRWBuffer*	SimRootPointIndexBuffer = nullptr;
+		// For debug purpose only
+		FRWBuffer*	SimRootPointIndexBuffer = nullptr;
 
-	#if RHI_RAYTRACING
-	FRayTracingGeometry*	RaytracingGeometry = nullptr;
-	FRWBuffer*				RaytracingPositionBuffer = nullptr;
-	uint32					RaytracingVertexCount = 0;
-	bool					bIsRTGeometryInitialized = false;
-	#endif
+		#if RHI_RAYTRACING
+		FRayTracingGeometry*	RaytracingGeometry = nullptr;
+		FRWBuffer*				RaytracingPositionBuffer = nullptr;
+		uint32					RaytracingVertexCount = 0;
+		bool					bIsRTGeometryInitialized = false;
+		#endif
 
-	float HairRadius = 0;
-	FVector HairWorldOffset = FVector::ZeroVector;
+		float HairRadius = 0;
+		FVector HairWorldOffset = FVector::ZeroVector;
 
-	inline bool IsValid() const
-	{
-		return
-			(RenderRestPosePositionBuffer && RenderRestPosePositionBuffer->SRV) &&
-			(RenderAttributeBuffer && RenderAttributeBuffer->SRV) &&
-			(SimRestPosePositionBuffer && SimRestPosePositionBuffer->SRV) &&
-			(SimAttributeBuffer && SimAttributeBuffer->SRV) &&
-			(Interpolation0Buffer && Interpolation0Buffer->SRV) &&
-			(Interpolation1Buffer && Interpolation1Buffer->SRV) &&
-			RenderVertexCount != 0 &&
-			SimVertexCount != 0;
-	}
+		inline bool IsValid() const
+		{
+			return
+				(RenderRestPosePositionBuffer && RenderRestPosePositionBuffer->SRV) &&
+				(RenderAttributeBuffer && RenderAttributeBuffer->SRV) &&
+				(SimRestPosePositionBuffer && SimRestPosePositionBuffer->SRV) &&
+				(SimAttributeBuffer && SimAttributeBuffer->SRV) &&
+				(Interpolation0Buffer && Interpolation0Buffer->SRV) &&
+				(Interpolation1Buffer && Interpolation1Buffer->SRV) &&
+				RenderVertexCount != 0 &&
+				SimVertexCount != 0;
+		}
+	};
+	TArray<FHairGroup> HairGroups;
 };
 
 struct FHairStrandsInterpolationOutput
 {
-	FRWBuffer* SimDeformedPositionBuffer[2] = { nullptr, nullptr };
-	FRWBuffer* RenderDeformedPositionBuffer[2] = { nullptr, nullptr };
-	uint32 CurrentIndex = 0;
-
-	FRWBuffer* RenderTangentBuffer = nullptr;
-	FRWBuffer* RenderAttributeBuffer = nullptr;
-
-	FRWBuffer* SimTangentBuffer = nullptr;
-
-	// Debug buffers (allocated on-the-fly if used)
-	FRWBuffer RenderPatchedAttributeBuffer;
-
 	// Input to the strand vertex factory. This allows to 
 	// abstract resources generation when debug mode are used
 	struct VertexFactoryInput
@@ -80,24 +72,41 @@ struct FHairStrandsInterpolationOutput
 			VertexCount = 0;
 		}
 	};
-	VertexFactoryInput VFInput;
 
-	inline bool IsValid() const
+	struct HairGroup
 	{
-		return
-			(SimDeformedPositionBuffer[0] && SimDeformedPositionBuffer[0]->SRV) &&
-			(SimDeformedPositionBuffer[1] && SimDeformedPositionBuffer[1]->SRV) &&
-			(RenderDeformedPositionBuffer[0] && RenderDeformedPositionBuffer[0]->SRV) &&
-			(RenderDeformedPositionBuffer[1] && RenderDeformedPositionBuffer[1]->SRV) &&
-			(RenderTangentBuffer && RenderTangentBuffer->SRV) &&
-			(RenderAttributeBuffer && RenderAttributeBuffer->SRV);
-	}
+		FRWBuffer* SimDeformedPositionBuffer[2] = { nullptr, nullptr };
+		FRWBuffer* RenderDeformedPositionBuffer[2] = { nullptr, nullptr };
+		uint32 CurrentIndex = 0;
+
+		FRWBuffer* RenderTangentBuffer = nullptr;
+		FRWBuffer* RenderAttributeBuffer = nullptr;
+
+		FRWBuffer* SimTangentBuffer = nullptr;
+
+		// Debug buffers (allocated on-the-fly if used)
+		FRWBuffer RenderPatchedAttributeBuffer;
+
+		VertexFactoryInput VFInput;
+
+		inline bool IsValid() const
+		{
+			return
+				(SimDeformedPositionBuffer[0] && SimDeformedPositionBuffer[0]->SRV) &&
+				(SimDeformedPositionBuffer[1] && SimDeformedPositionBuffer[1]->SRV) &&
+				(RenderDeformedPositionBuffer[0] && RenderDeformedPositionBuffer[0]->SRV) &&
+				(RenderDeformedPositionBuffer[1] && RenderDeformedPositionBuffer[1]->SRV) &&
+				(RenderTangentBuffer && RenderTangentBuffer->SRV) &&
+				(RenderAttributeBuffer && RenderAttributeBuffer->SRV);
+		}
+	};
+	TArray<HairGroup> HairGroups;
 };
 
 void ComputeHairStrandsInterpolation(
 	FRHICommandListImmediate& RHICmdList,
 	FHairStrandsInterpolationInput* Input,
 	FHairStrandsInterpolationOutput* Output, 
-	struct FHairStrandsProjectionHairData& HairData,
+	struct FHairStrandsProjectionHairData& HairDatas,
 	int32 LODIndex);
 
diff --git a/Engine/Plugins/Experimental/HairStrands/Source/HairStrandsCore/Public/HairStrandsVertexFactory.h b/Engine/Plugins/Experimental/HairStrands/Source/HairStrandsCore/Public/HairStrandsVertexFactory.h
--- a/Engine/Plugins/Experimental/HairStrands/Source/HairStrandsCore/Public/HairStrandsVertexFactory.h
+++ b/Engine/Plugins/Experimental/HairStrands/Source/HairStrandsCore/Public/HairStrandsVertexFactory.h
@@ -32,12 +32,16 @@ public:
 	struct FDataType
 	{
 		FHairStrandsInterpolationOutput* InterpolationOutput = nullptr;
-		
-		float MinStrandRadius = 0;
-		float MaxStrandRadius = 0;
-		float MaxStrandLength = 0;
-		float HairDensity = 1;
-		FVector HairWorldOffset = FVector::ZeroVector;
+		struct HairGroup
+		{
+			float MinStrandRadius = 0;
+			float MaxStrandRadius = 0;
+			float MaxStrandLength = 0;
+			float HairDensity = 1;
+			FVector HairWorldOffset = FVector::ZeroVector;
+		};
+
+		TArray<HairGroup> HairGroups;
 	};
 
 	/**
@@ -66,16 +70,16 @@ public:
 
 	static FVertexFactoryShaderParameters* ConstructShaderParameters(EShaderFrequency ShaderFrequency);
 
-	inline FRHIShaderResourceView* GetPositionSRV() const			{ check(Data.InterpolationOutput); return Data.InterpolationOutput->VFInput.HairPositionBuffer; };
-	inline FRHIShaderResourceView* GetPreviousPositionSRV() const	{ check(Data.InterpolationOutput); return Data.InterpolationOutput->VFInput.HairPreviousPositionBuffer; }
-	inline FRHIShaderResourceView* GetAttributeSRV() const			{ check(Data.InterpolationOutput); return Data.InterpolationOutput->VFInput.HairAttributeBuffer; }
-	inline FRHIShaderResourceView* GetTangentSRV() const			{ check(Data.InterpolationOutput); return Data.InterpolationOutput->VFInput.HairTangentBuffer; }
+	inline FRHIShaderResourceView* GetPositionSRV(uint32 GroupIndex) const			{ check(Data.InterpolationOutput); return Data.InterpolationOutput->HairGroups[GroupIndex].VFInput.HairPositionBuffer; };
+	inline FRHIShaderResourceView* GetPreviousPositionSRV(uint32 GroupIndex) const	{ check(Data.InterpolationOutput); return Data.InterpolationOutput->HairGroups[GroupIndex].VFInput.HairPreviousPositionBuffer; }
+	inline FRHIShaderResourceView* GetAttributeSRV(uint32 GroupIndex) const			{ check(Data.InterpolationOutput); return Data.InterpolationOutput->HairGroups[GroupIndex].VFInput.HairAttributeBuffer; }
+	inline FRHIShaderResourceView* GetTangentSRV(uint32 GroupIndex) const			{ check(Data.InterpolationOutput); return Data.InterpolationOutput->HairGroups[GroupIndex].VFInput.HairTangentBuffer; }
 
-	float GetMaxStrandRadius() const;
-	inline float GetMinStrandRadius() const							{ return Data.MinStrandRadius; }
-	inline float GetMaxStrandLength() const							{ return Data.MaxStrandLength; }
-	inline float GetHairDensity() const								{ return Data.HairDensity;  }
-	inline const FVector& GetWorldOffset() const					{ return Data.HairWorldOffset; }
+	float GetMaxStrandRadius(uint32 GroupIndex) const;
+	inline float GetMinStrandRadius(uint32 GroupIndex) const						{ return Data.HairGroups[GroupIndex].MinStrandRadius; }
+	inline float GetMaxStrandLength(uint32 GroupIndex) const						{ return Data.HairGroups[GroupIndex].MaxStrandLength; }
+	inline float GetHairDensity(uint32 GroupIndex) const							{ return Data.HairGroups[GroupIndex].HairDensity;  }
+	inline const FVector& GetWorldOffset(uint32 GroupIndex) const					{ return Data.HairGroups[GroupIndex].HairWorldOffset; }
 
 	const FDataType& GetData() const { return Data; }
 protected:
diff --git a/Engine/Plugins/Experimental/HairStrands/Source/HairStrandsNiagara/Private/NiagaraDataInterfaceHairStrands.cpp b/Engine/Plugins/Experimental/HairStrands/Source/HairStrandsNiagara/Private/NiagaraDataInterfaceHairStrands.cpp
--- a/Engine/Plugins/Experimental/HairStrands/Source/HairStrandsNiagara/Private/NiagaraDataInterfaceHairStrands.cpp
+++ b/Engine/Plugins/Experimental/HairStrands/Source/HairStrandsNiagara/Private/NiagaraDataInterfaceHairStrands.cpp
@@ -529,8 +529,9 @@ bool UNiagaraDataInterfaceHairStrands::InitPerInstanceData(void* PerInstanceData
 		}
 	}
 	GroomAsset = IsComponentValid() ? SourceComponent->GroomAsset : DefaultSource;
-	FHairStrandsDatas* StrandsDatas = IsComponentValid() ? SourceComponent->GetGuideStrandsDatas() : (DefaultSource != nullptr && DefaultSource->GetNumHairGroups() > 0) ? &DefaultSource->HairGroupsData[0].HairRenderData : nullptr;
-	FHairStrandsResource* StrandsResource = IsComponentValid() ? SourceComponent->GetGuideStrandsResource() : (DefaultSource != nullptr && DefaultSource->GetNumHairGroups() > 0) ? DefaultSource->HairGroupsData[0].HairStrandsResource : nullptr;
+	const uint32 GroupIndex = 0;
+	FHairStrandsDatas* StrandsDatas = IsComponentValid() ? SourceComponent->GetGuideStrandsDatas(GroupIndex) : (DefaultSource != nullptr && DefaultSource->GetNumHairGroups() > 0) ? &DefaultSource->HairGroupsData[0].HairRenderData : nullptr;
+	FHairStrandsResource* StrandsResource = IsComponentValid() ? SourceComponent->GetGuideStrandsResource(GroupIndex) : (DefaultSource != nullptr && DefaultSource->GetNumHairGroups() > 0) ? DefaultSource->HairGroupsData[0].HairStrandsResource : nullptr;
 
 	if (StrandsDatas == nullptr || StrandsResource == nullptr)
 	{
@@ -642,8 +643,9 @@ bool UNiagaraDataInterfaceHairStrands::PerInstanceTick(void* PerInstanceData, FN
 {
 	FNDIHairStrandsData* InstanceData = static_cast<FNDIHairStrandsData*>(PerInstanceData);
 
+	const uint32 GroupIndex = 0;
 	FHairStrandsDatas* StrandsDatas = IsComponentValid() ? 
-		SourceComponent->GetGuideStrandsDatas() : (DefaultSource != nullptr && DefaultSource->GetNumHairGroups() > 0) ? &DefaultSource->HairGroupsData[0].HairRenderData : nullptr;
+		SourceComponent->GetGuideStrandsDatas(GroupIndex) : (DefaultSource != nullptr && DefaultSource->GetNumHairGroups() > 0) ? &DefaultSource->HairGroupsData[0].HairRenderData : nullptr;
 	bool RequireReset = false;
 	if (StrandsDatas)
 	{
diff --git a/Engine/Source/Runtime/Renderer/Private/HairStrands/HairStrandsDebug.cpp b/Engine/Source/Runtime/Renderer/Private/HairStrands/HairStrandsDebug.cpp
--- a/Engine/Source/Runtime/Renderer/Private/HairStrands/HairStrandsDebug.cpp
+++ b/Engine/Source/Runtime/Renderer/Private/HairStrands/HairStrandsDebug.cpp
@@ -49,7 +49,7 @@ static FAutoConsoleVariableRef CVarHairDebugMeshProjection_HairDeformedFrames(TE
 // Helper functions for accessing interpolation data for debug purpose.
 // Definitions is in HairStrandsInterface.cpp
 void GetGroomInterpolationData(const EWorldType::Type WorldType, FHairStrandsProjectionMeshData& OutGeometries);
-void GetGroomInterpolationData(const EWorldType::Type WorldType, FHairStrandsProjectionHairDatas& OutHairData, TArray<int32>& OutLODIndices);
+void GetGroomInterpolationData(const EWorldType::Type WorldType, FHairStrandsProjectionHairData& OutHairData, TArray<int32>& OutLODIndices);
 
 enum class EHairDebugMode : uint8
 {
@@ -554,7 +554,7 @@ static void AddDebugProjectionMeshPass(
 	FRDGBuilder& GraphBuilder,
 	const FViewInfo* View,
 	const bool bClearDepth,
-	FHairStrandsProjectionMeshSection& MeshSectionData,
+	FHairStrandsProjectionMeshData::Section& MeshSectionData,
 	FRDGTextureRef& ColorTexture, 
 	FRDGTextureRef& DepthTexture)
 {	
@@ -707,7 +707,7 @@ static void AddDebugProjectionHairPass(
 	const EDebugProjectionHairType GeometryType,
 	const HairStrandsTriangleType PoseType,
 	const int32 LODIndex,
-	const FHairStrandsProjectionHairData& HairData,
+	const FHairStrandsProjectionHairData::HairGroup& HairData,
 	FRDGTextureRef& ColorTarget,
 	FRDGTextureRef& DepthTexture)
 {
@@ -911,18 +911,18 @@ void RenderHairStrandsDebugInfo(FRHICommandListImmediate& RHICmdList, TArray<FVi
 			check(ViewFamily.Scene && ViewFamily.Scene->GetWorld());
 			const bool bIsActive = DebugInfo.WorldType == ViewFamily.Scene->GetWorld()->WorldType;
 
-			Line = FString::Printf(TEXT(" * Id:%d | WorldType:%s | Group count : %d"), DebugInfo.Id, ToString(DebugInfo.WorldType), DebugInfo.GroupInfos.Num());
+			Line = FString::Printf(TEXT(" * Id:%d | WorldType:%s | Group count : %d"), DebugInfo.Id, ToString(DebugInfo.WorldType), DebugInfo.HairGroups.Num());
 			Canvas.DrawShadowedString(X, Y += YStep, *Line, GetStatsFont(), bIsActive ? DebugColor : InactiveColor);
 
-			for (const FHairStrandsDebugInfo::FGroupInfo& DebugGroupInfo : DebugInfo.GroupInfos)
+			for (const FHairStrandsDebugInfo::HairGroup& DebugHairGroup : DebugInfo.HairGroups)
 			{
 				Line = FString::Printf(TEXT("        |> CurveCount : %d | VertexCount : %d | MaxRadius : %f | MaxLength : %f | Skinned: %s | LOD count : %d"),
-					DebugGroupInfo.CurveCount,
-					DebugGroupInfo.VertexCount,
-					DebugGroupInfo.MaxRadius,
-					DebugGroupInfo.MaxLength,
-					DebugGroupInfo.bHasSkinInterpolation ? TEXT("True") : TEXT("False"),
-					DebugGroupInfo.LODCount);
+					DebugHairGroup.CurveCount,
+					DebugHairGroup.VertexCount,
+					DebugHairGroup.MaxRadius,
+					DebugHairGroup.MaxLength,
+					DebugHairGroup.bHasSkinInterpolation ? TEXT("True") : TEXT("False"),
+					DebugHairGroup.LODCount);
 				Canvas.DrawShadowedString(X, Y += YStep, *Line, GetStatsFont(), bIsActive ? DebugColor : InactiveColor);
 			}
 		}
@@ -1080,20 +1080,20 @@ void RenderHairStrandsDebugInfo(FRHICommandListImmediate& RHICmdList, TArray<FVi
 
 			if (GHairDebugMeshProjection_SkinCacheMesh > 0)
 			{
-				for (FHairStrandsProjectionMeshSection& Section : MeshProjectionData.Sections)
+				for (FHairStrandsProjectionMeshData::Section& Section : MeshProjectionData.Sections)
 				{
 					AddDebugProjectionMeshPass(GraphBuilder, &View, bClearDepth, Section, SceneColorTexture, DepthTexture);
 					bClearDepth = false;
 				}
 			}
 
-			FHairStrandsProjectionHairDatas HairProjectionDatas;
+			FHairStrandsProjectionHairData HairProjectionDatas;
 			TArray<int32> HairLODIndices;
 			GetGroomInterpolationData(WorldType, HairProjectionDatas, HairLODIndices);
-			check(HairProjectionDatas.HairData.Num() == HairLODIndices.Num());
-			for (int32 HairIndex=0; HairIndex < HairProjectionDatas.HairData.Num(); ++HairIndex)
+			check(HairProjectionDatas.HairGroups.Num() == HairLODIndices.Num());
+			for (int32 HairIndex=0; HairIndex < HairProjectionDatas.HairGroups.Num(); ++HairIndex)
 			{
-				const FHairStrandsProjectionHairData& Data = HairProjectionDatas.HairData[HairIndex];
+				const FHairStrandsProjectionHairData::HairGroup& Data = HairProjectionDatas.HairGroups[HairIndex];
 				const int32 LODIndex = HairLODIndices[HairIndex];
 
 				if (GHairDebugMeshProjection_HairRestTriangles > 0)
diff --git a/Engine/Source/Runtime/Renderer/Private/HairStrands/HairStrandsInterface.cpp b/Engine/Source/Runtime/Renderer/Private/HairStrands/HairStrandsInterface.cpp
--- a/Engine/Source/Runtime/Renderer/Private/HairStrands/HairStrandsInterface.cpp
+++ b/Engine/Source/Runtime/Renderer/Private/HairStrands/HairStrandsInterface.cpp
@@ -18,9 +18,9 @@ static FAutoConsoleVariableRef CVarHairStrandsRenderingEnable(TEXT("r.HairStrand
 
 ////////////////////////////////////////////////////////////////////////////////////////////////////////////////////////
 
-inline FHairStrandsProjectionMeshSection ConvertMeshSection(const FCachedGeometrySection& In, const FTransform& InTransform)
+inline FHairStrandsProjectionMeshData::Section ConvertMeshSection(const FCachedGeometrySection& In, const FTransform& InTransform)
 {
-	FHairStrandsProjectionMeshSection Out;
+	FHairStrandsProjectionMeshData::Section Out;
 	Out.IndexBuffer = In.IndexBuffer;
 	Out.PositionBuffer = In.PositionBuffer;
 	Out.TotalVertexCount = In.TotalVertexCount;
@@ -52,10 +52,10 @@ struct FHairStrandsManager
 	{
 		uint64 Id = 0;
 		EWorldType::Type WorldType = EWorldType::None;
-		FHairStrandsDebugInfo::FGroupInfos DebugGroupInfos;
+		FHairStrandsDebugInfo DebugInfo;
 		FHairStrandsInterpolationData InterpolationData;
-		FHairStrandsProjectionHairData RenProjectionHairData;
-		FHairStrandsProjectionHairData SimProjectionHairData;
+		FHairStrandsProjectionHairData RenProjectionHairDatas;
+		FHairStrandsProjectionHairData SimProjectionHairDatas;
 		FCachedGeometry CachedGeometry;
 		FTransform SkeletalLocalToWorld;
 		FVector SkeletalWorldBoundsCenter;
@@ -91,9 +91,9 @@ void RegisterHairStrands(
 	uint64 Id, 
 	EWorldType::Type WorldType,
 	const FHairStrandsInterpolationData& InterpolationData, 
-	const FHairStrandsProjectionHairData& RenProjectionData,
- 	const FHairStrandsProjectionHairData& SimProjectionData,
-	const FHairStrandsDebugInfo::FGroupInfos& GroupInfos)
+	const FHairStrandsProjectionHairData& RenProjectionDatas,
+ 	const FHairStrandsProjectionHairData& SimProjectionDatas,
+	const FHairStrandsDebugInfo& DebugInfo)
 {
 	for (int32 Index = 0; Index < GHairManager.Elements.Num(); ++Index)
 	{
@@ -109,13 +109,19 @@ void RegisterHairStrands(
 	E.Id = Id;
 	E.WorldType = WorldType;
 	E.InterpolationData = InterpolationData;
-	E.RenProjectionHairData = RenProjectionData;
-	E.RenProjectionHairData.LocalToWorld = FTransform::Identity;
-	E.SimProjectionHairData = SimProjectionData;
-	E.SimProjectionHairData.LocalToWorld = FTransform::Identity;
+	E.RenProjectionHairDatas = RenProjectionDatas;
+	for (FHairStrandsProjectionHairData::HairGroup& ProjectionData : E.RenProjectionHairDatas.HairGroups)
+	{
+		ProjectionData.LocalToWorld = FTransform::Identity;
+	}
+	E.SimProjectionHairDatas = SimProjectionDatas;
+	for (FHairStrandsProjectionHairData::HairGroup& ProjectionData : E.SimProjectionHairDatas.HairGroups)
+	{
+		ProjectionData.LocalToWorld = FTransform::Identity;
+	}
 	E.SkeletalLocalToWorld = FTransform::Identity;
 	E.SkeletalWorldBoundsCenter = FVector::ZeroVector;
-	E.DebugGroupInfos = GroupInfos;
+	E.DebugInfo = DebugInfo;
 }
 
 bool UpdateHairStrands(
@@ -130,8 +136,14 @@ bool UpdateHairStrands(
 		if (E.Id != Id || E.WorldType != WorldType)
 			continue;
 
-		E.RenProjectionHairData.LocalToWorld = HairLocalToWorld;
-		E.SimProjectionHairData.LocalToWorld = HairLocalToWorld;
+		for (FHairStrandsProjectionHairData::HairGroup& ProjectionData : E.RenProjectionHairDatas.HairGroups)
+		{
+			ProjectionData.LocalToWorld = HairLocalToWorld;
+		}
+		for (FHairStrandsProjectionHairData::HairGroup& ProjectionData : E.SimProjectionHairDatas.HairGroups)
+		{
+			ProjectionData.LocalToWorld = HairLocalToWorld;
+		}
 		E.SkeletalLocalToWorld = SkeletalLocalToWorld;
 		E.SkeletalWorldBoundsCenter = SkeletalWorldBoundsCenter;
 		return true;
@@ -160,18 +172,24 @@ bool UpdateHairStrands(
 	uint64 Id,
 	EWorldType::Type WorldType,
 	const FTransform& HairLocalToWorld,
-	const FHairStrandsProjectionHairData& RenProjectionData,
-	const FHairStrandsProjectionHairData& SimProjectionData)
+	const FHairStrandsProjectionHairData& RenProjectionDatas,
+	const FHairStrandsProjectionHairData& SimProjectionDatas)
 {
 	for (FHairStrandsManager::Element& E : GHairManager.Elements)
 	{
 		if (E.Id != Id || E.WorldType != WorldType)
 			continue;
 
-		E.RenProjectionHairData = RenProjectionData;
-		E.RenProjectionHairData.LocalToWorld = HairLocalToWorld;
-		E.SimProjectionHairData = SimProjectionData;
-		E.SimProjectionHairData.LocalToWorld = HairLocalToWorld;
+		E.RenProjectionHairDatas = RenProjectionDatas;
+		for (FHairStrandsProjectionHairData::HairGroup& ProjectionData : E.RenProjectionHairDatas.HairGroups)
+		{
+			ProjectionData.LocalToWorld = HairLocalToWorld;
+		}
+		E.SimProjectionHairDatas = SimProjectionDatas;
+		for (FHairStrandsProjectionHairData::HairGroup& ProjectionData : E.SimProjectionHairDatas.HairGroups)
+		{
+			ProjectionData.LocalToWorld = HairLocalToWorld;
+		}
 		return true;
 	}
 
@@ -201,14 +219,17 @@ void UnregisterHairStrands(uint64 Id, EWorldType::Type WorldType)
 	{
 		if (GHairManager.Elements[Index].Id == Id && GHairManager.Elements[Index].WorldType == WorldType)
 		{
-			FHairStrandsInterpolationData Out = GHairManager.Elements[Index].InterpolationData;
 			GHairManager.Elements[Index] = GHairManager.Elements[GHairManager.Elements.Num()-1];
 			GHairManager.Elements.SetNum(GHairManager.Elements.Num() - 1);
 		}
 	}
 }
 
-void RunHairStrandsInterpolation(FRHICommandListImmediate& RHICmdList, EWorldType::Type WorldType, TShaderMap<FGlobalShaderType>* ShaderMap, EHairStrandsInterpolationType Type)
+void RunHairStrandsInterpolation(
+	FRHICommandListImmediate& RHICmdList, 
+	EWorldType::Type WorldType, 
+	TShaderMap<FGlobalShaderType>* ShaderMap, 
+	EHairStrandsInterpolationType Type)
 {
 	check(IsInRenderingThread());
 
@@ -236,18 +257,18 @@ void RunHairStrandsInterpolation(FRHICommandListImmediate& RHICmdList, EWorldTyp
 				MeshData.Sections.Add(ConvertMeshSection(Section, E.SkeletalLocalToWorld));
 			}
 
-			//if (EHairStrandsInterpolationType::RenderStrands == Type)
+			for (FHairStrandsProjectionHairData::HairGroup& ProjectionHairData : E.SimProjectionHairDatas.HairGroups)
 			{
-				E.SimProjectionHairData.LODDatas[Q.LODIndex].RestRootCenter = Q.RestRootCenter;
-				ProjectHairStrandsOntoMesh(RHICmdList, ShaderMap, Q.LODIndex, MeshData, E.SimProjectionHairData);
-				UpdateHairStrandsMeshTriangles(RHICmdList, ShaderMap, Q.LODIndex, HairStrandsTriangleType::RestPose, MeshData, E.SimProjectionHairData);
+				ProjectionHairData.LODDatas[Q.LODIndex].RestRootCenter = Q.RestRootCenter;
+				ProjectHairStrandsOntoMesh(RHICmdList, ShaderMap, Q.LODIndex, MeshData, ProjectionHairData);
+				UpdateHairStrandsMeshTriangles(RHICmdList, ShaderMap, Q.LODIndex, HairStrandsTriangleType::RestPose, MeshData, ProjectionHairData);
 			}
 
-			//if (EHairStrandsInterpolationType::RenderStrands == Type)
+			for (FHairStrandsProjectionHairData::HairGroup& ProjectionHairData : E.RenProjectionHairDatas.HairGroups)
 			{
-				E.RenProjectionHairData.LODDatas[Q.LODIndex].RestRootCenter = Q.RestRootCenter;
-				ProjectHairStrandsOntoMesh(RHICmdList, ShaderMap, Q.LODIndex, MeshData, E.RenProjectionHairData);
-				UpdateHairStrandsMeshTriangles(RHICmdList, ShaderMap, Q.LODIndex, HairStrandsTriangleType::RestPose, MeshData, E.RenProjectionHairData);
+				ProjectionHairData.LODDatas[Q.LODIndex].RestRootCenter = Q.RestRootCenter;
+				ProjectHairStrandsOntoMesh(RHICmdList, ShaderMap, Q.LODIndex, MeshData, ProjectionHairData);
+				UpdateHairStrandsMeshTriangles(RHICmdList, ShaderMap, Q.LODIndex, HairStrandsTriangleType::RestPose, MeshData, ProjectionHairData);
 			}
 			Q.bProcessed = true;
 		}
@@ -270,16 +291,22 @@ void RunHairStrandsInterpolation(FRHICommandListImmediate& RHICmdList, EWorldTyp
 			MeshData.Sections.Add(ConvertMeshSection(Section, E.SkeletalLocalToWorld));
 		}
 
-		if (EHairStrandsInterpolationType::RenderStrands == Type && 0 <= E.FrameLODIndex && E.FrameLODIndex < E.RenProjectionHairData.LODDatas.Num() && E.RenProjectionHairData.LODDatas[E.FrameLODIndex].bIsValid)
+		for (FHairStrandsProjectionHairData::HairGroup& ProjectionHairData : E.RenProjectionHairDatas.HairGroups)
 		{
-			E.RenProjectionHairData.LODDatas[E.FrameLODIndex].DeformedRootCenter = E.SkeletalWorldBoundsCenter;
-			UpdateHairStrandsMeshTriangles(RHICmdList, ShaderMap, E.FrameLODIndex, HairStrandsTriangleType::DeformedPose, MeshData, E.RenProjectionHairData);
+			if (EHairStrandsInterpolationType::RenderStrands == Type && 0 <= E.FrameLODIndex && E.FrameLODIndex < ProjectionHairData.LODDatas.Num() && ProjectionHairData.LODDatas[E.FrameLODIndex].bIsValid)
+			{
+				ProjectionHairData.LODDatas[E.FrameLODIndex].DeformedRootCenter = E.SkeletalWorldBoundsCenter;
+				UpdateHairStrandsMeshTriangles(RHICmdList, ShaderMap, E.FrameLODIndex, HairStrandsTriangleType::DeformedPose, MeshData, ProjectionHairData);
+			}
 		}
 
-		if (EHairStrandsInterpolationType::SimulationStrands == Type && 0 <= E.FrameLODIndex && E.FrameLODIndex < E.SimProjectionHairData.LODDatas.Num() && E.SimProjectionHairData.LODDatas[E.FrameLODIndex].bIsValid)
+		for (FHairStrandsProjectionHairData::HairGroup& ProjectionHairData : E.SimProjectionHairDatas.HairGroups)
 		{
-			E.SimProjectionHairData.LODDatas[E.FrameLODIndex].DeformedRootCenter = E.SkeletalWorldBoundsCenter;
-			UpdateHairStrandsMeshTriangles(RHICmdList, ShaderMap, E.FrameLODIndex, HairStrandsTriangleType::DeformedPose, MeshData, E.SimProjectionHairData);
+			if (EHairStrandsInterpolationType::SimulationStrands == Type && 0 <= E.FrameLODIndex && E.FrameLODIndex < ProjectionHairData.LODDatas.Num() && ProjectionHairData.LODDatas[E.FrameLODIndex].bIsValid)
+			{
+				ProjectionHairData.LODDatas[E.FrameLODIndex].DeformedRootCenter = E.SkeletalWorldBoundsCenter;
+				UpdateHairStrandsMeshTriangles(RHICmdList, ShaderMap, E.FrameLODIndex, HairStrandsTriangleType::DeformedPose, MeshData, ProjectionHairData);
+			}
 		}
 	}
 
@@ -293,7 +320,7 @@ void RunHairStrandsInterpolation(FRHICommandListImmediate& RHICmdList, EWorldTyp
 
 			if (E.InterpolationData.Input && E.InterpolationData.Output && E.InterpolationData.Function)
 			{
-				E.InterpolationData.Function(RHICmdList, E.InterpolationData.Input, E.InterpolationData.Output, E.RenProjectionHairData, E.FrameLODIndex);
+				E.InterpolationData.Function(RHICmdList, E.InterpolationData.Input, E.InterpolationData.Output, E.RenProjectionHairDatas, E.FrameLODIndex);
 			}
 		}
 	}
@@ -313,7 +340,7 @@ void GetGroomInterpolationData(const EWorldType::Type WorldType, FHairStrandsPro
 	}
 }
 
-void GetGroomInterpolationData(const EWorldType::Type WorldType, FHairStrandsProjectionHairDatas& Out, TArray<int32>& OutLODIndices)
+void GetGroomInterpolationData(const EWorldType::Type WorldType, FHairStrandsProjectionHairData& Out, TArray<int32>& OutLODIndices)
 {
 	for (FHairStrandsManager::Element& E : GHairManager.Elements)
 	{
@@ -323,8 +350,11 @@ void GetGroomInterpolationData(const EWorldType::Type WorldType, FHairStrandsPro
 		const bool bHasDynamicMesh = E.CachedGeometry.Sections.Num() > 0;
 		if (bHasDynamicMesh)
 		{
-			Out.HairData.Add(E.RenProjectionHairData);
-			OutLODIndices.Add(E.FrameLODIndex);
+			for (FHairStrandsProjectionHairData::HairGroup& ProjectionHairData : E.RenProjectionHairDatas.HairGroups)
+			{
+				Out.HairGroups.Add(ProjectionHairData);
+				OutLODIndices.Add(E.FrameLODIndex);
+			}
 		}
 	}
 }
@@ -335,13 +365,19 @@ FHairStrandsDebugInfos GetHairStandsDebugInfos()
 	for (FHairStrandsManager::Element& E : GHairManager.Elements)
 	{
 		FHairStrandsDebugInfo& Info = Infos.AddDefaulted_GetRef();
+		Info = E.DebugInfo;
 		Info.Id = E.Id;
 		Info.WorldType = E.WorldType;
-		Info.GroupInfos = E.DebugGroupInfos;
-		for (FHairStrandsDebugInfo::FGroupInfo& GroupInfo : Info.GroupInfos)
-		{
-			GroupInfo.LODCount = E.RenProjectionHairData.LODDatas.Num();
-			GroupInfo.bHasSkinInterpolation = E.RenProjectionHairData.LODDatas.Num() > 0;
+
+		check(Info.HairGroups.Num() == E.RenProjectionHairDatas.HairGroups.Num());
+
+		const uint32 GroupCount = Info.HairGroups.Num();
+		for (uint32 GroupIt=0; GroupIt < GroupCount; ++GroupIt)
+		{		
+			FHairStrandsDebugInfo::HairGroup& GroupInfo = Info.HairGroups[GroupIt];
+			FHairStrandsProjectionHairData::HairGroup& ProjectionHair = E.RenProjectionHairDatas.HairGroups[GroupIt];
+			GroupInfo.LODCount = ProjectionHair.LODDatas.Num();
+			GroupInfo.bHasSkinInterpolation = ProjectionHair.LODDatas.Num() > 0;
 		}
 	}
 
diff --git a/Engine/Source/Runtime/Renderer/Private/HairStrands/HairStrandsMeshProjection.cpp b/Engine/Source/Runtime/Renderer/Private/HairStrands/HairStrandsMeshProjection.cpp
--- a/Engine/Source/Runtime/Renderer/Private/HairStrands/HairStrandsMeshProjection.cpp
+++ b/Engine/Source/Runtime/Renderer/Private/HairStrands/HairStrandsMeshProjection.cpp
@@ -50,8 +50,8 @@ static void AddHairStrandMeshProjectionPass(
 	TShaderMap<FGlobalShaderType>* ShaderMap,
 	const bool bClear,
 	const int32 LODIndex,
-	const FHairStrandsProjectionMeshSection& MeshSectionData,
-	const FHairStrandsProjectionHairData& RootData,
+	const FHairStrandsProjectionMeshData::Section& MeshSectionData,
+	const FHairStrandsProjectionHairData::HairGroup& RootData,
 	FRDGBufferRef RootDistanceBuffer)
 {
 	if (!RootData.RootPositionBuffer ||
@@ -110,7 +110,7 @@ void ProjectHairStrandsOntoMesh(
 	TShaderMap<FGlobalShaderType>* ShaderMap,
 	const int32 LODIndex,
 	FHairStrandsProjectionMeshData& ProjectionMeshData, 
-	FHairStrandsProjectionHairData& ProjectionHairData)
+	FHairStrandsProjectionHairData::HairGroup& ProjectionHairData)
 {
 	if (LODIndex < 0 || LODIndex >= ProjectionHairData.LODDatas.Num())
 		return;
@@ -120,7 +120,7 @@ void ProjectHairStrandsOntoMesh(
 	FRDGBufferRef RootDistanceBuffer = GraphBuilder.CreateBuffer(FRDGBufferDesc::CreateStructuredDesc(sizeof(float),  ProjectionHairData.RootCount),	TEXT("HairStrandsTriangleDistance"));
 
 	bool ClearDistance = true;
-	for (FHairStrandsProjectionMeshSection& MeshSection : ProjectionMeshData.Sections)
+	for (FHairStrandsProjectionMeshData::Section& MeshSection : ProjectionMeshData.Sections)
 	{
 		check(ProjectionHairData.LODDatas[LODIndex].LODIndex == LODIndex);
 		AddHairStrandMeshProjectionPass(GraphBuilder, ShaderMap, ClearDistance, LODIndex, MeshSection, ProjectionHairData, RootDistanceBuffer);
@@ -168,8 +168,8 @@ static void AddHairStrandUpdateMeshTrianglesPass(
 	TShaderMap<FGlobalShaderType>* ShaderMap,
 	const int32 LODIndex,
 	const HairStrandsTriangleType Type,
-	const FHairStrandsProjectionMeshSection& MeshSectionData,
-	FHairStrandsProjectionHairData& RootData)	
+	const FHairStrandsProjectionMeshData::Section& MeshSectionData,
+	FHairStrandsProjectionHairData::HairGroup& RootData)
 {
 	if (RootData.RootCount == 0 || LODIndex < 0 || LODIndex >= RootData.LODDatas.Num())
 	{
@@ -229,11 +229,11 @@ void UpdateHairStrandsMeshTriangles(
 	const int32 LODIndex,
 	const HairStrandsTriangleType Type,
 	FHairStrandsProjectionMeshData& ProjectionMeshData,
-	FHairStrandsProjectionHairData& ProjectionHairData)
+	FHairStrandsProjectionHairData::HairGroup& ProjectionHairData)
 {
 	FRDGBuilder GraphBuilder(RHICmdList);
 
-	for (FHairStrandsProjectionMeshSection& MeshSection : ProjectionMeshData.Sections)
+	for (FHairStrandsProjectionMeshData::Section& MeshSection : ProjectionMeshData.Sections)
 	{
 		AddHairStrandUpdateMeshTrianglesPass(GraphBuilder, ShaderMap, LODIndex, Type, MeshSection, ProjectionHairData);
 	}
diff --git a/Engine/Source/Runtime/Renderer/Private/HairStrands/HairStrandsMeshProjection.h b/Engine/Source/Runtime/Renderer/Private/HairStrands/HairStrandsMeshProjection.h
--- a/Engine/Source/Runtime/Renderer/Private/HairStrands/HairStrandsMeshProjection.h
+++ b/Engine/Source/Runtime/Renderer/Private/HairStrands/HairStrandsMeshProjection.h
@@ -4,6 +4,7 @@
 
 #include "CoreMinimal.h"
 #include "RendererInterface.h"
+#include "HairStrandsInterface.h"
 #include "Shader.h"
 
 /* Project hair strands onto a LOD mesh */
@@ -11,8 +12,8 @@ void ProjectHairStrandsOntoMesh(
 	FRHICommandListImmediate& RHICmdList, 
 	TShaderMap<FGlobalShaderType>* ShaderMap,
 	const int32 LODIndex,
-	struct FHairStrandsProjectionMeshData& ProjectionMeshData, 
-	struct FHairStrandsProjectionHairData& ProjectionHairData);
+	FHairStrandsProjectionMeshData& ProjectionMeshData, 
+	FHairStrandsProjectionHairData::HairGroup& ProjectionHairData);
 
 enum class HairStrandsTriangleType
 {
@@ -26,5 +27,5 @@ void UpdateHairStrandsMeshTriangles(
 	TShaderMap<FGlobalShaderType>* ShaderMap,
 	const int32 LODIndex,
 	const HairStrandsTriangleType Type,
-	struct FHairStrandsProjectionMeshData& ProjectionMeshData,
-	struct FHairStrandsProjectionHairData& ProjectionHairData);
\ No newline at end of file
+	FHairStrandsProjectionMeshData& ProjectionMeshData,
+	FHairStrandsProjectionHairData::HairGroup& ProjectionHairData);
\ No newline at end of file
diff --git a/Engine/Source/Runtime/Renderer/Public/HairStrandsInterface.h b/Engine/Source/Runtime/Renderer/Public/HairStrandsInterface.h
--- a/Engine/Source/Runtime/Renderer/Public/HairStrandsInterface.h
+++ b/Engine/Source/Runtime/Renderer/Public/HairStrandsInterface.h
@@ -52,7 +52,7 @@ typedef void (*THairStrandsInterpolationFunction)(
 	FRHICommandListImmediate& RHICmdList, 
 	struct FHairStrandsInterpolationInput* Input, 
 	struct FHairStrandsInterpolationOutput* Output, 
-	struct FHairStrandsProjectionHairData& HairData,
+	struct FHairStrandsProjectionHairData& HairProjection,
 	int32 LODIndex);
 
 struct FHairStrandsInterpolationData
@@ -64,16 +64,34 @@ struct FHairStrandsInterpolationData
 
 struct FRWBuffer;
 
-struct FHairStrandsProjectionHairData
+struct FHairStrandsProjectionMeshData
 {
-	FRHIShaderResourceView* RootPositionBuffer = nullptr;
-	FRHIShaderResourceView* RootNormalBuffer = nullptr;
-	FRWBuffer* VertexToCurveIndexBuffer = nullptr;
+	struct Section
+	{
+		FTransform LocalToWorld;
+		FRHIShaderResourceView* PositionBuffer = nullptr;
+		FRHIShaderResourceView* IndexBuffer = nullptr;
+		uint32 NumPrimitives = 0;
+		uint32 VertexBaseIndex = 0;
+		uint32 IndexBaseIndex = 0;
+		uint32 TotalVertexCount = 0;
+		uint32 TotalIndexCount = 0;
+		uint32 SectionIndex = 0;
+		int32 LODIndex = 0;
+	};
 
+	TArray<Section> Sections;
+};
+
+struct FHairStrandsProjectionHairData
+{
 	struct LODData
 	{
 		bool bIsValid = false;
 		int32 LODIndex = -1;
+
+		// The index buffers stores the mesh section & the triangle index into a single uint32 
+		// (3 highest bits store the section (up to 8 sections)
 		FRWBuffer* RootTriangleIndexBuffer = nullptr;
 		FRWBuffer* RootTriangleBarycentricBuffer = nullptr;
 
@@ -89,34 +107,20 @@ struct FHairStrandsProjectionHairData
 		FRWBuffer* DeformedRootTrianglePosition1Buffer = nullptr;
 		FRWBuffer* DeformedRootTrianglePosition2Buffer = nullptr;
 	};
-	TArray<LODData> LODDatas;
 
-	uint32 RootCount = 0;
-	FTransform LocalToWorld = FTransform::Identity;
-};
+	struct HairGroup
+	{
+		FRHIShaderResourceView* RootPositionBuffer = nullptr;
+		FRHIShaderResourceView* RootNormalBuffer = nullptr;
+		FRWBuffer* VertexToCurveIndexBuffer = nullptr;
 
-struct FHairStrandsProjectionMeshSection
-{
-	FTransform LocalToWorld;
-	FRHIShaderResourceView* PositionBuffer = nullptr;
-	FRHIShaderResourceView* IndexBuffer = nullptr;
-	uint32 NumPrimitives = 0;
-	uint32 VertexBaseIndex = 0;
-	uint32 IndexBaseIndex = 0;
-	uint32 TotalVertexCount = 0;
-	uint32 TotalIndexCount = 0;
-	uint32 SectionIndex = 0;
-	int32 LODIndex = 0;
-};
+		TArray<LODData> LODDatas;
 
-struct FHairStrandsProjectionMeshData
-{
-	TArray<FHairStrandsProjectionMeshSection> Sections;
-};
+		uint32 RootCount = 0;
+		FTransform LocalToWorld = FTransform::Identity;
+	};
 
-struct FHairStrandsProjectionHairDatas
-{
-	TArray<FHairStrandsProjectionHairData> HairData;
+	TArray<HairGroup> HairGroups;
 };
 
 RENDERER_API void AddHairStrandsProjectionQuery(FRHICommandListImmediate& RHICmdList, uint64 Id, EWorldType::Type WorldType, int32 LODIndex, const FVector& RestRootCenter);
@@ -129,7 +133,7 @@ struct FHairStrandsDebugInfo
 	uint64 Id = 0;
 	EWorldType::Type WorldType = EWorldType::None;
 
-	struct FGroupInfo
+	struct HairGroup
 	{
 		float MaxRadius = 0;
 		float MaxLength = 0;
@@ -139,8 +143,7 @@ struct FHairStrandsDebugInfo
 		bool bHasSkinInterpolation = false;
 		uint32 LODCount = 0;
 	};
-	typedef TArray<FGroupInfo> FGroupInfos;
-	FGroupInfos GroupInfos;
+	TArray<HairGroup> HairGroups;
 };
 
 typedef TArray<FHairStrandsDebugInfo> FHairStrandsDebugInfos;
@@ -155,7 +158,7 @@ RENDERER_API void RegisterHairStrands(
 	const FHairStrandsInterpolationData& E,
 	const FHairStrandsProjectionHairData& RenProjection,
 	const FHairStrandsProjectionHairData& SimProjection,
-	const FHairStrandsDebugInfo::FGroupInfos& GroupInfos);
+	const FHairStrandsDebugInfo& DebugInfo);
 
 RENDERER_API void UnregisterHairStrands(
 	uint64 Id, 
