Move hair deformed position/tangent buffers into group component resources to support multiple instance of the same assets.

(ushell-p4-cherrypick of 9439914 by Charles.deRousiers)

diff --git a/Engine/Plugins/Experimental/HairStrands/Source/HairStrandsCore/Private/GroomAsset.cpp b/Engine/Plugins/Experimental/HairStrands/Source/HairStrandsCore/Private/GroomAsset.cpp
--- a/Engine/Plugins/Experimental/HairStrands/Source/HairStrandsCore/Private/GroomAsset.cpp
+++ b/Engine/Plugins/Experimental/HairStrands/Source/HairStrandsCore/Private/GroomAsset.cpp
@@ -26,7 +26,7 @@ void CreateBuffer(const TArray<typename FormatType::Type>& InData, FRWBuffer& Ou
 }
 
 template<typename FormatType>
-void CreateBuffer(uint32 InVertexCount , FRWBuffer& OutBuffer)
+void CreateBuffer(uint32 InVertexCount, FRWBuffer& OutBuffer)
 {
 	const uint32 DataCount = InVertexCount;
 	const uint32 DataSizeInBytes = FormatType::SizeInByte*DataCount;
@@ -41,29 +41,52 @@ void CreateBuffer(uint32 InVertexCount , FRWBuffer& OutBuffer)
 
 /////////////////////////////////////////////////////////////////////////////////////////
 
-FHairStrandsResource::FHairStrandsResource(const FHairStrandsDatas::FRenderData& HairStrandRenderData) : 
-	RestPositionBuffer(), TangentBuffer(), AttributeBuffer(), RenderData(HairStrandRenderData)
+FHairStrandsRestResource::FHairStrandsRestResource(const FHairStrandsDatas::FRenderData& HairStrandRenderData) : 
+	RestPositionBuffer(), AttributeBuffer(), RenderData(HairStrandRenderData)
 {}
 
-void FHairStrandsResource::InitRHI()
+void FHairStrandsRestResource::InitRHI()
 {
 	const TArray<FHairStrandsPositionFormat::Type>& RenderingPositions = RenderData.RenderingPositions;
 	const TArray<FHairStrandsAttributeFormat::Type>& RenderingAttributes = RenderData.RenderingAttributes;
 
 	CreateBuffer<FHairStrandsPositionFormat>(RenderingPositions, RestPositionBuffer);
-	CreateBuffer<FHairStrandsPositionFormat>(RenderingPositions, DeformedPositionBuffer[0]);
-	CreateBuffer<FHairStrandsPositionFormat>(RenderingPositions, DeformedPositionBuffer[1]);
-	CreateBuffer<FHairStrandsTangentFormat>(RenderingPositions.Num() * FHairStrandsTangentFormat::ComponentCount, TangentBuffer);
 	CreateBuffer<FHairStrandsAttributeFormat>(RenderingAttributes, AttributeBuffer);
 }
 
-void FHairStrandsResource::ReleaseRHI()
+void FHairStrandsRestResource::ReleaseRHI()
 {
 	RestPositionBuffer.Release();
+	AttributeBuffer.Release();
+}
+
+/////////////////////////////////////////////////////////////////////////////////////////
+
+FHairStrandsDeformedResource::FHairStrandsDeformedResource(const FHairStrandsDatas::FRenderData& HairStrandRenderData, bool bInInitializedData) :
+	bInitializedData(bInInitializedData), RenderData(HairStrandRenderData)
+{}
+
+void FHairStrandsDeformedResource::InitRHI()
+{
+	const uint32 VertexCount = RenderData.RenderingPositions.Num();
+	if (bInitializedData)
+	{
+		CreateBuffer<FHairStrandsPositionFormat>(RenderData.RenderingPositions, DeformedPositionBuffer[0]);
+		CreateBuffer<FHairStrandsPositionFormat>(RenderData.RenderingPositions, DeformedPositionBuffer[1]);
+	}
+	else
+	{
+		CreateBuffer<FHairStrandsPositionFormat>(VertexCount, DeformedPositionBuffer[0]);
+		CreateBuffer<FHairStrandsPositionFormat>(VertexCount, DeformedPositionBuffer[1]);
+	}
+	CreateBuffer<FHairStrandsTangentFormat>(VertexCount * FHairStrandsTangentFormat::ComponentCount, TangentBuffer);
+}
+
+void FHairStrandsDeformedResource::ReleaseRHI()
+{
 	DeformedPositionBuffer[0].Release();
 	DeformedPositionBuffer[1].Release();
 	TangentBuffer.Release();
-	AttributeBuffer.Release();
 }
 
 /////////////////////////////////////////////////////////////////////////////////////////
@@ -241,13 +264,13 @@ void UGroomAsset::InitResource()
 {
 	for (FHairGroupData& GroupData : HairGroupsData)
 	{
-		GroupData.HairStrandsResource = new FHairStrandsResource(GroupData.HairRenderData.RenderData);
+		GroupData.HairStrandsRestResource = new FHairStrandsRestResource(GroupData.HairRenderData.RenderData);
 
-		BeginInitResource(GroupData.HairStrandsResource);
+		BeginInitResource(GroupData.HairStrandsRestResource);
 
-		GroupData.HairSimulationResource = new FHairStrandsResource(GroupData.HairSimulationData.RenderData);
+		GroupData.HairSimulationRestResource = new FHairStrandsRestResource(GroupData.HairSimulationData.RenderData);
 
-		BeginInitResource(GroupData.HairSimulationResource);
+		BeginInitResource(GroupData.HairSimulationRestResource);
 	}
 }
 
@@ -255,14 +278,14 @@ void UGroomAsset::UpdateResource()
 {
 	for (FHairGroupData& GroupData : HairGroupsData)
 	{
-		if (GroupData.HairStrandsResource)
+		if (GroupData.HairStrandsRestResource)
 		{
-			BeginUpdateResourceRHI(GroupData.HairStrandsResource);
+			BeginUpdateResourceRHI(GroupData.HairStrandsRestResource);
 		}
 
-		if (GroupData.HairSimulationResource)
+		if (GroupData.HairSimulationRestResource)
 		{
-			BeginUpdateResourceRHI(GroupData.HairSimulationResource);
+			BeginUpdateResourceRHI(GroupData.HairSimulationRestResource);
 		}
 	}
 }
@@ -271,28 +294,28 @@ void UGroomAsset::ReleaseResource()
 {
 	for (FHairGroupData& GroupData : HairGroupsData)
 	{
-		if (GroupData.HairStrandsResource)
+		if (GroupData.HairStrandsRestResource)
 		{
-			FHairStrandsResource* InResource = GroupData.HairStrandsResource;
+			FHairStrandsRestResource* InResource = GroupData.HairStrandsRestResource;
 			ENQUEUE_RENDER_COMMAND(ReleaseHairStrandsResourceCommand)(
 				[InResource](FRHICommandList& RHICmdList)
 			{
 				InResource->ReleaseResource();
 				delete InResource;
 			});
-			GroupData.HairStrandsResource = nullptr;
+			GroupData.HairStrandsRestResource = nullptr;
 		}
 
-		if (GroupData.HairSimulationResource)
+		if (GroupData.HairSimulationRestResource)
 		{
-			FHairStrandsResource* InResource = GroupData.HairSimulationResource;
+			FHairStrandsRestResource* InResource = GroupData.HairSimulationRestResource;
 			ENQUEUE_RENDER_COMMAND(ReleaseHairStrandsResourceCommand)(
 				[InResource](FRHICommandList& RHICmdList)
 			{
 				InResource->ReleaseResource();
 				delete InResource;
 			});
-			GroupData.HairSimulationResource = nullptr;
+			GroupData.HairSimulationRestResource = nullptr;
 		}
 	}
 }
diff --git a/Engine/Plugins/Experimental/HairStrands/Source/HairStrandsCore/Private/GroomComponent.cpp b/Engine/Plugins/Experimental/HairStrands/Source/HairStrandsCore/Private/GroomComponent.cpp
--- a/Engine/Plugins/Experimental/HairStrands/Source/HairStrandsCore/Private/GroomComponent.cpp
+++ b/Engine/Plugins/Experimental/HairStrands/Source/HairStrandsCore/Private/GroomComponent.cpp
@@ -528,14 +528,24 @@ FHairStrandsDatas* UGroomComponent::GetGuideStrandsDatas(uint32 GroupIndex)
 	return &GroomAsset->HairGroupsData[GroupIndex].HairSimulationData;
 }
 
-FHairStrandsResource* UGroomComponent::GetGuideStrandsResource(uint32 GroupIndex)
+FHairStrandsRestResource* UGroomComponent::GetGuideStrandsRestResource(uint32 GroupIndex)
 {
 	if (!GroomAsset || GroupIndex >= uint32(GroomAsset->GetNumHairGroups()))
 	{
 		return nullptr;
 	}
 
-	return GroomAsset->HairGroupsData[GroupIndex].HairSimulationResource;
+	return GroomAsset->HairGroupsData[GroupIndex].HairSimulationRestResource;
+}
+
+FHairStrandsDeformedResource* UGroomComponent::GetGuideStrandsDeformedResource(uint32 GroupIndex)
+{
+	if (GroupIndex >= uint32(HairGroupResources.Num()))
+	{
+		return nullptr;
+	}
+
+	return HairGroupResources[GroupIndex].SimDeformedResources;
 }
 
 template<typename T> void SafeDelete(T*& Data) 
@@ -623,7 +633,7 @@ void UGroomComponent::InitResources()
 	FHairStrandsDebugInfo DebugGroupInfo;
 	for (FHairGroupData& GroupData : GroomAsset->HairGroupsData)
 	{
-		if (!GroupData.HairStrandsResource)
+		if (!GroupData.HairStrandsRestResource)
 			return;
 
 		FHairStrandsDebugInfo::HairGroup& DebugHairGroup = DebugGroupInfo.HairGroups.AddDefaulted_GetRef();
@@ -658,8 +668,12 @@ void UGroomComponent::InitResources()
 			}
 		}
 		
-		Res.RenderResources = GroupData.HairStrandsResource;
-		Res.SimResources = GroupData.HairSimulationResource;
+		Res.RenderRestResources = GroupData.HairStrandsRestResource;
+		Res.SimRestResources = GroupData.HairSimulationRestResource;
+		Res.RenderDeformedResources = new FHairStrandsDeformedResource(GroupData.HairRenderData.RenderData, false);
+		Res.SimDeformedResources = new FHairStrandsDeformedResource(GroupData.HairSimulationData.RenderData, true);
+		BeginInitResource(Res.RenderDeformedResources);
+		BeginInitResource(Res.SimDeformedResources);
 
 		FHairStrandsInterpolationOutput::HairGroup& InterpolationOutputGroup = InterpolationOutput->HairGroups.AddDefaulted_GetRef();
 		FHairStrandsInterpolationInput::FHairGroup& InterpolationInputGroup = InterpolationInput->HairGroups.AddDefaulted_GetRef();
@@ -692,16 +706,16 @@ void UGroomComponent::InitResources()
 		{
 			FHairGroupResource& Res = (*LocalResources)[GroupIt];
 
-			FHairStrandsInterpolationInput::FHairGroup& InputGroup		= Interpolation.Input->HairGroups[GroupIt];
+			FHairStrandsInterpolationInput::FHairGroup& InputGroup 	= Interpolation.Input->HairGroups[GroupIt];
 			FHairStrandsInterpolationOutput::HairGroup& OutputGroup	= Interpolation.Output->HairGroups[GroupIt];
 
-			InputGroup.RenderRestPosePositionBuffer	= &Res.RenderResources->RestPositionBuffer;
-			InputGroup.RenderAttributeBuffer		= &Res.RenderResources->AttributeBuffer;
-			InputGroup.RenderVertexCount			=  Res.RenderResources->RenderData.RenderingPositions.Num() / FHairStrandsPositionFormat::ComponentCount;
+			InputGroup.RenderRestPosePositionBuffer	= &Res.RenderRestResources->RestPositionBuffer;
+			InputGroup.RenderAttributeBuffer		= &Res.RenderRestResources->AttributeBuffer;
+			InputGroup.RenderVertexCount			=  Res.RenderRestResources->RenderData.RenderingPositions.Num() / FHairStrandsPositionFormat::ComponentCount;
 
-			InputGroup.SimRestPosePositionBuffer	= &Res.SimResources->RestPositionBuffer;
-			InputGroup.SimAttributeBuffer			= &Res.SimResources->AttributeBuffer;
-			InputGroup.SimVertexCount				=  Res.SimResources->RenderData.RenderingPositions.Num() / FHairStrandsPositionFormat::ComponentCount;
+			InputGroup.SimRestPosePositionBuffer	= &Res.SimRestResources->RestPositionBuffer;
+			InputGroup.SimAttributeBuffer			= &Res.SimRestResources->AttributeBuffer;
+			InputGroup.SimVertexCount				=  Res.SimRestResources->RenderData.RenderingPositions.Num() / FHairStrandsPositionFormat::ComponentCount;
 			InputGroup.SimRootPointIndexBuffer		= &Res.InterpolationResource->SimRootPointIndexBuffer;
 
 			InputGroup.Interpolation0Buffer			= &Res.InterpolationResource->Interpolation0Buffer;
@@ -716,13 +730,13 @@ void UGroomComponent::InitResources()
 			}
 			#endif
 
-			OutputGroup.SimDeformedPositionBuffer[0]	= &Res.SimResources->DeformedPositionBuffer[0];
-			OutputGroup.SimDeformedPositionBuffer[1]	= &Res.SimResources->DeformedPositionBuffer[1];
-			OutputGroup.RenderDeformedPositionBuffer[0] = &Res.RenderResources->DeformedPositionBuffer[0];
-			OutputGroup.RenderDeformedPositionBuffer[1] = &Res.RenderResources->DeformedPositionBuffer[1];
-			OutputGroup.RenderAttributeBuffer			= &Res.RenderResources->AttributeBuffer;
-			OutputGroup.RenderTangentBuffer				= &Res.RenderResources->TangentBuffer;
-			OutputGroup.SimTangentBuffer				= &Res.SimResources->TangentBuffer;
+			OutputGroup.SimDeformedPositionBuffer[0]	= &Res.SimDeformedResources->DeformedPositionBuffer[0];
+			OutputGroup.SimDeformedPositionBuffer[1]	= &Res.SimDeformedResources->DeformedPositionBuffer[1];
+			OutputGroup.RenderDeformedPositionBuffer[0] = &Res.RenderDeformedResources->DeformedPositionBuffer[0];
+			OutputGroup.RenderDeformedPositionBuffer[1] = &Res.RenderDeformedResources->DeformedPositionBuffer[1];
+			OutputGroup.RenderAttributeBuffer			= &Res.RenderRestResources->AttributeBuffer;
+			OutputGroup.RenderTangentBuffer				= &Res.RenderDeformedResources->TangentBuffer;
+			OutputGroup.SimTangentBuffer				= &Res.SimDeformedResources->TangentBuffer;
 
 			RenProjectionDatas.HairGroups.Add(ToProjectionHairData(Res.RenRootResources));
 			SimProjectionDatas.HairGroups.Add(ToProjectionHairData(Res.SimRootResources));
@@ -755,6 +769,8 @@ void UGroomComponent::ReleaseResources()
 		SafeRelease(Res.InterpolationResource);
 		SafeRelease(Res.RenRootResources);
 		SafeRelease(Res.SimRootResources);
+		SafeRelease(Res.RenderDeformedResources);
+		SafeRelease(Res.SimDeformedResources);
 	#if RHI_RAYTRACING
 		SafeRelease(Res.RaytracingResources);
 	#endif
diff --git a/Engine/Plugins/Experimental/HairStrands/Source/HairStrandsCore/Public/GroomAsset.h b/Engine/Plugins/Experimental/HairStrands/Source/HairStrandsCore/Public/GroomAsset.h
--- a/Engine/Plugins/Experimental/HairStrands/Source/HairStrandsCore/Public/GroomAsset.h
+++ b/Engine/Plugins/Experimental/HairStrands/Source/HairStrandsCore/Public/GroomAsset.h
@@ -73,10 +73,10 @@ struct FHairStrandsRootResource : public FRenderResource
 };
 
 /* Render buffers that will be used for rendering */
-struct FHairStrandsResource : public FRenderResource
+struct FHairStrandsRestResource : public FRenderResource
 {
 	/** Build the hair strands resource */
-	FHairStrandsResource(const FHairStrandsDatas::FRenderData& HairStrandRenderData);
+	FHairStrandsRestResource(const FHairStrandsDatas::FRenderData& HairStrandRenderData);
 
 	/* Init the buffer */
 	virtual void InitRHI() override;
@@ -90,17 +90,38 @@ struct FHairStrandsResource : public FRenderResource
 	/* Strand hair rest position buffer */
 	FRWBuffer RestPositionBuffer;
 
-	/* Strand hair (previous) position buffer */
+	/* Strand hair offset buffer */
+	FRWBuffer AttributeBuffer;
+
+	/* Reference to the hair strands render data */
+	const FHairStrandsDatas::FRenderData& RenderData;
+};
+
+struct FHairStrandsDeformedResource : public FRenderResource
+{
+	/** Build the hair strands resource */
+	FHairStrandsDeformedResource(const FHairStrandsDatas::FRenderData& HairStrandRenderData, bool bInitializeData);
+
+	/* Init the buffer */
+	virtual void InitRHI() override;
+
+	/* Release the buffer */
+	virtual void ReleaseRHI() override;
+
+	/* Get the resource name */
+	virtual FString GetFriendlyName() const override { return TEXT("FHairStrandsDeformedResource"); }
+
+	/* Strand hair deformed position buffer (previous and current) */
 	FRWBuffer DeformedPositionBuffer[2];
 
-	/* Strand hair offset buffer */
+	/* Strand hair tangent buffer */
 	FRWBuffer TangentBuffer;
 
-	/* Strand hair offset buffer */
-	FRWBuffer AttributeBuffer;
-
 	/* Reference to the hair strands render data */
 	const FHairStrandsDatas::FRenderData& RenderData;
+
+	/* Whether the GPU data should be initialized with the asset data or not */
+	const bool bInitializedData;
 };
 
 struct FHairStrandsInterpolationResource : public FRenderResource
@@ -180,10 +201,10 @@ struct HAIRSTRANDSCORE_API FHairGroupData
 	FHairStrandsInterpolationDatas HairInterpolationData;
 
 	/** Interpolated hair render resource to be allocated */
-	FHairStrandsResource* HairStrandsResource = nullptr;
+	FHairStrandsRestResource* HairStrandsRestResource = nullptr;
 
 	/** Guide render resource to be allocated */
-	FHairStrandsResource* HairSimulationResource = nullptr;
+	FHairStrandsRestResource* HairSimulationRestResource = nullptr;
 
 	friend FArchive& operator<<(FArchive& Ar, FHairGroupData& GroupData);
 };
diff --git a/Engine/Plugins/Experimental/HairStrands/Source/HairStrandsCore/Public/GroomComponent.h b/Engine/Plugins/Experimental/HairStrands/Source/HairStrandsCore/Public/GroomComponent.h
--- a/Engine/Plugins/Experimental/HairStrands/Source/HairStrandsCore/Public/GroomComponent.h
+++ b/Engine/Plugins/Experimental/HairStrands/Source/HairStrandsCore/Public/GroomComponent.h
@@ -60,8 +60,11 @@ public:
 	/** Return the guide hairs datas */
 	FHairStrandsDatas* GetGuideStrandsDatas(uint32 GroupIndex);
 
-	/** Return the guide hairs resources*/
-	FHairStrandsResource* GetGuideStrandsResource(uint32 GroupIndex);
+	/** Return the guide hairs rest resources*/
+	FHairStrandsRestResource* GetGuideStrandsRestResource(uint32 GroupIndex);
+
+	/** Return the guide hairs deformed resources*/
+	FHairStrandsDeformedResource* GetGuideStrandsDeformedResource(uint32 GroupIndex);
 
 #if WITH_EDITOR
 	virtual void PostEditChangeProperty(FPropertyChangedEvent& PropertyChangedEvent) override;
@@ -69,16 +72,23 @@ public:
 
 	struct FHairGroupResource
 	{
+		// Sim to rendering interpolation resources
 		FHairStrandsInterpolationResource* InterpolationResource = nullptr;
+
+		// Projection resources
 		struct FHairStrandsRootResource* RenRootResources = nullptr;
 		struct FHairStrandsRootResource* SimRootResources = nullptr;
 	#if RHI_RAYTRACING
 		FHairStrandsRaytracingResource* RaytracingResources = nullptr;
 	#endif
 
-		// Resources own by the asset
-		FHairStrandsResource* RenderResources = nullptr;
-		FHairStrandsResource* SimResources = nullptr;
+		// Deformed position
+		FHairStrandsDeformedResource* RenderDeformedResources = nullptr;
+		FHairStrandsDeformedResource* SimDeformedResources = nullptr;
+
+		// Rest resources, owned by the asset
+		FHairStrandsRestResource* RenderRestResources = nullptr;
+		FHairStrandsRestResource* SimRestResources = nullptr;
 	};
 	typedef TArray<FHairGroupResource> FHairGroupResources;
 
diff --git a/Engine/Plugins/Experimental/HairStrands/Source/HairStrandsNiagara/Private/NiagaraDataInterfaceHairStrands.cpp b/Engine/Plugins/Experimental/HairStrands/Source/HairStrandsNiagara/Private/NiagaraDataInterfaceHairStrands.cpp
--- a/Engine/Plugins/Experimental/HairStrands/Source/HairStrandsNiagara/Private/NiagaraDataInterfaceHairStrands.cpp
+++ b/Engine/Plugins/Experimental/HairStrands/Source/HairStrandsNiagara/Private/NiagaraDataInterfaceHairStrands.cpp
@@ -194,26 +194,25 @@ void FNDICollisionGridBuffer::ClearBuffers(FRHICommandList& RHICmdList)
 
 //------------------------------------------------------------------------------------------------------------
 
-void FNDIHairStrandsBuffer::SetHairAsset(const FHairStrandsDatas*  HairStrandsDatas, const FHairStrandsResource*  HairStrandsResource)
+void FNDIHairStrandsBuffer::SetHairAsset(const FHairStrandsDatas*  HairStrandsDatas, const FHairStrandsRestResource* HairStrandsRestResource, const FHairStrandsDeformedResource*  HairStrandsDeformedResource)
 {
 	SourceDatas = HairStrandsDatas;
-	SourceResources = HairStrandsResource;
+	SourceRestResources = HairStrandsRestResource;
+	SourceDeformedResources = HairStrandsDeformedResource;
 }
 
 void FNDIHairStrandsBuffer::InitRHI()
 {
-	if (SourceDatas != nullptr && SourceResources != nullptr)
+	if (SourceDatas != nullptr && SourceRestResources != nullptr && SourceDeformedResources != nullptr)
 	{
-		{
-			const uint32 OffsetCount = SourceDatas->GetNumCurves() + 1;
-			const uint32 OffsetBytes = sizeof(uint32)*OffsetCount;
+		const uint32 OffsetCount = SourceDatas->GetNumCurves() + 1;
+		const uint32 OffsetBytes = sizeof(uint32)*OffsetCount;
 
-			CurvesOffsetsBuffer.Initialize(sizeof(uint32), OffsetCount, EPixelFormat::PF_R32_UINT, BUF_Static);
-			void* OffsetBufferData = RHILockVertexBuffer(CurvesOffsetsBuffer.Buffer, 0, OffsetBytes, RLM_WriteOnly);
+		CurvesOffsetsBuffer.Initialize(sizeof(uint32), OffsetCount, EPixelFormat::PF_R32_UINT, BUF_Static);
+		void* OffsetBufferData = RHILockVertexBuffer(CurvesOffsetsBuffer.Buffer, 0, OffsetBytes, RLM_WriteOnly);
 
-			FMemory::Memcpy(OffsetBufferData, SourceDatas->StrandsCurves.CurvesOffset.GetData(), OffsetBytes);
-			RHIUnlockVertexBuffer(CurvesOffsetsBuffer.Buffer);
-		}
+		FMemory::Memcpy(OffsetBufferData, SourceDatas->StrandsCurves.CurvesOffset.GetData(), OffsetBytes);
+		RHIUnlockVertexBuffer(CurvesOffsetsBuffer.Buffer);
 	}
 }
 
@@ -323,9 +322,9 @@ struct FNDIHairStrandsParametersCS : public FNiagaraDataInterfaceParametersCS
 			FNDICollisionGridBuffer* CurrentGridBuffer = ProxyData->CurrentGridBuffer;
 			FNDICollisionGridBuffer* DestinationGridBuffer = ProxyData->DestinationGridBuffer;
 
-			SetUAVParameter(RHICmdList, ComputeShaderRHI, PointsPositionsBuffer, HairStrandsBuffer->SourceResources->DeformedPositionBuffer[0].UAV);
+			SetUAVParameter(RHICmdList, ComputeShaderRHI, PointsPositionsBuffer, HairStrandsBuffer->SourceDeformedResources->DeformedPositionBuffer[0].UAV);
 			SetSRVParameter(RHICmdList, ComputeShaderRHI, CurvesOffsetsBuffer, HairStrandsBuffer->CurvesOffsetsBuffer.SRV);
-			SetSRVParameter(RHICmdList, ComputeShaderRHI, RestPositionsBuffer, HairStrandsBuffer->SourceResources->RestPositionBuffer.SRV);
+			SetSRVParameter(RHICmdList, ComputeShaderRHI, RestPositionsBuffer, HairStrandsBuffer->SourceRestResources->RestPositionBuffer.SRV);
 
 			RHICmdList.TransitionResource(EResourceTransitionAccess::EWritable, EResourceTransitionPipeline::EComputeToCompute, DestinationGridBuffer->GridDataBuffer.UAV);
 			SetUAVParameter(RHICmdList, ComputeShaderRHI, GridDestinationBuffer, DestinationGridBuffer->GridDataBuffer.UAV);
@@ -531,9 +530,10 @@ bool UNiagaraDataInterfaceHairStrands::InitPerInstanceData(void* PerInstanceData
 	GroomAsset = IsComponentValid() ? SourceComponent->GroomAsset : DefaultSource;
 	const uint32 GroupIndex = 0;
 	FHairStrandsDatas* StrandsDatas = IsComponentValid() ? SourceComponent->GetGuideStrandsDatas(GroupIndex) : (DefaultSource != nullptr && DefaultSource->GetNumHairGroups() > 0) ? &DefaultSource->HairGroupsData[0].HairRenderData : nullptr;
-	FHairStrandsResource* StrandsResource = IsComponentValid() ? SourceComponent->GetGuideStrandsResource(GroupIndex) : (DefaultSource != nullptr && DefaultSource->GetNumHairGroups() > 0) ? DefaultSource->HairGroupsData[0].HairStrandsResource : nullptr;
+	FHairStrandsRestResource* StrandsRestResource = IsComponentValid() ? SourceComponent->GetGuideStrandsRestResource(GroupIndex) : (DefaultSource != nullptr && DefaultSource->GetNumHairGroups() > 0) ? DefaultSource->HairGroupsData[0].HairStrandsRestResource : nullptr;
+	FHairStrandsDeformedResource* StrandsDeformedResource = IsComponentValid() ? SourceComponent->GetGuideStrandsDeformedResource(GroupIndex) : nullptr;
 
-	if (StrandsDatas == nullptr || StrandsResource == nullptr)
+	if (StrandsDatas == nullptr || StrandsRestResource == nullptr || StrandsDeformedResource == nullptr)
 	{
 		UE_LOG(LogHairStrands, Log, TEXT("Hair Strands data interface has no valid asset. Failed InitPerInstanceData - %s %d"), *GetFullName(), DefaultSource);
 		return false;
@@ -543,7 +543,7 @@ bool UNiagaraDataInterfaceHairStrands::InitPerInstanceData(void* PerInstanceData
 		FNDIHairStrandsBuffer* HairStrandsBuffer = new FNDIHairStrandsBuffer;
 		FNDICollisionGridBuffer* CurrentGridBuffer = new FNDICollisionGridBuffer;
 		FNDICollisionGridBuffer* DestinationGridBuffer = new FNDICollisionGridBuffer;
-		HairStrandsBuffer->SetHairAsset(StrandsDatas, StrandsResource);
+		HairStrandsBuffer->SetHairAsset(StrandsDatas, StrandsRestResource, StrandsDeformedResource);
 		
 
 		// Push instance data to RT
diff --git a/Engine/Plugins/Experimental/HairStrands/Source/HairStrandsNiagara/Public/NiagaraDataInterfaceHairStrands.h b/Engine/Plugins/Experimental/HairStrands/Source/HairStrandsNiagara/Public/NiagaraDataInterfaceHairStrands.h
--- a/Engine/Plugins/Experimental/HairStrands/Source/HairStrandsNiagara/Public/NiagaraDataInterfaceHairStrands.h
+++ b/Engine/Plugins/Experimental/HairStrands/Source/HairStrandsNiagara/Public/NiagaraDataInterfaceHairStrands.h
@@ -51,7 +51,7 @@ struct FNDICollisionGridBuffer : public FRenderResource
 struct FNDIHairStrandsBuffer : public FRenderResource
 {
 	/** Set the asset that will be used to affect the buffer */
-	void SetHairAsset(const FHairStrandsDatas*  HairStrandsDatas, const FHairStrandsResource*  HairStrandsResource );
+	void SetHairAsset(const FHairStrandsDatas*  HairStrandsDatas, const FHairStrandsRestResource*  HairStrandsRestResource, const FHairStrandsDeformedResource*  HairStrandsDeformedResource);
 
 	/** Init the buffer */
 	virtual void InitRHI() override;
@@ -69,7 +69,10 @@ struct FNDIHairStrandsBuffer : public FRenderResource
 	const FHairStrandsDatas* SourceDatas;
 
 	/** The strand asset resource from which to sample */
-	const FHairStrandsResource* SourceResources;
+	const FHairStrandsRestResource* SourceRestResources;
+
+	/** The strand deformed resource to write into */
+	const FHairStrandsDeformedResource* SourceDeformedResources;
 };
 
 /** Data stored per strand base instance*/
