Expose raytracing radius scale per instance.

(ushell-p4-cherrypick of 9440791 by Charles.deRousiers)

diff --git a/Engine/Plugins/Experimental/HairStrands/Source/HairStrandsCore/Private/GroomComponent.cpp b/Engine/Plugins/Experimental/HairStrands/Source/HairStrandsCore/Private/GroomComponent.cpp
--- a/Engine/Plugins/Experimental/HairStrands/Source/HairStrandsCore/Private/GroomComponent.cpp
+++ b/Engine/Plugins/Experimental/HairStrands/Source/HairStrandsCore/Private/GroomComponent.cpp
@@ -379,6 +379,7 @@ UGroomComponent::UGroomComponent(const FObjectInitializer& ObjectInitializer)
 	bSelectable = true;
 	RegisteredSkeletalMeshComponent = nullptr;
 	HairDensity = 1;
+	HairRaytracingRadiusScale = 1;
 	bSkinGroom = false;
 	InitializedResources = nullptr;
 	Mobility = EComponentMobility::Movable;
@@ -677,8 +678,9 @@ void UGroomComponent::InitResources()
 
 		FHairStrandsInterpolationOutput::HairGroup& InterpolationOutputGroup = InterpolationOutput->HairGroups.AddDefaulted_GetRef();
 		FHairStrandsInterpolationInput::FHairGroup& InterpolationInputGroup = InterpolationInput->HairGroups.AddDefaulted_GetRef();
-		InterpolationInputGroup.HairRadius		= GroupData.HairRenderData.StrandsCurves.MaxRadius;
+		InterpolationInputGroup.HairRadius = GroupData.HairRenderData.StrandsCurves.MaxRadius;
 		InterpolationInputGroup.HairWorldOffset = GroupData.HairRenderData.BoundingBox.GetCenter();
+		InterpolationInputGroup.HairRaytracingRadiusScale = HairRaytracingRadiusScale;
 	}
 
 	FHairStrandsInterpolationData Interpolation;
@@ -1020,6 +1022,37 @@ void UGroomComponent::PostEditChangeProperty(FPropertyChangedEvent& PropertyChan
 			}
 		}
 	}
+
+	if (PropertyName == GET_MEMBER_NAME_CHECKED(UGroomComponent, HairRaytracingRadiusScale) && InterpolationInput)
+	{
+		for (FHairStrandsInterpolationInput::FHairGroup& Group : InterpolationInput->HairGroups)
+		{
+			Group.HairRaytracingRadiusScale = HairRaytracingRadiusScale;
+		}
+	}
+}
+
+bool UGroomComponent::CanEditChange(const UProperty* InProperty) const
+{
+	if (InProperty)
+	{
+		FString PropertyName = InProperty->GetName();
+
+		if (FCString::Strcmp(*PropertyName, TEXT("HairRaytracingRadiusScale")) == 0)
+		{
+			bool bIsEditable = false;
+			#if RHI_RAYTRACING
+			FRayTracingGeometry* RayTracingGeometry = nullptr;
+			if (IsRayTracingEnabled())
+			{
+				bIsEditable = true;
+			}
+			#endif
+			return bIsEditable;
+		}
+	}
+
+	return Super::CanEditChange(InProperty);
 }
 #endif
 
diff --git a/Engine/Plugins/Experimental/HairStrands/Source/HairStrandsCore/Private/HairStrandsRendering.cpp b/Engine/Plugins/Experimental/HairStrands/Source/HairStrandsCore/Private/HairStrandsRendering.cpp
--- a/Engine/Plugins/Experimental/HairStrands/Source/HairStrandsCore/Private/HairStrandsRendering.cpp
+++ b/Engine/Plugins/Experimental/HairStrands/Source/HairStrandsCore/Private/HairStrandsRendering.cpp
@@ -11,8 +11,8 @@
 static int32 GHairDeformationType = 0;
 static FAutoConsoleVariableRef CVarHairDeformationType(TEXT("r.HairStrands.DeformationType"), GHairDeformationType, TEXT("Type of procedural deformation applied on hair strands (0:bypass, 1:wave, 2:normal)"));
 
-static float GHairRaytracingRadiusScale = 1;
-static FAutoConsoleVariableRef CVarHairRaytracingRadiusScale(TEXT("r.HairStrands.RaytracingRadiusScale"), GHairRaytracingRadiusScale, TEXT("Scale factor for raytracing hair strands geometry"));
+static float GHairRaytracingRadiusScale = 0;
+static FAutoConsoleVariableRef CVarHairRaytracingRadiusScale(TEXT("r.HairStrands.RaytracingRadiusScale"), GHairRaytracingRadiusScale, TEXT("Override the per instance scale factor for raytracing hair strands geometry (0: disabled, >0:enabled)"));
 
 
 static FIntVector ComputeDispatchCount(uint32 ItemCount, uint32 GroupSize)
@@ -470,7 +470,7 @@ void ComputeHairStrandsInterpolation(
 				AddGenerateRaytracingGeometryPass(
 					GraphBuilder,
 					Input.RenderVertexCount,
-					Input.HairRadius * GHairRaytracingRadiusScale,
+					Input.HairRadius * (GHairRaytracingRadiusScale > 0 ? GHairRaytracingRadiusScale : Input.HairRaytracingRadiusScale),
 					Input.HairWorldOffset,
 					Output.RenderDeformedPositionBuffer[CurrIndex]->SRV,
 					Input.RaytracingPositionBuffer->UAV);
diff --git a/Engine/Plugins/Experimental/HairStrands/Source/HairStrandsCore/Public/GroomComponent.h b/Engine/Plugins/Experimental/HairStrands/Source/HairStrandsCore/Public/GroomComponent.h
--- a/Engine/Plugins/Experimental/HairStrands/Source/HairStrandsCore/Public/GroomComponent.h
+++ b/Engine/Plugins/Experimental/HairStrands/Source/HairStrandsCore/Public/GroomComponent.h
@@ -36,6 +36,10 @@ public:
 	UPROPERTY(EditAnywhere, BlueprintReadOnly, Category = "HairStrands Rendering")
 	bool bSkinGroom;
 
+	/** Scale the hair geometry radius for ray tracing effects (e.g. shadow) */
+	UPROPERTY(EditAnywhere, BlueprintReadWrite, Category = "HairStrands Rendering", meta = (ClampMin = "0.0001", UIMin = "0.01", UIMax = "10.0"))
+	float HairRaytracingRadiusScale;
+
 	//~ Begin UActorComponent Interface.
 	virtual void OnRegister() override;
 	virtual void OnUnregister() override;
@@ -68,6 +72,7 @@ public:
 
 #if WITH_EDITOR
 	virtual void PostEditChangeProperty(FPropertyChangedEvent& PropertyChangedEvent) override;
+	virtual bool CanEditChange(const UProperty* InProperty) const override;
 #endif
 
 	struct FHairGroupResource
diff --git a/Engine/Plugins/Experimental/HairStrands/Source/HairStrandsCore/Public/HairStrandsRendering.h b/Engine/Plugins/Experimental/HairStrands/Source/HairStrandsCore/Public/HairStrandsRendering.h
--- a/Engine/Plugins/Experimental/HairStrands/Source/HairStrandsCore/Public/HairStrandsRendering.h
+++ b/Engine/Plugins/Experimental/HairStrands/Source/HairStrandsCore/Public/HairStrandsRendering.h
@@ -33,6 +33,7 @@ struct FHairStrandsInterpolationInput
 		#endif
 
 		float HairRadius = 0;
+		float HairRaytracingRadiusScale = 0;
 		FVector HairWorldOffset = FVector::ZeroVector;
 
 		inline bool IsValid() const
