Insure Skin groom can only be edited with the correct project settings are enabled.

(ushell-p4-cherrypick of 9441634 by Charles.deRousiers)

diff --git a/Engine/Plugins/Experimental/HairStrands/Source/HairStrandsCore/Private/GroomComponent.cpp b/Engine/Plugins/Experimental/HairStrands/Source/HairStrandsCore/Private/GroomComponent.cpp
--- a/Engine/Plugins/Experimental/HairStrands/Source/HairStrandsCore/Private/GroomComponent.cpp
+++ b/Engine/Plugins/Experimental/HairStrands/Source/HairStrandsCore/Private/GroomComponent.cpp
@@ -1042,7 +1042,6 @@ bool UGroomComponent::CanEditChange(const UProperty* InProperty) const
 		{
 			bool bIsEditable = false;
 			#if RHI_RAYTRACING
-			FRayTracingGeometry* RayTracingGeometry = nullptr;
 			if (IsRayTracingEnabled())
 			{
 				bIsEditable = true;
@@ -1050,6 +1049,11 @@ bool UGroomComponent::CanEditChange(const UProperty* InProperty) const
 			#endif
 			return bIsEditable;
 		}
+
+		if (FCString::Strcmp(*PropertyName, TEXT("bSkinGroom")) == 0)
+		{
+			return GetDefault<URendererSettings>()->bSupportSkinCacheShaders && !UAnimationSettings::Get()->bTickAnimationOnSkeletalMeshInit;
+		}
 	}
 
 	return Super::CanEditChange(InProperty);
