Fix missing unregister event and resource desctruction when a compnent is deleted.

(ushell-p4-cherrypick of 9458505 by Charles.deRousiers)

diff --git a/Engine/Plugins/Experimental/HairStrands/Source/HairStrandsCore/Private/GroomComponent.cpp b/Engine/Plugins/Experimental/HairStrands/Source/HairStrandsCore/Private/GroomComponent.cpp
--- a/Engine/Plugins/Experimental/HairStrands/Source/HairStrandsCore/Private/GroomComponent.cpp
+++ b/Engine/Plugins/Experimental/HairStrands/Source/HairStrandsCore/Private/GroomComponent.cpp
@@ -758,12 +758,11 @@ void UGroomComponent::ReleaseResources()
 {
 	// Unregister component interpolation resources
 	const FPrimitiveComponentId LocalComponentId = ComponentId;
-	const EWorldType::Type WorldType = GetWorld() ? EWorldType::Type(GetWorld()->WorldType) : EWorldType::None;
 	const uint64 Id = LocalComponentId.PrimIDValue;
 	ENQUEUE_RENDER_COMMAND(StaticMeshVertexBuffersLegacyInit)(
-		[Id, WorldType](FRHICommandListImmediate& RHICmdList)
+		[Id](FRHICommandListImmediate& RHICmdList)
 	{
-		UnregisterHairStrands(Id, WorldType);
+		UnregisterHairStrands(Id);
 	});
 
 	for (FHairGroupResource& Res : HairGroupResources)
@@ -840,6 +839,12 @@ void UGroomComponent::OnUnregister()
 	Super::OnUnregister();
 }
 
+void UGroomComponent::OnComponentDestroyed(bool bDestroyingHierarchy)
+{
+	ReleaseResources();
+	Super::OnComponentDestroyed(bDestroyingHierarchy);
+}
+
 void UGroomComponent::TickComponent(float DeltaTime, enum ELevelTick TickType, FActorComponentTickFunction *ThisTickFunction)
 {
 	Super::TickComponent(DeltaTime, TickType, ThisTickFunction);	
diff --git a/Engine/Plugins/Experimental/HairStrands/Source/HairStrandsCore/Public/GroomComponent.h b/Engine/Plugins/Experimental/HairStrands/Source/HairStrandsCore/Public/GroomComponent.h
--- a/Engine/Plugins/Experimental/HairStrands/Source/HairStrandsCore/Public/GroomComponent.h
+++ b/Engine/Plugins/Experimental/HairStrands/Source/HairStrandsCore/Public/GroomComponent.h
@@ -43,6 +43,7 @@ public:
 	//~ Begin UActorComponent Interface.
 	virtual void OnRegister() override;
 	virtual void OnUnregister() override;
+	virtual void OnComponentDestroyed(bool bDestroyingHierarchy) override;
 	virtual void TickComponent(float DeltaTime, enum ELevelTick TickType, FActorComponentTickFunction *ThisTickFunction) override;
 	virtual void SendRenderTransform_Concurrent() override;
 	//~ End UActorComponent Interface.
diff --git a/Engine/Source/Runtime/Renderer/Private/HairStrands/HairStrandsInterface.cpp b/Engine/Source/Runtime/Renderer/Private/HairStrands/HairStrandsInterface.cpp
--- a/Engine/Source/Runtime/Renderer/Private/HairStrands/HairStrandsInterface.cpp
+++ b/Engine/Source/Runtime/Renderer/Private/HairStrands/HairStrandsInterface.cpp
@@ -213,11 +213,11 @@ bool UpdateHairStrands(
 	return false;
 }
 
-void UnregisterHairStrands(uint64 Id, EWorldType::Type WorldType)
+void UnregisterHairStrands(uint64 Id)
 {
 	for (int32 Index=0;Index< GHairManager.Elements.Num();++Index)
 	{
-		if (GHairManager.Elements[Index].Id == Id && GHairManager.Elements[Index].WorldType == WorldType)
+		if (GHairManager.Elements[Index].Id == Id)
 		{
 			GHairManager.Elements[Index] = GHairManager.Elements[GHairManager.Elements.Num()-1];
 			GHairManager.Elements.SetNum(GHairManager.Elements.Num() - 1);
diff --git a/Engine/Source/Runtime/Renderer/Public/HairStrandsInterface.h b/Engine/Source/Runtime/Renderer/Public/HairStrandsInterface.h
--- a/Engine/Source/Runtime/Renderer/Public/HairStrandsInterface.h
+++ b/Engine/Source/Runtime/Renderer/Public/HairStrandsInterface.h
@@ -160,9 +160,7 @@ RENDERER_API void RegisterHairStrands(
 	const FHairStrandsProjectionHairData& SimProjection,
 	const FHairStrandsDebugInfo& DebugInfo);
 
-RENDERER_API void UnregisterHairStrands(
-	uint64 Id, 
-	EWorldType::Type WorldType);
+RENDERER_API void UnregisterHairStrands(uint64 Id);
 
 RENDERER_API bool UpdateHairStrands(
 	uint64 Id,
