- Add Air Drag
- Expose user parameters to control the simulation
- Clean of the hair strands interface
- Spring based material
- Fix transform update when multiple instances

(ushell-p4-cherrypick of 9469306 by [REDACTED_PERSON])

diff --git a/Engine/Plugins/Experimental/HairStrands/Shaders/Private/NiagaraCosseratRodMaterial.ush b/Engine/Plugins/Experimental/HairStrands/Shaders/Private/NiagaraCosseratRodMaterial.ush
--- a/Engine/Plugins/Experimental/HairStrands/Shaders/Private/NiagaraCosseratRodMaterial.ush
+++ b/Engine/Plugins/Experimental/HairStrands/Shaders/Private/NiagaraCosseratRodMaterial.ush
@@ -19,7 +19,9 @@ void SetupStretchRodMaterial(in int StrandsSize, in float YoungModulus, in float
 		OutMaterialCompliance = 4.0/(YoungModulus*PI*RodThickness*RestLength*RodThickness*DeltaTime*DeltaTime);
 
 		const float SumInverseMass = !ProjectConstraint ? ( SharedInverseMass[GGroupThreadId.x] + SharedInverseMass[GGroupThreadId.x-1] ) / (RestLength*RestLength) + 
-					SharedInverseInertia[GGroupThreadId.x-1]*4.0 : SharedInverseMass[GGroupThreadId.x];
+					SharedInverseInertia[GGroupThreadId.x-1]*4.0 : SharedInverseMass[GGroupThreadId.x] / (RestLength*RestLength) ;
+
+		//const float SumInverseMass = !ProjectConstraint ? ( SharedInverseMass[GGroupThreadId.x] + SharedInverseMass[GGroupThreadId.x-1] ) / (RestLength*RestLength) : SharedInverseMass[GGroupThreadId.x] / (RestLength*RestLength) ;
 
 		const float SchurDiagonal = ( (1.0  + MaterialDamping / DeltaTime ) * SumInverseMass + OutMaterialCompliance );
 		OutMaterialWeight = ( SchurDiagonal != 0.0 ) ? 1.0 / SchurDiagonal : 0.0;
@@ -61,6 +63,8 @@ void UpdateStretchRodMultiplier( in float RestLength, in float DeltaTime, in boo
 	{
 		SharedNodePosition[GGroupThreadId.x-1] -= PositionDelta * SharedInverseMass[GGroupThreadId.x-1];
 
+		//const float QuaternionDelta = 0.5 * MultiplyQuat(float4(cross( TargetDirection, DeltaLambda),0), q0 );
+
 		const float4 QuaternionDelta = -2.0 * float4( DeltaLambda.xyz * qebar.w + cross(DeltaLambda.xyz, qebar.xyz),
 						- dot(DeltaLambda.xyz, qebar.xyz) );
 
@@ -127,7 +131,7 @@ void ProjectStretchRodMaterial(in int StrandsSize, in float YoungModulus, in flo
 	{
 		// Compliance = 1.0 / (k * dt * dt)
 		// with k = (Y * I * L) 
-		// A is the polar moment of inertia (Pi * R * R * R * R / 4), L is the rest length and Y is the young modulus
+		// I is the polar moment of inertia (Pi * R * R * R * R / 4), L is the rest length and Y is the young modulus
 		OutMaterialCompliance = 64.0/(YoungModulus*PI*RestLength*RodThickness*RodThickness*RodThickness*RodThickness*DeltaTime*DeltaTime);
 
 		const float SumInverseMass = !ProjectConstraint ? ( SharedInverseInertia[GGroupThreadId.x] + SharedInverseInertia[GGroupThreadId.x-1] ) * 4.0 / (RestLength*RestLength) 
diff --git a/Engine/Plugins/Experimental/HairStrands/Shaders/Private/NiagaraDataInterfaceHairStrands.ush b/Engine/Plugins/Experimental/HairStrands/Shaders/Private/NiagaraDataInterfaceHairStrands.ush
--- a/Engine/Plugins/Experimental/HairStrands/Shaders/Private/NiagaraDataInterfaceHairStrands.ush
+++ b/Engine/Plugins/Experimental/HairStrands/Shaders/Private/NiagaraDataInterfaceHairStrands.ush
@@ -24,8 +24,6 @@ struct FDIHairStrandsContext
 	RWBuffer<uint4>		PointsPositionsBuffer;
 	Buffer<uint>		CurvesOffsetsBuffer;
 	Buffer<uint4>		RestPositionsBuffer;
-	Texture3D<int>		GridCurrentBuffer;
-	RWTexture3D<int>	GridDestinationBuffer;
 	float4x4			WorldTransform;
 	float4x4			WorldInverse;
 	float4				WorldRotation;
@@ -34,16 +32,13 @@ struct FDIHairStrandsContext
 	float				StrandDensity;
 	float				RootThickness;
 	float				TipThickness;
-	float4				GridOrigin;
-	int4				GridSize;
+	float4				BoxOrigin;
 }; 
 
 #define DIHAIRSTRANDS_DECLARE_CONSTANTS(NAME)\
 RWBuffer<uint4> PointsPositionsBuffer_##NAME;\
 Buffer<uint> CurvesOffsetsBuffer_##NAME;\
 Buffer<uint4> RestPositionsBuffer_##NAME;\
-Texture3D<int> GridCurrentBuffer_##NAME;\
-RWTexture3D<int> GridDestinationBuffer_##NAME;\
 float4x4 WorldTransform_##NAME;\
 float4x4 WorldInverse_##NAME;\
 float4 WorldRotation_##NAME;\
@@ -52,8 +47,7 @@ int StrandSize_##NAME;\
 float StrandDensity_##NAME;\
 float RootThickness_##NAME;\
 float TipThickness_##NAME;\
-float4 GridOrigin_##NAME;\
-int4 GridSize_##NAME;\
+float4 BoxOrigin_##NAME;\
 
 #define DIHAIRSTRANDS_MAKE_CONTEXT(NAME)\
 FDIHairStrandsContext DIContext;\
@@ -68,10 +62,7 @@ DIContext.StrandSize = StrandSize_##NAME;\
 DIContext.StrandDensity = StrandDensity_##NAME;\
 DIContext.RootThickness = RootThickness_##NAME;\
 DIContext.TipThickness = TipThickness_##NAME;\
-DIContext.GridOrigin = GridOrigin_##NAME;\
-DIContext.GridSize = GridSize_##NAME;\
-DIContext.GridCurrentBuffer = GridCurrentBuffer_##NAME;\
-DIContext.GridDestinationBuffer = GridDestinationBuffer_##NAME;\
+DIContext.BoxOrigin = BoxOrigin_##NAME;\
 
 /* -----------------------------------------------------------------
  * Box utilities
@@ -87,394 +78,7 @@ void DIHairStrands_GetBoxCenter(in FDIHairStrandsContext DIContext, out float3 O
 // Return the bounding box extent
 void DIHairStrands_GetBoxExtent(in FDIHairStrandsContext DIContext, out float3 OutBoxExtent)
 {
-	OutBoxExtent = -DIContext.GridOrigin.xyz;
-}
-
-/* -----------------------------------------------------------------
- * Grid utilities
- * -----------------------------------------------------------------
- */
-
-// Return true if the grid index is within the domain size
-bool DIHairStrands_IsIndexValid(in FDIHairStrandsContext DIContext, in int3 GridIndex)
-{
-	return ((GridIndex.x < (DIContext.GridSize.x-1)) && (GridIndex.y < (DIContext.GridSize.y-1)) && 
-			(GridIndex.z < (DIContext.GridSize.z-1)) && (GridIndex.x >= 0) && (GridIndex.y >= 0) && (GridIndex.z >= 0));
-}
-
-// Transform the node world position in grid space
-float3 DIHairStrands_GetGridPosition(in FDIHairStrandsContext DIContext, in float3 NodePosition)
-{
-	return ( mul(float4( NodePosition, 1.0), DIContext.WorldInverse).xyz - DIContext.GridOrigin.xyz );
-}
-
-// Return the grid hash given a grid index
-int DIHairStrands_GetGridHash(in FDIHairStrandsContext DIContext, in int3 GridIndex)
-{
-	return GridIndex.z * DIContext.GridSize.y * DIContext.GridSize.x + GridIndex.y * DIContext.GridSize.x + GridIndex.x;
-}
-
-// Compute the grid index and fraction
-void DIHairStrands_ComputeGridIndex(in FDIHairStrandsContext DIContext, in float3 GridPosition, in float3 GridOffset, out int3 OutGridIndex, out float3 OutGridFraction)
-{
-	const float3 GridCell = GridPosition / DIContext.GridOrigin.w - GridOffset;
-
-	OutGridIndex = floor(GridCell);
-	OutGridFraction = GridCell - OutGridIndex;
-}
-
-int3 DIHairStrands_GridNodeIndex( in FDIHairStrandsContext DIContext, in int3 GridIndex, in int DataOffset)
-{
-	return int3( GridIndex.x + DataOffset * DIContext.GridSize.x, GridIndex.yz);
-}
-
-void InterlockedAddFloat(in FDIHairStrandsContext DIContext, in int3 GridIndex, in int DataOffset, in float FloatValue)
-{
-	int IntValue = asint(FloatValue);
-	int CompareValue = 0;
-	int StoredValue = 0;
-
-	[allow_uav_condition] while(true)
-	{
-		InterlockedCompareExchange(DIContext.GridDestinationBuffer[
-				DIHairStrands_GridNodeIndex(DIContext,GridIndex,DataOffset)],CompareValue,IntValue,StoredValue);
-		if(StoredValue == CompareValue)
-		{
-			break;
-		}
-		CompareValue = StoredValue;
-		IntValue = asint(FloatValue+asfloat(StoredValue));
-	}
-}
-
-void DIHairStrands_AddNodeScalar( in FDIHairStrandsContext DIContext, in int3 GridIndex, in int DataOffset, in float NodeScalar, in bool EnableLocking)
-{
-	const int ConvertScalar = NodeScalar * SCALE;
-	if(EnableLocking)
-	{
-		InterlockedAdd(DIContext.GridDestinationBuffer[DIHairStrands_GridNodeIndex(DIContext,GridIndex,DataOffset)], ConvertScalar);
-		//InterlockedAddFloat(DIContext,GridIndex,DataOffset,NodeScalar);
-	}
-	else
-	{
-		DIContext.GridDestinationBuffer[DIHairStrands_GridNodeIndex(DIContext,GridIndex,DataOffset)] += ConvertScalar;
-		//const float FloatValue = asfloat(DIContext.GridDestinationBuffer[DIHairStrands_GridNodeIndex(DIContext,GridIndex,DataOffset)]) + NodeScalar;
-		//DIContext.GridDestinationBuffer[DIHairStrands_GridNodeIndex(DIContext,GridIndex,DataOffset)] = asint(FloatValue);
-	}
-	DeviceMemoryBarrier();
-}
-
-void DIHairStrands_AddNodeVector( in FDIHairStrandsContext DIContext, in int3 GridIndex, in int DataOffset, in float3 NodeVector, in bool EnableLocking)
-{
-	DIHairStrands_AddNodeScalar(DIContext,GridIndex,DataOffset,NodeVector.x,EnableLocking);
-	DIHairStrands_AddNodeScalar(DIContext,GridIndex,DataOffset+1,NodeVector.y,EnableLocking);
-	DIHairStrands_AddNodeScalar(DIContext,GridIndex,DataOffset+2,NodeVector.z,EnableLocking);
-}
-
-float DIHairStrands_GetGridScalar( in FDIHairStrandsContext DIContext, in int3 GridIndex, in int DataOffset)
-{
-	//return asfloat(DIContext.GridCurrentBuffer[DIHairStrands_GridNodeIndex(DIContext,GridIndex,DataOffset)]);
-	return DIContext.GridCurrentBuffer[DIHairStrands_GridNodeIndex(DIContext,GridIndex,DataOffset)] / SCALE;
-}
-
-float3 DIHairStrands_GetGridVector( in FDIHairStrandsContext DIContext, in int3 GridIndex, in int DataOffset)
-{
-	return float3(DIHairStrands_GetGridScalar(DIContext,GridIndex,DataOffset),
-				  DIHairStrands_GetGridScalar(DIContext,GridIndex,DataOffset+1),
-				  DIHairStrands_GetGridScalar(DIContext,GridIndex,DataOffset+2));
-}
-
-// Transfer the velocity from the node to the grid
-void DIHairStrands_BuildCollisionGrid(in FDIHairStrandsContext DIContext, in float3 NodePosition, in float3 NodeVelocity, in float NodeMass, 
-			in float3 NodeGradientX, in float3 NodeGradientY, in float3 NodeGradientZ, out bool OutBuildStatus)
-{
-	SharedNodePosition[GGroupThreadId.x] = NodePosition;
-	SharedPreviousPosition[GGroupThreadId.x] = NodeVelocity;
-	SharedInverseMass[GGroupThreadId.x] = NodeMass;
-	GroupMemoryBarrier(); 
-
-	OutBuildStatus = false;
-
-	const int LocalIndex = GGroupThreadId.x % DIContext.StrandSize;
-	if(LocalIndex != 0)
-	{
-		const float3 DeltaVelocity = SharedPreviousPosition[GGroupThreadId.x] - SharedPreviousPosition[GGroupThreadId.x-1];
-		const float3 DeltaPosition = SharedNodePosition[GGroupThreadId.x] - SharedNodePosition[GGroupThreadId.x-1];
-
-		const float2 SegmentWeight = (LocalIndex == 1) ? float2(1.0,0.5) : (LocalIndex == (DIContext.StrandSize-1) ) ? float2(0.5,1.0) : float2(0.5,0.5);
-		const float SegmentMass =  (SharedInverseMass[GGroupThreadId.x-1] * SegmentWeight.x + SharedInverseMass[GGroupThreadId.x] * SegmentWeight.y);
-		const float SegmentLength = length(DeltaPosition);
-
-		const float3 GridPositionA = DIHairStrands_GetGridPosition(DIContext,SharedNodePosition[GGroupThreadId.x-1]);
-		const float3 GridPositionB = DIHairStrands_GetGridPosition(DIContext,SharedNodePosition[GGroupThreadId.x]);
-		const float3 GridSegment = GridPositionB - GridPositionA;
-
-		const int NumSamples = ceil(SegmentLength / DIContext.GridOrigin.w);
-
-		const float SampleMass = SegmentMass / NumSamples;
-		const float3 SampleTangent = DeltaPosition / SegmentLength;
-
-		for( int SampleIndex = 0; SampleIndex < NumSamples; ++SampleIndex)
-		{
-			const float SampleCoord = (0.5+SampleIndex) / NumSamples;
-			const float3 SamplePosition = GridPositionA + SampleCoord * GridSegment;
-			//const float3 SampleMomentum = NodeVelocity * SampleMass;
-			const float3 SampleMomentum = (SharedPreviousPosition[GGroupThreadId.x-1] + SampleCoord * DeltaVelocity) * SampleMass;
-
-
-			for( int ComponentIndex = 0; ComponentIndex < 3; ++ComponentIndex)
-			{
-				const float3 NodeGradient = (ComponentIndex == 0) ? NodeGradientX : (ComponentIndex == 1) ? NodeGradientY : NodeGradientZ;
-				float3 GridOffset = float3(0.5,0.5,0.5); GridOffset[ComponentIndex] = 0.0;
-				int3 GridIndex = int3(0,0,0);
-				float3 GridFraction = float3(0,0,0);
-				DIHairStrands_ComputeGridIndex(DIContext,SamplePosition,GridOffset,GridIndex,GridFraction);
-
-				if(DIHairStrands_IsIndexValid(DIContext,GridIndex))
-				{
-					OutBuildStatus = true;
-					for(int i = 0; i < 2; ++i)
-					{
-						const float GridWeightX = (i == 0) ? 1.0 - GridFraction.x : GridFraction.x;
-						for(int j = 0; j < 2; ++j)
-						{
-							const float GridWeightY = (j == 0) ? 1.0 - GridFraction.y : GridFraction.y;
-							for(int k = 0; k < 2; ++k)
-							{
-								const float GridWeightZ = (k == 0) ? 1.0 - GridFraction.z : GridFraction.z;
-
-								const float GridWeightXYZ = GridWeightX*GridWeightY*GridWeightZ;
-
-								const int3 GridLocation = int3(GridIndex.x+i,GridIndex.y+j,GridIndex.z+k);
-								const float3 DeltaPosition = ((GridLocation+GridOffset)* DIContext.GridOrigin.w) - SamplePosition;
-								const float ExtendedMomentum = SampleMomentum[ComponentIndex] + dot(NodeGradient, DeltaPosition) * SampleMass;
-
-								DIHairStrands_AddNodeScalar(DIContext,GridLocation,MASS_OFFSET+ComponentIndex,SampleMass * GridWeightXYZ,true);
-								DIHairStrands_AddNodeScalar(DIContext,GridLocation,VELOCITY_OFFSET+ComponentIndex,ExtendedMomentum * GridWeightXYZ,true);
-								//DIHairStrands_AddNodeScalar(DIContext,GridLocation,TANGENT_OFFSET+ComponentIndex,SampleTangent[ComponentIndex] * GridWeightXYZ,true);
-							}
-						}
-					}
-				}
-			}
-		}
-	}
-}
-
-
-// Return the velocity from the grid to the node
-void DIHairStrands_QueryCollisionGrid(in FDIHairStrandsContext DIContext, in float3 NodePosition, out float3 OutGridVelocity,
-			out float3 OutGridTangent, out float OutGridMass, out float3 OutGridGradientX, out float3 OutGridGradientY, out float3 OutGridGradientZ)
-{
-	const float3 GridPosition = DIHairStrands_GetGridPosition(DIContext,NodePosition);
-	
-	OutGridVelocity = float3(0,0,0);
-	OutGridMass = 0.0;
-	OutGridTangent = float3(0,0,0);
-	OutGridGradientX = float3(0,0,0);
-	OutGridGradientY = float3(0,0,0);
-	OutGridGradientZ = float3(0,0,0);
-
-	for( int ComponentIndex = 0; ComponentIndex < 3; ++ComponentIndex)
-	{
-		float3 GridOffset = float3(0.5,0.5,0.5); GridOffset[ComponentIndex] = 0.0;
-		int3 GridIndex = int3(0,0,0);
-		float3 GridFraction = float3(0,0,0);
-		DIHairStrands_ComputeGridIndex(DIContext,GridPosition,GridOffset,GridIndex,GridFraction);
-		if(DIHairStrands_IsIndexValid(DIContext,GridIndex))
-		{
-			for(int i = 0; i < 2; ++i)
-			{
-				const float GridWeightX = (i == 0) ? 1.0 - GridFraction.x : GridFraction.x;
-				const float GridGradientX = (i == 0) ? -1.0: 1.0;
-				for(int j = 0; j < 2; ++j)
-				{
-					const float GridWeightY = (j == 0) ? 1.0 - GridFraction.y : GridFraction.y;
-					const float GridGradientY = (j == 0) ? -1.0: 1.0;
-					for(int k = 0; k < 2; ++k)
-					{
-						const float GridWeightZ = (k == 0) ? 1.0 - GridFraction.z : GridFraction.z;
-						const float GridGradientZ = (k == 0) ? -1.0: 1.0;
-
-						const float GridWeightXYZ = GridWeightX*GridWeightY*GridWeightZ;
-						const float3 GridGradientXYZ = float3(GridGradientX*GridWeightY*GridWeightZ,
-															  GridWeightX*GridGradientY*GridWeightZ,
-															  GridWeightX*GridWeightY*GridGradientZ) / DIContext.GridOrigin.w;
-
-						const int3 GridLocation = int3(GridIndex.x+i,GridIndex.y+j,GridIndex.z+k);
-					
-						const float LocalMass = DIHairStrands_GetGridScalar(DIContext,GridLocation,MASS_OFFSET+ComponentIndex);
-						if( LocalMass != 0.0 )
-						{
-							const float LocalVelocity = DIHairStrands_GetGridScalar(DIContext,GridLocation,VELOCITY_OFFSET+ComponentIndex) / LocalMass;
-							OutGridMass += LocalMass * GridWeightXYZ;
-							OutGridVelocity[ComponentIndex] += LocalVelocity * GridWeightXYZ;
-							OutGridTangent[ComponentIndex] += (DIHairStrands_GetGridScalar(DIContext,GridLocation,TANGENT_OFFSET+ComponentIndex) / LocalMass) * GridWeightXYZ;
-							//OutGridTangent[ComponentIndex] += DIHairStrands_GetGridScalar(DIContext,GridLocation,TANGENT_OFFSET+ComponentIndex) * GridWeightXYZ;
-							if( ComponentIndex == 0)
-							{
-								OutGridGradientX +=  GridGradientXYZ * LocalVelocity * GridWeightXYZ;
-							}
-							else if( ComponentIndex == 1)
-							{
-								OutGridGradientY +=  GridGradientXYZ * LocalVelocity * GridWeightXYZ;
-							}
-							else
-							{
-								OutGridGradientZ +=  GridGradientXYZ * LocalVelocity * GridWeightXYZ;
-							}
-						}
-					}
-				}
-			}
-		}
-	}
-	OutGridMass /= (3*DIContext.GridOrigin.w * DIContext.GridOrigin.w * DIContext.GridOrigin.w);
-	//if(OutGridMass!=0.0)
-	//{
-	//	OutGridTangent = normalize(OutGridTangent);
-	//}
-}
-
-// Return the grid index given a grid hash
-int3 DIHairStrands_ComputeGridIndex(in FDIHairStrandsContext DIContext, in int GridHash)
-{
-	const int GridIndexX = GridHash / ( DIContext.GridSize.y * DIContext.GridSize.z);
-	const int GridDummy = GridHash - GridIndexX * DIContext.GridSize.y * DIContext.GridSize.z;
-	const int GridIndexY = GridDummy / (DIContext.GridSize.z);
-	const int GridIndexZ = GridDummy % (DIContext.GridSize.z);
-
-	return int3(GridIndexX,GridIndexY,GridIndexZ);
-}
-
-// Compute the grid pressure
-float DIHairStrands_ComputeGridPressure(in FDIHairStrandsContext DIContext, in int3 PrevIndex)
-{	
-	float GridPressure = 0.0;
-
-	if(DIHairStrands_IsIndexValid(DIContext,PrevIndex))
-	{
-		const float VoxelVolume = DIContext.GridOrigin.w * DIContext.GridOrigin.w * DIContext.GridOrigin.w;
-		float GridDivergence = 0.0;
-		float GridFraction = 0.0;
-
-		[unroll]
-		for(int ComponentIndex = 0; ComponentIndex < 3; ++ComponentIndex)
-		{
-			int3 NextIndex = PrevIndex; NextIndex[ComponentIndex] += 1;
-
-			const float PrevMass = DIHairStrands_GetGridScalar(DIContext,PrevIndex,MASS_OFFSET+ComponentIndex);
-			const float PrevDensity = min(1,max(0,PrevMass / VoxelVolume));
-			
-			if( PrevMass > 0.)
-			{
-				GridDivergence -= DIHairStrands_GetGridScalar(DIContext,PrevIndex,VELOCITY_OFFSET+ComponentIndex) / PrevMass;
-				GridFraction += 1.0 / PrevDensity;
-			}
-
-			const float NextMass = DIHairStrands_GetGridScalar(DIContext,NextIndex,MASS_OFFSET+ComponentIndex);
-			const float NextDensity = min(1,max(0,NextMass / VoxelVolume));
-
-			if( NextMass > 0.)
-			{
-				GridDivergence += DIHairStrands_GetGridScalar(DIContext,NextIndex,VELOCITY_OFFSET+ComponentIndex) / NextMass;
-				GridFraction += 1.0 / NextDensity;
-			}
-		}
-		if( GridFraction > 0.)
-		{
-			GridDivergence *= DIContext.GridOrigin.w;
-			GridPressure = -GridDivergence / GridFraction;
-		}
-	}
-	return GridPressure;
-}
-
-// Return the velocity from the grid to the node
-void DIHairStrands_ProjectCollisionGrid(in FDIHairStrandsContext DIContext, in int GridHash, out bool OutProjectStatus)
-{	
-	OutProjectStatus = false;
-
-	const int3 NextIndex = DIHairStrands_ComputeGridIndex(DIContext,GridHash);
-	const float NextPressure = DIHairStrands_ComputeGridPressure(DIContext,NextIndex);
-	const float VoxelVolume = DIContext.GridOrigin.w * DIContext.GridOrigin.w * DIContext.GridOrigin.w;
-
-	[unroll]
-	for(int ComponentIndex = 0; ComponentIndex < 3; ++ComponentIndex)
-	{
-		int3 PrevIndex = NextIndex; PrevIndex[ComponentIndex] -= 1;
-		const float PrevPressure = DIHairStrands_ComputeGridPressure(DIContext,PrevIndex);
-
-		const float NextMass = DIHairStrands_GetGridScalar(DIContext,NextIndex,MASS_OFFSET+ComponentIndex);
-		const float NextDensity = min(1,max(0,NextMass / VoxelVolume));
-
-		const float PressureGradient = (NextPressure - PrevPressure) / ( DIContext.GridOrigin.w * NextDensity );
-
-		DIHairStrands_AddNodeScalar(DIContext,NextIndex,VELOCITY_OFFSET+ComponentIndex,-PressureGradient*NextMass,false);
-		DIHairStrands_AddNodeScalar(DIContext,NextIndex,TANGENT_OFFSET+ComponentIndex,-PressureGradient*NextMass,false);
-	}
-	return;
-
-
-	const int GridIndexX = GridHash / ( DIContext.GridSize.y * DIContext.GridSize.z);
-	const int GridDummy = GridHash - GridIndexX * DIContext.GridSize.y * DIContext.GridSize.z;
-	const int GridIndexY = GridDummy / (DIContext.GridSize.z);
-	const int GridIndexZ = GridDummy % (DIContext.GridSize.z);
-
-	const bool OddStage = (ShaderStageIndex % 2)  == 0;
-	const bool OddVoxel = ((GridIndexX+GridIndexY+GridIndexZ) % 2) == 0;
-
-	//if( (OddStage && OddVoxel) || (!OddStage && !OddVoxel) )
-	{
-		const float VoxelVolume = DIContext.GridOrigin.w * DIContext.GridOrigin.w * DIContext.GridOrigin.w;
-
-		const int3 PrevIndex = int3(GridIndexX,GridIndexY,GridIndexZ);
-		if(DIHairStrands_IsIndexValid(DIContext,PrevIndex))
-		{
-			OutProjectStatus = true;
-			float GridDivergence = 0.0;
-			float GridFraction = 0.0;
-			for( int ComponentIndex = 0; ComponentIndex < 3; ++ComponentIndex)
-			{
-				int3 NextIndex = PrevIndex; NextIndex[ComponentIndex] += 1;
-
-				const float PrevMass = DIHairStrands_GetGridScalar(DIContext,PrevIndex,MASS_OFFSET+ComponentIndex);
-				const float PrevDensity = min(1,max(0,PrevMass / VoxelVolume));
-				if( PrevMass > 0.)
-				{
-					GridDivergence -= DIHairStrands_GetGridScalar(DIContext,PrevIndex,VELOCITY_OFFSET+ComponentIndex) * PrevDensity / PrevMass;
-					GridFraction += PrevDensity;
-				}
-
-				const float NextMass = DIHairStrands_GetGridScalar(DIContext,NextIndex,MASS_OFFSET+ComponentIndex);
-				const float NextDensity = min(1,max(0,NextMass / VoxelVolume));
-				if( NextMass > 0.)
-				{
-					GridDivergence += DIHairStrands_GetGridScalar(DIContext,NextIndex,VELOCITY_OFFSET+ComponentIndex) * NextDensity / NextMass;
-					GridFraction += NextDensity;
-				}
-			}
-			if( GridFraction > 0.)
-			{
-				GridDivergence *= DIContext.GridOrigin.w;
-				const float GridPressure = -GridDivergence / GridFraction;
-
-				const float PressureGradient = GridPressure / DIContext.GridOrigin.w;
-				for( int ComponentIndexB = 0; ComponentIndexB < 3; ++ComponentIndexB)
-				{
-					int3 NextIndex = PrevIndex; NextIndex[ComponentIndexB] += 1;
-
-					const float PrevMass = DIHairStrands_GetGridScalar(DIContext,PrevIndex,MASS_OFFSET+ComponentIndexB);
-					DIHairStrands_AddNodeScalar(DIContext,PrevIndex,VELOCITY_OFFSET+ComponentIndexB,-PressureGradient*PrevMass,true);
-					DIHairStrands_AddNodeScalar(DIContext,PrevIndex,TANGENT_OFFSET+ComponentIndexB,-PressureGradient*PrevMass,true);
-
-					const float NextMass = DIHairStrands_GetGridScalar(DIContext,NextIndex,MASS_OFFSET+ComponentIndexB);
-					DIHairStrands_AddNodeScalar(DIContext,NextIndex,VELOCITY_OFFSET+ComponentIndexB,PressureGradient*NextMass,true);
-					DIHairStrands_AddNodeScalar(DIContext,NextIndex,TANGENT_OFFSET+ComponentIndexB,PressureGradient*NextMass,true);
-				}
-			}
-		}	
-	}
+	OutBoxExtent = -DIContext.BoxOrigin.xyz;
 }
 
 /* -----------------------------------------------------------------
@@ -560,8 +164,8 @@ void DIHairStrands_ComputeNodeMass(in FDIHairStrandsContext DIContext, in int No
 	const float NextThick = RootThickness * ( 1.0 - NextCoord ) + TipThickness * NextCoord;
 	const float MeanRadius = 0.25 * ( PrevThick + NextThick );
 
-	//OutNodeMass = StrandsDensity * EdgeLength * EdgeScale * M_PI * MeanRadius * MeanRadius;
-	OutNodeMass = StrandsDensity * EdgeLength * EdgeScale * M_PI * (PrevThick * PrevThick + NextThick * PrevThick + NextThick * NextThick) / 12.0;
+	OutNodeMass = StrandsDensity * EdgeLength * EdgeScale * M_PI * MeanRadius * MeanRadius;
+	//OutNodeMass = StrandsDensity * EdgeLength * EdgeScale * M_PI * (PrevThick * PrevThick + NextThick * PrevThick + NextThick * NextThick) / 12.0;
 }
 
 // Compute the node inertia
@@ -583,14 +187,13 @@ void DIHairStrands_ComputeNodeInertia(in FDIHairStrandsContext DIContext, in int
 	const float NextThick = RootThickness * ( 1.0 - NextCoord ) + TipThickness * NextCoord;
 	const float MeanRadius = 0.25 * ( PrevThick + NextThick );
 
-	const float EdgeMass = StrandsDensity * EdgeLength * M_PI * (PrevThick * PrevThick + NextThick * PrevThick + NextThick * NextThick) / 12.0;
-	const float RadialInertia = EdgeMass * ( 3.0 * MeanRadius * MeanRadius + EdgeLength * EdgeLength) / 12.0;
-	//const float RadialInertia = EdgeMass * ( MeanRadius * MeanRadius) / 4.0;
-	OutNodeInertia = float3( RadialInertia, RadialInertia, EdgeMass * MeanRadius * MeanRadius * 0.5 );
+	//const float EdgeMass = StrandsDensity * EdgeLength * M_PI * (PrevThick * PrevThick + NextThick * PrevThick + NextThick * NextThick) / 12.0;
+	//const float RadialInertia = EdgeMass * ( 3.0 * MeanRadius * MeanRadius + EdgeLength * EdgeLength) / 12.0;
+	//OutNodeInertia = float3( RadialInertia, RadialInertia, EdgeMass * MeanRadius * MeanRadius * 0.5 );
 
-	//const float EdgeMass = StrandsDensity * EdgeLength * M_PI * MeanRadius * MeanRadius;
-	//const float RadialInertia = EdgeMass * ( MeanRadius * MeanRadius) / 4.0;
-	//OutNodeInertia = float3( RadialInertia, RadialInertia, 2.0 * RadialInertia );
+	const float EdgeMass = StrandsDensity * EdgeLength * M_PI * MeanRadius * MeanRadius;
+	const float RadialInertia = EdgeMass * ( MeanRadius * MeanRadius) / 4.0;
+	OutNodeInertia = float3( RadialInertia, RadialInertia, 2.0 * RadialInertia );
 }
 
 // Compute the node position by linear interpolation over the points
@@ -673,7 +276,7 @@ void DIHairStrands_ComputeNodeOrientation(in FDIHairStrandsContext DIContext, in
  */
 
 // Compute the edge length value
-void DIHairStrands_ComputeEdgeLength(in FDIHairStrandsContext DIContext, in int NodeIndex, in float3 NodePosition, out float OutEdgeLength)
+void DIHairStrands_ComputeEdgeLength(in FDIHairStrandsContext DIContext, in int NodeIndex, in float3 NodePosition, in int NodeOffset, out float OutEdgeLength)
 {
 	SharedNodePosition[GGroupThreadId.x] = NodePosition;
 	GroupMemoryBarrier(); 
@@ -682,7 +285,7 @@ void DIHairStrands_ComputeEdgeLength(in FDIHairStrandsContext DIContext, in int
 	DIHairStrands_ComputeNodeIndices(DIContext,NodeIndex,StrandIndex,LocalIndex);
 
 	// L = ||P1-P0||
-	OutEdgeLength = (LocalIndex>0) ? length(SharedNodePosition[GGroupThreadId.x]- SharedNodePosition[GGroupThreadId.x-1]) : 0.0;
+	OutEdgeLength = (LocalIndex>NodeOffset) ? length(SharedNodePosition[GGroupThreadId.x]- SharedNodePosition[GGroupThreadId.x-1-NodeOffset]) : 0.0;
 }
 
 // Compute the edge darboux vector (diff between consecutive edge orientations)
@@ -831,6 +434,50 @@ void DIHairStrands_UpdateAngularVelocity(in FDIHairStrandsContext DIContext, in
 	}
 }
 
+/* -----------------------------------------------------------------
+ * Compute Air Drag Force
+ * -----------------------------------------------------------------
+ */
+
+// Air Density = 1 kg/m3 = 1e-3 g/cm3
+// Dynamic Viscosity = 1e-5 Pa.s = 1e-6 g/(cm.s)
+// Compute air drag force
+void DIHairStrands_ComputeAirDragForce(in FDIHairStrandsContext DIContext, in float AirDensity, in float AirViscosity, in float AirDrag, 
+	in float3 AirVelocity, in float NodeThickness, in float3 NodePosition, in float3 NodeVelocity, out float3 OutAirDrag)
+{	
+	const float3 DeltaVelocity = NodeVelocity - AirVelocity;
+	const float DeltaVNorm = length(DeltaVelocity);
+
+	SharedNodePosition[GGroupThreadId.x] = NodePosition;
+	GroupMemoryBarrier();
+	
+	float AlignedTangent = 0.0;
+	float TangentLength = 0.0;
+	if( DeltaVNorm > SMALL_NUMBER)
+	{
+		const int LocalIndex = GGroupThreadId.x % DIContext.StrandSize;
+		float3 NodeTangent = (LocalIndex == 0) ? SharedNodePosition[GGroupThreadId.x+1] - SharedNodePosition[GGroupThreadId.x] : 
+												 SharedNodePosition[GGroupThreadId.x] - SharedNodePosition[GGroupThreadId.x-1]; 
+		TangentLength = length(NodeTangent);
+		NodeTangent = normalize(NodeTangent);
+
+		AlignedTangent = abs(dot(NodeTangent,DeltaVelocity)/DeltaVNorm);
+	}
+	const float NodeDiameter = NodeThickness * (1.0 - AlignedTangent) + TangentLength * AlignedTangent;
+	const float NodeArea = NodeDiameter * NodeThickness;
+
+	// V = 4 * PI * R * R * R / 3 = PI * D * D * D / 6 = M / rho 
+	//const float NodeDiameter = pow(6.0 * (NodeMass / StrandsDensity ) / PI, 1.0 / 3.0);
+	//const float NodeArea = PI * NodeDiameter * NodeDiameter / 4.0;
+
+	const float LinearDrag = 3.0 * PI * NodeDiameter * AirViscosity * AirDrag;
+	const float QuadraticDrag = 0.5 * AirDensity * NodeArea * AirDrag;
+
+	const float DragCoeff = LinearDrag + QuadraticDrag * DeltaVNorm;
+
+	OutAirDrag = - DragCoeff * DeltaVelocity;
+}
+
 /* -----------------------------------------------------------------
  * Update node position orientation
  * -----------------------------------------------------------------
diff --git a/Engine/Plugins/Experimental/HairStrands/Shaders/Private/NiagaraDataInterfacePressureGrid.ush b/Engine/Plugins/Experimental/HairStrands/Shaders/Private/NiagaraDataInterfacePressureGrid.ush
--- a/Engine/Plugins/Experimental/HairStrands/Shaders/Private/NiagaraDataInterfacePressureGrid.ush
+++ b/Engine/Plugins/Experimental/HairStrands/Shaders/Private/NiagaraDataInterfacePressureGrid.ush
@@ -11,7 +11,7 @@ NiagaraDataInterfacePressureGrid.ush
  * -----------------------------------------------------------------
  */
 
-static const float FLOAT_PRECISION = 1e+6;
+static const float FLOAT_PRECISION = 1e+8;
 static const bool FAST_MODE = true;
 
 static const int FLUID_MASS_OFFSET = 0;
diff --git a/Engine/Plugins/Experimental/HairStrands/Shaders/Private/NiagaraHookeSpringMaterial.ush b/Engine/Plugins/Experimental/HairStrands/Shaders/Private/NiagaraHookeSpringMaterial.ush
--- a/Engine/Plugins/Experimental/HairStrands/Shaders/Private/NiagaraHookeSpringMaterial.ush
+++ b/Engine/Plugins/Experimental/HairStrands/Shaders/Private/NiagaraHookeSpringMaterial.ush
@@ -7,47 +7,57 @@
  * -----------------------------------------------------------------
  */
 
- void SetupStretchSpringMaterial(in int StrandsSize, in float YoungModulus, in float RodThickness, 
-	in float RestLength, in float DeltaTime, in bool ProjectConstraint, in float MaterialDamping, out float OutMaterialCompliance, out float OutMaterialWeight, out float OutMaterialMultiplier)
+ void ResetStretchSpringMaterial(in int StrandsSize, in float YoungModulus, in float RodThickness, 
+	in float RestLength, in float DeltaTime, in bool ProjectConstraint, in float MaterialDamping, out float OutMaterialCompliance, out float OutMaterialWeight, out float OutMaterialMultiplier, in int NodeOffset)
 {
-	const int LocalIndex = (GGroupThreadId.x % StrandsSize);
-	if( LocalIndex > 1 )
-	{
-		// Compliance = 1.0 / (k * dt * dt)
-		// with k = Y * A / L 
-		// A is the cross section area (Pi * R * R), L is the rest length and Y is the young modulus
-		OutMaterialCompliance = RestLength*4.0/(YoungModulus*PI*RodThickness*RodThickness*DeltaTime*DeltaTime);
+	// Compliance = 1.0 / (k * dt * dt)
+	// with k = L * L * (Y * A / L) (the L*L term before is coming from the fact that our constraint is dL/L and not dL)
+	// A is the cross section area (Pi * R * R), L is the rest length and Y is the young modulus
+	OutMaterialCompliance = 4.0/(YoungModulus*PI*RodThickness*RestLength*RodThickness*DeltaTime*DeltaTime);
 
-		const float SumInverseMass = !ProjectConstraint ? ( SharedInverseMass[GGroupThreadId.x] + SharedInverseMass[GGroupThreadId.x-1] ) : SharedInverseMass[GGroupThreadId.x];
-		const float SchurDiagonal = ( (1.0  + MaterialDamping / DeltaTime ) * SumInverseMass + OutMaterialCompliance );
-		OutMaterialWeight = ( SchurDiagonal != 0.0 ) ? 1.0 / SchurDiagonal : 0.0;
+	const float SumInverseMass = !ProjectConstraint ? ( SharedInverseMass[GGroupThreadId.x] + SharedInverseMass[GGroupThreadId.x-1-NodeOffset] ) / (RestLength*RestLength) :
+				SharedInverseMass[GGroupThreadId.x] / (RestLength*RestLength);
 
-		OutMaterialMultiplier = 0;
-	}
+	const float SchurDiagonal = ( (1.0  + MaterialDamping / DeltaTime ) * SumInverseMass + OutMaterialCompliance );
+	OutMaterialWeight = ( SchurDiagonal != 0.0 ) ? 1.0 / SchurDiagonal : 0.0;
+
+	OutMaterialMultiplier = 0;
 }
 
-void UpdateStretchSpringMultiplier(in float RestLength, in float DeltaTime, in bool ProjectConstraint, in float MaterialDamping, in float MaterialCompliance, in float MaterialWeight, inout float OutMaterialMultiplier)
+
+void UpdateStretchSpringMultiplier(in float RestLength, in float DeltaTime, in bool ProjectConstraint, in float MaterialDamping, in float MaterialCompliance, in float MaterialWeight, inout float OutMaterialMultiplier, in int NodeOffset)
 {
-	float3 EdgeDirection = SharedNodePosition[GGroupThreadId.x] - SharedNodePosition[GGroupThreadId.x-1];
-	const float3 DeltaVelocity = (EdgeDirection - ( SharedPreviousPosition[GGroupThreadId.x] - SharedPreviousPosition[GGroupThreadId.x-1] ) ) / DeltaTime; 
+	const int PreviousIndex = GGroupThreadId.x-1-NodeOffset;
+	float3 EdgeDirection = SharedNodePosition[GGroupThreadId.x] - SharedNodePosition[PreviousIndex];
+	const float3 DeltaVelocity = (EdgeDirection - ( SharedPreviousPosition[GGroupThreadId.x] - SharedPreviousPosition[PreviousIndex] ) ) / DeltaTime; 
 
 	const float EdgeLength = length(EdgeDirection);
 	EdgeDirection /= EdgeLength;
 
 	// XPBD lagrange multiplier update : dL = -(C+compliance*L) / (dC * invM * dCt + alpha)
-	const float DeltaLambda = -((EdgeLength-RestLength) + OutMaterialMultiplier * MaterialCompliance + MaterialDamping * dot(EdgeDirection,DeltaVelocity) ) * MaterialWeight;
+	const float DeltaLambda = -((EdgeLength-RestLength)/RestLength + OutMaterialMultiplier * MaterialCompliance + MaterialDamping * dot(EdgeDirection,DeltaVelocity)/RestLength ) * MaterialWeight;
 
 	// L += dL
 	OutMaterialMultiplier += DeltaLambda;
 
 	// XPBD position update : dX = dL * dCt * invM
-	const float3 PositionDelta = EdgeDirection * DeltaLambda;
+	const float3 PositionDelta = EdgeDirection * DeltaLambda / RestLength;
 
 	// dX += dX
 	SharedNodePosition[GGroupThreadId.x] += PositionDelta * SharedInverseMass[GGroupThreadId.x];
 	if(!ProjectConstraint)
 	{
-		SharedNodePosition[GGroupThreadId.x-1] -= PositionDelta * SharedInverseMass[GGroupThreadId.x-1];
+		SharedNodePosition[PreviousIndex] -= PositionDelta * SharedInverseMass[PreviousIndex];
+	}
+}
+
+void SetupStretchSpringMaterial(in int StrandsSize, in float YoungModulus, in float RodThickness, 
+	in float RestLength, in float DeltaTime, in bool ProjectConstraint, in float MaterialDamping, out float OutMaterialCompliance, out float OutMaterialWeight, out float OutMaterialMultiplier)
+{
+	const int LocalIndex = (GGroupThreadId.x % StrandsSize);
+	if( LocalIndex > 1 )
+	{
+		ResetStretchSpringMaterial(StrandsSize,YoungModulus,RodThickness,RestLength,DeltaTime,ProjectConstraint,MaterialDamping,OutMaterialCompliance,OutMaterialWeight,OutMaterialMultiplier,0);
 	}
 }
 
@@ -62,13 +72,13 @@ void SolveStretchSpringMaterial(in int StrandsSize, in float RestLength, in floa
 		// Process all the red rods
 		if (IsRed)
 		{	
-			UpdateStretchSpringMultiplier(RestLength,DeltaTime,false,MaterialDamping,MaterialCompliance,MaterialWeight,OutMaterialMultiplier);
+			UpdateStretchSpringMultiplier(RestLength,DeltaTime,false,MaterialDamping,MaterialCompliance,MaterialWeight,OutMaterialMultiplier,0);
 		}
 		// Process all the black rods
 		GroupMemoryBarrier(); 
 		if (!IsRed)
 		{
-			UpdateStretchSpringMultiplier(RestLength,DeltaTime,false,MaterialDamping,MaterialCompliance,MaterialWeight,OutMaterialMultiplier);
+			UpdateStretchSpringMultiplier(RestLength,DeltaTime,false,MaterialDamping,MaterialCompliance,MaterialWeight,OutMaterialMultiplier,0);
 		}
 		GroupMemoryBarrier(); 
 	}
@@ -80,16 +90,136 @@ void ProjectStretchSpringMaterial(in int StrandsSize, in float YoungModulus, in
 	float MaterialWeight = 0.0;
 	OutMaterialMultiplier = 0.0;
 
-	SetupStretchSpringMaterial(StrandsSize, YoungModulus, RodThickness, RestLength, DeltaTime, true, 0.0, MaterialCompliance, MaterialWeight, OutMaterialMultiplier);
+	const int LocalIndex = (GGroupThreadId.x % StrandsSize);
+	if( LocalIndex > 1 )
+	{
+		ResetStretchSpringMaterial(StrandsSize,YoungModulus,RodThickness,RestLength,DeltaTime,true,0.0,MaterialCompliance,MaterialWeight,OutMaterialMultiplier,0);
+
+		for(int i = 2; i < StrandsSize; ++i)
+		{
+			if( LocalIndex == i )
+			{
+				UpdateStretchSpringMultiplier(RestLength,DeltaTime,true,0.0,MaterialCompliance,MaterialWeight,OutMaterialMultiplier,0);
+			}
+			GroupMemoryBarrier(); 
+		}
+	}
+}
+
+/* -----------------------------------------------------------------
+ * Bend spring material
+ * -----------------------------------------------------------------
+ */
 
+ void SetupBendSpringMaterial(in int StrandsSize, in float YoungModulus, in float RodThickness, 
+	in float RestLength, in float DeltaTime, in bool ProjectConstraint, in float MaterialDamping, out float OutMaterialCompliance, out float OutMaterialWeight, out float OutMaterialMultiplier)
+{
 	const int LocalIndex = (GGroupThreadId.x % StrandsSize);
 	if( LocalIndex > 1 )
 	{
+		ResetStretchSpringMaterial(StrandsSize,YoungModulus,RodThickness,RestLength,DeltaTime,ProjectConstraint,MaterialDamping,OutMaterialCompliance,OutMaterialWeight,OutMaterialMultiplier,1);
+	}
+}
+
+void SolveBendSpringMaterial(in int StrandsSize, in float RestLength, in float DeltaTime, in float MaterialDamping, 
+			in float MaterialCompliance, in float MaterialWeight, in float MaterialMultiplier, out float OutMaterialMultiplier)
+{
+	OutMaterialMultiplier = MaterialMultiplier;
+	const int LocalIndex = (GGroupThreadId.x % StrandsSize);
+	if( LocalIndex > 1)
+	{
+		const int IsRed = ((GGroupThreadId.x/2) % 2) == 0;
+		// Process all the red rods
+		if (IsRed)
+		{	
+			UpdateStretchSpringMultiplier(RestLength,DeltaTime,false,MaterialDamping,MaterialCompliance,MaterialWeight,OutMaterialMultiplier,1);
+		}
+		// Process all the black rods
+		GroupMemoryBarrier(); 
+		if (!IsRed)
+		{
+			UpdateStretchSpringMultiplier(RestLength,DeltaTime,false,MaterialDamping,MaterialCompliance,MaterialWeight,OutMaterialMultiplier,1);
+		}
+		GroupMemoryBarrier(); 
+	}
+} 
+
+void ProjectBendSpringMaterial(in int StrandsSize, in float YoungModulus, in float RodThickness, in float RestLength, in float DeltaTime, out float OutMaterialMultiplier)
+{
+	float MaterialCompliance = 0.0;
+	float MaterialWeight = 0.0;
+	OutMaterialMultiplier = 0.0;
+
+	const int LocalIndex = (GGroupThreadId.x % StrandsSize);
+	if( LocalIndex > 1 )
+	{
+		ResetStretchSpringMaterial(StrandsSize,YoungModulus,RodThickness,RestLength,DeltaTime,true,0.0,MaterialCompliance,MaterialWeight,OutMaterialMultiplier,1);
+
 		for(int i = 2; i < StrandsSize; ++i)
 		{
 			if( LocalIndex == i )
 			{
-				UpdateStretchSpringMultiplier(RestLength,DeltaTime,true,0.0,MaterialCompliance,MaterialWeight,OutMaterialMultiplier);
+				UpdateStretchSpringMultiplier(RestLength,DeltaTime,true,0.0,MaterialCompliance,MaterialWeight,OutMaterialMultiplier,1);
+			}
+			GroupMemoryBarrier(); 
+		}
+	}
+}
+
+/* -----------------------------------------------------------------
+ * Twist spring material
+ * -----------------------------------------------------------------
+ */
+
+ void SetupTwistSpringMaterial(in int StrandsSize, in float YoungModulus, in float RodThickness, 
+	in float RestLength, in float DeltaTime, in bool ProjectConstraint, in float MaterialDamping, out float OutMaterialCompliance, out float OutMaterialWeight, out float OutMaterialMultiplier)
+{
+	const int LocalIndex = (GGroupThreadId.x % StrandsSize);
+	if( LocalIndex > 2 )
+	{
+		ResetStretchSpringMaterial(StrandsSize,YoungModulus,RodThickness,RestLength,DeltaTime,ProjectConstraint,MaterialDamping,OutMaterialCompliance,OutMaterialWeight,OutMaterialMultiplier,2);
+	}
+}
+
+void SolveTwistSpringMaterial(in int StrandsSize, in float RestLength, in float DeltaTime, in float MaterialDamping, 
+			in float MaterialCompliance, in float MaterialWeight, in float MaterialMultiplier, out float OutMaterialMultiplier)
+{
+	OutMaterialMultiplier = MaterialMultiplier;
+	const int LocalIndex = (GGroupThreadId.x % StrandsSize);
+	if( LocalIndex > 2)
+	{
+		const int IsRed = ((GGroupThreadId.x/3) % 2) == 0;
+		// Process all the red rods
+		if (IsRed)
+		{	
+			UpdateStretchSpringMultiplier(RestLength,DeltaTime,false,MaterialDamping,MaterialCompliance,MaterialWeight,OutMaterialMultiplier,2);
+		}
+		// Process all the black rods
+		GroupMemoryBarrier(); 
+		if (!IsRed)
+		{
+			UpdateStretchSpringMultiplier(RestLength,DeltaTime,false,MaterialDamping,MaterialCompliance,MaterialWeight,OutMaterialMultiplier,2);
+		}
+		GroupMemoryBarrier(); 
+	}
+} 
+
+void ProjectTwistSpringMaterial(in int StrandsSize, in float YoungModulus, in float RodThickness, in float RestLength, in float DeltaTime, out float OutMaterialMultiplier)
+{
+	float MaterialCompliance = 0.0;
+	float MaterialWeight = 0.0;
+	OutMaterialMultiplier = 0.0;
+
+	const int LocalIndex = (GGroupThreadId.x % StrandsSize);
+	if( LocalIndex > 2 )
+	{
+		ResetStretchSpringMaterial(StrandsSize,YoungModulus,RodThickness,RestLength,DeltaTime,true,0.0,MaterialCompliance,MaterialWeight,OutMaterialMultiplier,2);
+
+		for(int i = 3; i < StrandsSize; ++i)
+		{
+			if( LocalIndex == i )
+			{
+				UpdateStretchSpringMultiplier(RestLength,DeltaTime,true,0.0,MaterialCompliance,MaterialWeight,OutMaterialMultiplier,2);
 			}
 			GroupMemoryBarrier(); 
 		}
@@ -127,7 +257,7 @@ void ComputeMaterialFrame(in int StrandsSize)
  * -----------------------------------------------------------------
  */
 
- void SetupBendSpringMaterial(in int StrandsSize, in float YoungModulus, in float RodThickness, 
+ void SetupAngularSpringMaterial(in int StrandsSize, in float YoungModulus, in float RodThickness, 
 	in float RestLength, in float DeltaTime, in bool ProjectConstraint, in float MaterialDamping, out float OutMaterialCompliance, out float OutMaterialWeight, out float3 OutMaterialMultiplier)
 {
 	float MaterialMultiplier = 0;
@@ -136,7 +266,7 @@ void ComputeMaterialFrame(in int StrandsSize)
 	//ComputeMaterialFrame(StrandsSize);
 }
 
-void UpdateBendSpringMultiplier( in float3 RestDirection, in float DeltaTime, in bool ProjectConstraint, in float MaterialDamping, in float InMaterialCompliance, in float InMaterialWeight, inout float3 OutMaterialMultiplier)
+void UpdateAngularSpringMultiplier( in float RestLength, in float3 RestDirection, in float DeltaTime, in bool ProjectConstraint, in float MaterialDamping, in float InMaterialCompliance, in float InMaterialWeight, inout float3 OutMaterialMultiplier)
 {
 	const float3 EdgeDirection = SharedNodePosition[GGroupThreadId.x] - SharedNodePosition[GGroupThreadId.x-1];
 	const float3 DeltaVelocity = (EdgeDirection - ( SharedPreviousPosition[GGroupThreadId.x] - SharedPreviousPosition[GGroupThreadId.x-1] ) ) / DeltaTime; 
@@ -145,14 +275,14 @@ void UpdateBendSpringMultiplier( in float3 RestDirection, in float DeltaTime, in
 	const float3 TargetDirection = RotateVectorByQuat(RestDirection,ParentFrame);
 
 	// XPBD lagrange multiplier update : dL = -(C+compliance*L) / (dC * invM * dCt + alpha)
-	const float3 DeltaLambda = -(EdgeDirection - TargetDirection + 
-					OutMaterialMultiplier * InMaterialCompliance  + MaterialDamping * DeltaVelocity ) * InMaterialWeight;
+	const float3 DeltaLambda = -((EdgeDirection - TargetDirection)/RestLength + 
+					OutMaterialMultiplier * InMaterialCompliance  + MaterialDamping * DeltaVelocity /RestLength) * InMaterialWeight;
 
 	// L += dL
 	OutMaterialMultiplier += DeltaLambda;
 
 	// XPBD position update : dX = dL * dCt * invM
-	const float3 PositionDelta  = DeltaLambda;
+	const float3 PositionDelta  = DeltaLambda/RestLength;
 
 	// dX += dX
 	SharedNodePosition[GGroupThreadId.x] += PositionDelta * SharedInverseMass[GGroupThreadId.x];
@@ -162,11 +292,9 @@ void UpdateBendSpringMultiplier( in float3 RestDirection, in float DeltaTime, in
 	}
 }
 
-void SolveBendSpringMaterial(in int StrandsSize, in float3 RestDirection, in float DeltaTime, in float MaterialDamping,
+void SolveAngularSpringMaterial(in int StrandsSize, in float RestLength, in float3 RestDirection, in float DeltaTime, in float MaterialDamping,
 			in float MaterialCompliance, in float MaterialWeight, in float3 MaterialMultiplier, out float3 OutMaterialMultiplier)
 {
-	//ComputeMaterialFrame(StrandsSize);
-
 	OutMaterialMultiplier = MaterialMultiplier;
 	const int LocalIndex = (GGroupThreadId.x % StrandsSize);
 	if( LocalIndex > 1)
@@ -175,19 +303,19 @@ void SolveBendSpringMaterial(in int StrandsSize, in float3 RestDirection, in flo
 		// Process all the red rods
 		if (IsRed)
 		{	
-			UpdateBendSpringMultiplier(RestDirection,DeltaTime,false,MaterialDamping,MaterialCompliance,MaterialWeight,OutMaterialMultiplier);
+			UpdateAngularSpringMultiplier(RestLength,RestDirection,DeltaTime,false,MaterialDamping,MaterialCompliance,MaterialWeight,OutMaterialMultiplier);
 		}
 		GroupMemoryBarrier(); 
 
 		if (!IsRed)
 		{
-			UpdateBendSpringMultiplier(RestDirection,DeltaTime,false,MaterialDamping,MaterialCompliance,MaterialWeight,OutMaterialMultiplier);
+			UpdateAngularSpringMultiplier(RestLength,RestDirection,DeltaTime,false,MaterialDamping,MaterialCompliance,MaterialWeight,OutMaterialMultiplier);
 		}
 		GroupMemoryBarrier(); 
 	}
 } 
 
-void ProjectBendSpringMaterial(in int StrandsSize, in float YoungModulus, in float RodThickness, in float RestLength, in float3 RestDirection, in float DeltaTime,
+void ProjectAngularSpringMaterial(in int StrandsSize, in float YoungModulus, in float RodThickness, in float RestLength, in float3 RestDirection, in float DeltaTime,
 			out float3 OutMaterialMultiplier)
 {
 	float MaterialCompliance = 0.0;
@@ -195,7 +323,7 @@ void ProjectBendSpringMaterial(in int StrandsSize, in float YoungModulus, in flo
 	OutMaterialMultiplier = float3(0,0,0);
 
 	ComputeMaterialFrame(StrandsSize);
-	SetupBendSpringMaterial(StrandsSize, YoungModulus, RodThickness, RestLength, DeltaTime, true, 0.0, MaterialCompliance, MaterialWeight, OutMaterialMultiplier);
+	SetupAngularSpringMaterial(StrandsSize, YoungModulus, RodThickness, RestLength, DeltaTime, true, 0.0, MaterialCompliance, MaterialWeight, OutMaterialMultiplier);
 
 	float3 MaterialMultiplier = float3(0,0,0);
 	const int LocalIndex = (GGroupThreadId.x % StrandsSize);
@@ -205,7 +333,7 @@ void ProjectBendSpringMaterial(in int StrandsSize, in float YoungModulus, in flo
 		{
 			if( LocalIndex == i )
 			{
-				UpdateBendSpringMultiplier(RestDirection,DeltaTime,true,0.0,MaterialCompliance,MaterialWeight,OutMaterialMultiplier);
+				UpdateAngularSpringMultiplier(RestLength,RestDirection,DeltaTime,true,0.0,MaterialCompliance,MaterialWeight,OutMaterialMultiplier);
 			}
 			GroupMemoryBarrier(); 
 		}
@@ -213,7 +341,7 @@ void ProjectBendSpringMaterial(in int StrandsSize, in float YoungModulus, in flo
 }
 
 /* -----------------------------------------------------------------
- * Compute Rest Direction
+ * Compute Rest Direction and length
  * -----------------------------------------------------------------
  */
 
@@ -235,4 +363,4 @@ void ComputeRestDirection(in int StrandsSize, in float3 NodePosition, in float4
 		OutRestDirection = UnrotateVectorByQuat(EdgeDirection,ParentFrame);
 	}
 	GroupMemoryBarrier(); 
-}
\ No newline at end of file
+}
diff --git a/Engine/Plugins/Experimental/HairStrands/Source/HairStrandsNiagara/Private/NiagaraDataInterfaceHairStrands.cpp b/Engine/Plugins/Experimental/HairStrands/Source/HairStrandsNiagara/Private/NiagaraDataInterfaceHairStrands.cpp
--- a/Engine/Plugins/Experimental/HairStrands/Source/HairStrandsNiagara/Private/NiagaraDataInterfaceHairStrands.cpp
+++ b/Engine/Plugins/Experimental/HairStrands/Source/HairStrandsNiagara/Private/NiagaraDataInterfaceHairStrands.cpp
@@ -17,11 +17,8 @@
 
 //------------------------------------------------------------------------------------------------------------
 
-static const FName GetStrandDensityName(TEXT("GetStrandDensity"));
 static const FName GetStrandSizeName(TEXT("GetStrandSize"));
 static const FName GetNumStrandsName(TEXT("GetNumStrands"));
-static const FName GetRootThicknessName(TEXT("GetRootThickness"));
-static const FName GetTipThicknessName(TEXT("GetTipThickness"));
 static const FName GetWorldTransformName(TEXT("GetWorldTransform"));
 static const FName GetWorldInverseName(TEXT("GetWorldInverse"));
 static const FName GetPointPositionName(TEXT("GetPointPosition"));
@@ -56,26 +53,20 @@ static const FName ResetPointPositionName(TEXT("ResetPointPosition"));
 
 //------------------------------------------------------------------------------------------------------------
 
-static const FName BuildCollisionGridName(TEXT("BuildCollisionGrid"));
-static const FName QueryCollisionGridName(TEXT("QueryCollisionGrid"));
-static const FName ProjectCollisionGridName(TEXT("ProjectCollisionGrid"));
-
-//------------------------------------------------------------------------------------------------------------
-
 static const FName GetBoxCenterName(TEXT("GetBoxCenter"));
 static const FName GetBoxExtentName(TEXT("GetBoxExtent"));
 
 //------------------------------------------------------------------------------------------------------------
 
-static const FName SetupStretchSpringMaterialName(TEXT("SetupStretchSpringMaterial"));
-static const FName SolveStretchSpringMaterialName(TEXT("SolveStretchSpringMaterial"));
-static const FName ProjectStretchSpringMaterialName(TEXT("ProjectStretchSpringMaterial"));
+static const FName SetupDistanceSpringMaterialName(TEXT("SetupDistanceSpringMaterial"));
+static const FName SolveDistanceSpringMaterialName(TEXT("SolveDistanceSpringMaterial"));
+static const FName ProjectDistanceSpringMaterialName(TEXT("ProjectDistanceSpringMaterial"));
 
 //------------------------------------------------------------------------------------------------------------
 
-static const FName SetupBendSpringMaterialName(TEXT("SetupBendSpringMaterial"));
-static const FName SolveBendSpringMaterialName(TEXT("SolveBendSpringMaterial"));
-static const FName ProjectBendSpringMaterialName(TEXT("ProjectBendSpringMaterial"));
+static const FName SetupAngularSpringMaterialName(TEXT("SetupAngularSpringMaterial"));
+static const FName SolveAngularSpringMaterialName(TEXT("SolveAngularSpringMaterial"));
+static const FName ProjectAngularSpringMaterialName(TEXT("ProjectAngularSpringMaterial"));
 
 //------------------------------------------------------------------------------------------------------------
 
@@ -100,11 +91,12 @@ static const FName UpdateNodeOrientationName(TEXT("UpdateNodeOrientation"));
 
 //------------------------------------------------------------------------------------------------------------
 
+static const FName ComputeAirDragForceName(TEXT("ComputeAirDragForce"));  
+
+//------------------------------------------------------------------------------------------------------------
+
 const FString UNiagaraDataInterfaceHairStrands::NumStrandsName(TEXT("NumStrands_"));
 const FString UNiagaraDataInterfaceHairStrands::StrandSizeName(TEXT("StrandSize_"));
-const FString UNiagaraDataInterfaceHairStrands::StrandDensityName(TEXT("StrandDensity_"));
-const FString UNiagaraDataInterfaceHairStrands::RootThicknessName(TEXT("RootThickness_"));
-const FString UNiagaraDataInterfaceHairStrands::TipThicknessName(TEXT("TipThickness_"));
 const FString UNiagaraDataInterfaceHairStrands::WorldTransformName(TEXT("WorldTransform_"));
 const FString UNiagaraDataInterfaceHairStrands::WorldInverseName(TEXT("WorldInverse_"));
 const FString UNiagaraDataInterfaceHairStrands::WorldRotationName(TEXT("WorldRotation_"));
@@ -113,11 +105,7 @@ const FString UNiagaraDataInterfaceHairStrands::PointsPositionsBufferName(TEXT("
 const FString UNiagaraDataInterfaceHairStrands::CurvesOffsetsBufferName(TEXT("CurvesOffsetsBuffer_"));
 const FString UNiagaraDataInterfaceHairStrands::RestPositionsBufferName(TEXT("RestPositionsBuffer_"));
 
-const FString UNiagaraDataInterfaceHairStrands::GridCurrentBufferName(TEXT("GridCurrentBuffer_"));
-const FString UNiagaraDataInterfaceHairStrands::GridDestinationBufferName(TEXT("GridDestinationBuffer_"));
-
-const FString UNiagaraDataInterfaceHairStrands::GridSizeName(TEXT("GridSize_"));
-const FString UNiagaraDataInterfaceHairStrands::GridOriginName(TEXT("GridOrigin_"));
+const FString UNiagaraDataInterfaceHairStrands::BoxOriginName(TEXT("BoxOrigin_"));
 
 //------------------------------------------------------------------------------------------------------------
 
@@ -127,73 +115,30 @@ struct FNDIHairStrandsParametersName
 	{
 		NumStrandsName = UNiagaraDataInterfaceHairStrands::NumStrandsName + Suffix;
 		StrandSizeName = UNiagaraDataInterfaceHairStrands::StrandSizeName + Suffix;
-		StrandDensityName = UNiagaraDataInterfaceHairStrands::StrandDensityName + Suffix;
-		RootThicknessName = UNiagaraDataInterfaceHairStrands::RootThicknessName + Suffix;
-		TipThicknessName = UNiagaraDataInterfaceHairStrands::TipThicknessName + Suffix;
 		WorldTransformName = UNiagaraDataInterfaceHairStrands::WorldTransformName + Suffix;
 		WorldInverseName = UNiagaraDataInterfaceHairStrands::WorldInverseName + Suffix;
 		WorldRotationName = UNiagaraDataInterfaceHairStrands::WorldRotationName + Suffix;
+		BoxOriginName = UNiagaraDataInterfaceHairStrands::BoxOriginName + Suffix;
 
 		PointsPositionsBufferName = UNiagaraDataInterfaceHairStrands::PointsPositionsBufferName + Suffix;
 		CurvesOffsetsBufferName = UNiagaraDataInterfaceHairStrands::CurvesOffsetsBufferName + Suffix;
 		RestPositionsBufferName = UNiagaraDataInterfaceHairStrands::RestPositionsBufferName + Suffix;
-
-		GridCurrentBufferName = UNiagaraDataInterfaceHairStrands::GridCurrentBufferName + Suffix;
-		GridDestinationBufferName = UNiagaraDataInterfaceHairStrands::GridDestinationBufferName + Suffix;
-
-		GridSizeName = UNiagaraDataInterfaceHairStrands::GridSizeName + Suffix;
-		GridOriginName = UNiagaraDataInterfaceHairStrands::GridOriginName + Suffix;
 	}
 
 	FString NumStrandsName;
 	FString StrandSizeName;
-	FString StrandDensityName;
-	FString RootThicknessName;
-	FString TipThicknessName;
 	FString WorldTransformName;
 	FString WorldInverseName;
 	FString WorldRotationName;
+	FString BoxOriginName;
 
 	FString PointsPositionsBufferName;
 	FString CurvesOffsetsBufferName;
 	FString RestPositionsBufferName;
-
-	FString GridCurrentBufferName;
-	FString GridDestinationBufferName;
-
-	FString GridSizeName;
-	FString GridOriginName;
 };
 
 //------------------------------------------------------------------------------------------------------------
 
-void FNDICollisionGridBuffer::SetGridSize(const FUintVector4 InGridSize)
-{
-	GridSize = InGridSize;
-}
-
-void FNDICollisionGridBuffer::InitRHI()
-{
-	if(GridSize.X != 0 && GridSize.Y != 0 && GridSize.Z != 0)
-	{
-		static const uint32 NumComponents = 9;
-		GridDataBuffer.Initialize(sizeof(int32), GridSize.X*NumComponents, GridSize.Y,
-			GridSize.Z, EPixelFormat::PF_R32_SINT);
-	}
-}
-
-void FNDICollisionGridBuffer::ReleaseRHI()
-{
-	GridDataBuffer.Release();
-}
-
-void FNDICollisionGridBuffer::ClearBuffers(FRHICommandList& RHICmdList)
-{
-	ClearUAV(RHICmdList, GridDataBuffer, FLinearColor(0, 0, 0, 0));
-}
-
-//------------------------------------------------------------------------------------------------------------
-
 void FNDIHairStrandsBuffer::SetHairAsset(const FHairStrandsDatas*  HairStrandsDatas, const FHairStrandsRestResource* HairStrandsRestResource, const FHairStrandsDeformedResource*  HairStrandsDeformedResource)
 {
 	SourceDatas = HairStrandsDatas;
@@ -203,31 +148,30 @@ void FNDIHairStrandsBuffer::SetHairAsset(const FHairStrandsDatas*  HairStrandsDa
 
 void FNDIHairStrandsBuffer::InitRHI()
 {
-	if (SourceDatas != nullptr && SourceRestResources != nullptr && SourceDeformedResources != nullptr)
+	if (SourceDatas != nullptr && SourceRestResources != nullptr)
 	{
-		const uint32 OffsetCount = SourceDatas->GetNumCurves() + 1;
-		const uint32 OffsetBytes = sizeof(uint32)*OffsetCount;
+		{
+			const uint32 OffsetCount = SourceDatas->GetNumCurves() + 1;
+			const uint32 OffsetBytes = sizeof(uint32)*OffsetCount;
 
-		CurvesOffsetsBuffer.Initialize(sizeof(uint32), OffsetCount, EPixelFormat::PF_R32_UINT, BUF_Static);
-		void* OffsetBufferData = RHILockVertexBuffer(CurvesOffsetsBuffer.Buffer, 0, OffsetBytes, RLM_WriteOnly);
+			CurvesOffsetsBuffer.Initialize(sizeof(uint32), OffsetCount, EPixelFormat::PF_R32_UINT, BUF_Static);
+			void* OffsetBufferData = RHILockVertexBuffer(CurvesOffsetsBuffer.Buffer, 0, OffsetBytes, RLM_WriteOnly);
 
-		FMemory::Memcpy(OffsetBufferData, SourceDatas->StrandsCurves.CurvesOffset.GetData(), OffsetBytes);
-		RHIUnlockVertexBuffer(CurvesOffsetsBuffer.Buffer);
+			FMemory::Memcpy(OffsetBufferData, SourceDatas->StrandsCurves.CurvesOffset.GetData(), OffsetBytes);
+			RHIUnlockVertexBuffer(CurvesOffsetsBuffer.Buffer);
+		}
+		if (SourceDeformedResources == nullptr)
+		{
+			const uint32 PositionsCount = SourceDatas->GetNumPoints();
+			PointsPositionsBuffer.Initialize(FHairStrandsPositionFormat::SizeInByte, PositionsCount, FHairStrandsPositionFormat::Format, BUF_Static);
+		}
 	}
 }
 
 void FNDIHairStrandsBuffer::ReleaseRHI()
 {
 	CurvesOffsetsBuffer.Release();
-}
-
-//------------------------------------------------------------------------------------------------------------
-
-void FNDIHairStrandsData::SwapBuffers()
-{
-	FNDICollisionGridBuffer* StoredBufferPointer = CurrentGridBuffer;
-	CurrentGridBuffer = DestinationGridBuffer;
-	DestinationGridBuffer = StoredBufferPointer;
+	PointsPositionsBuffer.Release();
 }
 
 //------------------------------------------------------------------------------------------------------------
@@ -243,20 +187,12 @@ struct FNDIHairStrandsParametersCS : public FNiagaraDataInterfaceParametersCS
 		WorldRotation.Bind(ParameterMap, *ParamNames.WorldRotationName);
 		NumStrands.Bind(ParameterMap, *ParamNames.NumStrandsName);
 		StrandSize.Bind(ParameterMap, *ParamNames.StrandSizeName);
-		StrandDensity.Bind(ParameterMap, *ParamNames.StrandDensityName);
-		RootThickness.Bind(ParameterMap, *ParamNames.RootThicknessName);
-		TipThickness.Bind(ParameterMap, *ParamNames.TipThicknessName);
+		BoxOrigin.Bind(ParameterMap, *ParamNames.BoxOriginName);
 
 		PointsPositionsBuffer.Bind(ParameterMap, *ParamNames.PointsPositionsBufferName);
 		CurvesOffsetsBuffer.Bind(ParameterMap, *ParamNames.CurvesOffsetsBufferName);
 		RestPositionsBuffer.Bind(ParameterMap, *ParamNames.RestPositionsBufferName);
 
-		GridCurrentBuffer.Bind(ParameterMap, *ParamNames.GridCurrentBufferName);
-		GridDestinationBuffer.Bind(ParameterMap, *ParamNames.GridDestinationBufferName);
-
-		GridOrigin.Bind(ParameterMap, *ParamNames.GridOriginName);
-		GridSize.Bind(ParameterMap, *ParamNames.GridSizeName);
-
 		if (!PointsPositionsBuffer.IsBound())
 		{
 			UE_LOG(LogHairStrands, Warning, TEXT("Binding failed for FNDIHairStrandsParametersCS %s. Was it optimized out?"), *ParamNames.PointsPositionsBufferName)
@@ -269,16 +205,6 @@ struct FNDIHairStrandsParametersCS : public FNiagaraDataInterfaceParametersCS
 		{
 			UE_LOG(LogHairStrands, Warning, TEXT("Binding failed for FNDIHairStrandsParametersCS %s. Was it optimized out?"), *ParamNames.RestPositionsBufferName)
 		}
-
-		if (!GridCurrentBuffer.IsBound())
-		{
-			UE_LOG(LogHairStrands, Warning, TEXT("Binding failed for FNDIHairStrandsParametersCS %s. Was it optimized out?"), *ParamNames.GridCurrentBufferName)
-		}
-
-		if (!GridDestinationBuffer.IsBound())
-		{
-			UE_LOG(LogHairStrands, Warning, TEXT("Binding failed for FNDIHairStrandsParametersCS %s. Was it optimized out?"), *ParamNames.GridDestinationBufferName)
-		}
 	}
 
 	virtual void Serialize(FArchive& Ar) override
@@ -288,19 +214,11 @@ struct FNDIHairStrandsParametersCS : public FNiagaraDataInterfaceParametersCS
 		Ar << WorldRotation;
 		Ar << NumStrands;
 		Ar << StrandSize;
-		Ar << StrandDensity;
-		Ar << RootThickness;
-		Ar << TipThickness;
+		Ar << BoxOrigin;
 
 		Ar << PointsPositionsBuffer;
 		Ar << CurvesOffsetsBuffer;
 		Ar << RestPositionsBuffer;
-
-		Ar << GridCurrentBuffer;
-		Ar << GridDestinationBuffer;
-
-		Ar << GridOrigin;
-		Ar << GridSize;
 	}
 
 	virtual void Set(FRHICommandList& RHICmdList, const FNiagaraDataInterfaceSetArgs& Context) const override
@@ -319,30 +237,19 @@ struct FNDIHairStrandsParametersCS : public FNiagaraDataInterfaceParametersCS
 		{
 			FNDIHairStrandsBuffer* HairStrandsBuffer = ProxyData->HairStrandsBuffer;
 
-			FNDICollisionGridBuffer* CurrentGridBuffer = ProxyData->CurrentGridBuffer;
-			FNDICollisionGridBuffer* DestinationGridBuffer = ProxyData->DestinationGridBuffer;
+			FUnorderedAccessViewRHIRef PointPositionsUAV = (HairStrandsBuffer->SourceDeformedResources != nullptr) ? HairStrandsBuffer->SourceDeformedResources->DeformedPositionBuffer[0].UAV : 
+				HairStrandsBuffer->PointsPositionsBuffer.UAV;
 
-			SetUAVParameter(RHICmdList, ComputeShaderRHI, PointsPositionsBuffer, HairStrandsBuffer->SourceDeformedResources->DeformedPositionBuffer[0].UAV);
+			SetUAVParameter(RHICmdList, ComputeShaderRHI, PointsPositionsBuffer, PointPositionsUAV);
 			SetSRVParameter(RHICmdList, ComputeShaderRHI, CurvesOffsetsBuffer, HairStrandsBuffer->CurvesOffsetsBuffer.SRV);
 			SetSRVParameter(RHICmdList, ComputeShaderRHI, RestPositionsBuffer, HairStrandsBuffer->SourceRestResources->RestPositionBuffer.SRV);
 
-			RHICmdList.TransitionResource(EResourceTransitionAccess::EWritable, EResourceTransitionPipeline::EComputeToCompute, DestinationGridBuffer->GridDataBuffer.UAV);
-			SetUAVParameter(RHICmdList, ComputeShaderRHI, GridDestinationBuffer, DestinationGridBuffer->GridDataBuffer.UAV);
-
-			RHICmdList.TransitionResource(EResourceTransitionAccess::EReadable, EResourceTransitionPipeline::EComputeToCompute, CurrentGridBuffer->GridDataBuffer.UAV);
-			SetSRVParameter(RHICmdList, ComputeShaderRHI, GridCurrentBuffer, CurrentGridBuffer->GridDataBuffer.SRV);
-
 			SetShaderValue(RHICmdList, ComputeShaderRHI, WorldTransform, ProxyData->WorldTransform);
 			SetShaderValue(RHICmdList, ComputeShaderRHI, WorldInverse, ProxyData->WorldTransform.Inverse());
 			SetShaderValue(RHICmdList, ComputeShaderRHI, WorldRotation, ProxyData->WorldTransform.ToQuat());
 			SetShaderValue(RHICmdList, ComputeShaderRHI, NumStrands, ProxyData->NumStrands);
 			SetShaderValue(RHICmdList, ComputeShaderRHI, StrandSize, ProxyData->StrandSize);
-			SetShaderValue(RHICmdList, ComputeShaderRHI, StrandDensity, ProxyData->StrandDensity);
-			SetShaderValue(RHICmdList, ComputeShaderRHI, RootThickness, ProxyData->RootThickness);
-			SetShaderValue(RHICmdList, ComputeShaderRHI, TipThickness, ProxyData->TipThickness);
-
-			SetShaderValue(RHICmdList, ComputeShaderRHI, GridOrigin, ProxyData->GridOrigin);
-			SetShaderValue(RHICmdList, ComputeShaderRHI, GridSize, ProxyData->GridSize);
+			SetShaderValue(RHICmdList, ComputeShaderRHI, BoxOrigin, ProxyData->BoxOrigin);
 		}
 		else
 		{
@@ -350,27 +257,18 @@ struct FNDIHairStrandsParametersCS : public FNiagaraDataInterfaceParametersCS
 			SetSRVParameter(RHICmdList, ComputeShaderRHI, CurvesOffsetsBuffer, FNiagaraRenderer::GetDummyUIntBuffer().SRV);
 			SetSRVParameter(RHICmdList, ComputeShaderRHI, RestPositionsBuffer, FNiagaraRenderer::GetDummyFloatBuffer().SRV);
 
-			SetUAVParameter(RHICmdList, ComputeShaderRHI, GridDestinationBuffer, FNiagaraRenderer::GetDummyUIntBuffer().UAV);
-			SetSRVParameter(RHICmdList, ComputeShaderRHI, GridCurrentBuffer, FNiagaraRenderer::GetDummyUIntBuffer().SRV);
-
-			SetShaderValue(RHICmdList, ComputeShaderRHI, GridOrigin, FVector4(0,0,0,0));
-			SetShaderValue(RHICmdList, ComputeShaderRHI, GridSize, FUintVector4());
-
 			SetShaderValue(RHICmdList, ComputeShaderRHI, WorldTransform, FMatrix::Identity);
 			SetShaderValue(RHICmdList, ComputeShaderRHI, WorldInverse, FMatrix::Identity);
 			SetShaderValue(RHICmdList, ComputeShaderRHI, WorldRotation, FQuat::Identity);
 			SetShaderValue(RHICmdList, ComputeShaderRHI, NumStrands, 1);
 			SetShaderValue(RHICmdList, ComputeShaderRHI, StrandSize, 1);
-			SetShaderValue(RHICmdList, ComputeShaderRHI, StrandDensity, 1);
-			SetShaderValue(RHICmdList, ComputeShaderRHI, RootThickness, 0.1);
-			SetShaderValue(RHICmdList, ComputeShaderRHI, TipThickness, 0.1);
+			SetShaderValue(RHICmdList, ComputeShaderRHI, BoxOrigin, FVector4(0,0,0,0));
 		}
 	}
 
 	virtual void Unset(FRHICommandList& RHICmdList, const FNiagaraDataInterfaceSetArgs& Context) const override
 	{
 		SetUAVParameter(RHICmdList, Context.Shader->GetComputeShader(), PointsPositionsBuffer, nullptr);
-		SetUAVParameter(RHICmdList, Context.Shader->GetComputeShader(), GridDestinationBuffer, nullptr);
 	}
 
 private:
@@ -380,19 +278,11 @@ private:
 	FShaderParameter WorldRotation;
 	FShaderParameter NumStrands;
 	FShaderParameter StrandSize;
-	FShaderParameter StrandDensity;
-	FShaderParameter RootThickness;
-	FShaderParameter TipThickness;
+	FShaderParameter BoxOrigin;
 
 	FShaderResourceParameter PointsPositionsBuffer;
 	FShaderResourceParameter CurvesOffsetsBuffer;
 	FShaderResourceParameter RestPositionsBuffer;
-
-	FShaderResourceParameter GridCurrentBuffer;
-	FShaderResourceParameter GridDestinationBuffer;
-
-	FShaderParameter GridSize;
-	FShaderParameter GridOrigin;
 };
 
 //------------------------------------------------------------------------------------------------------------
@@ -416,6 +306,10 @@ void FNDIHairStrandsProxy::ConsumePerInstanceDataFromGameThread(void* PerInstanc
 	if (TargetData)
 	{
 		TargetData->WorldTransform = SourceData->WorldTransform;
+		TargetData->NumStrands = SourceData->NumStrands;
+		TargetData->StrandSize = SourceData->StrandSize;
+		TargetData->BoxOrigin = SourceData->BoxOrigin;
+		TargetData->HairStrandsBuffer = SourceData->HairStrandsBuffer;
 	}
 	else
 	{
@@ -423,8 +317,7 @@ void FNDIHairStrandsProxy::ConsumePerInstanceDataFromGameThread(void* PerInstanc
 	}
 }
 
-void FNDIHairStrandsProxy::InitializePerInstanceData(const FNiagaraSystemInstanceID& SystemInstance, FNDIHairStrandsBuffer* HairStrandsBuffer, FNDICollisionGridBuffer* CurrentGridBuffer, FNDICollisionGridBuffer* DestinationGridBuffer, uint32 NumStrands, const uint8 StrandSize,
-	const float StrandDensity, const float RootThickness, const float TipThickness, const FVector4& GridOrigin, const FUintVector4& GridSize)
+void FNDIHairStrandsProxy::InitializePerInstanceData(const FNiagaraSystemInstanceID& SystemInstance, FNDIHairStrandsBuffer* HairStrandsBuffer, const uint32 NumStrands, const uint8 StrandSize, const FVector4& BoxOrigin, const FMatrix& WorldTransform)
 {
 	check(IsInRenderingThread());
 
@@ -438,15 +331,10 @@ void FNDIHairStrandsProxy::InitializePerInstanceData(const FNiagaraSystemInstanc
 		TargetData = &SystemInstancesToProxyData.Add(SystemInstance);
 	}
 	TargetData->HairStrandsBuffer = HairStrandsBuffer;
-	TargetData->CurrentGridBuffer = CurrentGridBuffer;
-	TargetData->DestinationGridBuffer = DestinationGridBuffer;
 	TargetData->NumStrands = NumStrands;
 	TargetData->StrandSize = StrandSize;
-	TargetData->StrandDensity = StrandDensity;
-	TargetData->RootThickness = RootThickness;
-	TargetData->TipThickness = TipThickness;
-	TargetData->GridOrigin = GridOrigin;
-	TargetData->GridSize = GridSize;
+	TargetData->BoxOrigin = BoxOrigin;
+	TargetData->WorldTransform = WorldTransform;
 }
 
 void FNDIHairStrandsProxy::DestroyPerInstanceData(NiagaraEmitterInstanceBatcher* Batcher, const FNiagaraSystemInstanceID& SystemInstance)
@@ -462,17 +350,9 @@ void FNDIHairStrandsProxy::DestroyPerInstanceData(NiagaraEmitterInstanceBatcher*
 UNiagaraDataInterfaceHairStrands::UNiagaraDataInterfaceHairStrands(FObjectInitializer const& ObjectInitializer)
 	: Super(ObjectInitializer)
 	, StrandSize(EHairStrandsSize::Size4)
-	, StrandDensity(1.0)
-	, RootThickness(0.1)
-	, TipThickness(0.1)
 	, DefaultSource(nullptr)
 	, SourceActor(nullptr)
-	, SourceTransform(FMatrix::Identity)
-	, GridSizeX(10)
-	, GridSizeY(10)
-	, GridSizeZ(10)
 	, SourceComponent(nullptr)
-	, GroomAsset(nullptr)
 {
 
 	Proxy = MakeShared<FNDIHairStrandsProxy, ESPMode::ThreadSafe>();
@@ -483,11 +363,8 @@ bool UNiagaraDataInterfaceHairStrands::IsComponentValid() const
 	return (SourceComponent.IsValid() && SourceComponent != nullptr);
 }
 
-bool UNiagaraDataInterfaceHairStrands::InitPerInstanceData(void* PerInstanceData, FNiagaraSystemInstance* SystemInstance)
+void UNiagaraDataInterfaceHairStrands::UpdateSourceComponent(FNiagaraSystemInstance* SystemInstance)
 {
-	FNDIHairStrandsData* InstanceData = new (PerInstanceData) FNDIHairStrandsData();
-	check(InstanceData);
-	
 	SourceComponent = nullptr;
 	if (SourceActor)
 	{
@@ -527,13 +404,44 @@ bool UNiagaraDataInterfaceHairStrands::InitPerInstanceData(void* PerInstanceData
 			}
 		}
 	}
-	GroomAsset = IsComponentValid() ? SourceComponent->GroomAsset : DefaultSource;
+}
+
+void UNiagaraDataInterfaceHairStrands::ExtractDatasAndResources(FNiagaraSystemInstance* SystemInstance,FHairStrandsDatas*& OutStrandsDatas, 
+	FHairStrandsRestResource*& OutStrandsRestResource, FHairStrandsDeformedResource*& OutStrandsDeformedResource)
+{
+	UpdateSourceComponent(SystemInstance);
+
 	const uint32 GroupIndex = 0;
-	FHairStrandsDatas* StrandsDatas = IsComponentValid() ? SourceComponent->GetGuideStrandsDatas(GroupIndex) : (DefaultSource != nullptr && DefaultSource->GetNumHairGroups() > 0) ? &DefaultSource->HairGroupsData[0].HairRenderData : nullptr;
-	FHairStrandsRestResource* StrandsRestResource = IsComponentValid() ? SourceComponent->GetGuideStrandsRestResource(GroupIndex) : (DefaultSource != nullptr && DefaultSource->GetNumHairGroups() > 0) ? DefaultSource->HairGroupsData[0].HairStrandsRestResource : nullptr;
-	FHairStrandsDeformedResource* StrandsDeformedResource = IsComponentValid() ? SourceComponent->GetGuideStrandsDeformedResource(GroupIndex) : nullptr;
 
-	if (StrandsDatas == nullptr || StrandsRestResource == nullptr || StrandsDeformedResource == nullptr)
+	OutStrandsDatas = nullptr;
+	OutStrandsRestResource = nullptr;
+	OutStrandsDeformedResource = nullptr;
+
+	if (IsComponentValid())
+	{
+		OutStrandsDatas = SourceComponent->GetGuideStrandsDatas(GroupIndex);
+		OutStrandsRestResource = SourceComponent->GetGuideStrandsRestResource(GroupIndex);
+		OutStrandsDeformedResource = SourceComponent->GetGuideStrandsDeformedResource(GroupIndex);
+	}
+	else if (DefaultSource != nullptr && DefaultSource->GetNumHairGroups() > 0)
+	{
+		OutStrandsDatas = &DefaultSource->HairGroupsData[0].HairRenderData;
+		OutStrandsRestResource = DefaultSource->HairGroupsData[0].HairStrandsRestResource;
+	}
+}
+
+bool UNiagaraDataInterfaceHairStrands::InitPerInstanceData(void* PerInstanceData, FNiagaraSystemInstance* SystemInstance)
+{
+	FNDIHairStrandsData* InstanceData = new (PerInstanceData) FNDIHairStrandsData();
+	check(InstanceData);
+
+	FHairStrandsDatas* StrandsDatas =  nullptr;
+	FHairStrandsRestResource* StrandsRestResource = nullptr;
+	FHairStrandsDeformedResource* StrandsDeformedResource = nullptr;
+
+	ExtractDatasAndResources(SystemInstance, StrandsDatas, StrandsRestResource, StrandsDeformedResource);
+
+	if (StrandsDatas == nullptr || StrandsRestResource == nullptr)
 	{
 		UE_LOG(LogHairStrands, Log, TEXT("Hair Strands data interface has no valid asset. Failed InitPerInstanceData - %s %d"), *GetFullName(), DefaultSource);
 		return false;
@@ -541,68 +449,36 @@ bool UNiagaraDataInterfaceHairStrands::InitPerInstanceData(void* PerInstanceData
 	else
 	{ 
 		FNDIHairStrandsBuffer* HairStrandsBuffer = new FNDIHairStrandsBuffer;
-		FNDICollisionGridBuffer* CurrentGridBuffer = new FNDICollisionGridBuffer;
-		FNDICollisionGridBuffer* DestinationGridBuffer = new FNDICollisionGridBuffer;
 		HairStrandsBuffer->SetHairAsset(StrandsDatas, StrandsRestResource, StrandsDeformedResource);
-		
 
 		// Push instance data to RT
 		{
-			const uint32 LocalNumStrands = StrandsDatas->GetNumCurves();
-			const uint8 LocalStrandSize = static_cast<uint8>(StrandSize);
-			const FBox& LocalStrandsBox = StrandsDatas->BoundingBox;
+			const uint32 NumStrands = StrandsDatas->GetNumCurves();
+			const uint8 StrandsSize = static_cast<uint8>(StrandSize);
+			const FBox& StrandsBox = StrandsDatas->BoundingBox;
 
-			const FVector GridExtent = LocalStrandsBox.GetExtent();
-			const FVector GridLengths((LocalStrandsBox.Max.X - LocalStrandsBox.Min.X) / (GridSizeX-1), 
-									  (LocalStrandsBox.Max.Y - LocalStrandsBox.Min.Y) / (GridSizeY-1), 
-									  (LocalStrandsBox.Max.Z - LocalStrandsBox.Min.Z) / (GridSizeZ-1));
-			const float GridLength = GridLengths.GetMax();
+			const FVector BoxExtent = StrandsBox.GetExtent();
+			const FVector4 BoxOrigin(-BoxExtent.X, -BoxExtent.Y, -BoxExtent.Z, 0.0);
 
-			const FUintVector4 GridSize(GridSizeX, GridSizeY, GridSizeZ,0);
-			const FVector4 GridOrigin(-GridExtent.X, -GridExtent.Y, -GridExtent.Z, GridLength);
-			CurrentGridBuffer->SetGridSize(GridSize);
-			DestinationGridBuffer->SetGridSize(GridSize);
-
-			const float LocalStrandDensity = StrandDensity;
-			const float LocalRootThickness = RootThickness;
-			const float LocalTipThickness = TipThickness;
-
-			const FTransform WorldOffset(-LocalStrandsBox.GetCenter());
-			const FMatrix WorldTransform = WorldOffset.ToInverseMatrixWithScale() * SourceTransform;
+			const FTransform WorldOffset(-StrandsBox.GetCenter());
+			const FMatrix WorldTransform = WorldOffset.ToInverseMatrixWithScale();
 
 			InstanceData->WorldTransform = IsComponentValid() ? WorldTransform * SourceComponent->GetComponentToWorld().ToMatrixWithScale() : 
 				WorldTransform * SystemInstance->GetComponent()->GetComponentToWorld().ToMatrixWithScale();
-
-		//UE_LOG(LogHairStrands, Log, TEXT("Num Strands = %d | Strand Size = %d | Num Vertices = %d | Min = %f %f %f | Max = %f %f %f | Transform = %s"), LocalNumStrands, LocalStrandSize,
-				//StrandsDatas->GetNumPoints(), LocalStrandsBox.Min[0], LocalStrandsBox.Min[1], LocalStrandsBox.Min[2], LocalStrandsBox.Max[0], LocalStrandsBox.Max[1], LocalStrandsBox.Max[2], *InstanceData->WorldTransform.ToString());
-
-			TSet<int> RT_OutputShaderStages = OutputShaderStages;
-			TSet<int> RT_IterationShaderStages = IterationShaderStages;
-
-			InstanceData->StrandSize = LocalStrandSize;
-			InstanceData->StrandDensity = LocalStrandDensity;
-			InstanceData->RootThickness = LocalRootThickness;
-			InstanceData->TipThickness = LocalTipThickness;
-			InstanceData->NumStrands = LocalNumStrands;
+			InstanceData->StrandSize = StrandsSize;
+			InstanceData->NumStrands = NumStrands;
 			InstanceData->HairStrandsBuffer = HairStrandsBuffer;
-			InstanceData->CurrentGridBuffer = CurrentGridBuffer;
-			InstanceData->DestinationGridBuffer = DestinationGridBuffer;
-			InstanceData->GridSize = GridSize;
-			InstanceData->GridOrigin = GridOrigin;
+			InstanceData->BoxOrigin = BoxOrigin;
+
+			//UE_LOG(LogHairStrands, Log, TEXT("Num Strands = %d | Strand Size = %d | Num Vertices = %d | Min = %f %f %f | Max = %f %f %f | Transform = %s"), LocalNumStrands, LocalStrandSize,
+			//	StrandsDatas->GetNumPoints(), LocalStrandsBox.Min[0], LocalStrandsBox.Min[1], LocalStrandsBox.Min[2], LocalStrandsBox.Max[0], LocalStrandsBox.Max[1], LocalStrandsBox.Max[2], *InstanceData->WorldTransform.ToString());
 
 			FNDIHairStrandsProxy* ThisProxy = GetProxyAs<FNDIHairStrandsProxy>();
 			ENQUEUE_RENDER_COMMAND(FNiagaraDIPushInitialInstanceDataToRT) (
-				[ThisProxy, RT_OutputShaderStages, RT_IterationShaderStages, InstanceID = SystemInstance->GetId(), HairStrandsBuffer, CurrentGridBuffer, DestinationGridBuffer, LocalNumStrands, LocalStrandSize, LocalStrandDensity, LocalRootThickness, LocalTipThickness, GridSize, GridOrigin](FRHICommandListImmediate& CmdList)
+				[ThisProxy, InstanceID = SystemInstance->GetId(), HairStrandsBuffer, NumStrands, StrandsSize, BoxOrigin, WorldTransform](FRHICommandListImmediate& CmdList)
 			{
-				ThisProxy->OutputShaderStages = RT_OutputShaderStages;
-				ThisProxy->IterationShaderStages = RT_IterationShaderStages;
-				ThisProxy->SetElementCount(GridSize.X*GridSize.Y*GridSize.Z);
-
 				HairStrandsBuffer->InitResource();
-				CurrentGridBuffer->InitResource();
-				DestinationGridBuffer->InitResource();
-
-				ThisProxy->InitializePerInstanceData(InstanceID, HairStrandsBuffer, CurrentGridBuffer, DestinationGridBuffer, LocalNumStrands, LocalStrandSize, LocalStrandDensity, LocalRootThickness, LocalTipThickness, GridOrigin,GridSize);
+				ThisProxy->InitializePerInstanceData(InstanceID, HairStrandsBuffer, NumStrands, StrandsSize, BoxOrigin, WorldTransform);
 			}
 			);
 		}
@@ -615,27 +491,19 @@ void UNiagaraDataInterfaceHairStrands::DestroyPerInstanceData(void* PerInstanceD
 	FNDIHairStrandsData* InstanceData = static_cast<FNDIHairStrandsData*>(PerInstanceData);
 	InstanceData->~FNDIHairStrandsData();
 
-	if (InstanceData->HairStrandsBuffer && InstanceData->CurrentGridBuffer && InstanceData->DestinationGridBuffer)
+	if (InstanceData->HairStrandsBuffer)
 	{
 		FNDIHairStrandsProxy* ThisProxy = GetProxyAs<FNDIHairStrandsProxy>();
 		FNDIHairStrandsBuffer* HairStrandsBuffer = InstanceData->HairStrandsBuffer;
-		FNDICollisionGridBuffer* CurrentGridBuffer = InstanceData->CurrentGridBuffer;
-		FNDICollisionGridBuffer* DestinationGridBuffer = InstanceData->DestinationGridBuffer;
 		ENQUEUE_RENDER_COMMAND(FNiagaraDIDestroyInstanceData) (
-			[ThisProxy, HairStrandsBuffer, CurrentGridBuffer, DestinationGridBuffer, InstanceID = SystemInstance->GetId(), Batcher = SystemInstance->GetBatcher()](FRHICommandListImmediate& CmdList)
+			[ThisProxy, HairStrandsBuffer, InstanceID = SystemInstance->GetId(), Batcher = SystemInstance->GetBatcher()](FRHICommandListImmediate& CmdList)
 		{
 			HairStrandsBuffer->ReleaseResource();
-			CurrentGridBuffer->ReleaseResource();
-			DestinationGridBuffer->ReleaseResource();
 			ThisProxy->DestroyPerInstanceData(Batcher, InstanceID);
 			delete HairStrandsBuffer;
-			delete CurrentGridBuffer;
-			delete DestinationGridBuffer;
 		}
 		);
 		InstanceData->HairStrandsBuffer = nullptr;
-		InstanceData->CurrentGridBuffer = nullptr;
-		InstanceData->DestinationGridBuffer = nullptr;
 	}
 }
 
@@ -643,14 +511,17 @@ bool UNiagaraDataInterfaceHairStrands::PerInstanceTick(void* PerInstanceData, FN
 {
 	FNDIHairStrandsData* InstanceData = static_cast<FNDIHairStrandsData*>(PerInstanceData);
 
-	const uint32 GroupIndex = 0;
-	FHairStrandsDatas* StrandsDatas = IsComponentValid() ? 
-		SourceComponent->GetGuideStrandsDatas(GroupIndex) : (DefaultSource != nullptr && DefaultSource->GetNumHairGroups() > 0) ? &DefaultSource->HairGroupsData[0].HairRenderData : nullptr;
+	FHairStrandsDatas* StrandsDatas = nullptr;
+	FHairStrandsRestResource* StrandsRestResource = nullptr;
+	FHairStrandsDeformedResource* StrandsDeformedResource = nullptr;
+
+	ExtractDatasAndResources(SystemInstance, StrandsDatas, StrandsRestResource, StrandsDeformedResource);
+
 	bool RequireReset = false;
 	if (StrandsDatas)
 	{
 		const FTransform WorldOffset(-StrandsDatas->BoundingBox.GetCenter());
-		const FMatrix WorldTransform = WorldOffset.ToInverseMatrixWithScale() * SourceTransform;
+		const FMatrix WorldTransform = WorldOffset.ToInverseMatrixWithScale();
 
 		if (IsComponentValid())
 		{
@@ -675,17 +546,9 @@ bool UNiagaraDataInterfaceHairStrands::CopyToInternal(UNiagaraDataInterface* Des
 
 	UNiagaraDataInterfaceHairStrands* OtherTyped = CastChecked<UNiagaraDataInterfaceHairStrands>(Destination);
 	OtherTyped->StrandSize = StrandSize;
-	OtherTyped->StrandDensity = StrandDensity;
-	OtherTyped->RootThickness = RootThickness;
-	OtherTyped->TipThickness = TipThickness;
-	OtherTyped->GroomAsset = GroomAsset;
 	OtherTyped->SourceActor= SourceActor;
 	OtherTyped->SourceComponent = SourceComponent;
 	OtherTyped->DefaultSource = DefaultSource;
-	OtherTyped->SourceTransform = SourceTransform;
-	OtherTyped->GridSizeX = GridSizeX;
-	OtherTyped->GridSizeY = GridSizeY;
-	OtherTyped->GridSizeZ = GridSizeZ;
 
 	return true;
 }
@@ -699,10 +562,7 @@ bool UNiagaraDataInterfaceHairStrands::Equals(const UNiagaraDataInterface* Other
 	const UNiagaraDataInterfaceHairStrands* OtherTyped = CastChecked<const UNiagaraDataInterfaceHairStrands>(Other);
 
 	return  (OtherTyped->StrandSize == StrandSize)
-		&& (OtherTyped->StrandDensity == StrandDensity) && (OtherTyped->RootThickness == RootThickness)
-		&& (OtherTyped->TipThickness == TipThickness) && (OtherTyped->GroomAsset == GroomAsset) && (OtherTyped->SourceTransform == SourceTransform) 
-		&& (OtherTyped->SourceActor == SourceActor) && (OtherTyped->SourceComponent == SourceComponent) && (OtherTyped->DefaultSource == DefaultSource)
-		&& (OtherTyped->GridSizeX == GridSizeX) && (OtherTyped->GridSizeY == GridSizeY) && (OtherTyped->GridSizeZ == GridSizeZ);
+		&& (OtherTyped->SourceActor == SourceActor) && (OtherTyped->SourceComponent == SourceComponent) && (OtherTyped->DefaultSource == DefaultSource);
 }
 
 void UNiagaraDataInterfaceHairStrands::PostInitProperties()
@@ -727,16 +587,6 @@ void UNiagaraDataInterfaceHairStrands::GetFunctions(TArray<FNiagaraFunctionSigna
 
 		OutFunctions.Add(Sig);
 	}
-	{
-		FNiagaraFunctionSignature Sig;
-		Sig.Name = GetStrandDensityName;
-		Sig.bMemberFunction = true;
-		Sig.bRequiresContext = false;
-		Sig.Inputs.Add(FNiagaraVariable(FNiagaraTypeDefinition(GetClass()), TEXT("Hair Strands")));
-		Sig.Outputs.Add(FNiagaraVariable(FNiagaraTypeDefinition::GetFloatDef(), TEXT("Strand Density")));
-
-		OutFunctions.Add(Sig);
-	}
 	{
 		FNiagaraFunctionSignature Sig;
 		Sig.Name = GetStrandSizeName;
@@ -747,26 +597,6 @@ void UNiagaraDataInterfaceHairStrands::GetFunctions(TArray<FNiagaraFunctionSigna
 
 		OutFunctions.Add(Sig);
 	}
-	{
-		FNiagaraFunctionSignature Sig;
-		Sig.Name = GetRootThicknessName;
-		Sig.bMemberFunction = true;
-		Sig.bRequiresContext = false;
-		Sig.Inputs.Add(FNiagaraVariable(FNiagaraTypeDefinition(GetClass()), TEXT("Hair Strands")));
-		Sig.Outputs.Add(FNiagaraVariable(FNiagaraTypeDefinition::GetFloatDef(), TEXT("Root Thickness")));
-
-		OutFunctions.Add(Sig);
-	}
-	{
-		FNiagaraFunctionSignature Sig;
-		Sig.Name = GetTipThicknessName;
-		Sig.bMemberFunction = true;
-		Sig.bRequiresContext = false;
-		Sig.Inputs.Add(FNiagaraVariable(FNiagaraTypeDefinition(GetClass()), TEXT("Hair Strands")));
-		Sig.Outputs.Add(FNiagaraVariable(FNiagaraTypeDefinition::GetFloatDef(), TEXT("Tip Thickness")));
-
-		OutFunctions.Add(Sig);
-	}
 	{
 		FNiagaraFunctionSignature Sig;
 		Sig.Name = GetWorldTransformName;
@@ -856,6 +686,7 @@ void UNiagaraDataInterfaceHairStrands::GetFunctions(TArray<FNiagaraFunctionSigna
 		Sig.Inputs.Add(FNiagaraVariable(FNiagaraTypeDefinition(GetClass()), TEXT("Hair Strands")));
 		Sig.Inputs.Add(FNiagaraVariable(FNiagaraTypeDefinition::GetIntDef(), TEXT("Node Index")));
 		Sig.Inputs.Add(FNiagaraVariable(FNiagaraTypeDefinition::GetVec3Def(), TEXT("Node Position")));
+		Sig.Inputs.Add(FNiagaraVariable(FNiagaraTypeDefinition::GetIntDef(), TEXT("Node Offset")));
 		Sig.Outputs.Add(FNiagaraVariable(FNiagaraTypeDefinition::GetFloatDef(), TEXT("Edge Length")));
 
 		OutFunctions.Add(Sig);
@@ -999,49 +830,6 @@ void UNiagaraDataInterfaceHairStrands::GetFunctions(TArray<FNiagaraFunctionSigna
 
 		OutFunctions.Add(Sig);
 	}
-	{
-		FNiagaraFunctionSignature Sig;
-		Sig.Name = BuildCollisionGridName;
-		Sig.bMemberFunction = true;
-		Sig.bRequiresContext = false;
-		Sig.Inputs.Add(FNiagaraVariable(FNiagaraTypeDefinition(GetClass()), TEXT("Hair Strands")));
-		Sig.Inputs.Add(FNiagaraVariable(FNiagaraTypeDefinition::GetVec3Def(), TEXT("Node Position")));
-		Sig.Inputs.Add(FNiagaraVariable(FNiagaraTypeDefinition::GetVec3Def(), TEXT("Node Velocity")));
-		Sig.Inputs.Add(FNiagaraVariable(FNiagaraTypeDefinition::GetFloatDef(), TEXT("Node Mass")));
-		Sig.Inputs.Add(FNiagaraVariable(FNiagaraTypeDefinition::GetVec3Def(), TEXT("Node GradientX")));
-		Sig.Inputs.Add(FNiagaraVariable(FNiagaraTypeDefinition::GetVec3Def(), TEXT("Node GradientY")));
-		Sig.Inputs.Add(FNiagaraVariable(FNiagaraTypeDefinition::GetVec3Def(), TEXT("Node GradientZ")));
-		Sig.Outputs.Add(FNiagaraVariable(FNiagaraTypeDefinition::GetBoolDef(), TEXT("Build Status")));
-
-		OutFunctions.Add(Sig);
-	}
-	{
-		FNiagaraFunctionSignature Sig;
-		Sig.Name = ProjectCollisionGridName;
-		Sig.bMemberFunction = true;
-		Sig.bRequiresContext = false;
-		Sig.Inputs.Add(FNiagaraVariable(FNiagaraTypeDefinition(GetClass()), TEXT("Hair Strands")));
-		Sig.Inputs.Add(FNiagaraVariable(FNiagaraTypeDefinition::GetIntDef(), TEXT("Grid Hash")));
-		Sig.Outputs.Add(FNiagaraVariable(FNiagaraTypeDefinition::GetBoolDef(), TEXT("Project Status")));
-
-		OutFunctions.Add(Sig);
-	}
-	{
-		FNiagaraFunctionSignature Sig;
-		Sig.Name = QueryCollisionGridName;
-		Sig.bMemberFunction = true;
-		Sig.bRequiresContext = false;
-		Sig.Inputs.Add(FNiagaraVariable(FNiagaraTypeDefinition(GetClass()), TEXT("Hair Strands")));
-		Sig.Inputs.Add(FNiagaraVariable(FNiagaraTypeDefinition::GetVec3Def(), TEXT("Node Position")));
-		Sig.Outputs.Add(FNiagaraVariable(FNiagaraTypeDefinition::GetVec3Def(), TEXT("Grid Velocity")));
-		Sig.Outputs.Add(FNiagaraVariable(FNiagaraTypeDefinition::GetVec3Def(), TEXT("Grid Tangent")));
-		Sig.Outputs.Add(FNiagaraVariable(FNiagaraTypeDefinition::GetFloatDef(), TEXT("Grid Mass")));
-		Sig.Outputs.Add(FNiagaraVariable(FNiagaraTypeDefinition::GetVec3Def(), TEXT("Grid GradientX")));
-		Sig.Outputs.Add(FNiagaraVariable(FNiagaraTypeDefinition::GetVec3Def(), TEXT("Grid GradientY")));
-		Sig.Outputs.Add(FNiagaraVariable(FNiagaraTypeDefinition::GetVec3Def(), TEXT("Grid GradientZ")));
-
-		OutFunctions.Add(Sig);
-	}
 	{
 		FNiagaraFunctionSignature Sig;
 		Sig.Name = GetBoxCenterName;
@@ -1064,7 +852,7 @@ void UNiagaraDataInterfaceHairStrands::GetFunctions(TArray<FNiagaraFunctionSigna
 	}
 	{
 		FNiagaraFunctionSignature Sig;
-		Sig.Name = SetupStretchSpringMaterialName;
+		Sig.Name = SetupDistanceSpringMaterialName;
 		Sig.bMemberFunction = true;
 		Sig.bRequiresContext = false;
 		Sig.Inputs.Add(FNiagaraVariable(FNiagaraTypeDefinition(GetClass()), TEXT("Hair Strands")));
@@ -1072,6 +860,7 @@ void UNiagaraDataInterfaceHairStrands::GetFunctions(TArray<FNiagaraFunctionSigna
 		Sig.Inputs.Add(FNiagaraVariable(FNiagaraTypeDefinition::GetFloatDef(), TEXT("Node Thickness")));
 		Sig.Inputs.Add(FNiagaraVariable(FNiagaraTypeDefinition::GetFloatDef(), TEXT("Rest Length")));
 		Sig.Inputs.Add(FNiagaraVariable(FNiagaraTypeDefinition::GetFloatDef(), TEXT("Delta Time")));
+		Sig.Inputs.Add(FNiagaraVariable(FNiagaraTypeDefinition::GetIntDef(), TEXT("Node Offset")));
 		Sig.Inputs.Add(FNiagaraVariable(FNiagaraTypeDefinition::GetFloatDef(), TEXT("Material Damping")));
 		Sig.Outputs.Add(FNiagaraVariable(FNiagaraTypeDefinition::GetFloatDef(), TEXT("Material Compliance")));
 		Sig.Outputs.Add(FNiagaraVariable(FNiagaraTypeDefinition::GetFloatDef(), TEXT("Material Weight")));
@@ -1080,12 +869,13 @@ void UNiagaraDataInterfaceHairStrands::GetFunctions(TArray<FNiagaraFunctionSigna
 	}
 	{
 		FNiagaraFunctionSignature Sig;
-		Sig.Name = SolveStretchSpringMaterialName;
+		Sig.Name = SolveDistanceSpringMaterialName;
 		Sig.bMemberFunction = true;
 		Sig.bRequiresContext = false;
 		Sig.Inputs.Add(FNiagaraVariable(FNiagaraTypeDefinition(GetClass()), TEXT("Hair Strands")));
 		Sig.Inputs.Add(FNiagaraVariable(FNiagaraTypeDefinition::GetFloatDef(), TEXT("Rest Length")));
 		Sig.Inputs.Add(FNiagaraVariable(FNiagaraTypeDefinition::GetFloatDef(), TEXT("Delta Time")));
+		Sig.Inputs.Add(FNiagaraVariable(FNiagaraTypeDefinition::GetIntDef(), TEXT("Node Offset")));
 		Sig.Inputs.Add(FNiagaraVariable(FNiagaraTypeDefinition::GetFloatDef(), TEXT("Material Damping")));
 		Sig.Inputs.Add(FNiagaraVariable(FNiagaraTypeDefinition::GetFloatDef(), TEXT("Material Compliance")));
 		Sig.Inputs.Add(FNiagaraVariable(FNiagaraTypeDefinition::GetFloatDef(), TEXT("Material Weight")));
@@ -1095,7 +885,7 @@ void UNiagaraDataInterfaceHairStrands::GetFunctions(TArray<FNiagaraFunctionSigna
 	}
 	{
 		FNiagaraFunctionSignature Sig;
-		Sig.Name = ProjectStretchSpringMaterialName;
+		Sig.Name = ProjectDistanceSpringMaterialName;
 		Sig.bMemberFunction = true;
 		Sig.bRequiresContext = false;
 		Sig.Inputs.Add(FNiagaraVariable(FNiagaraTypeDefinition(GetClass()), TEXT("Hair Strands")));
@@ -1103,14 +893,14 @@ void UNiagaraDataInterfaceHairStrands::GetFunctions(TArray<FNiagaraFunctionSigna
 		Sig.Inputs.Add(FNiagaraVariable(FNiagaraTypeDefinition::GetFloatDef(), TEXT("Node Thickness")));
 		Sig.Inputs.Add(FNiagaraVariable(FNiagaraTypeDefinition::GetFloatDef(), TEXT("Rest Length")));
 		Sig.Inputs.Add(FNiagaraVariable(FNiagaraTypeDefinition::GetFloatDef(), TEXT("Delta Time")));
+		Sig.Inputs.Add(FNiagaraVariable(FNiagaraTypeDefinition::GetIntDef(), TEXT("Node Offset")));
 		Sig.Outputs.Add(FNiagaraVariable(FNiagaraTypeDefinition::GetFloatDef(), TEXT("Material Multiplier")));
 
 		OutFunctions.Add(Sig);
 	}
-
 	{
 		FNiagaraFunctionSignature Sig;
-		Sig.Name = SetupBendSpringMaterialName;
+		Sig.Name = SetupAngularSpringMaterialName;
 		Sig.bMemberFunction = true;
 		Sig.bRequiresContext = false;
 		Sig.Inputs.Add(FNiagaraVariable(FNiagaraTypeDefinition(GetClass()), TEXT("Hair Strands")));
@@ -1126,10 +916,11 @@ void UNiagaraDataInterfaceHairStrands::GetFunctions(TArray<FNiagaraFunctionSigna
 	}
 	{
 		FNiagaraFunctionSignature Sig;
-		Sig.Name = SolveBendSpringMaterialName;
+		Sig.Name = SolveAngularSpringMaterialName;
 		Sig.bMemberFunction = true;
 		Sig.bRequiresContext = false;
 		Sig.Inputs.Add(FNiagaraVariable(FNiagaraTypeDefinition(GetClass()), TEXT("Hair Strands")));
+		Sig.Inputs.Add(FNiagaraVariable(FNiagaraTypeDefinition::GetFloatDef(), TEXT("Rest Length")));
 		Sig.Inputs.Add(FNiagaraVariable(FNiagaraTypeDefinition::GetVec3Def(), TEXT("Rest Direction")));
 		Sig.Inputs.Add(FNiagaraVariable(FNiagaraTypeDefinition::GetFloatDef(), TEXT("Delta Time")));
 		Sig.Inputs.Add(FNiagaraVariable(FNiagaraTypeDefinition::GetFloatDef(), TEXT("Material Damping")));
@@ -1141,7 +932,7 @@ void UNiagaraDataInterfaceHairStrands::GetFunctions(TArray<FNiagaraFunctionSigna
 	}
 	{
 		FNiagaraFunctionSignature Sig;
-		Sig.Name = ProjectBendSpringMaterialName;
+		Sig.Name = ProjectAngularSpringMaterialName;
 		Sig.bMemberFunction = true;
 		Sig.bRequiresContext = false;
 		Sig.Inputs.Add(FNiagaraVariable(FNiagaraTypeDefinition(GetClass()), TEXT("Hair Strands")));
@@ -1305,15 +1096,29 @@ void UNiagaraDataInterfaceHairStrands::GetFunctions(TArray<FNiagaraFunctionSigna
 		Sig.Inputs.Add(FNiagaraVariable(FNiagaraTypeDefinition(GetClass()), TEXT("Hair Strands")));
 		Sig.Outputs.Add(FNiagaraVariable(FNiagaraTypeDefinition::GetBoolDef(), TEXT("Update Status")));
 
+		OutFunctions.Add(Sig);
+	}
+	{
+		FNiagaraFunctionSignature Sig;
+		Sig.Name = ComputeAirDragForceName;
+		Sig.bMemberFunction = true;
+		Sig.bRequiresContext = false;
+		Sig.Inputs.Add(FNiagaraVariable(FNiagaraTypeDefinition(GetClass()), TEXT("Hair Strands")));
+		Sig.Inputs.Add(FNiagaraVariable(FNiagaraTypeDefinition::GetFloatDef(), TEXT("Air Density")));
+		Sig.Inputs.Add(FNiagaraVariable(FNiagaraTypeDefinition::GetFloatDef(), TEXT("Air Viscosity")));
+		Sig.Inputs.Add(FNiagaraVariable(FNiagaraTypeDefinition::GetFloatDef(), TEXT("Air Drag")));
+		Sig.Inputs.Add(FNiagaraVariable(FNiagaraTypeDefinition::GetVec3Def(), TEXT("Air Velocity")));
+		Sig.Inputs.Add(FNiagaraVariable(FNiagaraTypeDefinition::GetFloatDef(), TEXT("Node Thickness")));
+		Sig.Inputs.Add(FNiagaraVariable(FNiagaraTypeDefinition::GetVec3Def(), TEXT("Node Position")));
+		Sig.Inputs.Add(FNiagaraVariable(FNiagaraTypeDefinition::GetVec3Def(), TEXT("Node Velocity")));
+		Sig.Outputs.Add(FNiagaraVariable(FNiagaraTypeDefinition::GetVec3Def(), TEXT("Drag Force")));
+
 		OutFunctions.Add(Sig);
 	}
 }
 
 DEFINE_NDI_DIRECT_FUNC_BINDER(UNiagaraDataInterfaceHairStrands, GetNumStrands);
-DEFINE_NDI_DIRECT_FUNC_BINDER(UNiagaraDataInterfaceHairStrands, GetStrandDensity);
 DEFINE_NDI_DIRECT_FUNC_BINDER(UNiagaraDataInterfaceHairStrands, GetStrandSize);
-DEFINE_NDI_DIRECT_FUNC_BINDER(UNiagaraDataInterfaceHairStrands, GetRootThickness);
-DEFINE_NDI_DIRECT_FUNC_BINDER(UNiagaraDataInterfaceHairStrands, GetTipThickness);
 DEFINE_NDI_DIRECT_FUNC_BINDER(UNiagaraDataInterfaceHairStrands, GetWorldTransform);
 DEFINE_NDI_DIRECT_FUNC_BINDER(UNiagaraDataInterfaceHairStrands, GetWorldInverse);
 DEFINE_NDI_DIRECT_FUNC_BINDER(UNiagaraDataInterfaceHairStrands, GetPointPosition);
@@ -1330,10 +1135,6 @@ DEFINE_NDI_DIRECT_FUNC_BINDER(UNiagaraDataInterfaceHairStrands, AttachNodeOrient
 DEFINE_NDI_DIRECT_FUNC_BINDER(UNiagaraDataInterfaceHairStrands, UpdatePointPosition);
 DEFINE_NDI_DIRECT_FUNC_BINDER(UNiagaraDataInterfaceHairStrands, ResetPointPosition);
 
-DEFINE_NDI_DIRECT_FUNC_BINDER(UNiagaraDataInterfaceHairStrands, BuildCollisionGrid);
-DEFINE_NDI_DIRECT_FUNC_BINDER(UNiagaraDataInterfaceHairStrands, QueryCollisionGrid);
-DEFINE_NDI_DIRECT_FUNC_BINDER(UNiagaraDataInterfaceHairStrands, ProjectCollisionGrid);
-
 DEFINE_NDI_DIRECT_FUNC_BINDER(UNiagaraDataInterfaceHairStrands, GetBoxExtent);
 DEFINE_NDI_DIRECT_FUNC_BINDER(UNiagaraDataInterfaceHairStrands, GetBoxCenter);
 
@@ -1342,13 +1143,13 @@ DEFINE_NDI_DIRECT_FUNC_BINDER(UNiagaraDataInterfaceHairStrands, AdvectNodeOrient
 DEFINE_NDI_DIRECT_FUNC_BINDER(UNiagaraDataInterfaceHairStrands, UpdateLinearVelocity);
 DEFINE_NDI_DIRECT_FUNC_BINDER(UNiagaraDataInterfaceHairStrands, UpdateAngularVelocity);
 
-DEFINE_NDI_DIRECT_FUNC_BINDER(UNiagaraDataInterfaceHairStrands, SetupStretchSpringMaterial);
-DEFINE_NDI_DIRECT_FUNC_BINDER(UNiagaraDataInterfaceHairStrands, SolveStretchSpringMaterial);
-DEFINE_NDI_DIRECT_FUNC_BINDER(UNiagaraDataInterfaceHairStrands, ProjectStretchSpringMaterial);
+DEFINE_NDI_DIRECT_FUNC_BINDER(UNiagaraDataInterfaceHairStrands, SetupDistanceSpringMaterial);
+DEFINE_NDI_DIRECT_FUNC_BINDER(UNiagaraDataInterfaceHairStrands, SolveDistanceSpringMaterial);
+DEFINE_NDI_DIRECT_FUNC_BINDER(UNiagaraDataInterfaceHairStrands, ProjectDistanceSpringMaterial);
 
-DEFINE_NDI_DIRECT_FUNC_BINDER(UNiagaraDataInterfaceHairStrands, SetupBendSpringMaterial);
-DEFINE_NDI_DIRECT_FUNC_BINDER(UNiagaraDataInterfaceHairStrands, SolveBendSpringMaterial);
-DEFINE_NDI_DIRECT_FUNC_BINDER(UNiagaraDataInterfaceHairStrands, ProjectBendSpringMaterial);
+DEFINE_NDI_DIRECT_FUNC_BINDER(UNiagaraDataInterfaceHairStrands, SetupAngularSpringMaterial);
+DEFINE_NDI_DIRECT_FUNC_BINDER(UNiagaraDataInterfaceHairStrands, SolveAngularSpringMaterial);
+DEFINE_NDI_DIRECT_FUNC_BINDER(UNiagaraDataInterfaceHairStrands, ProjectAngularSpringMaterial);
 
 DEFINE_NDI_DIRECT_FUNC_BINDER(UNiagaraDataInterfaceHairStrands, SetupStretchRodMaterial);
 DEFINE_NDI_DIRECT_FUNC_BINDER(UNiagaraDataInterfaceHairStrands, SolveStretchRodMaterial);
@@ -1363,6 +1164,7 @@ DEFINE_NDI_DIRECT_FUNC_BINDER(UNiagaraDataInterfaceHairStrands, ProjectStaticCol
 
 DEFINE_NDI_DIRECT_FUNC_BINDER(UNiagaraDataInterfaceHairStrands, ComputeRestDirection);
 DEFINE_NDI_DIRECT_FUNC_BINDER(UNiagaraDataInterfaceHairStrands, UpdateNodeOrientation);
+DEFINE_NDI_DIRECT_FUNC_BINDER(UNiagaraDataInterfaceHairStrands, ComputeAirDragForce);
 
 void UNiagaraDataInterfaceHairStrands::GetVMExternalFunction(const FVMExternalFunctionBindingInfo& BindingInfo, void* InstanceData, FVMExternalFunction &OutFunc)
 {
@@ -1371,26 +1173,11 @@ void UNiagaraDataInterfaceHairStrands::GetVMExternalFunction(const FVMExternalFu
 		check(BindingInfo.GetNumInputs() == 1 && BindingInfo.GetNumOutputs() == 1);
 		NDI_FUNC_BINDER(UNiagaraDataInterfaceHairStrands, GetNumStrands)::Bind(this, OutFunc);
 	}
-	else if (BindingInfo.Name == GetStrandDensityName)
-	{
-		check(BindingInfo.GetNumInputs() == 1 && BindingInfo.GetNumOutputs() == 1);
-		NDI_FUNC_BINDER(UNiagaraDataInterfaceHairStrands, GetStrandDensity)::Bind(this, OutFunc);
-	}
 	else if (BindingInfo.Name == GetStrandSizeName)
 	{
 		check(BindingInfo.GetNumInputs() == 1 && BindingInfo.GetNumOutputs() == 1);
 		NDI_FUNC_BINDER(UNiagaraDataInterfaceHairStrands, GetStrandSize)::Bind(this, OutFunc);
 	}
-	else if (BindingInfo.Name == GetRootThicknessName)
-	{
-		check(BindingInfo.GetNumInputs() == 1 && BindingInfo.GetNumOutputs() == 1);
-		NDI_FUNC_BINDER(UNiagaraDataInterfaceHairStrands, GetRootThickness)::Bind(this, OutFunc);
-	}
-	else if (BindingInfo.Name == GetTipThicknessName)
-	{
-		check(BindingInfo.GetNumInputs() == 1 && BindingInfo.GetNumOutputs() == 1);
-		NDI_FUNC_BINDER(UNiagaraDataInterfaceHairStrands, GetTipThickness)::Bind(this, OutFunc);
-	}
 	else if (BindingInfo.Name == GetWorldTransformName)
 	{
 		check(BindingInfo.GetNumInputs() == 1 && BindingInfo.GetNumOutputs() == 16);
@@ -1428,7 +1215,7 @@ void UNiagaraDataInterfaceHairStrands::GetVMExternalFunction(const FVMExternalFu
 	}
 	else if (BindingInfo.Name == ComputeEdgeLengthName)
 	{
-		check(BindingInfo.GetNumInputs() == 5 && BindingInfo.GetNumOutputs() == 1);
+		check(BindingInfo.GetNumInputs() == 6 && BindingInfo.GetNumOutputs() == 1);
 		NDI_FUNC_BINDER(UNiagaraDataInterfaceHairStrands, ComputeEdgeLength)::Bind(this, OutFunc);
 	}
 	else if (BindingInfo.Name == ComputeEdgeRotationName)
@@ -1486,21 +1273,6 @@ void UNiagaraDataInterfaceHairStrands::GetVMExternalFunction(const FVMExternalFu
 		check(BindingInfo.GetNumInputs() == 10 && BindingInfo.GetNumOutputs() == 3);
 		NDI_FUNC_BINDER(UNiagaraDataInterfaceHairStrands, UpdateAngularVelocity)::Bind(this, OutFunc);
 	}
-	else if (BindingInfo.Name == BuildCollisionGridName)
-	{
-		check(BindingInfo.GetNumInputs() == 17 && BindingInfo.GetNumOutputs() == 1);
-		NDI_FUNC_BINDER(UNiagaraDataInterfaceHairStrands, BuildCollisionGrid)::Bind(this, OutFunc);
-	}
-	else if (BindingInfo.Name == QueryCollisionGridName)
-	{
-		check(BindingInfo.GetNumInputs() == 4 && BindingInfo.GetNumOutputs() == 16);
-		NDI_FUNC_BINDER(UNiagaraDataInterfaceHairStrands, QueryCollisionGrid)::Bind(this, OutFunc);
-	}
-	else if (BindingInfo.Name == ProjectCollisionGridName)
-	{
-		check(BindingInfo.GetNumInputs() == 2 && BindingInfo.GetNumOutputs() == 1);
-		NDI_FUNC_BINDER(UNiagaraDataInterfaceHairStrands, ProjectCollisionGrid)::Bind(this, OutFunc);
-	}
 	else if (BindingInfo.Name == GetBoxExtentName)
 	{
 		check(BindingInfo.GetNumInputs() == 1 && BindingInfo.GetNumOutputs() == 3);
@@ -1511,35 +1283,35 @@ void UNiagaraDataInterfaceHairStrands::GetVMExternalFunction(const FVMExternalFu
 		check(BindingInfo.GetNumInputs() == 1 && BindingInfo.GetNumOutputs() == 3);
 		NDI_FUNC_BINDER(UNiagaraDataInterfaceHairStrands, GetBoxCenter)::Bind(this, OutFunc);
 	}
-	else if (BindingInfo.Name == SetupStretchSpringMaterialName)
+	else if (BindingInfo.Name == SetupDistanceSpringMaterialName)
 	{
-		check(BindingInfo.GetNumInputs() == 6 && BindingInfo.GetNumOutputs() == 3);
-		NDI_FUNC_BINDER(UNiagaraDataInterfaceHairStrands, SetupStretchSpringMaterial)::Bind(this, OutFunc);
+		check(BindingInfo.GetNumInputs() == 7 && BindingInfo.GetNumOutputs() == 3);
+		NDI_FUNC_BINDER(UNiagaraDataInterfaceHairStrands, SetupDistanceSpringMaterial)::Bind(this, OutFunc);
 	}
-	else if (BindingInfo.Name == SolveStretchSpringMaterialName)
+	else if (BindingInfo.Name == SolveDistanceSpringMaterialName)
 	{
-		check(BindingInfo.GetNumInputs() == 7 && BindingInfo.GetNumOutputs() == 1);
-		NDI_FUNC_BINDER(UNiagaraDataInterfaceHairStrands, SolveStretchSpringMaterial)::Bind(this, OutFunc);
+		check(BindingInfo.GetNumInputs() == 8 && BindingInfo.GetNumOutputs() == 1);
+		NDI_FUNC_BINDER(UNiagaraDataInterfaceHairStrands, SolveDistanceSpringMaterial)::Bind(this, OutFunc);
 	}
-	else if (BindingInfo.Name == ProjectStretchSpringMaterialName)
+	else if (BindingInfo.Name == ProjectDistanceSpringMaterialName)
 	{
-		check(BindingInfo.GetNumInputs() == 5 && BindingInfo.GetNumOutputs() == 1);
-		NDI_FUNC_BINDER(UNiagaraDataInterfaceHairStrands, ProjectStretchSpringMaterial)::Bind(this, OutFunc);
+		check(BindingInfo.GetNumInputs() == 6 && BindingInfo.GetNumOutputs() == 1);
+		NDI_FUNC_BINDER(UNiagaraDataInterfaceHairStrands, ProjectDistanceSpringMaterial)::Bind(this, OutFunc);
 	}
-	else if (BindingInfo.Name == SetupBendSpringMaterialName)
+	else if (BindingInfo.Name == SetupAngularSpringMaterialName)
 	{
 		check(BindingInfo.GetNumInputs() == 6 && BindingInfo.GetNumOutputs() == 5);
-		NDI_FUNC_BINDER(UNiagaraDataInterfaceHairStrands, SetupBendSpringMaterial)::Bind(this, OutFunc);
+		NDI_FUNC_BINDER(UNiagaraDataInterfaceHairStrands, SetupAngularSpringMaterial)::Bind(this, OutFunc);
 	}
-	else if (BindingInfo.Name == SolveBendSpringMaterialName)
+	else if (BindingInfo.Name == SolveAngularSpringMaterialName)
 	{
-		check(BindingInfo.GetNumInputs() == 11 && BindingInfo.GetNumOutputs() == 3);
-		NDI_FUNC_BINDER(UNiagaraDataInterfaceHairStrands, SolveBendSpringMaterial)::Bind(this, OutFunc);
+		check(BindingInfo.GetNumInputs() == 12 && BindingInfo.GetNumOutputs() == 3);
+		NDI_FUNC_BINDER(UNiagaraDataInterfaceHairStrands, SolveAngularSpringMaterial)::Bind(this, OutFunc);
 	}
-	else if (BindingInfo.Name == ProjectBendSpringMaterialName)
+	else if (BindingInfo.Name == ProjectAngularSpringMaterialName)
 	{
 		check(BindingInfo.GetNumInputs() == 8 && BindingInfo.GetNumOutputs() == 3);
-		NDI_FUNC_BINDER(UNiagaraDataInterfaceHairStrands, ProjectBendSpringMaterial)::Bind(this, OutFunc);
+		NDI_FUNC_BINDER(UNiagaraDataInterfaceHairStrands, ProjectAngularSpringMaterial)::Bind(this, OutFunc);
 	}
 	else if (BindingInfo.Name == SetupStretchRodMaterialName)
 	{
@@ -1591,6 +1363,11 @@ void UNiagaraDataInterfaceHairStrands::GetVMExternalFunction(const FVMExternalFu
 		check(BindingInfo.GetNumInputs() == 1 && BindingInfo.GetNumOutputs() == 1);
 		NDI_FUNC_BINDER(UNiagaraDataInterfaceHairStrands, UpdateNodeOrientation)::Bind(this, OutFunc);
 	}
+	else if (BindingInfo.Name == ComputeAirDragForceName)
+	{
+		check(BindingInfo.GetNumInputs() == 14 && BindingInfo.GetNumOutputs() == 3);
+		NDI_FUNC_BINDER(UNiagaraDataInterfaceHairStrands, ComputeAirDragForce)::Bind(this, OutFunc);
+	}
 }
 
 void WriteTransform(const FMatrix& ToWrite, FVectorVMContext& Context)
@@ -1655,46 +1432,11 @@ void UNiagaraDataInterfaceHairStrands::GetStrandSize(FVectorVMContext& Context)
 	}
 }
 
-void UNiagaraDataInterfaceHairStrands::GetStrandDensity(FVectorVMContext& Context)
-{
-	VectorVM::FUserPtrHandler<FNDIHairStrandsData> InstData(Context);
-	VectorVM::FExternalFuncRegisterHandler<float> OutStrandDensity(Context);
-
-	for (int32 InstanceIdx = 0; InstanceIdx < Context.NumInstances; ++InstanceIdx)
-	{
-		*OutStrandDensity.GetDestAndAdvance() = InstData->StrandDensity;
-	}
-}
-
-void UNiagaraDataInterfaceHairStrands::GetRootThickness(FVectorVMContext& Context)
-{
-	VectorVM::FUserPtrHandler<FNDIHairStrandsData> InstData(Context);
-	VectorVM::FExternalFuncRegisterHandler<float> OutRootThickness(Context);
-
-	for (int32 InstanceIdx = 0; InstanceIdx < Context.NumInstances; ++InstanceIdx)
-	{
-		*OutRootThickness.GetDestAndAdvance() = InstData->RootThickness;
-	}
-}
-
-void UNiagaraDataInterfaceHairStrands::GetTipThickness(FVectorVMContext& Context)
-{
-	VectorVM::FUserPtrHandler<FNDIHairStrandsData> InstData(Context);
-	VectorVM::FExternalFuncRegisterHandler<float> OutTipThickness(Context);
-
-	for (int32 InstanceIdx = 0; InstanceIdx < Context.NumInstances; ++InstanceIdx)
-	{
-		*OutTipThickness.GetDestAndAdvance() = InstData->TipThickness;
-	}
-}
-
 void UNiagaraDataInterfaceHairStrands::GetWorldTransform(FVectorVMContext& Context)
 {
 	VectorVM::FUserPtrHandler<FNDIHairStrandsData> InstData(Context);
 	FMatrix WorldTransform = InstData->WorldTransform;
 
-	//UE_LOG(LogHairStrands, Warning, TEXT("Get Hair Transform : %s"), *InstData->WorldTransform.ToString());
-
 	WriteTransform(WorldTransform, Context);
 }
 
@@ -1730,11 +1472,9 @@ void UNiagaraDataInterfaceHairStrands::GetBoxExtent(FVectorVMContext& Context)
 
 	for (int32 InstanceIdx = 0; InstanceIdx < Context.NumInstances; ++InstanceIdx)
 	{
-		*OutBoxExtentX.GetDestAndAdvance() = -InstData->GridOrigin.X;
-		*OutBoxExtentY.GetDestAndAdvance() = -InstData->GridOrigin.Y;
-		*OutBoxExtentZ.GetDestAndAdvance() = -InstData->GridOrigin.Z;
-
-		//UE_LOG(LogHairStrands, Warning, TEXT("Get Box Extent : %f %f %f"), InstData->GridOrigin.X, InstData->GridOrigin.Y, InstData->GridOrigin.Z);
+		*OutBoxExtentX.GetDestAndAdvance() = -InstData->BoxOrigin.X;
+		*OutBoxExtentY.GetDestAndAdvance() = -InstData->BoxOrigin.Y;
+		*OutBoxExtentZ.GetDestAndAdvance() = -InstData->BoxOrigin.Z;
 	}
 }
 
@@ -1823,47 +1563,32 @@ void UNiagaraDataInterfaceHairStrands::UpdateAngularVelocity(FVectorVMContext& C
 	// @todo : implement function for cpu 
 }
 
-void UNiagaraDataInterfaceHairStrands::BuildCollisionGrid(FVectorVMContext& Context)
-{
-	// @todo : implement function for cpu 
-}
-
-void UNiagaraDataInterfaceHairStrands::QueryCollisionGrid(FVectorVMContext& Context)
-{
-	// @todo : implement function for cpu 
-}
-
-void UNiagaraDataInterfaceHairStrands::ProjectCollisionGrid(FVectorVMContext& Context)
+void UNiagaraDataInterfaceHairStrands::SetupDistanceSpringMaterial(FVectorVMContext& Context)
 {
 	// @todo : implement function for cpu 
 }
 
-void UNiagaraDataInterfaceHairStrands::SetupStretchSpringMaterial(FVectorVMContext& Context)
+void UNiagaraDataInterfaceHairStrands::SolveDistanceSpringMaterial(FVectorVMContext& Context)
 {
 	// @todo : implement function for cpu 
 }
 
-void UNiagaraDataInterfaceHairStrands::SolveStretchSpringMaterial(FVectorVMContext& Context)
+void UNiagaraDataInterfaceHairStrands::ProjectDistanceSpringMaterial(FVectorVMContext& Context)
 {
 	// @todo : implement function for cpu 
 }
 
-void UNiagaraDataInterfaceHairStrands::ProjectStretchSpringMaterial(FVectorVMContext& Context)
+void UNiagaraDataInterfaceHairStrands::SetupAngularSpringMaterial(FVectorVMContext& Context)
 {
 	// @todo : implement function for cpu 
 }
 
-void UNiagaraDataInterfaceHairStrands::SetupBendSpringMaterial(FVectorVMContext& Context)
+void UNiagaraDataInterfaceHairStrands::SolveAngularSpringMaterial(FVectorVMContext& Context)
 {
 	// @todo : implement function for cpu 
 }
 
-void UNiagaraDataInterfaceHairStrands::SolveBendSpringMaterial(FVectorVMContext& Context)
-{
-	// @todo : implement function for cpu 
-}
-
-void UNiagaraDataInterfaceHairStrands::ProjectBendSpringMaterial(FVectorVMContext& Context)
+void UNiagaraDataInterfaceHairStrands::ProjectAngularSpringMaterial(FVectorVMContext& Context)
 {
 	// @todo : implement function for cpu 
 }
@@ -1919,6 +1644,11 @@ void UNiagaraDataInterfaceHairStrands::ProjectStaticCollisionConstraint(FVectorV
 	// @todo : implement function for cpu 
 }
 
+void UNiagaraDataInterfaceHairStrands::ComputeAirDragForce(FVectorVMContext& Context)
+{
+	// @todo : implement function for cpu 
+}
+
 bool UNiagaraDataInterfaceHairStrands::GetFunctionHLSL(const FName& DefinitionFunctionName, FString InstanceFunctionName, FNiagaraDataInterfaceGPUParamInfo& ParamInfo, FString& OutHLSL)
 {
 	FNDIHairStrandsParametersName ParamNames(ParamInfo.DataInterfaceHLSLSymbol);
@@ -1927,34 +1657,17 @@ bool UNiagaraDataInterfaceHairStrands::GetFunctionHLSL(const FName& DefinitionFu
 		{TEXT("InstanceFunctionName"), InstanceFunctionName},
 		{TEXT("NumStrandsName"), ParamNames.NumStrandsName},
 		{TEXT("StrandSizeName"), ParamNames.StrandSizeName},
-		{TEXT("StrandDensityName"), ParamNames.StrandDensityName},
-		{TEXT("RootThicknessName"), ParamNames.RootThicknessName},
-		{TEXT("TipThicknessName"), ParamNames.TipThicknessName},
 		{TEXT("WorldTransformName"), ParamNames.WorldTransformName},
 		{TEXT("WorldInverseName"), ParamNames.WorldInverseName},
 		{TEXT("WorldRotationName"), ParamNames.WorldRotationName},
 		{TEXT("PointsPositionsBufferName"), ParamNames.PointsPositionsBufferName},
 		{TEXT("CurvesOffsetsBufferName"), ParamNames.CurvesOffsetsBufferName},
 		{TEXT("RestPositionsBufferName"), ParamNames.RestPositionsBufferName},
-		{TEXT("GridCurrentBufferName"), ParamNames.GridCurrentBufferName},
-		{TEXT("GridDestinationBufferName"), ParamNames.GridDestinationBufferName},
-		{TEXT("GridOriginName"), ParamNames.GridOriginName},
-		{TEXT("GridSizeName"), ParamNames.GridSizeName},
+		{TEXT("BoxOriginName"), ParamNames.BoxOriginName},
 		{TEXT("HairStrandsContextName"), TEXT("DIHAIRSTRANDS_MAKE_CONTEXT(") + ParamInfo.DataInterfaceHLSLSymbol + TEXT(")")},
 	};
 
-	if (DefinitionFunctionName == GetStrandDensityName)
-	{
-		static const TCHAR *FormatSample = TEXT(R"(
-		void {InstanceFunctionName}(out float OutStrandDensity)
-		{
-			OutStrandDensity = {StrandDensityName};
-		}
-		)");
-		OutHLSL += FString::Format(FormatSample, ArgsSample);
-		return true;
-	}
-	else if (DefinitionFunctionName == GetStrandSizeName)
+	if (DefinitionFunctionName == GetStrandSizeName)
 	{
 		static const TCHAR *FormatSample = TEXT(R"(
 		void {InstanceFunctionName}(out int OutStrandSize)
@@ -1976,28 +1689,6 @@ bool UNiagaraDataInterfaceHairStrands::GetFunctionHLSL(const FName& DefinitionFu
 		OutHLSL += FString::Format(FormatSample, ArgsSample);
 		return true;
 	}
-	else if (DefinitionFunctionName == GetRootThicknessName)
-	{
-		static const TCHAR *FormatSample = TEXT(R"(
-		void {InstanceFunctionName}(out float OutRootThickness)
-		{
-			OutRootThickness = {RootThicknessName};
-		}
-		)");
-		OutHLSL += FString::Format(FormatSample, ArgsSample);
-		return true;
-	}
-	else if (DefinitionFunctionName == GetTipThicknessName)
-	{
-		static const TCHAR *FormatSample = TEXT(R"(
-		void {InstanceFunctionName}(out float OutTipThickness)
-		{
-			OutTipThickness = {TipThicknessName};
-		}
-		)");
-		OutHLSL += FString::Format(FormatSample, ArgsSample);
-		return true;
-	}
 	else if (DefinitionFunctionName == GetWorldTransformName)
 	{
 		static const TCHAR *FormatSample = TEXT(R"(
@@ -2078,9 +1769,9 @@ bool UNiagaraDataInterfaceHairStrands::GetFunctionHLSL(const FName& DefinitionFu
 	else if (DefinitionFunctionName == ComputeEdgeLengthName)
 	{
 		static const TCHAR *FormatSample = TEXT(R"(
-			void {InstanceFunctionName} (in int NodeIndex, in float3 NodePosition, out float OutEdgeLength)
+			void {InstanceFunctionName} (in int NodeIndex, in float3 NodePosition, in int NodeOffset, out float OutEdgeLength)
 			{
-				{HairStrandsContextName} DIHairStrands_ComputeEdgeLength(DIContext,NodeIndex,NodePosition,OutEdgeLength);
+				{HairStrandsContextName} DIHairStrands_ComputeEdgeLength(DIContext,NodeIndex,NodePosition,NodeOffset,OutEdgeLength);
 			}
 			)");
 		OutHLSL += FString::Format(FormatSample, ArgsSample);
@@ -2213,39 +1904,6 @@ bool UNiagaraDataInterfaceHairStrands::GetFunctionHLSL(const FName& DefinitionFu
 		OutHLSL += FString::Format(FormatSample, ArgsSample);
 		return true;
 	}
-	else if (DefinitionFunctionName == BuildCollisionGridName)
-	{
-		static const TCHAR *FormatSample = TEXT(R"(
-				void {InstanceFunctionName} (in float3 NodePosition, in float3 NodeVelocity, in float NodeMass, in float3 NodeGradientX, in float3 NodeGradientY, in float3 NodeGradientZ, out bool OutBuildStatus)
-				{
-					{HairStrandsContextName} DIHairStrands_BuildCollisionGrid(DIContext,NodePosition,NodeVelocity,NodeMass,NodeGradientX,NodeGradientY,NodeGradientZ,OutBuildStatus);
-				}
-				)");
-		OutHLSL += FString::Format(FormatSample, ArgsSample);
-		return true;
-	}
-	else if (DefinitionFunctionName == QueryCollisionGridName)
-	{
-		static const TCHAR *FormatSample = TEXT(R"(
-				void {InstanceFunctionName} (in float3 NodePosition, out float3 OutGridVelocity, out float3 OutGridTangent, out float GridOutput, out float3 OutGridGradientX, out float3 OutGridGradientY, out float3 OutGridGradientZ )
-				{
-					{HairStrandsContextName} DIHairStrands_QueryCollisionGrid(DIContext,NodePosition,OutGridVelocity,OutGridTangent,GridOutput,OutGridGradientX,OutGridGradientY,OutGridGradientZ);
-				}
-				)");
-		OutHLSL += FString::Format(FormatSample, ArgsSample);
-		return true;
-	}
-	else if (DefinitionFunctionName == ProjectCollisionGridName)
-	{
-		static const TCHAR *FormatSample = TEXT(R"(
-				void {InstanceFunctionName} (in int GridHash, out bool OutProjectStatus)
-				{
-					{HairStrandsContextName} DIHairStrands_ProjectCollisionGrid(DIContext,GridHash,OutProjectStatus);
-				}
-				)");
-		OutHLSL += FString::Format(FormatSample, ArgsSample);
-		return true;
-	}
 	else if (DefinitionFunctionName == GetBoxExtentName)
 	{
 		static const TCHAR *FormatSample = TEXT(R"(
@@ -2268,71 +1926,107 @@ bool UNiagaraDataInterfaceHairStrands::GetFunctionHLSL(const FName& DefinitionFu
 		OutHLSL += FString::Format(FormatSample, ArgsSample);
 		return true;
 	}
-	else if (DefinitionFunctionName == SetupStretchSpringMaterialName)
+	else if (DefinitionFunctionName == SetupDistanceSpringMaterialName)
 	{
 		static const TCHAR *FormatSample = TEXT(R"(
 				void {InstanceFunctionName} (in float YoungModulus, in float RodThickness, 
-in float RestLength, in float DeltaTime, in float MaterialDamping, out float OutMaterialCompliance, out float OutMaterialWeight, out float OutMaterialMultiplier)
+in float RestLength, in float DeltaTime, in int NodeOffset, in float MaterialDamping, out float OutMaterialCompliance, out float OutMaterialWeight, out float OutMaterialMultiplier)
 				{
-					{HairStrandsContextName} SetupStretchSpringMaterial(DIContext.StrandSize,YoungModulus,RodThickness,RestLength,DeltaTime,false,MaterialDamping,OutMaterialCompliance,OutMaterialWeight,OutMaterialMultiplier);
+					{HairStrandsContextName} 
+					if(NodeOffset == 0)
+					{
+						SetupStretchSpringMaterial(DIContext.StrandSize,YoungModulus,RodThickness,RestLength,DeltaTime,false,MaterialDamping,OutMaterialCompliance,OutMaterialWeight,OutMaterialMultiplier);
+					}
+					else if( NodeOffset == 1)
+					{
+						SetupBendSpringMaterial(DIContext.StrandSize,YoungModulus,RodThickness,RestLength,DeltaTime,false,MaterialDamping,OutMaterialCompliance,OutMaterialWeight,OutMaterialMultiplier);
+					}
+					else if( NodeOffset == 2)
+					{
+						SetupTwistSpringMaterial(DIContext.StrandSize,YoungModulus,RodThickness,RestLength,DeltaTime,false,MaterialDamping,OutMaterialCompliance,OutMaterialWeight,OutMaterialMultiplier);
+					}
 				}
 				)");
 		OutHLSL += FString::Format(FormatSample, ArgsSample);
 		return true;
 	}
-	else if (DefinitionFunctionName == SolveStretchSpringMaterialName)
+	else if (DefinitionFunctionName == SolveDistanceSpringMaterialName)
 	{
 		static const TCHAR *FormatSample = TEXT(R"(
-				void {InstanceFunctionName} (in float RestLength, in float DeltaTime, in float MaterialDamping, 
+				void {InstanceFunctionName} (in float RestLength, in float DeltaTime, in int NodeOffset, in float MaterialDamping, 
 		in float MaterialCompliance, in float MaterialWeight, in float MaterialMultiplier, out float OutMaterialMultiplier)
 				{
-					{HairStrandsContextName} SolveStretchSpringMaterial(DIContext.StrandSize,RestLength,DeltaTime,MaterialDamping,MaterialCompliance,MaterialWeight,MaterialMultiplier,OutMaterialMultiplier);
+					{HairStrandsContextName} 
+					if(NodeOffset == 0)
+					{
+						SolveStretchSpringMaterial(DIContext.StrandSize,RestLength,DeltaTime,MaterialDamping,MaterialCompliance,MaterialWeight,MaterialMultiplier,OutMaterialMultiplier);
+					}
+					else if(NodeOffset == 1)
+					{
+						SolveBendSpringMaterial(DIContext.StrandSize,RestLength,DeltaTime,MaterialDamping,MaterialCompliance,MaterialWeight,MaterialMultiplier,OutMaterialMultiplier);
+					}
+					else if(NodeOffset == 2)
+					{
+						SolveTwistSpringMaterial(DIContext.StrandSize,RestLength,DeltaTime,MaterialDamping,MaterialCompliance,MaterialWeight,MaterialMultiplier,OutMaterialMultiplier);
+					}
 				}
 				)");
 		OutHLSL += FString::Format(FormatSample, ArgsSample);
 		return true;
 	}
-	else if (DefinitionFunctionName == ProjectStretchSpringMaterialName)
+	else if (DefinitionFunctionName == ProjectDistanceSpringMaterialName)
 	{
 		static const TCHAR *FormatSample = TEXT(R"(
-				void {InstanceFunctionName} (in float YoungModulus, in float RodThickness, in float RestLength, in float DeltaTime, out float OutMaterialMultiplier)
+				void {InstanceFunctionName} (in float YoungModulus, in float RodThickness, in float RestLength, in float DeltaTime, in int NodeOffset, out float OutMaterialMultiplier)
 				{
-					{HairStrandsContextName} ProjectStretchSpringMaterial(DIContext.StrandSize,YoungModulus,RodThickness,RestLength,DeltaTime,OutMaterialMultiplier);
+					{HairStrandsContextName} 
+					if(NodeOffset == 0)
+					{
+						ProjectStretchSpringMaterial(DIContext.StrandSize,YoungModulus,RodThickness,RestLength,DeltaTime,OutMaterialMultiplier);
+					}
+					if(NodeOffset == 1)
+					{
+						ProjectBendSpringMaterial(DIContext.StrandSize,YoungModulus,RodThickness,RestLength,DeltaTime,OutMaterialMultiplier);
+					}
+					if(NodeOffset == 2)
+					{
+						ProjectTwistSpringMaterial(DIContext.StrandSize,YoungModulus,RodThickness,RestLength,DeltaTime,OutMaterialMultiplier);
+					}
 				}
 				)");
 		OutHLSL += FString::Format(FormatSample, ArgsSample);
 		return true;
 	}
-	else if (DefinitionFunctionName == SetupBendSpringMaterialName)
+	else if (DefinitionFunctionName == SetupAngularSpringMaterialName)
 	{
 		static const TCHAR *FormatSample = TEXT(R"(
 					void {InstanceFunctionName} (in float YoungModulus, in float RodThickness, 
 	in float RestLength, in float DeltaTime, in float MaterialDamping, out float OutMaterialCompliance, out float OutMaterialWeight, out float3 OutMaterialMultiplier)
 					{
-						{HairStrandsContextName} SetupBendSpringMaterial(DIContext.StrandSize,YoungModulus,RodThickness,RestLength,DeltaTime,false,MaterialDamping,OutMaterialCompliance,OutMaterialWeight,OutMaterialMultiplier);
+						{HairStrandsContextName} SetupAngularSpringMaterial(DIContext.StrandSize,YoungModulus,RodThickness,RestLength,DeltaTime,false,MaterialDamping,OutMaterialCompliance,OutMaterialWeight,OutMaterialMultiplier);
 					}
 					)");
 		OutHLSL += FString::Format(FormatSample, ArgsSample);
 		return true;
 	}
-	else if (DefinitionFunctionName == SolveBendSpringMaterialName)
+	else if (DefinitionFunctionName == SolveAngularSpringMaterialName)
 	{
 		static const TCHAR *FormatSample = TEXT(R"(
-					void {InstanceFunctionName} (in float3 RestDirection, in float DeltaTime, in float MaterialDamping, 
+					void {InstanceFunctionName} (in float RestLength, in float3 RestDirection, in float DeltaTime, in float MaterialDamping, 
 			in float MaterialCompliance, in float MaterialWeight, in float3 MaterialMultiplier, out float3 OutMaterialMultiplier)
 					{
-						{HairStrandsContextName} SolveBendSpringMaterial(DIContext.StrandSize,RestDirection,DeltaTime,MaterialDamping,MaterialCompliance,MaterialWeight,MaterialMultiplier,OutMaterialMultiplier);
+						{HairStrandsContextName} SolveAngularSpringMaterial(DIContext.StrandSize,RestLength, RestDirection,DeltaTime,MaterialDamping,MaterialCompliance,MaterialWeight,MaterialMultiplier,OutMaterialMultiplier);
 					}
 					)");
 		OutHLSL += FString::Format(FormatSample, ArgsSample);
 		return true;
 	}
-	else if (DefinitionFunctionName == ProjectBendSpringMaterialName)
+	else if (DefinitionFunctionName == ProjectAngularSpringMaterialName)
 	{
 		static const TCHAR *FormatSample = TEXT(R"(
 					void {InstanceFunctionName} (in float YoungModulus, in float RodThickness, in float RestLength, in float3 RestDirection, in float DeltaTime, out float OutMaterialMultiplier)
 					{
-						{HairStrandsContextName} ProjectBendSpringMaterial(DIContext.StrandSize,YoungModulus,RodThickness,RestLength,RestDirection,DeltaTime,OutMaterialMultiplier);
+						{HairStrandsContextName} ProjectAngularSpringMaterial(DIContext.StrandSize,YoungModulus,RodThickness,RestLength,RestDirection,DeltaTime,OutMaterialMultiplier);
 					}
 					)");
 		OutHLSL += FString::Format(FormatSample, ArgsSample);
@@ -2459,6 +2153,18 @@ in float RestLength, in float DeltaTime, in float MaterialDamping, out float Out
 		OutHLSL += FString::Format(FormatSample, ArgsSample);
 		return true;
 	}
+	else if (DefinitionFunctionName == ComputeAirDragForceName)
+	{
+		static const TCHAR *FormatSample = TEXT(R"(
+						void {InstanceFunctionName} ( in float AirDensity, in float AirViscosity, in float AirDrag, 
+	in float3 AirVelocity, in float NodeThickness, in float3 NodePosition, in float3 NodeVelocity, out float3 OutAirDrag )
+						{
+							{HairStrandsContextName} DIHairStrands_ComputeAirDragForce(DIContext,AirDensity,AirViscosity,AirDrag,AirVelocity,NodeThickness,NodePosition,NodeVelocity,OutAirDrag);
+						}
+						)");
+		OutHLSL += FString::Format(FormatSample, ArgsSample);
+		return true;
+	}
 
 	OutHLSL += TEXT("\n");
 	return false;
@@ -2486,6 +2192,9 @@ void UNiagaraDataInterfaceHairStrands::ProvidePerInstanceDataForRenderThread(voi
 	FNDIHairStrandsData* RenderThreadData = static_cast<FNDIHairStrandsData*>(DataForRenderThread);
 
 	RenderThreadData->WorldTransform = GameThreadData->WorldTransform;
+	RenderThreadData->HairStrandsBuffer = GameThreadData->HairStrandsBuffer;
+	RenderThreadData->NumStrands = GameThreadData->NumStrands;
+	RenderThreadData->StrandSize = GameThreadData->StrandSize;
 }
 
 FNiagaraDataInterfaceParametersCS*
@@ -2495,53 +2204,13 @@ UNiagaraDataInterfaceHairStrands::ConstructComputeParameters() const
 }
 
 void FNDIHairStrandsProxy::PreStage(FRHICommandList& RHICmdList, const FNiagaraDataInterfaceSetArgs& Context)
-{
-	FNDIHairStrandsData* ProxyData =
-		SystemInstancesToProxyData.Find(Context.SystemInstance);
-
-	if (ProxyData != nullptr)
-	{
-		if (!Context.IsIterationStage)
-		{
-			ProxyData->DestinationGridBuffer->ClearBuffers(RHICmdList);
-		}
-		else
-		{
-			ENQUEUE_RENDER_COMMAND(CopyCollisionGrid)(
-				[ProxyData](FRHICommandListImmediate& RHICmdList)
-			{
-				FRHICopyTextureInfo CopyInfo;
-				RHICmdList.CopyTexture(ProxyData->CurrentGridBuffer->GridDataBuffer.Buffer,
-					ProxyData->DestinationGridBuffer->GridDataBuffer.Buffer, CopyInfo);
-			});
-		}
-	}
-}
+{}
 
 void FNDIHairStrandsProxy::PostStage(FRHICommandList& RHICmdList, const FNiagaraDataInterfaceSetArgs& Context)
-{
-	FNDIHairStrandsData* ProxyData =
-		SystemInstancesToProxyData.Find(Context.SystemInstance);
-
-	if (Context.IsOutputStage)
-	{
-		if (ProxyData != nullptr)
-		{
-			ProxyData->SwapBuffers();
-		}
-	}
-}
+{}
 
 void FNDIHairStrandsProxy::ResetData(FRHICommandList& RHICmdList, const FNiagaraDataInterfaceSetArgs& Context)
-{
-	FNDIHairStrandsData* ProxyData = SystemInstancesToProxyData.Find(Context.SystemInstance);
-
-	if (ProxyData != nullptr)
-	{
-		ProxyData->CurrentGridBuffer->ClearBuffers(RHICmdList);
-		ProxyData->DestinationGridBuffer->ClearBuffers(RHICmdList);
-	}
-}
+{}
 
 
 #undef LOCTEXT_NAMESPACE
\ No newline at end of file
diff --git a/Engine/Plugins/Experimental/HairStrands/Source/HairStrandsNiagara/Private/NiagaraDataInterfacePhysicsAsset.cpp b/Engine/Plugins/Experimental/HairStrands/Source/HairStrandsNiagara/Private/NiagaraDataInterfacePhysicsAsset.cpp
--- a/Engine/Plugins/Experimental/HairStrands/Source/HairStrandsNiagara/Private/NiagaraDataInterfacePhysicsAsset.cpp
+++ b/Engine/Plugins/Experimental/HairStrands/Source/HairStrandsNiagara/Private/NiagaraDataInterfacePhysicsAsset.cpp
@@ -280,8 +280,8 @@ void FNDIPhysicsAssetBuffer::InitRHI()
 		CreateInternalBuffer<FVector4, FVector4, 1, EPixelFormat::PF_A32B32G32R32F, true>(AssetArrays->InverseTransform.Num(), AssetArrays->InverseTransform, InverseTransformBuffer);
 		CreateInternalBuffer<FVector4, FVector4, 1, EPixelFormat::PF_A32B32G32R32F, true>(AssetArrays->ElementExtent.Num(), AssetArrays->ElementExtent, ElementExtentBuffer);
 
-		UE_LOG(LogPhysicsAsset, Warning, TEXT("Num Capsules = %d | Num Spheres = %d | Num Boxes = %d"), AssetArrays->ElementOffsets.NumElements - AssetArrays->ElementOffsets.CapsuleOffset,
-			AssetArrays->ElementOffsets.CapsuleOffset - AssetArrays->ElementOffsets.SphereOffset, AssetArrays->ElementOffsets.SphereOffset - AssetArrays->ElementOffsets.BoxOffset);
+		//UE_LOG(LogPhysicsAsset, Warning, TEXT("Num Capsules = %d | Num Spheres = %d | Num Boxes = %d"), AssetArrays->ElementOffsets.NumElements - AssetArrays->ElementOffsets.CapsuleOffset,
+		//	AssetArrays->ElementOffsets.CapsuleOffset - AssetArrays->ElementOffsets.SphereOffset, AssetArrays->ElementOffsets.SphereOffset - AssetArrays->ElementOffsets.BoxOffset);
 	}
 }
 
diff --git a/Engine/Plugins/Experimental/HairStrands/Source/HairStrandsNiagara/Public/NiagaraDataInterfaceHairStrands.h b/Engine/Plugins/Experimental/HairStrands/Source/HairStrandsNiagara/Public/NiagaraDataInterfaceHairStrands.h
--- a/Engine/Plugins/Experimental/HairStrands/Source/HairStrandsNiagara/Public/NiagaraDataInterfaceHairStrands.h
+++ b/Engine/Plugins/Experimental/HairStrands/Source/HairStrandsNiagara/Public/NiagaraDataInterfaceHairStrands.h
@@ -22,31 +22,6 @@ enum class EHairStrandsSize : uint8
 	Size32 = 0x20 UMETA(DisplatName = "32")
 };
 
-/** Render buffers that will be used in hlsl functions */
-struct FNDICollisionGridBuffer : public FRenderResource
-{
-	/** Set the grid size */
-	void SetGridSize(const FUintVector4 GridSize);
-
-	/** Init the buffer */
-	virtual void InitRHI() override;
-
-	/** Release the buffer */
-	virtual void ReleaseRHI() override;
-
-	/** Clear all UAV*/
-	void ClearBuffers(FRHICommandList& RHICmdList);
-
-	/** Get the resource name */
-	virtual FString GetFriendlyName() const override { return TEXT("FNDICollisionGridBuffer"); }
-
-	/** Grid data texture */
-	FTextureRWBuffer3D GridDataBuffer;
-
-	/** Grid size that will be used for the collision*/
-	FUintVector4 GridSize;
-};
-
 /** Render buffers that will be used in hlsl functions */
 struct FNDIHairStrandsBuffer : public FRenderResource
 {
@@ -65,6 +40,9 @@ struct FNDIHairStrandsBuffer : public FRenderResource
 	/** Strand curves point offset buffer */
 	FRWBuffer CurvesOffsetsBuffer;
 
+	/** Deformed position buffer in case no ressource are there */
+	FRWBuffer PointsPositionsBuffer;
+
 	/** The strand asset datas from which to sample */
 	const FHairStrandsDatas* SourceDatas;
 
@@ -78,12 +56,6 @@ struct FNDIHairStrandsBuffer : public FRenderResource
 /** Data stored per strand base instance*/
 struct FNDIHairStrandsData
 {
-	/** Swap the current and the destination data */
-	void SwapBuffers();
-
-	/** Init the current and the destination data */
-	void InitBuffers();
-
 	/** Cached World transform. */
 	FMatrix WorldTransform;
 
@@ -93,40 +65,16 @@ struct FNDIHairStrandsData
 	/** Strand size */
 	int32 StrandSize;
 
-	/** Strand Density */
-	float StrandDensity;
-
-	/** Root Thickness */
-	float RootThickness;
-
-	/** Tip Thickness */
-	float TipThickness;
-
-	/** Grid Origin */
-	FVector4 GridOrigin;
-
-	/** Grid Size */
-	FUintVector4 GridSize;
+	/** Bounding box origin */
+	FVector4 BoxOrigin;
 
 	/** Strands Gpu buffer */
 	FNDIHairStrandsBuffer* HairStrandsBuffer;
-
-	/** Collision Grid Buffer A */
-	FNDICollisionGridBuffer* CollisionGridBufferA;
-
-	/** Collision Grid Buffer B */
-	FNDICollisionGridBuffer* CollisionGridBufferB;
-
-	/** Pointer to the current buffer */
-	FNDICollisionGridBuffer* CurrentGridBuffer;
-
-	/** Pointer to the destination buffer */
-	FNDICollisionGridBuffer* DestinationGridBuffer;
 };
 
 /** Data Interface for the strand base */
 UCLASS(EditInlineNew, Category = "Strands", meta = (DisplayName = "Hair Strands"))
-class HAIRSTRANDSNIAGARA_API UNiagaraDataInterfaceHairStrands : public UNiagaraDataInterfaceRWBase
+class HAIRSTRANDSNIAGARA_API UNiagaraDataInterfaceHairStrands : public UNiagaraDataInterface
 {
 	GENERATED_UCLASS_BODY()
 
@@ -136,18 +84,6 @@ public:
 	UPROPERTY(EditAnywhere, Category = "Spawn")
 	EHairStrandsSize StrandSize;
 
-	/** Density of each strand. */
-	UPROPERTY(EditAnywhere, Category = "Spawn")
-	float StrandDensity;
-
-	/** Strand Root thickness. */
-	UPROPERTY(EditAnywhere, Category = "Spawn")
-	float RootThickness;
-
-	/** Strand Tip thickness. */
-	UPROPERTY(EditAnywhere, Category = "Spawn")
-	float TipThickness;
-
 	/** Hair Strands Asset used to sample from when not overridden by a source actor from the scene. Also useful for previewing in the editor. */
 	UPROPERTY(EditAnywhere, Category = "Source")
 	UGroomAsset* DefaultSource;
@@ -156,28 +92,9 @@ public:
 	UPROPERTY(EditAnywhere, Category = "Source")
 	AActor* SourceActor;
 
-	/** Source transform to be applied to the asset. */
-	UPROPERTY(EditAnywhere, Category = "Source")
-	FMatrix SourceTransform;
-
-	/** Grid size along the X axis. */
-	UPROPERTY(EditAnywhere, Category = "Grid")
-	uint32 GridSizeX;
-
-	/** Grid size along the Y axis. */
-	UPROPERTY(EditAnywhere, Category = "Grid")
-	uint32 GridSizeY;
-
-	/** Grid size along the Z axis. */
-	UPROPERTY(EditAnywhere, Category = "Grid")
-	uint32 GridSizeZ;
-
 	/** The source component from which to sample */
 	TWeakObjectPtr<class UGroomComponent> SourceComponent;
 
-	/** The source asset from which to sample */
-	TWeakObjectPtr<class UGroomAsset> GroomAsset;
-
 	/** UObject Interface */
 	virtual void PostInitProperties() override;
 
@@ -198,24 +115,22 @@ public:
 	virtual void ProvidePerInstanceDataForRenderThread(void* DataForRenderThread, void* PerInstanceData, const FNiagaraSystemInstanceID& SystemInstance) override;
 	virtual void GetCommonHLSL(FString& OutHLSL) override;
 
+	/** Update the source component */
+	void UpdateSourceComponent(FNiagaraSystemInstance* SystemInstance);
+
 	/** Check if the component is Valid */
 	bool IsComponentValid() const;
 
+	/** Extract datas and resources */
+	void ExtractDatasAndResources(FNiagaraSystemInstance* SystemInstance, FHairStrandsDatas*& OutStrandsDatas,
+		FHairStrandsRestResource*& OutStrandsRestResource, FHairStrandsDeformedResource*& OutStrandsDeformedResource);
+
 	/** Get the number of strands */
 	void GetNumStrands(FVectorVMContext& Context);
 
 	/** Get the strand size  */
 	void GetStrandSize(FVectorVMContext& Context);
 
-	/** Get the strand density */
-	void GetStrandDensity(FVectorVMContext& Context);
-
-	/** Get the strand thickness at the root */
-	void GetRootThickness(FVectorVMContext& Context);
-
-	/** Get the strand thickness at the tip */
-	void GetTipThickness(FVectorVMContext& Context);
-
 	/** Get the world transform */
 	void GetWorldTransform(FVectorVMContext& Context);
 
@@ -273,38 +188,29 @@ public:
 	/** Update the node angular velocity based on the node orientation difference */
 	void UpdateAngularVelocity(FVectorVMContext& Context);
 
-	/** Build the collision grid */
-	void BuildCollisionGrid(FVectorVMContext& Context);
-
-	/** Query the collision grid */
-	void QueryCollisionGrid(FVectorVMContext& Context);
-
-	/** Project the collision grid */
-	void ProjectCollisionGrid(FVectorVMContext& Context);
-
 	/** Get the bounding box center */
 	void GetBoxCenter(FVectorVMContext& Context);
 
 	/** Get the bounding box extent */
 	void GetBoxExtent(FVectorVMContext& Context);
 
-	/** Setup the stretch spring material */
-	void SetupStretchSpringMaterial(FVectorVMContext& Context);
+	/** Setup the distance spring material */
+	void SetupDistanceSpringMaterial(FVectorVMContext& Context);
 
-	/** Solve the stretch spring material */
-	void SolveStretchSpringMaterial(FVectorVMContext& Context);
+	/** Solve the distance spring material */
+	void SolveDistanceSpringMaterial(FVectorVMContext& Context);
 
-	/** Project the stretch spring material */
-	void ProjectStretchSpringMaterial(FVectorVMContext& Context);
+	/** Project the distance spring material */
+	void ProjectDistanceSpringMaterial(FVectorVMContext& Context);
 
-	/** Setup the bend spring material */
-	void SetupBendSpringMaterial(FVectorVMContext& Context);
+	/** Setup the angular spring material */
+	void SetupAngularSpringMaterial(FVectorVMContext& Context);
 
-	/** Solve the bend spring material */
-	void SolveBendSpringMaterial(FVectorVMContext& Context);
+	/** Solve the angular spring material */
+	void SolveAngularSpringMaterial(FVectorVMContext& Context);
 
-	/** Project the bend spring material */
-	void ProjectBendSpringMaterial(FVectorVMContext& Context);
+	/** Project the angular spring material */
+	void ProjectAngularSpringMaterial(FVectorVMContext& Context);
 
 	/** Setup the stretch rod material */
 	void SetupStretchRodMaterial(FVectorVMContext& Context);
@@ -336,6 +242,9 @@ public:
 	/** Update the node orientation to match the bishop frame*/
 	void UpdateNodeOrientation(FVectorVMContext& Context);
 
+	/** Compute the air drag force */
+	void ComputeAirDragForce(FVectorVMContext& Context);
+
 	/** Name of the world transform */
 	static const FString WorldTransformName;
 
@@ -351,15 +260,6 @@ public:
 	/** Name of the strand size */
 	static const FString StrandSizeName;
 
-	/** Name of the strand density */
-	static const FString StrandDensityName;
-
-	/** Name of the root thickness */
-	static const FString RootThicknessName;
-
-	/** Name of the tip thickness */
-	static const FString TipThicknessName;
-
 	/** Name of the points positions buffer */
 	static const FString PointsPositionsBufferName;
 
@@ -369,17 +269,8 @@ public:
 	/** Name of the nodes positions buffer */
 	static const FString RestPositionsBufferName;
 
-	/** Name of the grid current buffer */
-	static const FString GridCurrentBufferName;
-
-	/** Name of the grid X velocity buffer */
-	static const FString GridDestinationBufferName;
-
-	/** Name of the grid origin  */
-	static const FString GridOriginName;
-
-	/** Name of the grid size */
-	static const FString GridSizeName;
+	/** Name of the box origin  */
+	static const FString BoxOriginName;
 
 protected:
 	/** Copy one niagara DI to this */
@@ -399,8 +290,7 @@ struct FNDIHairStrandsProxy : public FNiagaraDataInterfaceProxy
 	virtual void ConsumePerInstanceDataFromGameThread(void* PerInstanceData, const FNiagaraSystemInstanceID& Instance) override;
 
 	/** Initialize the Proxy data strands buffer */
-	void InitializePerInstanceData(const FNiagaraSystemInstanceID& SystemInstance, FNDIHairStrandsBuffer* StrandsBuffer, FNDICollisionGridBuffer* CollisionGridBufferA, FNDICollisionGridBuffer* CollisionGridBufferB, const uint32 NumStrands, const uint8 StrandSize,
-		const float StrandDensity, const float RootThickness, const float TipThickness, const FVector4& GridOrigin, const FUintVector4& GridSize);
+	void InitializePerInstanceData(const FNiagaraSystemInstanceID& SystemInstance, FNDIHairStrandsBuffer* StrandsBuffer, const uint32 NumStrands, const uint8 StrandSize, const FVector4& BoxOrigin, const FMatrix& WorldTransform);
 
 	/** Destroy the proxy data if necessary */
 	void DestroyPerInstanceData(NiagaraEmitterInstanceBatcher* Batcher, const FNiagaraSystemInstanceID& SystemInstance);
diff --git a/Engine/Plugins/FX/Niagara/Shaders/Private/NiagaraStrandsPhysics.ush b/Engine/Plugins/FX/Niagara/Shaders/Private/NiagaraStrandsPhysics.ush
--- a/Engine/Plugins/FX/Niagara/Shaders/Private/NiagaraStrandsPhysics.ush
+++ b/Engine/Plugins/FX/Niagara/Shaders/Private/NiagaraStrandsPhysics.ush
@@ -179,10 +179,10 @@ void SolveHookeSpringMaterial(const float InMaterialCompliance, const float InMa
 
  void ComputeStretchShearWeight(const float InMaterialCompliance, const float InRestLength, out float OutMaterialWeight)
 {
-	const float SchurDiagonal = ( SharedInverseMass[GGroupThreadId.x] + SharedInverseMass[GGroupThreadId.x-1] ) / (InRestLength*InRestLength) + 
-					SharedInverseInertia[GGroupThreadId.x-1]*4.0 + InMaterialCompliance;
 	//const float SchurDiagonal = ( SharedInverseMass[GGroupThreadId.x] + SharedInverseMass[GGroupThreadId.x-1] ) / (InRestLength*InRestLength) + 
-	//				 + InMaterialCompliance;
+	//				SharedInverseInertia[GGroupThreadId.x-1]*4.0 + InMaterialCompliance;
+	const float SchurDiagonal = ( SharedInverseMass[GGroupThreadId.x] + SharedInverseMass[GGroupThreadId.x-1] ) / (InRestLength*InRestLength) + 
+					 + InMaterialCompliance;
 	OutMaterialWeight = ( SchurDiagonal != 0.0 ) ? 1.0 / SchurDiagonal : 0.0;
 }
 
@@ -221,12 +221,12 @@ void UpdateStretchShearMultiplier(const float InMaterialCompliance, const float
 	SharedNodePosition[GGroupThreadId.x] += PositionDelta * SharedInverseMass[GGroupThreadId.x];
 	SharedNodePosition[GGroupThreadId.x-1] -= PositionDelta * SharedInverseMass[GGroupThreadId.x-1];
 
-	const float4 qebar = float4(-q0.y, q0.x, -q0.w, q0.z);	
-	const float4 QuaternionDelta = -2.0 * float4( DeltaLambda.xyz * qebar.w + cross(DeltaLambda.xyz, qebar.xyz),
-					- dot(DeltaLambda.xyz, qebar.xyz) );
+	//const float4 qebar = float4(-q0.y, q0.x, -q0.w, q0.z);	
+	//const float4 QuaternionDelta = -2.0 * float4( DeltaLambda.xyz * qebar.w + cross(DeltaLambda.xyz, qebar.xyz),
+	//				- dot(DeltaLambda.xyz, qebar.xyz) );
 
-	q0 += QuaternionDelta * SharedInverseInertia[GGroupThreadId.x-1];
-	SharedNodeOrientation[GGroupThreadId.x-1] = normalize(q0);
+	//q0 += QuaternionDelta * SharedInverseInertia[GGroupThreadId.x-1];
+	//SharedNodeOrientation[GGroupThreadId.x-1] = normalize(q0);
 }
 
 void SolveStretchShearMaterial(const float InMaterialCompliance, const float InMaterialWeight, const float InRestLength, const bool InIsRed, inout float3 OutMaterialMultiplier)
@@ -245,8 +245,6 @@ void SolveStretchShearMaterial(const float InMaterialCompliance, const float InM
 	GroupMemoryBarrier(); 
 } 
 
-
-
 /* -----------------------------------------------------------------
  * Angular Spring material
  * -----------------------------------------------------------------
@@ -280,7 +278,7 @@ void UpdateAngularSpringMultiplier(const float InMaterialCompliance, const float
 	EdgeDirection[2] = q0.w * q0.w - q0.x * q0.x - q0.y * q0.y + q0.z * q0.z;
 
 	// XPBD lagrange multiplier update : dL = -(C+compliance*L) / (dC * invM * dCt + alpha)
-	const float3 DeltaLambda = -((SharedNodePosition[GGroupThreadId.x] - SharedNodePosition[GGroupThreadId.x-1]) / InRestLength - EdgeDirection + 
+	const float3 DeltaLambda = -(normalize(SharedNodePosition[GGroupThreadId.x] - SharedNodePosition[GGroupThreadId.x-1]) - EdgeDirection + 
 					OutMaterialMultiplier * InMaterialCompliance) * InMaterialWeight;
 
 	// L += dL
