Fixed compilation warning and HairStrandsDatas linking errors in Development.

(ushell-p4-cherrypick of 9472597 by [REDACTED_PERSON])

diff --git a/Engine/Plugins/Experimental/HairStrands/Source/HairStrandsCore/Private/GroomAsset.cpp b/Engine/Plugins/Experimental/HairStrands/Source/HairStrandsCore/Private/GroomAsset.cpp
--- a/Engine/Plugins/Experimental/HairStrands/Source/HairStrandsCore/Private/GroomAsset.cpp
+++ b/Engine/Plugins/Experimental/HairStrands/Source/HairStrandsCore/Private/GroomAsset.cpp
@@ -63,7 +63,7 @@ void FHairStrandsRestResource::ReleaseRHI()
 /////////////////////////////////////////////////////////////////////////////////////////
 
 FHairStrandsDeformedResource::FHairStrandsDeformedResource(const FHairStrandsDatas::FRenderData& HairStrandRenderData, bool bInInitializedData) :
-	bInitializedData(bInInitializedData), RenderData(HairStrandRenderData)
+	RenderData(HairStrandRenderData), bInitializedData(bInInitializedData)
 {}
 
 void FHairStrandsDeformedResource::InitRHI()
diff --git a/Engine/Plugins/Experimental/HairStrands/Source/HairStrandsCore/Public/HairStrandsDatas.h b/Engine/Plugins/Experimental/HairStrands/Source/HairStrandsCore/Public/HairStrandsDatas.h
--- a/Engine/Plugins/Experimental/HairStrands/Source/HairStrandsCore/Public/HairStrandsDatas.h
+++ b/Engine/Plugins/Experimental/HairStrands/Source/HairStrandsCore/Public/HairStrandsDatas.h
@@ -207,16 +207,16 @@ struct FHairStrandsRootNormalFormat
 struct HAIRSTRANDSCORE_API FHairStrandsInterpolationDatas
 {
 	/** Serialize the interpolated points */
-	inline void Serialize(FArchive& Ar);
+	void Serialize(FArchive& Ar);
 
 	/** Set the number of interpolated points */
-	inline void SetNum(const uint32 NumPoints);
+	void SetNum(const uint32 NumPoints);
 
 	/** Reset the interpolated points to 0 */
-	inline void Reset();
+	void Reset();
 
 	/** Get the number of interpolated points */
-	inline uint32 Num() const { return PointsSimCurvesVertexIndex.Num(); }
+	uint32 Num() const { return PointsSimCurvesVertexIndex.Num(); }
 
 	/** Simulation curve indices, ordered by closest influence */
 	TArray<FIntVector> PointsSimCurvesIndex;
@@ -232,7 +232,7 @@ struct HAIRSTRANDSCORE_API FHairStrandsInterpolationDatas
 		TArray<FHairStrandsInterpolation0Format::Type> Interpolation0;
 		TArray<FHairStrandsInterpolation1Format::Type> Interpolation1;
 
-		inline void Serialize(FArchive& Ar);
+		void Serialize(FArchive& Ar);
 	} RenderData;
 };
 
@@ -240,16 +240,16 @@ struct HAIRSTRANDSCORE_API FHairStrandsInterpolationDatas
 struct HAIRSTRANDSCORE_API FHairStrandsPoints
 {
 	/** Serialize the points */
-	inline void Serialize(FArchive& Ar);
+	void Serialize(FArchive& Ar);
 
 	/** Set the number of points */
-	inline void SetNum(const uint32 NumPoints);
+	void SetNum(const uint32 NumPoints);
 
 	/** Reset the points to 0 */
-	inline void Reset();
+	void Reset();
 
 	/** Get the number of points */
-	inline uint32 Num() const { return PointsPosition.Num();  }
+	uint32 Num() const { return PointsPosition.Num();  }
 
 	/** Points position in local space */
 	TArray<FVector> PointsPosition;
@@ -265,16 +265,16 @@ struct HAIRSTRANDSCORE_API FHairStrandsPoints
 struct HAIRSTRANDSCORE_API FHairStrandsCurves
 {
 	/** Serialize the curves */
-	inline void Serialize(FArchive& Ar);
+	void Serialize(FArchive& Ar);
 
 	/** Set the number of Curves */
-	inline void SetNum(const uint32 NumPoints);
+	void SetNum(const uint32 NumPoints);
 
 	/** Reset the curves to 0 */
-	inline void Reset();
+	void Reset();
 
 	/** Get the number of Curves */
-	inline uint32 Num() const { return CurvesCount.Num(); }
+	uint32 Num() const { return CurvesCount.Num(); }
 
 	/** Number of points per rod */
 	TArray<uint16> CurvesCount;
@@ -299,13 +299,13 @@ struct HAIRSTRANDSCORE_API FHairStrandsCurves
 struct HAIRSTRANDSCORE_API FHairStrandsDatas
 {
 	/** Serialize all the hair strands datas */
-	inline void Serialize(FArchive& Ar);
+	void Serialize(FArchive& Ar);
 
 	/* Get the total number of points */
-	inline uint32 GetNumPoints() const { return StrandsPoints.Num(); }
+	uint32 GetNumPoints() const { return StrandsPoints.Num(); }
 
 	/* Get the total number of Curves */
-	inline uint32 GetNumCurves() const { return StrandsCurves.Num(); }
+	uint32 GetNumCurves() const { return StrandsCurves.Num(); }
 
 	void Reset();
 
@@ -326,6 +326,6 @@ struct HAIRSTRANDSCORE_API FHairStrandsDatas
 		TArray<FHairStrandsPositionFormat::Type> RenderingPositions;
 		TArray<FHairStrandsAttributeFormat::Type> RenderingAttributes;
 
-		inline void Serialize(FArchive& Ar);
+		void Serialize(FArchive& Ar);
 	} RenderData;
 };
