Removed pbrt and .hair format importers.

(ushell-p4-cherrypick of 9472801 by [REDACTED_PERSON])

diff --git a/Engine/Plugins/Experimental/HairStrands/Source/HairStrandsEditor/Private/HairFormatTranslator.cpp b/Engine/Plugins/Experimental/HairStrands/Source/HairStrandsEditor/Private/HairFormatTranslator.cpp
deleted file mode 100644
diff --git a/Engine/Plugins/Experimental/HairStrands/Source/HairStrandsEditor/Private/HairFormatTranslator.h b/Engine/Plugins/Experimental/HairStrands/Source/HairStrandsEditor/Private/HairFormatTranslator.h
deleted file mode 100644
diff --git a/Engine/Plugins/Experimental/HairStrands/Source/HairStrandsEditor/Private/HairStrandsEditor.cpp b/Engine/Plugins/Experimental/HairStrands/Source/HairStrandsEditor/Private/HairStrandsEditor.cpp
--- a/Engine/Plugins/Experimental/HairStrands/Source/HairStrandsEditor/Private/HairStrandsEditor.cpp
+++ b/Engine/Plugins/Experimental/HairStrands/Source/HairStrandsEditor/Private/HairStrandsEditor.cpp
@@ -4,8 +4,6 @@
 #include "GroomActions.h"
 
 #include "FbxHairTranslator.h"
-#include "HairFormatTranslator.h"
-#include "PbrtHairTranslator.h"
 
 IMPLEMENT_MODULE(FHairStrandsEditor, HairStrandsEditor);
 
@@ -18,8 +16,6 @@ void FHairStrandsEditor::StartupModule()
 	RegisteredAssetTypeActions.Add(AssetActions);
 
 	RegisterHairTranslator<FFbxHairTranslator>();
-	RegisterHairTranslator<FHairFormatTranslator>();
-	RegisterHairTranslator<FPbrtHairTranslator>();
 }
 
 void FHairStrandsEditor::ShutdownModule()
diff --git a/Engine/Plugins/Experimental/HairStrands/Source/HairStrandsEditor/Private/PbrtHairTranslator.cpp b/Engine/Plugins/Experimental/HairStrands/Source/HairStrandsEditor/Private/PbrtHairTranslator.cpp
deleted file mode 100644
diff --git a/Engine/Plugins/Experimental/HairStrands/Source/HairStrandsEditor/Private/PbrtHairTranslator.h b/Engine/Plugins/Experimental/HairStrands/Source/HairStrandsEditor/Private/PbrtHairTranslator.h
deleted file mode 100644
