Rename groom rendering properties.

(ushell-p4-cherrypick of 9479428 by Charles.deRousiers)

diff --git a/Engine/Plugins/Experimental/HairStrands/Source/HairStrandsCore/Private/GroomComponent.cpp b/Engine/Plugins/Experimental/HairStrands/Source/HairStrandsCore/Private/GroomComponent.cpp
--- a/Engine/Plugins/Experimental/HairStrands/Source/HairStrandsCore/Private/GroomComponent.cpp
+++ b/Engine/Plugins/Experimental/HairStrands/Source/HairStrandsCore/Private/GroomComponent.cpp
@@ -378,9 +378,9 @@ UGroomComponent::UGroomComponent(const FObjectInitializer& ObjectInitializer)
 	bAutoActivate = true;
 	bSelectable = true;
 	RegisteredSkeletalMeshComponent = nullptr;
-	HairDensity = 1;
+	HairShadowDensity = 1;
 	HairRaytracingRadiusScale = 1;
-	bSkinGroom = false;
+	bBindGroomToSkeletalMesh = false;
 	InitializedResources = nullptr;
 	Mobility = EComponentMobility::Movable;
 	MeshProjectionTickDelay = 0;
@@ -613,7 +613,7 @@ void UGroomComponent::InitResources()
 	WorldType = WorldType == EWorldType::Inactive ? EWorldType::Editor : WorldType;
 
 	// Insure the ticking of the Groom component always happens after the skeletalMeshComponent.
-	USkeletalMeshComponent* SkeletalMeshComponent = bSkinGroom && GetAttachParent() ? Cast<USkeletalMeshComponent>(GetAttachParent()) : nullptr;
+	USkeletalMeshComponent* SkeletalMeshComponent = bBindGroomToSkeletalMesh && GetAttachParent() ? Cast<USkeletalMeshComponent>(GetAttachParent()) : nullptr;
 	if (SkeletalMeshComponent)
 	{
 		RegisteredSkeletalMeshComponent = SkeletalMeshComponent;
@@ -626,7 +626,7 @@ void UGroomComponent::InitResources()
 	}
 
 	FTransform HairLocalToWorld = GetComponentTransform();
-	FTransform SkinLocalToWorld = bSkinGroom && SkeletalMeshComponent ? SkeletalMeshComponent->GetComponentTransform() : FTransform::Identity;
+	FTransform SkinLocalToWorld = bBindGroomToSkeletalMesh && SkeletalMeshComponent ? SkeletalMeshComponent->GetComponentTransform() : FTransform::Identity;
 	
 	InterpolationOutput = new FHairStrandsInterpolationOutput();
 	InterpolationInput = new FHairStrandsInterpolationInput();
@@ -1005,12 +1005,12 @@ void UGroomComponent::PostEditChangeProperty(FPropertyChangedEvent& PropertyChan
 	//  Init/release resources when setting the GroomAsset (or undoing)
 	const bool bRecreateResources =
 		(PropertyName == GET_MEMBER_NAME_CHECKED(UGroomComponent, GroomAsset) || PropertyThatChanged == nullptr) ||
-		PropertyName == GET_MEMBER_NAME_CHECKED(UGroomComponent, bSkinGroom);
+		PropertyName == GET_MEMBER_NAME_CHECKED(UGroomComponent, bBindGroomToSkeletalMesh);
 
 	const bool bSupportSkinProjection = GetDefault<URendererSettings>()->bSupportSkinCacheShaders && !UAnimationSettings::Get()->bTickAnimationOnSkeletalMeshInit;
 	if (!bSupportSkinProjection)
 	{
-		bSkinGroom = false;
+		bBindGroomToSkeletalMesh = false;
 	}
 
 	if (bRecreateResources)
@@ -1055,7 +1055,7 @@ bool UGroomComponent::CanEditChange(const UProperty* InProperty) const
 			return bIsEditable;
 		}
 
-		if (FCString::Strcmp(*PropertyName, TEXT("bSkinGroom")) == 0)
+		if (FCString::Strcmp(*PropertyName, TEXT("bBindGroomToSkeletalMesh")) == 0)
 		{
 			return GetDefault<URendererSettings>()->bSupportSkinCacheShaders && !UAnimationSettings::Get()->bTickAnimationOnSkeletalMeshInit;
 		}
diff --git a/Engine/Plugins/Experimental/HairStrands/Source/HairStrandsCore/Public/GroomComponent.h b/Engine/Plugins/Experimental/HairStrands/Source/HairStrandsCore/Public/GroomComponent.h
--- a/Engine/Plugins/Experimental/HairStrands/Source/HairStrandsCore/Public/GroomComponent.h
+++ b/Engine/Plugins/Experimental/HairStrands/Source/HairStrandsCore/Public/GroomComponent.h
@@ -23,21 +23,21 @@ public:
 	UPROPERTY(EditAnywhere, BlueprintReadWrite, Category = "Groom")
 	UGroomAsset* GroomAsset;
 
-	/** Controls the hair density, to reduce or increase hair count during shadow rendering. This allows to increase/decrease the shadowing on hair when the number of strand is not realistic */
-	UPROPERTY(EditAnywhere, BlueprintReadWrite, Category = "HairStrands Rendering", meta = (ClampMin = "0.0001", UIMin = "0.01", UIMax = "10.0"))
-	float HairDensity;
-
 	/** 
-	 * When activated, the hair groom will be attached and skinned onto the mesh, if the groom component is a child of a skeletal/skinned component.
+	 * When activated, the groom will be attached and skinned onto the skeletal mesh, if the groom component is a child of a skeletal/skinned component.
 	 * This requires the following projection settings: 
 	 * - Rendering settings: 'Skin cache' enabled
 	 * - Animation settings: 'Tick Animation On Skeletal Mesh Init' disabled
 	 */
-	UPROPERTY(EditAnywhere, BlueprintReadOnly, Category = "HairStrands Rendering")
-	bool bSkinGroom;
+	UPROPERTY(EditAnywhere, BlueprintReadOnly, Category = "Groom")
+	bool bBindGroomToSkeletalMesh;
+
+	/** Controls the hair density, to reduce or increase hair count during shadow rendering. This allows to increase/decrease the shadowing on hair when the number of strand is not realistic */
+	UPROPERTY(EditAnywhere, BlueprintReadWrite, Category = "Groom Rendering", meta = (ClampMin = "0.0001", UIMin = "0.01", UIMax = "10.0"))
+	float HairShadowDensity;
 
 	/** Scale the hair geometry radius for ray tracing effects (e.g. shadow) */
-	UPROPERTY(EditAnywhere, BlueprintReadWrite, Category = "HairStrands Rendering", meta = (ClampMin = "0.0001", UIMin = "0.01", UIMax = "10.0"))
+	UPROPERTY(EditAnywhere, BlueprintReadWrite, Category = "Groom Rendering", meta = (ClampMin = "0.0001", UIMin = "0.01", UIMax = "10.0"))
 	float HairRaytracingRadiusScale;
 
 	//~ Begin UActorComponent Interface.
