Fix missing clear resources.

(ushell-p4-cherrypick of 9485541 by Charles.deRousiers)

diff --git a/Engine/Plugins/Experimental/HairStrands/Source/HairStrandsCore/Private/GroomComponent.cpp b/Engine/Plugins/Experimental/HairStrands/Source/HairStrandsCore/Private/GroomComponent.cpp
--- a/Engine/Plugins/Experimental/HairStrands/Source/HairStrandsCore/Private/GroomComponent.cpp
+++ b/Engine/Plugins/Experimental/HairStrands/Source/HairStrandsCore/Private/GroomComponent.cpp
@@ -776,6 +776,7 @@ void UGroomComponent::ReleaseResources()
 		SafeRelease(Res.RaytracingResources);
 	#endif
 	}
+	HairGroupResources.Empty();
 
 	// Delay destruction as resources reference by the interpolation 
 	// structs are used on the rendering thread, 
