Fix precision issue during interpolation due to 16 bits positions.

 * Introduce relative center for rest & deform triangle positions.
 * Introduce relative center for rest & deformed hair positions positions.
 * Add optional f32 triangle position format to improve pecision issue.

(ushell-p4-cherrypick of 9487539 by Charles.deRousiers)

(ushell-p4-cherrypick of 9487542 by Charles.deRousiers)

diff --git a/Engine/Plugins/Experimental/HairStrands/Source/HairStrandsCore/Private/GroomAsset.cpp b/Engine/Plugins/Experimental/HairStrands/Source/HairStrandsCore/Private/GroomAsset.cpp
--- a/Engine/Plugins/Experimental/HairStrands/Source/HairStrandsCore/Private/GroomAsset.cpp
+++ b/Engine/Plugins/Experimental/HairStrands/Source/HairStrandsCore/Private/GroomAsset.cpp
@@ -163,14 +163,14 @@ void FHairStrandsRootResource::InitRHI()
 			// Create buffers. Initialization will be done by render passes
 			CreateBuffer<FHairStrandsCurveTriangleIndexFormat>(RootCount, LOD.RootTriangleIndexBuffer);
 			CreateBuffer<FHairStrandsCurveTriangleBarycentricFormat>(RootCount, LOD.RootTriangleBarycentricBuffer);
-			CreateBuffer<FHairStrandsCurveTranslationFormat>(RootCount, LOD.RestRootTrianglePosition0Buffer);
-			CreateBuffer<FHairStrandsCurveTranslationFormat>(RootCount, LOD.RestRootTrianglePosition1Buffer);
-			CreateBuffer<FHairStrandsCurveTranslationFormat>(RootCount, LOD.RestRootTrianglePosition2Buffer);
+			CreateBuffer<FHairStrandsMeshTrianglePositionFormat>(RootCount, LOD.RestRootTrianglePosition0Buffer);
+			CreateBuffer<FHairStrandsMeshTrianglePositionFormat>(RootCount, LOD.RestRootTrianglePosition1Buffer);
+			CreateBuffer<FHairStrandsMeshTrianglePositionFormat>(RootCount, LOD.RestRootTrianglePosition2Buffer);
 
 			// Strand hair roots translation and rotation in triangle-deformed position relative to the bound triangle 
-			CreateBuffer<FHairStrandsCurveTranslationFormat>(RootCount, LOD.DeformedRootTrianglePosition0Buffer);
-			CreateBuffer<FHairStrandsCurveTranslationFormat>(RootCount, LOD.DeformedRootTrianglePosition1Buffer);
-			CreateBuffer<FHairStrandsCurveTranslationFormat>(RootCount, LOD.DeformedRootTrianglePosition2Buffer);
+			CreateBuffer<FHairStrandsMeshTrianglePositionFormat>(RootCount, LOD.DeformedRootTrianglePosition0Buffer);
+			CreateBuffer<FHairStrandsMeshTrianglePositionFormat>(RootCount, LOD.DeformedRootTrianglePosition1Buffer);
+			CreateBuffer<FHairStrandsMeshTrianglePositionFormat>(RootCount, LOD.DeformedRootTrianglePosition2Buffer);
 		}
 	}
 }
diff --git a/Engine/Plugins/Experimental/HairStrands/Source/HairStrandsCore/Private/GroomComponent.cpp b/Engine/Plugins/Experimental/HairStrands/Source/HairStrandsCore/Private/GroomComponent.cpp
--- a/Engine/Plugins/Experimental/HairStrands/Source/HairStrandsCore/Private/GroomComponent.cpp
+++ b/Engine/Plugins/Experimental/HairStrands/Source/HairStrandsCore/Private/GroomComponent.cpp
@@ -24,6 +24,9 @@ static int32 GHairStrandsMeshProjectionForceLOD = -1;
 static FAutoConsoleVariableRef CVarHairStrandsMeshProjectionForceRefPoseEnable(TEXT("r.HairStrands.MeshProjection.RefPose"), GHairStrandsMeshProjectionForceRefPoseEnable, TEXT("Enable/Disable reference pose"));
 static FAutoConsoleVariableRef CVarHairStrandsMeshProjectionForceLOD(TEXT("r.HairStrands.MeshProjection.LOD"), GHairStrandsMeshProjectionForceLOD, TEXT("Force a specific LOD"));
 
+static int32 GHairStrandsMeshUseRelativePosition = 1;
+static FAutoConsoleVariableRef CVarHairStrandsMeshUseRelativePosition(TEXT("r.HairStrands.MeshProjection.RelativePosition"), GHairStrandsMeshUseRelativePosition, TEXT("Enable/Disable relative triangle position for improving positions"));
+
 /**
  * An material render proxy which overrides the debug mode parameter.
  */
@@ -378,6 +381,7 @@ UGroomComponent::UGroomComponent(const FObjectInitializer& ObjectInitializer)
 	bAutoActivate = true;
 	bSelectable = true;
 	RegisteredSkeletalMeshComponent = nullptr;
+	SkeletalPreviousPositionOffset = FVector::ZeroVector;
 	HairShadowDensity = 1;
 	HairRaytracingRadiusScale = 1;
 	bBindGroomToSkeletalMesh = false;
@@ -502,12 +506,12 @@ static FHairStrandsProjectionHairData::HairGroup ToProjectionHairData(FHairStran
 		LODData.RootTriangleIndexBuffer = &MeshLODData.RootTriangleIndexBuffer;
 		LODData.RootTriangleBarycentricBuffer = &MeshLODData.RootTriangleBarycentricBuffer;
 		
-		LODData.RestRootCenter					= MeshLODData.RestRootCenter;
+		LODData.RestPositionOffset					= MeshLODData.RestRootCenter;
 		LODData.RestRootTrianglePosition0Buffer = &MeshLODData.RestRootTrianglePosition0Buffer;
 		LODData.RestRootTrianglePosition1Buffer = &MeshLODData.RestRootTrianglePosition1Buffer;
 		LODData.RestRootTrianglePosition2Buffer = &MeshLODData.RestRootTrianglePosition2Buffer;
 
-		LODData.DeformedRootCenter					= MeshLODData.DeformedRootCenter;
+		LODData.DeformedPositionOffset					= MeshLODData.DeformedRootCenter;
 		LODData.DeformedRootTrianglePosition0Buffer = &MeshLODData.DeformedRootTrianglePosition0Buffer;
 		LODData.DeformedRootTrianglePosition1Buffer = &MeshLODData.DeformedRootTrianglePosition1Buffer;
 		LODData.DeformedRootTrianglePosition2Buffer = &MeshLODData.DeformedRootTrianglePosition2Buffer;
@@ -676,11 +680,16 @@ void UGroomComponent::InitResources()
 		BeginInitResource(Res.RenderDeformedResources);
 		BeginInitResource(Res.SimDeformedResources);
 
+		const FVector HairBoundCenter = GroupData.HairRenderData.BoundingBox.GetCenter();
+
 		FHairStrandsInterpolationOutput::HairGroup& InterpolationOutputGroup = InterpolationOutput->HairGroups.AddDefaulted_GetRef();
 		FHairStrandsInterpolationInput::FHairGroup& InterpolationInputGroup = InterpolationInput->HairGroups.AddDefaulted_GetRef();
 		InterpolationInputGroup.HairRadius = GroupData.HairRenderData.StrandsCurves.MaxRadius;
-		InterpolationInputGroup.HairWorldOffset = GroupData.HairRenderData.BoundingBox.GetCenter();
 		InterpolationInputGroup.HairRaytracingRadiusScale = HairRaytracingRadiusScale;
+		InterpolationInputGroup.InHairPositionOffset = HairBoundCenter;
+		// For skinned groom, these value will be updated during TickComponent() call
+		InterpolationInputGroup.OutHairPositionOffset = HairBoundCenter; 
+		InterpolationInputGroup.OutHairPreviousPositionOffset = HairBoundCenter;
 	}
 
 	FHairStrandsInterpolationData Interpolation;
@@ -854,9 +863,13 @@ void UGroomComponent::TickComponent(float DeltaTime, enum ELevelTick TickType, F
 	const uint64 Id = ComponentId.PrimIDValue;
 	const ERHIFeatureLevel::Type FeatureLevel = GetWorld() ? ERHIFeatureLevel::Type(GetWorld()->FeatureLevel) : ERHIFeatureLevel::Num;
 
+	FVector MeshPositionOffset = FVector::ZeroVector;
 	USkeletalMeshComponent* SkeletalMeshComponent = RegisteredSkeletalMeshComponent;	
 	if (SkeletalMeshComponent)
 	{
+		// The offset is based on the center of the skeletal mesh (which is computed based on the physics capsules/boxes/...)
+		MeshPositionOffset = SkeletalMeshComponent->CalcBounds(FTransform::Identity).GetBox().GetCenter();
+
 		const uint32 SkeletalLODCount = SkeletalMeshComponent->GetNumLODs();
 		const uint32 ResourceLODCount = HairGroupResources.Num() > 0 && HairGroupResources[0].RenRootResources ? HairGroupResources[0].RenRootResources->MeshProjectionLODs.Num() : 0;
 		if (SkeletalLODCount != ResourceLODCount)
@@ -911,7 +924,6 @@ void UGroomComponent::TickComponent(float DeltaTime, enum ELevelTick TickType, F
 
 		if (MeshProjectionState != EMeshProjectionState::Completed && MeshProjectionTickDelay == 0)
 		{
-
 			if (MeshProjectionState == EMeshProjectionState::Invalid)
 			{
 				MeshProjectionLODIndex = 0;
@@ -930,7 +942,7 @@ void UGroomComponent::TickComponent(float DeltaTime, enum ELevelTick TickType, F
 				SkeletalMeshComponent->SetForceRefPose(true);
 				SkeletalMeshComponent->SetForcedLOD(MeshProjectionLODIndex + 1);
 
-				const FVector RestRootCenter = SkeletalMeshComponent->Bounds.GetSphere().Center;
+				const FVector RestRootCenter = MeshPositionOffset;
 				const uint32 LODIndex = MeshProjectionLODIndex;
 				ENQUEUE_RENDER_COMMAND(FHairStrandsTick_Projection)(
 					[Id, WorldType, FeatureLevel, LODIndex, RestRootCenter](FRHICommandListImmediate& RHICmdList)
@@ -949,9 +961,29 @@ void UGroomComponent::TickComponent(float DeltaTime, enum ELevelTick TickType, F
 			SkeletalMeshComponent->SetForceRefPose(GHairStrandsMeshProjectionForceRefPoseEnable > 0);
 		}
 
+		// For skinned mesh update the relative center of hair positions after deformation
+		{
+			const FVector OutHairPositionOffset = MeshPositionOffset;
+			const FVector OutHairPreviousPositionOffset = SkeletalPreviousPositionOffset;
+			FHairStrandsInterpolationInput* LocalInterpolationInput = InterpolationInput;
+			ENQUEUE_RENDER_COMMAND(FHairStrandsTick_OutHairPositionOffsetUpdate)(
+				[OutHairPositionOffset, OutHairPreviousPositionOffset, LocalInterpolationInput](FRHICommandListImmediate& RHICmdList)
+			{
+				for (FHairStrandsInterpolationInput::FHairGroup& HairGroup : LocalInterpolationInput->HairGroups)
+				{
+					HairGroup.OutHairPositionOffset = OutHairPositionOffset;
+					HairGroup.OutHairPreviousPositionOffset = OutHairPreviousPositionOffset;
+				}
+			});
+
+			// First frame will be wrong ... 
+			SkeletalPreviousPositionOffset = OutHairPositionOffset;
+		}
+
 		// When a skeletal object with projection is enabled, activate the refresh of the bounding box to 
 		// insure the component/proxy bounding box always lies onto the actual skinned mesh
 		MarkRenderTransformDirty();
+
 	}
 
 	if (MeshProjectionTickDelay > 0)
@@ -959,16 +991,16 @@ void UGroomComponent::TickComponent(float DeltaTime, enum ELevelTick TickType, F
 		MeshProjectionTickDelay--;
 	}
 
-	const FVector WorldsBoundsCenter  = SkeletalMeshComponent ? SkeletalMeshComponent->Bounds.GetSphere().Center : FVector::ZeroVector;
+	const FVector DeformedPositionCenter = GHairStrandsMeshUseRelativePosition > 0 ? MeshPositionOffset : FVector::ZeroVector;
 	const FTransform SkinLocalToWorld = SkeletalMeshComponent ? SkeletalMeshComponent->GetComponentTransform() : FTransform();
 	const FTransform HairLocalToWorld = GetComponentTransform();
 	ENQUEUE_RENDER_COMMAND(FHairStrandsTick_TransformUpdate)(
-		[Id, WorldType, HairLocalToWorld, SkinLocalToWorld, WorldsBoundsCenter, FeatureLevel](FRHICommandListImmediate& RHICmdList)
+		[Id, WorldType, HairLocalToWorld, SkinLocalToWorld, DeformedPositionCenter, FeatureLevel](FRHICommandListImmediate& RHICmdList)
 	{		
 		if (ERHIFeatureLevel::Num == FeatureLevel)
 			return;
 
-		UpdateHairStrands(Id, WorldType, HairLocalToWorld, SkinLocalToWorld, WorldsBoundsCenter);
+		UpdateHairStrands(Id, WorldType, HairLocalToWorld, SkinLocalToWorld, DeformedPositionCenter);
 	});
 }
 
diff --git a/Engine/Plugins/Experimental/HairStrands/Source/HairStrandsCore/Private/HairStrandsRendering.cpp b/Engine/Plugins/Experimental/HairStrands/Source/HairStrandsCore/Private/HairStrandsRendering.cpp
--- a/Engine/Plugins/Experimental/HairStrands/Source/HairStrandsCore/Private/HairStrandsRendering.cpp
+++ b/Engine/Plugins/Experimental/HairStrands/Source/HairStrandsCore/Private/HairStrandsRendering.cpp
@@ -114,11 +114,11 @@ class FHairInterpolationCS : public FGlobalShader
 	BEGIN_SHADER_PARAMETER_STRUCT(FParameters, )
 		SHADER_PARAMETER(uint32, VertexCount)
 		SHADER_PARAMETER(uint32, DispatchCountX)
-		SHADER_PARAMETER(FVector, HairWorldOffset)
+		SHADER_PARAMETER(FVector, InHairPositionOffset)
+		SHADER_PARAMETER(FVector, OutHairPositionOffset)
 
-		SHADER_PARAMETER(FVector, RestPositionWorldCenter)
-		SHADER_PARAMETER(FVector, DeformedPositionWorldCenter)
-		SHADER_PARAMETER(FVector, OutRenderDeformedPositionCenter)
+		SHADER_PARAMETER(FVector, RestPositionOffset)
+		SHADER_PARAMETER(FVector, DeformedPositionOffset)
 
 		SHADER_PARAMETER_SRV(Buffer, RenderRestPosePositionBuffer)
 		SHADER_PARAMETER_UAV(RWBuffer, OutRenderDeformedPositionBuffer)
@@ -154,7 +154,8 @@ IMPLEMENT_GLOBAL_SHADER(FHairInterpolationCS, "/Engine/Private/HairStrands/HairS
 static void AddHairStrandsInterpolationPass(
 	FRDGBuilder& GraphBuilder,
 	const FHairStrandsProjectionHairData::HairGroup& InHairData,
-	const FVector& HairWorldOffset, 
+	const FVector& InHairWorldOffset, 
+	const FVector& OutHairWorldOffset,
 	const int32 LODIndex,
 	const uint32 VertexCount,
 	const FShaderResourceViewRHIRef& RenderRestPosePositionBuffer,
@@ -183,21 +184,19 @@ static void AddHairStrandsInterpolationPass(
 		Parameters->OutRenderAttributeBuffer = OutRenderAttributeBuffer;
 	}
 	Parameters->VertexCount = VertexCount;
-	Parameters->HairWorldOffset = HairWorldOffset;
+	Parameters->InHairPositionOffset = InHairWorldOffset;
+	Parameters->OutHairPositionOffset = OutHairWorldOffset;
 	Parameters->DispatchCountX = DispatchCount.X;
-	Parameters->OutRenderDeformedPositionCenter = FVector::ZeroVector;
 
 	const bool bSupportDynamicMesh = InHairData.RootCount > 0 && LODIndex >= 0 && LODIndex < InHairData.LODDatas.Num() && InHairData.LODDatas[LODIndex].bIsValid;
 	if (bSupportDynamicMesh)
 	{
-		Parameters->OutRenderDeformedPositionCenter = FVector::ZeroVector; // TODO
-
-		Parameters->RestPositionWorldCenter = FVector::ZeroVector; // TODO: InHairData.LODDatas[LODIndex].RestRootCenter;
+		Parameters->RestPositionOffset = InHairData.LODDatas[LODIndex].RestPositionOffset;
 		Parameters->RestPosition0Buffer = InHairData.LODDatas[LODIndex].RestRootTrianglePosition0Buffer->SRV;
 		Parameters->RestPosition1Buffer = InHairData.LODDatas[LODIndex].RestRootTrianglePosition1Buffer->SRV;
 		Parameters->RestPosition2Buffer = InHairData.LODDatas[LODIndex].RestRootTrianglePosition2Buffer->SRV;
 
-		Parameters->DeformedPositionWorldCenter = FVector::ZeroVector; // TODO: InHairData.LODDatas[LODIndex].DeformedRootCenter;
+		Parameters->DeformedPositionOffset = InHairData.LODDatas[LODIndex].DeformedPositionOffset;
 		Parameters->DeformedPosition0Buffer = InHairData.LODDatas[LODIndex].DeformedRootTrianglePosition0Buffer->SRV;
 		Parameters->DeformedPosition1Buffer = InHairData.LODDatas[LODIndex].DeformedRootTrianglePosition1Buffer->SRV;
 		Parameters->DeformedPosition2Buffer = InHairData.LODDatas[LODIndex].DeformedRootTrianglePosition2Buffer->SRV;
@@ -437,6 +436,8 @@ void ComputeHairStrandsInterpolation(
 			Output.VFInput.HairPreviousPositionBuffer = Output.SimDeformedPositionBuffer[SimIndex]->SRV;
 			Output.VFInput.HairTangentBuffer = Output.SimTangentBuffer->SRV;
 			Output.VFInput.HairAttributeBuffer = Input.SimAttributeBuffer->SRV;
+			Output.VFInput.HairPositionOffset = Input.OutHairPositionOffset;
+			Output.VFInput.HairPreviousPositionOffset = Input.OutHairPreviousPositionOffset;
 			Output.VFInput.VertexCount = Input.SimVertexCount;
 		}
 		else
@@ -452,7 +453,8 @@ void ComputeHairStrandsInterpolation(
 			AddHairStrandsInterpolationPass(
 				GraphBuilder,
 				InHairDatas.HairGroups[GroupIndex],
-				Input.HairWorldOffset,
+				Input.InHairPositionOffset,
+				Input.OutHairPositionOffset,
 				LODIndex,
 				Input.RenderVertexCount,
 				Input.RenderRestPosePositionBuffer->SRV,
@@ -477,7 +479,7 @@ void ComputeHairStrandsInterpolation(
 					GraphBuilder,
 					Input.RenderVertexCount,
 					Input.HairRadius * (GHairRaytracingRadiusScale > 0 ? GHairRaytracingRadiusScale : Input.HairRaytracingRadiusScale),
-					Input.HairWorldOffset,
+					Input.OutHairPositionOffset,
 					Output.RenderDeformedPositionBuffer[CurrIndex]->SRV,
 					Input.RaytracingPositionBuffer->UAV);
 			}
@@ -488,6 +490,8 @@ void ComputeHairStrandsInterpolation(
 			Output.VFInput.HairPreviousPositionBuffer = Output.RenderDeformedPositionBuffer[PrevIndex]->SRV;
 			Output.VFInput.HairTangentBuffer = Output.RenderTangentBuffer->SRV;
 			Output.VFInput.HairAttributeBuffer = DebugMode == EHairStrandsDebugMode::RenderHairStrands ? Output.RenderPatchedAttributeBuffer.SRV : Input.RenderAttributeBuffer->SRV;
+			Output.VFInput.HairPositionOffset = Input.OutHairPositionOffset;
+			Output.VFInput.HairPreviousPositionOffset = Input.OutHairPreviousPositionOffset;
 			Output.VFInput.VertexCount = Input.RenderVertexCount;
 
 			#if RHI_RAYTRACING
diff --git a/Engine/Plugins/Experimental/HairStrands/Source/HairStrandsCore/Private/HairStrandsVertexFactory.cpp b/Engine/Plugins/Experimental/HairStrands/Source/HairStrandsCore/Private/HairStrandsVertexFactory.cpp
--- a/Engine/Plugins/Experimental/HairStrands/Source/HairStrandsCore/Private/HairStrandsVertexFactory.cpp
+++ b/Engine/Plugins/Experimental/HairStrands/Source/HairStrandsCore/Private/HairStrandsVertexFactory.cpp
@@ -39,7 +39,8 @@ public:
 	FShaderParameter Length;
 	FShaderParameter RadiusAtDepth1_Primary;
 	FShaderParameter RadiusAtDepth1_Velocity;
-	FShaderParameter WorldOffset;
+	FShaderParameter PositionOffset;
+	FShaderParameter PreviousPositionOffset;
 	FShaderParameter Density;
 
 	FShaderResourceParameter PositionBuffer;
@@ -53,7 +54,8 @@ public:
 		Length.Bind(ParameterMap, TEXT("HairStrandsVF_Length"));
 		RadiusAtDepth1_Primary.Bind(ParameterMap, TEXT("HairStrandsVF_RadiusAtDepth1_Primary"));
 		RadiusAtDepth1_Velocity.Bind(ParameterMap, TEXT("HairStrandsVF_RadiusAtDepth1_Velocity"));
-		WorldOffset.Bind(ParameterMap, TEXT("HairStrandsVF_WorldOffset"));
+		PositionOffset.Bind(ParameterMap, TEXT("HairStrandsVF_PositionOffset"));
+		PreviousPositionOffset.Bind(ParameterMap, TEXT("HairStrandsVF_PreviousPositionOffset"));
 		Density.Bind(ParameterMap, TEXT("HairStrandsVF_Density"));	
 
 		PositionBuffer.Bind(ParameterMap, TEXT("HairStrandsVF_PositionBuffer"));
@@ -68,7 +70,8 @@ public:
 		Ar << Length;
 		Ar << RadiusAtDepth1_Primary;
 		Ar << RadiusAtDepth1_Velocity;
-		Ar << WorldOffset;
+		Ar << PositionOffset;
+		Ar << PreviousPositionOffset;
 		Ar << Density;
 
 		Ar << PositionBuffer;
@@ -105,7 +108,8 @@ public:
 		BindParam(ShaderBindings, TangentBuffer, VF->GetTangentSRV(GroupIndex));
 		BindParam(ShaderBindings, Radius, VF->GetMaxStrandRadius(GroupIndex));
 		BindParam(ShaderBindings, Length, VF->GetMaxStrandLength(GroupIndex));
-		BindParam(ShaderBindings, WorldOffset, VF->GetWorldOffset(GroupIndex));
+		BindParam(ShaderBindings, PositionOffset, VF->GetPositionOffset(GroupIndex));
+		BindParam(ShaderBindings, PreviousPositionOffset, VF->GetPreviousPositionOffset(GroupIndex));
 		BindParam(ShaderBindings, Density, VF->GetHairDensity(GroupIndex));
 		BindParam(ShaderBindings, RadiusAtDepth1_Primary, MinRadiusAtDepth1.Primary);
 		BindParam(ShaderBindings, RadiusAtDepth1_Velocity, MinRadiusAtDepth1.Velocity);		
diff --git a/Engine/Plugins/Experimental/HairStrands/Source/HairStrandsCore/Public/GroomComponent.h b/Engine/Plugins/Experimental/HairStrands/Source/HairStrandsCore/Public/GroomComponent.h
--- a/Engine/Plugins/Experimental/HairStrands/Source/HairStrandsCore/Public/GroomComponent.h
+++ b/Engine/Plugins/Experimental/HairStrands/Source/HairStrandsCore/Public/GroomComponent.h
@@ -111,6 +111,7 @@ private:
 		Completed
 	};
 	class USkeletalMeshComponent* RegisteredSkeletalMeshComponent;
+	FVector SkeletalPreviousPositionOffset;
 	int32 MeshProjectionLODIndex;
 	uint32 MeshProjectionTickDelay;
 	EMeshProjectionState MeshProjectionState;
diff --git a/Engine/Plugins/Experimental/HairStrands/Source/HairStrandsCore/Public/HairStrandsDatas.h b/Engine/Plugins/Experimental/HairStrands/Source/HairStrandsCore/Public/HairStrandsDatas.h
--- a/Engine/Plugins/Experimental/HairStrands/Source/HairStrandsCore/Public/HairStrandsDatas.h
+++ b/Engine/Plugins/Experimental/HairStrands/Source/HairStrandsCore/Public/HairStrandsDatas.h
@@ -149,22 +149,28 @@ struct FHairStrandsWeightFormat
 	static const EPixelFormat Format = PF_R32_FLOAT;
 };
 
-struct FHairStrandsCurveTranslationFormat
+/** 
+ * Skinned mesh triangle vertex position format
+ * Two precision options are available: 4x16bits or 4x32bits
+ * Triangle vertices are relative to their bounding box in order to preserve precision, 
+ * however this is sometime not enough for large asset, this is why by default the format 
+ * use 32bits precision
+*/
+struct FHairStrandsMeshTrianglePositionFormat
 {
-	using Type = FVector4_16;
+#if 0
+	using Type = FVector4;
 	static const uint32 ComponentCount = 1;
 	static const uint32 SizeInByte = sizeof(Type);
 	static const EVertexElementType VertexElementType = VET_Float4;
-	static const EPixelFormat Format = PF_FloatRGBA;
-};
-
-struct FHairStrandsCurveRotationFormat
-{
+	static const EPixelFormat Format = PF_A32B32G32R32F;
+#else
 	using Type = FVector4_16;
 	static const uint32 ComponentCount = 1;
 	static const uint32 SizeInByte = sizeof(Type);
 	static const EVertexElementType VertexElementType = VET_Float4;
 	static const EPixelFormat Format = PF_FloatRGBA;
+#endif
 };
 
 struct FHairStrandsCurveTriangleIndexFormat
diff --git a/Engine/Plugins/Experimental/HairStrands/Source/HairStrandsCore/Public/HairStrandsRendering.h b/Engine/Plugins/Experimental/HairStrands/Source/HairStrandsCore/Public/HairStrandsRendering.h
--- a/Engine/Plugins/Experimental/HairStrands/Source/HairStrandsCore/Public/HairStrandsRendering.h
+++ b/Engine/Plugins/Experimental/HairStrands/Source/HairStrandsCore/Public/HairStrandsRendering.h
@@ -34,7 +34,9 @@ struct FHairStrandsInterpolationInput
 
 		float HairRadius = 0;
 		float HairRaytracingRadiusScale = 0;
-		FVector HairWorldOffset = FVector::ZeroVector;
+		FVector InHairPositionOffset = FVector::ZeroVector;
+		FVector OutHairPositionOffset = FVector::ZeroVector;
+		FVector OutHairPreviousPositionOffset = FVector::ZeroVector;
 
 		inline bool IsValid() const
 		{
@@ -62,6 +64,8 @@ struct FHairStrandsInterpolationOutput
 		FShaderResourceViewRHIRef HairPreviousPositionBuffer = nullptr;
 		FShaderResourceViewRHIRef HairTangentBuffer = nullptr;
 		FShaderResourceViewRHIRef HairAttributeBuffer = nullptr;
+		FVector HairPositionOffset = FVector::ZeroVector;
+		FVector HairPreviousPositionOffset = FVector::ZeroVector;
 		uint32 VertexCount = 0;
 
 		inline void Reset()
@@ -70,6 +74,8 @@ struct FHairStrandsInterpolationOutput
 			HairPreviousPositionBuffer = nullptr;
 			HairTangentBuffer = nullptr;
 			HairAttributeBuffer = nullptr;
+			HairPositionOffset = FVector::ZeroVector;
+			HairPreviousPositionOffset = FVector::ZeroVector;
 			VertexCount = 0;
 		}
 	};
diff --git a/Engine/Plugins/Experimental/HairStrands/Source/HairStrandsCore/Public/HairStrandsVertexFactory.h b/Engine/Plugins/Experimental/HairStrands/Source/HairStrandsCore/Public/HairStrandsVertexFactory.h
--- a/Engine/Plugins/Experimental/HairStrands/Source/HairStrandsCore/Public/HairStrandsVertexFactory.h
+++ b/Engine/Plugins/Experimental/HairStrands/Source/HairStrandsCore/Public/HairStrandsVertexFactory.h
@@ -79,7 +79,8 @@ public:
 	inline float GetMinStrandRadius(uint32 GroupIndex) const						{ return Data.HairGroups[GroupIndex].MinStrandRadius; }
 	inline float GetMaxStrandLength(uint32 GroupIndex) const						{ return Data.HairGroups[GroupIndex].MaxStrandLength; }
 	inline float GetHairDensity(uint32 GroupIndex) const							{ return Data.HairGroups[GroupIndex].HairDensity;  }
-	inline const FVector& GetWorldOffset(uint32 GroupIndex) const					{ return Data.HairGroups[GroupIndex].HairWorldOffset; }
+	inline const FVector& GetPositionOffset(uint32 GroupIndex) const				{ return Data.InterpolationOutput->HairGroups[GroupIndex].VFInput.HairPositionOffset; }
+	inline const FVector& GetPreviousPositionOffset(uint32 GroupIndex) const		{ return Data.InterpolationOutput->HairGroups[GroupIndex].VFInput.HairPreviousPositionOffset; }
 
 	const FDataType& GetData() const { return Data; }
 protected:
diff --git a/Engine/Shaders/Private/HairStrands/HairStrandsInterpolation.usf b/Engine/Shaders/Private/HairStrands/HairStrandsInterpolation.usf
--- a/Engine/Shaders/Private/HairStrands/HairStrandsInterpolation.usf
+++ b/Engine/Shaders/Private/HairStrands/HairStrandsInterpolation.usf
@@ -5,7 +5,8 @@
 
 uint VertexCount;
 uint DispatchCountX;
-float3 HairWorldOffset;
+float3 InHairPositionOffset;
+float3 OutHairPositionOffset;
 
 Buffer<uint4>  RenderRestPosePositionBuffer;
 Buffer<uint4>  SimRestPosePositionBuffer;
@@ -13,18 +14,17 @@ Buffer<uint4>  DeformedSimPositionBuffer;
 Buffer<uint4>  Interpolation0Buffer;
 Buffer<uint4>  Interpolation1Buffer;
 
-float3			OutRenderDeformedPositionCenter; 
 RWBuffer<uint4> OutRenderDeformedPositionBuffer;
 
 #if PERMUTATION_DYNAMIC_GEOMETRY == 1
 
 // Compact all these buffers into 2 buffers: translation + quaternion
-float3			RestPositionWorldCenter;
+float3			RestPositionOffset;
 Buffer<float4>  RestPosition0Buffer;
 Buffer<float4>  RestPosition1Buffer;
 Buffer<float4>  RestPosition2Buffer;
 
-float3			DeformedPositionWorldCenter;
+float3			DeformedPositionOffset;
 Buffer<float4>  DeformedPosition0Buffer;
 Buffer<float4>  DeformedPosition1Buffer;
 Buffer<float4>  DeformedPosition2Buffer;
@@ -64,8 +64,9 @@ void MainCS(uint2 DispatchThreadId : SV_DispatchThreadID)
 		float3 GuideVertexWeights = float3(saturate((InterpolationData0.w & 0xFF) / 255.f), saturate(((InterpolationData0.w >> 8) & 0xFF) / 255.f), 0);
 		GuideVertexWeights.z = saturate(1.0f - (GuideVertexWeights.x + GuideVertexWeights.y));
 
+		// ControlPoint is in the local hair referential
 		float3 CurrOffset = 0;
-		float3 ControlPoint = f16tof32(RenderRestPosePositionBuffer[VertexIndex].xyz);
+		float3 ControlPoint = f16tof32(RenderRestPosePositionBuffer[VertexIndex].xyz) + InHairPositionOffset;
 		[unroll]
 		for (uint K = 0; K < 3; ++K)
 		{
@@ -83,10 +84,10 @@ void MainCS(uint2 DispatchThreadId : SV_DispatchThreadID)
 		{
 			const uint RootIndex						= VertexToRootIndexBuffer[VertexIndex];
 			const float3 RootBarycentric				= DecodeBarycentrics(RootBarycentricBuffer[RootIndex]);
-			const FHairMeshTriangle RestTriangle		= GetTriangleTransformation(RootIndex, RestPosition0Buffer, RestPosition1Buffer, RestPosition2Buffer, RestPositionWorldCenter);
-			const FHairMeshTriangle DeformedTriangle	= GetTriangleTransformation(RootIndex, DeformedPosition0Buffer, DeformedPosition1Buffer, DeformedPosition2Buffer, DeformedPositionWorldCenter);
+			const FHairMeshTriangle RestTriangle		= GetTriangleTransformation(RootIndex, RestPosition0Buffer, RestPosition1Buffer, RestPosition2Buffer, RestPositionOffset);
+			const FHairMeshTriangle DeformedTriangle	= GetTriangleTransformation(RootIndex, DeformedPosition0Buffer, DeformedPosition1Buffer, DeformedPosition2Buffer, DeformedPositionOffset);
 
-			ControlPoint = TransformPoint(ControlPoint + HairWorldOffset, RootBarycentric, RestTriangle, DeformedTriangle) - HairWorldOffset;
+			ControlPoint = TransformPoint(ControlPoint, RootBarycentric, RestTriangle, DeformedTriangle);
 
 			// Debug
 			#if 0
@@ -95,13 +96,13 @@ void MainCS(uint2 DispatchThreadId : SV_DispatchThreadID)
 				//ControlPoint = VertexIndex % 2 == 0 ? RestTriangle.P0 : RestTriangle.P0 + RestTriangle.N;
 
 				// Measure the stretch of the projection (to identify if there any local transform issue during the projeciton steps
-				ControlPoint = VertexIndex % 2 == 0 ? ControlPoint  : RestTriangle.P0 - HairWorldOffset;
+				ControlPoint = VertexIndex % 2 == 0 ? ControlPoint  : RestTriangle.P0;
 			}
 			#endif
 		}
 		#endif	
 
-		OutRenderDeformedPositionBuffer[VertexIndex] = float4(f32tof16(ControlPoint + CurrOffset - OutRenderDeformedPositionCenter), RenderRestPosePositionBuffer[VertexIndex].w);
+		OutRenderDeformedPositionBuffer[VertexIndex] = float4(f32tof16(ControlPoint + CurrOffset - OutHairPositionOffset), RenderRestPosePositionBuffer[VertexIndex].w);
 
 		#if PERMUTATION_DEBUG == 1
 		{
diff --git a/Engine/Shaders/Private/HairStrands/HairStrandsMeshProjectionCommon.ush b/Engine/Shaders/Private/HairStrands/HairStrandsMeshProjectionCommon.ush
--- a/Engine/Shaders/Private/HairStrands/HairStrandsMeshProjectionCommon.ush
+++ b/Engine/Shaders/Private/HairStrands/HairStrandsMeshProjectionCommon.ush
@@ -52,8 +52,8 @@ FHairMeshTriangle GetTriangleTransformation(
 	uint RootIndex,
 	Buffer<float4>  InPosition0Buffer,
 	Buffer<float4>  InPosition1Buffer,
-	Buffer<float4>  InPosition2Buffer, 
-	float3 CenterOffset=float3(0,0,0))
+	Buffer<float4>  InPosition2Buffer,
+	float3 CenterOffset)
 {
 	FHairMeshTriangle Out;
 
diff --git a/Engine/Shaders/Private/HairStrands/HairStrandsMeshProjectionHairDebug.usf b/Engine/Shaders/Private/HairStrands/HairStrandsMeshProjectionHairDebug.usf
--- a/Engine/Shaders/Private/HairStrands/HairStrandsMeshProjectionHairDebug.usf
+++ b/Engine/Shaders/Private/HairStrands/HairStrandsMeshProjectionHairDebug.usf
@@ -17,10 +17,12 @@ Buffer<float4> RootNormalBuffer;
 Buffer<uint>   RootBarycentricBuffer;
 #endif
 
+float3 RestPositionOffset;
 Buffer<float4>  RestPosition0Buffer;
 Buffer<float4>  RestPosition1Buffer;
 Buffer<float4>  RestPosition2Buffer;
 
+float3 DeformedPositionOffset;
 Buffer<float4>  DeformedPosition0Buffer;
 Buffer<float4>  DeformedPosition1Buffer;
 Buffer<float4>  DeformedPosition2Buffer;
@@ -43,8 +45,8 @@ void MainVS(
 	float3 N;
 	float3 P;
 	{
-		const FHairMeshTriangle RestTriangle     = GetTriangleTransformation(LineIndex, RestPosition0Buffer, RestPosition1Buffer, RestPosition2Buffer);
-		const FHairMeshTriangle DeformedTriangle = GetTriangleTransformation(LineIndex, DeformedPosition0Buffer, DeformedPosition1Buffer, DeformedPosition2Buffer);
+		const FHairMeshTriangle RestTriangle     = GetTriangleTransformation(LineIndex, RestPosition0Buffer, RestPosition1Buffer, RestPosition2Buffer, RestPositionOffset);
+		const FHairMeshTriangle DeformedTriangle = GetTriangleTransformation(LineIndex, DeformedPosition0Buffer, DeformedPosition1Buffer, DeformedPosition2Buffer, DeformedPositionOffset);
 
 		float3 P0 = RootPositionBuffer[LineIndex].xyz;
 		float3 LN = RootNormalBuffer[LineIndex].xyz;
@@ -72,8 +74,8 @@ void MainVS(
 
 	if (TriangleIndex > MaxRootCount) return;
 
-	const FHairMeshTriangle RestTriangle     = GetTriangleTransformation(TriangleIndex, RestPosition0Buffer, RestPosition1Buffer, RestPosition2Buffer);
-	const FHairMeshTriangle DeformedTriangle = GetTriangleTransformation(TriangleIndex, DeformedPosition0Buffer, DeformedPosition1Buffer, DeformedPosition2Buffer);
+	const FHairMeshTriangle RestTriangle     = GetTriangleTransformation(TriangleIndex, RestPosition0Buffer, RestPosition1Buffer, RestPosition2Buffer, RestPositionOffset);
+	const FHairMeshTriangle DeformedTriangle = GetTriangleTransformation(TriangleIndex, DeformedPosition0Buffer, DeformedPosition1Buffer, DeformedPosition2Buffer, DeformedPositionOffset);
 
 	FHairMeshTriangle Tri;
 	if (DeformedFrameEnable > 0)
diff --git a/Engine/Shaders/Private/HairStrands/HairStrandsMeshUpdate.usf b/Engine/Shaders/Private/HairStrands/HairStrandsMeshUpdate.usf
--- a/Engine/Shaders/Private/HairStrands/HairStrandsMeshUpdate.usf
+++ b/Engine/Shaders/Private/HairStrands/HairStrandsMeshUpdate.usf
@@ -3,7 +3,7 @@
 #include "../Common.ush"
 #include "HairStrandsMeshProjectionCommon.ush"
 
-uint RootTriangleWorldCenter;
+float3 RootTrianglePositionOffset;
 float4x4 MeshLocalToWorld;
 uint MeshSectionIndex;
 uint MeshMaxIndexCount;
@@ -33,7 +33,7 @@ void MainCS(uint3 DispatchThreadId : SV_DispatchThreadID)
 
 	const FHairMeshTriangle Tri = GetTriangleData(TriangleIndex.TriangleIndex, MeshLocalToWorld, MeshPositionBuffer, MeshIndexBuffer, MeshIndexOffset, MeshMaxIndexCount, MeshMaxVertexCount);
 
-	OutRootTrianglePosition0[RootIndex] = (float4(Tri.P0 - RootTriangleWorldCenter,1));
-	OutRootTrianglePosition1[RootIndex] = (float4(Tri.P1 - RootTriangleWorldCenter,1));
-	OutRootTrianglePosition2[RootIndex] = (float4(Tri.P2 - RootTriangleWorldCenter,1));
+	OutRootTrianglePosition0[RootIndex] = (float4(Tri.P0 - RootTrianglePositionOffset,1));
+	OutRootTrianglePosition1[RootIndex] = (float4(Tri.P1 - RootTrianglePositionOffset,1));
+	OutRootTrianglePosition2[RootIndex] = (float4(Tri.P2 - RootTrianglePositionOffset,1));
 }
\ No newline at end of file
diff --git a/Engine/Shaders/Private/HairStrands/HairStrandsVertexFactory.ush b/Engine/Shaders/Private/HairStrands/HairStrandsVertexFactory.ush
--- a/Engine/Shaders/Private/HairStrands/HairStrandsVertexFactory.ush
+++ b/Engine/Shaders/Private/HairStrands/HairStrandsVertexFactory.ush
@@ -20,7 +20,8 @@
 // This is set by compilation enviromenent of the vertex factory
 ///////////////////////////////////////////////////////////////////////////////
 
-float3 HairStrandsVF_WorldOffset; // Vertex position (3x16bits) are relative to this offset for improving precision while keeping vertex footprint low
+float3 HairStrandsVF_PositionOffset;			// Vertex position (3x16bits) are relative to this offset for improving precision while keeping vertex footprint low
+float3 HairStrandsVF_PreviousPositionOffset;	// Vertex position (3x16bits) are relative to this offset for improving precision while keeping vertex footprint low
 float  HairStrandsVF_Density;
 float  HairStrandsVF_Radius;
 float  HairStrandsVF_Length;
@@ -144,7 +145,7 @@ FControlPoint GetVertexPosition(FVertexFactoryInput Input)
 	
 	FControlPoint Output = (FControlPoint)0;
 
-	const float3 ControlPoint	= f16tof32(PackedData.xyz) + HairStrandsVF_WorldOffset;
+	const float3 ControlPoint	= f16tof32(PackedData.xyz) + HairStrandsVF_PositionOffset;
 	const uint PackedAlpha		= asuint(PackedData.a);
 	const uint ControlPointType = PackedAlpha & 0x3;
 	Output.WorldRadius			= ((PackedAlpha & 0x00FC) >> 2) / 63.f * HairStrandsVF_Radius;
@@ -169,7 +170,7 @@ float3 GetVertexPreviousPosition(FVertexFactoryInput Input)
 	const uint OffsetIndex = QuadIndex == 0 || QuadIndex == 2 || QuadIndex == 4 ? 0 : 1;
 	const uint4 PackedData = HairStrandsVF_PreviousPositionBuffer[BaseIndex + OffsetIndex];
 
-	const float3 ControlPoint = f16tof32(PackedData.xyz) + HairStrandsVF_WorldOffset;
+	const float3 ControlPoint = f16tof32(PackedData.xyz) + HairStrandsVF_PreviousPositionOffset;
 	const uint PackedAlpha = asuint(PackedData.a);
 	const uint ControlPointType = PackedAlpha & 0x3;
 
diff --git a/Engine/Source/Runtime/Renderer/Private/HairStrands/HairStrandsDebug.cpp b/Engine/Source/Runtime/Renderer/Private/HairStrands/HairStrandsDebug.cpp
--- a/Engine/Source/Runtime/Renderer/Private/HairStrands/HairStrandsDebug.cpp
+++ b/Engine/Source/Runtime/Renderer/Private/HairStrands/HairStrandsDebug.cpp
@@ -632,6 +632,8 @@ BEGIN_SHADER_PARAMETER_STRUCT(FHairProjectionHairDebugParameters, )
 	SHADER_PARAMETER(FVector2D, OutputResolution)
 	SHADER_PARAMETER(uint32, MaxRootCount)
 	SHADER_PARAMETER(uint32, DeformedFrameEnable)
+	SHADER_PARAMETER(FVector, RestPositionOffset)
+	SHADER_PARAMETER(FVector, DeformedPositionOffset)
 	SHADER_PARAMETER(FMatrix, RootLocalToWorld)
 
 	SHADER_PARAMETER_SRV(StructuredBuffer, RestPosition0Buffer)
@@ -741,10 +743,12 @@ static void AddDebugProjectionHairPass(
 		Parameters->RootBarycentricBuffer = HairData.LODDatas[LODIndex].RootTriangleBarycentricBuffer->SRV;
 	}
 
+	Parameters->RestPositionOffset = HairData.LODDatas[LODIndex].RestPositionOffset;
 	Parameters->RestPosition0Buffer = HairData.LODDatas[LODIndex].RestRootTrianglePosition0Buffer->SRV;
 	Parameters->RestPosition1Buffer = HairData.LODDatas[LODIndex].RestRootTrianglePosition1Buffer->SRV;
 	Parameters->RestPosition2Buffer = HairData.LODDatas[LODIndex].RestRootTrianglePosition2Buffer->SRV;
-		
+	
+	Parameters->DeformedPositionOffset = HairData.LODDatas[LODIndex].DeformedPositionOffset;
 	Parameters->DeformedPosition0Buffer = HairData.LODDatas[LODIndex].DeformedRootTrianglePosition0Buffer->SRV;
 	Parameters->DeformedPosition1Buffer = HairData.LODDatas[LODIndex].DeformedRootTrianglePosition1Buffer->SRV;
 	Parameters->DeformedPosition2Buffer = HairData.LODDatas[LODIndex].DeformedRootTrianglePosition2Buffer->SRV;
diff --git a/Engine/Source/Runtime/Renderer/Private/HairStrands/HairStrandsInterface.cpp b/Engine/Source/Runtime/Renderer/Private/HairStrands/HairStrandsInterface.cpp
--- a/Engine/Source/Runtime/Renderer/Private/HairStrands/HairStrandsInterface.cpp
+++ b/Engine/Source/Runtime/Renderer/Private/HairStrands/HairStrandsInterface.cpp
@@ -58,7 +58,7 @@ struct FHairStrandsManager
 		FHairStrandsProjectionHairData SimProjectionHairDatas;
 		FCachedGeometry CachedGeometry;
 		FTransform SkeletalLocalToWorld;
-		FVector SkeletalWorldBoundsCenter;
+		FVector SkeletalDeformedPositionOffset;
 		const FSkeletalMeshObject* MeshObject = nullptr;
 		int32 FrameLODIndex = -1;
 	};
@@ -69,7 +69,7 @@ struct FHairStrandsManager
 		uint64 Id = 0;
 		EWorldType::Type WorldType = EWorldType::None;
 		int32 LODIndex = -1;
-		FVector RestRootCenter = FVector::ZeroVector;
+		FVector RestPositionOffset = FVector::ZeroVector;
 		bool bProcessed = false;
 	};
 	TArray<ProjectionQuery> ProjectionsQueries;
@@ -120,7 +120,7 @@ void RegisterHairStrands(
 		ProjectionData.LocalToWorld = FTransform::Identity;
 	}
 	E.SkeletalLocalToWorld = FTransform::Identity;
-	E.SkeletalWorldBoundsCenter = FVector::ZeroVector;
+	E.SkeletalDeformedPositionOffset = FVector::ZeroVector;
 	E.DebugInfo = DebugInfo;
 }
 
@@ -129,7 +129,7 @@ bool UpdateHairStrands(
 	EWorldType::Type WorldType, 
 	const FTransform& HairLocalToWorld, 
 	const FTransform& SkeletalLocalToWorld,
-	const FVector& SkeletalWorldBoundsCenter)
+	const FVector& SkeletalDeformedPositionOffset)
 {
 	for (FHairStrandsManager::Element& E : GHairManager.Elements)
 	{
@@ -145,7 +145,7 @@ bool UpdateHairStrands(
 			ProjectionData.LocalToWorld = HairLocalToWorld;
 		}
 		E.SkeletalLocalToWorld = SkeletalLocalToWorld;
-		E.SkeletalWorldBoundsCenter = SkeletalWorldBoundsCenter;
+		E.SkeletalDeformedPositionOffset = SkeletalDeformedPositionOffset;
 		return true;
 	}
 
@@ -259,14 +259,14 @@ void RunHairStrandsInterpolation(
 
 			for (FHairStrandsProjectionHairData::HairGroup& ProjectionHairData : E.SimProjectionHairDatas.HairGroups)
 			{
-				ProjectionHairData.LODDatas[Q.LODIndex].RestRootCenter = Q.RestRootCenter;
+				ProjectionHairData.LODDatas[Q.LODIndex].RestPositionOffset = Q.RestPositionOffset;
 				ProjectHairStrandsOntoMesh(RHICmdList, ShaderMap, Q.LODIndex, MeshData, ProjectionHairData);
 				UpdateHairStrandsMeshTriangles(RHICmdList, ShaderMap, Q.LODIndex, HairStrandsTriangleType::RestPose, MeshData, ProjectionHairData);
 			}
 
 			for (FHairStrandsProjectionHairData::HairGroup& ProjectionHairData : E.RenProjectionHairDatas.HairGroups)
 			{
-				ProjectionHairData.LODDatas[Q.LODIndex].RestRootCenter = Q.RestRootCenter;
+				ProjectionHairData.LODDatas[Q.LODIndex].RestPositionOffset = Q.RestPositionOffset;
 				ProjectHairStrandsOntoMesh(RHICmdList, ShaderMap, Q.LODIndex, MeshData, ProjectionHairData);
 				UpdateHairStrandsMeshTriangles(RHICmdList, ShaderMap, Q.LODIndex, HairStrandsTriangleType::RestPose, MeshData, ProjectionHairData);
 			}
@@ -295,7 +295,7 @@ void RunHairStrandsInterpolation(
 		{
 			if (EHairStrandsInterpolationType::RenderStrands == Type && 0 <= E.FrameLODIndex && E.FrameLODIndex < ProjectionHairData.LODDatas.Num() && ProjectionHairData.LODDatas[E.FrameLODIndex].bIsValid)
 			{
-				ProjectionHairData.LODDatas[E.FrameLODIndex].DeformedRootCenter = E.SkeletalWorldBoundsCenter;
+				ProjectionHairData.LODDatas[E.FrameLODIndex].DeformedPositionOffset = E.SkeletalDeformedPositionOffset;
 				UpdateHairStrandsMeshTriangles(RHICmdList, ShaderMap, E.FrameLODIndex, HairStrandsTriangleType::DeformedPose, MeshData, ProjectionHairData);
 			}
 		}
@@ -304,7 +304,7 @@ void RunHairStrandsInterpolation(
 		{
 			if (EHairStrandsInterpolationType::SimulationStrands == Type && 0 <= E.FrameLODIndex && E.FrameLODIndex < ProjectionHairData.LODDatas.Num() && ProjectionHairData.LODDatas[E.FrameLODIndex].bIsValid)
 			{
-				ProjectionHairData.LODDatas[E.FrameLODIndex].DeformedRootCenter = E.SkeletalWorldBoundsCenter;
+				ProjectionHairData.LODDatas[E.FrameLODIndex].DeformedPositionOffset = E.SkeletalDeformedPositionOffset;
 				UpdateHairStrandsMeshTriangles(RHICmdList, ShaderMap, E.FrameLODIndex, HairStrandsTriangleType::DeformedPose, MeshData, ProjectionHairData);
 			}
 		}
diff --git a/Engine/Source/Runtime/Renderer/Private/HairStrands/HairStrandsMeshProjection.cpp b/Engine/Source/Runtime/Renderer/Private/HairStrands/HairStrandsMeshProjection.cpp
--- a/Engine/Source/Runtime/Renderer/Private/HairStrands/HairStrandsMeshProjection.cpp
+++ b/Engine/Source/Runtime/Renderer/Private/HairStrands/HairStrandsMeshProjection.cpp
@@ -139,7 +139,7 @@ class FHairUpdateMeshTriangleCS : public FGlobalShader
 	SHADER_USE_PARAMETER_STRUCT(FHairUpdateMeshTriangleCS, FGlobalShader);
 
 	BEGIN_SHADER_PARAMETER_STRUCT(FParameters, )
-		SHADER_PARAMETER(FVector, RootTriangleWorldCenter)
+		SHADER_PARAMETER(FVector, RootTrianglePositionOffset)
 		SHADER_PARAMETER(uint32, MaxRootCount)
 		SHADER_PARAMETER(FMatrix, RootLocalToWorld)
 
@@ -199,14 +199,14 @@ static void AddHairStrandUpdateMeshTrianglesPass(
 
 	if (Type == HairStrandsTriangleType::RestPose)
 	{
-		Parameters->RootTriangleWorldCenter = FVector::ZeroVector; // TODO  LODData.RestRootCenter;
+		Parameters->RootTrianglePositionOffset = LODData.RestPositionOffset;
 		Parameters->OutRootTrianglePosition0 = LODData.RestRootTrianglePosition0Buffer->UAV;
 		Parameters->OutRootTrianglePosition1 = LODData.RestRootTrianglePosition1Buffer->UAV;
 		Parameters->OutRootTrianglePosition2 = LODData.RestRootTrianglePosition2Buffer->UAV;
 	}
 	else if (Type == HairStrandsTriangleType::DeformedPose)
 	{
-		Parameters->RootTriangleWorldCenter = FVector::ZeroVector; // TODO LODData.DeformedRootCenter;
+		Parameters->RootTrianglePositionOffset = LODData.DeformedPositionOffset;
 		Parameters->OutRootTrianglePosition0 = LODData.DeformedRootTrianglePosition0Buffer->UAV;
 		Parameters->OutRootTrianglePosition1 = LODData.DeformedRootTrianglePosition1Buffer->UAV;
 		Parameters->OutRootTrianglePosition2 = LODData.DeformedRootTrianglePosition2Buffer->UAV;
diff --git a/Engine/Source/Runtime/Renderer/Public/HairStrandsInterface.h b/Engine/Source/Runtime/Renderer/Public/HairStrandsInterface.h
--- a/Engine/Source/Runtime/Renderer/Public/HairStrandsInterface.h
+++ b/Engine/Source/Runtime/Renderer/Public/HairStrandsInterface.h
@@ -96,13 +96,13 @@ struct FHairStrandsProjectionHairData
 		FRWBuffer* RootTriangleBarycentricBuffer = nullptr;
 
 		// Rest root triangles' positions are relative to root center (for preserving precision)
-		FVector RestRootCenter = FVector::ZeroVector;
+		FVector RestPositionOffset = FVector::ZeroVector;
 		FRWBuffer* RestRootTrianglePosition0Buffer = nullptr;
 		FRWBuffer* RestRootTrianglePosition1Buffer = nullptr;
 		FRWBuffer* RestRootTrianglePosition2Buffer = nullptr;
 
 		// Deformed root triangles' positions are relative to root center (for preserving precision)
-		FVector DeformedRootCenter = FVector::ZeroVector;
+		FVector DeformedPositionOffset = FVector::ZeroVector;
 		FRWBuffer* DeformedRootTrianglePosition0Buffer = nullptr;
 		FRWBuffer* DeformedRootTrianglePosition1Buffer = nullptr;
 		FRWBuffer* DeformedRootTrianglePosition2Buffer = nullptr;
@@ -178,7 +178,7 @@ RENDERER_API bool UpdateHairStrands(
 	EWorldType::Type WorldType,
 	const FTransform& HairLocalToWorld,
 	const FTransform& MeshLocalToWorld,
-	const FVector& SkeletalWorldBoundsCenter);
+	const FVector& SkeletalDeformedPositionOffset);
 
 RENDERER_API bool UpdateHairStrands(
 	uint64 Id, 
