Fix incorrect bounding box when using skin mesh.

(ushell-p4-cherrypick of 9488469 by Charles.deRousiers)

diff --git a/Engine/Plugins/Experimental/HairStrands/Source/HairStrandsCore/Private/GroomComponent.cpp b/Engine/Plugins/Experimental/HairStrands/Source/HairStrandsCore/Private/GroomComponent.cpp
--- a/Engine/Plugins/Experimental/HairStrands/Source/HairStrandsCore/Private/GroomComponent.cpp
+++ b/Engine/Plugins/Experimental/HairStrands/Source/HairStrandsCore/Private/GroomComponent.cpp
@@ -414,22 +414,23 @@ FBoxSphereBounds UGroomComponent::CalcBounds(const FTransform& InLocalToWorld) c
 	{
 		if (RegisteredSkeletalMeshComponent)
 		{
-			FBoxSphereBounds WorldBounds(EForceInit::ForceInit);
+			FBox WorldBox(EForceInit::ForceInit);
 			for (const FHairGroupData& GroupData : GroomAsset->HairGroupsData)
 			{
 				// Transform the bounding box with an extra offset coming from the skin animation
-				const FBoxSphereBounds MeshLocalBound = RegisteredSkeletalMeshComponent->CalcBounds(FTransform::Identity);
+				const FVector MeshTranslation = RegisteredSkeletalMeshComponent->CalcBounds(FTransform::Identity).GetBox().GetCenter();
 				const FVector HairTranslation = GroupData.HairRenderData.BoundingBox.GetCenter();
-				const FVector MeshTranslation = MeshLocalBound.GetSphere().Center;
 				const FVector LocalTranslation = MeshTranslation - HairTranslation;
 
 				FTransform LocalToWorld = InLocalToWorld;
 				LocalToWorld.SetLocation(LocalTranslation + InLocalToWorld.GetLocation());
-				FBoxSphereBounds GroupWorldBound(GroupData.HairRenderData.BoundingBox.TransformBy(LocalToWorld));
+				FBox GroupWorldBound = GroupData.HairRenderData.BoundingBox.TransformBy(LocalToWorld);
 			
-				WorldBounds = Union(GroupWorldBound, WorldBounds);
+				WorldBox += GroupWorldBound;
 			}
-			return WorldBounds;
+
+
+			return FBoxSphereBounds(WorldBox);
 
 		}
 		else
