Fix hair projection on skinned mesh in blueprint.

(ushell-p4-cherrypick of 9530920 by Charles.deRousiers)

diff --git a/Engine/Plugins/Experimental/HairStrands/Source/HairStrandsCore/Private/GroomComponent.cpp b/Engine/Plugins/Experimental/HairStrands/Source/HairStrandsCore/Private/GroomComponent.cpp
--- a/Engine/Plugins/Experimental/HairStrands/Source/HairStrandsCore/Private/GroomComponent.cpp
+++ b/Engine/Plugins/Experimental/HairStrands/Source/HairStrandsCore/Private/GroomComponent.cpp
@@ -27,6 +27,9 @@ static FAutoConsoleVariableRef CVarHairStrandsMeshProjectionForceLOD(TEXT("r.Hai
 static int32 GHairStrandsMeshUseRelativePosition = 1;
 static FAutoConsoleVariableRef CVarHairStrandsMeshUseRelativePosition(TEXT("r.HairStrands.MeshProjection.RelativePosition"), GHairStrandsMeshUseRelativePosition, TEXT("Enable/Disable relative triangle position for improving positions"));
 
+static int32 GHairStrandsMeshProjectionTickDelay = 1;
+static FAutoConsoleVariableRef CVarHairStrandsMeshProjectionTickDelay(TEXT("r.HairStrands.MeshProjection.TickDelay"), GHairStrandsMeshProjectionTickDelay, TEXT("Number of simulation tick to wait before projecting a groom onto a mesh"));
+
 /**
  * An material render proxy which overrides the debug mode parameter.
  */
@@ -806,13 +809,13 @@ void UGroomComponent::ReleaseResources()
 	MeshProjectionTickDelay = 0;
 	MeshProjectionState = EMeshProjectionState::Invalid;
 
-
 	// Insure the ticking of the Groom component always happens after the skeletalMeshComponent.
 	if (RegisteredSkeletalMeshComponent)
 	{
 		RemoveTickPrerequisiteComponent(RegisteredSkeletalMeshComponent);
-		RegisteredSkeletalMeshComponent = nullptr;
 	}
+	SkeletalPreviousPositionOffset = FVector::ZeroVector;
+	RegisteredSkeletalMeshComponent = nullptr;
 }
 
 void UGroomComponent::PostLoad()
@@ -856,6 +859,15 @@ void UGroomComponent::OnComponentDestroyed(bool bDestroyingHierarchy)
 	Super::OnComponentDestroyed(bDestroyingHierarchy);
 }
 
+void UGroomComponent::OnAttachmentChanged()
+{
+	Super::OnAttachmentChanged();
+	if (GroomAsset && !IsBeingDestroyed() && HasBeenCreated())
+	{
+		InitResources();
+	}
+}
+
 void UGroomComponent::TickComponent(float DeltaTime, enum ELevelTick TickType, FActorComponentTickFunction *ThisTickFunction)
 {
 	Super::TickComponent(DeltaTime, TickType, ThisTickFunction);	
@@ -923,35 +935,68 @@ void UGroomComponent::TickComponent(float DeltaTime, enum ELevelTick TickType, F
 			});
 		}
 
-		if (MeshProjectionState != EMeshProjectionState::Completed && MeshProjectionTickDelay == 0)
+		// State machine for binding/projecting hair onto skinned mesh. The state machine operates as follow:
+		// * For each mesh LOD
+		//   * Force the skeletal mesh in rest pose
+		//   * Wait for the skeletal mesh to tick/update the animation
+		//   * Issue a binding/projection query to the hair strands system
+		//   * Wait for the query to be completed
+		// * Put the skeletal mesh into its original pose
+		//
+		// In editor, the animation is not ticked, so we manually issue a re-initialization of the animation 
+		// to force the animation in rest pose or original pose
+		if (MeshProjectionState != EMeshProjectionState::Completed)
 		{
-			if (MeshProjectionState == EMeshProjectionState::Invalid)
-			{
-				MeshProjectionLODIndex = 0;
-				MeshProjectionState = EMeshProjectionState::WaitForData;
-			}
+			const uint32 TickDelay = FMath::Max(0, GHairStrandsMeshProjectionTickDelay);
 
-			if (MeshProjectionLODIndex == SkeletalMeshComponent->GetNumLODs())
+			if (MeshProjectionTickDelay == 0)
 			{
-				SkeletalMeshComponent->SetForceRefPose(false);
-				SkeletalMeshComponent->SetForcedLOD(0);
-				MeshProjectionState = EMeshProjectionState::Completed;
-			}
+				if (MeshProjectionState == EMeshProjectionState::Invalid)
+				{
+					MeshProjectionLODIndex = 0;
+					MeshProjectionState = EMeshProjectionState::InProgressBinding;
+				}
 
-			if (MeshProjectionLODIndex < SkeletalMeshComponent->GetNumLODs())
-			{
-				SkeletalMeshComponent->SetForceRefPose(true);
-				SkeletalMeshComponent->SetForcedLOD(MeshProjectionLODIndex + 1);
+				if (MeshProjectionLODIndex == SkeletalMeshComponent->GetNumLODs())
+				{
+					SkeletalMeshComponent->SetForceRefPose(false);
+					SkeletalMeshComponent->SetForcedLOD(0);
+					if (WorldType == EWorldType::Editor)
+					{
+						SkeletalMeshComponent->InitAnim(true);
+					}
+					MeshProjectionState = EMeshProjectionState::Completed;
+				}
 
-				const FVector RestRootCenter = MeshPositionOffset;
-				const uint32 LODIndex = MeshProjectionLODIndex;
-				ENQUEUE_RENDER_COMMAND(FHairStrandsTick_Projection)(
-					[Id, WorldType, FeatureLevel, LODIndex, RestRootCenter](FRHICommandListImmediate& RHICmdList)
+				if (MeshProjectionLODIndex < SkeletalMeshComponent->GetNumLODs())
 				{
-					AddHairStrandsProjectionQuery(RHICmdList, Id, WorldType, LODIndex, RestRootCenter);
-				});
-				MeshProjectionTickDelay += 5;
-				++MeshProjectionLODIndex;
+					if (MeshProjectionState == EMeshProjectionState::InProgressBinding)
+					{
+						SkeletalMeshComponent->SetForceRefPose(true);
+						SkeletalMeshComponent->SetForcedLOD(MeshProjectionLODIndex + 1);
+						if (WorldType == EWorldType::Editor)
+						{
+							SkeletalMeshComponent->InitAnim(true);
+						}
+						MeshProjectionState = EMeshProjectionState::WaitForRestPose;
+						MeshProjectionTickDelay += TickDelay;
+					}
+
+					if (MeshProjectionState == EMeshProjectionState::WaitForRestPose && MeshProjectionTickDelay == 0)
+					{
+						const FVector RestPositionOffset = MeshPositionOffset;
+						const uint32 LODIndex = MeshProjectionLODIndex;
+						ENQUEUE_RENDER_COMMAND(FHairStrandsTick_Projection)(
+							[Id, WorldType, FeatureLevel, LODIndex, RestPositionOffset](FRHICommandListImmediate& RHICmdList)
+						{
+							AddHairStrandsProjectionQuery(RHICmdList, Id, WorldType, LODIndex, RestPositionOffset);
+						});
+
+						MeshProjectionState = EMeshProjectionState::InProgressBinding;
+						MeshProjectionTickDelay += TickDelay;
+						++MeshProjectionLODIndex;
+					}
+				}
 			}
 		}
 		else
diff --git a/Engine/Plugins/Experimental/HairStrands/Source/HairStrandsCore/Public/GroomComponent.h b/Engine/Plugins/Experimental/HairStrands/Source/HairStrandsCore/Public/GroomComponent.h
--- a/Engine/Plugins/Experimental/HairStrands/Source/HairStrandsCore/Public/GroomComponent.h
+++ b/Engine/Plugins/Experimental/HairStrands/Source/HairStrandsCore/Public/GroomComponent.h
@@ -44,6 +44,7 @@ public:
 	virtual void OnRegister() override;
 	virtual void OnUnregister() override;
 	virtual void OnComponentDestroyed(bool bDestroyingHierarchy) override;
+	virtual void OnAttachmentChanged() override;
 	virtual void TickComponent(float DeltaTime, enum ELevelTick TickType, FActorComponentTickFunction *ThisTickFunction) override;
 	virtual void SendRenderTransform_Concurrent() override;
 	//~ End UActorComponent Interface.
@@ -107,7 +108,8 @@ private:
 	enum class EMeshProjectionState
 	{
 		Invalid,
-		WaitForData,
+		InProgressBinding,
+		WaitForRestPose,
 		Completed
 	};
 	class USkeletalMeshComponent* RegisteredSkeletalMeshComponent;
diff --git a/Engine/Source/Runtime/Engine/Private/Components/SkeletalMeshComponent.cpp b/Engine/Source/Runtime/Engine/Private/Components/SkeletalMeshComponent.cpp
--- a/Engine/Source/Runtime/Engine/Private/Components/SkeletalMeshComponent.cpp
+++ b/Engine/Source/Runtime/Engine/Private/Components/SkeletalMeshComponent.cpp
@@ -648,7 +648,7 @@ void USkeletalMeshComponent::InitAnim(bool bForceReinit)
 		// In Editor, animations won't get ticked. So Update once to get accurate representation instead of T-Pose.
 		// Also allow this to be an option to support pre-4.19 games that might need it..
 		const bool bTickAnimationNow =
-			((GetWorld()->WorldType == EWorldType::Editor) && !bUseRefPoseOnInitAnim)
+			((GetWorld()->WorldType == EWorldType::Editor) && !bUseRefPoseOnInitAnim && !bForceRefpose)
 			|| UAnimationSettings::Get()->bTickAnimationOnSkeletalMeshInit;
 
 		const bool bInitializedAnimInstance = InitializeAnimScriptInstance(bForceReinit, !bTickAnimationNow);
