Fix bounding box update under rotation.

(ushell-p4-cherrypick of 9542073 by Charles.deRousiers)

(ushell-p4-cherrypick of 9542098 by Charles.deRousiers)

diff --git a/Engine/Plugins/Experimental/HairStrands/Source/HairStrandsCore/Private/GroomComponent.cpp b/Engine/Plugins/Experimental/HairStrands/Source/HairStrandsCore/Private/GroomComponent.cpp
--- a/Engine/Plugins/Experimental/HairStrands/Source/HairStrandsCore/Private/GroomComponent.cpp
+++ b/Engine/Plugins/Experimental/HairStrands/Source/HairStrandsCore/Private/GroomComponent.cpp
@@ -417,24 +417,19 @@ FBoxSphereBounds UGroomComponent::CalcBounds(const FTransform& InLocalToWorld) c
 	{
 		if (RegisteredSkeletalMeshComponent)
 		{
-			FBox WorldBox(EForceInit::ForceInit);
+			FBox LocalBound(EForceInit::ForceInit);
 			for (const FHairGroupData& GroupData : GroomAsset->HairGroupsData)
 			{
-				// Transform the bounding box with an extra offset coming from the skin animation
-				const FVector MeshTranslation = RegisteredSkeletalMeshComponent->CalcBounds(FTransform::Identity).GetBox().GetCenter();
-				const FVector HairTranslation = GroupData.HairRenderData.BoundingBox.GetCenter();
-				const FVector LocalTranslation = MeshTranslation - HairTranslation;
-
-				FTransform LocalToWorld = InLocalToWorld;
-				LocalToWorld.SetLocation(LocalTranslation + InLocalToWorld.GetLocation());
-				FBox GroupWorldBound = GroupData.HairRenderData.BoundingBox.TransformBy(LocalToWorld);
-			
-				WorldBox += GroupWorldBound;
+				LocalBound += GroupData.HairRenderData.BoundingBox;
 			}
+			FBox WorldBound = LocalBound.TransformBy(InLocalToWorld);
 
+			const FVector MeshTranslation = RegisteredSkeletalMeshComponent->CalcBounds(InLocalToWorld).GetBox().GetCenter();
+			const FVector LocalAnimationTranslation = MeshTranslation - WorldBound.GetCenter();
 
-			return FBoxSphereBounds(WorldBox);
-
+			WorldBound.Max += LocalAnimationTranslation;
+			WorldBound.Min += LocalAnimationTranslation;
+			return FBoxSphereBounds(WorldBound);
 		}
 		else
 		{
