Enable by default 32bit precision position for skinned triangle positions.

This resolves the jittering issue due to lack of precision when using skinning on Megan.

(ushell-p4-cherrypick of 9545727 by Charles.deRousiers)

diff --git a/Engine/Plugins/Experimental/HairStrands/Source/HairStrandsCore/Public/HairStrandsDatas.h b/Engine/Plugins/Experimental/HairStrands/Source/HairStrandsCore/Public/HairStrandsDatas.h
--- a/Engine/Plugins/Experimental/HairStrands/Source/HairStrandsCore/Public/HairStrandsDatas.h
+++ b/Engine/Plugins/Experimental/HairStrands/Source/HairStrandsCore/Public/HairStrandsDatas.h
@@ -158,7 +158,7 @@ struct FHairStrandsWeightFormat
 */
 struct FHairStrandsMeshTrianglePositionFormat
 {
-#if 0
+#if 1
 	using Type = FVector4;
 	static const uint32 ComponentCount = 1;
 	static const uint32 SizeInByte = sizeof(Type);
