Fix wrong primitive count when updating hair geometry bottom BVH.

(ushell-p4-cherrypick of 9546589 by Charles.deRousiers)

diff --git a/Engine/Plugins/Experimental/HairStrands/Source/HairStrandsCore/Private/HairStrandsRendering.cpp b/Engine/Plugins/Experimental/HairStrands/Source/HairStrandsCore/Private/HairStrandsRendering.cpp
--- a/Engine/Plugins/Experimental/HairStrands/Source/HairStrandsCore/Private/HairStrandsRendering.cpp
+++ b/Engine/Plugins/Experimental/HairStrands/Source/HairStrandsCore/Private/HairStrandsRendering.cpp
@@ -353,7 +353,7 @@ static void BuildHairAccelerationStructure(FRHICommandList& RHICmdList, uint32 R
 	Initializer.IndexBuffer = nullptr;
 	Initializer.IndexBufferOffset = 0;
 	Initializer.GeometryType = RTGT_Triangles;
-	Initializer.TotalPrimitiveCount = RaytracingVertexCount;
+	Initializer.TotalPrimitiveCount = RaytracingVertexCount/3;
 	Initializer.bFastBuild = true;
 	Initializer.bAllowUpdate = true;
 
