Fix issue memory leaks due to skeletal mesh's skin cache resource being recreated every frame.

The groom component was resetting the skeletal mesh component's rest pose flag at every tick triggering a chain of update.

(ushell-p4-cherrypick of 9546595 by Charles.deRousiers)

diff --git a/Engine/Plugins/Experimental/HairStrands/Source/HairStrandsCore/Private/GroomComponent.cpp b/Engine/Plugins/Experimental/HairStrands/Source/HairStrandsCore/Private/GroomComponent.cpp
--- a/Engine/Plugins/Experimental/HairStrands/Source/HairStrandsCore/Private/GroomComponent.cpp
+++ b/Engine/Plugins/Experimental/HairStrands/Source/HairStrandsCore/Private/GroomComponent.cpp
@@ -994,12 +994,15 @@ void UGroomComponent::TickComponent(float DeltaTime, enum ELevelTick TickType, F
 				}
 			}
 		}
-		else
+
+		FSkeletalMeshConfiguration CurrentState;
+		CurrentState.ForceRefPose = GHairStrandsMeshProjectionForceRefPoseEnable > 0;
+		CurrentState.ForceLOD     = GHairStrandsMeshProjectionForceLOD >= 0 ? FMath::Clamp(GHairStrandsMeshProjectionForceLOD, 0, SkeletalMeshComponent->GetNumLODs() - 1) : -1;
+		if (!FSkeletalMeshConfiguration::Equals(CurrentState, SkeletalMeshConfiguration) && MeshProjectionState == EMeshProjectionState::Completed)
 		{
-			// For debug purpose
-			const int32 ForceLOD = GHairStrandsMeshProjectionForceLOD >= 0 ? FMath::Clamp(GHairStrandsMeshProjectionForceLOD, 0, SkeletalMeshComponent->GetNumLODs() - 1) : -1;
-			SkeletalMeshComponent->SetForcedLOD(ForceLOD + 1);
-			SkeletalMeshComponent->SetForceRefPose(GHairStrandsMeshProjectionForceRefPoseEnable > 0);
+			SkeletalMeshComponent->SetForcedLOD(CurrentState.ForceLOD + 1);
+			SkeletalMeshComponent->SetForceRefPose(CurrentState.ForceRefPose);
+			SkeletalMeshConfiguration = CurrentState;
 		}
 
 		// For skinned mesh update the relative center of hair positions after deformation
diff --git a/Engine/Plugins/Experimental/HairStrands/Source/HairStrandsCore/Public/GroomComponent.h b/Engine/Plugins/Experimental/HairStrands/Source/HairStrandsCore/Public/GroomComponent.h
--- a/Engine/Plugins/Experimental/HairStrands/Source/HairStrandsCore/Public/GroomComponent.h
+++ b/Engine/Plugins/Experimental/HairStrands/Source/HairStrandsCore/Public/GroomComponent.h
@@ -117,6 +117,17 @@ private:
 	int32 MeshProjectionLODIndex;
 	uint32 MeshProjectionTickDelay;
 	EMeshProjectionState MeshProjectionState;
+	
+	struct FSkeletalMeshConfiguration
+	{
+		int32 ForceLOD = -1;
+		bool ForceRefPose = false;
+		static bool Equals(const FSkeletalMeshConfiguration& A, const FSkeletalMeshConfiguration& B)
+		{
+			return A.ForceLOD == B.ForceLOD && A.ForceRefPose == B.ForceRefPose;
+		}
+	};
+	FSkeletalMeshConfiguration SkeletalMeshConfiguration;
 
 	void InitResources();
 	void ReleaseResources();
