Fixed crash when importing hair with "Override Guides" option an alembic that has guides.

(ushell-p4-cherrypick of 9604283 by [REDACTED_PERSON])

diff --git a/Engine/Plugins/Experimental/HairStrands/Source/HairStrandsCore/Private/GroomBuilder.cpp b/Engine/Plugins/Experimental/HairStrands/Source/HairStrandsCore/Private/GroomBuilder.cpp
--- a/Engine/Plugins/Experimental/HairStrands/Source/HairStrandsCore/Private/GroomBuilder.cpp
+++ b/Engine/Plugins/Experimental/HairStrands/Source/HairStrandsCore/Private/GroomBuilder.cpp
@@ -771,6 +771,11 @@ bool FGroomBuilder::BuildGroom(const FHairDescription& HairDescription, const FG
 
 			++GroupInfo.NumGuides;
 		}
+		else
+		{
+			// A guide but don't want to import it, so skip it
+			continue;
+		}
 
 		CurrentHairStrandsDatas->StrandsCurves.CurvesCount.Add(CurveNumVertices);
 
