Add/update position offset for rest/deformed sim/render triangle positions, and rest/deformed sim/render hair positions.

(ushell-p4-cherrypick of 9606096 by Charles.deRousiers)

diff --git a/Engine/Plugins/Experimental/HairStrands/Source/HairStrandsCore/Private/GroomAsset.cpp b/Engine/Plugins/Experimental/HairStrands/Source/HairStrandsCore/Private/GroomAsset.cpp
--- a/Engine/Plugins/Experimental/HairStrands/Source/HairStrandsCore/Private/GroomAsset.cpp
+++ b/Engine/Plugins/Experimental/HairStrands/Source/HairStrandsCore/Private/GroomAsset.cpp
@@ -41,8 +41,8 @@ void CreateBuffer(uint32 InVertexCount, FRWBuffer& OutBuffer)
 
 /////////////////////////////////////////////////////////////////////////////////////////
 
-FHairStrandsRestResource::FHairStrandsRestResource(const FHairStrandsDatas::FRenderData& HairStrandRenderData) : 
-	RestPositionBuffer(), AttributeBuffer(), RenderData(HairStrandRenderData)
+FHairStrandsRestResource::FHairStrandsRestResource(const FHairStrandsDatas::FRenderData& HairStrandRenderData, const FVector& InPositionOffset) :
+	RestPositionBuffer(), AttributeBuffer(), RenderData(HairStrandRenderData), PositionOffset(InPositionOffset)
 {}
 
 void FHairStrandsRestResource::InitRHI()
@@ -102,7 +102,6 @@ FHairStrandsRootResource::FHairStrandsRootResource(const FHairStrandsDatas* Hair
 	RootPositions.SetNum(RootCount);
 	RootNormals.SetNum(RootCount);
 
-
 	for (uint32 CurveIndex = 0; CurveIndex < CurveCount; ++CurveIndex)
 	{
 		const uint32 RootIndex = HairStrandsDatas->StrandsCurves.CurvesOffset[CurveIndex];
@@ -145,6 +144,7 @@ FHairStrandsRootResource::FHairStrandsRootResource(const FHairStrandsDatas* Hair
 	for (FMeshProjectionLOD& MeshProjectionLOD : MeshProjectionLODs)
 	{
 		MeshProjectionLOD.Status = FMeshProjectionLOD::EStatus::Invalid;
+		MeshProjectionLOD.RestRootCenter = HairStrandsDatas->BoundingBox.GetCenter();
 		MeshProjectionLOD.LODIndex = LODIndex++;
 	}
 }
@@ -264,11 +264,11 @@ void UGroomAsset::InitResource()
 {
 	for (FHairGroupData& GroupData : HairGroupsData)
 	{
-		GroupData.HairStrandsRestResource = new FHairStrandsRestResource(GroupData.HairRenderData.RenderData);
+		GroupData.HairStrandsRestResource = new FHairStrandsRestResource(GroupData.HairRenderData.RenderData, GroupData.HairRenderData.BoundingBox.GetCenter());
 
 		BeginInitResource(GroupData.HairStrandsRestResource);
 
-		GroupData.HairSimulationRestResource = new FHairStrandsRestResource(GroupData.HairSimulationData.RenderData);
+		GroupData.HairSimulationRestResource = new FHairStrandsRestResource(GroupData.HairSimulationData.RenderData, GroupData.HairSimulationData.BoundingBox.GetCenter());
 
 		BeginInitResource(GroupData.HairSimulationRestResource);
 	}
diff --git a/Engine/Plugins/Experimental/HairStrands/Source/HairStrandsCore/Private/GroomComponent.cpp b/Engine/Plugins/Experimental/HairStrands/Source/HairStrandsCore/Private/GroomComponent.cpp
--- a/Engine/Plugins/Experimental/HairStrands/Source/HairStrandsCore/Private/GroomComponent.cpp
+++ b/Engine/Plugins/Experimental/HairStrands/Source/HairStrandsCore/Private/GroomComponent.cpp
@@ -679,16 +679,16 @@ void UGroomComponent::InitResources()
 		BeginInitResource(Res.RenderDeformedResources);
 		BeginInitResource(Res.SimDeformedResources);
 
-		const FVector HairBoundCenter = GroupData.HairRenderData.BoundingBox.GetCenter();
+		const FVector RestHairPositionOffset = Res.RenderRestResources->PositionOffset;
 
 		FHairStrandsInterpolationOutput::HairGroup& InterpolationOutputGroup = InterpolationOutput->HairGroups.AddDefaulted_GetRef();
 		FHairStrandsInterpolationInput::FHairGroup& InterpolationInputGroup = InterpolationInput->HairGroups.AddDefaulted_GetRef();
 		InterpolationInputGroup.HairRadius = GroupData.HairRenderData.StrandsCurves.MaxRadius;
 		InterpolationInputGroup.HairRaytracingRadiusScale = HairRaytracingRadiusScale;
-		InterpolationInputGroup.InHairPositionOffset = HairBoundCenter;
+		InterpolationInputGroup.InHairPositionOffset = RestHairPositionOffset;
 		// For skinned groom, these value will be updated during TickComponent() call
-		InterpolationInputGroup.OutHairPositionOffset = HairBoundCenter; 
-		InterpolationInputGroup.OutHairPreviousPositionOffset = HairBoundCenter;
+		InterpolationInputGroup.OutHairPositionOffset = RestHairPositionOffset;
+		InterpolationInputGroup.OutHairPreviousPositionOffset = RestHairPositionOffset;
 	}
 
 	FHairStrandsInterpolationData Interpolation;
@@ -979,12 +979,34 @@ void UGroomComponent::TickComponent(float DeltaTime, enum ELevelTick TickType, F
 
 					if (MeshProjectionState == EMeshProjectionState::WaitForRestPose && MeshProjectionTickDelay == 0)
 					{
+						FHairGroupResources* LocalResources = &HairGroupResources;
+
 						const FVector RestPositionOffset = MeshPositionOffset;
 						const uint32 LODIndex = MeshProjectionLODIndex;
 						ENQUEUE_RENDER_COMMAND(FHairStrandsTick_Projection)(
-							[Id, WorldType, FeatureLevel, LODIndex, RestPositionOffset](FRHICommandListImmediate& RHICmdList)
+							[Id, WorldType, FeatureLevel, LODIndex, RestPositionOffset, LocalResources](FRHICommandListImmediate& RHICmdList)
 						{
 							AddHairStrandsProjectionQuery(RHICmdList, Id, WorldType, LODIndex, RestPositionOffset);
+
+							// Update rest (sim/render) position offsets. This is used by Niagara.
+							for (FHairGroupResource& Res : *LocalResources)
+							{
+								if (Res.RenRootResources)
+								{
+									for (FHairStrandsRootResource::FMeshProjectionLOD& MeshProjectionLOD : Res.RenRootResources->MeshProjectionLODs)
+									{
+										MeshProjectionLOD.RestRootCenter = RestPositionOffset;
+									}
+								}
+
+								if (Res.SimRootResources)
+								{
+									for (FHairStrandsRootResource::FMeshProjectionLOD& MeshProjectionLOD : Res.SimRootResources->MeshProjectionLODs)
+									{
+										MeshProjectionLOD.RestRootCenter = RestPositionOffset;
+									}
+								}
+							}
 						});
 
 						MeshProjectionState = EMeshProjectionState::InProgressBinding;
@@ -1007,17 +1029,29 @@ void UGroomComponent::TickComponent(float DeltaTime, enum ELevelTick TickType, F
 
 		// For skinned mesh update the relative center of hair positions after deformation
 		{
+			FHairGroupResources* LocalResources = &HairGroupResources;
+
 			const FVector OutHairPositionOffset = MeshPositionOffset;
 			const FVector OutHairPreviousPositionOffset = SkeletalPreviousPositionOffset;
 			FHairStrandsInterpolationInput* LocalInterpolationInput = InterpolationInput;
 			ENQUEUE_RENDER_COMMAND(FHairStrandsTick_OutHairPositionOffsetUpdate)(
-				[OutHairPositionOffset, OutHairPreviousPositionOffset, LocalInterpolationInput](FRHICommandListImmediate& RHICmdList)
+				[OutHairPositionOffset, OutHairPreviousPositionOffset, LocalInterpolationInput, LocalResources](FRHICommandListImmediate& RHICmdList)
 			{
 				for (FHairStrandsInterpolationInput::FHairGroup& HairGroup : LocalInterpolationInput->HairGroups)
 				{
 					HairGroup.OutHairPositionOffset = OutHairPositionOffset;
 					HairGroup.OutHairPreviousPositionOffset = OutHairPreviousPositionOffset;
 				}
+
+				// Update deformed (sim/render) hair position offsets. This is used by Niagara.
+				for (FHairGroupResource& Res : *LocalResources)
+				{
+					if (Res.RenderDeformedResources)
+						Res.RenderDeformedResources->PositionOffset = OutHairPositionOffset;
+
+					if (Res.SimDeformedResources)
+						Res.SimDeformedResources->PositionOffset = OutHairPositionOffset;
+				}
 			});
 
 			// First frame will be wrong ... 
@@ -1035,16 +1069,36 @@ void UGroomComponent::TickComponent(float DeltaTime, enum ELevelTick TickType, F
 		MeshProjectionTickDelay--;
 	}
 
+	FHairGroupResources* LocalResources = &HairGroupResources;
 	const FVector DeformedPositionCenter = GHairStrandsMeshUseRelativePosition > 0 ? MeshPositionOffset : FVector::ZeroVector;
 	const FTransform SkinLocalToWorld = SkeletalMeshComponent ? SkeletalMeshComponent->GetComponentTransform() : FTransform();
 	const FTransform HairLocalToWorld = GetComponentTransform();
 	ENQUEUE_RENDER_COMMAND(FHairStrandsTick_TransformUpdate)(
-		[Id, WorldType, HairLocalToWorld, SkinLocalToWorld, DeformedPositionCenter, FeatureLevel](FRHICommandListImmediate& RHICmdList)
+		[Id, WorldType, HairLocalToWorld, SkinLocalToWorld, DeformedPositionCenter, FeatureLevel, LocalResources](FRHICommandListImmediate& RHICmdList)
 	{		
 		if (ERHIFeatureLevel::Num == FeatureLevel)
 			return;
 
 		UpdateHairStrands(Id, WorldType, HairLocalToWorld, SkinLocalToWorld, DeformedPositionCenter);
+
+		// Update deformed (sim/render) triangles position offsets. This is used by Niagara.
+		for (FHairGroupResource& Res : *LocalResources)
+		{
+			if (Res.RenRootResources)
+			{
+				for (FHairStrandsRootResource::FMeshProjectionLOD& MeshProjectionLOD : Res.RenRootResources->MeshProjectionLODs)
+				{
+					MeshProjectionLOD.DeformedRootCenter = DeformedPositionCenter;
+				}
+			}
+			if (Res.SimRootResources)
+			{
+				for (FHairStrandsRootResource::FMeshProjectionLOD& MeshProjectionLOD : Res.SimRootResources->MeshProjectionLODs)
+				{
+					MeshProjectionLOD.DeformedRootCenter = DeformedPositionCenter;
+				}
+			}
+		}
 	});
 }
 
diff --git a/Engine/Plugins/Experimental/HairStrands/Source/HairStrandsCore/Public/GroomAsset.h b/Engine/Plugins/Experimental/HairStrands/Source/HairStrandsCore/Public/GroomAsset.h
--- a/Engine/Plugins/Experimental/HairStrands/Source/HairStrandsCore/Public/GroomAsset.h
+++ b/Engine/Plugins/Experimental/HairStrands/Source/HairStrandsCore/Public/GroomAsset.h
@@ -76,7 +76,7 @@ struct FHairStrandsRootResource : public FRenderResource
 struct FHairStrandsRestResource : public FRenderResource
 {
 	/** Build the hair strands resource */
-	FHairStrandsRestResource(const FHairStrandsDatas::FRenderData& HairStrandRenderData);
+	FHairStrandsRestResource(const FHairStrandsDatas::FRenderData& HairStrandRenderData, const FVector& PositionOffset);
 
 	/* Init the buffer */
 	virtual void InitRHI() override;
@@ -93,6 +93,9 @@ struct FHairStrandsRestResource : public FRenderResource
 	/* Strand hair offset buffer */
 	FRWBuffer AttributeBuffer;
 
+	/* Position offset as the rest positions are expressed in relative coordinate (16bits) */
+	FVector PositionOffset = FVector::ZeroVector;
+
 	/* Reference to the hair strands render data */
 	const FHairStrandsDatas::FRenderData& RenderData;
 };
@@ -117,6 +120,9 @@ struct FHairStrandsDeformedResource : public FRenderResource
 	/* Strand hair tangent buffer */
 	FRWBuffer TangentBuffer;
 
+	/* Position offset as the deformed positions are expressed in relative coordinate (16bits) */
+	FVector PositionOffset = FVector::ZeroVector;
+
 	/* Reference to the hair strands render data */
 	const FHairStrandsDatas::FRenderData& RenderData;
 
