Fix warning.

(ushell-p4-cherrypick of 9616524 by Charles.deRousiers)

diff --git a/Engine/Plugins/Experimental/HairStrands/Source/HairStrandsCore/Private/GroomAsset.cpp b/Engine/Plugins/Experimental/HairStrands/Source/HairStrandsCore/Private/GroomAsset.cpp
--- a/Engine/Plugins/Experimental/HairStrands/Source/HairStrandsCore/Private/GroomAsset.cpp
+++ b/Engine/Plugins/Experimental/HairStrands/Source/HairStrandsCore/Private/GroomAsset.cpp
@@ -42,7 +42,7 @@ void CreateBuffer(uint32 InVertexCount, FRWBuffer& OutBuffer)
 /////////////////////////////////////////////////////////////////////////////////////////
 
 FHairStrandsRestResource::FHairStrandsRestResource(const FHairStrandsDatas::FRenderData& HairStrandRenderData, const FVector& InPositionOffset) :
-	RestPositionBuffer(), AttributeBuffer(), RenderData(HairStrandRenderData), PositionOffset(InPositionOffset)
+	RestPositionBuffer(), AttributeBuffer(), PositionOffset(InPositionOffset), RenderData(HairStrandRenderData)
 {}
 
 void FHairStrandsRestResource::InitRHI()
