Add groom icon.

(ushell-p4-cherrypick of 9617270 by Charles.deRousiers)

diff --git a/Engine/Plugins/Experimental/HairStrands/Source/HairStrandsCore/Private/GroomActor.cpp b/Engine/Plugins/Experimental/HairStrands/Source/HairStrandsCore/Private/GroomActor.cpp
--- a/Engine/Plugins/Experimental/HairStrands/Source/HairStrandsCore/Private/GroomActor.cpp
+++ b/Engine/Plugins/Experimental/HairStrands/Source/HairStrandsCore/Private/GroomActor.cpp
@@ -2,11 +2,54 @@
 
 #include "GroomActor.h"
 #include "GroomComponent.h"
+#include "Components/BillboardComponent.h"
+#include "UObject/ConstructorHelpers.h"
+#include "Internationalization/Internationalization.h"
 
+#define LOCTEXT_NAMESPACE "HairStrands"
 
 AGroomActor::AGroomActor(const FObjectInitializer& ObjectInitializer)
 	: Super(ObjectInitializer)
 {
 	GroomComponent = CreateDefaultSubobject<UGroomComponent>(TEXT("GroomComponent0"));
 	RootComponent = GroomComponent;
-}
\ No newline at end of file
+
+#if WITH_EDITOR
+	struct FConstructorStatics
+	{
+		FName ID_Sprite;
+		FText NAME_Sprite;
+		FColor SceneBaseColor;
+		FVector SceneBaseSize;
+		FConstructorStatics()
+			: ID_Sprite(TEXT("GroomActor"))
+			, NAME_Sprite(LOCTEXT("RootSpriteInfo", "Groom Actor"))
+			, SceneBaseColor(100, 255, 255, 255)
+			, SceneBaseSize(600.0f, 600.0f, 400.0f)
+		{
+		}
+	};
+	static FConstructorStatics ConstructorStatics;
+
+	SpriteComponent = CreateEditorOnlyDefaultSubobject<UBillboardComponent>(TEXT("Sprite"));
+	if (SpriteComponent)
+	{
+		ConstructorHelpers::FObjectFinderOptional<UTexture2D> RootTextureObject = TEXT("/HairStrands/Icons/S_GroomActor");
+		
+		SpriteComponent->Sprite = RootTextureObject.Get();
+		SpriteComponent->SetRelativeScale3D_Direct(FVector(1.0f, 1.0f, 1.0f));
+		SpriteComponent->bHiddenInGame = false;
+		SpriteComponent->SpriteInfo.Category = ConstructorStatics.ID_Sprite;
+		SpriteComponent->SpriteInfo.DisplayName = ConstructorStatics.NAME_Sprite;
+		SpriteComponent->bIsScreenSizeScaled = true;
+		SpriteComponent->SetupAttachment(RootComponent);
+	}
+#endif
+}
+
+
+
+
+
+
+
diff --git a/Engine/Plugins/Experimental/HairStrands/Source/HairStrandsCore/Public/GroomActor.h b/Engine/Plugins/Experimental/HairStrands/Source/HairStrandsCore/Public/GroomActor.h
--- a/Engine/Plugins/Experimental/HairStrands/Source/HairStrandsCore/Public/GroomActor.h
+++ b/Engine/Plugins/Experimental/HairStrands/Source/HairStrandsCore/Public/GroomActor.h
@@ -18,7 +18,14 @@ class HAIRSTRANDSCORE_API AGroomActor : public AActor
 	UPROPERTY(Category = StrandHair, VisibleAnywhere, BlueprintReadOnly)
 	class UGroomComponent* GroomComponent;
 
+#if WITH_EDITORONLY_DATA
+protected:
+	/** Billboard used to see the scene in the editor */
+	UPROPERTY()
+		UBillboardComponent* SpriteComponent;
+#endif
+
 public:
 	/** Returns GroomComponent subobject **/
 	class UGroomComponent* GetGroomComponent() const { return GroomComponent; }
-};
\ No newline at end of file
+};
diff --git a/Engine/Plugins/Experimental/HairStrands/Source/HairStrandsEditor/HairStrandsEditor.Build.cs b/Engine/Plugins/Experimental/HairStrands/Source/HairStrandsEditor/HairStrandsEditor.Build.cs
--- a/Engine/Plugins/Experimental/HairStrands/Source/HairStrandsEditor/HairStrandsEditor.Build.cs
+++ b/Engine/Plugins/Experimental/HairStrands/Source/HairStrandsEditor/HairStrandsEditor.Build.cs
@@ -21,6 +21,7 @@ namespace UnrealBuildTool.Rules
 					"MainFrame",
 					"Slate",
 					"SlateCore",
+					"Projects",
 					"UnrealEd",
 					"AssetTools",
 				});
diff --git a/Engine/Plugins/Experimental/HairStrands/Source/HairStrandsEditor/Private/HairStrandsEditor.cpp b/Engine/Plugins/Experimental/HairStrands/Source/HairStrandsEditor/Private/HairStrandsEditor.cpp
--- a/Engine/Plugins/Experimental/HairStrands/Source/HairStrandsEditor/Private/HairStrandsEditor.cpp
+++ b/Engine/Plugins/Experimental/HairStrands/Source/HairStrandsEditor/Private/HairStrandsEditor.cpp
@@ -3,6 +3,11 @@
 #include "HairStrandsEditor.h"
 #include "GroomActions.h"
 
+#include "Styling/SlateStyleRegistry.h"
+#include "Styling/SlateTypes.h"
+#include "EditorStyleSet.h"
+#include "Interfaces/IPluginManager.h"
+#include "SlateOptMacros.h"
 #include "FbxHairTranslator.h"
 
 IMPLEMENT_MODULE(FHairStrandsEditor, HairStrandsEditor);
@@ -16,6 +21,26 @@ void FHairStrandsEditor::StartupModule()
 	RegisteredAssetTypeActions.Add(AssetActions);
 
 	RegisterHairTranslator<FFbxHairTranslator>();
+
+	// Only register once
+	if (!StyleSet.IsValid())
+	{
+		const FVector2D Icon16x16(16.0f, 16.0f);
+		const FVector2D Icon64x64(64.0f, 64.0f);
+		FString HairStrandsContent = IPluginManager::Get().FindPlugin("HairStrands")->GetBaseDir() + "/Content";
+
+		StyleSet = MakeShareable(new FSlateStyleSet("Groom"));
+		StyleSet->SetContentRoot(FPaths::EngineContentDir() / TEXT("Editor/Slate"));
+		StyleSet->SetCoreContentRoot(FPaths::EngineContentDir() / TEXT("Slate"));
+		
+		StyleSet->Set("ClassIcon.GroomComponent", new FSlateImageBrush(HairStrandsContent + "/Icons/S_Groom_16.png", Icon16x16));
+		StyleSet->Set("ClassThumbnail.GroomComponent", new FSlateImageBrush(HairStrandsContent + "/Icons/S_Groom_64.png", Icon64x64));
+		
+		StyleSet->Set("ClassIcon.GroomActor", new FSlateImageBrush(HairStrandsContent + "/Icons/S_Groom_16.png", Icon16x16));
+		StyleSet->Set("ClassThumbnail.GroomActor", new FSlateImageBrush(HairStrandsContent + "/Icons/S_Groom_64.png", Icon64x64));
+		
+		FSlateStyleRegistry::RegisterSlateStyle(*StyleSet.Get());
+	}
 }
 
 void FHairStrandsEditor::ShutdownModule()
@@ -32,6 +57,13 @@ void FHairStrandsEditor::ShutdownModule()
 			AssetTools.UnregisterAssetTypeActions(Action);
 		}
 	}
+
+	if (StyleSet.IsValid())
+	{
+		FSlateStyleRegistry::UnRegisterSlateStyle(*StyleSet.Get());
+		ensure(StyleSet.IsUnique());
+		StyleSet.Reset();
+	}
 }
 
 TArray<TSharedPtr<IHairStrandsTranslator>> FHairStrandsEditor::GetHairTranslators()
diff --git a/Engine/Plugins/Experimental/HairStrands/Source/HairStrandsEditor/Public/HairStrandsEditor.h b/Engine/Plugins/Experimental/HairStrands/Source/HairStrandsEditor/Public/HairStrandsEditor.h
--- a/Engine/Plugins/Experimental/HairStrands/Source/HairStrandsEditor/Public/HairStrandsEditor.h
+++ b/Engine/Plugins/Experimental/HairStrands/Source/HairStrandsEditor/Public/HairStrandsEditor.h
@@ -4,6 +4,8 @@
 
 #include "Modules/ModuleInterface.h"
 #include "Modules/ModuleManager.h"
+#include "Styling/ISlateStyle.h"
+#include "Styling/SlateStyle.h"
 
 #define HAIRSTRANDSEDITOR_MODULE_NAME TEXT("HairStrandsEditor")
 
@@ -48,4 +50,6 @@ private:
 	TArray<TSharedRef<class IAssetTypeActions>> RegisteredAssetTypeActions;
 
 	TArray<TFunction<TSharedPtr<IHairStrandsTranslator>()>> TranslatorSpawners;
+
+	TSharedPtr<FSlateStyleSet> StyleSet;
 };
