Update mesh status once the projection is done.

(ushell-p4-cherrypick of 9618673 by Charles.deRousiers)

diff --git a/Engine/Plugins/Experimental/HairStrands/Source/HairStrandsCore/Private/GroomComponent.cpp b/Engine/Plugins/Experimental/HairStrands/Source/HairStrandsCore/Private/GroomComponent.cpp
--- a/Engine/Plugins/Experimental/HairStrands/Source/HairStrandsCore/Private/GroomComponent.cpp
+++ b/Engine/Plugins/Experimental/HairStrands/Source/HairStrandsCore/Private/GroomComponent.cpp
@@ -514,9 +514,7 @@ static FHairStrandsProjectionHairData::HairGroup ToProjectionHairData(FHairStran
 		LODData.DeformedRootTrianglePosition0Buffer = &MeshLODData.DeformedRootTrianglePosition0Buffer;
 		LODData.DeformedRootTrianglePosition1Buffer = &MeshLODData.DeformedRootTrianglePosition1Buffer;
 		LODData.DeformedRootTrianglePosition2Buffer = &MeshLODData.DeformedRootTrianglePosition2Buffer;
-		LODData.bIsValid = 
-			MeshLODData.Status == FHairStrandsRootResource::FMeshProjectionLOD::EStatus::Projected ||
-			MeshLODData.Status == FHairStrandsRootResource::FMeshProjectionLOD::EStatus::Initialized;
+		LODData.bIsValid = MeshLODData.Status == FHairStrandsRootResource::FMeshProjectionLOD::EStatus::Completed;
 	}
 
 	return Out;
@@ -961,6 +959,24 @@ void UGroomComponent::TickComponent(float DeltaTime, enum ELevelTick TickType, F
 						SkeletalMeshComponent->InitAnim(true);
 					}
 					MeshProjectionState = EMeshProjectionState::Completed;
+
+					FHairGroupResources* LocalResources = &HairGroupResources;
+					ENQUEUE_RENDER_COMMAND(FHairStrandsTick_UpdateProjectionStatus)(
+						[LocalResources](FRHICommandListImmediate& RHICmdList)
+					{
+						for (FHairGroupResource& Res : *LocalResources)
+						{
+							for (FHairStrandsRootResource::FMeshProjectionLOD& MeshLODData : Res.RenRootResources->MeshProjectionLODs)
+							{
+								MeshLODData.Status = FHairStrandsRootResource::FMeshProjectionLOD::EStatus::Completed;
+							}
+
+							for (FHairStrandsRootResource::FMeshProjectionLOD& MeshLODData : Res.SimRootResources->MeshProjectionLODs)
+							{
+								MeshLODData.Status = FHairStrandsRootResource::FMeshProjectionLOD::EStatus::Completed;
+							}
+						}
+					});
 				}
 
 				if (MeshProjectionLODIndex < SkeletalMeshComponent->GetNumLODs())
diff --git a/Engine/Plugins/Experimental/HairStrands/Source/HairStrandsCore/Public/GroomAsset.h b/Engine/Plugins/Experimental/HairStrands/Source/HairStrandsCore/Public/GroomAsset.h
--- a/Engine/Plugins/Experimental/HairStrands/Source/HairStrandsCore/Public/GroomAsset.h
+++ b/Engine/Plugins/Experimental/HairStrands/Source/HairStrandsCore/Public/GroomAsset.h
@@ -34,7 +34,7 @@ struct FHairStrandsRootResource : public FRenderResource
 
 	struct FMeshProjectionLOD
 	{
-		enum class EStatus { Invalid, Initialized, Projected };
+		enum class EStatus { Invalid, Initialized, Completed };
 		EStatus Status = EStatus::Invalid;
 		int32 LODIndex = -1;
 
