* Add hair group information on the goom instance.
 * Add hair width override per group on instance.

(ushell-p4-cherrypick of 9621116 by Charles.deRousiers)

diff --git a/Engine/Plugins/Experimental/HairStrands/Source/HairStrandsCore/Private/GroomComponent.cpp b/Engine/Plugins/Experimental/HairStrands/Source/HairStrandsCore/Private/GroomComponent.cpp
--- a/Engine/Plugins/Experimental/HairStrands/Source/HairStrandsCore/Private/GroomComponent.cpp
+++ b/Engine/Plugins/Experimental/HairStrands/Source/HairStrandsCore/Private/GroomComponent.cpp
@@ -143,6 +143,8 @@ public:
 		for (uint32 GroupIt=0;GroupIt<GroupCount; GroupIt++)
 		{		
 			const FHairGroupData& InGroupData = Component->GroomAsset->HairGroupsData[GroupIt];
+			const FHairGroupDesc& InGroupDesc = Component->GroomGroupsDesc[GroupIt];
+
 			const UGroomComponent::FHairGroupResource& GroupResources = Component->HairGroupResources[GroupIt];
 
 			UMaterialInterface* Material = Component->GetMaterial(GroupIt);
@@ -153,7 +155,7 @@ public:
 
 			FHairStrandsVertexFactory::FDataType::HairGroup& VFGroupData = VFData.HairGroups.AddDefaulted_GetRef();
 			VFGroupData.MinStrandRadius		= 0;
-			VFGroupData.MaxStrandRadius		= InGroupData.HairRenderData.StrandsCurves.MaxRadius;
+			VFGroupData.MaxStrandRadius		= InGroupDesc.HairWidth > 0 ? InGroupDesc.HairWidth * 0.5f : InGroupData.HairRenderData.StrandsCurves.MaxRadius;
 			VFGroupData.MaxStrandLength		= InGroupData.HairRenderData.StrandsCurves.MaxLength;
 			VFGroupData.HairDensity			= InGroupData.HairRenderData.HairDensity;
 			VFGroupData.HairWorldOffset		= InGroupData.HairRenderData.BoundingBox.GetCenter();
@@ -368,7 +370,37 @@ private:
 	TArray<HairGroup> HairGroups;
 };
 
+static void UpdateHairGroupsDesc(UGroomAsset* GroomAsset, TArray<FHairGroupDesc>& GroomGroupsDesc)
+{
+	if (!GroomAsset)
+	{
+		GroomGroupsDesc.Empty();
+		return;
+	}
+
+	check(GroomAsset->HairGroupsInfo.Num() == GroomAsset->HairGroupsInfo.Num());
 
+	const uint32 GroupCount = GroomAsset->HairGroupsInfo.Num();
+	const bool bReinitOverride = GroupCount != GroomGroupsDesc.Num();
+	if (bReinitOverride)
+	{
+		GroomGroupsDesc.SetNum(GroupCount);
+	}
+
+	for (uint32 GroupIt = 0; GroupIt < GroupCount; ++GroupIt)
+	{
+		FHairGroupInfo& GroupInfo = GroomAsset->HairGroupsInfo[GroupIt];
+		FHairGroupData& GroupData = GroomAsset->HairGroupsData[GroupIt];
+
+		FHairGroupDesc& Desc = GroomGroupsDesc[GroupIt];
+		Desc.GuideCount = GroupInfo.NumGuides;
+		Desc.HairCount  = GroupInfo.NumCurves;
+		if (bReinitOverride || Desc.HairWidth == 0)
+		{
+			Desc.HairWidth = GroupData.HairRenderData.StrandsCurves.MaxRadius * 0.5f;
+		}
+	}
+}
 
 /////////////////////////////////////////////////////////////////////////////////////////
 /////////////////////////////////////////////////////////////////////////////////////////
@@ -825,6 +857,8 @@ void UGroomComponent::PostLoad()
 	{
 		InitResources();
 	}
+
+	UpdateHairGroupsDesc(GroomAsset, GroomGroupsDesc);
 }
 
 void UGroomComponent::OnRegister()
@@ -1167,6 +1201,7 @@ void UGroomComponent::PostEditChangeProperty(FPropertyChangedEvent& PropertyChan
 			if (GroomAsset)
 			{
 				InitResources();
+				UpdateHairGroupsDesc(GroomAsset, GroomGroupsDesc);
 			}
 			else
 			{
@@ -1182,6 +1217,11 @@ void UGroomComponent::PostEditChangeProperty(FPropertyChangedEvent& PropertyChan
 			Group.HairRaytracingRadiusScale = HairRaytracingRadiusScale;
 		}
 	}
+
+	if (PropertyName == GET_MEMBER_NAME_CHECKED(FHairGroupDesc, HairWidth))
+	{	
+		UpdateHairGroupsDesc(GroomAsset, GroomGroupsDesc);
+	}
 }
 
 bool UGroomComponent::CanEditChange(const UProperty* InProperty) const
diff --git a/Engine/Plugins/Experimental/HairStrands/Source/HairStrandsCore/Public/GroomComponent.h b/Engine/Plugins/Experimental/HairStrands/Source/HairStrandsCore/Public/GroomComponent.h
--- a/Engine/Plugins/Experimental/HairStrands/Source/HairStrandsCore/Public/GroomComponent.h
+++ b/Engine/Plugins/Experimental/HairStrands/Source/HairStrandsCore/Public/GroomComponent.h
@@ -11,7 +11,24 @@
 
 #include "GroomComponent.generated.h"
 
-/** Component that allows you to specify custom triangle mesh geometry */
+USTRUCT(BlueprintType)
+struct FHairGroupDesc
+{
+	GENERATED_USTRUCT_BODY()
+
+	/** Number of hairs within this hair group.  */
+	UPROPERTY(VisibleAnywhere, BlueprintReadOnly, Category = "Groom")
+	int32 HairCount;
+
+	/** Number of simulation guides within this hair group. */
+	UPROPERTY(VisibleAnywhere, BlueprintReadOnly, Category = "Groom")
+	int32 GuideCount;
+
+	/** Override the hair width (in centimeters) */
+	UPROPERTY(EditAnywhere, BlueprintReadWrite, Category = "Groom", meta = (ClampMin = "0.0001", UIMin = "0.001", UIMax = "1.0", SliderExponent = 6))
+	float HairWidth;
+};
+
 UCLASS(HideCategories = (Object, Physics, Activation, Mobility, "Components|Activation"), editinlinenew, meta = (BlueprintSpawnableComponent), ClassGroup = Rendering)
 class HAIRSTRANDSCORE_API UGroomComponent : public UMeshComponent
 {
@@ -19,7 +36,7 @@ class HAIRSTRANDSCORE_API UGroomComponent : public UMeshComponent
 
 public:
 
-	/** Hair strand asset used for rendering. */
+	/** Groom asset . */
 	UPROPERTY(EditAnywhere, BlueprintReadWrite, Category = "Groom")
 	UGroomAsset* GroomAsset;
 
@@ -99,6 +116,10 @@ public:
 	};
 	typedef TArray<FHairGroupResource> FHairGroupResources;
 
+	/** Groom's groups info. */
+	UPROPERTY(EditAnywhere, BlueprintReadWrite, Category = "Groom")
+	TArray<FHairGroupDesc> GroomGroupsDesc;
+
 	FHairGroupResources HairGroupResources;
 	struct FHairStrandsInterpolationOutput* InterpolationOutput = nullptr;
 	struct FHairStrandsInterpolationInput* InterpolationInput = nullptr;
