Fix warning on CIS.

(ushell-p4-cherrypick of 9630371 by Charles.deRousiers)

diff --git a/Engine/Plugins/Experimental/HairStrands/Source/HairStrandsCore/Private/GroomActor.cpp b/Engine/Plugins/Experimental/HairStrands/Source/HairStrandsCore/Private/GroomActor.cpp
--- a/Engine/Plugins/Experimental/HairStrands/Source/HairStrandsCore/Private/GroomActor.cpp
+++ b/Engine/Plugins/Experimental/HairStrands/Source/HairStrandsCore/Private/GroomActor.cpp
@@ -47,8 +47,7 @@ AGroomActor::AGroomActor(const FObjectInitializer& ObjectInitializer)
 #endif
 }
 
-
-
+#undef LOCTEXT_NAMESPACE
 
 
 
