- Groom skinning in niagara
- New strands interpolation (r.HairStrands.LocalInterpolation 1)
- Groom Index exposed to the simulation
- Asset refactoring
- Material damping
- Simulation parameters exposed to the user

diff --git a/Engine/Plugins/Experimental/HairStrands/Shaders/Private/NiagaraDataInterfaceHairStrands.ush b/Engine/Plugins/Experimental/HairStrands/Shaders/Private/NiagaraDataInterfaceHairStrands.ush
--- a/Engine/Plugins/Experimental/HairStrands/Shaders/Private/NiagaraDataInterfaceHairStrands.ush
+++ b/Engine/Plugins/Experimental/HairStrands/Shaders/Private/NiagaraDataInterfaceHairStrands.ush
@@ -21,48 +21,81 @@ static const int VELOCITY_OFFSET = 6;
 
 struct FDIHairStrandsContext
 {
-	RWBuffer<uint4>		PointsPositionsBuffer;
+	RWBuffer<uint4>		DeformedPositionBuffer;
+	Buffer<uint4>		RestPositionBuffer;
 	Buffer<uint>		CurvesOffsetsBuffer;
-	Buffer<uint4>		RestPositionsBuffer;
+	Buffer<float4>		RestTrianglePositionABuffer;
+	Buffer<float4>		RestTrianglePositionBBuffer;
+	Buffer<float4>		RestTrianglePositionCBuffer;
+	Buffer<float4>		DeformedTrianglePositionABuffer;
+	Buffer<float4>		DeformedTrianglePositionBBuffer;
+	Buffer<float4>		DeformedTrianglePositionCBuffer;
+	Buffer<uint>		RootBarycentricCoordinatesBuffer;
 	float4x4			WorldTransform;
 	float4x4			WorldInverse;
 	float4				WorldRotation;
 	int					NumStrands;
 	int					StrandSize;
-	float				StrandDensity;
-	float				RootThickness;
-	float				TipThickness;
-	float4				BoxOrigin;
+	float3				BoxCenter;
+	float3				BoxExtent;
+	int					HasRootAttached;
+	float3				RestRootOffset;
+	float3				DeformedRootOffset;
+	float3				RestPositionOffset;
+	float3				DeformedPositionOffset;
+	bool				ResetSimulation;
 }; 
 
 #define DIHAIRSTRANDS_DECLARE_CONSTANTS(NAME)\
-RWBuffer<uint4> PointsPositionsBuffer_##NAME;\
-Buffer<uint> CurvesOffsetsBuffer_##NAME;\
-Buffer<uint4> RestPositionsBuffer_##NAME;\
-float4x4 WorldTransform_##NAME;\
-float4x4 WorldInverse_##NAME;\
-float4 WorldRotation_##NAME;\
-int NumStrands_##NAME;\
-int StrandSize_##NAME;\
-float StrandDensity_##NAME;\
-float RootThickness_##NAME;\
-float TipThickness_##NAME;\
-float4 BoxOrigin_##NAME;\
+RWBuffer<uint4>		DeformedPositionBuffer_##NAME;\
+Buffer<uint4>		RestPositionBuffer_##NAME;\
+Buffer<uint>		CurvesOffsetsBuffer_##NAME;\
+Buffer<float4>		RestTrianglePositionABuffer_##NAME;\
+Buffer<float4>		RestTrianglePositionBBuffer_##NAME;\
+Buffer<float4>		RestTrianglePositionCBuffer_##NAME;\
+Buffer<float4>		DeformedTrianglePositionABuffer_##NAME;\
+Buffer<float4>		DeformedTrianglePositionBBuffer_##NAME;\
+Buffer<float4>		DeformedTrianglePositionCBuffer_##NAME;\
+Buffer<uint>		RootBarycentricCoordinatesBuffer_##NAME;\
+float4x4			WorldTransform_##NAME;\
+float4x4			WorldInverse_##NAME;\
+float4				WorldRotation_##NAME;\
+int					NumStrands_##NAME;\
+int					StrandSize_##NAME;\
+float3				BoxCenter_##NAME;\
+float3				BoxExtent_##NAME;\
+int					HasRootAttached_##NAME;\
+float3				RestRootOffset_##NAME;\
+float3				DeformedRootOffset_##NAME;\
+float3				RestPositionOffset_##NAME;\
+float3				DeformedPositionOffset_##NAME;\
+bool				ResetSimulation_##NAME;\
 
 #define DIHAIRSTRANDS_MAKE_CONTEXT(NAME)\
 FDIHairStrandsContext DIContext;\
-DIContext.PointsPositionsBuffer = PointsPositionsBuffer_##NAME;\
+DIContext.DeformedPositionBuffer = DeformedPositionBuffer_##NAME;\
+DIContext.RestPositionBuffer = RestPositionBuffer_##NAME;\
 DIContext.CurvesOffsetsBuffer = CurvesOffsetsBuffer_##NAME;\
-DIContext.RestPositionsBuffer = RestPositionsBuffer_##NAME;\
+DIContext.RestTrianglePositionABuffer = RestTrianglePositionABuffer_##NAME;\
+DIContext.RestTrianglePositionBBuffer = RestTrianglePositionBBuffer_##NAME;\
+DIContext.RestTrianglePositionCBuffer = RestTrianglePositionCBuffer_##NAME;\
+DIContext.DeformedTrianglePositionABuffer = DeformedTrianglePositionABuffer_##NAME;\
+DIContext.DeformedTrianglePositionBBuffer = DeformedTrianglePositionBBuffer_##NAME;\
+DIContext.DeformedTrianglePositionCBuffer = DeformedTrianglePositionCBuffer_##NAME;\
+DIContext.RootBarycentricCoordinatesBuffer = RootBarycentricCoordinatesBuffer_##NAME;\
 DIContext.WorldTransform = WorldTransform_##NAME;\
 DIContext.WorldInverse = WorldInverse_##NAME;\
 DIContext.WorldRotation = WorldRotation_##NAME;\
 DIContext.NumStrands = NumStrands_##NAME;\
 DIContext.StrandSize = StrandSize_##NAME;\
-DIContext.StrandDensity = StrandDensity_##NAME;\
-DIContext.RootThickness = RootThickness_##NAME;\
-DIContext.TipThickness = TipThickness_##NAME;\
-DIContext.BoxOrigin = BoxOrigin_##NAME;\
+DIContext.BoxCenter = BoxCenter_##NAME;\
+DIContext.BoxExtent = BoxExtent_##NAME;\
+DIContext.HasRootAttached =	HasRootAttached_##NAME;\
+DIContext.RestRootOffset =	RestRootOffset_##NAME;\
+DIContext.DeformedRootOffset = DeformedRootOffset_##NAME;\
+DIContext.RestPositionOffset =	RestPositionOffset_##NAME;\
+DIContext.DeformedPositionOffset = DeformedPositionOffset_##NAME;\
+DIContext.ResetSimulation = ResetSimulation_##NAME;\
 
 /* -----------------------------------------------------------------
  * Box utilities
@@ -72,13 +105,13 @@ DIContext.BoxOrigin = BoxOrigin_##NAME;\
 // Return the bounding box center
 void DIHairStrands_GetBoxCenter(in FDIHairStrandsContext DIContext, out float3 OutBoxCenter)
 {
-	OutBoxCenter = float3(0,0,0);
+	OutBoxCenter = DIContext.BoxCenter;
 }
 
 // Return the bounding box extent
 void DIHairStrands_GetBoxExtent(in FDIHairStrandsContext DIContext, out float3 OutBoxExtent)
 {
-	OutBoxExtent = -DIContext.BoxOrigin.xyz;
+	OutBoxExtent = DIContext.BoxExtent;
 }
 
 /* -----------------------------------------------------------------
@@ -122,13 +155,13 @@ float DIHairStrands_ComputeStrandLength(in FDIHairStrandsContext DIContext, in i
 		int PointOffset = DIContext.CurvesOffsetsBuffer[StrandIndex];
 		const int EdgeCount = DIContext.CurvesOffsetsBuffer[StrandIndex+1]-PointOffset-1;
 
-		float3 PointNext = f16tof32(DIContext.RestPositionsBuffer[PointOffset++].xyz);
+		float3 PointNext = f16tof32(DIContext.RestPositionBuffer[PointOffset++].xyz);
 		float3 PointPrev = PointNext;
 
 		for (int EdgeIndex = 0; EdgeIndex < EdgeCount; ++EdgeIndex)
 		{
 			PointPrev = PointNext;
-			PointNext = f16tof32(DIContext.RestPositionsBuffer[PointOffset++].xyz);
+			PointNext = f16tof32(DIContext.RestPositionBuffer[PointOffset++].xyz);
 			StrandLength += length(PointNext-PointPrev);
 		}
 	}
@@ -224,23 +257,23 @@ void DIHairStrands_ComputeNodePosition(in FDIHairStrandsContext DIContext, in in
 		const int FirstNext = PointPrev+1;
 
 		const float FirstAlpha = FirstPoint - (float)FirstPrev;
-		const float3 FirstPosition =  f16tof32(DIContext.RestPositionsBuffer[PointOffset+FirstPrev].xyz) * (1.0-FirstAlpha) + 
-						   f16tof32(DIContext.RestPositionsBuffer[PointOffset+FirstNext].xyz) * FirstAlpha;
-		const float3 DiffPosition = FirstPosition - f16tof32(DIContext.RestPositionsBuffer[PointOffset].xyz);
+		const float3 FirstPosition =  f16tof32(DIContext.RestPositionBuffer[PointOffset+FirstPrev].xyz) * (1.0-FirstAlpha) + 
+						   f16tof32(DIContext.RestPositionBuffer[PointOffset+FirstNext].xyz) * FirstAlpha;
+		const float3 DiffPosition = FirstPosition - f16tof32(DIContext.RestPositionBuffer[PointOffset].xyz);
 		const float3 EdgeDirection =  normalize(DiffPosition);
 
 		const float EdgeLength = DIHairStrands_ComputeStrandLength(DIContext,StrandIndex) * CoordScale;
-		//float3 EdgeDirection = f16tof32(DIContext.RestPositionsBuffer[PointOffset+PointNext].xyz) - 
-		//							 f16tof32(DIContext.RestPositionsBuffer[PointOffset+PointPrev].xyz);
-		OutNodePosition = 	f16tof32(DIContext.RestPositionsBuffer[PointOffset+PointPrev].xyz) - EdgeLength * normalize(EdgeDirection);					 
+		//float3 EdgeDirection = f16tof32(DIContext.RestPositionBuffer[PointOffset+PointNext].xyz) - 
+		//							 f16tof32(DIContext.RestPositionBuffer[PointOffset+PointPrev].xyz);
+		OutNodePosition = 	f16tof32(DIContext.RestPositionBuffer[PointOffset+PointPrev].xyz) - EdgeLength * normalize(EdgeDirection);					 
 	}
 	else
 	{
 		const float NodeAlpha = NodePoint - (float)PointPrev;
-		OutNodePosition =  f16tof32(DIContext.RestPositionsBuffer[PointOffset+PointPrev].xyz) * (1.0-NodeAlpha) + 
-						   f16tof32(DIContext.RestPositionsBuffer[PointOffset+PointNext].xyz) * NodeAlpha;
+		OutNodePosition =  f16tof32(DIContext.RestPositionBuffer[PointOffset+PointPrev].xyz) * (1.0-NodeAlpha) + 
+						   f16tof32(DIContext.RestPositionBuffer[PointOffset+PointNext].xyz) * NodeAlpha;
 	}
-	OutNodePosition = mul(float4( OutNodePosition, 1.0), DIContext.WorldTransform).xyz;
+	OutNodePosition = mul(float4( OutNodePosition+DIContext.RestPositionOffset, 1.0), DIContext.WorldTransform).xyz;
 }
 
 // Compute the node orientation
@@ -335,10 +368,12 @@ void DIHairStrands_ResetPointPosition(in FDIHairStrandsContext DIContext, in int
 		const int PointEnd = DIContext.CurvesOffsetsBuffer[StrandIndex+1]; 
 		const int PointCount = PointEnd - PointBegin;
 
+		const uint4 PositionOffset = uint4( f32tof16(DIContext.RestPositionOffset - DIContext.DeformedPositionOffset), 0 );
+
 		ReportStatus = true;
 		for (int PointIndex = 0; PointIndex < PointCount; ++PointIndex)
 		{
-			DIContext.PointsPositionsBuffer[PointBegin+PointIndex] = DIContext.RestPositionsBuffer[PointBegin+PointIndex];
+			DIContext.DeformedPositionBuffer[PointBegin+PointIndex] = DIContext.RestPositionBuffer[PointBegin+PointIndex] + PositionOffset;
 		}
 	}
 }
@@ -360,6 +395,8 @@ void DIHairStrands_UpdatePointPosition(in FDIHairStrandsContext DIContext, in in
 	const int PointEnd = DIContext.CurvesOffsetsBuffer[StrandIndex+1]; 
 	const int PointCount = PointEnd - PointBegin;
 
+	const float3 PositionOffset = DIContext.RestPositionOffset - DIContext.DeformedPositionOffset;
+
 	for (int PointIndex = 0; PointIndex < PointCount; ++PointIndex)
 	{
 		const float PointCoord = (float)(PointIndex) / (PointCount-1.0);
@@ -373,8 +410,8 @@ void DIHairStrands_UpdatePointPosition(in FDIHairStrandsContext DIContext, in in
 			ReportStatus = true;
 			const float PointAlpha = PointNode - (float)NodePrev;
 			const float3 LocalDisplace = SharedNodePosition[GGroupThreadId.x] * (1.0-PointAlpha) + SharedNodePosition[GGroupThreadId.x+1] * PointAlpha;
-			const float3 DeformedPosition = f16tof32(DIContext.RestPositionsBuffer[PointBegin+PointIndex].xyz) + LocalDisplace;
-			DIContext.PointsPositionsBuffer[PointBegin+PointIndex] = uint4( f32tof16(DeformedPosition), DIContext.RestPositionsBuffer[PointBegin+PointIndex].w);
+			const float3 DeformedPosition = f16tof32(DIContext.RestPositionBuffer[PointBegin+PointIndex].xyz) + LocalDisplace + PositionOffset;
+			DIContext.DeformedPositionBuffer[PointBegin+PointIndex] = uint4( f32tof16(DeformedPosition), DIContext.RestPositionBuffer[PointBegin+PointIndex].w);
 		}
 	}
 }
@@ -382,7 +419,7 @@ void DIHairStrands_UpdatePointPosition(in FDIHairStrandsContext DIContext, in in
 // Get the point position
 void DIHairStrands_GetPointPosition(in FDIHairStrandsContext DIContext, in int PointIndex, out float3 OutPointPosition)
 {
-	OutPointPosition = mul(float4( f16tof32(DIContext.PointsPositionsBuffer[PointIndex].xyz), 1.0), DIContext.WorldTransform).xyz;
+	OutPointPosition = mul(float4( f16tof32(DIContext.DeformedPositionBuffer[PointIndex].xyz+DIContext.DeformedPositionOffset), 1.0), DIContext.WorldTransform).xyz;
 }
 
 /* -----------------------------------------------------------------
@@ -507,6 +544,83 @@ void DIHairStrands_AttachNodeOrientation(in FDIHairStrandsContext DIContext, in
 	OutNodeOrientation = NormalizeQuat( MultiplyQuat(DIContext.WorldRotation,RestOrientation) );
 }
 
+/* -----------------------------------------------------------------
+ * Attached the root to the skinned cache
+ * -----------------------------------------------------------------
+ */
+
+float2 ExtractBarycentrics(in uint PackedCoordinates )
+{
+	return float2(f16tof32(PackedCoordinates & 0xFFFF), f16tof32((PackedCoordinates>>16) & 0xFFFF));
+}
+
+// Compuite rest position and orientation
+void DIHairStrands_ComputeLocalState(in FDIHairStrandsContext DIContext, in int NodeIndex, in float3 RestPosition, in float4 RestOrientation, 
+												out float3 OutLocalPosition, out float4 OutLocalOrientation)
+{
+	OutLocalPosition = RestPosition;
+	OutLocalOrientation = RestOrientation;
+
+	if( DIContext.HasRootAttached == 1)
+	{
+		const int CurveIndex = NodeIndex / DIContext.StrandSize;
+		const float2 RootUV = ExtractBarycentrics(DIContext.RootBarycentricCoordinatesBuffer[CurveIndex]);
+
+		const float3 RootProjection = DIContext.RestTrianglePositionABuffer[CurveIndex].xyz * RootUV.x + 
+									  DIContext.RestTrianglePositionBBuffer[CurveIndex].xyz * RootUV.y +
+									  DIContext.RestTrianglePositionCBuffer[CurveIndex].xyz * (1.0 - RootUV.x - RootUV.y) + DIContext.RestRootOffset;
+		const float3 DiffPosition = RestPosition - RootProjection;
+
+		float3 TangentU = normalize(DIContext.RestTrianglePositionBBuffer[CurveIndex].xyz - DIContext.RestTrianglePositionABuffer[CurveIndex].xyz);
+		float3 TangentV = normalize(DIContext.RestTrianglePositionCBuffer[CurveIndex].xyz - DIContext.RestTrianglePositionABuffer[CurveIndex].xyz);
+
+		const float3 Normal = normalize(cross(TangentU,TangentV));
+		TangentU = normalize(cross(TangentV,Normal));
+
+		const float3 RotationMatrix[3] = {float3(TangentU.x,TangentV.x,Normal.x),
+										  float3(TangentU.y,TangentV.y,Normal.y),
+										  float3(TangentU.z,TangentV.z,Normal.z)};
+		const float4 InverseOrientation = QuatFromMatrix(RotationMatrix);
+
+		OutLocalPosition = RotateVectorByQuat(DiffPosition,InverseOrientation);
+		//OutLocalPosition = float3( dot(DiffPosition,TangentU), dot(DiffPosition,TangentV), dot(DiffPosition,Normal) );
+		OutLocalOrientation = NormalizeQuat( MultiplyQuat(InverseOrientation,RestOrientation) );
+	}
+}
+
+// Update Node Position and orientation
+void DIHairStrands_AttachNodeState(in FDIHairStrandsContext DIContext, in int NodeIndex, in float3 LocalPosition, in float4 LocalOrientation, out float3 OutNodePosition, out float4 OutNodeOrientation)
+{	
+	float3 RestPosition =  LocalPosition;
+	float4 RestOrientation = LocalOrientation;
+	if( DIContext.HasRootAttached == 1)
+	{
+		const int CurveIndex = NodeIndex / DIContext.StrandSize;
+		const float2 RootUV = ExtractBarycentrics(DIContext.RootBarycentricCoordinatesBuffer[CurveIndex]);
+
+		const float3 RootProjection = DIContext.DeformedTrianglePositionABuffer[CurveIndex].xyz * RootUV.x + 
+									  DIContext.DeformedTrianglePositionBBuffer[CurveIndex].xyz * RootUV.y +
+									  DIContext.DeformedTrianglePositionCBuffer[CurveIndex].xyz * (1.0 - RootUV.x - RootUV.y) + DIContext.DeformedRootOffset;
+
+		float3 TangentU = normalize(DIContext.DeformedTrianglePositionBBuffer[CurveIndex].xyz - DIContext.DeformedTrianglePositionABuffer[CurveIndex].xyz);
+		float3 TangentV = normalize(DIContext.DeformedTrianglePositionCBuffer[CurveIndex].xyz - DIContext.DeformedTrianglePositionABuffer[CurveIndex].xyz);
+
+		const float3 Normal = normalize(cross(TangentU,TangentV));
+		TangentU = normalize(cross(TangentV,Normal));
+		
+		const float3 RotationMatrix[3] = { TangentU, TangentV, Normal };
+		const float4 TriangleOrientation = QuatFromMatrix(RotationMatrix);
+
+		const float3 DiffPosition = RotateVectorByQuat(LocalPosition,TriangleOrientation);  
+		//const float3 DiffPosition = LocalPosition.x * TangentU + LocalPosition.y * TangentV + LocalPosition.z * Normal;
+
+		RestPosition = RootProjection + DiffPosition;
+		RestOrientation = NormalizeQuat( MultiplyQuat(TriangleOrientation,LocalOrientation) );
+	}
+	DIHairStrands_AttachNodePosition(DIContext,RestPosition,OutNodePosition);
+	DIHairStrands_AttachNodeOrientation(DIContext,RestOrientation,OutNodeOrientation);
+}
+
 
 
 
diff --git a/Engine/Plugins/Experimental/HairStrands/Shaders/Private/NiagaraDataInterfacePressureGrid.ush b/Engine/Plugins/Experimental/HairStrands/Shaders/Private/NiagaraDataInterfacePressureGrid.ush
--- a/Engine/Plugins/Experimental/HairStrands/Shaders/Private/NiagaraDataInterfacePressureGrid.ush
+++ b/Engine/Plugins/Experimental/HairStrands/Shaders/Private/NiagaraDataInterfacePressureGrid.ush
@@ -427,6 +427,10 @@ void DIPressureGrid_SampleVelocityField(in FDIPressureGridContext DIContext, in
 	DIPressureGrid_SampleVelocityComponent(DIContext,GridPosition,float3(0.5,0.5,0.0),2,OutGridVelocity.z,GridMass.z,OutGridGradientZ);
 
 	OutNodeDensity = (GridMass.x + GridMass.y + GridMass.z)/(3.0*DIContext.GridOrigin.w*DIContext.GridOrigin.w*DIContext.GridOrigin.w);
+
+	//int3 GridIndex = int3(0,0,0); float3 GridFraction = float3(0,0,0);
+	//DIPressureGrid_GetGridIndex(DIContext,GridPosition,0,GridIndex,GridFraction);
+	//OutGridVelocity = GridIndex;
 }
 
 /* -----------------------------------------------------------------
@@ -460,7 +464,7 @@ void DIPressureGrid_SetSolidBoundary(in FDIPressureGridContext DIContext, in int
  // Update the solid weights and velocities
 void DIPressureGrid_ComputeBoundaryWeights(in FDIPressureGridContext DIContext, in int GridCell, out bool OutWeightsStatus)
 {
-	int3 GridIndex = DIPressureGrid_ComputeGridIndex(DIContext,GridCell);
+	const int3 GridIndex = DIPressureGrid_ComputeGridIndex(DIContext,GridCell);
 	OutWeightsStatus = false;
 
 	for(int ComponentIndex = 0; ComponentIndex < 3; ++ComponentIndex)
@@ -474,9 +478,9 @@ void DIPressureGrid_ComputeBoundaryWeights(in FDIPressureGridContext DIContext,
 				const float RightDistance = DIPressureGrid_GetGridScalar(DIContext,GridIndex,SOLID_DISTANCE_OFFSET);
 				const float RightVelocity = DIPressureGrid_GetGridScalar(DIContext,GridIndex,SOLID_VELOCITY_OFFSET+ComponentIndex);
 
-				GridIndex[ComponentIndex] -= 1;
-				const float LeftDistance = DIPressureGrid_GetGridScalar(DIContext,GridIndex,SOLID_DISTANCE_OFFSET);
-				const float LeftVelocity = DIPressureGrid_GetGridScalar(DIContext,GridIndex,SOLID_VELOCITY_OFFSET+ComponentIndex);
+				int3 LeftIndex = GridIndex; LeftIndex[ComponentIndex] -= 1;
+				const float LeftDistance = DIPressureGrid_GetGridScalar(DIContext,LeftIndex,SOLID_DISTANCE_OFFSET);
+				const float LeftVelocity = DIPressureGrid_GetGridScalar(DIContext,LeftIndex,SOLID_VELOCITY_OFFSET+ComponentIndex);
 
 				SolidWeight = max(0.0,min(1.0 - DIPressureGrid_GetFractionInside(LeftDistance,RightDistance),1.0));
 				SolidVelocity = 0.5 * (RightVelocity+LeftVelocity);
diff --git a/Engine/Plugins/Experimental/HairStrands/Shaders/Private/NiagaraHookeSpringMaterial.ush b/Engine/Plugins/Experimental/HairStrands/Shaders/Private/NiagaraHookeSpringMaterial.ush
--- a/Engine/Plugins/Experimental/HairStrands/Shaders/Private/NiagaraHookeSpringMaterial.ush
+++ b/Engine/Plugins/Experimental/HairStrands/Shaders/Private/NiagaraHookeSpringMaterial.ush
@@ -54,6 +54,9 @@ void UpdateStretchSpringMultiplier(in float RestLength, in float DeltaTime, in b
 void SetupStretchSpringMaterial(in int StrandsSize, in float YoungModulus, in float RodThickness, 
 	in float RestLength, in float DeltaTime, in bool ProjectConstraint, in float MaterialDamping, out float OutMaterialCompliance, out float OutMaterialWeight, out float OutMaterialMultiplier)
 {
+	OutMaterialCompliance = 0.0;
+	OutMaterialWeight = 0.0;
+	OutMaterialMultiplier = 0.0;
 	const int LocalIndex = (GGroupThreadId.x % StrandsSize);
 	if( LocalIndex > 1 )
 	{
@@ -114,6 +117,9 @@ void ProjectStretchSpringMaterial(in int StrandsSize, in float YoungModulus, in
  void SetupBendSpringMaterial(in int StrandsSize, in float YoungModulus, in float RodThickness, 
 	in float RestLength, in float DeltaTime, in bool ProjectConstraint, in float MaterialDamping, out float OutMaterialCompliance, out float OutMaterialWeight, out float OutMaterialMultiplier)
 {
+	OutMaterialCompliance = 0.0;
+	OutMaterialWeight = 0.0;
+	OutMaterialMultiplier = 0.0;
 	const int LocalIndex = (GGroupThreadId.x % StrandsSize);
 	if( LocalIndex > 1 )
 	{
@@ -174,6 +180,9 @@ void ProjectBendSpringMaterial(in int StrandsSize, in float YoungModulus, in flo
  void SetupTwistSpringMaterial(in int StrandsSize, in float YoungModulus, in float RodThickness, 
 	in float RestLength, in float DeltaTime, in bool ProjectConstraint, in float MaterialDamping, out float OutMaterialCompliance, out float OutMaterialWeight, out float OutMaterialMultiplier)
 {
+	OutMaterialCompliance = 0.0;
+	OutMaterialWeight = 0.0;
+	OutMaterialMultiplier = 0.0;
 	const int LocalIndex = (GGroupThreadId.x % StrandsSize);
 	if( LocalIndex > 2 )
 	{
@@ -237,13 +246,13 @@ void ComputeMaterialFrame(in int StrandsSize)
 	if( LocalIndex == 0 )
 	{
 		float4 NodeQuaternion = SharedNodeOrientation[GGroupThreadId.x];
-		float3 TangentPrev = normalize(SharedNodePosition[GGroupThreadId.x+1] - SharedNodePosition[GGroupThreadId.x]);
+		float3 TangentPrev = normalize(SharedPreviousPosition[GGroupThreadId.x+1] - SharedPreviousPosition[GGroupThreadId.x]);
 		float3 TangentNext = TangentPrev;
 
 		for (int EdgeIndex = 1, EdgeEnd = StrandsSize-1; EdgeIndex < EdgeEnd; ++EdgeIndex)
 		{
 			TangentPrev = TangentNext;
-			TangentNext = normalize(SharedNodePosition[GGroupThreadId.x+EdgeIndex+1] - SharedNodePosition[GGroupThreadId.x+EdgeIndex]);
+			TangentNext = normalize(SharedPreviousPosition[GGroupThreadId.x+EdgeIndex+1] - SharedPreviousPosition[GGroupThreadId.x+EdgeIndex]);
 			NodeQuaternion = NormalizeQuat( MultiplyQuat( FindQuatBetweenNormals(TangentPrev,TangentNext), NodeQuaternion ) );
 			SharedNodeOrientation[GGroupThreadId.x+EdgeIndex] = NodeQuaternion;
 			GroupMemoryBarrier();
@@ -260,10 +269,12 @@ void ComputeMaterialFrame(in int StrandsSize)
  void SetupAngularSpringMaterial(in int StrandsSize, in float YoungModulus, in float RodThickness, 
 	in float RestLength, in float DeltaTime, in bool ProjectConstraint, in float MaterialDamping, out float OutMaterialCompliance, out float OutMaterialWeight, out float3 OutMaterialMultiplier)
 {
+	OutMaterialCompliance = 0.0;
+	OutMaterialWeight = 0.0;
+	OutMaterialMultiplier = float3(0,0,0);
+
 	float MaterialMultiplier = 0;
 	SetupStretchSpringMaterial(StrandsSize,YoungModulus,RodThickness,RestLength,DeltaTime,ProjectConstraint,MaterialDamping,OutMaterialCompliance,OutMaterialWeight,MaterialMultiplier);
-	OutMaterialMultiplier = float3(0,0,0);
-	//ComputeMaterialFrame(StrandsSize);
 }
 
 void UpdateAngularSpringMultiplier( in float RestLength, in float3 RestDirection, in float DeltaTime, in bool ProjectConstraint, in float MaterialDamping, in float InMaterialCompliance, in float InMaterialWeight, inout float3 OutMaterialMultiplier)
@@ -347,7 +358,7 @@ void ProjectAngularSpringMaterial(in int StrandsSize, in float YoungModulus, in
 
 void ComputeRestDirection(in int StrandsSize, in float3 NodePosition, in float4 NodeOrientation, out float3 OutRestDirection)
 {
-	SharedNodePosition[GGroupThreadId.x] = NodePosition;
+	SharedPreviousPosition[GGroupThreadId.x] = NodePosition;
 	SharedNodeOrientation[GGroupThreadId.x] = NodeOrientation;
 	GroupMemoryBarrier(); 
 
diff --git a/Engine/Plugins/Experimental/HairStrands/Shaders/Private/NiagaraQuaternionUtils.ush b/Engine/Plugins/Experimental/HairStrands/Shaders/Private/NiagaraQuaternionUtils.ush
--- a/Engine/Plugins/Experimental/HairStrands/Shaders/Private/NiagaraQuaternionUtils.ush
+++ b/Engine/Plugins/Experimental/HairStrands/Shaders/Private/NiagaraQuaternionUtils.ush
@@ -35,7 +35,7 @@ float3 UnrotateVectorByQuat(in float3 Vector, in float4 Quat)
 	return Vector - Quat.w * T + cross(Quat.xyz,T);
 }
 
-float4 FindQuatBetweenHelper(in float3 A, in float3 B, in float NormAB)
+/*float4 FindQuatBetweenHelper(in float3 A, in float3 B, in float NormAB)
 {
 	float W = NormAB + dot(A,B);
 	float4 Quat = (W>1e-6*NormAB) ? float4(A.y*B.z-A.z*B.y,A.z*B.x-A.x*B.z,A.x*B.y-A.y*B.x,W) :
@@ -44,6 +44,47 @@ float4 FindQuatBetweenHelper(in float3 A, in float3 B, in float NormAB)
 	//float S = sqrt((NormAB + dot(A,B))*2.0);
 	//float4 Quat = float4(cross(A,B)/S,0.5*S);
 	return NormalizeQuat(Quat);
+}*/
+
+float4 FindQuatBetweenInternal(in float3 An, in float3 Bn)
+{
+	float3 Cn = normalize(An+Bn);
+	return float4(cross(An,Cn),dot(An,Cn));
+}
+
+float4 FindQuatBetweenHelper(in float3 A, in float3 B, in float NormAB)
+{
+	float3 An = normalize(A);
+	float3 Bn = normalize(B);
+
+	if( dot(An,Bn) >= 0.0)
+	{
+		return NormalizeQuat(FindQuatBetweenInternal(An,Bn));
+	}
+	else
+	{
+		float3 Cn = normalize(An+Bn);
+		if(dot(Cn,Cn) != 0.0)
+		{
+			return NormalizeQuat(MultiplyQuat(FindQuatBetweenInternal(An,Cn),FindQuatBetweenInternal(Cn,Bn)));
+		}
+		else
+		{
+			float3 A2 = An * An;
+			if(A2.x <= A2.y && A2.x <= A2.z)
+			{
+				return NormalizeQuat(float4(cross(An,float3(1,0,0)),0));
+			}
+			else if(A2.y <= A2.z)
+			{
+				return NormalizeQuat(float4(cross(An,float3(0,1,0)),0));
+			}
+			else
+			{
+				return NormalizeQuat(float4(cross(An,float3(0,0,1)),0));
+			}
+		}
+	}
 }
 
 float4 FindQuatBetweenNormals(in float3 NormalA, in float3 NormalB)
@@ -57,3 +98,63 @@ float4 FindQuatBetweenVectors(in float3 VectorA, in float3 VectorB)
 	float NormAB = sqrt(dot(VectorA,VectorA) * dot(VectorB,VectorB));
 	return FindQuatBetweenHelper(VectorA,VectorB,NormAB);
 }
+
+float4 QuatFromMatrix(in float3 R[3])
+{
+	float4 Q;
+	float s;
+	float tr = R[0][0] + R[1][1] + R[2][2];
+
+	if (tr > 0.0f) 
+	{
+		float InvS = 1.0 / sqrt(tr + 1.f);
+		s = 0.5f * InvS;
+
+		Q.w = 0.5f * (1.f / InvS);
+		Q.x = (R[1][2] - R[2][1]) * s;
+		Q.y = (R[2][0] - R[0][2]) * s;
+		Q.z = (R[0][1] - R[1][0]) * s;
+	} 
+	else 
+	{
+		uint i = 0;
+
+		if (R[1][1] > R[0][0])
+			i = 1;
+
+		if (R[2][2] > R[i][i])
+			i = 2;
+
+		const uint j = (i+1)%3;
+		const uint k = (j+1)%3;
+
+		tr = R[i][i] - R[j][j] - R[k][k];
+		float InvS = 1.0 / sqrt(tr + 1.0);
+		s = 0.5 * InvS;
+
+		const float Qi = 0.5 * (1.0 / InvS);
+		const float Qj = (R[i][j] + R[j][i]) * s;
+		const float Qk = (R[i][k] + R[k][i]) * s;
+
+		Q.w = (R[j][k] - R[k][j]) * s;
+		if(i == 0)
+		{
+			Q.x = Qi;
+			Q.y = Qj;
+			Q.z = Qk;
+		}
+		else if(i == 1)
+		{
+			Q.y = Qi;
+			Q.z = Qj;
+			Q.x = Qk;
+		}
+		else if(i == 2)
+		{
+			Q.z = Qi;
+			Q.x = Qj;
+			Q.y = Qk;
+		}
+	}
+	return NormalizeQuat(Q);
+}
diff --git a/Engine/Plugins/Experimental/HairStrands/Source/HairStrandsCore/Private/GroomAsset.cpp b/Engine/Plugins/Experimental/HairStrands/Source/HairStrandsCore/Private/GroomAsset.cpp
--- a/Engine/Plugins/Experimental/HairStrands/Source/HairStrandsCore/Private/GroomAsset.cpp
+++ b/Engine/Plugins/Experimental/HairStrands/Source/HairStrandsCore/Private/GroomAsset.cpp
@@ -144,7 +144,7 @@ FHairStrandsRootResource::FHairStrandsRootResource(const FHairStrandsDatas* Hair
 	for (FMeshProjectionLOD& MeshProjectionLOD : MeshProjectionLODs)
 	{
 		MeshProjectionLOD.Status = FMeshProjectionLOD::EStatus::Invalid;
-		MeshProjectionLOD.RestRootCenter = HairStrandsDatas->BoundingBox.GetCenter();
+		MeshProjectionLOD.RestRootOffset = HairStrandsDatas->BoundingBox.GetCenter();
 		MeshProjectionLOD.LODIndex = LODIndex++;
 	}
 }
diff --git a/Engine/Plugins/Experimental/HairStrands/Source/HairStrandsCore/Private/GroomComponent.cpp b/Engine/Plugins/Experimental/HairStrands/Source/HairStrandsCore/Private/GroomComponent.cpp
--- a/Engine/Plugins/Experimental/HairStrands/Source/HairStrandsCore/Private/GroomComponent.cpp
+++ b/Engine/Plugins/Experimental/HairStrands/Source/HairStrandsCore/Private/GroomComponent.cpp
@@ -537,12 +537,12 @@ static FHairStrandsProjectionHairData::HairGroup ToProjectionHairData(FHairStran
 		LODData.RootTriangleIndexBuffer = &MeshLODData.RootTriangleIndexBuffer;
 		LODData.RootTriangleBarycentricBuffer = &MeshLODData.RootTriangleBarycentricBuffer;
 		
-		LODData.RestPositionOffset					= MeshLODData.RestRootCenter;
+		LODData.RestPositionOffset					= MeshLODData.RestRootOffset;
 		LODData.RestRootTrianglePosition0Buffer = &MeshLODData.RestRootTrianglePosition0Buffer;
 		LODData.RestRootTrianglePosition1Buffer = &MeshLODData.RestRootTrianglePosition1Buffer;
 		LODData.RestRootTrianglePosition2Buffer = &MeshLODData.RestRootTrianglePosition2Buffer;
 
-		LODData.DeformedPositionOffset					= MeshLODData.DeformedRootCenter;
+		LODData.DeformedPositionOffset					= MeshLODData.DeformedRootOffset;
 		LODData.DeformedRootTrianglePosition0Buffer = &MeshLODData.DeformedRootTrianglePosition0Buffer;
 		LODData.DeformedRootTrianglePosition1Buffer = &MeshLODData.DeformedRootTrianglePosition1Buffer;
 		LODData.DeformedRootTrianglePosition2Buffer = &MeshLODData.DeformedRootTrianglePosition2Buffer;
@@ -582,6 +582,16 @@ FHairStrandsDeformedResource* UGroomComponent::GetGuideStrandsDeformedResource(u
 	return HairGroupResources[GroupIndex].SimDeformedResources;
 }
 
+FHairStrandsRootResource* UGroomComponent::GetGuideStrandsRootResource(uint32 GroupIndex)
+{
+	if (GroupIndex >= uint32(HairGroupResources.Num()))
+	{
+		return nullptr;
+	}
+
+	return HairGroupResources[GroupIndex].SimRootResources;
+}
+
 template<typename T> void SafeDelete(T*& Data) 
 { 
 	if (Data)
@@ -704,13 +714,20 @@ void UGroomComponent::InitResources()
 		
 		Res.RenderRestResources = GroupData.HairStrandsRestResource;
 		Res.SimRestResources = GroupData.HairSimulationRestResource;
+
 		Res.RenderDeformedResources = new FHairStrandsDeformedResource(GroupData.HairRenderData.RenderData, false);
 		Res.SimDeformedResources = new FHairStrandsDeformedResource(GroupData.HairSimulationData.RenderData, true);
+
 		BeginInitResource(Res.RenderDeformedResources);
 		BeginInitResource(Res.SimDeformedResources);
 
 		const FVector RestHairPositionOffset = Res.RenderRestResources->PositionOffset;
 
+		Res.RenderDeformedResources->PositionOffset = RestHairPositionOffset;
+		Res.RenderRestResources->PositionOffset = RestHairPositionOffset;
+		Res.SimDeformedResources->PositionOffset = RestHairPositionOffset;
+		Res.SimRestResources->PositionOffset = RestHairPositionOffset;
+
 		FHairStrandsInterpolationOutput::HairGroup& InterpolationOutputGroup = InterpolationOutput->HairGroups.AddDefaulted_GetRef();
 		FHairStrandsInterpolationInput::FHairGroup& InterpolationInputGroup = InterpolationInput->HairGroups.AddDefaulted_GetRef();
 		InterpolationInputGroup.HairRadius = GroupData.HairRenderData.StrandsCurves.MaxRadius;
@@ -1045,7 +1062,7 @@ void UGroomComponent::TickComponent(float DeltaTime, enum ELevelTick TickType, F
 								{
 									for (FHairStrandsRootResource::FMeshProjectionLOD& MeshProjectionLOD : Res.RenRootResources->MeshProjectionLODs)
 									{
-										MeshProjectionLOD.RestRootCenter = RestPositionOffset;
+										MeshProjectionLOD.RestRootOffset = RestPositionOffset;
 									}
 								}
 
@@ -1053,7 +1070,7 @@ void UGroomComponent::TickComponent(float DeltaTime, enum ELevelTick TickType, F
 								{
 									for (FHairStrandsRootResource::FMeshProjectionLOD& MeshProjectionLOD : Res.SimRootResources->MeshProjectionLODs)
 									{
-										MeshProjectionLOD.RestRootCenter = RestPositionOffset;
+										MeshProjectionLOD.RestRootOffset = RestPositionOffset;
 									}
 								}
 							}
@@ -1138,14 +1155,14 @@ void UGroomComponent::TickComponent(float DeltaTime, enum ELevelTick TickType, F
 			{
 				for (FHairStrandsRootResource::FMeshProjectionLOD& MeshProjectionLOD : Res.RenRootResources->MeshProjectionLODs)
 				{
-					MeshProjectionLOD.DeformedRootCenter = DeformedPositionCenter;
+					MeshProjectionLOD.DeformedRootOffset = DeformedPositionCenter;
 				}
 			}
 			if (Res.SimRootResources)
 			{
 				for (FHairStrandsRootResource::FMeshProjectionLOD& MeshProjectionLOD : Res.SimRootResources->MeshProjectionLODs)
 				{
-					MeshProjectionLOD.DeformedRootCenter = DeformedPositionCenter;
+					MeshProjectionLOD.DeformedRootOffset = DeformedPositionCenter;
 				}
 			}
 		}
diff --git a/Engine/Plugins/Experimental/HairStrands/Source/HairStrandsCore/Private/HairStrandsRendering.cpp b/Engine/Plugins/Experimental/HairStrands/Source/HairStrandsCore/Private/HairStrandsRendering.cpp
--- a/Engine/Plugins/Experimental/HairStrands/Source/HairStrandsCore/Private/HairStrandsRendering.cpp
+++ b/Engine/Plugins/Experimental/HairStrands/Source/HairStrandsCore/Private/HairStrandsRendering.cpp
@@ -14,6 +14,8 @@ static FAutoConsoleVariableRef CVarHairDeformationType(TEXT("r.HairStrands.Defor
 static float GHairRaytracingRadiusScale = 0;
 static FAutoConsoleVariableRef CVarHairRaytracingRadiusScale(TEXT("r.HairStrands.RaytracingRadiusScale"), GHairRaytracingRadiusScale, TEXT("Override the per instance scale factor for raytracing hair strands geometry (0: disabled, >0:enabled)"));
 
+static int32 GHairLocalInterpolation = 0;
+static FAutoConsoleVariableRef CVarHairLocalInterpolation(TEXT("r.HairStrands.LocalInterpolation"), GHairLocalInterpolation, TEXT("Disable the default skinning onto the skel mesh"));
 
 static FIntVector ComputeDispatchCount(uint32 ItemCount, uint32 GroupSize)
 {
@@ -109,7 +111,8 @@ class FHairInterpolationCS : public FGlobalShader
 	class FGroupSize : SHADER_PERMUTATION_INT("PERMUTATION_GROUP_SIZE", 2);
 	class FDebug : SHADER_PERMUTATION_INT("PERMUTATION_DEBUG", 2);
 	class FDynamicGeometry : SHADER_PERMUTATION_INT("PERMUTATION_DYNAMIC_GEOMETRY", 2);
-	using FPermutationDomain = TShaderPermutationDomain<FGroupSize, FDebug, FDynamicGeometry>;
+	class FLocalInterpolation : SHADER_PERMUTATION_INT("PERMUTATION_LOCAL_INTERPOLATION", 2);
+	using FPermutationDomain = TShaderPermutationDomain<FGroupSize, FDebug, FDynamicGeometry,FLocalInterpolation>;
 
 	BEGIN_SHADER_PARAMETER_STRUCT(FParameters, )
 		SHADER_PARAMETER(uint32, VertexCount)
@@ -120,6 +123,9 @@ class FHairInterpolationCS : public FGlobalShader
 		SHADER_PARAMETER(FVector, RestPositionOffset)
 		SHADER_PARAMETER(FVector, DeformedPositionOffset)
 
+		SHADER_PARAMETER(FVector, SimRestPositionOffset)
+		SHADER_PARAMETER(FVector, SimDeformedPositionOffset)
+
 		SHADER_PARAMETER_SRV(Buffer, RenderRestPosePositionBuffer)
 		SHADER_PARAMETER_UAV(RWBuffer, OutRenderDeformedPositionBuffer)
 
@@ -143,6 +149,18 @@ class FHairInterpolationCS : public FGlobalShader
 		SHADER_PARAMETER_SRV(Buffer<uint>, RootToTriangleIndex)
 		SHADER_PARAMETER_SRV(Buffer<uint>, VertexToRootIndexBuffer)
 
+		SHADER_PARAMETER_SRV(Buffer<float4>, SimRestPosition0Buffer)
+		SHADER_PARAMETER_SRV(Buffer<float4>, SimRestPosition1Buffer)
+		SHADER_PARAMETER_SRV(Buffer<float4>, SimRestPosition2Buffer)
+
+		SHADER_PARAMETER_SRV(Buffer<float4>, SimDeformedPosition0Buffer)
+		SHADER_PARAMETER_SRV(Buffer<float4>, SimDeformedPosition1Buffer)
+		SHADER_PARAMETER_SRV(Buffer<float4>, SimDeformedPosition2Buffer)
+
+		SHADER_PARAMETER_SRV(Buffer<uint>, SimRootBarycentricBuffer)
+		SHADER_PARAMETER_SRV(Buffer<uint>, SimRootToTriangleIndex)
+		SHADER_PARAMETER_SRV(Buffer<uint>, SimVertexToRootIndexBuffer)
+
 		END_SHADER_PARAMETER_STRUCT()
 
 public:
@@ -153,7 +171,8 @@ IMPLEMENT_GLOBAL_SHADER(FHairInterpolationCS, "/Engine/Private/HairStrands/HairS
 
 static void AddHairStrandsInterpolationPass(
 	FRDGBuilder& GraphBuilder,
-	const FHairStrandsProjectionHairData::HairGroup& InHairData,
+	const FHairStrandsProjectionHairData::HairGroup& InRenHairData,
+	const FHairStrandsProjectionHairData::HairGroup& InSimHairData,
 	const FVector& InHairWorldOffset, 
 	const FVector& OutHairWorldOffset,
 	const int32 LODIndex,
@@ -188,29 +207,44 @@ static void AddHairStrandsInterpolationPass(
 	Parameters->OutHairPositionOffset = OutHairWorldOffset;
 	Parameters->DispatchCountX = DispatchCount.X;
 
-	const bool bSupportDynamicMesh = InHairData.RootCount > 0 && LODIndex >= 0 && LODIndex < InHairData.LODDatas.Num() && InHairData.LODDatas[LODIndex].bIsValid;
+	const bool bSupportDynamicMesh = InRenHairData.RootCount > 0 && LODIndex >= 0 && LODIndex < InRenHairData.LODDatas.Num() && InRenHairData.LODDatas[LODIndex].bIsValid;
 	if (bSupportDynamicMesh)
 	{
-		Parameters->RestPositionOffset = InHairData.LODDatas[LODIndex].RestPositionOffset;
-		Parameters->RestPosition0Buffer = InHairData.LODDatas[LODIndex].RestRootTrianglePosition0Buffer->SRV;
-		Parameters->RestPosition1Buffer = InHairData.LODDatas[LODIndex].RestRootTrianglePosition1Buffer->SRV;
-		Parameters->RestPosition2Buffer = InHairData.LODDatas[LODIndex].RestRootTrianglePosition2Buffer->SRV;
-
-		Parameters->DeformedPositionOffset = InHairData.LODDatas[LODIndex].DeformedPositionOffset;
-		Parameters->DeformedPosition0Buffer = InHairData.LODDatas[LODIndex].DeformedRootTrianglePosition0Buffer->SRV;
-		Parameters->DeformedPosition1Buffer = InHairData.LODDatas[LODIndex].DeformedRootTrianglePosition1Buffer->SRV;
-		Parameters->DeformedPosition2Buffer = InHairData.LODDatas[LODIndex].DeformedRootTrianglePosition2Buffer->SRV;
-
-		Parameters->RootToTriangleIndex = InHairData.LODDatas[LODIndex].RootTriangleIndexBuffer->SRV;
-		Parameters->RootBarycentricBuffer = InHairData.LODDatas[LODIndex].RootTriangleBarycentricBuffer->SRV;
-		Parameters->VertexToRootIndexBuffer = InHairData.VertexToCurveIndexBuffer->SRV;
+		Parameters->RestPositionOffset = InRenHairData.LODDatas[LODIndex].RestPositionOffset;
+		Parameters->RestPosition0Buffer = InRenHairData.LODDatas[LODIndex].RestRootTrianglePosition0Buffer->SRV;
+		Parameters->RestPosition1Buffer = InRenHairData.LODDatas[LODIndex].RestRootTrianglePosition1Buffer->SRV;
+		Parameters->RestPosition2Buffer = InRenHairData.LODDatas[LODIndex].RestRootTrianglePosition2Buffer->SRV;
+
+		Parameters->DeformedPositionOffset = InRenHairData.LODDatas[LODIndex].DeformedPositionOffset;
+		Parameters->DeformedPosition0Buffer = InRenHairData.LODDatas[LODIndex].DeformedRootTrianglePosition0Buffer->SRV;
+		Parameters->DeformedPosition1Buffer = InRenHairData.LODDatas[LODIndex].DeformedRootTrianglePosition1Buffer->SRV;
+		Parameters->DeformedPosition2Buffer = InRenHairData.LODDatas[LODIndex].DeformedRootTrianglePosition2Buffer->SRV;
+
+		Parameters->RootToTriangleIndex = InRenHairData.LODDatas[LODIndex].RootTriangleIndexBuffer->SRV;
+		Parameters->RootBarycentricBuffer = InRenHairData.LODDatas[LODIndex].RootTriangleBarycentricBuffer->SRV;
+		Parameters->VertexToRootIndexBuffer = InRenHairData.VertexToCurveIndexBuffer->SRV;
+
+		Parameters->SimRestPositionOffset = InSimHairData.LODDatas[LODIndex].RestPositionOffset;
+		Parameters->SimRestPosition0Buffer = InSimHairData.LODDatas[LODIndex].RestRootTrianglePosition0Buffer->SRV;
+		Parameters->SimRestPosition1Buffer = InSimHairData.LODDatas[LODIndex].RestRootTrianglePosition1Buffer->SRV;
+		Parameters->SimRestPosition2Buffer = InSimHairData.LODDatas[LODIndex].RestRootTrianglePosition2Buffer->SRV;
+
+		Parameters->SimDeformedPositionOffset = InSimHairData.LODDatas[LODIndex].DeformedPositionOffset;
+		Parameters->SimDeformedPosition0Buffer = InSimHairData.LODDatas[LODIndex].DeformedRootTrianglePosition0Buffer->SRV;
+		Parameters->SimDeformedPosition1Buffer = InSimHairData.LODDatas[LODIndex].DeformedRootTrianglePosition1Buffer->SRV;
+		Parameters->SimDeformedPosition2Buffer = InSimHairData.LODDatas[LODIndex].DeformedRootTrianglePosition2Buffer->SRV;
+
+		Parameters->SimRootToTriangleIndex = InSimHairData.LODDatas[LODIndex].RootTriangleIndexBuffer->SRV;
+		Parameters->SimRootBarycentricBuffer = InSimHairData.LODDatas[LODIndex].RootTriangleBarycentricBuffer->SRV;
+		Parameters->SimVertexToRootIndexBuffer = InSimHairData.VertexToCurveIndexBuffer->SRV;
 	}
 
 	FHairInterpolationCS::FPermutationDomain PermutationVector;
 	PermutationVector.Set<FHairInterpolationCS::FGroupSize>(GetGroupSizePermutation(GroupSize));
 	PermutationVector.Set<FHairInterpolationCS::FDebug>(bCopySimAttributesToRenderAttributes ? 1 : 0);
 	PermutationVector.Set<FHairInterpolationCS::FDynamicGeometry>(bSupportDynamicMesh ? 1 : 0);
-	
+	PermutationVector.Set<FHairInterpolationCS::FLocalInterpolation>((GHairLocalInterpolation == 1) ? 1 : 0);
+
 	TShaderMap<FGlobalShaderType>* ShaderMap = GetGlobalShaderMap(ERHIFeatureLevel::SM5);
 
 	TShaderMapRef<FHairInterpolationCS> ComputeShader(ShaderMap, PermutationVector);
@@ -374,7 +408,8 @@ void ComputeHairStrandsInterpolation(
 	FRHICommandListImmediate& RHICmdList,
 	FHairStrandsInterpolationInput* InInput,
 	FHairStrandsInterpolationOutput* InOutput,
-	FHairStrandsProjectionHairData& InHairDatas,
+	FHairStrandsProjectionHairData& InRenHairDatas,
+	FHairStrandsProjectionHairData& InSimHairDatas,
 	int32 LODIndex)
 {
 	if (!InInput || !InOutput) return;
@@ -449,10 +484,12 @@ void ComputeHairStrandsInterpolation(
 				Output.RenderPatchedAttributeBuffer.Initialize(FHairStrandsAttributeFormat::SizeInByte, Input.RenderVertexCount, FHairStrandsAttributeFormat::Format, BUF_Static);
 			}
 
-			check(GroupIndex < uint32(InHairDatas.HairGroups.Num()));
+			check(GroupIndex < uint32(InRenHairDatas.HairGroups.Num()));
+			check(GroupIndex < uint32(InSimHairDatas.HairGroups.Num()));
 			AddHairStrandsInterpolationPass(
 				GraphBuilder,
-				InHairDatas.HairGroups[GroupIndex],
+				InRenHairDatas.HairGroups[GroupIndex],
+				InSimHairDatas.HairGroups[GroupIndex],
 				Input.InHairPositionOffset,
 				Input.OutHairPositionOffset,
 				LODIndex,
diff --git a/Engine/Plugins/Experimental/HairStrands/Source/HairStrandsCore/Public/GroomAsset.h b/Engine/Plugins/Experimental/HairStrands/Source/HairStrandsCore/Public/GroomAsset.h
--- a/Engine/Plugins/Experimental/HairStrands/Source/HairStrandsCore/Public/GroomAsset.h
+++ b/Engine/Plugins/Experimental/HairStrands/Source/HairStrandsCore/Public/GroomAsset.h
@@ -43,13 +43,13 @@ struct FHairStrandsRootResource : public FRenderResource
 		FRWBuffer RootTriangleBarycentricBuffer;
 	
 		/* Strand hair roots translation and rotation in rest position relative to the bound triangle. Positions are relative to the rest root center */
-		FVector	  RestRootCenter = FVector::ZeroVector;
+		FVector	  RestRootOffset = FVector::ZeroVector;
 		FRWBuffer RestRootTrianglePosition0Buffer;
 		FRWBuffer RestRootTrianglePosition1Buffer;
 		FRWBuffer RestRootTrianglePosition2Buffer;
 
 		/* Strand hair roots translation and rotation in triangle-deformed position relative to the bound triangle. Positions are relative the deformed root center*/
-		FVector   DeformedRootCenter = FVector::ZeroVector;
+		FVector   DeformedRootOffset = FVector::ZeroVector;
 		FRWBuffer DeformedRootTrianglePosition0Buffer;
 		FRWBuffer DeformedRootTrianglePosition1Buffer;
 		FRWBuffer DeformedRootTrianglePosition2Buffer;
diff --git a/Engine/Plugins/Experimental/HairStrands/Source/HairStrandsCore/Public/GroomComponent.h b/Engine/Plugins/Experimental/HairStrands/Source/HairStrandsCore/Public/GroomComponent.h
--- a/Engine/Plugins/Experimental/HairStrands/Source/HairStrandsCore/Public/GroomComponent.h
+++ b/Engine/Plugins/Experimental/HairStrands/Source/HairStrandsCore/Public/GroomComponent.h
@@ -89,6 +89,9 @@ public:
 	/** Return the guide hairs deformed resources*/
 	FHairStrandsDeformedResource* GetGuideStrandsDeformedResource(uint32 GroupIndex);
 
+	/** Return the guide hairs root resources*/
+	FHairStrandsRootResource* GetGuideStrandsRootResource(uint32 GroupIndex);
+
 #if WITH_EDITOR
 	virtual void PostEditChangeProperty(FPropertyChangedEvent& PropertyChangedEvent) override;
 	virtual bool CanEditChange(const UProperty* InProperty) const override;
diff --git a/Engine/Plugins/Experimental/HairStrands/Source/HairStrandsCore/Public/HairStrandsRendering.h b/Engine/Plugins/Experimental/HairStrands/Source/HairStrandsCore/Public/HairStrandsRendering.h
--- a/Engine/Plugins/Experimental/HairStrands/Source/HairStrandsCore/Public/HairStrandsRendering.h
+++ b/Engine/Plugins/Experimental/HairStrands/Source/HairStrandsCore/Public/HairStrandsRendering.h
@@ -114,6 +114,7 @@ void ComputeHairStrandsInterpolation(
 	FRHICommandListImmediate& RHICmdList,
 	FHairStrandsInterpolationInput* Input,
 	FHairStrandsInterpolationOutput* Output, 
-	struct FHairStrandsProjectionHairData& HairDatas,
+	struct FHairStrandsProjectionHairData& RenHairDatas,
+	struct FHairStrandsProjectionHairData& SimHairDatas,
 	int32 LODIndex);
 
diff --git a/Engine/Plugins/Experimental/HairStrands/Source/HairStrandsNiagara/Private/NiagaraDataInterfaceHairStrands.cpp b/Engine/Plugins/Experimental/HairStrands/Source/HairStrandsNiagara/Private/NiagaraDataInterfaceHairStrands.cpp
--- a/Engine/Plugins/Experimental/HairStrands/Source/HairStrandsNiagara/Private/NiagaraDataInterfaceHairStrands.cpp
+++ b/Engine/Plugins/Experimental/HairStrands/Source/HairStrandsNiagara/Private/NiagaraDataInterfaceHairStrands.cpp
@@ -15,6 +15,8 @@
 
 #define LOCTEXT_NAMESPACE "NiagaraDataInterfaceHairStrands"
 
+static const int32 ResetTick = 4;
+
 //------------------------------------------------------------------------------------------------------------
 
 static const FName GetStrandSizeName(TEXT("GetStrandSize"));
@@ -31,8 +33,12 @@ static const FName ComputeNodeMassName(TEXT("ComputeNodeMass"));
 static const FName ComputeNodeInertiaName(TEXT("ComputeNodeInertia"));
 static const FName ComputeEdgeLengthName(TEXT("ComputeEdgeLength"));
 static const FName ComputeEdgeRotationName(TEXT("ComputeEdgeRotation"));
+
+//------------------------------------------------------------------------------------------------------------
+
 static const FName ComputeRestPositionName(TEXT("ComputeRestPosition"));
 static const FName ComputeRestOrientationName(TEXT("ComputeRestOrientation"));
+static const FName ComputeLocalStateName(TEXT("ComputeLocalState"));
 
 //------------------------------------------------------------------------------------------------------------
 
@@ -45,6 +51,7 @@ static const FName UpdateAngularVelocityName(TEXT("UpdateAngularVelocity"));
 
 static const FName AttachNodePositionName(TEXT("AttachNodePosition"));
 static const FName AttachNodeOrientationName(TEXT("AttachNodeOrientation"));
+static const FName AttachNodeStateName(TEXT("AttachNodeState"));
 
 //------------------------------------------------------------------------------------------------------------
 
@@ -92,6 +99,9 @@ static const FName UpdateNodeOrientationName(TEXT("UpdateNodeOrientation"));
 //------------------------------------------------------------------------------------------------------------
 
 static const FName ComputeAirDragForceName(TEXT("ComputeAirDragForce"));  
+//------------------------------------------------------------------------------------------------------------
+
+static const FName NeedSimulationResetName(TEXT("NeedSimulationReset"));
 
 //------------------------------------------------------------------------------------------------------------
 
@@ -100,12 +110,29 @@ const FString UNiagaraDataInterfaceHairStrands::StrandSizeName(TEXT("StrandSize_
 const FString UNiagaraDataInterfaceHairStrands::WorldTransformName(TEXT("WorldTransform_"));
 const FString UNiagaraDataInterfaceHairStrands::WorldInverseName(TEXT("WorldInverse_"));
 const FString UNiagaraDataInterfaceHairStrands::WorldRotationName(TEXT("WorldRotation_"));
+const FString UNiagaraDataInterfaceHairStrands::BoxCenterName(TEXT("BoxCenter_"));
+const FString UNiagaraDataInterfaceHairStrands::BoxExtentName(TEXT("BoxExtent_"));
 
-const FString UNiagaraDataInterfaceHairStrands::PointsPositionsBufferName(TEXT("PointsPositionsBuffer_"));
+const FString UNiagaraDataInterfaceHairStrands::DeformedPositionBufferName(TEXT("DeformedPositionBuffer_"));
 const FString UNiagaraDataInterfaceHairStrands::CurvesOffsetsBufferName(TEXT("CurvesOffsetsBuffer_"));
-const FString UNiagaraDataInterfaceHairStrands::RestPositionsBufferName(TEXT("RestPositionsBuffer_"));
+const FString UNiagaraDataInterfaceHairStrands::RestPositionBufferName(TEXT("RestPositionBuffer_"));
 
-const FString UNiagaraDataInterfaceHairStrands::BoxOriginName(TEXT("BoxOrigin_"));
+const FString UNiagaraDataInterfaceHairStrands::ResetSimulationName(TEXT("ResetSimulation_"));
+const FString UNiagaraDataInterfaceHairStrands::HasRootAttachedName(TEXT("HasRootAttached_"));
+const FString UNiagaraDataInterfaceHairStrands::RootBarycentricCoordinatesName(TEXT("RootBarycentricCoordinatesBuffer_"));
+
+const FString UNiagaraDataInterfaceHairStrands::RestRootOffsetName(TEXT("RestRootOffset_"));
+const FString UNiagaraDataInterfaceHairStrands::RestTrianglePositionAName(TEXT("RestTrianglePositionABuffer_"));
+const FString UNiagaraDataInterfaceHairStrands::RestTrianglePositionBName(TEXT("RestTrianglePositionBBuffer_"));
+const FString UNiagaraDataInterfaceHairStrands::RestTrianglePositionCName(TEXT("RestTrianglePositionCBuffer_"));
+
+const FString UNiagaraDataInterfaceHairStrands::DeformedRootOffsetName(TEXT("DeformedRootOffset_"));
+const FString UNiagaraDataInterfaceHairStrands::DeformedTrianglePositionAName(TEXT("DeformedTrianglePositionABuffer_"));
+const FString UNiagaraDataInterfaceHairStrands::DeformedTrianglePositionBName(TEXT("DeformedTrianglePositionBBuffer_"));
+const FString UNiagaraDataInterfaceHairStrands::DeformedTrianglePositionCName(TEXT("DeformedTrianglePositionCBuffer_"));
+
+const FString UNiagaraDataInterfaceHairStrands::RestPositionOffsetName(TEXT("RestPositionOffset_"));
+const FString UNiagaraDataInterfaceHairStrands::DeformedPositionOffsetName(TEXT("DeformedPositionOffset_"));
 
 //------------------------------------------------------------------------------------------------------------
 
@@ -118,11 +145,29 @@ struct FNDIHairStrandsParametersName
 		WorldTransformName = UNiagaraDataInterfaceHairStrands::WorldTransformName + Suffix;
 		WorldInverseName = UNiagaraDataInterfaceHairStrands::WorldInverseName + Suffix;
 		WorldRotationName = UNiagaraDataInterfaceHairStrands::WorldRotationName + Suffix;
-		BoxOriginName = UNiagaraDataInterfaceHairStrands::BoxOriginName + Suffix;
+		BoxCenterName = UNiagaraDataInterfaceHairStrands::BoxCenterName + Suffix;
+		BoxExtentName = UNiagaraDataInterfaceHairStrands::BoxExtentName + Suffix;
 
-		PointsPositionsBufferName = UNiagaraDataInterfaceHairStrands::PointsPositionsBufferName + Suffix;
+		DeformedPositionBufferName = UNiagaraDataInterfaceHairStrands::DeformedPositionBufferName + Suffix;
 		CurvesOffsetsBufferName = UNiagaraDataInterfaceHairStrands::CurvesOffsetsBufferName + Suffix;
-		RestPositionsBufferName = UNiagaraDataInterfaceHairStrands::RestPositionsBufferName + Suffix;
+		RestPositionBufferName = UNiagaraDataInterfaceHairStrands::RestPositionBufferName + Suffix;
+
+		HasRootAttachedName = UNiagaraDataInterfaceHairStrands::HasRootAttachedName + Suffix;
+		ResetSimulationName = UNiagaraDataInterfaceHairStrands::ResetSimulationName + Suffix;
+		RootBarycentricCoordinatesName = UNiagaraDataInterfaceHairStrands::RootBarycentricCoordinatesName + Suffix;
+
+		RestRootOffsetName = UNiagaraDataInterfaceHairStrands::RestRootOffsetName + Suffix;
+		RestTrianglePositionAName = UNiagaraDataInterfaceHairStrands::RestTrianglePositionAName + Suffix;
+		RestTrianglePositionBName = UNiagaraDataInterfaceHairStrands::RestTrianglePositionBName + Suffix;
+		RestTrianglePositionCName = UNiagaraDataInterfaceHairStrands::RestTrianglePositionCName + Suffix;
+
+		DeformedRootOffsetName = UNiagaraDataInterfaceHairStrands::DeformedRootOffsetName + Suffix;
+		DeformedTrianglePositionAName = UNiagaraDataInterfaceHairStrands::DeformedTrianglePositionAName + Suffix;
+		DeformedTrianglePositionBName = UNiagaraDataInterfaceHairStrands::DeformedTrianglePositionBName + Suffix;
+		DeformedTrianglePositionCName = UNiagaraDataInterfaceHairStrands::DeformedTrianglePositionCName + Suffix;
+
+		RestPositionOffsetName = UNiagaraDataInterfaceHairStrands::RestPositionOffsetName + Suffix;
+		DeformedPositionOffsetName = UNiagaraDataInterfaceHairStrands::DeformedPositionOffsetName + Suffix;
 	}
 
 	FString NumStrandsName;
@@ -130,20 +175,40 @@ struct FNDIHairStrandsParametersName
 	FString WorldTransformName;
 	FString WorldInverseName;
 	FString WorldRotationName;
-	FString BoxOriginName;
+	FString BoxCenterName;
+	FString BoxExtentName;
 
-	FString PointsPositionsBufferName;
+	FString DeformedPositionBufferName;
 	FString CurvesOffsetsBufferName;
-	FString RestPositionsBufferName;
+	FString RestPositionBufferName;
+
+	FString ResetSimulationName;
+	FString HasRootAttachedName;
+	FString RootBarycentricCoordinatesName;
+
+	FString RestRootOffsetName;
+	FString RestTrianglePositionAName;
+	FString RestTrianglePositionBName;
+	FString RestTrianglePositionCName;
+
+	FString DeformedRootOffsetName;
+	FString DeformedTrianglePositionAName;
+	FString DeformedTrianglePositionBName;
+	FString DeformedTrianglePositionCName;
+
+	FString RestPositionOffsetName;
+	FString DeformedPositionOffsetName;
 };
 
 //------------------------------------------------------------------------------------------------------------
 
-void FNDIHairStrandsBuffer::SetHairAsset(const FHairStrandsDatas*  HairStrandsDatas, const FHairStrandsRestResource* HairStrandsRestResource, const FHairStrandsDeformedResource*  HairStrandsDeformedResource)
+void FNDIHairStrandsBuffer::SetHairAsset(const FHairStrandsDatas*  HairStrandsDatas, const FHairStrandsRestResource* HairStrandsRestResource, 
+	const FHairStrandsDeformedResource*  HairStrandsDeformedResource, const FHairStrandsRootResource* HairStrandsRootResource)
 {
 	SourceDatas = HairStrandsDatas;
 	SourceRestResources = HairStrandsRestResource;
 	SourceDeformedResources = HairStrandsDeformedResource;
+	SourceRootResources = HairStrandsRootResource;
 }
 
 void FNDIHairStrandsBuffer::InitRHI()
@@ -163,7 +228,20 @@ void FNDIHairStrandsBuffer::InitRHI()
 		if (SourceDeformedResources == nullptr)
 		{
 			const uint32 PositionsCount = SourceDatas->GetNumPoints();
-			PointsPositionsBuffer.Initialize(FHairStrandsPositionFormat::SizeInByte, PositionsCount, FHairStrandsPositionFormat::Format, BUF_Static);
+			DeformedPositionBuffer.Initialize(FHairStrandsPositionFormat::SizeInByte, PositionsCount, FHairStrandsPositionFormat::Format, BUF_Static);
+		}
+		if (SourceRootResources == nullptr)
+		{
+			const uint32 PositionsCount = SourceDatas->GetNumCurves();
+			RestTrianglePositionABuffer.Initialize(FHairStrandsMeshTrianglePositionFormat::SizeInByte, PositionsCount, FHairStrandsMeshTrianglePositionFormat::Format, BUF_Static);
+			RestTrianglePositionBBuffer.Initialize(FHairStrandsMeshTrianglePositionFormat::SizeInByte, PositionsCount, FHairStrandsMeshTrianglePositionFormat::Format, BUF_Static);
+			RestTrianglePositionCBuffer.Initialize(FHairStrandsMeshTrianglePositionFormat::SizeInByte, PositionsCount, FHairStrandsMeshTrianglePositionFormat::Format, BUF_Static);
+
+			DeformedTrianglePositionABuffer.Initialize(FHairStrandsMeshTrianglePositionFormat::SizeInByte, PositionsCount, FHairStrandsMeshTrianglePositionFormat::Format, BUF_Static);
+			DeformedTrianglePositionBBuffer.Initialize(FHairStrandsMeshTrianglePositionFormat::SizeInByte, PositionsCount, FHairStrandsMeshTrianglePositionFormat::Format, BUF_Static);
+			DeformedTrianglePositionCBuffer.Initialize(FHairStrandsMeshTrianglePositionFormat::SizeInByte, PositionsCount, FHairStrandsMeshTrianglePositionFormat::Format, BUF_Static);
+
+			RootBarycentricCoordinatesBuffer.Initialize(FHairStrandsCurveTriangleBarycentricFormat::SizeInByte, PositionsCount, FHairStrandsCurveTriangleBarycentricFormat::Format, BUF_Static);
 		}
 	}
 }
@@ -171,7 +249,17 @@ void FNDIHairStrandsBuffer::InitRHI()
 void FNDIHairStrandsBuffer::ReleaseRHI()
 {
 	CurvesOffsetsBuffer.Release();
-	PointsPositionsBuffer.Release();
+	DeformedPositionBuffer.Release();
+
+	RestTrianglePositionABuffer.Release();
+	RestTrianglePositionBBuffer.Release();
+	RestTrianglePositionCBuffer.Release();
+
+	DeformedTrianglePositionABuffer.Release();
+	DeformedTrianglePositionBBuffer.Release();
+	DeformedTrianglePositionCBuffer.Release();
+
+	RootBarycentricCoordinatesBuffer.Release();
 }
 
 //------------------------------------------------------------------------------------------------------------
@@ -187,23 +275,73 @@ struct FNDIHairStrandsParametersCS : public FNiagaraDataInterfaceParametersCS
 		WorldRotation.Bind(ParameterMap, *ParamNames.WorldRotationName);
 		NumStrands.Bind(ParameterMap, *ParamNames.NumStrandsName);
 		StrandSize.Bind(ParameterMap, *ParamNames.StrandSizeName);
-		BoxOrigin.Bind(ParameterMap, *ParamNames.BoxOriginName);
+		BoxCenter.Bind(ParameterMap, *ParamNames.BoxCenterName);
+		BoxExtent.Bind(ParameterMap, *ParamNames.BoxExtentName);
 
-		PointsPositionsBuffer.Bind(ParameterMap, *ParamNames.PointsPositionsBufferName);
+		DeformedPositionBuffer.Bind(ParameterMap, *ParamNames.DeformedPositionBufferName);
 		CurvesOffsetsBuffer.Bind(ParameterMap, *ParamNames.CurvesOffsetsBufferName);
-		RestPositionsBuffer.Bind(ParameterMap, *ParamNames.RestPositionsBufferName);
+		RestPositionBuffer.Bind(ParameterMap, *ParamNames.RestPositionBufferName);
+
+		ResetSimulation.Bind(ParameterMap, *ParamNames.ResetSimulationName);
+		HasRootAttached.Bind(ParameterMap,*ParamNames.HasRootAttachedName);
+		RestRootOffset.Bind(ParameterMap, *ParamNames.RestRootOffsetName);
+		DeformedRootOffset.Bind(ParameterMap, *ParamNames.DeformedRootOffsetName);
+
+		RestPositionOffset.Bind(ParameterMap, *ParamNames.RestPositionOffsetName);
+		DeformedPositionOffset.Bind(ParameterMap, *ParamNames.DeformedPositionOffsetName);
+
+		RootBarycentricCoordinatesBuffer.Bind(ParameterMap, *ParamNames.RootBarycentricCoordinatesName);
 
-		if (!PointsPositionsBuffer.IsBound())
+		RestTrianglePositionABuffer.Bind(ParameterMap, *ParamNames.RestTrianglePositionAName);
+		RestTrianglePositionBBuffer.Bind(ParameterMap, *ParamNames.RestTrianglePositionBName);
+		RestTrianglePositionCBuffer.Bind(ParameterMap, *ParamNames.RestTrianglePositionCName);
+
+		DeformedTrianglePositionABuffer.Bind(ParameterMap, *ParamNames.DeformedTrianglePositionAName);
+		DeformedTrianglePositionBBuffer.Bind(ParameterMap, *ParamNames.DeformedTrianglePositionBName);
+		DeformedTrianglePositionCBuffer.Bind(ParameterMap, *ParamNames.DeformedTrianglePositionCName);
+
+		if (!DeformedPositionBuffer.IsBound())
 		{
-			UE_LOG(LogHairStrands, Warning, TEXT("Binding failed for FNDIHairStrandsParametersCS %s. Was it optimized out?"), *ParamNames.PointsPositionsBufferName)
+			UE_LOG(LogHairStrands, Warning, TEXT("Binding failed for FNDIHairStrandsParametersCS %s. Was it optimized out?"), *ParamNames.DeformedPositionBufferName)
 		}
 		if (!CurvesOffsetsBuffer.IsBound())
 		{
 			UE_LOG(LogHairStrands, Warning, TEXT("Binding failed for FNDIHairStrandsParametersCS %s. Was it optimized out?"), *ParamNames.CurvesOffsetsBufferName)
 		}
-		if (!RestPositionsBuffer.IsBound())
+		if (!RestPositionBuffer.IsBound())
+		{
+			UE_LOG(LogHairStrands, Warning, TEXT("Binding failed for FNDIHairStrandsParametersCS %s. Was it optimized out?"), *ParamNames.RestPositionBufferName)
+		}
+
+		if (!RootBarycentricCoordinatesBuffer.IsBound())
+		{
+			UE_LOG(LogHairStrands, Warning, TEXT("Binding failed for FNDIHairStrandsParametersCS %s. Was it optimized out?"), *ParamNames.RootBarycentricCoordinatesName)
+		}
+
+		if (!RestTrianglePositionABuffer.IsBound())
+		{
+			UE_LOG(LogHairStrands, Warning, TEXT("Binding failed for FNDIHairStrandsParametersCS %s. Was it optimized out?"), *ParamNames.RestTrianglePositionAName)
+		}
+		if (!RestTrianglePositionBBuffer.IsBound())
+		{
+			UE_LOG(LogHairStrands, Warning, TEXT("Binding failed for FNDIHairStrandsParametersCS %s. Was it optimized out?"), *ParamNames.RestTrianglePositionBName)
+		}
+		if (!RestTrianglePositionCBuffer.IsBound())
+		{
+			UE_LOG(LogHairStrands, Warning, TEXT("Binding failed for FNDIHairStrandsParametersCS %s. Was it optimized out?"), *ParamNames.RestTrianglePositionCName)
+		}
+
+		if (!DeformedTrianglePositionABuffer.IsBound())
 		{
-			UE_LOG(LogHairStrands, Warning, TEXT("Binding failed for FNDIHairStrandsParametersCS %s. Was it optimized out?"), *ParamNames.RestPositionsBufferName)
+			UE_LOG(LogHairStrands, Warning, TEXT("Binding failed for FNDIHairStrandsParametersCS %s. Was it optimized out?"), *ParamNames.DeformedTrianglePositionAName)
+		}
+		if (!DeformedTrianglePositionBBuffer.IsBound())
+		{
+			UE_LOG(LogHairStrands, Warning, TEXT("Binding failed for FNDIHairStrandsParametersCS %s. Was it optimized out?"), *ParamNames.DeformedTrianglePositionBName)
+		}
+		if (!DeformedTrianglePositionCBuffer.IsBound())
+		{
+			UE_LOG(LogHairStrands, Warning, TEXT("Binding failed for FNDIHairStrandsParametersCS %s. Was it optimized out?"), *ParamNames.DeformedTrianglePositionCName)
 		}
 	}
 
@@ -214,16 +352,35 @@ struct FNDIHairStrandsParametersCS : public FNiagaraDataInterfaceParametersCS
 		Ar << WorldRotation;
 		Ar << NumStrands;
 		Ar << StrandSize;
-		Ar << BoxOrigin;
+		Ar << BoxCenter;
+		Ar << BoxExtent;
 
-		Ar << PointsPositionsBuffer;
+		Ar << DeformedPositionBuffer;
 		Ar << CurvesOffsetsBuffer;
-		Ar << RestPositionsBuffer;
+		Ar << RestPositionBuffer;
+
+		Ar << ResetSimulation;
+		Ar << HasRootAttached;
+		Ar << RestRootOffset;
+		Ar << DeformedRootOffset;
+
+		Ar << RootBarycentricCoordinatesBuffer;
+
+		Ar << RestTrianglePositionABuffer;
+		Ar << RestTrianglePositionBBuffer;
+		Ar << RestTrianglePositionCBuffer;
+
+		Ar << DeformedTrianglePositionABuffer;
+		Ar << DeformedTrianglePositionBBuffer;
+		Ar << DeformedTrianglePositionCBuffer;
+
+		Ar << RestPositionOffset;
+		Ar << DeformedPositionOffset;
 	}
 
 	virtual void Set(FRHICommandList& RHICmdList, const FNiagaraDataInterfaceSetArgs& Context) const override
 	{
-		check(IsInRenderingThread());
+		check(IsInRenderingThread()); 
 
 		FRHIComputeShader* ComputeShaderRHI = Context.Shader->GetComputeShader();
 
@@ -238,37 +395,117 @@ struct FNDIHairStrandsParametersCS : public FNiagaraDataInterfaceParametersCS
 			FNDIHairStrandsBuffer* HairStrandsBuffer = ProxyData->HairStrandsBuffer;
 
 			FUnorderedAccessViewRHIRef PointPositionsUAV = (HairStrandsBuffer->SourceDeformedResources != nullptr) ? HairStrandsBuffer->SourceDeformedResources->DeformedPositionBuffer[0].UAV : 
-				HairStrandsBuffer->PointsPositionsBuffer.UAV;
+				HairStrandsBuffer->DeformedPositionBuffer.UAV;
+
+			FShaderResourceViewRHIRef RestTrianglePositionASRV = HairStrandsBuffer->RestTrianglePositionABuffer.SRV;
+			FShaderResourceViewRHIRef RestTrianglePositionBSRV = HairStrandsBuffer->RestTrianglePositionBBuffer.SRV;
+			FShaderResourceViewRHIRef RestTrianglePositionCSRV = HairStrandsBuffer->RestTrianglePositionCBuffer.SRV;
+
+			FShaderResourceViewRHIRef DeformedTrianglePositionASRV = HairStrandsBuffer->DeformedTrianglePositionABuffer.SRV;
+			FShaderResourceViewRHIRef DeformedTrianglePositionBSRV = HairStrandsBuffer->DeformedTrianglePositionBBuffer.SRV;
+			FShaderResourceViewRHIRef DeformedTrianglePositionCSRV = HairStrandsBuffer->DeformedTrianglePositionCBuffer.SRV;
+
+			FShaderResourceViewRHIRef RootBarycentricCoordinatesSRV = HairStrandsBuffer->RootBarycentricCoordinatesBuffer.SRV;
+
+			int32 HasRootAttachedValue = 0;
+			bool NeedSimReset = (ProxyData->TickCount <= ResetTick);
+			FVector RestRootOffsetValue = FVector(0, 0, 0);
+			FVector DeformedRootOffsetValue = FVector(0, 0, 0);
+			FVector DeformedPositionOffsetValue = HairStrandsBuffer->SourceRestResources->PositionOffset;
+			FVector RestPositionOffsetValue = HairStrandsBuffer->SourceRestResources->PositionOffset;
+
+			if (HairStrandsBuffer->SourceRootResources != nullptr && HairStrandsBuffer->SourceRootResources->MeshProjectionLODs.Num() > 0)
+			{
+				RestTrianglePositionASRV = HairStrandsBuffer->SourceRootResources->MeshProjectionLODs[0].RestRootTrianglePosition0Buffer.SRV;
+				RestTrianglePositionBSRV = HairStrandsBuffer->SourceRootResources->MeshProjectionLODs[0].RestRootTrianglePosition1Buffer.SRV;
+				RestTrianglePositionCSRV = HairStrandsBuffer->SourceRootResources->MeshProjectionLODs[0].RestRootTrianglePosition2Buffer.SRV;
+
+				DeformedTrianglePositionASRV = HairStrandsBuffer->SourceRootResources->MeshProjectionLODs[0].DeformedRootTrianglePosition0Buffer.SRV;
+				DeformedTrianglePositionBSRV = HairStrandsBuffer->SourceRootResources->MeshProjectionLODs[0].DeformedRootTrianglePosition1Buffer.SRV;
+				DeformedTrianglePositionCSRV = HairStrandsBuffer->SourceRootResources->MeshProjectionLODs[0].DeformedRootTrianglePosition2Buffer.SRV;
+
+				RootBarycentricCoordinatesSRV = HairStrandsBuffer->SourceRootResources->MeshProjectionLODs[0].RootTriangleBarycentricBuffer.SRV;
+
+				HasRootAttachedValue = 1;
+				NeedSimReset = (HairStrandsBuffer->SourceRootResources->MeshProjectionLODs[0].Status != FHairStrandsRootResource::FMeshProjectionLOD::EStatus::Completed);
+				RestRootOffsetValue = HairStrandsBuffer->SourceRootResources->MeshProjectionLODs[0].RestRootOffset;
+				DeformedRootOffsetValue = ProxyData->HairStrandsBuffer->SourceRootResources->MeshProjectionLODs[0].DeformedRootOffset;
+			}
+			if (HairStrandsBuffer->SourceDeformedResources != nullptr)
+			{
+				DeformedPositionOffsetValue = HairStrandsBuffer->SourceDeformedResources->PositionOffset;
+			}
 
-			SetUAVParameter(RHICmdList, ComputeShaderRHI, PointsPositionsBuffer, PointPositionsUAV);
+			//UE_LOG(LogHairStrands, Log, TEXT("Root Offset = %f %f %f | %f %f %f "), RestRootOffsetValue[0], RestRootOffsetValue[1], RestRootOffsetValue[2],
+			//	DeformedRootOffsetValue[0], DeformedRootOffsetValue[1], DeformedRootOffsetValue[2]);
+
+			SetUAVParameter(RHICmdList, ComputeShaderRHI, DeformedPositionBuffer, PointPositionsUAV);
 			SetSRVParameter(RHICmdList, ComputeShaderRHI, CurvesOffsetsBuffer, HairStrandsBuffer->CurvesOffsetsBuffer.SRV);
-			SetSRVParameter(RHICmdList, ComputeShaderRHI, RestPositionsBuffer, HairStrandsBuffer->SourceRestResources->RestPositionBuffer.SRV);
+			SetSRVParameter(RHICmdList, ComputeShaderRHI, RestPositionBuffer, HairStrandsBuffer->SourceRestResources->RestPositionBuffer.SRV);
 
 			SetShaderValue(RHICmdList, ComputeShaderRHI, WorldTransform, ProxyData->WorldTransform);
 			SetShaderValue(RHICmdList, ComputeShaderRHI, WorldInverse, ProxyData->WorldTransform.Inverse());
 			SetShaderValue(RHICmdList, ComputeShaderRHI, WorldRotation, ProxyData->WorldTransform.ToQuat());
 			SetShaderValue(RHICmdList, ComputeShaderRHI, NumStrands, ProxyData->NumStrands);
 			SetShaderValue(RHICmdList, ComputeShaderRHI, StrandSize, ProxyData->StrandSize);
-			SetShaderValue(RHICmdList, ComputeShaderRHI, BoxOrigin, ProxyData->BoxOrigin);
+			SetShaderValue(RHICmdList, ComputeShaderRHI, BoxCenter, ProxyData->BoxCenter);
+			SetShaderValue(RHICmdList, ComputeShaderRHI, BoxExtent, ProxyData->BoxExtent);
+
+			SetShaderValue(RHICmdList, ComputeShaderRHI, ResetSimulation, NeedSimReset);
+			SetShaderValue(RHICmdList, ComputeShaderRHI, HasRootAttached, HasRootAttachedValue);
+			SetShaderValue(RHICmdList, ComputeShaderRHI, RestRootOffset, RestRootOffsetValue);
+			SetShaderValue(RHICmdList, ComputeShaderRHI, DeformedRootOffset, DeformedRootOffsetValue);
+
+			SetShaderValue(RHICmdList, ComputeShaderRHI, RestPositionOffset, RestPositionOffsetValue);
+			SetShaderValue(RHICmdList, ComputeShaderRHI, DeformedPositionOffset, DeformedPositionOffsetValue);
+
+			SetSRVParameter(RHICmdList, ComputeShaderRHI, RestTrianglePositionABuffer, RestTrianglePositionASRV);
+			SetSRVParameter(RHICmdList, ComputeShaderRHI, RestTrianglePositionBBuffer, RestTrianglePositionBSRV);
+			SetSRVParameter(RHICmdList, ComputeShaderRHI, RestTrianglePositionCBuffer, RestTrianglePositionCSRV);
+
+			SetSRVParameter(RHICmdList, ComputeShaderRHI, DeformedTrianglePositionABuffer, DeformedTrianglePositionASRV);
+			SetSRVParameter(RHICmdList, ComputeShaderRHI, DeformedTrianglePositionBBuffer, DeformedTrianglePositionBSRV);
+			SetSRVParameter(RHICmdList, ComputeShaderRHI, DeformedTrianglePositionCBuffer, DeformedTrianglePositionCSRV);
+
+			SetSRVParameter(RHICmdList, ComputeShaderRHI, RootBarycentricCoordinatesBuffer, RootBarycentricCoordinatesSRV);
 		}
 		else
 		{
-			SetUAVParameter(RHICmdList, ComputeShaderRHI, PointsPositionsBuffer, FNiagaraRenderer::GetDummyFloatBuffer().UAV);
+			SetUAVParameter(RHICmdList, ComputeShaderRHI, DeformedPositionBuffer, FNiagaraRenderer::GetDummyFloatBuffer().UAV);
 			SetSRVParameter(RHICmdList, ComputeShaderRHI, CurvesOffsetsBuffer, FNiagaraRenderer::GetDummyUIntBuffer().SRV);
-			SetSRVParameter(RHICmdList, ComputeShaderRHI, RestPositionsBuffer, FNiagaraRenderer::GetDummyFloatBuffer().SRV);
+			SetSRVParameter(RHICmdList, ComputeShaderRHI, RestPositionBuffer, FNiagaraRenderer::GetDummyFloatBuffer().SRV);
 
 			SetShaderValue(RHICmdList, ComputeShaderRHI, WorldTransform, FMatrix::Identity);
 			SetShaderValue(RHICmdList, ComputeShaderRHI, WorldInverse, FMatrix::Identity);
 			SetShaderValue(RHICmdList, ComputeShaderRHI, WorldRotation, FQuat::Identity);
 			SetShaderValue(RHICmdList, ComputeShaderRHI, NumStrands, 1);
 			SetShaderValue(RHICmdList, ComputeShaderRHI, StrandSize, 1);
-			SetShaderValue(RHICmdList, ComputeShaderRHI, BoxOrigin, FVector4(0,0,0,0));
+			SetShaderValue(RHICmdList, ComputeShaderRHI, BoxCenter, FVector(0,0,0));
+			SetShaderValue(RHICmdList, ComputeShaderRHI, BoxExtent, FVector(0,0,0));
+
+			SetShaderValue(RHICmdList, ComputeShaderRHI, ResetSimulation, false);
+			SetShaderValue(RHICmdList, ComputeShaderRHI, HasRootAttached, 0);
+			SetShaderValue(RHICmdList, ComputeShaderRHI, RestRootOffset, FVector4(0, 0, 0, 0));
+			SetShaderValue(RHICmdList, ComputeShaderRHI, DeformedRootOffset, FVector4(0, 0, 0, 0));
+
+			SetShaderValue(RHICmdList, ComputeShaderRHI, RestPositionOffset, FVector(0,0,0));
+			SetShaderValue(RHICmdList, ComputeShaderRHI, DeformedPositionOffset, FVector(0,0,0));
+
+			SetSRVParameter(RHICmdList, ComputeShaderRHI, RestTrianglePositionABuffer, FNiagaraRenderer::GetDummyFloatBuffer().SRV);
+			SetSRVParameter(RHICmdList, ComputeShaderRHI, RestTrianglePositionBBuffer, FNiagaraRenderer::GetDummyFloatBuffer().SRV);
+			SetSRVParameter(RHICmdList, ComputeShaderRHI, RestTrianglePositionCBuffer, FNiagaraRenderer::GetDummyFloatBuffer().SRV);
+
+			SetSRVParameter(RHICmdList, ComputeShaderRHI, DeformedTrianglePositionABuffer, FNiagaraRenderer::GetDummyFloatBuffer().SRV);
+			SetSRVParameter(RHICmdList, ComputeShaderRHI, DeformedTrianglePositionBBuffer, FNiagaraRenderer::GetDummyFloatBuffer().SRV);
+			SetSRVParameter(RHICmdList, ComputeShaderRHI, DeformedTrianglePositionCBuffer, FNiagaraRenderer::GetDummyFloatBuffer().SRV);
+
+			SetSRVParameter(RHICmdList, ComputeShaderRHI, RootBarycentricCoordinatesBuffer, FNiagaraRenderer::GetDummyFloatBuffer().SRV);
 		}
 	}
 
 	virtual void Unset(FRHICommandList& RHICmdList, const FNiagaraDataInterfaceSetArgs& Context) const override
 	{
-		SetUAVParameter(RHICmdList, Context.Shader->GetComputeShader(), PointsPositionsBuffer, nullptr);
+		SetUAVParameter(RHICmdList, Context.Shader->GetComputeShader(), DeformedPositionBuffer, nullptr);
 	}
 
 private:
@@ -278,11 +515,30 @@ private:
 	FShaderParameter WorldRotation;
 	FShaderParameter NumStrands;
 	FShaderParameter StrandSize;
-	FShaderParameter BoxOrigin;
+	FShaderParameter BoxCenter;
+	FShaderParameter BoxExtent;
 
-	FShaderResourceParameter PointsPositionsBuffer;
+	FShaderResourceParameter DeformedPositionBuffer;
 	FShaderResourceParameter CurvesOffsetsBuffer;
-	FShaderResourceParameter RestPositionsBuffer;
+	FShaderResourceParameter RestPositionBuffer;
+
+	FShaderParameter ResetSimulation;
+	FShaderParameter HasRootAttached;
+	FShaderParameter RestRootOffset;
+	FShaderParameter DeformedRootOffset;
+
+	FShaderResourceParameter RootBarycentricCoordinatesBuffer;
+
+	FShaderResourceParameter RestTrianglePositionABuffer;
+	FShaderResourceParameter RestTrianglePositionBBuffer;
+	FShaderResourceParameter RestTrianglePositionCBuffer;
+
+	FShaderResourceParameter DeformedTrianglePositionABuffer;
+	FShaderResourceParameter DeformedTrianglePositionBBuffer;
+	FShaderResourceParameter DeformedTrianglePositionCBuffer;
+
+	FShaderParameter RestPositionOffset;
+	FShaderParameter DeformedPositionOffset;
 };
 
 //------------------------------------------------------------------------------------------------------------
@@ -308,7 +564,9 @@ void FNDIHairStrandsProxy::ConsumePerInstanceDataFromGameThread(void* PerInstanc
 		TargetData->WorldTransform = SourceData->WorldTransform;
 		TargetData->NumStrands = SourceData->NumStrands;
 		TargetData->StrandSize = SourceData->StrandSize;
-		TargetData->BoxOrigin = SourceData->BoxOrigin;
+		TargetData->BoxCenter = SourceData->BoxCenter;
+		TargetData->BoxExtent = SourceData->BoxExtent;
+		TargetData->TickCount = SourceData->TickCount;
 		TargetData->HairStrandsBuffer = SourceData->HairStrandsBuffer;
 	}
 	else
@@ -317,7 +575,8 @@ void FNDIHairStrandsProxy::ConsumePerInstanceDataFromGameThread(void* PerInstanc
 	}
 }
 
-void FNDIHairStrandsProxy::InitializePerInstanceData(const FNiagaraSystemInstanceID& SystemInstance, FNDIHairStrandsBuffer* HairStrandsBuffer, const uint32 NumStrands, const uint8 StrandSize, const FVector4& BoxOrigin, const FMatrix& WorldTransform)
+void FNDIHairStrandsProxy::InitializePerInstanceData(const FNiagaraSystemInstanceID& SystemInstance, FNDIHairStrandsBuffer* HairStrandsBuffer, 
+	const uint32 NumStrands, const uint8 StrandSize, const FVector& BoxCenter, const FVector& BoxExtent, const FMatrix& WorldTransform)
 {
 	check(IsInRenderingThread());
 
@@ -333,7 +592,9 @@ void FNDIHairStrandsProxy::InitializePerInstanceData(const FNiagaraSystemInstanc
 	TargetData->HairStrandsBuffer = HairStrandsBuffer;
 	TargetData->NumStrands = NumStrands;
 	TargetData->StrandSize = StrandSize;
-	TargetData->BoxOrigin = BoxOrigin;
+	TargetData->BoxCenter = BoxCenter;
+	TargetData->BoxExtent = BoxExtent;
+	TargetData->TickCount = 0;
 	TargetData->WorldTransform = WorldTransform;
 }
 
@@ -353,6 +614,7 @@ UNiagaraDataInterfaceHairStrands::UNiagaraDataInterfaceHairStrands(FObjectInitia
 	, DefaultSource(nullptr)
 	, SourceActor(nullptr)
 	, SourceComponent(nullptr)
+	, GroupIndex(0)
 {
 
 	Proxy = MakeShared<FNDIHairStrandsProxy, ESPMode::ThreadSafe>();
@@ -407,26 +669,26 @@ void UNiagaraDataInterfaceHairStrands::UpdateSourceComponent(FNiagaraSystemInsta
 }
 
 void UNiagaraDataInterfaceHairStrands::ExtractDatasAndResources(FNiagaraSystemInstance* SystemInstance,FHairStrandsDatas*& OutStrandsDatas, 
-	FHairStrandsRestResource*& OutStrandsRestResource, FHairStrandsDeformedResource*& OutStrandsDeformedResource)
+	FHairStrandsRestResource*& OutStrandsRestResource, FHairStrandsDeformedResource*& OutStrandsDeformedResource, FHairStrandsRootResource*& OutStrandsRootResource)
 {
 	UpdateSourceComponent(SystemInstance);
 
-	const uint32 GroupIndex = 0;
-
 	OutStrandsDatas = nullptr;
 	OutStrandsRestResource = nullptr;
 	OutStrandsDeformedResource = nullptr;
+	OutStrandsRootResource = nullptr;
 
 	if (IsComponentValid())
 	{
 		OutStrandsDatas = SourceComponent->GetGuideStrandsDatas(GroupIndex);
 		OutStrandsRestResource = SourceComponent->GetGuideStrandsRestResource(GroupIndex);
 		OutStrandsDeformedResource = SourceComponent->GetGuideStrandsDeformedResource(GroupIndex);
+		OutStrandsRootResource = SourceComponent->GetGuideStrandsRootResource(GroupIndex);
 	}
-	else if (DefaultSource != nullptr && DefaultSource->GetNumHairGroups() > 0)
+	else if (DefaultSource != nullptr && (GroupIndex < DefaultSource->GetNumHairGroups()))
 	{
-		OutStrandsDatas = &DefaultSource->HairGroupsData[0].HairRenderData;
-		OutStrandsRestResource = DefaultSource->HairGroupsData[0].HairStrandsRestResource;
+		OutStrandsDatas = &DefaultSource->HairGroupsData[GroupIndex].HairRenderData;
+		OutStrandsRestResource = DefaultSource->HairGroupsData[GroupIndex].HairStrandsRestResource;
 	}
 }
 
@@ -438,8 +700,9 @@ bool UNiagaraDataInterfaceHairStrands::InitPerInstanceData(void* PerInstanceData
 	FHairStrandsDatas* StrandsDatas =  nullptr;
 	FHairStrandsRestResource* StrandsRestResource = nullptr;
 	FHairStrandsDeformedResource* StrandsDeformedResource = nullptr;
+	FHairStrandsRootResource* StrandsRootResource = nullptr;
 
-	ExtractDatasAndResources(SystemInstance, StrandsDatas, StrandsRestResource, StrandsDeformedResource);
+	ExtractDatasAndResources(SystemInstance, StrandsDatas, StrandsRestResource, StrandsDeformedResource, StrandsRootResource);
 
 	if (StrandsDatas == nullptr || StrandsRestResource == nullptr)
 	{
@@ -449,7 +712,7 @@ bool UNiagaraDataInterfaceHairStrands::InitPerInstanceData(void* PerInstanceData
 	else
 	{ 
 		FNDIHairStrandsBuffer* HairStrandsBuffer = new FNDIHairStrandsBuffer;
-		HairStrandsBuffer->SetHairAsset(StrandsDatas, StrandsRestResource, StrandsDeformedResource);
+		HairStrandsBuffer->SetHairAsset(StrandsDatas, StrandsRestResource, StrandsDeformedResource, StrandsRootResource);
 
 		// Push instance data to RT
 		{
@@ -458,27 +721,28 @@ bool UNiagaraDataInterfaceHairStrands::InitPerInstanceData(void* PerInstanceData
 			const FBox& StrandsBox = StrandsDatas->BoundingBox;
 
 			const FVector BoxExtent = StrandsBox.GetExtent();
-			const FVector4 BoxOrigin(-BoxExtent.X, -BoxExtent.Y, -BoxExtent.Z, 0.0);
+			const FVector BoxCenter = StrandsBox.GetCenter(); 
 
-			const FTransform WorldOffset(-StrandsBox.GetCenter());
-			const FMatrix WorldTransform = WorldOffset.ToInverseMatrixWithScale();
+			const FMatrix WorldTransform = IsComponentValid() ? SourceComponent->GetComponentToWorld().ToMatrixWithScale() :
+				SystemInstance->GetComponent()->GetComponentToWorld().ToMatrixWithScale();
 
-			InstanceData->WorldTransform = IsComponentValid() ? WorldTransform * SourceComponent->GetComponentToWorld().ToMatrixWithScale() : 
-				WorldTransform * SystemInstance->GetComponent()->GetComponentToWorld().ToMatrixWithScale();
+			InstanceData->WorldTransform = WorldTransform;
 			InstanceData->StrandSize = StrandsSize;
 			InstanceData->NumStrands = NumStrands;
 			InstanceData->HairStrandsBuffer = HairStrandsBuffer;
-			InstanceData->BoxOrigin = BoxOrigin;
+			InstanceData->BoxCenter = BoxCenter;
+			InstanceData->BoxExtent = BoxExtent;
+			InstanceData->TickCount = 0;
 
-			//UE_LOG(LogHairStrands, Log, TEXT("Num Strands = %d | Strand Size = %d | Num Vertices = %d | Min = %f %f %f | Max = %f %f %f | Transform = %s"), LocalNumStrands, LocalStrandSize,
-			//	StrandsDatas->GetNumPoints(), LocalStrandsBox.Min[0], LocalStrandsBox.Min[1], LocalStrandsBox.Min[2], LocalStrandsBox.Max[0], LocalStrandsBox.Max[1], LocalStrandsBox.Max[2], *InstanceData->WorldTransform.ToString());
+			//UE_LOG(LogHairStrands, Log, TEXT("Num Strands = %d | Strand Size = %d | Num Vertices = %d | Min = %f %f %f | Max = %f %f %f | Transform = %s"), NumStrands, StrandSize,
+			//	StrandsDatas->GetNumPoints(), StrandsBox.Min[0], StrandsBox.Min[1], StrandsBox.Min[2], StrandsBox.Max[0], StrandsBox.Max[1], StrandsBox.Max[2], *InstanceData->WorldTransform.ToString());
 
 			FNDIHairStrandsProxy* ThisProxy = GetProxyAs<FNDIHairStrandsProxy>();
 			ENQUEUE_RENDER_COMMAND(FNiagaraDIPushInitialInstanceDataToRT) (
-				[ThisProxy, InstanceID = SystemInstance->GetId(), HairStrandsBuffer, NumStrands, StrandsSize, BoxOrigin, WorldTransform](FRHICommandListImmediate& CmdList)
+				[ThisProxy, InstanceID = SystemInstance->GetId(), HairStrandsBuffer, NumStrands, StrandsSize, BoxCenter, BoxExtent, WorldTransform](FRHICommandListImmediate& CmdList)
 			{
 				HairStrandsBuffer->InitResource();
-				ThisProxy->InitializePerInstanceData(InstanceID, HairStrandsBuffer, NumStrands, StrandsSize, BoxOrigin, WorldTransform);
+				ThisProxy->InitializePerInstanceData(InstanceID, HairStrandsBuffer, NumStrands, StrandsSize, BoxCenter, BoxExtent, WorldTransform);
 			}
 			);
 		}
@@ -514,23 +778,28 @@ bool UNiagaraDataInterfaceHairStrands::PerInstanceTick(void* PerInstanceData, FN
 	FHairStrandsDatas* StrandsDatas = nullptr;
 	FHairStrandsRestResource* StrandsRestResource = nullptr;
 	FHairStrandsDeformedResource* StrandsDeformedResource = nullptr;
+	FHairStrandsRootResource* StrandsRootResource = nullptr;
+	InstanceData->TickCount = FMath::Min(ResetTick+1,InstanceData->TickCount+1);
 
-	ExtractDatasAndResources(SystemInstance, StrandsDatas, StrandsRestResource, StrandsDeformedResource);
+	ExtractDatasAndResources(SystemInstance, StrandsDatas, StrandsRestResource, StrandsDeformedResource, StrandsRootResource);
+	InstanceData->HairStrandsBuffer->SetHairAsset(StrandsDatas, StrandsRestResource, StrandsDeformedResource, StrandsRootResource);
 
 	bool RequireReset = false;
 	if (StrandsDatas)
 	{
-		const FTransform WorldOffset(-StrandsDatas->BoundingBox.GetCenter());
-		const FMatrix WorldTransform = WorldOffset.ToInverseMatrixWithScale();
+		if (StrandsDeformedResource != nullptr)
+		{
+			InstanceData->BoxCenter = StrandsDeformedResource->PositionOffset;
+		}
 
 		if (IsComponentValid())
 		{
-			InstanceData->WorldTransform = WorldTransform * SourceComponent->GetComponentToWorld().ToMatrixWithScale();
+			InstanceData->WorldTransform = SourceComponent->GetComponentToWorld().ToMatrixWithScale();
 			//UE_LOG(LogHairStrands, Warning, TEXT("Get Component Transform : %s"), *InstanceData->WorldTransform.ToString());
 		}
 		else
 		{
-			InstanceData->WorldTransform = WorldTransform * SystemInstance->GetComponent()->GetComponentToWorld().ToMatrixWithScale();
+			InstanceData->WorldTransform = SystemInstance->GetComponent()->GetComponentToWorld().ToMatrixWithScale();
 			//UE_LOG(LogHairStrands, Warning, TEXT("Get Instance Transform : %s"), *InstanceData->WorldTransform.ToString());
 		}
 	}
@@ -549,6 +818,7 @@ bool UNiagaraDataInterfaceHairStrands::CopyToInternal(UNiagaraDataInterface* Des
 	OtherTyped->SourceActor= SourceActor;
 	OtherTyped->SourceComponent = SourceComponent;
 	OtherTyped->DefaultSource = DefaultSource;
+	OtherTyped->GroupIndex = GroupIndex;
 
 	return true;
 }
@@ -562,7 +832,8 @@ bool UNiagaraDataInterfaceHairStrands::Equals(const UNiagaraDataInterface* Other
 	const UNiagaraDataInterfaceHairStrands* OtherTyped = CastChecked<const UNiagaraDataInterfaceHairStrands>(Other);
 
 	return  (OtherTyped->StrandSize == StrandSize)
-		&& (OtherTyped->SourceActor == SourceActor) && (OtherTyped->SourceComponent == SourceComponent) && (OtherTyped->DefaultSource == DefaultSource);
+		&& (OtherTyped->SourceActor == SourceActor) && (OtherTyped->SourceComponent == SourceComponent)
+		&& (OtherTyped->DefaultSource == DefaultSource) && (OtherTyped->GroupIndex == GroupIndex);
 }
 
 void UNiagaraDataInterfaceHairStrands::PostInitProperties()
@@ -725,6 +996,20 @@ void UNiagaraDataInterfaceHairStrands::GetFunctions(TArray<FNiagaraFunctionSigna
 
 		OutFunctions.Add(Sig);
 	}
+	{
+		FNiagaraFunctionSignature Sig;
+		Sig.Name = ComputeLocalStateName;
+		Sig.bMemberFunction = true;
+		Sig.bRequiresContext = false;
+		Sig.Inputs.Add(FNiagaraVariable(FNiagaraTypeDefinition(GetClass()), TEXT("Hair Strands")));
+		Sig.Inputs.Add(FNiagaraVariable(FNiagaraTypeDefinition::GetIntDef(), TEXT("Node Index")));
+		Sig.Inputs.Add(FNiagaraVariable(FNiagaraTypeDefinition::GetVec3Def(), TEXT("Rest Position")));
+		Sig.Inputs.Add(FNiagaraVariable(FNiagaraTypeDefinition::GetQuatDef(), TEXT("Rest Orientation")));
+		Sig.Outputs.Add(FNiagaraVariable(FNiagaraTypeDefinition::GetVec3Def(), TEXT("Local Position")));
+		Sig.Outputs.Add(FNiagaraVariable(FNiagaraTypeDefinition::GetQuatDef(), TEXT("Local Orientation")));
+
+		OutFunctions.Add(Sig);
+	}
 	{
 		FNiagaraFunctionSignature Sig;
 		Sig.Name = AdvectNodePositionName;
@@ -807,6 +1092,20 @@ void UNiagaraDataInterfaceHairStrands::GetFunctions(TArray<FNiagaraFunctionSigna
 
 		OutFunctions.Add(Sig);
 	}
+	{
+		FNiagaraFunctionSignature Sig;
+		Sig.Name = AttachNodeStateName;
+		Sig.bMemberFunction = true;
+		Sig.bRequiresContext = false;
+		Sig.Inputs.Add(FNiagaraVariable(FNiagaraTypeDefinition(GetClass()), TEXT("Hair Strands")));
+		Sig.Inputs.Add(FNiagaraVariable(FNiagaraTypeDefinition::GetIntDef(), TEXT("Node Index")));
+		Sig.Inputs.Add(FNiagaraVariable(FNiagaraTypeDefinition::GetVec3Def(), TEXT("Local Position")));
+		Sig.Inputs.Add(FNiagaraVariable(FNiagaraTypeDefinition::GetQuatDef(), TEXT("Local Orientation")));
+		Sig.Outputs.Add(FNiagaraVariable(FNiagaraTypeDefinition::GetVec3Def(), TEXT("Node Position")));
+		Sig.Outputs.Add(FNiagaraVariable(FNiagaraTypeDefinition::GetQuatDef(), TEXT("Node Orientation")));
+
+		OutFunctions.Add(Sig);
+	}
 	{
 		FNiagaraFunctionSignature Sig;
 		Sig.Name = UpdatePointPositionName;
@@ -1113,6 +1412,16 @@ void UNiagaraDataInterfaceHairStrands::GetFunctions(TArray<FNiagaraFunctionSigna
 		Sig.Inputs.Add(FNiagaraVariable(FNiagaraTypeDefinition::GetVec3Def(), TEXT("Node Velocity")));
 		Sig.Outputs.Add(FNiagaraVariable(FNiagaraTypeDefinition::GetVec3Def(), TEXT("Drag Force")));
 
+		OutFunctions.Add(Sig);
+	}
+	{
+		FNiagaraFunctionSignature Sig;
+		Sig.Name = NeedSimulationResetName;
+		Sig.bMemberFunction = true;
+		Sig.bRequiresContext = false;
+		Sig.Inputs.Add(FNiagaraVariable(FNiagaraTypeDefinition(GetClass()), TEXT("Hair Strands")));
+		Sig.Outputs.Add(FNiagaraVariable(FNiagaraTypeDefinition::GetBoolDef(), TEXT("Reset Simulation")));
+
 		OutFunctions.Add(Sig);
 	}
 }
@@ -1130,8 +1439,10 @@ DEFINE_NDI_DIRECT_FUNC_BINDER(UNiagaraDataInterfaceHairStrands, ComputeEdgeLengt
 DEFINE_NDI_DIRECT_FUNC_BINDER(UNiagaraDataInterfaceHairStrands, ComputeEdgeRotation);
 DEFINE_NDI_DIRECT_FUNC_BINDER(UNiagaraDataInterfaceHairStrands, ComputeRestPosition);
 DEFINE_NDI_DIRECT_FUNC_BINDER(UNiagaraDataInterfaceHairStrands, ComputeRestOrientation);
+DEFINE_NDI_DIRECT_FUNC_BINDER(UNiagaraDataInterfaceHairStrands, ComputeLocalState);
 DEFINE_NDI_DIRECT_FUNC_BINDER(UNiagaraDataInterfaceHairStrands, AttachNodePosition);
 DEFINE_NDI_DIRECT_FUNC_BINDER(UNiagaraDataInterfaceHairStrands, AttachNodeOrientation);
+DEFINE_NDI_DIRECT_FUNC_BINDER(UNiagaraDataInterfaceHairStrands, AttachNodeState);
 DEFINE_NDI_DIRECT_FUNC_BINDER(UNiagaraDataInterfaceHairStrands, UpdatePointPosition);
 DEFINE_NDI_DIRECT_FUNC_BINDER(UNiagaraDataInterfaceHairStrands, ResetPointPosition);
 
@@ -1165,6 +1476,7 @@ DEFINE_NDI_DIRECT_FUNC_BINDER(UNiagaraDataInterfaceHairStrands, ProjectStaticCol
 DEFINE_NDI_DIRECT_FUNC_BINDER(UNiagaraDataInterfaceHairStrands, ComputeRestDirection);
 DEFINE_NDI_DIRECT_FUNC_BINDER(UNiagaraDataInterfaceHairStrands, UpdateNodeOrientation);
 DEFINE_NDI_DIRECT_FUNC_BINDER(UNiagaraDataInterfaceHairStrands, ComputeAirDragForce);
+DEFINE_NDI_DIRECT_FUNC_BINDER(UNiagaraDataInterfaceHairStrands, NeedSimulationReset);
 
 void UNiagaraDataInterfaceHairStrands::GetVMExternalFunction(const FVMExternalFunctionBindingInfo& BindingInfo, void* InstanceData, FVMExternalFunction &OutFunc)
 {
@@ -1233,6 +1545,11 @@ void UNiagaraDataInterfaceHairStrands::GetVMExternalFunction(const FVMExternalFu
 		check(BindingInfo.GetNumInputs() == 5 && BindingInfo.GetNumOutputs() == 4);
 		NDI_FUNC_BINDER(UNiagaraDataInterfaceHairStrands, ComputeRestOrientation)::Bind(this, OutFunc);
 	}
+	else if (BindingInfo.Name == ComputeLocalStateName)
+	{
+		check(BindingInfo.GetNumInputs() == 9 && BindingInfo.GetNumOutputs() == 7);
+		NDI_FUNC_BINDER(UNiagaraDataInterfaceHairStrands, ComputeLocalState)::Bind(this, OutFunc);
+	}
 	else if (BindingInfo.Name == AttachNodePositionName)
 	{
 		check(BindingInfo.GetNumInputs() == 4 && BindingInfo.GetNumOutputs() == 3);
@@ -1243,6 +1560,11 @@ void UNiagaraDataInterfaceHairStrands::GetVMExternalFunction(const FVMExternalFu
 		check(BindingInfo.GetNumInputs() == 5 && BindingInfo.GetNumOutputs() == 4);
 		NDI_FUNC_BINDER(UNiagaraDataInterfaceHairStrands, AttachNodeOrientation)::Bind(this, OutFunc);
 	}
+	else if (BindingInfo.Name == AttachNodeStateName)
+	{
+		check(BindingInfo.GetNumInputs() == 9 && BindingInfo.GetNumOutputs() == 7);
+		NDI_FUNC_BINDER(UNiagaraDataInterfaceHairStrands, AttachNodeState)::Bind(this, OutFunc);
+	}
 	else if (BindingInfo.Name == UpdatePointPositionName)
 	{
 		check(BindingInfo.GetNumInputs() == 5 && BindingInfo.GetNumOutputs() == 1);
@@ -1368,6 +1690,11 @@ void UNiagaraDataInterfaceHairStrands::GetVMExternalFunction(const FVMExternalFu
 		check(BindingInfo.GetNumInputs() == 14 && BindingInfo.GetNumOutputs() == 3);
 		NDI_FUNC_BINDER(UNiagaraDataInterfaceHairStrands, ComputeAirDragForce)::Bind(this, OutFunc);
 	}
+	else if (BindingInfo.Name == NeedSimulationResetName)
+	{
+		check(BindingInfo.GetNumInputs() == 1 && BindingInfo.GetNumOutputs() == 1);
+		NDI_FUNC_BINDER(UNiagaraDataInterfaceHairStrands, NeedSimulationReset)::Bind(this, OutFunc);
+	}
 }
 
 void WriteTransform(const FMatrix& ToWrite, FVectorVMContext& Context)
@@ -1457,9 +1784,9 @@ void UNiagaraDataInterfaceHairStrands::GetBoxCenter(FVectorVMContext& Context)
 
 	for (int32 InstanceIdx = 0; InstanceIdx < Context.NumInstances; ++InstanceIdx)
 	{
-		*OutBoxCenterX.GetDestAndAdvance() = 0.0;
-		*OutBoxCenterY.GetDestAndAdvance() = 0.0;
-		*OutBoxCenterZ.GetDestAndAdvance() = 0.0;
+		*OutBoxCenterX.GetDestAndAdvance() = InstData->BoxCenter.X;
+		*OutBoxCenterY.GetDestAndAdvance() = InstData->BoxCenter.Y;
+		*OutBoxCenterZ.GetDestAndAdvance() = InstData->BoxCenter.Z;
 	}
 }
 
@@ -1472,9 +1799,9 @@ void UNiagaraDataInterfaceHairStrands::GetBoxExtent(FVectorVMContext& Context)
 
 	for (int32 InstanceIdx = 0; InstanceIdx < Context.NumInstances; ++InstanceIdx)
 	{
-		*OutBoxExtentX.GetDestAndAdvance() = -InstData->BoxOrigin.X;
-		*OutBoxExtentY.GetDestAndAdvance() = -InstData->BoxOrigin.Y;
-		*OutBoxExtentZ.GetDestAndAdvance() = -InstData->BoxOrigin.Z;
+		*OutBoxExtentX.GetDestAndAdvance() = InstData->BoxExtent.X;
+		*OutBoxExtentY.GetDestAndAdvance() = InstData->BoxExtent.Y;
+		*OutBoxExtentZ.GetDestAndAdvance() = InstData->BoxExtent.Z;
 	}
 }
 
@@ -1523,6 +1850,11 @@ void UNiagaraDataInterfaceHairStrands::ComputeRestOrientation(FVectorVMContext&
 	// @todo : implement function for cpu 
 }
 
+void UNiagaraDataInterfaceHairStrands::ComputeLocalState(FVectorVMContext& Context)
+{
+	// @todo : implement function for cpu 
+}
+
 void UNiagaraDataInterfaceHairStrands::UpdatePointPosition(FVectorVMContext& Context)
 {
 	// @todo : implement function for cpu 
@@ -1543,6 +1875,11 @@ void UNiagaraDataInterfaceHairStrands::AttachNodeOrientation(FVectorVMContext& C
 	// @todo : implement function for cpu 
 }
 
+void UNiagaraDataInterfaceHairStrands::AttachNodeState(FVectorVMContext& Context)
+{
+	// @todo : implement function for cpu 
+}
+
 void UNiagaraDataInterfaceHairStrands::AdvectNodePosition(FVectorVMContext& Context)
 {
 	// @todo : implement function for cpu 
@@ -1649,6 +1986,11 @@ void UNiagaraDataInterfaceHairStrands::ComputeAirDragForce(FVectorVMContext& Con
 	// @todo : implement function for cpu 
 }
 
+void UNiagaraDataInterfaceHairStrands::NeedSimulationReset(FVectorVMContext& Context)
+{
+	// @todo : implement function for cpu 
+}
+
 bool UNiagaraDataInterfaceHairStrands::GetFunctionHLSL(const FName& DefinitionFunctionName, FString InstanceFunctionName, FNiagaraDataInterfaceGPUParamInfo& ParamInfo, FString& OutHLSL)
 {
 	FNDIHairStrandsParametersName ParamNames(ParamInfo.DataInterfaceHLSLSymbol);
@@ -1660,10 +2002,11 @@ bool UNiagaraDataInterfaceHairStrands::GetFunctionHLSL(const FName& DefinitionFu
 		{TEXT("WorldTransformName"), ParamNames.WorldTransformName},
 		{TEXT("WorldInverseName"), ParamNames.WorldInverseName},
 		{TEXT("WorldRotationName"), ParamNames.WorldRotationName},
-		{TEXT("PointsPositionsBufferName"), ParamNames.PointsPositionsBufferName},
+		{TEXT("DeformedPositionBufferName"), ParamNames.DeformedPositionBufferName},
 		{TEXT("CurvesOffsetsBufferName"), ParamNames.CurvesOffsetsBufferName},
-		{TEXT("RestPositionsBufferName"), ParamNames.RestPositionsBufferName},
-		{TEXT("BoxOriginName"), ParamNames.BoxOriginName},
+		{TEXT("RestPositionBufferName"), ParamNames.RestPositionBufferName},
+		{TEXT("BoxCenterName"), ParamNames.BoxCenterName},
+		{TEXT("BoxExtentName"), ParamNames.BoxExtentName},
 		{TEXT("HairStrandsContextName"), TEXT("DIHAIRSTRANDS_MAKE_CONTEXT(") + ParamInfo.DataInterfaceHLSLSymbol + TEXT(")")},
 	};
 
@@ -1810,6 +2153,17 @@ bool UNiagaraDataInterfaceHairStrands::GetFunctionHLSL(const FName& DefinitionFu
 		OutHLSL += FString::Format(FormatSample, ArgsSample);
 		return true;
 	}
+	else if (DefinitionFunctionName == ComputeLocalStateName)
+	{
+		static const TCHAR *FormatSample = TEXT(R"(
+			void {InstanceFunctionName} (in int NodeIndex, in float3 RestPosition, in float4 RestOrientation, out float3 LocalPosition, out float4 LocalOrientation)
+			{
+				{HairStrandsContextName} DIHairStrands_ComputeLocalState(DIContext,NodeIndex,RestPosition,RestOrientation,LocalPosition,LocalOrientation);
+			}
+			)");
+		OutHLSL += FString::Format(FormatSample, ArgsSample);
+		return true;
+	}
 	else if (DefinitionFunctionName == AttachNodePositionName)
 	{
 		static const TCHAR *FormatSample = TEXT(R"(
@@ -1823,15 +2177,26 @@ bool UNiagaraDataInterfaceHairStrands::GetFunctionHLSL(const FName& DefinitionFu
 	}
 	else if (DefinitionFunctionName == AttachNodeOrientationName)
 	{
-	static const TCHAR *FormatSample = TEXT(R"(
-			void {InstanceFunctionName} (in float4 RestOrientation, out float4 NodeOrientation)
-			{
-				{HairStrandsContextName} DIHairStrands_AttachNodeOrientation(DIContext,RestOrientation,NodeOrientation);
-			}
-			)");
-	OutHLSL += FString::Format(FormatSample, ArgsSample);
-	return true;
+		static const TCHAR *FormatSample = TEXT(R"(
+				void {InstanceFunctionName} (in float4 RestOrientation, out float4 NodeOrientation)
+				{
+					{HairStrandsContextName} DIHairStrands_AttachNodeOrientation(DIContext,RestOrientation,NodeOrientation);
+				}
+				)");
+		OutHLSL += FString::Format(FormatSample, ArgsSample);
+		return true;
 	}	
+	else if (DefinitionFunctionName == AttachNodeStateName)
+	{
+		static const TCHAR *FormatSample = TEXT(R"(
+				void {InstanceFunctionName} (in int NodeIndex, in float3 LocalPosition, in float4 LocalOrientation, out float3 NodePosition, out float4 NodeOrientation)
+				{
+					{HairStrandsContextName} DIHairStrands_AttachNodeState(DIContext,NodeIndex,LocalPosition,LocalOrientation,NodePosition,NodeOrientation);
+				}
+				)");
+		OutHLSL += FString::Format(FormatSample, ArgsSample);
+		return true;
+	}
 	else if (DefinitionFunctionName == UpdatePointPositionName)
 	{
 		static const TCHAR *FormatSample = TEXT(R"(
@@ -2165,6 +2530,17 @@ in float RestLength, in float DeltaTime, in int NodeOffset, in float MaterialDam
 		OutHLSL += FString::Format(FormatSample, ArgsSample);
 		return true;
 	}
+	else if (DefinitionFunctionName == NeedSimulationResetName)
+	{
+		static const TCHAR *FormatSample = TEXT(R"(
+				void {InstanceFunctionName} ( out bool ResetSimulation)
+				{
+					{HairStrandsContextName} ResetSimulation  = DIContext.ResetSimulation;
+				}
+				)");
+		OutHLSL += FString::Format(FormatSample, ArgsSample);
+		return true;
+	}
 
 	OutHLSL += TEXT("\n");
 	return false;
@@ -2195,6 +2571,9 @@ void UNiagaraDataInterfaceHairStrands::ProvidePerInstanceDataForRenderThread(voi
 	RenderThreadData->HairStrandsBuffer = GameThreadData->HairStrandsBuffer;
 	RenderThreadData->NumStrands = GameThreadData->NumStrands;
 	RenderThreadData->StrandSize = GameThreadData->StrandSize;
+	RenderThreadData->BoxCenter = GameThreadData->BoxCenter;
+	RenderThreadData->BoxExtent = GameThreadData->BoxExtent;
+	RenderThreadData->TickCount = GameThreadData->TickCount;
 }
 
 FNiagaraDataInterfaceParametersCS*
diff --git a/Engine/Plugins/Experimental/HairStrands/Source/HairStrandsNiagara/Private/NiagaraDataInterfacePhysicsAsset.cpp b/Engine/Plugins/Experimental/HairStrands/Source/HairStrandsNiagara/Private/NiagaraDataInterfacePhysicsAsset.cpp
--- a/Engine/Plugins/Experimental/HairStrands/Source/HairStrandsNiagara/Private/NiagaraDataInterfacePhysicsAsset.cpp
+++ b/Engine/Plugins/Experimental/HairStrands/Source/HairStrandsNiagara/Private/NiagaraDataInterfacePhysicsAsset.cpp
@@ -492,20 +492,14 @@ bool UNiagaraDataInterfacePhysicsAsset::InitPerInstanceData(void* PerInstanceDat
 		}
 	}
 
-	if (SourceComponent != nullptr)
-	{
-		UE_LOG(LogPhysicsAsset, Log, TEXT("Found Skeletal Mesh %d %d %d"), SourceComponent.Get(), SourceComponent->GetPhysicsAsset(), DefaultSource);
-
-	}
-
 	PhysicsAsset = (SourceComponent != nullptr) ? SourceComponent->GetPhysicsAsset() : ( DefaultSource != nullptr ) ? DefaultSource : nullptr;
 
-	if (PhysicsAsset == nullptr )
-	{
-		UE_LOG(LogPhysicsAsset, Log, TEXT("Physics Asset data interface has no valid asset. Failed InitPerInstanceData - %s %d %d"), *GetFullName(), PhysicsAsset.Get(), DefaultSource);
-		return false;
-	}
-	else
+	//if (PhysicsAsset == nullptr )
+	//{
+	//	UE_LOG(LogPhysicsAsset, Log, TEXT("Physics Asset data interface has no valid asset. Failed InitPerInstanceData - %s %d %d"), *GetFullName(), PhysicsAsset.Get(), DefaultSource);
+	//	return false;
+	//}
+	//else
 	{
 		const FTransform WorldTransform = (SourceComponent != nullptr) ? SourceComponent->GetComponentTransform() : SystemInstance->GetComponent()->GetComponentTransform();
 		FNDIPhysicsAssetBuffer* AssetBuffer = new FNDIPhysicsAssetBuffer;
diff --git a/Engine/Plugins/Experimental/HairStrands/Source/HairStrandsNiagara/Private/NiagaraDataInterfacePressureGrid.cpp b/Engine/Plugins/Experimental/HairStrands/Source/HairStrandsNiagara/Private/NiagaraDataInterfacePressureGrid.cpp
--- a/Engine/Plugins/Experimental/HairStrands/Source/HairStrandsNiagara/Private/NiagaraDataInterfacePressureGrid.cpp
+++ b/Engine/Plugins/Experimental/HairStrands/Source/HairStrandsNiagara/Private/NiagaraDataInterfacePressureGrid.cpp
@@ -634,8 +634,9 @@ void UNiagaraDataInterfacePressureGrid::BuildGridTopology(FVectorVMContext& Cont
 		*OutGridOriginY.GetDestAndAdvance() = BoxOrigin.Y;
 		*OutGridOriginZ.GetDestAndAdvance() = BoxOrigin.Z;
 		*OutGridLength.GetDestAndAdvance() = MaxLength;
-		//UE_LOG(LogPressureGrid, Warning, TEXT("Get Grid Extent : %f %f %f | Grid Lengths : %f %f %f | Regular Extent : %f %f %f | Box Origin : %f %f %f"), 
-		//				GridExtent.X, GridExtent.Y, GridExtent.Z, GridLengths.X, GridLengths.Y, GridLengths.Z, RegularExtent.X, RegularExtent.Y, RegularExtent.Z, InstData->GridOrigin.X, InstData->GridOrigin.Y, InstData->GridOrigin.Z);
+		//UE_LOG(LogPressureGrid, Warning, TEXT("Get Grid Extent : %f %f %f | Grid Lengths : %f %f %f | Regular Extent : %f %f %f | Box Origin : %f %f %f | Grid Center : %f %f %f"), 
+		//				GridExtent.X, GridExtent.Y, GridExtent.Z, GridLengths.X, GridLengths.Y, GridLengths.Z, RegularExtent.X, RegularExtent.Y, RegularExtent.Z,
+		//	InstData->GridOrigin.X, InstData->GridOrigin.Y, InstData->GridOrigin.Z, GridCenter.X, GridCenter.Y, GridCenter.Z);
 	}
 
 	
diff --git a/Engine/Plugins/Experimental/HairStrands/Source/HairStrandsNiagara/Public/NiagaraDataInterfaceHairStrands.h b/Engine/Plugins/Experimental/HairStrands/Source/HairStrandsNiagara/Public/NiagaraDataInterfaceHairStrands.h
--- a/Engine/Plugins/Experimental/HairStrands/Source/HairStrandsNiagara/Public/NiagaraDataInterfaceHairStrands.h
+++ b/Engine/Plugins/Experimental/HairStrands/Source/HairStrandsNiagara/Public/NiagaraDataInterfaceHairStrands.h
@@ -26,7 +26,8 @@ enum class EHairStrandsSize : uint8
 struct FNDIHairStrandsBuffer : public FRenderResource
 {
 	/** Set the asset that will be used to affect the buffer */
-	void SetHairAsset(const FHairStrandsDatas*  HairStrandsDatas, const FHairStrandsRestResource*  HairStrandsRestResource, const FHairStrandsDeformedResource*  HairStrandsDeformedResource);
+	void SetHairAsset(const FHairStrandsDatas*  HairStrandsDatas, const FHairStrandsRestResource*  HairStrandsRestResource, 
+		const FHairStrandsDeformedResource*  HairStrandsDeformedResource, const FHairStrandsRootResource* HairStrandsRootResource );
 
 	/** Init the buffer */
 	virtual void InitRHI() override;
@@ -41,7 +42,28 @@ struct FNDIHairStrandsBuffer : public FRenderResource
 	FRWBuffer CurvesOffsetsBuffer;
 
 	/** Deformed position buffer in case no ressource are there */
-	FRWBuffer PointsPositionsBuffer;
+	FRWBuffer DeformedPositionBuffer;
+
+	/**Rest triangle position of vertex A*/
+	FRWBuffer RestTrianglePositionABuffer;
+
+	/**Rest triangle position of vertex B*/
+	FRWBuffer RestTrianglePositionBBuffer;
+
+	/**Rest triangle position of vertex C*/
+	FRWBuffer RestTrianglePositionCBuffer;
+
+	/**Deformed triangle position of vertex A*/
+	FRWBuffer DeformedTrianglePositionABuffer;
+
+	/**Deformed triangle position of vertex B*/
+	FRWBuffer DeformedTrianglePositionBBuffer;
+
+	/**Deformed triangle position of vertex C*/
+	FRWBuffer DeformedTrianglePositionCBuffer;
+
+	/**Root barycentric coordinates */
+	FRWBuffer RootBarycentricCoordinatesBuffer;
 
 	/** The strand asset datas from which to sample */
 	const FHairStrandsDatas* SourceDatas;
@@ -51,6 +73,9 @@ struct FNDIHairStrandsBuffer : public FRenderResource
 
 	/** The strand deformed resource to write into */
 	const FHairStrandsDeformedResource* SourceDeformedResources;
+
+	/** The strand root resource to write into */
+	const FHairStrandsRootResource* SourceRootResources;
 };
 
 /** Data stored per strand base instance*/
@@ -65,8 +90,14 @@ struct FNDIHairStrandsData
 	/** Strand size */
 	int32 StrandSize;
 
-	/** Bounding box origin */
-	FVector4 BoxOrigin;
+	/** Bounding box center */
+	FVector BoxCenter;
+
+	/** Bounding box extent */
+	FVector BoxExtent;
+
+	/** Tick Count*/
+	int32 TickCount;
 
 	/** Strands Gpu buffer */
 	FNDIHairStrandsBuffer* HairStrandsBuffer;
@@ -95,6 +126,10 @@ public:
 	/** The source component from which to sample */
 	TWeakObjectPtr<class UGroomComponent> SourceComponent;
 
+	/** Group Index to be used */
+	UPROPERTY(EditAnywhere, Category = "Source")
+	int32 GroupIndex;
+
 	/** UObject Interface */
 	virtual void PostInitProperties() override;
 
@@ -123,7 +158,7 @@ public:
 
 	/** Extract datas and resources */
 	void ExtractDatasAndResources(FNiagaraSystemInstance* SystemInstance, FHairStrandsDatas*& OutStrandsDatas,
-		FHairStrandsRestResource*& OutStrandsRestResource, FHairStrandsDeformedResource*& OutStrandsDeformedResource);
+		FHairStrandsRestResource*& OutStrandsRestResource, FHairStrandsDeformedResource*& OutStrandsDeformedResource, FHairStrandsRootResource*& OutStrandsRootResource);
 
 	/** Get the number of strands */
 	void GetNumStrands(FVectorVMContext& Context);
@@ -245,6 +280,15 @@ public:
 	/** Compute the air drag force */
 	void ComputeAirDragForce(FVectorVMContext& Context);
 
+	/** Get the rest position and orientation relative to the transform or to the skin cache */
+	void ComputeLocalState(FVectorVMContext& Context);
+
+	/** Attach the node position and orientation to the transform or to the skin cache */
+	void AttachNodeState(FVectorVMContext& Context);
+
+	/** Check if we need or not a simulation reset*/
+	void NeedSimulationReset(FVectorVMContext& Context);
+
 	/** Name of the world transform */
 	static const FString WorldTransformName;
 
@@ -261,16 +305,58 @@ public:
 	static const FString StrandSizeName;
 
 	/** Name of the points positions buffer */
-	static const FString PointsPositionsBufferName;
+	static const FString DeformedPositionBufferName;
 
 	/** Name of the curves offsets buffer */
 	static const FString CurvesOffsetsBufferName;
 
 	/** Name of the nodes positions buffer */
-	static const FString RestPositionsBufferName;
+	static const FString RestPositionBufferName;
+
+	/** Name of the box center  */
+	static const FString BoxCenterName;
+
+	/** Name of the box extent  */
+	static const FString BoxExtentName;
+
+	/** Param to check if the roots have been attached to the skin */
+	static const FString HasRootAttachedName;
+
+	/** boolean to check if we need to rest the simulation*/
+	static const FString ResetSimulationName;
+
+	/** Rest center of all the roots */
+	static const FString RestRootOffsetName;
+
+	/** Rest position of the triangle vertex A */
+	static const FString RestTrianglePositionAName;
+
+	/** Rest position of the triangle vertex B */
+	static const FString RestTrianglePositionBName;
+
+	/** Rest position of the triangle vertex C */
+	static const FString RestTrianglePositionCName;
+
+	/** Deformed center of all the roots */
+	static const FString DeformedRootOffsetName;
+
+	/** Deformed position of the triangle vertex A */
+	static const FString DeformedTrianglePositionAName;
+
+	/** Deformed position of the triangle vertex A */
+	static const FString DeformedTrianglePositionBName;
+
+	/** Deformed position of the triangle vertex A */
+	static const FString DeformedTrianglePositionCName;
+
+	/** Root barycentric coordinates */
+	static const FString RootBarycentricCoordinatesName;
+
+	/** Rest center of all the position */
+	static const FString RestPositionOffsetName;
 
-	/** Name of the box origin  */
-	static const FString BoxOriginName;
+	/** Deformed center of all the position */
+	static const FString DeformedPositionOffsetName;
 
 protected:
 	/** Copy one niagara DI to this */
@@ -290,7 +376,7 @@ struct FNDIHairStrandsProxy : public FNiagaraDataInterfaceProxy
 	virtual void ConsumePerInstanceDataFromGameThread(void* PerInstanceData, const FNiagaraSystemInstanceID& Instance) override;
 
 	/** Initialize the Proxy data strands buffer */
-	void InitializePerInstanceData(const FNiagaraSystemInstanceID& SystemInstance, FNDIHairStrandsBuffer* StrandsBuffer, const uint32 NumStrands, const uint8 StrandSize, const FVector4& BoxOrigin, const FMatrix& WorldTransform);
+	void InitializePerInstanceData(const FNiagaraSystemInstanceID& SystemInstance, FNDIHairStrandsBuffer* StrandsBuffer, const uint32 NumStrands, const uint8 StrandSize, const FVector& BoxCenter, const FVector& BoxExtent, const FMatrix& WorldTransform);
 
 	/** Destroy the proxy data if necessary */
 	void DestroyPerInstanceData(NiagaraEmitterInstanceBatcher* Batcher, const FNiagaraSystemInstanceID& SystemInstance);
diff --git a/Engine/Shaders/Private/HairStrands/HairStrandsInterpolation.usf b/Engine/Shaders/Private/HairStrands/HairStrandsInterpolation.usf
--- a/Engine/Shaders/Private/HairStrands/HairStrandsInterpolation.usf
+++ b/Engine/Shaders/Private/HairStrands/HairStrandsInterpolation.usf
@@ -32,6 +32,20 @@ Buffer<float4>  DeformedPosition2Buffer;
 Buffer<uint>	RootBarycentricBuffer;
 Buffer<uint>	RootToTriangleIndex;
 Buffer<uint>	VertexToRootIndexBuffer;
+
+float3			SimRestPositionOffset;
+Buffer<float4>  SimRestPosition0Buffer;
+Buffer<float4>  SimRestPosition1Buffer;
+Buffer<float4>  SimRestPosition2Buffer;
+
+float3			SimDeformedPositionOffset;
+Buffer<float4>  SimDeformedPosition0Buffer;
+Buffer<float4>  SimDeformedPosition1Buffer;
+Buffer<float4>  SimDeformedPosition2Buffer;
+
+Buffer<uint>	SimRootBarycentricBuffer;
+Buffer<uint>	SimRootToTriangleIndex;
+Buffer<uint>	SimVertexToRootIndexBuffer;
 #endif
 
 #if PERMUTATION_DEBUG == 1
@@ -71,23 +85,59 @@ void MainCS(uint2 DispatchThreadId : SV_DispatchThreadID)
 		for (uint K = 0; K < 3; ++K)
 		{
 			const uint GuideIndex = GuideVertexIndices[K];
-			const float3 RestGuidePoint	= f16tof32(SimRestPosePositionBuffer[GuideIndex].xyz);
+			const float3 RestGuidePoint	= f16tof32(SimRestPosePositionBuffer[GuideIndex].xyz) + InHairPositionOffset;
 
-			const float3 DeformedGuidePoint	= f16tof32(DeformedSimPositionBuffer[GuideIndex].xyz);
+			const float3 DeformedGuidePoint	= f16tof32(DeformedSimPositionBuffer[GuideIndex].xyz) + OutHairPositionOffset;
 			CurrOffset += (DeformedGuidePoint - RestGuidePoint) * GuideVertexWeights[K];
-
 		}
+		OutRenderDeformedPositionBuffer[VertexIndex] = float4(f32tof16(ControlPoint + CurrOffset - OutHairPositionOffset), RenderRestPosePositionBuffer[VertexIndex].w);
 
 		// Apply dynamic mesh deformation (translation / rotation)
 		// hair_todo: applied the transformation in local space and then rotation everything in global space
 		#if PERMUTATION_DYNAMIC_GEOMETRY == 1
 		{
-			const uint RootIndex						= VertexToRootIndexBuffer[VertexIndex];
-			const float3 RootBarycentric				= DecodeBarycentrics(RootBarycentricBuffer[RootIndex]);
-			const FHairMeshTriangle RestTriangle		= GetTriangleTransformation(RootIndex, RestPosition0Buffer, RestPosition1Buffer, RestPosition2Buffer, RestPositionOffset);
-			const FHairMeshTriangle DeformedTriangle	= GetTriangleTransformation(RootIndex, DeformedPosition0Buffer, DeformedPosition1Buffer, DeformedPosition2Buffer, DeformedPositionOffset);
+			#if PERMUTATION_LOCAL_INTERPOLATION == 1
+			{
+				CurrOffset = 0;
+				for (uint K = 0; K < 3; ++K)
+				{
+					const uint GuideIndex = GuideVertexIndices[K];
+					const uint RootIndex = SimVertexToRootIndexBuffer[GuideIndex];
+					const float3 RootBarycentric = DecodeBarycentrics(SimRootBarycentricBuffer[RootIndex]);
+
+					const FHairMeshTriangle RestTriangle		= GetTriangleTransformation(RootIndex, SimRestPosition0Buffer, SimRestPosition1Buffer, SimRestPosition2Buffer, SimRestPositionOffset);
+					const FHairMeshTriangle DeformedTriangle	= GetTriangleTransformation(RootIndex, SimDeformedPosition0Buffer, SimDeformedPosition1Buffer, SimDeformedPosition2Buffer, SimDeformedPositionOffset);
+
+					const float3 RestGuidePoint	= f16tof32(SimRestPosePositionBuffer[GuideIndex].xyz) + InHairPositionOffset;
+					const float3 LocalRestGuidePoint = ToTriangle(RestGuidePoint,RootBarycentric,RestTriangle);
+
+					const float3 DeformedGuidePoint	= f16tof32(DeformedSimPositionBuffer[GuideIndex].xyz) + OutHairPositionOffset;
+					const float3 LocalDeformedGuidePoint = ToTriangle(DeformedGuidePoint,RootBarycentric,DeformedTriangle);
+					CurrOffset += (LocalDeformedGuidePoint - LocalRestGuidePoint) * GuideVertexWeights[K];
+				}
+			
+				const uint RootIndex						= VertexToRootIndexBuffer[VertexIndex];
+				const float3 RootBarycentric				= DecodeBarycentrics(RootBarycentricBuffer[RootIndex]);
+				const FHairMeshTriangle RestTriangle		= GetTriangleTransformation(RootIndex, RestPosition0Buffer, RestPosition1Buffer, RestPosition2Buffer, RestPositionOffset);
+				const FHairMeshTriangle DeformedTriangle	= GetTriangleTransformation(RootIndex, DeformedPosition0Buffer, DeformedPosition1Buffer, DeformedPosition2Buffer, DeformedPositionOffset);
+
+				ControlPoint = ToTriangle(ControlPoint,RootBarycentric,RestTriangle) + CurrOffset;
+				ControlPoint = ToWorld(ControlPoint,RootBarycentric,DeformedTriangle);
+			}
+			#endif	
 
-			ControlPoint = TransformPoint(ControlPoint, RootBarycentric, RestTriangle, DeformedTriangle);
+			#if PERMUTATION_LOCAL_INTERPOLATION == 0
+			{
+				const uint RootIndex						= VertexToRootIndexBuffer[VertexIndex];
+				const float3 RootBarycentric				= DecodeBarycentrics(RootBarycentricBuffer[RootIndex]);
+				const FHairMeshTriangle RestTriangle		= GetTriangleTransformation(RootIndex, RestPosition0Buffer, RestPosition1Buffer, RestPosition2Buffer, RestPositionOffset);
+				const FHairMeshTriangle DeformedTriangle	= GetTriangleTransformation(RootIndex, DeformedPosition0Buffer, DeformedPosition1Buffer, DeformedPosition2Buffer, DeformedPositionOffset);
+
+				ControlPoint = TransformPoint(ControlPoint, RootBarycentric, RestTriangle, DeformedTriangle);
+			}
+			#endif
+
+			OutRenderDeformedPositionBuffer[VertexIndex] = float4(f32tof16(ControlPoint - OutHairPositionOffset), RenderRestPosePositionBuffer[VertexIndex].w);
 
 			// Debug
 			#if 0
@@ -102,8 +152,6 @@ void MainCS(uint2 DispatchThreadId : SV_DispatchThreadID)
 		}
 		#endif	
 
-		OutRenderDeformedPositionBuffer[VertexIndex] = float4(f32tof16(ControlPoint + CurrOffset - OutHairPositionOffset), RenderRestPosePositionBuffer[VertexIndex].w);
-
 		#if PERMUTATION_DEBUG == 1
 		{
 			const uint GuideIndex = GuideVertexIndices[0]; // Take the closest guide
diff --git a/Engine/Source/Runtime/Renderer/Private/HairStrands/HairStrandsInterface.cpp b/Engine/Source/Runtime/Renderer/Private/HairStrands/HairStrandsInterface.cpp
--- a/Engine/Source/Runtime/Renderer/Private/HairStrands/HairStrandsInterface.cpp
+++ b/Engine/Source/Runtime/Renderer/Private/HairStrands/HairStrandsInterface.cpp
@@ -320,7 +320,7 @@ void RunHairStrandsInterpolation(
 
 			if (E.InterpolationData.Input && E.InterpolationData.Output && E.InterpolationData.Function)
 			{
-				E.InterpolationData.Function(RHICmdList, E.InterpolationData.Input, E.InterpolationData.Output, E.RenProjectionHairDatas, E.FrameLODIndex);
+				E.InterpolationData.Function(RHICmdList, E.InterpolationData.Input, E.InterpolationData.Output, E.RenProjectionHairDatas, E.SimProjectionHairDatas, E.FrameLODIndex);
 			}
 		}
 	}
diff --git a/Engine/Source/Runtime/Renderer/Public/HairStrandsInterface.h b/Engine/Source/Runtime/Renderer/Public/HairStrandsInterface.h
--- a/Engine/Source/Runtime/Renderer/Public/HairStrandsInterface.h
+++ b/Engine/Source/Runtime/Renderer/Public/HairStrandsInterface.h
@@ -52,7 +52,8 @@ typedef void (*THairStrandsInterpolationFunction)(
 	FRHICommandListImmediate& RHICmdList, 
 	struct FHairStrandsInterpolationInput* Input, 
 	struct FHairStrandsInterpolationOutput* Output, 
-	struct FHairStrandsProjectionHairData& HairProjection,
+	struct FHairStrandsProjectionHairData& RenHairProjection,
+	struct FHairStrandsProjectionHairData& SimHairProjection,
 	int32 LODIndex);
 
 struct FHairStrandsInterpolationData
