Remove hair from debug view mode as this is currently not supported.

diff --git a/Engine/Plugins/Experimental/HairStrands/Source/HairStrandsCore/Private/GroomComponent.cpp b/Engine/Plugins/Experimental/HairStrands/Source/HairStrandsCore/Private/GroomComponent.cpp
--- a/Engine/Plugins/Experimental/HairStrands/Source/HairStrandsCore/Private/GroomComponent.cpp
+++ b/Engine/Plugins/Experimental/HairStrands/Source/HairStrandsCore/Private/GroomComponent.cpp
@@ -319,7 +319,9 @@ public:
 	}
 
 	virtual FPrimitiveViewRelevance GetViewRelevance(const FSceneView* View) const override
-	{
+	{		
+		const bool bIsViewModeValid = View->Family->ViewMode == VMI_Lit;
+
 		const EHairStrandsDebugMode DebugMode = GetHairStrandsDebugStrandsMode();
 		if (DebugMode != EHairStrandsDebugMode::None)
 		{
@@ -332,7 +334,7 @@ public:
 		}
 
 		FPrimitiveViewRelevance Result;
-		Result.bHairStrandsRelevance = true;
+		Result.bHairStrandsRelevance = bIsViewModeValid;
 
 		// Special pass for hair strands geometry (not part of the base pass, and shadowing is handlded in a custom fashion)
 		Result.bDrawRelevance = false;		
@@ -343,7 +345,7 @@ public:
 		// Selection only
 		#if WITH_EDITOR
 		{
-			const bool bIsSelected = (IsSelected() || IsHovered());
+			const bool bIsSelected = (IsSelected() || IsHovered()) && bIsViewModeValid;
 			Result.bEditorStaticSelectionRelevance = bIsSelected;
 			Result.bDrawRelevance = bIsSelected;
 		}
diff --git a/Engine/Source/Runtime/Renderer/Private/DeferredShadingRenderer.cpp b/Engine/Source/Runtime/Renderer/Private/DeferredShadingRenderer.cpp
--- a/Engine/Source/Runtime/Renderer/Private/DeferredShadingRenderer.cpp
+++ b/Engine/Source/Runtime/Renderer/Private/DeferredShadingRenderer.cpp
@@ -1891,9 +1891,11 @@ void FDeferredShadingSceneRenderer::Render(FRHICommandListImmediate& RHICmdList)
 
 	checkSlow(RHICmdList.IsOutsideRenderPass());
 
+	// #hair_todo: Add multi-view
 	FHairStrandsDatas* HairDatas = nullptr;
 	FHairStrandsDatas HairDatasStorage;
-	if (IsHairStrandsEnable(Scene->GetShaderPlatform()))
+	const bool bIsViewCompatible = Views.Num() > 0 && Views[0].Family->ViewMode == VMI_Lit; 
+	if (IsHairStrandsEnable(Scene->GetShaderPlatform()) && bIsViewCompatible)
 	{
 		HairDatasStorage.HairClusterPerViews = CreateHairStrandsClusters(RHICmdList, Scene, Views);
 		VoxelizeHairStrands(RHICmdList, Scene, Views, HairDatasStorage.HairClusterPerViews);
