* Insure that material override work correctly
 * Hide Niagara system slot onto the hair groom, as this is currently not used.

diff --git a/Engine/Plugins/Experimental/HairStrands/Source/HairStrandsCore/Private/GroomAsset.cpp b/Engine/Plugins/Experimental/HairStrands/Source/HairStrandsCore/Private/GroomAsset.cpp
--- a/Engine/Plugins/Experimental/HairStrands/Source/HairStrandsCore/Private/GroomAsset.cpp
+++ b/Engine/Plugins/Experimental/HairStrands/Source/HairStrandsCore/Private/GroomAsset.cpp
@@ -355,6 +355,7 @@ void UGroomAsset::PostEditChangeProperty(FPropertyChangedEvent& PropertyChangedE
 {
 	Super::PostEditChangeProperty(PropertyChangedEvent);
 	UpdateResource();
+	OnGroomAssetChanged.Broadcast();
 }
 #endif // WITH_EDITOR
 
diff --git a/Engine/Plugins/Experimental/HairStrands/Source/HairStrandsCore/Private/GroomComponent.cpp b/Engine/Plugins/Experimental/HairStrands/Source/HairStrandsCore/Private/GroomComponent.cpp
--- a/Engine/Plugins/Experimental/HairStrands/Source/HairStrandsCore/Private/GroomComponent.cpp
+++ b/Engine/Plugins/Experimental/HairStrands/Source/HairStrandsCore/Private/GroomComponent.cpp
@@ -428,10 +428,7 @@ UGroomComponent::UGroomComponent(const FObjectInitializer& ObjectInitializer)
 	MeshProjectionTickDelay = 0;
 	MeshProjectionLODIndex = -1;
 	MeshProjectionState = EMeshProjectionState::Invalid;
-	if (GEngine)
-	{
-		SetMaterial(0, GEngine->HairDefaultMaterial);
-	}
+	bIsGroomAssetCallbackRegistered = false;
 
 	SetCollisionProfileName(UCollisionProfile::PhysicsActor_ProfileName);
 }
@@ -879,6 +876,19 @@ void UGroomComponent::PostLoad()
 	}
 
 	UpdateHairGroupsDesc(GroomAsset, GroomGroupsDesc);
+
+#if WITH_EDITOR
+	if (!bIsGroomAssetCallbackRegistered)
+	{
+		GroomAsset->GetOnGroomAssetChanged().AddUObject(this, &UGroomComponent::Invalidate);
+		bIsGroomAssetCallbackRegistered = true;
+	}
+#endif
+}
+
+void UGroomComponent::Invalidate()
+{
+	MarkRenderStateDirty();
 }
 
 void UGroomComponent::OnRegister()
@@ -903,6 +913,15 @@ void UGroomComponent::OnUnregister()
 void UGroomComponent::OnComponentDestroyed(bool bDestroyingHierarchy)
 {
 	ReleaseResources();
+
+#if WITH_EDITOR
+	if (bIsGroomAssetCallbackRegistered)
+	{
+		GroomAsset->GetOnGroomAssetChanged().RemoveAll(this);
+		bIsGroomAssetCallbackRegistered = false;
+	}
+#endif
+
 	Super::OnComponentDestroyed(bDestroyingHierarchy);
 }
 
@@ -1204,6 +1223,7 @@ void UGroomComponent::PostEditChangeProperty(FPropertyChangedEvent& PropertyChan
 	const FName PropertyName = PropertyThatChanged ? PropertyThatChanged->GetFName() : NAME_None;
 
 	//  Init/release resources when setting the GroomAsset (or undoing)
+	const bool bAssetChanged = PropertyName == GET_MEMBER_NAME_CHECKED(UGroomComponent, GroomAsset);
 	const bool bRecreateResources =
 		(PropertyName == GET_MEMBER_NAME_CHECKED(UGroomComponent, GroomAsset) || PropertyThatChanged == nullptr) ||
 		PropertyName == GET_MEMBER_NAME_CHECKED(UGroomComponent, bBindGroomToSkeletalMesh);
@@ -1214,6 +1234,23 @@ void UGroomComponent::PostEditChangeProperty(FPropertyChangedEvent& PropertyChan
 		bBindGroomToSkeletalMesh = false;
 	}
 
+#if WITH_EDITOR
+	if (bAssetChanged)
+	{
+		if (bIsGroomAssetCallbackRegistered && GroomAsset)
+		{
+			GroomAsset->GetOnGroomAssetChanged().RemoveAll(this);
+		}
+		bIsGroomAssetCallbackRegistered = false;
+
+		if (GroomAsset)
+		{
+			GroomAsset->GetOnGroomAssetChanged().AddUObject(this, &UGroomComponent::Invalidate);
+			bIsGroomAssetCallbackRegistered = true;
+		}
+	}
+#endif
+
 	if (bRecreateResources)
 	{
 		if (GroomAsset != InitializedResources)
diff --git a/Engine/Plugins/Experimental/HairStrands/Source/HairStrandsCore/Public/GroomAsset.h b/Engine/Plugins/Experimental/HairStrands/Source/HairStrandsCore/Public/GroomAsset.h
--- a/Engine/Plugins/Experimental/HairStrands/Source/HairStrandsCore/Public/GroomAsset.h
+++ b/Engine/Plugins/Experimental/HairStrands/Source/HairStrandsCore/Public/GroomAsset.h
@@ -196,8 +196,9 @@ struct HAIRSTRANDSCORE_API FHairGroupInfo
 	UPROPERTY(EditAnywhere, Category="Rendering")
 	UMaterialInterface* Material = nullptr;
 
-	UPROPERTY(EditAnywhere, Category="Simulation", meta = (DisplayName = "Niagara System Asset"))
-	UNiagaraSystem* NiagaraAsset = nullptr;
+	// Currently hide the Nigara simulation slot as it is not used, and could confuse users
+	//UPROPERTY(EditAnywhere, Category="Simulation", meta = (DisplayName = "Niagara System Asset"))
+	//UNiagaraSystem* NiagaraAsset = nullptr;
 };
 
 struct HAIRSTRANDSCORE_API FHairGroupData
@@ -223,6 +224,11 @@ class HAIRSTRANDSCORE_API UGroomAsset : public UObject
 {
 	GENERATED_BODY()
 
+#if WITH_EDITOR
+	/** Notification when anything changed */
+	DECLARE_MULTICAST_DELEGATE(FOnGroomAssetChanged);
+#endif
+
 public:
 
 	UPROPERTY(EditAnywhere, EditFixedSize, BlueprintReadWrite, Category = "HairGroups", meta = (DisplayName = "Group"))
@@ -240,6 +246,7 @@ public:
 	float HairToGuideDensity = 0.1f;
 
 #if WITH_EDITOR
+	FOnGroomAssetChanged& GetOnGroomAssetChanged() { return OnGroomAssetChanged;  }
 
 	/**  Part of Uobject interface  */
 	virtual void PostEditChangeProperty(FPropertyChangedEvent& PropertyChangedEvent) override;
@@ -278,7 +285,9 @@ public:
 
 	int32 GetNumHairGroups() const;
 
-//private : 
-
+//private :
+#if WITH_EDITOR
+	FOnGroomAssetChanged OnGroomAssetChanged;
+#endif
 	TUniquePtr<FHairDescription> HairDescription;
 };
diff --git a/Engine/Plugins/Experimental/HairStrands/Source/HairStrandsCore/Public/GroomComponent.h b/Engine/Plugins/Experimental/HairStrands/Source/HairStrandsCore/Public/GroomComponent.h
--- a/Engine/Plugins/Experimental/HairStrands/Source/HairStrandsCore/Public/GroomComponent.h
+++ b/Engine/Plugins/Experimental/HairStrands/Source/HairStrandsCore/Public/GroomComponent.h
@@ -95,6 +95,7 @@ public:
 #if WITH_EDITOR
 	virtual void PostEditChangeProperty(FPropertyChangedEvent& PropertyChangedEvent) override;
 	virtual bool CanEditChange(const UProperty* InProperty) const override;
+	void Invalidate();
 #endif
 
 	struct FHairGroupResource
@@ -141,6 +142,7 @@ private:
 	int32 MeshProjectionLODIndex;
 	uint32 MeshProjectionTickDelay;
 	EMeshProjectionState MeshProjectionState;
+	bool bIsGroomAssetCallbackRegistered;
 	
 	struct FSkeletalMeshConfiguration
 	{
