Add groom actor factory, to create a groom actor when drag and dropping an actor into the scene.

diff --git a/Engine/Plugins/Experimental/HairStrands/Source/HairStrandsEditor/Private/GroomActorFactory.cpp b/Engine/Plugins/Experimental/HairStrands/Source/HairStrandsEditor/Private/GroomActorFactory.cpp
new file mode 100644
--- /dev/null
+++ b/Engine/Plugins/Experimental/HairStrands/Source/HairStrandsEditor/Private/GroomActorFactory.cpp
@@ -0,0 +1,74 @@
+// Copyright 1998-2019 Epic Games, Inc. All Rights Reserved.
+
+#include "GroomActorFactory.h"
+#include "GroomAsset.h"
+#include "GroomActor.h"
+#include "GroomComponent.h"
+
+#define LOCTEXT_NAMESPACE "GroomActorFactory"
+
+UGroomActorFactory::UGroomActorFactory(const FObjectInitializer& ObjectInitializer)
+	: Super(ObjectInitializer)
+{
+	DisplayName = LOCTEXT("DefaultName", "Groom");
+	NewActorClass = AGroomActor::StaticClass();
+	bUseSurfaceOrientation = false;
+	bShowInEditorQuickMenu = true;
+}
+
+bool UGroomActorFactory::CanCreateActorFrom( const FAssetData& AssetData, FText& OutErrorMsg )
+{
+	if ( !AssetData.IsValid() || !AssetData.GetClass()->IsChildOf( UGroomAsset::StaticClass() ) )
+	{
+		OutErrorMsg = NSLOCTEXT("CanCreateActor", "NoGroomAsset", "A valid groom asset must be specified.");
+		return false;
+	}
+
+	return true;
+}
+
+void UGroomActorFactory::PostSpawnActor( UObject* Asset, AActor* NewActor)
+{
+	Super::PostSpawnActor(Asset, NewActor);
+
+	UGroomAsset* GroomAsset = CastChecked<UGroomAsset>(Asset);
+
+	// Change properties
+	AGroomActor* GroomActor = CastChecked<AGroomActor>( NewActor );
+	UGroomComponent* GroomComponent = GroomActor->GetGroomComponent();
+	check(GroomComponent);
+
+	GroomComponent->UnregisterComponent();
+	GroomComponent->GroomAsset = GroomAsset;
+	GroomComponent->PostLoad();
+	GroomComponent->RegisterComponent();
+}
+
+UObject* UGroomActorFactory::GetAssetFromActorInstance(AActor* Instance)
+{
+	check(Instance->IsA(NewActorClass));
+	AGroomActor* SMA = CastChecked<AGroomActor>(Instance);
+
+	check(SMA->GetGroomComponent());
+	return SMA->GetGroomComponent()->GroomAsset;
+}
+
+void UGroomActorFactory::PostCreateBlueprint( UObject* Asset, AActor* CDO )
+{
+	if (Asset != NULL && CDO != NULL)
+	{
+		UGroomAsset* GroomAsset = CastChecked<UGroomAsset>(Asset);
+		AGroomActor* GroomActor = CastChecked<AGroomActor>(CDO);
+		UGroomComponent* GroomComponent = GroomActor->GetGroomComponent();
+
+		GroomComponent->GroomAsset = GroomAsset;
+	}
+}
+
+FQuat UGroomActorFactory::AlignObjectToSurfaceNormal(const FVector& InSurfaceNormal, const FQuat& ActorRotation) const
+{
+	return FindActorAlignmentRotation(ActorRotation, FVector(0.f, 0.f, 1.f), InSurfaceNormal);
+}
+
+
+#undef LOCTEXT_NAMESPACE 
\ No newline at end of file
diff --git a/Engine/Plugins/Experimental/HairStrands/Source/HairStrandsEditor/Public/GroomActorFactory.h b/Engine/Plugins/Experimental/HairStrands/Source/HairStrandsEditor/Public/GroomActorFactory.h
new file mode 100644
--- /dev/null
+++ b/Engine/Plugins/Experimental/HairStrands/Source/HairStrandsEditor/Public/GroomActorFactory.h
@@ -0,0 +1,29 @@
+// Copyright 1998-2019 Epic Games, Inc. All Rights Reserved.
+
+
+#pragma once
+
+#include "CoreMinimal.h"
+#include "UObject/ObjectMacros.h"
+#include "ActorFactories/ActorFactory.h"
+#include "GroomActorFactory.generated.h"
+
+class AActor;
+struct FAssetData;
+
+UCLASS(MinimalAPI, config=Editor)
+class UGroomActorFactory : public UActorFactory
+{
+	GENERATED_UCLASS_BODY()
+
+	//~ Begin UActorFactory Interface
+	virtual bool CanCreateActorFrom( const FAssetData& AssetData, FText& OutErrorMsg ) override;
+	virtual void PostSpawnActor( UObject* Asset, AActor* NewActor) override;
+	virtual void PostCreateBlueprint( UObject* Asset, AActor* CDO ) override;
+	virtual UObject* GetAssetFromActorInstance(AActor* ActorInstance) override;
+	virtual FQuat AlignObjectToSurfaceNormal(const FVector& InSurfaceNormal, const FQuat& ActorRotation) const override;
+	//~ End UActorFactory Interface
+};
+
+
+
