Add icon to Groom asset to it is easier to spot them in the asset browser.

diff --git a/Engine/Plugins/Experimental/HairStrands/Source/HairStrandsEditor/Private/HairStrandsEditor.cpp b/Engine/Plugins/Experimental/HairStrands/Source/HairStrandsEditor/Private/HairStrandsEditor.cpp
--- a/Engine/Plugins/Experimental/HairStrands/Source/HairStrandsEditor/Private/HairStrandsEditor.cpp
+++ b/Engine/Plugins/Experimental/HairStrands/Source/HairStrandsEditor/Private/HairStrandsEditor.cpp
@@ -38,6 +38,9 @@ void FHairStrandsEditor::StartupModule()
 		
 		StyleSet->Set("ClassIcon.GroomActor", new FSlateImageBrush(HairStrandsContent + "/Icons/S_Groom_16.png", Icon16x16));
 		StyleSet->Set("ClassThumbnail.GroomActor", new FSlateImageBrush(HairStrandsContent + "/Icons/S_Groom_64.png", Icon64x64));
+
+		StyleSet->Set("ClassIcon.GroomAsset", new FSlateImageBrush(HairStrandsContent + "/Icons/S_Groom_16.png", Icon16x16));
+		StyleSet->Set("ClassThumbnail.GroomAsset", new FSlateImageBrush(HairStrandsContent + "/Icons/S_Groom_64.png", Icon64x64));
 		
 		FSlateStyleRegistry::RegisterSlateStyle(*StyleSet.Get());
 	}
