Fix non-editor builds
#rnx

diff --git a/Engine/Plugins/Experimental/HairStrands/Source/HairStrandsCore/Private/GroomComponent.cpp b/Engine/Plugins/Experimental/HairStrands/Source/HairStrandsCore/Private/GroomComponent.cpp
--- a/Engine/Plugins/Experimental/HairStrands/Source/HairStrandsCore/Private/GroomComponent.cpp
+++ b/Engine/Plugins/Experimental/HairStrands/Source/HairStrandsCore/Private/GroomComponent.cpp
@@ -886,10 +886,12 @@ void UGroomComponent::PostLoad()
 #endif
 }
 
+#if WITH_EDITOR
 void UGroomComponent::Invalidate()
 {
 	MarkRenderStateDirty();
 }
+#endif
 
 void UGroomComponent::OnRegister()
 {
