Remove in scene groom icon which make focusing on groom not accurate.

diff --git a/Engine/Plugins/Experimental/HairStrands/Source/HairStrandsCore/Private/GroomActor.cpp b/Engine/Plugins/Experimental/HairStrands/Source/HairStrandsCore/Private/GroomActor.cpp
--- a/Engine/Plugins/Experimental/HairStrands/Source/HairStrandsCore/Private/GroomActor.cpp
+++ b/Engine/Plugins/Experimental/HairStrands/Source/HairStrandsCore/Private/GroomActor.cpp
@@ -14,7 +14,10 @@ AGroomActor::AGroomActor(const FObjectInitializer& ObjectInitializer)
 	GroomComponent = CreateDefaultSubobject<UGroomComponent>(TEXT("GroomComponent0"));
 	RootComponent = GroomComponent;
 
+// Disable in editor icon as this bias the actor bounds size, and prevent to have 
+// accurate focusing on the groom component.
 #if WITH_EDITOR
+#if 0
 	struct FConstructorStatics
 	{
 		FName ID_Sprite;
@@ -45,6 +48,7 @@ AGroomActor::AGroomActor(const FObjectInitializer& ObjectInitializer)
 		SpriteComponent->SetupAttachment(RootComponent);
 	}
 #endif
+#endif
 }
 
 #undef LOCTEXT_NAMESPACE
