Add groom check, for catching error.

diff --git a/Engine/Plugins/Experimental/HairStrands/Source/HairStrandsCore/Private/GroomComponent.cpp b/Engine/Plugins/Experimental/HairStrands/Source/HairStrandsCore/Private/GroomComponent.cpp
--- a/Engine/Plugins/Experimental/HairStrands/Source/HairStrandsCore/Private/GroomComponent.cpp
+++ b/Engine/Plugins/Experimental/HairStrands/Source/HairStrandsCore/Private/GroomComponent.cpp
@@ -14,6 +14,9 @@
 #include "Components/SkeletalMeshComponent.h"
 #include "Engine/RendererSettings.h"
 #include "Animation/AnimationSettings.h"
+#include "Logging/MessageLog.h"
+#include "Misc/UObjectToken.h"
+#include "Misc/MapErrors.h"
 
 static float GHairClipLength = -1;
 static FAutoConsoleVariableRef CVarHairClipLength(TEXT("r.HairStrands.DebugClipLength"), GHairClipLength, TEXT("Clip hair strands which have a lenth larger than this value. (default is -1, no effect)"));
@@ -30,6 +33,8 @@ static FAutoConsoleVariableRef CVarHairStrandsMeshUseRelativePosition(TEXT("r.Ha
 static int32 GHairStrandsMeshProjectionTickDelay = 1;
 static FAutoConsoleVariableRef CVarHairStrandsMeshProjectionTickDelay(TEXT("r.HairStrands.MeshProjection.TickDelay"), GHairStrandsMeshProjectionTickDelay, TEXT("Number of simulation tick to wait before projecting a groom onto a mesh"));
 
+#define LOCTEXT_NAMESPACE "GroomComponent"
+
 /**
  * An material render proxy which overrides the debug mode parameter.
  */
@@ -883,6 +888,7 @@ void UGroomComponent::PostLoad()
 		GroomAsset->GetOnGroomAssetChanged().AddUObject(this, &UGroomComponent::Invalidate);
 		bIsGroomAssetCallbackRegistered = true;
 	}
+	ValidateMaterials(false);
 #endif
 }
 
@@ -890,6 +896,7 @@ void UGroomComponent::PostLoad()
 void UGroomComponent::Invalidate()
 {
 	MarkRenderStateDirty();
+	ValidateMaterials(false);
 }
 #endif
 
@@ -1281,6 +1288,9 @@ void UGroomComponent::PostEditChangeProperty(FPropertyChangedEvent& PropertyChan
 	{	
 		UpdateHairGroupsDesc(GroomAsset, GroomGroupsDesc);
 	}
+#if WITH_EDITOR
+	ValidateMaterials(false);
+#endif
 }
 
 bool UGroomComponent::CanEditChange(const UProperty* InProperty) const
@@ -1311,6 +1321,116 @@ bool UGroomComponent::CanEditChange(const UProperty* InProperty) const
 }
 #endif
 
+#if WITH_EDITOR
+void UGroomComponent::ValidateMaterials(bool bMapCheck) const
+{
+	if (!GroomAsset)
+		return;
+	
+	FString Name = "";
+	if (GetOwner())
+	{
+		Name += GetOwner()->GetName() + "/";
+	}
+	Name += GetName() + "/" + GroomAsset->GetName();
+	
+	const ERHIFeatureLevel::Type FeatureLevel = GetScene() ? GetScene()->GetFeatureLevel() : ERHIFeatureLevel::Num;
+	for (uint32 MaterialIt = 0, MaterialCount = GetNumMaterials(); MaterialIt < MaterialCount; ++MaterialIt)
+	{
+		UMaterialInterface* OverrideMaterial = Super::GetMaterial(MaterialIt);
+
+		FMaterialResource* Material = nullptr;
+
+		if (OverrideMaterial)
+		{
+			Material = OverrideMaterial->GetMaterialResource(FeatureLevel);
+		}
+		else if (MaterialIt < uint32(GroomAsset->HairGroupsInfo.Num()) && GroomAsset->HairGroupsInfo[MaterialIt].Material)
+		{
+			Material = GroomAsset->HairGroupsInfo[MaterialIt].Material->GetMaterialResource(FeatureLevel);
+		}
+
+		if (Material)
+		{
+			if (!Material->IsUsedWithHairStrands())
+			{
+				if (bMapCheck)
+				{
+					FMessageLog("MapCheck").Warning()
+						->AddToken(FUObjectToken::Create(GroomAsset))
+						->AddToken(FTextToken::Create(LOCTEXT("MapCheck_Message_InvalidHairStrandsMaterial", "Groom's material needs to have UseHairStrands option enabled. Groom's material will be replaced with default hair strands shader.")))
+						->AddToken(FMapErrorToken::Create(FMapErrors::InvalidHairStrandsMaterial));
+				}
+				else
+				{
+					UE_LOG(LogHairStrands, Warning, TEXT("[Groom] %s - Groom's material needs to have UseHairStrands option enabled. Groom's material will be replaced with default hair strands shader."));
+				}
+			}
+			if (!Material->GetShadingModels().HasShadingModel(MSM_Hair))
+			{
+				if (bMapCheck)
+				{
+					FMessageLog("MapCheck").Warning()
+						->AddToken(FUObjectToken::Create(GroomAsset))
+						->AddToken(FTextToken::Create(LOCTEXT("MapCheck_Message_InvalidHairStrandsMaterial", "Groom's material needs to have Hair shadering model. Groom's material will be replaced with default hair strands shader.")))
+						->AddToken(FMapErrorToken::Create(FMapErrors::InvalidHairStrandsMaterial));
+				}
+				else
+				{
+					UE_LOG(LogHairStrands, Warning, TEXT("[Groom] %s - Groom's material needs to have Hair shadering model. Groom's material will be replaced with default hair strands shader."), *Name);
+				}
+			}
+			if (Material->GetBlendMode() != BLEND_Opaque)
+			{
+				if (bMapCheck)
+				{
+					FMessageLog("MapCheck").Warning()
+						->AddToken(FUObjectToken::Create(GroomAsset))
+						->AddToken(FTextToken::Create(LOCTEXT("MapCheck_Message_InvalidHairStrandsMaterial", "Groom's material needs to have opaque blend mode. Groom's material will be replaced with default hair strands shader.")))
+						->AddToken(FMapErrorToken::Create(FMapErrors::InvalidHairStrandsMaterial));
+				}
+				else
+				{
+					UE_LOG(LogHairStrands, Warning, TEXT("[Groom] %s - Groom's material needs to have Hair shadering model. Groom's material will be replaced with default hair strands shader."), *Name);
+				}
+			}
+		}
+		else
+		{
+			if (bMapCheck)
+			{
+				FMessageLog("MapCheck").Info()
+					->AddToken(FUObjectToken::Create(GroomAsset))
+					->AddToken(FTextToken::Create(LOCTEXT("MapCheck_Message_InvalidHairStrandsMaterial", "Groom's material is not set and will fallback on default hair strands shader.")))
+					->AddToken(FMapErrorToken::Create(FMapErrors::InvalidHairStrandsMaterial));
+			}
+			else
+			{
+				UE_LOG(LogHairStrands, Warning, TEXT("[Groom] %s - Groom's material is not set and will fallback on default hair strands shader."), *Name);
+			}
+		}
+	}
+}
+
+void UGroomComponent::CheckForErrors()
+{
+	Super::CheckForErrors();
+
+	const FCoreTexts& CoreTexts = FCoreTexts::Get();
+
+	// Get the mesh owner's name.
+	AActor* Owner = GetOwner();
+	FString OwnerName(*(CoreTexts.None.ToString()));
+	if (Owner)
+	{
+		OwnerName = Owner->GetName();
+	}
+
+	ValidateMaterials(true);
+}
+#endif
+
+
 FGroomComponentRecreateRenderStateContext::FGroomComponentRecreateRenderStateContext(UGroomAsset* GroomAsset)
 {
 	if (!GroomAsset)
@@ -1348,3 +1468,5 @@ FGroomComponentRecreateRenderStateContext::~FGroomComponentRecreateRenderStateCo
 		}
 	}
 }
+
+#undef LOCTEXT_NAMESPACE
\ No newline at end of file
diff --git a/Engine/Plugins/Experimental/HairStrands/Source/HairStrandsCore/Public/GroomComponent.h b/Engine/Plugins/Experimental/HairStrands/Source/HairStrandsCore/Public/GroomComponent.h
--- a/Engine/Plugins/Experimental/HairStrands/Source/HairStrandsCore/Public/GroomComponent.h
+++ b/Engine/Plugins/Experimental/HairStrands/Source/HairStrandsCore/Public/GroomComponent.h
@@ -93,8 +93,10 @@ public:
 	FHairStrandsRootResource* GetGuideStrandsRootResource(uint32 GroupIndex);
 
 #if WITH_EDITOR
+	virtual void CheckForErrors() override;
 	virtual void PostEditChangeProperty(FPropertyChangedEvent& PropertyChangedEvent) override;
 	virtual bool CanEditChange(const UProperty* InProperty) const override;
+	void ValidateMaterials(bool bMapCheck) const;
 	void Invalidate();
 #endif
 
diff --git a/Engine/Source/Runtime/Core/Private/Misc/MapErrors.cpp b/Engine/Source/Runtime/Core/Private/Misc/MapErrors.cpp
--- a/Engine/Source/Runtime/Core/Private/Misc/MapErrors.cpp
+++ b/Engine/Source/Runtime/Core/Private/Misc/MapErrors.cpp
@@ -17,6 +17,7 @@ FLazyName FMapErrors::CollisionEnabledNoCollisionGeom(TEXT("CollisionEnabledNoCo
 FLazyName FMapErrors::ShadowCasterUsingBoundsScale(TEXT("ShadowCasterUsingBoundsScale"));
 FLazyName FMapErrors::MultipleSkyLights(TEXT("MultipleSkyLights"));
 FLazyName FMapErrors::MultipleSkyAtmospheres(TEXT("MultipleSkyAtmospheres"));
+FLazyName FMapErrors::InvalidHairStrandsMaterial(TEXT("InvalidHairStrandsMaterial"));
 FLazyName FMapErrors::MultipleSkyAtmosphereTypes(TEXT("MultipleSkyAtmosphereTypes"));
 FLazyName FMapErrors::InvalidTrace(TEXT("InvalidTrace"));
 FLazyName FMapErrors::BrushZeroPolygons(TEXT("BrushZeroPolygons"));
diff --git a/Engine/Source/Runtime/Core/Public/Misc/MapErrors.h b/Engine/Source/Runtime/Core/Public/Misc/MapErrors.h
--- a/Engine/Source/Runtime/Core/Public/Misc/MapErrors.h
+++ b/Engine/Source/Runtime/Core/Public/Misc/MapErrors.h
@@ -61,6 +61,9 @@ struct CORE_API FMapErrors
 	static FLazyName MultipleSkyAtmospheres;
 	static FLazyName MultipleSkyAtmosphereTypes;
 
+	/** Actor gas invalid hair strands material */
+	static FLazyName InvalidHairStrandsMaterial;
+
 	/** {ActorName} has WorldTrace blocked. It will be considered to be world geometry */
 	static FLazyName InvalidTrace;
 
