Fix crash when a groom componet is changed to be attached to another object in a level editor.

This forces the scene proxy to be recreated, preventing it to used pointer on deleted resources.

diff --git a/Engine/Plugins/Experimental/HairStrands/Source/HairStrandsCore/Private/GroomComponent.cpp b/Engine/Plugins/Experimental/HairStrands/Source/HairStrandsCore/Private/GroomComponent.cpp
--- a/Engine/Plugins/Experimental/HairStrands/Source/HairStrandsCore/Private/GroomComponent.cpp
+++ b/Engine/Plugins/Experimental/HairStrands/Source/HairStrandsCore/Private/GroomComponent.cpp
@@ -863,6 +863,8 @@ void UGroomComponent::ReleaseResources()
 	}
 	SkeletalPreviousPositionOffset = FVector::ZeroVector;
 	RegisteredSkeletalMeshComponent = nullptr;
+
+	MarkRenderStateDirty();
 }
 
 void UGroomComponent::PostLoad()
