Fix missing parameter in groom/hair warnig message.

diff --git a/Engine/Plugins/Experimental/HairStrands/Source/HairStrandsCore/Private/GroomComponent.cpp b/Engine/Plugins/Experimental/HairStrands/Source/HairStrandsCore/Private/GroomComponent.cpp
--- a/Engine/Plugins/Experimental/HairStrands/Source/HairStrandsCore/Private/GroomComponent.cpp
+++ b/Engine/Plugins/Experimental/HairStrands/Source/HairStrandsCore/Private/GroomComponent.cpp
@@ -1365,7 +1365,7 @@ void UGroomComponent::ValidateMaterials(bool bMapCheck) const
 				}
 				else
 				{
-					UE_LOG(LogHairStrands, Warning, TEXT("[Groom] %s - Groom's material needs to have UseHairStrands option enabled. Groom's material will be replaced with default hair strands shader."));
+					UE_LOG(LogHairStrands, Warning, TEXT("[Groom] %s - Groom's material needs to have UseHairStrands option enabled. Groom's material will be replaced with default hair strands shader."), *Name);
 				}
 			}
 			if (!Material->GetShadingModels().HasShadingModel(MSM_Hair))
