Fixed hair vertex factory toi work with new gpu scene texture/buffer logic.
Done with Charles.
Bumpbing shader version as changing the vertex shader was nor enough.

diff --git a/Engine/Plugins/Experimental/HairStrands/Source/HairStrandsCore/Private/HairStrandsVertexFactory.cpp b/Engine/Plugins/Experimental/HairStrands/Source/HairStrandsCore/Private/HairStrandsVertexFactory.cpp
--- a/Engine/Plugins/Experimental/HairStrands/Source/HairStrandsCore/Private/HairStrandsVertexFactory.cpp
+++ b/Engine/Plugins/Experimental/HairStrands/Source/HairStrandsCore/Private/HairStrandsVertexFactory.cpp
@@ -115,8 +115,10 @@ bool FHairStrandsVertexFactory::ShouldCompilePermutation(EShaderPlatform Platfor
 
 void FHairStrandsVertexFactory::ModifyCompilationEnvironment(const FVertexFactoryType* Type, EShaderPlatform Platform, const FMaterial* Material, FShaderCompilerEnvironment& OutEnvironment)
 {
-	OutEnvironment.SetDefine(TEXT("VF_SUPPORTS_PRIMITIVE_SCENE_DATA"), Type->SupportsPrimitiveIdStream() && UseGPUScene(Platform, GetMaxSupportedFeatureLevel(Platform)));
+	const bool bUseGPUSceneAndPrimitiveIdStream = Type->SupportsPrimitiveIdStream() && UseGPUScene(Platform, GetMaxSupportedFeatureLevel(Platform));
+	OutEnvironment.SetDefine(TEXT("VF_SUPPORTS_PRIMITIVE_SCENE_DATA"), bUseGPUSceneAndPrimitiveIdStream);
 	OutEnvironment.SetDefine(TEXT("VF_STRAND_HAIR"), TEXT("1"));
+	OutEnvironment.SetDefine(TEXT("VF_GPU_SCENE_BUFFER"), bUseGPUSceneAndPrimitiveIdStream && !GPUSceneUseTexture2D(Platform));
 }
 
 void FHairStrandsVertexFactory::ValidateCompiledResult(const FVertexFactoryType* Type, EShaderPlatform Platform, const FShaderParameterMap& ParameterMap, TArray<FString>& OutErrors)
diff --git a/Engine/Shaders/Public/ShaderVersion.ush b/Engine/Shaders/Public/ShaderVersion.ush
--- a/Engine/Shaders/Public/ShaderVersion.ush
+++ b/Engine/Shaders/Public/ShaderVersion.ush
@@ -3,4 +3,4 @@
 // in Platform.ush (which should be included in any shader) it allows to invalidate the shader DDC.
 // 
 // If you are merging streams and there is a conflict with this GUID you should make a new GUID rather than taking one or the other.
-// GUID = CE7B9FFB-2E88-43DA-AD0A-C147887A37ED
+// GUID = FDD7186B-678E-403A-8A1C-F24AAEDE76C5
