* Add support for data driven color and roughness per hair vertices.
 * Expose these material attribute to the material graph
 * Add strand debug mode for visualizsing these attributes.
 * Add fallback loading for previously imported assets.

diff --git a/Engine/Plugins/Experimental/HairStrands/Source/HairStrandsCore/Private/GroomAsset.cpp b/Engine/Plugins/Experimental/HairStrands/Source/HairStrandsCore/Private/GroomAsset.cpp
--- a/Engine/Plugins/Experimental/HairStrands/Source/HairStrandsCore/Private/GroomAsset.cpp
+++ b/Engine/Plugins/Experimental/HairStrands/Source/HairStrandsCore/Private/GroomAsset.cpp
@@ -42,22 +42,25 @@ void CreateBuffer(uint32 InVertexCount, FRWBuffer& OutBuffer)
 /////////////////////////////////////////////////////////////////////////////////////////
 
 FHairStrandsRestResource::FHairStrandsRestResource(const FHairStrandsDatas::FRenderData& HairStrandRenderData, const FVector& InPositionOffset) :
-	RestPositionBuffer(), AttributeBuffer(), PositionOffset(InPositionOffset), RenderData(HairStrandRenderData)
+	RestPositionBuffer(), AttributeBuffer(), MaterialBuffer(), PositionOffset(InPositionOffset), RenderData(HairStrandRenderData)
 {}
 
 void FHairStrandsRestResource::InitRHI()
 {
-	const TArray<FHairStrandsPositionFormat::Type>& RenderingPositions = RenderData.RenderingPositions;
+	const TArray<FHairStrandsPositionFormat::Type>& RenderingPositions	 = RenderData.RenderingPositions;
 	const TArray<FHairStrandsAttributeFormat::Type>& RenderingAttributes = RenderData.RenderingAttributes;
+	const TArray<FHairStrandsMaterialFormat::Type>& RenderingMaterials	 = RenderData.RenderingMaterials;
 
 	CreateBuffer<FHairStrandsPositionFormat>(RenderingPositions, RestPositionBuffer);
 	CreateBuffer<FHairStrandsAttributeFormat>(RenderingAttributes, AttributeBuffer);
+	CreateBuffer<FHairStrandsMaterialFormat>(RenderingMaterials, MaterialBuffer);
 }
 
 void FHairStrandsRestResource::ReleaseRHI()
 {
 	RestPositionBuffer.Release();
 	AttributeBuffer.Release();
+	MaterialBuffer.Release();
 }
 
 /////////////////////////////////////////////////////////////////////////////////////////
diff --git a/Engine/Plugins/Experimental/HairStrands/Source/HairStrandsCore/Private/GroomBuilder.cpp b/Engine/Plugins/Experimental/HairStrands/Source/HairStrandsCore/Private/GroomBuilder.cpp
--- a/Engine/Plugins/Experimental/HairStrands/Source/HairStrandsCore/Private/GroomBuilder.cpp
+++ b/Engine/Plugins/Experimental/HairStrands/Source/HairStrandsCore/Private/GroomBuilder.cpp
@@ -207,9 +207,11 @@ namespace HairStrandsBuilder
 
 		TArray<FHairStrandsPositionFormat::Type>& OutPackedPositions = HairStrands.RenderData.RenderingPositions;
 		TArray<FHairStrandsAttributeFormat::Type>& OutPackedAttributes = HairStrands.RenderData.RenderingAttributes;
+		TArray<FHairStrandsMaterialFormat::Type>& OutPackedMaterials = HairStrands.RenderData.RenderingMaterials;
 
 		OutPackedPositions.SetNum(NumPoints * FHairStrandsPositionFormat::ComponentCount);
 		OutPackedAttributes.SetNum(NumPoints * FHairStrandsAttributeFormat::ComponentCount);
+		OutPackedMaterials.SetNum(NumPoints * FHairStrandsMaterialFormat::ComponentCount);
 
 		const FVector HairBoxCenter = HairStrands.BoundingBox.GetCenter();
 
@@ -253,6 +255,14 @@ namespace HairStrandsBuilder
 				PackedAttributes.RootV  = uint32(FMath::Clamp(TextureRootUV.Y*255.f, 0.f, 255.f));
 				PackedAttributes.IndexU = uint32(FMath::Clamp(TextureIndexUV.X, 0.f, 255.f));
 				PackedAttributes.IndexV = uint32(FMath::Clamp(TextureIndexUV.Y, 0.f, 255.f));
+
+				FHairStrandsMaterialFormat::Type& Material = OutPackedMaterials[PointIndex + IndexOffset];
+				// Cheap sRGB encoding instead of PointsBaseColor.ToFColor(true), as this makes the decompression 
+				// cheaper on GPU (since R8G8B8A8 sRGB format used/exposed not exposed)
+				Material.BaseColorR = FMath::Clamp(uint32(FMath::Sqrt(Points.PointsBaseColor[PointIndex + IndexOffset].R) * 0xFF), 0u, 0xFFu);
+				Material.BaseColorG = FMath::Clamp(uint32(FMath::Sqrt(Points.PointsBaseColor[PointIndex + IndexOffset].G) * 0xFF), 0u, 0xFFu);
+				Material.BaseColorB = FMath::Clamp(uint32(FMath::Sqrt(Points.PointsBaseColor[PointIndex + IndexOffset].B) * 0xFF), 0u, 0xFFu);
+				Material.Roughness  = FMath::Clamp(uint32(Points.PointsRoughness[PointIndex + IndexOffset] * 0xFF), 0u, 0xFFu);
 			}
 		}
 	}
@@ -711,8 +721,9 @@ bool FGroomBuilder::BuildGroom(const FHairDescription& HairDescription, const FG
 		GroomHairColor = GroomHairColorAttribute[GroomID];
 	}
 
-	TVertexAttributesConstRef<FVector> VertexPositions = HairDescription.VertexAttributes().GetAttributesRef<FVector>(HairAttribute::Vertex::Position);
-	TStrandAttributesConstRef<int> StrandNumVertices = HairDescription.StrandAttributes().GetAttributesRef<int>(HairAttribute::Strand::VertexCount);
+	TVertexAttributesConstRef<FVector> VertexPositions	= HairDescription.VertexAttributes().GetAttributesRef<FVector>(HairAttribute::Vertex::Position);
+	TVertexAttributesConstRef<FVector> VertexBaseColor	= HairDescription.VertexAttributes().GetAttributesRef<FVector>(HairAttribute::Vertex::Color);
+	TStrandAttributesConstRef<int> StrandNumVertices	= HairDescription.StrandAttributes().GetAttributesRef<int>(HairAttribute::Strand::VertexCount);
 
 	if (!VertexPositions.IsValid() || !StrandNumVertices.IsValid())
 	{
@@ -804,6 +815,8 @@ bool FGroomBuilder::BuildGroom(const FHairDescription& HairDescription, const FG
 			FVertexID VertexID(GlobalVertexIndex);
 
 			CurrentHairStrandsDatas->StrandsPoints.PointsPosition.Add(VertexPositions[VertexID]);
+			CurrentHairStrandsDatas->StrandsPoints.PointsBaseColor.Add(VertexBaseColor[VertexID]);
+			CurrentHairStrandsDatas->StrandsPoints.PointsRoughness.Add(0); // @hair_todo: add attribute read on the alembic for reading roughness per groom/strands/vertex
 
 			float VertexWidth = 0.f;
 			if (VertexWidths.IsValid())
@@ -964,6 +977,8 @@ void FGroomBuilder::GenerateGuides(const FHairStrandsDatas& InData, float Decima
 			OutData.StrandsPoints.PointsPosition[OutPointOffset] = InData.StrandsPoints.PointsPosition[OutPointIndex + InPointOffset];
 			OutData.StrandsPoints.PointsCoordU[OutPointOffset] = InData.StrandsPoints.PointsCoordU[OutPointIndex + InPointOffset];
 			OutData.StrandsPoints.PointsRadius[OutPointOffset] = InData.StrandsPoints.PointsRadius[OutPointIndex + InPointOffset] * InData.StrandsCurves.MaxRadius;
+			OutData.StrandsPoints.PointsBaseColor[OutPointOffset] = FLinearColor::Black;
+			OutData.StrandsPoints.PointsRoughness[OutPointOffset] = 0;
 		}
 	}
 
diff --git a/Engine/Plugins/Experimental/HairStrands/Source/HairStrandsCore/Private/GroomComponent.cpp b/Engine/Plugins/Experimental/HairStrands/Source/HairStrandsCore/Private/GroomComponent.cpp
--- a/Engine/Plugins/Experimental/HairStrands/Source/HairStrandsCore/Private/GroomComponent.cpp
+++ b/Engine/Plugins/Experimental/HairStrands/Source/HairStrandsCore/Private/GroomComponent.cpp
@@ -259,6 +259,8 @@ public:
 			case EHairStrandsDebugMode::RenderHairDimension			: DebugModeScalar = 4.f; break;
 			case EHairStrandsDebugMode::RenderHairRadiusVariation	: DebugModeScalar = 5.f; break;
 			case EHairStrandsDebugMode::RenderHairRootUDIM			: DebugModeScalar = 6.f; break;
+			case EHairStrandsDebugMode::RenderHairBaseColor			: DebugModeScalar = 7.f; break;
+			case EHairStrandsDebugMode::RenderHairRoughness			: DebugModeScalar = 8.f; break;
 			};
 
 			float HairMinRadius = FLT_MAX;
@@ -797,6 +799,7 @@ void UGroomComponent::InitResources()
 			OutputGroup.RenderDeformedPositionBuffer[0] = &Res.RenderDeformedResources->DeformedPositionBuffer[0];
 			OutputGroup.RenderDeformedPositionBuffer[1] = &Res.RenderDeformedResources->DeformedPositionBuffer[1];
 			OutputGroup.RenderAttributeBuffer			= &Res.RenderRestResources->AttributeBuffer;
+			OutputGroup.RenderMaterialBuffer			= &Res.RenderRestResources->MaterialBuffer;
 			OutputGroup.RenderTangentBuffer				= &Res.RenderDeformedResources->TangentBuffer;
 			OutputGroup.SimTangentBuffer				= &Res.SimDeformedResources->TangentBuffer;
 
diff --git a/Engine/Plugins/Experimental/HairStrands/Source/HairStrandsCore/Private/HairStrandsDatas.cpp b/Engine/Plugins/Experimental/HairStrands/Source/HairStrandsCore/Private/HairStrandsDatas.cpp
--- a/Engine/Plugins/Experimental/HairStrands/Source/HairStrandsCore/Private/HairStrandsDatas.cpp
+++ b/Engine/Plugins/Experimental/HairStrands/Source/HairStrandsCore/Private/HairStrandsDatas.cpp
@@ -98,6 +98,16 @@ FArchive& operator<<(FArchive& Ar, FPackedHairAttributeVertex& Vertex)
 	return Ar;
 }
 
+FArchive& operator<<(FArchive& Ar, FHairMaterialVertex& Vertex)
+{
+	Ar << Vertex.BaseColorR;
+	Ar << Vertex.BaseColorG;
+	Ar << Vertex.BaseColorB;
+	Ar << Vertex.Roughness;
+
+	return Ar;
+}
+
 FArchive& operator<<(FArchive& Ar, FHairInterpolation0Vertex& Vertex)
 {
 	Ar << Vertex.Index0;
@@ -136,9 +146,23 @@ void FHairStrandsInterpolationDatas::Serialize(FArchive& Ar)
 
 void FHairStrandsPoints::Serialize(FArchive& Ar)
 {
+	Ar.UsingCustomVersion(FReleaseObjectVersion::GUID);
+
 	Ar << PointsPosition;
 	Ar << PointsRadius;
 	Ar << PointsCoordU;
+
+	if (Ar.CustomVer(FReleaseObjectVersion::GUID) >= FReleaseObjectVersion::GroomAssetVersion2)
+	{
+		Ar << PointsBaseColor;
+		Ar << PointsRoughness;
+	}
+	else
+	{
+		const uint32 ElementCount = PointsPosition.Num();
+		PointsBaseColor.InsertZeroed(0, ElementCount);
+		PointsRoughness.InsertZeroed(0, ElementCount);
+	}
 }
 
 void FHairStrandsCurves::Serialize(FArchive& Ar)
@@ -153,8 +177,19 @@ void FHairStrandsCurves::Serialize(FArchive& Ar)
 
 void FHairStrandsDatas::FRenderData::Serialize(FArchive& Ar)
 {
+	Ar.UsingCustomVersion(FReleaseObjectVersion::GUID);
+
 	Ar << RenderingPositions;
 	Ar << RenderingAttributes;
+	if (Ar.CustomVer(FReleaseObjectVersion::GUID) >= FReleaseObjectVersion::GroomAssetVersion2)
+	{
+		Ar << RenderingMaterials;
+	}
+	else
+	{
+		const uint32 ElementCount = RenderingAttributes.Num();
+		RenderingMaterials.InsertZeroed(0, ElementCount);
+	}
 }
 
 void FHairStrandsDatas::Serialize(FArchive& Ar)
diff --git a/Engine/Plugins/Experimental/HairStrands/Source/HairStrandsCore/Private/HairStrandsRendering.cpp b/Engine/Plugins/Experimental/HairStrands/Source/HairStrandsCore/Private/HairStrandsRendering.cpp
--- a/Engine/Plugins/Experimental/HairStrands/Source/HairStrandsCore/Private/HairStrandsRendering.cpp
+++ b/Engine/Plugins/Experimental/HairStrands/Source/HairStrandsCore/Private/HairStrandsRendering.cpp
@@ -471,6 +471,7 @@ void ComputeHairStrandsInterpolation(
 			Output.VFInput.HairPreviousPositionBuffer = Output.SimDeformedPositionBuffer[SimIndex]->SRV;
 			Output.VFInput.HairTangentBuffer = Output.SimTangentBuffer->SRV;
 			Output.VFInput.HairAttributeBuffer = Input.SimAttributeBuffer->SRV;
+			Output.VFInput.HairMaterialBuffer = Output.RenderMaterialBuffer->SRV;
 			Output.VFInput.HairPositionOffset = Input.OutHairPositionOffset;
 			Output.VFInput.HairPreviousPositionOffset = Input.OutHairPreviousPositionOffset;
 			Output.VFInput.VertexCount = Input.SimVertexCount;
@@ -527,6 +528,7 @@ void ComputeHairStrandsInterpolation(
 			Output.VFInput.HairPreviousPositionBuffer = Output.RenderDeformedPositionBuffer[PrevIndex]->SRV;
 			Output.VFInput.HairTangentBuffer = Output.RenderTangentBuffer->SRV;
 			Output.VFInput.HairAttributeBuffer = DebugMode == EHairStrandsDebugMode::RenderHairStrands ? Output.RenderPatchedAttributeBuffer.SRV : Input.RenderAttributeBuffer->SRV;
+			Output.VFInput.HairMaterialBuffer = Output.RenderMaterialBuffer->SRV;
 			Output.VFInput.HairPositionOffset = Input.OutHairPositionOffset;
 			Output.VFInput.HairPreviousPositionOffset = Input.OutHairPreviousPositionOffset;
 			Output.VFInput.VertexCount = Input.RenderVertexCount;
diff --git a/Engine/Plugins/Experimental/HairStrands/Source/HairStrandsCore/Private/HairStrandsVertexFactory.cpp b/Engine/Plugins/Experimental/HairStrands/Source/HairStrandsCore/Private/HairStrandsVertexFactory.cpp
--- a/Engine/Plugins/Experimental/HairStrands/Source/HairStrandsCore/Private/HairStrandsVertexFactory.cpp
+++ b/Engine/Plugins/Experimental/HairStrands/Source/HairStrandsCore/Private/HairStrandsVertexFactory.cpp
@@ -46,6 +46,7 @@ public:
 	FShaderResourceParameter PositionBuffer;
 	FShaderResourceParameter PreviousPositionBuffer;
 	FShaderResourceParameter AttributeBuffer;
+	FShaderResourceParameter MaterialBuffer;
 	FShaderResourceParameter TangentBuffer;
 
 	virtual void Bind(const FShaderParameterMap& ParameterMap) override
@@ -59,6 +60,7 @@ public:
 		PositionBuffer.Bind(ParameterMap, TEXT("HairStrandsVF_PositionBuffer"));
 		PreviousPositionBuffer.Bind(ParameterMap, TEXT("HairStrandsVF_PreviousPositionBuffer"));
 		AttributeBuffer.Bind(ParameterMap, TEXT("HairStrandsVF_AttributeBuffer"));
+		MaterialBuffer.Bind(ParameterMap, TEXT("HairStrandsVF_MaterialBuffer"));
 		TangentBuffer.Bind(ParameterMap, TEXT("HairStrandsVF_TangentBuffer"));
 	}
 
@@ -75,6 +77,7 @@ public:
 		Ar << PositionBuffer;
 		Ar << PreviousPositionBuffer;
 		Ar << AttributeBuffer;
+		Ar << MaterialBuffer;
 		Ar << TangentBuffer;
 	}
 
@@ -96,6 +99,7 @@ public:
 		BindParam(ShaderBindings, PositionBuffer, VF->GetPositionSRV(GroupIndex));
 		BindParam(ShaderBindings, PreviousPositionBuffer, VF->GetPreviousPositionSRV(GroupIndex));
 		BindParam(ShaderBindings, AttributeBuffer, VF->GetAttributeSRV(GroupIndex));
+		BindParam(ShaderBindings, MaterialBuffer, VF->GetMaterialSRV(GroupIndex));
 		BindParam(ShaderBindings, TangentBuffer, VF->GetTangentSRV(GroupIndex));
 		BindParam(ShaderBindings, Radius, VF->GetMaxStrandRadius(GroupIndex));
 		BindParam(ShaderBindings, Length, VF->GetMaxStrandLength(GroupIndex));
diff --git a/Engine/Plugins/Experimental/HairStrands/Source/HairStrandsCore/Public/GroomAsset.h b/Engine/Plugins/Experimental/HairStrands/Source/HairStrandsCore/Public/GroomAsset.h
--- a/Engine/Plugins/Experimental/HairStrands/Source/HairStrandsCore/Public/GroomAsset.h
+++ b/Engine/Plugins/Experimental/HairStrands/Source/HairStrandsCore/Public/GroomAsset.h
@@ -90,9 +90,12 @@ struct FHairStrandsRestResource : public FRenderResource
 	/* Strand hair rest position buffer */
 	FRWBuffer RestPositionBuffer;
 
-	/* Strand hair offset buffer */
+	/* Strand hair attribute buffer */
 	FRWBuffer AttributeBuffer;
 
+	/* Strand hair material buffer */
+	FRWBuffer MaterialBuffer;
+	
 	/* Position offset as the rest positions are expressed in relative coordinate (16bits) */
 	FVector PositionOffset = FVector::ZeroVector;
 
diff --git a/Engine/Plugins/Experimental/HairStrands/Source/HairStrandsCore/Public/HairStrandsDatas.h b/Engine/Plugins/Experimental/HairStrands/Source/HairStrandsCore/Public/HairStrandsDatas.h
--- a/Engine/Plugins/Experimental/HairStrands/Source/HairStrandsCore/Public/HairStrandsDatas.h
+++ b/Engine/Plugins/Experimental/HairStrands/Source/HairStrandsCore/Public/HairStrandsDatas.h
@@ -35,6 +35,18 @@ struct FPackedHairAttributeVertex
 	friend FArchive& operator<<(FArchive& Ar, FPackedHairAttributeVertex& Vertex);
 };
 
+struct FHairMaterialVertex
+{
+	// sRGB color space
+	uint8 BaseColorR;
+	uint8 BaseColorG;
+	uint8 BaseColorB;
+
+	uint8 Roughness;
+
+	friend FArchive& operator<<(FArchive& Ar, FHairMaterialVertex& Vertex);
+};
+
 struct FHairInterpolation0Vertex
 {
 	uint16 Index0;
@@ -83,6 +95,15 @@ struct FHairStrandsAttributeFormat
 	static const EPixelFormat Format = PF_R16G16B16A16_UINT;
 };
 
+struct FHairStrandsMaterialFormat
+{
+	typedef FHairMaterialVertex Type;
+	static const uint32 ComponentCount = 1;
+	static const uint32 SizeInByte = sizeof(Type);
+	static const EVertexElementType VertexElementType = VET_UByte4;
+	static const EPixelFormat Format = PF_R8G8B8A8;
+};
+
 struct FHairStrandsTangentFormat
 {
 	typedef FPackedNormal Type;
@@ -270,6 +291,12 @@ struct HAIRSTRANDSCORE_API FHairStrandsPoints
 
 	/** Normalized length */
 	TArray<float> PointsCoordU; // [0..1]
+
+	/** Material base color */
+	TArray<FLinearColor> PointsBaseColor; // [0..1]
+
+	/** Material roughness */
+	TArray<float> PointsRoughness; // [0..1]
 };
 
 /** Hair strands Curves attribute */
@@ -336,6 +363,7 @@ struct HAIRSTRANDSCORE_API FHairStrandsDatas
 	{
 		TArray<FHairStrandsPositionFormat::Type> RenderingPositions;
 		TArray<FHairStrandsAttributeFormat::Type> RenderingAttributes;
+		TArray<FHairStrandsMaterialFormat::Type> RenderingMaterials;
 
 		void Serialize(FArchive& Ar);
 	} RenderData;
diff --git a/Engine/Plugins/Experimental/HairStrands/Source/HairStrandsCore/Public/HairStrandsRendering.h b/Engine/Plugins/Experimental/HairStrands/Source/HairStrandsCore/Public/HairStrandsRendering.h
--- a/Engine/Plugins/Experimental/HairStrands/Source/HairStrandsCore/Public/HairStrandsRendering.h
+++ b/Engine/Plugins/Experimental/HairStrands/Source/HairStrandsCore/Public/HairStrandsRendering.h
@@ -64,6 +64,7 @@ struct FHairStrandsInterpolationOutput
 		FShaderResourceViewRHIRef HairPreviousPositionBuffer = nullptr;
 		FShaderResourceViewRHIRef HairTangentBuffer = nullptr;
 		FShaderResourceViewRHIRef HairAttributeBuffer = nullptr;
+		FShaderResourceViewRHIRef HairMaterialBuffer = nullptr;
 		FVector HairPositionOffset = FVector::ZeroVector;
 		FVector HairPreviousPositionOffset = FVector::ZeroVector;
 		uint32 VertexCount = 0;
@@ -74,6 +75,7 @@ struct FHairStrandsInterpolationOutput
 			HairPreviousPositionBuffer = nullptr;
 			HairTangentBuffer = nullptr;
 			HairAttributeBuffer = nullptr;
+			HairMaterialBuffer = nullptr;
 			HairPositionOffset = FVector::ZeroVector;
 			HairPreviousPositionOffset = FVector::ZeroVector;
 			VertexCount = 0;
@@ -88,6 +90,7 @@ struct FHairStrandsInterpolationOutput
 
 		FRWBuffer* RenderTangentBuffer = nullptr;
 		FRWBuffer* RenderAttributeBuffer = nullptr;
+		FRWBuffer* RenderMaterialBuffer = nullptr;
 
 		FRWBuffer* SimTangentBuffer = nullptr;
 
@@ -105,6 +108,7 @@ struct FHairStrandsInterpolationOutput
 				(RenderDeformedPositionBuffer[1] && RenderDeformedPositionBuffer[1]->SRV) &&
 				(RenderTangentBuffer && RenderTangentBuffer->SRV) &&
 				(RenderAttributeBuffer && RenderAttributeBuffer->SRV);
+				(RenderMaterialBuffer && RenderMaterialBuffer->SRV);
 		}
 	};
 	TArray<HairGroup> HairGroups;
diff --git a/Engine/Plugins/Experimental/HairStrands/Source/HairStrandsCore/Public/HairStrandsVertexFactory.h b/Engine/Plugins/Experimental/HairStrands/Source/HairStrandsCore/Public/HairStrandsVertexFactory.h
--- a/Engine/Plugins/Experimental/HairStrands/Source/HairStrandsCore/Public/HairStrandsVertexFactory.h
+++ b/Engine/Plugins/Experimental/HairStrands/Source/HairStrandsCore/Public/HairStrandsVertexFactory.h
@@ -73,6 +73,7 @@ public:
 	inline FRHIShaderResourceView* GetPositionSRV(uint32 GroupIndex) const			{ check(Data.InterpolationOutput); return Data.InterpolationOutput->HairGroups[GroupIndex].VFInput.HairPositionBuffer; };
 	inline FRHIShaderResourceView* GetPreviousPositionSRV(uint32 GroupIndex) const	{ check(Data.InterpolationOutput); return Data.InterpolationOutput->HairGroups[GroupIndex].VFInput.HairPreviousPositionBuffer; }
 	inline FRHIShaderResourceView* GetAttributeSRV(uint32 GroupIndex) const			{ check(Data.InterpolationOutput); return Data.InterpolationOutput->HairGroups[GroupIndex].VFInput.HairAttributeBuffer; }
+	inline FRHIShaderResourceView* GetMaterialSRV(uint32 GroupIndex) const			{ check(Data.InterpolationOutput); return Data.InterpolationOutput->HairGroups[GroupIndex].VFInput.HairMaterialBuffer; }
 	inline FRHIShaderResourceView* GetTangentSRV(uint32 GroupIndex) const			{ check(Data.InterpolationOutput); return Data.InterpolationOutput->HairGroups[GroupIndex].VFInput.HairTangentBuffer; }
 
 	float GetMaxStrandRadius(uint32 GroupIndex) const;
diff --git a/Engine/Shaders/Private/HairStrands/HairStrandsVertexFactory.ush b/Engine/Shaders/Private/HairStrands/HairStrandsVertexFactory.ush
--- a/Engine/Shaders/Private/HairStrands/HairStrandsVertexFactory.ush
+++ b/Engine/Shaders/Private/HairStrands/HairStrandsVertexFactory.ush
@@ -32,6 +32,7 @@ float  HairStrandsVF_Length;
 Buffer<uint4>  HairStrandsVF_PositionBuffer;
 Buffer<uint4>  HairStrandsVF_PreviousPositionBuffer;
 Buffer<uint4>  HairStrandsVF_AttributeBuffer;
+Buffer<float4> HairStrandsVF_MaterialBuffer;
 Buffer<float4> HairStrandsVF_TangentBuffer;
 
 // ## A quad is made of the following topology and indexing:
@@ -80,7 +81,6 @@ struct FHairVertexAttributes
 	float3 HairBaseColor;
 	float  HairSeed;
 	float  HairRoughness;
-	float  HairSpecular;
 };
 
 FHairVertexAttributes GetHairStrandsAttributes(uint HairPrimitiveId, float2 VertexUV)
@@ -106,6 +106,15 @@ FHairVertexAttributes GetHairStrandsAttributes(uint HairPrimitiveId, float2 Vert
 	// Varying along the strand
 	Out.HairUV.y		= VertexUV.y;
 	Out.HairUV.x		= lerp(Attributes0.UCoord, Attributes1.UCoord, VertexUV.x);
+
+	// Instead of doing the actual interpolation we could rely on dithering as an optimization
+	const float4 Material0 = HairStrandsVF_MaterialBuffer[VertexId0];
+	const float4 Material1 = HairStrandsVF_MaterialBuffer[VertexId1];
+
+	// Cheap/manual sRGB decoding as the underlying format is linear quantization
+	Out.HairBaseColor	= lerp(Material0.xyz*Material0.xyz, Material1.xyz*Material1.xyz, VertexUV.x);
+	Out.HairRoughness	= lerp(Material0.w, Material1.w, VertexUV.x);
+
 	return Out; 
 }
 
@@ -129,6 +138,16 @@ float  GetHairStrandsSeed(uint HairPrimitiveId, float2 HairPrimitiveUV)
 	return GetHairStrandsAttributes(HairPrimitiveId, HairPrimitiveUV).HairSeed;
 }
 
+float3 GetHairStrandsBaseColor(uint HairPrimitiveId, float2 HairPrimitiveUV)
+{
+	return GetHairStrandsAttributes(HairPrimitiveId, HairPrimitiveUV).HairBaseColor;
+}
+
+float GetHairStrandsRoughness(uint HairPrimitiveId, float2 HairPrimitiveUV)
+{
+	return GetHairStrandsAttributes(HairPrimitiveId, HairPrimitiveUV).HairRoughness;
+}
+
 struct FVertexFactoryInput
 {
 #if VF_USE_PRIMITIVE_SCENE_DATA
diff --git a/Engine/Shaders/Private/MaterialTemplate.ush b/Engine/Shaders/Private/MaterialTemplate.ush
--- a/Engine/Shaders/Private/MaterialTemplate.ush
+++ b/Engine/Shaders/Private/MaterialTemplate.ush
@@ -1002,6 +1002,8 @@ float2 GetHairStrandsRootUV(uint HairPrimitiveId, float2 VertexUV);
 float2 GetHairStrandsUV(uint HairPrimitiveId, float2 VertexUV);
 float2 GetHairStrandsDimensions(uint HairPrimitiveId, float2 VertexUV);
 float  GetHairStrandsSeed(uint HairPrimitiveId, float2 VertexUV);
+float  GetHairStrandsRoughness(uint HairPrimitiveId, float2 HairPrimitiveUV);
+float3 GetHairStrandsBaseColor(uint HairPrimitiveId, float2 HairPrimitiveUV);
 #endif
 
 float2 MaterialExpressionGetHairRootUV(FMaterialPixelParameters Parameters)
@@ -1040,6 +1042,23 @@ float MaterialExpressionGetHairSeed(FMaterialPixelParameters Parameters)
 #endif
 }
 
+float3 MaterialExpressionGetHairBaseColor(FMaterialPixelParameters Parameters)
+{
+#if VF_STRAND_HAIR
+	return GetHairStrandsBaseColor(Parameters.HairPrimitiveId, Parameters.HairPrimitiveUV);
+#else
+	return float3(0,0,0);
+#endif
+}
+
+float MaterialExpressionGetHairRoughness(FMaterialPixelParameters Parameters)
+{
+#if VF_STRAND_HAIR
+	return GetHairStrandsRoughness(Parameters.HairPrimitiveId, Parameters.HairPrimitiveUV);
+#else
+	return 0;
+#endif
+}
 float3 MaterialExpressionGetHairTangent(FMaterialPixelParameters Parameters)
 {
 	return Parameters.TangentToWorld[2];
diff --git a/Engine/Source/Runtime/Core/Public/UObject/ReleaseObjectVersion.h b/Engine/Source/Runtime/Core/Public/UObject/ReleaseObjectVersion.h
--- a/Engine/Source/Runtime/Core/Public/UObject/ReleaseObjectVersion.h
+++ b/Engine/Source/Runtime/Core/Public/UObject/ReleaseObjectVersion.h
@@ -85,6 +85,7 @@ struct CORE_API FReleaseObjectVersion
 
 		// Groom asset version
 		GroomAssetVersion1,
+		GroomAssetVersion2,
 
 		// -----<new versions can be added above this line>-------------------------------------------------
 		VersionPlusOne,
diff --git a/Engine/Source/Runtime/Engine/Private/Materials/HLSLMaterialTranslator.h b/Engine/Source/Runtime/Engine/Private/Materials/HLSLMaterialTranslator.h
--- a/Engine/Source/Runtime/Engine/Private/Materials/HLSLMaterialTranslator.h
+++ b/Engine/Source/Runtime/Engine/Private/Materials/HLSLMaterialTranslator.h
@@ -6644,6 +6644,16 @@ protected:
 		return AddCodeChunk(MCT_Float2, TEXT("MaterialExpressionGetHairRootUV(Parameters)"));
 	}
 
+	virtual int32 GetHairBaseColor() override
+	{
+		return AddCodeChunk(MCT_Float3, TEXT("MaterialExpressionGetHairBaseColor(Parameters)"));
+	}
+
+	virtual int32 GetHairRoughness() override
+	{
+		return AddCodeChunk(MCT_Float1, TEXT("MaterialExpressionGetHairRoughness(Parameters)"));
+	}
+
 	virtual int32 DistanceToNearestSurface(int32 PositionArg) override
 	{
 		if (ErrorUnlessFeatureLevelSupported(ERHIFeatureLevel::SM5) == INDEX_NONE)
diff --git a/Engine/Source/Runtime/Engine/Private/Materials/MaterialExpressions.cpp b/Engine/Source/Runtime/Engine/Private/Materials/MaterialExpressions.cpp
--- a/Engine/Source/Runtime/Engine/Private/Materials/MaterialExpressions.cpp
+++ b/Engine/Source/Runtime/Engine/Private/Materials/MaterialExpressions.cpp
@@ -16834,6 +16834,8 @@ UMaterialExpressionHairAttributes::UMaterialExpressionHairAttributes(const FObje
 	Outputs.Add(FExpressionOutput(TEXT("Seed")));
 	Outputs.Add(FExpressionOutput(TEXT("World Tangent"), 1, 1, 1, 1, 0));
 	Outputs.Add(FExpressionOutput(TEXT("Root UV"), 1, 1, 1, 0, 0));
+	Outputs.Add(FExpressionOutput(TEXT("BaseColor"), 1, 1, 1, 1, 0));
+	Outputs.Add(FExpressionOutput(TEXT("Roughness"), 1, 1, 0, 0, 0));
 #endif
 }
 
@@ -16860,6 +16862,14 @@ int32 UMaterialExpressionHairAttributes::Compile(class FMaterialCompiler* Compil
 	{
 		return Compiler->GetHairRootUV();
 	}
+	else if (OutputIndex == 7)
+	{
+		return Compiler->GetHairBaseColor();
+	}
+	else if (OutputIndex == 8)
+	{
+		return Compiler->GetHairRoughness();
+	}
 
 	return Compiler->Errorf(TEXT("Invalid input parameter"));
 }
diff --git a/Engine/Source/Runtime/Engine/Public/MaterialCompiler.h b/Engine/Source/Runtime/Engine/Public/MaterialCompiler.h
--- a/Engine/Source/Runtime/Engine/Public/MaterialCompiler.h
+++ b/Engine/Source/Runtime/Engine/Public/MaterialCompiler.h
@@ -357,6 +357,8 @@ public:
 	virtual int32 GetHairSeed() = 0;
 	virtual int32 GetHairTangent() = 0;
 	virtual int32 GetHairRootUV() = 0;
+	virtual int32 GetHairBaseColor() = 0;
+	virtual int32 GetHairRoughness() = 0;
 	virtual int32 CustomPrimitiveData(int32 OutputIndex, EMaterialValueType Type) = 0;
 	virtual int32 ShadingModel(EMaterialShadingModel InSelectedShadingModel) = 0;
 
@@ -615,6 +617,8 @@ public:
 	virtual int32 GetHairSeed() override { return Compiler->GetHairSeed(); }
 	virtual int32 GetHairTangent() override { return Compiler->GetHairTangent(); }
 	virtual int32 GetHairRootUV() override { return Compiler->GetHairRootUV(); }
+	virtual int32 GetHairBaseColor() override { return Compiler->GetHairBaseColor(); }
+	virtual int32 GetHairRoughness() override { return Compiler->GetHairRoughness(); }
 
 	virtual int32 RotateScaleOffsetTexCoords(int32 TexCoordCodeIndex, int32 RotationScale, int32 Offset) override
 	{
diff --git a/Engine/Source/Runtime/Renderer/Private/HairStrands/HairStrandsDebug.cpp b/Engine/Source/Runtime/Renderer/Private/HairStrands/HairStrandsDebug.cpp
--- a/Engine/Source/Runtime/Renderer/Private/HairStrands/HairStrandsDebug.cpp
+++ b/Engine/Source/Runtime/Renderer/Private/HairStrands/HairStrandsDebug.cpp
@@ -131,15 +131,17 @@ EHairStrandsDebugMode GetHairStrandsDebugStrandsMode()
 {
 	switch (GHairStrandsDebugStrandsMode)
 	{
-	case 0:  return EHairStrandsDebugMode::None;
-	case 1:  return EHairStrandsDebugMode::SimHairStrands;
-	case 2:  return EHairStrandsDebugMode::RenderHairStrands;
-	case 3:  return EHairStrandsDebugMode::RenderHairRootUV;
-	case 4:  return EHairStrandsDebugMode::RenderHairUV;
-	case 5:  return EHairStrandsDebugMode::RenderHairSeed;
-	case 6:  return EHairStrandsDebugMode::RenderHairDimension;
-	case 7:  return EHairStrandsDebugMode::RenderHairRadiusVariation;
-	case 8:  return EHairStrandsDebugMode::RenderHairRootUDIM;
+	case  0:  return EHairStrandsDebugMode::None;
+	case  1:  return EHairStrandsDebugMode::SimHairStrands;
+	case  2:  return EHairStrandsDebugMode::RenderHairStrands;
+	case  3:  return EHairStrandsDebugMode::RenderHairRootUV;
+	case  4:  return EHairStrandsDebugMode::RenderHairRootUDIM;
+	case  5:  return EHairStrandsDebugMode::RenderHairUV;
+	case  6:  return EHairStrandsDebugMode::RenderHairSeed;
+	case  7:  return EHairStrandsDebugMode::RenderHairDimension;
+	case  8:  return EHairStrandsDebugMode::RenderHairRadiusVariation;
+	case  9:  return EHairStrandsDebugMode::RenderHairBaseColor;
+	case 10:  return EHairStrandsDebugMode::RenderHairRoughness;
 	default: return EHairStrandsDebugMode::None;
 	};
 }
@@ -152,11 +154,13 @@ static const TCHAR* ToString(EHairStrandsDebugMode DebugMode)
 	case EHairStrandsDebugMode::SimHairStrands				: return TEXT("Simulation strands");
 	case EHairStrandsDebugMode::RenderHairStrands			: return TEXT("Rendering strands influences");
 	case EHairStrandsDebugMode::RenderHairRootUV			: return TEXT("Roots UV");
+	case EHairStrandsDebugMode::RenderHairRootUDIM			: return TEXT("Roots UV UDIM texture index");
 	case EHairStrandsDebugMode::RenderHairUV				: return TEXT("Hair UV");
 	case EHairStrandsDebugMode::RenderHairSeed				: return TEXT("Hair seed");
 	case EHairStrandsDebugMode::RenderHairDimension			: return TEXT("Hair dimensions");
 	case EHairStrandsDebugMode::RenderHairRadiusVariation	: return TEXT("Hair radius variation");
-	case EHairStrandsDebugMode::RenderHairRootUDIM			: return TEXT("Root UV UDIM texture index");		
+	case EHairStrandsDebugMode::RenderHairBaseColor			: return TEXT("Hair vertices color");
+	case EHairStrandsDebugMode::RenderHairRoughness			: return TEXT("Hair vertices roughness");
 	default													: return TEXT("None");
 	};
 }
diff --git a/Engine/Source/Runtime/Renderer/Public/HairStrandsInterface.h b/Engine/Source/Runtime/Renderer/Public/HairStrandsInterface.h
--- a/Engine/Source/Runtime/Renderer/Public/HairStrandsInterface.h
+++ b/Engine/Source/Runtime/Renderer/Public/HairStrandsInterface.h
@@ -21,10 +21,12 @@ enum class EHairStrandsDebugMode : uint8
 	RenderHairStrands,
 	RenderHairUV,
 	RenderHairRootUV,
+	RenderHairRootUDIM,
 	RenderHairSeed,
 	RenderHairDimension,
 	RenderHairRadiusVariation,
-	RenderHairRootUDIM,
+	RenderHairBaseColor,
+	RenderHairRoughness,
 	Count
 };
 
