Fix localization error.

diff --git a/Engine/Plugins/Experimental/HairStrands/Source/HairStrandsCore/Private/GroomComponent.cpp b/Engine/Plugins/Experimental/HairStrands/Source/HairStrandsCore/Private/GroomComponent.cpp
--- a/Engine/Plugins/Experimental/HairStrands/Source/HairStrandsCore/Private/GroomComponent.cpp
+++ b/Engine/Plugins/Experimental/HairStrands/Source/HairStrandsCore/Private/GroomComponent.cpp
@@ -1363,12 +1363,12 @@ void UGroomComponent::ValidateMaterials(bool bMapCheck) const
 				{
 					FMessageLog("MapCheck").Warning()
 						->AddToken(FUObjectToken::Create(GroomAsset))
-						->AddToken(FTextToken::Create(LOCTEXT("MapCheck_Message_InvalidHairStrandsMaterial", "Groom's material needs to have UseHairStrands option enabled. Groom's material will be replaced with default hair strands shader.")))
+						->AddToken(FTextToken::Create(LOCTEXT("MapCheck_Message_HairStrandsMissingUseHairStrands", "Groom's material needs to enable the UseHairStrands option. Groom's material will be replaced with default hair strands shader.")))
 						->AddToken(FMapErrorToken::Create(FMapErrors::InvalidHairStrandsMaterial));
 				}
 				else
 				{
-					UE_LOG(LogHairStrands, Warning, TEXT("[Groom] %s - Groom's material needs to have UseHairStrands option enabled. Groom's material will be replaced with default hair strands shader."), *Name);
+					UE_LOG(LogHairStrands, Warning, TEXT("[Groom] %s - Groom's material needs to enable the UseHairStrands option. Groom's material will be replaced with default hair strands shader."), *Name);
 				}
 			}
 			if (!Material->GetShadingModels().HasShadingModel(MSM_Hair))
@@ -1377,12 +1377,12 @@ void UGroomComponent::ValidateMaterials(bool bMapCheck) const
 				{
 					FMessageLog("MapCheck").Warning()
 						->AddToken(FUObjectToken::Create(GroomAsset))
-						->AddToken(FTextToken::Create(LOCTEXT("MapCheck_Message_InvalidHairStrandsMaterial", "Groom's material needs to have Hair shadering model. Groom's material will be replaced with default hair strands shader.")))
+						->AddToken(FTextToken::Create(LOCTEXT("MapCheck_Message_HairStrandsInvalidShadingModel", "Groom's material needs to have Hair shading model. Groom's material will be replaced with default hair strands shader.")))
 						->AddToken(FMapErrorToken::Create(FMapErrors::InvalidHairStrandsMaterial));
 				}
 				else
 				{
-					UE_LOG(LogHairStrands, Warning, TEXT("[Groom] %s - Groom's material needs to have Hair shadering model. Groom's material will be replaced with default hair strands shader."), *Name);
+					UE_LOG(LogHairStrands, Warning, TEXT("[Groom] %s - Groom's material needs to have Hair shading model. Groom's material will be replaced with default hair strands shader."), *Name);
 				}
 			}
 			if (Material->GetBlendMode() != BLEND_Opaque)
@@ -1391,12 +1391,12 @@ void UGroomComponent::ValidateMaterials(bool bMapCheck) const
 				{
 					FMessageLog("MapCheck").Warning()
 						->AddToken(FUObjectToken::Create(GroomAsset))
-						->AddToken(FTextToken::Create(LOCTEXT("MapCheck_Message_InvalidHairStrandsMaterial", "Groom's material needs to have opaque blend mode. Groom's material will be replaced with default hair strands shader.")))
+						->AddToken(FTextToken::Create(LOCTEXT("MapCheck_Message_HairStrandsInvalidBlendMode", "Groom's material needs to have Opaque blend mode. Groom's material will be replaced with default hair strands shader.")))
 						->AddToken(FMapErrorToken::Create(FMapErrors::InvalidHairStrandsMaterial));
 				}
 				else
 				{
-					UE_LOG(LogHairStrands, Warning, TEXT("[Groom] %s - Groom's material needs to have Hair shadering model. Groom's material will be replaced with default hair strands shader."), *Name);
+					UE_LOG(LogHairStrands, Warning, TEXT("[Groom] %s - Groom's material needs to have Opaque blend mode. Groom's material will be replaced with default hair strands shader."), *Name);
 				}
 			}
 		}
@@ -1406,7 +1406,7 @@ void UGroomComponent::ValidateMaterials(bool bMapCheck) const
 			{
 				FMessageLog("MapCheck").Info()
 					->AddToken(FUObjectToken::Create(GroomAsset))
-					->AddToken(FTextToken::Create(LOCTEXT("MapCheck_Message_InvalidHairStrandsMaterial", "Groom's material is not set and will fallback on default hair strands shader.")))
+					->AddToken(FTextToken::Create(LOCTEXT("MapCheck_Message_HairStrandsMissingMaterial", "Groom's material is not set and will fallback on default hair strands shader.")))
 					->AddToken(FMapErrorToken::Create(FMapErrors::InvalidHairStrandsMaterial));
 			}
 			else
