Fix non-initialized material attribute array size.

diff --git a/Engine/Plugins/Experimental/HairStrands/Source/HairStrandsCore/Private/HairStrandsDatas.cpp b/Engine/Plugins/Experimental/HairStrands/Source/HairStrandsCore/Private/HairStrandsDatas.cpp
--- a/Engine/Plugins/Experimental/HairStrands/Source/HairStrandsCore/Private/HairStrandsDatas.cpp
+++ b/Engine/Plugins/Experimental/HairStrands/Source/HairStrandsCore/Private/HairStrandsDatas.cpp
@@ -23,6 +23,8 @@ void FHairStrandsPoints::SetNum(const uint32 NumPoints)
 	PointsPosition.SetNum(NumPoints);
 	PointsRadius.SetNum(NumPoints);
 	PointsCoordU.SetNum(NumPoints);
+	PointsBaseColor.SetNum(NumPoints);
+	PointsRoughness.SetNum(NumPoints);
 }
 
 void FHairStrandsInterpolationDatas::Reset()
