fixed invisible hair when animated and in raytracing mode.

diff --git a/Engine/Plugins/Experimental/HairStrands/Source/HairStrandsCore/Private/GroomComponent.cpp b/Engine/Plugins/Experimental/HairStrands/Source/HairStrandsCore/Private/GroomComponent.cpp
--- a/Engine/Plugins/Experimental/HairStrands/Source/HairStrandsCore/Private/GroomComponent.cpp
+++ b/Engine/Plugins/Experimental/HairStrands/Source/HairStrandsCore/Private/GroomComponent.cpp
@@ -35,6 +35,11 @@ static FAutoConsoleVariableRef CVarHairStrandsMeshProjectionTickDelay(TEXT("r.Ha
 
 #define LOCTEXT_NAMESPACE "GroomComponent"
 
+static float GetGroupMaxHairRadius(const FHairGroupDesc& GroupDesc, const FHairGroupData& GroupData)
+{
+	return GroupDesc.HairWidth > 0 ? GroupDesc.HairWidth * 0.5f : GroupData.HairRenderData.StrandsCurves.MaxRadius;
+}
+
 /**
  * An material render proxy which overrides the debug mode parameter.
  */
@@ -160,7 +165,7 @@ public:
 
 			FHairStrandsVertexFactory::FDataType::HairGroup& VFGroupData = VFData.HairGroups.AddDefaulted_GetRef();
 			VFGroupData.MinStrandRadius		= 0;
-			VFGroupData.MaxStrandRadius		= InGroupDesc.HairWidth > 0 ? InGroupDesc.HairWidth * 0.5f : InGroupData.HairRenderData.StrandsCurves.MaxRadius;
+			VFGroupData.MaxStrandRadius		= GetGroupMaxHairRadius(InGroupDesc, InGroupData);
 			VFGroupData.MaxStrandLength		= InGroupData.HairRenderData.StrandsCurves.MaxLength;
 			VFGroupData.HairDensity			= InGroupData.HairRenderData.HairDensity;
 			VFGroupData.HairWorldOffset		= InGroupData.HairRenderData.BoundingBox.GetCenter();
@@ -682,6 +687,7 @@ void UGroomComponent::InitResources()
 	InterpolationInput = new FHairStrandsInterpolationInput();
 
 	FHairStrandsDebugInfo DebugGroupInfo;
+	int32 GroupIt = 0;
 	for (FHairGroupData& GroupData : GroomAsset->HairGroupsData)
 	{
 		if (!GroupData.HairStrandsRestResource)
@@ -737,12 +743,17 @@ void UGroomComponent::InitResources()
 
 		FHairStrandsInterpolationOutput::HairGroup& InterpolationOutputGroup = InterpolationOutput->HairGroups.AddDefaulted_GetRef();
 		FHairStrandsInterpolationInput::FHairGroup& InterpolationInputGroup = InterpolationInput->HairGroups.AddDefaulted_GetRef();
-		InterpolationInputGroup.HairRadius = GroupData.HairRenderData.StrandsCurves.MaxRadius;
+
+		check(GroupIt < GroomGroupsDesc.Num());
+		const FHairGroupDesc& InGroupDesc = GroomGroupsDesc[GroupIt];
+		InterpolationInputGroup.HairRadius = GetGroupMaxHairRadius(InGroupDesc, GroupData);
 		InterpolationInputGroup.HairRaytracingRadiusScale = HairRaytracingRadiusScale;
 		InterpolationInputGroup.InHairPositionOffset = RestHairPositionOffset;
 		// For skinned groom, these value will be updated during TickComponent() call
 		InterpolationInputGroup.OutHairPositionOffset = RestHairPositionOffset;
 		InterpolationInputGroup.OutHairPreviousPositionOffset = RestHairPositionOffset;
+
+		GroupIt++;
 	}
 
 	FHairStrandsInterpolationData Interpolation;
diff --git a/Engine/Shaders/Private/HairStrands/HairStrandsRaytracingGeometry.usf b/Engine/Shaders/Private/HairStrands/HairStrandsRaytracingGeometry.usf
--- a/Engine/Shaders/Private/HairStrands/HairStrandsRaytracingGeometry.usf
+++ b/Engine/Shaders/Private/HairStrands/HairStrandsRaytracingGeometry.usf
@@ -38,7 +38,7 @@ FControlPoint GetControlPoint(uint Index)
 
 		Out.Position 	= f16tof32(PackedData.xyz) + StrandHairWorldOffset;
 		Out.PointType 	= PackedAlpha & 0x3;
-		Out.WorldRadius	= ((PackedAlpha & 0x00FC) >> 2) / 63.f * StrandHairRadius;	
+		Out.WorldRadius	= StrandHairRadius * float((PackedAlpha & 0x00FC) >> 2) / 63.0f;
 	}
 
 	return Out;
