Fix groom import when the asset does not have any color info

diff --git a/Engine/Plugins/Experimental/HairStrands/Source/HairStrandsCore/Private/GroomBuilder.cpp b/Engine/Plugins/Experimental/HairStrands/Source/HairStrandsCore/Private/GroomBuilder.cpp
--- a/Engine/Plugins/Experimental/HairStrands/Source/HairStrandsCore/Private/GroomBuilder.cpp
+++ b/Engine/Plugins/Experimental/HairStrands/Source/HairStrandsCore/Private/GroomBuilder.cpp
@@ -731,6 +731,8 @@ bool FGroomBuilder::BuildGroom(const FHairDescription& HairDescription, const FG
 		return false;
 	}
 
+	const bool bHasBaseColorAttribute = VertexBaseColor.IsValid();
+
 	TVertexAttributesConstRef<float> VertexWidths = HairDescription.VertexAttributes().GetAttributesRef<float>(HairAttribute::Vertex::Width);
 	TStrandAttributesConstRef<float> StrandWidths = HairDescription.StrandAttributes().GetAttributesRef<float>(HairAttribute::Strand::Width);
 
@@ -815,7 +817,7 @@ bool FGroomBuilder::BuildGroom(const FHairDescription& HairDescription, const FG
 			FVertexID VertexID(GlobalVertexIndex);
 
 			CurrentHairStrandsDatas->StrandsPoints.PointsPosition.Add(VertexPositions[VertexID]);
-			CurrentHairStrandsDatas->StrandsPoints.PointsBaseColor.Add(VertexBaseColor[VertexID]);
+			CurrentHairStrandsDatas->StrandsPoints.PointsBaseColor.Add(bHasBaseColorAttribute ? VertexBaseColor[VertexID] : FLinearColor::Black);
 			CurrentHairStrandsDatas->StrandsPoints.PointsRoughness.Add(0); // @hair_todo: add attribute read on the alembic for reading roughness per groom/strands/vertex
 
 			float VertexWidth = 0.f;
