Fix hair crash.

Insure hair group desc are filled in before used (in initresources)
Recreate proxy when the hair width change, as currently only the component data were updated (not the proxy/vertexfactory)

diff --git a/Engine/Plugins/Experimental/HairStrands/Source/HairStrandsCore/Private/GroomComponent.cpp b/Engine/Plugins/Experimental/HairStrands/Source/HairStrandsCore/Private/GroomComponent.cpp
--- a/Engine/Plugins/Experimental/HairStrands/Source/HairStrandsCore/Private/GroomComponent.cpp
+++ b/Engine/Plugins/Experimental/HairStrands/Source/HairStrandsCore/Private/GroomComponent.cpp
@@ -891,12 +891,12 @@ void UGroomComponent::PostLoad()
 		GroomAsset->ConditionalPostLoad();
 	}
 
+	UpdateHairGroupsDesc(GroomAsset, GroomGroupsDesc);
 	if (!IsTemplate())
 	{
 		InitResources();
 	}
 
-	UpdateHairGroupsDesc(GroomAsset, GroomGroupsDesc);
 
 #if WITH_EDITOR
 	if (!bIsGroomAssetCallbackRegistered)
@@ -1282,8 +1282,8 @@ void UGroomComponent::PostEditChangeProperty(FPropertyChangedEvent& PropertyChan
 		{
 			if (GroomAsset)
 			{
-				InitResources();
 				UpdateHairGroupsDesc(GroomAsset, GroomGroupsDesc);
+				InitResources();
 			}
 			else
 			{
@@ -1303,6 +1303,7 @@ void UGroomComponent::PostEditChangeProperty(FPropertyChangedEvent& PropertyChan
 	if (PropertyName == GET_MEMBER_NAME_CHECKED(FHairGroupDesc, HairWidth))
 	{	
 		UpdateHairGroupsDesc(GroomAsset, GroomGroupsDesc);
+		MarkRenderStateDirty();
 	}
 #if WITH_EDITOR
 	ValidateMaterials(false);
