Add automatic enable/disable of hair simulation based if a niagara comonent is attached to the groom component.

diff --git a/Engine/Plugins/Experimental/HairStrands/Source/HairStrandsCore/Private/GroomComponent.cpp b/Engine/Plugins/Experimental/HairStrands/Source/HairStrandsCore/Private/GroomComponent.cpp
--- a/Engine/Plugins/Experimental/HairStrands/Source/HairStrandsCore/Private/GroomComponent.cpp
+++ b/Engine/Plugins/Experimental/HairStrands/Source/HairStrandsCore/Private/GroomComponent.cpp
@@ -17,6 +17,7 @@
 #include "Logging/MessageLog.h"
 #include "Misc/UObjectToken.h"
 #include "Misc/MapErrors.h"
+#include "NiagaraComponent.h"
 
 static float GHairClipLength = -1;
 static FAutoConsoleVariableRef CVarHairClipLength(TEXT("r.HairStrands.DebugClipLength"), GHairClipLength, TEXT("Clip hair strands which have a lenth larger than this value. (default is -1, no effect)"));
@@ -654,6 +655,60 @@ void CallbackMeshObjectCallback(
 	});
 }
 
+static bool IsSimulationEnabled(const USceneComponent* Component)
+{
+	check(Component);
+
+	// If the groom component has an Niagara component attached, we assume it has simulation capabilities
+	bool bHasNiagaraSimulationComponent = false;
+	for (int32 ChildIt = 0, ChildCount = Component->GetNumChildrenComponents(); ChildIt < ChildCount; ++ChildIt)
+	{
+		const USceneComponent* ChildComponent = Component->GetChildComponent(ChildIt);
+		const UNiagaraComponent* NiagaraComponent = Cast<UNiagaraComponent>(ChildComponent);
+		if (NiagaraComponent != nullptr)
+		{
+			bHasNiagaraSimulationComponent = true;
+			break;
+		}
+	}
+
+	return bHasNiagaraSimulationComponent;
+}
+
+void UGroomComponent::OnChildAttached(USceneComponent* ChildComponent)
+{
+	const UNiagaraComponent* NiagaraComponent = Cast<UNiagaraComponent>(ChildComponent);
+	if (NiagaraComponent && InterpolationInput)
+	{
+		FHairStrandsInterpolationInput* LocalInterpolationInput = InterpolationInput;
+		ENQUEUE_RENDER_COMMAND(FHairStrandsTick_UpdateSimulationEnable)(
+			[LocalInterpolationInput](FRHICommandListImmediate& RHICmdList)
+		{
+			for (FHairStrandsInterpolationInput::FHairGroup& HairGroup : LocalInterpolationInput->HairGroups)
+			{
+				HairGroup.bIsSimulationEnable = true;
+			}
+		});
+	}
+}
+
+void UGroomComponent::OnChildDetached(USceneComponent* ChildComponent)
+{
+	const UNiagaraComponent* NiagaraComponent = Cast<UNiagaraComponent>(ChildComponent);
+	if (NiagaraComponent && InterpolationInput)
+	{
+		FHairStrandsInterpolationInput* LocalInterpolationInput = InterpolationInput;
+		ENQUEUE_RENDER_COMMAND(FHairStrandsTick_UpdateSimulationDisable)(
+			[LocalInterpolationInput](FRHICommandListImmediate& RHICmdList)
+		{
+			for (FHairStrandsInterpolationInput::FHairGroup& HairGroup : LocalInterpolationInput->HairGroups)
+			{
+				HairGroup.bIsSimulationEnable = false;
+			}
+		});
+	}
+}
+
 void UGroomComponent::InitResources()
 {
 	ReleaseResources();
@@ -680,6 +735,8 @@ void UGroomComponent::InitResources()
 		SkeletalMeshComponent->MeshObjectCallbackData = CallbackData;
 	}
 
+	const bool bIsSimulationEnable = IsSimulationEnabled(this);
+
 	FTransform HairLocalToWorld = GetComponentTransform();
 	FTransform SkinLocalToWorld = bBindGroomToSkeletalMesh && SkeletalMeshComponent ? SkeletalMeshComponent->GetComponentTransform() : FTransform::Identity;
 	
@@ -752,6 +809,7 @@ void UGroomComponent::InitResources()
 		// For skinned groom, these value will be updated during TickComponent() call
 		InterpolationInputGroup.OutHairPositionOffset = RestHairPositionOffset;
 		InterpolationInputGroup.OutHairPreviousPositionOffset = RestHairPositionOffset;
+		InterpolationInputGroup.bIsSimulationEnable = bIsSimulationEnable;
 
 		GroupIt++;
 	}
diff --git a/Engine/Plugins/Experimental/HairStrands/Source/HairStrandsCore/Private/HairStrandsRendering.cpp b/Engine/Plugins/Experimental/HairStrands/Source/HairStrandsCore/Private/HairStrandsRendering.cpp
--- a/Engine/Plugins/Experimental/HairStrands/Source/HairStrandsCore/Private/HairStrandsRendering.cpp
+++ b/Engine/Plugins/Experimental/HairStrands/Source/HairStrandsCore/Private/HairStrandsRendering.cpp
@@ -14,8 +14,8 @@ static FAutoConsoleVariableRef CVarHairDeformationType(TEXT("r.HairStrands.Defor
 static float GHairRaytracingRadiusScale = 0;
 static FAutoConsoleVariableRef CVarHairRaytracingRadiusScale(TEXT("r.HairStrands.RaytracingRadiusScale"), GHairRaytracingRadiusScale, TEXT("Override the per instance scale factor for raytracing hair strands geometry (0: disabled, >0:enabled)"));
 
-static int32 GHairLocalInterpolation = 0;
-static FAutoConsoleVariableRef CVarHairLocalInterpolation(TEXT("r.HairStrands.LocalInterpolation"), GHairLocalInterpolation, TEXT("Disable the default skinning onto the skel mesh"));
+static int32 GHairStrandsInterpolateSimulation = 1;
+static FAutoConsoleVariableRef CVarHairInterpolateSimulation(TEXT("r.HairStrands.InterpolateSimulation"), GHairStrandsInterpolateSimulation, TEXT("Enable/disable simulation output during the hair interpolation"));
 
 static FIntVector ComputeDispatchCount(uint32 ItemCount, uint32 GroupSize)
 {
@@ -111,8 +111,8 @@ class FHairInterpolationCS : public FGlobalShader
 	class FGroupSize : SHADER_PERMUTATION_INT("PERMUTATION_GROUP_SIZE", 2);
 	class FDebug : SHADER_PERMUTATION_INT("PERMUTATION_DEBUG", 2);
 	class FDynamicGeometry : SHADER_PERMUTATION_INT("PERMUTATION_DYNAMIC_GEOMETRY", 2);
-	class FLocalInterpolation : SHADER_PERMUTATION_INT("PERMUTATION_LOCAL_INTERPOLATION", 2);
-	using FPermutationDomain = TShaderPermutationDomain<FGroupSize, FDebug, FDynamicGeometry,FLocalInterpolation>;
+	class FSimulation : SHADER_PERMUTATION_INT("PERMUTATION_SIMULATION", 2);
+	using FPermutationDomain = TShaderPermutationDomain<FGroupSize, FDebug, FDynamicGeometry, FSimulation>;
 
 	BEGIN_SHADER_PARAMETER_STRUCT(FParameters, )
 		SHADER_PARAMETER(uint32, VertexCount)
@@ -176,6 +176,7 @@ static void AddHairStrandsInterpolationPass(
 	const FVector& InHairWorldOffset, 
 	const FVector& OutHairWorldOffset,
 	const int32 LODIndex,
+	const bool bHasSimulationEnable,
 	const uint32 VertexCount,
 	const FShaderResourceViewRHIRef& RenderRestPosePositionBuffer,
 	const FShaderResourceViewRHIRef& Interpolation0Buffer,
@@ -243,7 +244,7 @@ static void AddHairStrandsInterpolationPass(
 	PermutationVector.Set<FHairInterpolationCS::FGroupSize>(GetGroupSizePermutation(GroupSize));
 	PermutationVector.Set<FHairInterpolationCS::FDebug>(bCopySimAttributesToRenderAttributes ? 1 : 0);
 	PermutationVector.Set<FHairInterpolationCS::FDynamicGeometry>(bSupportDynamicMesh ? 1 : 0);
-	PermutationVector.Set<FHairInterpolationCS::FLocalInterpolation>((GHairLocalInterpolation == 1) ? 1 : 0);
+	PermutationVector.Set<FHairInterpolationCS::FSimulation>(bHasSimulationEnable ? 1 : 0);
 
 	TShaderMap<FGlobalShaderType>* ShaderMap = GetGlobalShaderMap(ERHIFeatureLevel::SM5);
 
@@ -485,6 +486,7 @@ void ComputeHairStrandsInterpolation(
 				Output.RenderPatchedAttributeBuffer.Initialize(FHairStrandsAttributeFormat::SizeInByte, Input.RenderVertexCount, FHairStrandsAttributeFormat::Format, BUF_Static);
 			}
 
+			const bool bHasSimulationEnabled = Input.bIsSimulationEnable && GHairStrandsInterpolateSimulation;
 			check(GroupIndex < uint32(InRenHairDatas.HairGroups.Num()));
 			check(GroupIndex < uint32(InSimHairDatas.HairGroups.Num()));
 			AddHairStrandsInterpolationPass(
@@ -494,6 +496,7 @@ void ComputeHairStrandsInterpolation(
 				Input.InHairPositionOffset,
 				Input.OutHairPositionOffset,
 				LODIndex,
+				bHasSimulationEnabled,
 				Input.RenderVertexCount,
 				Input.RenderRestPosePositionBuffer->SRV,
 				Input.Interpolation0Buffer->SRV,
diff --git a/Engine/Plugins/Experimental/HairStrands/Source/HairStrandsCore/Public/GroomComponent.h b/Engine/Plugins/Experimental/HairStrands/Source/HairStrandsCore/Public/GroomComponent.h
--- a/Engine/Plugins/Experimental/HairStrands/Source/HairStrandsCore/Public/GroomComponent.h
+++ b/Engine/Plugins/Experimental/HairStrands/Source/HairStrandsCore/Public/GroomComponent.h
@@ -129,6 +129,12 @@ public:
 	FHairGroupResources HairGroupResources;
 	struct FHairStrandsInterpolationOutput* InterpolationOutput = nullptr;
 	struct FHairStrandsInterpolationInput* InterpolationInput = nullptr;
+
+protected:
+	// Used for tracking if a Niagara component is attached or not
+	virtual void OnChildAttached(USceneComponent* ChildComponent) override;
+	virtual void OnChildDetached(USceneComponent* ChildComponent) override;
+
 private:
 	void* InitializedResources;
 
diff --git a/Engine/Plugins/Experimental/HairStrands/Source/HairStrandsCore/Public/HairStrandsRendering.h b/Engine/Plugins/Experimental/HairStrands/Source/HairStrandsCore/Public/HairStrandsRendering.h
--- a/Engine/Plugins/Experimental/HairStrands/Source/HairStrandsCore/Public/HairStrandsRendering.h
+++ b/Engine/Plugins/Experimental/HairStrands/Source/HairStrandsCore/Public/HairStrandsRendering.h
@@ -32,6 +32,7 @@ struct FHairStrandsInterpolationInput
 		bool					bIsRTGeometryInitialized = false;
 		#endif
 
+		bool bIsSimulationEnable = false;
 		float HairRadius = 0;
 		float HairRaytracingRadiusScale = 0;
 		FVector InHairPositionOffset = FVector::ZeroVector;
diff --git a/Engine/Shaders/Private/HairStrands/HairStrandsInterpolation.usf b/Engine/Shaders/Private/HairStrands/HairStrandsInterpolation.usf
--- a/Engine/Shaders/Private/HairStrands/HairStrandsInterpolation.usf
+++ b/Engine/Shaders/Private/HairStrands/HairStrandsInterpolation.usf
@@ -67,16 +67,21 @@ void MainCS(uint2 DispatchThreadId : SV_DispatchThreadID)
 	const uint VertexIndex = DispatchThreadId.x + DispatchThreadId.y * DispatchCountX * GROUP_SIZE;
 	if (VertexIndex < VertexCount)
 	{
-		const uint4 InterpolationData0	= Interpolation0Buffer[VertexIndex];
-		const uint4 InterpolationData1	= Interpolation1Buffer[VertexIndex];		
-
 		uint3 GuideVertexIndices = 0;
-		GuideVertexIndices.x = InterpolationData0.x | (InterpolationData1.x << 16);
-		GuideVertexIndices.y = InterpolationData0.y | (InterpolationData1.y << 16);
-		GuideVertexIndices.z = InterpolationData0.z | (InterpolationData1.z << 16);
+		float3 GuideVertexWeights = 0;
+		#if PERMUTATION_SIMULATION == 1 || PERMUTATION_DEBUG == 1
+		{
+			const uint4 InterpolationData0	= Interpolation0Buffer[VertexIndex];
+			const uint4 InterpolationData1	= Interpolation1Buffer[VertexIndex];
+
+			GuideVertexIndices.x = InterpolationData0.x | (InterpolationData1.x << 16);
+			GuideVertexIndices.y = InterpolationData0.y | (InterpolationData1.y << 16);
+			GuideVertexIndices.z = InterpolationData0.z | (InterpolationData1.z << 16);
 
-		float3 GuideVertexWeights = float3(saturate((InterpolationData0.w & 0xFF) / 255.f), saturate(((InterpolationData0.w >> 8) & 0xFF) / 255.f), 0);
-		GuideVertexWeights.z = saturate(1.0f - (GuideVertexWeights.x + GuideVertexWeights.y));
+			GuideVertexWeights = float3(saturate((InterpolationData0.w & 0xFF) / 255.f), saturate(((InterpolationData0.w >> 8) & 0xFF) / 255.f), 0);
+			GuideVertexWeights.z = saturate(1.0f - (GuideVertexWeights.x + GuideVertexWeights.y));
+		}
+		#endif
 
 		// 1. Rigid transformation
 		// ControlPoint is in the local hair referential
@@ -85,6 +90,7 @@ void MainCS(uint2 DispatchThreadId : SV_DispatchThreadID)
 		{
 			float3 CurrOffset = 0;
 			float3 ControlPoint = f16tof32(RenderRestPosePositionBuffer[VertexIndex].xyz) + InHairPositionOffset;
+			#if PERMUTATION_SIMULATION == 1
 			[unroll]
 			for (uint K = 0; K < 3; ++K)
 			{
@@ -94,6 +100,7 @@ void MainCS(uint2 DispatchThreadId : SV_DispatchThreadID)
 				const float3 DeformedGuidePoint	= f16tof32(DeformedSimPositionBuffer[GuideIndex].xyz) + OutHairPositionOffset;
 				CurrOffset += (DeformedGuidePoint - RestGuidePoint) * GuideVertexWeights[K];
 			}
+			#endif
 			OutRenderDeformedPositionBuffer[VertexIndex] = float4(f32tof16(ControlPoint + CurrOffset - OutHairPositionOffset), RenderRestPosePositionBuffer[VertexIndex].w);
 		}
 		#endif	
@@ -106,7 +113,7 @@ void MainCS(uint2 DispatchThreadId : SV_DispatchThreadID)
 			float3 CurrOffset = 0;
 
 			// Compute the simulation offset in hair local space (i.e., triangle)
-			#if PERMUTATION_LOCAL_INTERPOLATION == 1
+			#if PERMUTATION_SIMULATION == 1
 			for (uint K = 0; K < 3; ++K)
 			{
 				const uint GuideIndex = GuideVertexIndices[K];
