Hair sim is waiting to have the animation loaded to start

diff --git a/Engine/Plugins/Experimental/HairStrands/Source/HairStrandsCore/Private/GroomComponent.cpp b/Engine/Plugins/Experimental/HairStrands/Source/HairStrandsCore/Private/GroomComponent.cpp
--- a/Engine/Plugins/Experimental/HairStrands/Source/HairStrandsCore/Private/GroomComponent.cpp
+++ b/Engine/Plugins/Experimental/HairStrands/Source/HairStrandsCore/Private/GroomComponent.cpp
@@ -709,9 +709,20 @@ void UGroomComponent::OnChildDetached(USceneComponent* ChildComponent)
 	}
 }
 
+void UGroomComponent::ResetSimulation()
+{
+	USkeletalMeshComponent* SkeletalMeshComponent = GetAttachParent() ? Cast<USkeletalMeshComponent>(GetAttachParent()) : nullptr;
+	if (SkeletalMeshComponent)
+	{
+		SkeletalMeshComponent->OnBoneTransformsFinalized.RemoveDynamic(this, &UGroomComponent::ResetSimulation);
+		bResetSimulation = false;
+	}
+}
+
 void UGroomComponent::InitResources()
 {
 	ReleaseResources();
+	bResetSimulation = true;
 
 	if (!GroomAsset || GroomAsset->GetNumHairGroups() == 0)
 		return;
@@ -978,6 +989,13 @@ void UGroomComponent::OnRegister()
 {
 	Super::OnRegister();
 
+	// Insure the ticking of the Groom component always happens after the skeletalMeshComponent.
+	USkeletalMeshComponent* SkeletalMeshComponent = GetAttachParent() ? Cast<USkeletalMeshComponent>(GetAttachParent()) : nullptr;
+	if (SkeletalMeshComponent)
+	{
+		SkeletalMeshComponent->OnBoneTransformsFinalized.AddDynamic(this, &UGroomComponent::ResetSimulation);
+	}
+
 	const EWorldType::Type WorldType = GetWorld() ? EWorldType::Type(GetWorld()->WorldType) : EWorldType::None;
 	const uint64 Id = ComponentId.PrimIDValue;
 
diff --git a/Engine/Plugins/Experimental/HairStrands/Source/HairStrandsCore/Public/GroomComponent.h b/Engine/Plugins/Experimental/HairStrands/Source/HairStrandsCore/Public/GroomComponent.h
--- a/Engine/Plugins/Experimental/HairStrands/Source/HairStrandsCore/Public/GroomComponent.h
+++ b/Engine/Plugins/Experimental/HairStrands/Source/HairStrandsCore/Public/GroomComponent.h
@@ -57,6 +57,13 @@ public:
 	UPROPERTY(EditAnywhere, BlueprintReadWrite, Category = "Groom Rendering", meta = (ClampMin = "0.0001", UIMin = "0.01", UIMax = "10.0"))
 	float HairRaytracingRadiusScale;
 
+	/** Boolean to check when animation has been loaded */
+	bool bResetSimulation;
+
+	/** Listen for the animation event to trigger the sim */
+	UFUNCTION()
+	void ResetSimulation();
+
 	//~ Begin UActorComponent Interface.
 	virtual void OnRegister() override;
 	virtual void OnUnregister() override;
diff --git a/Engine/Plugins/Experimental/HairStrands/Source/HairStrandsNiagara/Private/NiagaraDataInterfaceHairStrands.cpp b/Engine/Plugins/Experimental/HairStrands/Source/HairStrandsNiagara/Private/NiagaraDataInterfaceHairStrands.cpp
--- a/Engine/Plugins/Experimental/HairStrands/Source/HairStrandsNiagara/Private/NiagaraDataInterfaceHairStrands.cpp
+++ b/Engine/Plugins/Experimental/HairStrands/Source/HairStrandsNiagara/Private/NiagaraDataInterfaceHairStrands.cpp
@@ -15,7 +15,7 @@
 
 #define LOCTEXT_NAMESPACE "NiagaraDataInterfaceHairStrands"
 
-static const int32 ResetTick = 4;
+static const int32 MaxDelay = 5;
 
 //------------------------------------------------------------------------------------------------------------
 
@@ -408,7 +408,7 @@ struct FNDIHairStrandsParametersCS : public FNiagaraDataInterfaceParametersCS
 			FShaderResourceViewRHIRef RootBarycentricCoordinatesSRV = HairStrandsBuffer->RootBarycentricCoordinatesBuffer.SRV;
 
 			int32 HasRootAttachedValue = 0;
-			bool NeedSimReset = (ProxyData->TickCount <= ResetTick);
+			bool bNeedSimReset = (ProxyData->TickCount <= ProxyData->ResetTick);
 			FVector RestRootOffsetValue = FVector(0, 0, 0);
 			FVector DeformedRootOffsetValue = FVector(0, 0, 0);
 			FVector DeformedPositionOffsetValue = HairStrandsBuffer->SourceRestResources->PositionOffset;
@@ -427,7 +427,7 @@ struct FNDIHairStrandsParametersCS : public FNiagaraDataInterfaceParametersCS
 				RootBarycentricCoordinatesSRV = HairStrandsBuffer->SourceRootResources->MeshProjectionLODs[0].RootTriangleBarycentricBuffer.SRV;
 
 				HasRootAttachedValue = 1;
-				NeedSimReset = (HairStrandsBuffer->SourceRootResources->MeshProjectionLODs[0].Status != FHairStrandsRootResource::FMeshProjectionLOD::EStatus::Completed);
+				//NeedSimReset = (HairStrandsBuffer->SourceRootResources->MeshProjectionLODs[0].Status != FHairStrandsRootResource::FMeshProjectionLOD::EStatus::Completed);
 				RestRootOffsetValue = HairStrandsBuffer->SourceRootResources->MeshProjectionLODs[0].RestRootOffset;
 				DeformedRootOffsetValue = ProxyData->HairStrandsBuffer->SourceRootResources->MeshProjectionLODs[0].DeformedRootOffset;
 			}
@@ -451,7 +451,7 @@ struct FNDIHairStrandsParametersCS : public FNiagaraDataInterfaceParametersCS
 			SetShaderValue(RHICmdList, ComputeShaderRHI, BoxCenter, ProxyData->BoxCenter);
 			SetShaderValue(RHICmdList, ComputeShaderRHI, BoxExtent, ProxyData->BoxExtent);
 
-			SetShaderValue(RHICmdList, ComputeShaderRHI, ResetSimulation, NeedSimReset);
+			SetShaderValue(RHICmdList, ComputeShaderRHI, ResetSimulation, bNeedSimReset);
 			SetShaderValue(RHICmdList, ComputeShaderRHI, HasRootAttached, HasRootAttachedValue);
 			SetShaderValue(RHICmdList, ComputeShaderRHI, RestRootOffset, RestRootOffsetValue);
 			SetShaderValue(RHICmdList, ComputeShaderRHI, DeformedRootOffset, DeformedRootOffsetValue);
@@ -567,6 +567,8 @@ void FNDIHairStrandsProxy::ConsumePerInstanceDataFromGameThread(void* PerInstanc
 		TargetData->BoxCenter = SourceData->BoxCenter;
 		TargetData->BoxExtent = SourceData->BoxExtent;
 		TargetData->TickCount = SourceData->TickCount;
+		TargetData->ResetTick = SourceData->ResetTick;
+		TargetData->ForceReset = SourceData->ForceReset;
 		TargetData->HairStrandsBuffer = SourceData->HairStrandsBuffer;
 	}
 	else
@@ -595,6 +597,8 @@ void FNDIHairStrandsProxy::InitializePerInstanceData(const FNiagaraSystemInstanc
 	TargetData->BoxCenter = BoxCenter;
 	TargetData->BoxExtent = BoxExtent;
 	TargetData->TickCount = 0;
+	TargetData->ResetTick = MaxDelay;
+	TargetData->ForceReset = false;
 	TargetData->WorldTransform = WorldTransform;
 }
 
@@ -733,6 +737,8 @@ bool UNiagaraDataInterfaceHairStrands::InitPerInstanceData(void* PerInstanceData
 			InstanceData->BoxCenter = BoxCenter;
 			InstanceData->BoxExtent = BoxExtent;
 			InstanceData->TickCount = 0;
+			InstanceData->ForceReset = false;
+			InstanceData->ResetTick = MaxDelay;
 
 			//UE_LOG(LogHairStrands, Log, TEXT("Num Strands = %d | Strand Size = %d | Num Vertices = %d | Min = %f %f %f | Max = %f %f %f | Transform = %s"), NumStrands, StrandSize,
 			//	StrandsDatas->GetNumPoints(), StrandsBox.Min[0], StrandsBox.Min[1], StrandsBox.Min[2], StrandsBox.Max[0], StrandsBox.Max[1], StrandsBox.Max[2], *InstanceData->WorldTransform.ToString());
@@ -779,11 +785,20 @@ bool UNiagaraDataInterfaceHairStrands::PerInstanceTick(void* PerInstanceData, FN
 	FHairStrandsRestResource* StrandsRestResource = nullptr;
 	FHairStrandsDeformedResource* StrandsDeformedResource = nullptr;
 	FHairStrandsRootResource* StrandsRootResource = nullptr;
-	InstanceData->TickCount = FMath::Min(ResetTick+1,InstanceData->TickCount+1);
+	InstanceData->TickCount = FMath::Min(MaxDelay+1,InstanceData->TickCount+1);
 
 	ExtractDatasAndResources(SystemInstance, StrandsDatas, StrandsRestResource, StrandsDeformedResource, StrandsRootResource);
 	InstanceData->HairStrandsBuffer->SetHairAsset(StrandsDatas, StrandsRestResource, StrandsDeformedResource, StrandsRootResource);
 
+	if (SourceComponent != nullptr)
+	{
+		if (!InstanceData->ForceReset && !SourceComponent->bResetSimulation)
+		{
+			InstanceData->ResetTick = FMath::Min(MaxDelay, InstanceData->TickCount);
+		}
+		InstanceData->ForceReset = SourceComponent->bResetSimulation;
+	}
+
 	bool RequireReset = false;
 	if (StrandsDatas)
 	{
@@ -2574,6 +2589,8 @@ void UNiagaraDataInterfaceHairStrands::ProvidePerInstanceDataForRenderThread(voi
 	RenderThreadData->BoxCenter = GameThreadData->BoxCenter;
 	RenderThreadData->BoxExtent = GameThreadData->BoxExtent;
 	RenderThreadData->TickCount = GameThreadData->TickCount;
+	RenderThreadData->ResetTick = GameThreadData->ResetTick;
+	RenderThreadData->ForceReset = GameThreadData->ForceReset;
 }
 
 FNiagaraDataInterfaceParametersCS*
diff --git a/Engine/Plugins/Experimental/HairStrands/Source/HairStrandsNiagara/Public/NiagaraDataInterfaceHairStrands.h b/Engine/Plugins/Experimental/HairStrands/Source/HairStrandsNiagara/Public/NiagaraDataInterfaceHairStrands.h
--- a/Engine/Plugins/Experimental/HairStrands/Source/HairStrandsNiagara/Public/NiagaraDataInterfaceHairStrands.h
+++ b/Engine/Plugins/Experimental/HairStrands/Source/HairStrandsNiagara/Public/NiagaraDataInterfaceHairStrands.h
@@ -99,6 +99,12 @@ struct FNDIHairStrandsData
 	/** Tick Count*/
 	int32 TickCount;
 
+	/** Force reset simulation */
+	bool ForceReset;
+
+	/** Reset Tick*/
+	int32 ResetTick;
+
 	/** Strands Gpu buffer */
 	FNDIHairStrandsBuffer* HairStrandsBuffer;
 };
