Change root position (used for hair projection on mesh) to 32-bits float, as they don't use relative position, and this could cause precision issue.

diff --git a/Engine/Plugins/Experimental/HairStrands/Source/HairStrandsCore/Private/GroomAsset.cpp b/Engine/Plugins/Experimental/HairStrands/Source/HairStrandsCore/Private/GroomAsset.cpp
--- a/Engine/Plugins/Experimental/HairStrands/Source/HairStrandsCore/Private/GroomAsset.cpp
+++ b/Engine/Plugins/Experimental/HairStrands/Source/HairStrandsCore/Private/GroomAsset.cpp
@@ -126,13 +126,13 @@ FHairStrandsRootResource::FHairStrandsRootResource(const FHairStrandsDatas* Hair
 			N0 = FVector(0, 0, 1);
 		}
 
-		FVector4_16 P;
+		FHairStrandsRootPositionFormat::Type P;
 		P.X = P0.X;
 		P.Y = P0.Y;
 		P.Z = P0.Z;
 		P.W = 1;
 
-		FVector4_16 N;
+		FHairStrandsRootNormalFormat::Type N;
 		N.X = N0.X;
 		N.Y = N0.Y;
 		N.Z = N0.Z;
diff --git a/Engine/Plugins/Experimental/HairStrands/Source/HairStrandsCore/Public/HairStrandsDatas.h b/Engine/Plugins/Experimental/HairStrands/Source/HairStrandsCore/Public/HairStrandsDatas.h
--- a/Engine/Plugins/Experimental/HairStrands/Source/HairStrandsCore/Public/HairStrandsDatas.h
+++ b/Engine/Plugins/Experimental/HairStrands/Source/HairStrandsCore/Public/HairStrandsDatas.h
@@ -219,11 +219,11 @@ struct FHairStrandsCurveTriangleBarycentricFormat
 
 struct FHairStrandsRootPositionFormat
 {
-	typedef FVector4_16 Type;
+	typedef FVector4 Type;
 	static const uint32 ComponentCount = 1;
 	static const uint32 SizeInByte = sizeof(Type);
-	static const EVertexElementType VertexElementType = VET_UShort4;
-	static const EPixelFormat Format = PF_FloatRGBA;
+	static const EVertexElementType VertexElementType = VET_Float4;
+	static const EPixelFormat Format = PF_A32B32G32R32F;
 };
 
 struct FHairStrandsRootNormalFormat
