Add quality level (low/medium/high) for building hair/groom interpolation data.

Default is high, and maps onto the previous metric/technique we use.
Low is a simple nearest neighbors search based on distance, using a grid look up for speed.

diff --git a/Engine/Plugins/Experimental/HairStrands/Source/HairStrandsCore/Private/GroomAsset.cpp b/Engine/Plugins/Experimental/HairStrands/Source/HairStrandsCore/Private/GroomAsset.cpp
--- a/Engine/Plugins/Experimental/HairStrands/Source/HairStrandsCore/Private/GroomAsset.cpp
+++ b/Engine/Plugins/Experimental/HairStrands/Source/HairStrandsCore/Private/GroomAsset.cpp
@@ -338,7 +338,7 @@ void UGroomAsset::PostLoad()
 	check(HairGroupsData.Num() > 0);
 	if (HairGroupsData[0].HairSimulationData.GetNumCurves() == 0 || HairGroupsData[0].HairInterpolationData.Num() == 0)
 	{
-		FGroomBuilder::BuildData(this);
+		FGroomBuilder::BuildData(this, 2);
 	}
 
 	if (!IsTemplate())
diff --git a/Engine/Plugins/Experimental/HairStrands/Source/HairStrandsCore/Private/GroomBuilder.cpp b/Engine/Plugins/Experimental/HairStrands/Source/HairStrandsCore/Private/GroomBuilder.cpp
--- a/Engine/Plugins/Experimental/HairStrands/Source/HairStrandsCore/Private/GroomBuilder.cpp
+++ b/Engine/Plugins/Experimental/HairStrands/Source/HairStrandsCore/Private/GroomBuilder.cpp
@@ -380,183 +380,464 @@ namespace HairInterpolationBuilder
 		B = Temp;
 	}
 
-	void BuildInterpolationData(FHairStrandsInterpolationDatas& InterpolationData,
-		const FHairStrandsDatas& SimStrandsData,
-		const FHairStrandsDatas& RenStrandsData)
+	struct FMetrics
 	{
-		TRACE_CPUPROFILER_EVENT_SCOPE(HairInterpolationBuilder::BuildInterpolationData);
+		static const uint32 ClosestCount = 3;
+		static const uint32 KGuideCount = 6;
+		float KMinMetrics[KGuideCount];
+		int32 KClosestGuideIndices[KGuideCount];
+	};
 
-		InterpolationData.SetNum(RenStrandsData.GetNumPoints());
+	// Randomize influence guide to break interpolation coherence, and create a more random/natural pattern
+	static void SelectFinalGuides(int32* ClosestGuideIndices, FRandomStream& Random, FMetrics& Metric)
+	{
+		const static bool bRandomizeInterpolation = false;
+		const static bool bUseUniqueGuide = false;
 
-		typedef TArray<FHairRoot> FRoots;
+		uint32 RandIndex0 = 0;
+		uint32 RandIndex1 = 1;
+		uint32 RandIndex2 = 2;
+		if (bRandomizeInterpolation)
+		{
+			do
+			{
+				RandIndex0 = Random.RandRange(0, FMetrics::KGuideCount - 1);
+				RandIndex1 = Random.RandRange(0, FMetrics::KGuideCount - 1);
+				RandIndex2 = Random.RandRange(0, FMetrics::KGuideCount - 1);
+
+			} while (RandIndex0 == RandIndex1 || RandIndex0 == RandIndex2 || RandIndex1 == RandIndex2);
+		}
 
-		// Extract strand roots
-		auto ExtractRoot = [](const FHairStrandsDatas& InData, TArray<FHairRoot>& OutRoots)
+		ClosestGuideIndices[0] = Metric.KClosestGuideIndices[RandIndex0];
+		ClosestGuideIndices[1] = Metric.KClosestGuideIndices[RandIndex1];
+		ClosestGuideIndices[2] = Metric.KClosestGuideIndices[RandIndex2];
+
+		if (bUseUniqueGuide)
+		{
+			ClosestGuideIndices[1] = Metric.KClosestGuideIndices[RandIndex0];
+			ClosestGuideIndices[2] = Metric.KClosestGuideIndices[RandIndex0];
+		}
+
+
+		float MinMetrics[FMetrics::KGuideCount];
+		MinMetrics[0] = Metric.KMinMetrics[RandIndex0];
+		MinMetrics[1] = Metric.KMinMetrics[RandIndex1];
+		MinMetrics[2] = Metric.KMinMetrics[RandIndex2];
+
+
+		while (!(MinMetrics[0] <= MinMetrics[1] && MinMetrics[1] <= MinMetrics[2]))
 		{
-			const uint32 CurveCount = InData.StrandsCurves.Num();
-			OutRoots.Reserve(CurveCount);
-			for (uint32 CurveIndex = 0; CurveIndex < CurveCount; ++CurveIndex)
+			if (MinMetrics[0] > MinMetrics[1])
+			{
+				SwapValue(MinMetrics[0], MinMetrics[1]);
+				SwapValue(ClosestGuideIndices[0], ClosestGuideIndices[1]);
+			}
+
+			if (MinMetrics[1] > MinMetrics[2])
 			{
-				const uint32 PointOffset = InData.StrandsCurves.CurvesOffset[CurveIndex];
-				const uint32 PointCount  = InData.StrandsCurves.CurvesCount[CurveIndex];
-				const float  CurveLength = InData.StrandsCurves.CurvesLength[CurveIndex] * InData.StrandsCurves.MaxLength;
-				check(PointCount > 1);
-				const FVector& P0 = InData.StrandsPoints.PointsPosition[PointOffset];
-				const FVector& P1 = InData.StrandsPoints.PointsPosition[PointOffset+1];
-				FVector N = (P1 - P0).GetSafeNormal();
-
-				// Fallback in case the initial points are too close (this happens on certain assets)
-				if (FVector::DotProduct(N, N) == 0)
+				SwapValue(MinMetrics[1], MinMetrics[2]);
+				SwapValue(ClosestGuideIndices[1], ClosestGuideIndices[2]);
+			}
+		}
+
+		// If there less than 3 valid guides, fill the rest with existing valid guides
+		// This can happen due to the normal-orientation based rejection above
+		if (ClosestGuideIndices[1] < 0)
+		{
+			ClosestGuideIndices[1] = ClosestGuideIndices[0];
+			MinMetrics[1] = MinMetrics[0];
+		}
+		if (ClosestGuideIndices[2] < 0)
+		{
+			ClosestGuideIndices[2] = ClosestGuideIndices[1];
+			MinMetrics[2] = MinMetrics[1];
+		}
+
+		check(MinMetrics[0] <= MinMetrics[1]);
+		check(MinMetrics[1] <= MinMetrics[2]);
+	};
+
+	// Simple closest distance metric
+	static void ComputeSimpleMetric(
+		FMetrics& Metrics1, 
+		const FHairRoot& RenRoot, 
+		const FHairRoot& GuideRoot, 
+		const int32 RenCurveIndex,
+		const int32 SimCurveIndex)
+	{
+		const float Metric = FVector::Dist(GuideRoot.Position, RenRoot.Position);
+		if (Metric < Metrics1.KMinMetrics[FMetrics::ClosestCount - 1])
+		{
+			int32 LastGuideIndex = SimCurveIndex;
+			float LastMetric = Metric;
+			for (uint32 Index = 0; Index < FMetrics::ClosestCount; ++Index)
+			{
+				if (Metric < Metrics1.KMinMetrics[Index])
 				{
-					N = FVector(0, 0, 1);
+					SwapValue(Metrics1.KClosestGuideIndices[Index], LastGuideIndex);
+					SwapValue(Metrics1.KMinMetrics[Index], LastMetric);
 				}
-				OutRoots.Add({ P0, PointCount, N, PointOffset, CurveLength });
 			}
-		};
+		}
+	}
 
-		FRoots RenRoots, SimRoots;
-		ExtractRoot(RenStrandsData, RenRoots);
-		ExtractRoot(SimStrandsData, SimRoots);
+	// Complex pairing metric
+	static void ComputeAdvandedMetric(FMetrics& Metrics0,
+		const FHairStrandsDatas& RenStrandsData,
+		const FHairStrandsDatas& SimStrandsData,
+		const int32 RenCurveIndex,
+		const int32 SimCurveIndex)
+	{
+		const float Metric = 1.0 - ComputeCurvesMetric<16>(RenStrandsData, RenCurveIndex, SimStrandsData, SimCurveIndex, 0.0f, 1.0f, 1.0f);
+		if (Metric < Metrics0.KMinMetrics[FMetrics::KGuideCount - 1])
+		{
+			int32 LastGuideIndex = SimCurveIndex;
+			float LastMetric = Metric;
+			for (uint32 Index = 0; Index < FMetrics::KGuideCount; ++Index)
+			{
+				if (Metric < Metrics0.KMinMetrics[Index])
+				{
+					SwapValue(Metrics0.KClosestGuideIndices[Index], LastGuideIndex);
+					SwapValue(Metrics0.KMinMetrics[Index], LastMetric);
+				}
+			}
+		}
+	}
 
-		// Find k-closest guide:
-		// N hairs, M guides. Complexity = MxN
-		// #hair_todo: build an acceleration structure for fast(er) look up or at least make it run on parallel (e.g. with ParallelFor)
-		uint32 TotalInvalidInterpolationCount = 0;
-		const static float MinWeightDistance = 0.0001f;
-		const static uint32 GuideCount = 3;
-		const static uint32 KGuideCount = GuideCount * 2;
+	struct FRootsGrid
+	{
+		FVector MinBound = FVector::ZeroVector;
+		FVector MaxBound = FVector::ZeroVector;
+		
+		const uint32 MaxLookupDistance = 2;
+		const FIntVector GridResolution = FIntVector(32, 32, 32);
 
-		const static bool bRandomizeInterpolation = false;
-		const static bool bUseUniqueGuide = false;
-		const static bool bPrintDebugMetric = false;
+		TArray<int32> GridIndirection;
+		TArray<TArray<int32>> RootIndices;
 
-		FRandomStream Random;
-		const uint32 RenCurveCount = RenStrandsData.GetNumCurves();
-		const uint32 SimCurveCount = SimStrandsData.GetNumCurves();
+		FORCEINLINE FIntVector ClampToVolume(const FIntVector& CellCoord) const
+		{
+			return FIntVector(
+				FMath::Clamp(CellCoord.X, 0, GridResolution.X - 1),
+				FMath::Clamp(CellCoord.Y, 0, GridResolution.Y - 1),
+				FMath::Clamp(CellCoord.Z, 0, GridResolution.Z - 1));
+		}
 
-		ParallelFor(RenCurveCount, 
-		[
-			RenCurveCount, &RenRoots, &RenStrandsData,
-			SimCurveCount, &SimRoots, &SimStrandsData, 
-			&TotalInvalidInterpolationCount,  
-			&InterpolationData, 
-			&Random
-		] (uint32 RenCurveIndex) 
-		//for (uint32 RenCurveIndex = 0; RenCurveIndex < RenCurveCount; ++RenCurveIndex)
+		FORCEINLINE FIntVector ToCellCoord(const FVector& P) const
 		{
-			const FHairRoot& StrandRoot = RenRoots[RenCurveIndex];
+			const FVector F = ((P - MinBound) / (MaxBound - MinBound));
+			const FIntVector CellCoord = FIntVector(FMath::FloorToInt(F.X * GridResolution.X), FMath::FloorToInt(F.Y * GridResolution.Y), FMath::FloorToInt(F.Z * GridResolution.Z));			
+			return ClampToVolume(CellCoord);
+		}
+
+		uint32 ToIndex(const FIntVector& CellCoord) const
+		{
+			uint32 CellIndex = CellCoord.X + CellCoord.Y * GridResolution.X + CellCoord.Z * GridResolution.X * GridResolution.Y;
+			check(CellIndex < uint32(GridIndirection.Num()));
+			return CellIndex;
+		}
 
-			int32 ClosestGuideIndices[GuideCount];
+		void InsertRoots(TArray<FHairRoot>& Roots, const FVector& InMinBound, const FVector& InMaxBound)
+		{
+			MinBound = InMinBound;
+			MaxBound = InMaxBound;
+			GridIndirection.SetNumZeroed(GridResolution.X*GridResolution.Y*GridResolution.Z);
+			RootIndices.Empty();
+			RootIndices.AddDefaulted(); // Add a default empty list for the null element
+
+			const uint32 RootCount = Roots.Num();
+			for (uint32 RootIt = 0; RootIt < RootCount; ++RootIt)
 			{
-				float KMinMetrics[KGuideCount];
-				int32 KClosestGuideIndices[KGuideCount];
-				for (uint32 ClearIt = 0; ClearIt < KGuideCount; ++ClearIt)
+				const FHairRoot& Root = Roots[RootIt];
+				const FIntVector CellCoord = ToCellCoord(Root.Position);
+				const uint32 Index = ToIndex(CellCoord);
+				if (GridIndirection[Index] == 0)
 				{
-					KMinMetrics[ClearIt] = FLT_MAX;
-					KClosestGuideIndices[ClearIt] = -1;
+					GridIndirection[Index] = RootIndices.Num();
+					RootIndices.AddDefaulted();
 				}
+				
+				TArray<int32>& CellGuideIndices = RootIndices[GridIndirection[Index]];
+				CellGuideIndices.Add(RootIt);
+			}
+		}
+
+		void FindClosestRoots(
+			const uint32 RenCurveIndex,
+			const TArray<FHairRoot>& RenRoots,
+			const TArray<FHairRoot>& SimRoots,
+			const FHairStrandsDatas& RenStrandsData,
+			const FHairStrandsDatas& SimStrandsData,
+			FRandomStream& Random,
+			int32* ClosestGuideIndices) const
+		{
+			const FHairRoot& RenRoot = RenRoots[RenCurveIndex];
+			const FIntVector PointCoord = ToCellCoord(RenRoot.Position);
+
+			FMetrics Metrics;
+			for (uint32 ClearIt = 0; ClearIt < FMetrics::ClosestCount; ++ClearIt)
+			{
+				Metrics.KMinMetrics[ClearIt] = FLT_MAX;
+				Metrics.KClosestGuideIndices[ClearIt] = -1;
+			}
 
-				for (uint32 SimCurveIndex = 0; SimCurveIndex < SimCurveCount; ++SimCurveIndex)
+			// Up to 5x5X5 lookup within a 32x32x32 volume
+			// * Tracks two metrics: a complex (pairing) one, and simple (closest) one
+			// * If the complex one fails to find a valid guide, we fallback on the simplest one
+			//for (uint32 R = 1; R < MaxLookupDistance; ++R)
+			{
+				int32 Offset = MaxLookupDistance;
+				for (int32 X = -Offset; X <= Offset; ++X)
 				{
-					const FHairRoot& GuideRoot = SimRoots[SimCurveIndex];
-					//if ((StrandRoot.Position - GuideRoot.Position).Size() < 5.0f)
+					for (int32 Y = -Offset; Y <= Offset; ++Y)
 					{
-						const float Metric = 1.0 - ComputeCurvesMetric<16>(RenStrandsData, RenCurveIndex, SimStrandsData, SimCurveIndex, 0.0f, 1.0f, 1.0f);
-						if (Metric < KMinMetrics[KGuideCount - 1])
+						for (int32 Z = -Offset; Z <= Offset; ++Z)
 						{
-							int32 LastGuideIndex = SimCurveIndex;
-							float LastMetric = Metric;
-							for (uint32 Index = 0; Index < KGuideCount; ++Index)
+							const FIntVector CellCoord = ClampToVolume(PointCoord + FIntVector(X, Y, Z));
+							const uint32 CellIndex = ToIndex(CellCoord);
+
+							if (GridIndirection[CellIndex] == 0)
+								continue;
+
+							const TArray<int32>& Elements = RootIndices[GridIndirection[CellIndex]];
+					
+							for (int32 SimCurveIndex : Elements)
 							{
-								if (Metric < KMinMetrics[Index])
+								const FHairRoot& GuideRoot = SimRoots[SimCurveIndex];
 								{
-									SwapValue(KClosestGuideIndices[Index], LastGuideIndex);
-									SwapValue(KMinMetrics[Index], LastMetric);
+									ComputeSimpleMetric(Metrics, RenRoot, GuideRoot, RenCurveIndex, SimCurveIndex);
 								}
 							}
 						}
 					}
 				}
 
-				// Debug
-				//if (bPrintDebugMetric)
+				// Early out if we have found all closest guide during a ring/layer step.
+				// This early out is not conservative, as the complex metric might find better guides one or multiple step further.
+				//if (Metrics.KClosestGuideIndices[FMetrics::ClosestCount - 1] != -1 && R >= 2)
 				//{
-				//	const FHairInterpolationMetric ClosestMetric = ComputeInterpolationMetric(StrandRoot, SimRoots[KClosestGuideIndices[0]]);
-				//	const float Threshold = 20;
-				//	if (ClosestMetric.Metric > Threshold)
-				//	{
-				//		PrintInterpolationMetric(ClosestMetric);
-				//		++TotalInvalidInterpolationCount;
-				//	}
+				//	break;
 				//}
+			}
 
-				// Randomize influence guide to break interpolation coherence, and create a more random/natural pattern
-				{
-					uint32 RandIndex0 = 0;
-					uint32 RandIndex1 = 1;
-					uint32 RandIndex2 = 2;
-					if (bRandomizeInterpolation)
-					{
-						do
-						{
-							RandIndex0 = Random.RandRange(0, KGuideCount - 1);
-							RandIndex1 = Random.RandRange(0, KGuideCount - 1);
-							RandIndex2 = Random.RandRange(0, KGuideCount - 1);
+			// If no valid guide have been found, switch to a simpler metric
+			SelectFinalGuides(ClosestGuideIndices, Random, Metrics);
 
-						} while (RandIndex0 == RandIndex1 || RandIndex0 == RandIndex2 || RandIndex1 == RandIndex2);
-					}
-					ClosestGuideIndices[0] = KClosestGuideIndices[RandIndex0];
-					ClosestGuideIndices[1] = KClosestGuideIndices[RandIndex1];
-					ClosestGuideIndices[2] = KClosestGuideIndices[RandIndex2];
+			check(ClosestGuideIndices[0] >= 0);
+			check(ClosestGuideIndices[1] >= 0);
+			check(ClosestGuideIndices[2] >= 0);
+		}
 
-					if (bUseUniqueGuide)
-					{
-						ClosestGuideIndices[1] = KClosestGuideIndices[RandIndex0];
-						ClosestGuideIndices[2] = KClosestGuideIndices[RandIndex0];
-					}
 
-					float MinMetrics[GuideCount];
-					MinMetrics[0] = KMinMetrics[RandIndex0];
-					MinMetrics[1] = KMinMetrics[RandIndex1];
-					MinMetrics[2] = KMinMetrics[RandIndex2];
+		void FindBestClosestRoots(
+			const uint32 RenCurveIndex,
+			const TArray<FHairRoot>& RenRoots,
+			const TArray<FHairRoot>& SimRoots,
+			const FHairStrandsDatas& RenStrandsData,
+			const FHairStrandsDatas& SimStrandsData,
+			FRandomStream& Random,
+			int32* ClosestGuideIndices) const 
+		{
+			const FHairRoot& RenRoot = RenRoots[RenCurveIndex];
+			const FIntVector PointCoord = ToCellCoord(RenRoot.Position);
 
+			FMetrics Metrics0;
+			FMetrics Metrics1;
+			for (uint32 ClearIt = 0; ClearIt < FMetrics::ClosestCount; ++ClearIt)
+			{
+				Metrics0.KMinMetrics[ClearIt] = FLT_MAX;
+				Metrics0.KClosestGuideIndices[ClearIt] = -1;
+				Metrics1.KMinMetrics[ClearIt] = FLT_MAX;
+				Metrics1.KClosestGuideIndices[ClearIt] = -1;
+			}
 
-					while (!(MinMetrics[0] <= MinMetrics[1] && MinMetrics[1] <= MinMetrics[2]))
+			// Up to 7x7X7 lookup within a 32x32x32 volume
+			// * Tracks two metrics: a complex (pairing) one, and simple (closest) one
+			// * If the complex one fails to find a valid guide, we fallback on the simplest one
+			//for (uint32 R = 1; R < MaxLookupDistance; ++R)
+			{
+				int32 Offset = MaxLookupDistance;
+				for (int32 X = -Offset; X <= Offset; ++X)
+				{
+					for (int32 Y = -Offset; Y <= Offset; ++Y) 
 					{
-						if (MinMetrics[0] > MinMetrics[1])
+						for (int32 Z = -Offset; Z <= Offset; ++Z)
 						{
-							SwapValue(MinMetrics[0], MinMetrics[1]);
-							SwapValue(ClosestGuideIndices[0], ClosestGuideIndices[1]);
-						}
+							const FIntVector CellCoord = ClampToVolume(PointCoord + FIntVector(X, Y, Z));
+							const uint32 CellIndex = ToIndex(CellCoord);
 
-						if (MinMetrics[1] > MinMetrics[2])
-						{
-							SwapValue(MinMetrics[1], MinMetrics[2]);
-							SwapValue(ClosestGuideIndices[1], ClosestGuideIndices[2]);
+							if (GridIndirection[CellIndex] == 0)
+								continue;
+
+							const TArray<int32>& Elements = RootIndices[GridIndirection[CellIndex]];
+
+							for (int32 SimCurveIndex : Elements)
+							{
+								const FHairRoot& GuideRoot = SimRoots[SimCurveIndex];
+								{
+									ComputeSimpleMetric(Metrics1, RenRoot, GuideRoot, RenCurveIndex, SimCurveIndex);
+									ComputeAdvandedMetric(Metrics0, RenStrandsData, SimStrandsData, RenCurveIndex, SimCurveIndex);
+								}
+							}
 						}
 					}
+				}
 
-					// If there less than 3 valid guides, fill the rest with existing valid guides
-					// This can happen due to the normal-orientation based rejection above
-					if (ClosestGuideIndices[1] < 0)
-					{
-						ClosestGuideIndices[1] = ClosestGuideIndices[0];
-						MinMetrics[1] = MinMetrics[0];
-					}
-					if (ClosestGuideIndices[2] < 0)
-					{
-						ClosestGuideIndices[2] = ClosestGuideIndices[1];
-						MinMetrics[2] = MinMetrics[1];
-					}
+				// Early out if we have found all closest guide during a ring/layer step.
+				// This early out is not conservative, as the complex metric might find better guides one or multiple step further.
+				//if (Metrics0.KClosestGuideIndices[FMetrics::ClosestCount - 1] != -1 && R >= 2)
+				//{
+				//	break;
+				//}
+			}
 
-					check(ClosestGuideIndices[0] >= 0);
-					check(ClosestGuideIndices[1] >= 0);
-					check(ClosestGuideIndices[2] >= 0);
-					check(MinMetrics[0] <= MinMetrics[1]);
-					check(MinMetrics[1] <= MinMetrics[2]);
-				}
+			// If no valid guide have been found, switch to a simpler metric
+			SelectFinalGuides(ClosestGuideIndices, Random, Metrics0);
+			if (ClosestGuideIndices[0] == -1)
+			{
+				SelectFinalGuides(ClosestGuideIndices, Random, Metrics1);
+			}
+
+			check(ClosestGuideIndices[0] >= 0);
+			check(ClosestGuideIndices[1] >= 0);
+			check(ClosestGuideIndices[2] >= 0);
+		}
+	};
+
+	static void FindBestRoots(
+		const uint32 RenCurveIndex,
+		const TArray<FHairRoot>& RenRoots,
+		const TArray<FHairRoot>& SimRoots,
+		const FHairStrandsDatas& RenStrandsData,
+		const FHairStrandsDatas& SimStrandsData,
+		FRandomStream& Random,
+		int32* ClosestGuideIndices)
+	{
+		FMetrics Metrics;
+		for (uint32 ClearIt = 0; ClearIt < FMetrics::ClosestCount; ++ClearIt)
+		{
+			Metrics.KMinMetrics[ClearIt] = FLT_MAX;
+			Metrics.KClosestGuideIndices[ClearIt] = -1;
+		}
+
+		for(const FHairRoot& GuideRoot : SimRoots)
+		{
+			const uint32 SimCurveIndex = GuideRoot.Index;
+			ComputeAdvandedMetric(Metrics, RenStrandsData, SimStrandsData, RenCurveIndex, SimCurveIndex);
+		}
+			
+		SelectFinalGuides(ClosestGuideIndices, Random, Metrics);
+
+		check(ClosestGuideIndices[0] >= 0);
+		check(ClosestGuideIndices[1] >= 0);
+		check(ClosestGuideIndices[2] >= 0);
+	}
+
+
+	// Extract strand roots
+	static void ExtractRoots(const FHairStrandsDatas& InData, TArray<FHairRoot>& OutRoots, FVector& MinBound, FVector& MaxBound)
+	{
+		MinBound = FVector(FLT_MAX, FLT_MAX, FLT_MAX);
+		MaxBound = FVector(-FLT_MAX, -FLT_MAX, -FLT_MAX);
+		const uint32 CurveCount = InData.StrandsCurves.Num();
+		OutRoots.Reserve(CurveCount);
+		for (uint32 CurveIndex = 0; CurveIndex < CurveCount; ++CurveIndex)
+		{
+			const uint32 PointOffset = InData.StrandsCurves.CurvesOffset[CurveIndex];
+			const uint32 PointCount = InData.StrandsCurves.CurvesCount[CurveIndex];
+			const float  CurveLength = InData.StrandsCurves.CurvesLength[CurveIndex] * InData.StrandsCurves.MaxLength;
+			check(PointCount > 1);
+			const FVector& P0 = InData.StrandsPoints.PointsPosition[PointOffset];
+			const FVector& P1 = InData.StrandsPoints.PointsPosition[PointOffset + 1];
+			FVector N = (P1 - P0).GetSafeNormal();
+
+			// Fallback in case the initial points are too close (this happens on certain assets)
+			if (FVector::DotProduct(N, N) == 0)
+			{
+				N = FVector(0, 0, 1);
+			}
+			OutRoots.Add({ P0, PointCount, N, PointOffset, CurveLength });
+
+			MinBound = MinBound.ComponentMin(P0);
+			MaxBound = MaxBound.ComponentMax(P0);
+		}
+	}
+
+	enum class EHairInterpolationDataQuality
+	{
+		Low,
+		Medium,
+		High
+	};
+
+	static void BuildInterpolationData(
+		FHairStrandsInterpolationDatas& InterpolationData,
+		const FHairStrandsDatas& SimStrandsData,
+		const FHairStrandsDatas& RenStrandsData,
+		const EHairInterpolationDataQuality Quality)
+	{
+		TRACE_CPUPROFILER_EVENT_SCOPE(HairInterpolationBuilder::BuildInterpolationData);
+
+		InterpolationData.SetNum(RenStrandsData.GetNumPoints());
+
+		typedef TArray<FHairRoot> FRoots;
+
+		// Build acceleration structure for fast nearest-neighbors lookup.
+		// This is used only for low quality interpolation as high quality 
+		// interpolation require broader search
+		FRoots RenRoots, SimRoots;
+		FRootsGrid RootsGrid;
+		{
+			FVector RenMinBound, RenMaxBound;
+			FVector SimMinBound, SimMaxBound;
+			ExtractRoots(RenStrandsData, RenRoots, RenMinBound, RenMaxBound);
+			ExtractRoots(SimStrandsData, SimRoots, SimMinBound, SimMaxBound);
+
+			if (Quality == EHairInterpolationDataQuality::Low || Quality == EHairInterpolationDataQuality::Medium)
+			{
+				// Build a conservative bound, to insure all queries will fall 
+				// into the grid volume.
+				const FVector MinBound = RenMinBound.ComponentMin(SimMinBound);
+				const FVector MaxBound = RenMaxBound.ComponentMax(SimMaxBound);
+				RootsGrid.InsertRoots(SimRoots, MinBound, MaxBound);
 			}
+		}
+
+		// Find k-closest guide:
+		uint32 TotalInvalidInterpolationCount = 0;
+		const static float MinWeightDistance = 0.0001f;
+
+		FRandomStream Random;
+		const uint32 RenCurveCount = RenStrandsData.GetNumCurves();
+		const uint32 SimCurveCount = SimStrandsData.GetNumCurves();
 
+		ParallelFor(RenCurveCount, 
+		[
+			Quality,
+			RenCurveCount, &RenRoots, &RenStrandsData,
+			SimCurveCount, &SimRoots, &SimStrandsData, 
+			&RootsGrid,
+			&TotalInvalidInterpolationCount,  
+			&InterpolationData, 
+			&Random
+		] (uint32 RenCurveIndex) 
+		//for (uint32 RenCurveIndex = 0; RenCurveIndex < RenCurveCount; ++RenCurveIndex)
+		{
+			const FHairRoot& StrandRoot = RenRoots[RenCurveIndex];
+
+			int32 ClosestGuideIndices[FMetrics::ClosestCount];
+			if (Quality == EHairInterpolationDataQuality::Low)
+			{
+				RootsGrid.FindClosestRoots(RenCurveIndex, RenRoots, SimRoots, RenStrandsData, SimStrandsData, Random, ClosestGuideIndices);
+			}
+			else if (Quality == EHairInterpolationDataQuality::Medium)
+			{
+				RootsGrid.FindBestClosestRoots(RenCurveIndex, RenRoots, SimRoots, RenStrandsData, SimStrandsData, Random, ClosestGuideIndices);
+			}
+			else // (Quality == EHairInterpolationDataQuality::High)
+			{
+				FindBestRoots(RenCurveIndex, RenRoots, SimRoots, RenStrandsData, SimStrandsData, Random, ClosestGuideIndices);
+			}
 
 			const uint32 RendPointCount	= RenStrandsData.StrandsCurves.CurvesCount[RenCurveIndex];
 			const uint32 RenOffset		= RenStrandsData.StrandsCurves.CurvesOffset[RenCurveIndex];
@@ -568,7 +849,7 @@ namespace HairInterpolationBuilder
 				const float RenPointDistance = RenStrandsData.StrandsPoints.PointsCoordU[PointGlobalIndex] * RenStrandsData.StrandsCurves.CurvesLength[RenCurveIndex] * RenStrandsData.StrandsCurves.MaxLength;
 
 				float TotalWeight = 0;
-				for (uint32 KIndex = 0; KIndex < GuideCount; ++KIndex)
+				for (uint32 KIndex = 0; KIndex < FMetrics::ClosestCount; ++KIndex)
 				{
 
 				#define WEIGHT_METHOD 0
@@ -594,7 +875,7 @@ namespace HairInterpolationBuilder
 							}
 							break;
 						}
-						PrevSimPointDistance = SimPointDistance;					
+						PrevSimPointDistance = SimPointDistance;
 					}
 					check(SimPointCount > 0);
 					ClosestSimPointIndex = FMath::Clamp(ClosestSimPointIndex, 0u, uint32(SimPointCount - 1));
@@ -610,7 +891,7 @@ namespace HairInterpolationBuilder
 				#if WEIGHT_METHOD == 1
 					const uint32 SimCurveIndex = ClosestGuideIndices[KIndex];
 					const uint32 SimOffset = SimStrandsData.StrandsCurves.CurvesOffset[SimCurveIndex];
-					const FVector& SimRootPointPosition = SimStrandsData.StrandsPoints[SimOffset];
+					const FVector& SimRootPointPosition = SimStrandsData.StrandsPoints.PointsPosition[SimOffset];
 					const float Weight = 1.0f / FMath::Max(MinWeightDistance, FVector::Distance(RenPointPosition, SimRootPointPosition));
 					InterpolationData.PointsSimCurvesIndex[PointGlobalIndex][KIndex] = SimCurveIndex;
 					InterpolationData.PointsSimCurvesVertexIndex[PointGlobalIndex][KIndex] = SimOffset;
@@ -622,7 +903,7 @@ namespace HairInterpolationBuilder
 					check(SimPointCount > 0);
 					const uint32 SimCurveIndex = ClosestGuideIndices[KIndex];
 					const uint32 SimPointIndex = FMath::Clamp(RenPointIndex, 0, SimPointCount - 1);
-					const FVector& SimPointPosition = SimStrandsData.StrandsPoints[SimPointIndex + SimOffset];
+					const FVector& SimPointPosition = SimStrandsData.StrandsPoints.PointsPosition[SimPointIndex + SimOffset];
 					const float Weight = 1.0f / FMath::Max(MinWeightDistance, FVector::Distance(RenPointPosition, SimPointPosition));
 					InterpolationData.PointsSimCurvesIndex[PointGlobalIndex][KIndex] = SimCurveIndex;
 					InterpolationData.PointsSimCurvesVertexIndex[PointGlobalIndex][KIndex] = SimPointIndex + SimOffset;
@@ -632,17 +913,12 @@ namespace HairInterpolationBuilder
 					TotalWeight += Weight;
 				}
 
-				for (int32 KIndex = 0; KIndex < GuideCount; ++KIndex)
+				for (int32 KIndex = 0; KIndex < FMetrics::ClosestCount; ++KIndex)
 				{
 					InterpolationData.PointsSimCurvesVertexWeights[PointGlobalIndex][KIndex] /= TotalWeight;
 				}
 			}
 		});
-
-		if (bPrintDebugMetric)
-		{
-			UE_LOG(LogGroomBuilder, Log, TEXT("Invalid interpolation count: %d/%d)"), TotalInvalidInterpolationCount, RenCurveCount);
-		}
 	}
 
 	/** Build data for interpolation between simulation and rendering */
@@ -884,14 +1160,14 @@ bool FGroomBuilder::BuildGroom(const FHairDescription& HairDescription, const FG
 		GroomAsset->HairGroupsData.Add(MoveTemp(GroupData));
 	}
 
-	BuildData(GroomAsset);
+	BuildData(GroomAsset, uint8(BuildSettings.InterpolationQuality));
 
 	GroomAsset->InitResource();
 
 	return true;
 }
 
-void FGroomBuilder::BuildData(UGroomAsset* GroomAsset)
+void FGroomBuilder::BuildData(UGroomAsset* GroomAsset, uint8 QualityLevel)
 {
 	if (!GroomAsset)
 	{
@@ -930,8 +1206,15 @@ void FGroomBuilder::BuildData(UGroomAsset* GroomAsset)
 		FHairStrandsInterpolationDatas& HairInterpolationData = GroupData.HairInterpolationData;
 
 		// Build InterpolationData from render and simulation HairStrandsDatas
+		HairInterpolationBuilder::EHairInterpolationDataQuality InterpolationQuality = HairInterpolationBuilder::EHairInterpolationDataQuality::High;
+		switch (QualityLevel)
+		{
+			case 0: InterpolationQuality = HairInterpolationBuilder::EHairInterpolationDataQuality::Low; break;
+			case 1: InterpolationQuality = HairInterpolationBuilder::EHairInterpolationDataQuality::Medium; break;
+			case 2: InterpolationQuality = HairInterpolationBuilder::EHairInterpolationDataQuality::High; break;
+		}
 		SlowTask.EnterProgressFrame(1.0f, LOCTEXT("BuildInterpolationData", "Building hair interpolation data"));
-		HairInterpolationBuilder::BuildInterpolationData(HairInterpolationData, HairSimulationData, HairRenderData);
+		HairInterpolationBuilder::BuildInterpolationData(HairInterpolationData, HairSimulationData, HairRenderData, InterpolationQuality);
 
 		// Build Rendering data for InterpolationData
 		SlowTask.EnterProgressFrame(1.0f, LOCTEXT("BuildInterpolationRenderData", "Building render data for hair interpolation"));
diff --git a/Engine/Plugins/Experimental/HairStrands/Source/HairStrandsCore/Public/GroomBuilder.h b/Engine/Plugins/Experimental/HairStrands/Source/HairStrandsCore/Public/GroomBuilder.h
--- a/Engine/Plugins/Experimental/HairStrands/Source/HairStrandsCore/Public/GroomBuilder.h
+++ b/Engine/Plugins/Experimental/HairStrands/Source/HairStrandsCore/Public/GroomBuilder.h
@@ -11,5 +11,5 @@ struct HAIRSTRANDSCORE_API FGroomBuilder
 {
 	static bool BuildGroom(const class FHairDescription& HairDescription, const struct FGroomBuildSettings& BuildSettings, UGroomAsset* GroomAsset);
 	static void GenerateGuides(const FHairStrandsDatas& InData, float DecimationPercentage, FHairStrandsDatas& OutData);
-	static void BuildData(UGroomAsset* GroomAsset);
+	static void BuildData(UGroomAsset* GroomAsset, uint8 QualityLevel);
 };
\ No newline at end of file
diff --git a/Engine/Plugins/Experimental/HairStrands/Source/HairStrandsCore/Public/GroomSettings.h b/Engine/Plugins/Experimental/HairStrands/Source/HairStrandsCore/Public/GroomSettings.h
--- a/Engine/Plugins/Experimental/HairStrands/Source/HairStrandsCore/Public/GroomSettings.h
+++ b/Engine/Plugins/Experimental/HairStrands/Source/HairStrandsCore/Public/GroomSettings.h
@@ -18,11 +18,23 @@ struct HAIRSTRANDSCORE_API FGroomConversionSettings
 
 	/** Rotation in Euler angles in degrees to fix up or front axes */
 	UPROPERTY(EditAnywhere, BlueprintReadWrite, Category = Conversion)
-		FVector Rotation;
+	FVector Rotation;
 
 	/** Scale value to convert file unit into centimeters */
 	UPROPERTY(EditAnywhere, BlueprintReadWrite, Category = Conversion)
-		FVector Scale;
+	FVector Scale;
+};
+
+/**
+ * Specifies the overal rendering/shading model for a material
+ * @warning Check UMaterialInstance::Serialize if changed!
+ */
+UENUM()
+enum class EGroomInterpolationQuality : uint8
+{
+	Low		UMETA(DisplayName = "Low",		ToolTip = "Build interpolation data based on nearst neighbor search. Low quality interpolation data, but fast to build (takes a few minutes)"),
+	Medium	UMETA(DisplayName = "Medium",	ToolTip = "Build interpolation data using curve shape matching search but within a limited spatial range. This is a tradeoff between Low and high quality in term of quality & build time (can takes several dozen of minutes)"),
+	High	UMETA(DisplayName = "High",		ToolTip = "Build interpolation data using curve shape matching search. This result in high quality interpolation data, but is relatively slow to build (can takes several dozen of minutes)"),
 };
 
 USTRUCT(BlueprintType)
@@ -33,6 +45,7 @@ struct HAIRSTRANDSCORE_API FGroomBuildSettings
 	FGroomBuildSettings()
 		: bOverrideGuides(false)
 		, HairToGuideDensity(0.1f)
+		, InterpolationQuality(EGroomInterpolationQuality::High)
 	{}
 
 	/** Flag to override the imported guides with generated guides. */
@@ -42,4 +55,8 @@ struct HAIRSTRANDSCORE_API FGroomBuildSettings
 	/** Density factor for converting hair into guide curve if no guides are provided. */
 	UPROPERTY(EditAnywhere, BlueprintReadWrite, Category = "BuildSettings", meta = (ClampMin = "0.01", UIMin = "0.01", UIMax = "1.0"))
 	float HairToGuideDensity;
+
+	/** Interpolation data quality. */
+	UPROPERTY(EditAnywhere, BlueprintReadWrite, Category = "BuildSettings")
+	EGroomInterpolationQuality InterpolationQuality;
 };
