Fix crash when importing groom with high/med quality interpolation data.

diff --git a/Engine/Plugins/Experimental/HairStrands/Source/HairStrandsCore/Private/GroomBuilder.cpp b/Engine/Plugins/Experimental/HairStrands/Source/HairStrandsCore/Private/GroomBuilder.cpp
--- a/Engine/Plugins/Experimental/HairStrands/Source/HairStrandsCore/Private/GroomBuilder.cpp
+++ b/Engine/Plugins/Experimental/HairStrands/Source/HairStrandsCore/Private/GroomBuilder.cpp
@@ -382,14 +382,19 @@ namespace HairInterpolationBuilder
 
 	struct FMetrics
 	{
-		static const uint32 ClosestCount = 3;
-		static const uint32 KGuideCount = 6;
-		float KMinMetrics[KGuideCount];
-		int32 KClosestGuideIndices[KGuideCount];
+		static const uint32 Count = 6;
+		float KMinMetrics[Count];
+		int32 KClosestGuideIndices[Count];
+	};
+
+	struct FClosestGuides
+	{
+		static const uint32 Count = 3;
+		int32 Indices[Count];
 	};
 
 	// Randomize influence guide to break interpolation coherence, and create a more random/natural pattern
-	static void SelectFinalGuides(int32* ClosestGuideIndices, FRandomStream& Random, FMetrics& Metric)
+	static void SelectFinalGuides(FClosestGuides& ClosestGuides, FRandomStream& Random, const FMetrics& Metric)
 	{
 		const static bool bRandomizeInterpolation = false;
 		const static bool bUseUniqueGuide = false;
@@ -401,25 +406,25 @@ namespace HairInterpolationBuilder
 		{
 			do
 			{
-				RandIndex0 = Random.RandRange(0, FMetrics::KGuideCount - 1);
-				RandIndex1 = Random.RandRange(0, FMetrics::KGuideCount - 1);
-				RandIndex2 = Random.RandRange(0, FMetrics::KGuideCount - 1);
+				RandIndex0 = Random.RandRange(0, FMetrics::Count - 1);
+				RandIndex1 = Random.RandRange(0, FMetrics::Count - 1);
+				RandIndex2 = Random.RandRange(0, FMetrics::Count - 1);
 
 			} while (RandIndex0 == RandIndex1 || RandIndex0 == RandIndex2 || RandIndex1 == RandIndex2);
 		}
 
-		ClosestGuideIndices[0] = Metric.KClosestGuideIndices[RandIndex0];
-		ClosestGuideIndices[1] = Metric.KClosestGuideIndices[RandIndex1];
-		ClosestGuideIndices[2] = Metric.KClosestGuideIndices[RandIndex2];
+		ClosestGuides.Indices[0] = Metric.KClosestGuideIndices[RandIndex0];
+		ClosestGuides.Indices[1] = Metric.KClosestGuideIndices[RandIndex1];
+		ClosestGuides.Indices[2] = Metric.KClosestGuideIndices[RandIndex2];
 
 		if (bUseUniqueGuide)
 		{
-			ClosestGuideIndices[1] = Metric.KClosestGuideIndices[RandIndex0];
-			ClosestGuideIndices[2] = Metric.KClosestGuideIndices[RandIndex0];
+			ClosestGuides.Indices[1] = Metric.KClosestGuideIndices[RandIndex0];
+			ClosestGuides.Indices[2] = Metric.KClosestGuideIndices[RandIndex0];
 		}
 
 
-		float MinMetrics[FMetrics::KGuideCount];
+		float MinMetrics[FMetrics::Count];
 		MinMetrics[0] = Metric.KMinMetrics[RandIndex0];
 		MinMetrics[1] = Metric.KMinMetrics[RandIndex1];
 		MinMetrics[2] = Metric.KMinMetrics[RandIndex2];
@@ -430,26 +435,26 @@ namespace HairInterpolationBuilder
 			if (MinMetrics[0] > MinMetrics[1])
 			{
 				SwapValue(MinMetrics[0], MinMetrics[1]);
-				SwapValue(ClosestGuideIndices[0], ClosestGuideIndices[1]);
+				SwapValue(ClosestGuides.Indices[0], ClosestGuides.Indices[1]);
 			}
 
 			if (MinMetrics[1] > MinMetrics[2])
 			{
 				SwapValue(MinMetrics[1], MinMetrics[2]);
-				SwapValue(ClosestGuideIndices[1], ClosestGuideIndices[2]);
+				SwapValue(ClosestGuides.Indices[1], ClosestGuides.Indices[2]);
 			}
 		}
 
 		// If there less than 3 valid guides, fill the rest with existing valid guides
 		// This can happen due to the normal-orientation based rejection above
-		if (ClosestGuideIndices[1] < 0)
+		if (ClosestGuides.Indices[1] < 0)
 		{
-			ClosestGuideIndices[1] = ClosestGuideIndices[0];
+			ClosestGuides.Indices[1] = ClosestGuides.Indices[0];
 			MinMetrics[1] = MinMetrics[0];
 		}
-		if (ClosestGuideIndices[2] < 0)
+		if (ClosestGuides.Indices[2] < 0)
 		{
-			ClosestGuideIndices[2] = ClosestGuideIndices[1];
+			ClosestGuides.Indices[2] = ClosestGuides.Indices[1];
 			MinMetrics[2] = MinMetrics[1];
 		}
 
@@ -466,11 +471,11 @@ namespace HairInterpolationBuilder
 		const int32 SimCurveIndex)
 	{
 		const float Metric = FVector::Dist(GuideRoot.Position, RenRoot.Position);
-		if (Metric < Metrics1.KMinMetrics[FMetrics::ClosestCount - 1])
+		if (Metric < Metrics1.KMinMetrics[FMetrics::Count - 1])
 		{
 			int32 LastGuideIndex = SimCurveIndex;
 			float LastMetric = Metric;
-			for (uint32 Index = 0; Index < FMetrics::ClosestCount; ++Index)
+			for (uint32 Index = 0; Index < FMetrics::Count; ++Index)
 			{
 				if (Metric < Metrics1.KMinMetrics[Index])
 				{
@@ -489,11 +494,11 @@ namespace HairInterpolationBuilder
 		const int32 SimCurveIndex)
 	{
 		const float Metric = 1.0 - ComputeCurvesMetric<16>(RenStrandsData, RenCurveIndex, SimStrandsData, SimCurveIndex, 0.0f, 1.0f, 1.0f);
-		if (Metric < Metrics0.KMinMetrics[FMetrics::KGuideCount - 1])
+		if (Metric < Metrics0.KMinMetrics[FMetrics::Count - 1])
 		{
 			int32 LastGuideIndex = SimCurveIndex;
 			float LastMetric = Metric;
-			for (uint32 Index = 0; Index < FMetrics::KGuideCount; ++Index)
+			for (uint32 Index = 0; Index < FMetrics::Count; ++Index)
 			{
 				if (Metric < Metrics0.KMinMetrics[Index])
 				{
@@ -562,20 +567,19 @@ namespace HairInterpolationBuilder
 			}
 		}
 
-		void FindClosestRoots(
+		FClosestGuides FindClosestRoots(
 			const uint32 RenCurveIndex,
 			const TArray<FHairRoot>& RenRoots,
 			const TArray<FHairRoot>& SimRoots,
 			const FHairStrandsDatas& RenStrandsData,
 			const FHairStrandsDatas& SimStrandsData,
-			FRandomStream& Random,
-			int32* ClosestGuideIndices) const
+			FRandomStream& Random) const
 		{
 			const FHairRoot& RenRoot = RenRoots[RenCurveIndex];
 			const FIntVector PointCoord = ToCellCoord(RenRoot.Position);
 
 			FMetrics Metrics;
-			for (uint32 ClearIt = 0; ClearIt < FMetrics::ClosestCount; ++ClearIt)
+			for (uint32 ClearIt = 0; ClearIt < FMetrics::Count; ++ClearIt)
 			{
 				Metrics.KMinMetrics[ClearIt] = FLT_MAX;
 				Metrics.KClosestGuideIndices[ClearIt] = -1;
@@ -621,29 +625,31 @@ namespace HairInterpolationBuilder
 			}
 
 			// If no valid guide have been found, switch to a simpler metric
-			SelectFinalGuides(ClosestGuideIndices, Random, Metrics);
+			FClosestGuides ClosestGuides;
+			SelectFinalGuides(ClosestGuides, Random, Metrics);
+
+			check(ClosestGuides.Indices[0] >= 0);
+			check(ClosestGuides.Indices[1] >= 0);
+			check(ClosestGuides.Indices[2] >= 0);
 
-			check(ClosestGuideIndices[0] >= 0);
-			check(ClosestGuideIndices[1] >= 0);
-			check(ClosestGuideIndices[2] >= 0);
+			return ClosestGuides;
 		}
 
 
-		void FindBestClosestRoots(
+		FClosestGuides FindBestClosestRoots(
 			const uint32 RenCurveIndex,
 			const TArray<FHairRoot>& RenRoots,
 			const TArray<FHairRoot>& SimRoots,
 			const FHairStrandsDatas& RenStrandsData,
 			const FHairStrandsDatas& SimStrandsData,
-			FRandomStream& Random,
-			int32* ClosestGuideIndices) const 
+			FRandomStream& Random) const 
 		{
 			const FHairRoot& RenRoot = RenRoots[RenCurveIndex];
 			const FIntVector PointCoord = ToCellCoord(RenRoot.Position);
 
 			FMetrics Metrics0;
 			FMetrics Metrics1;
-			for (uint32 ClearIt = 0; ClearIt < FMetrics::ClosestCount; ++ClearIt)
+			for (uint32 ClearIt = 0; ClearIt < FMetrics::Count; ++ClearIt)
 			{
 				Metrics0.KMinMetrics[ClearIt] = FLT_MAX;
 				Metrics0.KClosestGuideIndices[ClearIt] = -1;
@@ -685,54 +691,58 @@ namespace HairInterpolationBuilder
 
 				// Early out if we have found all closest guide during a ring/layer step.
 				// This early out is not conservative, as the complex metric might find better guides one or multiple step further.
-				//if (Metrics0.KClosestGuideIndices[FMetrics::ClosestCount - 1] != -1 && R >= 2)
+				//if (Metrics0.KClosestGuideIndices[FMetrics::KGuideCount - 1] != -1 && R >= 2)
 				//{
 				//	break;
 				//}
 			}
 
 			// If no valid guide have been found, switch to a simpler metric
-			SelectFinalGuides(ClosestGuideIndices, Random, Metrics0);
-			if (ClosestGuideIndices[0] == -1)
+			FClosestGuides ClosestGuides;
+			SelectFinalGuides(ClosestGuides, Random, Metrics0);
+			if (ClosestGuides.Indices[0] == -1)
 			{
-				SelectFinalGuides(ClosestGuideIndices, Random, Metrics1);
+				SelectFinalGuides(ClosestGuides, Random, Metrics1);
 			}
 
-			check(ClosestGuideIndices[0] >= 0);
-			check(ClosestGuideIndices[1] >= 0);
-			check(ClosestGuideIndices[2] >= 0);
+			check(ClosestGuides.Indices[0] >= 0);
+			check(ClosestGuides.Indices[1] >= 0);
+			check(ClosestGuides.Indices[2] >= 0);
+
+			return ClosestGuides;
 		}
 	};
 
-	static void FindBestRoots(
+	static FClosestGuides FindBestRoots(
 		const uint32 RenCurveIndex,
 		const TArray<FHairRoot>& RenRoots,
 		const TArray<FHairRoot>& SimRoots,
 		const FHairStrandsDatas& RenStrandsData,
 		const FHairStrandsDatas& SimStrandsData,
-		FRandomStream& Random,
-		int32* ClosestGuideIndices)
+		FRandomStream& Random)
 	{
 		FMetrics Metrics;
-		for (uint32 ClearIt = 0; ClearIt < FMetrics::ClosestCount; ++ClearIt)
+		for (uint32 ClearIt = 0; ClearIt < FMetrics::Count; ++ClearIt)
 		{
 			Metrics.KMinMetrics[ClearIt] = FLT_MAX;
 			Metrics.KClosestGuideIndices[ClearIt] = -1;
 		}
 
-		for(const FHairRoot& GuideRoot : SimRoots)
+		const uint32 SimRootsCount = SimRoots.Num();
+		for (uint32 SimCurveIndex =0; SimCurveIndex<SimRootsCount; ++SimCurveIndex)
 		{
-			const uint32 SimCurveIndex = GuideRoot.Index;
 			ComputeAdvandedMetric(Metrics, RenStrandsData, SimStrandsData, RenCurveIndex, SimCurveIndex);
 		}
 			
-		SelectFinalGuides(ClosestGuideIndices, Random, Metrics);
+		FClosestGuides ClosestGuides;
+		SelectFinalGuides(ClosestGuides, Random, Metrics);
 
-		check(ClosestGuideIndices[0] >= 0);
-		check(ClosestGuideIndices[1] >= 0);
-		check(ClosestGuideIndices[2] >= 0);
-	}
+		check(ClosestGuides.Indices[0] >= 0);
+		check(ClosestGuides.Indices[1] >= 0);
+		check(ClosestGuides.Indices[2] >= 0);
 
+		return ClosestGuides;
+	}
 
 	// Extract strand roots
 	static void ExtractRoots(const FHairStrandsDatas& InData, TArray<FHairRoot>& OutRoots, FVector& MinBound, FVector& MaxBound)
@@ -825,18 +835,18 @@ namespace HairInterpolationBuilder
 		{
 			const FHairRoot& StrandRoot = RenRoots[RenCurveIndex];
 
-			int32 ClosestGuideIndices[FMetrics::ClosestCount];
+			FClosestGuides ClosestGuides;
 			if (Quality == EHairInterpolationDataQuality::Low)
 			{
-				RootsGrid.FindClosestRoots(RenCurveIndex, RenRoots, SimRoots, RenStrandsData, SimStrandsData, Random, ClosestGuideIndices);
+				ClosestGuides = RootsGrid.FindClosestRoots(RenCurveIndex, RenRoots, SimRoots, RenStrandsData, SimStrandsData, Random);
 			}
 			else if (Quality == EHairInterpolationDataQuality::Medium)
 			{
-				RootsGrid.FindBestClosestRoots(RenCurveIndex, RenRoots, SimRoots, RenStrandsData, SimStrandsData, Random, ClosestGuideIndices);
+				ClosestGuides = RootsGrid.FindBestClosestRoots(RenCurveIndex, RenRoots, SimRoots, RenStrandsData, SimStrandsData, Random);
 			}
 			else // (Quality == EHairInterpolationDataQuality::High)
 			{
-				FindBestRoots(RenCurveIndex, RenRoots, SimRoots, RenStrandsData, SimStrandsData, Random, ClosestGuideIndices);
+				ClosestGuides = FindBestRoots(RenCurveIndex, RenRoots, SimRoots, RenStrandsData, SimStrandsData, Random);
 			}
 
 			const uint32 RendPointCount	= RenStrandsData.StrandsCurves.CurvesCount[RenCurveIndex];
@@ -849,14 +859,14 @@ namespace HairInterpolationBuilder
 				const float RenPointDistance = RenStrandsData.StrandsPoints.PointsCoordU[PointGlobalIndex] * RenStrandsData.StrandsCurves.CurvesLength[RenCurveIndex] * RenStrandsData.StrandsCurves.MaxLength;
 
 				float TotalWeight = 0;
-				for (uint32 KIndex = 0; KIndex < FMetrics::ClosestCount; ++KIndex)
+				for (uint32 KIndex = 0; KIndex < FClosestGuides::Count; ++KIndex)
 				{
 
 				#define WEIGHT_METHOD 0
 
 					// Find the closest vertex on the guide which matches the strand vertex distance along its curve
 				#if WEIGHT_METHOD == 0
-					const uint32 SimCurveIndex = ClosestGuideIndices[KIndex];
+					const uint32 SimCurveIndex = ClosestGuides.Indices[KIndex];
 					const uint32 SimOffset = SimStrandsData.StrandsCurves.CurvesOffset[SimCurveIndex];
 					float PrevSimPointDistance = 0;
 					uint32 ClosestSimPointIndex = 0;
@@ -913,7 +923,7 @@ namespace HairInterpolationBuilder
 					TotalWeight += Weight;
 				}
 
-				for (int32 KIndex = 0; KIndex < FMetrics::ClosestCount; ++KIndex)
+				for (int32 KIndex = 0; KIndex < FClosestGuides::Count; ++KIndex)
 				{
 					InterpolationData.PointsSimCurvesVertexWeights[PointGlobalIndex][KIndex] /= TotalWeight;
 				}
