Fixed progress bar not updating during groom import.

diff --git a/Engine/Plugins/Experimental/HairStrands/Source/HairStrandsCore/Private/GroomBuilder.cpp b/Engine/Plugins/Experimental/HairStrands/Source/HairStrandsCore/Private/GroomBuilder.cpp
--- a/Engine/Plugins/Experimental/HairStrands/Source/HairStrandsCore/Private/GroomBuilder.cpp
+++ b/Engine/Plugins/Experimental/HairStrands/Source/HairStrandsCore/Private/GroomBuilder.cpp
@@ -821,6 +821,10 @@ namespace HairInterpolationBuilder
 		const uint32 RenCurveCount = RenStrandsData.GetNumCurves();
 		const uint32 SimCurveCount = SimStrandsData.GetNumCurves();
 
+		TAtomic<uint32> CompletedTasks = 0;
+		FScopedSlowTask SlowTask(RenCurveCount, LOCTEXT("BuildInterpolationData", "Building groom simulation data"));
+		SlowTask.MakeDialog();
+
 		ParallelFor(RenCurveCount, 
 		[
 			Quality,
@@ -829,10 +833,22 @@ namespace HairInterpolationBuilder
 			&RootsGrid,
 			&TotalInvalidInterpolationCount,  
 			&InterpolationData, 
-			&Random
+			&Random,
+			&CompletedTasks,
+			&SlowTask
 		] (uint32 RenCurveIndex) 
 		//for (uint32 RenCurveIndex = 0; RenCurveIndex < RenCurveCount; ++RenCurveIndex)
 		{
+			TRACE_CPUPROFILER_EVENT_SCOPE(HairInterpolationBuilder::ComputingClosestGuidesAndWeights);
+
+			++CompletedTasks;
+
+			if (IsInGameThread())
+			{
+				uint32 CurrentCompletedTasks = CompletedTasks.Exchange(0);
+				SlowTask.EnterProgressFrame(CurrentCompletedTasks, LOCTEXT("ComputeGuidesAndWeights", "Computing closest guides and weights"));
+			}
+
 			const FHairRoot& StrandRoot = RenRoots[RenCurveIndex];
 
 			FClosestGuides ClosestGuides;
@@ -977,6 +993,8 @@ bool FGroomBuilder::BuildGroom(const FHairDescription& HairDescription, const FG
 		return false;
 	}
 
+	TRACE_CPUPROFILER_EVENT_SCOPE(FGroomBuilder::BuildGroom);
+
 	// Convert HairDescription to HairStrandsDatas
 	// For now, just convert HairDescription to HairStrandsDatas
 	int32 NumCurves = HairDescription.GetNumStrands();
@@ -1184,10 +1202,7 @@ void FGroomBuilder::BuildData(UGroomAsset* GroomAsset, uint8 QualityLevel)
 		return;
 	}
 
-	float NumTasks = 4.0f * GroomAsset->GetNumHairGroups();
-
-	FScopedSlowTask SlowTask(NumTasks, LOCTEXT("BuildData", "Building groom render and simulation data"));
-	SlowTask.MakeDialog();
+	TRACE_CPUPROFILER_EVENT_SCOPE(FGroomBuilder::BuildData);
 
 	for (int32 Index = 0; Index < GroomAsset->GetNumHairGroups(); ++Index)
 	{
@@ -1207,10 +1222,8 @@ void FGroomBuilder::BuildData(UGroomAsset* GroomAsset, uint8 QualityLevel)
 		}
 
 		// Build RenderData for HairStrandsDatas
-		SlowTask.EnterProgressFrame(1.0f, LOCTEXT("BuildHairRenderData", "Building render data for hair"));
 		HairStrandsBuilder::BuildRenderData(HairRenderData);
 
-		SlowTask.EnterProgressFrame(1.0f, LOCTEXT("BuildSimRenderData", "Building render data for guides"));
 		HairStrandsBuilder::BuildRenderData(HairSimulationData);
 
 		FHairStrandsInterpolationDatas& HairInterpolationData = GroupData.HairInterpolationData;
@@ -1223,11 +1236,9 @@ void FGroomBuilder::BuildData(UGroomAsset* GroomAsset, uint8 QualityLevel)
 			case 1: InterpolationQuality = HairInterpolationBuilder::EHairInterpolationDataQuality::Medium; break;
 			case 2: InterpolationQuality = HairInterpolationBuilder::EHairInterpolationDataQuality::High; break;
 		}
-		SlowTask.EnterProgressFrame(1.0f, LOCTEXT("BuildInterpolationData", "Building hair interpolation data"));
 		HairInterpolationBuilder::BuildInterpolationData(HairInterpolationData, HairSimulationData, HairRenderData, InterpolationQuality);
 
 		// Build Rendering data for InterpolationData
-		SlowTask.EnterProgressFrame(1.0f, LOCTEXT("BuildInterpolationRenderData", "Building render data for hair interpolation"));
 		HairInterpolationBuilder::BuildRenderData(HairInterpolationData);
 	}
 }
