Fix niagara crash with the grrom component delegate
Adaptive substepping to have a fix dt independently of the engine deltatime
Projection interpolation for better stability

diff --git a/Engine/Plugins/Experimental/HairStrands/Shaders/Private/NiagaraDataInterfaceHairStrands.ush b/Engine/Plugins/Experimental/HairStrands/Shaders/Private/NiagaraDataInterfaceHairStrands.ush
--- a/Engine/Plugins/Experimental/HairStrands/Shaders/Private/NiagaraDataInterfaceHairStrands.ush
+++ b/Engine/Plugins/Experimental/HairStrands/Shaders/Private/NiagaraDataInterfaceHairStrands.ush
@@ -221,12 +221,13 @@ void DIHairStrands_ComputeNodeInertia(in FDIHairStrandsContext DIContext, in int
 	const float MeanRadius = 0.25 * ( PrevThick + NextThick );
 
 	//const float EdgeMass = StrandsDensity * EdgeLength * M_PI * (PrevThick * PrevThick + NextThick * PrevThick + NextThick * NextThick) / 12.0;
-	//const float RadialInertia = EdgeMass * ( 3.0 * MeanRadius * MeanRadius + EdgeLength * EdgeLength) / 12.0;
-	//OutNodeInertia = float3( RadialInertia, RadialInertia, EdgeMass * MeanRadius * MeanRadius * 0.5 );
-
 	const float EdgeMass = StrandsDensity * EdgeLength * M_PI * MeanRadius * MeanRadius;
-	const float RadialInertia = EdgeMass * ( MeanRadius * MeanRadius) / 4.0;
-	OutNodeInertia = float3( RadialInertia, RadialInertia, 2.0 * RadialInertia );
+	const float RadialInertia = EdgeMass * ( 3.0 * MeanRadius * MeanRadius + EdgeLength * EdgeLength) / 12.0;
+	OutNodeInertia = float3( RadialInertia, RadialInertia, EdgeMass * MeanRadius * MeanRadius * 0.5 );
+
+	//const float EdgeMass = StrandsDensity * EdgeLength * M_PI * MeanRadius * MeanRadius;
+	//const float RadialInertia = EdgeMass * ( MeanRadius * MeanRadius) / 4.0;
+	//OutNodeInertia = float3( RadialInertia, RadialInertia, 2.0 * RadialInertia );
 }
 
 // Compute the node position by linear interpolation over the points
diff --git a/Engine/Plugins/Experimental/HairStrands/Shaders/Private/NiagaraHookeSpringMaterial.ush b/Engine/Plugins/Experimental/HairStrands/Shaders/Private/NiagaraHookeSpringMaterial.ush
--- a/Engine/Plugins/Experimental/HairStrands/Shaders/Private/NiagaraHookeSpringMaterial.ush
+++ b/Engine/Plugins/Experimental/HairStrands/Shaders/Private/NiagaraHookeSpringMaterial.ush
@@ -261,6 +261,20 @@ void ComputeMaterialFrame(in int StrandsSize)
 	GroupMemoryBarrier();
 }
 
+/*void ComputeMaterialFrame(in int StrandsSize)
+{
+	const int LocalIndex = (GGroupThreadId.x % StrandsSize);
+	if( (LocalIndex > 0) && (LocalIndex < StrandsSize-1) )
+	{
+		float4 NodeQuaternion = SharedNodeOrientation[GGroupThreadId.x];
+		float3 TangentPrev = normalize(SharedPreviousPosition[GGroupThreadId.x+1] - SharedPreviousPosition[GGroupThreadId.x]);
+		float3 TangentNext = normalize(SharedNodePosition[GGroupThreadId.x+1] - SharedNodePosition[GGroupThreadId.x]);
+
+		SharedNodeOrientation[GGroupThreadId.x] = NormalizeQuat( MultiplyQuat( FindQuatBetweenNormals(TangentPrev,TangentNext), SharedPreviousOrientation[GGroupThreadId.x] ) );
+	}
+	GroupMemoryBarrier();
+}*/
+
 /* -----------------------------------------------------------------
  * Bend spring material
  * -----------------------------------------------------------------
diff --git a/Engine/Plugins/Experimental/HairStrands/Source/HairStrandsCore/Private/GroomComponent.cpp b/Engine/Plugins/Experimental/HairStrands/Source/HairStrandsCore/Private/GroomComponent.cpp
--- a/Engine/Plugins/Experimental/HairStrands/Source/HairStrandsCore/Private/GroomComponent.cpp
+++ b/Engine/Plugins/Experimental/HairStrands/Source/HairStrandsCore/Private/GroomComponent.cpp
@@ -711,18 +711,15 @@ void UGroomComponent::OnChildDetached(USceneComponent* ChildComponent)
 
 void UGroomComponent::ResetSimulation()
 {
-	USkeletalMeshComponent* SkeletalMeshComponent = GetAttachParent() ? Cast<USkeletalMeshComponent>(GetAttachParent()) : nullptr;
-	if (SkeletalMeshComponent)
-	{
-		SkeletalMeshComponent->OnBoneTransformsFinalized.RemoveDynamic(this, &UGroomComponent::ResetSimulation);
-		bResetSimulation = false;
-	}
+	bResetSimulation = false;
+	//UE_LOG(LogHairStrands, Warning, TEXT("Groom Reset = %d"), bResetSimulation);
 }
 
 void UGroomComponent::InitResources()
 {
 	ReleaseResources();
 	bResetSimulation = true;
+	//UE_LOG(LogHairStrands, Warning, TEXT("Groom Init = %d"), bResetSimulation);
 
 	if (!GroomAsset || GroomAsset->GetNumHairGroups() == 0)
 		return;
@@ -744,6 +741,11 @@ void UGroomComponent::InitResources()
 		CallbackData.Run = CallbackMeshObjectCallback;
 		CallbackData.UserData = (uint64(LocalComponentId.PrimIDValue) & 0xFFFFFFFF) | (uint64(WorldType) << 32);
 		SkeletalMeshComponent->MeshObjectCallbackData = CallbackData;
+
+
+		{
+			SkeletalMeshComponent->OnBoneTransformsFinalized.AddDynamic(this, &UGroomComponent::ResetSimulation);
+		}
 	}
 
 	const bool bIsSimulationEnable = IsSimulationEnabled(this);
@@ -951,6 +953,13 @@ void UGroomComponent::ReleaseResources()
 	SkeletalPreviousPositionOffset = FVector::ZeroVector;
 	RegisteredSkeletalMeshComponent = nullptr;
 
+	USkeletalMeshComponent* SkeletalMeshComponent = GetAttachParent() ? Cast<USkeletalMeshComponent>(GetAttachParent()) : nullptr;
+	if (SkeletalMeshComponent)
+	{
+		SkeletalMeshComponent->OnBoneTransformsFinalized.RemoveDynamic(this, &UGroomComponent::ResetSimulation);
+		bResetSimulation = true;
+	}
+
 	MarkRenderStateDirty();
 }
 
@@ -995,10 +1004,10 @@ void UGroomComponent::OnRegister()
 
 	// Insure the ticking of the Groom component always happens after the skeletalMeshComponent.
 	USkeletalMeshComponent* SkeletalMeshComponent = GetAttachParent() ? Cast<USkeletalMeshComponent>(GetAttachParent()) : nullptr;
-	if (SkeletalMeshComponent)
+	/*if (SkeletalMeshComponent)
 	{
 		SkeletalMeshComponent->OnBoneTransformsFinalized.AddDynamic(this, &UGroomComponent::ResetSimulation);
-	}
+	}*/
 
 	const EWorldType::Type WorldType = GetWorld() ? EWorldType::Type(GetWorld()->WorldType) : EWorldType::None;
 	const uint64 Id = ComponentId.PrimIDValue;
diff --git a/Engine/Plugins/Experimental/HairStrands/Source/HairStrandsNiagara/Private/NiagaraDataInterfaceHairStrands.cpp b/Engine/Plugins/Experimental/HairStrands/Source/HairStrandsNiagara/Private/NiagaraDataInterfaceHairStrands.cpp
--- a/Engine/Plugins/Experimental/HairStrands/Source/HairStrandsNiagara/Private/NiagaraDataInterfaceHairStrands.cpp
+++ b/Engine/Plugins/Experimental/HairStrands/Source/HairStrandsNiagara/Private/NiagaraDataInterfaceHairStrands.cpp
@@ -15,7 +15,7 @@
 
 #define LOCTEXT_NAMESPACE "NiagaraDataInterfaceHairStrands"
 
-static const int32 MaxDelay = 5;
+static const int32 MaxDelay = 8;
 
 //------------------------------------------------------------------------------------------------------------
 
@@ -414,6 +414,8 @@ struct FNDIHairStrandsParametersCS : public FNiagaraDataInterfaceParametersCS
 			FVector DeformedPositionOffsetValue = HairStrandsBuffer->SourceRestResources->PositionOffset;
 			FVector RestPositionOffsetValue = HairStrandsBuffer->SourceRestResources->PositionOffset;
 
+			//UE_LOG(LogHairStrands, Log, TEXT("Need Reset = %d %d %d"), bNeedSimReset, ProxyData->TickCount, ProxyData->ResetTick);
+
 			if (HairStrandsBuffer->SourceRootResources != nullptr && HairStrandsBuffer->SourceRootResources->MeshProjectionLODs.Num() > 0)
 			{
 				RestTrianglePositionASRV = HairStrandsBuffer->SourceRootResources->MeshProjectionLODs[0].RestRootTrianglePosition0Buffer.SRV;
@@ -598,7 +600,7 @@ void FNDIHairStrandsProxy::InitializePerInstanceData(const FNiagaraSystemInstanc
 	TargetData->BoxExtent = BoxExtent;
 	TargetData->TickCount = 0;
 	TargetData->ResetTick = MaxDelay;
-	TargetData->ForceReset = false;
+	TargetData->ForceReset = true;
 	TargetData->WorldTransform = WorldTransform;
 }
 
@@ -737,7 +739,7 @@ bool UNiagaraDataInterfaceHairStrands::InitPerInstanceData(void* PerInstanceData
 			InstanceData->BoxCenter = BoxCenter;
 			InstanceData->BoxExtent = BoxExtent;
 			InstanceData->TickCount = 0;
-			InstanceData->ForceReset = false;
+			InstanceData->ForceReset = true;
 			InstanceData->ResetTick = MaxDelay;
 
 			//UE_LOG(LogHairStrands, Log, TEXT("Num Strands = %d | Strand Size = %d | Num Vertices = %d | Min = %f %f %f | Max = %f %f %f | Transform = %s"), NumStrands, StrandSize,
@@ -792,10 +794,11 @@ bool UNiagaraDataInterfaceHairStrands::PerInstanceTick(void* PerInstanceData, FN
 
 	if (SourceComponent != nullptr)
 	{
-		if (!InstanceData->ForceReset && !SourceComponent->bResetSimulation)
+		if (!InstanceData->ForceReset && !SourceComponent->bResetSimulation && (InstanceData->ResetTick == MaxDelay))
 		{
-			InstanceData->ResetTick = FMath::Min(MaxDelay, InstanceData->TickCount);
+			InstanceData->ResetTick = FMath::Min(MaxDelay, InstanceData->TickCount+1);
 		}
+		//UE_LOG(LogHairStrands, Warning, TEXT("Reset Simulation : %d %d %d %d"), InstanceData->ForceReset, SourceComponent->bResetSimulation, InstanceData->TickCount, InstanceData->ResetTick);
 		InstanceData->ForceReset = SourceComponent->bResetSimulation;
 	}
 
