Fix Clang Win64 error copying variable of type 'TAtomic<uint32>' invokes deleted constructor
#rnx

diff --git a/Engine/Plugins/Experimental/HairStrands/Source/HairStrandsCore/Private/GroomBuilder.cpp b/Engine/Plugins/Experimental/HairStrands/Source/HairStrandsCore/Private/GroomBuilder.cpp
--- a/Engine/Plugins/Experimental/HairStrands/Source/HairStrandsCore/Private/GroomBuilder.cpp
+++ b/Engine/Plugins/Experimental/HairStrands/Source/HairStrandsCore/Private/GroomBuilder.cpp
@@ -821,7 +821,7 @@ namespace HairInterpolationBuilder
 		const uint32 RenCurveCount = RenStrandsData.GetNumCurves();
 		const uint32 SimCurveCount = SimStrandsData.GetNumCurves();
 
-		TAtomic<uint32> CompletedTasks = 0;
+		TAtomic<uint32> CompletedTasks(0);
 		FScopedSlowTask SlowTask(RenCurveCount, LOCTEXT("BuildInterpolationData", "Building groom simulation data"));
 		SlowTask.MakeDialog();
 
