Add hair shadow density live tweak and overridable  per instance.

diff --git a/Engine/Plugins/Experimental/HairStrands/Source/HairStrandsCore/Private/GroomComponent.cpp b/Engine/Plugins/Experimental/HairStrands/Source/HairStrandsCore/Private/GroomComponent.cpp
--- a/Engine/Plugins/Experimental/HairStrands/Source/HairStrandsCore/Private/GroomComponent.cpp
+++ b/Engine/Plugins/Experimental/HairStrands/Source/HairStrandsCore/Private/GroomComponent.cpp
@@ -41,6 +41,11 @@ static float GetGroupMaxHairRadius(const FHairGroupDesc& GroupDesc, const FHairG
 	return GroupDesc.HairWidth > 0 ? GroupDesc.HairWidth * 0.5f : GroupData.HairRenderData.StrandsCurves.MaxRadius;
 }
 
+static float GetGroupHairShadowDensity(const FHairGroupDesc& GroupDesc, const FHairGroupData& GroupData)
+{
+	return GroupDesc.HairShadowDensity > 0 ? GroupDesc.HairShadowDensity : GroupData.HairRenderData.HairDensity;
+}
+
 /**
  * An material render proxy which overrides the debug mode parameter.
  */
@@ -168,7 +173,7 @@ public:
 			VFGroupData.MinStrandRadius		= 0;
 			VFGroupData.MaxStrandRadius		= GetGroupMaxHairRadius(InGroupDesc, InGroupData);
 			VFGroupData.MaxStrandLength		= InGroupData.HairRenderData.StrandsCurves.MaxLength;
-			VFGroupData.HairDensity			= InGroupData.HairRenderData.HairDensity;
+			VFGroupData.HairDensity			= GetGroupHairShadowDensity(InGroupDesc, InGroupData);
 			VFGroupData.HairWorldOffset		= InGroupData.HairRenderData.BoundingBox.GetCenter();
 
 			#if RHI_RAYTRACING
@@ -409,12 +414,16 @@ static void UpdateHairGroupsDesc(UGroomAsset* GroomAsset, TArray<FHairGroupDesc>
 		FHairGroupData& GroupData = GroomAsset->HairGroupsData[GroupIt];
 
 		FHairGroupDesc& Desc = GroomGroupsDesc[GroupIt];
-		Desc.GuideCount = GroupInfo.NumGuides;
-		Desc.HairCount  = GroupInfo.NumCurves;
+		Desc.GuideCount	= GroupInfo.NumGuides;
+		Desc.HairCount	= GroupInfo.NumCurves;
 		if (bReinitOverride || Desc.HairWidth == 0)
 		{
 			Desc.HairWidth = GroupData.HairRenderData.StrandsCurves.MaxRadius * 0.5f;
 		}
+		if (bReinitOverride || Desc.HairShadowDensity == 0)
+		{
+			Desc.HairShadowDensity = GroupData.HairRenderData.HairDensity;
+		}
 	}
 }
 
@@ -433,7 +442,6 @@ UGroomComponent::UGroomComponent(const FObjectInitializer& ObjectInitializer)
 	bSelectable = true;
 	RegisteredSkeletalMeshComponent = nullptr;
 	SkeletalPreviousPositionOffset = FVector::ZeroVector;
-	HairShadowDensity = 1;
 	HairRaytracingRadiusScale = 1;
 	bBindGroomToSkeletalMesh = false;
 	InitializedResources = nullptr;
@@ -1389,11 +1397,13 @@ void UGroomComponent::PostEditChangeProperty(FPropertyChangedEvent& PropertyChan
 		}
 	}
 
-	if (PropertyName == GET_MEMBER_NAME_CHECKED(FHairGroupDesc, HairWidth))
+	if (PropertyName == GET_MEMBER_NAME_CHECKED(FHairGroupDesc, HairWidth) || 
+		PropertyName == GET_MEMBER_NAME_CHECKED(FHairGroupDesc, HairShadowDensity))
 	{	
 		UpdateHairGroupsDesc(GroomAsset, GroomGroupsDesc);
 		MarkRenderStateDirty();
 	}
+
 #if WITH_EDITOR
 	ValidateMaterials(false);
 #endif
diff --git a/Engine/Plugins/Experimental/HairStrands/Source/HairStrandsCore/Public/GroomComponent.h b/Engine/Plugins/Experimental/HairStrands/Source/HairStrandsCore/Public/GroomComponent.h
--- a/Engine/Plugins/Experimental/HairStrands/Source/HairStrandsCore/Public/GroomComponent.h
+++ b/Engine/Plugins/Experimental/HairStrands/Source/HairStrandsCore/Public/GroomComponent.h
@@ -27,6 +27,10 @@ struct FHairGroupDesc
 	/** Override the hair width (in centimeters) */
 	UPROPERTY(EditAnywhere, BlueprintReadWrite, Category = "Groom", meta = (ClampMin = "0.0001", UIMin = "0.001", UIMax = "1.0", SliderExponent = 6))
 	float HairWidth;
+
+	/** Override the hair shadow density factor (unit less).  */
+	UPROPERTY(EditAnywhere, BlueprintReadWrite, Category = "Groom", meta = (ClampMin = "0.0001", UIMin = "0.001", UIMax = "10.0", SliderExponent = 6))
+	float HairShadowDensity;
 };
 
 UCLASS(HideCategories = (Object, Physics, Activation, Mobility, "Components|Activation"), editinlinenew, meta = (BlueprintSpawnableComponent), ClassGroup = Rendering)
@@ -49,10 +53,6 @@ public:
 	UPROPERTY(EditAnywhere, BlueprintReadOnly, Category = "Groom")
 	bool bBindGroomToSkeletalMesh;
 
-	/** Controls the hair density, to reduce or increase hair count during shadow rendering. This allows to increase/decrease the shadowing on hair when the number of strand is not realistic */
-	UPROPERTY(EditAnywhere, BlueprintReadWrite, Category = "Groom Rendering", meta = (ClampMin = "0.0001", UIMin = "0.01", UIMax = "10.0"))
-	float HairShadowDensity;
-
 	/** Scale the hair geometry radius for ray tracing effects (e.g. shadow) */
 	UPROPERTY(EditAnywhere, BlueprintReadWrite, Category = "Groom Rendering", meta = (ClampMin = "0.0001", UIMin = "0.01", UIMax = "10.0"))
 	float HairRaytracingRadiusScale;
