Removed #pragma once.

diff --git a/Engine/Plugins/Experimental/HairStrands/Source/HairStrandsCore/Private/GroomBuilder.cpp b/Engine/Plugins/Experimental/HairStrands/Source/HairStrandsCore/Private/GroomBuilder.cpp
--- a/Engine/Plugins/Experimental/HairStrands/Source/HairStrandsCore/Private/GroomBuilder.cpp
+++ b/Engine/Plugins/Experimental/HairStrands/Source/HairStrandsCore/Private/GroomBuilder.cpp
@@ -1,7 +1,5 @@
 // Copyright 1998-2019 Epic Games, Inc. All Rights Reserved.
 
-#pragma once
-
 #include "GroomBuilder.h"
 #include "GroomAsset.h"
 #include "GroomComponent.h"
diff --git a/Engine/Plugins/Experimental/HairStrands/Source/HairStrandsEditor/Private/GroomImportOptionsWindow.cpp b/Engine/Plugins/Experimental/HairStrands/Source/HairStrandsEditor/Private/GroomImportOptionsWindow.cpp
--- a/Engine/Plugins/Experimental/HairStrands/Source/HairStrandsEditor/Private/GroomImportOptionsWindow.cpp
+++ b/Engine/Plugins/Experimental/HairStrands/Source/HairStrandsEditor/Private/GroomImportOptionsWindow.cpp
@@ -1,7 +1,5 @@
 // Copyright 1998-2019 Epic Games, Inc. All Rights Reserved.
 
-#pragma once
-
 #include "GroomImportOptionsWindow.h"
 
 #include "EditorStyleSet.h"
