Change hair plugins names from hair strands to groom (only the user exposed name).
Add icon for the alembic groom importer plugin.

diff --git a/Engine/Plugins/Experimental/AlembicHairImporter/AlembicHairImporter.uplugin b/Engine/Plugins/Experimental/AlembicHairImporter/AlembicHairImporter.uplugin
--- a/Engine/Plugins/Experimental/AlembicHairImporter/AlembicHairImporter.uplugin
+++ b/Engine/Plugins/Experimental/AlembicHairImporter/AlembicHairImporter.uplugin
@@ -2,7 +2,7 @@
 	"FileVersion": 1,
 	"Version": 1,
 	"VersionName": "0.1",
-	"FriendlyName": "Alembic Hair Strands Importer",
+	"FriendlyName": "Alembic Groom Importer",
 	"Description": "Import Hair Strands from Alembic file",
 	"Category": "Geometry",
 	"CreatedBy": "Epic Games, Inc.",
diff --git a/Engine/Plugins/Experimental/HairStrands/HairStrands.uplugin b/Engine/Plugins/Experimental/HairStrands/HairStrands.uplugin
--- a/Engine/Plugins/Experimental/HairStrands/HairStrands.uplugin
+++ b/Engine/Plugins/Experimental/HairStrands/HairStrands.uplugin
@@ -2,8 +2,8 @@
 	"FileVersion": 3,
 	"Version": 1,
 	"VersionName": "0.1",
-	"FriendlyName": "Hair Strands",
-	"Description": "Build Hair Strands for rendering and simulation",
+	"FriendlyName": "Groom",
+	"Description": "Rendering and simulation of grooms",
 	"Category": "Geometry",
 	"CreatedBy": "Epic Games, Inc.",
 	"CreatedByURL": "http://epicgames.com",
