Change raytracing strands scale for shadow to be defined per hair group.

This makes the property more consistent.

diff --git a/Engine/Plugins/Experimental/HairStrands/Source/HairStrandsCore/Private/GroomComponent.cpp b/Engine/Plugins/Experimental/HairStrands/Source/HairStrandsCore/Private/GroomComponent.cpp
--- a/Engine/Plugins/Experimental/HairStrands/Source/HairStrandsCore/Private/GroomComponent.cpp
+++ b/Engine/Plugins/Experimental/HairStrands/Source/HairStrandsCore/Private/GroomComponent.cpp
@@ -46,6 +46,11 @@ static float GetGroupHairShadowDensity(const FHairGroupDesc& GroupDesc, const FH
 	return GroupDesc.HairShadowDensity > 0 ? GroupDesc.HairShadowDensity : GroupData.HairRenderData.HairDensity;
 }
 
+static float GetGroupHairRaytracingRadiusScale(const FHairGroupDesc& GroupDesc, const FHairGroupData& GroupData)
+{
+	return GroupDesc.HairRaytracingRadiusScale;
+}
+
 /**
  * An material render proxy which overrides the debug mode parameter.
  */
@@ -424,6 +429,11 @@ static void UpdateHairGroupsDesc(UGroomAsset* GroomAsset, TArray<FHairGroupDesc>
 		{
 			Desc.HairShadowDensity = GroupData.HairRenderData.HairDensity;
 		}
+
+		if (bReinitOverride)
+		{
+			Desc.HairRaytracingRadiusScale = 1;
+		}
 	}
 }
 
@@ -442,7 +452,6 @@ UGroomComponent::UGroomComponent(const FObjectInitializer& ObjectInitializer)
 	bSelectable = true;
 	RegisteredSkeletalMeshComponent = nullptr;
 	SkeletalPreviousPositionOffset = FVector::ZeroVector;
-	HairRaytracingRadiusScale = 1;
 	bBindGroomToSkeletalMesh = false;
 	InitializedResources = nullptr;
 	Mobility = EComponentMobility::Movable;
@@ -826,12 +835,12 @@ void UGroomComponent::InitResources()
 		check(GroupIt < GroomGroupsDesc.Num());
 		const FHairGroupDesc& InGroupDesc = GroomGroupsDesc[GroupIt];
 		InterpolationInputGroup.HairRadius = GetGroupMaxHairRadius(InGroupDesc, GroupData);
-		InterpolationInputGroup.HairRaytracingRadiusScale = HairRaytracingRadiusScale;
+		InterpolationInputGroup.HairRaytracingRadiusScale = GetGroupHairRaytracingRadiusScale(InGroupDesc, GroupData);
 		InterpolationInputGroup.InRenderHairPositionOffset = RenderRestHairPositionOffset;
 		InterpolationInputGroup.InSimHairPositionOffset = SimRestHairPositionOffset;
 
 		// For skinned groom, these value will be updated during TickComponent() call
-		// Deformed sim & render are expressed within the referencial (unlike rest pose)
+		// Deformed sim & render are expressed within the referential (unlike rest pose)
 		InterpolationInputGroup.OutHairPositionOffset = RenderRestHairPositionOffset;
 		InterpolationInputGroup.OutHairPreviousPositionOffset = RenderRestHairPositionOffset;
 		InterpolationInputGroup.bIsSimulationEnable = bIsSimulationEnable;
@@ -1389,11 +1398,12 @@ void UGroomComponent::PostEditChangeProperty(FPropertyChangedEvent& PropertyChan
 		}
 	}
 
-	if (PropertyName == GET_MEMBER_NAME_CHECKED(UGroomComponent, HairRaytracingRadiusScale) && InterpolationInput)
+	if (PropertyName == GET_MEMBER_NAME_CHECKED(FHairGroupDesc, HairRaytracingRadiusScale) && InterpolationInput && InterpolationInput->HairGroups.Num() == GroomGroupsDesc.Num())
 	{
-		for (FHairStrandsInterpolationInput::FHairGroup& Group : InterpolationInput->HairGroups)
+		const int32 GroupCount = InterpolationInput->HairGroups.Num();
+		for (int32 GroupIt = 0; GroupIt < GroupCount; ++GroupIt)
 		{
-			Group.HairRaytracingRadiusScale = HairRaytracingRadiusScale;
+			InterpolationInput->HairGroups[GroupIt].HairRaytracingRadiusScale = GroomGroupsDesc[GroupIt].HairRaytracingRadiusScale;
 		}
 	}
 
diff --git a/Engine/Plugins/Experimental/HairStrands/Source/HairStrandsCore/Public/GroomComponent.h b/Engine/Plugins/Experimental/HairStrands/Source/HairStrandsCore/Public/GroomComponent.h
--- a/Engine/Plugins/Experimental/HairStrands/Source/HairStrandsCore/Public/GroomComponent.h
+++ b/Engine/Plugins/Experimental/HairStrands/Source/HairStrandsCore/Public/GroomComponent.h
@@ -31,6 +31,10 @@ struct FHairGroupDesc
 	/** Override the hair shadow density factor (unit less).  */
 	UPROPERTY(EditAnywhere, BlueprintReadWrite, Category = "Groom", meta = (ClampMin = "0.0001", UIMin = "0.001", UIMax = "10.0", SliderExponent = 6))
 	float HairShadowDensity;
+
+	/** Scale the hair geometry radius for ray tracing effects (e.g. shadow) */
+	UPROPERTY(EditAnywhere, BlueprintReadWrite, Category = "Groom", meta = (ClampMin = "0.0001", UIMin = "0.001", UIMax = "10.0", SliderExponent = 6))
+	float HairRaytracingRadiusScale;
 };
 
 UCLASS(HideCategories = (Object, Physics, Activation, Mobility, "Components|Activation"), editinlinenew, meta = (BlueprintSpawnableComponent), ClassGroup = Rendering)
@@ -53,10 +57,6 @@ public:
 	UPROPERTY(EditAnywhere, BlueprintReadOnly, Category = "Groom")
 	bool bBindGroomToSkeletalMesh;
 
-	/** Scale the hair geometry radius for ray tracing effects (e.g. shadow) */
-	UPROPERTY(EditAnywhere, BlueprintReadWrite, Category = "Groom Rendering", meta = (ClampMin = "0.0001", UIMin = "0.01", UIMax = "10.0"))
-	float HairRaytracingRadiusScale;
-
 	/** Boolean to check when animation has been loaded */
 	bool bResetSimulation;
 
