Add mode for outputing the render rest position

diff --git a/Engine/Plugins/Experimental/HairStrands/Source/HairStrandsCore/Private/HairStrandsRendering.cpp b/Engine/Plugins/Experimental/HairStrands/Source/HairStrandsCore/Private/HairStrandsRendering.cpp
--- a/Engine/Plugins/Experimental/HairStrands/Source/HairStrandsCore/Private/HairStrandsRendering.cpp
+++ b/Engine/Plugins/Experimental/HairStrands/Source/HairStrandsCore/Private/HairStrandsRendering.cpp
@@ -9,7 +9,7 @@
 #include "HairStrandsInterface.h"
 
 static int32 GHairDeformationType = 0;
-static FAutoConsoleVariableRef CVarHairDeformationType(TEXT("r.HairStrands.DeformationType"), GHairDeformationType, TEXT("Type of procedural deformation applied on hair strands (0:bypass, 1:wave, 2:normal)"));
+static FAutoConsoleVariableRef CVarHairDeformationType(TEXT("r.HairStrands.DeformationType"), GHairDeformationType, TEXT("Type of procedural deformation applied on hair strands (0:use simulation's output, 1:use rest strands, 2: use rest guides, 3:wave pattern, 4:follow root normal)"));
 
 static float GHairRaytracingRadiusScale = 0;
 static FAutoConsoleVariableRef CVarHairRaytracingRadiusScale(TEXT("r.HairStrands.RaytracingRadiusScale"), GHairRaytracingRadiusScale, TEXT("Override the per instance scale factor for raytracing hair strands geometry (0: disabled, >0:enabled)"));
@@ -41,6 +41,29 @@ inline uint32 GetGroupSizePermutation(uint32 GroupSize)
 	return GroupSize == 64 ? 0 : (GroupSize == 32 ? 1 : 2);
 }
 
+enum class EDeformationType : uint8
+{
+	Simulation,		// Use the output of the hair simulation
+	RestStrands,	// Use the rest strands position (no weighted interpolation)
+	RestGuide,		// Use the rest guide as input of the interpolation (no deformation), only weighted interpolation
+	Wave,			// Apply a wave pattern to deform the guides
+	NormalDirection // Apply a stretch pattern aligned with the guide root's normal
+};
+
+static EDeformationType GetDeformationType()
+{
+	switch (GHairDeformationType)
+	{
+	case 0: return EDeformationType::Simulation;
+	case 1: return EDeformationType::RestStrands;
+	case 2: return EDeformationType::RestGuide;
+	case 3: return EDeformationType::Wave;
+	case 4: return EDeformationType::NormalDirection;
+	}
+
+	return EDeformationType::Simulation;
+}
+
 class FDeformGuideCS : public FGlobalShader
 {
 	DECLARE_GLOBAL_SHADER(FDeformGuideCS);
@@ -66,7 +89,7 @@ IMPLEMENT_GLOBAL_SHADER(FDeformGuideCS, "/Engine/Private/HairStrands/HairStrands
 
 static void AddDeformSimHairStrandsPass(
 	FRDGBuilder& GraphBuilder,
-	uint32 DeformationType,
+	EDeformationType DeformationType,
 	uint32 VertexCount,
 	FRHIShaderResourceView* SimRestPosePositionBuffer,
 	FRHIShaderResourceView* SimRootIndexBuffer,
@@ -75,6 +98,16 @@ static void AddDeformSimHairStrandsPass(
 	static uint32 IterationCount = 0;
 	++IterationCount;
 
+	int32 InternalDeformationType = -1;
+	switch (DeformationType)
+	{
+	case EDeformationType::RestGuide: InternalDeformationType = 0; break;
+	case EDeformationType::Wave: InternalDeformationType = 1; break;
+	case EDeformationType::NormalDirection: InternalDeformationType = 2; break;
+	}
+
+	if (InternalDeformationType < 0) return;
+
 	FDeformGuideCS::FParameters* Parameters = GraphBuilder.AllocParameters<FDeformGuideCS::FParameters>();
 	Parameters->SimRestPosePositionBuffer = SimRestPosePositionBuffer;
 	Parameters->SimRootIndexBuffer = SimRootIndexBuffer;
@@ -85,7 +118,7 @@ static void AddDeformSimHairStrandsPass(
 
 	FDeformGuideCS::FPermutationDomain PermutationVector;
 	PermutationVector.Set<FDeformGuideCS::FGroupSize>(GetGroupSizePermutation(GroupSize));
-	PermutationVector.Set<FDeformGuideCS::FDeformationType>(DeformationType);
+	PermutationVector.Set<FDeformGuideCS::FDeformationType>(InternalDeformationType);
 
 	TShaderMap<FGlobalShaderType>* ShaderMap = GetGlobalShaderMap(ERHIFeatureLevel::SM5);
 
@@ -436,13 +469,8 @@ void ComputeHairStrandsInterpolation(
 		const uint32 CurrIndex = Output.CurrentIndex;
 		const uint32 PrevIndex = (Output.CurrentIndex + 1) % 2;
 
-		// Procedural deformers in place of actual physics solver
-		// 0: Simulation
-		// 1: Bypass
-		// 2: Wave
-		// 3: Straighten hair in direction of the root's normal
-		const int32 DeformationType = FMath::Clamp(GHairDeformationType, 0, Input.SimRootPointIndexBuffer ? 3 : 2) - 1;
-		if (GHairDeformationType > 0)
+		const EDeformationType DeformationType = GetDeformationType();
+		if (DeformationType != EDeformationType::RestStrands && DeformationType != EDeformationType::Simulation)
 		{
 			AddDeformSimHairStrandsPass(
 				GraphBuilder,
@@ -511,10 +539,22 @@ void ComputeHairStrandsInterpolation(
 				Output.RenderDeformedPositionBuffer[CurrIndex]->UAV,
 				DebugMode == EHairStrandsDebugMode::RenderHairStrands ? Output.RenderPatchedAttributeBuffer.UAV : nullptr);
 
+
+			if (DeformationType == EDeformationType::RestStrands)
+			{
+				Output.VFInput.HairPositionBuffer = Input.RenderRestPosePositionBuffer->SRV;
+				Output.VFInput.HairPreviousPositionBuffer = Input.RenderRestPosePositionBuffer->SRV;
+			}
+			else
+			{
+				Output.VFInput.HairPositionBuffer = Output.RenderDeformedPositionBuffer[CurrIndex]->SRV;
+				Output.VFInput.HairPreviousPositionBuffer = Output.RenderDeformedPositionBuffer[PrevIndex]->SRV;
+			}
+
 			AddHairTangentPass(
 				GraphBuilder,
 				Input.RenderVertexCount,
-				Output.RenderDeformedPositionBuffer[CurrIndex]->SRV,
+				Output.VFInput.HairPositionBuffer,
 				Output.RenderTangentBuffer->UAV);
 
 			#if RHI_RAYTRACING
@@ -525,14 +565,12 @@ void ComputeHairStrandsInterpolation(
 					Input.RenderVertexCount,
 					Input.HairRadius * (GHairRaytracingRadiusScale > 0 ? GHairRaytracingRadiusScale : Input.HairRaytracingRadiusScale),
 					Input.OutHairPositionOffset,
-					Output.RenderDeformedPositionBuffer[CurrIndex]->SRV,
+					Output.VFInput.HairPositionBuffer,
 					Input.RaytracingPositionBuffer->UAV);
 			}
 			#endif
 			GraphBuilder.Execute();
 
-			Output.VFInput.HairPositionBuffer = Output.RenderDeformedPositionBuffer[CurrIndex]->SRV;
-			Output.VFInput.HairPreviousPositionBuffer = Output.RenderDeformedPositionBuffer[PrevIndex]->SRV;
 			Output.VFInput.HairTangentBuffer = Output.RenderTangentBuffer->SRV;
 			Output.VFInput.HairAttributeBuffer = DebugMode == EHairStrandsDebugMode::RenderHairStrands ? Output.RenderPatchedAttributeBuffer.SRV : Input.RenderAttributeBuffer->SRV;
 			Output.VFInput.HairMaterialBuffer = Output.RenderMaterialBuffer->SRV;
