Add debug code for isolating a specific hair strands or cluster.

diff --git a/Engine/Plugins/Experimental/HairStrands/Source/HairStrandsCore/Private/HairStrandsRendering.cpp b/Engine/Plugins/Experimental/HairStrands/Source/HairStrandsCore/Private/HairStrandsRendering.cpp
--- a/Engine/Plugins/Experimental/HairStrands/Source/HairStrandsCore/Private/HairStrandsRendering.cpp
+++ b/Engine/Plugins/Experimental/HairStrands/Source/HairStrandsCore/Private/HairStrandsRendering.cpp
@@ -142,7 +142,7 @@ class FHairInterpolationCS : public FGlobalShader
 	SHADER_USE_PARAMETER_STRUCT(FHairInterpolationCS, FGlobalShader);
 
 	class FGroupSize : SHADER_PERMUTATION_INT("PERMUTATION_GROUP_SIZE", 2);
-	class FDebug : SHADER_PERMUTATION_INT("PERMUTATION_DEBUG", 2);
+	class FDebug : SHADER_PERMUTATION_INT("PERMUTATION_DEBUG", 3);
 	class FDynamicGeometry : SHADER_PERMUTATION_INT("PERMUTATION_DYNAMIC_GEOMETRY", 2);
 	class FSimulation : SHADER_PERMUTATION_INT("PERMUTATION_SIMULATION", 2);
 	using FPermutationDomain = TShaderPermutationDomain<FGroupSize, FDebug, FDynamicGeometry, FSimulation>;
@@ -153,6 +153,7 @@ class FHairInterpolationCS : public FGlobalShader
 		SHADER_PARAMETER(FVector, InRenderHairPositionOffset)
 		SHADER_PARAMETER(FVector, InSimHairPositionOffset)
 		SHADER_PARAMETER(FVector, OutHairPositionOffset)
+		SHADER_PARAMETER(FIntPoint, HairStrandsCullIndex)
 
 		SHADER_PARAMETER(FVector, RestPositionOffset)
 		SHADER_PARAMETER(FVector, DeformedPositionOffset)
@@ -181,7 +182,7 @@ class FHairInterpolationCS : public FGlobalShader
 
 		SHADER_PARAMETER_SRV(Buffer<uint>, RootBarycentricBuffer)
 		SHADER_PARAMETER_SRV(Buffer<uint>, RootToTriangleIndex)
-		SHADER_PARAMETER_SRV(Buffer<uint>, VertexToRootIndexBuffer)
+		SHADER_PARAMETER_SRV(Buffer<uint>, RenVertexToRootIndexBuffer)
 
 		SHADER_PARAMETER_SRV(Buffer<float4>, SimRestPosition0Buffer)
 		SHADER_PARAMETER_SRV(Buffer<float4>, SimRestPosition1Buffer)
@@ -207,7 +208,7 @@ static void AddHairStrandsInterpolationPass(
 	FRDGBuilder& GraphBuilder,
 	const FHairStrandsProjectionHairData::HairGroup& InRenHairData,
 	const FHairStrandsProjectionHairData::HairGroup& InSimHairData,
-	const FVector& InRenderHairWorldOffset, 
+	const FVector& InRenderHairWorldOffset,
 	const FVector& InSimHairWorldOffset,
 	const FVector& OutHairWorldOffset,
 	const int32 LODIndex,
@@ -233,6 +234,16 @@ static void AddHairStrandsInterpolationPass(
 	Parameters->Interpolation0Buffer = Interpolation0Buffer;
 	Parameters->Interpolation1Buffer = Interpolation1Buffer;
 	Parameters->OutRenderDeformedPositionBuffer = OutRenderPositionBuffer;
+	Parameters->HairStrandsCullIndex = FIntPoint(-1, -1);
+	const FHairCullInfo Info = GetHairStrandsCullInfo();
+	const bool bCullingEnable = Info.CullMode != EHairCullMode::None && InRenHairData.VertexToCurveIndexBuffer && InSimHairData.VertexToCurveIndexBuffer;
+	if (bCullingEnable)
+	{
+		if (Info.CullMode == EHairCullMode::Sim)
+			Parameters->HairStrandsCullIndex.Y = Info.ExplicitIndex >= 0 ? Info.ExplicitIndex : FMath::Clamp(uint32(Info.NormalizedIndex * InSimHairData.RootCount), 0u, InSimHairData.RootCount - 1);
+		if (Info.CullMode == EHairCullMode::Render)
+			Parameters->HairStrandsCullIndex.X = Info.ExplicitIndex >= 0 ? Info.ExplicitIndex : FMath::Clamp(uint32(Info.NormalizedIndex * InRenHairData.RootCount), 0u, InRenHairData.RootCount - 1);
+	}
 	if (bCopySimAttributesToRenderAttributes)
 	{
 		Parameters->SimAttributeBuffer = SimAttributeBuffer;
@@ -245,6 +256,12 @@ static void AddHairStrandsInterpolationPass(
 	Parameters->DispatchCountX = DispatchCount.X;
 
 	const bool bSupportDynamicMesh = InRenHairData.RootCount > 0 && LODIndex >= 0 && LODIndex < InRenHairData.LODDatas.Num() && InRenHairData.LODDatas[LODIndex].bIsValid;
+	if (bSupportDynamicMesh || bCullingEnable)
+	{
+		Parameters->RenVertexToRootIndexBuffer = InRenHairData.VertexToCurveIndexBuffer->SRV;
+		Parameters->SimVertexToRootIndexBuffer = InSimHairData.VertexToCurveIndexBuffer->SRV;
+	}
+
 	if (bSupportDynamicMesh)
 	{
 		Parameters->RestPositionOffset = InRenHairData.LODDatas[LODIndex].RestPositionOffset;
@@ -259,7 +276,6 @@ static void AddHairStrandsInterpolationPass(
 
 		Parameters->RootToTriangleIndex = InRenHairData.LODDatas[LODIndex].RootTriangleIndexBuffer->SRV;
 		Parameters->RootBarycentricBuffer = InRenHairData.LODDatas[LODIndex].RootTriangleBarycentricBuffer->SRV;
-		Parameters->VertexToRootIndexBuffer = InRenHairData.VertexToCurveIndexBuffer->SRV;
 
 		Parameters->SimRestPositionOffset = InSimHairData.LODDatas[LODIndex].RestPositionOffset;
 		Parameters->SimRestPosition0Buffer = InSimHairData.LODDatas[LODIndex].RestRootTrianglePosition0Buffer->SRV;
@@ -273,12 +289,11 @@ static void AddHairStrandsInterpolationPass(
 
 		Parameters->SimRootToTriangleIndex = InSimHairData.LODDatas[LODIndex].RootTriangleIndexBuffer->SRV;
 		Parameters->SimRootBarycentricBuffer = InSimHairData.LODDatas[LODIndex].RootTriangleBarycentricBuffer->SRV;
-		Parameters->SimVertexToRootIndexBuffer = InSimHairData.VertexToCurveIndexBuffer->SRV;
 	}
 
 	FHairInterpolationCS::FPermutationDomain PermutationVector;
 	PermutationVector.Set<FHairInterpolationCS::FGroupSize>(GetGroupSizePermutation(GroupSize));
-	PermutationVector.Set<FHairInterpolationCS::FDebug>(bCopySimAttributesToRenderAttributes ? 1 : 0);
+	PermutationVector.Set<FHairInterpolationCS::FDebug>(bCopySimAttributesToRenderAttributes ? 1 : (bCullingEnable ? 2 : 0));
 	PermutationVector.Set<FHairInterpolationCS::FDynamicGeometry>(bSupportDynamicMesh ? 1 : 0);
 	PermutationVector.Set<FHairInterpolationCS::FSimulation>(bHasSimulationEnable ? 1 : 0);
 
@@ -517,7 +532,7 @@ void ComputeHairStrandsInterpolation(
 				Output.RenderPatchedAttributeBuffer.Initialize(FHairStrandsAttributeFormat::SizeInByte, Input.RenderVertexCount, FHairStrandsAttributeFormat::Format, BUF_Static);
 			}
 
-			const bool bHasSimulationEnabled = Input.bIsSimulationEnable && GHairStrandsInterpolateSimulation;
+			const bool bHasSimulationEnabled = Input.bIsSimulationEnable && GHairStrandsInterpolateSimulation && DeformationType != EDeformationType::RestStrands;
 			check(GroupIndex < uint32(InRenHairDatas.HairGroups.Num()));
 			check(GroupIndex < uint32(InSimHairDatas.HairGroups.Num()));
 			AddHairStrandsInterpolationPass(
@@ -539,18 +554,9 @@ void ComputeHairStrandsInterpolation(
 				Output.RenderDeformedPositionBuffer[CurrIndex]->UAV,
 				DebugMode == EHairStrandsDebugMode::RenderHairStrands ? Output.RenderPatchedAttributeBuffer.UAV : nullptr);
 
-
-			if (DeformationType == EDeformationType::RestStrands)
-			{
-				Output.VFInput.HairPositionBuffer = Input.RenderRestPosePositionBuffer->SRV;
-				Output.VFInput.HairPreviousPositionBuffer = Input.RenderRestPosePositionBuffer->SRV;
-			}
-			else
-			{
-				Output.VFInput.HairPositionBuffer = Output.RenderDeformedPositionBuffer[CurrIndex]->SRV;
-				Output.VFInput.HairPreviousPositionBuffer = Output.RenderDeformedPositionBuffer[PrevIndex]->SRV;
-			}
-
+			Output.VFInput.HairPositionBuffer = Output.RenderDeformedPositionBuffer[CurrIndex]->SRV;
+			Output.VFInput.HairPreviousPositionBuffer = Output.RenderDeformedPositionBuffer[PrevIndex]->SRV;
+			
 			AddHairTangentPass(
 				GraphBuilder,
 				Input.RenderVertexCount,
diff --git a/Engine/Shaders/Private/HairStrands/HairStrandsInterpolation.usf b/Engine/Shaders/Private/HairStrands/HairStrandsInterpolation.usf
--- a/Engine/Shaders/Private/HairStrands/HairStrandsInterpolation.usf
+++ b/Engine/Shaders/Private/HairStrands/HairStrandsInterpolation.usf
@@ -8,6 +8,7 @@ uint DispatchCountX;
 float3 InRenderHairPositionOffset;
 float3 InSimHairPositionOffset;
 float3 OutHairPositionOffset;
+int2 HairStrandsCullIndex;
 
 Buffer<uint4>  RenderRestPosePositionBuffer;
 Buffer<uint4>  SimRestPosePositionBuffer;
@@ -15,6 +16,11 @@ Buffer<uint4>  DeformedSimPositionBuffer;
 Buffer<uint4>  Interpolation0Buffer;
 Buffer<uint4>  Interpolation1Buffer;
 
+#if PERMUTATION_DYNAMIC_GEOMETRY == 1 || PERMUTATION_DEBUG == 2
+Buffer<uint>   RenVertexToRootIndexBuffer;
+Buffer<uint>   SimVertexToRootIndexBuffer;
+#endif
+
 RWBuffer<uint4> OutRenderDeformedPositionBuffer;
 
 #if PERMUTATION_DYNAMIC_GEOMETRY == 1
@@ -32,7 +38,6 @@ Buffer<float4>  DeformedPosition2Buffer;
 
 Buffer<uint>	RootBarycentricBuffer;
 Buffer<uint>	RootToTriangleIndex;
-Buffer<uint>	VertexToRootIndexBuffer;
 
 float3			SimRestPositionOffset;
 Buffer<float4>  SimRestPosition0Buffer;
@@ -46,7 +51,6 @@ Buffer<float4>  SimDeformedPosition2Buffer;
 
 Buffer<uint>	SimRootBarycentricBuffer;
 Buffer<uint>	SimRootToTriangleIndex;
-Buffer<uint>	SimVertexToRootIndexBuffer;
 #endif
 
 #if PERMUTATION_DEBUG == 1
@@ -70,7 +74,7 @@ void MainCS(uint2 DispatchThreadId : SV_DispatchThreadID)
 	{
 		uint3 GuideVertexIndices = 0;
 		float3 GuideVertexWeights = 0;
-		#if PERMUTATION_SIMULATION == 1 || PERMUTATION_DEBUG == 1
+		#if PERMUTATION_SIMULATION == 1 || PERMUTATION_DEBUG == 1 || PERMUTATION_DEBUG == 2
 		{
 			const uint4 InterpolationData0	= Interpolation0Buffer[VertexIndex];
 			const uint4 InterpolationData1	= Interpolation1Buffer[VertexIndex];
@@ -134,7 +138,7 @@ void MainCS(uint2 DispatchThreadId : SV_DispatchThreadID)
 			#endif		
 
 			// Transform hair from rest post to deformed pose, based triangle deformation + simulation offset (optional depending of the simulation is actual running or not
-			const uint RootIndex						= VertexToRootIndexBuffer[VertexIndex];
+			const uint RootIndex						= RenVertexToRootIndexBuffer[VertexIndex];
 			const float3 RootBarycentric				= DecodeBarycentrics(RootBarycentricBuffer[RootIndex]);
 			const FHairMeshTriangle RestTriangle		= GetTriangleTransformation(RootIndex, RestPosition0Buffer, RestPosition1Buffer, RestPosition2Buffer, RestPositionOffset);
 			const FHairMeshTriangle DeformedTriangle	= GetTriangleTransformation(RootIndex, DeformedPosition0Buffer, DeformedPosition1Buffer, DeformedPosition2Buffer, DeformedPositionOffset);
@@ -164,5 +168,36 @@ void MainCS(uint2 DispatchThreadId : SV_DispatchThreadID)
 			OutRenderAttributeBuffer[VertexIndex] = GuideAttribute;
 		}
 		#endif
+
+		// Cull hair strands or clusters which are not matching the provided index
+		#if PERMUTATION_DEBUG == 2
+		{
+			bool bIsValid = false;
+			if (HairStrandsCullIndex.y >= 0)
+			{
+
+				for (uint K = 0; K < 3; ++K)
+				{
+					const uint GuideIndex = GuideVertexIndices[K];
+					const uint RootIndex = SimVertexToRootIndexBuffer[GuideIndex];
+					if (HairStrandsCullIndex.y == RootIndex)
+					{
+						bIsValid = true;
+						break;
+					}
+				}
+			}
+			else if (HairStrandsCullIndex.x >= 0)
+			{
+				const uint RootIndex = RenVertexToRootIndexBuffer[VertexIndex];
+				bIsValid = HairStrandsCullIndex.x == RootIndex;
+			}
+
+			if (!bIsValid)
+			{
+				OutRenderDeformedPositionBuffer[VertexIndex] = float4(float3(1.#INF, 1.#INF, 1.#INF), 2);
+			}
+		}
+		#endif
 	}
 }
\ No newline at end of file
diff --git a/Engine/Source/Runtime/Renderer/Private/HairStrands/HairStrandsDebug.cpp b/Engine/Source/Runtime/Renderer/Private/HairStrands/HairStrandsDebug.cpp
--- a/Engine/Source/Runtime/Renderer/Private/HairStrands/HairStrandsDebug.cpp
+++ b/Engine/Source/Runtime/Renderer/Private/HairStrands/HairStrandsDebug.cpp
@@ -79,6 +79,23 @@ static FAutoConsoleVariableRef CVarHairStrandsDebugPPLL(									TEXT("r.HairStr
 void GetGroomInterpolationData(const EWorldType::Type WorldType, FHairStrandsProjectionMeshData& OutGeometries);
 void GetGroomInterpolationData(const EWorldType::Type WorldType, const bool bRenderData, FHairStrandsProjectionHairData& OutHairData, TArray<int32>& OutLODIndices);
 
+static int32 GHairStrandsCull = 0;
+static int32 GHairStrandsCullIndex = -1;
+static int32 GHairStrandsUpdateCullIndex = 0;
+static float GHairStrandsCullNormalizedIndex = -1;
+static FAutoConsoleVariableRef CVarHairStrandsCull			(TEXT("r.HairStrands.Cull"), GHairStrandsCull, TEXT("Cull hair strands (0:disabled, 1: render cull, 2: sim cull)."));
+static FAutoConsoleVariableRef CVarHairStrandsCullIndex		(TEXT("r.HairStrands.Cull.Index"), GHairStrandsCullIndex, TEXT("Hair strands index to be kept. Other will be culled."));
+static FAutoConsoleVariableRef CVarHairStrandsUpdateCullIndex(TEXT("r.HairStrands.Cull.Update"), GHairStrandsUpdateCullIndex, TEXT("Update the guide index to be kept using mouse position for fast selection."));
+
+FHairCullInfo GetHairStrandsCullInfo()
+{
+	FHairCullInfo Out;
+	Out.CullMode		= GHairStrandsCull == 1 ? EHairCullMode::Render : (GHairStrandsCull == 2 ? EHairCullMode::Sim : EHairCullMode::None);
+	Out.ExplicitIndex	= GHairStrandsCullIndex >= 0 ? GHairStrandsCullIndex : -1;
+	Out.NormalizedIndex = GHairStrandsCullNormalizedIndex;
+	return Out;
+}
+
 enum class EHairDebugMode : uint8
 {
 	None,
@@ -1065,6 +1082,14 @@ void RenderHairStrandsDebugInfo(FRHICommandListImmediate& RHICmdList, TArray<FVi
 	if (Views.Num() == 0)
 		return;
 
+	if (GHairStrandsUpdateCullIndex)
+	{
+		const FViewInfo& View = Views[0];
+		const float TotalPixelCount = View.ViewRect.Width() * View.ViewRect.Height();
+		const float Index = View.CursorPos.X + View.CursorPos.Y * View.ViewRect.Width();
+		GHairStrandsCullNormalizedIndex = Index / TotalPixelCount;
+	}
+
 	// Only render debug information for the main view
 	const uint32 ViewIndex = 0;
 	FViewInfo& View = Views[ViewIndex];
diff --git a/Engine/Source/Runtime/Renderer/Public/HairStrandsInterface.h b/Engine/Source/Runtime/Renderer/Public/HairStrandsInterface.h
--- a/Engine/Source/Runtime/Renderer/Public/HairStrandsInterface.h
+++ b/Engine/Source/Runtime/Renderer/Public/HairStrandsInterface.h
@@ -191,3 +191,18 @@ RENDERER_API bool UpdateHairStrands(
 
 RENDERER_API bool IsHairStrandsSupported(const EShaderPlatform Platform);
 bool IsHairStrandsEnable(EShaderPlatform Platform);
+
+// Return strands & guide indices to be preserved, while all others strands/guides should be culled
+enum class EHairCullMode : uint8
+{
+	None,
+	Render,
+	Sim
+};
+struct FHairCullInfo
+{
+	int32 ExplicitIndex = -1; 
+	float NormalizedIndex = 0; // [0,1]
+	EHairCullMode CullMode = EHairCullMode::None;
+};
+RENDERER_API FHairCullInfo GetHairStrandsCullInfo();
