Constraint Projection for stretch and bend
Adaptive time stepping
Position smoothing based on bessel filter
Slerp in parrallel transport
Root position and orientation interpolation

diff --git a/Engine/Plugins/Experimental/HairStrands/Shaders/Private/NiagaraDataInterfaceHairStrands.ush b/Engine/Plugins/Experimental/HairStrands/Shaders/Private/NiagaraDataInterfaceHairStrands.ush
--- a/Engine/Plugins/Experimental/HairStrands/Shaders/Private/NiagaraDataInterfaceHairStrands.ush
+++ b/Engine/Plugins/Experimental/HairStrands/Shaders/Private/NiagaraDataInterfaceHairStrands.ush
@@ -277,10 +277,77 @@ void DIHairStrands_ComputeNodePosition(in FDIHairStrandsContext DIContext, in in
 	OutNodePosition = mul(float4( OutNodePosition+DIContext.RestPositionOffset, 1.0), DIContext.WorldTransform).xyz;
 }
 
+// Compute the node position by linear interpolation over the points
+void DIHairStrands_SmoothNodePosition(in FDIHairStrandsContext DIContext, in int NodeIndex, inout float3 OutNodePosition)
+{
+	SharedNodePosition[GGroupThreadId.x] = OutNodePosition;
+	GroupMemoryBarrier();
+
+	const int LocalIndex = (GGroupThreadId.x % DIContext.StrandSize);
+	if(LocalIndex == 0)
+	{
+		const float Beta = 0.2;
+		float3 DirM1 = SharedNodePosition[GGroupThreadId.x+1] - SharedNodePosition[GGroupThreadId.x];
+		float3 DirM2 = DirM1;
+
+		const float Gamma1 = 2.0 * (1.0-Beta);
+		const float Gamma2 = - (1.0-Beta)*(1.0-Beta);
+		const float Gamma3 = Beta*Beta;
+
+		float3 NodePosition = SharedNodePosition[GGroupThreadId.x];
+		SharedPreviousPosition[GGroupThreadId.x] = NodePosition;
+		
+		for( int i = GGroupThreadId.x, end = GGroupThreadId.x+DIContext.StrandSize-1; i < end; ++i)
+		{
+			const float3 DirM3 = SharedNodePosition[i+1] - SharedNodePosition[i];
+			const float3 DirMi = Gamma1 * DirM1 + Gamma2 * DirM2 + Gamma3 * DirM3;
+
+			SharedPreviousPosition[i+1]  = SharedPreviousPosition[i] + DirMi;
+			GroupMemoryBarrier();
+
+			DirM2 = DirM1;
+			DirM1 = DirMi;
+		}
+	}
+	GroupMemoryBarrier();
+	OutNodePosition = SharedPreviousPosition[GGroupThreadId.x];
+}
+
 // Compute the node orientation
-void DIHairStrands_ComputeNodeOrientation(in FDIHairStrandsContext DIContext, in int NodeIndex, out float4 OutNodeOrientation)
+void DIHairStrands_ComputeNodeOrientation(in FDIHairStrandsContext DIContext, in int NodeIndex, in float3 NodePosition, out float4 OutNodeOrientation)
 {
-	int StrandIndex = 0, LocalIndex = 0;
+	SharedNodePosition[GGroupThreadId.x] = NodePosition;
+	SharedNodeOrientation[GGroupThreadId.x] = QUATERNION_IDENTITY;
+	GroupMemoryBarrier();
+
+	const int LocalIndex = (GGroupThreadId.x % DIContext.StrandSize);
+	if(LocalIndex == 0)
+	{
+		float4 NodeQuaternion = QUATERNION_IDENTITY;
+		float3 TangentNext = float3(0,0,1);
+		float3 TangentPrev = TangentNext;
+
+		float3 PositionNext = SharedNodePosition[GGroupThreadId.x];
+		float3 PositionPrev = PositionNext;
+
+		for (int EdgeIndex = 0, EdgeEnd = DIContext.StrandSize-1; EdgeIndex < EdgeEnd; ++EdgeIndex)
+		{
+			PositionPrev = PositionNext;
+			PositionNext = SharedNodePosition[GGroupThreadId.x+EdgeIndex+1];
+
+			TangentPrev = TangentNext;
+			TangentNext = normalize(PositionNext - PositionPrev);
+			NodeQuaternion = NormalizeQuat( MultiplyQuat( FindQuatBetweenNormals(TangentPrev,TangentNext), NodeQuaternion ) );
+
+			SharedNodeOrientation[GGroupThreadId.x+EdgeIndex] = NodeQuaternion;
+			//GroupMemoryBarrier();
+		}
+		//GroupMemoryBarrier();
+	}
+	GroupMemoryBarrier();
+	OutNodeOrientation = 	SharedNodeOrientation[GGroupThreadId.x];
+
+	/*int StrandIndex = 0, LocalIndex = 0;
 	DIHairStrands_ComputeNodeIndices(DIContext,NodeIndex,StrandIndex,LocalIndex);
 	const int RootIndex = StrandIndex * DIContext.StrandSize;
 
@@ -301,7 +368,7 @@ void DIHairStrands_ComputeNodeOrientation(in FDIHairStrandsContext DIContext, in
 		TangentNext = normalize(PositionNext - PositionPrev);
 		NodeQuaternion = NormalizeQuat( MultiplyQuat( FindQuatBetweenNormals(TangentPrev,TangentNext), NodeQuaternion ) );
 	}
-	OutNodeOrientation = NodeQuaternion;
+	OutNodeOrientation = NodeQuaternion;*/
 }
 
 /* -----------------------------------------------------------------
diff --git a/Engine/Plugins/Experimental/HairStrands/Shaders/Private/NiagaraHookeSpringMaterial.ush b/Engine/Plugins/Experimental/HairStrands/Shaders/Private/NiagaraHookeSpringMaterial.ush
--- a/Engine/Plugins/Experimental/HairStrands/Shaders/Private/NiagaraHookeSpringMaterial.ush
+++ b/Engine/Plugins/Experimental/HairStrands/Shaders/Private/NiagaraHookeSpringMaterial.ush
@@ -87,26 +87,38 @@ void SolveStretchSpringMaterial(in int StrandsSize, in float RestLength, in floa
 	}
 } 
 
-void ProjectStretchSpringMaterial(in int StrandsSize, in float YoungModulus, in float RodThickness, in float RestLength, in float DeltaTime, out float OutMaterialMultiplier)
+void ProjectStretchSpringMaterial(in int StrandsSize, in float YoungModulus, in float RodThickness, in float RestLength, in float NodeMass, in float DeltaTime, in float3 NodePosition, out float3 OutNodePosition)
 {
-	float MaterialCompliance = 0.0;
-	float MaterialWeight = 0.0;
-	OutMaterialMultiplier = 0.0;
-
 	const int LocalIndex = (GGroupThreadId.x % StrandsSize);
-	if( LocalIndex > 1 )
+
+	SharedNodePosition[GGroupThreadId.x] = NodePosition;
+	SharedPreviousPosition[GGroupThreadId.x] = NodePosition;
+	SharedInverseMass[GGroupThreadId.x] = (LocalIndex > 1 && NodeMass != 0.0) ? 1.0 / NodeMass : 0.0;
+	GroupMemoryBarrier();
+
+	if(DeltaTime != 0.0)
 	{
-		ResetStretchSpringMaterial(StrandsSize,YoungModulus,RodThickness,RestLength,DeltaTime,true,0.0,MaterialCompliance,MaterialWeight,OutMaterialMultiplier,0);
+		float MaterialCompliance = 0.0;
+		float MaterialWeight = 0.0;
+		float MaterialMultiplier = 0.0;
 
-		for(int i = 2; i < StrandsSize; ++i)
+		if( LocalIndex > 1 )
 		{
-			if( LocalIndex == i )
+			ResetStretchSpringMaterial(StrandsSize,YoungModulus,RodThickness,RestLength,DeltaTime,true,0.0,MaterialCompliance,MaterialWeight,MaterialMultiplier,0);
+
+			for(int i = 2; i < StrandsSize; ++i)
 			{
-				UpdateStretchSpringMultiplier(RestLength,DeltaTime,true,0.0,MaterialCompliance,MaterialWeight,OutMaterialMultiplier,0);
+				if( LocalIndex == i )
+				{
+					UpdateStretchSpringMultiplier(RestLength,DeltaTime,true,0.0,MaterialCompliance,MaterialWeight,MaterialMultiplier,0);
+				}
+				GroupMemoryBarrier(); 
 			}
-			GroupMemoryBarrier(); 
 		}
+		
 	}
+	GroupMemoryBarrier();
+	OutNodePosition = SharedNodePosition[GGroupThreadId.x] ;
 }
 
 /* -----------------------------------------------------------------
@@ -150,26 +162,37 @@ void SolveBendSpringMaterial(in int StrandsSize, in float RestLength, in float D
 	}
 } 
 
-void ProjectBendSpringMaterial(in int StrandsSize, in float YoungModulus, in float RodThickness, in float RestLength, in float DeltaTime, out float OutMaterialMultiplier)
+void ProjectBendSpringMaterial(in int StrandsSize, in float YoungModulus, in float RodThickness, in float RestLength, in float NodeMass, in float DeltaTime, in float3 NodePosition, out float3 OutNodePosition)
 {
-	float MaterialCompliance = 0.0;
-	float MaterialWeight = 0.0;
-	OutMaterialMultiplier = 0.0;
-
 	const int LocalIndex = (GGroupThreadId.x % StrandsSize);
-	if( LocalIndex > 1 )
+
+	SharedNodePosition[GGroupThreadId.x] = NodePosition;
+	SharedPreviousPosition[GGroupThreadId.x] = NodePosition;
+	SharedInverseMass[GGroupThreadId.x] = (LocalIndex > 1 && NodeMass != 0.0) ? 1.0 / NodeMass : 0.0;
+	GroupMemoryBarrier();
+
+	if(DeltaTime != 0.0)
 	{
-		ResetStretchSpringMaterial(StrandsSize,YoungModulus,RodThickness,RestLength,DeltaTime,true,0.0,MaterialCompliance,MaterialWeight,OutMaterialMultiplier,1);
+		float MaterialCompliance = 0.0;
+		float MaterialWeight = 0.0;
+		float MaterialMultiplier = 0.0;
 
-		for(int i = 2; i < StrandsSize; ++i)
+		if( LocalIndex > 1 )
 		{
-			if( LocalIndex == i )
+			ResetStretchSpringMaterial(StrandsSize,YoungModulus,RodThickness,RestLength,DeltaTime,true,0.0,MaterialCompliance,MaterialWeight,MaterialMultiplier,1);
+
+			for(int i = 2; i < StrandsSize; ++i)
 			{
-				UpdateStretchSpringMultiplier(RestLength,DeltaTime,true,0.0,MaterialCompliance,MaterialWeight,OutMaterialMultiplier,1);
+				if( LocalIndex == i )
+				{
+					UpdateStretchSpringMultiplier(RestLength,DeltaTime,true,0.0,MaterialCompliance,MaterialWeight,MaterialMultiplier,1);
+				}
+				GroupMemoryBarrier(); 
 			}
-			GroupMemoryBarrier(); 
 		}
 	}
+	GroupMemoryBarrier();
+	OutNodePosition = SharedNodePosition[GGroupThreadId.x];
 }
 
 /* -----------------------------------------------------------------
@@ -213,26 +236,37 @@ void SolveTwistSpringMaterial(in int StrandsSize, in float RestLength, in float
 	}
 } 
 
-void ProjectTwistSpringMaterial(in int StrandsSize, in float YoungModulus, in float RodThickness, in float RestLength, in float DeltaTime, out float OutMaterialMultiplier)
+void ProjectTwistSpringMaterial(in int StrandsSize, in float YoungModulus, in float RodThickness, in float RestLength, in float NodeMass, in float DeltaTime, in float3 NodePosition, out float3 OutNodePosition)
 {
-	float MaterialCompliance = 0.0;
-	float MaterialWeight = 0.0;
-	OutMaterialMultiplier = 0.0;
-
 	const int LocalIndex = (GGroupThreadId.x % StrandsSize);
-	if( LocalIndex > 2 )
+
+	SharedNodePosition[GGroupThreadId.x] = NodePosition;
+	SharedPreviousPosition[GGroupThreadId.x] = NodePosition;
+	SharedInverseMass[GGroupThreadId.x] = (LocalIndex > 1 && NodeMass != 0.0) ? 1.0 / NodeMass : 0.0;
+	GroupMemoryBarrier();
+
+	if(DeltaTime != 0.0)
 	{
-		ResetStretchSpringMaterial(StrandsSize,YoungModulus,RodThickness,RestLength,DeltaTime,true,0.0,MaterialCompliance,MaterialWeight,OutMaterialMultiplier,2);
+		float MaterialCompliance = 0.0;
+		float MaterialWeight = 0.0;
+		float MaterialMultiplier = 0.0;
 
-		for(int i = 3; i < StrandsSize; ++i)
+		if( LocalIndex > 2 )
 		{
-			if( LocalIndex == i )
+			ResetStretchSpringMaterial(StrandsSize,YoungModulus,RodThickness,RestLength,DeltaTime,true,0.0,MaterialCompliance,MaterialWeight,MaterialMultiplier,2);
+
+			for(int i = 3; i < StrandsSize; ++i)
 			{
-				UpdateStretchSpringMultiplier(RestLength,DeltaTime,true,0.0,MaterialCompliance,MaterialWeight,OutMaterialMultiplier,2);
+				if( LocalIndex == i )
+				{
+					UpdateStretchSpringMultiplier(RestLength,DeltaTime,true,0.0,MaterialCompliance,MaterialWeight,MaterialMultiplier,2);
+				}
+				GroupMemoryBarrier(); 
 			}
-			GroupMemoryBarrier(); 
 		}
 	}
+	GroupMemoryBarrier();
+	OutNodePosition = SharedNodePosition[GGroupThreadId.x];
 }
 
 /* -----------------------------------------------------------------
@@ -246,6 +280,29 @@ void ComputeMaterialFrame(in int StrandsSize)
 	if( LocalIndex == 0 )
 	{
 		float4 NodeQuaternion = SharedNodeOrientation[GGroupThreadId.x];
+		float3 TangentPrev = normalize(SharedNodePosition[GGroupThreadId.x+1] - SharedNodePosition[GGroupThreadId.x]);
+		float3 TangentNext = TangentPrev;
+
+		for (int EdgeIndex = 1, EdgeEnd = StrandsSize-1; EdgeIndex < EdgeEnd; ++EdgeIndex)
+		{
+			TangentPrev = TangentNext;
+			TangentNext = normalize(SharedNodePosition[GGroupThreadId.x+EdgeIndex+1] - SharedNodePosition[GGroupThreadId.x+EdgeIndex]);
+
+			NodeQuaternion = NormalizeQuat( MultiplyQuat( FindQuatBetweenNormals(TangentPrev,TangentNext), NodeQuaternion ) );
+			SharedNodeOrientation[GGroupThreadId.x+EdgeIndex] = NodeQuaternion;
+			GroupMemoryBarrier();
+		}
+		GroupMemoryBarrier();
+	}
+	GroupMemoryBarrier();
+}
+
+void UpdateMaterialFrame(in int StrandsSize)
+{
+	const int LocalIndex = (GGroupThreadId.x % StrandsSize);
+	if( LocalIndex == 0 )
+	{
+		float4 NodeQuaternion = SharedPreviousOrientation[GGroupThreadId.x];
 		float3 TangentPrev = normalize(SharedPreviousPosition[GGroupThreadId.x+1] - SharedPreviousPosition[GGroupThreadId.x]);
 		float3 TangentNext = TangentPrev;
 
@@ -253,24 +310,30 @@ void ComputeMaterialFrame(in int StrandsSize)
 		{
 			TangentPrev = TangentNext;
 			TangentNext = normalize(SharedPreviousPosition[GGroupThreadId.x+EdgeIndex+1] - SharedPreviousPosition[GGroupThreadId.x+EdgeIndex]);
+
 			NodeQuaternion = NormalizeQuat( MultiplyQuat( FindQuatBetweenNormals(TangentPrev,TangentNext), NodeQuaternion ) );
 			SharedNodeOrientation[GGroupThreadId.x+EdgeIndex] = NodeQuaternion;
 			GroupMemoryBarrier();
 		}
+		GroupMemoryBarrier();
 	}
 	GroupMemoryBarrier();
 }
 
-/*void ComputeMaterialFrame(in int StrandsSize)
+
+/*void UpdateMaterialFrame(in int StrandsSize)
 {
 	const int LocalIndex = (GGroupThreadId.x % StrandsSize);
 	if( (LocalIndex > 0) && (LocalIndex < StrandsSize-1) )
 	{
-		float4 NodeQuaternion = SharedNodeOrientation[GGroupThreadId.x];
-		float3 TangentPrev = normalize(SharedPreviousPosition[GGroupThreadId.x+1] - SharedPreviousPosition[GGroupThreadId.x]);
-		float3 TangentNext = normalize(SharedNodePosition[GGroupThreadId.x+1] - SharedNodePosition[GGroupThreadId.x]);
+		float4 NodeQuaternion = SharedPreviousOrientation[GGroupThreadId.x];
+		float3 TangentPrev = RotateVectorByQuat( float3(0,0,1), NodeQuaternion);
+		float3 TangentNext = normalize(SharedPreviousPosition[GGroupThreadId.x+1] - SharedPreviousPosition[GGroupThreadId.x]);
+
+		//float3 TangentPrev = normalize(SharedPreviousPosition[GGroupThreadId.x+1] - SharedPreviousPosition[GGroupThreadId.x]);
+		//float3 TangentNext = normalize(SharedNodePosition[GGroupThreadId.x+1] - SharedNodePosition[GGroupThreadId.x]);
 
-		SharedNodeOrientation[GGroupThreadId.x] = NormalizeQuat( MultiplyQuat( FindQuatBetweenNormals(TangentPrev,TangentNext), SharedPreviousOrientation[GGroupThreadId.x] ) );
+		SharedNodeOrientation[GGroupThreadId.x] = NormalizeQuat( MultiplyQuat( FindQuatBetweenNormals(TangentPrev,TangentNext), NodeQuaternion) );
 	}
 	GroupMemoryBarrier();
 }*/
@@ -296,7 +359,7 @@ void UpdateAngularSpringMultiplier( in float RestLength, in float3 RestDirection
 	const float3 EdgeDirection = SharedNodePosition[GGroupThreadId.x] - SharedNodePosition[GGroupThreadId.x-1];
 	const float3 DeltaVelocity = (EdgeDirection - ( SharedPreviousPosition[GGroupThreadId.x] - SharedPreviousPosition[GGroupThreadId.x-1] ) ) / DeltaTime; 
 
-	const float4 ParentFrame = SharedNodeOrientation[GGroupThreadId.x-2];
+	const float4 ParentFrame = SlerpQuat(SharedNodeOrientation[GGroupThreadId.x-2],SharedNodeOrientation[GGroupThreadId.x-2],0.5);
 	const float3 TargetDirection = RotateVectorByQuat(RestDirection,ParentFrame);
 
 	// XPBD lagrange multiplier update : dL = -(C+compliance*L) / (dC * invM * dCt + alpha)
@@ -320,49 +383,137 @@ void UpdateAngularSpringMultiplier( in float RestLength, in float3 RestDirection
 void SolveAngularSpringMaterial(in int StrandsSize, in float RestLength, in float3 RestDirection, in float DeltaTime, in float MaterialDamping,
 			in float MaterialCompliance, in float MaterialWeight, in float3 MaterialMultiplier, out float3 OutMaterialMultiplier)
 {
+	
 	OutMaterialMultiplier = MaterialMultiplier;
 	const int LocalIndex = (GGroupThreadId.x % StrandsSize);
 	if( LocalIndex > 1)
 	{
 		const int IsRed = (GGroupThreadId.x % 2) == 0;
 		// Process all the red rods
-		if (IsRed)
+		if (!IsRed)
 		{	
 			UpdateAngularSpringMultiplier(RestLength,RestDirection,DeltaTime,false,MaterialDamping,MaterialCompliance,MaterialWeight,OutMaterialMultiplier);
 		}
 		GroupMemoryBarrier(); 
 
-		if (!IsRed)
+		if (IsRed)
 		{
 			UpdateAngularSpringMultiplier(RestLength,RestDirection,DeltaTime,false,MaterialDamping,MaterialCompliance,MaterialWeight,OutMaterialMultiplier);
 		}
 		GroupMemoryBarrier(); 
 	}
-} 
-
-void ProjectAngularSpringMaterial(in int StrandsSize, in float YoungModulus, in float RodThickness, in float RestLength, in float3 RestDirection, in float DeltaTime,
-			out float3 OutMaterialMultiplier)
+}
+/*
+groupshared float4 SharedRestDatas[THREADGROUP_SIZE];
+void SolveAngularSpringMaterial(in int StrandsSize, in float RestLength, in float3 RestDirection, in float DeltaTime, in float MaterialDamping,
+			in float MaterialCompliance, in float MaterialWeight, in float3 MaterialMultiplier, out float3 OutMaterialMultiplier)
 {
-	float MaterialCompliance = 0.0;
-	float MaterialWeight = 0.0;
-	OutMaterialMultiplier = float3(0,0,0);
+	SharedRestDatas[GGroupThreadId.x] = float4(RestDirection,RestLength);
+	GroupMemoryBarrier(); 
 
-	ComputeMaterialFrame(StrandsSize);
-	SetupAngularSpringMaterial(StrandsSize, YoungModulus, RodThickness, RestLength, DeltaTime, true, 0.0, MaterialCompliance, MaterialWeight, OutMaterialMultiplier);
+	OutMaterialMultiplier = MaterialMultiplier;
+	const int LocalIndex = (GGroupThreadId.x % StrandsSize);
+	if( LocalIndex == 0)
+	{
+		for(int i = GGroupThreadId.x+2; i < GGroupThreadId.x+StrandsSize; ++i)
+		{
+			const float3 EdgeDirection = SharedNodePosition[i] - SharedNodePosition[i-1];
+
+			const float4 ParentFrame = SharedNodeOrientation[i-2];
+			const float3 TargetDirection = RotateVectorByQuat(SharedRestDatas[i].xyz,ParentFrame);
+
+			// XPBD lagrange multiplier update : dL = -(C+compliance*L) / (dC * invM * dCt + alpha)
+			const float3 DeltaLambda = -(EdgeDirection - TargetDirection) / (SharedInverseMass[i]+SharedInverseMass[i-1]);
+
+			// XPBD position update : dX = dL * dCt * invM
+			const float3 PositionDelta  = DeltaLambda;
+
+			SharedDeltaPosition[i] = PositionDelta;
+			GroupMemoryBarrier();
 
-	float3 MaterialMultiplier = float3(0,0,0);
+			// dX += dX
+			SharedNodePosition[i] += PositionDelta * SharedInverseMass[i];
+			GroupMemoryBarrier();
+			SharedNodePosition[i-1] -= PositionDelta * SharedInverseMass[i-1];
+			GroupMemoryBarrier();
+
+			if( i > GGroupThreadId.x+2 )
+			{
+				const float3 TangentInit = normalize(SharedNodePosition[i-2] - SharedNodePosition[i-3]);
+				const float3 TangentPrev = normalize(SharedNodePosition[i-1] - SharedNodePosition[i-2]);
+				const float3 TangentNext = normalize(SharedNodePosition[i] - SharedNodePosition[i-1]);
+
+				SharedNodeOrientation[i-2] = NormalizeQuat( MultiplyQuat( FindQuatBetweenNormals(TangentInit,TangentPrev), SharedNodeOrientation[i-3] ) );
+				GroupMemoryBarrier();
+				SharedNodeOrientation[i-1] = NormalizeQuat( MultiplyQuat( FindQuatBetweenNormals(TangentPrev,TangentNext), SharedNodeOrientation[i-2] ) );
+				GroupMemoryBarrier();
+			}
+			else
+			{
+				const float3 TangentPrev = normalize(SharedNodePosition[i-1] - SharedNodePosition[i-2]);
+				const float3 TangentNext = normalize(SharedNodePosition[i] - SharedNodePosition[i-1]);
+
+				SharedNodeOrientation[i-1] = NormalizeQuat( MultiplyQuat( FindQuatBetweenNormals(TangentPrev,TangentNext), SharedNodeOrientation[i-2] ) );
+				GroupMemoryBarrier();
+			}
+			GroupMemoryBarrier();
+		}
+		GroupMemoryBarrier();
+		for(int i = GGroupThreadId.x+2; i < GGroupThreadId.x+StrandsSize; ++i)
+		{
+			SharedNodePosition[i] += SharedDeltaPosition[i] * SharedInverseMass[i];
+			SharedNodePosition[i-1] -= SharedDeltaPosition[i] * SharedInverseMass[i-1];
+			GroupMemoryBarrier();
+		}
+		GroupMemoryBarrier();
+		for(int i = GGroupThreadId.x+2; i < GGroupThreadId.x+StrandsSize; ++i)
+		{
+			const float3 TangentPrev = normalize(SharedNodePosition[i-1] - SharedNodePosition[i-2]);
+			const float3 TangentNext = normalize(SharedNodePosition[i] - SharedNodePosition[i-1]);
+
+			SharedNodeOrientation[i-1] = NormalizeQuat( MultiplyQuat( FindQuatBetweenNormals(TangentPrev,TangentNext), SharedNodeOrientation[i-2] ) );
+			GroupMemoryBarrier();
+		}
+		GroupMemoryBarrier();
+	} 
+}*/
+
+void ProjectAngularSpringMaterial(in int StrandsSize, in float YoungModulus, in float RodThickness, in float RestLength, in float3 RestDirection, in float NodeMass, in float DeltaTime,
+			 in float3 NodePosition, out float3 OutNodePosition)
+{
 	const int LocalIndex = (GGroupThreadId.x % StrandsSize);
-	if( LocalIndex > 1 )
+
+	SharedNodePosition[GGroupThreadId.x] = NodePosition;
+	SharedPreviousPosition[GGroupThreadId.x] = NodePosition;
+	SharedInverseMass[GGroupThreadId.x] = (LocalIndex > 1 && NodeMass != 0.0) ? 1.0 / NodeMass : 0.0;
+	GroupMemoryBarrier();
+
+	if(DeltaTime != 0.0)
 	{
-		for(int i = 2; i < StrandsSize; ++i)
+
+		float MaterialCompliance = 0.0;
+		float MaterialWeight = 0.0;
+		float3 MaterialMultiplier = float3(0,0,0);
+
+		//ComputeMaterialFrame(StrandsSize);
+		SetupAngularSpringMaterial(StrandsSize, YoungModulus, RodThickness, RestLength, DeltaTime, true, 0.0, MaterialCompliance, MaterialWeight, MaterialMultiplier);
+
+		//if( LocalIndex > 1 )
 		{
-			if( LocalIndex == i )
+			for(int i = 2; i < StrandsSize; ++i)
 			{
-				UpdateAngularSpringMultiplier(RestLength,RestDirection,DeltaTime,true,0.0,MaterialCompliance,MaterialWeight,OutMaterialMultiplier);
+				ComputeMaterialFrame(StrandsSize);
+				if( LocalIndex == i )
+				{
+					UpdateAngularSpringMultiplier(RestLength,RestDirection,DeltaTime,true,0.0,MaterialCompliance,MaterialWeight,MaterialMultiplier);
+				}
+				GroupMemoryBarrier(); 
 			}
-			GroupMemoryBarrier(); 
 		}
 	}
+
+	GroupMemoryBarrier();
+	OutNodePosition = SharedNodePosition[GGroupThreadId.x];
 }
 
 /* -----------------------------------------------------------------
@@ -372,7 +523,7 @@ void ProjectAngularSpringMaterial(in int StrandsSize, in float YoungModulus, in
 
 void ComputeRestDirection(in int StrandsSize, in float3 NodePosition, in float4 NodeOrientation, out float3 OutRestDirection)
 {
-	SharedPreviousPosition[GGroupThreadId.x] = NodePosition;
+	SharedNodePosition[GGroupThreadId.x] = NodePosition;
 	SharedNodeOrientation[GGroupThreadId.x] = NodeOrientation;
 	GroupMemoryBarrier(); 
 
@@ -384,7 +535,7 @@ void ComputeRestDirection(in int StrandsSize, in float3 NodePosition, in float4
 	{
 		const float3 EdgeDirection = SharedNodePosition[GGroupThreadId.x] - SharedNodePosition[GGroupThreadId.x-1];
 
-		const float4 ParentFrame = SharedNodeOrientation[GGroupThreadId.x-2];
+		const float4 ParentFrame = SlerpQuat(SharedNodeOrientation[GGroupThreadId.x-2],SharedNodeOrientation[GGroupThreadId.x-2],0.5);
 		OutRestDirection = UnrotateVectorByQuat(EdgeDirection,ParentFrame);
 	}
 	GroupMemoryBarrier(); 
diff --git a/Engine/Plugins/Experimental/HairStrands/Shaders/Private/NiagaraQuaternionUtils.ush b/Engine/Plugins/Experimental/HairStrands/Shaders/Private/NiagaraQuaternionUtils.ush
--- a/Engine/Plugins/Experimental/HairStrands/Shaders/Private/NiagaraQuaternionUtils.ush
+++ b/Engine/Plugins/Experimental/HairStrands/Shaders/Private/NiagaraQuaternionUtils.ush
@@ -35,16 +35,54 @@ float3 UnrotateVectorByQuat(in float3 Vector, in float4 Quat)
 	return Vector - Quat.w * T + cross(Quat.xyz,T);
 }
 
-/*float4 FindQuatBetweenHelper(in float3 A, in float3 B, in float NormAB)
+float4 SlerpQuat(in float4 QuatA, in float4 QuatB, float Slerp)
+{
+	// Get cosine of angle between quats.
+	const float RawCosom = 
+		    QuatA.x * QuatB.x +
+			QuatA.y * QuatB.y +
+			QuatA.z * QuatB.z +
+			QuatA.w * QuatB.w;
+	// Unaligned quats - compensate, results in taking shorter route.
+	const float Cosom = (RawCosom >= 0.0) ? RawCosom : -RawCosom;
+	
+	float Scale0, Scale1;
+
+	if( Cosom < 0.9999f )
+	{	
+		const float Omega = acos(Cosom);
+		const float InvSin = 1.f/sin(Omega);
+		Scale0 = sin( (1.f - Slerp) * Omega ) * InvSin;
+		Scale1 = sin( Slerp * Omega ) * InvSin;
+	}
+	else
+	{
+		// Use linear interpolation.
+		Scale0 = 1.0f - Slerp;
+		Scale1 = Slerp;	
+	}
+
+	// In keeping with our flipped Cosom:
+	Scale1 = (RawCosom >= 0.0) ? Scale1 : -Scale1; 
+
+	float4 Result;
+		
+	Result.x = Scale0 * QuatA.x + Scale1 * QuatB.x;
+	Result.y = Scale0 * QuatA.y + Scale1 * QuatB.y;
+	Result.z = Scale0 * QuatA.z + Scale1 * QuatB.z;
+	Result.w = Scale0 * QuatA.w + Scale1 * QuatB.w;
+
+	return NormalizeQuat(Result);
+}
+
+float4 FindQuatBetweenHelper(in float3 A, in float3 B, in float NormAB)
 {
 	float W = NormAB + dot(A,B);
 	float4 Quat = (W>1e-6*NormAB) ? float4(A.y*B.z-A.z*B.y,A.z*B.x-A.x*B.z,A.x*B.y-A.y*B.x,W) :
 			(abs(A.x) > abs(A.y)) ? float4(-A.z,0.0,A.x,0.0) : float4(0.0,-A.z,A.y,0.0);
 
-	//float S = sqrt((NormAB + dot(A,B))*2.0);
-	//float4 Quat = float4(cross(A,B)/S,0.5*S);
 	return NormalizeQuat(Quat);
-}*/
+}
 
 float4 FindQuatBetweenInternal(in float3 An, in float3 Bn)
 {
@@ -52,41 +90,6 @@ float4 FindQuatBetweenInternal(in float3 An, in float3 Bn)
 	return float4(cross(An,Cn),dot(An,Cn));
 }
 
-float4 FindQuatBetweenHelper(in float3 A, in float3 B, in float NormAB)
-{
-	float3 An = normalize(A);
-	float3 Bn = normalize(B);
-
-	if( dot(An,Bn) >= 0.0)
-	{
-		return NormalizeQuat(FindQuatBetweenInternal(An,Bn));
-	}
-	else
-	{
-		float3 Cn = normalize(An+Bn);
-		if(dot(Cn,Cn) != 0.0)
-		{
-			return NormalizeQuat(MultiplyQuat(FindQuatBetweenInternal(An,Cn),FindQuatBetweenInternal(Cn,Bn)));
-		}
-		else
-		{
-			float3 A2 = An * An;
-			if(A2.x <= A2.y && A2.x <= A2.z)
-			{
-				return NormalizeQuat(float4(cross(An,float3(1,0,0)),0));
-			}
-			else if(A2.y <= A2.z)
-			{
-				return NormalizeQuat(float4(cross(An,float3(0,1,0)),0));
-			}
-			else
-			{
-				return NormalizeQuat(float4(cross(An,float3(0,0,1)),0));
-			}
-		}
-	}
-}
-
 float4 FindQuatBetweenNormals(in float3 NormalA, in float3 NormalB)
 {
 	float NormAB = 1.0;
diff --git a/Engine/Plugins/Experimental/HairStrands/Source/HairStrandsNiagara/Private/NiagaraDataInterfaceHairStrands.cpp b/Engine/Plugins/Experimental/HairStrands/Source/HairStrandsNiagara/Private/NiagaraDataInterfaceHairStrands.cpp
--- a/Engine/Plugins/Experimental/HairStrands/Source/HairStrandsNiagara/Private/NiagaraDataInterfaceHairStrands.cpp
+++ b/Engine/Plugins/Experimental/HairStrands/Source/HairStrandsNiagara/Private/NiagaraDataInterfaceHairStrands.cpp
@@ -438,8 +438,8 @@ struct FNDIHairStrandsParametersCS : public FNiagaraDataInterfaceParametersCS
 				DeformedPositionOffsetValue = HairStrandsBuffer->SourceDeformedResources->PositionOffset;
 			}
 
-			//UE_LOG(LogHairStrands, Log, TEXT("Root Offset = %f %f %f | %f %f %f "), RestRootOffsetValue[0], RestRootOffsetValue[1], RestRootOffsetValue[2],
-			//	DeformedRootOffsetValue[0], DeformedRootOffsetValue[1], DeformedRootOffsetValue[2]);
+			UE_LOG(LogHairStrands, Log, TEXT("Root Offset = %f %f %f | %f %f %f | %f %f %f"), RestRootOffsetValue[0], RestRootOffsetValue[1], RestRootOffsetValue[2],
+				DeformedRootOffsetValue[0], DeformedRootOffsetValue[1], DeformedRootOffsetValue[2], DeformedPositionOffsetValue[0], DeformedPositionOffsetValue[1], DeformedPositionOffsetValue[2]);
 
 			SetUAVParameter(RHICmdList, ComputeShaderRHI, DeformedPositionBuffer, PointPositionsUAV);
 			SetSRVParameter(RHICmdList, ComputeShaderRHI, CurvesOffsetsBuffer, HairStrandsBuffer->CurvesOffsetsBuffer.SRV);
@@ -935,6 +935,7 @@ void UNiagaraDataInterfaceHairStrands::GetFunctions(TArray<FNiagaraFunctionSigna
 		Sig.bRequiresContext = false;
 		Sig.Inputs.Add(FNiagaraVariable(FNiagaraTypeDefinition(GetClass()), TEXT("Hair Strands")));
 		Sig.Inputs.Add(FNiagaraVariable(FNiagaraTypeDefinition::GetIntDef(), TEXT("Node Index")));
+		Sig.Inputs.Add(FNiagaraVariable(FNiagaraTypeDefinition::GetVec3Def(), TEXT("Node Position")));
 		Sig.Outputs.Add(FNiagaraVariable(FNiagaraTypeDefinition::GetQuatDef(), TEXT("Node Orientation")));
 
 		OutFunctions.Add(Sig);
@@ -1209,9 +1210,11 @@ void UNiagaraDataInterfaceHairStrands::GetFunctions(TArray<FNiagaraFunctionSigna
 		Sig.Inputs.Add(FNiagaraVariable(FNiagaraTypeDefinition::GetFloatDef(), TEXT("Stretch Stiffness")));
 		Sig.Inputs.Add(FNiagaraVariable(FNiagaraTypeDefinition::GetFloatDef(), TEXT("Node Thickness")));
 		Sig.Inputs.Add(FNiagaraVariable(FNiagaraTypeDefinition::GetFloatDef(), TEXT("Rest Length")));
+		Sig.Inputs.Add(FNiagaraVariable(FNiagaraTypeDefinition::GetFloatDef(), TEXT("Node Mass")));
 		Sig.Inputs.Add(FNiagaraVariable(FNiagaraTypeDefinition::GetFloatDef(), TEXT("Delta Time")));
 		Sig.Inputs.Add(FNiagaraVariable(FNiagaraTypeDefinition::GetIntDef(), TEXT("Node Offset")));
-		Sig.Outputs.Add(FNiagaraVariable(FNiagaraTypeDefinition::GetFloatDef(), TEXT("Material Multiplier")));
+		Sig.Inputs.Add(FNiagaraVariable(FNiagaraTypeDefinition::GetVec3Def(), TEXT("Node Position")));
+		Sig.Outputs.Add(FNiagaraVariable(FNiagaraTypeDefinition::GetVec3Def(), TEXT("Node Position")));
 
 		OutFunctions.Add(Sig);
 	}
@@ -1257,8 +1260,10 @@ void UNiagaraDataInterfaceHairStrands::GetFunctions(TArray<FNiagaraFunctionSigna
 		Sig.Inputs.Add(FNiagaraVariable(FNiagaraTypeDefinition::GetFloatDef(), TEXT("Node Thickness")));
 		Sig.Inputs.Add(FNiagaraVariable(FNiagaraTypeDefinition::GetFloatDef(), TEXT("Rest Length")));
 		Sig.Inputs.Add(FNiagaraVariable(FNiagaraTypeDefinition::GetVec3Def(), TEXT("Rest Direction")));
+		Sig.Inputs.Add(FNiagaraVariable(FNiagaraTypeDefinition::GetFloatDef(), TEXT("Node Mass")));
 		Sig.Inputs.Add(FNiagaraVariable(FNiagaraTypeDefinition::GetFloatDef(), TEXT("Delta Time")));
-		Sig.Outputs.Add(FNiagaraVariable(FNiagaraTypeDefinition::GetVec3Def(), TEXT("Material Multiplier")));
+		Sig.Inputs.Add(FNiagaraVariable(FNiagaraTypeDefinition::GetVec3Def(), TEXT("Node Position")));
+		Sig.Outputs.Add(FNiagaraVariable(FNiagaraTypeDefinition::GetVec3Def(), TEXT("Node Position")));
 
 		OutFunctions.Add(Sig);
 	}
@@ -1530,7 +1535,7 @@ void UNiagaraDataInterfaceHairStrands::GetVMExternalFunction(const FVMExternalFu
 	}
 	else if (BindingInfo.Name == ComputeNodeOrientationName)
 	{
-		check(BindingInfo.GetNumInputs() == 2 && BindingInfo.GetNumOutputs() == 4);
+		check(BindingInfo.GetNumInputs() == 5 && BindingInfo.GetNumOutputs() == 4);
 		NDI_FUNC_BINDER(UNiagaraDataInterfaceHairStrands, ComputeNodeOrientation)::Bind(this, OutFunc);
 	}
 	else if (BindingInfo.Name == ComputeNodeMassName)
@@ -1635,7 +1640,7 @@ void UNiagaraDataInterfaceHairStrands::GetVMExternalFunction(const FVMExternalFu
 	}
 	else if (BindingInfo.Name == ProjectDistanceSpringMaterialName)
 	{
-		check(BindingInfo.GetNumInputs() == 6 && BindingInfo.GetNumOutputs() == 1);
+		check(BindingInfo.GetNumInputs() == 10 && BindingInfo.GetNumOutputs() == 3);
 		NDI_FUNC_BINDER(UNiagaraDataInterfaceHairStrands, ProjectDistanceSpringMaterial)::Bind(this, OutFunc);
 	}
 	else if (BindingInfo.Name == SetupAngularSpringMaterialName)
@@ -1650,7 +1655,7 @@ void UNiagaraDataInterfaceHairStrands::GetVMExternalFunction(const FVMExternalFu
 	}
 	else if (BindingInfo.Name == ProjectAngularSpringMaterialName)
 	{
-		check(BindingInfo.GetNumInputs() == 8 && BindingInfo.GetNumOutputs() == 3);
+		check(BindingInfo.GetNumInputs() == 12 && BindingInfo.GetNumOutputs() == 3);
 		NDI_FUNC_BINDER(UNiagaraDataInterfaceHairStrands, ProjectAngularSpringMaterial)::Bind(this, OutFunc);
 	}
 	else if (BindingInfo.Name == SetupStretchRodMaterialName)
@@ -2089,6 +2094,7 @@ bool UNiagaraDataInterfaceHairStrands::GetFunctionHLSL(const FName& DefinitionFu
 			void {InstanceFunctionName} (in int NodeIndex, out float3 OutNodePosition)
 			{
 				{HairStrandsContextName} DIHairStrands_ComputeNodePosition(DIContext,NodeIndex,OutNodePosition);
+				DIHairStrands_SmoothNodePosition(DIContext,NodeIndex,OutNodePosition);
 			}
 			)");
 		OutHLSL += FString::Format(FormatSample, ArgsSample);
@@ -2097,9 +2103,9 @@ bool UNiagaraDataInterfaceHairStrands::GetFunctionHLSL(const FName& DefinitionFu
 	else if (DefinitionFunctionName == ComputeNodeOrientationName)
 	{
 		static const TCHAR *FormatSample = TEXT(R"(
-			void {InstanceFunctionName} (in int NodeIndex, out float4 OutNodeOrientation)
+			void {InstanceFunctionName} (in int NodeIndex, in float3 NodePosition, out float4 OutNodeOrientation)
 			{
-				{HairStrandsContextName} DIHairStrands_ComputeNodeOrientation(DIContext,NodeIndex,OutNodeOrientation);
+				{HairStrandsContextName} DIHairStrands_ComputeNodeOrientation(DIContext,NodeIndex,NodePosition,OutNodeOrientation);
 			}
 			)");
 		OutHLSL += FString::Format(FormatSample, ArgsSample);
@@ -2360,20 +2366,20 @@ in float RestLength, in float DeltaTime, in int NodeOffset, in float MaterialDam
 	else if (DefinitionFunctionName == ProjectDistanceSpringMaterialName)
 	{
 		static const TCHAR *FormatSample = TEXT(R"(
-				void {InstanceFunctionName} (in float YoungModulus, in float RodThickness, in float RestLength, in float DeltaTime, in int NodeOffset, out float OutMaterialMultiplier)
+				void {InstanceFunctionName} (in float YoungModulus, in float RodThickness, in float RestLength, in float NodeMass, in float DeltaTime, in int NodeOffset, in float3 NodePosition, out float3 OutNodePosition)
 				{
 					{HairStrandsContextName} 
 					if(NodeOffset == 0)
 					{
-						ProjectStretchSpringMaterial(DIContext.StrandSize,YoungModulus,RodThickness,RestLength,DeltaTime,OutMaterialMultiplier);
+						ProjectStretchSpringMaterial(DIContext.StrandSize,YoungModulus,RodThickness,RestLength,NodeMass,DeltaTime,NodePosition,OutNodePosition);
 					}
 					if(NodeOffset == 1)
 					{
-						ProjectBendSpringMaterial(DIContext.StrandSize,YoungModulus,RodThickness,RestLength,DeltaTime,OutMaterialMultiplier);
+						ProjectBendSpringMaterial(DIContext.StrandSize,YoungModulus,RodThickness,RestLength,NodeMass,DeltaTime,NodePosition,OutNodePosition);
 					}
 					if(NodeOffset == 2)
 					{
-						ProjectTwistSpringMaterial(DIContext.StrandSize,YoungModulus,RodThickness,RestLength,DeltaTime,OutMaterialMultiplier);
+						ProjectTwistSpringMaterial(DIContext.StrandSize,YoungModulus,RodThickness,RestLength,NodeMass,DeltaTime,NodePosition,OutNodePosition);
 					}
 				}
 				)");
@@ -2407,9 +2413,9 @@ in float RestLength, in float DeltaTime, in int NodeOffset, in float MaterialDam
 	else if (DefinitionFunctionName == ProjectAngularSpringMaterialName)
 	{
 		static const TCHAR *FormatSample = TEXT(R"(
-					void {InstanceFunctionName} (in float YoungModulus, in float RodThickness, in float RestLength, in float3 RestDirection, in float DeltaTime, out float OutMaterialMultiplier)
+					void {InstanceFunctionName} (in float YoungModulus, in float RodThickness, in float RestLength, in float3 RestDirection, in float NodeMass, in float DeltaTime, in float3 NodePosition, out float3 OutNodePosition)
 					{
-						{HairStrandsContextName} ProjectAngularSpringMaterial(DIContext.StrandSize,YoungModulus,RodThickness,RestLength,RestDirection,DeltaTime,OutMaterialMultiplier);
+						{HairStrandsContextName} ProjectAngularSpringMaterial(DIContext.StrandSize,YoungModulus,RodThickness,RestLength,RestDirection,NodeMass,DeltaTime,NodePosition,OutNodePosition);
 					}
 					)");
 		OutHLSL += FString::Format(FormatSample, ArgsSample);
@@ -2529,7 +2535,7 @@ in float RestLength, in float DeltaTime, in int NodeOffset, in float MaterialDam
 		static const TCHAR *FormatSample = TEXT(R"(
 					void {InstanceFunctionName} ( out bool OutUpdateStatus)
 					{
-						{HairStrandsContextName} ComputeMaterialFrame(DIContext.StrandSize);
+						{HairStrandsContextName} UpdateMaterialFrame(DIContext.StrandSize);
 						OutUpdateStatus = true;
 					}
 					)");
