remove log in hair strands data interface

diff --git a/Engine/Plugins/Experimental/HairStrands/Source/HairStrandsNiagara/Private/NiagaraDataInterfaceHairStrands.cpp b/Engine/Plugins/Experimental/HairStrands/Source/HairStrandsNiagara/Private/NiagaraDataInterfaceHairStrands.cpp
--- a/Engine/Plugins/Experimental/HairStrands/Source/HairStrandsNiagara/Private/NiagaraDataInterfaceHairStrands.cpp
+++ b/Engine/Plugins/Experimental/HairStrands/Source/HairStrandsNiagara/Private/NiagaraDataInterfaceHairStrands.cpp
@@ -438,8 +438,8 @@ struct FNDIHairStrandsParametersCS : public FNiagaraDataInterfaceParametersCS
 				DeformedPositionOffsetValue = HairStrandsBuffer->SourceDeformedResources->PositionOffset;
 			}
 
-			UE_LOG(LogHairStrands, Log, TEXT("Root Offset = %f %f %f | %f %f %f | %f %f %f"), RestRootOffsetValue[0], RestRootOffsetValue[1], RestRootOffsetValue[2],
-				DeformedRootOffsetValue[0], DeformedRootOffsetValue[1], DeformedRootOffsetValue[2], DeformedPositionOffsetValue[0], DeformedPositionOffsetValue[1], DeformedPositionOffsetValue[2]);
+			//UE_LOG(LogHairStrands, Log, TEXT("Root Offset = %f %f %f | %f %f %f | %f %f %f"), RestRootOffsetValue[0], RestRootOffsetValue[1], RestRootOffsetValue[2],
+			//	DeformedRootOffsetValue[0], DeformedRootOffsetValue[1], DeformedRootOffsetValue[2], DeformedPositionOffsetValue[0], DeformedPositionOffsetValue[1], DeformedPositionOffsetValue[2]);
 
 			SetUAVParameter(RHICmdList, ComputeShaderRHI, DeformedPositionBuffer, PointPositionsUAV);
 			SetSRVParameter(RHICmdList, ComputeShaderRHI, CurvesOffsetsBuffer, HairStrandsBuffer->CurvesOffsetsBuffer.SRV);
