Fix CIS issue.

diff --git a/Engine/Plugins/Experimental/HairStrands/Source/HairStrandsCore/Private/HairStrandsRendering.cpp b/Engine/Plugins/Experimental/HairStrands/Source/HairStrandsCore/Private/HairStrandsRendering.cpp
--- a/Engine/Plugins/Experimental/HairStrands/Source/HairStrandsCore/Private/HairStrandsRendering.cpp
+++ b/Engine/Plugins/Experimental/HairStrands/Source/HairStrandsCore/Private/HairStrandsRendering.cpp
@@ -223,7 +223,6 @@ static void AddHairStrandsInterpolationPass(
 	const FUnorderedAccessViewRHIRef& OutRenderPositionBuffer,
 	const FUnorderedAccessViewRHIRef& OutRenderAttributeBuffer)
 {
-	const bool bCopySimAttributesToRenderAttributes = SimAttributeBuffer != nullptr && OutRenderAttributeBuffer != nullptr;
 	const uint32 GroupSize = ComputeGroupSize();
 	const FIntVector DispatchCount = ComputeDispatchCount(VertexCount, GroupSize);
 
@@ -235,8 +234,21 @@ static void AddHairStrandsInterpolationPass(
 	Parameters->Interpolation1Buffer = Interpolation1Buffer;
 	Parameters->OutRenderDeformedPositionBuffer = OutRenderPositionBuffer;
 	Parameters->HairStrandsCullIndex = FIntPoint(-1, -1);
+	Parameters->VertexCount = VertexCount;
+	Parameters->InRenderHairPositionOffset = InRenderHairWorldOffset;
+	Parameters->InSimHairPositionOffset = InSimHairWorldOffset;
+	Parameters->OutHairPositionOffset = OutHairWorldOffset;
+	Parameters->DispatchCountX = DispatchCount.X;
+	
+	const bool bIsVertexToCurveBuffersValid = InRenHairData.VertexToCurveIndexBuffer && InSimHairData.VertexToCurveIndexBuffer;
+	if (bIsVertexToCurveBuffersValid)
+	{
+		Parameters->RenVertexToRootIndexBuffer = InRenHairData.VertexToCurveIndexBuffer->SRV;
+		Parameters->SimVertexToRootIndexBuffer = InSimHairData.VertexToCurveIndexBuffer->SRV;
+	}
+
 	const FHairCullInfo Info = GetHairStrandsCullInfo();
-	const bool bCullingEnable = Info.CullMode != EHairCullMode::None && InRenHairData.VertexToCurveIndexBuffer && InSimHairData.VertexToCurveIndexBuffer;
+	const bool bCullingEnable = Info.CullMode != EHairCullMode::None && bIsVertexToCurveBuffersValid;
 	if (bCullingEnable)
 	{
 		if (Info.CullMode == EHairCullMode::Sim)
@@ -244,24 +256,15 @@ static void AddHairStrandsInterpolationPass(
 		if (Info.CullMode == EHairCullMode::Render)
 			Parameters->HairStrandsCullIndex.X = Info.ExplicitIndex >= 0 ? Info.ExplicitIndex : FMath::Clamp(uint32(Info.NormalizedIndex * InRenHairData.RootCount), 0u, InRenHairData.RootCount - 1);
 	}
+
+	const bool bCopySimAttributesToRenderAttributes = SimAttributeBuffer != nullptr && OutRenderAttributeBuffer != nullptr;
 	if (bCopySimAttributesToRenderAttributes)
 	{
 		Parameters->SimAttributeBuffer = SimAttributeBuffer;
 		Parameters->OutRenderAttributeBuffer = OutRenderAttributeBuffer;
 	}
-	Parameters->VertexCount = VertexCount;
-	Parameters->InRenderHairPositionOffset = InRenderHairWorldOffset;
-	Parameters->InSimHairPositionOffset = InSimHairWorldOffset;
-	Parameters->OutHairPositionOffset = OutHairWorldOffset;
-	Parameters->DispatchCountX = DispatchCount.X;
-
-	const bool bSupportDynamicMesh = InRenHairData.RootCount > 0 && LODIndex >= 0 && LODIndex < InRenHairData.LODDatas.Num() && InRenHairData.LODDatas[LODIndex].bIsValid;
-	if (bSupportDynamicMesh || bCullingEnable)
-	{
-		Parameters->RenVertexToRootIndexBuffer = InRenHairData.VertexToCurveIndexBuffer->SRV;
-		Parameters->SimVertexToRootIndexBuffer = InSimHairData.VertexToCurveIndexBuffer->SRV;
-	}
 
+	const bool bSupportDynamicMesh = InRenHairData.RootCount > 0 && LODIndex >= 0 && LODIndex < InRenHairData.LODDatas.Num() && InRenHairData.LODDatas[LODIndex].bIsValid && bIsVertexToCurveBuffersValid;
 	if (bSupportDynamicMesh)
 	{
 		Parameters->RestPositionOffset = InRenHairData.LODDatas[LODIndex].RestPositionOffset;
