Expose smoothing to the user
Fix sim when number of steps = 1

diff --git a/Engine/Plugins/Experimental/HairStrands/Shaders/Private/NiagaraDataInterfaceHairStrands.ush b/Engine/Plugins/Experimental/HairStrands/Shaders/Private/NiagaraDataInterfaceHairStrands.ush
--- a/Engine/Plugins/Experimental/HairStrands/Shaders/Private/NiagaraDataInterfaceHairStrands.ush
+++ b/Engine/Plugins/Experimental/HairStrands/Shaders/Private/NiagaraDataInterfaceHairStrands.ush
@@ -278,7 +278,7 @@ void DIHairStrands_ComputeNodePosition(in FDIHairStrandsContext DIContext, in in
 }
 
 // Compute the node position by linear interpolation over the points
-void DIHairStrands_SmoothNodePosition(in FDIHairStrandsContext DIContext, in int NodeIndex, inout float3 OutNodePosition)
+void DIHairStrands_SmoothNodePosition(in FDIHairStrandsContext DIContext, in int NodeIndex, in float SmoothingFilter, inout float3 OutNodePosition)
 {
 	SharedNodePosition[GGroupThreadId.x] = OutNodePosition;
 	GroupMemoryBarrier();
@@ -286,7 +286,7 @@ void DIHairStrands_SmoothNodePosition(in FDIHairStrandsContext DIContext, in int
 	const int LocalIndex = (GGroupThreadId.x % DIContext.StrandSize);
 	if(LocalIndex == 0)
 	{
-		const float Beta = 0.2;
+		const float Beta = SmoothingFilter;
 		float3 DirM1 = SharedNodePosition[GGroupThreadId.x+1] - SharedNodePosition[GGroupThreadId.x];
 		float3 DirM2 = DirM1;
 
diff --git a/Engine/Plugins/Experimental/HairStrands/Source/HairStrandsNiagara/Private/NiagaraDataInterfaceHairStrands.cpp b/Engine/Plugins/Experimental/HairStrands/Source/HairStrandsNiagara/Private/NiagaraDataInterfaceHairStrands.cpp
--- a/Engine/Plugins/Experimental/HairStrands/Source/HairStrandsNiagara/Private/NiagaraDataInterfaceHairStrands.cpp
+++ b/Engine/Plugins/Experimental/HairStrands/Source/HairStrandsNiagara/Private/NiagaraDataInterfaceHairStrands.cpp
@@ -924,6 +924,7 @@ void UNiagaraDataInterfaceHairStrands::GetFunctions(TArray<FNiagaraFunctionSigna
 		Sig.bRequiresContext = false;
 		Sig.Inputs.Add(FNiagaraVariable(FNiagaraTypeDefinition(GetClass()), TEXT("Hair Strands")));
 		Sig.Inputs.Add(FNiagaraVariable(FNiagaraTypeDefinition::GetIntDef(), TEXT("Node Index")));
+		Sig.Inputs.Add(FNiagaraVariable(FNiagaraTypeDefinition::GetFloatDef(), TEXT("Smoothing Filter")));
 		Sig.Outputs.Add(FNiagaraVariable(FNiagaraTypeDefinition::GetVec3Def(), TEXT("Node Position")));
 
 		OutFunctions.Add(Sig);
@@ -1530,7 +1531,7 @@ void UNiagaraDataInterfaceHairStrands::GetVMExternalFunction(const FVMExternalFu
 	}
 	else if (BindingInfo.Name == ComputeNodePositionName)
 	{
-		check(BindingInfo.GetNumInputs() == 2 && BindingInfo.GetNumOutputs() == 3);
+		check(BindingInfo.GetNumInputs() == 3 && BindingInfo.GetNumOutputs() == 3);
 		NDI_FUNC_BINDER(UNiagaraDataInterfaceHairStrands, ComputeNodePosition)::Bind(this, OutFunc);
 	}
 	else if (BindingInfo.Name == ComputeNodeOrientationName)
@@ -2091,10 +2092,10 @@ bool UNiagaraDataInterfaceHairStrands::GetFunctionHLSL(const FName& DefinitionFu
 	else if (DefinitionFunctionName == ComputeNodePositionName)
 	{
 		static const TCHAR *FormatSample = TEXT(R"(
-			void {InstanceFunctionName} (in int NodeIndex, out float3 OutNodePosition)
+			void {InstanceFunctionName} (in int NodeIndex, in float SmoothingFilter, out float3 OutNodePosition)
 			{
 				{HairStrandsContextName} DIHairStrands_ComputeNodePosition(DIContext,NodeIndex,OutNodePosition);
-				DIHairStrands_SmoothNodePosition(DIContext,NodeIndex,OutNodePosition);
+				DIHairStrands_SmoothNodePosition(DIContext,NodeIndex,SmoothingFilter,OutNodePosition);
 			}
 			)");
 		OutHLSL += FString::Format(FormatSample, ArgsSample);
