Fix interpolation quantization due to lack of guide inter-vertex interpolation during motion interpolation.

diff --git a/Engine/Plugins/Experimental/HairStrands/Source/HairStrandsCore/Private/GroomBuilder.cpp b/Engine/Plugins/Experimental/HairStrands/Source/HairStrandsCore/Private/GroomBuilder.cpp
--- a/Engine/Plugins/Experimental/HairStrands/Source/HairStrandsCore/Private/GroomBuilder.cpp
+++ b/Engine/Plugins/Experimental/HairStrands/Source/HairStrandsCore/Private/GroomBuilder.cpp
@@ -419,6 +419,8 @@ namespace HairInterpolationBuilder
 		{
 			ClosestGuides.Indices[1] = Metric.KClosestGuideIndices[RandIndex0];
 			ClosestGuides.Indices[2] = Metric.KClosestGuideIndices[RandIndex0];
+			RandIndex1 = RandIndex0;
+			RandIndex2 = RandIndex0;
 		}
 
 
@@ -771,6 +773,44 @@ namespace HairInterpolationBuilder
 		}
 	}
 
+	struct FVertexInterpolationDesc
+	{
+		uint32 Index0 = 0;
+		uint32 Index1 = 0;
+		float T = 0;
+	};
+
+	// Find the vertex along a sim curve 'SimCurveIndex', which has the same parametric distance than the render distance 'RenPointDistance'
+	static FVertexInterpolationDesc FindMatchingVertex(const float RenPointDistance, const FHairStrandsDatas& SimStrandsData, const uint32 SimCurveIndex)
+	{
+		const uint32 SimOffset = SimStrandsData.StrandsCurves.CurvesOffset[SimCurveIndex];
+
+		const float CurveLength = SimStrandsData.StrandsCurves.CurvesLength[SimCurveIndex] * SimStrandsData.StrandsCurves.MaxLength;
+
+		// Find with with vertex the vertex should be paired
+		const uint32 SimPointCount = SimStrandsData.StrandsCurves.CurvesCount[SimCurveIndex];
+		for (uint32 SimPointIndex = 0; SimPointIndex < SimPointCount-1; ++SimPointIndex)
+		{
+			const float SimPointDistance0 = SimStrandsData.StrandsPoints.PointsCoordU[SimPointIndex + SimOffset] * CurveLength;
+			const float SimPointDistance1 = SimStrandsData.StrandsPoints.PointsCoordU[SimPointIndex + SimOffset + 1] * CurveLength;
+			if (SimPointDistance0 <= RenPointDistance && RenPointDistance <= SimPointDistance1)
+			{
+				const float SegmentLength = SimPointDistance1 - SimPointDistance0;
+				FVertexInterpolationDesc Out;
+				Out.Index0 = SimPointIndex;
+				Out.Index1 = SimPointIndex+1;
+				Out.T = (RenPointDistance - SimPointDistance0) / (SegmentLength>0? SegmentLength : 1);
+				Out.T = FMath::Clamp(Out.T, 0.f, 1.f);
+				return Out;
+			}
+		}
+		FVertexInterpolationDesc Desc;
+		Desc.Index0 = SimPointCount-2;
+		Desc.Index1 = SimPointCount-1;
+		Desc.T = 1;
+		return Desc;
+	}
+
 	enum class EHairInterpolationDataQuality
 	{
 		Low,
@@ -882,55 +922,45 @@ namespace HairInterpolationBuilder
 				#if WEIGHT_METHOD == 0
 					const uint32 SimCurveIndex = ClosestGuides.Indices[KIndex];
 					const uint32 SimOffset = SimStrandsData.StrandsCurves.CurvesOffset[SimCurveIndex];
-					float PrevSimPointDistance = 0;
-					uint32 ClosestSimPointIndex = 0;
-					const uint32 SimPointCount = SimStrandsData.StrandsCurves.CurvesCount[SimCurveIndex];
-					for (uint32 SimPointIndex = 0; SimPointIndex < SimPointCount; ++SimPointIndex, ++ClosestSimPointIndex)
-					{
-						const float SimPointDistance = SimStrandsData.StrandsPoints.PointsCoordU[SimPointIndex + SimOffset] * SimStrandsData.StrandsCurves.CurvesLength[SimCurveIndex] * SimStrandsData.StrandsCurves.MaxLength;
-						if (RenPointDistance >= PrevSimPointDistance && RenPointDistance <= SimPointDistance)
-						{
-							//const float GuideDistance = Guide.Points[SimPointIndex].Distance;
-							const float D0 = FMath::Abs(PrevSimPointDistance - RenPointDistance);
-							const float D1 = FMath::Abs(SimPointDistance - RenPointDistance);
-							if (D0 > D1)
-							{
-								ClosestSimPointIndex = SimPointIndex;
-							}
-							break;
-						}
-						PrevSimPointDistance = SimPointDistance;
-					}
-					check(SimPointCount > 0);
-					ClosestSimPointIndex = FMath::Clamp(ClosestSimPointIndex, 0u, uint32(SimPointCount - 1));
+					const FVertexInterpolationDesc Desc = FindMatchingVertex(RenPointDistance, SimStrandsData, SimCurveIndex);
+					const FVector& SimPointPosition0 = SimStrandsData.StrandsPoints.PointsPosition[Desc.Index0 + SimOffset];
+					const FVector& SimPointPosition1 = SimStrandsData.StrandsPoints.PointsPosition[Desc.Index1 + SimOffset];
+					const float Weight = 1.0f / FMath::Max(MinWeightDistance, FVector::Distance(RenPointPosition, FMath::Lerp(SimPointPosition0, SimPointPosition1, Desc.T)));
 
-					const FVector& SimPointPosition = SimStrandsData.StrandsPoints.PointsPosition[ClosestSimPointIndex + SimOffset];
-					const float Weight = 1.0f / FMath::Max(MinWeightDistance, FVector::Distance(RenPointPosition, SimPointPosition));
 					InterpolationData.PointsSimCurvesIndex[PointGlobalIndex][KIndex] = SimCurveIndex;
-					InterpolationData.PointsSimCurvesVertexIndex[PointGlobalIndex][KIndex] = ClosestSimPointIndex + SimOffset;
+					InterpolationData.PointsSimCurvesVertexIndex[PointGlobalIndex][KIndex] = Desc.Index0 + SimOffset;
+					InterpolationData.PointsSimCurvesVertexLerp[PointGlobalIndex][KIndex] = Desc.T;
 					InterpolationData.PointsSimCurvesVertexWeights[PointGlobalIndex][KIndex] = Weight;
 				#endif
 
 					// Use only the root as a *constant* weight for deformation along each vertex
+					// Still compute the closest vertex (in parametric distance) to know on which vertex the offset/delta should be computed
 				#if WEIGHT_METHOD == 1
-					const uint32 SimCurveIndex = ClosestGuideIndices[KIndex];
+					const uint32 SimCurveIndex = ClosestGuides.Indices[KIndex];
 					const uint32 SimOffset = SimStrandsData.StrandsCurves.CurvesOffset[SimCurveIndex];
 					const FVector& SimRootPointPosition = SimStrandsData.StrandsPoints.PointsPosition[SimOffset];
-					const float Weight = 1.0f / FMath::Max(MinWeightDistance, FVector::Distance(RenPointPosition, SimRootPointPosition));
+					const FVector& RenRootPointPosition = RenStrandsData.StrandsPoints.PointsPosition[RenOffset];
+					const float Weight = 1.0f / FMath::Max(MinWeightDistance, FVector::Distance(RenRootPointPosition, SimRootPointPosition));
+					const FVertexInterpolationDesc Desc = FindMatchingVertex(RenPointDistance, SimStrandsData, SimCurveIndex);
+
 					InterpolationData.PointsSimCurvesIndex[PointGlobalIndex][KIndex] = SimCurveIndex;
-					InterpolationData.PointsSimCurvesVertexIndex[PointGlobalIndex][KIndex] = SimOffset;
+					InterpolationData.PointsSimCurvesVertexIndex[PointGlobalIndex][KIndex] = Desc.Index0 + SimOffset;
+					InterpolationData.PointsSimCurvesVertexLerp[PointGlobalIndex][KIndex] = Desc.T;
 					InterpolationData.PointsSimCurvesVertexWeights[PointGlobalIndex][KIndex] = Weight;
 				#endif
 
 					// Use the *same vertex index* to match guide vertex with strand vertex
 				#if WEIGHT_METHOD == 2
 					check(SimPointCount > 0);
-					const uint32 SimCurveIndex = ClosestGuideIndices[KIndex];
+					const uint32 SimCurveIndex = ClosestGuides.Indices[KIndex];
 					const uint32 SimPointIndex = FMath::Clamp(RenPointIndex, 0, SimPointCount - 1);
 					const FVector& SimPointPosition = SimStrandsData.StrandsPoints.PointsPosition[SimPointIndex + SimOffset];
 					const float Weight = 1.0f / FMath::Max(MinWeightDistance, FVector::Distance(RenPointPosition, SimPointPosition));
+					const FVertexInterpolationDesc Desc = FindMatchingVertex(RenPointDistance, SimStrandsData, SimCurveIndex);
+
 					InterpolationData.PointsSimCurvesIndex[PointGlobalIndex][KIndex] = SimCurveIndex;
-					InterpolationData.PointsSimCurvesVertexIndex[PointGlobalIndex][KIndex] = SimPointIndex + SimOffset;
+					InterpolationData.PointsSimCurvesVertexIndex[PointGlobalIndex][KIndex] = Desc.Index0 + SimOffset;
+					InterpolationData.PointsSimCurvesVertexLerp[PointGlobalIndex][KIndex] = Desc.T;
 					InterpolationData.PointsSimCurvesVertexWeights[PointGlobalIndex][KIndex] = Weight;
 				#endif
 
@@ -967,6 +997,7 @@ namespace HairInterpolationBuilder
 		{
 			const FIntVector& Indices = HairInterpolation.PointsSimCurvesVertexIndex[PointIndex];
 			const FVector& Weights = HairInterpolation.PointsSimCurvesVertexWeights[PointIndex];
+			const FVector& S = HairInterpolation.PointsSimCurvesVertexLerp[PointIndex];
 
 			FHairStrandsInterpolation0Format::Type& OutInterp0 = OutPointsInterpolation0[PointIndex];
 			OutInterp0.Index0 = LowerPart(Indices[0]);
@@ -979,7 +1010,11 @@ namespace HairInterpolationBuilder
 			OutInterp1.VertexIndex0 = UpperPart(Indices[0]);
 			OutInterp1.VertexIndex1 = UpperPart(Indices[1]);
 			OutInterp1.VertexIndex2 = UpperPart(Indices[2]);
-			OutInterp1.Pad0 = 0;
+			OutInterp1.VertexLerp0  = S[0] * 255.f;
+			OutInterp1.VertexLerp1  = S[1] * 255.f;
+			OutInterp1.VertexLerp2  = S[2] * 255.f;
+			OutInterp1.Pad0			= 0;
+			OutInterp1.Pad1			= 0;
 		}	
 	}
 }
diff --git a/Engine/Plugins/Experimental/HairStrands/Source/HairStrandsCore/Private/HairStrandsDatas.cpp b/Engine/Plugins/Experimental/HairStrands/Source/HairStrandsCore/Private/HairStrandsDatas.cpp
--- a/Engine/Plugins/Experimental/HairStrands/Source/HairStrandsCore/Private/HairStrandsDatas.cpp
+++ b/Engine/Plugins/Experimental/HairStrands/Source/HairStrandsCore/Private/HairStrandsDatas.cpp
@@ -6,6 +6,7 @@
 void FHairStrandsInterpolationDatas::SetNum(const uint32 NumCurves)
 {
 	PointsSimCurvesVertexWeights.SetNum(NumCurves);
+	PointsSimCurvesVertexLerp.SetNum(NumCurves);
 	PointsSimCurvesVertexIndex.SetNum(NumCurves);
 	PointsSimCurvesIndex.SetNum(NumCurves);
 }
@@ -30,6 +31,7 @@ void FHairStrandsPoints::SetNum(const uint32 NumPoints)
 void FHairStrandsInterpolationDatas::Reset()
 {
 	PointsSimCurvesVertexWeights.Reset();
+	PointsSimCurvesVertexLerp.Reset();
 	PointsSimCurvesVertexIndex.Reset();
 	PointsSimCurvesIndex.Reset();
 }
@@ -123,10 +125,37 @@ FArchive& operator<<(FArchive& Ar, FHairInterpolation0Vertex& Vertex)
 
 FArchive& operator<<(FArchive& Ar, FHairInterpolation1Vertex& Vertex)
 {
-	Ar << Vertex.VertexIndex0;
-	Ar << Vertex.VertexIndex1;
-	Ar << Vertex.VertexIndex2;
-	Ar << Vertex.Pad0;
+	Ar.UsingCustomVersion(FReleaseObjectVersion::GUID);
+
+	if (Ar.CustomVer(FReleaseObjectVersion::GUID) >= FReleaseObjectVersion::GroomAssetVersion3)
+	{
+		Ar << Vertex.VertexIndex0;
+		Ar << Vertex.VertexIndex1;
+		Ar << Vertex.VertexIndex2;
+
+		Ar << Vertex.VertexLerp0;
+		Ar << Vertex.VertexLerp1;
+		Ar << Vertex.VertexLerp2;
+
+		Ar << Vertex.Pad0;
+		Ar << Vertex.Pad1;
+	}
+	else
+	{
+		Ar << Vertex.VertexIndex0;
+		Ar << Vertex.VertexIndex1;
+		Ar << Vertex.VertexIndex2;
+
+		uint8 Pad0 = 0;
+		Ar << Pad0;
+
+		if (Ar.IsLoading())
+		{
+			Vertex.VertexLerp0 = 0;
+			Vertex.VertexLerp1 = 0;
+			Vertex.VertexLerp2 = 0;
+		}
+	}
 
 	return Ar;
 }
@@ -139,10 +168,26 @@ void FHairStrandsInterpolationDatas::FRenderData::Serialize(FArchive& Ar)
 
 void FHairStrandsInterpolationDatas::Serialize(FArchive& Ar)
 {
+	Ar.UsingCustomVersion(FReleaseObjectVersion::GUID);
+
 	Ar << PointsSimCurvesVertexWeights;
 	Ar << PointsSimCurvesVertexIndex;
 	Ar << PointsSimCurvesIndex;
 
+	if (Ar.CustomVer(FReleaseObjectVersion::GUID) >= FReleaseObjectVersion::GroomAssetVersion3)
+	{
+		Ar << PointsSimCurvesVertexLerp;
+	}
+	else if (Ar.IsLoading())
+	{
+		const uint32 ElmentCount = PointsSimCurvesVertexIndex.Num();
+		PointsSimCurvesVertexLerp.SetNum(ElmentCount);
+		for (FVector& S : PointsSimCurvesVertexLerp)
+		{
+			S = FVector::ZeroVector;
+		}
+	}
+
 	RenderData.Serialize(Ar);
 }
 
diff --git a/Engine/Plugins/Experimental/HairStrands/Source/HairStrandsCore/Public/HairStrandsDatas.h b/Engine/Plugins/Experimental/HairStrands/Source/HairStrandsCore/Public/HairStrandsDatas.h
--- a/Engine/Plugins/Experimental/HairStrands/Source/HairStrandsCore/Public/HairStrandsDatas.h
+++ b/Engine/Plugins/Experimental/HairStrands/Source/HairStrandsCore/Public/HairStrandsDatas.h
@@ -16,8 +16,6 @@ struct FPackedHairVertex
 	uint8 ControlPointType : 2;
 	uint8 NormalizedRadius : 6;
 	uint8 NormalizedLength;
-
-	friend FArchive& operator<<(FArchive& Ar, FPackedHairVertex& Vertex);
 };
 
 struct FPackedHairAttributeVertex
@@ -31,8 +29,6 @@ struct FPackedHairAttributeVertex
 	uint8 IndexV;
 	uint8 Unused0;
 	uint8 Unused1;
-
-	friend FArchive& operator<<(FArchive& Ar, FPackedHairAttributeVertex& Vertex);
 };
 
 struct FHairMaterialVertex
@@ -43,8 +39,6 @@ struct FHairMaterialVertex
 	uint8 BaseColorB;
 
 	uint8 Roughness;
-
-	friend FArchive& operator<<(FArchive& Ar, FHairMaterialVertex& Vertex);
 };
 
 struct FHairInterpolation0Vertex
@@ -55,8 +49,6 @@ struct FHairInterpolation0Vertex
 
 	uint8 VertexWeight0;
 	uint8 VertexWeight1;
-
-	friend FArchive& operator<<(FArchive& Ar, FPackedHairVertex& FHairInterpolation0Vertex);
 };
 
 struct FHairInterpolation1Vertex
@@ -64,9 +56,13 @@ struct FHairInterpolation1Vertex
 	uint8 VertexIndex0;
 	uint8 VertexIndex1;
 	uint8 VertexIndex2;
-	uint8 Pad0;
 
-	friend FArchive& operator<<(FArchive& Ar, FPackedHairVertex& FHairInterpolation1Vertex);
+	uint8 VertexLerp0;
+	uint8 VertexLerp1;
+	uint8 VertexLerp2;
+
+	uint8 Pad0;
+	uint8 Pad1;
 };
 
 struct FVector4_16
@@ -131,8 +127,8 @@ struct FHairStrandsInterpolation1Format
 
 	static const uint32 ComponentCount = 1;
 	static const uint32 SizeInByte = sizeof(Type);
-	static const EVertexElementType VertexElementType = VET_UByte4;
-	static const EPixelFormat Format = PF_R8G8B8A8_UINT;
+	static const EVertexElementType VertexElementType = VET_UShort4;
+	static const EPixelFormat Format = PF_R16G16B16A16_UINT;
 };
 
 struct FHairStrandsRootIndexFormat
@@ -256,6 +252,9 @@ struct HAIRSTRANDSCORE_API FHairStrandsInterpolationDatas
 	/** Closest vertex indices on simulation curve, ordered by closest influence */
 	TArray<FIntVector> PointsSimCurvesVertexIndex;
 
+	/** Lerp value between the closest vertex indices and the next one, ordered by closest influence */
+	TArray<FVector> PointsSimCurvesVertexLerp;
+
 	/** Weight of vertex indices on simulation curve, ordered by closest influence */
 	TArray<FVector>	PointsSimCurvesVertexWeights;
 
diff --git a/Engine/Shaders/Private/HairStrands/HairStrandsInterpolation.usf b/Engine/Shaders/Private/HairStrands/HairStrandsInterpolation.usf
--- a/Engine/Shaders/Private/HairStrands/HairStrandsInterpolation.usf
+++ b/Engine/Shaders/Private/HairStrands/HairStrandsInterpolation.usf
@@ -66,6 +66,30 @@ RWBuffer<float4> OutRenderAttributeBuffer;
 	#error Unknown group size	
 #endif
 
+float3 ComputeStaticGeometryOffset(uint GuideIndex, float GuideVertexWeight)
+{	
+	const float3 RestGuidePoint = f16tof32(SimRestPosePositionBuffer[GuideIndex].xyz) + InSimHairPositionOffset;
+
+	const float3 DeformedGuidePoint = f16tof32(DeformedSimPositionBuffer[GuideIndex].xyz) + OutHairPositionOffset;
+	return (DeformedGuidePoint - RestGuidePoint) * GuideVertexWeight;
+}
+
+float3 ComputeDynamicGeometryOffset(
+	uint GuideIndex, 
+	float GuideVertexWeight, 
+	FHairMeshTriangle RestTriangle, 
+	FHairMeshTriangle DeformedTriangle,
+	float3 RootBarycentric)
+{
+	const float3 RestGuidePoint = f16tof32(SimRestPosePositionBuffer[GuideIndex].xyz) + InSimHairPositionOffset;
+	const float3 LocalRestGuidePoint = ToTriangle(RestGuidePoint, RootBarycentric, RestTriangle);
+
+	const float3 DeformedGuidePoint = f16tof32(DeformedSimPositionBuffer[GuideIndex].xyz) + OutHairPositionOffset;
+	const float3 LocalDeformedGuidePoint = ToTriangle(DeformedGuidePoint, RootBarycentric, DeformedTriangle);
+
+	return (LocalDeformedGuidePoint - LocalRestGuidePoint) * GuideVertexWeight;
+}
+
 [numthreads(GROUP_SIZE, 1, 1)]
 void MainCS(uint2 DispatchThreadId : SV_DispatchThreadID)
 {
@@ -73,18 +97,23 @@ void MainCS(uint2 DispatchThreadId : SV_DispatchThreadID)
 	if (VertexIndex < VertexCount)
 	{
 		uint3 GuideVertexIndices = 0;
+		float3 GuideVertexLerps = 0;
 		float3 GuideVertexWeights = 0;
 		#if PERMUTATION_SIMULATION == 1 || PERMUTATION_DEBUG == 1 || PERMUTATION_DEBUG == 2
 		{
 			const uint4 InterpolationData0	= Interpolation0Buffer[VertexIndex];
 			const uint4 InterpolationData1	= Interpolation1Buffer[VertexIndex];
 
-			GuideVertexIndices.x = InterpolationData0.x | (InterpolationData1.x << 16);
-			GuideVertexIndices.y = InterpolationData0.y | (InterpolationData1.y << 16);
-			GuideVertexIndices.z = InterpolationData0.z | (InterpolationData1.z << 16);
+			GuideVertexIndices.x = InterpolationData0.x | ( (InterpolationData1.x & 0xFF) << 16);
+			GuideVertexIndices.y = InterpolationData0.y | (((InterpolationData1.x >> 8) & 0xFF) << 16);
+			GuideVertexIndices.z = InterpolationData0.z | ( (InterpolationData1.y & 0xFF) << 16);
 
-			GuideVertexWeights = float3(saturate((InterpolationData0.w & 0xFF) / 255.f), saturate(((InterpolationData0.w >> 8) & 0xFF) / 255.f), 0);
+			GuideVertexWeights   = float3(saturate((InterpolationData0.w & 0xFF) / 255.f), saturate(((InterpolationData0.w >> 8) & 0xFF) / 255.f), 0);
 			GuideVertexWeights.z = saturate(1.0f - (GuideVertexWeights.x + GuideVertexWeights.y));
+
+			GuideVertexLerps.x = saturate(((InterpolationData1.y >> 8) & 0xFF)/255.f);
+			GuideVertexLerps.y = saturate(( InterpolationData1.z & 0xFF      )/255.f);
+			GuideVertexLerps.z = saturate(((InterpolationData1.z >> 8) & 0xFF)/255.f);
 		}
 		#endif
 
@@ -99,11 +128,13 @@ void MainCS(uint2 DispatchThreadId : SV_DispatchThreadID)
 			[unroll]
 			for (uint K = 0; K < 3; ++K)
 			{
-				const uint GuideIndex = GuideVertexIndices[K];
-				const float3 RestGuidePoint	= f16tof32(SimRestPosePositionBuffer[GuideIndex].xyz) + InSimHairPositionOffset;
+				const uint GuideIndex0 = GuideVertexIndices[K];
+				const uint GuideIndex1 = GuideIndex0 + 1;
+
+				const float3 Offset0 = ComputeStaticGeometryOffset(GuideIndex0, GuideVertexWeights[K]);
+				const float3 Offset1 = ComputeStaticGeometryOffset(GuideIndex1, GuideVertexWeights[K]);
 
-				const float3 DeformedGuidePoint	= f16tof32(DeformedSimPositionBuffer[GuideIndex].xyz) + OutHairPositionOffset;
-				CurrOffset += (DeformedGuidePoint - RestGuidePoint) * GuideVertexWeights[K];
+				CurrOffset += lerp(Offset0, Offset1, GuideVertexLerps[K]);
 			}
 			#endif
 			OutRenderDeformedPositionBuffer[VertexIndex] = float4(f32tof16(ControlPoint + CurrOffset - OutHairPositionOffset), RenderRestPosePositionBuffer[VertexIndex].w);
@@ -121,19 +152,19 @@ void MainCS(uint2 DispatchThreadId : SV_DispatchThreadID)
 			#if PERMUTATION_SIMULATION == 1
 			for (uint K = 0; K < 3; ++K)
 			{
-				const uint GuideIndex = GuideVertexIndices[K];
-				const uint RootIndex = SimVertexToRootIndexBuffer[GuideIndex];
+				const uint GuideIndex0 = GuideVertexIndices[K];
+				const uint GuideIndex1 = GuideIndex0 + 1;
+
+				const uint RootIndex = SimVertexToRootIndexBuffer[GuideIndex0];
 				const float3 RootBarycentric = DecodeBarycentrics(SimRootBarycentricBuffer[RootIndex]);
 
 				const FHairMeshTriangle RestTriangle		= GetTriangleTransformation(RootIndex, SimRestPosition0Buffer, SimRestPosition1Buffer, SimRestPosition2Buffer, SimRestPositionOffset);
 				const FHairMeshTriangle DeformedTriangle	= GetTriangleTransformation(RootIndex, SimDeformedPosition0Buffer, SimDeformedPosition1Buffer, SimDeformedPosition2Buffer, SimDeformedPositionOffset);
 
-				const float3 RestGuidePoint	= f16tof32(SimRestPosePositionBuffer[GuideIndex].xyz) + InSimHairPositionOffset;
-				const float3 LocalRestGuidePoint = ToTriangle(RestGuidePoint,RootBarycentric,RestTriangle);
+				const float3 Offset0 = ComputeDynamicGeometryOffset(GuideIndex0, GuideVertexWeights[K], RestTriangle, DeformedTriangle, RootBarycentric);
+				const float3 Offset1 = ComputeDynamicGeometryOffset(GuideIndex1, GuideVertexWeights[K], RestTriangle, DeformedTriangle, RootBarycentric);
 
-				const float3 DeformedGuidePoint	= f16tof32(DeformedSimPositionBuffer[GuideIndex].xyz) + OutHairPositionOffset;
-				const float3 LocalDeformedGuidePoint = ToTriangle(DeformedGuidePoint,RootBarycentric,DeformedTriangle);
-				CurrOffset += (LocalDeformedGuidePoint - LocalRestGuidePoint) * GuideVertexWeights[K];
+				CurrOffset += lerp(Offset0, Offset1, GuideVertexLerps[K]);
 			}
 			#endif		
 
diff --git a/Engine/Source/Runtime/Core/Public/UObject/ReleaseObjectVersion.h b/Engine/Source/Runtime/Core/Public/UObject/ReleaseObjectVersion.h
--- a/Engine/Source/Runtime/Core/Public/UObject/ReleaseObjectVersion.h
+++ b/Engine/Source/Runtime/Core/Public/UObject/ReleaseObjectVersion.h
@@ -90,6 +90,9 @@ struct CORE_API FReleaseObjectVersion
 		// Store applied version of Animation Modifier to use when reverting
 		SerializeAnimModifierState,
 
+		// Groom asset version
+		GroomAssetVersion3,
+
 		// -----<new versions can be added above this line>-------------------------------------------------
 		VersionPlusOne,
 		LatestVersion = VersionPlusOne - 1
