Exposed more options to user for building groom interpolation data, as there is no obvious 'one-fit-all' settings for low guide count.

diff --git a/Engine/Plugins/Experimental/HairStrands/Source/HairStrandsCore/Private/GroomAsset.cpp b/Engine/Plugins/Experimental/HairStrands/Source/HairStrandsCore/Private/GroomAsset.cpp
--- a/Engine/Plugins/Experimental/HairStrands/Source/HairStrandsCore/Private/GroomAsset.cpp
+++ b/Engine/Plugins/Experimental/HairStrands/Source/HairStrandsCore/Private/GroomAsset.cpp
@@ -338,7 +338,7 @@ void UGroomAsset::PostLoad()
 	check(HairGroupsData.Num() > 0);
 	if (HairGroupsData[0].HairSimulationData.GetNumCurves() == 0 || HairGroupsData[0].HairInterpolationData.Num() == 0)
 	{
-		FGroomBuilder::BuildData(this, 2);
+		FGroomBuilder::BuildData(this, 2, 0, false, false);
 	}
 
 	if (!IsTemplate())
diff --git a/Engine/Plugins/Experimental/HairStrands/Source/HairStrandsCore/Private/GroomBuilder.cpp b/Engine/Plugins/Experimental/HairStrands/Source/HairStrandsCore/Private/GroomBuilder.cpp
--- a/Engine/Plugins/Experimental/HairStrands/Source/HairStrandsCore/Private/GroomBuilder.cpp
+++ b/Engine/Plugins/Experimental/HairStrands/Source/HairStrandsCore/Private/GroomBuilder.cpp
@@ -392,11 +392,13 @@ namespace HairInterpolationBuilder
 	};
 
 	// Randomize influence guide to break interpolation coherence, and create a more random/natural pattern
-	static void SelectFinalGuides(FClosestGuides& ClosestGuides, FRandomStream& Random, const FMetrics& Metric)
+	static void SelectFinalGuides(
+		FClosestGuides& ClosestGuides, 
+		FRandomStream& Random, 
+		const FMetrics& Metric, 
+		const bool bRandomizeInterpolation, 
+		const bool bUseUniqueGuide)
 	{
-		const static bool bRandomizeInterpolation = false;
-		const static bool bUseUniqueGuide = false;
-
 		uint32 RandIndex0 = 0;
 		uint32 RandIndex1 = 1;
 		uint32 RandIndex2 = 2;
@@ -573,6 +575,8 @@ namespace HairInterpolationBuilder
 			const TArray<FHairRoot>& SimRoots,
 			const FHairStrandsDatas& RenStrandsData,
 			const FHairStrandsDatas& SimStrandsData,
+			const bool bRandomized,
+			const bool bUnique,
 			FRandomStream& Random) const
 		{
 			const FHairRoot& RenRoot = RenRoots[RenCurveIndex];
@@ -626,7 +630,7 @@ namespace HairInterpolationBuilder
 
 			// If no valid guide have been found, switch to a simpler metric
 			FClosestGuides ClosestGuides;
-			SelectFinalGuides(ClosestGuides, Random, Metrics);
+			SelectFinalGuides(ClosestGuides, Random, Metrics, bRandomized, bUnique);
 
 			check(ClosestGuides.Indices[0] >= 0);
 			check(ClosestGuides.Indices[1] >= 0);
@@ -642,6 +646,8 @@ namespace HairInterpolationBuilder
 			const TArray<FHairRoot>& SimRoots,
 			const FHairStrandsDatas& RenStrandsData,
 			const FHairStrandsDatas& SimStrandsData,
+			const bool bRandomized,
+			const bool bUnique,
 			FRandomStream& Random) const 
 		{
 			const FHairRoot& RenRoot = RenRoots[RenCurveIndex];
@@ -699,10 +705,10 @@ namespace HairInterpolationBuilder
 
 			// If no valid guide have been found, switch to a simpler metric
 			FClosestGuides ClosestGuides;
-			SelectFinalGuides(ClosestGuides, Random, Metrics0);
+			SelectFinalGuides(ClosestGuides, Random, Metrics0, bRandomized, bUnique);
 			if (ClosestGuides.Indices[0] == -1)
 			{
-				SelectFinalGuides(ClosestGuides, Random, Metrics1);
+				SelectFinalGuides(ClosestGuides, Random, Metrics1, bRandomized, bUnique);
 			}
 
 			check(ClosestGuides.Indices[0] >= 0);
@@ -719,6 +725,8 @@ namespace HairInterpolationBuilder
 		const TArray<FHairRoot>& SimRoots,
 		const FHairStrandsDatas& RenStrandsData,
 		const FHairStrandsDatas& SimStrandsData,
+		const bool bRandomized,
+		const bool bUnique,
 		FRandomStream& Random)
 	{
 		FMetrics Metrics;
@@ -735,7 +743,7 @@ namespace HairInterpolationBuilder
 		}
 			
 		FClosestGuides ClosestGuides;
-		SelectFinalGuides(ClosestGuides, Random, Metrics);
+		SelectFinalGuides(ClosestGuides, Random, Metrics, bRandomized, bUnique);
 
 		check(ClosestGuides.Indices[0] >= 0);
 		check(ClosestGuides.Indices[1] >= 0);
@@ -811,18 +819,33 @@ namespace HairInterpolationBuilder
 		return Desc;
 	}
 
-	enum class EHairInterpolationDataQuality
+	enum class EHairInterpolationWeightMethod : uint8
+	{
+		ParametricDistance,
+		RootDistance,
+		VertexIndex
+	};
+
+	enum class EHairInterpolationDataQuality : uint8
 	{
 		Low,
 		Medium,
 		High
 	};
 
+	struct FHairInterpolationSettings
+	{
+		EHairInterpolationDataQuality Quality		= EHairInterpolationDataQuality::High;
+		EHairInterpolationWeightMethod WeightMethod = EHairInterpolationWeightMethod::ParametricDistance;
+		bool bRandomizedGuides						= false;
+		bool bUseUniqueGuide						= false;
+	};
+
 	static void BuildInterpolationData(
 		FHairStrandsInterpolationDatas& InterpolationData,
 		const FHairStrandsDatas& SimStrandsData,
 		const FHairStrandsDatas& RenStrandsData,
-		const EHairInterpolationDataQuality Quality)
+		const FHairInterpolationSettings& Settings)
 	{
 		TRACE_CPUPROFILER_EVENT_SCOPE(HairInterpolationBuilder::BuildInterpolationData);
 
@@ -841,7 +864,7 @@ namespace HairInterpolationBuilder
 			ExtractRoots(RenStrandsData, RenRoots, RenMinBound, RenMaxBound);
 			ExtractRoots(SimStrandsData, SimRoots, SimMinBound, SimMaxBound);
 
-			if (Quality == EHairInterpolationDataQuality::Low || Quality == EHairInterpolationDataQuality::Medium)
+			if (Settings.Quality == EHairInterpolationDataQuality::Low || Settings.Quality == EHairInterpolationDataQuality::Medium)
 			{
 				// Build a conservative bound, to insure all queries will fall 
 				// into the grid volume.
@@ -865,7 +888,7 @@ namespace HairInterpolationBuilder
 
 		ParallelFor(RenCurveCount, 
 		[
-			Quality,
+			Settings,
 			RenCurveCount, &RenRoots, &RenStrandsData,
 			SimCurveCount, &SimRoots, &SimStrandsData, 
 			&RootsGrid,
@@ -890,17 +913,17 @@ namespace HairInterpolationBuilder
 			const FHairRoot& StrandRoot = RenRoots[RenCurveIndex];
 
 			FClosestGuides ClosestGuides;
-			if (Quality == EHairInterpolationDataQuality::Low)
+			if (Settings.Quality == EHairInterpolationDataQuality::Low)
 			{
-				ClosestGuides = RootsGrid.FindClosestRoots(RenCurveIndex, RenRoots, SimRoots, RenStrandsData, SimStrandsData, Random);
+				ClosestGuides = RootsGrid.FindClosestRoots(RenCurveIndex, RenRoots, SimRoots, RenStrandsData, SimStrandsData, Settings.bRandomizedGuides, Settings.bUseUniqueGuide, Random);
 			}
-			else if (Quality == EHairInterpolationDataQuality::Medium)
+			else if (Settings.Quality == EHairInterpolationDataQuality::Medium)
 			{
-				ClosestGuides = RootsGrid.FindBestClosestRoots(RenCurveIndex, RenRoots, SimRoots, RenStrandsData, SimStrandsData, Random);
+				ClosestGuides = RootsGrid.FindBestClosestRoots(RenCurveIndex, RenRoots, SimRoots, RenStrandsData, SimStrandsData, Settings.bRandomizedGuides, Settings.bUseUniqueGuide, Random);
 			}
-			else // (Quality == EHairInterpolationDataQuality::High)
+			else // (Settings.Quality == EHairInterpolationDataQuality::High)
 			{
-				ClosestGuides = FindBestRoots(RenCurveIndex, RenRoots, SimRoots, RenStrandsData, SimStrandsData, Random);
+				ClosestGuides = FindBestRoots(RenCurveIndex, RenRoots, SimRoots, RenStrandsData, SimStrandsData, Settings.bRandomizedGuides, Settings.bUseUniqueGuide, Random);
 			}
 
 			const uint32 RendPointCount	= RenStrandsData.StrandsCurves.CurvesCount[RenCurveIndex];
@@ -915,56 +938,56 @@ namespace HairInterpolationBuilder
 				float TotalWeight = 0;
 				for (uint32 KIndex = 0; KIndex < FClosestGuides::Count; ++KIndex)
 				{
-
-				#define WEIGHT_METHOD 0
-
 					// Find the closest vertex on the guide which matches the strand vertex distance along its curve
-				#if WEIGHT_METHOD == 0
-					const uint32 SimCurveIndex = ClosestGuides.Indices[KIndex];
-					const uint32 SimOffset = SimStrandsData.StrandsCurves.CurvesOffset[SimCurveIndex];
-					const FVertexInterpolationDesc Desc = FindMatchingVertex(RenPointDistance, SimStrandsData, SimCurveIndex);
-					const FVector& SimPointPosition0 = SimStrandsData.StrandsPoints.PointsPosition[Desc.Index0 + SimOffset];
-					const FVector& SimPointPosition1 = SimStrandsData.StrandsPoints.PointsPosition[Desc.Index1 + SimOffset];
-					const float Weight = 1.0f / FMath::Max(MinWeightDistance, FVector::Distance(RenPointPosition, FMath::Lerp(SimPointPosition0, SimPointPosition1, Desc.T)));
-
-					InterpolationData.PointsSimCurvesIndex[PointGlobalIndex][KIndex] = SimCurveIndex;
-					InterpolationData.PointsSimCurvesVertexIndex[PointGlobalIndex][KIndex] = Desc.Index0 + SimOffset;
-					InterpolationData.PointsSimCurvesVertexLerp[PointGlobalIndex][KIndex] = Desc.T;
-					InterpolationData.PointsSimCurvesVertexWeights[PointGlobalIndex][KIndex] = Weight;
-				#endif
+					if (Settings.WeightMethod == EHairInterpolationWeightMethod::ParametricDistance)
+					{
+						const uint32 SimCurveIndex = ClosestGuides.Indices[KIndex];
+						const uint32 SimOffset = SimStrandsData.StrandsCurves.CurvesOffset[SimCurveIndex];
+						const FVertexInterpolationDesc Desc = FindMatchingVertex(RenPointDistance, SimStrandsData, SimCurveIndex);
+						const FVector& SimPointPosition0 = SimStrandsData.StrandsPoints.PointsPosition[Desc.Index0 + SimOffset];
+						const FVector& SimPointPosition1 = SimStrandsData.StrandsPoints.PointsPosition[Desc.Index1 + SimOffset];
+						const float Weight = 1.0f / FMath::Max(MinWeightDistance, FVector::Distance(RenPointPosition, FMath::Lerp(SimPointPosition0, SimPointPosition1, Desc.T)));
+
+						InterpolationData.PointsSimCurvesIndex[PointGlobalIndex][KIndex] = SimCurveIndex;
+						InterpolationData.PointsSimCurvesVertexIndex[PointGlobalIndex][KIndex] = Desc.Index0 + SimOffset;
+						InterpolationData.PointsSimCurvesVertexLerp[PointGlobalIndex][KIndex] = Desc.T;
+						InterpolationData.PointsSimCurvesVertexWeights[PointGlobalIndex][KIndex] = Weight;
+					}
 
 					// Use only the root as a *constant* weight for deformation along each vertex
 					// Still compute the closest vertex (in parametric distance) to know on which vertex the offset/delta should be computed
-				#if WEIGHT_METHOD == 1
-					const uint32 SimCurveIndex = ClosestGuides.Indices[KIndex];
-					const uint32 SimOffset = SimStrandsData.StrandsCurves.CurvesOffset[SimCurveIndex];
-					const FVector& SimRootPointPosition = SimStrandsData.StrandsPoints.PointsPosition[SimOffset];
-					const FVector& RenRootPointPosition = RenStrandsData.StrandsPoints.PointsPosition[RenOffset];
-					const float Weight = 1.0f / FMath::Max(MinWeightDistance, FVector::Distance(RenRootPointPosition, SimRootPointPosition));
-					const FVertexInterpolationDesc Desc = FindMatchingVertex(RenPointDistance, SimStrandsData, SimCurveIndex);
-
-					InterpolationData.PointsSimCurvesIndex[PointGlobalIndex][KIndex] = SimCurveIndex;
-					InterpolationData.PointsSimCurvesVertexIndex[PointGlobalIndex][KIndex] = Desc.Index0 + SimOffset;
-					InterpolationData.PointsSimCurvesVertexLerp[PointGlobalIndex][KIndex] = Desc.T;
-					InterpolationData.PointsSimCurvesVertexWeights[PointGlobalIndex][KIndex] = Weight;
-				#endif
+					if (Settings.WeightMethod == EHairInterpolationWeightMethod::RootDistance)
+					{
+						const uint32 SimCurveIndex = ClosestGuides.Indices[KIndex];
+						const uint32 SimOffset = SimStrandsData.StrandsCurves.CurvesOffset[SimCurveIndex];
+						const FVector& SimRootPointPosition = SimStrandsData.StrandsPoints.PointsPosition[SimOffset];
+						const FVector& RenRootPointPosition = RenStrandsData.StrandsPoints.PointsPosition[RenOffset];
+						const float Weight = 1.0f / FMath::Max(MinWeightDistance, FVector::Distance(RenRootPointPosition, SimRootPointPosition));
+						const FVertexInterpolationDesc Desc = FindMatchingVertex(RenPointDistance, SimStrandsData, SimCurveIndex);
+
+						InterpolationData.PointsSimCurvesIndex[PointGlobalIndex][KIndex] = SimCurveIndex;
+						InterpolationData.PointsSimCurvesVertexIndex[PointGlobalIndex][KIndex] = Desc.Index0 + SimOffset;
+						InterpolationData.PointsSimCurvesVertexLerp[PointGlobalIndex][KIndex] = Desc.T;
+						InterpolationData.PointsSimCurvesVertexWeights[PointGlobalIndex][KIndex] = Weight;
+					}
 
 					// Use the *same vertex index* to match guide vertex with strand vertex
-				#if WEIGHT_METHOD == 2
-					check(SimPointCount > 0);
-					const uint32 SimCurveIndex = ClosestGuides.Indices[KIndex];
-					const uint32 SimPointIndex = FMath::Clamp(RenPointIndex, 0, SimPointCount - 1);
-					const FVector& SimPointPosition = SimStrandsData.StrandsPoints.PointsPosition[SimPointIndex + SimOffset];
-					const float Weight = 1.0f / FMath::Max(MinWeightDistance, FVector::Distance(RenPointPosition, SimPointPosition));
-					const FVertexInterpolationDesc Desc = FindMatchingVertex(RenPointDistance, SimStrandsData, SimCurveIndex);
-
-					InterpolationData.PointsSimCurvesIndex[PointGlobalIndex][KIndex] = SimCurveIndex;
-					InterpolationData.PointsSimCurvesVertexIndex[PointGlobalIndex][KIndex] = Desc.Index0 + SimOffset;
-					InterpolationData.PointsSimCurvesVertexLerp[PointGlobalIndex][KIndex] = Desc.T;
-					InterpolationData.PointsSimCurvesVertexWeights[PointGlobalIndex][KIndex] = Weight;
-				#endif
-
-					TotalWeight += Weight;
+					if (Settings.WeightMethod == EHairInterpolationWeightMethod::VertexIndex)
+					{
+						const uint32 SimCurveIndex = ClosestGuides.Indices[KIndex];
+						const uint32 SimOffset = SimStrandsData.StrandsCurves.CurvesOffset[SimCurveIndex];
+						const uint32 SimPointCount = SimStrandsData.StrandsCurves.CurvesCount[SimCurveIndex];
+						const uint32 SimPointIndex = FMath::Clamp(RenPointIndex, 0u, SimPointCount - 1);
+						const FVector& SimPointPosition = SimStrandsData.StrandsPoints.PointsPosition[SimPointIndex + SimOffset];
+						const float Weight = 1.0f / FMath::Max(MinWeightDistance, FVector::Distance(RenPointPosition, SimPointPosition));
+
+						InterpolationData.PointsSimCurvesIndex[PointGlobalIndex][KIndex] = SimCurveIndex;
+						InterpolationData.PointsSimCurvesVertexIndex[PointGlobalIndex][KIndex] = SimPointIndex + SimOffset;
+						InterpolationData.PointsSimCurvesVertexLerp[PointGlobalIndex][KIndex] = 1;
+						InterpolationData.PointsSimCurvesVertexWeights[PointGlobalIndex][KIndex] = Weight;
+					}
+
+					TotalWeight += InterpolationData.PointsSimCurvesVertexWeights[PointGlobalIndex][KIndex];
 				}
 
 				for (int32 KIndex = 0; KIndex < FClosestGuides::Count; ++KIndex)
@@ -1221,14 +1244,14 @@ bool FGroomBuilder::BuildGroom(const FHairDescription& HairDescription, const FG
 		GroomAsset->HairGroupsData.Add(MoveTemp(GroupData));
 	}
 
-	BuildData(GroomAsset, uint8(BuildSettings.InterpolationQuality));
+	BuildData(GroomAsset, uint8(BuildSettings.InterpolationQuality), uint8(BuildSettings.InterpolationDistance), BuildSettings.bRandomizeGuide, BuildSettings.bUseUniqueGuide);
 
 	GroomAsset->InitResource();
 
 	return true;
 }
 
-void FGroomBuilder::BuildData(UGroomAsset* GroomAsset, uint8 QualityLevel)
+void FGroomBuilder::BuildData(UGroomAsset* GroomAsset, uint8 QualityLevel, uint8 WeightMethod, bool bRandomize, bool bUnique)
 {
 	if (!GroomAsset)
 	{
@@ -1262,14 +1285,22 @@ void FGroomBuilder::BuildData(UGroomAsset* GroomAsset, uint8 QualityLevel)
 		FHairStrandsInterpolationDatas& HairInterpolationData = GroupData.HairInterpolationData;
 
 		// Build InterpolationData from render and simulation HairStrandsDatas
-		HairInterpolationBuilder::EHairInterpolationDataQuality InterpolationQuality = HairInterpolationBuilder::EHairInterpolationDataQuality::High;
+		HairInterpolationBuilder::FHairInterpolationSettings Settings;
+		Settings.bRandomizedGuides = bRandomize;
+		Settings.bUseUniqueGuide = bUnique;
+		switch (WeightMethod)
+		{
+			case 0: Settings.WeightMethod = HairInterpolationBuilder::EHairInterpolationWeightMethod::ParametricDistance; break;
+			case 1: Settings.WeightMethod = HairInterpolationBuilder::EHairInterpolationWeightMethod::RootDistance; break;
+			case 2: Settings.WeightMethod = HairInterpolationBuilder::EHairInterpolationWeightMethod::VertexIndex; break;
+		}
 		switch (QualityLevel)
 		{
-			case 0: InterpolationQuality = HairInterpolationBuilder::EHairInterpolationDataQuality::Low; break;
-			case 1: InterpolationQuality = HairInterpolationBuilder::EHairInterpolationDataQuality::Medium; break;
-			case 2: InterpolationQuality = HairInterpolationBuilder::EHairInterpolationDataQuality::High; break;
+			case 0: Settings.Quality = HairInterpolationBuilder::EHairInterpolationDataQuality::Low; break;
+			case 1: Settings.Quality = HairInterpolationBuilder::EHairInterpolationDataQuality::Medium; break;
+			case 2: Settings.Quality = HairInterpolationBuilder::EHairInterpolationDataQuality::High; break;
 		}
-		HairInterpolationBuilder::BuildInterpolationData(HairInterpolationData, HairSimulationData, HairRenderData, InterpolationQuality);
+		HairInterpolationBuilder::BuildInterpolationData(HairInterpolationData, HairSimulationData, HairRenderData, Settings);
 
 		// Build Rendering data for InterpolationData
 		HairInterpolationBuilder::BuildRenderData(HairInterpolationData);
diff --git a/Engine/Plugins/Experimental/HairStrands/Source/HairStrandsCore/Public/GroomBuilder.h b/Engine/Plugins/Experimental/HairStrands/Source/HairStrandsCore/Public/GroomBuilder.h
--- a/Engine/Plugins/Experimental/HairStrands/Source/HairStrandsCore/Public/GroomBuilder.h
+++ b/Engine/Plugins/Experimental/HairStrands/Source/HairStrandsCore/Public/GroomBuilder.h
@@ -11,5 +11,5 @@ struct HAIRSTRANDSCORE_API FGroomBuilder
 {
 	static bool BuildGroom(const class FHairDescription& HairDescription, const struct FGroomBuildSettings& BuildSettings, UGroomAsset* GroomAsset);
 	static void GenerateGuides(const FHairStrandsDatas& InData, float DecimationPercentage, FHairStrandsDatas& OutData);
-	static void BuildData(UGroomAsset* GroomAsset, uint8 QualityLevel);
+	static void BuildData(UGroomAsset* GroomAsset, uint8 QualityLevel, uint8 WeightMethod, bool bRandomize, bool bUnique);
 };
\ No newline at end of file
diff --git a/Engine/Plugins/Experimental/HairStrands/Source/HairStrandsCore/Public/GroomSettings.h b/Engine/Plugins/Experimental/HairStrands/Source/HairStrandsCore/Public/GroomSettings.h
--- a/Engine/Plugins/Experimental/HairStrands/Source/HairStrandsCore/Public/GroomSettings.h
+++ b/Engine/Plugins/Experimental/HairStrands/Source/HairStrandsCore/Public/GroomSettings.h
@@ -25,10 +25,6 @@ struct HAIRSTRANDSCORE_API FGroomConversionSettings
 	FVector Scale;
 };
 
-/**
- * Specifies the overal rendering/shading model for a material
- * @warning Check UMaterialInstance::Serialize if changed!
- */
 UENUM()
 enum class EGroomInterpolationQuality : uint8
 {
@@ -37,6 +33,14 @@ enum class EGroomInterpolationQuality : uint8
 	High	UMETA(DisplayName = "High",		ToolTip = "Build interpolation data using curve shape matching search. This result in high quality interpolation data, but is relatively slow to build (can takes several dozen of minutes)"),
 };
 
+UENUM()
+enum class EGroomInterpolationWeight : uint8
+{
+	Parametric	UMETA(DisplayName = "Parametric", ToolTip = "Build interpolation data based on curve parametric distance"),
+	Root		UMETA(DisplayName = "Root", ToolTip = "Build interpolation data based on distance between guide's root and strands's root"),
+	Index		UMETA(DisplayName = "Index", ToolTip = "Build interpolation data based on guide and strands vertex indices"),
+};
+
 USTRUCT(BlueprintType)
 struct HAIRSTRANDSCORE_API FGroomBuildSettings
 {
@@ -46,6 +50,9 @@ struct HAIRSTRANDSCORE_API FGroomBuildSettings
 		: bOverrideGuides(false)
 		, HairToGuideDensity(0.1f)
 		, InterpolationQuality(EGroomInterpolationQuality::High)
+		, InterpolationDistance(EGroomInterpolationWeight::Parametric)
+		, bRandomizeGuide(false)
+		, bUseUniqueGuide(false)
 	{}
 
 	/** Flag to override the imported guides with generated guides. */
@@ -59,4 +66,16 @@ struct HAIRSTRANDSCORE_API FGroomBuildSettings
 	/** Interpolation data quality. */
 	UPROPERTY(EditAnywhere, BlueprintReadWrite, Category = "BuildSettings")
 	EGroomInterpolationQuality InterpolationQuality;
+
+	/** Interpolation distance metric. */
+	UPROPERTY(EditAnywhere, BlueprintReadWrite, Category = "BuildSettings")
+	EGroomInterpolationWeight InterpolationDistance;
+
+	/** Randomize which guides affect a given hair strand. */
+	UPROPERTY(EditAnywhere, BlueprintReadWrite, Category = "BuildSettings")
+	bool bRandomizeGuide;
+
+	/** Force a hair strand to be affected by a unique guide. */
+	UPROPERTY(EditAnywhere, BlueprintReadWrite, Category = "BuildSettings")
+	bool bUseUniqueGuide;
 };
