Fix a hair import issue when a groom contains guides, but the user choose to override guides with generated ones.

diff --git a/Engine/Plugins/Experimental/HairStrands/Source/HairStrandsCore/Private/GroomBuilder.cpp b/Engine/Plugins/Experimental/HairStrands/Source/HairStrandsCore/Private/GroomBuilder.cpp
--- a/Engine/Plugins/Experimental/HairStrands/Source/HairStrandsCore/Private/GroomBuilder.cpp
+++ b/Engine/Plugins/Experimental/HairStrands/Source/HairStrandsCore/Private/GroomBuilder.cpp
@@ -1156,6 +1156,7 @@ bool FGroomBuilder::BuildGroom(const FHairDescription& HairDescription, const FG
 		else
 		{
 			// A guide but don't want to import it, so skip it
+			GlobalVertexIndex += CurveNumVertices;
 			continue;
 		}
 
@@ -1333,23 +1334,27 @@ void FGroomBuilder::GenerateGuides(const FHairStrandsDatas& InData, float Decima
 	OutData.HairDensity = InData.HairDensity;
 
 	uint32 OutPointOffset = 0; 
-	for (uint32 OutCurveIndex = 0; OutCurveIndex < OutCurveCount; OutCurveIndex++)
+	for (uint32 OutCurveIndex = 0; OutCurveIndex < OutCurveCount; ++OutCurveIndex)
 	{
-		const uint32 InCurveIndex = CurveIndices[OutCurveIndex];
-		OutData.StrandsCurves.CurvesCount[OutCurveIndex] = InData.StrandsCurves.CurvesCount[InCurveIndex];
+		const uint32 InCurveIndex	= CurveIndices[OutCurveIndex];
+		const uint32 InPointOffset	= InData.StrandsCurves.CurvesOffset[InCurveIndex];
+		const uint32 PointCount		= InData.StrandsCurves.CurvesCount[InCurveIndex];
+		OutData.StrandsCurves.CurvesCount[OutCurveIndex]  = PointCount;
 		OutData.StrandsCurves.CurvesRootUV[OutCurveIndex] = InData.StrandsCurves.CurvesRootUV[InCurveIndex];
 		OutData.StrandsCurves.CurvesOffset[OutCurveIndex] = OutPointOffset;
+		OutData.StrandsCurves.CurvesLength[OutCurveIndex] = InData.StrandsCurves.CurvesLength[InCurveIndex];
+		OutData.StrandsCurves.MaxLength = InData.StrandsCurves.MaxLength;
+		OutData.StrandsCurves.MaxRadius = InData.StrandsCurves.MaxRadius;
 
-		const uint32 InPointOffset = InData.StrandsCurves.CurvesOffset[InCurveIndex];
-		const uint32 OutPointCount = OutData.StrandsCurves.CurvesCount[OutCurveIndex];
-		for (uint32 OutPointIndex = 0; OutPointIndex < OutPointCount; ++OutPointIndex, ++OutPointOffset)
+		for (uint32 PointIndex = 0; PointIndex < PointCount; ++PointIndex)
 		{
-			OutData.StrandsPoints.PointsPosition[OutPointOffset] = InData.StrandsPoints.PointsPosition[OutPointIndex + InPointOffset];
-			OutData.StrandsPoints.PointsCoordU[OutPointOffset] = InData.StrandsPoints.PointsCoordU[OutPointIndex + InPointOffset];
-			OutData.StrandsPoints.PointsRadius[OutPointOffset] = InData.StrandsPoints.PointsRadius[OutPointIndex + InPointOffset] * InData.StrandsCurves.MaxRadius;
-			OutData.StrandsPoints.PointsBaseColor[OutPointOffset] = FLinearColor::Black;
-			OutData.StrandsPoints.PointsRoughness[OutPointOffset] = 0;
+			OutData.StrandsPoints.PointsPosition [PointIndex + OutPointOffset] = InData.StrandsPoints.PointsPosition	[PointIndex + InPointOffset];
+			OutData.StrandsPoints.PointsCoordU	 [PointIndex + OutPointOffset] = InData.StrandsPoints.PointsCoordU		[PointIndex + InPointOffset];
+			OutData.StrandsPoints.PointsRadius	 [PointIndex + OutPointOffset] = InData.StrandsPoints.PointsRadius		[PointIndex + InPointOffset];
+			OutData.StrandsPoints.PointsBaseColor[PointIndex + OutPointOffset] = FLinearColor::Black;
+			OutData.StrandsPoints.PointsRoughness[PointIndex + OutPointOffset] = 0;
 		}
+		OutPointOffset += PointCount;
 	}
 
 	HairStrandsBuilder::BuildInternalData(OutData, false);
