Fixed build warnings due to new deprecation.

#rnx

diff --git a/Engine/Plugins/Experimental/HairStrands/Source/HairStrandsNiagara/Private/NiagaraDataInterfaceHairStrands.cpp b/Engine/Plugins/Experimental/HairStrands/Source/HairStrandsNiagara/Private/NiagaraDataInterfaceHairStrands.cpp
--- a/Engine/Plugins/Experimental/HairStrands/Source/HairStrandsNiagara/Private/NiagaraDataInterfaceHairStrands.cpp
+++ b/Engine/Plugins/Experimental/HairStrands/Source/HairStrandsNiagara/Private/NiagaraDataInterfaceHairStrands.cpp
@@ -408,7 +408,7 @@ struct FNDIHairStrandsParametersCS : public FNiagaraDataInterfaceParametersCS
 			FShaderResourceViewRHIRef RootBarycentricCoordinatesSRV = HairStrandsBuffer->RootBarycentricCoordinatesBuffer.SRV;
 
 			int32 HasRootAttachedValue = 0;
-			bool bNeedSimReset = (ProxyData->TickCount <= ProxyData->ResetTick);
+			int32 bNeedSimReset = (ProxyData->TickCount <= ProxyData->ResetTick ? 1 : 0);
 			FVector RestRootOffsetValue = FVector(0, 0, 0);
 			FVector DeformedRootOffsetValue = FVector(0, 0, 0);
 			FVector DeformedPositionOffsetValue = HairStrandsBuffer->SourceRestResources->PositionOffset;
@@ -485,7 +485,7 @@ struct FNDIHairStrandsParametersCS : public FNiagaraDataInterfaceParametersCS
 			SetShaderValue(RHICmdList, ComputeShaderRHI, BoxCenter, FVector(0,0,0));
 			SetShaderValue(RHICmdList, ComputeShaderRHI, BoxExtent, FVector(0,0,0));
 
-			SetShaderValue(RHICmdList, ComputeShaderRHI, ResetSimulation, false);
+			SetShaderValue(RHICmdList, ComputeShaderRHI, ResetSimulation, 0);
 			SetShaderValue(RHICmdList, ComputeShaderRHI, HasRootAttached, 0);
 			SetShaderValue(RHICmdList, ComputeShaderRHI, RestRootOffset, FVector4(0, 0, 0, 0));
 			SetShaderValue(RHICmdList, ComputeShaderRHI, DeformedRootOffset, FVector4(0, 0, 0, 0));
