Stiffness scaling along the hairs

diff --git a/Engine/Plugins/Experimental/HairStrands/Shaders/Private/NiagaraHookeSpringMaterial.ush b/Engine/Plugins/Experimental/HairStrands/Shaders/Private/NiagaraHookeSpringMaterial.ush
--- a/Engine/Plugins/Experimental/HairStrands/Shaders/Private/NiagaraHookeSpringMaterial.ush
+++ b/Engine/Plugins/Experimental/HairStrands/Shaders/Private/NiagaraHookeSpringMaterial.ush
@@ -327,11 +327,11 @@ void UpdateMaterialFrame(in int StrandsSize)
 	if( (LocalIndex > 0) && (LocalIndex < StrandsSize-1) )
 	{
 		float4 NodeQuaternion = SharedPreviousOrientation[GGroupThreadId.x];
-		float3 TangentPrev = RotateVectorByQuat( float3(0,0,1), NodeQuaternion);
-		float3 TangentNext = normalize(SharedPreviousPosition[GGroupThreadId.x+1] - SharedPreviousPosition[GGroupThreadId.x]);
+		//float3 TangentPrev = RotateVectorByQuat( float3(0,0,1), NodeQuaternion);
+		//float3 TangentNext = normalize(SharedPreviousPosition[GGroupThreadId.x+1] - SharedPreviousPosition[GGroupThreadId.x]);
 
-		//float3 TangentPrev = normalize(SharedPreviousPosition[GGroupThreadId.x+1] - SharedPreviousPosition[GGroupThreadId.x]);
-		//float3 TangentNext = normalize(SharedNodePosition[GGroupThreadId.x+1] - SharedNodePosition[GGroupThreadId.x]);
+		float3 TangentPrev = normalize(SharedPreviousPosition[GGroupThreadId.x+1] - SharedPreviousPosition[GGroupThreadId.x]);
+		float3 TangentNext = normalize(SharedNodePosition[GGroupThreadId.x+1] - SharedNodePosition[GGroupThreadId.x]);
 
 		SharedNodeOrientation[GGroupThreadId.x] = NormalizeQuat( MultiplyQuat( FindQuatBetweenNormals(TangentPrev,TangentNext), NodeQuaternion) );
 	}
