Added missing GroomAsset check on GroomComponent::PostLoad.

diff --git a/Engine/Plugins/Experimental/HairStrands/Source/HairStrandsCore/Private/GroomComponent.cpp b/Engine/Plugins/Experimental/HairStrands/Source/HairStrandsCore/Private/GroomComponent.cpp
--- a/Engine/Plugins/Experimental/HairStrands/Source/HairStrandsCore/Private/GroomComponent.cpp
+++ b/Engine/Plugins/Experimental/HairStrands/Source/HairStrandsCore/Private/GroomComponent.cpp
@@ -998,7 +998,7 @@ void UGroomComponent::PostLoad()
 
 
 #if WITH_EDITOR
-	if (!bIsGroomAssetCallbackRegistered)
+	if (GroomAsset && !bIsGroomAssetCallbackRegistered)
 	{
 		GroomAsset->GetOnGroomAssetChanged().AddUObject(this, &UGroomComponent::Invalidate);
 		bIsGroomAssetCallbackRegistered = true;
