Fixed Groom resources not being initialized when compiling in BP editor.
Fixed crash when setting a GroomAsset on GroomComponent when that asset has a different number of groups than the previous one.
Fixed OnGroomAssetChanged callback not being removed from the right asset.

diff --git a/Engine/Plugins/Experimental/HairStrands/Source/HairStrandsCore/Private/GroomComponent.cpp b/Engine/Plugins/Experimental/HairStrands/Source/HairStrandsCore/Private/GroomComponent.cpp
--- a/Engine/Plugins/Experimental/HairStrands/Source/HairStrandsCore/Private/GroomComponent.cpp
+++ b/Engine/Plugins/Experimental/HairStrands/Source/HairStrandsCore/Private/GroomComponent.cpp
@@ -991,11 +991,7 @@ void UGroomComponent::PostLoad()
 	}
 
 	UpdateHairGroupsDesc(GroomAsset, GroomGroupsDesc);
-	if (!IsTemplate())
-	{
-		InitResources();
-	}
-
+	InitResources();
 
 #if WITH_EDITOR
 	if (GroomAsset && !bIsGroomAssetCallbackRegistered)
@@ -1019,6 +1015,12 @@ void UGroomComponent::OnRegister()
 {
 	Super::OnRegister();
 
+	if (!InitializedResources)
+	{
+		UpdateHairGroupsDesc(GroomAsset, GroomGroupsDesc);
+		InitResources();
+	}
+
 	// Insure the ticking of the Groom component always happens after the skeletalMeshComponent.
 	USkeletalMeshComponent* SkeletalMeshComponent = GetAttachParent() ? Cast<USkeletalMeshComponent>(GetAttachParent()) : nullptr;
 	/*if (SkeletalMeshComponent)
@@ -1346,15 +1348,39 @@ void UGroomComponent::GetUsedMaterials(TArray<UMaterialInterface*>& OutMaterials
 }
 
 #if WITH_EDITOR
-void UGroomComponent::PostEditChangeProperty(FPropertyChangedEvent& PropertyChangedEvent) 
+void UGroomComponent::PreEditChange(UProperty* PropertyAboutToChange)
 {
-	Super::PostEditChangeProperty(PropertyChangedEvent);
+	Super::PreEditChange(PropertyAboutToChange);
+
+	const FName PropertyName = PropertyAboutToChange ? PropertyAboutToChange->GetFName() : NAME_None;
+	const bool bAssetAboutToChanged = PropertyName == GET_MEMBER_NAME_CHECKED(UGroomComponent, GroomAsset);
+	if (bAssetAboutToChanged)
+	{
+		// Remove the callback on the GroomAsset about to be replaced
+		if (bIsGroomAssetCallbackRegistered && GroomAsset)
+		{
+			GroomAsset->GetOnGroomAssetChanged().RemoveAll(this);
+		}
+		bIsGroomAssetCallbackRegistered = false;
+	}
+}
 
+void UGroomComponent::PostEditChangeProperty(FPropertyChangedEvent& PropertyChangedEvent) 
+{
 	UProperty* PropertyThatChanged = PropertyChangedEvent.Property;
 	const FName PropertyName = PropertyThatChanged ? PropertyThatChanged->GetFName() : NAME_None;
 
 	//  Init/release resources when setting the GroomAsset (or undoing)
 	const bool bAssetChanged = PropertyName == GET_MEMBER_NAME_CHECKED(UGroomComponent, GroomAsset);
+	if (bAssetChanged)
+	{
+		// Release the resources before Super::PostEditChangeProperty so that they get
+		// re-initialized in OnRegister
+		ReleaseResources();
+	}
+
+	Super::PostEditChangeProperty(PropertyChangedEvent);
+
 	const bool bRecreateResources =
 		(PropertyName == GET_MEMBER_NAME_CHECKED(UGroomComponent, GroomAsset) || PropertyThatChanged == nullptr) ||
 		PropertyName == GET_MEMBER_NAME_CHECKED(UGroomComponent, bBindGroomToSkeletalMesh);
@@ -1368,14 +1394,9 @@ void UGroomComponent::PostEditChangeProperty(FPropertyChangedEvent& PropertyChan
 #if WITH_EDITOR
 	if (bAssetChanged)
 	{
-		if (bIsGroomAssetCallbackRegistered && GroomAsset)
-		{
-			GroomAsset->GetOnGroomAssetChanged().RemoveAll(this);
-		}
-		bIsGroomAssetCallbackRegistered = false;
-
 		if (GroomAsset)
 		{
+			// Set the callback on the new GroomAsset being assigned
 			GroomAsset->GetOnGroomAssetChanged().AddUObject(this, &UGroomComponent::Invalidate);
 			bIsGroomAssetCallbackRegistered = true;
 		}
diff --git a/Engine/Plugins/Experimental/HairStrands/Source/HairStrandsCore/Public/GroomComponent.h b/Engine/Plugins/Experimental/HairStrands/Source/HairStrandsCore/Public/GroomComponent.h
--- a/Engine/Plugins/Experimental/HairStrands/Source/HairStrandsCore/Public/GroomComponent.h
+++ b/Engine/Plugins/Experimental/HairStrands/Source/HairStrandsCore/Public/GroomComponent.h
@@ -101,6 +101,7 @@ public:
 
 #if WITH_EDITOR
 	virtual void CheckForErrors() override;
+	virtual void PreEditChange(UProperty* PropertyAboutToChange) override;
 	virtual void PostEditChangeProperty(FPropertyChangedEvent& PropertyChangedEvent) override;
 	virtual bool CanEditChange(const UProperty* InProperty) const override;
 	void ValidateMaterials(bool bMapCheck) const;
