Fix import crash with low guide count.

Change the low/medium quality interpolation to do a progressive spactial exploration covering the entire volume if needed, to insure we find at least one guide per strands
Change the randomization to be just a permutation/selection of the three closest guides.

diff --git a/Engine/Plugins/Experimental/HairStrands/Source/HairStrandsCore/Private/GroomBuilder.cpp b/Engine/Plugins/Experimental/HairStrands/Source/HairStrandsCore/Private/GroomBuilder.cpp
--- a/Engine/Plugins/Experimental/HairStrands/Source/HairStrandsCore/Private/GroomBuilder.cpp
+++ b/Engine/Plugins/Experimental/HairStrands/Source/HairStrandsCore/Private/GroomBuilder.cpp
@@ -380,7 +380,7 @@ namespace HairInterpolationBuilder
 
 	struct FMetrics
 	{
-		static const uint32 Count = 6;
+		static const uint32 Count = 3;
 		float KMinMetrics[Count];
 		int32 KClosestGuideIndices[Count];
 	};
@@ -395,22 +395,34 @@ namespace HairInterpolationBuilder
 	static void SelectFinalGuides(
 		FClosestGuides& ClosestGuides, 
 		FRandomStream& Random, 
-		const FMetrics& Metric, 
+		const FMetrics& InMetric, 
 		const bool bRandomizeInterpolation, 
 		const bool bUseUniqueGuide)
 	{
+		FMetrics Metric = InMetric;
+		check(Metric.KClosestGuideIndices[0] >= 0);
+
+		// If some indices are invalid (for instance, found a valid single guide, fill in the rest with the valid ones)
+		if (Metric.KClosestGuideIndices[1] < 0)
+		{
+			Metric.KClosestGuideIndices[1] = Metric.KClosestGuideIndices[0];
+			Metric.KMinMetrics[1] = Metric.KMinMetrics[0];
+		}
+		if (Metric.KClosestGuideIndices[2] < 0)
+		{
+			Metric.KClosestGuideIndices[2] = Metric.KClosestGuideIndices[1];
+			Metric.KMinMetrics[2] = Metric.KMinMetrics[1];
+		}
+
 		uint32 RandIndex0 = 0;
 		uint32 RandIndex1 = 1;
 		uint32 RandIndex2 = 2;
 		if (bRandomizeInterpolation)
 		{
-			do
-			{
-				RandIndex0 = Random.RandRange(0, FMetrics::Count - 1);
-				RandIndex1 = Random.RandRange(0, FMetrics::Count - 1);
-				RandIndex2 = Random.RandRange(0, FMetrics::Count - 1);
-
-			} while (RandIndex0 == RandIndex1 || RandIndex0 == RandIndex2 || RandIndex1 == RandIndex2);
+			// This randomization makes certain strands being affected by 1, 2, or 3 guides
+			RandIndex0 = Random.RandRange(0, FMetrics::Count - 1);
+			RandIndex1 = Random.RandRange(0, FMetrics::Count - 1);
+			RandIndex2 = Random.RandRange(0, FMetrics::Count - 1);
 		}
 
 		ClosestGuides.Indices[0] = Metric.KClosestGuideIndices[RandIndex0];
@@ -462,7 +474,7 @@ namespace HairInterpolationBuilder
 
 		check(MinMetrics[0] <= MinMetrics[1]);
 		check(MinMetrics[1] <= MinMetrics[2]);
-	};
+	}
 
 	// Simple closest distance metric
 	static void ComputeSimpleMetric(
@@ -516,14 +528,22 @@ namespace HairInterpolationBuilder
 		FVector MinBound = FVector::ZeroVector;
 		FVector MaxBound = FVector::ZeroVector;
 		
-		const uint32 MaxLookupDistance = 2;
+		const uint32 MaxLookupDistance = 31;
 		const FIntVector GridResolution = FIntVector(32, 32, 32);
 
 		TArray<int32> GridIndirection;
 		TArray<TArray<int32>> RootIndices;
+		
+		FORCEINLINE bool IsValid(const FIntVector& P) const
+		{
+			return	0 <= P.X && P.X < GridResolution.X &&
+					0 <= P.Y && P.Y < GridResolution.Y &&
+					0 <= P.Z && P.Z < GridResolution.Z;
+		}
 
-		FORCEINLINE FIntVector ClampToVolume(const FIntVector& CellCoord) const
+		FORCEINLINE FIntVector ClampToVolume(const FIntVector& CellCoord, bool& bIsValid) const
 		{
+			bIsValid = IsValid(CellCoord);
 			return FIntVector(
 				FMath::Clamp(CellCoord.X, 0, GridResolution.X - 1),
 				FMath::Clamp(CellCoord.Y, 0, GridResolution.Y - 1),
@@ -532,9 +552,10 @@ namespace HairInterpolationBuilder
 
 		FORCEINLINE FIntVector ToCellCoord(const FVector& P) const
 		{
+			bool bIsValid = false;
 			const FVector F = ((P - MinBound) / (MaxBound - MinBound));
 			const FIntVector CellCoord = FIntVector(FMath::FloorToInt(F.X * GridResolution.X), FMath::FloorToInt(F.Y * GridResolution.Y), FMath::FloorToInt(F.Z * GridResolution.Z));			
-			return ClampToVolume(CellCoord);
+			return ClampToVolume(CellCoord, bIsValid);
 		}
 
 		uint32 ToIndex(const FIntVector& CellCoord) const
@@ -569,6 +590,29 @@ namespace HairInterpolationBuilder
 			}
 		}
 
+		FORCEINLINE void SearchCell(
+			const FIntVector& CellCoord,
+			const uint32 RenCurveIndex,
+			const FHairRoot& RenRoot,
+			const TArray<FHairRoot>& SimRoots,
+			FMetrics& Metrics) const 
+		{
+			const uint32 CellIndex = ToIndex(CellCoord);
+
+			if (GridIndirection[CellIndex] == 0)
+				return;
+
+			const TArray<int32>& Elements = RootIndices[GridIndirection[CellIndex]];
+
+			for (int32 SimCurveIndex : Elements)
+			{
+				const FHairRoot& GuideRoot = SimRoots[SimCurveIndex];
+				{
+					ComputeSimpleMetric(Metrics, RenRoot, GuideRoot, RenCurveIndex, SimCurveIndex);
+				}
+			}
+		}
+
 		FClosestGuides FindClosestRoots(
 			const uint32 RenCurveIndex,
 			const TArray<FHairRoot>& RenRoots,
@@ -589,43 +633,55 @@ namespace HairInterpolationBuilder
 				Metrics.KClosestGuideIndices[ClearIt] = -1;
 			}
 
-			// Up to 5x5X5 lookup within a 32x32x32 volume
-			// * Tracks two metrics: a complex (pairing) one, and simple (closest) one
-			// * If the complex one fails to find a valid guide, we fallback on the simplest one
-			//for (uint32 R = 1; R < MaxLookupDistance; ++R)
+			for (int32 Offset = 1; Offset <= int32(MaxLookupDistance); ++Offset)
 			{
-				int32 Offset = MaxLookupDistance;
+				// Center
+				{
+					bool bIsValid = false;
+					const FIntVector CellCoord = ClampToVolume(PointCoord, bIsValid);
+					if (bIsValid) SearchCell(PointCoord, RenCurveIndex, RenRoot, SimRoots, Metrics);
+				}
+
+				// Top & Bottom
 				for (int32 X = -Offset; X <= Offset; ++X)
+				for (int32 Y = -Offset; Y <= Offset; ++Y)
 				{
-					for (int32 Y = -Offset; Y <= Offset; ++Y)
-					{
-						for (int32 Z = -Offset; Z <= Offset; ++Z)
-						{
-							const FIntVector CellCoord = ClampToVolume(PointCoord + FIntVector(X, Y, Z));
-							const uint32 CellIndex = ToIndex(CellCoord);
-
-							if (GridIndirection[CellIndex] == 0)
-								continue;
-
-							const TArray<int32>& Elements = RootIndices[GridIndirection[CellIndex]];
-					
-							for (int32 SimCurveIndex : Elements)
-							{
-								const FHairRoot& GuideRoot = SimRoots[SimCurveIndex];
-								{
-									ComputeSimpleMetric(Metrics, RenRoot, GuideRoot, RenCurveIndex, SimCurveIndex);
-								}
-							}
-						}
-					}
+					bool bIsValid0 = false, bIsValid1 = false;
+					const FIntVector CellCoord0 = ClampToVolume(PointCoord + FIntVector(X, Y, Offset), bIsValid0);
+					const FIntVector CellCoord1 = ClampToVolume(PointCoord + FIntVector(X, Y,-Offset), bIsValid1);
+					if (bIsValid0) SearchCell(CellCoord0, RenCurveIndex, RenRoot, SimRoots, Metrics);
+					if (bIsValid1) SearchCell(CellCoord1, RenCurveIndex, RenRoot, SimRoots, Metrics);
+				}
+
+				const int32 OffsetMinusOne = Offset - 1;
+				// Front & Back
+				for (int32 X = -Offset; X <= Offset; ++X)
+				for (int32 Z = -OffsetMinusOne; Z <= OffsetMinusOne; ++Z)
+				{
+					bool bIsValid0 = false, bIsValid1 = false;
+					const FIntVector CellCoord0 = ClampToVolume(PointCoord + FIntVector(X,  Offset, Z), bIsValid0);
+					const FIntVector CellCoord1 = ClampToVolume(PointCoord + FIntVector(X, -Offset, Z), bIsValid1);
+					if (bIsValid0) SearchCell(CellCoord0, RenCurveIndex, RenRoot, SimRoots, Metrics);
+					if (bIsValid1) SearchCell(CellCoord1, RenCurveIndex, RenRoot, SimRoots, Metrics);
+				}
+				
+				// Left & Right
+				for (int32 Y = -OffsetMinusOne; Y <= OffsetMinusOne; ++Y)
+				for (int32 Z = -OffsetMinusOne; Z <= OffsetMinusOne; ++Z)
+				{
+					bool bIsValid0 = false, bIsValid1 = false;
+					const FIntVector CellCoord0 = ClampToVolume(PointCoord + FIntVector( Offset, Y, Z), bIsValid0);
+					const FIntVector CellCoord1 = ClampToVolume(PointCoord + FIntVector(-Offset, Y, Z), bIsValid1);
+					if (bIsValid0) SearchCell(CellCoord0, RenCurveIndex, RenRoot, SimRoots, Metrics);
+					if (bIsValid1) SearchCell(CellCoord1, RenCurveIndex, RenRoot, SimRoots, Metrics);
 				}
 
 				// Early out if we have found all closest guide during a ring/layer step.
 				// This early out is not conservative, as the complex metric might find better guides one or multiple step further.
-				//if (Metrics.KClosestGuideIndices[FMetrics::ClosestCount - 1] != -1 && R >= 2)
-				//{
-				//	break;
-				//}
+				if (Metrics.KClosestGuideIndices[FMetrics::Count-1] != -1 && Offset >= 2)
+				{
+					break;
+				}
 			}
 
 			// If no valid guide have been found, switch to a simpler metric
@@ -640,6 +696,33 @@ namespace HairInterpolationBuilder
 		}
 
 
+		FORCEINLINE void SearchCell(
+			const FIntVector& CellCoord,
+			const uint32 RenCurveIndex,
+			const FHairRoot& RenRoot,
+			const TArray<FHairRoot>& SimRoots,
+			const FHairStrandsDatas& RenStrandsData,
+			const FHairStrandsDatas& SimStrandsData,
+			FMetrics& Metrics0,
+			FMetrics& Metrics1) const
+		{
+			const uint32 CellIndex = ToIndex(CellCoord);
+
+			if (GridIndirection[CellIndex] == 0)
+				return;
+
+			const TArray<int32>& Elements = RootIndices[GridIndirection[CellIndex]];
+
+			for (int32 SimCurveIndex : Elements)
+			{
+				const FHairRoot& GuideRoot = SimRoots[SimCurveIndex];
+				{
+					ComputeSimpleMetric(Metrics1, RenRoot, GuideRoot, RenCurveIndex, SimCurveIndex);
+					ComputeAdvandedMetric(Metrics0, RenStrandsData, SimStrandsData, RenCurveIndex, SimCurveIndex);
+				}
+			}
+		}
+
 		FClosestGuides FindBestClosestRoots(
 			const uint32 RenCurveIndex,
 			const TArray<FHairRoot>& RenRoots,
@@ -663,50 +746,64 @@ namespace HairInterpolationBuilder
 				Metrics1.KClosestGuideIndices[ClearIt] = -1;
 			}
 
-			// Up to 7x7X7 lookup within a 32x32x32 volume
-			// * Tracks two metrics: a complex (pairing) one, and simple (closest) one
-			// * If the complex one fails to find a valid guide, we fallback on the simplest one
-			//for (uint32 R = 1; R < MaxLookupDistance; ++R)
+			for (int32 Offset = 1; Offset <= int32(MaxLookupDistance); ++Offset)
 			{
-				int32 Offset = MaxLookupDistance;
+				// Center
+				{
+					bool bIsValid = false;
+					const FIntVector CellCoord = ClampToVolume(PointCoord, bIsValid);
+					if (bIsValid) SearchCell(CellCoord, RenCurveIndex, RenRoot, SimRoots, RenStrandsData, SimStrandsData, Metrics0, Metrics1);
+				}
+
+				// Top & Bottom
 				for (int32 X = -Offset; X <= Offset; ++X)
+				for (int32 Y = -Offset; Y <= Offset; ++Y)
 				{
-					for (int32 Y = -Offset; Y <= Offset; ++Y) 
-					{
-						for (int32 Z = -Offset; Z <= Offset; ++Z)
-						{
-							const FIntVector CellCoord = ClampToVolume(PointCoord + FIntVector(X, Y, Z));
-							const uint32 CellIndex = ToIndex(CellCoord);
-
-							if (GridIndirection[CellIndex] == 0)
-								continue;
-
-							const TArray<int32>& Elements = RootIndices[GridIndirection[CellIndex]];
-
-							for (int32 SimCurveIndex : Elements)
-							{
-								const FHairRoot& GuideRoot = SimRoots[SimCurveIndex];
-								{
-									ComputeSimpleMetric(Metrics1, RenRoot, GuideRoot, RenCurveIndex, SimCurveIndex);
-									ComputeAdvandedMetric(Metrics0, RenStrandsData, SimStrandsData, RenCurveIndex, SimCurveIndex);
-								}
-							}
-						}
-					}
+					bool bIsValid0 = false, bIsValid1 = false;
+					const FIntVector CellCoord0 = ClampToVolume(PointCoord + FIntVector(X, Y, Offset), bIsValid0);
+					const FIntVector CellCoord1 = ClampToVolume(PointCoord + FIntVector(X, Y,-Offset), bIsValid1);
+					if (bIsValid0) SearchCell(CellCoord0, RenCurveIndex, RenRoot, SimRoots, RenStrandsData, SimStrandsData, Metrics0, Metrics1);
+					if (bIsValid1) SearchCell(CellCoord1, RenCurveIndex, RenRoot, SimRoots, RenStrandsData, SimStrandsData, Metrics0, Metrics1);
+				}
+
+				const int32 OffsetMinusOne = Offset - 1;
+				// Front & Back
+				for (int32 X = -Offset; X <= Offset; ++X)
+				for (int32 Z = -OffsetMinusOne; Z <= OffsetMinusOne; ++Z)
+				{
+					bool bIsValid0 = false, bIsValid1 = false;
+					const FIntVector CellCoord0 = ClampToVolume(PointCoord + FIntVector(X, Offset, Z), bIsValid0);
+					const FIntVector CellCoord1 = ClampToVolume(PointCoord + FIntVector(X, -Offset, Z), bIsValid1);
+					if (bIsValid0) SearchCell(CellCoord0, RenCurveIndex, RenRoot, SimRoots, RenStrandsData, SimStrandsData, Metrics0, Metrics1);
+					if (bIsValid1) SearchCell(CellCoord1, RenCurveIndex, RenRoot, SimRoots, RenStrandsData, SimStrandsData, Metrics0, Metrics1);
+				}
+				
+				// Left & Right
+				for (int32 Y = -OffsetMinusOne; Y <= OffsetMinusOne; ++Y)
+				for (int32 Z = -OffsetMinusOne; Z <= OffsetMinusOne; ++Z)
+				{
+					bool bIsValid0 = false, bIsValid1 = false;
+					const FIntVector CellCoord0 = ClampToVolume(PointCoord + FIntVector( Offset, Y, Z), bIsValid0);
+					const FIntVector CellCoord1 = ClampToVolume(PointCoord + FIntVector(-Offset, Y, Z), bIsValid1);
+					if (bIsValid0) SearchCell(CellCoord0, RenCurveIndex, RenRoot, SimRoots, RenStrandsData, SimStrandsData, Metrics0, Metrics1);
+					if (bIsValid1) SearchCell(CellCoord1, RenCurveIndex, RenRoot, SimRoots, RenStrandsData, SimStrandsData, Metrics0, Metrics1);
 				}
 
 				// Early out if we have found all closest guide during a ring/layer step.
 				// This early out is not conservative, as the complex metric might find better guides one or multiple step further.
-				//if (Metrics0.KClosestGuideIndices[FMetrics::KGuideCount - 1] != -1 && R >= 2)
-				//{
-				//	break;
-				//}
+				if ((Metrics0.KClosestGuideIndices[FMetrics::Count-1] != -1 || Metrics1.KClosestGuideIndices[FMetrics::Count - 1] != -1) && Offset >= 2)
+				{
+					break;
+				}
 			}
 
 			// If no valid guide have been found, switch to a simpler metric
 			FClosestGuides ClosestGuides;
-			SelectFinalGuides(ClosestGuides, Random, Metrics0, bRandomized, bUnique);
-			if (ClosestGuides.Indices[0] == -1)
+			if (Metrics0.KClosestGuideIndices[0] != -1)
+			{
+				SelectFinalGuides(ClosestGuides, Random, Metrics0, bRandomized, bUnique);
+			}
+			else
 			{
 				SelectFinalGuides(ClosestGuides, Random, Metrics1, bRandomized, bUnique);
 			}
