Fix hair strands guide visualization having the wrong offset when there is no simulation enabled.

diff --git a/Engine/Plugins/Experimental/HairStrands/Source/HairStrandsCore/Private/HairStrandsRendering.cpp b/Engine/Plugins/Experimental/HairStrands/Source/HairStrandsCore/Private/HairStrandsRendering.cpp
--- a/Engine/Plugins/Experimental/HairStrands/Source/HairStrandsCore/Private/HairStrandsRendering.cpp
+++ b/Engine/Plugins/Experimental/HairStrands/Source/HairStrandsCore/Private/HairStrandsRendering.cpp
@@ -517,13 +517,15 @@ void ComputeHairStrandsInterpolation(
 
 			GraphBuilder.Execute();
 
+			const bool bHasSimulationEnabled = Input.bIsSimulationEnable && GHairStrandsInterpolateSimulation && DeformationType != EDeformationType::RestStrands;
+
 			Output.VFInput.HairPositionBuffer = Output.SimDeformedPositionBuffer[SimIndex]->SRV;
 			Output.VFInput.HairPreviousPositionBuffer = Output.SimDeformedPositionBuffer[SimIndex]->SRV;
 			Output.VFInput.HairTangentBuffer = Output.SimTangentBuffer->SRV;
 			Output.VFInput.HairAttributeBuffer = Input.SimAttributeBuffer->SRV;
 			Output.VFInput.HairMaterialBuffer = Output.RenderMaterialBuffer->SRV;
-			Output.VFInput.HairPositionOffset = Input.OutHairPositionOffset;
-			Output.VFInput.HairPreviousPositionOffset = Input.OutHairPreviousPositionOffset;
+			Output.VFInput.HairPositionOffset = bHasSimulationEnabled ? Input.OutHairPositionOffset : Input.InSimHairPositionOffset;
+			Output.VFInput.HairPreviousPositionOffset = bHasSimulationEnabled ? Input.OutHairPreviousPositionOffset : Input.InSimHairPositionOffset;
 			Output.VFInput.VertexCount = Input.SimVertexCount;
 		}
 		else
