Fix Hair interpolation
Fix crash with skelmesh scaling
Groom bounding box used
Hair viscosity enabled

diff --git a/Engine/Plugins/Experimental/HairStrands/HairStrands.uplugin b/Engine/Plugins/Experimental/HairStrands/HairStrands.uplugin
--- a/Engine/Plugins/Experimental/HairStrands/HairStrands.uplugin
+++ b/Engine/Plugins/Experimental/HairStrands/HairStrands.uplugin
@@ -36,7 +36,7 @@
 		{
 			"Name": "HairStrandsNiagara",
 			"Type": "RunTime",
-			"LoadingPhase": "Default",
+			"LoadingPhase": "PostConfigInit",
 			"WhitelistPlatforms": [
 				"Win64",
 				"Win32"
diff --git a/Engine/Plugins/Experimental/HairStrands/Shaders/Private/NiagaraAngularSpringMaterial.ush b/Engine/Plugins/Experimental/HairStrands/Shaders/Private/NiagaraAngularSpringMaterial.ush
new file mode 100644
--- /dev/null
+++ b/Engine/Plugins/Experimental/HairStrands/Shaders/Private/NiagaraAngularSpringMaterial.ush
@@ -0,0 +1,542 @@
+// Copyright 1998-2019 Epic Games, Inc. All Rights Reserved.
+
+#pragma once
+
+/* -----------------------------------------------------------------
+ * Stretch spring material
+ * -----------------------------------------------------------------
+ */
+
+ void ResetStretchSpringMaterial(in int StrandsSize, in float YoungModulus, in float RodThickness, 
+	in float RestLength, in float DeltaTime, in bool ProjectConstraint, in float MaterialDamping, out float OutMaterialCompliance, out float OutMaterialWeight, out float OutMaterialMultiplier, in int NodeOffset)
+{
+	// Compliance = 1.0 / (k * dt * dt)
+	// with k = L * L * (Y * A / L) (the L*L term before is coming from the fact that our constraint is dL/L and not dL)
+	// A is the cross section area (Pi * R * R), L is the rest length and Y is the young modulus
+	OutMaterialCompliance = 4.0/(YoungModulus*PI*RodThickness*RestLength*RodThickness*DeltaTime*DeltaTime);
+
+	const float SumInverseMass = !ProjectConstraint ? ( SharedInverseMass[GGroupThreadId.x] + SharedInverseMass[GGroupThreadId.x-1-NodeOffset] ) / (RestLength*RestLength) :
+				SharedInverseMass[GGroupThreadId.x] / (RestLength*RestLength);
+
+	const float SchurDiagonal = ( (1.0  + MaterialDamping / DeltaTime ) * SumInverseMass + OutMaterialCompliance );
+	OutMaterialWeight = ( SchurDiagonal != 0.0 ) ? 1.0 / SchurDiagonal : 0.0;
+
+	OutMaterialMultiplier = 0;
+}
+
+
+void UpdateStretchSpringMultiplier(in float RestLength, in float DeltaTime, in bool ProjectConstraint, in float MaterialDamping, in float MaterialCompliance, in float MaterialWeight, inout float OutMaterialMultiplier, in int NodeOffset)
+{
+	const int PreviousIndex = GGroupThreadId.x-1-NodeOffset;
+	float3 EdgeDirection = SharedNodePosition[GGroupThreadId.x] - SharedNodePosition[PreviousIndex];
+	const float3 DeltaVelocity = (EdgeDirection - ( SharedPreviousPosition[GGroupThreadId.x] - SharedPreviousPosition[PreviousIndex] ) ) / DeltaTime; 
+
+	const float EdgeLength = length(EdgeDirection);
+	EdgeDirection /= EdgeLength;
+
+	// XPBD lagrange multiplier update : dL = -(C+compliance*L) / (dC * invM * dCt + alpha)
+	const float DeltaLambda = -((EdgeLength-RestLength)/RestLength + OutMaterialMultiplier * MaterialCompliance + MaterialDamping * dot(EdgeDirection,DeltaVelocity)/RestLength ) * MaterialWeight;
+
+	// L += dL
+	OutMaterialMultiplier += DeltaLambda;
+
+	// XPBD position update : dX = dL * dCt * invM
+	const float3 PositionDelta = EdgeDirection * DeltaLambda / RestLength;
+
+	// dX += dX
+	SharedNodePosition[GGroupThreadId.x] += PositionDelta * SharedInverseMass[GGroupThreadId.x];
+	if(!ProjectConstraint)
+	{
+		SharedNodePosition[PreviousIndex] -= PositionDelta * SharedInverseMass[PreviousIndex];
+	}
+}
+
+void SetupStretchSpringMaterial(in int StrandsSize, in float YoungModulus, in float RodThickness, 
+	in float RestLength, in float DeltaTime, in bool ProjectConstraint, in float MaterialDamping, out float OutMaterialCompliance, out float OutMaterialWeight, out float OutMaterialMultiplier)
+{
+	OutMaterialCompliance = 0.0;
+	OutMaterialWeight = 0.0;
+	OutMaterialMultiplier = 0.0;
+	const int LocalIndex = (GGroupThreadId.x % StrandsSize);
+	if( LocalIndex > 1 )
+	{
+		ResetStretchSpringMaterial(StrandsSize,YoungModulus,RodThickness,RestLength,DeltaTime,ProjectConstraint,MaterialDamping,OutMaterialCompliance,OutMaterialWeight,OutMaterialMultiplier,0);
+	}
+}
+
+void SolveStretchSpringMaterial(in int StrandsSize, in float RestLength, in float DeltaTime, in float MaterialDamping, 
+			in float MaterialCompliance, in float MaterialWeight, in float MaterialMultiplier, out float OutMaterialMultiplier)
+{
+	OutMaterialMultiplier = MaterialMultiplier;
+	const int LocalIndex = (GGroupThreadId.x % StrandsSize);
+	if( LocalIndex > 1)
+	{
+		const int IsRed = (GGroupThreadId.x % 2) == 0;
+		// Process all the red rods
+		if (IsRed)
+		{	
+			UpdateStretchSpringMultiplier(RestLength,DeltaTime,false,MaterialDamping,MaterialCompliance,MaterialWeight,OutMaterialMultiplier,0);
+		}
+		// Process all the black rods
+		GroupMemoryBarrier(); 
+		if (!IsRed)
+		{
+			UpdateStretchSpringMultiplier(RestLength,DeltaTime,false,MaterialDamping,MaterialCompliance,MaterialWeight,OutMaterialMultiplier,0);
+		}
+		GroupMemoryBarrier(); 
+	}
+} 
+
+void ProjectStretchSpringMaterial(in int StrandsSize, in float YoungModulus, in float RodThickness, in float RestLength, in float NodeMass, in float DeltaTime, in float3 NodePosition, out float3 OutNodePosition)
+{
+	const int LocalIndex = (GGroupThreadId.x % StrandsSize);
+
+	SharedNodePosition[GGroupThreadId.x] = NodePosition;
+	SharedPreviousPosition[GGroupThreadId.x] = NodePosition;
+	SharedInverseMass[GGroupThreadId.x] = (LocalIndex > 1 && NodeMass != 0.0) ? 1.0 / NodeMass : 0.0;
+	GroupMemoryBarrier();
+
+	if(DeltaTime != 0.0)
+	{
+		float MaterialCompliance = 0.0;
+		float MaterialWeight = 0.0;
+		float MaterialMultiplier = 0.0;
+
+		if( LocalIndex > 1 )
+		{
+			ResetStretchSpringMaterial(StrandsSize,YoungModulus,RodThickness,RestLength,DeltaTime,true,0.0,MaterialCompliance,MaterialWeight,MaterialMultiplier,0);
+
+			for(int i = 2; i < StrandsSize; ++i)
+			{
+				if( LocalIndex == i )
+				{
+					UpdateStretchSpringMultiplier(RestLength,DeltaTime,true,0.0,MaterialCompliance,MaterialWeight,MaterialMultiplier,0);
+				}
+				GroupMemoryBarrier(); 
+			}
+		}
+		
+	}
+	GroupMemoryBarrier();
+	OutNodePosition = SharedNodePosition[GGroupThreadId.x] ;
+}
+
+/* -----------------------------------------------------------------
+ * Bend spring material
+ * -----------------------------------------------------------------
+ */
+
+ void SetupBendSpringMaterial(in int StrandsSize, in float YoungModulus, in float RodThickness, 
+	in float RestLength, in float DeltaTime, in bool ProjectConstraint, in float MaterialDamping, out float OutMaterialCompliance, out float OutMaterialWeight, out float OutMaterialMultiplier)
+{
+	OutMaterialCompliance = 0.0;
+	OutMaterialWeight = 0.0;
+	OutMaterialMultiplier = 0.0;
+	const int LocalIndex = (GGroupThreadId.x % StrandsSize);
+	if( LocalIndex > 1 )
+	{
+		ResetStretchSpringMaterial(StrandsSize,YoungModulus,RodThickness,RestLength,DeltaTime,ProjectConstraint,MaterialDamping,OutMaterialCompliance,OutMaterialWeight,OutMaterialMultiplier,1);
+	}
+}
+
+void SolveBendSpringMaterial(in int StrandsSize, in float RestLength, in float DeltaTime, in float MaterialDamping, 
+			in float MaterialCompliance, in float MaterialWeight, in float MaterialMultiplier, out float OutMaterialMultiplier)
+{
+	OutMaterialMultiplier = MaterialMultiplier;
+	const int LocalIndex = (GGroupThreadId.x % StrandsSize);
+	if( LocalIndex > 1)
+	{
+		const int IsRed = ((GGroupThreadId.x/2) % 2) == 0;
+		// Process all the red rods
+		if (IsRed)
+		{	
+			UpdateStretchSpringMultiplier(RestLength,DeltaTime,false,MaterialDamping,MaterialCompliance,MaterialWeight,OutMaterialMultiplier,1);
+		}
+		// Process all the black rods
+		GroupMemoryBarrier(); 
+		if (!IsRed)
+		{
+			UpdateStretchSpringMultiplier(RestLength,DeltaTime,false,MaterialDamping,MaterialCompliance,MaterialWeight,OutMaterialMultiplier,1);
+		}
+		GroupMemoryBarrier(); 
+	}
+} 
+
+void ProjectBendSpringMaterial(in int StrandsSize, in float YoungModulus, in float RodThickness, in float RestLength, in float NodeMass, in float DeltaTime, in float3 NodePosition, out float3 OutNodePosition)
+{
+	const int LocalIndex = (GGroupThreadId.x % StrandsSize);
+
+	SharedNodePosition[GGroupThreadId.x] = NodePosition;
+	SharedPreviousPosition[GGroupThreadId.x] = NodePosition;
+	SharedInverseMass[GGroupThreadId.x] = (LocalIndex > 1 && NodeMass != 0.0) ? 1.0 / NodeMass : 0.0;
+	GroupMemoryBarrier();
+
+	if(DeltaTime != 0.0)
+	{
+		float MaterialCompliance = 0.0;
+		float MaterialWeight = 0.0;
+		float MaterialMultiplier = 0.0;
+
+		if( LocalIndex > 1 )
+		{
+			ResetStretchSpringMaterial(StrandsSize,YoungModulus,RodThickness,RestLength,DeltaTime,true,0.0,MaterialCompliance,MaterialWeight,MaterialMultiplier,1);
+
+			for(int i = 2; i < StrandsSize; ++i)
+			{
+				if( LocalIndex == i )
+				{
+					UpdateStretchSpringMultiplier(RestLength,DeltaTime,true,0.0,MaterialCompliance,MaterialWeight,MaterialMultiplier,1);
+				}
+				GroupMemoryBarrier(); 
+			}
+		}
+	}
+	GroupMemoryBarrier();
+	OutNodePosition = SharedNodePosition[GGroupThreadId.x];
+}
+
+/* -----------------------------------------------------------------
+ * Twist spring material
+ * -----------------------------------------------------------------
+ */
+
+ void SetupTwistSpringMaterial(in int StrandsSize, in float YoungModulus, in float RodThickness, 
+	in float RestLength, in float DeltaTime, in bool ProjectConstraint, in float MaterialDamping, out float OutMaterialCompliance, out float OutMaterialWeight, out float OutMaterialMultiplier)
+{
+	OutMaterialCompliance = 0.0;
+	OutMaterialWeight = 0.0;
+	OutMaterialMultiplier = 0.0;
+	const int LocalIndex = (GGroupThreadId.x % StrandsSize);
+	if( LocalIndex > 2 )
+	{
+		ResetStretchSpringMaterial(StrandsSize,YoungModulus,RodThickness,RestLength,DeltaTime,ProjectConstraint,MaterialDamping,OutMaterialCompliance,OutMaterialWeight,OutMaterialMultiplier,2);
+	}
+}
+
+void SolveTwistSpringMaterial(in int StrandsSize, in float RestLength, in float DeltaTime, in float MaterialDamping, 
+			in float MaterialCompliance, in float MaterialWeight, in float MaterialMultiplier, out float OutMaterialMultiplier)
+{
+	OutMaterialMultiplier = MaterialMultiplier;
+	const int LocalIndex = (GGroupThreadId.x % StrandsSize);
+	if( LocalIndex > 2)
+	{
+		const int IsRed = ((GGroupThreadId.x/3) % 2) == 0;
+		// Process all the red rods
+		if (IsRed)
+		{	
+			UpdateStretchSpringMultiplier(RestLength,DeltaTime,false,MaterialDamping,MaterialCompliance,MaterialWeight,OutMaterialMultiplier,2);
+		}
+		// Process all the black rods
+		GroupMemoryBarrier(); 
+		if (!IsRed)
+		{
+			UpdateStretchSpringMultiplier(RestLength,DeltaTime,false,MaterialDamping,MaterialCompliance,MaterialWeight,OutMaterialMultiplier,2);
+		}
+		GroupMemoryBarrier(); 
+	}
+} 
+
+void ProjectTwistSpringMaterial(in int StrandsSize, in float YoungModulus, in float RodThickness, in float RestLength, in float NodeMass, in float DeltaTime, in float3 NodePosition, out float3 OutNodePosition)
+{
+	const int LocalIndex = (GGroupThreadId.x % StrandsSize);
+
+	SharedNodePosition[GGroupThreadId.x] = NodePosition;
+	SharedPreviousPosition[GGroupThreadId.x] = NodePosition;
+	SharedInverseMass[GGroupThreadId.x] = (LocalIndex > 1 && NodeMass != 0.0) ? 1.0 / NodeMass : 0.0;
+	GroupMemoryBarrier();
+
+	if(DeltaTime != 0.0)
+	{
+		float MaterialCompliance = 0.0;
+		float MaterialWeight = 0.0;
+		float MaterialMultiplier = 0.0;
+
+		if( LocalIndex > 2 )
+		{
+			ResetStretchSpringMaterial(StrandsSize,YoungModulus,RodThickness,RestLength,DeltaTime,true,0.0,MaterialCompliance,MaterialWeight,MaterialMultiplier,2);
+
+			for(int i = 3; i < StrandsSize; ++i)
+			{
+				if( LocalIndex == i )
+				{
+					UpdateStretchSpringMultiplier(RestLength,DeltaTime,true,0.0,MaterialCompliance,MaterialWeight,MaterialMultiplier,2);
+				}
+				GroupMemoryBarrier(); 
+			}
+		}
+	}
+	GroupMemoryBarrier();
+	OutNodePosition = SharedNodePosition[GGroupThreadId.x];
+}
+
+/* -----------------------------------------------------------------
+ * Compute Bishop frame
+ * -----------------------------------------------------------------
+ */
+
+void ComputeMaterialFrame(in int StrandsSize)
+{
+	const int LocalIndex = (GGroupThreadId.x % StrandsSize);
+	if( LocalIndex == 0 )
+	{
+		float4 NodeQuaternion = SharedNodeOrientation[GGroupThreadId.x];
+		float3 TangentPrev = normalize(SharedNodePosition[GGroupThreadId.x+1] - SharedNodePosition[GGroupThreadId.x]);
+		float3 TangentNext = TangentPrev;
+
+		for (int EdgeIndex = 1, EdgeEnd = StrandsSize-1; EdgeIndex < EdgeEnd; ++EdgeIndex)
+		{
+			TangentPrev = TangentNext;
+			TangentNext = normalize(SharedNodePosition[GGroupThreadId.x+EdgeIndex+1] - SharedNodePosition[GGroupThreadId.x+EdgeIndex]);
+
+			NodeQuaternion = NormalizeQuat( MultiplyQuat( FindQuatBetweenNormals(TangentPrev,TangentNext), NodeQuaternion ) );
+			SharedNodeOrientation[GGroupThreadId.x+EdgeIndex] = NodeQuaternion;
+			GroupMemoryBarrier();
+		}
+		GroupMemoryBarrier();
+	}
+	GroupMemoryBarrier();
+}
+
+void UpdateMaterialFrame(in int StrandsSize)
+{
+	const int LocalIndex = (GGroupThreadId.x % StrandsSize);
+	if( LocalIndex == 0 )
+	{
+		float4 NodeQuaternion = SharedPreviousOrientation[GGroupThreadId.x];
+		float3 TangentPrev = normalize(SharedPreviousPosition[GGroupThreadId.x+1] - SharedPreviousPosition[GGroupThreadId.x]);
+		float3 TangentNext = TangentPrev;
+
+		for (int EdgeIndex = 1, EdgeEnd = StrandsSize-1; EdgeIndex < EdgeEnd; ++EdgeIndex)
+		{
+			TangentPrev = TangentNext;
+			TangentNext = normalize(SharedPreviousPosition[GGroupThreadId.x+EdgeIndex+1] - SharedPreviousPosition[GGroupThreadId.x+EdgeIndex]);
+
+			NodeQuaternion = NormalizeQuat( MultiplyQuat( FindQuatBetweenNormals(TangentPrev,TangentNext), NodeQuaternion ) );
+			SharedNodeOrientation[GGroupThreadId.x+EdgeIndex] = NodeQuaternion;
+			GroupMemoryBarrier();
+		}
+		GroupMemoryBarrier();
+	}
+	GroupMemoryBarrier();
+}
+
+
+/*void UpdateMaterialFrame(in int StrandsSize)
+{
+	const int LocalIndex = (GGroupThreadId.x % StrandsSize);
+	if( (LocalIndex > 0) && (LocalIndex < StrandsSize-1) )
+	{
+		float4 NodeQuaternion = SharedPreviousOrientation[GGroupThreadId.x];
+		//float3 TangentPrev = RotateVectorByQuat( float3(0,0,1), NodeQuaternion);
+		//float3 TangentNext = normalize(SharedPreviousPosition[GGroupThreadId.x+1] - SharedPreviousPosition[GGroupThreadId.x]);
+
+		float3 TangentPrev = normalize(SharedPreviousPosition[GGroupThreadId.x+1] - SharedPreviousPosition[GGroupThreadId.x]);
+		float3 TangentNext = normalize(SharedNodePosition[GGroupThreadId.x+1] - SharedNodePosition[GGroupThreadId.x]);
+
+		SharedNodeOrientation[GGroupThreadId.x] = NormalizeQuat( MultiplyQuat( FindQuatBetweenNormals(TangentPrev,TangentNext), NodeQuaternion) );
+	}
+	GroupMemoryBarrier();
+}*/
+
+/* -----------------------------------------------------------------
+ * Bend spring material
+ * -----------------------------------------------------------------
+ */
+
+ void SetupAngularSpringMaterial(in int StrandsSize, in float YoungModulus, in float RodThickness, 
+	in float RestLength, in float DeltaTime, in bool ProjectConstraint, in float MaterialDamping, out float OutMaterialCompliance, out float OutMaterialWeight, out float3 OutMaterialMultiplier)
+{
+	OutMaterialCompliance = 0.0;
+	OutMaterialWeight = 0.0;
+	OutMaterialMultiplier = float3(0,0,0);
+
+	float MaterialMultiplier = 0;
+	SetupStretchSpringMaterial(StrandsSize,YoungModulus,RodThickness,RestLength,DeltaTime,ProjectConstraint,MaterialDamping,OutMaterialCompliance,OutMaterialWeight,MaterialMultiplier);
+}
+
+void UpdateAngularSpringMultiplier( in float RestLength, in float3 RestDirection, in float DeltaTime, in bool ProjectConstraint, in float MaterialDamping, in float InMaterialCompliance, in float InMaterialWeight, inout float3 OutMaterialMultiplier)
+{
+	const float3 EdgeDirection = SharedNodePosition[GGroupThreadId.x] - SharedNodePosition[GGroupThreadId.x-1];
+	const float3 DeltaVelocity = (EdgeDirection - ( SharedPreviousPosition[GGroupThreadId.x] - SharedPreviousPosition[GGroupThreadId.x-1] ) ) / DeltaTime; 
+
+	const float4 ParentFrame = SlerpQuat(SharedNodeOrientation[GGroupThreadId.x-2],SharedNodeOrientation[GGroupThreadId.x-2],0.5);
+	const float3 TargetDirection = RotateVectorByQuat(RestDirection,ParentFrame);
+
+	// XPBD lagrange multiplier update : dL = -(C+compliance*L) / (dC * invM * dCt + alpha)
+	const float3 DeltaLambda = -((EdgeDirection - TargetDirection)/RestLength + 
+					OutMaterialMultiplier * InMaterialCompliance  + MaterialDamping * DeltaVelocity /RestLength) * InMaterialWeight;
+
+	// L += dL
+	OutMaterialMultiplier += DeltaLambda;
+
+	// XPBD position update : dX = dL * dCt * invM
+	const float3 PositionDelta  = DeltaLambda/RestLength;
+
+	// dX += dX
+	SharedNodePosition[GGroupThreadId.x] += PositionDelta * SharedInverseMass[GGroupThreadId.x];
+	if(!ProjectConstraint)
+	{
+		SharedNodePosition[GGroupThreadId.x-1] -= PositionDelta * SharedInverseMass[GGroupThreadId.x-1];
+	}
+}
+
+void SolveAngularSpringMaterial(in int StrandsSize, in float RestLength, in float3 RestDirection, in float DeltaTime, in float MaterialDamping,
+			in float MaterialCompliance, in float MaterialWeight, in float3 MaterialMultiplier, out float3 OutMaterialMultiplier)
+{
+	
+	OutMaterialMultiplier = MaterialMultiplier;
+	const int LocalIndex = (GGroupThreadId.x % StrandsSize);
+	if( LocalIndex > 1)
+	{
+		const int IsRed = (GGroupThreadId.x % 2) == 0;
+		// Process all the red rods
+		if (!IsRed)
+		{	
+			UpdateAngularSpringMultiplier(RestLength,RestDirection,DeltaTime,false,MaterialDamping,MaterialCompliance,MaterialWeight,OutMaterialMultiplier);
+		}
+		GroupMemoryBarrier(); 
+
+		if (IsRed)
+		{
+			UpdateAngularSpringMultiplier(RestLength,RestDirection,DeltaTime,false,MaterialDamping,MaterialCompliance,MaterialWeight,OutMaterialMultiplier);
+		}
+		GroupMemoryBarrier(); 
+	}
+}
+/*
+groupshared float4 SharedRestDatas[THREADGROUP_SIZE];
+void SolveAngularSpringMaterial(in int StrandsSize, in float RestLength, in float3 RestDirection, in float DeltaTime, in float MaterialDamping,
+			in float MaterialCompliance, in float MaterialWeight, in float3 MaterialMultiplier, out float3 OutMaterialMultiplier)
+{
+	SharedRestDatas[GGroupThreadId.x] = float4(RestDirection,RestLength);
+	GroupMemoryBarrier(); 
+
+	OutMaterialMultiplier = MaterialMultiplier;
+	const int LocalIndex = (GGroupThreadId.x % StrandsSize);
+	if( LocalIndex == 0)
+	{
+		for(int i = GGroupThreadId.x+2; i < GGroupThreadId.x+StrandsSize; ++i)
+		{
+			const float3 EdgeDirection = SharedNodePosition[i] - SharedNodePosition[i-1];
+
+			const float4 ParentFrame = SharedNodeOrientation[i-2];
+			const float3 TargetDirection = RotateVectorByQuat(SharedRestDatas[i].xyz,ParentFrame);
+
+			// XPBD lagrange multiplier update : dL = -(C+compliance*L) / (dC * invM * dCt + alpha)
+			const float3 DeltaLambda = -(EdgeDirection - TargetDirection) / (SharedInverseMass[i]+SharedInverseMass[i-1]);
+
+			// XPBD position update : dX = dL * dCt * invM
+			const float3 PositionDelta  = DeltaLambda;
+
+			SharedDeltaPosition[i] = PositionDelta;
+			GroupMemoryBarrier();
+
+			// dX += dX
+			SharedNodePosition[i] += PositionDelta * SharedInverseMass[i];
+			GroupMemoryBarrier();
+			SharedNodePosition[i-1] -= PositionDelta * SharedInverseMass[i-1];
+			GroupMemoryBarrier();
+
+			if( i > GGroupThreadId.x+2 )
+			{
+				const float3 TangentInit = normalize(SharedNodePosition[i-2] - SharedNodePosition[i-3]);
+				const float3 TangentPrev = normalize(SharedNodePosition[i-1] - SharedNodePosition[i-2]);
+				const float3 TangentNext = normalize(SharedNodePosition[i] - SharedNodePosition[i-1]);
+
+				SharedNodeOrientation[i-2] = NormalizeQuat( MultiplyQuat( FindQuatBetweenNormals(TangentInit,TangentPrev), SharedNodeOrientation[i-3] ) );
+				GroupMemoryBarrier();
+				SharedNodeOrientation[i-1] = NormalizeQuat( MultiplyQuat( FindQuatBetweenNormals(TangentPrev,TangentNext), SharedNodeOrientation[i-2] ) );
+				GroupMemoryBarrier();
+			}
+			else
+			{
+				const float3 TangentPrev = normalize(SharedNodePosition[i-1] - SharedNodePosition[i-2]);
+				const float3 TangentNext = normalize(SharedNodePosition[i] - SharedNodePosition[i-1]);
+
+				SharedNodeOrientation[i-1] = NormalizeQuat( MultiplyQuat( FindQuatBetweenNormals(TangentPrev,TangentNext), SharedNodeOrientation[i-2] ) );
+				GroupMemoryBarrier();
+			}
+			GroupMemoryBarrier();
+		}
+		GroupMemoryBarrier();
+		for(int i = GGroupThreadId.x+2; i < GGroupThreadId.x+StrandsSize; ++i)
+		{
+			SharedNodePosition[i] += SharedDeltaPosition[i] * SharedInverseMass[i];
+			SharedNodePosition[i-1] -= SharedDeltaPosition[i] * SharedInverseMass[i-1];
+			GroupMemoryBarrier();
+		}
+		GroupMemoryBarrier();
+		for(int i = GGroupThreadId.x+2; i < GGroupThreadId.x+StrandsSize; ++i)
+		{
+			const float3 TangentPrev = normalize(SharedNodePosition[i-1] - SharedNodePosition[i-2]);
+			const float3 TangentNext = normalize(SharedNodePosition[i] - SharedNodePosition[i-1]);
+
+			SharedNodeOrientation[i-1] = NormalizeQuat( MultiplyQuat( FindQuatBetweenNormals(TangentPrev,TangentNext), SharedNodeOrientation[i-2] ) );
+			GroupMemoryBarrier();
+		}
+		GroupMemoryBarrier();
+	} 
+}*/
+
+void ProjectAngularSpringMaterial(in int StrandsSize, in float YoungModulus, in float RodThickness, in float RestLength, in float3 RestDirection, in float NodeMass, in float DeltaTime,
+			 in float3 NodePosition, out float3 OutNodePosition)
+{
+	const int LocalIndex = (GGroupThreadId.x % StrandsSize);
+
+	SharedNodePosition[GGroupThreadId.x] = NodePosition;
+	SharedPreviousPosition[GGroupThreadId.x] = NodePosition;
+	SharedInverseMass[GGroupThreadId.x] = (LocalIndex > 1 && NodeMass != 0.0) ? 1.0 / NodeMass : 0.0;
+	GroupMemoryBarrier();
+
+	if(DeltaTime != 0.0)
+	{
+
+		float MaterialCompliance = 0.0;
+		float MaterialWeight = 0.0;
+		float3 MaterialMultiplier = float3(0,0,0);
+
+		//ComputeMaterialFrame(StrandsSize);
+		SetupAngularSpringMaterial(StrandsSize, YoungModulus, RodThickness, RestLength, DeltaTime, true, 0.0, MaterialCompliance, MaterialWeight, MaterialMultiplier);
+
+		//if( LocalIndex > 1 )
+		{
+			for(int i = 2; i < StrandsSize; ++i)
+			{
+				ComputeMaterialFrame(StrandsSize);
+				if( LocalIndex == i )
+				{
+					UpdateAngularSpringMultiplier(RestLength,RestDirection,DeltaTime,true,0.0,MaterialCompliance,MaterialWeight,MaterialMultiplier);
+				}
+				GroupMemoryBarrier(); 
+			}
+		}
+	}
+
+	GroupMemoryBarrier();
+	OutNodePosition = SharedNodePosition[GGroupThreadId.x];
+}
+
+/* -----------------------------------------------------------------
+ * Compute Rest Direction and length
+ * -----------------------------------------------------------------
+ */
+
+void ComputeRestDirection(in int StrandsSize, in float3 NodePosition, in float4 NodeOrientation, out float3 OutRestDirection)
+{
+	SharedNodePosition[GGroupThreadId.x] = NodePosition;
+	SharedNodeOrientation[GGroupThreadId.x] = NodeOrientation;
+	GroupMemoryBarrier(); 
+
+	ComputeMaterialFrame(StrandsSize);
+
+	OutRestDirection = float3(0,0,0);
+	const int LocalIndex = (GGroupThreadId.x % StrandsSize);
+	if( LocalIndex > 1 )
+	{
+		const float3 EdgeDirection = SharedNodePosition[GGroupThreadId.x] - SharedNodePosition[GGroupThreadId.x-1];
+
+		const float4 ParentFrame = SlerpQuat(SharedNodeOrientation[GGroupThreadId.x-2],SharedNodeOrientation[GGroupThreadId.x-2],0.5);
+		OutRestDirection = UnrotateVectorByQuat(EdgeDirection,ParentFrame);
+	}
+	GroupMemoryBarrier(); 
+}
diff --git a/Engine/Plugins/Experimental/HairStrands/Shaders/Private/NiagaraCopyBoundingBox.usf b/Engine/Plugins/Experimental/HairStrands/Shaders/Private/NiagaraCopyBoundingBox.usf
new file mode 100644
--- /dev/null
+++ b/Engine/Plugins/Experimental/HairStrands/Shaders/Private/NiagaraCopyBoundingBox.usf
@@ -0,0 +1,21 @@
+// Copyright 1998-2019 Epic Games, Inc. All Rights Reserved.
+
+uint NumElements;
+Buffer<uint> BoundingBoxBuffer;
+RWBuffer<uint> OutNodeBoundBuffer;
+
+[numthreads(THREAD_COUNT, 1, 1)]
+void MainCS(uint2 DispatchThreadId : SV_DispatchThreadID)
+{
+	const uint ElementIndex = DispatchThreadId.x;
+	if (ElementIndex < NumElements)
+	{
+		OutNodeBoundBuffer[0] = BoundingBoxBuffer[0];
+		OutNodeBoundBuffer[1] = BoundingBoxBuffer[1];
+		OutNodeBoundBuffer[2] = BoundingBoxBuffer[2];
+
+		OutNodeBoundBuffer[3] = BoundingBoxBuffer[3];
+		OutNodeBoundBuffer[4] = BoundingBoxBuffer[4];
+		OutNodeBoundBuffer[5] = BoundingBoxBuffer[5];
+	}
+}
\ No newline at end of file
diff --git a/Engine/Plugins/Experimental/HairStrands/Shaders/Private/NiagaraDataInterfaceHairStrands.ush b/Engine/Plugins/Experimental/HairStrands/Shaders/Private/NiagaraDataInterfaceHairStrands.ush
--- a/Engine/Plugins/Experimental/HairStrands/Shaders/Private/NiagaraDataInterfaceHairStrands.ush
+++ b/Engine/Plugins/Experimental/HairStrands/Shaders/Private/NiagaraDataInterfaceHairStrands.ush
@@ -11,13 +11,8 @@ NiagaraDataInterfaceHairStrands.ush
  * -----------------------------------------------------------------
  */
 
-#define WRAP_SIZE 32
 static const float M_PI = 3.14159265;
-static const float SCALE = 1e+6;
-
-static const int MASS_OFFSET = 0;
-static const int TANGENT_OFFSET = 3;
-static const int VELOCITY_OFFSET = 6;
+static const float3 MAX_POS = float3(1e+8,1e+8,1e+8);
 
 struct FDIHairStrandsContext
 {
@@ -44,6 +39,8 @@ struct FDIHairStrandsContext
 	float3				RestPositionOffset;
 	float3				DeformedPositionOffset;
 	bool				ResetSimulation;
+	RWBuffer<uint>		BoundingBoxBuffer;
+	Buffer<uint>		NodeBoundBuffer;
 }; 
 
 #define DIHAIRSTRANDS_DECLARE_CONSTANTS(NAME)\
@@ -70,6 +67,8 @@ float3				DeformedRootOffset_##NAME;\
 float3				RestPositionOffset_##NAME;\
 float3				DeformedPositionOffset_##NAME;\
 bool				ResetSimulation_##NAME;\
+RWBuffer<uint>	    BoundingBoxBuffer_##NAME;\
+Buffer<uint>		NodeBoundBuffer_##NAME;\
 
 #define DIHAIRSTRANDS_MAKE_CONTEXT(NAME)\
 FDIHairStrandsContext DIContext;\
@@ -96,6 +95,33 @@ DIContext.DeformedRootOffset = DeformedRootOffset_##NAME;\
 DIContext.RestPositionOffset =	RestPositionOffset_##NAME;\
 DIContext.DeformedPositionOffset = DeformedPositionOffset_##NAME;\
 DIContext.ResetSimulation = ResetSimulation_##NAME;\
+DIContext.BoundingBoxBuffer = BoundingBoxBuffer_##NAME;\
+DIContext.NodeBoundBuffer = NodeBoundBuffer_##NAME;\
+
+/* -----------------------------------------------------------------
+ * Vector utilities
+ * -----------------------------------------------------------------
+ */
+
+ float3 Uint3ToFloat3( in uint3 Uint3Vector)
+ {
+	uint3 SignedVector = Uint3Vector;
+	SignedVector.x ^= (((SignedVector.x >> 31)-1) | 0x80000000);
+	SignedVector.y ^= (((SignedVector.y >> 31)-1) | 0x80000000);
+	SignedVector.z ^= (((SignedVector.z >> 31)-1) | 0x80000000);
+
+	return float3(asfloat(SignedVector.x),asfloat(SignedVector.y),asfloat(SignedVector.z));
+ }
+
+ uint3 Float3ToUint3( in float3 Float3Vector)
+ {
+	uint3 UnsignedVector = uint3(asuint(Float3Vector.x),asuint(Float3Vector.y),asuint(Float3Vector.z));
+	UnsignedVector.x ^= (1+~(UnsignedVector.x >> 31) | 0x80000000);
+	UnsignedVector.y ^= (1+~(UnsignedVector.y >> 31) | 0x80000000);
+	UnsignedVector.z ^= (1+~(UnsignedVector.z >> 31) | 0x80000000);
+
+	return UnsignedVector;
+ }
 
 /* -----------------------------------------------------------------
  * Box utilities
@@ -105,13 +131,19 @@ DIContext.ResetSimulation = ResetSimulation_##NAME;\
 // Return the bounding box center
 void DIHairStrands_GetBoxCenter(in FDIHairStrandsContext DIContext, out float3 OutBoxCenter)
 {
-	OutBoxCenter = DIContext.BoxCenter;
+	const float3 BoxMin = Uint3ToFloat3(uint3(DIContext.NodeBoundBuffer[0],DIContext.NodeBoundBuffer[1],DIContext.NodeBoundBuffer[2]));
+	const float3 BoxMax = Uint3ToFloat3(uint3(DIContext.NodeBoundBuffer[3],DIContext.NodeBoundBuffer[4],DIContext.NodeBoundBuffer[5]));
+
+	OutBoxCenter = 0.5 * (BoxMin+BoxMax);
 }
 
 // Return the bounding box extent
 void DIHairStrands_GetBoxExtent(in FDIHairStrandsContext DIContext, out float3 OutBoxExtent)
 {
-	OutBoxExtent = DIContext.BoxExtent;
+	const float3 BoxMin = Uint3ToFloat3(uint3(DIContext.NodeBoundBuffer[0],DIContext.NodeBoundBuffer[1],DIContext.NodeBoundBuffer[2]));
+	const float3 BoxMax = Uint3ToFloat3(uint3(DIContext.NodeBoundBuffer[3],DIContext.NodeBoundBuffer[4],DIContext.NodeBoundBuffer[5]));
+
+	OutBoxExtent =  (BoxMax-BoxMin);
 }
 
 /* -----------------------------------------------------------------
@@ -168,6 +200,99 @@ float DIHairStrands_ComputeStrandLength(in FDIHairStrandsContext DIContext, in i
 	return StrandLength;
 }
 
+/* -----------------------------------------------------------------
+ * Bounding box generation
+ * -----------------------------------------------------------------
+ */
+
+ #define WRAP_SIZE 32
+
+ /*
+ const uint3 LocalPosition = Float3ToUint3(NodePosition);
+
+if(isfinite(NodePosition))
+{
+	InterlockedMin(DIContext.BoundingBoxBuffer[0],LocalPosition.x);
+	InterlockedMin(DIContext.BoundingBoxBuffer[1],LocalPosition.y);
+	InterlockedMin(DIContext.BoundingBoxBuffer[2],LocalPosition.z);
+
+	InterlockedMax(DIContext.BoundingBoxBuffer[3],LocalPosition.x);
+	InterlockedMax(DIContext.BoundingBoxBuffer[4],LocalPosition.y);
+	InterlockedMax(DIContext.BoundingBoxBuffer[5],LocalPosition.z);
+}
+
+AllMemoryBarrier();
+
+OutBoundMin = Uint3ToFloat3(uint3(DIContext.NodeBoundBuffer[0],DIContext.NodeBoundBuffer[1],DIContext.NodeBoundBuffer[2]));
+OutBoundMax = Uint3ToFloat3(uint3(DIContext.NodeBoundBuffer[3],DIContext.NodeBoundBuffer[4],DIContext.NodeBoundBuffer[5])); 
+*/
+
+void DIHairStrands_BuildBoundingBox(in FDIHairStrandsContext DIContext, in float3 NodePosition, out bool BuildStatus, out float3 OutBoundMin, out float3 OutBoundMax)
+{
+	BuildStatus = true;
+
+	const uint3 LocalPosition = Float3ToUint3(NodePosition);
+
+	if(isfinite(NodePosition.x) && isfinite(NodePosition.y) && isfinite(NodePosition.z))
+	{
+		InterlockedMin(DIContext.BoundingBoxBuffer[0],LocalPosition.x);
+		InterlockedMin(DIContext.BoundingBoxBuffer[1],LocalPosition.y);
+		InterlockedMin(DIContext.BoundingBoxBuffer[2],LocalPosition.z);
+
+		InterlockedMax(DIContext.BoundingBoxBuffer[3],LocalPosition.x);
+		InterlockedMax(DIContext.BoundingBoxBuffer[4],LocalPosition.y);
+		InterlockedMax(DIContext.BoundingBoxBuffer[5],LocalPosition.z);
+	}
+
+	DeviceMemoryBarrier();
+
+	OutBoundMin = Uint3ToFloat3(uint3(DIContext.NodeBoundBuffer[0],DIContext.NodeBoundBuffer[1],DIContext.NodeBoundBuffer[2]));
+	OutBoundMax = Uint3ToFloat3(uint3(DIContext.NodeBoundBuffer[3],DIContext.NodeBoundBuffer[4],DIContext.NodeBoundBuffer[5])); 
+
+	return;
+
+	SharedNodePosition[GGroupThreadId.x] = NodePosition;
+	GroupMemoryBarrier();
+
+	const int LocalIndex = (GGroupThreadId.x % DIContext.StrandSize);
+	if(LocalIndex == 0)
+	{
+		const float FLOAT_MAX = 1e+10;
+
+		float3 BoundingBoxMin = float3(FLOAT_MAX,FLOAT_MAX,FLOAT_MAX);
+		float3 BoundingBoxMax = -BoundingBoxMin;
+
+		[unroll]
+		for( int i = GGroupThreadId.x; i < GGroupThreadId.x+DIContext.StrandSize; ++i)
+		{
+			BoundingBoxMin = min(BoundingBoxMin,SharedNodePosition[i]);
+			BoundingBoxMax = max(BoundingBoxMax,SharedNodePosition[i]);
+		}
+
+		if(isfinite(BoundingBoxMin.x) && isfinite(BoundingBoxMin.y) && isfinite(BoundingBoxMin.z))
+		{
+			const uint3 MinPosition = Float3ToUint3(BoundingBoxMin);
+
+			InterlockedMin(DIContext.BoundingBoxBuffer[0],MinPosition.x);
+			InterlockedMin(DIContext.BoundingBoxBuffer[1],MinPosition.y);
+			InterlockedMin(DIContext.BoundingBoxBuffer[2],MinPosition.z);
+		}
+		if(isfinite(BoundingBoxMax.x) && isfinite(BoundingBoxMax.y) && isfinite(BoundingBoxMax.z))
+		{
+			const uint3 MaxPosition = Float3ToUint3(BoundingBoxMax);
+
+			InterlockedMax(DIContext.BoundingBoxBuffer[3],MaxPosition.x);
+			InterlockedMax(DIContext.BoundingBoxBuffer[4],MaxPosition.y);
+			InterlockedMax(DIContext.BoundingBoxBuffer[5],MaxPosition.z);
+		}
+	}
+
+	AllMemoryBarrier();
+
+	OutBoundMin = Uint3ToFloat3(uint3(DIContext.NodeBoundBuffer[0],DIContext.NodeBoundBuffer[1],DIContext.NodeBoundBuffer[2]));
+	OutBoundMax = Uint3ToFloat3(uint3(DIContext.NodeBoundBuffer[3],DIContext.NodeBoundBuffer[4],DIContext.NodeBoundBuffer[5]));
+}
+
 /* -----------------------------------------------------------------
  * Node Mass/Inertia/Position/Orientation computation
  * -----------------------------------------------------------------
@@ -346,29 +471,6 @@ void DIHairStrands_ComputeNodeOrientation(in FDIHairStrandsContext DIContext, in
 	}
 	GroupMemoryBarrier();
 	OutNodeOrientation = 	SharedNodeOrientation[GGroupThreadId.x];
-
-	/*int StrandIndex = 0, LocalIndex = 0;
-	DIHairStrands_ComputeNodeIndices(DIContext,NodeIndex,StrandIndex,LocalIndex);
-	const int RootIndex = StrandIndex * DIContext.StrandSize;
-
-	float4 NodeQuaternion = QUATERNION_IDENTITY;
-	float3 TangentNext = float3(0,0,1);
-	float3 TangentPrev = TangentNext;
-
-	float3 PositionNext = float3(0.0,0.0,0.0);
-	DIHairStrands_ComputeNodePosition(DIContext,RootIndex,PositionNext);
-	float3 PositionPrev = PositionNext;
-
-	for (int EdgeIndex = RootIndex, EdgeEnd = min(RootIndex+DIContext.StrandSize-1,NodeIndex+1); EdgeIndex < EdgeEnd; ++EdgeIndex)
-	{
-		PositionPrev = PositionNext;
-		DIHairStrands_ComputeNodePosition(DIContext,EdgeIndex+1,PositionNext);
-
-		TangentPrev = TangentNext;
-		TangentNext = normalize(PositionNext - PositionPrev);
-		NodeQuaternion = NormalizeQuat( MultiplyQuat( FindQuatBetweenNormals(TangentPrev,TangentNext), NodeQuaternion ) );
-	}
-	OutNodeOrientation = NodeQuaternion;*/
 }
 
 /* -----------------------------------------------------------------
@@ -496,23 +598,34 @@ void DIHairStrands_GetPointPosition(in FDIHairStrandsContext DIContext, in int P
  */
 
 // Add external force to the linear velocity and advect node position
-void DIHairStrands_AdvectNodePosition(in FDIHairStrandsContext DIContext, in float NodeMass, in float InverssMass, in float3 ExternalForce, in float DeltaTime, inout float3 OutLinearVelocity, inout float3 OutNodePosition)
+void DIHairStrands_AdvectNodePosition(in FDIHairStrandsContext DIContext, in float NodeMass, in bool IsPositionMobile, in float3 ExternalForce, in float3 ForceGradient, in float DeltaTime, inout float3 OutLinearVelocity, inout float3 OutNodePosition)
 {	
-	OutLinearVelocity += DeltaTime * InverssMass * ExternalForce;
+	if(IsPositionMobile && NodeMass != 0.0)
+	{
+		const float3 ImplicitGradient =  float3(NodeMass,NodeMass,NodeMass) / DeltaTime - ForceGradient;
+		const float3 InverseGradient = float3(1.0/ImplicitGradient.x, 1.0/ImplicitGradient.y, 1.0/ImplicitGradient.z);
 
-	OutNodePosition += OutLinearVelocity * DeltaTime;
+		OutLinearVelocity += InverseGradient * ExternalForce;
+		OutNodePosition += OutLinearVelocity * DeltaTime;
+	}
 }
 
 // Add external torque to the angular velocity and advect node orientation
-void DIHairStrands_AdvectNodeOrientation(in FDIHairStrandsContext DIContext, in float3 NodeInertia, in float3 InverseInertia, in float3 ExternalTorque, in float DeltaTime, inout float3 OutAngularVelocity, inout float4 OutNodeOrientation)
+void DIHairStrands_AdvectNodeOrientation(in FDIHairStrandsContext DIContext, in float3 NodeInertia, in bool IsOrientationMobile, in float3 ExternalTorque, in float3 TorqueGradient, in float DeltaTime, inout float3 OutAngularVelocity, inout float4 OutNodeOrientation)
 {	
-	OutAngularVelocity += DeltaTime * InverseInertia * (ExternalTorque - cross(OutAngularVelocity, NodeInertia * OutAngularVelocity));
+	if(IsOrientationMobile && NodeInertia.x != 0.0 && NodeInertia.y != 0.0 && NodeInertia.z != 0.0)
+	{
+		const float3 ImplicitGradient = NodeInertia / DeltaTime - TorqueGradient;
+		const float3 InverseGradient = float3(1.0/ImplicitGradient.x, 1.0/ImplicitGradient.y, 1.0/ImplicitGradient.z);
+
+		OutAngularVelocity += InverseGradient * (ExternalTorque - cross(OutAngularVelocity, NodeInertia * OutAngularVelocity));
 
-	OutNodeOrientation = OutNodeOrientation + 0.5 * DeltaTime * float4(
-        OutAngularVelocity.xyz * OutNodeOrientation.w + cross(OutAngularVelocity.xyz, OutNodeOrientation.xyz),
-        - dot(OutAngularVelocity.xyz, OutNodeOrientation.xyz));
+		OutNodeOrientation = OutNodeOrientation + 0.5 * DeltaTime * float4(
+			OutAngularVelocity.xyz * OutNodeOrientation.w + cross(OutAngularVelocity.xyz, OutNodeOrientation.xyz),
+			- dot(OutAngularVelocity.xyz, OutNodeOrientation.xyz));
 
-	OutNodeOrientation = NormalizeQuat(OutNodeOrientation);
+		OutNodeOrientation = NormalizeQuat(OutNodeOrientation);
+	}
 }
 
 // Update the node linear velocity based on the node position difference
@@ -539,50 +652,6 @@ void DIHairStrands_UpdateAngularVelocity(in FDIHairStrandsContext DIContext, in
 	}
 }
 
-/* -----------------------------------------------------------------
- * Compute Air Drag Force
- * -----------------------------------------------------------------
- */
-
-// Air Density = 1 kg/m3 = 1e-3 g/cm3
-// Dynamic Viscosity = 1e-5 Pa.s = 1e-6 g/(cm.s)
-// Compute air drag force
-void DIHairStrands_ComputeAirDragForce(in FDIHairStrandsContext DIContext, in float AirDensity, in float AirViscosity, in float AirDrag, 
-	in float3 AirVelocity, in float NodeThickness, in float3 NodePosition, in float3 NodeVelocity, out float3 OutAirDrag)
-{	
-	const float3 DeltaVelocity = NodeVelocity - AirVelocity;
-	const float DeltaVNorm = length(DeltaVelocity);
-
-	SharedNodePosition[GGroupThreadId.x] = NodePosition;
-	GroupMemoryBarrier();
-	
-	float AlignedTangent = 0.0;
-	float TangentLength = 0.0;
-	if( DeltaVNorm > SMALL_NUMBER)
-	{
-		const int LocalIndex = GGroupThreadId.x % DIContext.StrandSize;
-		float3 NodeTangent = (LocalIndex == 0) ? SharedNodePosition[GGroupThreadId.x+1] - SharedNodePosition[GGroupThreadId.x] : 
-												 SharedNodePosition[GGroupThreadId.x] - SharedNodePosition[GGroupThreadId.x-1]; 
-		TangentLength = length(NodeTangent);
-		NodeTangent = normalize(NodeTangent);
-
-		AlignedTangent = abs(dot(NodeTangent,DeltaVelocity)/DeltaVNorm);
-	}
-	const float NodeDiameter = NodeThickness * (1.0 - AlignedTangent) + TangentLength * AlignedTangent;
-	const float NodeArea = NodeDiameter * NodeThickness;
-
-	// V = 4 * PI * R * R * R / 3 = PI * D * D * D / 6 = M / rho 
-	//const float NodeDiameter = pow(6.0 * (NodeMass / StrandsDensity ) / PI, 1.0 / 3.0);
-	//const float NodeArea = PI * NodeDiameter * NodeDiameter / 4.0;
-
-	const float LinearDrag = 3.0 * PI * NodeDiameter * AirViscosity * AirDrag;
-	const float QuadraticDrag = 0.5 * AirDensity * NodeArea * AirDrag;
-
-	const float DragCoeff = LinearDrag + QuadraticDrag * DeltaVNorm;
-
-	OutAirDrag = - DragCoeff * DeltaVelocity;
-}
-
 /* -----------------------------------------------------------------
  * Update node position orientation
  * -----------------------------------------------------------------
diff --git a/Engine/Plugins/Experimental/HairStrands/Shaders/Private/NiagaraDataInterfacePhysicsAsset.ush b/Engine/Plugins/Experimental/HairStrands/Shaders/Private/NiagaraDataInterfacePhysicsAsset.ush
--- a/Engine/Plugins/Experimental/HairStrands/Shaders/Private/NiagaraDataInterfacePhysicsAsset.ush
+++ b/Engine/Plugins/Experimental/HairStrands/Shaders/Private/NiagaraDataInterfacePhysicsAsset.ush
@@ -21,24 +21,36 @@ struct FDIPhysicsAssetContext
 	uint4 ElementOffsets;
 	Buffer<float4>		CurrentTransformBuffer;
 	Buffer<float4>		InverseTransformBuffer;
+	Buffer<float4>		RestTransformBuffer;
+	Buffer<float4>		RestInverseBuffer;
 	Buffer<float4>		PreviousTransformBuffer;
 	Buffer<float4>		ElementExtentBuffer;
+	float3				BoxOrigin;
+	float3				BoxExtent;
 }; 
 
 #define DIPHYSICSASSET_DECLARE_CONSTANTS(NAME)\
 uint4 ElementOffsets_##NAME;\
 Buffer<float4>		CurrentTransformBuffer_##NAME;\
 Buffer<float4>		InverseTransformBuffer_##NAME;\
+Buffer<float4>		RestTransformBuffer_##NAME;\
+Buffer<float4>		RestInverseBuffer_##NAME;\
 Buffer<float4>		PreviousTransformBuffer_##NAME;\
 Buffer<float4>		ElementExtentBuffer_##NAME;\
+float3 BoxOrigin_##NAME;\
+float3 BoxExtent_##NAME;\
 
 #define DIPHYSICSASSET_MAKE_CONTEXT(NAME)\
 FDIPhysicsAssetContext DIContext; \
 DIContext.ElementOffsets = ElementOffsets_##NAME;\
 DIContext.CurrentTransformBuffer = CurrentTransformBuffer_##NAME;\
 DIContext.InverseTransformBuffer = InverseTransformBuffer_##NAME;\
+DIContext.RestTransformBuffer = RestTransformBuffer_##NAME;\
+DIContext.RestInverseBuffer = RestInverseBuffer_##NAME;\
 DIContext.PreviousTransformBuffer = PreviousTransformBuffer_##NAME;\
 DIContext.ElementExtentBuffer = ElementExtentBuffer_##NAME;\
+DIContext.BoxOrigin = BoxOrigin_##NAME;\
+DIContext.BoxExtent = BoxExtent_##NAME;\
 
 /* -----------------------------------------------------------------
  * Element accessors
@@ -96,6 +108,18 @@ float3x4 DIPhysicsAsset_GetPreviousTransform(in FDIPhysicsAssetContext DIContext
     return float3x4(DIContext.PreviousTransformBuffer[BufferOffset], DIContext.PreviousTransformBuffer[BufferOffset+1], DIContext.PreviousTransformBuffer[BufferOffset+2]);
 }
 
+float3x4 DIPhysicsAsset_GetRestTransform(in FDIPhysicsAssetContext DIContext, in int ElementIndex)
+{
+	const int BufferOffset = 3 * ElementIndex;
+    return float3x4(DIContext.RestTransformBuffer[BufferOffset], DIContext.RestTransformBuffer[BufferOffset+1], DIContext.RestTransformBuffer[BufferOffset+2]);
+}
+
+float3x4 DIPhysicsAsset_GetRestInverse(in FDIPhysicsAssetContext DIContext, in int ElementIndex)
+{
+	const int BufferOffset = 3 * ElementIndex;
+    return float3x4(DIContext.RestInverseBuffer[BufferOffset], DIContext.RestInverseBuffer[BufferOffset+1], DIContext.RestInverseBuffer[BufferOffset+2]);
+}
+
 /* -----------------------------------------------------------------
  * Collision detection utils
  * -----------------------------------------------------------------
@@ -181,25 +205,21 @@ float DIPhysicsAsset_GetBoxProjection(in float3 LocalPosition, in float3 BoxExte
 	return BoxDistance;
 }
 
+
 /* -----------------------------------------------------------------
- * Get the closest point to the physics asset 
+ * Get the closest element to the node position
  * -----------------------------------------------------------------
  */
 
 // Given a world space position (NodePosition) compute the sphere closest point (position,normal,velocity)
-void DIPhysicsAsset_GetClosestPoint(in FDIPhysicsAssetContext DIContext, in float3 NodePosition, in float DeltaTime, out float3 OutClosestPosition, 
-										  out float3 OutClosestNormal, out float3 OutClosestVelocity, out float OutClosestDistance)
+int DIPhysicsAsset_GetClosestElement(in FDIPhysicsAssetContext DIContext, in float3 NodePosition, out float3 OutClosestPosition, 
+										  out float3 OutClosestNormal, out float OutMinDistance )
 {	
 	float MinDistance = MAX_DISTANCE;
 	int ElementIndex = -1;
 	float3 CollisionPosition = float3(0,0,0);
 	float3 CollisionNormal = float3(0,0,0);
 
-	OutClosestNormal = float3(0,0,0);
-	OutClosestPosition = float3(0,0,0);
-	OutClosestVelocity = float3(0,0,0);
-	OutClosestDistance = 0.0;
-
 	const int SpheresBegin = DIContext.ElementOffsets[SPHERE_INDEX];
 	const int SpheresEnd = DIContext.ElementOffsets[SPHERE_INDEX+1];
 	for (int SphereIndex = SpheresBegin; SphereIndex < SpheresEnd; ++SphereIndex)
@@ -229,6 +249,31 @@ void DIPhysicsAsset_GetClosestPoint(in FDIPhysicsAssetContext DIContext, in floa
 		DIPhysicsAsset_GetBoxProjection(LocalPosition, DIContext.ElementExtentBuffer[BoxIndex].xyz, BoxIndex,
 				CollisionPosition, CollisionNormal, ElementIndex, MinDistance);
 	}
+	OutClosestPosition = CollisionPosition;
+	OutClosestNormal = CollisionNormal;
+	OutMinDistance = MinDistance;
+	return ElementIndex;
+}
+
+/* -----------------------------------------------------------------
+ * Get the closest point to the physics asset 
+ * -----------------------------------------------------------------
+ */
+
+// Given a world space position (NodePosition) compute the sphere closest point (position,normal,velocity)
+void DIPhysicsAsset_GetClosestPoint(in FDIPhysicsAssetContext DIContext, in float3 NodePosition, in float DeltaTime, out float3 OutClosestPosition, 
+										  out float3 OutClosestNormal, out float3 OutClosestVelocity, out float OutClosestDistance)
+{	
+	float3 CollisionPosition = float3(0,0,0);
+	float3 CollisionNormal = float3(0,0,0);
+	float MinDistance = 0.0;
+
+	const int ElementIndex = DIPhysicsAsset_GetClosestElement(DIContext,NodePosition,CollisionPosition,CollisionNormal,MinDistance);
+
+	OutClosestNormal = float3(0,0,0);
+	OutClosestPosition = float3(0,0,0);
+	OutClosestVelocity = float3(0,0,0);
+	OutClosestDistance = 0.0;
 
 	if (ElementIndex != -1)
 	{
@@ -240,3 +285,63 @@ void DIPhysicsAsset_GetClosestPoint(in FDIPhysicsAssetContext DIContext, in floa
 		OutClosestDistance = MinDistance;
 	}
 }
+
+/* -----------------------------------------------------------------
+ * Get the closest point to the physics asset 
+ * -----------------------------------------------------------------
+ */
+
+void DIPhysicsAsset_GetTexturePoint(in FDIPhysicsAssetContext DIContext, in float3 NodePosition, out int OutElementIndex, out float3 OutTexturePosition)
+{
+	float3 CollisionPosition = float3(0,0,0);
+	float3 CollisionNormal = float3(0,0,0);
+	float MinDistance = 0.0;
+
+	OutTexturePosition = float3(0,0,0);
+
+	OutElementIndex = DIPhysicsAsset_GetClosestElement(DIContext,NodePosition,CollisionPosition,CollisionNormal,MinDistance);
+
+	const float AxisExtent = max( max(DIContext.BoxExtent.x, DIContext.BoxExtent.y), DIContext.BoxExtent.z );
+	const float3 BoxExtent = float3(AxisExtent,AxisExtent,AxisExtent);
+
+	const float3 BoxMinimum = DIContext.BoxOrigin - BoxExtent;
+	const float3 BoxSize = 2.0 * BoxExtent;
+
+	if (OutElementIndex != -1)
+	{
+		const float3 LocalPosition = mul(DIPhysicsAsset_GetInverseTransform(DIContext,OutElementIndex), float4(NodePosition,1.0)).xyz;
+		const float3 RestPosition = mul(DIPhysicsAsset_GetRestTransform(DIContext,OutElementIndex), float4(LocalPosition,1.0)).xyz;
+
+		OutTexturePosition = clamp( (RestPosition - BoxMinimum)/BoxSize, float3(0,0,0), float3(1,1,1) );
+	}
+}
+
+void DIPhysicsAsset_GetProjectionPoint(in FDIPhysicsAssetContext DIContext, in float3 NodePosition, in float DeltaTime, in int ElementIndex,
+			in float TextureValue, in float3 TextureGradient, out float3 OutClosestPosition, out float3 OutClosestNormal, out float3 OutClosestVelocity, out float OutClosestDistance)
+{
+	OutClosestNormal = float3(0,0,0);
+	OutClosestPosition = float3(0,0,0);
+	OutClosestVelocity = float3(0,0,0);
+	OutClosestDistance = 0.0;
+
+	if (ElementIndex != -1)
+	{
+		const float AxisExtent = max( max(DIContext.BoxExtent.x, DIContext.BoxExtent.y), DIContext.BoxExtent.z );
+
+		const float3 LocalPosition = mul(DIPhysicsAsset_GetInverseTransform(DIContext,ElementIndex), float4(NodePosition,1.0)).xyz;
+		const float3 RestPosition = mul(DIPhysicsAsset_GetRestTransform(DIContext,ElementIndex), float4(LocalPosition,1.0)).xyz;
+
+		const float3 RestProjection = RestPosition - TextureValue * 2.0 * AxisExtent * TextureGradient;
+
+		const float3 CollisionPosition = mul(DIPhysicsAsset_GetRestInverse(DIContext,ElementIndex), float4(RestProjection,1.0)).xyz;
+		const float3 CollisionNormal = mul(DIPhysicsAsset_GetRestInverse(DIContext,ElementIndex), float4(TextureGradient,0.0)).xyz;
+
+		const float3 PreviousPosition = mul(DIPhysicsAsset_GetPreviousTransform(DIContext,ElementIndex), float4(CollisionPosition,1.0)).xyz;
+
+		OutClosestPosition = mul(DIPhysicsAsset_GetCurrentTransform(DIContext,ElementIndex), float4(CollisionPosition,1.0)).xyz;
+		OutClosestNormal = mul(DIPhysicsAsset_GetCurrentTransform(DIContext,ElementIndex), float4(CollisionNormal,0.0)).xyz;
+		OutClosestVelocity = ( OutClosestPosition - PreviousPosition ) / DeltaTime;
+		OutClosestDistance = TextureValue * 2.0 * AxisExtent;
+	}
+}
+
diff --git a/Engine/Plugins/Experimental/HairStrands/Shaders/Private/NiagaraDataInterfacePressureGrid.ush b/Engine/Plugins/Experimental/HairStrands/Shaders/Private/NiagaraDataInterfacePressureGrid.ush
--- a/Engine/Plugins/Experimental/HairStrands/Shaders/Private/NiagaraDataInterfacePressureGrid.ush
+++ b/Engine/Plugins/Experimental/HairStrands/Shaders/Private/NiagaraDataInterfacePressureGrid.ush
@@ -53,6 +53,23 @@ DIContext.GridDestinationBuffer = GridDestinationBuffer_##NAME;\
  * -----------------------------------------------------------------
  */
 
+// Compute the grid origin and cell length
+void DIPressureGrid_BuildGridTopology(in FDIPressureGridContext DIContext, in float3 GridCenter, in float3 GridExtent, out float3 OutGridOrigin, out float OutGridLength)
+{
+	const float3 GridLengths = float3(2.0*GridExtent.x/(DIContext.GridSize.x-1),
+									  2.0*GridExtent.y/(DIContext.GridSize.y-1),
+									  2.0*GridExtent.z/(DIContext.GridSize.z-1));
+
+	const float MaxLength = max(GridLengths.x,max(GridLengths.y,GridLengths.z));
+
+	const float3 RegularExtent = float3((DIContext.GridSize.x-1)*MaxLength,
+										(DIContext.GridSize.y-1)*MaxLength,
+										(DIContext.GridSize.z-1)*MaxLength);
+
+	OutGridOrigin = GridCenter - 0.5 * RegularExtent;
+	OutGridLength = MaxLength;
+}
+
 // Return true if the grid index is in the valid space
 bool DIPressureGrid_IsIndexValid(in int3 GridIndex, in int4 ValidSpace)
 {
@@ -86,15 +103,15 @@ bool DIPressureGrid_IsFaceValid(in FDIPressureGridContext DIContext, in int3 Gri
 }
 
 // Transform the node world position in grid space
-float3 DIPressureGrid_GetGridPosition(in FDIPressureGridContext DIContext, in float3 NodePosition)
+float3 DIPressureGrid_GetGridPosition(in FDIPressureGridContext DIContext, in float3 NodePosition, in float3 GridOrigin)
 {
-	return ( mul(float4( NodePosition, 1.0), DIContext.WorldInverse).xyz - DIContext.GridOrigin.xyz );
+	return ( mul(float4( NodePosition, 1.0), DIContext.WorldInverse).xyz - GridOrigin );
 }
 
 // Compute the grid index and fraction given a grid position
-void DIPressureGrid_GetGridIndex(in FDIPressureGridContext DIContext, in float3 GridPosition, in float3 GridOffset, out int3 OutGridIndex, out float3 OutGridFraction)
+void DIPressureGrid_GetGridIndex(in FDIPressureGridContext DIContext, in float3 GridPosition, in float3 GridOffset, in float GridLength, out int3 OutGridIndex, out float3 OutGridFraction)
 {
-	const float3 GridCell = GridPosition / DIContext.GridOrigin.w - GridOffset;
+	const float3 GridCell = GridPosition / GridLength - GridOffset;
 
 	OutGridIndex = floor(GridCell);
 	OutGridFraction = GridCell - OutGridIndex;
@@ -127,10 +144,10 @@ int DIPressureGrid_ComputeGridCell(in FDIPressureGridContext DIContext, in int3
 }
 
 // Compute the world grid position
-void DIPressureGrid_GetCellPosition(in FDIPressureGridContext DIContext, in int GridCell, out float3 OutCellPosition)
+void DIPressureGrid_GetCellPosition(in FDIPressureGridContext DIContext, in int GridCell, in float3 GridOrigin, in float GridLength, out float3 OutCellPosition)
 {
 	const int3 GridIndex = DIPressureGrid_ComputeGridIndex(DIContext,GridCell);
-	OutCellPosition = mul(float4((GridIndex+float3(0.5,0.5,0.5))*DIContext.GridOrigin.w+DIContext.GridOrigin.xyz,1.0),DIContext.WorldTransform).xyz;
+	OutCellPosition = mul(float4((GridIndex+float3(0.5,0.5,0.5))*GridLength+GridOrigin,1.0),DIContext.WorldTransform).xyz;
 }
 
 /* -----------------------------------------------------------------
@@ -195,7 +212,7 @@ void DIPressureGrid_TransferCellVelocity(in FDIPressureGridContext DIContext, in
 
  void DIPressureGrid_InterlockedAddFloat(in FDIPressureGridContext DIContext, in int3 GridIndex, in int DataOffset, in float FloatValue)
 {
-	int IntValue = asint(FloatValue);
+	/*int IntValue = asint(FloatValue);
 	int CompareValue = 0;
 	int StoredValue = 0;
 
@@ -209,7 +226,7 @@ void DIPressureGrid_TransferCellVelocity(in FDIPressureGridContext DIContext, in
 		}
 		CompareValue = StoredValue;
 		IntValue = asint(FloatValue+asfloat(StoredValue));
-	}
+	}*/
 }
 
 void DIPressureGrid_InterlockedAddInt(in FDIPressureGridContext DIContext, in int3 GridIndex, in int DataOffset, in int IntValue)
@@ -279,10 +296,10 @@ float3 DIPressureGrid_GetGridVector( in FDIPressureGridContext DIContext, in int
 
 // Transfer the velocity component from the node to the grid
 void DIPressureGrid_BuildVelocityComponent(in FDIPressureGridContext DIContext, in float3 GridPosition, in float SampleMass, in float SampleMomentum,
-			in float3 MomentumGradient, in float3 ComponentOffset, in int ComponentIndex )
+			in float3 MomentumGradient, in float3 ComponentOffset, in int ComponentIndex, in float GridLength )
 {
 	int3 GridIndex = int3(0,0,0); float3 GridFraction = float3(0,0,0);
-	DIPressureGrid_GetGridIndex(DIContext,GridPosition,ComponentOffset,GridIndex,GridFraction);
+	DIPressureGrid_GetGridIndex(DIContext,GridPosition,ComponentOffset,GridLength,GridIndex,GridFraction);
 
 	if(DIPressureGrid_IsFaceValid(DIContext,GridIndex,ComponentIndex) && 
 	   DIPressureGrid_IsFaceValid(DIContext,GridIndex+int3(1,1,1),ComponentIndex))
@@ -302,7 +319,7 @@ void DIPressureGrid_BuildVelocityComponent(in FDIPressureGridContext DIContext,
 					const float GridWeightXYZ = GridWeightX*GridWeightY*GridWeightZ;
 
 					const int3 GridLocation = int3(GridIndex.x+i,GridIndex.y+j,GridIndex.z+k);
-					const float3 DeltaPosition = ((GridLocation+ComponentOffset)*DIContext.GridOrigin.w) - GridPosition;
+					const float3 DeltaPosition = ((GridLocation+ComponentOffset)*GridLength) - GridPosition;
 					const float GridMomentum = SampleMomentum + dot(MomentumGradient, DeltaPosition);
 
 					DIPressureGrid_InterlockedAddScalar(DIContext,GridLocation,FLUID_MASS_OFFSET+ComponentIndex,SampleMass*GridWeightXYZ);
@@ -315,7 +332,7 @@ void DIPressureGrid_BuildVelocityComponent(in FDIPressureGridContext DIContext,
 
 // Transfer the velocity from the node to the grid
 void DIPressureGrid_BuildVelocityField(in FDIPressureGridContext DIContext, in int StrandsSize, in float3 NodePosition, in float NodeMass, in float3 NodeVelocity, 
-			in float3 VelocityGradientX, in float3 VelocityGradientY, in float3 VelocityGradientZ, out bool OutBuildStatus)
+			in float3 VelocityGradientX, in float3 VelocityGradientY, in float3 VelocityGradientZ, in float3 GridOrigin, in float GridLength, out bool OutBuildStatus)
 {
 	SharedNodePosition[GGroupThreadId.x] = NodePosition;
 	SharedPreviousPosition[GGroupThreadId.x] = NodeVelocity;
@@ -333,7 +350,7 @@ void DIPressureGrid_BuildVelocityField(in FDIPressureGridContext DIContext, in i
 		const float SegmentMass =  (SharedInverseMass[GGroupThreadId.x-1] * SegmentWeight.x + SharedInverseMass[GGroupThreadId.x] * SegmentWeight.y);
 		const float SegmentLength = length(DeltaPosition);
 
-		const int NumSamples = ceil(SegmentLength / DIContext.GridOrigin.w);
+		const int NumSamples = ceil(SegmentLength / GridLength);
 
 		const float SampleMass = SegmentMass / NumSamples;
 
@@ -347,11 +364,11 @@ void DIPressureGrid_BuildVelocityField(in FDIPressureGridContext DIContext, in i
 			const float3 SamplePosition = SharedNodePosition[GGroupThreadId.x] + SampleCoord * DeltaPosition;
 			const float3 SampleMomentum = (SharedPreviousPosition[GGroupThreadId.x] + SampleCoord * DeltaVelocity) * SampleMass;
 
-			const float3 GridPosition = DIPressureGrid_GetGridPosition(DIContext,SamplePosition);
+			const float3 GridPosition = DIPressureGrid_GetGridPosition(DIContext,SamplePosition,GridOrigin);
 
-			DIPressureGrid_BuildVelocityComponent(DIContext,GridPosition,SampleMass,SampleMomentum.x,MomentumGradientX,float3(0.0,0.5,0.5),0);
-			DIPressureGrid_BuildVelocityComponent(DIContext,GridPosition,SampleMass,SampleMomentum.y,MomentumGradientY,float3(0.5,0.0,0.5),1);
-			DIPressureGrid_BuildVelocityComponent(DIContext,GridPosition,SampleMass,SampleMomentum.z,MomentumGradientZ,float3(0.5,0.5,0.0),2);
+			DIPressureGrid_BuildVelocityComponent(DIContext,GridPosition,SampleMass,SampleMomentum.x,MomentumGradientX,float3(0.0,0.5,0.5),0,GridLength);
+			DIPressureGrid_BuildVelocityComponent(DIContext,GridPosition,SampleMass,SampleMomentum.y,MomentumGradientY,float3(0.5,0.0,0.5),1,GridLength);
+			DIPressureGrid_BuildVelocityComponent(DIContext,GridPosition,SampleMass,SampleMomentum.z,MomentumGradientZ,float3(0.5,0.5,0.0),2,GridLength);
 		}
 	}
 	OutBuildStatus = true;
@@ -363,11 +380,11 @@ void DIPressureGrid_BuildVelocityField(in FDIPressureGridContext DIContext, in i
  */
 
 // Sample the grid velocity field and gradient component at the node local position 
-void DIPressureGrid_SampleVelocityComponent(in FDIPressureGridContext DIContext, in float3 GridPosition, in float3 ComponentOffset, in int ComponentIndex, 
+void DIPressureGrid_SampleVelocityComponent(in FDIPressureGridContext DIContext, in float3 GridPosition, in float3 ComponentOffset, in int ComponentIndex, in float GridLength,
 						out float OutGridVelocity, out float OutGridMass, out float3 OutGridGradient)
 {
 	int3 GridIndex = int3(0,0,0); float3 GridFraction = float3(0,0,0);
-	DIPressureGrid_GetGridIndex(DIContext,GridPosition,ComponentOffset,GridIndex,GridFraction);
+	DIPressureGrid_GetGridIndex(DIContext,GridPosition,ComponentOffset,GridLength,GridIndex,GridFraction);
 
 	OutGridVelocity = 0.0;
 	OutGridGradient = float3(0,0,0);
@@ -395,7 +412,7 @@ void DIPressureGrid_SampleVelocityComponent(in FDIPressureGridContext DIContext,
 					const float GridWeightXYZ = GridWeightX*GridWeightY*GridWeightZ;
 					const float3 GridGradientXYZ = float3(GridGradientX*GridWeightY*GridWeightZ,
 														  GridWeightX*GridGradientY*GridWeightZ,
-														  GridWeightX*GridWeightY*GridGradientZ) / DIContext.GridOrigin.w;
+														  GridWeightX*GridWeightY*GridGradientZ) / GridLength;
 
 					const int3 GridLocation = int3(GridIndex.x+i,GridIndex.y+j,GridIndex.z+k);
 					const float GridVelocity = DIPressureGrid_GetGridScalar(DIContext,GridLocation,FLUID_VELOCITY_OFFSET+ComponentIndex);
@@ -411,7 +428,7 @@ void DIPressureGrid_SampleVelocityComponent(in FDIPressureGridContext DIContext,
 }
 
 // Sample the grid velocity field and gradient at the node world position 
-void DIPressureGrid_SampleVelocityField(in FDIPressureGridContext DIContext, in float3 NodePosition, out float3 OutGridVelocity, out float OutNodeDensity,
+void DIPressureGrid_SampleVelocityField(in FDIPressureGridContext DIContext, in float3 NodePosition, in float3 GridOrigin, in float GridLength, out float3 OutGridVelocity, out float OutNodeDensity,
 			out float3 OutGridGradientX, out float3 OutGridGradientY, out float3 OutGridGradientZ)
 {
 	OutGridVelocity = float3(0,0,0);
@@ -421,15 +438,15 @@ void DIPressureGrid_SampleVelocityField(in FDIPressureGridContext DIContext, in
 
 	float3 GridMass = float3(0,0,0);
 
-	const float3 GridPosition = DIPressureGrid_GetGridPosition(DIContext,NodePosition);
-	DIPressureGrid_SampleVelocityComponent(DIContext,GridPosition,float3(0.0,0.5,0.5),0,OutGridVelocity.x,GridMass.x,OutGridGradientX);
-	DIPressureGrid_SampleVelocityComponent(DIContext,GridPosition,float3(0.5,0.0,0.5),1,OutGridVelocity.y,GridMass.y,OutGridGradientY);
-	DIPressureGrid_SampleVelocityComponent(DIContext,GridPosition,float3(0.5,0.5,0.0),2,OutGridVelocity.z,GridMass.z,OutGridGradientZ);
+	const float3 GridPosition = DIPressureGrid_GetGridPosition(DIContext,NodePosition,GridOrigin);
+	DIPressureGrid_SampleVelocityComponent(DIContext,GridPosition,float3(0.0,0.5,0.5),0,GridLength,OutGridVelocity.x,GridMass.x,OutGridGradientX);
+	DIPressureGrid_SampleVelocityComponent(DIContext,GridPosition,float3(0.5,0.0,0.5),1,GridLength,OutGridVelocity.y,GridMass.y,OutGridGradientY);
+	DIPressureGrid_SampleVelocityComponent(DIContext,GridPosition,float3(0.5,0.5,0.0),2,GridLength,OutGridVelocity.z,GridMass.z,OutGridGradientZ);
 
-	OutNodeDensity = (GridMass.x + GridMass.y + GridMass.z)/(3.0*DIContext.GridOrigin.w*DIContext.GridOrigin.w*DIContext.GridOrigin.w);
+	OutNodeDensity = (GridMass.x + GridMass.y + GridMass.z)/(3.0*GridLength*GridLength*GridLength);
 
 	//int3 GridIndex = int3(0,0,0); float3 GridFraction = float3(0,0,0);
-	//DIPressureGrid_GetGridIndex(DIContext,GridPosition,0,GridIndex,GridFraction);
+	//DIPressureGrid_GetGridIndex(DIContext,GridPosition,0,GridLength,GridIndex,GridFraction);
 	//OutGridVelocity = GridIndex;
 }
 
@@ -542,9 +559,10 @@ float DIPressureGrid_ComputeGridPressure(in FDIPressureGridContext DIContext, in
 		}
 		if( GridFraction > 0.)
 		{
-			GridPressure = -GridDivergence / ( DIContext.GridOrigin.w * DIContext.GridOrigin.w * GridFraction);
+			//GridPressure = -GridDivergence / ( GridLength * GridLength * GridFraction);
+			GridPressure = -GridDivergence / (GridFraction);
 		}
-		OutCellDivergence = 3.0;//GridDivergence;
+		OutCellDivergence = GridDivergence;
 	}
 	return GridPressure;
 }
@@ -561,7 +579,8 @@ void DIPressureGrid_ProjectVelocityField(in FDIPressureGridContext DIContext, in
 		if( (PrevIndex.x+PrevIndex.y+PrevIndex.z+ShaderStageIndex) % 2 == 0)
 		{
 			OutProjectStatus = true;
-			const float GridPressure = DIPressureGrid_ComputeGridPressure(DIContext,PrevIndex,OutCellDivergence) * DIContext.GridOrigin.w * DIContext.GridOrigin.w;
+			//const float GridPressure = DIPressureGrid_ComputeGridPressure(DIContext,PrevIndex,OutCellDivergence) * GridLength * GridLength;
+			const float GridPressure = DIPressureGrid_ComputeGridPressure(DIContext,PrevIndex,OutCellDivergence);
 		
 			for( int ComponentIndex = 0; ComponentIndex < 3; ++ComponentIndex)
 			{
diff --git a/Engine/Plugins/Experimental/HairStrands/Shaders/Private/NiagaraResetBoundingBox.usf b/Engine/Plugins/Experimental/HairStrands/Shaders/Private/NiagaraResetBoundingBox.usf
new file mode 100644
--- /dev/null
+++ b/Engine/Plugins/Experimental/HairStrands/Shaders/Private/NiagaraResetBoundingBox.usf
@@ -0,0 +1,22 @@
+// Copyright 1998-2019 Epic Games, Inc. All Rights Reserved.
+
+uint NumElements;
+RWBuffer<uint> OutBoundingBoxBuffer;
+
+[numthreads(THREAD_COUNT, 1, 1)]
+void MainCS(uint2 DispatchThreadId : SV_DispatchThreadID)
+{
+	const uint UINT_MAX = 0xffffffff;
+
+	const uint ElementIndex = DispatchThreadId.x;
+	if (ElementIndex < NumElements)
+	{
+		OutBoundingBoxBuffer[0] = UINT_MAX;
+		OutBoundingBoxBuffer[1] = UINT_MAX;
+		OutBoundingBoxBuffer[2] = UINT_MAX;
+
+		OutBoundingBoxBuffer[3] = 0;
+		OutBoundingBoxBuffer[4] = 0;
+		OutBoundingBoxBuffer[5] = 0;
+	}
+}
\ No newline at end of file
diff --git a/Engine/Plugins/Experimental/HairStrands/Shaders/Private/NiagaraStrandsExternalForce.ush b/Engine/Plugins/Experimental/HairStrands/Shaders/Private/NiagaraStrandsExternalForce.ush
new file mode 100644
--- /dev/null
+++ b/Engine/Plugins/Experimental/HairStrands/Shaders/Private/NiagaraStrandsExternalForce.ush
@@ -0,0 +1,68 @@
+// Copyright 1998-2019 Epic Games, Inc. All Rights Reserved.
+
+/*=============================================================================
+NiagaraStrandsExternalForce.ush 
+=============================================================================*/
+
+#pragma once
+
+/* -----------------------------------------------------------------
+ * Compute Air Drag Force
+ * -----------------------------------------------------------------
+ */
+
+// Air Density = 1 kg/m3 = 1e-3 g/cm3
+// Dynamic Viscosity = 1e-5 Pa.s = 1e-6 g/(cm.s)
+// Compute air drag force
+void ComputeAirDragForce(in int StrandSize, in float AirDensity, in float AirViscosity, in float AirDrag, 
+	in float3 AirVelocity, in float NodeThickness, in float3 NodePosition, in float3 NodeVelocity, out float3 OutAirDrag, out float3 OutForceGradient)
+{	
+	const float3 DeltaVelocity = NodeVelocity - AirVelocity;
+	const float DeltaVNorm = length(DeltaVelocity);
+
+	SharedNodePosition[GGroupThreadId.x] = NodePosition;
+	GroupMemoryBarrier();
+	
+	float AlignedTangent = 0.0;
+	float TangentLength = 0.0;
+	if( DeltaVNorm > 1e-6)
+	{
+		const int LocalIndex = GGroupThreadId.x % StrandSize;
+		float3 NodeTangent = (LocalIndex == 0) ? SharedNodePosition[GGroupThreadId.x+1] - SharedNodePosition[GGroupThreadId.x] : 
+												 SharedNodePosition[GGroupThreadId.x] - SharedNodePosition[GGroupThreadId.x-1]; 
+		TangentLength = length(NodeTangent);
+		NodeTangent = normalize(NodeTangent);
+
+		AlignedTangent = abs(dot(NodeTangent,DeltaVelocity)/DeltaVNorm);
+	}
+	const float NodeDiameter = NodeThickness * (1.0 - AlignedTangent) + TangentLength * AlignedTangent;
+	const float NodeArea = NodeDiameter * NodeThickness;
+
+	// V = 4 * PI * R * R * R / 3 = PI * D * D * D / 6 = M / rho 
+	//const float NodeDiameter = pow(6.0 * (NodeMass / StrandsDensity ) / PI, 1.0 / 3.0);
+	//const float NodeArea = PI * NodeDiameter * NodeDiameter / 4.0;
+
+	const float LinearDrag = 3.0 * PI * NodeDiameter * AirViscosity * AirDrag;
+	const float QuadraticDrag = 0.5 * AirDensity * NodeArea * AirDrag;
+
+	const float DragCoeff = LinearDrag + QuadraticDrag * DeltaVNorm;
+
+	OutAirDrag = - DragCoeff * DeltaVelocity;
+
+	OutForceGradient = - float3(DragCoeff,DragCoeff,DragCoeff);
+	if( DeltaVNorm > 1e-6)
+	{
+		OutForceGradient -= QuadraticDrag * float3(DeltaVelocity.x*DeltaVelocity.x,
+												   DeltaVelocity.y*DeltaVelocity.y,
+												   DeltaVelocity.z*DeltaVelocity.z) / DeltaVNorm;
+	}
+}
+
+
+
+
+
+
+
+
+
diff --git a/Engine/Plugins/Experimental/HairStrands/Source/HairStrandsCore/Private/GroomComponent.cpp b/Engine/Plugins/Experimental/HairStrands/Source/HairStrandsCore/Private/GroomComponent.cpp
--- a/Engine/Plugins/Experimental/HairStrands/Source/HairStrandsCore/Private/GroomComponent.cpp
+++ b/Engine/Plugins/Experimental/HairStrands/Source/HairStrandsCore/Private/GroomComponent.cpp
@@ -1405,7 +1405,7 @@ void UGroomComponent::PostEditChangeProperty(FPropertyChangedEvent& PropertyChan
 
 	if (bRecreateResources)
 	{
-		if (GroomAsset != InitializedResources)
+		//if (GroomAsset != InitializedResources)
 		{
 			if (GroomAsset)
 			{
diff --git a/Engine/Plugins/Experimental/HairStrands/Source/HairStrandsNiagara/Private/NiagaraDataInterfaceHairStrands.cpp b/Engine/Plugins/Experimental/HairStrands/Source/HairStrandsNiagara/Private/NiagaraDataInterfaceHairStrands.cpp
--- a/Engine/Plugins/Experimental/HairStrands/Source/HairStrandsNiagara/Private/NiagaraDataInterfaceHairStrands.cpp
+++ b/Engine/Plugins/Experimental/HairStrands/Source/HairStrandsNiagara/Private/NiagaraDataInterfaceHairStrands.cpp
@@ -9,6 +9,10 @@
 
 #include "ShaderParameterUtils.h"
 #include "ClearQuad.h"
+#include "RenderGraphBuilder.h"
+#include "RenderGraphUtils.h"
+#include "ShaderParameterStruct.h"
+#include "GlobalShader.h"
 
 #include "GroomComponent.h"
 #include "GroomAsset.h"  
@@ -62,6 +66,7 @@ static const FName ResetPointPositionName(TEXT("ResetPointPosition"));
 
 static const FName GetBoxCenterName(TEXT("GetBoxCenter"));
 static const FName GetBoxExtentName(TEXT("GetBoxExtent"));
+static const FName BuildBoundingBoxName(TEXT("BuildBoundingBox"));
 
 //------------------------------------------------------------------------------------------------------------
 
@@ -99,6 +104,7 @@ static const FName UpdateNodeOrientationName(TEXT("UpdateNodeOrientation"));
 //------------------------------------------------------------------------------------------------------------
 
 static const FName ComputeAirDragForceName(TEXT("ComputeAirDragForce"));  
+
 //------------------------------------------------------------------------------------------------------------
 
 static const FName NeedSimulationResetName(TEXT("NeedSimulationReset"));
@@ -134,6 +140,9 @@ const FString UNiagaraDataInterfaceHairStrands::DeformedTrianglePositionCName(TE
 const FString UNiagaraDataInterfaceHairStrands::RestPositionOffsetName(TEXT("RestPositionOffset_"));
 const FString UNiagaraDataInterfaceHairStrands::DeformedPositionOffsetName(TEXT("DeformedPositionOffset_"));
 
+const FString UNiagaraDataInterfaceHairStrands::BoundingBoxBufferName(TEXT("BoundingBoxBuffer_"));
+const FString UNiagaraDataInterfaceHairStrands::NodeBoundBufferName(TEXT("NodeBoundBuffer_"));
+
 //------------------------------------------------------------------------------------------------------------
 
 struct FNDIHairStrandsParametersName
@@ -168,6 +177,9 @@ struct FNDIHairStrandsParametersName
 
 		RestPositionOffsetName = UNiagaraDataInterfaceHairStrands::RestPositionOffsetName + Suffix;
 		DeformedPositionOffsetName = UNiagaraDataInterfaceHairStrands::DeformedPositionOffsetName + Suffix;
+
+		BoundingBoxBufferName = UNiagaraDataInterfaceHairStrands::BoundingBoxBufferName + Suffix;
+		NodeBoundBufferName = UNiagaraDataInterfaceHairStrands::NodeBoundBufferName + Suffix;
 	}
 
 	FString NumStrandsName;
@@ -198,8 +210,118 @@ struct FNDIHairStrandsParametersName
 
 	FString RestPositionOffsetName;
 	FString DeformedPositionOffsetName;
+
+	FString BoundingBoxBufferName;
+	FString NodeBoundBufferName;
+};
+
+//------------------------------------------------------------------------------------------------------------
+
+#define NIAGARA_HAIR_STRANDS_THREAD_COUNT 64
+
+class FCopyBoundingBoxCS : public FGlobalShader
+{
+	DECLARE_GLOBAL_SHADER(FCopyBoundingBoxCS);
+	SHADER_USE_PARAMETER_STRUCT(FCopyBoundingBoxCS, FGlobalShader);
+
+	BEGIN_SHADER_PARAMETER_STRUCT(FParameters, )
+		SHADER_PARAMETER(uint32, NumElements)
+		SHADER_PARAMETER_SRV(Buffer, BoundingBoxBuffer)
+		SHADER_PARAMETER_UAV(RWBuffer, OutNodeBoundBuffer)
+		END_SHADER_PARAMETER_STRUCT()
+
+public:
+	static bool ShouldCompilePermutation(const FGlobalShaderPermutationParameters& Parameters)
+	{
+		return RHISupportsComputeShaders(Parameters.Platform);
+	}
+
+	static void ModifyCompilationEnvironment(const FGlobalShaderPermutationParameters& Parameters, FShaderCompilerEnvironment& OutEnvironment)
+	{
+		FGlobalShader::ModifyCompilationEnvironment(Parameters, OutEnvironment);
+		OutEnvironment.SetDefine(TEXT("THREAD_COUNT"), NIAGARA_HAIR_STRANDS_THREAD_COUNT);
+	}
+};
+
+class FResetBoundingBoxCS : public FGlobalShader
+{
+	DECLARE_GLOBAL_SHADER(FResetBoundingBoxCS);
+	SHADER_USE_PARAMETER_STRUCT(FResetBoundingBoxCS, FGlobalShader);
+
+	BEGIN_SHADER_PARAMETER_STRUCT(FParameters, )
+		SHADER_PARAMETER(uint32, NumElements)
+		SHADER_PARAMETER_UAV(RWBuffer, OutBoundingBoxBuffer)
+		END_SHADER_PARAMETER_STRUCT()
+
+public:
+	static bool ShouldCompilePermutation(const FGlobalShaderPermutationParameters& Parameters)
+	{
+		return RHISupportsComputeShaders(Parameters.Platform);
+	}
+
+	static void ModifyCompilationEnvironment(const FGlobalShaderPermutationParameters& Parameters, FShaderCompilerEnvironment& OutEnvironment)
+	{
+		FGlobalShader::ModifyCompilationEnvironment(Parameters, OutEnvironment);
+		OutEnvironment.SetDefine(TEXT("THREAD_COUNT"), NIAGARA_HAIR_STRANDS_THREAD_COUNT);
+	}
 };
 
+
+IMPLEMENT_GLOBAL_SHADER(FResetBoundingBoxCS, "/Plugin/Experimental/HairStrands/Private/NiagaraResetBoundingBox.usf", "MainCS", SF_Compute);
+IMPLEMENT_GLOBAL_SHADER(FCopyBoundingBoxCS, "/Plugin/Experimental/HairStrands/Private/NiagaraCopyBoundingBox.usf", "MainCS", SF_Compute);
+
+static void AddResetBoundingBoxPass(
+	FRDGBuilder& GraphBuilder,
+	FRHIUnorderedAccessView* OutBoundingBoxBuffer)
+{
+	const uint32 GroupSize = NIAGARA_HAIR_STRANDS_THREAD_COUNT;
+	const uint32 NumElements = 1;
+
+	FResetBoundingBoxCS::FParameters* Parameters = GraphBuilder.AllocParameters<FResetBoundingBoxCS::FParameters>();
+	Parameters->OutBoundingBoxBuffer = OutBoundingBoxBuffer;
+	Parameters->NumElements = NumElements;
+	
+	TShaderMap<FGlobalShaderType>* ShaderMap = GetGlobalShaderMap(ERHIFeatureLevel::SM5);
+
+	const uint32 DispatchCount = FMath::DivideAndRoundUp(NumElements, GroupSize);
+	check(DispatchCount <= 65535);
+
+	TShaderMapRef<FResetBoundingBoxCS> ComputeShader(ShaderMap);
+	FComputeShaderUtils::AddPass(
+		GraphBuilder,
+		RDG_EVENT_NAME("ResetBoundingBox"),
+		*ComputeShader,
+		Parameters,
+		FIntVector(DispatchCount, 1, 1));
+}
+
+static void AddCopyBoundingBoxPass(
+	FRDGBuilder& GraphBuilder,
+	FRHIShaderResourceView* BoundingBoxBuffer,
+	FRHIUnorderedAccessView* OutNodeBoundBuffer)
+{
+	const uint32 GroupSize = NIAGARA_HAIR_STRANDS_THREAD_COUNT;
+	const uint32 NumElements = 1;
+
+	FCopyBoundingBoxCS::FParameters* Parameters = GraphBuilder.AllocParameters<FCopyBoundingBoxCS::FParameters>();
+	Parameters->BoundingBoxBuffer = BoundingBoxBuffer;
+	Parameters->OutNodeBoundBuffer = OutNodeBoundBuffer;
+	Parameters->NumElements = NumElements;
+
+	TShaderMap<FGlobalShaderType>* ShaderMap = GetGlobalShaderMap(ERHIFeatureLevel::SM5);
+
+	const uint32 DispatchCount = FMath::DivideAndRoundUp(NumElements, GroupSize);
+	check(DispatchCount <= 65535);
+
+	TShaderMapRef<FCopyBoundingBoxCS> ComputeShader(ShaderMap);
+	FComputeShaderUtils::AddPass(
+		GraphBuilder,
+		RDG_EVENT_NAME("CopyBoundingBox"),
+		*ComputeShader,
+		Parameters,
+		FIntVector(DispatchCount, 1, 1));
+}
+
 //------------------------------------------------------------------------------------------------------------
 
 void FNDIHairStrandsBuffer::SetHairAsset(const FHairStrandsDatas*  HairStrandsDatas, const FHairStrandsRestResource* HairStrandsRestResource, 
@@ -211,6 +333,45 @@ void FNDIHairStrandsBuffer::SetHairAsset(const FHairStrandsDatas*  HairStrandsDa
 	SourceRootResources = HairStrandsRootResource;
 }
 
+void FNDIHairStrandsBuffer::ClearBuffer(FRHICommandList& RHICmdList)
+{
+	FRHIUnorderedAccessView* BoundingBoxBufferUAV = BoundingBoxBuffer.UAV;
+	FRHIShaderResourceView* BoundingBoxBufferSRV = BoundingBoxBuffer.SRV;
+	FRHIUnorderedAccessView* NodeBoundBufferUAV = NodeBoundBuffer.UAV;
+
+	ENQUEUE_RENDER_COMMAND(FNDIHairStrandsCopyBoundingBox)(
+		[BoundingBoxBufferUAV, NodeBoundBufferUAV, BoundingBoxBufferSRV]
+	(FRHICommandListImmediate& RHICmdListImm)
+	{
+		FRDGBuilder GraphBuilder(RHICmdListImm);
+
+		RHICmdListImm.TransitionResource(EResourceTransitionAccess::EReadable, EResourceTransitionPipeline::EComputeToCompute, BoundingBoxBufferUAV);
+		RHICmdListImm.TransitionResource(EResourceTransitionAccess::EWritable, EResourceTransitionPipeline::EComputeToCompute, NodeBoundBufferUAV);
+
+		AddCopyBoundingBoxPass(
+			GraphBuilder,
+			BoundingBoxBufferSRV, NodeBoundBufferUAV);
+
+		GraphBuilder.Execute();
+	});
+
+	ENQUEUE_RENDER_COMMAND(FNDIHairStrandsResetBoundingBox)(
+		[BoundingBoxBufferUAV, NodeBoundBufferUAV, BoundingBoxBufferSRV]
+	(FRHICommandListImmediate& RHICmdListImm)
+	{
+		FRDGBuilder GraphBuilder(RHICmdListImm);
+
+		RHICmdListImm.TransitionResource(EResourceTransitionAccess::EWritable, EResourceTransitionPipeline::EComputeToCompute, BoundingBoxBufferUAV);
+		RHICmdListImm.TransitionResource(EResourceTransitionAccess::EReadable, EResourceTransitionPipeline::EComputeToCompute, NodeBoundBufferUAV);
+
+		AddResetBoundingBoxPass(
+			GraphBuilder,
+			BoundingBoxBufferUAV);
+
+		GraphBuilder.Execute();
+	});
+}
+
 void FNDIHairStrandsBuffer::InitRHI()
 {
 	if (SourceDatas != nullptr && SourceRestResources != nullptr)
@@ -225,6 +386,31 @@ void FNDIHairStrandsBuffer::InitRHI()
 			FMemory::Memcpy(OffsetBufferData, SourceDatas->StrandsCurves.CurvesOffset.GetData(), OffsetBytes);
 			RHIUnlockVertexBuffer(CurvesOffsetsBuffer.Buffer);
 		}
+		{
+			static const TArray<uint32> ZeroData = {UINT_MAX,UINT_MAX,UINT_MAX,0,0,0};
+
+			const uint32 BoundCount = 6;
+			const uint32 BoundBytes = sizeof(uint32)*BoundCount;
+
+			BoundingBoxBuffer.Initialize(sizeof(uint32), BoundCount, EPixelFormat::PF_R32_UINT, BUF_Static);
+			void* BoundBufferData = RHILockVertexBuffer(BoundingBoxBuffer.Buffer, 0, BoundBytes, RLM_WriteOnly);
+			
+			FMemory::Memcpy(BoundBufferData, ZeroData.GetData(), BoundBytes);
+			RHIUnlockVertexBuffer(BoundingBoxBuffer.Buffer);
+		}
+
+		{
+			static const TArray<uint32> ZeroData = { UINT_MAX,UINT_MAX,UINT_MAX,0,0,0};
+
+			const uint32 BoundCount = 6;
+			const uint32 BoundBytes = sizeof(uint32)*BoundCount;
+
+			NodeBoundBuffer.Initialize(sizeof(uint32), BoundCount, EPixelFormat::PF_R32_UINT, BUF_Static);
+			void* BoundBufferData = RHILockVertexBuffer(NodeBoundBuffer.Buffer, 0, BoundBytes, RLM_WriteOnly);
+
+			FMemory::Memcpy(BoundBufferData, ZeroData.GetData(), BoundBytes);
+			RHIUnlockVertexBuffer(NodeBoundBuffer.Buffer);
+		}
 		if (SourceDeformedResources == nullptr)
 		{
 			const uint32 PositionsCount = SourceDatas->GetNumPoints();
@@ -260,6 +446,8 @@ void FNDIHairStrandsBuffer::ReleaseRHI()
 	DeformedTrianglePositionCBuffer.Release();
 
 	RootBarycentricCoordinatesBuffer.Release();
+	BoundingBoxBuffer.Release();
+	NodeBoundBuffer.Release();
 }
 
 //------------------------------------------------------------------------------------------------------------
@@ -300,6 +488,9 @@ struct FNDIHairStrandsParametersCS : public FNiagaraDataInterfaceParametersCS
 		DeformedTrianglePositionBBuffer.Bind(ParameterMap, *ParamNames.DeformedTrianglePositionBName);
 		DeformedTrianglePositionCBuffer.Bind(ParameterMap, *ParamNames.DeformedTrianglePositionCName);
 
+		BoundingBoxBuffer.Bind(ParameterMap, *ParamNames.BoundingBoxBufferName);
+		NodeBoundBuffer.Bind(ParameterMap, *ParamNames.NodeBoundBufferName);
+
 		if (!DeformedPositionBuffer.IsBound())
 		{
 			UE_LOG(LogHairStrands, Warning, TEXT("Binding failed for FNDIHairStrandsParametersCS %s. Was it optimized out?"), *ParamNames.DeformedPositionBufferName)
@@ -343,6 +534,14 @@ struct FNDIHairStrandsParametersCS : public FNiagaraDataInterfaceParametersCS
 		{
 			UE_LOG(LogHairStrands, Warning, TEXT("Binding failed for FNDIHairStrandsParametersCS %s. Was it optimized out?"), *ParamNames.DeformedTrianglePositionCName)
 		}
+		if (!BoundingBoxBuffer.IsBound())
+		{
+			UE_LOG(LogHairStrands, Warning, TEXT("Binding failed for FNDIHairStrandsParametersCS %s. Was it optimized out?"), *ParamNames.BoundingBoxBufferName)
+		}
+		if (!NodeBoundBuffer.IsBound())
+		{
+			UE_LOG(LogHairStrands, Warning, TEXT("Binding failed for FNDIHairStrandsParametersCS %s. Was it optimized out?"), *ParamNames.BoundingBoxBufferName)
+		}
 	}
 
 	virtual void Serialize(FArchive& Ar) override
@@ -376,6 +575,9 @@ struct FNDIHairStrandsParametersCS : public FNiagaraDataInterfaceParametersCS
 
 		Ar << RestPositionOffset;
 		Ar << DeformedPositionOffset;
+
+		Ar << BoundingBoxBuffer;
+		Ar << NodeBoundBuffer;
 	}
 
 	virtual void Set(FRHICommandList& RHICmdList, const FNiagaraDataInterfaceSetArgs& Context) const override
@@ -442,12 +644,16 @@ struct FNDIHairStrandsParametersCS : public FNiagaraDataInterfaceParametersCS
 			//	DeformedRootOffsetValue[0], DeformedRootOffsetValue[1], DeformedRootOffsetValue[2], DeformedPositionOffsetValue[0], DeformedPositionOffsetValue[1], DeformedPositionOffsetValue[2]);
 
 			SetUAVParameter(RHICmdList, ComputeShaderRHI, DeformedPositionBuffer, PointPositionsUAV);
+			RHICmdList.TransitionResource(EResourceTransitionAccess::EWritable, EResourceTransitionPipeline::EComputeToCompute, HairStrandsBuffer->BoundingBoxBuffer.UAV);
+			SetUAVParameter(RHICmdList, ComputeShaderRHI, BoundingBoxBuffer, HairStrandsBuffer->BoundingBoxBuffer.UAV);
+			RHICmdList.TransitionResource(EResourceTransitionAccess::EReadable, EResourceTransitionPipeline::EComputeToCompute, HairStrandsBuffer->NodeBoundBuffer.UAV);
+			SetSRVParameter(RHICmdList, ComputeShaderRHI, NodeBoundBuffer, HairStrandsBuffer->NodeBoundBuffer.SRV);
 			SetSRVParameter(RHICmdList, ComputeShaderRHI, CurvesOffsetsBuffer, HairStrandsBuffer->CurvesOffsetsBuffer.SRV);
 			SetSRVParameter(RHICmdList, ComputeShaderRHI, RestPositionBuffer, HairStrandsBuffer->SourceRestResources->RestPositionBuffer.SRV);
 
-			SetShaderValue(RHICmdList, ComputeShaderRHI, WorldTransform, ProxyData->WorldTransform);
-			SetShaderValue(RHICmdList, ComputeShaderRHI, WorldInverse, ProxyData->WorldTransform.Inverse());
-			SetShaderValue(RHICmdList, ComputeShaderRHI, WorldRotation, ProxyData->WorldTransform.ToQuat());
+			SetShaderValue(RHICmdList, ComputeShaderRHI, WorldTransform, ProxyData->WorldTransform.ToMatrixWithScale());
+			SetShaderValue(RHICmdList, ComputeShaderRHI, WorldInverse, ProxyData->WorldTransform.ToMatrixWithScale().Inverse());
+			SetShaderValue(RHICmdList, ComputeShaderRHI, WorldRotation, ProxyData->WorldTransform.ToMatrixNoScale().ToQuat());
 			SetShaderValue(RHICmdList, ComputeShaderRHI, NumStrands, ProxyData->NumStrands);
 			SetShaderValue(RHICmdList, ComputeShaderRHI, StrandSize, ProxyData->StrandSize);
 			SetShaderValue(RHICmdList, ComputeShaderRHI, BoxCenter, ProxyData->BoxCenter);
@@ -474,6 +680,8 @@ struct FNDIHairStrandsParametersCS : public FNiagaraDataInterfaceParametersCS
 		else
 		{
 			SetUAVParameter(RHICmdList, ComputeShaderRHI, DeformedPositionBuffer, FNiagaraRenderer::GetDummyFloatBuffer().UAV);
+			SetUAVParameter(RHICmdList, ComputeShaderRHI, BoundingBoxBuffer, FNiagaraRenderer::GetDummyUIntBuffer().UAV);
+			SetSRVParameter(RHICmdList, ComputeShaderRHI, NodeBoundBuffer, FNiagaraRenderer::GetDummyUIntBuffer().SRV);
 			SetSRVParameter(RHICmdList, ComputeShaderRHI, CurvesOffsetsBuffer, FNiagaraRenderer::GetDummyUIntBuffer().SRV);
 			SetSRVParameter(RHICmdList, ComputeShaderRHI, RestPositionBuffer, FNiagaraRenderer::GetDummyFloatBuffer().SRV);
 
@@ -508,6 +716,7 @@ struct FNDIHairStrandsParametersCS : public FNiagaraDataInterfaceParametersCS
 	virtual void Unset(FRHICommandList& RHICmdList, const FNiagaraDataInterfaceSetArgs& Context) const override
 	{
 		SetUAVParameter(RHICmdList, Context.Shader->GetComputeShader(), DeformedPositionBuffer, nullptr);
+		SetUAVParameter(RHICmdList, Context.Shader->GetComputeShader(), BoundingBoxBuffer, nullptr);
 	}
 
 private:
@@ -541,6 +750,9 @@ private:
 
 	FShaderParameter RestPositionOffset;
 	FShaderParameter DeformedPositionOffset;
+
+	FShaderResourceParameter BoundingBoxBuffer;
+	FShaderResourceParameter NodeBoundBuffer;
 };
 
 //------------------------------------------------------------------------------------------------------------
@@ -580,7 +792,7 @@ void FNDIHairStrandsProxy::ConsumePerInstanceDataFromGameThread(void* PerInstanc
 }
 
 void FNDIHairStrandsProxy::InitializePerInstanceData(const FNiagaraSystemInstanceID& SystemInstance, FNDIHairStrandsBuffer* HairStrandsBuffer, 
-	const uint32 NumStrands, const uint8 StrandSize, const FVector& BoxCenter, const FVector& BoxExtent, const FMatrix& WorldTransform)
+	const uint32 NumStrands, const uint8 StrandSize, const FVector& BoxCenter, const FVector& BoxExtent, const FTransform& WorldTransform)
 {
 	check(IsInRenderingThread());
 
@@ -653,6 +865,7 @@ void UNiagaraDataInterfaceHairStrands::UpdateSourceComponent(FNiagaraSystemInsta
 			if (UGroomComponent* ParentComp = Cast<UGroomComponent>(SimComp->GetAttachParent()))
 			{
 				SourceComponent = ParentComp;
+				SourceComponent = ParentComp;
 			}
 			else if (UGroomComponent* OuterComp = SimComp->GetTypedOuter<UGroomComponent>())
 			{
@@ -729,8 +942,8 @@ bool UNiagaraDataInterfaceHairStrands::InitPerInstanceData(void* PerInstanceData
 			const FVector BoxExtent = StrandsBox.GetExtent();
 			const FVector BoxCenter = StrandsBox.GetCenter(); 
 
-			const FMatrix WorldTransform = IsComponentValid() ? SourceComponent->GetComponentToWorld().ToMatrixWithScale() :
-				SystemInstance->GetComponent()->GetComponentToWorld().ToMatrixWithScale();
+			const FTransform WorldTransform = IsComponentValid() ? SourceComponent->GetComponentToWorld():
+				SystemInstance->GetComponent()->GetComponentToWorld();
 
 			InstanceData->WorldTransform = WorldTransform;
 			InstanceData->StrandSize = StrandsSize;
@@ -809,15 +1022,14 @@ bool UNiagaraDataInterfaceHairStrands::PerInstanceTick(void* PerInstanceData, FN
 		{
 			InstanceData->BoxCenter = StrandsDeformedResource->PositionOffset;
 		}
-
 		if (IsComponentValid())
 		{
-			InstanceData->WorldTransform = SourceComponent->GetComponentToWorld().ToMatrixWithScale();
+			InstanceData->WorldTransform = SourceComponent->GetComponentToWorld();
 			//UE_LOG(LogHairStrands, Warning, TEXT("Get Component Transform : %s"), *InstanceData->WorldTransform.ToString());
 		}
 		else
 		{
-			InstanceData->WorldTransform = SystemInstance->GetComponent()->GetComponentToWorld().ToMatrixWithScale();
+			InstanceData->WorldTransform = SystemInstance->GetComponent()->GetComponentToWorld();
 			//UE_LOG(LogHairStrands, Warning, TEXT("Get Instance Transform : %s"), *InstanceData->WorldTransform.ToString());
 		}
 	}
@@ -1037,8 +1249,9 @@ void UNiagaraDataInterfaceHairStrands::GetFunctions(TArray<FNiagaraFunctionSigna
 		Sig.bRequiresContext = false;
 		Sig.Inputs.Add(FNiagaraVariable(FNiagaraTypeDefinition(GetClass()), TEXT("Hair Strands")));
 		Sig.Inputs.Add(FNiagaraVariable(FNiagaraTypeDefinition::GetFloatDef(), TEXT("Node Mass")));
-		Sig.Inputs.Add(FNiagaraVariable(FNiagaraTypeDefinition::GetFloatDef(), TEXT("Inverse Mass")));
+		Sig.Inputs.Add(FNiagaraVariable(FNiagaraTypeDefinition::GetBoolDef(), TEXT("Position Mobile")));
 		Sig.Inputs.Add(FNiagaraVariable(FNiagaraTypeDefinition::GetVec3Def(), TEXT("External Force")));
+		Sig.Inputs.Add(FNiagaraVariable(FNiagaraTypeDefinition::GetVec3Def(), TEXT("Force Gradient")));
 		Sig.Inputs.Add(FNiagaraVariable(FNiagaraTypeDefinition::GetFloatDef(), TEXT("Delta Time")));
 		Sig.Inputs.Add(FNiagaraVariable(FNiagaraTypeDefinition::GetVec3Def(), TEXT("Linear Velocity")));
 		Sig.Inputs.Add(FNiagaraVariable(FNiagaraTypeDefinition::GetVec3Def(), TEXT("Node Position")));
@@ -1054,8 +1267,9 @@ void UNiagaraDataInterfaceHairStrands::GetFunctions(TArray<FNiagaraFunctionSigna
 		Sig.bRequiresContext = false;
 		Sig.Inputs.Add(FNiagaraVariable(FNiagaraTypeDefinition(GetClass()), TEXT("Hair Strands")));
 		Sig.Inputs.Add(FNiagaraVariable(FNiagaraTypeDefinition::GetVec3Def(), TEXT("Node Inertia")));
-		Sig.Inputs.Add(FNiagaraVariable(FNiagaraTypeDefinition::GetVec3Def(), TEXT("Inverse Inertia")));
+		Sig.Inputs.Add(FNiagaraVariable(FNiagaraTypeDefinition::GetBoolDef(), TEXT("Orientation Mobile")));
 		Sig.Inputs.Add(FNiagaraVariable(FNiagaraTypeDefinition::GetVec3Def(), TEXT("External Torque")));
+		Sig.Inputs.Add(FNiagaraVariable(FNiagaraTypeDefinition::GetVec3Def(), TEXT("Torque Gradient")));
 		Sig.Inputs.Add(FNiagaraVariable(FNiagaraTypeDefinition::GetFloatDef(), TEXT("Delta Time")));
 		Sig.Inputs.Add(FNiagaraVariable(FNiagaraTypeDefinition::GetVec3Def(), TEXT("Angular Velocity")));
 		Sig.Inputs.Add(FNiagaraVariable(FNiagaraTypeDefinition::GetQuatDef(), TEXT("Node Orientation")));
@@ -1169,6 +1383,19 @@ void UNiagaraDataInterfaceHairStrands::GetFunctions(TArray<FNiagaraFunctionSigna
 
 		OutFunctions.Add(Sig);
 	}
+	{
+		FNiagaraFunctionSignature Sig;
+		Sig.Name = BuildBoundingBoxName;
+		Sig.bMemberFunction = true;
+		Sig.bRequiresContext = false;
+		Sig.Inputs.Add(FNiagaraVariable(FNiagaraTypeDefinition(GetClass()), TEXT("Hair Strands")));
+		Sig.Inputs.Add(FNiagaraVariable(FNiagaraTypeDefinition::GetVec3Def(), TEXT("Node Position")));
+		Sig.Outputs.Add(FNiagaraVariable(FNiagaraTypeDefinition::GetBoolDef(), TEXT("Build Status")));
+		Sig.Outputs.Add(FNiagaraVariable(FNiagaraTypeDefinition::GetVec3Def(), TEXT("Bound Min")));
+		Sig.Outputs.Add(FNiagaraVariable(FNiagaraTypeDefinition::GetVec3Def(), TEXT("Bound Max")));
+
+		OutFunctions.Add(Sig);
+	}
 	{
 		FNiagaraFunctionSignature Sig;
 		Sig.Name = SetupDistanceSpringMaterialName;
@@ -1435,6 +1662,7 @@ void UNiagaraDataInterfaceHairStrands::GetFunctions(TArray<FNiagaraFunctionSigna
 		Sig.Inputs.Add(FNiagaraVariable(FNiagaraTypeDefinition::GetVec3Def(), TEXT("Node Position")));
 		Sig.Inputs.Add(FNiagaraVariable(FNiagaraTypeDefinition::GetVec3Def(), TEXT("Node Velocity")));
 		Sig.Outputs.Add(FNiagaraVariable(FNiagaraTypeDefinition::GetVec3Def(), TEXT("Drag Force")));
+		Sig.Outputs.Add(FNiagaraVariable(FNiagaraTypeDefinition::GetVec3Def(), TEXT("Drag Gradient")));
 
 		OutFunctions.Add(Sig);
 	}
@@ -1472,6 +1700,7 @@ DEFINE_NDI_DIRECT_FUNC_BINDER(UNiagaraDataInterfaceHairStrands, ResetPointPositi
 
 DEFINE_NDI_DIRECT_FUNC_BINDER(UNiagaraDataInterfaceHairStrands, GetBoxExtent);
 DEFINE_NDI_DIRECT_FUNC_BINDER(UNiagaraDataInterfaceHairStrands, GetBoxCenter);
+DEFINE_NDI_DIRECT_FUNC_BINDER(UNiagaraDataInterfaceHairStrands, BuildBoundingBox);
 
 DEFINE_NDI_DIRECT_FUNC_BINDER(UNiagaraDataInterfaceHairStrands, AdvectNodePosition);
 DEFINE_NDI_DIRECT_FUNC_BINDER(UNiagaraDataInterfaceHairStrands, AdvectNodeOrientation);
@@ -1601,12 +1830,12 @@ void UNiagaraDataInterfaceHairStrands::GetVMExternalFunction(const FVMExternalFu
 	}
 	else if (BindingInfo.Name == AdvectNodePositionName)
 	{
-		check(BindingInfo.GetNumInputs() == 13 && BindingInfo.GetNumOutputs() == 6);
+		check(BindingInfo.GetNumInputs() == 16 && BindingInfo.GetNumOutputs() == 6);
 		NDI_FUNC_BINDER(UNiagaraDataInterfaceHairStrands, AdvectNodePosition)::Bind(this, OutFunc);
 	}
 	else if (BindingInfo.Name == AdvectNodeOrientationName)
 	{
-		check(BindingInfo.GetNumInputs() == 18 && BindingInfo.GetNumOutputs() == 7);
+		check(BindingInfo.GetNumInputs() == 19 && BindingInfo.GetNumOutputs() == 7);
 		NDI_FUNC_BINDER(UNiagaraDataInterfaceHairStrands, AdvectNodeOrientation)::Bind(this, OutFunc);
 	}
 	else if (BindingInfo.Name == UpdateLinearVelocityName)
@@ -1629,6 +1858,11 @@ void UNiagaraDataInterfaceHairStrands::GetVMExternalFunction(const FVMExternalFu
 		check(BindingInfo.GetNumInputs() == 1 && BindingInfo.GetNumOutputs() == 3);
 		NDI_FUNC_BINDER(UNiagaraDataInterfaceHairStrands, GetBoxCenter)::Bind(this, OutFunc);
 	}
+	else if (BindingInfo.Name == BuildBoundingBoxName)
+	{
+		check(BindingInfo.GetNumInputs() == 4 && BindingInfo.GetNumOutputs() == 7);
+		NDI_FUNC_BINDER(UNiagaraDataInterfaceHairStrands, BuildBoundingBox)::Bind(this, OutFunc);
+	}
 	else if (BindingInfo.Name == SetupDistanceSpringMaterialName)
 	{
 		check(BindingInfo.GetNumInputs() == 7 && BindingInfo.GetNumOutputs() == 3);
@@ -1711,7 +1945,7 @@ void UNiagaraDataInterfaceHairStrands::GetVMExternalFunction(const FVMExternalFu
 	}
 	else if (BindingInfo.Name == ComputeAirDragForceName)
 	{
-		check(BindingInfo.GetNumInputs() == 14 && BindingInfo.GetNumOutputs() == 3);
+		check(BindingInfo.GetNumInputs() == 14 && BindingInfo.GetNumOutputs() == 6);
 		NDI_FUNC_BINDER(UNiagaraDataInterfaceHairStrands, ComputeAirDragForce)::Bind(this, OutFunc);
 	}
 	else if (BindingInfo.Name == NeedSimulationResetName)
@@ -1786,7 +2020,7 @@ void UNiagaraDataInterfaceHairStrands::GetStrandSize(FVectorVMContext& Context)
 void UNiagaraDataInterfaceHairStrands::GetWorldTransform(FVectorVMContext& Context)
 {
 	VectorVM::FUserPtrHandler<FNDIHairStrandsData> InstData(Context);
-	FMatrix WorldTransform = InstData->WorldTransform;
+	FMatrix WorldTransform = InstData->WorldTransform.ToMatrixWithScale();
 
 	WriteTransform(WorldTransform, Context);
 }
@@ -1794,11 +2028,14 @@ void UNiagaraDataInterfaceHairStrands::GetWorldTransform(FVectorVMContext& Conte
 void UNiagaraDataInterfaceHairStrands::GetWorldInverse(FVectorVMContext& Context)
 {
 	VectorVM::FUserPtrHandler<FNDIHairStrandsData> InstData(Context);
-	FMatrix WorldInverse = InstData->WorldTransform.Inverse();
+	FMatrix WorldInverse = InstData->WorldTransform.ToMatrixWithScale().Inverse();
 
 	WriteTransform(WorldInverse, Context);
 }
 
+void UNiagaraDataInterfaceHairStrands::BuildBoundingBox(FVectorVMContext& Context)
+{}
+
 void UNiagaraDataInterfaceHairStrands::GetBoxCenter(FVectorVMContext& Context)
 {
 	VectorVM::FUserPtrHandler<FNDIHairStrandsData> InstData(Context);
@@ -2247,12 +2484,12 @@ bool UNiagaraDataInterfaceHairStrands::GetFunctionHLSL(const FName& DefinitionFu
 	else if (DefinitionFunctionName == AdvectNodePositionName)
 	{
 		static const TCHAR *FormatSample = TEXT(R"(
-			void {InstanceFunctionName} (in float NodeMass, in float InverssMass, in float3 ExternalForce, in float DeltaTime, 
+			void {InstanceFunctionName} (in float NodeMass, in bool IsPositionMobile, in float3 ExternalForce, in float3 ForceGradient, in float DeltaTime, 
 									     in float3 LinearVelocity, in float3 NodePosition, out float3 OutLinearVelocity, out float3 OutNodePosition)
 			{
 				OutLinearVelocity = LinearVelocity;
 				OutNodePosition = NodePosition;
-				{HairStrandsContextName} DIHairStrands_AdvectNodePosition(DIContext,NodeMass,InverssMass,ExternalForce,DeltaTime,OutLinearVelocity,OutNodePosition);
+				{HairStrandsContextName} DIHairStrands_AdvectNodePosition(DIContext,NodeMass,IsPositionMobile,ExternalForce,ForceGradient,DeltaTime,OutLinearVelocity,OutNodePosition);
 			}
 			)");
 		OutHLSL += FString::Format(FormatSample, ArgsSample);
@@ -2261,12 +2498,12 @@ bool UNiagaraDataInterfaceHairStrands::GetFunctionHLSL(const FName& DefinitionFu
 	else if (DefinitionFunctionName == AdvectNodeOrientationName)
 	{
 		static const TCHAR *FormatSample = TEXT(R"(
-			void {InstanceFunctionName} (in float3 NodeInertia, in float3 InverseInertia, in float3 ExternalTorque, in float DeltaTime, 
+			void {InstanceFunctionName} (in float3 NodeInertia, in bool IsOrientationMobile, in float3 ExternalTorque, in float3 TorqueGradient, in float DeltaTime, 
 										 in float3 AngularVelocity, in float4 NodeOrientation, out float3 OutAngularVelocity, out float4 OutNodeOrientation)
 			{
 				OutAngularVelocity = AngularVelocity;
 				OutNodeOrientation = NodeOrientation;
-				{HairStrandsContextName} DIHairStrands_AdvectNodeOrientation(DIContext,NodeInertia,InverseInertia,ExternalTorque,DeltaTime,OutAngularVelocity,OutNodeOrientation);
+				{HairStrandsContextName} DIHairStrands_AdvectNodeOrientation(DIContext,NodeInertia,IsOrientationMobile,ExternalTorque,TorqueGradient,DeltaTime,OutAngularVelocity,OutNodeOrientation);
 			}
 			)");
 		OutHLSL += FString::Format(FormatSample, ArgsSample);
@@ -2316,6 +2553,17 @@ bool UNiagaraDataInterfaceHairStrands::GetFunctionHLSL(const FName& DefinitionFu
 		OutHLSL += FString::Format(FormatSample, ArgsSample);
 		return true;
 	}
+	else if (DefinitionFunctionName == BuildBoundingBoxName)
+	{
+		static const TCHAR *FormatSample = TEXT(R"(
+					void {InstanceFunctionName} (in float3 NodePosition, out bool OutBuildStatus, out float3 OutBoundMin, out float3 OutBoundMax)
+					{
+						{HairStrandsContextName} DIHairStrands_BuildBoundingBox(DIContext,NodePosition,OutBuildStatus,OutBoundMin,OutBoundMax);
+					}
+					)");
+		OutHLSL += FString::Format(FormatSample, ArgsSample);
+		return true;
+	}
 	else if (DefinitionFunctionName == SetupDistanceSpringMaterialName)
 	{
 		static const TCHAR *FormatSample = TEXT(R"(
@@ -2547,9 +2795,9 @@ in float RestLength, in float DeltaTime, in int NodeOffset, in float MaterialDam
 	{
 		static const TCHAR *FormatSample = TEXT(R"(
 						void {InstanceFunctionName} ( in float AirDensity, in float AirViscosity, in float AirDrag, 
-	in float3 AirVelocity, in float NodeThickness, in float3 NodePosition, in float3 NodeVelocity, out float3 OutAirDrag )
+	in float3 AirVelocity, in float NodeThickness, in float3 NodePosition, in float3 NodeVelocity, out float3 OutAirDrag, out float3 OutDragGradient )
 						{
-							{HairStrandsContextName} DIHairStrands_ComputeAirDragForce(DIContext,AirDensity,AirViscosity,AirDrag,AirVelocity,NodeThickness,NodePosition,NodeVelocity,OutAirDrag);
+							{HairStrandsContextName} ComputeAirDragForce(DIContext.StrandSize,AirDensity,AirViscosity,AirDrag,AirVelocity,NodeThickness,NodePosition,NodeVelocity,OutAirDrag,OutDragGradient);
 						}
 						)");
 		OutHLSL += FString::Format(FormatSample, ArgsSample);
@@ -2574,6 +2822,7 @@ in float RestLength, in float DeltaTime, in int NodeOffset, in float MaterialDam
 void UNiagaraDataInterfaceHairStrands::GetCommonHLSL(FString& OutHLSL)
 {
 	OutHLSL += TEXT("#include \"/Plugin/Experimental/HairStrands/Private/NiagaraQuaternionUtils.ush\"\n");
+	OutHLSL += TEXT("#include \"/Plugin/Experimental/HairStrands/Private/NiagaraStrandsExternalForce.ush\"\n");
 	OutHLSL += TEXT("#include \"/Plugin/Experimental/HairStrands/Private/NiagaraHookeSpringMaterial.ush\"\n");
 	OutHLSL += TEXT("#include \"/Plugin/Experimental/HairStrands/Private/NiagaraCosseratRodMaterial.ush\"\n");
 	OutHLSL += TEXT("#include \"/Plugin/Experimental/HairStrands/Private/NiagaraStaticCollisionConstraint.ush\"\n");
@@ -2610,7 +2859,18 @@ UNiagaraDataInterfaceHairStrands::ConstructComputeParameters() const
 }
 
 void FNDIHairStrandsProxy::PreStage(FRHICommandList& RHICmdList, const FNiagaraDataInterfaceSetArgs& Context)
-{}
+{
+	if (Context.ShaderStageIndex == 0)
+	{
+		FNDIHairStrandsData* ProxyData =
+			SystemInstancesToProxyData.Find(Context.SystemInstance);
+
+		if (ProxyData != nullptr && ProxyData->HairStrandsBuffer != nullptr)
+		{
+			ProxyData->HairStrandsBuffer->ClearBuffer(RHICmdList);
+		}
+	}
+}
 
 void FNDIHairStrandsProxy::PostStage(FRHICommandList& RHICmdList, const FNiagaraDataInterfaceSetArgs& Context)
 {}
diff --git a/Engine/Plugins/Experimental/HairStrands/Source/HairStrandsNiagara/Private/NiagaraDataInterfacePhysicsAsset.cpp b/Engine/Plugins/Experimental/HairStrands/Source/HairStrandsNiagara/Private/NiagaraDataInterfacePhysicsAsset.cpp
--- a/Engine/Plugins/Experimental/HairStrands/Source/HairStrandsNiagara/Private/NiagaraDataInterfacePhysicsAsset.cpp
+++ b/Engine/Plugins/Experimental/HairStrands/Source/HairStrandsNiagara/Private/NiagaraDataInterfacePhysicsAsset.cpp
@@ -24,6 +24,8 @@ static const FName GetNumCapsulesName(TEXT("GetNumCapsules"));
 //------------------------------------------------------------------------------------------------------------
 
 static const FName GetClosestPointName(TEXT("GetClosestPoint"));
+static const FName GetTexturePointName(TEXT("GetTexturePoint"));
+static const FName GetProjectionPointName(TEXT("GetProjectionPoint"));
 
 //------------------------------------------------------------------------------------------------------------
 
@@ -32,8 +34,13 @@ const FString UNiagaraDataInterfacePhysicsAsset::ElementOffsetsName(TEXT("Elemen
 const FString UNiagaraDataInterfacePhysicsAsset::CurrentTransformBufferName(TEXT("CurrentTransformBuffer_"));
 const FString UNiagaraDataInterfacePhysicsAsset::PreviousTransformBufferName(TEXT("PreviousTransformBuffer_"));
 const FString UNiagaraDataInterfacePhysicsAsset::InverseTransformBufferName(TEXT("InverseTransformBuffer_"));
+const FString UNiagaraDataInterfacePhysicsAsset::RestTransformBufferName(TEXT("RestTransformBuffer_"));
+const FString UNiagaraDataInterfacePhysicsAsset::RestInverseBufferName(TEXT("RestInverseBuffer_"));
 const FString UNiagaraDataInterfacePhysicsAsset::ElementExtentBufferName(TEXT("ElementExtentBuffer_"));
 
+const FString UNiagaraDataInterfacePhysicsAsset::BoxOriginName(TEXT("BoxOrigin_"));
+const FString UNiagaraDataInterfacePhysicsAsset::BoxExtentName(TEXT("BoxExtent_"));
+
 //------------------------------------------------------------------------------------------------------------
 
 struct FNDIPhysicsAssetParametersName
@@ -45,7 +52,12 @@ struct FNDIPhysicsAssetParametersName
 		CurrentTransformBufferName = UNiagaraDataInterfacePhysicsAsset::CurrentTransformBufferName + Suffix;
 		PreviousTransformBufferName = UNiagaraDataInterfacePhysicsAsset::PreviousTransformBufferName + Suffix;
 		InverseTransformBufferName = UNiagaraDataInterfacePhysicsAsset::InverseTransformBufferName + Suffix;
+		RestTransformBufferName = UNiagaraDataInterfacePhysicsAsset::RestTransformBufferName + Suffix;
+		RestInverseBufferName = UNiagaraDataInterfacePhysicsAsset::RestInverseBufferName + Suffix;
 		ElementExtentBufferName = UNiagaraDataInterfacePhysicsAsset::ElementExtentBufferName + Suffix;
+
+		BoxOriginName = UNiagaraDataInterfacePhysicsAsset::BoxOriginName + Suffix;
+		BoxExtentName = UNiagaraDataInterfacePhysicsAsset::BoxExtentName + Suffix;
 	}
 
 	FString ElementOffsetsName;
@@ -53,7 +65,12 @@ struct FNDIPhysicsAssetParametersName
 	FString CurrentTransformBufferName;
 	FString PreviousTransformBufferName;
 	FString InverseTransformBufferName;
+	FString RestTransformBufferName;
+	FString RestInverseBufferName;
 	FString ElementExtentBufferName;
+
+	FString BoxOriginName;
+	FString BoxExtentName;
 };
 
 //------------------------------------------------------------------------------------------------------------
@@ -103,15 +120,10 @@ void CreateInternalArrays(const TWeakObjectPtr<UPhysicsAsset> PhysicsAsset, cons
 			const FReferenceSkeleton* RefSkeleton = &PhysicsAsset->GetPreviewMesh()->RefSkeleton;
 			if (RefSkeleton != nullptr)
 			{
-				TArray<FTransform> BoneTransforms;
-				if (SkeletalMesh != nullptr)
-				{
-					BoneTransforms = SkeletalMesh->GetComponentSpaceTransforms();
-				}
-				else
-				{
-					FAnimationRuntime::FillUpComponentSpaceTransforms(*RefSkeleton, RefSkeleton->GetRefBonePose(), BoneTransforms);
-				}
+				TArray<FTransform> RestTransforms;
+				FAnimationRuntime::FillUpComponentSpaceTransforms(*RefSkeleton, RefSkeleton->GetRefBonePose(), RestTransforms);
+
+				TArray<FTransform> BoneTransforms = (SkeletalMesh != nullptr) ? SkeletalMesh->GetComponentSpaceTransforms() : RestTransforms;
 				const bool bHasMasterPoseComponent = (SkeletalMesh != nullptr) && SkeletalMesh->MasterPoseComponent.IsValid();
 
 				uint32 NumBoxes = 0;
@@ -139,6 +151,8 @@ void CreateInternalArrays(const TWeakObjectPtr<UPhysicsAsset> PhysicsAsset, cons
 
 				OutAssetArrays->CurrentTransform.SetNum(NumTransforms);
 				OutAssetArrays->InverseTransform.SetNum(NumTransforms);
+				OutAssetArrays->RestInverse.SetNum(NumTransforms);
+				OutAssetArrays->RestTransform.SetNum(NumTransforms);
 				OutAssetArrays->PreviousTransform.SetNum(NumTransforms);
 				OutAssetArrays->ElementExtent.SetNum(NumExtents);
 
@@ -149,10 +163,15 @@ void CreateInternalArrays(const TWeakObjectPtr<UPhysicsAsset> PhysicsAsset, cons
 					const int32 BoneIndex = RefSkeleton->FindBoneIndex(BoneName);
 					if (BoneIndex != INDEX_NONE)
 					{
+						const FTransform RestTransform = RestTransforms[BoneIndex];
 						const FTransform BoneTransform = bHasMasterPoseComponent ? SkeletalMesh->GetBoneTransform(BoneIndex) : BoneTransforms[BoneIndex] * WorldTransform;
 
 						for (const FKBoxElem& BoxElem : BodySetup->AggGeom.BoxElems)
 						{
+							const FTransform RestElement = FTransform(BoxElem.Rotation, BoxElem.Center) * RestTransform;
+							FillCurrentTransforms(RestElement, ElementCount, OutAssetArrays->RestTransform, OutAssetArrays->RestInverse);
+							--ElementCount;
+
 							const FTransform ElementTransform = FTransform(BoxElem.Rotation, BoxElem.Center) * BoneTransform;
 							OutAssetArrays->ElementExtent[ElementCount] = FVector4(BoxElem.X, BoxElem.Y, BoxElem.Z, 0);
 							FillCurrentTransforms(ElementTransform, ElementCount, OutAssetArrays->CurrentTransform, OutAssetArrays->InverseTransform);
@@ -160,6 +179,10 @@ void CreateInternalArrays(const TWeakObjectPtr<UPhysicsAsset> PhysicsAsset, cons
 
 						for (const FKSphereElem& SphereElem : BodySetup->AggGeom.SphereElems)
 						{
+							const FTransform RestElement = FTransform(SphereElem.Center) * RestTransform;
+							FillCurrentTransforms(RestElement, ElementCount, OutAssetArrays->RestTransform, OutAssetArrays->RestInverse);
+							--ElementCount;
+
 							const FTransform ElementTransform = FTransform(SphereElem.Center) * BoneTransform;
 							OutAssetArrays->ElementExtent[ElementCount] = FVector4(SphereElem.Radius, 0, 0, 0);
 							FillCurrentTransforms(ElementTransform, ElementCount, OutAssetArrays->CurrentTransform, OutAssetArrays->InverseTransform);
@@ -167,6 +190,10 @@ void CreateInternalArrays(const TWeakObjectPtr<UPhysicsAsset> PhysicsAsset, cons
 
 						for (const FKSphylElem& CapsuleElem : BodySetup->AggGeom.SphylElems)
 						{
+							const FTransform RestElement = FTransform(CapsuleElem.Rotation, CapsuleElem.Center) * RestTransform;
+							FillCurrentTransforms(RestElement, ElementCount, OutAssetArrays->RestTransform, OutAssetArrays->RestInverse);
+							--ElementCount;
+
 							const FTransform ElementTransform = FTransform(CapsuleElem.Rotation, CapsuleElem.Center) * BoneTransform;
 							OutAssetArrays->ElementExtent[ElementCount] = FVector4(CapsuleElem.Radius, CapsuleElem.Length, 0, 0);
 							FillCurrentTransforms(ElementTransform, ElementCount, OutAssetArrays->CurrentTransform, OutAssetArrays->InverseTransform);
@@ -278,6 +305,8 @@ void FNDIPhysicsAssetBuffer::InitRHI()
 		CreateInternalBuffer<FVector4, FVector4, 1, EPixelFormat::PF_A32B32G32R32F, true>(AssetArrays->CurrentTransform.Num(), AssetArrays->CurrentTransform, CurrentTransformBuffer);
 		CreateInternalBuffer<FVector4, FVector4, 1, EPixelFormat::PF_A32B32G32R32F, true>(AssetArrays->PreviousTransform.Num(), AssetArrays->PreviousTransform, PreviousTransformBuffer);
 		CreateInternalBuffer<FVector4, FVector4, 1, EPixelFormat::PF_A32B32G32R32F, true>(AssetArrays->InverseTransform.Num(), AssetArrays->InverseTransform, InverseTransformBuffer);
+		CreateInternalBuffer<FVector4, FVector4, 1, EPixelFormat::PF_A32B32G32R32F, true>(AssetArrays->RestTransform.Num(), AssetArrays->RestTransform, RestTransformBuffer);
+		CreateInternalBuffer<FVector4, FVector4, 1, EPixelFormat::PF_A32B32G32R32F, true>(AssetArrays->RestInverse.Num(), AssetArrays->RestInverse, RestInverseBuffer);
 		CreateInternalBuffer<FVector4, FVector4, 1, EPixelFormat::PF_A32B32G32R32F, true>(AssetArrays->ElementExtent.Num(), AssetArrays->ElementExtent, ElementExtentBuffer);
 
 		//UE_LOG(LogPhysicsAsset, Warning, TEXT("Num Capsules = %d | Num Spheres = %d | Num Boxes = %d"), AssetArrays->ElementOffsets.NumElements - AssetArrays->ElementOffsets.CapsuleOffset,
@@ -290,6 +319,8 @@ void FNDIPhysicsAssetBuffer::ReleaseRHI()
 	CurrentTransformBuffer.Release();
 	PreviousTransformBuffer.Release();
 	InverseTransformBuffer.Release();
+	RestTransformBuffer.Release();
+	RestInverseBuffer.Release();
 	ElementExtentBuffer.Release();
 }
 
@@ -306,8 +337,13 @@ struct FNDIPhysicsAssetParametersCS : public FNiagaraDataInterfaceParametersCS
 		CurrentTransformBuffer.Bind(ParameterMap, *ParamNames.CurrentTransformBufferName);
 		PreviousTransformBuffer.Bind(ParameterMap, *ParamNames.PreviousTransformBufferName);
 		InverseTransformBuffer.Bind(ParameterMap, *ParamNames.InverseTransformBufferName);
+		RestTransformBuffer.Bind(ParameterMap, *ParamNames.RestTransformBufferName);
+		RestInverseBuffer.Bind(ParameterMap, *ParamNames.RestInverseBufferName);
 		ElementExtentBuffer.Bind(ParameterMap, *ParamNames.ElementExtentBufferName);
 
+		BoxOrigin.Bind(ParameterMap, *ParamNames.BoxOriginName);
+		BoxExtent.Bind(ParameterMap, *ParamNames.BoxExtentName);
+
 		if (!CurrentTransformBuffer.IsBound())
 		{
 			UE_LOG(LogPhysicsAsset, Warning, TEXT("Binding failed for FNDIPhysicsAssetParametersCS %s. Was it optimized out?"), *ParamNames.CurrentTransformBufferName)
@@ -320,6 +356,14 @@ struct FNDIPhysicsAssetParametersCS : public FNiagaraDataInterfaceParametersCS
 		{
 			UE_LOG(LogPhysicsAsset, Warning, TEXT("Binding failed for FNDIPhysicsAssetParametersCS %s. Was it optimized out?"), *ParamNames.InverseTransformBufferName)
 		}
+		if (!RestTransformBuffer.IsBound())
+		{
+			UE_LOG(LogPhysicsAsset, Warning, TEXT("Binding failed for FNDIPhysicsAssetParametersCS %s. Was it optimized out?"), *ParamNames.RestTransformBufferName)
+		}
+		if (!RestInverseBuffer.IsBound())
+		{
+			UE_LOG(LogPhysicsAsset, Warning, TEXT("Binding failed for FNDIPhysicsAssetParametersCS %s. Was it optimized out?"), *ParamNames.RestInverseBufferName)
+		}
 		if (!ElementExtentBuffer.IsBound())
 		{
 			UE_LOG(LogPhysicsAsset, Warning, TEXT("Binding failed for FNDIPhysicsAssetParametersCS %s. Was it optimized out?"), *ParamNames.ElementExtentBufferName)
@@ -332,7 +376,11 @@ struct FNDIPhysicsAssetParametersCS : public FNiagaraDataInterfaceParametersCS
 		Ar << CurrentTransformBuffer;
 		Ar << PreviousTransformBuffer;
 		Ar << InverseTransformBuffer;
+		Ar << RestTransformBuffer;
+		Ar << RestInverseBuffer;
 		Ar << ElementExtentBuffer;
+		Ar << BoxOrigin;
+		Ar << BoxExtent;
 	}
 
 	virtual void Set(FRHICommandList& RHICmdList, const FNiagaraDataInterfaceSetArgs& Context) const override
@@ -353,20 +401,27 @@ struct FNDIPhysicsAssetParametersCS : public FNiagaraDataInterfaceParametersCS
 			SetSRVParameter(RHICmdList, ComputeShaderRHI, CurrentTransformBuffer, AssetBuffer->CurrentTransformBuffer.SRV);
 			SetSRVParameter(RHICmdList, ComputeShaderRHI, PreviousTransformBuffer, AssetBuffer->PreviousTransformBuffer.SRV);
 			SetSRVParameter(RHICmdList, ComputeShaderRHI, InverseTransformBuffer, AssetBuffer->InverseTransformBuffer.SRV);
+			SetSRVParameter(RHICmdList, ComputeShaderRHI, RestTransformBuffer, AssetBuffer->RestTransformBuffer.SRV);
+			SetSRVParameter(RHICmdList, ComputeShaderRHI, RestInverseBuffer, AssetBuffer->RestInverseBuffer.SRV);
 			SetSRVParameter(RHICmdList, ComputeShaderRHI, ElementExtentBuffer, AssetBuffer->ElementExtentBuffer.SRV);
 
 			SetShaderValue(RHICmdList, ComputeShaderRHI, ElementOffsets, AssetBuffer->AssetArrays->ElementOffsets);
-			
+			SetShaderValue(RHICmdList, ComputeShaderRHI, BoxOrigin, ProxyData->BoxOrigin);
+			SetShaderValue(RHICmdList, ComputeShaderRHI, BoxExtent, ProxyData->BoxExtent);
 		}
 		else
 		{
 			SetSRVParameter(RHICmdList, ComputeShaderRHI, CurrentTransformBuffer, FNiagaraRenderer::GetDummyFloatBuffer().SRV);
 			SetSRVParameter(RHICmdList, ComputeShaderRHI, PreviousTransformBuffer, FNiagaraRenderer::GetDummyFloatBuffer().SRV);
 			SetSRVParameter(RHICmdList, ComputeShaderRHI, InverseTransformBuffer, FNiagaraRenderer::GetDummyFloatBuffer().SRV);
+			SetSRVParameter(RHICmdList, ComputeShaderRHI, RestTransformBuffer, FNiagaraRenderer::GetDummyFloatBuffer().SRV);
+			SetSRVParameter(RHICmdList, ComputeShaderRHI, RestInverseBuffer, FNiagaraRenderer::GetDummyFloatBuffer().SRV);
 			SetSRVParameter(RHICmdList, ComputeShaderRHI, ElementExtentBuffer, FNiagaraRenderer::GetDummyFloatBuffer().SRV);
 
 			static const FElementOffset DummyOffsets(0,0,0,0);
 			SetShaderValue(RHICmdList, ComputeShaderRHI, ElementOffsets, DummyOffsets);
+			SetShaderValue(RHICmdList, ComputeShaderRHI, BoxOrigin, FVector(0,0,0));
+			SetShaderValue(RHICmdList, ComputeShaderRHI, BoxExtent, FVector(0,0,0));
 		}
 	}
 
@@ -381,7 +436,12 @@ private:
 	FShaderResourceParameter CurrentTransformBuffer;
 	FShaderResourceParameter PreviousTransformBuffer;
 	FShaderResourceParameter InverseTransformBuffer;
+	FShaderResourceParameter RestTransformBuffer;
+	FShaderResourceParameter RestInverseBuffer;
 	FShaderResourceParameter ElementExtentBuffer;
+
+	FShaderParameter BoxOrigin;
+	FShaderParameter BoxExtent;
 };
 
 //------------------------------------------------------------------------------------------------------------
@@ -405,6 +465,8 @@ void FNDIPhysicsAssetProxy::ConsumePerInstanceDataFromGameThread(void* PerInstan
 	if (TargetData)
 	{
 		TargetData->AssetBuffer = SourceData->AssetBuffer;
+		TargetData->BoxOrigin = SourceData->BoxOrigin;
+		TargetData->BoxExtent = SourceData->BoxExtent;
 	}
 	else
 	{
@@ -412,7 +474,7 @@ void FNDIPhysicsAssetProxy::ConsumePerInstanceDataFromGameThread(void* PerInstan
 	}
 }
 
-void FNDIPhysicsAssetProxy::InitializePerInstanceData(const FNiagaraSystemInstanceID& SystemInstance, FNDIPhysicsAssetBuffer* AssetBuffer)
+void FNDIPhysicsAssetProxy::InitializePerInstanceData(const FNiagaraSystemInstanceID& SystemInstance, FNDIPhysicsAssetBuffer* AssetBuffer, const FVector& BoxOrigin , const FVector& BoxExtent)
 {
 	check(IsInRenderingThread());
 
@@ -426,6 +488,8 @@ void FNDIPhysicsAssetProxy::InitializePerInstanceData(const FNiagaraSystemInstan
 		TargetData = &SystemInstancesToProxyData.Add(SystemInstance);
 	}
 	TargetData->AssetBuffer = AssetBuffer;
+	TargetData->BoxOrigin = BoxOrigin;
+	TargetData->BoxExtent = BoxExtent;
 }
 
 void FNDIPhysicsAssetProxy::DestroyPerInstanceData(NiagaraEmitterInstanceBatcher* Batcher, const FNiagaraSystemInstanceID& SystemInstance)
@@ -505,16 +569,22 @@ bool UNiagaraDataInterfacePhysicsAsset::InitPerInstanceData(void* PerInstanceDat
 		FNDIPhysicsAssetBuffer* AssetBuffer = new FNDIPhysicsAssetBuffer;
 		AssetBuffer->SetupArrays(PhysicsAsset, SourceComponent, WorldTransform);
 
+		const bool HasPreviewMesh = PhysicsAsset != nullptr && PhysicsAsset->GetPreviewMesh() != nullptr;
+		const FVector BoxOrigin = HasPreviewMesh ? PhysicsAsset->GetPreviewMesh()->GetImportedBounds().Origin : FVector(0, 0, 0);
+		const FVector BoxExtent = HasPreviewMesh ? PhysicsAsset->GetPreviewMesh()->GetImportedBounds().BoxExtent : FVector(1, 1, 1);
+
 		// Push instance data to RT
 		{
 			InstanceData->AssetBuffer = AssetBuffer;
+			InstanceData->BoxOrigin = BoxOrigin;
+			InstanceData->BoxExtent = BoxExtent;
 
 			FNDIPhysicsAssetProxy* ThisProxy = GetProxyAs<FNDIPhysicsAssetProxy>();
 			ENQUEUE_RENDER_COMMAND(FNiagaraDIPushInitialInstanceDataToRT) (
-				[ThisProxy, InstanceID = SystemInstance->GetId(), AssetBuffer](FRHICommandListImmediate& CmdList)
+				[ThisProxy, InstanceID = SystemInstance->GetId(), AssetBuffer, BoxOrigin, BoxExtent](FRHICommandListImmediate& CmdList)
 			{
 				AssetBuffer->InitResource();
-				ThisProxy->InitializePerInstanceData(InstanceID, AssetBuffer);
+				ThisProxy->InitializePerInstanceData(InstanceID, AssetBuffer, BoxOrigin, BoxExtent);
 			}
 			);
 		}
@@ -637,6 +707,36 @@ void UNiagaraDataInterfacePhysicsAsset::GetFunctions(TArray<FNiagaraFunctionSign
 		Sig.Outputs.Add(FNiagaraVariable(FNiagaraTypeDefinition::GetVec3Def(), TEXT("Closest Velocity")));
 		Sig.Outputs.Add(FNiagaraVariable(FNiagaraTypeDefinition::GetFloatDef(), TEXT("Closest Distance")));
 
+		OutFunctions.Add(Sig);
+	}
+	{
+		FNiagaraFunctionSignature Sig;
+		Sig.Name = GetTexturePointName;
+		Sig.bMemberFunction = true;
+		Sig.bRequiresContext = false;
+		Sig.Inputs.Add(FNiagaraVariable(FNiagaraTypeDefinition(GetClass()), TEXT("Physics Asset")));
+		Sig.Inputs.Add(FNiagaraVariable(FNiagaraTypeDefinition::GetVec3Def(), TEXT("World Position")));
+		Sig.Outputs.Add(FNiagaraVariable(FNiagaraTypeDefinition::GetIntDef(), TEXT("Element Index")));
+		Sig.Outputs.Add(FNiagaraVariable(FNiagaraTypeDefinition::GetVec3Def(), TEXT("Texture Position")));
+
+		OutFunctions.Add(Sig);
+	}
+	{
+		FNiagaraFunctionSignature Sig;
+		Sig.Name = GetProjectionPointName;
+		Sig.bMemberFunction = true;
+		Sig.bRequiresContext = false;
+		Sig.Inputs.Add(FNiagaraVariable(FNiagaraTypeDefinition(GetClass()), TEXT("Physics Asset")));
+		Sig.Inputs.Add(FNiagaraVariable(FNiagaraTypeDefinition::GetVec3Def(), TEXT("World Position")));
+		Sig.Inputs.Add(FNiagaraVariable(FNiagaraTypeDefinition::GetFloatDef(), TEXT("Delta Time")));
+		Sig.Inputs.Add(FNiagaraVariable(FNiagaraTypeDefinition::GetIntDef(), TEXT("Element Index")));
+		Sig.Inputs.Add(FNiagaraVariable(FNiagaraTypeDefinition::GetFloatDef(), TEXT("Texture Value")));
+		Sig.Inputs.Add(FNiagaraVariable(FNiagaraTypeDefinition::GetVec3Def(), TEXT("Texture Gradient")));
+		Sig.Outputs.Add(FNiagaraVariable(FNiagaraTypeDefinition::GetVec3Def(), TEXT("Closest Position")));
+		Sig.Outputs.Add(FNiagaraVariable(FNiagaraTypeDefinition::GetVec3Def(), TEXT("Closest Normal")));
+		Sig.Outputs.Add(FNiagaraVariable(FNiagaraTypeDefinition::GetVec3Def(), TEXT("Closest Velocity")));
+		Sig.Outputs.Add(FNiagaraVariable(FNiagaraTypeDefinition::GetFloatDef(), TEXT("Closest Distance")));
+
 		OutFunctions.Add(Sig);
 	}
 }
@@ -645,6 +745,8 @@ DEFINE_NDI_DIRECT_FUNC_BINDER(UNiagaraDataInterfacePhysicsAsset, GetNumBoxes);
 DEFINE_NDI_DIRECT_FUNC_BINDER(UNiagaraDataInterfacePhysicsAsset, GetNumSpheres);
 DEFINE_NDI_DIRECT_FUNC_BINDER(UNiagaraDataInterfacePhysicsAsset, GetNumCapsules);
 DEFINE_NDI_DIRECT_FUNC_BINDER(UNiagaraDataInterfacePhysicsAsset, GetClosestPoint);
+DEFINE_NDI_DIRECT_FUNC_BINDER(UNiagaraDataInterfacePhysicsAsset, GetTexturePoint);
+DEFINE_NDI_DIRECT_FUNC_BINDER(UNiagaraDataInterfacePhysicsAsset, GetProjectionPoint);
 
 void UNiagaraDataInterfacePhysicsAsset::GetVMExternalFunction(const FVMExternalFunctionBindingInfo& BindingInfo, void* InstanceData, FVMExternalFunction &OutFunc)
 {
@@ -668,6 +770,16 @@ void UNiagaraDataInterfacePhysicsAsset::GetVMExternalFunction(const FVMExternalF
 		check(BindingInfo.GetNumInputs() == 5 && BindingInfo.GetNumOutputs() == 10);
 		NDI_FUNC_BINDER(UNiagaraDataInterfacePhysicsAsset, GetClosestPoint)::Bind(this, OutFunc);
 	}
+	else if (BindingInfo.Name == GetTexturePointName)
+	{
+		check(BindingInfo.GetNumInputs() == 4 && BindingInfo.GetNumOutputs() == 4);
+		NDI_FUNC_BINDER(UNiagaraDataInterfacePhysicsAsset, GetTexturePoint)::Bind(this, OutFunc);
+	}
+	else if (BindingInfo.Name == GetProjectionPointName)
+	{
+		check(BindingInfo.GetNumInputs() == 10 && BindingInfo.GetNumOutputs() == 10);
+		NDI_FUNC_BINDER(UNiagaraDataInterfacePhysicsAsset, GetProjectionPoint)::Bind(this, OutFunc);
+	}
 }
 
 void UNiagaraDataInterfacePhysicsAsset::GetNumBoxes(FVectorVMContext& Context)
@@ -686,6 +798,14 @@ void UNiagaraDataInterfacePhysicsAsset::GetClosestPoint(FVectorVMContext& Contex
 {
 }
 
+void UNiagaraDataInterfacePhysicsAsset::GetTexturePoint(FVectorVMContext& Context)
+{
+}
+
+void UNiagaraDataInterfacePhysicsAsset::GetProjectionPoint(FVectorVMContext& Context)
+{
+}
+
 bool UNiagaraDataInterfacePhysicsAsset::GetFunctionHLSL(const FName& DefinitionFunctionName, FString InstanceFunctionName, FNiagaraDataInterfaceGPUParamInfo& ParamInfo, FString& OutHLSL)
 {
 	FNDIPhysicsAssetParametersName ParamNames(ParamInfo.DataInterfaceHLSLSymbol);
@@ -749,6 +869,30 @@ bool UNiagaraDataInterfacePhysicsAsset::GetFunctionHLSL(const FName& DefinitionF
 		OutHLSL += FString::Format(FormatSample, ArgsSample);
 		return true;
 	}
+	else if (DefinitionFunctionName == GetTexturePointName)
+	{
+		static const TCHAR *FormatSample = TEXT(R"(
+		void {InstanceFunctionName}(in float3 WorldPosition, out int OutElementIndex, out float3 OutTexturePosition)
+		{
+			{PhysicsAssetContextName} DIPhysicsAsset_GetTexturePoint(DIContext,WorldPosition,OutElementIndex,OutTexturePosition);
+		}
+		)");
+		OutHLSL += FString::Format(FormatSample, ArgsSample);
+		return true;
+	}
+	else if (DefinitionFunctionName == GetProjectionPointName)
+	{
+		static const TCHAR *FormatSample = TEXT(R"(
+		void {InstanceFunctionName}(in float3 WorldPosition, in float DeltaTime, in int ElementIndex, in float TextureValue, in float3 TextureGradient, out float3 OutClosestPosition, 
+							out float3 OutClosestNormal, out float3 OutClosestVelocity, out float OutClosestDistance)
+		{
+			{PhysicsAssetContextName} DIPhysicsAsset_GetProjectionPoint(DIContext,WorldPosition,DeltaTime,ElementIndex,TextureValue,TextureGradient,
+				OutClosestPosition,OutClosestNormal,OutClosestVelocity,OutClosestDistance);
+		}
+		)");
+		OutHLSL += FString::Format(FormatSample, ArgsSample);
+		return true;
+	}
 
 	OutHLSL += TEXT("\n");
 	return false;
@@ -772,6 +916,8 @@ void UNiagaraDataInterfacePhysicsAsset::ProvidePerInstanceDataForRenderThread(vo
 	if (GameThreadData != nullptr && RenderThreadData != nullptr)
 	{
 		RenderThreadData->AssetBuffer = GameThreadData->AssetBuffer;
+		RenderThreadData->BoxOrigin = GameThreadData->BoxOrigin;
+		RenderThreadData->BoxExtent = GameThreadData->BoxExtent;
 	}
 	check(Proxy);
 }
diff --git a/Engine/Plugins/Experimental/HairStrands/Source/HairStrandsNiagara/Private/NiagaraDataInterfacePressureGrid.cpp b/Engine/Plugins/Experimental/HairStrands/Source/HairStrandsNiagara/Private/NiagaraDataInterfacePressureGrid.cpp
--- a/Engine/Plugins/Experimental/HairStrands/Source/HairStrandsNiagara/Private/NiagaraDataInterfacePressureGrid.cpp
+++ b/Engine/Plugins/Experimental/HairStrands/Source/HairStrandsNiagara/Private/NiagaraDataInterfacePressureGrid.cpp
@@ -403,6 +403,8 @@ void UNiagaraDataInterfacePressureGrid::GetFunctions(TArray<FNiagaraFunctionSign
 		Sig.Inputs.Add(FNiagaraVariable(FNiagaraTypeDefinition::GetVec3Def(), TEXT("Velocity GradientX")));
 		Sig.Inputs.Add(FNiagaraVariable(FNiagaraTypeDefinition::GetVec3Def(), TEXT("Velocity GradientY")));
 		Sig.Inputs.Add(FNiagaraVariable(FNiagaraTypeDefinition::GetVec3Def(), TEXT("Velocity GradientZ")));
+		Sig.Inputs.Add(FNiagaraVariable(FNiagaraTypeDefinition::GetVec3Def(), TEXT("Grid Origin")));
+		Sig.Inputs.Add(FNiagaraVariable(FNiagaraTypeDefinition::GetFloatDef(), TEXT("Grid Length")));
 		Sig.Outputs.Add(FNiagaraVariable(FNiagaraTypeDefinition::GetBoolDef(), TEXT("Build Status")));
 
 		OutFunctions.Add(Sig);
@@ -414,6 +416,8 @@ void UNiagaraDataInterfacePressureGrid::GetFunctions(TArray<FNiagaraFunctionSign
 		Sig.bRequiresContext = false;
 		Sig.Inputs.Add(FNiagaraVariable(FNiagaraTypeDefinition(GetClass()), TEXT("Pressure Grid")));
 		Sig.Inputs.Add(FNiagaraVariable(FNiagaraTypeDefinition::GetVec3Def(), TEXT("Node Position")));
+		Sig.Inputs.Add(FNiagaraVariable(FNiagaraTypeDefinition::GetVec3Def(), TEXT("Grid Origin")));
+		Sig.Inputs.Add(FNiagaraVariable(FNiagaraTypeDefinition::GetFloatDef(), TEXT("Grid Length")));
 		Sig.Outputs.Add(FNiagaraVariable(FNiagaraTypeDefinition::GetVec3Def(), TEXT("Node Velocity")));
 		Sig.Outputs.Add(FNiagaraVariable(FNiagaraTypeDefinition::GetFloatDef(), TEXT("Node Density")));
 		Sig.Outputs.Add(FNiagaraVariable(FNiagaraTypeDefinition::GetVec3Def(), TEXT("Velocity GradientX")));
@@ -440,6 +444,8 @@ void UNiagaraDataInterfacePressureGrid::GetFunctions(TArray<FNiagaraFunctionSign
 		Sig.bRequiresContext = false;
 		Sig.Inputs.Add(FNiagaraVariable(FNiagaraTypeDefinition(GetClass()), TEXT("Pressure Grid")));
 		Sig.Inputs.Add(FNiagaraVariable(FNiagaraTypeDefinition::GetIntDef(), TEXT("Grid Cell")));
+		Sig.Inputs.Add(FNiagaraVariable(FNiagaraTypeDefinition::GetVec3Def(), TEXT("Grid Origin")));
+		Sig.Inputs.Add(FNiagaraVariable(FNiagaraTypeDefinition::GetFloatDef(), TEXT("Grid Length")));
 		Sig.Outputs.Add(FNiagaraVariable(FNiagaraTypeDefinition::GetVec3Def(), TEXT("Grid Position")));
 
 		OutFunctions.Add(Sig);
@@ -519,7 +525,7 @@ void UNiagaraDataInterfacePressureGrid::GetVMExternalFunction(const FVMExternalF
 {
 	if (BindingInfo.Name == BuildVelocityFieldName)
 	{
-		check(BindingInfo.GetNumInputs() == 19 && BindingInfo.GetNumOutputs() == 1);
+		check(BindingInfo.GetNumInputs() == 23 && BindingInfo.GetNumOutputs() == 1);
 		NDI_FUNC_BINDER(UNiagaraDataInterfacePressureGrid, BuildVelocityField)::Bind(this, OutFunc);
 	}
 	else if (BindingInfo.Name == ProjectVelocityFieldName)
@@ -529,12 +535,12 @@ void UNiagaraDataInterfacePressureGrid::GetVMExternalFunction(const FVMExternalF
 	}
 	else if (BindingInfo.Name == GetCellPositionName)
 	{
-		check(BindingInfo.GetNumInputs() == 2 && BindingInfo.GetNumOutputs() == 3);
+		check(BindingInfo.GetNumInputs() == 6 && BindingInfo.GetNumOutputs() == 3);
 		NDI_FUNC_BINDER(UNiagaraDataInterfacePressureGrid, GetCellPosition)::Bind(this, OutFunc);
 	}
 	else if (BindingInfo.Name == SampleVelocityFieldName)
 	{
-		check(BindingInfo.GetNumInputs() == 4 && BindingInfo.GetNumOutputs() == 13);
+		check(BindingInfo.GetNumInputs() == 8 && BindingInfo.GetNumOutputs() == 13);
 		NDI_FUNC_BINDER(UNiagaraDataInterfacePressureGrid, SampleVelocityField)::Bind(this, OutFunc);
 	}
 	else if (BindingInfo.Name == SetSolidBoundaryName)
@@ -638,8 +644,6 @@ void UNiagaraDataInterfacePressureGrid::BuildGridTopology(FVectorVMContext& Cont
 		//				GridExtent.X, GridExtent.Y, GridExtent.Z, GridLengths.X, GridLengths.Y, GridLengths.Z, RegularExtent.X, RegularExtent.Y, RegularExtent.Z,
 		//	InstData->GridOrigin.X, InstData->GridOrigin.Y, InstData->GridOrigin.Z, GridCenter.X, GridCenter.Y, GridCenter.Z);
 	}
-
-	
 }
 
 void UNiagaraDataInterfacePressureGrid::UpdateGridTransform(FVectorVMContext& Context)
@@ -716,9 +720,10 @@ bool UNiagaraDataInterfacePressureGrid::GetFunctionHLSL(const FName& DefinitionF
 	if (DefinitionFunctionName == BuildVelocityFieldName)
 	{
 		static const TCHAR *FormatSample = TEXT(R"(
-				void {InstanceFunctionName} (in int StrandsSize, in float3 NodePosition, in float NodeMass, in float3 NodeVelocity, in float3 VelocityGradientX, in float3 VelocityGradientY, in float3 VelocityGradientZ, out bool OutBuildStatus)
+				void {InstanceFunctionName} (in int StrandsSize, in float3 NodePosition, in float NodeMass, in float3 NodeVelocity, in float3 VelocityGradientX, in float3 VelocityGradientY, in float3 VelocityGradientZ, 
+							in float3 GridOrigin, in float GridLength, out bool OutBuildStatus)
 				{
-					{PressureGridContextName} DIPressureGrid_BuildVelocityField(DIContext,StrandsSize,NodePosition,NodeMass,NodeVelocity,VelocityGradientX,VelocityGradientY,VelocityGradientZ,OutBuildStatus);
+					{PressureGridContextName} DIPressureGrid_BuildVelocityField(DIContext,StrandsSize,NodePosition,NodeMass,NodeVelocity,VelocityGradientX,VelocityGradientY,VelocityGradientZ,GridOrigin,GridLength,OutBuildStatus);
 				}
 				)");
 		OutHLSL += FString::Format(FormatSample, ArgsSample);
@@ -727,9 +732,9 @@ bool UNiagaraDataInterfacePressureGrid::GetFunctionHLSL(const FName& DefinitionF
 	else if (DefinitionFunctionName == SampleVelocityFieldName)
 	{
 		static const TCHAR *FormatSample = TEXT(R"(
-				void {InstanceFunctionName} (in float3 NodePosition, out float3 OutGridVelocity, out float OutGridDensity, out float3 OutGridGradientX, out float3 OutGridGradientY, out float3 OutGridGradientZ )
+				void {InstanceFunctionName} (in float3 NodePosition, in float3 GridVelocity, in float GridLength, out float3 OutGridVelocity, out float OutGridDensity, out float3 OutGridGradientX, out float3 OutGridGradientY, out float3 OutGridGradientZ )
 				{
-					{PressureGridContextName} DIPressureGrid_SampleVelocityField(DIContext,NodePosition,OutGridVelocity,OutGridDensity,OutGridGradientX,OutGridGradientY,OutGridGradientZ);
+					{PressureGridContextName} DIPressureGrid_SampleVelocityField(DIContext,NodePosition,GridVelocity,GridLength,OutGridVelocity,OutGridDensity,OutGridGradientX,OutGridGradientY,OutGridGradientZ);
 				}
 				)");
 		OutHLSL += FString::Format(FormatSample, ArgsSample);
@@ -749,9 +754,9 @@ bool UNiagaraDataInterfacePressureGrid::GetFunctionHLSL(const FName& DefinitionF
 	else if (DefinitionFunctionName == GetCellPositionName)
 	{
 		static const TCHAR *FormatSample = TEXT(R"(
-				void {InstanceFunctionName} (in int GridCell, out float3 OutGridPosition)
+				void {InstanceFunctionName} (in int GridCell, in float3 GridOrigin, in float GridLength, out float3 OutGridPosition)
 				{
-					{PressureGridContextName} DIPressureGrid_GetCellPosition(DIContext,GridCell,OutGridPosition);
+					{PressureGridContextName} DIPressureGrid_GetCellPosition(DIContext,GridCell,GridOrigin,GridLength,OutGridPosition);
 				}
 				)");
 		OutHLSL += FString::Format(FormatSample, ArgsSample);
@@ -790,7 +795,17 @@ bool UNiagaraDataInterfacePressureGrid::GetFunctionHLSL(const FName& DefinitionF
 		OutHLSL += FString::Format(FormatSample, ArgsSample);
 		return true;
 	}
-	
+	else if (DefinitionFunctionName == BuildGridTopologyName)
+	{
+		static const TCHAR *FormatSample = TEXT(R"(
+				void {InstanceFunctionName} (in float3 GridCenter, in float3 GridExtent, out float3 OutGridOrigin, out float OutGridLength)
+				{
+					{PressureGridContextName} DIPressureGrid_BuildGridTopology(DIContext,GridCenter,GridExtent,OutGridOrigin,OutGridLength);
+				}
+				)");
+		OutHLSL += FString::Format(FormatSample, ArgsSample);
+		return true;
+	}
 
 	OutHLSL += TEXT("\n");
 	return false;
diff --git a/Engine/Plugins/Experimental/HairStrands/Source/HairStrandsNiagara/Public/NiagaraDataInterfaceHairStrands.h b/Engine/Plugins/Experimental/HairStrands/Source/HairStrandsNiagara/Public/NiagaraDataInterfaceHairStrands.h
--- a/Engine/Plugins/Experimental/HairStrands/Source/HairStrandsNiagara/Public/NiagaraDataInterfaceHairStrands.h
+++ b/Engine/Plugins/Experimental/HairStrands/Source/HairStrandsNiagara/Public/NiagaraDataInterfaceHairStrands.h
@@ -38,6 +38,9 @@ struct FNDIHairStrandsBuffer : public FRenderResource
 	/** Get the resource name */
 	virtual FString GetFriendlyName() const override { return TEXT("FNDIHairStrandsBuffer"); }
 
+	/** Clear the bounding box buffer */
+	void ClearBuffer(FRHICommandList& RHICmdList);
+
 	/** Strand curves point offset buffer */
 	FRWBuffer CurvesOffsetsBuffer;
 
@@ -65,6 +68,12 @@ struct FNDIHairStrandsBuffer : public FRenderResource
 	/**Root barycentric coordinates */
 	FRWBuffer RootBarycentricCoordinatesBuffer;
 
+	/** Bounding Box Buffer*/
+	FRWBuffer BoundingBoxBuffer;
+
+	/** Node Bound Buffer*/
+	FRWBuffer NodeBoundBuffer;
+
 	/** The strand asset datas from which to sample */
 	const FHairStrandsDatas* SourceDatas;
 
@@ -82,7 +91,7 @@ struct FNDIHairStrandsBuffer : public FRenderResource
 struct FNDIHairStrandsData
 {
 	/** Cached World transform. */
-	FMatrix WorldTransform;
+	FTransform WorldTransform;
 
 	/** Number of strands*/
 	int32 NumStrands;
@@ -235,6 +244,9 @@ public:
 	/** Get the bounding box extent */
 	void GetBoxExtent(FVectorVMContext& Context);
 
+	/** Build the groom bounding box */
+	void BuildBoundingBox(FVectorVMContext& Context);
+
 	/** Setup the distance spring material */
 	void SetupDistanceSpringMaterial(FVectorVMContext& Context);
 
@@ -316,6 +328,12 @@ public:
 	/** Name of the curves offsets buffer */
 	static const FString CurvesOffsetsBufferName;
 
+	/** Name of bounding box buffer */
+	static const FString BoundingBoxBufferName;
+
+	/** Name of node bound buffer */
+	static const FString NodeBoundBufferName;
+
 	/** Name of the nodes positions buffer */
 	static const FString RestPositionBufferName;
 
@@ -382,7 +400,7 @@ struct FNDIHairStrandsProxy : public FNiagaraDataInterfaceProxy
 	virtual void ConsumePerInstanceDataFromGameThread(void* PerInstanceData, const FNiagaraSystemInstanceID& Instance) override;
 
 	/** Initialize the Proxy data strands buffer */
-	void InitializePerInstanceData(const FNiagaraSystemInstanceID& SystemInstance, FNDIHairStrandsBuffer* StrandsBuffer, const uint32 NumStrands, const uint8 StrandSize, const FVector& BoxCenter, const FVector& BoxExtent, const FMatrix& WorldTransform);
+	void InitializePerInstanceData(const FNiagaraSystemInstanceID& SystemInstance, FNDIHairStrandsBuffer* StrandsBuffer, const uint32 NumStrands, const uint8 StrandSize, const FVector& BoxCenter, const FVector& BoxExtent, const FTransform& WorldTransform);
 
 	/** Destroy the proxy data if necessary */
 	void DestroyPerInstanceData(NiagaraEmitterInstanceBatcher* Batcher, const FNiagaraSystemInstanceID& SystemInstance);
diff --git a/Engine/Plugins/Experimental/HairStrands/Source/HairStrandsNiagara/Public/NiagaraDataInterfacePhysicsAsset.h b/Engine/Plugins/Experimental/HairStrands/Source/HairStrandsNiagara/Public/NiagaraDataInterfacePhysicsAsset.h
--- a/Engine/Plugins/Experimental/HairStrands/Source/HairStrandsNiagara/Public/NiagaraDataInterfacePhysicsAsset.h
+++ b/Engine/Plugins/Experimental/HairStrands/Source/HairStrandsNiagara/Public/NiagaraDataInterfacePhysicsAsset.h
@@ -31,6 +31,8 @@ struct FNDIPhysicsAssetArrays
 	TArray<FVector4> CurrentTransform;
 	TArray<FVector4> InverseTransform;
 	TArray<FVector4> PreviousTransform;
+	TArray<FVector4> RestTransform;
+	TArray<FVector4> RestInverse;
 	TArray<FVector4> ElementExtent;
 };
 
@@ -64,6 +66,12 @@ struct FNDIPhysicsAssetBuffer : public FRenderResource
 	/** Inverse transform buffer*/
 	FRWBuffer InverseTransformBuffer;
 
+	/** Rest transform buffer */
+	FRWBuffer RestTransformBuffer;
+
+	/** Rest transform buffer */
+	FRWBuffer RestInverseBuffer;
+
 	/** Element extent buffer */
 	FRWBuffer ElementExtentBuffer;
 
@@ -85,6 +93,12 @@ struct FNDIPhysicsAssetData
 {
 	/** Physics asset Gpu buffer */
 	FNDIPhysicsAssetBuffer* AssetBuffer;
+
+	/** Bounding box center */
+	FVector BoxOrigin;
+
+	/** Bounding box extent */
+	FVector BoxExtent;
 };
 
 /** Data Interface for the strand base */
@@ -141,6 +155,12 @@ public:
 	/** Get the closest point */
 	void GetClosestPoint(FVectorVMContext& Context);
 
+	/** Get the closest texture point */
+	void GetTexturePoint(FVectorVMContext& Context);
+
+	/** Get the projection point */
+	void GetProjectionPoint(FVectorVMContext& Context);
+
 	/** Name of element offsets */
 	static const FString ElementOffsetsName;
 
@@ -153,9 +173,21 @@ public:
 	/** Name of the inverse transform buffer */
 	static const FString InverseTransformBufferName;
 
+	/** Name of the rest transform buffer */
+	static const FString RestTransformBufferName;
+
+	/** Name of the rest inverse transform buffer */
+	static const FString RestInverseBufferName;
+
 	/** Name of the element extent buffer */
 	static const FString ElementExtentBufferName;
 
+	/** Init Box Origin */
+	static const FString BoxOriginName;
+
+	/** Init Box extent */
+	static const FString BoxExtentName;
+
 protected:
 	/** Copy one niagara DI to this */
 	virtual bool CopyToInternal(UNiagaraDataInterface* Destination) const override;
@@ -174,7 +206,7 @@ struct FNDIPhysicsAssetProxy : public FNiagaraDataInterfaceProxy
 	virtual void ConsumePerInstanceDataFromGameThread(void* PerInstanceData, const FNiagaraSystemInstanceID& Instance) override;
 
 	/** Initialize the Proxy data strands buffer */
-	void InitializePerInstanceData(const FNiagaraSystemInstanceID& SystemInstance, FNDIPhysicsAssetBuffer* AssetBuffer);
+	void InitializePerInstanceData(const FNiagaraSystemInstanceID& SystemInstance, FNDIPhysicsAssetBuffer* AssetBuffer, const FVector& BoxOrigin, const FVector& BoxExtent);
 
 	/** Destroy the proxy data if necessary */
 	void DestroyPerInstanceData(NiagaraEmitterInstanceBatcher* Batcher, const FNiagaraSystemInstanceID& SystemInstance);
diff --git a/Engine/Shaders/Private/HairStrands/HairStrandsInterpolation.usf b/Engine/Shaders/Private/HairStrands/HairStrandsInterpolation.usf
--- a/Engine/Shaders/Private/HairStrands/HairStrandsInterpolation.usf
+++ b/Engine/Shaders/Private/HairStrands/HairStrandsInterpolation.usf
@@ -164,7 +164,7 @@ void MainCS(uint2 DispatchThreadId : SV_DispatchThreadID)
 				const float3 Offset0 = ComputeDynamicGeometryOffset(GuideIndex0, GuideVertexWeights[K], RestTriangle, DeformedTriangle, RootBarycentric);
 				const float3 Offset1 = ComputeDynamicGeometryOffset(GuideIndex1, GuideVertexWeights[K], RestTriangle, DeformedTriangle, RootBarycentric);
 
-				CurrOffset += lerp(Offset0, Offset1, GuideVertexLerps[K]);
+				CurrOffset += VectorToWorld(lerp(Offset0, Offset1, GuideVertexLerps[K]),DeformedTriangle);
 			}
 			#endif		
 
@@ -174,8 +174,8 @@ void MainCS(uint2 DispatchThreadId : SV_DispatchThreadID)
 			const FHairMeshTriangle RestTriangle		= GetTriangleTransformation(RootIndex, RestPosition0Buffer, RestPosition1Buffer, RestPosition2Buffer, RestPositionOffset);
 			const FHairMeshTriangle DeformedTriangle	= GetTriangleTransformation(RootIndex, DeformedPosition0Buffer, DeformedPosition1Buffer, DeformedPosition2Buffer, DeformedPositionOffset);
 
-			ControlPoint = ToTriangle(ControlPoint,RootBarycentric,RestTriangle) + CurrOffset;
-			ControlPoint = ToWorld(ControlPoint,RootBarycentric,DeformedTriangle);
+			ControlPoint = ToTriangle(ControlPoint,RootBarycentric,RestTriangle);
+			ControlPoint = ToWorld(ControlPoint,RootBarycentric,DeformedTriangle) + CurrOffset;
 
 			OutRenderDeformedPositionBuffer[VertexIndex] = float4(f32tof16(ControlPoint - OutHairPositionOffset), RenderRestPosePositionBuffer[VertexIndex].w);
 
diff --git a/Engine/Shaders/Private/HairStrands/HairStrandsMeshProjectionCommon.ush b/Engine/Shaders/Private/HairStrands/HairStrandsMeshProjectionCommon.ush
--- a/Engine/Shaders/Private/HairStrands/HairStrandsMeshProjectionCommon.ush
+++ b/Engine/Shaders/Private/HairStrands/HairStrandsMeshProjectionCommon.ush
@@ -85,6 +85,21 @@ float3 ToWorld(const float3 P, float3 B, const FHairMeshTriangle T)
 
 }
 
+float3 VectorToWorld(const float3 P, const FHairMeshTriangle T)
+{
+	const float3 Rotation = P.x * T.E0 + P.y * T.E1 + P.z * T.N;
+	return Rotation;
+}
+
+float3 VectorToTriangle(const float3 P, const FHairMeshTriangle T)
+{
+	float3 Out;
+	Out.x = dot(P, T.E0);
+	Out.y = dot(P, T.E1);
+	Out.z = dot(P, T.N);
+	return Out;
+}
+
 float3 TransformPoint(float3 P, float3 B, const FHairMeshTriangle RestTriangle, const FHairMeshTriangle DeformedTriangle)
 {
 	float3 LocalP = ToTriangle(P, B, RestTriangle);
