Fix crash whe  opening BP viewport with niagara enabled

diff --git a/Engine/Plugins/Experimental/HairStrands/Source/HairStrandsNiagara/Private/NiagaraDataInterfaceHairStrands.cpp b/Engine/Plugins/Experimental/HairStrands/Source/HairStrandsNiagara/Private/NiagaraDataInterfaceHairStrands.cpp
--- a/Engine/Plugins/Experimental/HairStrands/Source/HairStrandsNiagara/Private/NiagaraDataInterfaceHairStrands.cpp
+++ b/Engine/Plugins/Experimental/HairStrands/Source/HairStrandsNiagara/Private/NiagaraDataInterfaceHairStrands.cpp
@@ -590,7 +590,6 @@ struct FNDIHairStrandsParametersCS : public FNiagaraDataInterfaceParametersCS
 			static_cast<FNDIHairStrandsProxy*>(Context.DataInterface);
 		FNDIHairStrandsData* ProxyData =
 			InterfaceProxy->SystemInstancesToProxyData.Find(Context.SystemInstance);
-		ensure(ProxyData);
 
 		if (ProxyData != nullptr)
 		{
@@ -770,25 +769,17 @@ void FNDIHairStrandsProxy::DeferredDestroy()
 void FNDIHairStrandsProxy::ConsumePerInstanceDataFromGameThread(void* PerInstanceData, const FNiagaraSystemInstanceID& Instance)
 {
 	FNDIHairStrandsData* SourceData = static_cast<FNDIHairStrandsData*>(PerInstanceData);
-	FNDIHairStrandsData* TargetData = SystemInstancesToProxyData.Find(Instance);
+	FNDIHairStrandsData& TargetData = SystemInstancesToProxyData.FindOrAdd(Instance);
 
-	ensure(TargetData);
-	if (TargetData)
-	{
-		TargetData->WorldTransform = SourceData->WorldTransform;
-		TargetData->NumStrands = SourceData->NumStrands;
-		TargetData->StrandSize = SourceData->StrandSize;
-		TargetData->BoxCenter = SourceData->BoxCenter;
-		TargetData->BoxExtent = SourceData->BoxExtent;
-		TargetData->TickCount = SourceData->TickCount;
-		TargetData->ResetTick = SourceData->ResetTick;
-		TargetData->ForceReset = SourceData->ForceReset;
-		TargetData->HairStrandsBuffer = SourceData->HairStrandsBuffer;
-	}
-	else
-	{
-		UE_LOG(LogHairStrands, Log, TEXT("ConsumePerInstanceDataFromGameThread() ... could not find %s"), *FNiagaraUtilities::SystemInstanceIDToString(Instance));
-	}
+	TargetData.WorldTransform = SourceData->WorldTransform;
+	TargetData.NumStrands = SourceData->NumStrands;
+	TargetData.StrandSize = SourceData->StrandSize;
+	TargetData.BoxCenter = SourceData->BoxCenter;
+	TargetData.BoxExtent = SourceData->BoxExtent;
+	TargetData.TickCount = SourceData->TickCount;
+	TargetData.ResetTick = SourceData->ResetTick;
+	TargetData.ForceReset = SourceData->ForceReset;
+	TargetData.HairStrandsBuffer = SourceData->HairStrandsBuffer;
 }
 
 void FNDIHairStrandsProxy::InitializePerInstanceData(const FNiagaraSystemInstanceID& SystemInstance, FNDIHairStrandsBuffer* HairStrandsBuffer, 
@@ -796,24 +787,16 @@ void FNDIHairStrandsProxy::InitializePerInstanceData(const FNiagaraSystemInstanc
 {
 	check(IsInRenderingThread());
 
-	FNDIHairStrandsData* TargetData = SystemInstancesToProxyData.Find(SystemInstance);
-	if (TargetData != nullptr)
-	{
-		DeferredDestroyList.Remove(SystemInstance);
-	}
-	else
-	{
-		TargetData = &SystemInstancesToProxyData.Add(SystemInstance);
-	}
-	TargetData->HairStrandsBuffer = HairStrandsBuffer;
-	TargetData->NumStrands = NumStrands;
-	TargetData->StrandSize = StrandSize;
-	TargetData->BoxCenter = BoxCenter;
-	TargetData->BoxExtent = BoxExtent;
-	TargetData->TickCount = 0;
-	TargetData->ResetTick = MaxDelay;
-	TargetData->ForceReset = true;
-	TargetData->WorldTransform = WorldTransform;
+	FNDIHairStrandsData& TargetData = SystemInstancesToProxyData.FindOrAdd(SystemInstance);
+	TargetData.HairStrandsBuffer = HairStrandsBuffer;
+	TargetData.NumStrands = NumStrands;
+	TargetData.StrandSize = StrandSize;
+	TargetData.BoxCenter = BoxCenter;
+	TargetData.BoxExtent = BoxExtent;
+	TargetData.TickCount = 0;
+	TargetData.ResetTick = MaxDelay;
+	TargetData.ForceReset = true;
+	TargetData.WorldTransform = WorldTransform;
 }
 
 void FNDIHairStrandsProxy::DestroyPerInstanceData(NiagaraEmitterInstanceBatcher* Batcher, const FNiagaraSystemInstanceID& SystemInstance)
@@ -984,7 +967,7 @@ void UNiagaraDataInterfaceHairStrands::DestroyPerInstanceData(void* PerInstanceD
 			[ThisProxy, HairStrandsBuffer, InstanceID = SystemInstance->GetId(), Batcher = SystemInstance->GetBatcher()](FRHICommandListImmediate& CmdList)
 		{
 			HairStrandsBuffer->ReleaseResource();
-			ThisProxy->DestroyPerInstanceData(Batcher, InstanceID);
+			ThisProxy->SystemInstancesToProxyData.Remove(InstanceID);
 			delete HairStrandsBuffer;
 		}
 		);
diff --git a/Engine/Plugins/Experimental/HairStrands/Source/HairStrandsNiagara/Private/NiagaraDataInterfacePhysicsAsset.cpp b/Engine/Plugins/Experimental/HairStrands/Source/HairStrandsNiagara/Private/NiagaraDataInterfacePhysicsAsset.cpp
--- a/Engine/Plugins/Experimental/HairStrands/Source/HairStrandsNiagara/Private/NiagaraDataInterfacePhysicsAsset.cpp
+++ b/Engine/Plugins/Experimental/HairStrands/Source/HairStrandsNiagara/Private/NiagaraDataInterfacePhysicsAsset.cpp
@@ -393,7 +393,6 @@ struct FNDIPhysicsAssetParametersCS : public FNiagaraDataInterfaceParametersCS
 			static_cast<FNDIPhysicsAssetProxy*>(Context.DataInterface);
 		FNDIPhysicsAssetData* ProxyData =
 			InterfaceProxy->SystemInstancesToProxyData.Find(Context.SystemInstance);
-		ensure(ProxyData);
 
 		if (ProxyData != nullptr)
 		{
@@ -459,37 +458,22 @@ void FNDIPhysicsAssetProxy::DeferredDestroy()
 void FNDIPhysicsAssetProxy::ConsumePerInstanceDataFromGameThread(void* PerInstanceData, const FNiagaraSystemInstanceID& Instance)
 {
 	FNDIPhysicsAssetData* SourceData = static_cast<FNDIPhysicsAssetData*>(PerInstanceData);
-	FNDIPhysicsAssetData* TargetData = SystemInstancesToProxyData.Find(Instance);
+	FNDIPhysicsAssetData& TargetData = SystemInstancesToProxyData.FindOrAdd(Instance);
 
-	ensure(TargetData);
-	if (TargetData)
-	{
-		TargetData->AssetBuffer = SourceData->AssetBuffer;
-		TargetData->BoxOrigin = SourceData->BoxOrigin;
-		TargetData->BoxExtent = SourceData->BoxExtent;
-	}
-	else
-	{
-		UE_LOG(LogPhysicsAsset, Log, TEXT("ConsumePerInstanceDataFromGameThread() ... could not find %d"), Instance);
-	}
+	TargetData.AssetBuffer = SourceData->AssetBuffer;
+	TargetData.BoxOrigin = SourceData->BoxOrigin;
+	TargetData.BoxExtent = SourceData->BoxExtent;
 }
 
 void FNDIPhysicsAssetProxy::InitializePerInstanceData(const FNiagaraSystemInstanceID& SystemInstance, FNDIPhysicsAssetBuffer* AssetBuffer, const FVector& BoxOrigin , const FVector& BoxExtent)
 {
 	check(IsInRenderingThread());
 
-	FNDIPhysicsAssetData* TargetData = SystemInstancesToProxyData.Find(SystemInstance);
-	if (TargetData != nullptr)
-	{
-		DeferredDestroyList.Remove(SystemInstance);
-	}
-	else
-	{
-		TargetData = &SystemInstancesToProxyData.Add(SystemInstance);
-	}
-	TargetData->AssetBuffer = AssetBuffer;
-	TargetData->BoxOrigin = BoxOrigin;
-	TargetData->BoxExtent = BoxExtent;
+	FNDIPhysicsAssetData& TargetData = SystemInstancesToProxyData.FindOrAdd(SystemInstance);
+
+	TargetData.AssetBuffer = AssetBuffer;
+	TargetData.BoxOrigin = BoxOrigin;
+	TargetData.BoxExtent = BoxExtent;
 }
 
 void FNDIPhysicsAssetProxy::DestroyPerInstanceData(NiagaraEmitterInstanceBatcher* Batcher, const FNiagaraSystemInstanceID& SystemInstance)
@@ -605,7 +589,7 @@ void UNiagaraDataInterfacePhysicsAsset::DestroyPerInstanceData(void* PerInstance
 			[ThisProxy, InBuffer,  InstanceID = SystemInstance->GetId(), Batcher = SystemInstance->GetBatcher()](FRHICommandListImmediate& CmdList)
 		{
 			InBuffer->ReleaseResource();
-			ThisProxy->DestroyPerInstanceData(Batcher, InstanceID);
+			ThisProxy->SystemInstancesToProxyData.Remove(InstanceID);
 			delete InBuffer;
 		}
 		);
diff --git a/Engine/Plugins/Experimental/HairStrands/Source/HairStrandsNiagara/Private/NiagaraDataInterfacePressureGrid.cpp b/Engine/Plugins/Experimental/HairStrands/Source/HairStrandsNiagara/Private/NiagaraDataInterfacePressureGrid.cpp
--- a/Engine/Plugins/Experimental/HairStrands/Source/HairStrandsNiagara/Private/NiagaraDataInterfacePressureGrid.cpp
+++ b/Engine/Plugins/Experimental/HairStrands/Source/HairStrandsNiagara/Private/NiagaraDataInterfacePressureGrid.cpp
@@ -148,7 +148,6 @@ struct FNDIPressureGridParametersCS : public FNiagaraDataInterfaceParametersCS
 			static_cast<FNDIPressureGridProxy*>(Context.DataInterface);
 		FNDIPressureGridData* ProxyData =
 			InterfaceProxy->SystemInstancesToProxyData.Find(Context.SystemInstance);
-		ensure(ProxyData);
 
 		if (ProxyData != nullptr)
 		{
@@ -216,18 +215,10 @@ void FNDIPressureGridProxy::DeferredDestroy()
 void FNDIPressureGridProxy::ConsumePerInstanceDataFromGameThread(void* PerInstanceData, const FNiagaraSystemInstanceID& Instance)
 {
 	FNDIPressureGridData* SourceData = static_cast<FNDIPressureGridData*>(PerInstanceData);
-	FNDIPressureGridData* TargetData = SystemInstancesToProxyData.Find(Instance);
+	FNDIPressureGridData& TargetData = SystemInstancesToProxyData.FindOrAdd(Instance);
 
-	ensure(TargetData);
-	if (TargetData)
-	{
-		TargetData->WorldTransform = SourceData->WorldTransform;
-		TargetData->GridOrigin = SourceData->GridOrigin;
-	}
-	else
-	{
-		UE_LOG(LogPressureGrid, Log, TEXT("ConsumePerInstanceDataFromGameThread() ... could not find %s"), *FNiagaraUtilities::SystemInstanceIDToString(Instance));
-	}
+	TargetData.WorldTransform = SourceData->WorldTransform;
+	TargetData.GridOrigin = SourceData->GridOrigin;
 }
 
 void FNDIPressureGridProxy::InitializePerInstanceData(const FNiagaraSystemInstanceID& SystemInstance, FNDIPressureGridBuffer* CurrentGridBuffer,
@@ -235,19 +226,12 @@ FNDIPressureGridBuffer* DestinationGridBuffer, const FVector4& GridOrigin, const
 {
 	check(IsInRenderingThread());
 
-	FNDIPressureGridData* TargetData = SystemInstancesToProxyData.Find(SystemInstance);
-	if (TargetData != nullptr)
-	{
-		DeferredDestroyList.Remove(SystemInstance);
-	}
-	else
-	{
-		TargetData = &SystemInstancesToProxyData.Add(SystemInstance);
-	}
-	TargetData->CurrentGridBuffer = CurrentGridBuffer;
-	TargetData->DestinationGridBuffer = DestinationGridBuffer;
-	TargetData->GridOrigin = GridOrigin;
-	TargetData->GridSize = GridSize;
+	FNDIPressureGridData& TargetData = SystemInstancesToProxyData.FindOrAdd(SystemInstance);
+	
+	TargetData.CurrentGridBuffer = CurrentGridBuffer;
+	TargetData.DestinationGridBuffer = DestinationGridBuffer;
+	TargetData.GridOrigin = GridOrigin;
+	TargetData.GridSize = GridSize;
 }
 
 void FNDIPressureGridProxy::DestroyPerInstanceData(NiagaraEmitterInstanceBatcher* Batcher, const FNiagaraSystemInstanceID& SystemInstance)
@@ -329,7 +313,7 @@ void UNiagaraDataInterfacePressureGrid::DestroyPerInstanceData(void* PerInstance
 		{
 			CurrentGridBuffer->ReleaseResource();
 			DestinationGridBuffer->ReleaseResource();
-			ThisProxy->DestroyPerInstanceData(Batcher, InstanceID);
+			ThisProxy->SystemInstancesToProxyData.Remove(InstanceID);
 			delete CurrentGridBuffer;
 			delete DestinationGridBuffer;
 		}
