Submitting files I've missed in CL 10975676.

diff --git a/Engine/Plugins/Experimental/HairStrands/Source/HairStrandsNiagara/Private/NiagaraDataInterfaceHairStrands.cpp b/Engine/Plugins/Experimental/HairStrands/Source/HairStrandsNiagara/Private/NiagaraDataInterfaceHairStrands.cpp
--- a/Engine/Plugins/Experimental/HairStrands/Source/HairStrandsNiagara/Private/NiagaraDataInterfaceHairStrands.cpp
+++ b/Engine/Plugins/Experimental/HairStrands/Source/HairStrandsNiagara/Private/NiagaraDataInterfaceHairStrands.cpp
@@ -2235,12 +2235,12 @@ void UNiagaraDataInterfaceHairStrands::NeedSimulationReset(FVectorVMContext& Con
 	// @todo : implement function for cpu 
 }
 
-bool UNiagaraDataInterfaceHairStrands::GetFunctionHLSL(const FNiagaraDataInterfaceGPUParamInfo& ParamInfo, const FNiagaraDataInterfaceGeneratedFunction& FunctionInfo, FString& OutHLSL)
+bool UNiagaraDataInterfaceHairStrands::GetFunctionHLSL(const FNiagaraDataInterfaceGPUParamInfo& ParamInfo, const FNiagaraDataInterfaceGeneratedFunction& FunctionInfo, int FunctionInstanceIndex, FString& OutHLSL)
 {
 	FNDIHairStrandsParametersName ParamNames(ParamInfo.DataInterfaceHLSLSymbol);
 
 	TMap<FString, FStringFormatArg> ArgsSample = {
-		{TEXT("InstanceFunctionName"), InstanceFunctionName},
+		{TEXT("InstanceFunctionName"), FunctionInfo.InstanceName},
 		{TEXT("NumStrandsName"), ParamNames.NumStrandsName},
 		{TEXT("StrandSizeName"), ParamNames.StrandSizeName},
 		{TEXT("WorldTransformName"), ParamNames.WorldTransformName},
@@ -2254,7 +2254,7 @@ bool UNiagaraDataInterfaceHairStrands::GetFunctionHLSL(const FNiagaraDataInterfa
 		{TEXT("HairStrandsContextName"), TEXT("DIHAIRSTRANDS_MAKE_CONTEXT(") + ParamInfo.DataInterfaceHLSLSymbol + TEXT(")")},
 	};
 
-	if (DefinitionFunctionName == GetStrandSizeName)
+	if (FunctionInfo.DefinitionName == GetStrandSizeName)
 	{
 		static const TCHAR *FormatSample = TEXT(R"(
 		void {InstanceFunctionName}(out int OutStrandSize)
@@ -2265,7 +2265,7 @@ bool UNiagaraDataInterfaceHairStrands::GetFunctionHLSL(const FNiagaraDataInterfa
 		OutHLSL += FString::Format(FormatSample, ArgsSample);
 		return true;
 	}
-	else if (DefinitionFunctionName == GetNumStrandsName)
+	else if (FunctionInfo.DefinitionName == GetNumStrandsName)
 	{
 		static const TCHAR *FormatSample = TEXT(R"(
 		void {InstanceFunctionName}(out int OutNumStrands)
@@ -2276,7 +2276,7 @@ bool UNiagaraDataInterfaceHairStrands::GetFunctionHLSL(const FNiagaraDataInterfa
 		OutHLSL += FString::Format(FormatSample, ArgsSample);
 		return true;
 	}
-	else if (DefinitionFunctionName == GetWorldTransformName)
+	else if (FunctionInfo.DefinitionName == GetWorldTransformName)
 	{
 		static const TCHAR *FormatSample = TEXT(R"(
 		void {InstanceFunctionName}(out float4x4 OutWorldTransform)
@@ -2287,7 +2287,7 @@ bool UNiagaraDataInterfaceHairStrands::GetFunctionHLSL(const FNiagaraDataInterfa
 		OutHLSL += FString::Format(FormatSample, ArgsSample);
 		return true;
 	}
-	else if (DefinitionFunctionName == GetWorldInverseName)
+	else if (FunctionInfo.DefinitionName == GetWorldInverseName)
 	{
 		static const TCHAR *FormatSample = TEXT(R"(
 		void {InstanceFunctionName}(out float4x4 OutWorldInverse)
@@ -2298,7 +2298,7 @@ bool UNiagaraDataInterfaceHairStrands::GetFunctionHLSL(const FNiagaraDataInterfa
 		OutHLSL += FString::Format(FormatSample, ArgsSample);
 		return true;
 	}
-	else if (DefinitionFunctionName == GetPointPositionName)
+	else if (FunctionInfo.DefinitionName == GetPointPositionName)
 	{
 		static const TCHAR *FormatSample = TEXT(R"(
 			void {InstanceFunctionName} (in int PointIndex, out float3 OutPointPosition)
@@ -2309,7 +2309,7 @@ bool UNiagaraDataInterfaceHairStrands::GetFunctionHLSL(const FNiagaraDataInterfa
 		OutHLSL += FString::Format(FormatSample, ArgsSample);
 		return true;
 	}
-	else if (DefinitionFunctionName == ComputeNodePositionName)
+	else if (FunctionInfo.DefinitionName == ComputeNodePositionName)
 	{
 		static const TCHAR *FormatSample = TEXT(R"(
 			void {InstanceFunctionName} (in int NodeIndex, in float SmoothingFilter, out float3 OutNodePosition)
@@ -2321,7 +2321,7 @@ bool UNiagaraDataInterfaceHairStrands::GetFunctionHLSL(const FNiagaraDataInterfa
 		OutHLSL += FString::Format(FormatSample, ArgsSample);
 		return true;
 	}
-	else if (DefinitionFunctionName == ComputeNodeOrientationName)
+	else if (FunctionInfo.DefinitionName == ComputeNodeOrientationName)
 	{
 		static const TCHAR *FormatSample = TEXT(R"(
 			void {InstanceFunctionName} (in int NodeIndex, in float3 NodePosition, out float4 OutNodeOrientation)
@@ -2332,7 +2332,7 @@ bool UNiagaraDataInterfaceHairStrands::GetFunctionHLSL(const FNiagaraDataInterfa
 		OutHLSL += FString::Format(FormatSample, ArgsSample);
 		return true;
 	}
-	else if (DefinitionFunctionName == ComputeNodeMassName)
+	else if (FunctionInfo.DefinitionName == ComputeNodeMassName)
 	{
 		static const TCHAR *FormatSample = TEXT(R"(
 			void {InstanceFunctionName} (in int NodeIndex, in float StrandsDensity, in float RootThickness, in float TipThickness, out float OutNodeMass)
@@ -2343,7 +2343,7 @@ bool UNiagaraDataInterfaceHairStrands::GetFunctionHLSL(const FNiagaraDataInterfa
 		OutHLSL += FString::Format(FormatSample, ArgsSample);
 		return true;
 	}
-	else if (DefinitionFunctionName == ComputeNodeInertiaName)
+	else if (FunctionInfo.DefinitionName == ComputeNodeInertiaName)
 	{
 		static const TCHAR *FormatSample = TEXT(R"(
 			void {InstanceFunctionName} (in int NodeIndex, in float StrandsDensity, in float RootThickness, in float TipThickness, out float3 OutNodeInertia)
@@ -2354,7 +2354,7 @@ bool UNiagaraDataInterfaceHairStrands::GetFunctionHLSL(const FNiagaraDataInterfa
 		OutHLSL += FString::Format(FormatSample, ArgsSample);
 		return true;
 	}
-	else if (DefinitionFunctionName == ComputeEdgeLengthName)
+	else if (FunctionInfo.DefinitionName == ComputeEdgeLengthName)
 	{
 		static const TCHAR *FormatSample = TEXT(R"(
 			void {InstanceFunctionName} (in int NodeIndex, in float3 NodePosition, in int NodeOffset, out float OutEdgeLength)
@@ -2365,7 +2365,7 @@ bool UNiagaraDataInterfaceHairStrands::GetFunctionHLSL(const FNiagaraDataInterfa
 		OutHLSL += FString::Format(FormatSample, ArgsSample);
 		return true;
 	}
-	else if (DefinitionFunctionName == ComputeEdgeRotationName)
+	else if (FunctionInfo.DefinitionName == ComputeEdgeRotationName)
 	{
 		static const TCHAR *FormatSample = TEXT(R"(
 			void {InstanceFunctionName} (in int NodeIndex, in float4 NodeOrientation, out float4 OutEdgeRotation)
@@ -2376,7 +2376,7 @@ bool UNiagaraDataInterfaceHairStrands::GetFunctionHLSL(const FNiagaraDataInterfa
 		OutHLSL += FString::Format(FormatSample, ArgsSample);
 		return true;
 	}
-	else if (DefinitionFunctionName == ComputeRestPositionName)
+	else if (FunctionInfo.DefinitionName == ComputeRestPositionName)
 	{
 		static const TCHAR *FormatSample = TEXT(R"(
 			void {InstanceFunctionName} (in float3 NodePosition, out float3 OutRestPosition)
@@ -2387,7 +2387,7 @@ bool UNiagaraDataInterfaceHairStrands::GetFunctionHLSL(const FNiagaraDataInterfa
 		OutHLSL += FString::Format(FormatSample, ArgsSample);
 		return true;
 	}
-	else if (DefinitionFunctionName == ComputeRestOrientationName)
+	else if (FunctionInfo.DefinitionName == ComputeRestOrientationName)
 	{
 		static const TCHAR *FormatSample = TEXT(R"(
 			void {InstanceFunctionName} (in float4 NodeOrientation, out float4 OutRestOrientation)
@@ -2398,7 +2398,7 @@ bool UNiagaraDataInterfaceHairStrands::GetFunctionHLSL(const FNiagaraDataInterfa
 		OutHLSL += FString::Format(FormatSample, ArgsSample);
 		return true;
 	}
-	else if (DefinitionFunctionName == ComputeLocalStateName)
+	else if (FunctionInfo.DefinitionName == ComputeLocalStateName)
 	{
 		static const TCHAR *FormatSample = TEXT(R"(
 			void {InstanceFunctionName} (in int NodeIndex, in float3 RestPosition, in float4 RestOrientation, out float3 LocalPosition, out float4 LocalOrientation)
@@ -2409,7 +2409,7 @@ bool UNiagaraDataInterfaceHairStrands::GetFunctionHLSL(const FNiagaraDataInterfa
 		OutHLSL += FString::Format(FormatSample, ArgsSample);
 		return true;
 	}
-	else if (DefinitionFunctionName == AttachNodePositionName)
+	else if (FunctionInfo.DefinitionName == AttachNodePositionName)
 	{
 		static const TCHAR *FormatSample = TEXT(R"(
 			void {InstanceFunctionName} (in float3 RestPosition, out float3 NodePosition)
@@ -2420,7 +2420,7 @@ bool UNiagaraDataInterfaceHairStrands::GetFunctionHLSL(const FNiagaraDataInterfa
 		OutHLSL += FString::Format(FormatSample, ArgsSample);
 		return true;
 	}
-	else if (DefinitionFunctionName == AttachNodeOrientationName)
+	else if (FunctionInfo.DefinitionName == AttachNodeOrientationName)
 	{
 		static const TCHAR *FormatSample = TEXT(R"(
 				void {InstanceFunctionName} (in float4 RestOrientation, out float4 NodeOrientation)
@@ -2431,7 +2431,7 @@ bool UNiagaraDataInterfaceHairStrands::GetFunctionHLSL(const FNiagaraDataInterfa
 		OutHLSL += FString::Format(FormatSample, ArgsSample);
 		return true;
 	}	
-	else if (DefinitionFunctionName == AttachNodeStateName)
+	else if (FunctionInfo.DefinitionName == AttachNodeStateName)
 	{
 		static const TCHAR *FormatSample = TEXT(R"(
 				void {InstanceFunctionName} (in int NodeIndex, in float3 LocalPosition, in float4 LocalOrientation, out float3 NodePosition, out float4 NodeOrientation)
@@ -2442,7 +2442,7 @@ bool UNiagaraDataInterfaceHairStrands::GetFunctionHLSL(const FNiagaraDataInterfa
 		OutHLSL += FString::Format(FormatSample, ArgsSample);
 		return true;
 	}
-	else if (DefinitionFunctionName == UpdatePointPositionName)
+	else if (FunctionInfo.DefinitionName == UpdatePointPositionName)
 	{
 		static const TCHAR *FormatSample = TEXT(R"(
 			void {InstanceFunctionName} (in int NodeIndex, in float3 NodeDisplace, out bool OutReportStatus)
@@ -2453,7 +2453,7 @@ bool UNiagaraDataInterfaceHairStrands::GetFunctionHLSL(const FNiagaraDataInterfa
 		OutHLSL += FString::Format(FormatSample, ArgsSample);
 		return true;
 	}
-	else if (DefinitionFunctionName == ResetPointPositionName)
+	else if (FunctionInfo.DefinitionName == ResetPointPositionName)
 	{
 		static const TCHAR *FormatSample = TEXT(R"(
 			void {InstanceFunctionName} (in int NodeIndex, out bool OutReportStatus)
@@ -2464,7 +2464,7 @@ bool UNiagaraDataInterfaceHairStrands::GetFunctionHLSL(const FNiagaraDataInterfa
 		OutHLSL += FString::Format(FormatSample, ArgsSample);
 		return true;
 	}
-	else if (DefinitionFunctionName == AdvectNodePositionName)
+	else if (FunctionInfo.DefinitionName == AdvectNodePositionName)
 	{
 		static const TCHAR *FormatSample = TEXT(R"(
 			void {InstanceFunctionName} (in float NodeMass, in bool IsPositionMobile, in float3 ExternalForce, in float3 ForceGradient, in float DeltaTime, 
@@ -2478,7 +2478,7 @@ bool UNiagaraDataInterfaceHairStrands::GetFunctionHLSL(const FNiagaraDataInterfa
 		OutHLSL += FString::Format(FormatSample, ArgsSample);
 		return true;
 	}
-	else if (DefinitionFunctionName == AdvectNodeOrientationName)
+	else if (FunctionInfo.DefinitionName == AdvectNodeOrientationName)
 	{
 		static const TCHAR *FormatSample = TEXT(R"(
 			void {InstanceFunctionName} (in float3 NodeInertia, in bool IsOrientationMobile, in float3 ExternalTorque, in float3 TorqueGradient, in float DeltaTime, 
@@ -2492,7 +2492,7 @@ bool UNiagaraDataInterfaceHairStrands::GetFunctionHLSL(const FNiagaraDataInterfa
 		OutHLSL += FString::Format(FormatSample, ArgsSample);
 		return true;
 	}
-	else if (DefinitionFunctionName == UpdateLinearVelocityName)
+	else if (FunctionInfo.DefinitionName == UpdateLinearVelocityName)
 	{
 		static const TCHAR *FormatSample = TEXT(R"(
 			void {InstanceFunctionName} (in float3 PreviousPosition, in float3 NodePosition, in float DeltaTime, out float3 OutLinearVelocity)
@@ -2503,7 +2503,7 @@ bool UNiagaraDataInterfaceHairStrands::GetFunctionHLSL(const FNiagaraDataInterfa
 		OutHLSL += FString::Format(FormatSample, ArgsSample);
 		return true;
 	}
-	else if (DefinitionFunctionName == UpdateAngularVelocityName)
+	else if (FunctionInfo.DefinitionName == UpdateAngularVelocityName)
 	{
 		static const TCHAR *FormatSample = TEXT(R"(
 			void {InstanceFunctionName} (in float4 PreviousOrientation, in float4 NodeOrientation, in float DeltaTime, out float3 OutAngularVelocity)
@@ -2514,7 +2514,7 @@ bool UNiagaraDataInterfaceHairStrands::GetFunctionHLSL(const FNiagaraDataInterfa
 		OutHLSL += FString::Format(FormatSample, ArgsSample);
 		return true;
 	}
-	else if (DefinitionFunctionName == GetBoxExtentName)
+	else if (FunctionInfo.DefinitionName == GetBoxExtentName)
 	{
 		static const TCHAR *FormatSample = TEXT(R"(
 				void {InstanceFunctionName} (out float3 OutBoxExtent)
@@ -2525,7 +2525,7 @@ bool UNiagaraDataInterfaceHairStrands::GetFunctionHLSL(const FNiagaraDataInterfa
 		OutHLSL += FString::Format(FormatSample, ArgsSample);
 		return true;
 	}
-	else if (DefinitionFunctionName == GetBoxCenterName)
+	else if (FunctionInfo.DefinitionName == GetBoxCenterName)
 	{
 		static const TCHAR *FormatSample = TEXT(R"(
 				void {InstanceFunctionName} (out float3 OutBoxCenter)
@@ -2536,7 +2536,7 @@ bool UNiagaraDataInterfaceHairStrands::GetFunctionHLSL(const FNiagaraDataInterfa
 		OutHLSL += FString::Format(FormatSample, ArgsSample);
 		return true;
 	}
-	else if (DefinitionFunctionName == BuildBoundingBoxName)
+	else if (FunctionInfo.DefinitionName == BuildBoundingBoxName)
 	{
 		static const TCHAR *FormatSample = TEXT(R"(
 					void {InstanceFunctionName} (in float3 NodePosition, out bool OutBuildStatus, out float3 OutBoundMin, out float3 OutBoundMax)
@@ -2547,7 +2547,7 @@ bool UNiagaraDataInterfaceHairStrands::GetFunctionHLSL(const FNiagaraDataInterfa
 		OutHLSL += FString::Format(FormatSample, ArgsSample);
 		return true;
 	}
-	else if (DefinitionFunctionName == SetupDistanceSpringMaterialName)
+	else if (FunctionInfo.DefinitionName == SetupDistanceSpringMaterialName)
 	{
 		static const TCHAR *FormatSample = TEXT(R"(
 				void {InstanceFunctionName} (in float YoungModulus, in float RodThickness, 
@@ -2571,7 +2571,7 @@ in float RestLength, in float DeltaTime, in int NodeOffset, in float MaterialDam
 		OutHLSL += FString::Format(FormatSample, ArgsSample);
 		return true;
 	}
-	else if (DefinitionFunctionName == SolveDistanceSpringMaterialName)
+	else if (FunctionInfo.DefinitionName == SolveDistanceSpringMaterialName)
 	{
 		static const TCHAR *FormatSample = TEXT(R"(
 				void {InstanceFunctionName} (in float RestLength, in float DeltaTime, in int NodeOffset, in float MaterialDamping, 
@@ -2595,7 +2595,7 @@ in float RestLength, in float DeltaTime, in int NodeOffset, in float MaterialDam
 		OutHLSL += FString::Format(FormatSample, ArgsSample);
 		return true;
 	}
-	else if (DefinitionFunctionName == ProjectDistanceSpringMaterialName)
+	else if (FunctionInfo.DefinitionName == ProjectDistanceSpringMaterialName)
 	{
 		static const TCHAR *FormatSample = TEXT(R"(
 				void {InstanceFunctionName} (in float YoungModulus, in float RodThickness, in float RestLength, in float NodeMass, in float DeltaTime, in int NodeOffset, in float3 NodePosition, out float3 OutNodePosition)
@@ -2618,7 +2618,7 @@ in float RestLength, in float DeltaTime, in int NodeOffset, in float MaterialDam
 		OutHLSL += FString::Format(FormatSample, ArgsSample);
 		return true;
 	}
-	else if (DefinitionFunctionName == SetupAngularSpringMaterialName)
+	else if (FunctionInfo.DefinitionName == SetupAngularSpringMaterialName)
 	{
 		static const TCHAR *FormatSample = TEXT(R"(
 					void {InstanceFunctionName} (in float YoungModulus, in float RodThickness, 
@@ -2630,7 +2630,7 @@ in float RestLength, in float DeltaTime, in int NodeOffset, in float MaterialDam
 		OutHLSL += FString::Format(FormatSample, ArgsSample);
 		return true;
 	}
-	else if (DefinitionFunctionName == SolveAngularSpringMaterialName)
+	else if (FunctionInfo.DefinitionName == SolveAngularSpringMaterialName)
 	{
 		static const TCHAR *FormatSample = TEXT(R"(
 					void {InstanceFunctionName} (in float RestLength, in float3 RestDirection, in float DeltaTime, in float MaterialDamping, 
@@ -2642,7 +2642,7 @@ in float RestLength, in float DeltaTime, in int NodeOffset, in float MaterialDam
 		OutHLSL += FString::Format(FormatSample, ArgsSample);
 		return true;
 	}
-	else if (DefinitionFunctionName == ProjectAngularSpringMaterialName)
+	else if (FunctionInfo.DefinitionName == ProjectAngularSpringMaterialName)
 	{
 		static const TCHAR *FormatSample = TEXT(R"(
 					void {InstanceFunctionName} (in float YoungModulus, in float RodThickness, in float RestLength, in float3 RestDirection, in float NodeMass, in float DeltaTime, in float3 NodePosition, out float3 OutNodePosition)
@@ -2653,7 +2653,7 @@ in float RestLength, in float DeltaTime, in int NodeOffset, in float MaterialDam
 		OutHLSL += FString::Format(FormatSample, ArgsSample);
 		return true;
 	}
-	else if (DefinitionFunctionName == SetupStretchRodMaterialName)
+	else if (FunctionInfo.DefinitionName == SetupStretchRodMaterialName)
 	{
 		static const TCHAR *FormatSample = TEXT(R"(
 					void {InstanceFunctionName} (in float YoungModulus, in float RodThickness, 
@@ -2665,7 +2665,7 @@ in float RestLength, in float DeltaTime, in int NodeOffset, in float MaterialDam
 		OutHLSL += FString::Format(FormatSample, ArgsSample);
 		return true;
 	}
-	else if (DefinitionFunctionName == SolveStretchRodMaterialName)
+	else if (FunctionInfo.DefinitionName == SolveStretchRodMaterialName)
 	{
 		static const TCHAR *FormatSample = TEXT(R"(
 					void {InstanceFunctionName} (in float RestLength, in float DeltaTime, in float MaterialDamping, 
@@ -2677,7 +2677,7 @@ in float RestLength, in float DeltaTime, in int NodeOffset, in float MaterialDam
 		OutHLSL += FString::Format(FormatSample, ArgsSample);
 		return true;
 	}
-	else if (DefinitionFunctionName == ProjectStretchRodMaterialName)
+	else if (FunctionInfo.DefinitionName == ProjectStretchRodMaterialName)
 	{
 		static const TCHAR *FormatSample = TEXT(R"(
 					void {InstanceFunctionName} (in float YoungModulus, in float RodThickness, in float RestLength, in float DeltaTime, out float3 OutMaterialMultiplier)
@@ -2688,7 +2688,7 @@ in float RestLength, in float DeltaTime, in int NodeOffset, in float MaterialDam
 		OutHLSL += FString::Format(FormatSample, ArgsSample);
 		return true;
 	}
-	else if (DefinitionFunctionName == SetupBendRodMaterialName)
+	else if (FunctionInfo.DefinitionName == SetupBendRodMaterialName)
 	{
 		static const TCHAR *FormatSample = TEXT(R"(
 						void {InstanceFunctionName} (in float YoungModulus, in float RodThickness, 
@@ -2700,7 +2700,7 @@ in float RestLength, in float DeltaTime, in int NodeOffset, in float MaterialDam
 		OutHLSL += FString::Format(FormatSample, ArgsSample);
 		return true;
 	}
-	else if (DefinitionFunctionName == SolveBendRodMaterialName)
+	else if (FunctionInfo.DefinitionName == SolveBendRodMaterialName)
 	{
 		static const TCHAR *FormatSample = TEXT(R"(
 						void {InstanceFunctionName} (in float RestLength, in float4 RestDarboux, in float DeltaTime, in float MaterialDamping, 
@@ -2712,7 +2712,7 @@ in float RestLength, in float DeltaTime, in int NodeOffset, in float MaterialDam
 		OutHLSL += FString::Format(FormatSample, ArgsSample);
 		return true;
 	}
-	else if (DefinitionFunctionName == ProjectBendRodMaterialName)
+	else if (FunctionInfo.DefinitionName == ProjectBendRodMaterialName)
 	{
 		static const TCHAR *FormatSample = TEXT(R"(
 						void {InstanceFunctionName} (in float YoungModulus, in float RodThickness, in float RestLength, in float4 RestDarboux, in float DeltaTime, out float3 OutMaterialMultiplier)
@@ -2723,7 +2723,7 @@ in float RestLength, in float DeltaTime, in int NodeOffset, in float MaterialDam
 		OutHLSL += FString::Format(FormatSample, ArgsSample);
 		return true;
 	}
-	else if (DefinitionFunctionName == SolveStaticCollisionConstraintName)
+	else if (FunctionInfo.DefinitionName == SolveStaticCollisionConstraintName)
 	{
 		static const TCHAR *FormatSample = TEXT(R"(
 							void {InstanceFunctionName} (in float PenetrationDepth, in float3 CollisionPosition, in float3 CollisionVelocity, in float3 CollisionNormal, 
@@ -2737,7 +2737,7 @@ in float RestLength, in float DeltaTime, in int NodeOffset, in float MaterialDam
 		OutHLSL += FString::Format(FormatSample, ArgsSample);
 		return true;
 	}
-	else if (DefinitionFunctionName == ProjectStaticCollisionConstraintName)
+	else if (FunctionInfo.DefinitionName == ProjectStaticCollisionConstraintName)
 	{
 			static const TCHAR *FormatSample = TEXT(R"(
 						void {InstanceFunctionName} (in float PenetrationDepth, in float3 CollisionPosition, in float3 CollisionVelocity, in float3 CollisionNormal, 
@@ -2751,7 +2751,7 @@ in float RestLength, in float DeltaTime, in int NodeOffset, in float MaterialDam
 		OutHLSL += FString::Format(FormatSample, ArgsSample);
 		return true;
 	}
-	else if (DefinitionFunctionName == ComputeRestDirectionName)
+	else if (FunctionInfo.DefinitionName == ComputeRestDirectionName)
 	{
 		static const TCHAR *FormatSample = TEXT(R"(
 							void {InstanceFunctionName} (in float3 NodePosition, in float4 NodeOrientation, out float3 OutRestDirection)
@@ -2762,7 +2762,7 @@ in float RestLength, in float DeltaTime, in int NodeOffset, in float MaterialDam
 		OutHLSL += FString::Format(FormatSample, ArgsSample);
 		return true;
 	}
-	else if (DefinitionFunctionName == UpdateNodeOrientationName)
+	else if (FunctionInfo.DefinitionName == UpdateNodeOrientationName)
 	{
 		static const TCHAR *FormatSample = TEXT(R"(
 					void {InstanceFunctionName} ( out bool OutUpdateStatus)
@@ -2774,7 +2774,7 @@ in float RestLength, in float DeltaTime, in int NodeOffset, in float MaterialDam
 		OutHLSL += FString::Format(FormatSample, ArgsSample);
 		return true;
 	}
-	else if (DefinitionFunctionName == ComputeAirDragForceName)
+	else if (FunctionInfo.DefinitionName == ComputeAirDragForceName)
 	{
 		static const TCHAR *FormatSample = TEXT(R"(
 						void {InstanceFunctionName} ( in float AirDensity, in float AirViscosity, in float AirDrag, 
@@ -2786,7 +2786,7 @@ in float RestLength, in float DeltaTime, in int NodeOffset, in float MaterialDam
 		OutHLSL += FString::Format(FormatSample, ArgsSample);
 		return true;
 	}
-	else if (DefinitionFunctionName == NeedSimulationResetName)
+	else if (FunctionInfo.DefinitionName == NeedSimulationResetName)
 	{
 		static const TCHAR *FormatSample = TEXT(R"(
 				void {InstanceFunctionName} ( out bool ResetSimulation)
diff --git a/Engine/Plugins/Experimental/HairStrands/Source/HairStrandsNiagara/Private/NiagaraDataInterfacePhysicsAsset.cpp b/Engine/Plugins/Experimental/HairStrands/Source/HairStrandsNiagara/Private/NiagaraDataInterfacePhysicsAsset.cpp
--- a/Engine/Plugins/Experimental/HairStrands/Source/HairStrandsNiagara/Private/NiagaraDataInterfacePhysicsAsset.cpp
+++ b/Engine/Plugins/Experimental/HairStrands/Source/HairStrandsNiagara/Private/NiagaraDataInterfacePhysicsAsset.cpp
@@ -790,12 +790,12 @@ void UNiagaraDataInterfacePhysicsAsset::GetProjectionPoint(FVectorVMContext& Con
 {
 }
 
-bool UNiagaraDataInterfacePhysicsAsset::GetFunctionHLSL(const FNiagaraDataInterfaceGPUParamInfo& ParamInfo, const FNiagaraDataInterfaceGeneratedFunction& FunctionInfo, FString& OutHLSL)
+bool UNiagaraDataInterfacePhysicsAsset::GetFunctionHLSL(const FNiagaraDataInterfaceGPUParamInfo& ParamInfo, const FNiagaraDataInterfaceGeneratedFunction& FunctionInfo, int FunctionInstanceIndex, FString& OutHLSL)
 {
 	FNDIPhysicsAssetParametersName ParamNames(ParamInfo.DataInterfaceHLSLSymbol);
 
 	TMap<FString, FStringFormatArg> ArgsSample = {
-		{TEXT("InstanceFunctionName"), InstanceFunctionName},
+		{TEXT("InstanceFunctionName"), FunctionInfo.InstanceName},
 		{TEXT("ElementOffsetsName"), ParamNames.ElementOffsetsName},
 		{TEXT("CurrentTransformBufferName"), ParamNames.CurrentTransformBufferName},
 		{TEXT("PreviousTransformBufferName"), ParamNames.PreviousTransformBufferName},
@@ -804,7 +804,7 @@ bool UNiagaraDataInterfacePhysicsAsset::GetFunctionHLSL(const FNiagaraDataInterf
 		{TEXT("PhysicsAssetContextName"), TEXT("DIPHYSICSASSET_MAKE_CONTEXT(") + ParamInfo.DataInterfaceHLSLSymbol + TEXT(")")},
 	};
 
-	if (DefinitionFunctionName == GetNumBoxesName)
+	if (FunctionInfo.DefinitionName == GetNumBoxesName)
 	{
 		static const TCHAR *FormatSample = TEXT(R"(
 		void {InstanceFunctionName}(out int OutNumBoxes)
@@ -816,7 +816,7 @@ bool UNiagaraDataInterfacePhysicsAsset::GetFunctionHLSL(const FNiagaraDataInterf
 		OutHLSL += FString::Format(FormatSample, ArgsSample);
 		return true;
 	}
-	else if (DefinitionFunctionName == GetNumCapsulesName)
+	else if (FunctionInfo.DefinitionName == GetNumCapsulesName)
 	{
 		static const TCHAR *FormatSample = TEXT(R"(
 		void {InstanceFunctionName}(out int OutNumCapsules)
@@ -828,7 +828,7 @@ bool UNiagaraDataInterfacePhysicsAsset::GetFunctionHLSL(const FNiagaraDataInterf
 		OutHLSL += FString::Format(FormatSample, ArgsSample);
 		return true;
 	}
-	else if (DefinitionFunctionName == GetNumSpheresName)
+	else if (FunctionInfo.DefinitionName == GetNumSpheresName)
 	{
 		static const TCHAR *FormatSample = TEXT(R"(
 		void {InstanceFunctionName}(out int OutNumSpheres)
@@ -840,7 +840,7 @@ bool UNiagaraDataInterfacePhysicsAsset::GetFunctionHLSL(const FNiagaraDataInterf
 		OutHLSL += FString::Format(FormatSample, ArgsSample);
 		return true;
 	}
-	else if (DefinitionFunctionName == GetClosestPointName)
+	else if (FunctionInfo.DefinitionName == GetClosestPointName)
 	{
 		static const TCHAR *FormatSample = TEXT(R"(
 		void {InstanceFunctionName}(in float3 WorldPosition, in float DeltaTime, out float3 OutClosestPosition, 
@@ -853,7 +853,7 @@ bool UNiagaraDataInterfacePhysicsAsset::GetFunctionHLSL(const FNiagaraDataInterf
 		OutHLSL += FString::Format(FormatSample, ArgsSample);
 		return true;
 	}
-	else if (DefinitionFunctionName == GetTexturePointName)
+	else if (FunctionInfo.DefinitionName == GetTexturePointName)
 	{
 		static const TCHAR *FormatSample = TEXT(R"(
 		void {InstanceFunctionName}(in float3 WorldPosition, out int OutElementIndex, out float3 OutTexturePosition)
@@ -864,7 +864,7 @@ bool UNiagaraDataInterfacePhysicsAsset::GetFunctionHLSL(const FNiagaraDataInterf
 		OutHLSL += FString::Format(FormatSample, ArgsSample);
 		return true;
 	}
-	else if (DefinitionFunctionName == GetProjectionPointName)
+	else if (FunctionInfo.DefinitionName == GetProjectionPointName)
 	{
 		static const TCHAR *FormatSample = TEXT(R"(
 		void {InstanceFunctionName}(in float3 WorldPosition, in float DeltaTime, in int ElementIndex, in float TextureValue, in float3 TextureGradient, out float3 OutClosestPosition, 
diff --git a/Engine/Plugins/Experimental/HairStrands/Source/HairStrandsNiagara/Private/NiagaraDataInterfacePressureGrid.cpp b/Engine/Plugins/Experimental/HairStrands/Source/HairStrandsNiagara/Private/NiagaraDataInterfacePressureGrid.cpp
--- a/Engine/Plugins/Experimental/HairStrands/Source/HairStrandsNiagara/Private/NiagaraDataInterfacePressureGrid.cpp
+++ b/Engine/Plugins/Experimental/HairStrands/Source/HairStrandsNiagara/Private/NiagaraDataInterfacePressureGrid.cpp
@@ -686,12 +686,12 @@ void UNiagaraDataInterfacePressureGrid::UpdateGridTransform(FVectorVMContext& Co
 	//UE_LOG(LogPressureGrid, Warning, TEXT("Get Grid Transform : %s"), *InstData->WorldTransform.ToString() );
 }
 
-bool UNiagaraDataInterfacePressureGrid::GetFunctionHLSL(const FNiagaraDataInterfaceGPUParamInfo& ParamInfo, const FNiagaraDataInterfaceGeneratedFunction& FunctionInfo, FString& OutHLSL)
+bool UNiagaraDataInterfacePressureGrid::GetFunctionHLSL(const FNiagaraDataInterfaceGPUParamInfo& ParamInfo, const FNiagaraDataInterfaceGeneratedFunction& FunctionInfo, int FunctionInstanceIndex, FString& OutHLSL)
 {
 	FNDIPressureGridParametersName ParamNames(ParamInfo.DataInterfaceHLSLSymbol);
 
 	TMap<FString, FStringFormatArg> ArgsSample = {
-		{TEXT("InstanceFunctionName"), InstanceFunctionName},
+		{TEXT("InstanceFunctionName"), FunctionInfo.InstanceName},
 		{TEXT("GridCurrentBufferName"), ParamNames.GridCurrentBufferName},
 		{TEXT("GridDestinationBufferName"), ParamNames.GridDestinationBufferName},
 		{TEXT("GridOriginName"), ParamNames.GridOriginName},
@@ -701,7 +701,7 @@ bool UNiagaraDataInterfacePressureGrid::GetFunctionHLSL(const FNiagaraDataInterf
 		{TEXT("PressureGridContextName"), TEXT("DIPRESSUREGRID_MAKE_CONTEXT(") + ParamInfo.DataInterfaceHLSLSymbol + TEXT(")")},
 	};
 
-	if (DefinitionFunctionName == BuildVelocityFieldName)
+	if (FunctionInfo.DefinitionName == BuildVelocityFieldName)
 	{
 		static const TCHAR *FormatSample = TEXT(R"(
 				void {InstanceFunctionName} (in int StrandsSize, in float3 NodePosition, in float NodeMass, in float3 NodeVelocity, in float3 VelocityGradientX, in float3 VelocityGradientY, in float3 VelocityGradientZ, 
@@ -713,7 +713,7 @@ bool UNiagaraDataInterfacePressureGrid::GetFunctionHLSL(const FNiagaraDataInterf
 		OutHLSL += FString::Format(FormatSample, ArgsSample);
 		return true;
 	}
-	else if (DefinitionFunctionName == SampleVelocityFieldName)
+	else if (FunctionInfo.DefinitionName == SampleVelocityFieldName)
 	{
 		static const TCHAR *FormatSample = TEXT(R"(
 				void {InstanceFunctionName} (in float3 NodePosition, in float3 GridVelocity, in float GridLength, out float3 OutGridVelocity, out float OutGridDensity, out float3 OutGridGradientX, out float3 OutGridGradientY, out float3 OutGridGradientZ )
@@ -724,7 +724,7 @@ bool UNiagaraDataInterfacePressureGrid::GetFunctionHLSL(const FNiagaraDataInterf
 		OutHLSL += FString::Format(FormatSample, ArgsSample);
 		return true;
 	}
-	else if (DefinitionFunctionName == ProjectVelocityFieldName)
+	else if (FunctionInfo.DefinitionName == ProjectVelocityFieldName)
 	{
 		static const TCHAR *FormatSample = TEXT(R"(
 				void {InstanceFunctionName} (in int GridCell, out bool OutProjectStatus, out float OutCellDivergence)
@@ -735,7 +735,7 @@ bool UNiagaraDataInterfacePressureGrid::GetFunctionHLSL(const FNiagaraDataInterf
 		OutHLSL += FString::Format(FormatSample, ArgsSample);
 		return true;
 	}
-	else if (DefinitionFunctionName == GetCellPositionName)
+	else if (FunctionInfo.DefinitionName == GetCellPositionName)
 	{
 		static const TCHAR *FormatSample = TEXT(R"(
 				void {InstanceFunctionName} (in int GridCell, in float3 GridOrigin, in float GridLength, out float3 OutGridPosition)
@@ -746,7 +746,7 @@ bool UNiagaraDataInterfacePressureGrid::GetFunctionHLSL(const FNiagaraDataInterf
 		OutHLSL += FString::Format(FormatSample, ArgsSample);
 		return true;
 	}
-	else if (DefinitionFunctionName == SetSolidBoundaryName)
+	else if (FunctionInfo.DefinitionName == SetSolidBoundaryName)
 	{
 		static const TCHAR *FormatSample = TEXT(R"(
 				void {InstanceFunctionName} (in int GridCell, in float SolidDistance, in float3 SolidVelocity, out bool OutBoundaryStatus)
@@ -757,7 +757,7 @@ bool UNiagaraDataInterfacePressureGrid::GetFunctionHLSL(const FNiagaraDataInterf
 		OutHLSL += FString::Format(FormatSample, ArgsSample);
 		return true;
 	}
-	else if (DefinitionFunctionName == ComputeBoundaryWeightsName)
+	else if (FunctionInfo.DefinitionName == ComputeBoundaryWeightsName)
 	{
 		static const TCHAR *FormatSample = TEXT(R"(
 				void {InstanceFunctionName} (in int GridCell, out bool OutWeightsStatus)
@@ -768,7 +768,7 @@ bool UNiagaraDataInterfacePressureGrid::GetFunctionHLSL(const FNiagaraDataInterf
 		OutHLSL += FString::Format(FormatSample, ArgsSample);
 		return true;
 	}
-	else if (DefinitionFunctionName == TransferCellVelocityName)
+	else if (FunctionInfo.DefinitionName == TransferCellVelocityName)
 	{
 		static const TCHAR *FormatSample = TEXT(R"(
 				void {InstanceFunctionName} (in int GridCell, out bool OutTransferStatus)
@@ -779,7 +779,7 @@ bool UNiagaraDataInterfacePressureGrid::GetFunctionHLSL(const FNiagaraDataInterf
 		OutHLSL += FString::Format(FormatSample, ArgsSample);
 		return true;
 	}
-	else if (DefinitionFunctionName == BuildGridTopologyName)
+	else if (FunctionInfo.DefinitionName == BuildGridTopologyName)
 	{
 		static const TCHAR *FormatSample = TEXT(R"(
 				void {InstanceFunctionName} (in float3 GridCenter, in float3 GridExtent, out float3 OutGridOrigin, out float OutGridLength)
