Copying //UE4/Dev-Physics to Main (//UE4/Main)

diff --git a/Engine/Config/BaseScalability.ini b/Engine/Config/BaseScalability.ini
--- a/Engine/Config/BaseScalability.ini
+++ b/Engine/Config/BaseScalability.ini
@@ -569,23 +569,43 @@ grass.DensityScale=1.0
 r.HairStrands.SkyLighting.SampleCount=8
 r.HairStrands.SkyLighting.JitterIntegration=0
 r.HairStrands.SkyAO.SampleCount=8
+r.HairStrands.DeepShadow.SuperSampling=0
+r.HairStrands.VisibilityPPLL=0
+r.HairStrands.RasterizationScale=0.5
+r.HairStrands.VelocityRasterizationScale=1
 
 [ShadingQuality@1]
 r.HairStrands.SkyLighting.SampleCount=16
 r.HairStrands.SkyLighting.JitterIntegration=0
 r.HairStrands.SkyAO.SampleCount=16
+r.HairStrands.DeepShadow.SuperSampling=0
+r.HairStrands.VisibilityPPLL=0
+r.HairStrands.RasterizationScale=0.5
+r.HairStrands.VelocityRasterizationScale=1
 
 [ShadingQuality@2]
 r.HairStrands.SkyLighting.SampleCount=16
 r.HairStrands.SkyLighting.JitterIntegration=0
 r.HairStrands.SkyAO.SampleCount=16
+r.HairStrands.DeepShadow.SuperSampling=0
+r.HairStrands.VisibilityPPLL=0
+r.HairStrands.RasterizationScale=0.5
+r.HairStrands.VelocityRasterizationScale=1
 
 [ShadingQuality@3]
 r.HairStrands.SkyLighting.SampleCount=16
 r.HairStrands.SkyLighting.JitterIntegration=0
 r.HairStrands.SkyAO.SampleCount=16
+r.HairStrands.DeepShadow.SuperSampling=0
+r.HairStrands.VisibilityPPLL=0
+r.HairStrands.RasterizationScale=0.5
+r.HairStrands.VelocityRasterizationScale=1
 
 [ShadingQuality@Cine]
 r.HairStrands.SkyLighting.SampleCount=64
 r.HairStrands.SkyLighting.JitterIntegration=1
-r.HairStrands.SkyAO.SampleCount=32
\ No newline at end of file
+r.HairStrands.SkyAO.SampleCount=32
+r.HairStrands.DeepShadow.SuperSampling=1
+r.HairStrands.RasterizationScale=1
+r.HairStrands.VelocityRasterizationScale=1
+r.HairStrands.VisibilityPPLL=1
\ No newline at end of file
diff --git a/Engine/Plugins/Experimental/AlembicHairImporter/Source/AlembicHairTranslator/Private/AlembicHairTranslator.cpp b/Engine/Plugins/Experimental/AlembicHairImporter/Source/AlembicHairTranslator/Private/AlembicHairTranslator.cpp
--- a/Engine/Plugins/Experimental/AlembicHairImporter/Source/AlembicHairTranslator/Private/AlembicHairTranslator.cpp
+++ b/Engine/Plugins/Experimental/AlembicHairImporter/Source/AlembicHairTranslator/Private/AlembicHairTranslator.cpp
@@ -264,7 +264,7 @@ namespace AlembicHairTranslatorUtils
 			TStrandAttributesRef<AttributeType> StrandAttributeRef = HairDescription.StrandAttributes().GetAttributesRef<AttributeType>(AttributeName);
 			if (!StrandAttributeRef.IsValid())
 			{
-				HairDescription.StrandAttributes().RegisterAttribute<AttributeType>(AttributeName);
+				HairDescription.StrandAttributes().RegisterAttribute<AttributeType>(AttributeName, 1, AttributeType::ZeroVector);
 				StrandAttributeRef = HairDescription.StrandAttributes().GetAttributesRef<AttributeType>(AttributeName);
 			}
 
@@ -345,7 +345,15 @@ namespace AlembicHairTranslatorUtils
 				break;
 				case Alembic::Util::kInt32POD:
 				{
-					ConvertAlembicAttribute<Alembic::AbcGeom::IInt32GeomParam, Alembic::Abc::Int32ArraySamplePtr, int>(HairDescription, StartStrandID, NumStrands, StartVertexID, NumVertices, Parameters, PropName);
+					switch (Extent)
+					{
+					case 1:
+						ConvertAlembicAttribute<Alembic::AbcGeom::IInt32GeomParam, Alembic::Abc::Int32ArraySamplePtr, int>(HairDescription, StartStrandID, NumStrands, StartVertexID, NumVertices, Parameters, PropName);
+						break;
+					case 3:
+						ConvertAlembicAttribute<Alembic::AbcGeom::IV3iGeomParam, Alembic::Abc::V3iArraySamplePtr, FVector>(HairDescription, StartStrandID, NumStrands, StartVertexID, NumVertices, Parameters, PropName, DataType.getExtent());
+						break;
+					}
 				}
 				break;
 				case Alembic::Util::kFloat32POD:
diff --git a/Engine/Plugins/Experimental/HairStrands/Config/BaseHairStrands.ini b/Engine/Plugins/Experimental/HairStrands/Config/BaseHairStrands.ini
--- a/Engine/Plugins/Experimental/HairStrands/Config/BaseHairStrands.ini
+++ b/Engine/Plugins/Experimental/HairStrands/Config/BaseHairStrands.ini
@@ -7,3 +7,5 @@ fx.UseShaderStages=1
 +ClassRedirects=(OldName="HairStrandsActor",NewName="GroomActor")
 +ClassRedirects=(OldName="HairStrandsAsset",NewName="GroomAsset")
 +ClassRedirects=(OldName="HairStrandsComponent",NewName="GroomComponent",ValueChanges=(("HairStrandsComponent0","GroomComponent0")))
++ClassRedirects=(OldName="/Script/HairStrandsEditor.GroomAssetImportData",NewName="/Script/HairStrandsCore.GroomAssetImportData")
++ClassRedirects=(OldName="/Script/HairStrandsEditor.GroomImportOptions",NewName="/Script/HairStrandsCore.GroomImportOptions")
\ No newline at end of file
diff --git a/Engine/Plugins/Experimental/HairStrands/Shaders/Private/NiagaraClearPressureGrid.usf b/Engine/Plugins/Experimental/HairStrands/Shaders/Private/NiagaraClearPressureGrid.usf
new file mode 100644
--- /dev/null
+++ b/Engine/Plugins/Experimental/HairStrands/Shaders/Private/NiagaraClearPressureGrid.usf
@@ -0,0 +1,41 @@
+// Copyright Epic Games, Inc. All Rights Reserved.
+
+static const int FLUID_PRESSURE_OFFSET = 6;
+static const int FLUID_DISTANCE_OFFSET = 14;
+static const int NUM_ATTRIBUTES = 17;
+
+static const float MAX_DISTANCE = 1e+8;
+
+int3 GridSize;
+int CopyPressure;
+
+Texture3D<int>		GridCurrentBuffer;
+RWTexture3D<int>	GridDestinationBuffer;
+
+[numthreads(THREAD_COUNT, 1, 1)]
+void MainCS(uint2 DispatchThreadId : SV_DispatchThreadID)
+{
+	const uint GridHash = DispatchThreadId.x;
+	if (GridHash < (GridSize.x+1)*(GridSize.y+1)*(GridSize.z+1))
+	{
+		const int GridSizeYZ = (GridSize.y+1) * (GridSize.z+1);
+		const int GridHashYZ = GridHash % GridSizeYZ;
+
+		const int3 GridIndex = int3(GridHash / GridSizeYZ, GridHashYZ / (GridSize.z+1), GridHashYZ % (GridSize.z+1));
+
+		[unroll]
+		for(int i = 0; i < NUM_ATTRIBUTES; ++i)
+		{
+			const int3 AttrIndex = int3(GridIndex.x+i*(GridSize.x+1),GridIndex.yz);
+			GridDestinationBuffer[AttrIndex] = 0;
+		}
+		if(CopyPressure)
+		{
+			const int3 PressureIndex = int3(GridIndex.x+FLUID_PRESSURE_OFFSET*(GridSize.x+1),GridIndex.yz);
+			GridDestinationBuffer[PressureIndex] = GridCurrentBuffer[PressureIndex];
+		}
+
+		const int3 DistanceIndex = int3(GridIndex.x+FLUID_DISTANCE_OFFSET*(GridSize.x+1),GridIndex.yz);
+		GridDestinationBuffer[DistanceIndex] = MAX_DISTANCE;
+	}
+}
\ No newline at end of file
diff --git a/Engine/Plugins/Experimental/HairStrands/Shaders/Private/NiagaraCopyBoundingBox.usf b/Engine/Plugins/Experimental/HairStrands/Shaders/Private/NiagaraCopyBoundingBox.usf
--- a/Engine/Plugins/Experimental/HairStrands/Shaders/Private/NiagaraCopyBoundingBox.usf
+++ b/Engine/Plugins/Experimental/HairStrands/Shaders/Private/NiagaraCopyBoundingBox.usf
@@ -1,12 +1,23 @@
 // Copyright Epic Games, Inc. All Rights Reserved.
 
 uint NumElements;
-Buffer<uint> BoundingBoxBuffer;
+RWBuffer<uint> BoundingBoxBuffer;
 RWBuffer<uint> OutNodeBoundBuffer;
 
+uint FloatToUint( in float FloatValue)
+{
+	uint3 UnsignedValue = asuint(FloatValue);
+	UnsignedValue ^= (1+~(UnsignedValue >> 31) | 0x80000000);
+
+	return UnsignedValue;
+}
+
 [numthreads(THREAD_COUNT, 1, 1)]
-void MainCS(uint2 DispatchThreadId : SV_DispatchThreadID)
+void MainCS(uint2 DispatchThreadId : SV_DispatchThreadID)  
 {
+	const uint UINT_MAX = FloatToUint(1e+8);
+	const uint UINT_MIN = FloatToUint(-1e+8);
+
 	const uint ElementIndex = DispatchThreadId.x;
 	if (ElementIndex < NumElements)
 	{
@@ -17,5 +28,13 @@ void MainCS(uint2 DispatchThreadId : SV_DispatchThreadID)
 		OutNodeBoundBuffer[3] = BoundingBoxBuffer[3];
 		OutNodeBoundBuffer[4] = BoundingBoxBuffer[4];
 		OutNodeBoundBuffer[5] = BoundingBoxBuffer[5];
+
+		BoundingBoxBuffer[0] = UINT_MAX;
+		BoundingBoxBuffer[1] = UINT_MAX;
+		BoundingBoxBuffer[2] = UINT_MAX;
+
+		BoundingBoxBuffer[3] = UINT_MIN;
+		BoundingBoxBuffer[4] = UINT_MIN;
+		BoundingBoxBuffer[5] = UINT_MIN;
 	}
 }
\ No newline at end of file
diff --git a/Engine/Plugins/Experimental/HairStrands/Shaders/Private/NiagaraDataInterfaceHairStrands.ush b/Engine/Plugins/Experimental/HairStrands/Shaders/Private/NiagaraDataInterfaceHairStrands.ush
--- a/Engine/Plugins/Experimental/HairStrands/Shaders/Private/NiagaraDataInterfaceHairStrands.ush
+++ b/Engine/Plugins/Experimental/HairStrands/Shaders/Private/NiagaraDataInterfaceHairStrands.ush
@@ -143,7 +143,8 @@ void DIHairStrands_GetBoxExtent(in FDIHairStrandsContext DIContext, out float3 O
 	const float3 BoxMin = Uint3ToFloat3(uint3(DIContext.NodeBoundBuffer[0],DIContext.NodeBoundBuffer[1],DIContext.NodeBoundBuffer[2]));
 	const float3 BoxMax = Uint3ToFloat3(uint3(DIContext.NodeBoundBuffer[3],DIContext.NodeBoundBuffer[4],DIContext.NodeBoundBuffer[5]));
 
-	OutBoxExtent =  (BoxMax-BoxMin);
+	//OutBoxExtent =  0.5 * (BoxMax-BoxMin);
+	OutBoxExtent =   (BoxMax-BoxMin);
 }
 
 /* -----------------------------------------------------------------
@@ -478,6 +479,32 @@ void DIHairStrands_ComputeNodeOrientation(in FDIHairStrandsContext DIContext, in
  * -----------------------------------------------------------------
  */
 
+// Compute the edge volume value
+void DIHairStrands_ComputeEdgeVolume(in FDIHairStrandsContext DIContext, in int NodeIndex, in float3 NodePosition, out float OutEdgeVolume)
+{
+	SharedNodePosition[GGroupThreadId.x] = NodePosition;
+	GroupMemoryBarrier(); 
+
+	int StrandIndex = 0, LocalIndex = 0;
+	DIHairStrands_ComputeNodeIndices(DIContext,NodeIndex,StrandIndex,LocalIndex);
+
+	// L = ||P1-P0||
+	OutEdgeVolume = 0.0;
+	if(LocalIndex > 2)
+	{
+		const int NodeIndexA = GGroupThreadId.x;
+		const int NodeIndexB = GGroupThreadId.x-1;
+		const int NodeIndexC = GGroupThreadId.x-2;
+		const int NodeIndexD = GGroupThreadId.x-3;
+
+		const float3 EdgeVectorA = SharedNodePosition[NodeIndexB] - SharedNodePosition[NodeIndexA];
+		const float3 EdgeVectorB = SharedNodePosition[NodeIndexC] - SharedNodePosition[NodeIndexA];
+		const float3 EdgeVectorC = SharedNodePosition[NodeIndexD] - SharedNodePosition[NodeIndexA];
+
+		OutEdgeVolume = dot(cross(EdgeVectorB,EdgeVectorC),EdgeVectorA);
+	}
+}
+
 // Compute the edge length value
 void DIHairStrands_ComputeEdgeLength(in FDIHairStrandsContext DIContext, in int NodeIndex, in float3 NodePosition, in int NodeOffset, out float OutEdgeLength)
 {
diff --git a/Engine/Plugins/Experimental/HairStrands/Shaders/Private/NiagaraDataInterfacePhysicsAsset.ush b/Engine/Plugins/Experimental/HairStrands/Shaders/Private/NiagaraDataInterfacePhysicsAsset.ush
--- a/Engine/Plugins/Experimental/HairStrands/Shaders/Private/NiagaraDataInterfacePhysicsAsset.ush
+++ b/Engine/Plugins/Experimental/HairStrands/Shaders/Private/NiagaraDataInterfacePhysicsAsset.ush
@@ -24,6 +24,7 @@ struct FDIPhysicsAssetContext
 	Buffer<float4>		RestTransformBuffer;
 	Buffer<float4>		RestInverseBuffer;
 	Buffer<float4>		PreviousTransformBuffer;
+	Buffer<float4>		PreviousInverseBuffer;
 	Buffer<float4>		ElementExtentBuffer;
 	float3				BoxOrigin;
 	float3				BoxExtent;
@@ -36,6 +37,7 @@ Buffer<float4>		InverseTransformBuffer_##NAME;\
 Buffer<float4>		RestTransformBuffer_##NAME;\
 Buffer<float4>		RestInverseBuffer_##NAME;\
 Buffer<float4>		PreviousTransformBuffer_##NAME;\
+Buffer<float4>		PreviousInverseBuffer_##NAME;\
 Buffer<float4>		ElementExtentBuffer_##NAME;\
 float3 BoxOrigin_##NAME;\
 float3 BoxExtent_##NAME;\
@@ -48,6 +50,7 @@ DIContext.InverseTransformBuffer = InverseTransformBuffer_##NAME;\
 DIContext.RestTransformBuffer = RestTransformBuffer_##NAME;\
 DIContext.RestInverseBuffer = RestInverseBuffer_##NAME;\
 DIContext.PreviousTransformBuffer = PreviousTransformBuffer_##NAME;\
+DIContext.PreviousInverseBuffer = PreviousInverseBuffer_##NAME;\
 DIContext.ElementExtentBuffer = ElementExtentBuffer_##NAME;\
 DIContext.BoxOrigin = BoxOrigin_##NAME;\
 DIContext.BoxExtent = BoxExtent_##NAME;\
@@ -108,6 +111,12 @@ float3x4 DIPhysicsAsset_GetPreviousTransform(in FDIPhysicsAssetContext DIContext
     return float3x4(DIContext.PreviousTransformBuffer[BufferOffset], DIContext.PreviousTransformBuffer[BufferOffset+1], DIContext.PreviousTransformBuffer[BufferOffset+2]);
 }
 
+float3x4 DIPhysicsAsset_GetPreviousInverse(in FDIPhysicsAssetContext DIContext, in int ElementIndex)
+{
+	const int BufferOffset = 3 * ElementIndex;
+    return float3x4(DIContext.PreviousInverseBuffer[BufferOffset], DIContext.PreviousInverseBuffer[BufferOffset+1], DIContext.PreviousInverseBuffer[BufferOffset+2]);
+}
+
 float3x4 DIPhysicsAsset_GetRestTransform(in FDIPhysicsAssetContext DIContext, in int ElementIndex)
 {
 	const int BufferOffset = 3 * ElementIndex;
@@ -211,9 +220,17 @@ float DIPhysicsAsset_GetBoxProjection(in float3 LocalPosition, in float3 BoxExte
  * -----------------------------------------------------------------
  */
 
+float3 DIPhysicsAsset_GetLocalPosition(in FDIPhysicsAssetContext DIContext, in float3 NodePosition, in int ElementIndex, in float TimeFraction )
+{
+	const float3 CurrentLocal = mul(DIPhysicsAsset_GetInverseTransform(DIContext,ElementIndex), float4(NodePosition,1.0)).xyz;
+	const float3 PreviousLocal = mul(DIPhysicsAsset_GetPreviousInverse(DIContext,ElementIndex), float4(NodePosition,1.0)).xyz;
+
+	return PreviousLocal + TimeFraction * (CurrentLocal-PreviousLocal);
+}
+
 // Given a world space position (NodePosition) compute the sphere closest point (position,normal,velocity)
 int DIPhysicsAsset_GetClosestElement(in FDIPhysicsAssetContext DIContext, in float3 NodePosition, out float3 OutClosestPosition, 
-										  out float3 OutClosestNormal, out float OutMinDistance )
+										  out float3 OutClosestNormal, out float OutMinDistance, in float TimeFraction )
 {	
 	float MinDistance = MAX_DISTANCE;
 	int ElementIndex = -1;
@@ -224,7 +241,7 @@ int DIPhysicsAsset_GetClosestElement(in FDIPhysicsAssetContext DIContext, in flo
 	const int SpheresEnd = DIContext.ElementOffsets[SPHERE_INDEX+1];
 	for (int SphereIndex = SpheresBegin; SphereIndex < SpheresEnd; ++SphereIndex)
 	{
-		const float3 LocalPosition = mul(DIPhysicsAsset_GetInverseTransform(DIContext,SphereIndex), float4(NodePosition,1.0)).xyz;
+		const float3 LocalPosition = DIPhysicsAsset_GetLocalPosition(DIContext,NodePosition,SphereIndex,TimeFraction);
 
 		DIPhysicsAsset_GetSphereProjection(LocalPosition, float3(0,0,0), DIContext.ElementExtentBuffer[SphereIndex].x, SphereIndex,
 				CollisionPosition, CollisionNormal, ElementIndex, MinDistance);
@@ -234,7 +251,7 @@ int DIPhysicsAsset_GetClosestElement(in FDIPhysicsAssetContext DIContext, in flo
 	const int CapsulesEnd = DIContext.ElementOffsets[CAPSULE_INDEX+1];
 	for (int CapsuleIndex = CapsulesBegin; CapsuleIndex < CapsulesEnd; ++CapsuleIndex)
 	{
-		const float3 LocalPosition = mul(DIPhysicsAsset_GetInverseTransform(DIContext,CapsuleIndex), float4(NodePosition,1.0)).xyz;
+		const float3 LocalPosition = DIPhysicsAsset_GetLocalPosition(DIContext,NodePosition,CapsuleIndex,TimeFraction);
 
 		DIPhysicsAsset_GetCapsuleProjection(LocalPosition, DIContext.ElementExtentBuffer[CapsuleIndex].xy, CapsuleIndex,
 				CollisionPosition, CollisionNormal, ElementIndex, MinDistance);
@@ -244,7 +261,7 @@ int DIPhysicsAsset_GetClosestElement(in FDIPhysicsAssetContext DIContext, in flo
 	const int BoxesEnd = DIContext.ElementOffsets[BOX_INDEX+1];
 	for (int BoxIndex = BoxesBegin; BoxIndex < BoxesEnd; ++BoxIndex)
 	{
-		const float3 LocalPosition = mul(DIPhysicsAsset_GetInverseTransform(DIContext,BoxIndex), float4(NodePosition,1.0)).xyz;
+		const float3 LocalPosition = DIPhysicsAsset_GetLocalPosition(DIContext,NodePosition,BoxIndex,TimeFraction);
 
 		DIPhysicsAsset_GetBoxProjection(LocalPosition, DIContext.ElementExtentBuffer[BoxIndex].xyz, BoxIndex,
 				CollisionPosition, CollisionNormal, ElementIndex, MinDistance);
@@ -261,14 +278,14 @@ int DIPhysicsAsset_GetClosestElement(in FDIPhysicsAssetContext DIContext, in flo
  */
 
 // Given a world space position (NodePosition) compute the sphere closest point (position,normal,velocity)
-void DIPhysicsAsset_GetClosestPoint(in FDIPhysicsAssetContext DIContext, in float3 NodePosition, in float DeltaTime, out float3 OutClosestPosition, 
+void DIPhysicsAsset_GetClosestPoint(in FDIPhysicsAssetContext DIContext, in float3 NodePosition, in float DeltaTime, in float TimeFraction, out float3 OutClosestPosition, 
 										  out float3 OutClosestNormal, out float3 OutClosestVelocity, out float OutClosestDistance)
 {	
 	float3 CollisionPosition = float3(0,0,0);
 	float3 CollisionNormal = float3(0,0,0);
 	float MinDistance = 0.0;
 
-	const int ElementIndex = DIPhysicsAsset_GetClosestElement(DIContext,NodePosition,CollisionPosition,CollisionNormal,MinDistance);
+	const int ElementIndex = DIPhysicsAsset_GetClosestElement(DIContext,NodePosition,CollisionPosition,CollisionNormal,MinDistance,TimeFraction);
 
 	OutClosestNormal = float3(0,0,0);
 	OutClosestPosition = float3(0,0,0);
@@ -278,10 +295,15 @@ void DIPhysicsAsset_GetClosestPoint(in FDIPhysicsAssetContext DIContext, in floa
 	if (ElementIndex != -1)
 	{
 		const float3 PreviousPosition = mul(DIPhysicsAsset_GetPreviousTransform(DIContext,ElementIndex), float4(CollisionPosition,1.0)).xyz;
+		const float3 CurrentPosition = mul(DIPhysicsAsset_GetCurrentTransform(DIContext,ElementIndex), float4(CollisionPosition,1.0)).xyz;
 
-		OutClosestPosition = mul(DIPhysicsAsset_GetCurrentTransform(DIContext,ElementIndex), float4(CollisionPosition,1.0)).xyz;
-		OutClosestNormal = mul(DIPhysicsAsset_GetCurrentTransform(DIContext,ElementIndex), float4(CollisionNormal,0.0)).xyz;
-		OutClosestVelocity = ( OutClosestPosition - PreviousPosition ) / DeltaTime;
+		const float3 PreviousNormal = mul(DIPhysicsAsset_GetPreviousTransform(DIContext,ElementIndex), float4(CollisionNormal,0.0)).xyz;
+		const float3 CurrentNormal = mul(DIPhysicsAsset_GetCurrentTransform(DIContext,ElementIndex), float4(CollisionNormal,0.0)).xyz;
+
+		OutClosestPosition = PreviousPosition + TimeFraction * (CurrentPosition-PreviousPosition);
+		OutClosestNormal = normalize(PreviousNormal + TimeFraction * (CurrentNormal-PreviousNormal));
+
+		OutClosestVelocity = ( CurrentPosition - PreviousPosition ) / DeltaTime;
 		OutClosestDistance = MinDistance;
 	}
 }
@@ -299,7 +321,7 @@ void DIPhysicsAsset_GetTexturePoint(in FDIPhysicsAssetContext DIContext, in floa
 
 	OutTexturePosition = float3(0,0,0);
 
-	OutElementIndex = DIPhysicsAsset_GetClosestElement(DIContext,NodePosition,CollisionPosition,CollisionNormal,MinDistance);
+	OutElementIndex = DIPhysicsAsset_GetClosestElement(DIContext,NodePosition,CollisionPosition,CollisionNormal,MinDistance, 1.0);
 
 	const float AxisExtent = max( max(DIContext.BoxExtent.x, DIContext.BoxExtent.y), DIContext.BoxExtent.z );
 	const float3 BoxExtent = float3(AxisExtent,AxisExtent,AxisExtent);
diff --git a/Engine/Plugins/Experimental/HairStrands/Shaders/Private/NiagaraDataInterfacePressureGrid.ush b/Engine/Plugins/Experimental/HairStrands/Shaders/Private/NiagaraDataInterfacePressureGrid.ush
--- a/Engine/Plugins/Experimental/HairStrands/Shaders/Private/NiagaraDataInterfacePressureGrid.ush
+++ b/Engine/Plugins/Experimental/HairStrands/Shaders/Private/NiagaraDataInterfacePressureGrid.ush
@@ -11,8 +11,10 @@ NiagaraDataInterfacePressureGrid.ush
  * -----------------------------------------------------------------
  */
 
-static const float FLOAT_PRECISION = 1e+8;
-static const bool FAST_MODE = true;
+#define FAST_MODE 1
+#define HAIR_MODE 1
+
+static const float FLOAT_PRECISION = 1e+4;
 
 static const int FLUID_MASS_OFFSET = 0;
 static const int FLUID_VELOCITY_OFFSET = 3;
@@ -20,13 +22,15 @@ static const int FLUID_PRESSURE_OFFSET = 6;
 static const int SOLID_WEIGHT_OFFSET = 7;
 static const int SOLID_VELOCITY_OFFSET = 10;
 static const int SOLID_DISTANCE_OFFSET = 13;
+static const int FLUID_DISTANCE_OFFSET = 14;
+static const int FLUID_DENSITY_OFFSET = 15;
 
 struct FDIPressureGridContext
 {
 	float4x4			WorldTransform;
 	float4x4			WorldInverse;
 	float4				GridOrigin;
-	int4				GridSize;
+	int3				GridSize;
 	Texture3D<int>		GridCurrentBuffer;
 	RWTexture3D<int>	GridDestinationBuffer;
 }; 
@@ -35,7 +39,7 @@ struct FDIPressureGridContext
 float4x4 WorldTransform_##NAME;\
 float4x4 WorldInverse_##NAME;\
 float4 GridOrigin_##NAME;\
-int4 GridSize_##NAME;\
+int3 GridSize_##NAME;\
 Texture3D<int> GridCurrentBuffer_##NAME;\
 RWTexture3D<int> GridDestinationBuffer_##NAME;\
 
@@ -49,170 +53,216 @@ DIContext.GridCurrentBuffer = GridCurrentBuffer_##NAME;\
 DIContext.GridDestinationBuffer = GridDestinationBuffer_##NAME;\
 
 /* -----------------------------------------------------------------
- * Grid utilities
+ * Grid topology conversion
  * -----------------------------------------------------------------
  */
 
-// Compute the grid origin and cell length
-void DIPressureGrid_BuildGridTopology(in FDIPressureGridContext DIContext, in float3 GridCenter, in float3 GridExtent, out float3 OutGridOrigin, out float OutGridLength)
+// Compute the grid origin and grid length
+void DIPressureGrid_BuildGridTopology(in FDIPressureGridContext DIContext, in float3 GridCenter, 
+									  in float3 GridExtent, out float3 OutGridOrigin, out float OutGridLength)
 {
-	const float3 GridLengths = float3(2.0*GridExtent.x/(DIContext.GridSize.x-1),
-									  2.0*GridExtent.y/(DIContext.GridSize.y-1),
-									  2.0*GridExtent.z/(DIContext.GridSize.z-1));
+	const float3 GridLengths = float3(2.0*GridExtent.x/(DIContext.GridSize.x-2),
+									  2.0*GridExtent.y/(DIContext.GridSize.y-2),
+									  2.0*GridExtent.z/(DIContext.GridSize.z-2));
 
 	const float MaxLength = max(GridLengths.x,max(GridLengths.y,GridLengths.z));
 
-	const float3 RegularExtent = float3((DIContext.GridSize.x-1)*MaxLength,
-										(DIContext.GridSize.y-1)*MaxLength,
-										(DIContext.GridSize.z-1)*MaxLength);
+	const float3 RegularSize = float3(DIContext.GridSize.x*MaxLength,
+									  DIContext.GridSize.y*MaxLength,
+									  DIContext.GridSize.z*MaxLength);
 
-	OutGridOrigin = GridCenter - 0.5 * RegularExtent;
+	OutGridOrigin = GridCenter - 0.5 * RegularSize;
 	OutGridLength = MaxLength;
 }
 
-// Return true if the grid index is in the valid space
-bool DIPressureGrid_IsIndexValid(in int3 GridIndex, in int4 ValidSpace)
+// Transform the node world position in grid space
+float3 DIPressureGrid_GetGridPosition(in FDIPressureGridContext DIContext, in float3 NodePosition, in float3 GridOrigin)
 {
-	return ((GridIndex.x < (ValidSpace.x-1)) && (GridIndex.y < (ValidSpace.y-1)) && 
-			(GridIndex.z < (ValidSpace.z-1)) && (GridIndex.x >= 0) && (GridIndex.y >= 0) && (GridIndex.z >= 0));
+	return NodePosition - GridOrigin;
 }
 
-// Return true if the grid cell index is in any voxel center (0,N-2)
-bool DIPressureGrid_IsCellValid(in FDIPressureGridContext DIContext, in int3 GridIndex)
+/* -----------------------------------------------------------------
+ * Grid cell/node/face validity
+ * -----------------------------------------------------------------
+ */
+
+// Return true if the grid cell index is in any voxel center (0,N-1)
+bool DIPressureGrid_IsCellValid(in FDIPressureGridContext DIContext, in int3 CellIndex)
 {
-	return DIPressureGrid_IsIndexValid(GridIndex,DIContext.GridSize);
+	return (CellIndex.x >= 0 && CellIndex.x < DIContext.GridSize.x) && 
+		   (CellIndex.y >= 0 && CellIndex.y < DIContext.GridSize.y) &&
+		   (CellIndex.z >= 0 && CellIndex.z < DIContext.GridSize.z);
 }
 
-// Return true if the grid index is in any voxel face for the component one (0,N-1) and in the voxel center for the others (0,N-2)
-bool DIPressureGrid_IsFaceValid(in FDIPressureGridContext DIContext, in int3 GridIndex, in int ComponentIndex)
+// Return true if the grid face index is for the component one (0,N) and in the voxel center for the others (0,N-1)
+bool DIPressureGrid_IsFaceValid(in FDIPressureGridContext DIContext, in int3 FaceIndex, in int ComponentIndex)
 {
-	int4 ValidSpace = DIContext.GridSize; //ValidSpace[ComponentIndex] += 1;
-	if( ComponentIndex == 0 )
-	{
-		ValidSpace.x += 1;
-	}
-	else if( ComponentIndex == 1 )
-	{
-		ValidSpace.y += 1;
-	}
-	else if( ComponentIndex == 2 )
+	switch(ComponentIndex)
 	{
-		ValidSpace.z += 1;
+		case 0 :
+			return (FaceIndex.x >= 0 && FaceIndex.x <= DIContext.GridSize.x) && 
+			   (FaceIndex.y >= 0 && FaceIndex.y < DIContext.GridSize.y) &&
+			   (FaceIndex.z >= 0 && FaceIndex.z < DIContext.GridSize.z);
+		case 1 : 
+			return (FaceIndex.x >= 0 && FaceIndex.x < DIContext.GridSize.x) && 
+			   (FaceIndex.y >= 0 && FaceIndex.y <= DIContext.GridSize.y) &&
+			   (FaceIndex.z >= 0 && FaceIndex.z < DIContext.GridSize.z);
+		default : 
+			return (FaceIndex.x >= 0 && FaceIndex.x < DIContext.GridSize.x) && 
+			   (FaceIndex.y >= 0 && FaceIndex.y < DIContext.GridSize.y) &&
+			   (FaceIndex.z >= 0 && FaceIndex.z <= DIContext.GridSize.z);
 	}
-	return DIPressureGrid_IsIndexValid(GridIndex,ValidSpace);
 }
 
-// Transform the node world position in grid space
-float3 DIPressureGrid_GetGridPosition(in FDIPressureGridContext DIContext, in float3 NodePosition, in float3 GridOrigin)
+// Return true if the grid node index is valid (0,N)
+bool DIPressureGrid_IsNodeValid(in FDIPressureGridContext DIContext, in int3 NodeIndex)
 {
-	return ( mul(float4( NodePosition, 1.0), DIContext.WorldInverse).xyz - GridOrigin );
+	return (NodeIndex.x >= 0 && NodeIndex.x <= DIContext.GridSize.x) && 
+		   (NodeIndex.y >= 0 && NodeIndex.y <= DIContext.GridSize.y) &&
+		   (NodeIndex.z >= 0 && NodeIndex.z <= DIContext.GridSize.z);
 }
 
-// Compute the grid index and fraction given a grid position
-void DIPressureGrid_GetGridIndex(in FDIPressureGridContext DIContext, in float3 GridPosition, in float3 GridOffset, in float GridLength, out int3 OutGridIndex, out float3 OutGridFraction)
+// Return true if the grid face index is inside the grid
+bool DIPressureGrid_IsFaceInside(in FDIPressureGridContext DIContext, in int3 FaceIndex, in int ComponentIndex)
 {
-	const float3 GridCell = GridPosition / GridLength - GridOffset;
+	switch(ComponentIndex)
+	{
+		case 0 :
+			return (FaceIndex.x >  0 && FaceIndex.x < DIContext.GridSize.x) && 
+			   (FaceIndex.y >= 0 && FaceIndex.y < DIContext.GridSize.y) &&
+			   (FaceIndex.z >= 0 && FaceIndex.z < DIContext.GridSize.z);
+		case 1 :
+			return (FaceIndex.x >= 0 && FaceIndex.x < DIContext.GridSize.x) && 
+			   (FaceIndex.y >  0 && FaceIndex.y < DIContext.GridSize.y) &&
+			   (FaceIndex.z >= 0 && FaceIndex.z < DIContext.GridSize.z);
+		default : 
+			return (FaceIndex.x >= 0 && FaceIndex.x < DIContext.GridSize.x) && 
+			   (FaceIndex.y >= 0 && FaceIndex.y < DIContext.GridSize.y) &&
+			   (FaceIndex.z >  0 && FaceIndex.z < DIContext.GridSize.z);
+	}
+}
 
-	OutGridIndex = floor(GridCell);
-	OutGridFraction = GridCell - OutGridIndex;
+bool DIPressureGrid_IsCellInside(in FDIPressureGridContext DIContext, in int3 CellIndex)
+{
+	return (CellIndex.x > 0 && CellIndex.x < (DIContext.GridSize.x-1)) && 
+		   (CellIndex.y > 0 && CellIndex.y < (DIContext.GridSize.y-1)) &&
+		   (CellIndex.z > 0 && CellIndex.z < (DIContext.GridSize.z-1));
 }
 
-// Compute the data index from agrid index and a data offset
-int3 DIPressureGrid_GetDataIndex( in FDIPressureGridContext DIContext, in int3 GridIndex, in int DataOffset)
+bool DIPressureGrid_IsNodeInside(in FDIPressureGridContext DIContext, in int3 NodeIndex)
 {
-	return int3( GridIndex.x + DataOffset * DIContext.GridSize.x, GridIndex.yz);
+	return (NodeIndex.x > 0 && NodeIndex.x < DIContext.GridSize.x) && 
+		   (NodeIndex.y > 0 && NodeIndex.y < DIContext.GridSize.y) &&
+		   (NodeIndex.z > 0 && NodeIndex.z < DIContext.GridSize.z);
 }
 
 /* -----------------------------------------------------------------
- * Get grid information from hash value
+ * Grid indexing utilities
  * -----------------------------------------------------------------
  */
 
-// Compute the grid index given a grid hash
-int3 DIPressureGrid_ComputeGridIndex(in FDIPressureGridContext DIContext, in int GridCell)
+// Compute the grid index and fraction given a grid position and a grid offset
+int3 DIPressureGrid_ComputeGridIndex(in FDIPressureGridContext DIContext, in float3 GridPosition, in float3 GridOffset, 
+								 in float GridLength, out float3 OutGridFraction)
 {
-	const int GridSizeXY = DIContext.GridSize.y * DIContext.GridSize.z;
-	const int GridCellYZ = GridCell % GridSizeXY;
+	const float3 GridCoordinate = GridPosition / GridLength - GridOffset;
 
-	return int3(GridCell / GridSizeXY, GridCellYZ / DIContext.GridSize.z,GridCellYZ % DIContext.GridSize.z);
+	const int3 GridIndex = floor(GridCoordinate);
+	OutGridFraction = GridCoordinate - GridIndex;
+	return GridIndex;
 }
 
-// Compute the grid hash given a grid index
-int DIPressureGrid_ComputeGridCell(in FDIPressureGridContext DIContext, in int3 GridIndex)
+// Compute the grid index given a grid hash
+int3 DIPressureGrid_GetGridIndex(in FDIPressureGridContext DIContext, in int GridHash)
 {
-	return GridIndex.x * DIContext.GridSize.y * DIContext.GridSize.z + GridIndex.y * DIContext.GridSize.z + GridIndex.z;
-}
+	const int GridSizeYZ = (DIContext.GridSize.y+1) * (DIContext.GridSize.z+1);
+	const int GridHashYZ = GridHash % GridSizeYZ;
 
-// Compute the world grid position
-void DIPressureGrid_GetCellPosition(in FDIPressureGridContext DIContext, in int GridCell, in float3 GridOrigin, in float GridLength, out float3 OutCellPosition)
-{
-	const int3 GridIndex = DIPressureGrid_ComputeGridIndex(DIContext,GridCell);
-	OutCellPosition = mul(float4((GridIndex+float3(0.5,0.5,0.5))*GridLength+GridOrigin,1.0),DIContext.WorldTransform).xyz;
+	return int3(GridHash / GridSizeYZ, GridHashYZ / (DIContext.GridSize.z+1), GridHashYZ % (DIContext.GridSize.z+1));
 }
 
 /* -----------------------------------------------------------------
- * Transfer the velocity and mass buffer from the current to the destination buffer
+ * Get Cell/Node/Face world positions
  * -----------------------------------------------------------------
  */
-
-void DIPressureGrid_TransferGridComponent(in FDIPressureGridContext DIContext, in int3 GridIndex, in int ComponentIndex)
+ 
+ // Compute the world face position
+void DIPressureGrid_GetFacePosition(in FDIPressureGridContext DIContext, in int GridHash, 
+									in float3 GridOrigin, in float GridLength, in int ComponentIndex, out float3 OutFacePosition)
 {
-	const int GridMomentum = DIContext.GridCurrentBuffer[DIPressureGrid_GetDataIndex(DIContext,GridIndex,FLUID_VELOCITY_OFFSET+ComponentIndex)];
-	const int GridMass = DIContext.GridCurrentBuffer[DIPressureGrid_GetDataIndex(DIContext,GridIndex,FLUID_MASS_OFFSET+ComponentIndex)];
+	const int3 GridIndex = DIPressureGrid_GetGridIndex(DIContext,GridHash);
 
-	if(GridMass != 0)
+	if( DIPressureGrid_IsFaceValid(DIContext,GridIndex,ComponentIndex) )
 	{
-		if(FAST_MODE)
-		{
-			const float GridMassFloat = float(GridMass) / FLOAT_PRECISION;
-			const float GridVelocityFloat = ( float(GridMomentum) / FLOAT_PRECISION ) / GridMassFloat;
-
-			DIContext.GridDestinationBuffer[DIPressureGrid_GetDataIndex(DIContext,GridIndex,FLUID_VELOCITY_OFFSET+ComponentIndex)] = asint(GridVelocityFloat);
-			DIContext.GridDestinationBuffer[DIPressureGrid_GetDataIndex(DIContext,GridIndex,FLUID_MASS_OFFSET+ComponentIndex)] = asint(GridMassFloat);
-		}
-		else
-		{
-			DIContext.GridDestinationBuffer[DIPressureGrid_GetDataIndex(DIContext,GridIndex,FLUID_VELOCITY_OFFSET+ComponentIndex)] = asint(asfloat(GridMomentum)/asfloat(GridMass));
-			DIContext.GridDestinationBuffer[DIPressureGrid_GetDataIndex(DIContext,GridIndex,FLUID_MASS_OFFSET+ComponentIndex)] = GridMass;
-		}
+		float3 ComponentOffset = float3(0.5,0.5,0.5); ComponentOffset[ComponentIndex] = 0.0;
+		OutFacePosition = (GridIndex+ComponentOffset)*GridLength+GridOrigin;
 	}
 	else
 	{
-		DIContext.GridDestinationBuffer[DIPressureGrid_GetDataIndex(DIContext,GridIndex,FLUID_VELOCITY_OFFSET+ComponentIndex)] = 0;
-		DIContext.GridDestinationBuffer[DIPressureGrid_GetDataIndex(DIContext,GridIndex,FLUID_MASS_OFFSET+ComponentIndex)] = 0;
+		OutFacePosition = float3(0,0,0);
 	}
-	DeviceMemoryBarrier();
 }
 
-void DIPressureGrid_TransferCellVelocity(in FDIPressureGridContext DIContext, in int GridCell, out bool OutTransferStatus)
+// Compute the world cell position
+void DIPressureGrid_GetCellPosition(in FDIPressureGridContext DIContext, in int GridHash, 
+									in float3 GridOrigin, in float GridLength, out float3 OutCellPosition)
 {
-	const int3 GridIndex = DIPressureGrid_ComputeGridIndex(DIContext,GridCell);
-	OutTransferStatus = false;
-	//if(DIPressureGrid_IsFaceValid(DIContext,GridIndex,0))
+	const int3 GridIndex = DIPressureGrid_GetGridIndex(DIContext,GridHash);
+
+	if( DIPressureGrid_IsCellValid(DIContext,GridIndex) )
+	{
+		OutCellPosition = (GridIndex+float3(0.5,0.5,0.5))*GridLength+GridOrigin;
+	}
+	else
 	{
-		OutTransferStatus = true;
-		DIPressureGrid_TransferGridComponent(DIContext,GridIndex,0);
+		OutCellPosition = float3(0,0,0);
 	}
-	//if(DIPressureGrid_IsFaceValid(DIContext,GridIndex,1))
+}
+
+// Compute the world node position
+void DIPressureGrid_GetNodePosition(in FDIPressureGridContext DIContext, in int GridHash, 
+									in float3 GridOrigin, in float GridLength, out float3 OutNodePosition)
+{
+	const int3 GridIndex = DIPressureGrid_GetGridIndex(DIContext,GridHash);
+
+	if( DIPressureGrid_IsNodeValid(DIContext,GridIndex) )
 	{
-		OutTransferStatus = true;
-		DIPressureGrid_TransferGridComponent(DIContext,GridIndex,1);
+		OutNodePosition = GridIndex*GridLength+GridOrigin;
 	}
-	//if(DIPressureGrid_IsFaceValid(DIContext,GridIndex,2))
+	else
 	{
-		OutTransferStatus = true;
-		DIPressureGrid_TransferGridComponent(DIContext,GridIndex,2);
+		OutNodePosition = float3(0,0,0);
 	}
 }
 
 /* -----------------------------------------------------------------
- * Interlocked Add Float (slow but accurate version) and Int (fast but approximate version)
+ * Interlocked Add/Min Float (slow but accurate version) and Int (fast but approximate version)
  * -----------------------------------------------------------------
  */
 
- void DIPressureGrid_InterlockedAddFloat(in FDIPressureGridContext DIContext, in int3 GridIndex, in int DataOffset, in float FloatValue)
+// Compute the data index from agrid index and a data offset
+int3 DIPressureGrid_GetDataIndex( in FDIPressureGridContext DIContext, in int3 GridIndex, in int DataOffset)
 {
-	/*int IntValue = asint(FloatValue);
+	return int3( GridIndex.x + DataOffset * (DIContext.GridSize.x+1), GridIndex.yz);
+}
+
+#if FAST_MODE == 1
+
+void DIPressureGrid_InterlockedAddInt(in FDIPressureGridContext DIContext, in int3 GridIndex, in int DataOffset, in int IntValue)
+{
+	InterlockedAdd(DIContext.GridDestinationBuffer[DIPressureGrid_GetDataIndex(DIContext,GridIndex,DataOffset)], IntValue);
+}
+
+void DIPressureGrid_InterlockedMinInt(in FDIPressureGridContext DIContext, in int3 GridIndex, in int DataOffset, in int IntValue)
+{
+	InterlockedMin(DIContext.GridDestinationBuffer[DIPressureGrid_GetDataIndex(DIContext,GridIndex,DataOffset)], IntValue);
+}
+
+#else
+
+void DIPressureGrid_InterlockedAddFloat(in FDIPressureGridContext DIContext, in int3 GridIndex, in int DataOffset, in float FloatValue)
+{
+	int IntValue = asint(FloatValue);
 	int CompareValue = 0;
 	int StoredValue = 0;
 
@@ -226,30 +276,64 @@ void DIPressureGrid_TransferCellVelocity(in FDIPressureGridContext DIContext, in
 		}
 		CompareValue = StoredValue;
 		IntValue = asint(FloatValue+asfloat(StoredValue));
-	}*/
+	}
 }
 
-void DIPressureGrid_InterlockedAddInt(in FDIPressureGridContext DIContext, in int3 GridIndex, in int DataOffset, in int IntValue)
+void DIPressureGrid_InterlockedMinFloat(in FDIPressureGridContext DIContext, in int3 GridIndex, in int DataOffset, in float FloatValue)
 {
-	InterlockedAdd(DIContext.GridDestinationBuffer[DIPressureGrid_GetDataIndex(DIContext,GridIndex,DataOffset)], IntValue);
+	int IntValue = asint(FloatValue);
+	int CompareValue = 0;
+	int StoredValue = 0;
+
+	[allow_uav_condition] while(true)
+	{
+		InterlockedCompareExchange(DIContext.GridDestinationBuffer[
+				DIPressureGrid_GetDataIndex(DIContext,GridIndex,DataOffset)],CompareValue,IntValue,StoredValue);
+		if(StoredValue == CompareValue)
+		{
+			break;
+		}
+		CompareValue = StoredValue;
+		IntValue = asint(min(FloatValue,asfloat(StoredValue)));
+	}
 }
 
+#endif
+
 /* -----------------------------------------------------------------
- * Update the velocity and mass buffer
+ * Datas accessors
  * -----------------------------------------------------------------
  */
 
+ void DIPressureGrid_InterlockedMinScalar(in FDIPressureGridContext DIContext, in int3 GridIndex, in int DataOffset, in float FloatValue)
+{
+#if FAST_MODE == 1
+	const int IntValue = FloatValue * FLOAT_PRECISION;
+	DIPressureGrid_InterlockedMinInt(DIContext,GridIndex,DataOffset,IntValue);
+#else
+	DIPressureGrid_InterlockedMinFloat(DIContext,GridIndex,DataOffset,FloatValue);
+#endif
+}
+
 void DIPressureGrid_InterlockedAddScalar(in FDIPressureGridContext DIContext, in int3 GridIndex, in int DataOffset, in float FloatValue)
 {
-	if(FAST_MODE)
-	{
-		const int IntValue = FloatValue * FLOAT_PRECISION;
-		DIPressureGrid_InterlockedAddInt(DIContext,GridIndex,DataOffset,IntValue);
-	}
-	else
-	{
-		DIPressureGrid_InterlockedAddFloat(DIContext,GridIndex,DataOffset,FloatValue);
-	}
+#if FAST_MODE == 1
+	const int IntValue = FloatValue * FLOAT_PRECISION;
+	DIPressureGrid_InterlockedAddInt(DIContext,GridIndex,DataOffset,IntValue);
+#else
+	DIPressureGrid_InterlockedAddFloat(DIContext,GridIndex,DataOffset,FloatValue);
+#endif
+}
+
+float DIPressureGrid_ConvertGridScalar(in FDIPressureGridContext DIContext, in int3 GridIndex, in int DataOffset)
+{
+	const int IntValue = DIContext.GridCurrentBuffer[DIPressureGrid_GetDataIndex(DIContext,GridIndex,DataOffset)];
+
+#if FAST_MODE == 1
+	return float(IntValue) / FLOAT_PRECISION;
+#else
+	return asfloat(IntValue);
+#endif
 }
 
 void DIPressureGrid_AddGridScalar(in FDIPressureGridContext DIContext, in int3 GridIndex, in int DataOffset, in float FloatValue)
@@ -289,36 +373,184 @@ float3 DIPressureGrid_GetGridVector( in FDIPressureGridContext DIContext, in int
 				  DIPressureGrid_GetGridScalar(DIContext,GridIndex,DataOffset+2));
 }
 
+/* -----------------------------------------------------------------
+ * Transfer the velocity, mass and distance buffer from the current to the destination buffer
+ * -----------------------------------------------------------------
+ */
+
+void DIPressureGrid_TransferCellDistance(in FDIPressureGridContext DIContext, in int3 GridIndex, in float GridLength )
+{
+#if HAIR_MODE == 0
+	float FluidDistance = DIPressureGrid_ConvertGridScalar(DIContext,GridIndex,FLUID_DISTANCE_OFFSET);
+
+	[branch]
+	if(FluidDistance < 0.5 * GridLength)
+	{
+		const float SolidDistance = (DIPressureGrid_GetGridScalar(DIContext,int3(GridIndex.x+0,GridIndex.y+0,GridIndex.z+0),SOLID_DISTANCE_OFFSET) + 
+									 DIPressureGrid_GetGridScalar(DIContext,int3(GridIndex.x+0,GridIndex.y+0,GridIndex.z+1),SOLID_DISTANCE_OFFSET) +
+									 DIPressureGrid_GetGridScalar(DIContext,int3(GridIndex.x+0,GridIndex.y+1,GridIndex.z+0),SOLID_DISTANCE_OFFSET) +
+									 DIPressureGrid_GetGridScalar(DIContext,int3(GridIndex.x+0,GridIndex.y+1,GridIndex.z+1),SOLID_DISTANCE_OFFSET) +
+									 DIPressureGrid_GetGridScalar(DIContext,int3(GridIndex.x+1,GridIndex.y+0,GridIndex.z+0),SOLID_DISTANCE_OFFSET) +
+									 DIPressureGrid_GetGridScalar(DIContext,int3(GridIndex.x+1,GridIndex.y+0,GridIndex.z+1),SOLID_DISTANCE_OFFSET) +
+									 DIPressureGrid_GetGridScalar(DIContext,int3(GridIndex.x+1,GridIndex.y+1,GridIndex.z+0),SOLID_DISTANCE_OFFSET) +
+									 DIPressureGrid_GetGridScalar(DIContext,int3(GridIndex.x+1,GridIndex.y+1,GridIndex.z+1),SOLID_DISTANCE_OFFSET) ) / 8.0;
+
+		if(SolidDistance < 0.0) FluidDistance = -0.5 * GridLength;
+	}
+
+	DIPressureGrid_SetGridScalar(DIContext,GridIndex,FLUID_DISTANCE_OFFSET,FluidDistance);
+#endif
+}
+
+
+void DIPressureGrid_TransferFaceVelocity(in FDIPressureGridContext DIContext, in int3 GridIndex, in int ComponentIndex)
+{
+	const float GridMomentum = DIPressureGrid_ConvertGridScalar(DIContext,GridIndex,FLUID_VELOCITY_OFFSET+ComponentIndex);
+	const float GridMass = DIPressureGrid_ConvertGridScalar(DIContext,GridIndex,FLUID_MASS_OFFSET+ComponentIndex);
+
+	float GridVelocity = (GridMass != 0.0) ? GridMomentum / GridMass : 0.0;
+	//if(ComponentIndex == 2) GridVelocity -= 981.0 * 0.033;
+
+	DIPressureGrid_SetGridScalar(DIContext,GridIndex,FLUID_VELOCITY_OFFSET+ComponentIndex,GridVelocity);
+	DIPressureGrid_SetGridScalar(DIContext,GridIndex,FLUID_MASS_OFFSET+ComponentIndex,GridMass);
+}
+
+void DIPressureGrid_TransferCellVelocity(in FDIPressureGridContext DIContext, in int GridHash, in float GridLength, out bool OutTransferStatus)
+{
+	const int3 GridIndex = DIPressureGrid_GetGridIndex(DIContext,GridHash);
+
+	OutTransferStatus = false;
+	if(DIPressureGrid_IsFaceValid(DIContext,GridIndex,0))
+	{
+		DIPressureGrid_TransferFaceVelocity(DIContext,GridIndex,0);
+	}
+	if(DIPressureGrid_IsFaceValid(DIContext,GridIndex,1))
+	{
+		DIPressureGrid_TransferFaceVelocity(DIContext,GridIndex,1);
+	}
+	if(DIPressureGrid_IsFaceValid(DIContext,GridIndex,2))
+	{
+		DIPressureGrid_TransferFaceVelocity(DIContext,GridIndex,2);
+	}
+#if HAIR_MODE == 0
+	if(DIPressureGrid_IsCellValid(DIContext,GridIndex))
+	{
+		DIPressureGrid_TransferCellDistance(DIContext,GridIndex,GridLength);
+	}
+	if(DIPressureGrid_IsCellValid(DIContext,GridIndex))
+	{
+		const float PrevDensity = DIPressureGrid_ConvertGridScalar(DIContext,GridIndex,FLUID_DENSITY_OFFSET);
+		const float NextDensity = DIPressureGrid_ConvertGridScalar(DIContext,GridIndex,FLUID_DENSITY_OFFSET+1);
+
+		const float DensityRatio = (PrevDensity != 0.0) ? GridLength * (1.0-NextDensity / PrevDensity) / 0.033 : 0.0;
+
+		DIPressureGrid_SetGridScalar(DIContext,GridIndex,FLUID_DENSITY_OFFSET,DensityRatio);
+		DIPressureGrid_SetGridScalar(DIContext,GridIndex,FLUID_DENSITY_OFFSET+1,NextDensity);
+	}
+#endif
+}
+
+/* -----------------------------------------------------------------
+ * Interpolation function and gradient
+ * -----------------------------------------------------------------
+ */
+
+ float DIPressureGrid_GetShapeFunction( in float3 GridFraction, in int3 ijk)
+ {
+	const float GridWeightX = (ijk.x == 0) ? 1.0 - GridFraction.x : GridFraction.x;
+	const float GridWeightY = (ijk.y == 0) ? 1.0 - GridFraction.y : GridFraction.y;
+	const float GridWeightZ = (ijk.z == 0) ? 1.0 - GridFraction.z : GridFraction.z;
+
+	return GridWeightX*GridWeightY*GridWeightZ;
+ }
+
+ float3 DIPressureGrid_GetShapeGradient( in float3 GridFraction, in int3 ijk)
+ {
+	const float GridWeightX = (ijk.x == 0) ? 1.0 - GridFraction.x : GridFraction.x;
+	const float GridWeightY = (ijk.y == 0) ? 1.0 - GridFraction.y : GridFraction.y;
+	const float GridWeightZ = (ijk.z == 0) ? 1.0 - GridFraction.z : GridFraction.z;
+
+	const float GridGradientX = (ijk.x == 0) ? -1.0: 1.0;
+	const float GridGradientY = (ijk.y == 0) ? -1.0: 1.0;
+	const float GridGradientZ = (ijk.z == 0) ? -1.0: 1.0;
+
+	return  float3(GridGradientX*GridWeightY*GridWeightZ,
+				   GridWeightX*GridGradientY*GridWeightZ,
+				   GridWeightX*GridWeightY*GridGradientZ);
+ }
+
 /* -----------------------------------------------------------------
  * Build Velocity Field from nodes
  * -----------------------------------------------------------------
  */
+ 
+// Transfer the distance from the node to the grid
+void DIPressureGrid_BuildDistanceField(in FDIPressureGridContext DIContext, in float3 NodePosition, in float3 GridOrigin, in float GridLength)
+{
+#if HAIR_MODE == 0
+	const float3 GridPosition = DIPressureGrid_GetGridPosition(DIContext,NodePosition,GridOrigin);
+
+	float3 GridFraction = float3(0,0,0);
+	const int3 GridIndex = DIPressureGrid_ComputeGridIndex(DIContext,GridPosition,float3(0.5,0.5,0.5),GridLength,GridFraction);
+
+	const float NodeRadius = GridLength * 0.5 * 1.414213;
+
+	for(int i = -2; i <= 2; ++i)
+	for(int j = -2; j <= 2; ++j)
+	for(int k = -2; k <= 2; ++k)
+	{
+		const int3 GridLocation = int3(GridIndex.x+i,GridIndex.y+j,GridIndex.z+k);
+
+		if( DIPressureGrid_IsCellValid(DIContext,GridLocation) )
+		{
+			const float3 CellPosition = (GridLocation+float3(0.5,0.5,0.5))*GridLength+GridOrigin;
+			const float CellDistance = length(CellPosition-NodePosition) - NodeRadius; 
+
+			DIPressureGrid_InterlockedMinScalar(DIContext,GridLocation,FLUID_DISTANCE_OFFSET,CellDistance);
+		}
+	}
+#endif
+}
+ 
+// Transfer the velocity component from the node to the grid
+void DIPressureGrid_BuildVelocityFace(in FDIPressureGridContext DIContext, in int3 GridIndex, in float3 GridFraction, in int3 ijk, in float3 GridPosition, in float SampleMass, in float SampleMomentum,
+			in float3 MomentumGradient, in float3 ComponentOffset, in int ComponentIndex, in float GridLength )
+{
+	const int3 GridLocation = int3(GridIndex.x+ijk.x,GridIndex.y+ijk.y,GridIndex.z+ijk.z);
+
+	if( DIPressureGrid_IsFaceValid(DIContext,GridLocation,ComponentIndex) )
+	{
+		const float ShapeFunction = DIPressureGrid_GetShapeFunction(GridFraction,ijk);
+
+		const float3 DeltaPosition = (GridLocation+ComponentOffset)*GridLength - GridPosition;
+		const float GridMomentum = SampleMomentum + dot(MomentumGradient, DeltaPosition);
+
+		DIPressureGrid_InterlockedAddScalar(DIContext,GridLocation,FLUID_MASS_OFFSET+ComponentIndex,SampleMass*ShapeFunction);
+		DIPressureGrid_InterlockedAddScalar(DIContext,GridLocation,FLUID_VELOCITY_OFFSET+ComponentIndex,GridMomentum*ShapeFunction);
+	}
+}
 
 // Transfer the velocity component from the node to the grid
 void DIPressureGrid_BuildVelocityComponent(in FDIPressureGridContext DIContext, in float3 GridPosition, in float SampleMass, in float SampleMomentum,
 			in float3 MomentumGradient, in float3 ComponentOffset, in int ComponentIndex, in float GridLength )
 {
-	int3 GridIndex = int3(0,0,0); float3 GridFraction = float3(0,0,0);
-	DIPressureGrid_GetGridIndex(DIContext,GridPosition,ComponentOffset,GridLength,GridIndex,GridFraction);
+	float3 GridFraction = float3(0,0,0);
+	const int3 GridIndex = DIPressureGrid_ComputeGridIndex(DIContext,GridPosition,ComponentOffset,GridLength,GridFraction);
 
-	if(DIPressureGrid_IsFaceValid(DIContext,GridIndex,ComponentIndex) && 
-	   DIPressureGrid_IsFaceValid(DIContext,GridIndex+int3(1,1,1),ComponentIndex))
+	for(int i = 0; i < 2; ++i)
 	{
-		[unroll]
-		for(int i = 0; i < 2; ++i)
+		const float GridWeightX = (i == 0) ? 1.0 - GridFraction.x : GridFraction.x;
+		for(int j = 0; j < 2; ++j)
 		{
-			const float GridWeightX = (i == 0) ? 1.0 - GridFraction.x : GridFraction.x;
-			[unroll]
-			for(int j = 0; j < 2; ++j)
+			const float GridWeightY = (j == 0) ? 1.0 - GridFraction.y : GridFraction.y;
+			for(int k = 0; k < 2; ++k)
 			{
-				const float GridWeightY = (j == 0) ? 1.0 - GridFraction.y : GridFraction.y;
-				[unroll]
-				for(int k = 0; k < 2; ++k)
-				{
-					const float GridWeightZ = (k == 0) ? 1.0 - GridFraction.z : GridFraction.z;
-					const float GridWeightXYZ = GridWeightX*GridWeightY*GridWeightZ;
+				const float GridWeightZ = (k == 0) ? 1.0 - GridFraction.z : GridFraction.z;
+				const float GridWeightXYZ = GridWeightX*GridWeightY*GridWeightZ;
 
-					const int3 GridLocation = int3(GridIndex.x+i,GridIndex.y+j,GridIndex.z+k);
+				const int3 GridLocation = int3(GridIndex.x+i,GridIndex.y+j,GridIndex.z+k);
+				if( DIPressureGrid_IsFaceValid(DIContext,GridLocation,ComponentIndex) )
+				{
 					const float3 DeltaPosition = ((GridLocation+ComponentOffset)*GridLength) - GridPosition;
 					const float GridMomentum = SampleMomentum + dot(MomentumGradient, DeltaPosition);
 
@@ -328,47 +560,118 @@ void DIPressureGrid_BuildVelocityComponent(in FDIPressureGridContext DIContext,
 			}
 		}
 	}
+
+	/*DIPressureGrid_BuildVelocityFace(DIContext,GridIndex,GridFraction,int3(0,0,0),GridPosition,SampleMass,
+		SampleMomentum,MomentumGradient,ComponentOffset,ComponentIndex,GridLength);
+	DIPressureGrid_BuildVelocityFace(DIContext,GridIndex,GridFraction,int3(0,0,1),GridPosition,SampleMass,
+		SampleMomentum,MomentumGradient,ComponentOffset,ComponentIndex,GridLength);
+	DIPressureGrid_BuildVelocityFace(DIContext,GridIndex,GridFraction,int3(0,1,0),GridPosition,SampleMass,
+		SampleMomentum,MomentumGradient,ComponentOffset,ComponentIndex,GridLength);
+	DIPressureGrid_BuildVelocityFace(DIContext,GridIndex,GridFraction,int3(0,1,1),GridPosition,SampleMass,
+		SampleMomentum,MomentumGradient,ComponentOffset,ComponentIndex,GridLength);
+	DIPressureGrid_BuildVelocityFace(DIContext,GridIndex,GridFraction,int3(1,0,0),GridPosition,SampleMass,
+		SampleMomentum,MomentumGradient,ComponentOffset,ComponentIndex,GridLength);
+	DIPressureGrid_BuildVelocityFace(DIContext,GridIndex,GridFraction,int3(1,0,1),GridPosition,SampleMass,
+		SampleMomentum,MomentumGradient,ComponentOffset,ComponentIndex,GridLength);
+	DIPressureGrid_BuildVelocityFace(DIContext,GridIndex,GridFraction,int3(1,1,0),GridPosition,SampleMass,
+		SampleMomentum,MomentumGradient,ComponentOffset,ComponentIndex,GridLength);
+	DIPressureGrid_BuildVelocityFace(DIContext,GridIndex,GridFraction,int3(1,1,1),GridPosition,SampleMass,
+		SampleMomentum,MomentumGradient,ComponentOffset,ComponentIndex,GridLength);*/
 }
 
 // Transfer the velocity from the node to the grid
-void DIPressureGrid_BuildVelocityField(in FDIPressureGridContext DIContext, in int StrandsSize, in float3 NodePosition, in float NodeMass, in float3 NodeVelocity, 
-			in float3 VelocityGradientX, in float3 VelocityGradientY, in float3 VelocityGradientZ, in float3 GridOrigin, in float GridLength, out bool OutBuildStatus)
+void DIPressureGrid_BuildSampleVelocity(in FDIPressureGridContext DIContext, in float3 NodePosition, in float NodeMass, in float3 NodeVelocity, 
+			in float3 VelocityGradientX, in float3 VelocityGradientY, in float3 VelocityGradientZ, in float3 GridOrigin, in float GridLength)
 {
-	SharedNodePosition[GGroupThreadId.x] = NodePosition;
-	SharedPreviousPosition[GGroupThreadId.x] = NodeVelocity;
-	SharedInverseMass[GGroupThreadId.x] = NodeMass;
-	GroupMemoryBarrier();
+	const float3 NodeMomentum = NodeVelocity * NodeMass;
+
+	const float3 MomentumGradientX = VelocityGradientX * NodeMass;
+	const float3 MomentumGradientY = VelocityGradientY * NodeMass;
+	const float3 MomentumGradientZ = VelocityGradientZ * NodeMass;
+
+	const float3 GridPosition = DIPressureGrid_GetGridPosition(DIContext,NodePosition,GridOrigin);
+
+	DIPressureGrid_BuildVelocityComponent(DIContext,GridPosition,NodeMass,NodeMomentum.x,MomentumGradientX,float3(0.0,0.5,0.5),0,GridLength);
+	DIPressureGrid_BuildVelocityComponent(DIContext,GridPosition,NodeMass,NodeMomentum.y,MomentumGradientY,float3(0.5,0.0,0.5),1,GridLength);
+	DIPressureGrid_BuildVelocityComponent(DIContext,GridPosition,NodeMass,NodeMomentum.z,MomentumGradientZ,float3(0.5,0.5,0.0),2,GridLength);
+
+#if HAIR_MODE == 0
+	DIPressureGrid_BuildDistanceField(DIContext,NodePosition,GridOrigin,GridLength);
+
+	const float3 PreviousPosition = NodePosition - NodeVelocity * 0.033;
+	const float3 GridPositionP = DIPressureGrid_GetGridPosition(DIContext,PreviousPosition,GridOrigin);
 
-	const int LocalIndex = GGroupThreadId.x % StrandsSize;
+	float3 GridFractionP = float3(0,0,0);
+	const int3 GridIndexP = DIPressureGrid_ComputeGridIndex(DIContext,GridPositionP,float3(0.5,0.5,0.5),GridLength,GridFractionP);
 
-	if( LocalIndex > 0 )
+	float3 GridFractionN = float3(0,0,0);
+	const int3 GridIndexN = DIPressureGrid_ComputeGridIndex(DIContext,GridPosition,float3(0.5,0.5,0.5),GridLength,GridFractionN);
+
+	for( int i = 0; i < 2; ++i)
+	{
+	for( int j = 0; j < 2; ++j)
+	{
+	for( int k = 0; k < 2; ++k)
 	{
-		const float3 DeltaVelocity = SharedPreviousPosition[GGroupThreadId.x] - SharedPreviousPosition[GGroupThreadId.x-1];
-		const float3 DeltaPosition = SharedNodePosition[GGroupThreadId.x] - SharedNodePosition[GGroupThreadId.x-1];
+		const int3 ijk = int3(i,j,k);
+		const int3 GridLocationP = int3(GridIndexP.x+ijk.x,GridIndexP.y+ijk.y,GridIndexP.z+ijk.z);
 
-		const float2 SegmentWeight = (LocalIndex == 1) ? float2(1.0,0.5) : (LocalIndex == (StrandsSize-1) ) ? float2(0.5,1.0) : float2(0.5,0.5);
-		const float SegmentMass =  (SharedInverseMass[GGroupThreadId.x-1] * SegmentWeight.x + SharedInverseMass[GGroupThreadId.x] * SegmentWeight.y);
-		const float SegmentLength = length(DeltaPosition);
+		if( DIPressureGrid_IsCellValid(DIContext,GridLocationP) )
+		{
+			const float ShapeFunction = DIPressureGrid_GetShapeFunction(GridFractionP,ijk);
+			DIPressureGrid_InterlockedAddScalar(DIContext,GridLocationP,FLUID_DENSITY_OFFSET,ShapeFunction);
+		}
+		const int3 GridLocationN = int3(GridIndexN.x+ijk.x,GridIndexN.y+ijk.y,GridIndexN.z+ijk.z);
 
-		const int NumSamples = ceil(SegmentLength / GridLength);
+		if( DIPressureGrid_IsCellValid(DIContext,GridLocationN) )
+		{
+			const float ShapeFunction = DIPressureGrid_GetShapeFunction(GridFractionN,ijk);
+			DIPressureGrid_InterlockedAddScalar(DIContext,GridLocationN,FLUID_DENSITY_OFFSET+1,ShapeFunction);
+		}
+	}
+	}
+	}
+#endif
+}
 
-		const float SampleMass = SegmentMass / NumSamples;
+// Transfer the velocity from the node to the grid
+void DIPressureGrid_BuildVelocityField(in FDIPressureGridContext DIContext, in int StrandsSize, in float3 NodePosition, in float NodeMass, in float3 NodeVelocity, 
+			in float3 VelocityGradientX, in float3 VelocityGradientY, in float3 VelocityGradientZ, in float3 GridOrigin, in float GridLength, out bool OutBuildStatus)
+{
+	if(StrandsSize == 1)
+	{
+		DIPressureGrid_BuildSampleVelocity(DIContext,NodePosition,NodeMass,NodeVelocity,VelocityGradientX,VelocityGradientY,VelocityGradientZ,GridOrigin,GridLength);
+	}
+	else
+	{
+		SharedNodePosition[GGroupThreadId.x] = NodePosition;
+		SharedPreviousPosition[GGroupThreadId.x] = NodeVelocity;
+		SharedInverseMass[GGroupThreadId.x] = NodeMass;
+		GroupMemoryBarrier();
 
-		const float3 MomentumGradientX = VelocityGradientX * SampleMass;
-		const float3 MomentumGradientY = VelocityGradientY * SampleMass;
-		const float3 MomentumGradientZ = VelocityGradientZ * SampleMass;
+		const int LocalIndex = GGroupThreadId.x % StrandsSize;
 
-		for( int SampleIndex = 0; SampleIndex < NumSamples; ++SampleIndex)
+		if( LocalIndex > 0 )
 		{
-			const float SampleCoord = (0.5+SampleIndex) / NumSamples;
-			const float3 SamplePosition = SharedNodePosition[GGroupThreadId.x] + SampleCoord * DeltaPosition;
-			const float3 SampleMomentum = (SharedPreviousPosition[GGroupThreadId.x] + SampleCoord * DeltaVelocity) * SampleMass;
+			const float3 DeltaVelocity = SharedPreviousPosition[GGroupThreadId.x] - SharedPreviousPosition[GGroupThreadId.x-1];
+			const float3 DeltaPosition = SharedNodePosition[GGroupThreadId.x] - SharedNodePosition[GGroupThreadId.x-1];
+
+			const float2 SegmentWeight = (LocalIndex == 1) ? float2(1.0,0.5) : (LocalIndex == (StrandsSize-1) ) ? float2(0.5,1.0) : float2(0.5,0.5);
+			const float SegmentMass =  (SharedInverseMass[GGroupThreadId.x-1] * SegmentWeight.x + SharedInverseMass[GGroupThreadId.x] * SegmentWeight.y);
+			const float SegmentLength = length(DeltaPosition);
+
+			const int NumSamples = ceil(SegmentLength / GridLength);
+			const float SampleMass = SegmentMass / NumSamples;
 
-			const float3 GridPosition = DIPressureGrid_GetGridPosition(DIContext,SamplePosition,GridOrigin);
+			for( int SampleIndex = 0; SampleIndex < NumSamples; ++SampleIndex)
+			{
+				const float SampleCoord = (0.5+SampleIndex) / NumSamples;
+				const float3 SamplePosition = SharedNodePosition[GGroupThreadId.x] + SampleCoord * DeltaPosition;
+				const float3 SampleVelocity = (SharedPreviousPosition[GGroupThreadId.x] + SampleCoord * DeltaVelocity);
 
-			DIPressureGrid_BuildVelocityComponent(DIContext,GridPosition,SampleMass,SampleMomentum.x,MomentumGradientX,float3(0.0,0.5,0.5),0,GridLength);
-			DIPressureGrid_BuildVelocityComponent(DIContext,GridPosition,SampleMass,SampleMomentum.y,MomentumGradientY,float3(0.5,0.0,0.5),1,GridLength);
-			DIPressureGrid_BuildVelocityComponent(DIContext,GridPosition,SampleMass,SampleMomentum.z,MomentumGradientZ,float3(0.5,0.5,0.0),2,GridLength);
+				DIPressureGrid_BuildSampleVelocity(DIContext,SamplePosition,SampleMass,SampleVelocity,
+						VelocityGradientX,VelocityGradientY,VelocityGradientZ,GridOrigin,GridLength);
+			}
 		}
 	}
 	OutBuildStatus = true;
@@ -378,43 +681,59 @@ void DIPressureGrid_BuildVelocityField(in FDIPressureGridContext DIContext, in i
  * Sample Velocity field from grid
  * -----------------------------------------------------------------
  */
+ 
+// Sample the grid velocity field and gradient component at the node local position 
+void DIPressureGrid_SampleVelocityFace(in FDIPressureGridContext DIContext, in int3 GridIndex, in float3 GridFraction, in int3 ijk, in float3 GridPosition, in float3 ComponentOffset, in int ComponentIndex, in float GridLength,
+						inout float OutGridVelocity, inout float OutGridMass, inout float3 OutGridGradient)
+{
+	const int3 GridLocation = int3(GridIndex.x+ijk.x,GridIndex.y+ijk.y,GridIndex.z+ijk.z);
+
+	if( DIPressureGrid_IsFaceValid(DIContext,GridLocation,ComponentIndex) )
+	{
+		const float ShapeFunction = DIPressureGrid_GetShapeFunction(GridFraction,ijk);
+		const float3 ShapeGradient = DIPressureGrid_GetShapeGradient(GridFraction,ijk) / GridLength;
+
+		const float GridVelocity = DIPressureGrid_GetGridScalar(DIContext,GridLocation,FLUID_VELOCITY_OFFSET+ComponentIndex);
+		const float GridMass = DIPressureGrid_GetGridScalar(DIContext,GridLocation,FLUID_MASS_OFFSET+ComponentIndex);
+
+		OutGridMass += GridMass * ShapeFunction;
+		OutGridVelocity += GridVelocity * ShapeFunction;
+		OutGridGradient += GridVelocity * ShapeGradient;
+	}
+}
 
 // Sample the grid velocity field and gradient component at the node local position 
 void DIPressureGrid_SampleVelocityComponent(in FDIPressureGridContext DIContext, in float3 GridPosition, in float3 ComponentOffset, in int ComponentIndex, in float GridLength,
 						out float OutGridVelocity, out float OutGridMass, out float3 OutGridGradient)
 {
-	int3 GridIndex = int3(0,0,0); float3 GridFraction = float3(0,0,0);
-	DIPressureGrid_GetGridIndex(DIContext,GridPosition,ComponentOffset,GridLength,GridIndex,GridFraction);
+	float3 GridFraction = float3(0,0,0); 
+	const int3 GridIndex = DIPressureGrid_ComputeGridIndex(DIContext,GridPosition,ComponentOffset,GridLength,GridFraction);
 
 	OutGridVelocity = 0.0;
-	OutGridGradient = float3(0,0,0);
+	OutGridGradient = float3(0,0,0); 
 	OutGridMass = 0.0;
 
-	if(DIPressureGrid_IsFaceValid(DIContext,GridIndex,ComponentIndex) && 
-	   DIPressureGrid_IsFaceValid(DIContext,GridIndex+int3(1,1,1),ComponentIndex))
+	for(int i = 0; i < 2; ++i)
 	{
-		[unroll]
-		for(int i = 0; i < 2; ++i)
+		const float GridWeightX = (i == 0) ? 1.0 - GridFraction.x : GridFraction.x;
+		const float GridGradientX = (i == 0) ? -1.0: 1.0;
+		for(int j = 0; j < 2; ++j)
 		{
-			const float GridWeightX = (i == 0) ? 1.0 - GridFraction.x : GridFraction.x;
-			const float GridGradientX = (i == 0) ? -1.0: 1.0;
-			[unroll]
-			for(int j = 0; j < 2; ++j)
+			const float GridWeightY = (j == 0) ? 1.0 - GridFraction.y : GridFraction.y;
+			const float GridGradientY = (j == 0) ? -1.0: 1.0;
+			for(int k = 0; k < 2; ++k)
 			{
-				const float GridWeightY = (j == 0) ? 1.0 - GridFraction.y : GridFraction.y;
-				const float GridGradientY = (j == 0) ? -1.0: 1.0;
-				[unroll]
-				for(int k = 0; k < 2; ++k)
-				{
-					const float GridWeightZ = (k == 0) ? 1.0 - GridFraction.z : GridFraction.z;
-					const float GridGradientZ = (k == 0) ? -1.0: 1.0;
+				const float GridWeightZ = (k == 0) ? 1.0 - GridFraction.z : GridFraction.z;
+				const float GridGradientZ = (k == 0) ? -1.0: 1.0;
 
+				const int3 GridLocation = int3(GridIndex.x+i,GridIndex.y+j,GridIndex.z+k);
+
+				if( DIPressureGrid_IsFaceValid(DIContext,GridLocation,ComponentIndex) )
+				{
 					const float GridWeightXYZ = GridWeightX*GridWeightY*GridWeightZ;
 					const float3 GridGradientXYZ = float3(GridGradientX*GridWeightY*GridWeightZ,
-														  GridWeightX*GridGradientY*GridWeightZ,
-														  GridWeightX*GridWeightY*GridGradientZ) / GridLength;
-
-					const int3 GridLocation = int3(GridIndex.x+i,GridIndex.y+j,GridIndex.z+k);
+														GridWeightX*GridGradientY*GridWeightZ,
+														GridWeightX*GridWeightY*GridGradientZ) / GridLength;
 					const float GridVelocity = DIPressureGrid_GetGridScalar(DIContext,GridLocation,FLUID_VELOCITY_OFFSET+ComponentIndex);
 					const float GridMass = DIPressureGrid_GetGridScalar(DIContext,GridLocation,FLUID_MASS_OFFSET+ComponentIndex);
 
@@ -425,6 +744,23 @@ void DIPressureGrid_SampleVelocityComponent(in FDIPressureGridContext DIContext,
 			}
 		}
 	}
+
+	/*DIPressureGrid_SampleVelocityFace(DIContext,GridIndex,GridFraction,int3(0,0,0),GridPosition,
+		ComponentOffset,ComponentIndex,GridLength,OutGridVelocity,OutGridMass,OutGridGradient);
+	DIPressureGrid_SampleVelocityFace(DIContext,GridIndex,GridFraction,int3(0,0,1),GridPosition,
+		ComponentOffset,ComponentIndex,GridLength,OutGridVelocity,OutGridMass,OutGridGradient);
+	DIPressureGrid_SampleVelocityFace(DIContext,GridIndex,GridFraction,int3(0,1,0),GridPosition,
+		ComponentOffset,ComponentIndex,GridLength,OutGridVelocity,OutGridMass,OutGridGradient);
+	DIPressureGrid_SampleVelocityFace(DIContext,GridIndex,GridFraction,int3(0,1,1),GridPosition,
+		ComponentOffset,ComponentIndex,GridLength,OutGridVelocity,OutGridMass,OutGridGradient);
+	DIPressureGrid_SampleVelocityFace(DIContext,GridIndex,GridFraction,int3(1,0,0),GridPosition,
+		ComponentOffset,ComponentIndex,GridLength,OutGridVelocity,OutGridMass,OutGridGradient);
+	DIPressureGrid_SampleVelocityFace(DIContext,GridIndex,GridFraction,int3(1,0,1),GridPosition,
+		ComponentOffset,ComponentIndex,GridLength,OutGridVelocity,OutGridMass,OutGridGradient);
+	DIPressureGrid_SampleVelocityFace(DIContext,GridIndex,GridFraction,int3(1,1,0),GridPosition,
+		ComponentOffset,ComponentIndex,GridLength,OutGridVelocity,OutGridMass,OutGridGradient);
+	DIPressureGrid_SampleVelocityFace(DIContext,GridIndex,GridFraction,int3(1,1,1),GridPosition,
+		ComponentOffset,ComponentIndex,GridLength,OutGridVelocity,OutGridMass,OutGridGradient);*/
 }
 
 // Sample the grid velocity field and gradient at the node world position 
@@ -443,11 +779,54 @@ void DIPressureGrid_SampleVelocityField(in FDIPressureGridContext DIContext, in
 	DIPressureGrid_SampleVelocityComponent(DIContext,GridPosition,float3(0.5,0.0,0.5),1,GridLength,OutGridVelocity.y,GridMass.y,OutGridGradientY);
 	DIPressureGrid_SampleVelocityComponent(DIContext,GridPosition,float3(0.5,0.5,0.0),2,GridLength,OutGridVelocity.z,GridMass.z,OutGridGradientZ);
 
-	OutNodeDensity = (GridMass.x + GridMass.y + GridMass.z)/(3.0*GridLength*GridLength*GridLength);
+	OutNodeDensity = (GridMass.x+GridMass.y+GridMass.z)/(3.0*GridLength*GridLength*GridLength);
+}
+
+/* -----------------------------------------------------------------
+ * Velocity Accessors
+ * -----------------------------------------------------------------
+ */
+
+void DIPressureGrid_AddGridVelocity(in FDIPressureGridContext DIContext, in int GridHash, in float3 GridVelocity, out bool OutAddStatus)
+{
+	const int3 GridIndex = DIPressureGrid_GetGridIndex(DIContext,GridHash);
+	OutAddStatus = true;
 
-	//int3 GridIndex = int3(0,0,0); float3 GridFraction = float3(0,0,0);
-	//DIPressureGrid_GetGridIndex(DIContext,GridPosition,0,GridLength,GridIndex,GridFraction);
-	//OutGridVelocity = GridIndex;
+	if( DIPressureGrid_IsFaceValid(DIContext,GridIndex,0) )
+	{
+		OutAddStatus = true;
+		DIPressureGrid_AddGridScalar(DIContext,GridIndex,FLUID_VELOCITY_OFFSET,GridVelocity.x);
+	}
+	if( DIPressureGrid_IsFaceValid(DIContext,GridIndex,1) )
+	{
+		OutAddStatus = true;
+		DIPressureGrid_AddGridScalar(DIContext,GridIndex,FLUID_VELOCITY_OFFSET+1,GridVelocity.y);
+	}
+	if( DIPressureGrid_IsFaceValid(DIContext,GridIndex,2) )
+	{
+		OutAddStatus = true;
+		DIPressureGrid_AddGridScalar(DIContext,GridIndex,FLUID_VELOCITY_OFFSET+2,GridVelocity.z);
+	}
+}
+
+void DIPressureGrid_GetGridVelocity(in FDIPressureGridContext DIContext, in int GridHash, out float3 OutGridVelocity)
+{
+	const int3 GridIndex = DIPressureGrid_GetGridIndex(DIContext,GridHash);
+
+	OutGridVelocity = float3(0,0,0);
+
+	if( DIPressureGrid_IsFaceValid(DIContext,GridIndex,0) )
+	{
+		OutGridVelocity.x = DIPressureGrid_GetGridScalar(DIContext,GridIndex,FLUID_VELOCITY_OFFSET);
+	}
+	if( DIPressureGrid_IsFaceValid(DIContext,GridIndex,1) )
+	{
+		OutGridVelocity.y = DIPressureGrid_GetGridScalar(DIContext,GridIndex,FLUID_VELOCITY_OFFSET+1);
+	}
+	if( DIPressureGrid_IsFaceValid(DIContext,GridIndex,2) )
+	{
+		OutGridVelocity.z = DIPressureGrid_GetGridScalar(DIContext,GridIndex,FLUID_VELOCITY_OFFSET+2);
+	}
 }
 
 /* -----------------------------------------------------------------
@@ -455,7 +834,8 @@ void DIPressureGrid_SampleVelocityField(in FDIPressureGridContext DIContext, in
  * -----------------------------------------------------------------
  */
 
- float DIPressureGrid_GetFractionInside(in float SignedDistanceA, in float SignedDistanceB)
+ #if HAIR_MODE == 0
+ float DIPressureGrid_GetFractionInsideEdge(in float SignedDistanceA, in float SignedDistanceB)
  {
 	const bool IsInsideA = SignedDistanceA<0.0;
 	const bool IsInsideB = SignedDistanceB<0.0;
@@ -464,50 +844,232 @@ void DIPressureGrid_SampleVelocityField(in FDIPressureGridContext DIContext, in
 		   (!IsInsideA && IsInsideB) ? SignedDistanceB / (SignedDistanceB - SignedDistanceA)  : 0.0;
  }
 
+ float DIPressureGrid_GetFractionInsideImplA(in float SignedDistanceA, in float SignedDistanceB, in float SignedDistanceD)
+ {
+	const float Side0 = 1.0 - DIPressureGrid_GetFractionInsideEdge(SignedDistanceA,SignedDistanceD);
+	const float Side1 = 1.0 - DIPressureGrid_GetFractionInsideEdge(SignedDistanceA,SignedDistanceB);
+	return 1.0 - 0.5 * Side0 * Side1;
+ }
+
+ float DIPressureGrid_GetFractionInsideImplB(in float SignedDistanceA, in float SignedDistanceB, in float SignedDistanceC, in float SignedDistanceD)
+ {
+	const float Side0 = DIPressureGrid_GetFractionInsideEdge(SignedDistanceA,SignedDistanceD);
+	const float Side1 = DIPressureGrid_GetFractionInsideEdge(SignedDistanceB,SignedDistanceC);
+	return 0.5 * ( Side0 + Side1);
+ }
+
+ float DIPressureGrid_GetFractionInsideImplC(in float SignedDistanceA, in float SignedDistanceB, in float SignedDistanceC, in float SignedDistanceD)
+ {
+	const float MiddlePoint = 0.25 * (SignedDistanceA+SignedDistanceB+SignedDistanceC+SignedDistanceD);
+	if(MiddlePoint < 0.0)
+	{
+		float FaceArea = 0.0;
+
+		const float Side1 = 1.0 - DIPressureGrid_GetFractionInsideEdge(SignedDistanceA,SignedDistanceD);
+		const float Side3 = 1.0 - DIPressureGrid_GetFractionInsideEdge(SignedDistanceC,SignedDistanceD);
+		FaceArea += 0.5 * Side3 * Side1;
+
+		const float Side2 = 1.0 - DIPressureGrid_GetFractionInsideEdge(SignedDistanceC,SignedDistanceB);
+		const float Side0 = 1.0 - DIPressureGrid_GetFractionInsideEdge(SignedDistanceA,SignedDistanceB);
+		FaceArea += 0.5 * Side0 * Side2;
+
+		return 1.0 - FaceArea;
+	}
+	else
+	{
+		float FaceArea = 0.0;
+
+		const float Side0 = 1.0 - DIPressureGrid_GetFractionInsideEdge(SignedDistanceA,SignedDistanceB);
+		const float Side1 = 1.0 - DIPressureGrid_GetFractionInsideEdge(SignedDistanceA,SignedDistanceD);
+		FaceArea += 0.5 * Side0 * Side1;
+
+		const float Side2 = 1.0 - DIPressureGrid_GetFractionInsideEdge(SignedDistanceC,SignedDistanceB);
+		const float Side3 = 1.0 - DIPressureGrid_GetFractionInsideEdge(SignedDistanceC,SignedDistanceD);
+		FaceArea += 0.5 * Side3 * Side2;
+
+		return FaceArea;
+
+	}
+ }
+
+ float DIPressureGrid_GetFractionInsideImplD(in float SignedDistanceA, in float SignedDistanceB, in float SignedDistanceD)
+ {
+	const float Side0 = DIPressureGrid_GetFractionInsideEdge(SignedDistanceA,SignedDistanceD);
+	const float Side1 = DIPressureGrid_GetFractionInsideEdge(SignedDistanceA,SignedDistanceB);
+	return 0.5 * Side0 * Side1;
+ }
+
+ float DIPressureGrid_GetFractionInsideFace(in float SignedDistanceA, in float SignedDistanceB, in float SignedDistanceC, in float SignedDistanceD )
+ {
+	const bool IsInsideA = SignedDistanceA<0.0;
+	const bool IsInsideB = SignedDistanceB<0.0;
+	const bool IsInsideC = SignedDistanceC<0.0;
+	const bool IsInsideD = SignedDistanceD<0.0;
+
+	const int InsideCount = IsInsideA + IsInsideB + IsInsideC + IsInsideD;
+
+	[branch]
+	if(InsideCount == 4)
+	{
+		return 1.0;
+	}
+	else if(InsideCount == 3)
+	{
+		[branch]
+		if(!IsInsideA)
+		{
+			return DIPressureGrid_GetFractionInsideImplA(SignedDistanceA,SignedDistanceB,SignedDistanceD);
+		}
+		else if(!IsInsideB)
+		{
+			return DIPressureGrid_GetFractionInsideImplA(SignedDistanceB,SignedDistanceC,SignedDistanceA);
+		}
+		else if(!IsInsideC)
+		{
+			return DIPressureGrid_GetFractionInsideImplA(SignedDistanceC,SignedDistanceD,SignedDistanceB);
+		}
+		else if(!IsInsideD)
+		{
+			return DIPressureGrid_GetFractionInsideImplA(SignedDistanceD,SignedDistanceA,SignedDistanceC);
+		}
+	}
+	else if(InsideCount == 2)
+	{
+		[branch]
+		if(IsInsideA && IsInsideB)
+		{
+			return DIPressureGrid_GetFractionInsideImplB(SignedDistanceA,SignedDistanceB,SignedDistanceC, SignedDistanceD);
+		}
+		else if(IsInsideB && IsInsideC)
+		{
+			return DIPressureGrid_GetFractionInsideImplB(SignedDistanceB,SignedDistanceC,SignedDistanceD, SignedDistanceA);
+		}
+		else if(IsInsideC && IsInsideD)
+		{
+			return DIPressureGrid_GetFractionInsideImplB(SignedDistanceC,SignedDistanceD,SignedDistanceA, SignedDistanceB);
+		}
+		else if(IsInsideD && IsInsideA)
+		{
+			return DIPressureGrid_GetFractionInsideImplB(SignedDistanceD,SignedDistanceA,SignedDistanceB, SignedDistanceC);
+		}
+		else if(IsInsideA && IsInsideC)
+		{
+			return DIPressureGrid_GetFractionInsideImplC(SignedDistanceA,SignedDistanceB,SignedDistanceC, SignedDistanceD);
+		}
+		else if(IsInsideB && IsInsideD)
+		{
+			return DIPressureGrid_GetFractionInsideImplC(SignedDistanceB,SignedDistanceC,SignedDistanceD, SignedDistanceA);
+		}
+	}
+	else if(InsideCount == 1)
+	{
+		[branch]
+		if(IsInsideA)
+		{
+			return DIPressureGrid_GetFractionInsideImplA(SignedDistanceA,SignedDistanceB,SignedDistanceD);
+		}
+		else if(IsInsideB)
+		{
+			return DIPressureGrid_GetFractionInsideImplA(SignedDistanceB,SignedDistanceC,SignedDistanceA);
+		}
+		else if(IsInsideC)
+		{
+			return DIPressureGrid_GetFractionInsideImplA(SignedDistanceC,SignedDistanceD,SignedDistanceB);
+		}
+		else if(IsInsideD)
+		{
+			return DIPressureGrid_GetFractionInsideImplA(SignedDistanceD,SignedDistanceA,SignedDistanceC);
+		}
+	}
+	return 0;
+ }
+
+ #endif
+
  // Update the solid weights and velocities
-void DIPressureGrid_SetSolidBoundary(in FDIPressureGridContext DIContext, in int GridCell, in float SolidDistance,
+void DIPressureGrid_SetSolidBoundary(in FDIPressureGridContext DIContext, in int GridHash, in float SolidDistance,
 																in float3 SolidVelocity, out bool OutBoundaryStatus)
 {
-	int3 GridIndex = DIPressureGrid_ComputeGridIndex(DIContext,GridCell);
+	const int3 GridIndex = DIPressureGrid_GetGridIndex(DIContext,GridHash);
 	OutBoundaryStatus = false;
-	if( DIPressureGrid_IsCellValid(DIContext,GridIndex))
+
+	#if HAIR_MODE == 0
+
+	[branch]
+	if( DIPressureGrid_IsNodeValid(DIContext,GridIndex))
 	{
 		OutBoundaryStatus = true;
 		DIPressureGrid_SetGridScalar(DIContext,GridIndex,SOLID_DISTANCE_OFFSET,SolidDistance);
 		DIPressureGrid_SetGridVector(DIContext,GridIndex,SOLID_VELOCITY_OFFSET,SolidVelocity);
 	}
+
+	#endif
 }
 
- // Update the solid weights and velocities
-void DIPressureGrid_ComputeBoundaryWeights(in FDIPressureGridContext DIContext, in int GridCell, out bool OutWeightsStatus)
+#if HAIR_MODE == 0
+
+void DIPressureGrid_ComputeFaceWeights(in FDIPressureGridContext DIContext, in int3 GridLocation, in int ComponentIndex, out float SignedDistance, inout float NodeVelocity )
 {
-	const int3 GridIndex = DIPressureGrid_ComputeGridIndex(DIContext,GridCell);
-	OutWeightsStatus = false;
+	SignedDistance = DIPressureGrid_GetGridScalar(DIContext,GridLocation,SOLID_DISTANCE_OFFSET);
+	NodeVelocity += DIPressureGrid_GetGridScalar(DIContext,GridLocation,SOLID_VELOCITY_OFFSET+ComponentIndex);
+}
 
-	for(int ComponentIndex = 0; ComponentIndex < 3; ++ComponentIndex)
+void DIPressureGrid_ComputeComponentWeights(in FDIPressureGridContext DIContext, in int3 GridIndex, in int ComponentIndex)
+{
+	[branch]
+	if( DIPressureGrid_IsFaceInside(DIContext,GridIndex,ComponentIndex))
 	{
-		if( DIPressureGrid_IsFaceValid(DIContext,GridIndex,ComponentIndex))
+		float SignedDistance[4];
+		float NodeVelocity = 0.0;
+
+		[branch]
+		if(ComponentIndex == 0)
 		{
-			float SolidWeight = 0.0, SolidVelocity = 0.0;
-			if( GridIndex[ComponentIndex] > 0 && GridIndex[ComponentIndex] < (DIContext.GridSize[ComponentIndex]-1) )
-			{
-				OutWeightsStatus = true;
-				const float RightDistance = DIPressureGrid_GetGridScalar(DIContext,GridIndex,SOLID_DISTANCE_OFFSET);
-				const float RightVelocity = DIPressureGrid_GetGridScalar(DIContext,GridIndex,SOLID_VELOCITY_OFFSET+ComponentIndex);
+			DIPressureGrid_ComputeFaceWeights(DIContext,int3(GridIndex.x,GridIndex.y+0,GridIndex.z+0),ComponentIndex,SignedDistance[0],NodeVelocity);
+			DIPressureGrid_ComputeFaceWeights(DIContext,int3(GridIndex.x,GridIndex.y+0,GridIndex.z+1),ComponentIndex,SignedDistance[1],NodeVelocity);
+			DIPressureGrid_ComputeFaceWeights(DIContext,int3(GridIndex.x,GridIndex.y+1,GridIndex.z+0),ComponentIndex,SignedDistance[2],NodeVelocity);
+			DIPressureGrid_ComputeFaceWeights(DIContext,int3(GridIndex.x,GridIndex.y+1,GridIndex.z+1),ComponentIndex,SignedDistance[3],NodeVelocity);
+		}
+		else if(ComponentIndex == 1)
+		{
+			DIPressureGrid_ComputeFaceWeights(DIContext,int3(GridIndex.x+0,GridIndex.y,GridIndex.z+0),ComponentIndex,SignedDistance[0],NodeVelocity);
+			DIPressureGrid_ComputeFaceWeights(DIContext,int3(GridIndex.x+0,GridIndex.y,GridIndex.z+1),ComponentIndex,SignedDistance[1],NodeVelocity);
+			DIPressureGrid_ComputeFaceWeights(DIContext,int3(GridIndex.x+1,GridIndex.y,GridIndex.z+0),ComponentIndex,SignedDistance[2],NodeVelocity);
+			DIPressureGrid_ComputeFaceWeights(DIContext,int3(GridIndex.x+1,GridIndex.y,GridIndex.z+1),ComponentIndex,SignedDistance[3],NodeVelocity);
+		}
+		else
+		{
+			DIPressureGrid_ComputeFaceWeights(DIContext,int3(GridIndex.x+0,GridIndex.y+0,GridIndex.z),ComponentIndex,SignedDistance[0],NodeVelocity);
+			DIPressureGrid_ComputeFaceWeights(DIContext,int3(GridIndex.x+0,GridIndex.y+1,GridIndex.z),ComponentIndex,SignedDistance[1],NodeVelocity);
+			DIPressureGrid_ComputeFaceWeights(DIContext,int3(GridIndex.x+1,GridIndex.y+0,GridIndex.z),ComponentIndex,SignedDistance[2],NodeVelocity);
+			DIPressureGrid_ComputeFaceWeights(DIContext,int3(GridIndex.x+1,GridIndex.y+1,GridIndex.z),ComponentIndex,SignedDistance[3],NodeVelocity);
+		}
 
-				int3 LeftIndex = GridIndex; LeftIndex[ComponentIndex] -= 1;
-				const float LeftDistance = DIPressureGrid_GetGridScalar(DIContext,LeftIndex,SOLID_DISTANCE_OFFSET);
-				const float LeftVelocity = DIPressureGrid_GetGridScalar(DIContext,LeftIndex,SOLID_VELOCITY_OFFSET+ComponentIndex);
+		const float SolidWeight = max(0.0,min(1.0 - DIPressureGrid_GetFractionInsideFace(SignedDistance[0],SignedDistance[1],SignedDistance[3],SignedDistance[2]),1.0));
+		const float SolidVelocity = 0.25 * NodeVelocity;
 
-				SolidWeight = max(0.0,min(1.0 - DIPressureGrid_GetFractionInside(LeftDistance,RightDistance),1.0));
-				SolidVelocity = 0.5 * (RightVelocity+LeftVelocity);
-			}
-			DIPressureGrid_SetGridScalar(DIContext,GridIndex,SOLID_WEIGHT_OFFSET+ComponentIndex,SolidWeight);
-			DIPressureGrid_SetGridScalar(DIContext,GridIndex,SOLID_VELOCITY_OFFSET+ComponentIndex,SolidVelocity);
-		}
+		DIPressureGrid_SetGridScalar(DIContext,GridIndex,SOLID_WEIGHT_OFFSET+ComponentIndex,SolidWeight);
+		DIPressureGrid_SetGridScalar(DIContext,GridIndex,SOLID_VELOCITY_OFFSET+ComponentIndex,SolidVelocity);
 	}
 }
 
+#endif
+
+ // Update the solid weights and velocities
+void DIPressureGrid_ComputeBoundaryWeights(in FDIPressureGridContext DIContext, in int GridHash, out bool OutWeightsStatus)
+{
+	const int3 GridIndex = DIPressureGrid_GetGridIndex(DIContext,GridHash);
+	OutWeightsStatus = false;
+
+	#if HAIR_MODE == 0
+
+	DIPressureGrid_ComputeComponentWeights(DIContext,GridIndex,0);
+	DIPressureGrid_ComputeComponentWeights(DIContext,GridIndex,1);
+	DIPressureGrid_ComputeComponentWeights(DIContext,GridIndex,2);
+
+	#endif
+}
+
 /* -----------------------------------------------------------------
  * Sample Velocity field from grid
  * -----------------------------------------------------------------
@@ -521,89 +1083,113 @@ void DIPressureGrid_ComputeBoundaryWeights(in FDIPressureGridContext DIContext,
 // the schur complement is J * inv(M) * Jt * P * dt = -C0 -> P * dt = -C0 / (wnx * dx * dx * wnx * dx * dx / ( wnx * mnx) + wpx * ....)
 // P * dt = - sum( vnx * wnx * dx * dx - vpx * wpx * dx * dx ) / sum (wnx * dx * dx * wnx * dx * dx / ( wnx * mnx);
 //        = - sum( vnx * wnx ) / sum( wnx * dx * dx / mnx ) = -sum( vnx * wnx ) / ( dx * dx * sum( wnx / mnx ) )
-float DIPressureGrid_ComputeGridPressure(in FDIPressureGridContext DIContext, in int3 PrevIndex, out float OutCellDivergence)
+// P * dt * dx * dx = Pt = -sum( vnx * wnx ) / ( sum( wnx / mnx ) )
+// dv = wnx * dx * dx * P * dt / (wnx * mnx) = Pt / mnx
+
+#if HAIR_MODE == 0
+
+float DIPressureGrid_ComputeGridPressure(in FDIPressureGridContext DIContext, in int3 PrevIndex, in float CentreDistance, out float3 PrevFraction, out float3 NextFraction)
 {	
 	float GridPressure = 0.0;
-	OutCellDivergence = 5.0;
+	float GridDivergence = 0.0;
+	float GridFraction = 0.0;
 
-	if(DIPressureGrid_IsCellValid(DIContext,PrevIndex))
+	for(int ComponentIndex = 0; ComponentIndex < 3; ++ComponentIndex)
 	{
-		float GridDivergence = 0.0;
-		float GridFraction = 0.0;
-
-		[unroll]
-		for(int ComponentIndex = 0; ComponentIndex < 3; ++ComponentIndex)
-		{
-			int3 NextIndex = PrevIndex; NextIndex[ComponentIndex] += 1;
+		int3 NextIndex = PrevIndex; NextIndex[ComponentIndex] += 1;
 
-			const float PrevMass = DIPressureGrid_GetGridScalar(DIContext,PrevIndex,FLUID_MASS_OFFSET+ComponentIndex);
-			const float PrevWeight = DIPressureGrid_GetGridScalar(DIContext,PrevIndex,SOLID_WEIGHT_OFFSET+ComponentIndex);
+		const float PrevDistance = DIPressureGrid_GetGridScalar(DIContext,PrevIndex,FLUID_DISTANCE_OFFSET);
+		PrevFraction[ComponentIndex] = (PrevDistance < 0.0) ? 1.0 : max(DIPressureGrid_GetFractionInsideEdge(CentreDistance,PrevDistance),0.01);
+		const float PrevWeight = DIPressureGrid_GetGridScalar(DIContext,PrevIndex,SOLID_WEIGHT_OFFSET+ComponentIndex);
 			
-			GridDivergence -= DIPressureGrid_GetGridScalar(DIContext,PrevIndex,SOLID_VELOCITY_OFFSET+ComponentIndex) * (1.0-PrevWeight);
+		GridDivergence -= DIPressureGrid_GetGridScalar(DIContext,PrevIndex,SOLID_VELOCITY_OFFSET+ComponentIndex) * (1.0-PrevWeight);
+		GridDivergence -= DIPressureGrid_GetGridScalar(DIContext,PrevIndex,FLUID_VELOCITY_OFFSET+ComponentIndex) * PrevWeight;
+		GridFraction += PrevWeight / PrevFraction[ComponentIndex];
 
-			if( PrevMass > 0.)
-			{
-				GridDivergence -= DIPressureGrid_GetGridScalar(DIContext,PrevIndex,FLUID_VELOCITY_OFFSET+ComponentIndex) * PrevWeight;
-				GridFraction += PrevWeight / PrevMass;
-			}
-
-			const float NextMass = DIPressureGrid_GetGridScalar(DIContext,NextIndex,FLUID_MASS_OFFSET+ComponentIndex);
-			const float NextWeight = DIPressureGrid_GetGridScalar(DIContext,NextIndex,SOLID_WEIGHT_OFFSET+ComponentIndex);
+		const float NextDistance = DIPressureGrid_GetGridScalar(DIContext,NextIndex,FLUID_DISTANCE_OFFSET);
+		NextFraction[ComponentIndex] = (NextDistance < 0.0) ? 1.0 : max(DIPressureGrid_GetFractionInsideEdge(CentreDistance,NextDistance),0.01);
+		const float NextWeight = DIPressureGrid_GetGridScalar(DIContext,NextIndex,SOLID_WEIGHT_OFFSET+ComponentIndex);
 			
-			GridDivergence += DIPressureGrid_GetGridScalar(DIContext,NextIndex,SOLID_VELOCITY_OFFSET+ComponentIndex) * (1.0-NextWeight);
-			if( NextMass > 0.)
-			{
-				GridDivergence += DIPressureGrid_GetGridScalar(DIContext,NextIndex,FLUID_VELOCITY_OFFSET+ComponentIndex) * NextWeight;
-				GridFraction += NextWeight / NextMass;
-			}
-		}
-		if( GridFraction > 0.)
-		{
-			//GridPressure = -GridDivergence / ( GridLength * GridLength * GridFraction);
-			GridPressure = -GridDivergence / (GridFraction);
-		}
-		OutCellDivergence = GridDivergence;
+		GridDivergence += DIPressureGrid_GetGridScalar(DIContext,NextIndex,SOLID_VELOCITY_OFFSET+ComponentIndex) * (1.0-NextWeight);
+		GridDivergence += DIPressureGrid_GetGridScalar(DIContext,NextIndex,FLUID_VELOCITY_OFFSET+ComponentIndex) * NextWeight;
+		GridFraction += NextWeight / NextFraction[ComponentIndex];
+	}
+	if( GridFraction > 0.0)
+	{
+		//GridDivergence += DIPressureGrid_GetGridScalar(DIContext,PrevIndex,FLUID_DENSITY_OFFSET);
+		GridPressure = -GridDivergence / (GridFraction);
 	}
 	return GridPressure;
 }
 
+#endif
+
 // Project the pressure grid to be divergence free
-void DIPressureGrid_ProjectVelocityField(in FDIPressureGridContext DIContext, in int GridCell, out bool OutProjectStatus, out float OutCellDivergence)
+void DIPressureGrid_ProjectVelocityField(in FDIPressureGridContext DIContext, in int GridHash, in int InitStage, out bool OutProjectStatus)
 {	
 	OutProjectStatus = false;
-	OutCellDivergence = 10.0;
 
-	const int3 PrevIndex = DIPressureGrid_ComputeGridIndex(DIContext,GridCell);
-	if( DIPressureGrid_IsCellValid(DIContext,PrevIndex))
+	#if HAIR_MODE == 0
+
+	const int3 PrevIndex = DIPressureGrid_GetGridIndex(DIContext,GridHash);
+	const bool IsInitStage = (ShaderStageIndex == InitStage) || (ShaderStageIndex == (InitStage+1));
+
+	if(DIPressureGrid_IsCellValid(DIContext,PrevIndex))
 	{
 		if( (PrevIndex.x+PrevIndex.y+PrevIndex.z+ShaderStageIndex) % 2 == 0)
 		{
 			OutProjectStatus = true;
-			//const float GridPressure = DIPressureGrid_ComputeGridPressure(DIContext,PrevIndex,OutCellDivergence) * GridLength * GridLength;
-			const float GridPressure = DIPressureGrid_ComputeGridPressure(DIContext,PrevIndex,OutCellDivergence);
-		
-			for( int ComponentIndex = 0; ComponentIndex < 3; ++ComponentIndex)
+			const float CentreDistance = DIPressureGrid_GetGridScalar(DIContext,PrevIndex,FLUID_DISTANCE_OFFSET);
+			if( CentreDistance < 0.0)
 			{
-				int3 NextIndex = PrevIndex; NextIndex[ComponentIndex] += 1;
+				float3 PrevFraction = float3(0,0,0);
+				float3 NextFraction = float3(0,0,0);
 
-				const float PrevMass = DIPressureGrid_GetGridScalar(DIContext,PrevIndex,FLUID_MASS_OFFSET+ComponentIndex);
-				const float PrevWeight = DIPressureGrid_GetGridScalar(DIContext,PrevIndex,SOLID_WEIGHT_OFFSET+ComponentIndex);
-
-				if( PrevMass > 0.0)
+				float GridPressure = DIPressureGrid_ComputeGridPressure(DIContext,PrevIndex,CentreDistance,PrevFraction,NextFraction);
+				if(!IsInitStage)
 				{
-					DIPressureGrid_AddGridScalar(DIContext,PrevIndex,FLUID_VELOCITY_OFFSET+ComponentIndex,-GridPressure*PrevWeight/PrevMass);
+					DIPressureGrid_AddGridScalar(DIContext,PrevIndex,FLUID_PRESSURE_OFFSET,GridPressure);
+				}
+				else 
+				{
+					GridPressure =  DIPressureGrid_GetGridScalar(DIContext,PrevIndex,FLUID_PRESSURE_OFFSET);
 				}
 
-				const float NextMass = DIPressureGrid_GetGridScalar(DIContext,NextIndex,FLUID_MASS_OFFSET+ComponentIndex);
-				const float NextWeight = DIPressureGrid_GetGridScalar(DIContext,NextIndex,SOLID_WEIGHT_OFFSET+ComponentIndex);
-
-				if( NextMass > 0.0)
+				for( int ComponentIndex = 0; ComponentIndex < 3; ++ComponentIndex)
 				{
-					DIPressureGrid_AddGridScalar(DIContext,NextIndex,FLUID_VELOCITY_OFFSET+ComponentIndex,GridPressure*NextWeight/NextMass);
+					int3 NextIndex = PrevIndex; NextIndex[ComponentIndex] += 1;
+
+					const float PrevWeight = DIPressureGrid_GetGridScalar(DIContext,PrevIndex,SOLID_WEIGHT_OFFSET+ComponentIndex);
+					if( PrevWeight > 0.0 && PrevFraction[ComponentIndex] > 0.0)
+					{
+						DIPressureGrid_AddGridScalar(DIContext,PrevIndex,FLUID_VELOCITY_OFFSET+ComponentIndex,-GridPressure/PrevFraction[ComponentIndex]);
+					}
+					else
+					{
+						const float SolidVelocity = DIPressureGrid_GetGridScalar(DIContext,PrevIndex,SOLID_VELOCITY_OFFSET+ComponentIndex);
+						DIPressureGrid_SetGridScalar(DIContext,PrevIndex,FLUID_VELOCITY_OFFSET+ComponentIndex,SolidVelocity);
+					}
+
+					const float NextWeight = DIPressureGrid_GetGridScalar(DIContext,NextIndex,SOLID_WEIGHT_OFFSET+ComponentIndex);
+					if( NextWeight > 0.0 && NextFraction[ComponentIndex] > 0.0)
+					{
+						DIPressureGrid_AddGridScalar(DIContext,NextIndex,FLUID_VELOCITY_OFFSET+ComponentIndex,GridPressure/NextFraction[ComponentIndex]);
+					}
+					else
+					{
+						const float SolidVelocity = DIPressureGrid_GetGridScalar(DIContext,NextIndex,SOLID_VELOCITY_OFFSET+ComponentIndex);
+						DIPressureGrid_SetGridScalar(DIContext,NextIndex,FLUID_VELOCITY_OFFSET+ComponentIndex,SolidVelocity);
+					}
 				}
 			}
+			else
+			{
+				DIPressureGrid_SetGridScalar(DIContext,PrevIndex,FLUID_PRESSURE_OFFSET,0.0);
+			}
 		}
 	}
+
+	#endif
 }
 
 
diff --git a/Engine/Plugins/Experimental/HairStrands/Shaders/Private/NiagaraHookeSpringMaterial.ush b/Engine/Plugins/Experimental/HairStrands/Shaders/Private/NiagaraHookeSpringMaterial.ush
--- a/Engine/Plugins/Experimental/HairStrands/Shaders/Private/NiagaraHookeSpringMaterial.ush
+++ b/Engine/Plugins/Experimental/HairStrands/Shaders/Private/NiagaraHookeSpringMaterial.ush
@@ -64,26 +64,29 @@ void SetupStretchSpringMaterial(in int StrandsSize, in float YoungModulus, in fl
 	}
 }
 
-void SolveStretchSpringMaterial(in int StrandsSize, in float RestLength, in float DeltaTime, in float MaterialDamping, 
+void SolveStretchSpringMaterial(in bool EnableConstraint, in int StrandsSize, in float RestLength, in float DeltaTime, in float MaterialDamping, 
 			in float MaterialCompliance, in float MaterialWeight, in float MaterialMultiplier, out float OutMaterialMultiplier)
 {
 	OutMaterialMultiplier = MaterialMultiplier;
-	const int LocalIndex = (GGroupThreadId.x % StrandsSize);
-	if( LocalIndex > 1)
+	if(EnableConstraint)
 	{
-		const int IsRed = (GGroupThreadId.x % 2) == 0;
-		// Process all the red rods
-		if (IsRed)
-		{	
-			UpdateStretchSpringMultiplier(RestLength,DeltaTime,false,MaterialDamping,MaterialCompliance,MaterialWeight,OutMaterialMultiplier,0);
-		}
-		// Process all the black rods
-		GroupMemoryBarrier(); 
-		if (!IsRed)
+		const int LocalIndex = (GGroupThreadId.x % StrandsSize);
+		if( LocalIndex > 1)
 		{
-			UpdateStretchSpringMultiplier(RestLength,DeltaTime,false,MaterialDamping,MaterialCompliance,MaterialWeight,OutMaterialMultiplier,0);
+			const int IsRed = (GGroupThreadId.x % 2) == 0;
+			// Process all the red rods
+			if (IsRed)
+			{	
+				UpdateStretchSpringMultiplier(RestLength,DeltaTime,false,MaterialDamping,MaterialCompliance,MaterialWeight,OutMaterialMultiplier,0);
+			}
+			// Process all the black rods
+			GroupMemoryBarrier(); 
+			if (!IsRed)
+			{
+				UpdateStretchSpringMultiplier(RestLength,DeltaTime,false,MaterialDamping,MaterialCompliance,MaterialWeight,OutMaterialMultiplier,0);
+			}
+			GroupMemoryBarrier(); 
 		}
-		GroupMemoryBarrier(); 
 	}
 } 
 
@@ -139,26 +142,29 @@ void ProjectStretchSpringMaterial(in int StrandsSize, in float YoungModulus, in
 	}
 }
 
-void SolveBendSpringMaterial(in int StrandsSize, in float RestLength, in float DeltaTime, in float MaterialDamping, 
+void SolveBendSpringMaterial(in bool EnableConstraint, in int StrandsSize, in float RestLength, in float DeltaTime, in float MaterialDamping, 
 			in float MaterialCompliance, in float MaterialWeight, in float MaterialMultiplier, out float OutMaterialMultiplier)
 {
 	OutMaterialMultiplier = MaterialMultiplier;
-	const int LocalIndex = (GGroupThreadId.x % StrandsSize);
-	if( LocalIndex > 1)
+	if(EnableConstraint)
 	{
-		const int IsRed = ((GGroupThreadId.x/2) % 2) == 0;
-		// Process all the red rods
-		if (IsRed)
-		{	
-			UpdateStretchSpringMultiplier(RestLength,DeltaTime,false,MaterialDamping,MaterialCompliance,MaterialWeight,OutMaterialMultiplier,1);
-		}
-		// Process all the black rods
-		GroupMemoryBarrier(); 
-		if (!IsRed)
+		const int LocalIndex = (GGroupThreadId.x % StrandsSize);
+		if( LocalIndex > 1)
 		{
-			UpdateStretchSpringMultiplier(RestLength,DeltaTime,false,MaterialDamping,MaterialCompliance,MaterialWeight,OutMaterialMultiplier,1);
+			const int IsRed = ((GGroupThreadId.x/2) % 2) == 0;
+			// Process all the red rods
+			if (IsRed)
+			{	
+				UpdateStretchSpringMultiplier(RestLength,DeltaTime,false,MaterialDamping,MaterialCompliance,MaterialWeight,OutMaterialMultiplier,1);
+			}
+			// Process all the black rods
+			GroupMemoryBarrier(); 
+			if (!IsRed)
+			{
+				UpdateStretchSpringMultiplier(RestLength,DeltaTime,false,MaterialDamping,MaterialCompliance,MaterialWeight,OutMaterialMultiplier,1);
+			}
+			GroupMemoryBarrier(); 
 		}
-		GroupMemoryBarrier(); 
 	}
 } 
 
@@ -195,6 +201,102 @@ void ProjectBendSpringMaterial(in int StrandsSize, in float YoungModulus, in flo
 	OutNodePosition = SharedNodePosition[GGroupThreadId.x];
 }
 
+/* -----------------------------------------------------------------
+ * Volume spring material
+ * -----------------------------------------------------------------
+ */
+
+
+ void SetupVolumeSpringMaterial(in int StrandsSize, in float YoungModulus,
+	in float RestVolume, in float DeltaTime, in bool ProjectConstraint, in float MaterialDamping, out float OutMaterialCompliance, out float OutMaterialWeight, out float OutMaterialMultiplier)
+{
+	OutMaterialCompliance = 0.0;
+	OutMaterialWeight = 0.0;
+	OutMaterialMultiplier = 0.0;
+
+	const int LocalIndex = (GGroupThreadId.x % StrandsSize);
+	if( LocalIndex > 2 )
+	{
+		// Compliance = 1.0 / (k * dt * dt)
+		// with k = (Y / V) 
+		// V is the constraint inital volume and Y is the young modulus
+		OutMaterialCompliance = RestVolume/(YoungModulus*DeltaTime*DeltaTime);
+	}
+}
+ 
+void UpdateVolumeSpringMultiplier(in float RestVolume, in float DeltaTime, in bool ProjectConstraint, in float MaterialDamping, in float MaterialCompliance, in float MaterialWeight, inout float OutMaterialMultiplier)
+{
+	const int NodeIndexA = GGroupThreadId.x;
+	const int NodeIndexB = GGroupThreadId.x-1;
+	const int NodeIndexC = GGroupThreadId.x-2;
+	const int NodeIndexD = GGroupThreadId.x-3;
+
+	const float3 EdgeVectorA = SharedNodePosition[NodeIndexB] - SharedNodePosition[NodeIndexA];
+	const float3 EdgeVectorB = SharedNodePosition[NodeIndexC] - SharedNodePosition[NodeIndexA];
+	const float3 EdgeVectorC = SharedNodePosition[NodeIndexD] - SharedNodePosition[NodeIndexA];
+
+	const float NodesVolume = dot(cross(EdgeVectorB,EdgeVectorC),EdgeVectorA);
+
+	const float3 ConstraintGradB = cross(EdgeVectorB,EdgeVectorC);
+	const float3 ConstraintGradC = cross(EdgeVectorC,EdgeVectorA);
+	const float3 ConstraintGradD = cross(EdgeVectorA,EdgeVectorB);
+	const float3 ConstraintGradA = -ConstraintGradB-ConstraintGradC-ConstraintGradD;
+
+	const float3 NodeVelocityA = ( SharedNodePosition[NodeIndexA] - SharedPreviousPosition[NodeIndexA] ) / DeltaTime;
+	const float3 NodeVelocityB = ( SharedNodePosition[NodeIndexB] - SharedPreviousPosition[NodeIndexB] ) / DeltaTime;
+	const float3 NodeVelocityC = ( SharedNodePosition[NodeIndexC] - SharedPreviousPosition[NodeIndexC] ) / DeltaTime;
+	const float3 NodeVelocityD = ( SharedNodePosition[NodeIndexD] - SharedPreviousPosition[NodeIndexD] ) / DeltaTime;
+
+	const float SumInverseMass = !ProjectConstraint ? ( dot(ConstraintGradA,ConstraintGradA) * SharedInverseMass[NodeIndexA] +
+														dot(ConstraintGradB,ConstraintGradB) * SharedInverseMass[NodeIndexB] +
+													    dot(ConstraintGradC,ConstraintGradC) * SharedInverseMass[NodeIndexC] +
+														dot(ConstraintGradD,ConstraintGradD) * SharedInverseMass[NodeIndexD] ) :
+														dot(ConstraintGradA,ConstraintGradA) * SharedInverseMass[NodeIndexA] ;
+
+	const float SchurDiagonal = (1.0 + MaterialDamping) * SumInverseMass + MaterialCompliance; 
+	const float ConstraintWeight = ( SchurDiagonal != 0.0 ) ? 1.0 / SchurDiagonal : 0.0;
+
+	// XPBD lagrange multiplier update : dL = -(C+compliance*L) / (dC * invM * dCt + alpha)
+	const float DeltaLambda = -((NodesVolume - RestVolume) + OutMaterialMultiplier * MaterialCompliance + MaterialDamping * 
+		(dot(ConstraintGradA,NodeVelocityA)+dot(ConstraintGradB,NodeVelocityB)+dot(ConstraintGradC,NodeVelocityC)+dot(ConstraintGradD,NodeVelocityD) ) ) * ConstraintWeight;
+
+	// L += dL
+	//OutMaterialMultiplier += DeltaLambda;
+	OutMaterialMultiplier = NodesVolume;
+
+	// XPBD position update : dX += dL * dCt * invM
+	SharedNodePosition[NodeIndexA] += ConstraintGradA * DeltaLambda * SharedInverseMass[NodeIndexA];
+	if(!ProjectConstraint)
+	{
+		SharedNodePosition[NodeIndexB] += ConstraintGradB * DeltaLambda * SharedInverseMass[NodeIndexB]; 
+		SharedNodePosition[NodeIndexC] += ConstraintGradC * DeltaLambda * SharedInverseMass[NodeIndexC];
+		SharedNodePosition[NodeIndexD] += ConstraintGradD * DeltaLambda * SharedInverseMass[NodeIndexD];
+	}
+}
+
+void SolveVolumeSpringMaterial(in bool EnableConstraint, in int StrandsSize, in float RestVolume, in float DeltaTime, in float MaterialDamping, 
+			in float MaterialCompliance, in float MaterialWeight, in float MaterialMultiplier, out float OutMaterialMultiplier)
+{
+	OutMaterialMultiplier = MaterialMultiplier;
+	if(EnableConstraint)
+	{
+		const int LocalIndex = (GGroupThreadId.x % StrandsSize);
+		if( LocalIndex > 2)
+		{
+			for(int i = 3; i < StrandsSize; ++i)
+			{
+				if( LocalIndex == i )
+				{
+					//UpdateVolumeSpringMultiplier(RestVolume,DeltaTime,false,MaterialDamping,MaterialCompliance,MaterialWeight,OutMaterialMultiplier);
+					UpdateVolumeSpringMultiplier(RestVolume,DeltaTime,false,0.0,0.0,MaterialWeight,OutMaterialMultiplier);
+				}
+				GroupMemoryBarrier(); 
+			}
+		}
+	}
+	GroupMemoryBarrier();
+} 
+
 /* -----------------------------------------------------------------
  * Twist spring material
  * -----------------------------------------------------------------
@@ -213,26 +315,29 @@ void ProjectBendSpringMaterial(in int StrandsSize, in float YoungModulus, in flo
 	}
 }
 
-void SolveTwistSpringMaterial(in int StrandsSize, in float RestLength, in float DeltaTime, in float MaterialDamping, 
+void SolveTwistSpringMaterial(in bool EnableConstraint, in int StrandsSize, in float RestLength, in float DeltaTime, in float MaterialDamping, 
 			in float MaterialCompliance, in float MaterialWeight, in float MaterialMultiplier, out float OutMaterialMultiplier)
 {
 	OutMaterialMultiplier = MaterialMultiplier;
-	const int LocalIndex = (GGroupThreadId.x % StrandsSize);
-	if( LocalIndex > 2)
+	if(EnableConstraint)
 	{
-		const int IsRed = ((GGroupThreadId.x/3) % 2) == 0;
-		// Process all the red rods
-		if (IsRed)
-		{	
-			UpdateStretchSpringMultiplier(RestLength,DeltaTime,false,MaterialDamping,MaterialCompliance,MaterialWeight,OutMaterialMultiplier,2);
-		}
-		// Process all the black rods
-		GroupMemoryBarrier(); 
-		if (!IsRed)
+		const int LocalIndex = (GGroupThreadId.x % StrandsSize);
+		if( LocalIndex > 2)
 		{
-			UpdateStretchSpringMultiplier(RestLength,DeltaTime,false,MaterialDamping,MaterialCompliance,MaterialWeight,OutMaterialMultiplier,2);
+			const int IsRed = ((GGroupThreadId.x/3) % 2) == 0;
+			// Process all the red rods
+			if (IsRed)
+			{	
+				UpdateStretchSpringMultiplier(RestLength,DeltaTime,false,MaterialDamping,MaterialCompliance,MaterialWeight,OutMaterialMultiplier,2);
+			}
+			// Process all the black rods
+			GroupMemoryBarrier(); 
+			if (!IsRed)
+			{
+				UpdateStretchSpringMultiplier(RestLength,DeltaTime,false,MaterialDamping,MaterialCompliance,MaterialWeight,OutMaterialMultiplier,2);
+			}
+			GroupMemoryBarrier(); 
 		}
-		GroupMemoryBarrier(); 
 	}
 } 
 
@@ -362,7 +467,7 @@ void UpdateAngularSpringMultiplier( in float RestLength, in float3 RestDirection
 	const float4 ParentFrame = SlerpQuat(SharedNodeOrientation[GGroupThreadId.x-2],SharedNodeOrientation[GGroupThreadId.x-2],0.5);
 	const float3 TargetDirection = RotateVectorByQuat(RestDirection,ParentFrame);
 
-	// XPBD lagrange multiplier update : dL = -(C+compliance*L) / (dC * invM * dCt + alpha)
+	// XPBD lagrang : dL = -(C+compliance*L) / (dC * invM * dCt + alpha)
 	const float3 DeltaLambda = -((EdgeDirection - TargetDirection)/RestLength + 
 					OutMaterialMultiplier * InMaterialCompliance  + MaterialDamping * DeltaVelocity /RestLength) * InMaterialWeight;
 
@@ -380,27 +485,30 @@ void UpdateAngularSpringMultiplier( in float RestLength, in float3 RestDirection
 	}
 }
 
-void SolveAngularSpringMaterial(in int StrandsSize, in float RestLength, in float3 RestDirection, in float DeltaTime, in float MaterialDamping,
+void SolveAngularSpringMaterial(in bool EnableConstraint, in int StrandsSize, in float RestLength, in float3 RestDirection, in float DeltaTime, in float MaterialDamping,
 			in float MaterialCompliance, in float MaterialWeight, in float3 MaterialMultiplier, out float3 OutMaterialMultiplier)
 {
-	
 	OutMaterialMultiplier = MaterialMultiplier;
-	const int LocalIndex = (GGroupThreadId.x % StrandsSize);
-	if( LocalIndex > 1)
-	{
-		const int IsRed = (GGroupThreadId.x % 2) == 0;
-		// Process all the red rods
-		if (!IsRed)
-		{	
-			UpdateAngularSpringMultiplier(RestLength,RestDirection,DeltaTime,false,MaterialDamping,MaterialCompliance,MaterialWeight,OutMaterialMultiplier);
-		}
-		GroupMemoryBarrier(); 
 
-		if (IsRed)
+	if(EnableConstraint)
+	{
+		const int LocalIndex = (GGroupThreadId.x % StrandsSize);
+		if( LocalIndex > 1)
 		{
-			UpdateAngularSpringMultiplier(RestLength,RestDirection,DeltaTime,false,MaterialDamping,MaterialCompliance,MaterialWeight,OutMaterialMultiplier);
+			const int IsRed = (GGroupThreadId.x % 2) == 0;
+			// Process all the red rods
+			if (!IsRed)
+			{	
+				UpdateAngularSpringMultiplier(RestLength,RestDirection,DeltaTime,false,MaterialDamping,MaterialCompliance,MaterialWeight,OutMaterialMultiplier);
+			}
+			GroupMemoryBarrier(); 
+
+			if (IsRed)
+			{
+				UpdateAngularSpringMultiplier(RestLength,RestDirection,DeltaTime,false,MaterialDamping,MaterialCompliance,MaterialWeight,OutMaterialMultiplier);
+			}
+			GroupMemoryBarrier(); 
 		}
-		GroupMemoryBarrier(); 
 	}
 }
 /*
