Copying //UE4/Dev-Physics to Main (//UE4/Main)

diff --git a/Engine/Plugins/Experimental/HairStrands/Shaders/Private/NiagaraStaticCollisionConstraint.ush b/Engine/Plugins/Experimental/HairStrands/Shaders/Private/NiagaraStaticCollisionConstraint.ush
--- a/Engine/Plugins/Experimental/HairStrands/Shaders/Private/NiagaraStaticCollisionConstraint.ush
+++ b/Engine/Plugins/Experimental/HairStrands/Shaders/Private/NiagaraStaticCollisionConstraint.ush
@@ -3,9 +3,10 @@
 #pragma once
 
 /* -----------------------------------------------------------------
- * Static collision constraint
+ * Static hard collision constraint
  * -----------------------------------------------------------------
  */
+
  float3 ComputeStaticCollisionDisplacement( in float3 DeltaPosition, in float3 DeltaVelocity, in float3 CollisionNormal, in float PenetrationDepth,
 							  in float StaticFriction, in float KineticFriction, in bool IsProjection)
  {
@@ -46,10 +47,10 @@
 	return float3(0,0,0);
  }
 
-void SolveStaticCollisionConstraint(in int StrandsSize, in float PenetrationDepth, in float3 CollisionPosition, in float3 CollisionVelocity, in float3 CollisionNormal, 
+ void SolveStaticCollisionConstraint(in bool EnableConstraint, in int StrandsSize, in float PenetrationDepth, in float3 CollisionPosition, in float3 CollisionVelocity, in float3 CollisionNormal, 
 				in float StaticFriction, in float KineticFriction, in float DeltaTime, in bool IsProjection, in float3 PreviousPosition, inout float3 OutNodePosition )
 {
-	if( GGroupThreadId.x % StrandsSize > 1)
+	if( GGroupThreadId.x % StrandsSize > 1 && EnableConstraint)
 	{
 		const float3 NodeVelocity = (OutNodePosition - PreviousPosition)/DeltaTime;
 		const float3 DeltaPosition = OutNodePosition - CollisionPosition;
@@ -58,4 +59,145 @@ void SolveStaticCollisionConstraint(in int StrandsSize, in float PenetrationDept
 		const float3 PositionDisplacement = ComputeStaticCollisionDisplacement(DeltaPosition,DeltaVelocity,CollisionNormal,PenetrationDepth,StaticFriction,KineticFriction,IsProjection);
 		OutNodePosition -= PositionDisplacement;
 	}
-}
\ No newline at end of file
+}
+
+/* -----------------------------------------------------------------
+ * Static soft collision constraint
+ * -----------------------------------------------------------------
+ */
+
+ void UpdateSoftCollisionConstraint(in float3 DeltaPosition, in float3 DeltaVelocity, in float3 CollisionNormal, in float PenetrationDepth,
+							  in float StaticFriction, in float KineticFriction, in float DeltaTime, in bool ProjectConstraint, in float MaterialDamping, 
+			in float MaterialCompliance, in float MaterialWeight, inout float3 OutMaterialMultiplier )
+{
+	const float DeltaVNormal = dot(DeltaVelocity,CollisionNormal);
+	const float DeltaPNormal = dot(DeltaPosition,CollisionNormal)-PenetrationDepth;
+
+	if (DeltaPNormal < 0.0)
+	{
+		const float LocalCompliance = ((PenetrationDepth <= 0.0) || ( PenetrationDepth > 0.0 && DeltaPNormal < -PenetrationDepth)) ? 0.0 : MaterialCompliance;// * max(0.0,min(1.0,(1.0 + DeltaPNormal / PenetrationDepth)));
+
+		const float DeltaFNormal = ProjectConstraint ? DeltaPNormal : DeltaVNormal;
+		const float NormDeltaFN = abs(DeltaFNormal);
+
+		const float3 CollisionTangent = DeltaVelocity - DeltaVNormal * CollisionNormal;
+		const float TangentLength = length(CollisionTangent);
+
+		float3 CollisionTangentA = (TangentLength > 0.0) ? CollisionTangent/TangentLength : (abs(CollisionNormal.x-1.0f) > 0.0) ? float3(1,0,0) : 
+																					        (abs(CollisionNormal.y-1.0f) > 0.0) ? float3(0,1,0) :
+																						     float3(0,0,1);
+		float3 CollisionTangentB = cross(CollisionNormal,CollisionTangentA);
+		CollisionTangentA = cross(CollisionTangentB,CollisionNormal);
+
+		const float TangentLengthA = length(CollisionTangentA);
+		const float TangentLengthB = length(CollisionTangentB);
+
+		CollisionTangentA = (TangentLengthA > 0.0) ? CollisionTangentA/TangentLengthA : float3(0,0,0);
+		CollisionTangentB = (TangentLengthB > 0.0) ? CollisionTangentB/TangentLengthB : float3(0,0,0);
+
+		const float DeltaVTangentA = dot( DeltaVelocity, CollisionTangentA);
+		const float DeltaVTangentB = dot( DeltaVelocity, CollisionTangentB);
+
+		const float NormDeltaVTA = abs(DeltaVTangentA);
+		const float NormDeltaVTB = abs(DeltaVTangentB);
+
+		const float AlphaTangentA = ( NormDeltaVTA < StaticFriction * NormDeltaFN ) ? 1.0 : (KineticFriction > 0.0) ? max(0.0,min(KineticFriction*NormDeltaFN/NormDeltaVTA, 1.0)) : 0.0;
+		const float AlphaTangentB = ( NormDeltaVTB < StaticFriction * NormDeltaFN ) ? 1.0 : (KineticFriction > 0.0) ? max(0.0,min(KineticFriction*NormDeltaFN/NormDeltaVTB, 1.0)) : 0.0;
+
+		const float ScaledInverseMass = ( (1.0  + MaterialDamping / DeltaTime ) * SharedInverseMass[GGroupThreadId.x] );
+
+		const float3 SchurDiagonal = float3(ScaledInverseMass + LocalCompliance, ScaledInverseMass + LocalCompliance, ScaledInverseMass + LocalCompliance );
+		const float MaterialWeightN = ( SchurDiagonal.x != 0.0 ) ? 1.0 / SchurDiagonal.x : 0.0;
+		const float MaterialWeightA = ( SchurDiagonal.y != 0.0 ) ? 1.0 / SchurDiagonal.y : 0.0;
+		const float MaterialWeightB = ( SchurDiagonal.z != 0.0 ) ? 1.0 / SchurDiagonal.z : 0.0;
+
+		const float3 DeltaLambda = float3(
+		    -( DeltaFNormal + OutMaterialMultiplier.x * LocalCompliance + MaterialDamping * DeltaVNormal / DeltaTime ) * MaterialWeightN,
+			-( DeltaVTangentA + OutMaterialMultiplier.y * LocalCompliance + MaterialDamping * DeltaVTangentA / DeltaTime ) * MaterialWeightA * AlphaTangentA,
+			-( DeltaVTangentB + OutMaterialMultiplier.z * LocalCompliance + MaterialDamping * DeltaVTangentB / DeltaTime ) * MaterialWeightB * AlphaTangentB);
+
+		OutMaterialMultiplier += DeltaLambda;
+
+		SharedNodePosition[GGroupThreadId.x] += SharedInverseMass[GGroupThreadId.x] * (DeltaLambda.x * CollisionNormal + DeltaLambda.y * CollisionTangentA + DeltaLambda.z * CollisionTangentB );
+	}
+	else
+	{
+		OutMaterialMultiplier = float3(0,0,0);
+	}
+}
+
+ void SetupSoftCollisionConstraint(in int StrandsSize, in float ConstraintStiffness,
+	in float DeltaTime, in float MaterialDamping, out float OutMaterialCompliance, out float OutMaterialWeight, out float3 OutMaterialMultiplier)
+{
+	OutMaterialCompliance = 0.0;
+	OutMaterialWeight = 0.0;
+	OutMaterialMultiplier = float3(0.0,0.0,0.0);
+
+	const int LocalIndex = (GGroupThreadId.x % StrandsSize);
+	if( LocalIndex > 1 )
+	{
+		OutMaterialCompliance = (ConstraintStiffness > 0.0) ? 1.0/(ConstraintStiffness*DeltaTime*DeltaTime) : 0.0;
+
+		//const float SumInverseMass = SharedInverseMass[GGroupThreadId.x];
+		//const float SchurDiagonal = ( (1.0  + MaterialDamping / DeltaTime ) * SumInverseMass + OutMaterialCompliance );
+		//OutMaterialWeight = ( SchurDiagonal != 0.0 ) ? 1.0 / SchurDiagonal : 0.0;
+	}
+}
+
+void SolveSoftCollisionConstraint(in bool EnableConstraint, in int StrandsSize, in float PenetrationDepth, in float3 CollisionPosition, in float3 CollisionVelocity, in float3 CollisionNormal, 
+				in float StaticFriction, in float KineticFriction, in float DeltaTime, in float MaterialDamping, 
+			in float MaterialCompliance, in float MaterialWeight, in float3 MaterialMultiplier, out float3 OutMaterialMultiplier )
+{
+	OutMaterialMultiplier = MaterialMultiplier;
+	if(EnableConstraint)
+	{
+		const int LocalIndex = (GGroupThreadId.x % StrandsSize);
+		if( LocalIndex > 1)
+		{
+			const float3 NodeVelocity = (SharedNodePosition[GGroupThreadId.x] - SharedPreviousPosition[GGroupThreadId.x])/DeltaTime;
+			const float3 DeltaPosition = SharedNodePosition[GGroupThreadId.x] - CollisionPosition;
+			const float3 DeltaVelocity = (NodeVelocity - CollisionVelocity) * DeltaTime;
+
+			UpdateSoftCollisionConstraint(DeltaPosition,DeltaVelocity,CollisionNormal,PenetrationDepth,
+					StaticFriction,KineticFriction,DeltaTime,false,MaterialDamping,MaterialCompliance,MaterialWeight,OutMaterialMultiplier);
+
+			GroupMemoryBarrier();
+		}
+	}
+}
+
+void ProjectSoftCollisionConstraint(in int StrandsSize, in float ConstraintStiffness, in float PenetrationDepth, in float3 CollisionPosition, in float3 CollisionVelocity, in float3 CollisionNormal, 
+				in float StaticFriction, in float KineticFriction, in float NodeMass, in float DeltaTime, in float3 PreviousPosition, in float3 NodePosition, out float3 OutNodePosition  )
+{
+	const int LocalIndex = (GGroupThreadId.x % StrandsSize);
+
+	SharedNodePosition[GGroupThreadId.x] = NodePosition;
+	SharedPreviousPosition[GGroupThreadId.x] = PreviousPosition;
+	SharedInverseMass[GGroupThreadId.x] = (LocalIndex > 1 && NodeMass != 0.0) ? 1.0 / NodeMass : 0.0;
+	GroupMemoryBarrier();
+
+	if(DeltaTime != 0.0)
+	{
+		float MaterialCompliance = 0.0;
+		float MaterialWeight = 0.0;
+		float3 MaterialMultiplier = float3(0.0,0.0,0.0);
+
+		if( LocalIndex > 1)
+		{
+			SetupSoftCollisionConstraint(StrandsSize,ConstraintStiffness,DeltaTime,0.0,MaterialCompliance,MaterialWeight,MaterialMultiplier);
+
+			const float3 NodeVelocity = (SharedNodePosition[GGroupThreadId.x] - SharedPreviousPosition[GGroupThreadId.x])/DeltaTime;
+			const float3 DeltaPosition = SharedNodePosition[GGroupThreadId.x] - CollisionPosition;
+			const float3 DeltaVelocity = (NodeVelocity - CollisionVelocity) * DeltaTime;
+
+			UpdateSoftCollisionConstraint(DeltaPosition,DeltaVelocity,CollisionNormal,PenetrationDepth,
+					StaticFriction,KineticFriction,DeltaTime,true,0.0,MaterialCompliance,MaterialWeight,MaterialMultiplier);
+
+			GroupMemoryBarrier();
+		}
+	}
+	GroupMemoryBarrier();
+	OutNodePosition = SharedNodePosition[GGroupThreadId.x];
+}
+
+
diff --git a/Engine/Plugins/Experimental/HairStrands/Source/HairStrandsCore/HairStrandsCore.Build.cs b/Engine/Plugins/Experimental/HairStrands/Source/HairStrandsCore/HairStrandsCore.Build.cs
--- a/Engine/Plugins/Experimental/HairStrands/Source/HairStrandsCore/HairStrandsCore.Build.cs
+++ b/Engine/Plugins/Experimental/HairStrands/Source/HairStrandsCore/HairStrandsCore.Build.cs
@@ -23,6 +23,15 @@ namespace UnrealBuildTool.Rules
 					"ChaosCore",
 					"Chaos"
 				});
+
+			if (Target.bBuildEditor == true)
+			{
+				PrivateDependencyModuleNames.AddRange(
+					new string[]
+					{
+					"DerivedDataCache",
+					});
+			}
 		}
 	}
 }
diff --git a/Engine/Plugins/Experimental/HairStrands/Source/HairStrandsCore/Private/GroomAsset.cpp b/Engine/Plugins/Experimental/HairStrands/Source/HairStrandsCore/Private/GroomAsset.cpp
--- a/Engine/Plugins/Experimental/HairStrands/Source/HairStrandsCore/Private/GroomAsset.cpp
+++ b/Engine/Plugins/Experimental/HairStrands/Source/HairStrandsCore/Private/GroomAsset.cpp
@@ -1,15 +1,34 @@
 // Copyright Epic Games, Inc. All Rights Reserved.
 
 #include "GroomAsset.h"
+#include "EngineUtils.h"
+#include "GroomAssetImportData.h"
 #include "GroomBuilder.h"
+#include "GroomImportOptions.h"
+#include "GroomSettings.h"
 #include "RenderingThread.h"
 #include "HairStrandsVertexFactory.h"
 #include "Misc/Paths.h"
+#include "Serialization/LargeMemoryReader.h"
+#include "Serialization/LargeMemoryWriter.h"
+#include "Serialization/MemoryReader.h"
+#include "Serialization/MemoryWriter.h"
+#include "UObject/PhysicsObjectVersion.h"
+#include "UObject/ReleaseObjectVersion.h"
 
 #if WITH_EDITORONLY_DATA
+#include "DerivedDataCacheInterface.h"
 #include "EditorFramework/AssetImportData.h"
 #endif
 
+template<typename BufferType>
+void UploadDataToBuffer(BufferType& OutBuffer, uint32 DataSizeInBytes, void* InCpuData)
+{
+	void* BufferData = RHILockVertexBuffer(OutBuffer.Buffer, 0, DataSizeInBytes, RLM_WriteOnly);
+	FMemory::Memcpy(BufferData, InCpuData, DataSizeInBytes);
+	RHIUnlockVertexBuffer(OutBuffer.Buffer);
+}
+
 template<typename FormatType>
 void CreateBuffer(const TArray<typename FormatType::Type>& InData, FRWBuffer& OutBuffer)
 {
@@ -94,6 +113,249 @@ void FHairStrandsDeformedResource::ReleaseRHI()
 
 /////////////////////////////////////////////////////////////////////////////////////////
 
+struct FClusterGrid
+{
+	struct FRenderingCurve
+	{
+		uint32 RenOffset;
+		uint32 RendPointCount;
+		float Area;
+	};
+	struct FCluster
+	{
+		TArray<FRenderingCurve> ClusterCurves;
+	};
+
+	FClusterGrid(float VoxelCentimeterSize, const FVector& InMinBound, const FVector& InMaxBound)
+	{
+		MinBound = InMinBound;
+		MaxBound = InMaxBound;
+
+		FVector VoxelCount = (MaxBound - MinBound) * FVector(1.0f / VoxelCentimeterSize);
+		GridResolution = FIntVector(FMath::CeilToInt(VoxelCount.X), FMath::CeilToInt(VoxelCount.Y), FMath::CeilToInt(VoxelCount.Z));
+
+		ClusterInfo.SetNum(GridResolution.X * GridResolution.Y * GridResolution.Z);
+	}
+
+	FORCEINLINE bool IsValid(const FIntVector& P) const
+	{
+		return	0 <= P.X && P.X < GridResolution.X &&
+			0 <= P.Y && P.Y < GridResolution.Y &&
+			0 <= P.Z && P.Z < GridResolution.Z;
+	}
+
+	FORCEINLINE FIntVector ClampToVolume(const FIntVector& CellCoord, bool& bIsValid) const
+	{
+		bIsValid = IsValid(CellCoord);
+		return FIntVector(
+			FMath::Clamp(CellCoord.X, 0, GridResolution.X - 1),
+			FMath::Clamp(CellCoord.Y, 0, GridResolution.Y - 1),
+			FMath::Clamp(CellCoord.Z, 0, GridResolution.Z - 1));
+	}
+
+	FORCEINLINE FIntVector ToCellCoord(const FVector& P) const
+	{
+		bool bIsValid = false;
+		const FVector F = ((P - MinBound) / (MaxBound - MinBound));
+		const FIntVector CellCoord = FIntVector(FMath::FloorToInt(F.X * GridResolution.X), FMath::FloorToInt(F.Y * GridResolution.Y), FMath::FloorToInt(F.Z * GridResolution.Z));
+		return ClampToVolume(CellCoord, bIsValid);
+	}
+
+	uint32 ToIndex(const FIntVector& CellCoord) const
+	{
+		uint32 CellIndex = CellCoord.X + CellCoord.Y * GridResolution.X + CellCoord.Z * GridResolution.X * GridResolution.Y;
+		check(CellIndex < uint32(ClusterInfo.Num()));
+		return CellIndex;
+	}
+
+	void InsertRenderingCurve(FRenderingCurve& Curve, const FVector& Root)
+	{
+		FIntVector CellCoord = ToCellCoord(Root);
+		uint32 Index = ToIndex(CellCoord);
+		FCluster& Cluster = ClusterInfo[Index];
+		Cluster.ClusterCurves.Add(Curve);
+	}
+
+	TArray<FCluster> GetAllClusters() { return ClusterInfo; }
+
+private:
+
+	FVector MinBound;
+	FVector MaxBound;
+
+	FIntVector GridResolution;
+
+	TArray<FCluster> ClusterInfo;
+};
+
+FHairStrandsClusterCullingResource::FHairStrandsClusterCullingResource(const FHairGroupData& GroupData)
+{
+	const FHairStrandsDatas& SimStrandsData = GroupData.HairSimulationData;
+	const FHairStrandsDatas& RenStrandsData = GroupData.HairRenderData;
+	const FHairStrandsInterpolationDatas& InterpolationData = GroupData.HairInterpolationData;
+	const uint32 RenCurveCount = RenStrandsData.GetNumCurves();
+	const uint32 SimCurveCount = SimStrandsData.GetNumCurves();
+	const uint32 PointCount = RenStrandsData.GetNumPoints();
+
+
+	// Allocate look up arrays for as many hair vertices as needed.
+	VertexCount = PointCount;
+	VertexToClusterIdArray.SetNum(PointCount);
+	ClusterVertexIdArray.SetNum(PointCount);
+
+	// Allocate cluster per voxel containing contains >=1 render curve root
+	const uint32 VoxelcountAlongLargerSide = 256; // TODO expose in UI
+	const float VoxelCentimeterSize = GroupData.HairRenderData.BoundingBox.GetSize().GetAbsMax() / VoxelcountAlongLargerSide;
+	FClusterGrid ClusterGrid(VoxelCentimeterSize, GroupData.HairRenderData.BoundingBox.Min, GroupData.HairRenderData.BoundingBox.Max);
+
+	for (uint32 RenCurveIndex = 0; RenCurveIndex < RenCurveCount; ++RenCurveIndex)
+	{
+		FClusterGrid::FRenderingCurve RCurve;
+		RCurve.RendPointCount = RenStrandsData.StrandsCurves.CurvesCount[RenCurveIndex];
+		RCurve.RenOffset = RenStrandsData.StrandsCurves.CurvesOffset[RenCurveIndex];
+
+		RCurve.Area = 0.0f;
+		for (uint32 RenPointIndex = 0; RenPointIndex < RCurve.RendPointCount; ++RenPointIndex)
+		{
+			uint32 PointGlobalIndex = RenPointIndex + RCurve.RenOffset;
+			const FVector& V0 = RenStrandsData.StrandsPoints.PointsPosition[PointGlobalIndex];
+			if (RenPointIndex > 0)
+			{
+				const FVector& V1 = RenStrandsData.StrandsPoints.PointsPosition[PointGlobalIndex-1];
+				FVector OutDir;
+				float OutLength;
+				(V1 - V0).ToDirectionAndLength(OutDir, OutLength);
+				RCurve.Area += RenStrandsData.StrandsPoints.PointsRadius[PointGlobalIndex] * OutLength;
+			}
+		}
+
+		const FVector Root = RenStrandsData.StrandsPoints.PointsPosition[RCurve.RenOffset];
+		ClusterGrid.InsertRenderingCurve(RCurve, Root);
+	}
+
+	// Count Cluster
+	ClusterCount = 0;
+	for (auto& Cluster : ClusterGrid.GetAllClusters())
+	{
+		ClusterCount += Cluster.ClusterCurves.Num() > 0 ? 1 : 0;
+	}
+	ClusterInfoArray.SetNum(ClusterCount);
+	ClusterIndexRadiusScaleInfoArray.SetNum(ClusterCount);
+
+	TArray<TArray<uint32>> ClusterToVertexIndexLod0;	// List of indices per cluster
+	ClusterToVertexIndexLod0.SetNum(ClusterCount);
+	TArray<uint32> ClusterVertexIndexCountLod1;			// Index count per cluster
+	ClusterVertexIndexCountLod1.SetNum(ClusterCount);
+
+	// Write out cluster information
+	uint32 Index = 0;
+	for (auto& Cluster : ClusterGrid.GetAllClusters())
+	{
+		if (Cluster.ClusterCurves.Num() == 0)
+		{
+			continue;
+		}
+
+		uint32 ClusterCurveCount = Cluster.ClusterCurves.Num();
+		uint32 Lod1CurveCount = FMath::Max(ClusterCurveCount / 4, 1u);		// a fourth of the curves for the lowest lod
+		uint32 Lod0VertexCount = 0;
+		uint32 Lod1VertexCount = 0;
+
+		// Prepare data to reconstruct vertex/index count to radius scale
+		float Lod0StrandArea = 0.0f;
+		float Lod1StrandArea = 0.0f;
+
+		// Sort curve to have largest area first, so that lower area curves with less influence are removed first.
+		// This also helps the radius scaling to not explode.
+		Cluster.ClusterCurves.Sort([](const FClusterGrid::FRenderingCurve& A, const FClusterGrid::FRenderingCurve& B) -> bool
+		{
+				return A.Area > B.Area;
+		});
+
+		uint32 CurveIndex = 0;
+		for (auto& ClusterCurve : Cluster.ClusterCurves)
+		{
+			for (uint32 RenPointIndex = 0; RenPointIndex < ClusterCurve.RendPointCount; ++RenPointIndex)
+			{
+				uint32 PointGlobalIndex = RenPointIndex + ClusterCurve.RenOffset;
+				VertexToClusterIdArray[PointGlobalIndex] = Index;
+				ClusterToVertexIndexLod0[Index].Add(PointGlobalIndex);
+			}
+			Lod0VertexCount += ClusterCurve.RendPointCount;
+			Lod0StrandArea += ClusterCurve.Area;
+			if (CurveIndex < Lod1CurveCount)
+			{
+				Lod1VertexCount += ClusterCurve.RendPointCount;
+				Lod1StrandArea += ClusterCurve.Area;
+			}
+			CurveIndex++;
+		}
+		if (Cluster.ClusterCurves.Num() > 0)
+		{
+			ClusterVertexIndexCountLod1[Index] = Lod1VertexCount;
+
+			const float ClusterRadiusScaleLod1 = Lod0StrandArea / Lod1StrandArea; // Could use coverage ratio (relative to cluster 0&2 bounding boxes) instead of hair area.
+			const float Numer = ClusterRadiusScaleLod1 - 1.0f;
+			const float Denom = float(Lod1VertexCount) - float(Lod0VertexCount);
+			const float AValue = Denom == 0.0f ? 0.0f : Numer / Denom;
+			ClusterIndexRadiusScaleInfoArray[Index] = AValue;
+			// ClusterIndexRadiusScaleInfoArray contains the 'a' value from the linear equation y=ax+b. Radius scale can then be recovered as ( 1.0 + a * (VertexCount - Lod0VertexCount) )
+		}
+
+		Index += Cluster.ClusterCurves.Num() > 0 ? 1 : 0;
+	}
+
+	// Now compute the cluster/strand information
+	uint32 VertexCountSum = 0;
+	for (uint32 ClusterId = 0; ClusterId < ClusterCount; ++ClusterId)
+	{
+		const uint32 ClusterVertexCountLod0 = ClusterToVertexIndexLod0[ClusterId].Num();
+		ClusterInfoArray[ClusterId].FirstVertexId = VertexCountSum;
+		ClusterInfoArray[ClusterId].VertexIdCountLodHigh = ClusterVertexCountLod0;
+		ClusterInfoArray[ClusterId].VertexIdCountLodLow = ClusterVertexIndexCountLod1[ClusterId];
+		ClusterInfoArray[ClusterId].UnusedUint = 0;
+		VertexCountSum += ClusterVertexCountLod0;
+	}
+	check(ClusterVertexIdArray.Num() == VertexCountSum);
+
+	// Flatten global VertexId for each cluster
+	for (uint32 ClusterId = 0; ClusterId < ClusterCount; ++ClusterId)
+	{
+		uint32 ClusterFirstVertexId = ClusterInfoArray[ClusterId].FirstVertexId;
+		uint32 ClusterVertexIdCount = ClusterInfoArray[ClusterId].VertexIdCountLodHigh;
+		TArray<uint32>& ClusterVertices = ClusterToVertexIndexLod0[ClusterId];
+		for (uint32 ClusterVertexId = 0; ClusterVertexId < ClusterVertexIdCount; ++ClusterVertexId)
+		{
+			ClusterVertexIdArray[ClusterFirstVertexId + ClusterVertexId] = ClusterVertices[ClusterVertexId];
+		}
+	}
+}
+
+void FHairStrandsClusterCullingResource::InitRHI()
+{
+	ClusterInfoBuffer.Initialize(sizeof(uint32) * 4, ClusterCount, EPixelFormat::PF_R32G32B32A32_UINT, BUF_Static);
+	UploadDataToBuffer(ClusterInfoBuffer, sizeof(uint32) * 4 * ClusterCount, ClusterInfoArray.GetData());
+
+	VertexToClusterIdBuffer.Initialize(sizeof(uint32), VertexToClusterIdArray.Num(), EPixelFormat::PF_R32_UINT, BUF_Static);
+	UploadDataToBuffer(VertexToClusterIdBuffer, sizeof(uint32) * VertexToClusterIdArray.Num(), VertexToClusterIdArray.GetData());
+
+	ClusterVertexIdBuffer.Initialize(sizeof(uint32), ClusterVertexIdArray.Num(), EPixelFormat::PF_R32_UINT, BUF_Static);
+	UploadDataToBuffer(ClusterVertexIdBuffer, sizeof(uint32) * ClusterVertexIdArray.Num(), ClusterVertexIdArray.GetData());
+
+	ClusterIndexRadiusScaleInfoBuffer.Initialize(sizeof(float), ClusterIndexRadiusScaleInfoArray.Num(), EPixelFormat::PF_R32_FLOAT, BUF_Static); // TODO use fp16
+	UploadDataToBuffer(ClusterIndexRadiusScaleInfoBuffer, sizeof(float) * ClusterIndexRadiusScaleInfoArray.Num(), ClusterIndexRadiusScaleInfoArray.GetData());
+}
+
+void FHairStrandsClusterCullingResource::ReleaseRHI()
+{
+	ClusterInfoBuffer.Release();
+	VertexToClusterIdBuffer.Release();
+	ClusterVertexIdBuffer.Release();
+	ClusterIndexRadiusScaleInfoBuffer.Release();
+}
+
+/////////////////////////////////////////////////////////////////////////////////////////
+
 FHairStrandsRootResource::FHairStrandsRootResource(const FHairStrandsDatas* HairStrandsDatas, uint32 LODCount):
 	RootCount(HairStrandsDatas ? HairStrandsDatas->GetNumCurves() : 0)
 {
@@ -146,7 +408,7 @@ FHairStrandsRootResource::FHairStrandsRootResource(const FHairStrandsDatas* Hair
 	uint32 LODIndex = 0;
 	for (FMeshProjectionLOD& MeshProjectionLOD : MeshProjectionLODs)
 	{
-		MeshProjectionLOD.Status = FMeshProjectionLOD::EStatus::Invalid;
+		MeshProjectionLOD.Status = FHairStrandsProjectionHairData::LODData::EStatus::Invalid;
 		MeshProjectionLOD.RestRootOffset = HairStrandsDatas->BoundingBox.GetCenter();
 		MeshProjectionLOD.LODIndex = LODIndex++;
 	}
@@ -161,7 +423,7 @@ void FHairStrandsRootResource::InitRHI()
 		CreateBuffer<FHairStrandsRootNormalFormat>(RootNormals, RootNormalBuffer);	
 		for (FMeshProjectionLOD& LOD : MeshProjectionLODs)
 		{
-			LOD.Status = FHairStrandsRootResource::FMeshProjectionLOD::EStatus::Initialized;
+			LOD.Status = FHairStrandsProjectionHairData::LODData::EStatus::Initialized;
 
 			// Create buffers. Initialization will be done by render passes
 			CreateBuffer<FHairStrandsCurveTriangleIndexFormat>(RootCount, LOD.RootTriangleIndexBuffer);
@@ -186,7 +448,7 @@ void FHairStrandsRootResource::ReleaseRHI()
 
 	for (FMeshProjectionLOD& LOD : MeshProjectionLODs)
 	{
-		LOD.Status = FHairStrandsRootResource::FMeshProjectionLOD::EStatus::Invalid;
+		LOD.Status = FHairStrandsProjectionHairData::LODData::EStatus::Invalid;
 		LOD.RootTriangleIndexBuffer.Release();
 		LOD.RootTriangleBarycentricBuffer.Release();
 		LOD.RestRootTrianglePosition0Buffer.Release();
@@ -260,7 +522,44 @@ void UGroomAsset::Serialize(FArchive& Ar)
 {
 	Super::Serialize(Ar);
 
-	Ar << HairGroupsData;
+	Ar.UsingCustomVersion(FPhysicsObjectVersion::GUID);
+	Ar.UsingCustomVersion(FReleaseObjectVersion::GUID); // Needed to support MeshDescription AttributesSet serialization
+
+	if (Ar.CustomVer(FPhysicsObjectVersion::GUID) >= FPhysicsObjectVersion::GroomWithDescription)
+	{
+		FStripDataFlags StripFlags(Ar);
+		if (StripFlags.IsEditorDataStripped() || (Ar.IsSaving() && !CanRebuildFromDescription()))
+		{
+			// When cooking data or serializing old format to new format,
+			// serialize the computed groom data
+			Ar << HairGroupsData;
+		}
+#if WITH_EDITORONLY_DATA
+		else
+		{
+			// When serializing data for editor, serialize the HairDescription as bulk data
+			// The computed groom data is fetched from the Derived Data Cache
+			if (!HairDescriptionBulkData)
+			{
+				// When loading, bulk data can be null so instantiate a new one to serialize into
+				HairDescriptionBulkData = MakeUnique<FHairDescriptionBulkData>();
+			}
+
+			HairDescriptionBulkData->Serialize(Ar, this);
+
+			// Serialize the HairGroupsData directly into the asset if it couldn't be cached in the DDC
+			if (!bIsCacheable)
+			{
+				Ar << HairGroupsData;
+			}
+		}
+#endif // WITH_EDITORONLY_DATA
+	}
+	else
+	{
+		// Old format serialized the computed groom data directly
+		Ar << HairGroupsData;
+	}
 }
 
 void UGroomAsset::InitResource()
@@ -335,10 +634,24 @@ void UGroomAsset::PostLoad()
 {
 	Super::PostLoad();
 
+#if WITH_EDITORONLY_DATA
+	if (bIsCacheable)
+	{
+		CacheDerivedData();
+	}
+#endif
+
 	check(HairGroupsData.Num() > 0);
 	if (HairGroupsData[0].HairSimulationData.GetNumCurves() == 0 || HairGroupsData[0].HairInterpolationData.Num() == 0)
 	{
-		FGroomBuilder::BuildData(this, 2, 0, false, false);
+		FGroomBuildSettings Settings;
+		Settings.bOverrideGuides = false;
+		Settings.HairToGuideDensity = 0.1f;
+		Settings.InterpolationQuality = EGroomInterpolationQuality::High;
+		Settings.InterpolationDistance = EGroomInterpolationWeight::Parametric;
+		Settings.bRandomizeGuide = false;
+		Settings.bUseUniqueGuide = false;
+		FGroomBuilder::BuildData(this, Settings);
 	}
 
 	if (!IsTemplate())
@@ -397,3 +710,226 @@ FArchive& operator<<(FArchive& Ar, FHairGroupData& GroupData)
 
 	return Ar;
 }
+
+FArchive& operator<<(FArchive& Ar, FHairGroupInfo& GroupInfo)
+{
+	Ar.UsingCustomVersion(FPhysicsObjectVersion::GUID);
+
+	Ar << GroupInfo.GroupID;
+	Ar << GroupInfo.NumCurves;
+	Ar << GroupInfo.NumGuides;
+
+	if (Ar.CustomVer(FPhysicsObjectVersion::GUID) >= FPhysicsObjectVersion::GroomWithImportSettings)
+	{
+		Ar << GroupInfo.ImportSettings;
+	}
+	else
+	{
+		if (Ar.IsLoading())
+		{
+			bool bIsAutoGenerated = false;
+			Ar << bIsAutoGenerated;
+
+			GroupInfo.ImportSettings = FGroomBuildSettings();
+			GroupInfo.ImportSettings.InterpolationQuality = EGroomInterpolationQuality::Unknown;
+			GroupInfo.ImportSettings.InterpolationDistance = EGroomInterpolationWeight::Unknown;
+			GroupInfo.ImportSettings.bOverrideGuides = bIsAutoGenerated;
+		}
+		else
+		{
+			bool bIsAutoGenerated = GroupInfo.ImportSettings.bOverrideGuides;
+			Ar << bIsAutoGenerated;
+		}
+	}
+
+	// Ignoring Material since it's null
+
+	return Ar;
+}
+
+bool UGroomAsset::CanRebuildFromDescription() const
+{
+#if WITH_EDITORONLY_DATA
+	return HairDescriptionBulkData.IsValid() && !HairDescriptionBulkData->IsEmpty();
+#else
+	return false;
+#endif
+}
+
+// If groom derived data needs to be rebuilt (new format, serialization
+// differences, etc.) replace the version GUID below with a new one.
+// In case of merge conflicts with DDC versions, you MUST generate a new GUID
+// and set this new GUID as the version.
+#define GROOM_DERIVED_DATA_VERSION TEXT("42406925F1934F9794A08A8EE51BA05D")
+
+#if WITH_EDITORONLY_DATA
+
+namespace GroomDerivedDataCacheUtils
+{
+	const FString& GetGroomDerivedDataVersion()
+	{
+		static FString CachedVersionString(GROOM_DERIVED_DATA_VERSION);
+		return CachedVersionString;
+	}
+
+	FString BuildGroomDerivedDataKey(const FString& KeySuffix)
+	{
+		return FDerivedDataCacheInterface::BuildCacheKey(TEXT("GROOM"), *GetGroomDerivedDataVersion(), *KeySuffix);
+	}
+
+	void SerializeGroomBuildSettingsForDDC(FArchive& Ar, FGroomBuildSettings& BuildSettings)
+	{
+		// Note: this serializer is only used to build the groom DDC key, no versioning is required
+		Ar << BuildSettings.bOverrideGuides;
+		Ar << BuildSettings.HairToGuideDensity;
+		Ar << BuildSettings.InterpolationQuality;
+		Ar << BuildSettings.InterpolationDistance;
+		Ar << BuildSettings.bRandomizeGuide;
+		Ar << BuildSettings.bUseUniqueGuide;
+	}
+}
+
+FString UGroomAsset::BuildDerivedDataKeySuffix(const FGroomBuildSettings& BuildSettings)
+{
+	// Serialize the build settings into a temporary array
+	// The archive is flagged as persistent so that machines of different endianness produce identical binary results.
+	TArray<uint8> TempBytes;
+	TempBytes.Reserve(64);
+	FMemoryWriter Ar(TempBytes, /*bIsPersistent=*/ true);
+
+	GroomDerivedDataCacheUtils::SerializeGroomBuildSettingsForDDC(Ar, const_cast<FGroomBuildSettings&>(BuildSettings));
+
+	FString KeySuffix;
+	if (HairDescriptionBulkData)
+	{
+		// Reserve twice the size of TempBytes because of ByteToHex below + 3 for "ID" and \0
+		KeySuffix.Reserve(HairDescriptionBulkData->GetIdString().Len() + TempBytes.Num() * 2 + 3);
+		KeySuffix += TEXT("ID");
+		KeySuffix += HairDescriptionBulkData->GetIdString();
+	}
+	else
+	{
+		KeySuffix.Reserve(TempBytes.Num() * 2 + 1);
+	}
+
+	// Now convert the raw bytes to a string
+	const uint8* SettingsAsBytes = TempBytes.GetData();
+	for (int32 ByteIndex = 0; ByteIndex < TempBytes.Num(); ++ByteIndex)
+	{
+		ByteToHex(SettingsAsBytes[ByteIndex], KeySuffix);
+	}
+
+	return KeySuffix;
+}
+
+void UGroomAsset::CommitHairDescription(FHairDescription&& InHairDescription)
+{
+	HairDescription = MakeUnique<FHairDescription>(InHairDescription);
+
+	if (!HairDescriptionBulkData)
+	{
+		HairDescriptionBulkData = MakeUnique<FHairDescriptionBulkData>();
+	}
+	HairDescriptionBulkData->SaveHairDescription(*HairDescription);
+}
+
+bool UGroomAsset::CacheDerivedData(const FGroomBuildSettings* InBuildSettings)
+{
+	if (!HairDescriptionBulkData)
+	{
+		return false;
+	}
+
+	static FGroomBuildSettings DefaultBuildSettings;
+	const FGroomBuildSettings* BuildSettings = &DefaultBuildSettings;
+	if (InBuildSettings)
+	{
+		// Use the provided build settings since a new build is requested
+		BuildSettings = InBuildSettings;
+	}
+	else if (AssetImportData)
+	{
+		// Otherwise, use the build settings that were used previously at import/re-import
+		UGroomAssetImportData* GroomImportData = Cast<UGroomAssetImportData>(AssetImportData);
+		if (GroomImportData && GroomImportData->ImportOptions)
+		{
+			BuildSettings = &GroomImportData->ImportOptions->BuildSettings;
+		}
+	}
+
+	bool bSuccess = true;
+
+	const FString KeySuffix = BuildDerivedDataKeySuffix(*BuildSettings);
+	const FString DerivedDataKey = GroomDerivedDataCacheUtils::BuildGroomDerivedDataKey(KeySuffix);
+
+	TArray<uint8> DerivedData;
+	if (GetDerivedDataCacheRef().GetSynchronous(*DerivedDataKey, DerivedData))
+	{
+		FMemoryReader Ar(DerivedData, /*bIsPersistent=*/ true);
+
+		int64 UncompressedSize = 0;
+		Ar << UncompressedSize;
+
+		uint8* DecompressionBuffer = reinterpret_cast<uint8*>(FMemory::Malloc(UncompressedSize));
+		Ar.SerializeCompressed(DecompressionBuffer, 0, NAME_Zlib);
+
+		FLargeMemoryReader LargeMemReader(DecompressionBuffer, UncompressedSize, ELargeMemoryReaderFlags::Persistent | ELargeMemoryReaderFlags::TakeOwnership);
+		if (HairGroupsInfo.Num() > 0)
+		{
+			// HairGroupsInfo serialized from the asset, so ignore the one from the DDC
+			TArray<FHairGroupInfo> TempHairGroupsInfo;
+			LargeMemReader << TempHairGroupsInfo;
+		}
+		else
+		{
+			// This is the case where a groom with the same build settings was previously pushed to DDC,
+			// but this is a new asset so HairGroupsInfo isn't serialized yet
+			LargeMemReader << HairGroupsInfo;
+		}
+
+		LargeMemReader << HairGroupsData;
+	}
+	else
+	{
+		// Load the HairDescription from the bulk data if needed
+		if (!HairDescription)
+		{
+			HairDescription = MakeUnique<FHairDescription>();
+			HairDescriptionBulkData->LoadHairDescription(*HairDescription);
+		}
+
+		// Build groom data with the new build settings
+		bSuccess = FGroomBuilder::BuildGroom(*HairDescription, *BuildSettings, this);
+		if (bSuccess)
+		{
+			// Using a LargeMemoryWriter for serialization since the data can be bigger than 2 GB
+			FLargeMemoryWriter LargeMemWriter(0, /*bIsPersistent=*/ true);
+			LargeMemWriter << HairGroupsInfo;
+			LargeMemWriter << HairGroupsData;
+
+			int64 UncompressedSize = LargeMemWriter.TotalSize();
+
+			// Then the content of the LargeMemWriter is compressed into a MemoryWriter
+			// Compression ratio can reach about 5:2 depending on the data
+			// Since the DDC doesn't support data bigger than 2 GB
+			// we can compute a size threshold to skip the caching when
+			// the uncompressed size exceeds the threshold
+			static constexpr const int64 SizeThreshold =  (int64) MAX_int32 * 2.5;
+			bIsCacheable = UncompressedSize < SizeThreshold;
+			if (bIsCacheable)
+			{
+				FMemoryWriter CompressedArchive(DerivedData, true);
+
+				CompressedArchive << UncompressedSize; // needed for allocating decompression buffer
+				CompressedArchive.SerializeCompressed(LargeMemWriter.GetData(), UncompressedSize, NAME_Zlib);
+
+				GetDerivedDataCacheRef().Put(*DerivedDataKey, DerivedData);
+			}
+		}
+	}
+	return bSuccess;
+}
+
+#endif // WITH_EDITORONLY_DATA
+
+#undef LOCTEXT_NAMESPACE
diff --git a/Engine/Plugins/Experimental/HairStrands/Source/HairStrandsCore/Private/GroomBuilder.cpp b/Engine/Plugins/Experimental/HairStrands/Source/HairStrandsCore/Private/GroomBuilder.cpp
--- a/Engine/Plugins/Experimental/HairStrands/Source/HairStrandsCore/Private/GroomBuilder.cpp
+++ b/Engine/Plugins/Experimental/HairStrands/Source/HairStrandsCore/Private/GroomBuilder.cpp
@@ -5,6 +5,7 @@
 #include "GroomComponent.h"
 #include "GroomSettings.h"
 #include "HairDescription.h"
+#include "GroomSettings.h"
 
 #include "Async/ParallelFor.h"
 #include "HAL/IConsoleManager.h"
@@ -194,7 +195,7 @@ namespace HairStrandsBuilder
 	}
 
 	/** Build the packed datas for gpu rendering/simulation */
-	void BuildRenderData(FHairStrandsDatas& HairStrands)
+	void BuildRenderData(FHairStrandsDatas& HairStrands, const TArray<uint8>& RandomSeeds)
 	{
 		TRACE_CPUPROFILER_EVENT_SCOPE(HairStrandsBuilder::BuildRenderData);
 
@@ -216,10 +217,9 @@ namespace HairStrandsBuilder
 		FHairStrandsCurves& Curves = HairStrands.StrandsCurves;
 		FHairStrandsPoints& Points = HairStrands.StrandsPoints;
 
-		FRandomStream Random;
 		for (uint32 CurveIndex = 0; CurveIndex < NumCurves; ++CurveIndex)
 		{
-			const float CurveSeed = Random.RandHelper(255);
+			const uint8 CurveSeed = RandomSeeds[CurveIndex];
 			const int32 IndexOffset = Curves.CurvesOffset[CurveIndex];
 			const uint16& PointCount = Curves.CurvesCount[CurveIndex];
 			for (int32 PointIndex = 0; PointIndex < PointCount; ++PointIndex)
@@ -394,7 +394,7 @@ namespace HairInterpolationBuilder
 	// Randomize influence guide to break interpolation coherence, and create a more random/natural pattern
 	static void SelectFinalGuides(
 		FClosestGuides& ClosestGuides, 
-		FRandomStream& Random, 
+		const FIntVector& RandomIndices,
 		const FMetrics& InMetric, 
 		const bool bRandomizeInterpolation, 
 		const bool bUseUniqueGuide)
@@ -420,9 +420,9 @@ namespace HairInterpolationBuilder
 		if (bRandomizeInterpolation)
 		{
 			// This randomization makes certain strands being affected by 1, 2, or 3 guides
-			RandIndex0 = Random.RandRange(0, FMetrics::Count - 1);
-			RandIndex1 = Random.RandRange(0, FMetrics::Count - 1);
-			RandIndex2 = Random.RandRange(0, FMetrics::Count - 1);
+			RandIndex0 = RandomIndices[0];
+			RandIndex1 = RandomIndices[1];
+			RandIndex2 = RandomIndices[2];
 		}
 
 		ClosestGuides.Indices[0] = Metric.KClosestGuideIndices[RandIndex0];
@@ -621,7 +621,7 @@ namespace HairInterpolationBuilder
 			const FHairStrandsDatas& SimStrandsData,
 			const bool bRandomized,
 			const bool bUnique,
-			FRandomStream& Random) const
+			const FIntVector& RandomIndices) const
 		{
 			const FHairRoot& RenRoot = RenRoots[RenCurveIndex];
 			const FIntVector PointCoord = ToCellCoord(RenRoot.Position);
@@ -686,7 +686,7 @@ namespace HairInterpolationBuilder
 
 			// If no valid guide have been found, switch to a simpler metric
 			FClosestGuides ClosestGuides;
-			SelectFinalGuides(ClosestGuides, Random, Metrics, bRandomized, bUnique);
+			SelectFinalGuides(ClosestGuides, RandomIndices, Metrics, bRandomized, bUnique);
 
 			check(ClosestGuides.Indices[0] >= 0);
 			check(ClosestGuides.Indices[1] >= 0);
@@ -731,7 +731,7 @@ namespace HairInterpolationBuilder
 			const FHairStrandsDatas& SimStrandsData,
 			const bool bRandomized,
 			const bool bUnique,
-			FRandomStream& Random) const 
+			const FIntVector& RandomIndices) const
 		{
 			const FHairRoot& RenRoot = RenRoots[RenCurveIndex];
 			const FIntVector PointCoord = ToCellCoord(RenRoot.Position);
@@ -801,11 +801,11 @@ namespace HairInterpolationBuilder
 			FClosestGuides ClosestGuides;
 			if (Metrics0.KClosestGuideIndices[0] != -1)
 			{
-				SelectFinalGuides(ClosestGuides, Random, Metrics0, bRandomized, bUnique);
+				SelectFinalGuides(ClosestGuides, RandomIndices, Metrics0, bRandomized, bUnique);
 			}
 			else
 			{
-				SelectFinalGuides(ClosestGuides, Random, Metrics1, bRandomized, bUnique);
+				SelectFinalGuides(ClosestGuides, RandomIndices, Metrics1, bRandomized, bUnique);
 			}
 
 			check(ClosestGuides.Indices[0] >= 0);
@@ -824,7 +824,7 @@ namespace HairInterpolationBuilder
 		const FHairStrandsDatas& SimStrandsData,
 		const bool bRandomized,
 		const bool bUnique,
-		FRandomStream& Random)
+		const FIntVector& RandomIndices)
 	{
 		FMetrics Metrics;
 		for (uint32 ClearIt = 0; ClearIt < FMetrics::Count; ++ClearIt)
@@ -840,7 +840,7 @@ namespace HairInterpolationBuilder
 		}
 			
 		FClosestGuides ClosestGuides;
-		SelectFinalGuides(ClosestGuides, Random, Metrics, bRandomized, bUnique);
+		SelectFinalGuides(ClosestGuides, RandomIndices, Metrics, bRandomized, bUnique);
 
 		check(ClosestGuides.Indices[0] >= 0);
 		check(ClosestGuides.Indices[1] >= 0);
@@ -916,33 +916,12 @@ namespace HairInterpolationBuilder
 		return Desc;
 	}
 
-	enum class EHairInterpolationWeightMethod : uint8
-	{
-		ParametricDistance,
-		RootDistance,
-		VertexIndex
-	};
-
-	enum class EHairInterpolationDataQuality : uint8
-	{
-		Low,
-		Medium,
-		High
-	};
-
-	struct FHairInterpolationSettings
-	{
-		EHairInterpolationDataQuality Quality		= EHairInterpolationDataQuality::High;
-		EHairInterpolationWeightMethod WeightMethod = EHairInterpolationWeightMethod::ParametricDistance;
-		bool bRandomizedGuides						= false;
-		bool bUseUniqueGuide						= false;
-	};
-
 	static void BuildInterpolationData(
 		FHairStrandsInterpolationDatas& InterpolationData,
 		const FHairStrandsDatas& SimStrandsData,
 		const FHairStrandsDatas& RenStrandsData,
-		const FHairInterpolationSettings& Settings)
+		const FGroomBuildSettings& Settings,
+		const TArray<FIntVector>& RandomGuideIndices)
 	{
 		TRACE_CPUPROFILER_EVENT_SCOPE(HairInterpolationBuilder::BuildInterpolationData);
 
@@ -961,7 +940,7 @@ namespace HairInterpolationBuilder
 			ExtractRoots(RenStrandsData, RenRoots, RenMinBound, RenMaxBound);
 			ExtractRoots(SimStrandsData, SimRoots, SimMinBound, SimMaxBound);
 
-			if (Settings.Quality == EHairInterpolationDataQuality::Low || Settings.Quality == EHairInterpolationDataQuality::Medium)
+			if (Settings.InterpolationQuality == EGroomInterpolationQuality::Low || Settings.InterpolationQuality == EGroomInterpolationQuality::Medium)
 			{
 				// Build a conservative bound, to insure all queries will fall 
 				// into the grid volume.
@@ -975,7 +954,6 @@ namespace HairInterpolationBuilder
 		uint32 TotalInvalidInterpolationCount = 0;
 		const static float MinWeightDistance = 0.0001f;
 
-		FRandomStream Random;
 		const uint32 RenCurveCount = RenStrandsData.GetNumCurves();
 		const uint32 SimCurveCount = SimStrandsData.GetNumCurves();
 
@@ -991,7 +969,7 @@ namespace HairInterpolationBuilder
 			&RootsGrid,
 			&TotalInvalidInterpolationCount,  
 			&InterpolationData, 
-			&Random,
+			&RandomGuideIndices,
 			&CompletedTasks,
 			&SlowTask
 		] (uint32 RenCurveIndex) 
@@ -1010,17 +988,17 @@ namespace HairInterpolationBuilder
 			const FHairRoot& StrandRoot = RenRoots[RenCurveIndex];
 
 			FClosestGuides ClosestGuides;
-			if (Settings.Quality == EHairInterpolationDataQuality::Low)
+			if (Settings.InterpolationQuality == EGroomInterpolationQuality::Low)
 			{
-				ClosestGuides = RootsGrid.FindClosestRoots(RenCurveIndex, RenRoots, SimRoots, RenStrandsData, SimStrandsData, Settings.bRandomizedGuides, Settings.bUseUniqueGuide, Random);
+				ClosestGuides = RootsGrid.FindClosestRoots(RenCurveIndex, RenRoots, SimRoots, RenStrandsData, SimStrandsData, Settings.bRandomizeGuide, Settings.bUseUniqueGuide, RandomGuideIndices[RenCurveIndex]);
 			}
-			else if (Settings.Quality == EHairInterpolationDataQuality::Medium)
+			else if (Settings.InterpolationQuality == EGroomInterpolationQuality::Medium)
 			{
-				ClosestGuides = RootsGrid.FindBestClosestRoots(RenCurveIndex, RenRoots, SimRoots, RenStrandsData, SimStrandsData, Settings.bRandomizedGuides, Settings.bUseUniqueGuide, Random);
+				ClosestGuides = RootsGrid.FindBestClosestRoots(RenCurveIndex, RenRoots, SimRoots, RenStrandsData, SimStrandsData, Settings.bRandomizeGuide, Settings.bUseUniqueGuide, RandomGuideIndices[RenCurveIndex]);
 			}
-			else // (Settings.Quality == EHairInterpolationDataQuality::High)
+			else // (Settings.Quality == EGroomInterpolationQuality::High)
 			{
-				ClosestGuides = FindBestRoots(RenCurveIndex, RenRoots, SimRoots, RenStrandsData, SimStrandsData, Settings.bRandomizedGuides, Settings.bUseUniqueGuide, Random);
+				ClosestGuides = FindBestRoots(RenCurveIndex, RenRoots, SimRoots, RenStrandsData, SimStrandsData, Settings.bRandomizeGuide, Settings.bUseUniqueGuide, RandomGuideIndices[RenCurveIndex]);
 			}
 
 			const uint32 RendPointCount	= RenStrandsData.StrandsCurves.CurvesCount[RenCurveIndex];
@@ -1036,7 +1014,7 @@ namespace HairInterpolationBuilder
 				for (uint32 KIndex = 0; KIndex < FClosestGuides::Count; ++KIndex)
 				{
 					// Find the closest vertex on the guide which matches the strand vertex distance along its curve
-					if (Settings.WeightMethod == EHairInterpolationWeightMethod::ParametricDistance)
+					if (Settings.InterpolationDistance == EGroomInterpolationWeight::Parametric)
 					{
 						const uint32 SimCurveIndex = ClosestGuides.Indices[KIndex];
 						const uint32 SimOffset = SimStrandsData.StrandsCurves.CurvesOffset[SimCurveIndex];
@@ -1053,7 +1031,7 @@ namespace HairInterpolationBuilder
 
 					// Use only the root as a *constant* weight for deformation along each vertex
 					// Still compute the closest vertex (in parametric distance) to know on which vertex the offset/delta should be computed
-					if (Settings.WeightMethod == EHairInterpolationWeightMethod::RootDistance)
+					if (Settings.InterpolationDistance == EGroomInterpolationWeight::Root)
 					{
 						const uint32 SimCurveIndex = ClosestGuides.Indices[KIndex];
 						const uint32 SimOffset = SimStrandsData.StrandsCurves.CurvesOffset[SimCurveIndex];
@@ -1069,7 +1047,7 @@ namespace HairInterpolationBuilder
 					}
 
 					// Use the *same vertex index* to match guide vertex with strand vertex
-					if (Settings.WeightMethod == EHairInterpolationWeightMethod::VertexIndex)
+					if (Settings.InterpolationDistance == EGroomInterpolationWeight::Index)
 					{
 						const uint32 SimCurveIndex = ClosestGuides.Indices[KIndex];
 						const uint32 SimOffset = SimStrandsData.StrandsCurves.CurvesOffset[SimCurveIndex];
@@ -1137,8 +1115,141 @@ namespace HairInterpolationBuilder
 			OutInterp1.Pad1			= 0;
 		}	
 	}
+
+	/** Fill the GroomAsset with the interpolation data that exists in the HairDescription */
+	void FillInterpolationData(UGroomAsset* GroomAsset, const FHairDescription& HairDescription)
+	{
+		TRACE_CPUPROFILER_EVENT_SCOPE(HairInterpolationBuilder::FillInterpolationData);
+
+		TStrandAttributesConstRef<FVector> ClosestGuides = HairDescription.StrandAttributes().GetAttributesRef<FVector>(HairAttribute::Strand::ClosestGuides);
+		TStrandAttributesConstRef<FVector> GuideWeights = HairDescription.StrandAttributes().GetAttributesRef<FVector>(HairAttribute::Strand::GuideWeights);
+
+		if (!GroomAsset || !ClosestGuides.IsValid() || !GuideWeights.IsValid())
+		{
+			return;
+		}
+
+		const int32 NumGroups = GroomAsset->HairGroupsData.Num();
+		for (int32 GroupID = 0; GroupID < NumGroups; ++GroupID)
+		{
+			FHairGroupData& GroupData = GroomAsset->HairGroupsData[GroupID];
+
+			const FHairStrandsDatas& RenderData = GroupData.HairRenderData;
+			const FHairStrandsDatas& SimulationData = GroupData.HairSimulationData;
+			FHairStrandsInterpolationDatas& InterpolationData = GroupData.HairInterpolationData;
+
+			// Group must include imported guides data for imported interpolation data to be meaningful
+			// Otherwise, interpolation data will be built from generated guide data
+			if (SimulationData.GetNumCurves() == 0)
+			{
+				continue;
+			}
+
+			InterpolationData.SetNum(RenderData.GetNumPoints());
+			for (uint32 CurveIndex = 0; CurveIndex < RenderData.GetNumCurves(); ++CurveIndex)
+			{
+				const uint32 CurveOffset = RenderData.StrandsCurves.CurvesOffset[CurveIndex];
+				const uint16 CurveNumVertices = RenderData.StrandsCurves.CurvesCount[CurveIndex];
+
+				const FStrandID StrandID(RenderData.StrandsCurves.StrandIDs[CurveIndex]);
+				const FVector StrandClosestGuides = ClosestGuides[StrandID];
+				const FVector StrandGuideWeights = GuideWeights[StrandID];
+				for (uint16 VertexIndex = 0; VertexIndex < CurveNumVertices; ++VertexIndex)
+				{
+					const uint32 PointGlobalIndex = VertexIndex + CurveOffset;
+					const FVector& RenPointPosition = RenderData.StrandsPoints.PointsPosition[PointGlobalIndex];
+					const float RenPointDistance = RenderData.StrandsPoints.PointsCoordU[PointGlobalIndex] * RenderData.StrandsCurves.CurvesLength[CurveIndex] * RenderData.StrandsCurves.MaxLength;
+
+					float TotalWeight = 0;
+					for (uint32 GuideIndex = 0; GuideIndex < FClosestGuides::Count; ++GuideIndex)
+					{
+						int32 ImportedGroomID = StrandClosestGuides[GuideIndex];
+						const int* SimCurveIndex = SimulationData.StrandsCurves.GroomIDToIndex.Find(ImportedGroomID);
+
+						if (SimCurveIndex && *SimCurveIndex >= 0)
+						{
+							// Fill the interpolation data using the ParametricDistance algorithm with a constant weight for all vertices along the strand
+							const uint32 SimOffset = SimulationData.StrandsCurves.CurvesOffset[*SimCurveIndex];
+							const FVertexInterpolationDesc Desc = FindMatchingVertex(RenPointDistance, SimulationData, *SimCurveIndex);
+
+							InterpolationData.PointsSimCurvesIndex[PointGlobalIndex][GuideIndex] = *SimCurveIndex;
+							InterpolationData.PointsSimCurvesVertexIndex[PointGlobalIndex][GuideIndex] = Desc.Index0 + SimOffset;
+							InterpolationData.PointsSimCurvesVertexLerp[PointGlobalIndex][GuideIndex] = Desc.T;
+							InterpolationData.PointsSimCurvesVertexWeights[PointGlobalIndex][GuideIndex] = StrandGuideWeights[GuideIndex];
+							TotalWeight += InterpolationData.PointsSimCurvesVertexWeights[PointGlobalIndex][GuideIndex];
+						}
+					}
+
+					// Normalize the weights
+					if (TotalWeight > 0.f)
+					{
+						for (int32 GuideIndex = 0; GuideIndex < FClosestGuides::Count; ++GuideIndex)
+						{
+							InterpolationData.PointsSimCurvesVertexWeights[PointGlobalIndex][GuideIndex] /= TotalWeight;
+						}
+					}
+				}
+			}
+		}
+
+		// Deallocate the memory used for indices mapping since it's not needed anymore
+		for (int32 GroupID = 0; GroupID < NumGroups; ++GroupID)
+		{
+			FHairGroupData& GroupData = GroomAsset->HairGroupsData[GroupID];
+
+			FHairStrandsDatas& RenderData = GroupData.HairRenderData;
+			RenderData.StrandsCurves.GroomIDToIndex.Empty();
+			RenderData.StrandsCurves.StrandIDs.Empty();
+
+			FHairStrandsDatas& SimulationData = GroupData.HairSimulationData;
+			SimulationData.StrandsCurves.GroomIDToIndex.Empty();
+			SimulationData.StrandsCurves.StrandIDs.Empty();
+		}
+	}
 }
 
+class FGroomDataRandomizer
+{
+public:
+	FGroomDataRandomizer(int Seed, int NumRenderCurves, int NumSimCurves)
+	{
+		Random.Initialize(Seed);
+
+		RenderCurveSeeds.SetNumUninitialized(NumRenderCurves);
+		for (int Index = 0; Index < NumRenderCurves; ++Index)
+		{
+			RenderCurveSeeds[Index] = Random.RandHelper(255);
+		}
+
+		SimCurveSeeds.SetNumUninitialized(NumSimCurves);
+		for (int Index = 0; Index < NumSimCurves; ++Index)
+		{
+			SimCurveSeeds[Index] = Random.RandHelper(255);
+		}
+
+		// This randomization makes certain strands being affected by 1, 2, or 3 guides
+		GuideIndices.SetNumUninitialized(NumRenderCurves);
+		for (int Index = 0; Index < NumRenderCurves; ++Index)
+		{
+			FIntVector& RandomIndices = GuideIndices[Index];
+			for (int GuideIndex = 0; GuideIndex < HairInterpolationBuilder::FMetrics::Count; ++GuideIndex)
+			{
+				RandomIndices[GuideIndex] = Random.RandRange(0, HairInterpolationBuilder::FMetrics::Count - 1);
+			}
+		}
+	}
+
+	const TArray<uint8>& GetRenderCurveSeeds() const { return RenderCurveSeeds; }
+	const TArray<uint8>& GetSimCurveSeeds() const { return SimCurveSeeds; }
+	const TArray<FIntVector>& GetRandomGuideIndices() const { return GuideIndices; }
+
+private:
+	FRandomStream Random;
+	TArray<uint8> RenderCurveSeeds;
+	TArray<uint8> SimCurveSeeds;
+	TArray<FIntVector> GuideIndices;
+};
+
 bool FGroomBuilder::BuildGroom(const FHairDescription& HairDescription, const FGroomBuildSettings& BuildSettings, UGroomAsset* GroomAsset)
 {
 	if (!GroomAsset)
@@ -1198,9 +1309,17 @@ bool FGroomBuilder::BuildGroom(const FHairDescription& HairDescription, const FG
 
 	TStrandAttributesConstRef<int> StrandGuides = HairDescription.StrandAttributes().GetAttributesRef<int>(HairAttribute::Strand::Guide);
 	TStrandAttributesConstRef<int> GroupIDs = HairDescription.StrandAttributes().GetAttributesRef<int>(HairAttribute::Strand::GroupID);
+	TStrandAttributesConstRef<int> StrandIDs = HairDescription.StrandAttributes().GetAttributesRef<int>(HairAttribute::Strand::ID);
 
 	bool bImportGuides = !BuildSettings.bOverrideGuides;
 
+	TStrandAttributesConstRef<FVector> ClosestGuides = HairDescription.StrandAttributes().GetAttributesRef<FVector>(HairAttribute::Strand::ClosestGuides);
+	TStrandAttributesConstRef<FVector> GuideWeights = HairDescription.StrandAttributes().GetAttributesRef<FVector>(HairAttribute::Strand::GuideWeights);
+
+	// To use ClosestGuides and GuideWeights attributes, guides must be imported from HairDescription and
+	// must include StrandID attribute since ClosestGuides references those IDs
+	bool bCanUseClosestGuidesAndWeights = bImportGuides && StrandIDs.IsValid() && ClosestGuides.IsValid() && GuideWeights.IsValid();
+	
 	FHairStrandsDatas* CurrentHairStrandsDatas = nullptr;
 
 	typedef TPair<FHairGroupInfo, FHairGroupData> FHairGroup;
@@ -1259,6 +1378,17 @@ bool FGroomBuilder::BuildGroom(const FHairDescription& HairDescription, const FG
 
 		CurrentHairStrandsDatas->StrandsCurves.CurvesCount.Add(CurveNumVertices);
 
+		if (bCanUseClosestGuidesAndWeights)
+		{
+			// ClosesGuides needs mapping of StrandID (the strand index in the HairDescription)
+			CurrentHairStrandsDatas->StrandsCurves.StrandIDs.Add(CurveIndex);
+
+			// and ImportedGroomID (the imported ID associated with a strand) to index in StrandCurves
+			const int ImportedGroomID = StrandIDs[StrandID];
+			const int StrandCurveIndex = CurrentHairStrandsDatas->StrandsCurves.CurvesCount.Num() - 1;
+			CurrentHairStrandsDatas->StrandsCurves.GroomIDToIndex.Add(ImportedGroomID, StrandCurveIndex);
+		}
+
 		if (bHasUVData)
 		{
 			CurrentHairStrandsDatas->StrandsCurves.CurvesRootUV.Add(StrandRootUV[StrandID]);
@@ -1328,10 +1458,13 @@ bool FGroomBuilder::BuildGroom(const FHairDescription& HairDescription, const FG
 		}
 		else
 		{
-			GroomAsset->HairToGuideDensity = FMath::Clamp(BuildSettings.HairToGuideDensity, 0.01f, 1.f);
+			GroomAsset->HairToGuideDensity = FMath::Clamp(BuildSettings.HairToGuideDensity, 0.f, 1.f);
 		}
 	}
 
+	GroomAsset->HairGroupsInfo.Reset();
+	GroomAsset->HairGroupsData.Reset();
+
 	for (TPair<int32, FHairGroup> HairGroupIt : HairGroups)
 	{
 		FHairGroup& Group = HairGroupIt.Value;
@@ -1342,14 +1475,21 @@ bool FGroomBuilder::BuildGroom(const FHairDescription& HairDescription, const FG
 		GroomAsset->HairGroupsData.Add(MoveTemp(GroupData));
 	}
 
-	BuildData(GroomAsset, uint8(BuildSettings.InterpolationQuality), uint8(BuildSettings.InterpolationDistance), BuildSettings.bRandomizeGuide, BuildSettings.bUseUniqueGuide);
+	// If there's usable closest guides and guide weights attributes, fill them into the asset
+	// This step requires the HairSimulationData (guides) to be filled prior to this
+	if (bCanUseClosestGuidesAndWeights)
+	{
+		HairInterpolationBuilder::FillInterpolationData(GroomAsset, HairDescription);
+	}
+
+	BuildData(GroomAsset, BuildSettings);
 
 	GroomAsset->InitResource();
 
 	return true;
 }
 
-void FGroomBuilder::BuildData(UGroomAsset* GroomAsset, uint8 QualityLevel, uint8 WeightMethod, bool bRandomize, bool bUnique)
+void FGroomBuilder::BuildData(UGroomAsset* GroomAsset, const FGroomBuildSettings& BuildSettings)
 {
 	if (!GroomAsset)
 	{
@@ -1369,36 +1509,29 @@ void FGroomBuilder::BuildData(UGroomAsset* GroomAsset, uint8 QualityLevel, uint8
 		{
 			FHairGroupInfo& GroupInfo = GroomAsset->HairGroupsInfo[Index];
 
-			float GuideDensity = FMath::Clamp(GroomAsset->HairToGuideDensity, 0.01f, 1.f);
+			const float GuideDensity = FMath::Clamp(GroomAsset->HairToGuideDensity, 0.f, 1.f);
 			GenerateGuides(HairRenderData, GuideDensity, HairSimulationData);
-			GroupInfo.bIsAutoGenerated = true;
+			GroupInfo.ImportSettings = BuildSettings;
 			GroupInfo.NumGuides = HairSimulationData.GetNumCurves();
 		}
 
+		// Precompute the random values used to build the data
+		// In particular, BuildInterpolationData needs this because it parallelizes its computation
+		FGroomDataRandomizer Random(Index, HairRenderData.GetNumCurves(), HairSimulationData.GetNumCurves());
+
 		// Build RenderData for HairStrandsDatas
-		HairStrandsBuilder::BuildRenderData(HairRenderData);
+		HairStrandsBuilder::BuildRenderData(HairRenderData, Random.GetRenderCurveSeeds());
 
-		HairStrandsBuilder::BuildRenderData(HairSimulationData);
+		HairStrandsBuilder::BuildRenderData(HairSimulationData, Random.GetSimCurveSeeds());
 
 		FHairStrandsInterpolationDatas& HairInterpolationData = GroupData.HairInterpolationData;
 
 		// Build InterpolationData from render and simulation HairStrandsDatas
-		HairInterpolationBuilder::FHairInterpolationSettings Settings;
-		Settings.bRandomizedGuides = bRandomize;
-		Settings.bUseUniqueGuide = bUnique;
-		switch (WeightMethod)
-		{
-			case 0: Settings.WeightMethod = HairInterpolationBuilder::EHairInterpolationWeightMethod::ParametricDistance; break;
-			case 1: Settings.WeightMethod = HairInterpolationBuilder::EHairInterpolationWeightMethod::RootDistance; break;
-			case 2: Settings.WeightMethod = HairInterpolationBuilder::EHairInterpolationWeightMethod::VertexIndex; break;
-		}
-		switch (QualityLevel)
+		// Skip building if interpolation data was provided by the source file
+		if (HairInterpolationData.Num() == 0)
 		{
-			case 0: Settings.Quality = HairInterpolationBuilder::EHairInterpolationDataQuality::Low; break;
-			case 1: Settings.Quality = HairInterpolationBuilder::EHairInterpolationDataQuality::Medium; break;
-			case 2: Settings.Quality = HairInterpolationBuilder::EHairInterpolationDataQuality::High; break;
+			HairInterpolationBuilder::BuildInterpolationData(HairInterpolationData, HairSimulationData, HairRenderData, BuildSettings, Random.GetRandomGuideIndices());
 		}
-		HairInterpolationBuilder::BuildInterpolationData(HairInterpolationData, HairSimulationData, HairRenderData, Settings);
 
 		// Build Rendering data for InterpolationData
 		HairInterpolationBuilder::BuildRenderData(HairInterpolationData);
diff --git a/Engine/Plugins/Experimental/HairStrands/Source/HairStrandsCore/Private/GroomComponent.cpp b/Engine/Plugins/Experimental/HairStrands/Source/HairStrandsCore/Private/GroomComponent.cpp
--- a/Engine/Plugins/Experimental/HairStrands/Source/HairStrandsCore/Private/GroomComponent.cpp
+++ b/Engine/Plugins/Experimental/HairStrands/Source/HairStrandsCore/Private/GroomComponent.cpp
@@ -12,6 +12,7 @@
 #include "UObject/UObjectIterator.h"
 #include "GlobalShader.h"
 #include "Components/SkeletalMeshComponent.h"
+#include "Rendering/SkeletalMeshRenderData.h"
 #include "Engine/RendererSettings.h"
 #include "Animation/AnimationSettings.h"
 #include "Logging/MessageLog.h"
@@ -23,34 +24,11 @@ static float GHairClipLength = -1;
 static FAutoConsoleVariableRef CVarHairClipLength(TEXT("r.HairStrands.DebugClipLength"), GHairClipLength, TEXT("Clip hair strands which have a lenth larger than this value. (default is -1, no effect)"));
 float GetHairClipLength() { return GHairClipLength > 0 ? GHairClipLength : 100000;  }
 
-static int32 GHairStrandsMeshProjectionForceRefPoseEnable = 0;
-static int32 GHairStrandsMeshProjectionForceLOD = -1;
-static FAutoConsoleVariableRef CVarHairStrandsMeshProjectionForceRefPoseEnable(TEXT("r.HairStrands.MeshProjection.RefPose"), GHairStrandsMeshProjectionForceRefPoseEnable, TEXT("Enable/Disable reference pose"));
-static FAutoConsoleVariableRef CVarHairStrandsMeshProjectionForceLOD(TEXT("r.HairStrands.MeshProjection.LOD"), GHairStrandsMeshProjectionForceLOD, TEXT("Force a specific LOD"));
-
 static int32 GHairStrandsMeshUseRelativePosition = 1;
 static FAutoConsoleVariableRef CVarHairStrandsMeshUseRelativePosition(TEXT("r.HairStrands.MeshProjection.RelativePosition"), GHairStrandsMeshUseRelativePosition, TEXT("Enable/Disable relative triangle position for improving positions"));
 
-static int32 GHairStrandsMeshProjectionTickDelay = 1;
-static FAutoConsoleVariableRef CVarHairStrandsMeshProjectionTickDelay(TEXT("r.HairStrands.MeshProjection.TickDelay"), GHairStrandsMeshProjectionTickDelay, TEXT("Number of simulation tick to wait before projecting a groom onto a mesh"));
-
 #define LOCTEXT_NAMESPACE "GroomComponent"
 
-static float GetGroupMaxHairRadius(const FHairGroupDesc& GroupDesc, const FHairGroupData& GroupData)
-{
-	return GroupDesc.HairWidth > 0 ? GroupDesc.HairWidth * 0.5f : GroupData.HairRenderData.StrandsCurves.MaxRadius;
-}
-
-static float GetGroupHairShadowDensity(const FHairGroupDesc& GroupDesc, const FHairGroupData& GroupData)
-{
-	return GroupDesc.HairShadowDensity > 0 ? GroupDesc.HairShadowDensity : GroupData.HairRenderData.HairDensity;
-}
-
-static float GetGroupHairRaytracingRadiusScale(const FHairGroupDesc& GroupDesc, const FHairGroupData& GroupData)
-{
-	return GroupDesc.HairRaytracingRadiusScale;
-}
-
 /**
  * An material render proxy which overrides the debug mode parameter.
  */
@@ -155,6 +133,7 @@ public:
 		check(Component->GroomAsset);
 		check(Component->GroomAsset->GetNumHairGroups() > 0);
 		check(Component->InterpolationOutput);
+		ComponentId = Component->ComponentId.PrimIDValue;
 
 		FHairStrandsVertexFactory::FDataType VFData;
 		VFData.InterpolationOutput = Component->InterpolationOutput;
@@ -174,16 +153,9 @@ public:
 				Material = GEngine->HairDefaultMaterial;
 			}
 
-			FHairStrandsVertexFactory::FDataType::HairGroup& VFGroupData = VFData.HairGroups.AddDefaulted_GetRef();
-			VFGroupData.MinStrandRadius		= 0;
-			VFGroupData.MaxStrandRadius		= GetGroupMaxHairRadius(InGroupDesc, InGroupData);
-			VFGroupData.MaxStrandLength		= InGroupData.HairRenderData.StrandsCurves.MaxLength;
-			VFGroupData.HairDensity			= GetGroupHairShadowDensity(InGroupDesc, InGroupData);
-			VFGroupData.HairWorldOffset		= InGroupData.HairRenderData.BoundingBox.GetCenter();
-
 			#if RHI_RAYTRACING
 			FRayTracingGeometry* RayTracingGeometry = nullptr;
-			if (IsRayTracingEnabled() && GroupResources.RaytracingResources)
+			if (IsHairRayTracingEnabled() && GroupResources.RaytracingResources)
 			{
 				RayTracingGeometry = &GroupResources.RaytracingResources->RayTracingGeometry;
 			}
@@ -191,11 +163,11 @@ public:
 
 			HairGroup& OutGroupData = HairGroups.Add_GetRef(
 			{
-				GroupIt,
 				Material
 				#if RHI_RAYTRACING
 				, RayTracingGeometry
 				#endif
+				, GroupResources.HairGroupPublicDatas
 			});
 
 		}
@@ -220,7 +192,7 @@ public:
 
 	virtual void GetDynamicRayTracingInstances(FRayTracingMaterialGatheringContext & Context, TArray<struct FRayTracingInstance>& OutRayTracingInstances) override
 	{
-		if (!IsRayTracingEnabled() || HairGroups.Num() == 0)
+		if (!IsHairRayTracingEnabled() || HairGroups.Num() == 0)
 			return;
 
 		for (const HairGroup& GroupData : HairGroups)
@@ -239,9 +211,7 @@ public:
 
 	virtual void GetDynamicMeshElements(const TArray<const FSceneView*>& Views, const FSceneViewFamily& ViewFamily, uint32 VisibilityMap, FMeshElementCollector& Collector) const override
 	{
-		const uint32 GroupCount = HairGroups.Num();
-		check(GroupCount == VertexFactory.GetData().HairGroups.Num());
-		if (GroupCount == 0)
+		if (!VertexFactory.GetData().InterpolationOutput)
 			return;
 
 		bool bHasOneElementValid = false;
@@ -256,6 +226,8 @@ public:
 		if (!bHasOneElementValid) 
 			return;
 
+		const uint32 GroupCount = VertexFactory.GetData().InterpolationOutput->HairGroups.Num();
+
 		QUICK_SCOPE_CYCLE_COUNTER(STAT_HairStrandsSceneProxy_GetDynamicMeshElements);
 
 
@@ -277,19 +249,20 @@ public:
 			case EHairStrandsDebugMode::RenderHairRootUDIM			: DebugModeScalar = 6.f; break;
 			case EHairStrandsDebugMode::RenderHairBaseColor			: DebugModeScalar = 7.f; break;
 			case EHairStrandsDebugMode::RenderHairRoughness			: DebugModeScalar = 8.f; break;
+			case EHairStrandsDebugMode::RenderVisCluster			: DebugModeScalar = 0.f; break;
 			};
 
-			float HairMinRadius = FLT_MAX;
+
 			float HairMaxRadius = 0;
-			for (uint32 GroupIt = 0; GroupIt < GroupCount; ++GroupIt)
+			for (FHairStrandsInterpolationOutput::HairGroup& HairGroup : VertexFactory.GetData().InterpolationOutput->HairGroups)
 			{
-				HairMinRadius = FMath::Min(HairMinRadius, VertexFactory.GetMinStrandRadius(GroupIt));
-				HairMaxRadius = FMath::Max(HairMaxRadius, VertexFactory.GetMaxStrandRadius(GroupIt));
+				HairMaxRadius = FMath::Max(HairMaxRadius, HairGroup.VFInput.HairRadius);
 			}
+
 			const float HairClipLength = GetHairClipLength();
 			auto DebugMaterial = new FHairDebugModeMaterialRenderProxy(
 				GEngine->HairDebugMaterial ? GEngine->HairDebugMaterial->GetRenderProxy() : nullptr,
-				DebugModeScalar, HairMinRadius, HairMaxRadius, HairClipLength);
+				DebugModeScalar, 0, HairMaxRadius, HairClipLength);
 			Collector.RegisterOneFrameMaterialProxy(DebugMaterial);
 			MaterialProxy = DebugMaterial;
 		}
@@ -302,9 +275,7 @@ public:
 				{
 					const HairGroup& GroupData = HairGroups[GroupIt];
 					const uint32 HairVertexCount = VertexFactory.GetData().InterpolationOutput->HairGroups[GroupIt].VFInput.VertexCount;
-					const FHairStrandsVertexFactory::FDataType::HairGroup& HairGroup = VertexFactory.GetData().HairGroups[GroupIt];
 
-					const FSceneView* View = Views[ViewIndex];
 					// Draw the mesh.
 					FMeshBatch& Mesh = Collector.AllocateMesh();
 					FMeshBatchElement& BatchElement = Mesh.Elements[0];
@@ -322,11 +293,19 @@ public:
 					FDynamicPrimitiveUniformBuffer& DynamicPrimitiveUniformBuffer = Collector.AllocateOneFrameResource<FDynamicPrimitiveUniformBuffer>();
 					DynamicPrimitiveUniformBuffer.Set(GetLocalToWorld(), PreviousLocalToWorld, GetBounds(), GetLocalBounds(), true, false, bDrawVelocity, bOutputVelocity);
 					BatchElement.PrimitiveUniformBufferResource = &DynamicPrimitiveUniformBuffer.UniformBuffer;
+
 					BatchElement.FirstIndex = 0;
-					BatchElement.NumPrimitives = HairVertexCount * 2;
+					BatchElement.NumInstances = 1;
+					BatchElement.NumPrimitives = 0;
+					BatchElement.IndirectArgsBuffer = GroupData.PublicData->GetDrawIndirectBuffer().Buffer.GetReference();
+					BatchElement.IndirectArgsOffset = 0;
+					
+					// Setup our vertex factor custom data
+					BatchElement.VertexFactoryUserData = const_cast<void*>(reinterpret_cast<const void*>(GroupData.PublicData));
+
 					BatchElement.MinVertexIndex = 0;
 					BatchElement.MaxVertexIndex = HairVertexCount * 6;
-					BatchElement.UserData = reinterpret_cast<void*>(uint64(GroupIt));
+					BatchElement.UserData = reinterpret_cast<void*>(uint64(ComponentId));
 					Mesh.ReverseCulling = IsLocalToWorldDeterminantNegative();
 					Mesh.Type = PT_TriangleList;
 					Mesh.DepthPriorityGroup = SDPG_World;
@@ -383,15 +362,16 @@ public:
 	uint32 GetAllocatedSize(void) const { return(FPrimitiveSceneProxy::GetAllocatedSize()); }
 private:
 
+	uint32 ComponentId = 0;
 	FHairStrandsVertexFactory VertexFactory;
 	FMaterialRelevance MaterialRelevance;
 	struct HairGroup
 	{
-		uint32 GroupIndex = 0;
 		UMaterialInterface* Material = nullptr;
 	#if RHI_RAYTRACING
 		FRayTracingGeometry* RayTracingGeometry = nullptr;
 	#endif
+		FHairGroupPublicData* PublicData = nullptr;
 	};
 	TArray<HairGroup> HairGroups;
 };
@@ -404,7 +384,7 @@ static void UpdateHairGroupsDesc(UGroomAsset* GroomAsset, TArray<FHairGroupDesc>
 		return;
 	}
 
-	check(GroomAsset->HairGroupsInfo.Num() == GroomAsset->HairGroupsInfo.Num());
+	check(GroomAsset->HairGroupsInfo.Num() == GroomAsset->HairGroupsData.Num());
 
 	const uint32 GroupCount = GroomAsset->HairGroupsInfo.Num();
 	const bool bReinitOverride = GroupCount != GroomGroupsDesc.Num();
@@ -421,9 +401,11 @@ static void UpdateHairGroupsDesc(UGroomAsset* GroomAsset, TArray<FHairGroupDesc>
 		FHairGroupDesc& Desc = GroomGroupsDesc[GroupIt];
 		Desc.GuideCount	= GroupInfo.NumGuides;
 		Desc.HairCount	= GroupInfo.NumCurves;
+		Desc.HairLength = GroupData.HairRenderData.StrandsCurves.MaxLength;
+
 		if (bReinitOverride || Desc.HairWidth == 0)
 		{
-			Desc.HairWidth = GroupData.HairRenderData.StrandsCurves.MaxRadius * 0.5f;
+			Desc.HairWidth = GroupData.HairRenderData.StrandsCurves.MaxRadius * 2.0f;
 		}
 		if (bReinitOverride || Desc.HairShadowDensity == 0)
 		{
@@ -432,7 +414,7 @@ static void UpdateHairGroupsDesc(UGroomAsset* GroomAsset, TArray<FHairGroupDesc>
 
 		if (bReinitOverride)
 		{
-			Desc.HairRaytracingRadiusScale = 1;
+			Desc.ReInit();
 		}
 	}
 }
@@ -455,9 +437,6 @@ UGroomComponent::UGroomComponent(const FObjectInitializer& ObjectInitializer)
 	bBindGroomToSkeletalMesh = false;
 	InitializedResources = nullptr;
 	Mobility = EComponentMobility::Movable;
-	MeshProjectionTickDelay = 0;
-	MeshProjectionLODIndex = -1;
-	MeshProjectionState = EMeshProjectionState::Invalid;
 	bIsGroomAssetCallbackRegistered = false;
 
 	SetCollisionProfileName(UCollisionProfile::PhysicsActor_ProfileName);
@@ -567,16 +546,17 @@ static FHairStrandsProjectionHairData::HairGroup ToProjectionHairData(FHairStran
 		LODData.RootTriangleIndexBuffer = &MeshLODData.RootTriangleIndexBuffer;
 		LODData.RootTriangleBarycentricBuffer = &MeshLODData.RootTriangleBarycentricBuffer;
 		
-		LODData.RestPositionOffset					= MeshLODData.RestRootOffset;
+		LODData.RestPositionOffset				= &MeshLODData.RestRootOffset;
 		LODData.RestRootTrianglePosition0Buffer = &MeshLODData.RestRootTrianglePosition0Buffer;
 		LODData.RestRootTrianglePosition1Buffer = &MeshLODData.RestRootTrianglePosition1Buffer;
 		LODData.RestRootTrianglePosition2Buffer = &MeshLODData.RestRootTrianglePosition2Buffer;
 
-		LODData.DeformedPositionOffset					= MeshLODData.DeformedRootOffset;
+		LODData.DeformedPositionOffset				= &MeshLODData.DeformedRootOffset;
 		LODData.DeformedRootTrianglePosition0Buffer = &MeshLODData.DeformedRootTrianglePosition0Buffer;
 		LODData.DeformedRootTrianglePosition1Buffer = &MeshLODData.DeformedRootTrianglePosition1Buffer;
 		LODData.DeformedRootTrianglePosition2Buffer = &MeshLODData.DeformedRootTrianglePosition2Buffer;
-		LODData.bIsValid = MeshLODData.Status == FHairStrandsRootResource::FMeshProjectionLOD::EStatus::Completed;
+
+		LODData.Status = &MeshLODData.Status;
 	}
 
 	return Out;
@@ -738,6 +718,8 @@ void UGroomComponent::InitResources()
 	bResetSimulation = true;
 	//UE_LOG(LogHairStrands, Warning, TEXT("Groom Init = %d"), bResetSimulation);
 
+	UpdateHairGroupsDesc(GroomAsset, GroomGroupsDesc);
+
 	if (!GroomAsset || GroomAsset->GetNumHairGroups() == 0)
 		return;
 
@@ -749,20 +731,22 @@ void UGroomComponent::InitResources()
 
 	// Insure the ticking of the Groom component always happens after the skeletalMeshComponent.
 	USkeletalMeshComponent* SkeletalMeshComponent = bBindGroomToSkeletalMesh && GetAttachParent() ? Cast<USkeletalMeshComponent>(GetAttachParent()) : nullptr;
+	const bool bHasValidSketalMesh = SkeletalMeshComponent && SkeletalMeshComponent->GetSkeletalMeshRenderData();
+
+	// Always setup the callback, even if the SkeletalMeshData are not ready yet
 	if (SkeletalMeshComponent)
 	{
-		RegisteredSkeletalMeshComponent = SkeletalMeshComponent;
-		AddTickPrerequisiteComponent(SkeletalMeshComponent);
-
 		FSkeletalMeshObjectCallbackData CallbackData;
 		CallbackData.Run = CallbackMeshObjectCallback;
 		CallbackData.UserData = (uint64(LocalComponentId.PrimIDValue) & 0xFFFFFFFF) | (uint64(WorldType) << 32);
 		SkeletalMeshComponent->MeshObjectCallbackData = CallbackData;
+	}
 
-
-		{
-			SkeletalMeshComponent->OnBoneTransformsFinalized.AddDynamic(this, &UGroomComponent::ResetSimulation);
-		}
+	if (bHasValidSketalMesh)
+	{
+		RegisteredSkeletalMeshComponent = SkeletalMeshComponent;
+		AddTickPrerequisiteComponent(SkeletalMeshComponent);
+		SkeletalMeshComponent->OnBoneTransformsFinalized.AddDynamic(this, &UGroomComponent::ResetSimulation);
 	}
 
 	const bool bIsSimulationEnable = IsSimulationEnabled(this);
@@ -793,14 +777,14 @@ void UGroomComponent::InitResources()
 		BeginInitResource(Res.InterpolationResource);
 
 		#if RHI_RAYTRACING
-		if (IsRayTracingEnabled())
+		if (IsHairRayTracingEnabled())
 		{
 			Res.RaytracingResources = new FHairStrandsRaytracingResource(GroupData.HairRenderData);
 			BeginInitResource(Res.RaytracingResources);
 		}
 		#endif
 
-		if (SkeletalMeshComponent)
+		if (bHasValidSketalMesh)
 		{
 			const uint32 LODCount = SkeletalMeshComponent->GetNumLODs();
 			if (LODCount > 0)
@@ -817,9 +801,16 @@ void UGroomComponent::InitResources()
 
 		Res.RenderDeformedResources = new FHairStrandsDeformedResource(GroupData.HairRenderData.RenderData, false);
 		Res.SimDeformedResources = new FHairStrandsDeformedResource(GroupData.HairSimulationData.RenderData, true);
+		Res.ClusterCullingResources = new FHairStrandsClusterCullingResource(GroupData);
+
+		const uint32 HairVertexCount = GroupData.HairRenderData.GetNumPoints();
+		const uint32 GroupInstanceVertexCount = HairVertexCount * 6; // 6 vertex per point for a quad
+		Res.HairGroupPublicDatas = new FHairGroupPublicData(GroupIt, GroupInstanceVertexCount, Res.ClusterCullingResources->ClusterCount, HairVertexCount);
 
 		BeginInitResource(Res.RenderDeformedResources);
 		BeginInitResource(Res.SimDeformedResources);
+		BeginInitResource(Res.ClusterCullingResources);
+		BeginInitResource(Res.HairGroupPublicDatas);
 
 		const FVector RenderRestHairPositionOffset = Res.RenderRestResources->PositionOffset;
 		const FVector SimRestHairPositionOffset = Res.SimRestResources->PositionOffset;
@@ -833,9 +824,11 @@ void UGroomComponent::InitResources()
 		FHairStrandsInterpolationInput::FHairGroup& InterpolationInputGroup = InterpolationInput->HairGroups.AddDefaulted_GetRef();
 
 		check(GroupIt < GroomGroupsDesc.Num());
-		const FHairGroupDesc& InGroupDesc = GroomGroupsDesc[GroupIt];
-		InterpolationInputGroup.HairRadius = GetGroupMaxHairRadius(InGroupDesc, GroupData);
-		InterpolationInputGroup.HairRaytracingRadiusScale = GetGroupHairRaytracingRadiusScale(InGroupDesc, GroupData);
+		InterpolationInputGroup.GroupDesc = GroomGroupsDesc[GroupIt];
+
+		InterpolationOutputGroup.HairGroupPublicData = Res.HairGroupPublicDatas;
+
+		// Rest post offset
 		InterpolationInputGroup.InRenderHairPositionOffset = RenderRestHairPositionOffset;
 		InterpolationInputGroup.InSimHairPositionOffset = SimRestHairPositionOffset;
 
@@ -852,6 +845,7 @@ void UGroomComponent::InitResources()
 	Interpolation.Input  = InterpolationInput;
 	Interpolation.Output = InterpolationOutput;
 	Interpolation.Function = ComputeHairStrandsInterpolation;
+	Interpolation.ResetFunction = ResetHairStrandsInterpolation;
 
 	FHairGroupResources* LocalResources = &HairGroupResources;
 	const uint64 Id = LocalComponentId.PrimIDValue;
@@ -862,12 +856,15 @@ void UGroomComponent::InitResources()
 			LocalResources,
 			HairLocalToWorld, SkinLocalToWorld,
 			WorldType,
-			DebugGroupInfo
+			DebugGroupInfo,
+			bHasValidSketalMesh,
+			SkeletalMeshComponent
 		]
 		(FRHICommandListImmediate& RHICmdList)
 	{
 		FHairStrandsProjectionHairData RenProjectionDatas;
 		FHairStrandsProjectionHairData SimProjectionDatas;
+		FHairStrandsPrimitiveResources PrimitiveResources;
 		const uint32 GroupCount = LocalResources->Num();
 		for (uint32 GroupIt=0;GroupIt<GroupCount; ++GroupIt)
 		{
@@ -887,9 +884,15 @@ void UGroomComponent::InitResources()
 
 			InputGroup.Interpolation0Buffer			= &Res.InterpolationResource->Interpolation0Buffer;
 			InputGroup.Interpolation1Buffer			= &Res.InterpolationResource->Interpolation1Buffer;
+			
+			InputGroup.ClusterCount					= Res.ClusterCullingResources->ClusterCount;
+			InputGroup.ClusterVertexCount			= Res.ClusterCullingResources->VertexCount;
+			InputGroup.VertexToClusterIdBuffer		= &Res.ClusterCullingResources->VertexToClusterIdBuffer;
+			InputGroup.ClusterVertexIdBuffer		= &Res.ClusterCullingResources->ClusterVertexIdBuffer;
+			InputGroup.ClusterIndexRadiusScaleInfoBuffer= &Res.ClusterCullingResources->ClusterIndexRadiusScaleInfoBuffer;
 
 			#if RHI_RAYTRACING
-			if (IsRayTracingEnabled())
+			if (IsHairRayTracingEnabled())
 			{
 				InputGroup.RaytracingGeometry		= &Res.RaytracingResources->RayTracingGeometry;
 				InputGroup.RaytracingPositionBuffer	= &Res.RaytracingResources->PositionBuffer;
@@ -905,17 +908,68 @@ void UGroomComponent::InitResources()
 			OutputGroup.RenderMaterialBuffer			= &Res.RenderRestResources->MaterialBuffer;
 			OutputGroup.RenderTangentBuffer				= &Res.RenderDeformedResources->TangentBuffer;
 			OutputGroup.SimTangentBuffer				= &Res.SimDeformedResources->TangentBuffer;
+			
+			OutputGroup.RenderClusterAABBBuffer			= &Res.HairGroupPublicDatas->GetClusterAABBBuffer();
+			OutputGroup.RenderGroupAABBBuffer			= &Res.HairGroupPublicDatas->GetGroupAABBBuffer();
+			OutputGroup.ClusterInfoBuffer				= &Res.ClusterCullingResources->ClusterInfoBuffer;
 
 			RenProjectionDatas.HairGroups.Add(ToProjectionHairData(Res.RenRootResources));
 			SimProjectionDatas.HairGroups.Add(ToProjectionHairData(Res.SimRootResources));
+
+			FHairStrandsPrimitiveResources::FHairGroup& Group = PrimitiveResources.Groups.AddDefaulted_GetRef();
+			Group.ClusterAABBBuffer = &Res.HairGroupPublicDatas->GetClusterAABBBuffer();
+			Group.GroupAABBBuffer	= &Res.HairGroupPublicDatas->GetGroupAABBBuffer();
+			Group.ClusterCount		= Res.ClusterCullingResources->ClusterCount;
+		}
+
+		if (bHasValidSketalMesh)
+		{		
+			// Convert FSkeletalMeshRenderData into FHairStrandsProjectionMeshData for the initial projection
+			FHairStrandsProjectionMeshData MeshData;
+			FSkeletalMeshRenderData* RenderData = SkeletalMeshComponent->GetSkeletalMeshRenderData();
+
+			uint32 LODIndex = 0;
+			for (FSkeletalMeshLODRenderData& LODRenderData : RenderData->LODRenderData)
+			{
+				FHairStrandsProjectionMeshData::LOD& LOD = MeshData.LODs.AddDefaulted_GetRef();
+				uint32 SectionIndex = 0;
+				for (FSkelMeshRenderSection& InSection : LODRenderData.RenderSections)
+				{
+					FHairStrandsProjectionMeshData::Section& OutSection = LOD.Sections.AddDefaulted_GetRef();
+					
+					OutSection.PositionBuffer	= LODRenderData.StaticVertexBuffers.PositionVertexBuffer.GetSRV();
+					OutSection.IndexBuffer		= LODRenderData.MultiSizeIndexContainer.GetIndexBuffer()->GetSRV();
+					OutSection.TotalVertexCount	= LODRenderData.StaticVertexBuffers.PositionVertexBuffer.GetNumVertices();
+					OutSection.TotalIndexCount	= LODRenderData.MultiSizeIndexContainer.GetIndexBuffer()->Num();
+					OutSection.NumPrimitives	= InSection.NumTriangles;
+					OutSection.VertexBaseIndex	= InSection.BaseVertexIndex;
+					OutSection.IndexBaseIndex	= InSection.BaseIndex;
+					OutSection.SectionIndex		= SectionIndex;
+					OutSection.LODIndex			= LODIndex;
+
+					++SectionIndex;
+				}
+				++LODIndex;
+			}
+
+			// The offset is based on the center of the skeletal mesh (which is computed based on the physics capsules/boxes/...)
+			const FVector MeshPositionOffset = SkeletalMeshComponent->CalcBounds(FTransform::Identity).GetBox().GetCenter();
+			RunProjection(
+				RHICmdList,
+				HairLocalToWorld,
+				MeshPositionOffset,
+				MeshData, 
+				RenProjectionDatas,
+				SimProjectionDatas);
 		}
 
 		RegisterHairStrands(
 			Id,
 			WorldType,
 			Interpolation,
-			RenProjectionDatas, 
+			RenProjectionDatas,
 			SimProjectionDatas,
+			PrimitiveResources,
 			DebugGroupInfo);
 	});
 }
@@ -937,6 +991,8 @@ void UGroomComponent::ReleaseResources()
 		SafeRelease(Res.RenRootResources);
 		SafeRelease(Res.SimRootResources);
 		SafeRelease(Res.RenderDeformedResources);
+		SafeRelease(Res.ClusterCullingResources);
+		SafeRelease(Res.HairGroupPublicDatas);
 		SafeRelease(Res.SimDeformedResources);
 	#if RHI_RAYTRACING
 		SafeRelease(Res.RaytracingResources);
@@ -958,10 +1014,6 @@ void UGroomComponent::ReleaseResources()
 	InterpolationOutput = nullptr;
 	InitializedResources = nullptr;
 
-	MeshProjectionLODIndex = -1;
-	MeshProjectionTickDelay = 0;
-	MeshProjectionState = EMeshProjectionState::Invalid;
-
 	// Insure the ticking of the Groom component always happens after the skeletalMeshComponent.
 	if (RegisteredSkeletalMeshComponent)
 	{
@@ -990,7 +1042,6 @@ void UGroomComponent::PostLoad()
 		GroomAsset->ConditionalPostLoad();
 	}
 
-	UpdateHairGroupsDesc(GroomAsset, GroomGroupsDesc);
 	InitResources();
 
 #if WITH_EDITOR
@@ -1017,7 +1068,6 @@ void UGroomComponent::OnRegister()
 
 	if (!InitializedResources)
 	{
-		UpdateHairGroupsDesc(GroomAsset, GroomGroupsDesc);
 		InitResources();
 	}
 
@@ -1063,6 +1113,7 @@ void UGroomComponent::OnAttachmentChanged()
 	Super::OnAttachmentChanged();
 	if (GroomAsset && !IsBeingDestroyed() && HasBeenCreated())
 	{
+		FlushRenderingCommands();
 		InitResources();
 	}
 }
@@ -1076,180 +1127,33 @@ void UGroomComponent::TickComponent(float DeltaTime, enum ELevelTick TickType, F
 	const ERHIFeatureLevel::Type FeatureLevel = GetWorld() ? ERHIFeatureLevel::Type(GetWorld()->FeatureLevel) : ERHIFeatureLevel::Num;
 
 	FVector MeshPositionOffset = FVector::ZeroVector;
-	USkeletalMeshComponent* SkeletalMeshComponent = RegisteredSkeletalMeshComponent;	
-	if (SkeletalMeshComponent)
-	{
-		// The offset is based on the center of the skeletal mesh (which is computed based on the physics capsules/boxes/...)
-		MeshPositionOffset = SkeletalMeshComponent->CalcBounds(FTransform::Identity).GetBox().GetCenter();
+	USkeletalMeshComponent* SkeletalMeshComponent = RegisteredSkeletalMeshComponent;
 
-		const uint32 SkeletalLODCount = SkeletalMeshComponent->GetNumLODs();
-		const uint32 ResourceLODCount = HairGroupResources.Num() > 0 && HairGroupResources[0].RenRootResources ? HairGroupResources[0].RenRootResources->MeshProjectionLODs.Num() : 0;
-		if (SkeletalLODCount != ResourceLODCount)
+	if (USkeletalMeshComponent* ParentComp = Cast<USkeletalMeshComponent>(GetAttachParent()))
+	{
+		if (ParentComp->GetNumBones() > 0)
 		{
-			auto InitRootResource = [](FHairStrandsRootResource*& PersistentResources, FHairStrandsRootResource* LocalResources)
-			{
-				SafeRelease(PersistentResources);
-				PersistentResources = LocalResources;
-				if (PersistentResources)
-				{
-					BeginInitResource(PersistentResources);
-				}
-			};
+			const int32 BoneIndex = FMath::Min(1, ParentComp->GetNumBones()-1);
+			const FMatrix NextBoneMatrix = ParentComp->GetBoneMatrix(BoneIndex);
 
-			TArray<FHairStrandsRootResource*> GroupRenRootResources;
-			TArray<FHairStrandsRootResource*> GroupSimRootResources;
-			const uint32 GroupCount = GroomAsset ? GroomAsset->GetNumHairGroups() : 0;
-			check(GroupCount == HairGroupResources.Num());
-			for (uint32 GroupIndex=0; GroupIndex < GroupCount; ++GroupIndex)
+			const float BoneDistance = FVector::DistSquared(PrevBoneMatrix.GetOrigin(), NextBoneMatrix.GetOrigin());
+			if (ParentComp->GetTeleportDistanceThreshold() > 0.0 && BoneDistance >
+				ParentComp->GetTeleportDistanceThreshold() * ParentComp->GetTeleportDistanceThreshold())
 			{
-				const FHairGroupData* HairGroupsData = &GroomAsset->HairGroupsData[GroupIndex];
-				FHairGroupResource& Res = HairGroupResources[GroupIndex];
-				FHairStrandsRootResource* LocalRenRootResources = SkeletalLODCount > 0 ? new FHairStrandsRootResource(&HairGroupsData->HairRenderData, SkeletalLODCount) : nullptr;
-				FHairStrandsRootResource* LocalSimRootResources = SkeletalLODCount > 0 ? new FHairStrandsRootResource(&HairGroupsData->HairSimulationData, SkeletalLODCount) : nullptr;
-
-				InitRootResource(Res.RenRootResources, LocalRenRootResources);
-				InitRootResource(Res.SimRootResources, LocalSimRootResources);
-
-				GroupRenRootResources.Add(LocalRenRootResources);
-				GroupSimRootResources.Add(LocalSimRootResources);
+				bResetSimulation = true;
 			}
-
-			FTransform HairLocalToWorld = GetComponentTransform();
-			ENQUEUE_RENDER_COMMAND(FHairStrandsTick_LODUpdate)(
-				[Id, WorldType, GroupRenRootResources, GroupSimRootResources, HairLocalToWorld](FRHICommandListImmediate& RHICmdList)
-			{
-				FHairStrandsProjectionHairData RenProjectionDatas;
-				for (FHairStrandsRootResource* Res : GroupRenRootResources)
-				{
-					RenProjectionDatas.HairGroups.Add(ToProjectionHairData(Res));
-				}
-
-				FHairStrandsProjectionHairData SimProjectionDatas;
-				for (FHairStrandsRootResource* Res : GroupSimRootResources)
-				{
-					SimProjectionDatas.HairGroups.Add(ToProjectionHairData(Res));
-				}
-
-				UpdateHairStrands(Id, WorldType, HairLocalToWorld, RenProjectionDatas, SimProjectionDatas);
-			});
-		}
-
-		// State machine for binding/projecting hair onto skinned mesh. The state machine operates as follow:
-		// * For each mesh LOD
-		//   * Force the skeletal mesh in rest pose
-		//   * Wait for the skeletal mesh to tick/update the animation
-		//   * Issue a binding/projection query to the hair strands system
-		//   * Wait for the query to be completed
-		// * Put the skeletal mesh into its original pose
-		//
-		// In editor, the animation is not ticked, so we manually issue a re-initialization of the animation 
-		// to force the animation in rest pose or original pose
-		if (MeshProjectionState != EMeshProjectionState::Completed)
-		{
-			const uint32 TickDelay = FMath::Max(0, GHairStrandsMeshProjectionTickDelay);
-
-			if (MeshProjectionTickDelay == 0)
+			else
 			{
-				if (MeshProjectionState == EMeshProjectionState::Invalid)
-				{
-					MeshProjectionLODIndex = 0;
-					MeshProjectionState = EMeshProjectionState::InProgressBinding;
-				}
-
-				if (MeshProjectionLODIndex == SkeletalMeshComponent->GetNumLODs())
-				{
-					SkeletalMeshComponent->SetForceRefPose(false);
-					SkeletalMeshComponent->SetForcedLOD(0);
-					if (WorldType == EWorldType::Editor)
-					{
-						SkeletalMeshComponent->InitAnim(true);
-					}
-					MeshProjectionState = EMeshProjectionState::Completed;
-
-					FHairGroupResources* LocalResources = &HairGroupResources;
-					ENQUEUE_RENDER_COMMAND(FHairStrandsTick_UpdateProjectionStatus)(
-						[LocalResources](FRHICommandListImmediate& RHICmdList)
-					{
-						for (FHairGroupResource& Res : *LocalResources)
-						{
-							for (FHairStrandsRootResource::FMeshProjectionLOD& MeshLODData : Res.RenRootResources->MeshProjectionLODs)
-							{
-								MeshLODData.Status = FHairStrandsRootResource::FMeshProjectionLOD::EStatus::Completed;
-							}
-
-							for (FHairStrandsRootResource::FMeshProjectionLOD& MeshLODData : Res.SimRootResources->MeshProjectionLODs)
-							{
-								MeshLODData.Status = FHairStrandsRootResource::FMeshProjectionLOD::EStatus::Completed;
-							}
-						}
-					});
-				}
-
-				if (MeshProjectionLODIndex < SkeletalMeshComponent->GetNumLODs())
-				{
-					if (MeshProjectionState == EMeshProjectionState::InProgressBinding)
-					{
-						SkeletalMeshComponent->SetForceRefPose(true);
-						SkeletalMeshComponent->SetForcedLOD(MeshProjectionLODIndex + 1);
-						if (WorldType == EWorldType::Editor)
-						{
-							SkeletalMeshComponent->InitAnim(true);
-						}
-						MeshProjectionState = EMeshProjectionState::WaitForRestPose;
-						MeshProjectionTickDelay += TickDelay;
-					}
-
-					if (MeshProjectionState == EMeshProjectionState::WaitForRestPose && MeshProjectionTickDelay == 0)
-					{
-						FHairGroupResources* LocalResources = &HairGroupResources;
-
-						const FVector RestPositionOffset = MeshPositionOffset;
-						const uint32 LODIndex = MeshProjectionLODIndex;
-						ENQUEUE_RENDER_COMMAND(FHairStrandsTick_Projection)(
-							[Id, WorldType, FeatureLevel, LODIndex, RestPositionOffset, LocalResources](FRHICommandListImmediate& RHICmdList)
-						{
-							AddHairStrandsProjectionQuery(RHICmdList, Id, WorldType, LODIndex, RestPositionOffset);
-
-							// Update rest (sim/render) position offsets. This is used by Niagara.
-							for (FHairGroupResource& Res : *LocalResources)
-							{
-								if (Res.RenRootResources)
-								{
-									for (FHairStrandsRootResource::FMeshProjectionLOD& MeshProjectionLOD : Res.RenRootResources->MeshProjectionLODs)
-									{
-										MeshProjectionLOD.RestRootOffset = RestPositionOffset;
-									}
-								}
-
-								if (Res.SimRootResources)
-								{
-									for (FHairStrandsRootResource::FMeshProjectionLOD& MeshProjectionLOD : Res.SimRootResources->MeshProjectionLODs)
-									{
-										MeshProjectionLOD.RestRootOffset = RestPositionOffset;
-									}
-								}
-							}
-						});
-
-						MeshProjectionState = EMeshProjectionState::InProgressBinding;
-						MeshProjectionTickDelay += TickDelay;
-						++MeshProjectionLODIndex;
-					}
-				}
+				bResetSimulation = false;
 			}
+			PrevBoneMatrix = NextBoneMatrix;
 		}
-
-		FSkeletalMeshConfiguration CurrentState;
-		CurrentState.ForceRefPose = GHairStrandsMeshProjectionForceRefPoseEnable > 0;
-		CurrentState.ForceLOD     = GHairStrandsMeshProjectionForceLOD >= 0 ? FMath::Clamp(GHairStrandsMeshProjectionForceLOD, 0, SkeletalMeshComponent->GetNumLODs() - 1) : -1;
-		if (!FSkeletalMeshConfiguration::Equals(CurrentState, SkeletalMeshConfiguration) && MeshProjectionState == EMeshProjectionState::Completed)
-		{
-			SkeletalMeshComponent->SetForcedLOD(CurrentState.ForceLOD + 1);
-			SkeletalMeshComponent->SetForceRefPose(CurrentState.ForceRefPose);
-			SkeletalMeshConfiguration = CurrentState;
-		}
-
+	}
+	if (SkeletalMeshComponent)
+	{
+		// The offset is based on the center of the skeletal mesh (which is computed based on the physics capsules/boxes/...)
 		// For skinned mesh update the relative center of hair positions after deformation
+		MeshPositionOffset = SkeletalMeshComponent->CalcBounds(FTransform::Identity).GetBox().GetCenter();
 		{
 			FHairGroupResources* LocalResources = &HairGroupResources;
 
@@ -1283,12 +1187,6 @@ void UGroomComponent::TickComponent(float DeltaTime, enum ELevelTick TickType, F
 		// When a skeletal object with projection is enabled, activate the refresh of the bounding box to 
 		// insure the component/proxy bounding box always lies onto the actual skinned mesh
 		MarkRenderTransformDirty();
-
-	}
-
-	if (MeshProjectionTickDelay > 0)
-	{
-		MeshProjectionTickDelay--;
 	}
 
 	FHairGroupResources* LocalResources = &HairGroupResources;
@@ -1372,7 +1270,11 @@ void UGroomComponent::PostEditChangeProperty(FPropertyChangedEvent& PropertyChan
 
 	//  Init/release resources when setting the GroomAsset (or undoing)
 	const bool bAssetChanged = PropertyName == GET_MEMBER_NAME_CHECKED(UGroomComponent, GroomAsset);
-	if (bAssetChanged)
+	const bool bRecreateResources =
+		(bAssetChanged || PropertyThatChanged == nullptr) ||
+		PropertyName == GET_MEMBER_NAME_CHECKED(UGroomComponent, bBindGroomToSkeletalMesh);
+
+	if (bRecreateResources)
 	{
 		// Release the resources before Super::PostEditChangeProperty so that they get
 		// re-initialized in OnRegister
@@ -1381,11 +1283,7 @@ void UGroomComponent::PostEditChangeProperty(FPropertyChangedEvent& PropertyChan
 
 	Super::PostEditChangeProperty(PropertyChangedEvent);
 
-	const bool bRecreateResources =
-		(PropertyName == GET_MEMBER_NAME_CHECKED(UGroomComponent, GroomAsset) || PropertyThatChanged == nullptr) ||
-		PropertyName == GET_MEMBER_NAME_CHECKED(UGroomComponent, bBindGroomToSkeletalMesh);
-
-	const bool bSupportSkinProjection = GetDefault<URendererSettings>()->bSupportSkinCacheShaders && !UAnimationSettings::Get()->bTickAnimationOnSkeletalMeshInit;
+	const bool bSupportSkinProjection = GetDefault<URendererSettings>()->bSupportSkinCacheShaders;
 	if (!bSupportSkinProjection)
 	{
 		bBindGroomToSkeletalMesh = false;
@@ -1403,35 +1301,26 @@ void UGroomComponent::PostEditChangeProperty(FPropertyChangedEvent& PropertyChan
 	}
 #endif
 
-	if (bRecreateResources)
-	{
-		//if (GroomAsset != InitializedResources)
+	if (PropertyName == GET_MEMBER_NAME_CHECKED(FHairGroupDesc, HairWidth) || 
+		PropertyName == GET_MEMBER_NAME_CHECKED(FHairGroupDesc, HairLength) ||
+		PropertyName == GET_MEMBER_NAME_CHECKED(FHairGroupDesc, HairRootScale) ||
+		PropertyName == GET_MEMBER_NAME_CHECKED(FHairGroupDesc, HairTipScale) ||
+		PropertyName == GET_MEMBER_NAME_CHECKED(FHairGroupDesc, HairClipLength) ||
+		PropertyName == GET_MEMBER_NAME_CHECKED(FHairGroupDesc, HairShadowDensity) ||
+		PropertyName == GET_MEMBER_NAME_CHECKED(FHairGroupDesc, HairRaytracingRadiusScale) ||
+		PropertyName == GET_MEMBER_NAME_CHECKED(FHairGroupDesc, LodBias) ||
+		PropertyName == GET_MEMBER_NAME_CHECKED(FHairGroupDesc, LodAverageVertexPerPixel))
+	{	
+		UpdateHairGroupsDesc(GroomAsset, GroomGroupsDesc);
+
+		if (InterpolationInput)
 		{
-			if (GroomAsset)
-			{
-				UpdateHairGroupsDesc(GroomAsset, GroomGroupsDesc);
-				InitResources();
-			}
-			else
+			const int32 GroupCount = InterpolationInput->HairGroups.Num();
+			for (int32 GroupIt = 0; GroupIt < GroupCount; ++GroupIt)
 			{
-				ReleaseResources();
+				InterpolationInput->HairGroups[GroupIt].GroupDesc = GroomGroupsDesc[GroupIt];
 			}
 		}
-	}
-
-	if (PropertyName == GET_MEMBER_NAME_CHECKED(FHairGroupDesc, HairRaytracingRadiusScale) && InterpolationInput && InterpolationInput->HairGroups.Num() == GroomGroupsDesc.Num())
-	{
-		const int32 GroupCount = InterpolationInput->HairGroups.Num();
-		for (int32 GroupIt = 0; GroupIt < GroupCount; ++GroupIt)
-		{
-			InterpolationInput->HairGroups[GroupIt].HairRaytracingRadiusScale = GroomGroupsDesc[GroupIt].HairRaytracingRadiusScale;
-		}
-	}
-
-	if (PropertyName == GET_MEMBER_NAME_CHECKED(FHairGroupDesc, HairWidth) || 
-		PropertyName == GET_MEMBER_NAME_CHECKED(FHairGroupDesc, HairShadowDensity))
-	{	
-		UpdateHairGroupsDesc(GroomAsset, GroomGroupsDesc);
 		MarkRenderStateDirty();
 	}
 
@@ -1450,7 +1339,7 @@ bool UGroomComponent::CanEditChange(const FProperty* InProperty) const
 		{
 			bool bIsEditable = false;
 			#if RHI_RAYTRACING
-			if (IsRayTracingEnabled())
+			if (IsHairRayTracingEnabled())
 			{
 				bIsEditable = true;
 			}
@@ -1460,7 +1349,7 @@ bool UGroomComponent::CanEditChange(const FProperty* InProperty) const
 
 		if (FCString::Strcmp(*PropertyName, TEXT("bBindGroomToSkeletalMesh")) == 0)
 		{
-			return GetDefault<URendererSettings>()->bSupportSkinCacheShaders && !UAnimationSettings::Get()->bTickAnimationOnSkeletalMeshInit;
+			return GetDefault<URendererSettings>()->bSupportSkinCacheShaders;
 		}
 	}
 
diff --git a/Engine/Plugins/Experimental/HairStrands/Source/HairStrandsCore/Private/HairAttributes.cpp b/Engine/Plugins/Experimental/HairStrands/Source/HairStrandsCore/Private/HairAttributes.cpp
--- a/Engine/Plugins/Experimental/HairStrands/Source/HairStrandsCore/Private/HairAttributes.cpp
+++ b/Engine/Plugins/Experimental/HairStrands/Source/HairStrandsCore/Private/HairAttributes.cpp
@@ -15,6 +15,8 @@ namespace HairAttribute
 	const FName Strand::RootUV("groom_root_uv");
 	const FName Strand::VertexCount("vertexcount");
 	const FName Strand::Width("groom_width");
+	const FName Strand::ClosestGuides("groom_closest_guides");
+	const FName Strand::GuideWeights("groom_guide_weights");
 
 	const FName Groom::Color("groom_color");
 	const FName Groom::Width("groom_width");
diff --git a/Engine/Plugins/Experimental/HairStrands/Source/HairStrandsCore/Private/HairDescription.cpp b/Engine/Plugins/Experimental/HairStrands/Source/HairStrandsCore/Private/HairDescription.cpp
--- a/Engine/Plugins/Experimental/HairStrands/Source/HairStrandsCore/Private/HairDescription.cpp
+++ b/Engine/Plugins/Experimental/HairStrands/Source/HairStrandsCore/Private/HairDescription.cpp
@@ -1,6 +1,9 @@
 // Copyright Epic Games, Inc. All Rights Reserved.
 
 #include "HairDescription.h"
+#include "Misc/SecureHash.h"
+#include "Serialization/BulkDataReader.h"
+#include "Serialization/BulkDataWriter.h"
 
 const FStrandID FStrandID::Invalid(TNumericLimits<uint32>::Max());
 const FGroomID FGroomID::Invalid(TNumericLimits<uint32>::Max());
@@ -43,6 +46,16 @@ FStrandID FHairDescription::AddStrand()
 	return StrandID;
 }
 
+void FHairDescription::Reset()
+{
+	NumVertices = 0;
+	NumStrands = 0;
+
+	VertexAttributesSet.Initialize(0);
+	StrandAttributesSet.Initialize(0);
+	GroomAttributesSet.Initialize(0);
+}
+
 bool FHairDescription::IsValid() const
 {
 	return (NumStrands > 0) && (NumVertices > 0);
@@ -56,4 +69,82 @@ void FHairDescription::Serialize(FArchive& Ar)
 	Ar << VertexAttributesSet;
 	Ar << StrandAttributesSet;
 	Ar << GroomAttributesSet;
-}
\ No newline at end of file
+}
+
+#if WITH_EDITORONLY_DATA
+
+void FHairDescriptionBulkData::Serialize(FArchive& Ar, UObject* Owner)
+{
+	if (Ar.IsLoading())
+	{
+		// If loading, take the package custom version so it can be applied to the bulk data archive
+		// when unpacking HairDescription from it
+		CustomVersions = Ar.GetCustomVersions();
+	}
+
+	BulkData.Serialize(Ar, Owner);
+
+	Ar << Guid;
+}
+
+void FHairDescriptionBulkData::SaveHairDescription(FHairDescription& HairDescription)
+{
+	BulkData.RemoveBulkData();
+
+	if (HairDescription.IsValid())
+	{
+		FBulkDataWriter Ar(BulkData, /*bIsPersistent*/ true);
+		HairDescription.Serialize(Ar);
+
+		// Preserve CustomVersions at save time so we can reuse the same ones when reloading direct from memory
+		CustomVersions = Ar.GetCustomVersions();
+	}
+
+	// Use bulk data hash instead of guid to identify content to improve DDC cache hit
+	ComputeGuidFromHash();
+}
+
+
+void FHairDescriptionBulkData::LoadHairDescription(FHairDescription& HairDescription)
+{
+	HairDescription.Reset();
+
+	if (!IsEmpty())
+	{
+		FBulkDataReader Ar(BulkData, /*bIsPersistent*/ true);
+
+		// Propagate the custom version information from the package to the bulk data, so that the HairDescription
+		// is serialized with the same versioning
+		Ar.SetCustomVersions(CustomVersions);
+
+		HairDescription.Serialize(Ar);
+	}
+}
+
+void FHairDescriptionBulkData::Empty()
+{
+	BulkData.RemoveBulkData();
+}
+
+FString FHairDescriptionBulkData::GetIdString() const
+{
+	FString GuidString = Guid.ToString();
+	GuidString += TEXT("X");
+	return GuidString;
+}
+
+void FHairDescriptionBulkData::ComputeGuidFromHash()
+{
+	uint32 Hash[5] = {};
+
+	if (BulkData.GetBulkDataSize() > 0)
+	{
+		void* Buffer = BulkData.Lock(LOCK_READ_ONLY);
+		FSHA1::HashBuffer(Buffer, BulkData.GetBulkDataSize(), (uint8*)Hash);
+		BulkData.Unlock();
+	}
+
+	Guid = FGuid(Hash[0] ^ Hash[4], Hash[1], Hash[2], Hash[3]);
+}
+
+#endif // #if WITH_EDITORONLY_DATA
