Copying //UE4/Dev-Physics to Main (//UE4/Main)

diff --git a/Engine/Plugins/Experimental/HairStrands/Source/HairStrandsCore/Private/HairStrandsRendering.cpp b/Engine/Plugins/Experimental/HairStrands/Source/HairStrandsCore/Private/HairStrandsRendering.cpp
--- a/Engine/Plugins/Experimental/HairStrands/Source/HairStrandsCore/Private/HairStrandsRendering.cpp
+++ b/Engine/Plugins/Experimental/HairStrands/Source/HairStrandsCore/Private/HairStrandsRendering.cpp
@@ -7,6 +7,10 @@
 #include "ShaderParameterStruct.h"
 #include "GlobalShader.h"
 #include "HairStrandsInterface.h"
+#include "SceneView.h"
+#include "Containers/ResourceArray.h"
+#include "HAL/ConsoleManager.h"
+#include "GpuDebugRendering.h"
 
 static int32 GHairDeformationType = 0;
 static FAutoConsoleVariableRef CVarHairDeformationType(TEXT("r.HairStrands.DeformationType"), GHairDeformationType, TEXT("Type of procedural deformation applied on hair strands (0:use simulation's output, 1:use rest strands, 2: use rest guides, 3:wave pattern, 4:follow root normal)"));
@@ -17,6 +21,13 @@ static FAutoConsoleVariableRef CVarHairRaytracingRadiusScale(TEXT("r.HairStrands
 static int32 GHairStrandsInterpolateSimulation = 1;
 static FAutoConsoleVariableRef CVarHairInterpolateSimulation(TEXT("r.HairStrands.InterpolateSimulation"), GHairStrandsInterpolateSimulation, TEXT("Enable/disable simulation output during the hair interpolation"));
 
+static float GStrandHairWidth = 0.0f;
+static FAutoConsoleVariableRef CVarStrandHairWidth(TEXT("r.HairStrands.StrandWidth"), GStrandHairWidth, TEXT("Width of hair strand"));
+
+static int32 GStrandHairInterpolationDebug = 0;
+static FAutoConsoleVariableRef CVarStrandHairInterpolationDebug(TEXT("r.HairStrands.Interpolation.Debug"), GStrandHairInterpolationDebug, TEXT("Enable debug rendering for hair interpolation"));
+
+
 static FIntVector ComputeDispatchCount(uint32 ItemCount, uint32 GroupSize)
 {
 	const uint32 BatchCount = FMath::DivideAndRoundUp(ItemCount, GroupSize);
@@ -29,6 +40,18 @@ static FIntVector ComputeDispatchCount(uint32 ItemCount, uint32 GroupSize)
 	return FIntVector(DispatchCountX, DispatchCountY, 1);
 }
 
+// Same as above but the group count is what matters and is preserved
+static FIntVector ComputeDispatchGroupCount2D(uint32 GroupCount)
+{
+	const uint32 DispatchCountX = FMath::FloorToInt(FMath::Sqrt(GroupCount));
+	const uint32 DispatchCountY = DispatchCountX + FMath::DivideAndRoundUp(GroupCount - DispatchCountX * DispatchCountX, DispatchCountX);
+
+	check(DispatchCountX <= 65535);
+	check(DispatchCountY <= 65535);
+	check(GroupCount <= DispatchCountX * DispatchCountY);
+	return FIntVector(DispatchCountX, DispatchCountY, 1);
+}
+
 inline uint32 ComputeGroupSize()
 {
 	const uint32 GroupSize = IsRHIDeviceAMD() ? 64 : (IsRHIDeviceNVIDIA() ? 32 : 64);
@@ -47,7 +70,8 @@ enum class EDeformationType : uint8
 	RestStrands,	// Use the rest strands position (no weighted interpolation)
 	RestGuide,		// Use the rest guide as input of the interpolation (no deformation), only weighted interpolation
 	Wave,			// Apply a wave pattern to deform the guides
-	NormalDirection // Apply a stretch pattern aligned with the guide root's normal
+	NormalDirection,// Apply a stretch pattern aligned with the guide root's normal
+	OffsetGuide		// Offset the guides
 };
 
 static EDeformationType GetDeformationType()
@@ -70,12 +94,15 @@ class FDeformGuideCS : public FGlobalShader
 	SHADER_USE_PARAMETER_STRUCT(FDeformGuideCS, FGlobalShader);
 
 	class FGroupSize : SHADER_PERMUTATION_INT("PERMUTATION_GROUP_SIZE", 2);
-	class FDeformationType : SHADER_PERMUTATION_INT("PERMUTATION_DEFORMATION", 3);
+	class FDeformationType : SHADER_PERMUTATION_INT("PERMUTATION_DEFORMATION", 4);
 	using FPermutationDomain = TShaderPermutationDomain<FGroupSize, FDeformationType>;
 
 	BEGIN_SHADER_PARAMETER_STRUCT(FParameters, )
 		SHADER_PARAMETER(uint32, VertexCount)
 		SHADER_PARAMETER(uint32, IterationCount)
+		SHADER_PARAMETER(FVector, SimRestOffset)
+		SHADER_PARAMETER(FVector, SimDeformedOffset)
+		SHADER_PARAMETER(uint32, DispatchCountX)
 		SHADER_PARAMETER_SRV(Buffer, SimRestPosePositionBuffer)
 		SHADER_PARAMETER_SRV(Buffer, SimRootIndexBuffer)
 		SHADER_PARAMETER_UAV(RWBuffer, OutSimDeformedPositionBuffer)
@@ -93,7 +120,9 @@ static void AddDeformSimHairStrandsPass(
 	uint32 VertexCount,
 	FRHIShaderResourceView* SimRestPosePositionBuffer,
 	FRHIShaderResourceView* SimRootIndexBuffer,
-	FRHIUnorderedAccessView* OutSimDeformedPositionBuffer)
+	FRHIUnorderedAccessView* OutSimDeformedPositionBuffer,
+	FVector& SimRestOffset,
+	FVector& SimDeformedOffset)
 {
 	static uint32 IterationCount = 0;
 	++IterationCount;
@@ -104,17 +133,25 @@ static void AddDeformSimHairStrandsPass(
 	case EDeformationType::RestGuide: InternalDeformationType = 0; break;
 	case EDeformationType::Wave: InternalDeformationType = 1; break;
 	case EDeformationType::NormalDirection: InternalDeformationType = 2; break;
+	case EDeformationType::OffsetGuide : InternalDeformationType = 3; break;
 	}
 
 	if (InternalDeformationType < 0) return;
 
+	const uint32 GroupSize = ComputeGroupSize();
+	const uint32 DispatchCount = FMath::DivideAndRoundUp(VertexCount, GroupSize);
+	const uint32 DispatchCountX = 16;
+	const uint32 DispatchCountY = FMath::DivideAndRoundUp(DispatchCount, DispatchCountX);
+
 	FDeformGuideCS::FParameters* Parameters = GraphBuilder.AllocParameters<FDeformGuideCS::FParameters>();
 	Parameters->SimRestPosePositionBuffer = SimRestPosePositionBuffer;
 	Parameters->SimRootIndexBuffer = SimRootIndexBuffer;
 	Parameters->OutSimDeformedPositionBuffer = OutSimDeformedPositionBuffer;
 	Parameters->VertexCount = VertexCount;
 	Parameters->IterationCount = IterationCount % 10000;
-	const uint32 GroupSize = ComputeGroupSize();
+	Parameters->SimDeformedOffset = SimDeformedOffset;
+	Parameters->SimRestOffset = SimRestOffset;
+	Parameters->DispatchCountX = DispatchCountX;
 
 	FDeformGuideCS::FPermutationDomain PermutationVector;
 	PermutationVector.Set<FDeformGuideCS::FGroupSize>(GetGroupSizePermutation(GroupSize));
@@ -122,54 +159,79 @@ static void AddDeformSimHairStrandsPass(
 
 	TShaderMap<FGlobalShaderType>* ShaderMap = GetGlobalShaderMap(ERHIFeatureLevel::SM5);
 
-	const uint32 DispatchCount = FMath::DivideAndRoundUp(VertexCount, GroupSize);
-	check(DispatchCount <= 65535);
-
 	TShaderMapRef<FDeformGuideCS> ComputeShader(ShaderMap, PermutationVector);
 	FComputeShaderUtils::AddPass(
 		GraphBuilder,
 		RDG_EVENT_NAME("DeformSimHairStrands"),
 		*ComputeShader,
 		Parameters,
-		FIntVector(DispatchCount, 1, 1));
+		FIntVector(DispatchCountX, DispatchCountY, 1));
 }
 
 ///////////////////////////////////////////////////////////////////////////////////////////////////
 
+struct FHairScaleAndClipDesc
+{
+	float InHairLength;
+	float InHairRadius;
+	float OutHairRadius;
+	float MaxOutHairRadius;
+	float HairRadiusRootScale;
+	float HairRadiusTipScale;
+	float HairLengthClip;
+};
+
 class FHairInterpolationCS : public FGlobalShader
 {
 	DECLARE_GLOBAL_SHADER(FHairInterpolationCS);
 	SHADER_USE_PARAMETER_STRUCT(FHairInterpolationCS, FGlobalShader);
 
 	class FGroupSize : SHADER_PERMUTATION_INT("PERMUTATION_GROUP_SIZE", 2);
-	class FDebug : SHADER_PERMUTATION_INT("PERMUTATION_DEBUG", 3);
+	class FDebug : SHADER_PERMUTATION_INT("PERMUTATION_DEBUG", 2);
 	class FDynamicGeometry : SHADER_PERMUTATION_INT("PERMUTATION_DYNAMIC_GEOMETRY", 2);
 	class FSimulation : SHADER_PERMUTATION_INT("PERMUTATION_SIMULATION", 2);
-	using FPermutationDomain = TShaderPermutationDomain<FGroupSize, FDebug, FDynamicGeometry, FSimulation>;
+	class FScaleAndClip : SHADER_PERMUTATION_INT("PERMUTATION_SCALE_AND_CLIP", 2);
+	using FPermutationDomain = TShaderPermutationDomain<FGroupSize, FDebug, FDynamicGeometry, FSimulation, FScaleAndClip>;
 
 	BEGIN_SHADER_PARAMETER_STRUCT(FParameters, )
+		SHADER_PARAMETER_STRUCT_INCLUDE(ShaderDrawDebug::FShaderDrawDebugParameters, ShaderDrawParameters)
 		SHADER_PARAMETER(uint32, VertexCount)
 		SHADER_PARAMETER(uint32, DispatchCountX)
+		SHADER_PARAMETER(uint32, HairDebugMode)
 		SHADER_PARAMETER(FVector, InRenderHairPositionOffset)
 		SHADER_PARAMETER(FVector, InSimHairPositionOffset)
 		SHADER_PARAMETER(FVector, OutHairPositionOffset)
 		SHADER_PARAMETER(FIntPoint, HairStrandsCullIndex)
 
+		SHADER_PARAMETER(float, InHairLength)
+		SHADER_PARAMETER(float, InHairRadius)
+		SHADER_PARAMETER(float, OutHairRadius)
+		SHADER_PARAMETER(float, MaxOutHairRadius)
+		SHADER_PARAMETER(float, HairRadiusRootScale)
+		SHADER_PARAMETER(float, HairRadiusTipScale)
+		SHADER_PARAMETER(float, HairLengthClip)
+
 		SHADER_PARAMETER(FVector, RestPositionOffset)
 		SHADER_PARAMETER(FVector, DeformedPositionOffset)
 
 		SHADER_PARAMETER(FVector, SimRestPositionOffset)
 		SHADER_PARAMETER(FVector, SimDeformedPositionOffset)
 
+		SHADER_PARAMETER(FMatrix, LocalToWorldMatrix)
+
 		SHADER_PARAMETER_SRV(Buffer, RenderRestPosePositionBuffer)
 		SHADER_PARAMETER_UAV(RWBuffer, OutRenderDeformedPositionBuffer)
 
+		SHADER_PARAMETER_UAV(RWBuffer, OutClusterBuffer)
+		SHADER_PARAMETER_SRV(Buffer, VertexToClusterIdBuffer)
+
 		SHADER_PARAMETER_SRV(Buffer, SimRestPosePositionBuffer)
 		SHADER_PARAMETER_SRV(Buffer, DeformedSimPositionBuffer)
 
 		SHADER_PARAMETER_SRV(Buffer, Interpolation0Buffer)
 		SHADER_PARAMETER_SRV(Buffer, Interpolation1Buffer)
 
+		SHADER_PARAMETER_SRV(Buffer, AttributeBuffer)
 		SHADER_PARAMETER_SRV(Buffer, SimAttributeBuffer)
 		SHADER_PARAMETER_UAV(RWBuffer, OutRenderAttributeBuffer)
 		SHADER_PARAMETER_SRV(Buffer<float4>, RestPosition0Buffer)
@@ -206,22 +268,27 @@ IMPLEMENT_GLOBAL_SHADER(FHairInterpolationCS, "/Engine/Private/HairStrands/HairS
 
 static void AddHairStrandsInterpolationPass(
 	FRDGBuilder& GraphBuilder,
+	const FShaderDrawDebugData* ShaderDrawData,
 	const FHairStrandsProjectionHairData::HairGroup& InRenHairData,
 	const FHairStrandsProjectionHairData::HairGroup& InSimHairData,
 	const FVector& InRenderHairWorldOffset,
 	const FVector& InSimHairWorldOffset,
 	const FVector& OutHairWorldOffset,
+	const FHairScaleAndClipDesc ScaleAndClipDesc,
 	const int32 LODIndex,
 	const bool bHasSimulationEnable,
+	const bool bPatchedAttributeBuffer, 
 	const uint32 VertexCount,
 	const FShaderResourceViewRHIRef& RenderRestPosePositionBuffer,
+	const FShaderResourceViewRHIRef& RenderAttributeBuffer,
 	const FShaderResourceViewRHIRef& Interpolation0Buffer,
 	const FShaderResourceViewRHIRef& Interpolation1Buffer,
 	const FShaderResourceViewRHIRef& SimRestPosePositionBuffer,
 	const FShaderResourceViewRHIRef& SimDeformedPositionBuffer,
 	const FShaderResourceViewRHIRef& SimAttributeBuffer,
 	const FUnorderedAccessViewRHIRef& OutRenderPositionBuffer,
-	const FUnorderedAccessViewRHIRef& OutRenderAttributeBuffer)
+	const FUnorderedAccessViewRHIRef& OutRenderAttributeBuffer,
+	const FShaderResourceViewRHIRef& VertexToClusterIdBuffer)
 {
 	const uint32 GroupSize = ComputeGroupSize();
 	const FIntVector DispatchCount = ComputeDispatchCount(VertexCount, GroupSize);
@@ -240,6 +307,23 @@ static void AddHairStrandsInterpolationPass(
 	Parameters->OutHairPositionOffset = OutHairWorldOffset;
 	Parameters->DispatchCountX = DispatchCount.X;
 	
+	const bool bNeedScaleOrClip = 
+		ScaleAndClipDesc.InHairRadius != ScaleAndClipDesc.OutHairRadius || 
+		ScaleAndClipDesc.HairRadiusRootScale != 1 || 
+		ScaleAndClipDesc.HairRadiusTipScale != 1 ||
+		ScaleAndClipDesc.HairLengthClip < 1;
+
+	Parameters->InHairLength = ScaleAndClipDesc.InHairLength;
+	Parameters->InHairRadius = ScaleAndClipDesc.InHairRadius;
+	Parameters->OutHairRadius = ScaleAndClipDesc.OutHairRadius;
+	Parameters->MaxOutHairRadius = ScaleAndClipDesc.MaxOutHairRadius;
+	Parameters->HairRadiusRootScale = ScaleAndClipDesc.HairRadiusRootScale;
+	Parameters->HairRadiusTipScale = ScaleAndClipDesc.HairRadiusTipScale;
+	Parameters->HairLengthClip = ScaleAndClipDesc.HairLengthClip * ScaleAndClipDesc.InHairLength; // HairLengthClip is the normalized length and we convert it to world length
+	if (bNeedScaleOrClip)
+	{
+		Parameters->AttributeBuffer = RenderAttributeBuffer;
+	}
 	const bool bIsVertexToCurveBuffersValid = InRenHairData.VertexToCurveIndexBuffer && InSimHairData.VertexToCurveIndexBuffer;
 	if (bIsVertexToCurveBuffersValid)
 	{
@@ -247,32 +331,56 @@ static void AddHairStrandsInterpolationPass(
 		Parameters->SimVertexToRootIndexBuffer = InSimHairData.VertexToCurveIndexBuffer->SRV;
 	}
 
-	const FHairCullInfo Info = GetHairStrandsCullInfo();
-	const bool bCullingEnable = Info.CullMode != EHairCullMode::None && bIsVertexToCurveBuffersValid;
-	if (bCullingEnable)
-	{
-		if (Info.CullMode == EHairCullMode::Sim)
-			Parameters->HairStrandsCullIndex.Y = Info.ExplicitIndex >= 0 ? Info.ExplicitIndex : FMath::Clamp(uint32(Info.NormalizedIndex * InSimHairData.RootCount), 0u, InSimHairData.RootCount - 1);
-		if (Info.CullMode == EHairCullMode::Render)
-			Parameters->HairStrandsCullIndex.X = Info.ExplicitIndex >= 0 ? Info.ExplicitIndex : FMath::Clamp(uint32(Info.NormalizedIndex * InRenHairData.RootCount), 0u, InRenHairData.RootCount - 1);
-	}
+	Parameters->VertexToClusterIdBuffer = VertexToClusterIdBuffer;
+	
+	Parameters->LocalToWorldMatrix = InRenHairData.LocalToWorld.ToMatrixWithScale();
 
-	const bool bCopySimAttributesToRenderAttributes = SimAttributeBuffer != nullptr && OutRenderAttributeBuffer != nullptr;
-	if (bCopySimAttributesToRenderAttributes)
+	// Debug rendering
+	Parameters->HairDebugMode = 0;
 	{
-		Parameters->SimAttributeBuffer = SimAttributeBuffer;
-		Parameters->OutRenderAttributeBuffer = OutRenderAttributeBuffer;
+		const FHairCullInfo Info = GetHairStrandsCullInfo();
+		const bool bCullingEnable = Info.CullMode != EHairCullMode::None && bIsVertexToCurveBuffersValid;
+
+		if (bPatchedAttributeBuffer)
+		{
+			Parameters->HairDebugMode = 1;
+			check(SimAttributeBuffer != nullptr);
+			check(OutRenderAttributeBuffer != nullptr);
+		}
+		else if (GStrandHairInterpolationDebug > 0)
+		{
+			Parameters->HairDebugMode = 4;
+		}
+		else if (GetHairStrandsDebugStrandsMode() == EHairStrandsDebugMode::RenderVisCluster)
+		{
+			Parameters->HairDebugMode = 3;
+		}
+		else if (bCullingEnable)
+		{
+			Parameters->HairDebugMode = 2;
+			
+			if (Info.CullMode == EHairCullMode::Sim)
+				Parameters->HairStrandsCullIndex.Y = Info.ExplicitIndex >= 0 ? Info.ExplicitIndex : FMath::Clamp(uint32(Info.NormalizedIndex * InSimHairData.RootCount), 0u, InSimHairData.RootCount - 1);
+			if (Info.CullMode == EHairCullMode::Render)
+				Parameters->HairStrandsCullIndex.X = Info.ExplicitIndex >= 0 ? Info.ExplicitIndex : FMath::Clamp(uint32(Info.NormalizedIndex * InRenHairData.RootCount), 0u, InRenHairData.RootCount - 1);
+		}
+
+		if (Parameters->HairDebugMode > 0)
+		{
+			Parameters->SimAttributeBuffer = SimAttributeBuffer;
+			Parameters->OutRenderAttributeBuffer = OutRenderAttributeBuffer;
+		}
 	}
 
 	const bool bSupportDynamicMesh = InRenHairData.RootCount > 0 && LODIndex >= 0 && LODIndex < InRenHairData.LODDatas.Num() && InRenHairData.LODDatas[LODIndex].bIsValid && bIsVertexToCurveBuffersValid;
 	if (bSupportDynamicMesh)
 	{
-		Parameters->RestPositionOffset = InRenHairData.LODDatas[LODIndex].RestPositionOffset;
+		Parameters->RestPositionOffset = *InRenHairData.LODDatas[LODIndex].RestPositionOffset;
 		Parameters->RestPosition0Buffer = InRenHairData.LODDatas[LODIndex].RestRootTrianglePosition0Buffer->SRV;
 		Parameters->RestPosition1Buffer = InRenHairData.LODDatas[LODIndex].RestRootTrianglePosition1Buffer->SRV;
 		Parameters->RestPosition2Buffer = InRenHairData.LODDatas[LODIndex].RestRootTrianglePosition2Buffer->SRV;
 
-		Parameters->DeformedPositionOffset = InRenHairData.LODDatas[LODIndex].DeformedPositionOffset;
+		Parameters->DeformedPositionOffset = *InRenHairData.LODDatas[LODIndex].DeformedPositionOffset;
 		Parameters->DeformedPosition0Buffer = InRenHairData.LODDatas[LODIndex].DeformedRootTrianglePosition0Buffer->SRV;
 		Parameters->DeformedPosition1Buffer = InRenHairData.LODDatas[LODIndex].DeformedRootTrianglePosition1Buffer->SRV;
 		Parameters->DeformedPosition2Buffer = InRenHairData.LODDatas[LODIndex].DeformedRootTrianglePosition2Buffer->SRV;
@@ -280,12 +388,12 @@ static void AddHairStrandsInterpolationPass(
 		Parameters->RootToTriangleIndex = InRenHairData.LODDatas[LODIndex].RootTriangleIndexBuffer->SRV;
 		Parameters->RootBarycentricBuffer = InRenHairData.LODDatas[LODIndex].RootTriangleBarycentricBuffer->SRV;
 
-		Parameters->SimRestPositionOffset = InSimHairData.LODDatas[LODIndex].RestPositionOffset;
+		Parameters->SimRestPositionOffset = *InSimHairData.LODDatas[LODIndex].RestPositionOffset;
 		Parameters->SimRestPosition0Buffer = InSimHairData.LODDatas[LODIndex].RestRootTrianglePosition0Buffer->SRV;
 		Parameters->SimRestPosition1Buffer = InSimHairData.LODDatas[LODIndex].RestRootTrianglePosition1Buffer->SRV;
 		Parameters->SimRestPosition2Buffer = InSimHairData.LODDatas[LODIndex].RestRootTrianglePosition2Buffer->SRV;
 
-		Parameters->SimDeformedPositionOffset = InSimHairData.LODDatas[LODIndex].DeformedPositionOffset;
+		Parameters->SimDeformedPositionOffset = *InSimHairData.LODDatas[LODIndex].DeformedPositionOffset;
 		Parameters->SimDeformedPosition0Buffer = InSimHairData.LODDatas[LODIndex].DeformedRootTrianglePosition0Buffer->SRV;
 		Parameters->SimDeformedPosition1Buffer = InSimHairData.LODDatas[LODIndex].DeformedRootTrianglePosition1Buffer->SRV;
 		Parameters->SimDeformedPosition2Buffer = InSimHairData.LODDatas[LODIndex].DeformedRootTrianglePosition2Buffer->SRV;
@@ -294,11 +402,17 @@ static void AddHairStrandsInterpolationPass(
 		Parameters->SimRootBarycentricBuffer = InSimHairData.LODDatas[LODIndex].RootTriangleBarycentricBuffer->SRV;
 	}
 
+	if (ShaderDrawDebug::IsShaderDrawDebugEnabled() && ShaderDrawData)
+	{
+		ShaderDrawDebug::SetParameters(GraphBuilder, *ShaderDrawData, Parameters->ShaderDrawParameters);
+	}
+
 	FHairInterpolationCS::FPermutationDomain PermutationVector;
 	PermutationVector.Set<FHairInterpolationCS::FGroupSize>(GetGroupSizePermutation(GroupSize));
-	PermutationVector.Set<FHairInterpolationCS::FDebug>(bCopySimAttributesToRenderAttributes ? 1 : (bCullingEnable ? 2 : 0));
+	PermutationVector.Set<FHairInterpolationCS::FDebug>(Parameters->HairDebugMode > 0 ? 1 : 0);
 	PermutationVector.Set<FHairInterpolationCS::FDynamicGeometry>(bSupportDynamicMesh ? 1 : 0);
 	PermutationVector.Set<FHairInterpolationCS::FSimulation>(bHasSimulationEnable ? 1 : 0);
+	PermutationVector.Set<FHairInterpolationCS::FScaleAndClip>(bNeedScaleOrClip ? 1 : 0);
 
 	TShaderMap<FGlobalShaderType>* ShaderMap = GetGlobalShaderMap(ERHIFeatureLevel::SM5);
 
@@ -313,6 +427,68 @@ static void AddHairStrandsInterpolationPass(
 
 ///////////////////////////////////////////////////////////////////////////////////////////////////
 
+
+class FHairClusterAABBCS : public FGlobalShader
+{
+	DECLARE_GLOBAL_SHADER(FHairClusterAABBCS);
+	SHADER_USE_PARAMETER_STRUCT(FHairClusterAABBCS, FGlobalShader);
+
+	class FGroupSize : SHADER_PERMUTATION_INT("PERMUTATION_GROUP_SIZE", 2);
+	using FPermutationDomain = TShaderPermutationDomain<FGroupSize>;
+
+	BEGIN_SHADER_PARAMETER_STRUCT(FParameters, )
+		SHADER_PARAMETER(uint32, DispatchCountX)
+		SHADER_PARAMETER(uint32, ClusterCount)
+		SHADER_PARAMETER(FVector, OutHairPositionOffset)
+		SHADER_PARAMETER(FMatrix, LocalToWorldMatrix)
+		SHADER_PARAMETER_SRV(Buffer, RenderDeformedPositionBuffer)
+		SHADER_PARAMETER_SRV(Buffer, ClusterVertexIdBuffer)
+		SHADER_PARAMETER_SRV(Buffer, ClusterInfoBuffer)
+		SHADER_PARAMETER_UAV(RWBuffer, OutClusterAABBBuffer)
+		SHADER_PARAMETER_UAV(RWBuffer, OutGroupAABBBuffer)
+	END_SHADER_PARAMETER_STRUCT()
+
+public:
+	static bool ShouldCompilePermutation(const FGlobalShaderPermutationParameters& Parameters) { return IsHairStrandsSupported(Parameters.Platform); }
+};
+
+IMPLEMENT_GLOBAL_SHADER(FHairClusterAABBCS, "/Engine/Private/HairStrands/HairStrandsInterpolation.usf", "ClusterAABBEvaluationCS", SF_Compute);
+
+static void AddHairClusterAABBPass(
+	FRDGBuilder& GraphBuilder,
+	const FHairStrandsProjectionHairData::HairGroup& InRenHairData,
+	const FVector& OutHairWorldOffset,
+	FHairStrandClusterData::FHairGroup& ClusterData,
+	const FShaderResourceViewRHIRef& RenderPositionBuffer)
+{
+	const uint32 GroupSize = ComputeGroupSize();
+	const FIntVector DispatchCount = ComputeDispatchGroupCount2D(ClusterData.ClusterCount);
+
+	FHairClusterAABBCS::FParameters* Parameters = GraphBuilder.AllocParameters<FHairClusterAABBCS::FParameters>();
+	Parameters->DispatchCountX = DispatchCount.X;
+	Parameters->ClusterCount = ClusterData.ClusterCount;
+	Parameters->LocalToWorldMatrix = InRenHairData.LocalToWorld.ToMatrixWithScale();
+	Parameters->OutHairPositionOffset = OutHairWorldOffset;
+	Parameters->RenderDeformedPositionBuffer = RenderPositionBuffer;
+	Parameters->ClusterVertexIdBuffer = ClusterData.ClusterVertexIdBuffer->SRV;
+	Parameters->ClusterInfoBuffer = ClusterData.ClusterInfoBuffer->SRV;
+	Parameters->OutClusterAABBBuffer = ClusterData.HairGroupPublicPtr->GetClusterAABBBuffer().UAV;
+	Parameters->OutGroupAABBBuffer = ClusterData.HairGroupPublicPtr->GetGroupAABBBuffer().UAV;
+
+	FHairClusterAABBCS::FPermutationDomain PermutationVector;
+	PermutationVector.Set<FHairClusterAABBCS::FGroupSize>(GetGroupSizePermutation(GroupSize));
+	TShaderMapRef<FHairClusterAABBCS> ComputeShader(GetGlobalShaderMap(ERHIFeatureLevel::SM5), PermutationVector);
+
+	FComputeShaderUtils::AddPass(
+		GraphBuilder,
+		RDG_EVENT_NAME("HairStrandsClusterAABB"),
+		*ComputeShader,
+		Parameters,
+		DispatchCount);
+}
+
+///////////////////////////////////////////////////////////////////////////////////////////////////
+
 class FHairTangentCS : public FGlobalShader
 {
 	DECLARE_GLOBAL_SHADER(FHairTangentCS);
@@ -423,6 +599,56 @@ static void AddGenerateRaytracingGeometryPass(
 
 ///////////////////////////////////////////////////////////////////////////////////////////////////
 
+class FClearClusterAABBCS : public FGlobalShader
+{
+	DECLARE_GLOBAL_SHADER(FClearClusterAABBCS);
+	SHADER_USE_PARAMETER_STRUCT(FClearClusterAABBCS, FGlobalShader);
+
+	BEGIN_SHADER_PARAMETER_STRUCT(FParameters, )
+		SHADER_PARAMETER_UAV(RWBuffer, OutClusterAABBBuffer)
+		SHADER_PARAMETER_UAV(RWBuffer, OutGroupAABBBuffer)
+		SHADER_PARAMETER(uint32, ClusterCount)
+	END_SHADER_PARAMETER_STRUCT()
+
+public:
+	static bool ShouldCompilePermutation(const FGlobalShaderPermutationParameters& Parameters) { return IsHairStrandsSupported(Parameters.Platform); }
+	static void ModifyCompilationEnvironment(const FGlobalShaderPermutationParameters& Parameters, FShaderCompilerEnvironment& OutEnvironment)
+	{
+		FGlobalShader::ModifyCompilationEnvironment(Parameters, OutEnvironment);
+		OutEnvironment.SetDefine(TEXT("SHADER_CLEARCLUSTERAABB"), 1);
+	}
+};
+
+IMPLEMENT_GLOBAL_SHADER(FClearClusterAABBCS, "/Engine/Private/HairStrands/HairStrandsClusterCulling.usf", "MainClearClusterAABBCS", SF_Compute);
+
+static void AddClearClusterAABBPass(
+	FRDGBuilder& GraphBuilder,
+	uint32 ClusterCount,
+	FRHIUnorderedAccessView* OutClusterAABBuffer,
+	FRHIUnorderedAccessView* OutGroupAABBuffer)
+{
+	check(OutClusterAABBuffer);
+
+	FClearClusterAABBCS::FParameters* Parameters = GraphBuilder.AllocParameters<FClearClusterAABBCS::FParameters>();
+	Parameters->ClusterCount = ClusterCount;
+	Parameters->OutClusterAABBBuffer = OutClusterAABBuffer;
+	Parameters->OutGroupAABBBuffer = OutGroupAABBuffer;
+
+	
+	TShaderMap<FGlobalShaderType>* ShaderMap = GetGlobalShaderMap(ERHIFeatureLevel::SM5);
+	TShaderMapRef<FClearClusterAABBCS> ComputeShader(ShaderMap);
+
+	const FIntVector DispatchCount = DispatchCount.DivideAndRoundUp(FIntVector(ClusterCount * 6, 1, 1) , FIntVector(64, 1, 1));
+	FComputeShaderUtils::AddPass(
+		GraphBuilder,
+		RDG_EVENT_NAME("HairStrandsClearClusterAABB"),
+		*ComputeShader,
+		Parameters,
+		DispatchCount);
+}
+
+///////////////////////////////////////////////////////////////////////////////////////////////////
+
 #if RHI_RAYTRACING
 static void UpdateHairAccelerationStructure(FRHICommandList& RHICmdList, FRayTracingGeometry* RayTracingGeometry)
 {
@@ -461,18 +687,21 @@ static void BuildHairAccelerationStructure(FRHICommandList& RHICmdList, uint32 R
 
 void ComputeHairStrandsInterpolation(
 	FRHICommandListImmediate& RHICmdList,
+	const FShaderDrawDebugData* ShaderDrawData,
 	FHairStrandsInterpolationInput* InInput,
 	FHairStrandsInterpolationOutput* InOutput,
 	FHairStrandsProjectionHairData& InRenHairDatas,
 	FHairStrandsProjectionHairData& InSimHairDatas,
-	int32 LODIndex)
+	int32 LODIndex,
+	FHairStrandClusterData* ClusterData)
 {
 	if (!InInput || !InOutput) return;
 
+
 	const uint32 GroupCount = InOutput->HairGroups.Num();
 	for (uint32 GroupIndex = 0; GroupIndex < GroupCount; ++GroupIndex)
 	{
-		FHairStrandsInterpolationInput::FHairGroup& Input   = InInput->HairGroups[GroupIndex];
+		FHairStrandsInterpolationInput::FHairGroup& Input  = InInput->HairGroups[GroupIndex];
 		FHairStrandsInterpolationOutput::HairGroup& Output = InOutput->HairGroups[GroupIndex];
 		Output.VFInput.Reset();
 
@@ -496,7 +725,8 @@ void ComputeHairStrandsInterpolation(
 				Input.SimVertexCount,
 				Input.SimRestPosePositionBuffer->SRV,
 				Input.SimRootPointIndexBuffer ? Input.SimRootPointIndexBuffer->SRV : nullptr,
-				Output.SimDeformedPositionBuffer[CurrIndex]->UAV);
+				Output.SimDeformedPositionBuffer[CurrIndex]->UAV, Input.InSimHairPositionOffset,
+				Input.OutHairPositionOffset );
 		}
 
 		// If the deformation is driven by the physics simulation, then the output is always the 0 index
@@ -527,40 +757,90 @@ void ComputeHairStrandsInterpolation(
 			Output.VFInput.HairPositionOffset = bHasSimulationEnabled ? Input.OutHairPositionOffset : Input.InSimHairPositionOffset;
 			Output.VFInput.HairPreviousPositionOffset = bHasSimulationEnabled ? Input.OutHairPreviousPositionOffset : Input.InSimHairPositionOffset;
 			Output.VFInput.VertexCount = Input.SimVertexCount;
+			Output.VFInput.HairRadius = (GStrandHairWidth > 0 ? GStrandHairWidth : Input.GroupDesc.HairWidth) * 0.5f;
+			Output.VFInput.HairLength = Input.GroupDesc.HairLength;
+			Output.VFInput.HairDensity = Input.GroupDesc.HairShadowDensity;
 		}
 		else
 		{
+			check(ClusterData);
+			const bool bDebugModePatchedAttributeBuffer = DebugMode == EHairStrandsDebugMode::RenderHairStrands || DebugMode == EHairStrandsDebugMode::RenderVisCluster;
+
 			const uint32 BufferSizeInBytes = Input.RenderVertexCount * FHairStrandsAttributeFormat::SizeInByte;
-			if (DebugMode == EHairStrandsDebugMode::RenderHairStrands && Output.RenderPatchedAttributeBuffer.NumBytes != BufferSizeInBytes)
+			if (bDebugModePatchedAttributeBuffer && Output.RenderPatchedAttributeBuffer.NumBytes != BufferSizeInBytes)
 			{
 				Output.RenderPatchedAttributeBuffer.Release();
 				Output.RenderPatchedAttributeBuffer.Initialize(FHairStrandsAttributeFormat::SizeInByte, Input.RenderVertexCount, FHairStrandsAttributeFormat::Format, BUF_Static);
 			}
 
+			AddClearClusterAABBPass(
+				GraphBuilder,
+				Input.ClusterCount,
+				Output.RenderClusterAABBBuffer->UAV,
+				Output.RenderGroupAABBBuffer->UAV);
+
+			FHairScaleAndClipDesc ScaleAndClipDesc;
+			ScaleAndClipDesc.InHairLength = Input.GroupDesc.HairLength;
+			ScaleAndClipDesc.InHairRadius = Input.GroupDesc.HairWidth * 0.5f;
+			ScaleAndClipDesc.OutHairRadius = (GStrandHairWidth > 0 ? GStrandHairWidth : Input.GroupDesc.HairWidth) * 0.5f;
+			ScaleAndClipDesc.MaxOutHairRadius = ScaleAndClipDesc.OutHairRadius * FMath::Max(1.f, FMath::Max(Input.GroupDesc.HairRootScale, Input.GroupDesc.HairTipScale));
+			ScaleAndClipDesc.HairRadiusRootScale = Input.GroupDesc.HairRootScale;
+			ScaleAndClipDesc.HairRadiusTipScale = Input.GroupDesc.HairTipScale;
+			ScaleAndClipDesc.HairLengthClip = Input.GroupDesc.HairClipLength;
+
 			const bool bHasSimulationEnabled = Input.bIsSimulationEnable && GHairStrandsInterpolateSimulation && DeformationType != EDeformationType::RestStrands;
 			check(GroupIndex < uint32(InRenHairDatas.HairGroups.Num()));
 			check(GroupIndex < uint32(InSimHairDatas.HairGroups.Num()));
 			AddHairStrandsInterpolationPass(
 				GraphBuilder,
+				ShaderDrawData,
 				InRenHairDatas.HairGroups[GroupIndex],
 				InSimHairDatas.HairGroups[GroupIndex],
 				Input.InRenderHairPositionOffset,
 				Input.InSimHairPositionOffset,
 				Input.OutHairPositionOffset,
+				ScaleAndClipDesc,
 				LODIndex,
 				bHasSimulationEnabled,
+				bDebugModePatchedAttributeBuffer,
 				Input.RenderVertexCount,
 				Input.RenderRestPosePositionBuffer->SRV,
+				Input.RenderAttributeBuffer->SRV,
 				Input.Interpolation0Buffer->SRV,
 				Input.Interpolation1Buffer->SRV,
 				Input.SimRestPosePositionBuffer->SRV,
 				Output.SimDeformedPositionBuffer[SimIndex]->SRV,
-				DebugMode == EHairStrandsDebugMode::RenderHairStrands ? Input.SimAttributeBuffer->SRV : nullptr,
+				Input.SimAttributeBuffer->SRV,
 				Output.RenderDeformedPositionBuffer[CurrIndex]->UAV,
-				DebugMode == EHairStrandsDebugMode::RenderHairStrands ? Output.RenderPatchedAttributeBuffer.UAV : nullptr);
-
+				Output.RenderPatchedAttributeBuffer.UAV,
+				Input.VertexToClusterIdBuffer->SRV);
+
+			// Initialize group cluster data for culling by the renderer
+			FHairStrandClusterData::FHairGroup& HairGroupCluster = ClusterData->HairGroups.Emplace_GetRef();
+			HairGroupCluster.ClusterCount = Input.ClusterCount;
+			HairGroupCluster.VertexCount = Input.ClusterVertexCount;
+			HairGroupCluster.GroupAABBBuffer = Output.RenderGroupAABBBuffer;
+			HairGroupCluster.ClusterAABBBuffer = Output.RenderClusterAABBBuffer;
+			HairGroupCluster.ClusterInfoBuffer = Output.ClusterInfoBuffer;
+			HairGroupCluster.VertexToClusterIdBuffer = Input.VertexToClusterIdBuffer;
+			HairGroupCluster.ClusterVertexIdBuffer = Input.ClusterVertexIdBuffer;
+			HairGroupCluster.ClusterIndexRadiusScaleInfoBuffer = Input.ClusterIndexRadiusScaleInfoBuffer;
+			HairGroupCluster.HairGroupPublicPtr = Output.HairGroupPublicData;
+			HairGroupCluster.LodBias = Input.GroupDesc.LodBias;
+			HairGroupCluster.LodAverageVertexPerPixel = Input.GroupDesc.LodAverageVertexPerPixel;
+
+			Output.VFInput.HairRadius = ScaleAndClipDesc.MaxOutHairRadius;
+			Output.VFInput.HairLength = Input.GroupDesc.HairLength;
+			Output.VFInput.HairDensity = Input.GroupDesc.HairShadowDensity;
 			Output.VFInput.HairPositionBuffer = Output.RenderDeformedPositionBuffer[CurrIndex]->SRV;
 			Output.VFInput.HairPreviousPositionBuffer = Output.RenderDeformedPositionBuffer[PrevIndex]->SRV;
+
+			AddHairClusterAABBPass(
+				GraphBuilder,
+				InRenHairDatas.HairGroups[GroupIndex],
+				Input.OutHairPositionOffset,
+				HairGroupCluster,
+				Output.RenderDeformedPositionBuffer[CurrIndex]->SRV);
 			
 			AddHairTangentPass(
 				GraphBuilder,
@@ -569,12 +849,13 @@ void ComputeHairStrandsInterpolation(
 				Output.RenderTangentBuffer->UAV);
 
 			#if RHI_RAYTRACING
-			if (IsRayTracingEnabled())
+			if (IsHairRayTracingEnabled())
 			{
+				const float HairRadiusScaleRT = (GHairRaytracingRadiusScale > 0 ? GHairRaytracingRadiusScale : Input.GroupDesc.HairRaytracingRadiusScale);
 				AddGenerateRaytracingGeometryPass(
 					GraphBuilder,
 					Input.RenderVertexCount,
-					Input.HairRadius * (GHairRaytracingRadiusScale > 0 ? GHairRaytracingRadiusScale : Input.HairRaytracingRadiusScale),
+					Output.VFInput.HairRadius * HairRadiusScaleRT,
 					Input.OutHairPositionOffset,
 					Output.VFInput.HairPositionBuffer,
 					Input.RaytracingPositionBuffer->UAV);
@@ -582,15 +863,15 @@ void ComputeHairStrandsInterpolation(
 			#endif
 			GraphBuilder.Execute();
 
-			Output.VFInput.HairTangentBuffer = Output.RenderTangentBuffer->SRV;
-			Output.VFInput.HairAttributeBuffer = DebugMode == EHairStrandsDebugMode::RenderHairStrands ? Output.RenderPatchedAttributeBuffer.SRV : Input.RenderAttributeBuffer->SRV;
-			Output.VFInput.HairMaterialBuffer = Output.RenderMaterialBuffer->SRV;
-			Output.VFInput.HairPositionOffset = Input.OutHairPositionOffset;
-			Output.VFInput.HairPreviousPositionOffset = Input.OutHairPreviousPositionOffset;
-			Output.VFInput.VertexCount = Input.RenderVertexCount;
+			Output.VFInput.HairTangentBuffer			= Output.RenderTangentBuffer->SRV;
+			Output.VFInput.HairAttributeBuffer			= bDebugModePatchedAttributeBuffer ? Output.RenderPatchedAttributeBuffer.SRV : Input.RenderAttributeBuffer->SRV;
+			Output.VFInput.HairMaterialBuffer			= Output.RenderMaterialBuffer->SRV;
+			Output.VFInput.HairPositionOffset			= Input.OutHairPositionOffset;
+			Output.VFInput.HairPreviousPositionOffset	= Input.OutHairPreviousPositionOffset;
+			Output.VFInput.VertexCount					= Input.RenderVertexCount;
 
 			#if RHI_RAYTRACING
-			if (IsRayTracingEnabled())
+			if (IsHairRayTracingEnabled())
 			{
 				FRHIUnorderedAccessView* UAV = Input.RaytracingPositionBuffer->UAV;
 				RHICmdList.TransitionResources(EResourceTransitionAccess::EReadable, EResourceTransitionPipeline::EComputeToGfx, &UAV, 1);
@@ -608,3 +889,36 @@ void ComputeHairStrandsInterpolation(
 		Output.CurrentIndex = PrevIndex;
 	}
 }
+
+void ResetHairStrandsInterpolation(
+	FRHICommandListImmediate& RHICmdList,
+	FHairStrandsInterpolationInput* InInput,
+	FHairStrandsInterpolationOutput* InOutput)
+{
+	if (!InInput || !InOutput) return;
+
+	const uint32 GroupCount = InOutput->HairGroups.Num();
+	for (uint32 GroupIndex = 0; GroupIndex < GroupCount; ++GroupIndex)
+	{
+		FHairStrandsInterpolationInput::FHairGroup& Input = InInput->HairGroups[GroupIndex];
+		FHairStrandsInterpolationOutput::HairGroup& Output = InOutput->HairGroups[GroupIndex];
+		if (!Input.IsValid() || !Output.IsValid()) return;
+
+		DECLARE_GPU_STAT(HairStrandsResetInterpolation);
+		SCOPED_DRAW_EVENT(RHICmdList, HairStrandsResetInterpolation);
+		SCOPED_GPU_STAT(RHICmdList, HairStrandsResetInterpolation);
+
+		const uint32 CurrIndex = Output.CurrentIndex;
+		const uint32 PrevIndex = (Output.CurrentIndex + 1) % 2;
+		FRDGBuilder GraphBuilder(RHICmdList);
+		AddDeformSimHairStrandsPass(
+			GraphBuilder,
+			EDeformationType::OffsetGuide,
+			Input.SimVertexCount,
+			Input.SimRestPosePositionBuffer->SRV,
+			Input.SimRootPointIndexBuffer ? Input.SimRootPointIndexBuffer->SRV : nullptr,
+			Output.SimDeformedPositionBuffer[CurrIndex]->UAV, Input.InSimHairPositionOffset,
+			Input.OutHairPositionOffset);
+		GraphBuilder.Execute();
+	}
+}
\ No newline at end of file
diff --git a/Engine/Plugins/Experimental/HairStrands/Source/HairStrandsCore/Private/HairStrandsVertexFactory.cpp b/Engine/Plugins/Experimental/HairStrands/Source/HairStrandsCore/Private/HairStrandsVertexFactory.cpp
--- a/Engine/Plugins/Experimental/HairStrands/Source/HairStrandsCore/Private/HairStrandsVertexFactory.cpp
+++ b/Engine/Plugins/Experimental/HairStrands/Source/HairStrandsCore/Private/HairStrandsVertexFactory.cpp
@@ -12,25 +12,34 @@
 #include "MeshMaterialShader.h"
 #include "HairStrandsInterface.h"
 
-static float GStrandHairWidth = 0.0f;
-static FAutoConsoleVariableRef CVarStrandHairWidth(TEXT("r.HairStrands.StrandWidth"), GStrandHairWidth, TEXT("Width of hair strand"));
-
-float FHairStrandsVertexFactory::GetMaxStrandRadius(uint32 GroupIndex) const
-{
-	return GStrandHairWidth > 0 ? GStrandHairWidth * 0.5f : Data.HairGroups[GroupIndex].MaxStrandRadius;
-}
-
-#define OPTIMIZE_OFF 0
-
-#if OPTIMIZE_OFF
-	#pragma optimize("", off)
-#endif
-
 /////////////////////////////////////////////////////////////////////////////////////////
 
 template<typename T> inline void BindParam(FMeshDrawSingleShaderBindings& ShaderBindings, const FShaderResourceParameter& Param, T* Value)	{ if (Param.IsBound() && Value) ShaderBindings.Add(Param, Value); }
 template<typename T> inline void BindParam(FMeshDrawSingleShaderBindings& ShaderBindings, const FShaderParameter& Param, const T& Value)	{ if (Param.IsBound()) ShaderBindings.Add(Param, Value); }
 
+class FDummyCulledDispatchVertexIdsBuffer : public FVertexBuffer
+{
+public:
+	FShaderResourceViewRHIRef SRVUint;
+	FShaderResourceViewRHIRef SRVFloat;
+
+	virtual void InitRHI() override
+	{
+		FRHIResourceCreateInfo CreateInfo;
+		void* BufferData = nullptr;
+		uint32 NumBytes = sizeof(uint32) * 4;
+		VertexBufferRHI = RHICreateAndLockVertexBuffer(NumBytes, BUF_Static | BUF_ShaderResource, CreateInfo, BufferData);
+		uint32* DummyContents = (uint32*)BufferData;
+		DummyContents[0] = DummyContents[1] = DummyContents[2] = DummyContents[3] = 0;
+		RHIUnlockVertexBuffer(VertexBufferRHI);
+
+		SRVUint = RHICreateShaderResourceView(VertexBufferRHI, sizeof(uint32), PF_R32_UINT);
+		SRVFloat = RHICreateShaderResourceView(VertexBufferRHI, sizeof(uint32), PF_R32_FLOAT);
+	}
+
+};
+TGlobalResource<FDummyCulledDispatchVertexIdsBuffer> GDummyCulledDispatchVertexIdsBuffer;
+
 class FHairStrandsVertexFactoryShaderParameters : public FVertexFactoryShaderParameters
 {
 public:
@@ -42,6 +51,7 @@ public:
 	FShaderParameter PositionOffset;
 	FShaderParameter PreviousPositionOffset;
 	FShaderParameter Density;
+	FShaderParameter Culling;
 
 	FShaderResourceParameter PositionBuffer;
 	FShaderResourceParameter PreviousPositionBuffer;
@@ -49,6 +59,9 @@ public:
 	FShaderResourceParameter MaterialBuffer;
 	FShaderResourceParameter TangentBuffer;
 
+	FShaderResourceParameter CulledVertexIdsBuffer;
+	FShaderResourceParameter CulledVertexRadiusScaleBuffer;
+
 	virtual void Bind(const FShaderParameterMap& ParameterMap) override
 	{
 		Radius.Bind(ParameterMap, TEXT("HairStrandsVF_Radius"));
@@ -56,12 +69,16 @@ public:
 		PositionOffset.Bind(ParameterMap, TEXT("HairStrandsVF_PositionOffset"));
 		PreviousPositionOffset.Bind(ParameterMap, TEXT("HairStrandsVF_PreviousPositionOffset"));
 		Density.Bind(ParameterMap, TEXT("HairStrandsVF_Density"));	
+		Culling.Bind(ParameterMap, TEXT("HairStrandsVF_CullingEnable"));
 
 		PositionBuffer.Bind(ParameterMap, TEXT("HairStrandsVF_PositionBuffer"));
 		PreviousPositionBuffer.Bind(ParameterMap, TEXT("HairStrandsVF_PreviousPositionBuffer"));
 		AttributeBuffer.Bind(ParameterMap, TEXT("HairStrandsVF_AttributeBuffer"));
 		MaterialBuffer.Bind(ParameterMap, TEXT("HairStrandsVF_MaterialBuffer"));
 		TangentBuffer.Bind(ParameterMap, TEXT("HairStrandsVF_TangentBuffer"));
+
+		CulledVertexIdsBuffer.Bind(ParameterMap, TEXT("CulledVertexIdsBuffer"));
+		CulledVertexRadiusScaleBuffer.Bind(ParameterMap, TEXT("CulledVertexRadiusScaleBuffer"));
 	}
 
 	virtual void Serialize(FArchive& Ar) override
@@ -73,12 +90,16 @@ public:
 		Ar << PositionOffset;
 		Ar << PreviousPositionOffset;
 		Ar << Density;
+		Ar << Culling;
 
 		Ar << PositionBuffer;
 		Ar << PreviousPositionBuffer;
 		Ar << AttributeBuffer;
 		Ar << MaterialBuffer;
 		Ar << TangentBuffer;
+
+		Ar << CulledVertexIdsBuffer;
+		Ar << CulledVertexRadiusScaleBuffer;
 	}
 
 	virtual void GetElementShaderBindings(
@@ -95,7 +116,9 @@ public:
 	{
 		const FHairStrandsVertexFactory* VF = static_cast<const FHairStrandsVertexFactory*>(VertexFactory);
 
-		const uint64 GroupIndex = reinterpret_cast<uint64>(BatchElement.UserData);
+		const FHairGroupPublicData* GroupPublicData = reinterpret_cast<const FHairGroupPublicData*>(BatchElement.VertexFactoryUserData);
+		check(GroupPublicData);
+		const uint64 GroupIndex = GroupPublicData->GetGroupIndex();
 		BindParam(ShaderBindings, PositionBuffer, VF->GetPositionSRV(GroupIndex));
 		BindParam(ShaderBindings, PreviousPositionBuffer, VF->GetPreviousPositionSRV(GroupIndex));
 		BindParam(ShaderBindings, AttributeBuffer, VF->GetAttributeSRV(GroupIndex));
@@ -106,6 +129,22 @@ public:
 		BindParam(ShaderBindings, PositionOffset, VF->GetPositionOffset(GroupIndex));
 		BindParam(ShaderBindings, PreviousPositionOffset, VF->GetPreviousPositionOffset(GroupIndex));
 		BindParam(ShaderBindings, Density, VF->GetHairDensity(GroupIndex));
+		
+		FShaderResourceViewRHIRef CulledDispatchVertexIdsSRV = GDummyCulledDispatchVertexIdsBuffer.SRVUint;
+		FShaderResourceViewRHIRef CulledCompactedRadiusScaleBufferSRV = GDummyCulledDispatchVertexIdsBuffer.SRVFloat;
+
+		FHairGroupPublicData* PublicGroupData = VF->GetHairGroupPublicData(GroupIndex);
+		
+		const bool bCulling = PublicGroupData ? PublicGroupData->GetCullingResultAvailable() : false;
+		if (bCulling && PublicGroupData)
+		{
+			CulledDispatchVertexIdsSRV = PublicGroupData->GetCulledVertexIdBuffer().SRV;
+			CulledCompactedRadiusScaleBufferSRV = PublicGroupData->GetCulledVertexRadiusScaleBuffer().SRV;
+		}
+
+		BindParam(ShaderBindings, Culling, bCulling ? 1 : 0);
+		ShaderBindings.Add(CulledVertexIdsBuffer, CulledDispatchVertexIdsSRV);
+		ShaderBindings.Add(CulledVertexRadiusScaleBuffer, CulledCompactedRadiusScaleBufferSRV);
 	}
 };
 
@@ -213,9 +252,4 @@ void FHairStrandsVertexFactory::ReleaseRHI()
 	FVertexFactory::ReleaseRHI();
 }
 
-
 IMPLEMENT_VERTEX_FACTORY_TYPE_EX(FHairStrandsVertexFactory,"/Engine/Private/HairStrands/HairStrandsVertexFactory.ush",true,false,true,true,true,true,true);
-
-#if OPTIMIZE_OFF
-	#pragma optimize("", on)
-#endif
\ No newline at end of file
diff --git a/Engine/Plugins/Experimental/HairStrands/Source/HairStrandsCore/Public/GroomAsset.h b/Engine/Plugins/Experimental/HairStrands/Source/HairStrandsCore/Public/GroomAsset.h
--- a/Engine/Plugins/Experimental/HairStrands/Source/HairStrandsCore/Public/GroomAsset.h
+++ b/Engine/Plugins/Experimental/HairStrands/Source/HairStrandsCore/Public/GroomAsset.h
@@ -8,11 +8,13 @@
 #include "HairDescription.h"
 #include "HairStrandsDatas.h"
 #include "RenderResource.h"
-
+#include "GroomSettings.h"
+#include "HairStrandsInterface.h"
 #include "GroomAsset.generated.h"
 
 class UMaterialInterface;
 class UNiagaraSystem;
+struct FHairGroupData;
 
 /* Render buffers for root deformation for dynamic meshes */
 struct FHairStrandsRootResource : public FRenderResource
@@ -34,8 +36,7 @@ struct FHairStrandsRootResource : public FRenderResource
 
 	struct FMeshProjectionLOD
 	{
-		enum class EStatus { Invalid, Initialized, Completed };
-		EStatus Status = EStatus::Invalid;
+		FHairStrandsProjectionHairData::LODData::EStatus Status = FHairStrandsProjectionHairData::LODData::EStatus::Invalid;
 		int32 LODIndex = -1;
 
 		/* Triangle on which a root is attached */
@@ -133,6 +134,50 @@ struct FHairStrandsDeformedResource : public FRenderResource
 	const bool bInitializedData;
 };
 
+struct FHairStrandsClusterCullingResource : public FRenderResource
+{
+	FHairStrandsClusterCullingResource(const FHairGroupData& GroupData);
+
+	/* Init the buffer */
+	virtual void InitRHI() override;
+
+	/* Release the buffer */
+	virtual void ReleaseRHI() override;
+
+	/* Get the resource name */
+	virtual FString GetFriendlyName() const override { return TEXT("FHairStrandsClusterResource"); }
+
+
+	struct FClusterInfo
+	{
+		uint32 FirstVertexId;
+		uint32 VertexIdCountLodHigh;
+		uint32 VertexIdCountLodLow;
+		uint32 UnusedUint;
+	};
+	TArray<FClusterInfo> ClusterInfoArray;
+	TArray<uint32> VertexToClusterIdArray;
+	TArray<uint32> ClusterVertexIdArray;
+	TArray<float> ClusterIndexRadiusScaleInfoArray;
+
+	/* Cluster info buffer */
+	FReadBuffer ClusterInfoBuffer;		// Contains: Start VertexId, VertexId Count (for VertexToClusterIdBuffer)
+
+	/* VertexId => ClusterId to know which AABB to contribute to*/
+	FReadBuffer VertexToClusterIdBuffer;
+
+	/* Concatenated data for each cluster: list of VertexId pointed to by ClusterInfoBuffer */
+	FReadBuffer ClusterVertexIdBuffer;
+
+	/* Contains information to recove the radius scale to apply per cluster when decimating vertex count */
+	FReadBuffer ClusterIndexRadiusScaleInfoBuffer;
+
+	/* Number of cluster  */
+	uint32 ClusterCount;
+	/* Number of vertex  */
+	uint32 VertexCount;
+};
+
 struct FHairStrandsInterpolationResource : public FRenderResource
 {
 	/** Build the hair strands resource */
@@ -193,8 +238,8 @@ struct HAIRSTRANDSCORE_API FHairGroupInfo
 	UPROPERTY(VisibleAnywhere, Category = "Info", meta = (DisplayName = "Guide Count"))
 	int32 NumGuides = 0;
 
-	UPROPERTY(VisibleAnywhere, Category = "Info", meta = (DisplayName = "Auto-generated Guides", ToolTip = "Checked when guides were auto-generated"))
-	bool bIsAutoGenerated = false;
+	UPROPERTY(VisibleAnywhere, Category = "Info", meta = (DisplayName = "Import options", ToolTip = "Show the options used at import time"))
+	FGroomBuildSettings ImportSettings;
 
 	UPROPERTY(EditAnywhere, Category="Rendering")
 	UMaterialInterface* Material = nullptr;
@@ -202,6 +247,8 @@ struct HAIRSTRANDSCORE_API FHairGroupInfo
 	// Currently hide the Nigara simulation slot as it is not used, and could confuse users
 	//UPROPERTY(EditAnywhere, Category="Simulation", meta = (DisplayName = "Niagara System Asset"))
 	//UNiagaraSystem* NiagaraAsset = nullptr;
+
+	friend FArchive& operator<<(FArchive& Ar, FHairGroupInfo& GroupInfo);
 };
 
 struct HAIRSTRANDSCORE_API FHairGroupData
@@ -288,9 +335,29 @@ public:
 
 	int32 GetNumHairGroups() const;
 
+	/** Returns true if the asset has the HairDescription needed to recompute its groom data */
+	bool CanRebuildFromDescription() const;
+
 //private :
 #if WITH_EDITOR
 	FOnGroomAssetChanged OnGroomAssetChanged;
 #endif
+
+#if WITH_EDITORONLY_DATA
+public:
+	/** Commits a HairDescription to buffer for serialization */
+	void CommitHairDescription(FHairDescription&& HairDescription);
+
+	/** Caches the computed groom data with the given build settings from/to the Derived Data Cache, building it if needed */
+	bool CacheDerivedData(const struct FGroomBuildSettings* BuildSettings = nullptr);
+
+private:
+	FString BuildDerivedDataKeySuffix(const struct FGroomBuildSettings& BuildSettings);
+
 	TUniquePtr<FHairDescription> HairDescription;
+	TUniquePtr<FHairDescriptionBulkData> HairDescriptionBulkData;
+
+	UPROPERTY()
+	bool bIsCacheable = true;
+#endif
 };
diff --git a/Engine/Plugins/Experimental/HairStrands/Source/HairStrandsCore/Public/GroomAssetImportData.h b/Engine/Plugins/Experimental/HairStrands/Source/HairStrandsCore/Public/GroomAssetImportData.h
new file mode 100644
--- /dev/null
+++ b/Engine/Plugins/Experimental/HairStrands/Source/HairStrandsCore/Public/GroomAssetImportData.h
@@ -0,0 +1,19 @@
+// Copyright Epic Games, Inc. All Rights Reserved.
+
+#pragma once
+
+#include "CoreMinimal.h"
+
+#include "EditorFramework/AssetImportData.h"
+#include "GroomAssetImportData.generated.h"
+
+UCLASS()
+class HAIRSTRANDSCORE_API UGroomAssetImportData : public UAssetImportData
+{
+	GENERATED_UCLASS_BODY()
+
+public:
+
+	UPROPERTY()
+	class UGroomImportOptions* ImportOptions;
+};
diff --git a/Engine/Plugins/Experimental/HairStrands/Source/HairStrandsCore/Public/GroomBuilder.h b/Engine/Plugins/Experimental/HairStrands/Source/HairStrandsCore/Public/GroomBuilder.h
--- a/Engine/Plugins/Experimental/HairStrands/Source/HairStrandsCore/Public/GroomBuilder.h
+++ b/Engine/Plugins/Experimental/HairStrands/Source/HairStrandsCore/Public/GroomBuilder.h
@@ -3,6 +3,9 @@
 #pragma once
 
 #include "CoreMinimal.h"
+#include "UObject/Object.h"
+#include "UObject/ObjectMacros.h"
+#include "GroomSettings.h"
 
 struct FHairStrandsDatas;
 class UGroomAsset;
@@ -11,5 +14,5 @@ struct HAIRSTRANDSCORE_API FGroomBuilder
 {
 	static bool BuildGroom(const class FHairDescription& HairDescription, const struct FGroomBuildSettings& BuildSettings, UGroomAsset* GroomAsset);
 	static void GenerateGuides(const FHairStrandsDatas& InData, float DecimationPercentage, FHairStrandsDatas& OutData);
-	static void BuildData(UGroomAsset* GroomAsset, uint8 QualityLevel, uint8 WeightMethod, bool bRandomize, bool bUnique);
+	static void BuildData(UGroomAsset* GroomAsset, const FGroomBuildSettings& BuildSettings);
 };
\ No newline at end of file
diff --git a/Engine/Plugins/Experimental/HairStrands/Source/HairStrandsCore/Public/GroomComponent.h b/Engine/Plugins/Experimental/HairStrands/Source/HairStrandsCore/Public/GroomComponent.h
--- a/Engine/Plugins/Experimental/HairStrands/Source/HairStrandsCore/Public/GroomComponent.h
+++ b/Engine/Plugins/Experimental/HairStrands/Source/HairStrandsCore/Public/GroomComponent.h
@@ -8,35 +8,10 @@
 #include "Components/MeshComponent.h"
 #include "GroomAsset.h"
 #include "RHIDefinitions.h"
+#include "GroomDesc.h"
 
 #include "GroomComponent.generated.h"
 
-USTRUCT(BlueprintType)
-struct FHairGroupDesc
-{
-	GENERATED_USTRUCT_BODY()
-
-	/** Number of hairs within this hair group.  */
-	UPROPERTY(VisibleAnywhere, BlueprintReadOnly, Category = "Groom")
-	int32 HairCount;
-
-	/** Number of simulation guides within this hair group. */
-	UPROPERTY(VisibleAnywhere, BlueprintReadOnly, Category = "Groom")
-	int32 GuideCount;
-
-	/** Override the hair width (in centimeters) */
-	UPROPERTY(EditAnywhere, BlueprintReadWrite, Category = "Groom", meta = (ClampMin = "0.0001", UIMin = "0.001", UIMax = "1.0", SliderExponent = 6))
-	float HairWidth;
-
-	/** Override the hair shadow density factor (unit less).  */
-	UPROPERTY(EditAnywhere, BlueprintReadWrite, Category = "Groom", meta = (ClampMin = "0.0001", UIMin = "0.001", UIMax = "10.0", SliderExponent = 6))
-	float HairShadowDensity;
-
-	/** Scale the hair geometry radius for ray tracing effects (e.g. shadow) */
-	UPROPERTY(EditAnywhere, BlueprintReadWrite, Category = "Groom", meta = (ClampMin = "0.0001", UIMin = "0.001", UIMax = "10.0", SliderExponent = 6))
-	float HairRaytracingRadiusScale;
-};
-
 UCLASS(HideCategories = (Object, Physics, Activation, Mobility, "Components|Activation"), editinlinenew, meta = (BlueprintSpawnableComponent), ClassGroup = Rendering)
 class HAIRSTRANDSCORE_API UGroomComponent : public UMeshComponent
 {
@@ -60,6 +35,9 @@ public:
 	/** Boolean to check when animation has been loaded */
 	bool bResetSimulation;
 
+	/** Previous bone matrix to compare the difference and decide to reset or not the simulation */
+	FMatrix	PrevBoneMatrix;
+
 	/** Listen for the animation event to trigger the sim */
 	UFUNCTION()
 	void ResetSimulation();
@@ -124,6 +102,9 @@ public:
 		FHairStrandsDeformedResource* RenderDeformedResources = nullptr;
 		FHairStrandsDeformedResource* SimDeformedResources = nullptr;
 
+		FHairStrandsClusterCullingResource* ClusterCullingResources = nullptr; // TODO merge into FHairGroupPublicData
+		FHairGroupPublicData* HairGroupPublicDatas = nullptr;
+
 		// Rest resources, owned by the asset
 		FHairStrandsRestResource* RenderRestResources = nullptr;
 		FHairStrandsRestResource* SimRestResources = nullptr;
@@ -146,30 +127,9 @@ protected:
 private:
 	void* InitializedResources;
 
-	enum class EMeshProjectionState
-	{
-		Invalid,
-		InProgressBinding,
-		WaitForRestPose,
-		Completed
-	};
 	class USkeletalMeshComponent* RegisteredSkeletalMeshComponent;
 	FVector SkeletalPreviousPositionOffset;
-	int32 MeshProjectionLODIndex;
-	uint32 MeshProjectionTickDelay;
-	EMeshProjectionState MeshProjectionState;
-	bool bIsGroomAssetCallbackRegistered;
-	
-	struct FSkeletalMeshConfiguration
-	{
-		int32 ForceLOD = -1;
-		bool ForceRefPose = false;
-		static bool Equals(const FSkeletalMeshConfiguration& A, const FSkeletalMeshConfiguration& B)
-		{
-			return A.ForceLOD == B.ForceLOD && A.ForceRefPose == B.ForceRefPose;
-		}
-	};
-	FSkeletalMeshConfiguration SkeletalMeshConfiguration;
+	bool bIsGroomAssetCallbackRegistered;	
 
 	void InitResources();
 	void ReleaseResources();
diff --git a/Engine/Plugins/Experimental/HairStrands/Source/HairStrandsCore/Public/GroomDesc.h b/Engine/Plugins/Experimental/HairStrands/Source/HairStrandsCore/Public/GroomDesc.h
new file mode 100644
--- /dev/null
+++ b/Engine/Plugins/Experimental/HairStrands/Source/HairStrandsCore/Public/GroomDesc.h
@@ -0,0 +1,74 @@
+// Copyright Epic Games, Inc. All Rights Reserved.
+
+#pragma once
+
+#include "CoreMinimal.h"
+#include "UObject/ObjectMacros.h"
+#include "Engine/EngineTypes.h"
+#include "GroomDesc.generated.h"
+
+USTRUCT(BlueprintType)
+struct FHairGroupDesc
+{
+	GENERATED_USTRUCT_BODY()
+
+	/** Number of hairs within this hair group. */
+	UPROPERTY(VisibleAnywhere, BlueprintReadOnly, Category = "Groom")
+	int32 HairCount;
+
+	/** Number of simulation guides within this hair group. */
+	UPROPERTY(VisibleAnywhere, BlueprintReadOnly, Category = "Groom")
+	int32 GuideCount;
+
+	/** Length of the longest hair strands */
+	UPROPERTY(VisibleAnywhere, BlueprintReadOnly, Category = "Groom")
+	float HairLength;
+
+	/** Hair width (in centimeters) */
+	UPROPERTY(EditAnywhere, BlueprintReadWrite, Category = "Groom", meta = (ClampMin = "0.0001", UIMin = "0.001", UIMax = "1.0", SliderExponent = 6))
+	float HairWidth;
+
+	/** Scale the hair width at the root */
+	UPROPERTY(EditAnywhere, BlueprintReadWrite, Category = "Groom", AdvancedDisplay, meta = (ClampMin = "0.0001", UIMin = "0.001", UIMax = "2.0", SliderExponent = 6))
+	float HairRootScale;
+
+	/** Scale the hair with at the tip */
+	UPROPERTY(EditAnywhere, BlueprintReadWrite, Category = "Groom", AdvancedDisplay, meta = (ClampMin = "0.0001", UIMin = "0.001", UIMax = "2.0", SliderExponent = 6))
+	float HairTipScale;
+
+	/** Normalized hair clip length, i.e. at which length hair will be clipped. 1 means no clipping. 0 means hairs are fully clipped */
+	UPROPERTY(EditAnywhere, BlueprintReadWrite, Category = "Groom", AdvancedDisplay, meta = (ClampMin = "0.0", UIMin = "0.0", UIMax = "1.0", SliderExponent = 1))
+	float HairClipLength;
+
+	/** Override the hair shadow density factor (unit less). */
+	UPROPERTY(EditAnywhere, BlueprintReadWrite, Category = "Groom", AdvancedDisplay, meta = (ClampMin = "0.0001", UIMin = "0.001", UIMax = "10.0", SliderExponent = 6))
+	float HairShadowDensity;
+
+	/** Scale the hair geometry radius for ray tracing effects (e.g. shadow) */
+	UPROPERTY(EditAnywhere, BlueprintReadWrite, Category = "Groom", AdvancedDisplay, meta = (ClampMin = "0.0001", UIMin = "0.001", UIMax = "10.0", SliderExponent = 6))
+	float HairRaytracingRadiusScale;
+
+	/** The targeted budget of hair vertex per pixel. Cluster strands will be decimated based on that. Used when r.HairStrands.Cluster.Culling = 1. */
+	UPROPERTY(EditAnywhere, BlueprintReadWrite, Category = "Performance", AdvancedDisplay, meta = (ClampMin = "0.0", UIMin = "0.0", UIMax = "6.0", SliderExponent = 1))
+	float LodAverageVertexPerPixel;
+	/** Bias the selected LOD. A value >0 will progressively select lower detailed lods. Used when r.HairStrands.Cluster.Culling = 1. */
+	UPROPERTY(EditAnywhere, AdvancedDisplay, BlueprintReadWrite, Category = "Performance", AdvancedDisplay, meta = (ClampMin = "-1.0", ClampMax = "1.0", UIMin = "-1.0", UIMax = "1.0", SliderExponent = 1))
+	float LodBias;
+
+	// Contructor used to initialised new values
+	FHairGroupDesc()
+	{
+		ReInit();
+	}
+
+	void ReInit()
+	{
+		HairRaytracingRadiusScale = 1.0f;
+		HairRootScale = 1.0f;
+		HairTipScale = 1.0f;
+		HairClipLength = 1.0f;
+
+		LodAverageVertexPerPixel = 3.0f; // Good for quality on heavy assets
+		LodBias = 0.0f;
+	}
+};
diff --git a/Engine/Plugins/Experimental/HairStrands/Source/HairStrandsCore/Public/GroomImportOptions.h b/Engine/Plugins/Experimental/HairStrands/Source/HairStrandsCore/Public/GroomImportOptions.h
new file mode 100644
--- /dev/null
+++ b/Engine/Plugins/Experimental/HairStrands/Source/HairStrandsCore/Public/GroomImportOptions.h
@@ -0,0 +1,21 @@
+// Copyright Epic Games, Inc. All Rights Reserved.
+
+#pragma once
+
+#include "CoreMinimal.h"
+#include "GroomSettings.h"
+#include "UObject/Object.h"
+#include "GroomImportOptions.generated.h"
+
+UCLASS(BlueprintType, config = EditorPerProjectUserSettings, HideCategories = ("Hidden"))
+class HAIRSTRANDSCORE_API UGroomImportOptions : public UObject
+{
+	GENERATED_UCLASS_BODY()
+
+public:
+	UPROPERTY(config, EditAnywhere, BlueprintReadWrite, meta = (ShowOnlyInnerProperties), Category = Conversion)
+	FGroomConversionSettings ConversionSettings;
+
+	UPROPERTY(config, EditAnywhere, BlueprintReadWrite, meta = (ShowOnlyInnerProperties), Category = BuildSettings)
+	FGroomBuildSettings BuildSettings;
+};
diff --git a/Engine/Plugins/Experimental/HairStrands/Source/HairStrandsCore/Public/GroomSettings.h b/Engine/Plugins/Experimental/HairStrands/Source/HairStrandsCore/Public/GroomSettings.h
--- a/Engine/Plugins/Experimental/HairStrands/Source/HairStrandsCore/Public/GroomSettings.h
+++ b/Engine/Plugins/Experimental/HairStrands/Source/HairStrandsCore/Public/GroomSettings.h
@@ -31,6 +31,7 @@ enum class EGroomInterpolationQuality : uint8
 	Low		UMETA(DisplayName = "Low",		ToolTip = "Build interpolation data based on nearst neighbor search. Low quality interpolation data, but fast to build (takes a few minutes)"),
 	Medium	UMETA(DisplayName = "Medium",	ToolTip = "Build interpolation data using curve shape matching search but within a limited spatial range. This is a tradeoff between Low and high quality in term of quality & build time (can takes several dozen of minutes)"),
 	High	UMETA(DisplayName = "High",		ToolTip = "Build interpolation data using curve shape matching search. This result in high quality interpolation data, but is relatively slow to build (can takes several dozen of minutes)"),
+	Unknown	UMETA(Hidden),
 };
 
 UENUM()
@@ -39,6 +40,7 @@ enum class EGroomInterpolationWeight : uint8
 	Parametric	UMETA(DisplayName = "Parametric", ToolTip = "Build interpolation data based on curve parametric distance"),
 	Root		UMETA(DisplayName = "Root", ToolTip = "Build interpolation data based on distance between guide's root and strands's root"),
 	Index		UMETA(DisplayName = "Index", ToolTip = "Build interpolation data based on guide and strands vertex indices"),
+	Unknown		UMETA(Hidden),
 };
 
 USTRUCT(BlueprintType)
@@ -60,7 +62,7 @@ struct HAIRSTRANDSCORE_API FGroomBuildSettings
 	bool bOverrideGuides;
 
 	/** Density factor for converting hair into guide curve if no guides are provided. */
-	UPROPERTY(EditAnywhere, BlueprintReadWrite, Category = "BuildSettings", meta = (ClampMin = "0.01", UIMin = "0.01", UIMax = "1.0"))
+	UPROPERTY(EditAnywhere, BlueprintReadWrite, Category = "BuildSettings", meta = (ClampMin = "0", UIMin = "0", UIMax = "1.0"))
 	float HairToGuideDensity;
 
 	/** Interpolation data quality. */
@@ -78,4 +80,6 @@ struct HAIRSTRANDSCORE_API FGroomBuildSettings
 	/** Force a hair strand to be affected by a unique guide. */
 	UPROPERTY(EditAnywhere, BlueprintReadWrite, Category = "BuildSettings")
 	bool bUseUniqueGuide;
+
+	friend FArchive& operator<<(FArchive& Ar, FGroomBuildSettings& GroupInfo);
 };
diff --git a/Engine/Plugins/Experimental/HairStrands/Source/HairStrandsCore/Public/HairAttributes.h b/Engine/Plugins/Experimental/HairStrands/Source/HairStrandsCore/Public/HairAttributes.h
--- a/Engine/Plugins/Experimental/HairStrands/Source/HairStrandsCore/Public/HairAttributes.h
+++ b/Engine/Plugins/Experimental/HairStrands/Source/HairStrandsCore/Public/HairAttributes.h
@@ -17,11 +17,13 @@ namespace HairAttribute
 	{
 		extern HAIRSTRANDSCORE_API const FName Color;			// FVector
 		extern HAIRSTRANDSCORE_API const FName GroupID;			// int
-		extern HAIRSTRANDSCORE_API const FName Guide;			// bool
+		extern HAIRSTRANDSCORE_API const FName Guide;			// int
 		extern HAIRSTRANDSCORE_API const FName ID;				// int
 		extern HAIRSTRANDSCORE_API const FName RootUV;			// FVector2D
 		extern HAIRSTRANDSCORE_API const FName VertexCount;		// int
 		extern HAIRSTRANDSCORE_API const FName Width;			// float
+		extern HAIRSTRANDSCORE_API const FName ClosestGuides;	// FVector
+		extern HAIRSTRANDSCORE_API const FName GuideWeights;	// FVector
 	}
 
 	namespace Groom
diff --git a/Engine/Plugins/Experimental/HairStrands/Source/HairStrandsCore/Public/HairDescription.h b/Engine/Plugins/Experimental/HairStrands/Source/HairStrandsCore/Public/HairDescription.h
--- a/Engine/Plugins/Experimental/HairStrands/Source/HairStrandsCore/Public/HairDescription.h
+++ b/Engine/Plugins/Experimental/HairStrands/Source/HairStrandsCore/Public/HairDescription.h
@@ -8,6 +8,7 @@
 #include "MeshAttributeArray.h" // for TMeshAttributesRef
 #include "MeshDescription.h" // for TVertexAttributesRef
 #include "MeshTypes.h" // for FElementID, FVertexID
+#include "Serialization/BulkData.h"
 
 struct FStrandID : public FElementID
 {
@@ -74,6 +75,7 @@ public:
 	FVertexID AddVertex();
 	FStrandID AddStrand();
 
+	void Reset();
 	bool IsValid() const;
 
 	TAttributesSet<FVertexID>& VertexAttributes() { return VertexAttributesSet; }
@@ -101,6 +103,51 @@ private:
 	int32 NumStrands;
 };
 
+/**
+ * Bulk data storage for FHairDescription
+ */
+struct HAIRSTRANDSCORE_API FHairDescriptionBulkData
+{
+public:
+	FHairDescriptionBulkData()
+	{
+		BulkData.SetBulkDataFlags(BULKDATA_SerializeCompressed | BULKDATA_SerializeCompressedZLIB);
+	}
+
+#if WITH_EDITORONLY_DATA
+	void Serialize(FArchive& Ar, UObject* Owner);
+
+	/** Stores a new hair description in the bulk data */
+	void SaveHairDescription(FHairDescription& HairDescription);
+
+	/** Loads the hair description from the bulk data */
+	void LoadHairDescription(FHairDescription& HairDescription);
+
+	/** Empties the bulk data */
+	void Empty();
+
+	/** Returns true if there is nothing in the bulk data */
+	bool IsEmpty() const { return BulkData.GetBulkDataSize() == 0; }
+
+	/** Returns unique ID string for this bulk data */
+	FString GetIdString() const;
+
+private:
+	/** Computes a GUID from the hash of the bulk data, useful to prevent recomputing content already in cache. */
+	void ComputeGuidFromHash();
+#endif
+
+private:
+	/** Internally store bulk data as bytes */
+	FByteBulkData BulkData;
+
+	/** GUID associated with the stored bulk data */
+	FGuid Guid;
+
+	/** Custom version to propagate to archive when serializing the bulk data */
+	FCustomVersionContainer CustomVersions;
+};
+
 template <typename AttributeType> using TStrandAttributesRef = TMeshAttributesRef<FStrandID, AttributeType>;
 template <typename AttributeType> using TGroomAttributesRef = TMeshAttributesRef<FGroomID, AttributeType>;
 
diff --git a/Engine/Plugins/Experimental/HairStrands/Source/HairStrandsCore/Public/HairStrandsDatas.h b/Engine/Plugins/Experimental/HairStrands/Source/HairStrandsCore/Public/HairStrandsDatas.h
--- a/Engine/Plugins/Experimental/HairStrands/Source/HairStrandsCore/Public/HairStrandsDatas.h
+++ b/Engine/Plugins/Experimental/HairStrands/Source/HairStrandsCore/Public/HairStrandsDatas.h
@@ -325,6 +325,12 @@ struct HAIRSTRANDSCORE_API FHairStrandsCurves
 	/** Roots UV. Support UDIM coordinate up to 256x256 */
 	TArray<FVector2D> CurvesRootUV; // [0..256]
 
+	/** Strand ID associated with each curve */
+	TArray<int> StrandIDs;
+
+	/** Mapping of imported Groom ID to index */
+	TMap<int, int> GroomIDToIndex;
+
 	/** Max strands Curves length */
 	float MaxLength = 0;
 
diff --git a/Engine/Plugins/Experimental/HairStrands/Source/HairStrandsCore/Public/HairStrandsRendering.h b/Engine/Plugins/Experimental/HairStrands/Source/HairStrandsCore/Public/HairStrandsRendering.h
--- a/Engine/Plugins/Experimental/HairStrands/Source/HairStrandsCore/Public/HairStrandsRendering.h
+++ b/Engine/Plugins/Experimental/HairStrands/Source/HairStrandsCore/Public/HairStrandsRendering.h
@@ -6,6 +6,8 @@
 #include "UObject/ObjectMacros.h"
 #include "Engine/EngineTypes.h"
 #include "RenderGraphResources.h"
+#include "Shader.h"
+#include "GroomDesc.h"
 
 struct FHairStrandsInterpolationInput
 {
@@ -32,9 +34,16 @@ struct FHairStrandsInterpolationInput
 		bool					bIsRTGeometryInitialized = false;
 		#endif
 
+		uint32		ClusterCount = 0;
+		uint32		ClusterVertexCount = 0;
+		FReadBuffer*VertexToClusterIdBuffer = nullptr;
+		FReadBuffer*ClusterVertexIdBuffer = nullptr;
+		FReadBuffer*ClusterIndexRadiusScaleInfoBuffer = nullptr;
+
 		bool bIsSimulationEnable = false;
-		float HairRadius = 0;
-		float HairRaytracingRadiusScale = 0;
+		
+		FHairGroupDesc GroupDesc;
+
 		FVector InRenderHairPositionOffset = FVector::ZeroVector;
 		FVector InSimHairPositionOffset = FVector::ZeroVector;
 		FVector OutHairPositionOffset = FVector::ZeroVector;
@@ -67,9 +76,13 @@ struct FHairStrandsInterpolationOutput
 		FShaderResourceViewRHIRef HairTangentBuffer = nullptr;
 		FShaderResourceViewRHIRef HairAttributeBuffer = nullptr;
 		FShaderResourceViewRHIRef HairMaterialBuffer = nullptr;
+
 		FVector HairPositionOffset = FVector::ZeroVector;
 		FVector HairPreviousPositionOffset = FVector::ZeroVector;
 		uint32 VertexCount = 0;
+		float HairRadius = 0;
+		float HairLength = 0;
+		float HairDensity = 0;
 
 		inline void Reset()
 		{
@@ -81,6 +94,9 @@ struct FHairStrandsInterpolationOutput
 			HairPositionOffset = FVector::ZeroVector;
 			HairPreviousPositionOffset = FVector::ZeroVector;
 			VertexCount = 0;
+			HairRadius = 0;
+			HairLength = 0;
+			HairDensity = 0;
 		}
 	};
 
@@ -96,9 +112,14 @@ struct FHairStrandsInterpolationOutput
 
 		FRWBuffer* SimTangentBuffer = nullptr;
 
+		FRWBuffer* RenderGroupAABBBuffer = nullptr;
+		FRWBuffer* RenderClusterAABBBuffer = nullptr;
+		FReadBuffer* ClusterInfoBuffer = nullptr;
+
 		// Debug buffers (allocated on-the-fly if used)
 		FRWBuffer RenderPatchedAttributeBuffer;
 
+		class FHairGroupPublicData* HairGroupPublicData = nullptr;
 		VertexFactoryInput VFInput;
 
 		inline bool IsValid() const
@@ -116,11 +137,22 @@ struct FHairStrandsInterpolationOutput
 	TArray<HairGroup> HairGroups;
 };
 
+// Reset the interpolation data. This needs to be called prior to ComputeHairStrandsInterpolation 
+// and prior to the actual hair simulation in order to insure that:
+//  1) when hair simulation is enabled, the first frame is correct
+//  2) when hair simulation is enabled/disabled (i.e., toggle/change) 
+//     we reset to deform buffer to rest state)
+void ResetHairStrandsInterpolation(
+	FRHICommandListImmediate& RHICmdList,
+	FHairStrandsInterpolationInput* InInput,
+	FHairStrandsInterpolationOutput* InOutput);
+
 void ComputeHairStrandsInterpolation(
 	FRHICommandListImmediate& RHICmdList,
+	const struct FShaderDrawDebugData* DebugShaderData,
 	FHairStrandsInterpolationInput* Input,
 	FHairStrandsInterpolationOutput* Output, 
 	struct FHairStrandsProjectionHairData& RenHairDatas,
 	struct FHairStrandsProjectionHairData& SimHairDatas,
-	int32 LODIndex);
-
+	int32 LODIndex,
+	struct FHairStrandClusterData* ClusterData);
diff --git a/Engine/Plugins/Experimental/HairStrands/Source/HairStrandsCore/Public/HairStrandsVertexFactory.h b/Engine/Plugins/Experimental/HairStrands/Source/HairStrandsCore/Public/HairStrandsVertexFactory.h
--- a/Engine/Plugins/Experimental/HairStrands/Source/HairStrandsCore/Public/HairStrandsVertexFactory.h
+++ b/Engine/Plugins/Experimental/HairStrands/Source/HairStrandsCore/Public/HairStrandsVertexFactory.h
@@ -9,6 +9,8 @@
 #include "RenderGraphResources.h"
 #include "HairStrandsDatas.h"
 #include "HairStrandsRendering.h"
+#include "HairStrandsInterface.h"
+#include "PrimitiveSceneProxy.h"
 
 class FMaterial;
 class FSceneView;
@@ -32,16 +34,6 @@ public:
 	struct FDataType
 	{
 		FHairStrandsInterpolationOutput* InterpolationOutput = nullptr;
-		struct HairGroup
-		{
-			float MinStrandRadius = 0;
-			float MaxStrandRadius = 0;
-			float MaxStrandLength = 0;
-			float HairDensity = 1;
-			FVector HairWorldOffset = FVector::ZeroVector;
-		};
-
-		TArray<HairGroup> HairGroups;
 	};
 
 	/**
@@ -76,12 +68,13 @@ public:
 	inline FRHIShaderResourceView* GetMaterialSRV(uint32 GroupIndex) const			{ check(Data.InterpolationOutput); return Data.InterpolationOutput->HairGroups[GroupIndex].VFInput.HairMaterialBuffer; }
 	inline FRHIShaderResourceView* GetTangentSRV(uint32 GroupIndex) const			{ check(Data.InterpolationOutput); return Data.InterpolationOutput->HairGroups[GroupIndex].VFInput.HairTangentBuffer; }
 
-	float GetMaxStrandRadius(uint32 GroupIndex) const;
-	inline float GetMinStrandRadius(uint32 GroupIndex) const						{ return Data.HairGroups[GroupIndex].MinStrandRadius; }
-	inline float GetMaxStrandLength(uint32 GroupIndex) const						{ return Data.HairGroups[GroupIndex].MaxStrandLength; }
-	inline float GetHairDensity(uint32 GroupIndex) const							{ return Data.HairGroups[GroupIndex].HairDensity;  }
-	inline const FVector& GetPositionOffset(uint32 GroupIndex) const				{ return Data.InterpolationOutput->HairGroups[GroupIndex].VFInput.HairPositionOffset; }
-	inline const FVector& GetPreviousPositionOffset(uint32 GroupIndex) const		{ return Data.InterpolationOutput->HairGroups[GroupIndex].VFInput.HairPreviousPositionOffset; }
+	inline float GetMaxStrandRadius(uint32 GroupIndex) const						{ check(Data.InterpolationOutput); return Data.InterpolationOutput->HairGroups[GroupIndex].VFInput.HairRadius; };
+	inline float GetMaxStrandLength(uint32 GroupIndex) const						{ check(Data.InterpolationOutput); return Data.InterpolationOutput->HairGroups[GroupIndex].VFInput.HairLength; };
+	inline float GetHairDensity(uint32 GroupIndex) const							{ check(Data.InterpolationOutput); return Data.InterpolationOutput->HairGroups[GroupIndex].VFInput.HairDensity; };
+	inline const FVector& GetPositionOffset(uint32 GroupIndex) const				{ check(Data.InterpolationOutput); return Data.InterpolationOutput->HairGroups[GroupIndex].VFInput.HairPositionOffset; }
+	inline const FVector& GetPreviousPositionOffset(uint32 GroupIndex) const		{ check(Data.InterpolationOutput); return Data.InterpolationOutput->HairGroups[GroupIndex].VFInput.HairPreviousPositionOffset; }
+	
+	inline FHairGroupPublicData* GetHairGroupPublicData(uint32 GroupIndex) const	{ check(Data.InterpolationOutput); return Data.InterpolationOutput->HairGroups[GroupIndex].HairGroupPublicData; }
 
 	const FDataType& GetData() const { return Data; }
 protected:
diff --git a/Engine/Plugins/Experimental/HairStrands/Source/HairStrandsEditor/HairStrandsEditor.Build.cs b/Engine/Plugins/Experimental/HairStrands/Source/HairStrandsEditor/HairStrandsEditor.Build.cs
--- a/Engine/Plugins/Experimental/HairStrands/Source/HairStrandsEditor/HairStrandsEditor.Build.cs
+++ b/Engine/Plugins/Experimental/HairStrands/Source/HairStrandsEditor/HairStrandsEditor.Build.cs
@@ -22,6 +22,7 @@ namespace UnrealBuildTool.Rules
 					"Slate",
 					"SlateCore",
 					"Projects",
+					"ToolMenus",
 					"UnrealEd",
 					"AssetTools",
 				});
diff --git a/Engine/Plugins/Experimental/HairStrands/Source/HairStrandsEditor/Private/GroomActions.cpp b/Engine/Plugins/Experimental/HairStrands/Source/HairStrandsEditor/Private/GroomActions.cpp
--- a/Engine/Plugins/Experimental/HairStrands/Source/HairStrandsEditor/Private/GroomActions.cpp
+++ b/Engine/Plugins/Experimental/HairStrands/Source/HairStrandsEditor/Private/GroomActions.cpp
@@ -4,7 +4,11 @@
 #include "GroomAsset.h"
 
 #include "EditorFramework/AssetImportData.h"
+#include "GroomAssetImportData.h"
+#include "GroomImportOptions.h"
+#include "GroomImportOptionsWindow.h"
 #include "Toolkits/SimpleAssetEditor.h"
+#include "ToolMenuSection.h"
 
 #define LOCTEXT_NAMESPACE "AssetTypeActions"
 
@@ -16,11 +20,23 @@ bool FGroomActions::CanFilter()
 	return true;
 }
 
-void FGroomActions::GetActions(const TArray<UObject*>& InObjects, FMenuBuilder& MenuBuilder)
+void FGroomActions::GetActions(const TArray<UObject*>& InObjects, FToolMenuSection& Section)
 {
-	FAssetTypeActions_Base::GetActions(InObjects, MenuBuilder);
-}
+	FAssetTypeActions_Base::GetActions(InObjects, Section);
+
+	TArray<TWeakObjectPtr<UGroomAsset>> GroomAssets = GetTypedWeakObjectPtrs<UGroomAsset>(InObjects);
 
+	Section.AddMenuEntry(
+		"RebuildGroom",
+		LOCTEXT("RebuildGroom", "Rebuild"),
+		LOCTEXT("RebuildGroomTooltip", "Rebuild the groom with new build settings"),
+		FSlateIcon(FEditorStyle::GetStyleSetName(), "ContentBrowser.AssetActions"),
+		FUIAction(
+			FExecuteAction::CreateSP(this, &FGroomActions::ExecuteRebuild, GroomAssets),
+			FCanExecuteAction::CreateSP(this, &FGroomActions::CanRebuild, GroomAssets)
+		)
+	);
+}
 
 uint32 FGroomActions::GetCategories()
 {
@@ -37,7 +53,10 @@ void FGroomActions::GetResolvedSourceFilePaths(const TArray<UObject*>& TypeAsset
 	for (UObject* Asset : TypeAssets)
 	{
 		const UGroomAsset* GroomAsset = CastChecked<UGroomAsset>(Asset);
-		GroomAsset->AssetImportData->ExtractFilenames(OutSourceFilePaths);
+		if (GroomAsset && GroomAsset->AssetImportData)
+		{
+			GroomAsset->AssetImportData->ExtractFilenames(OutSourceFilePaths);
+		}
 	}
 }
 
@@ -62,4 +81,49 @@ void FGroomActions::OpenAssetEditor(const TArray<UObject*>& InObjects, TSharedPt
 	FSimpleAssetEditor::CreateEditor(EToolkitMode::Standalone, EditWithinLevelEditor, InObjects);
 }
 
+bool FGroomActions::CanRebuild(TArray<TWeakObjectPtr<UGroomAsset>> Objects) const
+{
+	for (TWeakObjectPtr<UGroomAsset> GroomAsset : Objects)
+	{
+		if (GroomAsset.IsValid() && GroomAsset->CanRebuildFromDescription())
+		{
+			return true;
+		}
+	}
+	return false;
+}
+
+void FGroomActions::ExecuteRebuild(TArray<TWeakObjectPtr<UGroomAsset>> Objects) const
+{
+	for (TWeakObjectPtr<UGroomAsset> GroomAsset : Objects)
+	{
+		if (GroomAsset.IsValid() && GroomAsset->CanRebuildFromDescription() && GroomAsset->AssetImportData)
+		{
+			UGroomAssetImportData* GroomAssetImportData = Cast<UGroomAssetImportData>(GroomAsset->AssetImportData);
+			if (GroomAssetImportData && GroomAssetImportData->ImportOptions)
+			{
+				FString Filename(GroomAssetImportData->GetFirstFilename());
+
+				// Duplicate the options to prevent dirtying the asset when they are modified but the rebuild is cancelled
+				UGroomImportOptions* CurrentOptions = DuplicateObject<UGroomImportOptions>(GroomAssetImportData->ImportOptions, nullptr);
+				TSharedPtr<SGroomImportOptionsWindow> GroomOptionWindow = SGroomImportOptionsWindow::DisplayRebuildOptions(CurrentOptions, Filename);
+
+				if (!GroomOptionWindow->ShouldImport())
+				{
+					continue;
+				}
+
+				bool bSucceeded = GroomAsset->CacheDerivedData(&CurrentOptions->BuildSettings);
+				if (bSucceeded)
+				{
+					// Move the transient ImportOptions to the asset package and set it on the GroomAssetImportData for serialization
+					CurrentOptions->Rename(nullptr, GroomAssetImportData);
+					GroomAssetImportData->ImportOptions = CurrentOptions;
+					GroomAsset->MarkPackageDirty();
+				}
+			}
+		}
+	}
+}
+
 #undef LOCTEXT_NAMESPACE
\ No newline at end of file
diff --git a/Engine/Plugins/Experimental/HairStrands/Source/HairStrandsEditor/Private/GroomImportOptionsWindow.cpp b/Engine/Plugins/Experimental/HairStrands/Source/HairStrandsEditor/Private/GroomImportOptionsWindow.cpp
--- a/Engine/Plugins/Experimental/HairStrands/Source/HairStrandsEditor/Private/GroomImportOptionsWindow.cpp
+++ b/Engine/Plugins/Experimental/HairStrands/Source/HairStrandsEditor/Private/GroomImportOptionsWindow.cpp
@@ -79,7 +79,7 @@ void SGroomImportOptionsWindow::Construct(const FArguments& InArgs)
 			[
 				SAssignNew(ImportButton, SButton)
 				.HAlign(HAlign_Center)
-				.Text(LOCTEXT("GroomOptionsWindow_Import", "Import"))
+				.Text(InArgs._ButtonLabel)
 				.IsEnabled(this, &SGroomImportOptionsWindow::CanImport)
 				.OnClicked(this, &SGroomImportOptionsWindow::OnImport)
 			]
@@ -88,7 +88,6 @@ void SGroomImportOptionsWindow::Construct(const FArguments& InArgs)
 				SNew(SButton)
 				.HAlign(HAlign_Center)
 				.Text(LOCTEXT("GroomOptionsWindow_Cancel", "Cancel"))
-				.ToolTipText(LOCTEXT("GroomOptionsWindow_Cancel_ToolTip", "Cancels importing this groom file"))
 				.OnClicked(this, &SGroomImportOptionsWindow::OnCancel)
 			]
 		]
@@ -100,14 +99,41 @@ bool SGroomImportOptionsWindow::CanImport()  const
 	return true;
 }
 
-TSharedPtr<SGroomImportOptionsWindow> SGroomImportOptionsWindow::DisplayOptions(UGroomImportOptions* ImportOptions, const FString& FilePath)
+enum class EGroomOptionsVisibility : uint8
+{
+	None = 0x00,
+	ConversionOptions = 0x01,
+	BuildOptions = 0x02,
+	All = ConversionOptions | BuildOptions
+};
+
+ENUM_CLASS_FLAGS(EGroomOptionsVisibility);
+
+TSharedPtr<SGroomImportOptionsWindow> DisplayOptions(UGroomImportOptions* ImportOptions, const FString& FilePath, EGroomOptionsVisibility VisibilityFlag, FText WindowTitle, FText InButtonLabel)
 {
 	TSharedRef<SWindow> Window = SNew(SWindow)
-		.Title(LOCTEXT("GroomOptionsWindow_WindowTitle", "Groom Import Options"))
+		.Title(WindowTitle)
 		.SizingRule(ESizingRule::Autosized);
 
 	TSharedPtr<SGroomImportOptionsWindow> OptionsWindow;
 
+	FProperty* ConversionOptionsProperty = FindField<FProperty>(ImportOptions->GetClass(), GET_MEMBER_NAME_CHECKED(UGroomImportOptions, ConversionSettings));
+	if (ConversionOptionsProperty)
+	{
+		if (EnumHasAnyFlags(VisibilityFlag, EGroomOptionsVisibility::ConversionOptions))
+		{
+			ConversionOptionsProperty->SetMetaData(TEXT("ShowOnlyInnerProperties"), TEXT("1"));
+			ConversionOptionsProperty->SetMetaData(TEXT("Category"), TEXT("Conversion"));
+		}
+		else
+		{
+			// Note that UGroomImportOptions HideCategories named "Hidden",
+			// but the hiding doesn't work with ShowOnlyInnerProperties 
+			ConversionOptionsProperty->RemoveMetaData(TEXT("ShowOnlyInnerProperties"));
+			ConversionOptionsProperty->SetMetaData(TEXT("Category"), TEXT("Hidden"));
+		}
+	}
+
 	FString FileName = FPaths::GetCleanFilename(FilePath);
 	Window->SetContent
 	(
@@ -115,6 +141,7 @@ TSharedPtr<SGroomImportOptionsWindow> SGroomImportOptionsWindow::DisplayOptions(
 		.ImportOptions(ImportOptions)
 		.WidgetWindow(Window)
 		.FullPath(FText::FromString(FileName))
+		.ButtonLabel(InButtonLabel)
 	);
 
 	TSharedPtr<SWindow> ParentWindow;
@@ -130,4 +157,15 @@ TSharedPtr<SGroomImportOptionsWindow> SGroomImportOptionsWindow::DisplayOptions(
 	return OptionsWindow;
 }
 
+TSharedPtr<SGroomImportOptionsWindow> SGroomImportOptionsWindow::DisplayImportOptions(UGroomImportOptions* ImportOptions, const FString& FilePath)
+{
+	return DisplayOptions(ImportOptions, FilePath, EGroomOptionsVisibility::All, LOCTEXT("GroomOptionsWindow_WindowTitle", "Groom Import Options"), LOCTEXT("GroomOptionsWindow_Import", "Import"));
+}
+
+TSharedPtr<SGroomImportOptionsWindow> SGroomImportOptionsWindow::DisplayRebuildOptions(UGroomImportOptions* ImportOptions, const FString& FilePath)
+{
+	return DisplayOptions(ImportOptions, FilePath, EGroomOptionsVisibility::BuildOptions, LOCTEXT("GroomBuildOptionsWindow_WindowTitle", "Groom Build Options"), LOCTEXT("GroomOptionsWindow_Build", "Build"));
+}
+
+
 #undef LOCTEXT_NAMESPACE
diff --git a/Engine/Plugins/Experimental/HairStrands/Source/HairStrandsEditor/Private/GroomImportOptionsWindow.h b/Engine/Plugins/Experimental/HairStrands/Source/HairStrandsEditor/Private/GroomImportOptionsWindow.h
--- a/Engine/Plugins/Experimental/HairStrands/Source/HairStrandsEditor/Private/GroomImportOptionsWindow.h
+++ b/Engine/Plugins/Experimental/HairStrands/Source/HairStrandsEditor/Private/GroomImportOptionsWindow.h
@@ -19,18 +19,21 @@ public:
 		: _ImportOptions(nullptr)
 		, _WidgetWindow()
 		, _FullPath()
+		, _ButtonLabel()
 	{}
 
 	SLATE_ARGUMENT(UGroomImportOptions*, ImportOptions)
 	SLATE_ARGUMENT(TSharedPtr<SWindow>, WidgetWindow)
 	SLATE_ARGUMENT(FText, FullPath)
+	SLATE_ARGUMENT(FText, ButtonLabel)
 	SLATE_END_ARGS()
 
 public:
 	void Construct(const FArguments& InArgs);
 	virtual bool SupportsKeyboardFocus() const override { return true; }
 
-	static TSharedPtr<SGroomImportOptionsWindow> DisplayOptions(UGroomImportOptions* ImportOptions, const FString& FilePath);
+	static TSharedPtr<SGroomImportOptionsWindow> DisplayImportOptions(UGroomImportOptions* ImportOptions, const FString& FilePath);
+	static TSharedPtr<SGroomImportOptionsWindow> DisplayRebuildOptions(UGroomImportOptions* ImportOptions, const FString& FilePath);
 
 	FReply OnImport()
 	{
diff --git a/Engine/Plugins/Experimental/HairStrands/Source/HairStrandsEditor/Private/HairStrandsFactory.cpp b/Engine/Plugins/Experimental/HairStrands/Source/HairStrandsEditor/Private/HairStrandsFactory.cpp
--- a/Engine/Plugins/Experimental/HairStrands/Source/HairStrandsEditor/Private/HairStrandsFactory.cpp
+++ b/Engine/Plugins/Experimental/HairStrands/Source/HairStrandsEditor/Private/HairStrandsFactory.cpp
@@ -75,7 +75,7 @@ UObject* UHairStrandsFactory::FactoryCreateFile(UClass* InClass, UObject* InPare
 	if (!GIsRunningUnattendedScript && !IsAutomatedImport())
 	{
 		// Display import options and handle user cancellation
-		TSharedPtr<SGroomImportOptionsWindow> GroomOptionWindow = SGroomImportOptionsWindow::DisplayOptions(ImportOptions, Filename);
+		TSharedPtr<SGroomImportOptionsWindow> GroomOptionWindow = SGroomImportOptionsWindow::DisplayImportOptions(ImportOptions, Filename);
 		if (!GroomOptionWindow->ShouldImport())
 		{
 			bOutOperationCanceled = true;
diff --git a/Engine/Plugins/Experimental/HairStrands/Source/HairStrandsEditor/Private/HairStrandsImporter.cpp b/Engine/Plugins/Experimental/HairStrands/Source/HairStrandsEditor/Private/HairStrandsImporter.cpp
--- a/Engine/Plugins/Experimental/HairStrands/Source/HairStrandsEditor/Private/HairStrandsImporter.cpp
+++ b/Engine/Plugins/Experimental/HairStrands/Source/HairStrandsEditor/Private/HairStrandsImporter.cpp
@@ -39,7 +39,8 @@ UGroomAsset* FHairStrandsImporter::ImportHair(const FHairImportContext& ImportCo
 		}
 	}
 
-	bool bSucceeded = FGroomBuilder::BuildGroom(HairDescription, ImportContext.ImportOptions->BuildSettings, HairAsset);
+	HairAsset->CommitHairDescription(MoveTemp(HairDescription));
+	bool bSucceeded = HairAsset->CacheDerivedData(&ImportContext.ImportOptions->BuildSettings);
 	if (!bSucceeded)
 	{
 		// Purge the newly created asset that failed to import
@@ -51,7 +52,7 @@ UGroomAsset* FHairStrandsImporter::ImportHair(const FHairImportContext& ImportCo
 		return nullptr;
 	}
 
-	HairAsset->HairDescription = MakeUnique<FHairDescription>(MoveTemp(HairDescription));
+	HairAsset->InitResource();
 
 	return HairAsset;
 }
diff --git a/Engine/Plugins/Experimental/HairStrands/Source/HairStrandsEditor/Private/ReimportHairStrandsFactory.cpp b/Engine/Plugins/Experimental/HairStrands/Source/HairStrandsEditor/Private/ReimportHairStrandsFactory.cpp
--- a/Engine/Plugins/Experimental/HairStrands/Source/HairStrandsEditor/Private/ReimportHairStrandsFactory.cpp
+++ b/Engine/Plugins/Experimental/HairStrands/Source/HairStrandsEditor/Private/ReimportHairStrandsFactory.cpp
@@ -97,7 +97,7 @@ EReimportResult::Type UReimportHairStrandsFactory::Reimport(UObject* Obj)
 
 		if (!GIsRunningUnattendedScript && !IsAutomatedImport())
 		{
-			TSharedPtr<SGroomImportOptionsWindow> GroomOptionWindow = SGroomImportOptionsWindow::DisplayOptions(CurrentOptions, CurrentFilename);
+			TSharedPtr<SGroomImportOptionsWindow> GroomOptionWindow = SGroomImportOptionsWindow::DisplayImportOptions(CurrentOptions, CurrentFilename);
 
 			if (!GroomOptionWindow->ShouldImport())
 			{
@@ -115,9 +115,6 @@ EReimportResult::Type UReimportHairStrandsFactory::Reimport(UObject* Obj)
 			return EReimportResult::Failed;
 		}
 
-		FScopedSlowTask Progress( (float) 1, LOCTEXT("ReimportHairAsset", "Reimporting hair asset..."), true );
-		Progress.MakeDialog(true);
-
 		FHairDescription HairDescription;
 		if (!SelectedTranslator->Translate(CurrentFilename, HairDescription, CurrentOptions->ConversionSettings))
 		{
diff --git a/Engine/Plugins/Experimental/HairStrands/Source/HairStrandsEditor/Public/GroomActions.h b/Engine/Plugins/Experimental/HairStrands/Source/HairStrandsEditor/Public/GroomActions.h
--- a/Engine/Plugins/Experimental/HairStrands/Source/HairStrandsEditor/Public/GroomActions.h
+++ b/Engine/Plugins/Experimental/HairStrands/Source/HairStrandsEditor/Public/GroomActions.h
@@ -6,9 +6,10 @@
 #include "Templates/SharedPointer.h"
 
 class ISlateStyle;
+class UGroomAsset;
 
 /**
- * Implements an action for UHairStrands assets.
+ * Implements an action for groom assets.
  */
 class FGroomActions : public FAssetTypeActions_Base
 {
@@ -26,7 +27,7 @@ public:
 	//~ FAssetTypeActions_Base overrides
 
 	virtual bool CanFilter() override;
-	virtual void GetActions(const TArray<UObject*>& InObjects, FMenuBuilder& MenuBuilder) override;
+	virtual void GetActions(const TArray<UObject*>& InObjects, FToolMenuSection& Section) override;
 	virtual uint32 GetCategories() override;
 	virtual FText GetName() const override;
 	virtual void GetResolvedSourceFilePaths(const TArray<UObject*>& TypeAssets, TArray<FString>& OutSourceFilePaths) const override;
@@ -35,4 +36,10 @@ public:
 	virtual bool HasActions(const TArray<UObject*>& InObjects) const override;
 	virtual bool IsImportedAsset() const override { return true; }
 	virtual void OpenAssetEditor(const TArray<UObject*>& InObjects, TSharedPtr<IToolkitHost> EditWithinLevelEditor = TSharedPtr<IToolkitHost>()) override;
+
+private:
+
+	/** Callback for Rebuild groom action */
+	bool CanRebuild(TArray<TWeakObjectPtr<UGroomAsset>> Objects) const;
+	void ExecuteRebuild(TArray<TWeakObjectPtr<UGroomAsset>> Objects) const;
 };
