Copying //UE4/Dev-Physics to Main (//UE4/Main)

diff --git a/Engine/Plugins/Experimental/HairStrands/Source/HairStrandsNiagara/Private/NiagaraDataInterfaceHairStrands.cpp b/Engine/Plugins/Experimental/HairStrands/Source/HairStrandsNiagara/Private/NiagaraDataInterfaceHairStrands.cpp
--- a/Engine/Plugins/Experimental/HairStrands/Source/HairStrandsNiagara/Private/NiagaraDataInterfaceHairStrands.cpp
+++ b/Engine/Plugins/Experimental/HairStrands/Source/HairStrandsNiagara/Private/NiagaraDataInterfaceHairStrands.cpp
@@ -19,7 +19,7 @@
 
 #define LOCTEXT_NAMESPACE "NiagaraDataInterfaceHairStrands"
 
-static const int32 MaxDelay = 8;
+static const int32 MaxDelay = 5;
 
 //------------------------------------------------------------------------------------------------------------
 
@@ -96,6 +96,10 @@ static const FName ProjectBendRodMaterialName(TEXT("ProjectBendRodMaterial"));
 static const FName SolveStaticCollisionConstraintName(TEXT("SolveStaticCollisionConstraint"));
 static const FName ProjectStaticCollisionConstraintName(TEXT("ProjectStaticCollisionConstraint"));
 
+static const FName SetupSoftCollisionConstraintName(TEXT("SetupSoftCollisionConstraint"));
+static const FName SolveSoftCollisionConstraintName(TEXT("SolveSoftCollisionConstraint"));
+static const FName ProjectSoftCollisionConstraintName(TEXT("ProjectSoftCollisionConstraint"));
+
 //------------------------------------------------------------------------------------------------------------
 
 static const FName ComputeRestDirectionName(TEXT("ComputeRestDirection"));
@@ -226,7 +230,7 @@ class FCopyBoundingBoxCS : public FGlobalShader
 
 	BEGIN_SHADER_PARAMETER_STRUCT(FParameters, )
 		SHADER_PARAMETER(uint32, NumElements)
-		SHADER_PARAMETER_SRV(Buffer, BoundingBoxBuffer)
+		SHADER_PARAMETER_UAV(RWBuffer, BoundingBoxBuffer)
 		SHADER_PARAMETER_UAV(RWBuffer, OutNodeBoundBuffer)
 		END_SHADER_PARAMETER_STRUCT()
 
@@ -243,17 +247,42 @@ public:
 	}
 };
 
-class FResetBoundingBoxCS : public FGlobalShader
+IMPLEMENT_GLOBAL_SHADER(FCopyBoundingBoxCS, "/Plugin/Experimental/HairStrands/Private/NiagaraCopyBoundingBox.usf", "MainCS", SF_Compute);
+
+static void AddCopyBoundingBoxPass(
+	FRDGBuilder& GraphBuilder,
+	FRHIUnorderedAccessView* BoundingBoxBuffer,
+	FRHIUnorderedAccessView* OutNodeBoundBuffer)
 {
-	DECLARE_GLOBAL_SHADER(FResetBoundingBoxCS);
-	SHADER_USE_PARAMETER_STRUCT(FResetBoundingBoxCS, FGlobalShader);
+	const uint32 GroupSize = NIAGARA_HAIR_STRANDS_THREAD_COUNT;
+	const uint32 NumElements = 1;
 
-	BEGIN_SHADER_PARAMETER_STRUCT(FParameters, )
-		SHADER_PARAMETER(uint32, NumElements)
-		SHADER_PARAMETER_UAV(RWBuffer, OutBoundingBoxBuffer)
-		END_SHADER_PARAMETER_STRUCT()
+	FCopyBoundingBoxCS::FParameters* Parameters = GraphBuilder.AllocParameters<FCopyBoundingBoxCS::FParameters>();
+	Parameters->BoundingBoxBuffer = BoundingBoxBuffer;
+	Parameters->OutNodeBoundBuffer = OutNodeBoundBuffer;
+	Parameters->NumElements = NumElements;
+
+	TShaderMap<FGlobalShaderType>* ShaderMap = GetGlobalShaderMap(ERHIFeatureLevel::SM5);
+
+	const uint32 DispatchCount = FMath::DivideAndRoundUp(NumElements, GroupSize);
+
+	TShaderMapRef<FCopyBoundingBoxCS> ComputeShader(ShaderMap);
+	FComputeShaderUtils::AddPass(
+		GraphBuilder,
+		RDG_EVENT_NAME("CopyBoundingBox"),
+		*ComputeShader,
+		Parameters,
+		FIntVector(DispatchCount, 1, 1));
+}
+
+//------------------------------------------------------------------------------------------------------------
+
+class FDirectCopyBoundingBoxCS : public FGlobalShader
+{
+	DECLARE_SHADER_TYPE(FDirectCopyBoundingBoxCS, Global);
 
 public:
+
 	static bool ShouldCompilePermutation(const FGlobalShaderPermutationParameters& Parameters)
 	{
 		return RHISupportsComputeShaders(Parameters.Platform);
@@ -264,112 +293,142 @@ public:
 		FGlobalShader::ModifyCompilationEnvironment(Parameters, OutEnvironment);
 		OutEnvironment.SetDefine(TEXT("THREAD_COUNT"), NIAGARA_HAIR_STRANDS_THREAD_COUNT);
 	}
+
+	/** Default constructor. */
+	FDirectCopyBoundingBoxCS() {}
+
+	/** Initialization constructor. */
+	explicit FDirectCopyBoundingBoxCS(const ShaderMetaType::CompiledShaderInitializerType& Initializer);
+
+	/** Serialization. */
+	virtual bool Serialize(FArchive& Ar) override;
+
+	/**
+	 * Set parameters.
+	 */
+	void SetParameters(
+		FRHICommandList& RHICmdList,
+		FRHIUnorderedAccessView* InputBoundingBox,
+		FRHIUnorderedAccessView* OutputBoundingBox,
+		int32 NumElements);
+
+	/**
+	 * Unbinds any buffers that have been bound.
+	 */
+	void UnbindBuffers(FRHICommandList& RHICmdList);
+
+private:
+
+	FShaderParameter NumElements;
+	FShaderResourceParameter InputBoundingBox;
+	FShaderResourceParameter OutputBoundingBox;
 };
 
+IMPLEMENT_SHADER_TYPE(, FDirectCopyBoundingBoxCS, TEXT("/Plugin/Experimental/HairStrands/Private/NiagaraCopyBoundingBox.usf"), TEXT("MainCS"), SF_Compute);
 
-IMPLEMENT_GLOBAL_SHADER(FResetBoundingBoxCS, "/Plugin/Experimental/HairStrands/Private/NiagaraResetBoundingBox.usf", "MainCS", SF_Compute);
-IMPLEMENT_GLOBAL_SHADER(FCopyBoundingBoxCS, "/Plugin/Experimental/HairStrands/Private/NiagaraCopyBoundingBox.usf", "MainCS", SF_Compute);
+FDirectCopyBoundingBoxCS::FDirectCopyBoundingBoxCS(const ShaderMetaType::CompiledShaderInitializerType& Initializer)
+	: FGlobalShader(Initializer)
+{
+	NumElements.Bind(Initializer.ParameterMap, TEXT("NumElements"));
+	InputBoundingBox.Bind(Initializer.ParameterMap, TEXT("BoundingBoxBuffer"));
+	OutputBoundingBox.Bind(Initializer.ParameterMap, TEXT("OutNodeBoundBuffer"));
+}
 
-static void AddResetBoundingBoxPass(
-	FRDGBuilder& GraphBuilder,
-	FRHIUnorderedAccessView* OutBoundingBoxBuffer)
+bool FDirectCopyBoundingBoxCS::Serialize(FArchive& Ar)
 {
-	const uint32 GroupSize = NIAGARA_HAIR_STRANDS_THREAD_COUNT;
-	const uint32 NumElements = 1;
+	bool bShaderHasOutdatedParameters = FGlobalShader::Serialize(Ar);
+	Ar << NumElements;
+	Ar << InputBoundingBox;
+	Ar << OutputBoundingBox;
+	return bShaderHasOutdatedParameters;
+}
 
-	FResetBoundingBoxCS::FParameters* Parameters = GraphBuilder.AllocParameters<FResetBoundingBoxCS::FParameters>();
-	Parameters->OutBoundingBoxBuffer = OutBoundingBoxBuffer;
-	Parameters->NumElements = NumElements;
-	
-	TShaderMap<FGlobalShaderType>* ShaderMap = GetGlobalShaderMap(ERHIFeatureLevel::SM5);
+void FDirectCopyBoundingBoxCS::SetParameters(
+	FRHICommandList& RHICmdList,
+	FRHIUnorderedAccessView* InInputBoundingBox,
+	FRHIUnorderedAccessView* InOutputBoundingBox,
+	int32 InNumElements)
+{
+	FRHIComputeShader* ComputeShaderRHI = GetComputeShader();
 
-	const uint32 DispatchCount = FMath::DivideAndRoundUp(NumElements, GroupSize);
-	check(DispatchCount <= 65535);
+	RHICmdList.TransitionResource(EResourceTransitionAccess::ERWBarrier, EResourceTransitionPipeline::EComputeToCompute, InInputBoundingBox);
+	SetUAVParameter(RHICmdList, ComputeShaderRHI, InputBoundingBox, InInputBoundingBox);
 
-	TShaderMapRef<FResetBoundingBoxCS> ComputeShader(ShaderMap);
-	FComputeShaderUtils::AddPass(
-		GraphBuilder,
-		RDG_EVENT_NAME("ResetBoundingBox"),
-		*ComputeShader,
-		Parameters,
-		FIntVector(DispatchCount, 1, 1));
+	RHICmdList.TransitionResource(EResourceTransitionAccess::ERWBarrier, EResourceTransitionPipeline::EComputeToCompute, InOutputBoundingBox);
+	SetUAVParameter(RHICmdList, ComputeShaderRHI, OutputBoundingBox, InOutputBoundingBox);
+
+	SetShaderValue(RHICmdList, ComputeShaderRHI, NumElements, InNumElements); 
+
+	RHICmdList.TransitionResource(EResourceTransitionAccess::ERWBarrier, EResourceTransitionPipeline::EComputeToCompute, InInputBoundingBox);
+	RHICmdList.TransitionResource(EResourceTransitionAccess::ERWBarrier, EResourceTransitionPipeline::EComputeToCompute, InOutputBoundingBox);
 }
 
-static void AddCopyBoundingBoxPass(
-	FRDGBuilder& GraphBuilder,
-	FRHIShaderResourceView* BoundingBoxBuffer,
+void FDirectCopyBoundingBoxCS::UnbindBuffers(FRHICommandList& RHICmdList)
+{
+	FRHIComputeShader* ComputeShaderRHI = GetComputeShader();
+	SetUAVParameter(RHICmdList, ComputeShaderRHI, InputBoundingBox, nullptr);
+	SetUAVParameter(RHICmdList, ComputeShaderRHI, OutputBoundingBox, nullptr);
+}
+
+static void AddDirectCopyBoundingBoxPass(
+	FRHICommandListImmediate& RHICmdList,
+	FRHIUnorderedAccessView* BoundingBoxBuffer,
 	FRHIUnorderedAccessView* OutNodeBoundBuffer)
 {
 	const uint32 GroupSize = NIAGARA_HAIR_STRANDS_THREAD_COUNT;
 	const uint32 NumElements = 1;
 
-	FCopyBoundingBoxCS::FParameters* Parameters = GraphBuilder.AllocParameters<FCopyBoundingBoxCS::FParameters>();
-	Parameters->BoundingBoxBuffer = BoundingBoxBuffer;
-	Parameters->OutNodeBoundBuffer = OutNodeBoundBuffer;
-	Parameters->NumElements = NumElements;
-
 	TShaderMap<FGlobalShaderType>* ShaderMap = GetGlobalShaderMap(ERHIFeatureLevel::SM5);
-
 	const uint32 DispatchCount = FMath::DivideAndRoundUp(NumElements, GroupSize);
-	check(DispatchCount <= 65535);
 
-	TShaderMapRef<FCopyBoundingBoxCS> ComputeShader(ShaderMap);
-	FComputeShaderUtils::AddPass(
-		GraphBuilder,
-		RDG_EVENT_NAME("CopyBoundingBox"),
-		*ComputeShader,
-		Parameters,
-		FIntVector(DispatchCount, 1, 1));
-}
+	TShaderMapRef<FDirectCopyBoundingBoxCS> ComputeShader(ShaderMap);
+	RHICmdList.SetComputeShader(ComputeShader->GetComputeShader());
 
-//------------------------------------------------------------------------------------------------------------
+	ComputeShader->SetParameters(RHICmdList, BoundingBoxBuffer, OutNodeBoundBuffer, NumElements);
 
-void FNDIHairStrandsBuffer::SetHairAsset(const FHairStrandsDatas*  HairStrandsDatas, const FHairStrandsRestResource* HairStrandsRestResource, 
-	const FHairStrandsDeformedResource*  HairStrandsDeformedResource, const FHairStrandsRootResource* HairStrandsRootResource)
-{
-	SourceDatas = HairStrandsDatas;
-	SourceRestResources = HairStrandsRestResource;
-	SourceDeformedResources = HairStrandsDeformedResource;
-	SourceRootResources = HairStrandsRootResource;
+	DispatchComputeShader(RHICmdList, *ComputeShader, DispatchCount, 1, 1);
+	ComputeShader->UnbindBuffers(RHICmdList);
 }
 
-void FNDIHairStrandsBuffer::ClearBuffer(FRHICommandList& RHICmdList)
+//------------------------------------------------------------------------------------------------------------
+
+inline void ClearBuffer(FRHICommandList& RHICmdList, FNDIHairStrandsBuffer* HairStrandsBuffer)
 {
-	FRHIUnorderedAccessView* BoundingBoxBufferUAV = BoundingBoxBuffer.UAV;
-	FRHIShaderResourceView* BoundingBoxBufferSRV = BoundingBoxBuffer.SRV;
-	FRHIUnorderedAccessView* NodeBoundBufferUAV = NodeBoundBuffer.UAV;
+	FRHIUnorderedAccessView* BoundingBoxBufferUAV = HairStrandsBuffer->BoundingBoxBuffer.UAV;
+	FRHIShaderResourceView* BoundingBoxBufferSRV = HairStrandsBuffer->BoundingBoxBuffer.SRV;
+	FRHIUnorderedAccessView* NodeBoundBufferUAV = HairStrandsBuffer->NodeBoundBuffer.UAV;
 
-	ENQUEUE_RENDER_COMMAND(FNDIHairStrandsCopyBoundingBox)(
-		[BoundingBoxBufferUAV, NodeBoundBufferUAV, BoundingBoxBufferSRV]
-	(FRHICommandListImmediate& RHICmdListImm)
+	if (BoundingBoxBufferUAV != nullptr && BoundingBoxBufferSRV != nullptr && NodeBoundBufferUAV != nullptr)
 	{
-		FRDGBuilder GraphBuilder(RHICmdListImm);
-
-		RHICmdListImm.TransitionResource(EResourceTransitionAccess::EReadable, EResourceTransitionPipeline::EComputeToCompute, BoundingBoxBufferUAV);
-		RHICmdListImm.TransitionResource(EResourceTransitionAccess::EWritable, EResourceTransitionPipeline::EComputeToCompute, NodeBoundBufferUAV);
+		ENQUEUE_RENDER_COMMAND(FNDIHairStrandsCopyBoundingBox)(
+			[BoundingBoxBufferUAV, NodeBoundBufferUAV, BoundingBoxBufferSRV, HairStrandsBuffer]
+		(FRHICommandListImmediate& RHICmdListImm)
+		{
+		/*	FRDGBuilder GraphBuilder(RHICmdListImm);
 
-		AddCopyBoundingBoxPass(
-			GraphBuilder,
-			BoundingBoxBufferSRV, NodeBoundBufferUAV);
+			RHICmdListImm.TransitionResource(EResourceTransitionAccess::ERWBarrier, EResourceTransitionPipeline::EComputeToCompute, BoundingBoxBufferUAV);
+			RHICmdListImm.TransitionResource(EResourceTransitionAccess::ERWBarrier, EResourceTransitionPipeline::EComputeToCompute, NodeBoundBufferUAV);
 
-		GraphBuilder.Execute();
-	});
+			AddCopyBoundingBoxPass(
+				GraphBuilder,
+				BoundingBoxBufferUAV, NodeBoundBufferUAV);
 
-	ENQUEUE_RENDER_COMMAND(FNDIHairStrandsResetBoundingBox)(
-		[BoundingBoxBufferUAV, NodeBoundBufferUAV, BoundingBoxBufferSRV]
-	(FRHICommandListImmediate& RHICmdListImm)
-	{
-		FRDGBuilder GraphBuilder(RHICmdListImm);
+			GraphBuilder.Execute();*/
 
-		RHICmdListImm.TransitionResource(EResourceTransitionAccess::EWritable, EResourceTransitionPipeline::EComputeToCompute, BoundingBoxBufferUAV);
-		RHICmdListImm.TransitionResource(EResourceTransitionAccess::EReadable, EResourceTransitionPipeline::EComputeToCompute, NodeBoundBufferUAV);
+			AddDirectCopyBoundingBoxPass(RHICmdListImm, HairStrandsBuffer->BoundingBoxBuffer.UAV, HairStrandsBuffer->NodeBoundBuffer.UAV);
+		});
+	}
+}
 
-		AddResetBoundingBoxPass(
-			GraphBuilder,
-			BoundingBoxBufferUAV);
+//------------------------------------------------------------------------------------------------------------
 
-		GraphBuilder.Execute();
-	});
+void FNDIHairStrandsBuffer::Initialize(const FHairStrandsDatas*  HairStrandsDatas, const FHairStrandsRestResource* HairStrandsRestResource, 
+	const FHairStrandsDeformedResource*  HairStrandsDeformedResource, const FHairStrandsRootResource* HairStrandsRootResource)
+{
+	SourceDatas = HairStrandsDatas;
+	SourceRestResources = HairStrandsRestResource;
+	SourceDeformedResources = HairStrandsDeformedResource;
+	SourceRootResources = HairStrandsRootResource;
 }
 
 void FNDIHairStrandsBuffer::InitRHI()
@@ -387,7 +446,7 @@ void FNDIHairStrandsBuffer::InitRHI()
 			RHIUnlockVertexBuffer(CurvesOffsetsBuffer.Buffer);
 		}
 		{
-			static const TArray<uint32> ZeroData = {UINT_MAX,UINT_MAX,UINT_MAX,0,0,0};
+			static const TArray<uint32> ZeroData = { UINT_MAX,UINT_MAX,UINT_MAX,0,0,0 };
 
 			const uint32 BoundCount = 6;
 			const uint32 BoundBytes = sizeof(uint32)*BoundCount;
@@ -400,7 +459,7 @@ void FNDIHairStrandsBuffer::InitRHI()
 		}
 
 		{
-			static const TArray<uint32> ZeroData = { UINT_MAX,UINT_MAX,UINT_MAX,0,0,0};
+			static const TArray<uint32> ZeroData = { UINT_MAX,UINT_MAX,UINT_MAX,0,0,0 };
 
 			const uint32 BoundCount = 6;
 			const uint32 BoundBytes = sizeof(uint32)*BoundCount;
@@ -416,19 +475,6 @@ void FNDIHairStrandsBuffer::InitRHI()
 			const uint32 PositionsCount = SourceDatas->GetNumPoints();
 			DeformedPositionBuffer.Initialize(FHairStrandsPositionFormat::SizeInByte, PositionsCount, FHairStrandsPositionFormat::Format, BUF_Static);
 		}
-		if (SourceRootResources == nullptr)
-		{
-			const uint32 PositionsCount = SourceDatas->GetNumCurves();
-			RestTrianglePositionABuffer.Initialize(FHairStrandsMeshTrianglePositionFormat::SizeInByte, PositionsCount, FHairStrandsMeshTrianglePositionFormat::Format, BUF_Static);
-			RestTrianglePositionBBuffer.Initialize(FHairStrandsMeshTrianglePositionFormat::SizeInByte, PositionsCount, FHairStrandsMeshTrianglePositionFormat::Format, BUF_Static);
-			RestTrianglePositionCBuffer.Initialize(FHairStrandsMeshTrianglePositionFormat::SizeInByte, PositionsCount, FHairStrandsMeshTrianglePositionFormat::Format, BUF_Static);
-
-			DeformedTrianglePositionABuffer.Initialize(FHairStrandsMeshTrianglePositionFormat::SizeInByte, PositionsCount, FHairStrandsMeshTrianglePositionFormat::Format, BUF_Static);
-			DeformedTrianglePositionBBuffer.Initialize(FHairStrandsMeshTrianglePositionFormat::SizeInByte, PositionsCount, FHairStrandsMeshTrianglePositionFormat::Format, BUF_Static);
-			DeformedTrianglePositionCBuffer.Initialize(FHairStrandsMeshTrianglePositionFormat::SizeInByte, PositionsCount, FHairStrandsMeshTrianglePositionFormat::Format, BUF_Static);
-
-			RootBarycentricCoordinatesBuffer.Initialize(FHairStrandsCurveTriangleBarycentricFormat::SizeInByte, PositionsCount, FHairStrandsCurveTriangleBarycentricFormat::Format, BUF_Static);
-		}
 	}
 }
 
@@ -436,18 +482,69 @@ void FNDIHairStrandsBuffer::ReleaseRHI()
 {
 	CurvesOffsetsBuffer.Release();
 	DeformedPositionBuffer.Release();
+	BoundingBoxBuffer.Release();
+	NodeBoundBuffer.Release();
+}
+
+//------------------------------------------------------------------------------------------------------------
 
-	RestTrianglePositionABuffer.Release();
-	RestTrianglePositionBBuffer.Release();
-	RestTrianglePositionCBuffer.Release();
+void FNDIHairStrandsData::Release()
+{
+	if (HairStrandsBuffer)
+	{
+		BeginReleaseResource(HairStrandsBuffer);
+		ENQUEUE_RENDER_COMMAND(DeleteResource)(
+			[ParamPointerToRelease = HairStrandsBuffer](FRHICommandListImmediate& RHICmdList)
+		{
+			delete ParamPointerToRelease;
+		});
+		HairStrandsBuffer = nullptr;
+	}
+}
 
-	DeformedTrianglePositionABuffer.Release();
-	DeformedTrianglePositionBBuffer.Release();
-	DeformedTrianglePositionCBuffer.Release();
+bool FNDIHairStrandsData::Init(UNiagaraDataInterfaceHairStrands* Interface, FNiagaraSystemInstance* SystemInstance)
+{
+	HairStrandsBuffer = nullptr;
 
-	RootBarycentricCoordinatesBuffer.Release();
-	BoundingBoxBuffer.Release();
-	NodeBoundBuffer.Release();
+	if (Interface != nullptr && SystemInstance != nullptr)
+	{
+		FHairStrandsDatas* StrandsDatas = nullptr;
+		FHairStrandsRestResource* StrandsRestResource = nullptr;
+		FHairStrandsDeformedResource* StrandsDeformedResource = nullptr;
+		FHairStrandsRootResource* StrandsRootResource = nullptr;
+
+		{
+			Interface->ExtractDatasAndResources(SystemInstance, StrandsDatas, StrandsRestResource, StrandsDeformedResource, StrandsRootResource);
+
+			HairStrandsBuffer = new FNDIHairStrandsBuffer();
+			HairStrandsBuffer->Initialize(StrandsDatas, StrandsRestResource, StrandsDeformedResource, StrandsRootResource);
+			BeginInitResource(HairStrandsBuffer);
+
+			WorldTransform = Interface->IsComponentValid() ? Interface->SourceComponent->GetComponentToWorld() :
+				SystemInstance->GetComponent()->GetComponentToWorld();
+			StrandSize = static_cast<uint8>(Interface->StrandSize);
+
+			if (StrandsDatas != nullptr)
+			{
+				const FBox& StrandsBox = StrandsDatas->BoundingBox;
+
+				NumStrands = StrandsDatas->GetNumCurves();
+				BoxCenter = StrandsBox.GetCenter();
+				BoxExtent = StrandsBox.GetExtent();
+			}
+			else
+			{
+				NumStrands = 0;
+				BoxCenter = FVector(0, 0, 0);
+				BoxExtent = FVector(0, 0, 0);
+			}
+			TickCount = 0;
+			ForceReset = true;
+			ResetTick = MaxDelay;
+		}
+	}
+
+	return true;
 }
 
 //------------------------------------------------------------------------------------------------------------
@@ -591,61 +688,51 @@ struct FNDIHairStrandsParametersCS : public FNiagaraDataInterfaceParametersCS
 		FNDIHairStrandsData* ProxyData =
 			InterfaceProxy->SystemInstancesToProxyData.Find(Context.SystemInstance);
 
-		if (ProxyData != nullptr)
+		const bool IsHairValid = ProxyData != nullptr && ProxyData->HairStrandsBuffer != nullptr && ProxyData->HairStrandsBuffer->IsInitialized();
+		const bool IsRootValid = IsHairValid && ProxyData->HairStrandsBuffer->SourceRootResources != nullptr && ProxyData->HairStrandsBuffer->SourceRootResources->IsInitialized();
+		const bool IsRestValid = IsHairValid && ProxyData->HairStrandsBuffer->SourceRestResources != nullptr && ProxyData->HairStrandsBuffer->SourceRestResources->IsInitialized();
+		const bool IsDeformedValid = IsHairValid && ProxyData->HairStrandsBuffer->SourceDeformedResources != nullptr && ProxyData->HairStrandsBuffer->SourceDeformedResources->IsInitialized();
+
+		if (IsHairValid && IsRestValid)
 		{
 			FNDIHairStrandsBuffer* HairStrandsBuffer = ProxyData->HairStrandsBuffer;
 
-			FUnorderedAccessViewRHIRef PointPositionsUAV = (HairStrandsBuffer->SourceDeformedResources != nullptr) ? HairStrandsBuffer->SourceDeformedResources->DeformedPositionBuffer[0].UAV : 
-				HairStrandsBuffer->DeformedPositionBuffer.UAV;
-
-			FShaderResourceViewRHIRef RestTrianglePositionASRV = HairStrandsBuffer->RestTrianglePositionABuffer.SRV;
-			FShaderResourceViewRHIRef RestTrianglePositionBSRV = HairStrandsBuffer->RestTrianglePositionBBuffer.SRV;
-			FShaderResourceViewRHIRef RestTrianglePositionCSRV = HairStrandsBuffer->RestTrianglePositionCBuffer.SRV;
+			FUnorderedAccessViewRHIRef PointPositionsUAV = IsDeformedValid ?
+				HairStrandsBuffer->SourceDeformedResources->DeformedPositionBuffer[0].UAV : HairStrandsBuffer->DeformedPositionBuffer.UAV;
 
-			FShaderResourceViewRHIRef DeformedTrianglePositionASRV = HairStrandsBuffer->DeformedTrianglePositionABuffer.SRV;
-			FShaderResourceViewRHIRef DeformedTrianglePositionBSRV = HairStrandsBuffer->DeformedTrianglePositionBBuffer.SRV;
-			FShaderResourceViewRHIRef DeformedTrianglePositionCSRV = HairStrandsBuffer->DeformedTrianglePositionCBuffer.SRV;
+			const int32 HasRootAttachedValue = (IsRootValid && HairStrandsBuffer->SourceRootResources->MeshProjectionLODs.Num() > 0 && 
+				HairStrandsBuffer->SourceRootResources->MeshProjectionLODs[0].Status == FHairStrandsProjectionHairData::LODData::EStatus::Completed);
 
-			FShaderResourceViewRHIRef RootBarycentricCoordinatesSRV = HairStrandsBuffer->RootBarycentricCoordinatesBuffer.SRV;
+			const FHairStrandsRootResource::FMeshProjectionLOD* MeshProjection = (HasRootAttachedValue == 1) ? &(HairStrandsBuffer->SourceRootResources->MeshProjectionLODs[0]) : nullptr;
 
-			int32 HasRootAttachedValue = 0;
-			int32 bNeedSimReset = (ProxyData->TickCount <= ProxyData->ResetTick ? 1 : 0);
-			FVector RestRootOffsetValue = FVector(0, 0, 0);
-			FVector DeformedRootOffsetValue = FVector(0, 0, 0);
-			FVector DeformedPositionOffsetValue = HairStrandsBuffer->SourceRestResources->PositionOffset;
-			FVector RestPositionOffsetValue = HairStrandsBuffer->SourceRestResources->PositionOffset;
-
-			//UE_LOG(LogHairStrands, Log, TEXT("Need Reset = %d %d %d"), bNeedSimReset, ProxyData->TickCount, ProxyData->ResetTick);
-
-			if (HairStrandsBuffer->SourceRootResources != nullptr && HairStrandsBuffer->SourceRootResources->MeshProjectionLODs.Num() > 0)
-			{
-				RestTrianglePositionASRV = HairStrandsBuffer->SourceRootResources->MeshProjectionLODs[0].RestRootTrianglePosition0Buffer.SRV;
-				RestTrianglePositionBSRV = HairStrandsBuffer->SourceRootResources->MeshProjectionLODs[0].RestRootTrianglePosition1Buffer.SRV;
-				RestTrianglePositionCSRV = HairStrandsBuffer->SourceRootResources->MeshProjectionLODs[0].RestRootTrianglePosition2Buffer.SRV;
+			FShaderResourceViewRHIRef RestTrianglePositionASRV = (HasRootAttachedValue == 1) ?
+				MeshProjection->RestRootTrianglePosition0Buffer.SRV : FNiagaraRenderer::GetDummyFloatBuffer().SRV;
+			FShaderResourceViewRHIRef RestTrianglePositionBSRV = (HasRootAttachedValue == 1) ?
+				MeshProjection->RestRootTrianglePosition1Buffer.SRV : FNiagaraRenderer::GetDummyFloatBuffer().SRV;
+			FShaderResourceViewRHIRef RestTrianglePositionCSRV = (HasRootAttachedValue == 1) ?
+				MeshProjection->RestRootTrianglePosition2Buffer.SRV : FNiagaraRenderer::GetDummyFloatBuffer().SRV;
 
-				DeformedTrianglePositionASRV = HairStrandsBuffer->SourceRootResources->MeshProjectionLODs[0].DeformedRootTrianglePosition0Buffer.SRV;
-				DeformedTrianglePositionBSRV = HairStrandsBuffer->SourceRootResources->MeshProjectionLODs[0].DeformedRootTrianglePosition1Buffer.SRV;
-				DeformedTrianglePositionCSRV = HairStrandsBuffer->SourceRootResources->MeshProjectionLODs[0].DeformedRootTrianglePosition2Buffer.SRV;
+			FShaderResourceViewRHIRef DeformedTrianglePositionASRV = (HasRootAttachedValue == 1) ?
+				MeshProjection->DeformedRootTrianglePosition0Buffer.SRV : FNiagaraRenderer::GetDummyFloatBuffer().SRV;
+			FShaderResourceViewRHIRef DeformedTrianglePositionBSRV = (HasRootAttachedValue == 1) ?
+				MeshProjection->DeformedRootTrianglePosition1Buffer.SRV : FNiagaraRenderer::GetDummyFloatBuffer().SRV;
+			FShaderResourceViewRHIRef DeformedTrianglePositionCSRV = (HasRootAttachedValue == 1) ?
+				MeshProjection->DeformedRootTrianglePosition2Buffer.SRV : FNiagaraRenderer::GetDummyFloatBuffer().SRV;
 
-				RootBarycentricCoordinatesSRV = HairStrandsBuffer->SourceRootResources->MeshProjectionLODs[0].RootTriangleBarycentricBuffer.SRV;
-
-				HasRootAttachedValue = 1;
-				//NeedSimReset = (HairStrandsBuffer->SourceRootResources->MeshProjectionLODs[0].Status != FHairStrandsRootResource::FMeshProjectionLOD::EStatus::Completed);
-				RestRootOffsetValue = HairStrandsBuffer->SourceRootResources->MeshProjectionLODs[0].RestRootOffset;
-				DeformedRootOffsetValue = ProxyData->HairStrandsBuffer->SourceRootResources->MeshProjectionLODs[0].DeformedRootOffset;
-			}
-			if (HairStrandsBuffer->SourceDeformedResources != nullptr)
-			{
-				DeformedPositionOffsetValue = HairStrandsBuffer->SourceDeformedResources->PositionOffset;
-			}
+			FShaderResourceViewRHIRef RootBarycentricCoordinatesSRV = (HasRootAttachedValue == 1) ? 
+				MeshProjection->RootTriangleBarycentricBuffer.SRV : FNiagaraRenderer::GetDummyFloatBuffer().SRV;
 
-			//UE_LOG(LogHairStrands, Log, TEXT("Root Offset = %f %f %f | %f %f %f | %f %f %f"), RestRootOffsetValue[0], RestRootOffsetValue[1], RestRootOffsetValue[2],
-			//	DeformedRootOffsetValue[0], DeformedRootOffsetValue[1], DeformedRootOffsetValue[2], DeformedPositionOffsetValue[0], DeformedPositionOffsetValue[1], DeformedPositionOffsetValue[2]);
+			int32 bNeedSimReset = (ProxyData->TickCount <= ProxyData->ResetTick ? 1 : 0);
+			FVector RestRootOffsetValue = (HasRootAttachedValue == 1) ? MeshProjection->RestRootOffset : FVector(0, 0, 0);
+			FVector DeformedRootOffsetValue = (HasRootAttachedValue == 1) ? MeshProjection->DeformedRootOffset : FVector(0, 0, 0);
+			FVector DeformedPositionOffsetValue = IsDeformedValid ? HairStrandsBuffer->SourceDeformedResources->PositionOffset :
+																	HairStrandsBuffer->SourceRestResources->PositionOffset;
+			FVector RestPositionOffsetValue = HairStrandsBuffer->SourceRestResources->PositionOffset;
 
 			SetUAVParameter(RHICmdList, ComputeShaderRHI, DeformedPositionBuffer, PointPositionsUAV);
-			RHICmdList.TransitionResource(EResourceTransitionAccess::EWritable, EResourceTransitionPipeline::EComputeToCompute, HairStrandsBuffer->BoundingBoxBuffer.UAV);
+			RHICmdList.TransitionResource(EResourceTransitionAccess::ERWBarrier, EResourceTransitionPipeline::EComputeToCompute, HairStrandsBuffer->BoundingBoxBuffer.UAV);
 			SetUAVParameter(RHICmdList, ComputeShaderRHI, BoundingBoxBuffer, HairStrandsBuffer->BoundingBoxBuffer.UAV);
-			RHICmdList.TransitionResource(EResourceTransitionAccess::EReadable, EResourceTransitionPipeline::EComputeToCompute, HairStrandsBuffer->NodeBoundBuffer.UAV);
+			RHICmdList.TransitionResource(EResourceTransitionAccess::ERWBarrier, EResourceTransitionPipeline::EComputeToCompute, HairStrandsBuffer->NodeBoundBuffer.UAV);
 			SetSRVParameter(RHICmdList, ComputeShaderRHI, NodeBoundBuffer, HairStrandsBuffer->NodeBoundBuffer.SRV);
 			SetSRVParameter(RHICmdList, ComputeShaderRHI, CurvesOffsetsBuffer, HairStrandsBuffer->CurvesOffsetsBuffer.SRV);
 			SetSRVParameter(RHICmdList, ComputeShaderRHI, RestPositionBuffer, HairStrandsBuffer->SourceRestResources->RestPositionBuffer.SRV);
@@ -769,34 +856,40 @@ void FNDIHairStrandsProxy::DeferredDestroy()
 void FNDIHairStrandsProxy::ConsumePerInstanceDataFromGameThread(void* PerInstanceData, const FNiagaraSystemInstanceID& Instance)
 {
 	FNDIHairStrandsData* SourceData = static_cast<FNDIHairStrandsData*>(PerInstanceData);
-	FNDIHairStrandsData& TargetData = SystemInstancesToProxyData.FindOrAdd(Instance);
+	FNDIHairStrandsData* TargetData = &(SystemInstancesToProxyData.FindOrAdd(Instance));
 
-	TargetData.WorldTransform = SourceData->WorldTransform;
-	TargetData.NumStrands = SourceData->NumStrands;
-	TargetData.StrandSize = SourceData->StrandSize;
-	TargetData.BoxCenter = SourceData->BoxCenter;
-	TargetData.BoxExtent = SourceData->BoxExtent;
-	TargetData.TickCount = SourceData->TickCount;
-	TargetData.ResetTick = SourceData->ResetTick;
-	TargetData.ForceReset = SourceData->ForceReset;
-	TargetData.HairStrandsBuffer = SourceData->HairStrandsBuffer;
+	ensure(TargetData);
+	if (TargetData)
+	{
+		TargetData->WorldTransform = SourceData->WorldTransform;
+		TargetData->NumStrands = SourceData->NumStrands;
+		TargetData->StrandSize = SourceData->StrandSize;
+		TargetData->BoxCenter = SourceData->BoxCenter;
+		TargetData->BoxExtent = SourceData->BoxExtent;
+		TargetData->TickCount = SourceData->TickCount;
+		TargetData->ResetTick = SourceData->ResetTick;
+		TargetData->ForceReset = SourceData->ForceReset;
+		TargetData->HairStrandsBuffer = SourceData->HairStrandsBuffer;
+	}
+	else
+	{
+		UE_LOG(LogHairStrands, Log, TEXT("ConsumePerInstanceDataFromGameThread() ... could not find %s"), *FNiagaraUtilities::SystemInstanceIDToString(Instance));
+	}
 }
 
-void FNDIHairStrandsProxy::InitializePerInstanceData(const FNiagaraSystemInstanceID& SystemInstance, FNDIHairStrandsBuffer* HairStrandsBuffer, 
-	const uint32 NumStrands, const uint8 StrandSize, const FVector& BoxCenter, const FVector& BoxExtent, const FTransform& WorldTransform)
+void FNDIHairStrandsProxy::InitializePerInstanceData(const FNiagaraSystemInstanceID& SystemInstance)
 {
 	check(IsInRenderingThread());
 
-	FNDIHairStrandsData& TargetData = SystemInstancesToProxyData.FindOrAdd(SystemInstance);
-	TargetData.HairStrandsBuffer = HairStrandsBuffer;
-	TargetData.NumStrands = NumStrands;
-	TargetData.StrandSize = StrandSize;
-	TargetData.BoxCenter = BoxCenter;
-	TargetData.BoxExtent = BoxExtent;
-	TargetData.TickCount = 0;
-	TargetData.ResetTick = MaxDelay;
-	TargetData.ForceReset = true;
-	TargetData.WorldTransform = WorldTransform;
+	FNDIHairStrandsData* TargetData = SystemInstancesToProxyData.Find(SystemInstance);
+	if (TargetData != nullptr)
+	{
+		DeferredDestroyList.Remove(SystemInstance);
+	}
+	else
+	{
+		TargetData = &SystemInstancesToProxyData.Add(SystemInstance);
+	}
 }
 
 void FNDIHairStrandsProxy::DestroyPerInstanceData(NiagaraEmitterInstanceBatcher* Batcher, const FNiagaraSystemInstanceID& SystemInstance)
@@ -826,7 +919,7 @@ bool UNiagaraDataInterfaceHairStrands::IsComponentValid() const
 	return (SourceComponent.IsValid() && SourceComponent != nullptr);
 }
 
-void UNiagaraDataInterfaceHairStrands::UpdateSourceComponent(FNiagaraSystemInstance* SystemInstance)
+void UNiagaraDataInterfaceHairStrands::ExtractSourceComponent(FNiagaraSystemInstance* SystemInstance)
 {
 	SourceComponent = nullptr;
 	if (SourceActor)
@@ -873,7 +966,7 @@ void UNiagaraDataInterfaceHairStrands::UpdateSourceComponent(FNiagaraSystemInsta
 void UNiagaraDataInterfaceHairStrands::ExtractDatasAndResources(FNiagaraSystemInstance* SystemInstance,FHairStrandsDatas*& OutStrandsDatas, 
 	FHairStrandsRestResource*& OutStrandsRestResource, FHairStrandsDeformedResource*& OutStrandsDeformedResource, FHairStrandsRootResource*& OutStrandsRootResource)
 {
-	UpdateSourceComponent(SystemInstance);
+	ExtractSourceComponent(SystemInstance);
 
 	OutStrandsDatas = nullptr;
 	OutStrandsRestResource = nullptr;
@@ -899,80 +992,35 @@ bool UNiagaraDataInterfaceHairStrands::InitPerInstanceData(void* PerInstanceData
 	FNDIHairStrandsData* InstanceData = new (PerInstanceData) FNDIHairStrandsData();
 	check(InstanceData);
 
-	FHairStrandsDatas* StrandsDatas =  nullptr;
-	FHairStrandsRestResource* StrandsRestResource = nullptr;
-	FHairStrandsDeformedResource* StrandsDeformedResource = nullptr;
-	FHairStrandsRootResource* StrandsRootResource = nullptr;
-
-	ExtractDatasAndResources(SystemInstance, StrandsDatas, StrandsRestResource, StrandsDeformedResource, StrandsRootResource);
-
-	if (StrandsDatas == nullptr || StrandsRestResource == nullptr)
+	/*FNDIHairStrandsProxy* ThisProxy = GetProxyAs<FNDIHairStrandsProxy>();
+	ENQUEUE_RENDER_COMMAND(FNiagaraDIPushInitialInstanceDataToRT) (
+		[ThisProxy, InstanceID = SystemInstance->GetId()](FRHICommandListImmediate& CmdList)
 	{
-		UE_LOG(LogHairStrands, Log, TEXT("Hair Strands data interface has no valid asset. Failed InitPerInstanceData - %s %d"), *GetFullName(), DefaultSource);
-		return false;
+		ThisProxy->InitializePerInstanceData(InstanceID);
 	}
-	else
-	{ 
-		FNDIHairStrandsBuffer* HairStrandsBuffer = new FNDIHairStrandsBuffer;
-		HairStrandsBuffer->SetHairAsset(StrandsDatas, StrandsRestResource, StrandsDeformedResource, StrandsRootResource);
+	);*/
 
-		// Push instance data to RT
-		{
-			const uint32 NumStrands = StrandsDatas->GetNumCurves();
-			const uint8 StrandsSize = static_cast<uint8>(StrandSize);
-			const FBox& StrandsBox = StrandsDatas->BoundingBox;
+	//UE_LOG(LogHairStrands, Warning, TEXT("Init Strand Size : %d"), StrandSize);
 
-			const FVector BoxExtent = StrandsBox.GetExtent();
-			const FVector BoxCenter = StrandsBox.GetCenter(); 
-
-			const FTransform WorldTransform = IsComponentValid() ? SourceComponent->GetComponentToWorld():
-				SystemInstance->GetComponent()->GetComponentToWorld();
 
-			InstanceData->WorldTransform = WorldTransform;
-			InstanceData->StrandSize = StrandsSize;
-			InstanceData->NumStrands = NumStrands;
-			InstanceData->HairStrandsBuffer = HairStrandsBuffer;
-			InstanceData->BoxCenter = BoxCenter;
-			InstanceData->BoxExtent = BoxExtent;
-			InstanceData->TickCount = 0;
-			InstanceData->ForceReset = true;
-			InstanceData->ResetTick = MaxDelay;
-
-			//UE_LOG(LogHairStrands, Log, TEXT("Num Strands = %d | Strand Size = %d | Num Vertices = %d | Min = %f %f %f | Max = %f %f %f | Transform = %s"), NumStrands, StrandSize,
-			//	StrandsDatas->GetNumPoints(), StrandsBox.Min[0], StrandsBox.Min[1], StrandsBox.Min[2], StrandsBox.Max[0], StrandsBox.Max[1], StrandsBox.Max[2], *InstanceData->WorldTransform.ToString());
-
-			FNDIHairStrandsProxy* ThisProxy = GetProxyAs<FNDIHairStrandsProxy>();
-			ENQUEUE_RENDER_COMMAND(FNiagaraDIPushInitialInstanceDataToRT) (
-				[ThisProxy, InstanceID = SystemInstance->GetId(), HairStrandsBuffer, NumStrands, StrandsSize, BoxCenter, BoxExtent, WorldTransform](FRHICommandListImmediate& CmdList)
-			{
-				HairStrandsBuffer->InitResource();
-				ThisProxy->InitializePerInstanceData(InstanceID, HairStrandsBuffer, NumStrands, StrandsSize, BoxCenter, BoxExtent, WorldTransform);
-			}
-			);
-		}
-	}
-	return true;
+	return InstanceData->Init(this, SystemInstance);
 }
 
 void UNiagaraDataInterfaceHairStrands::DestroyPerInstanceData(void* PerInstanceData, FNiagaraSystemInstance* SystemInstance)
 {
 	FNDIHairStrandsData* InstanceData = static_cast<FNDIHairStrandsData*>(PerInstanceData);
+
+	InstanceData->Release();
 	InstanceData->~FNDIHairStrandsData();
 
-	if (InstanceData->HairStrandsBuffer)
+	FNDIHairStrandsProxy* ThisProxy = GetProxyAs<FNDIHairStrandsProxy>();
+	ENQUEUE_RENDER_COMMAND(FNiagaraDIDestroyInstanceData) (
+		[ThisProxy, InstanceID = SystemInstance->GetId(), Batcher = SystemInstance->GetBatcher()](FRHICommandListImmediate& CmdList)
 	{
-		FNDIHairStrandsProxy* ThisProxy = GetProxyAs<FNDIHairStrandsProxy>();
-		FNDIHairStrandsBuffer* HairStrandsBuffer = InstanceData->HairStrandsBuffer;
-		ENQUEUE_RENDER_COMMAND(FNiagaraDIDestroyInstanceData) (
-			[ThisProxy, HairStrandsBuffer, InstanceID = SystemInstance->GetId(), Batcher = SystemInstance->GetBatcher()](FRHICommandListImmediate& CmdList)
-		{
-			HairStrandsBuffer->ReleaseResource();
-			ThisProxy->SystemInstancesToProxyData.Remove(InstanceID);
-			delete HairStrandsBuffer;
-		}
-		);
-		InstanceData->HairStrandsBuffer = nullptr;
+		ThisProxy->SystemInstancesToProxyData.Remove(InstanceID);
+		//ThisProxy->DestroyPerInstanceData(Batcher, InstanceID);
 	}
+	);
 }
 
 bool UNiagaraDataInterfaceHairStrands::PerInstanceTick(void* PerInstanceData, FNiagaraSystemInstance* SystemInstance, float InDeltaSeconds)
@@ -986,7 +1034,7 @@ bool UNiagaraDataInterfaceHairStrands::PerInstanceTick(void* PerInstanceData, FN
 	InstanceData->TickCount = FMath::Min(MaxDelay+1,InstanceData->TickCount+1);
 
 	ExtractDatasAndResources(SystemInstance, StrandsDatas, StrandsRestResource, StrandsDeformedResource, StrandsRootResource);
-	InstanceData->HairStrandsBuffer->SetHairAsset(StrandsDatas, StrandsRestResource, StrandsDeformedResource, StrandsRootResource);
+	InstanceData->HairStrandsBuffer->Initialize(StrandsDatas, StrandsRestResource, StrandsDeformedResource, StrandsRootResource);
 
 	if (SourceComponent != nullptr)
 	{
@@ -994,6 +1042,17 @@ bool UNiagaraDataInterfaceHairStrands::PerInstanceTick(void* PerInstanceData, FN
 		{
 			InstanceData->ResetTick = FMath::Min(MaxDelay, InstanceData->TickCount+1);
 		}
+
+		//if (!InstanceData->ForceReset && !SourceComponent->bResetSimulation)
+		if (!SourceComponent->bResetSimulation)
+		{
+			InstanceData->TickCount = FMath::Min(MaxDelay + 1, InstanceData->TickCount + 1);
+		}
+		else
+		{
+			InstanceData->TickCount = 0;
+		}
+		InstanceData->ResetTick = MaxDelay;
 		//UE_LOG(LogHairStrands, Warning, TEXT("Reset Simulation : %d %d %d %d"), InstanceData->ForceReset, SourceComponent->bResetSimulation, InstanceData->TickCount, InstanceData->ResetTick);
 		InstanceData->ForceReset = SourceComponent->bResetSimulation;
 	}
@@ -1402,6 +1461,7 @@ void UNiagaraDataInterfaceHairStrands::GetFunctions(TArray<FNiagaraFunctionSigna
 		Sig.bMemberFunction = true;
 		Sig.bRequiresContext = false;
 		Sig.Inputs.Add(FNiagaraVariable(FNiagaraTypeDefinition(GetClass()), TEXT("Hair Strands")));
+		Sig.Inputs.Add(FNiagaraVariable(FNiagaraTypeDefinition::GetBoolDef(), TEXT("Enable Constraint")));
 		Sig.Inputs.Add(FNiagaraVariable(FNiagaraTypeDefinition::GetFloatDef(), TEXT("Rest Length")));
 		Sig.Inputs.Add(FNiagaraVariable(FNiagaraTypeDefinition::GetFloatDef(), TEXT("Delta Time")));
 		Sig.Inputs.Add(FNiagaraVariable(FNiagaraTypeDefinition::GetIntDef(), TEXT("Node Offset")));
@@ -1451,6 +1511,7 @@ void UNiagaraDataInterfaceHairStrands::GetFunctions(TArray<FNiagaraFunctionSigna
 		Sig.bMemberFunction = true;
 		Sig.bRequiresContext = false;
 		Sig.Inputs.Add(FNiagaraVariable(FNiagaraTypeDefinition(GetClass()), TEXT("Hair Strands")));
+		Sig.Inputs.Add(FNiagaraVariable(FNiagaraTypeDefinition::GetBoolDef(), TEXT("Enable Constraint")));
 		Sig.Inputs.Add(FNiagaraVariable(FNiagaraTypeDefinition::GetFloatDef(), TEXT("Rest Length")));
 		Sig.Inputs.Add(FNiagaraVariable(FNiagaraTypeDefinition::GetVec3Def(), TEXT("Rest Direction")));
 		Sig.Inputs.Add(FNiagaraVariable(FNiagaraTypeDefinition::GetFloatDef(), TEXT("Delta Time")));
@@ -1578,6 +1639,7 @@ void UNiagaraDataInterfaceHairStrands::GetFunctions(TArray<FNiagaraFunctionSigna
 		Sig.bMemberFunction = true;
 		Sig.bRequiresContext = false;
 		Sig.Inputs.Add(FNiagaraVariable(FNiagaraTypeDefinition(GetClass()), TEXT("Hair Strands")));
+		Sig.Inputs.Add(FNiagaraVariable(FNiagaraTypeDefinition::GetBoolDef(), TEXT("Enable Constraint")));
 		Sig.Inputs.Add(FNiagaraVariable(FNiagaraTypeDefinition::GetFloatDef(), TEXT("Penetration Depth")));
 		Sig.Inputs.Add(FNiagaraVariable(FNiagaraTypeDefinition::GetVec3Def(), TEXT("Collision Position")));
 		Sig.Inputs.Add(FNiagaraVariable(FNiagaraTypeDefinition::GetVec3Def(), TEXT("Collision Velocity")));
@@ -1609,6 +1671,64 @@ void UNiagaraDataInterfaceHairStrands::GetFunctions(TArray<FNiagaraFunctionSigna
 
 		OutFunctions.Add(Sig);
 	}
+	{
+		FNiagaraFunctionSignature Sig;
+		Sig.Name = SetupSoftCollisionConstraintName;
+		Sig.bMemberFunction = true;
+		Sig.bRequiresContext = false;
+		Sig.Inputs.Add(FNiagaraVariable(FNiagaraTypeDefinition(GetClass()), TEXT("Hair Strands")));
+		Sig.Inputs.Add(FNiagaraVariable(FNiagaraTypeDefinition::GetFloatDef(), TEXT("Collision Stiffness")));
+		Sig.Inputs.Add(FNiagaraVariable(FNiagaraTypeDefinition::GetFloatDef(), TEXT("Delta Time")));
+		Sig.Inputs.Add(FNiagaraVariable(FNiagaraTypeDefinition::GetFloatDef(), TEXT("Material Damping")));
+		Sig.Outputs.Add(FNiagaraVariable(FNiagaraTypeDefinition::GetFloatDef(), TEXT("Material Compliance")));
+		Sig.Outputs.Add(FNiagaraVariable(FNiagaraTypeDefinition::GetFloatDef(), TEXT("Material Weight")));
+		Sig.Outputs.Add(FNiagaraVariable(FNiagaraTypeDefinition::GetVec3Def(), TEXT("Material Multiplier")));
+
+		OutFunctions.Add(Sig);
+	}
+	{
+		FNiagaraFunctionSignature Sig;
+		Sig.Name = SolveSoftCollisionConstraintName;
+		Sig.bMemberFunction = true;
+		Sig.bRequiresContext = false;
+		Sig.Inputs.Add(FNiagaraVariable(FNiagaraTypeDefinition(GetClass()), TEXT("Hair Strands")));
+		Sig.Inputs.Add(FNiagaraVariable(FNiagaraTypeDefinition::GetBoolDef(), TEXT("Enable Constraint")));
+		Sig.Inputs.Add(FNiagaraVariable(FNiagaraTypeDefinition::GetFloatDef(), TEXT("Penetration Depth")));
+		Sig.Inputs.Add(FNiagaraVariable(FNiagaraTypeDefinition::GetVec3Def(), TEXT("Collision Position")));
+		Sig.Inputs.Add(FNiagaraVariable(FNiagaraTypeDefinition::GetVec3Def(), TEXT("Collision Velocity")));
+		Sig.Inputs.Add(FNiagaraVariable(FNiagaraTypeDefinition::GetVec3Def(), TEXT("Collision Normal")));
+		Sig.Inputs.Add(FNiagaraVariable(FNiagaraTypeDefinition::GetFloatDef(), TEXT("Static Friction")));
+		Sig.Inputs.Add(FNiagaraVariable(FNiagaraTypeDefinition::GetFloatDef(), TEXT("Kinetic Friction")));
+		Sig.Inputs.Add(FNiagaraVariable(FNiagaraTypeDefinition::GetFloatDef(), TEXT("Delta Time")));
+		Sig.Inputs.Add(FNiagaraVariable(FNiagaraTypeDefinition::GetFloatDef(), TEXT("Material Damping")));
+		Sig.Inputs.Add(FNiagaraVariable(FNiagaraTypeDefinition::GetFloatDef(), TEXT("Material Compliance")));
+		Sig.Inputs.Add(FNiagaraVariable(FNiagaraTypeDefinition::GetFloatDef(), TEXT("Material Weight")));
+		Sig.Inputs.Add(FNiagaraVariable(FNiagaraTypeDefinition::GetVec3Def(), TEXT("Material Multiplier")));
+		Sig.Outputs.Add(FNiagaraVariable(FNiagaraTypeDefinition::GetVec3Def(), TEXT("Material Multiplier")));
+
+		OutFunctions.Add(Sig);
+	}
+	{
+		FNiagaraFunctionSignature Sig;
+		Sig.Name = ProjectSoftCollisionConstraintName;
+		Sig.bMemberFunction = true;
+		Sig.bRequiresContext = false;
+		Sig.Inputs.Add(FNiagaraVariable(FNiagaraTypeDefinition(GetClass()), TEXT("Hair Strands")));
+		Sig.Inputs.Add(FNiagaraVariable(FNiagaraTypeDefinition::GetFloatDef(), TEXT("Collision Stiffness")));
+		Sig.Inputs.Add(FNiagaraVariable(FNiagaraTypeDefinition::GetFloatDef(), TEXT("Penetration Depth")));
+		Sig.Inputs.Add(FNiagaraVariable(FNiagaraTypeDefinition::GetVec3Def(), TEXT("Collision Position")));
+		Sig.Inputs.Add(FNiagaraVariable(FNiagaraTypeDefinition::GetVec3Def(), TEXT("Collision Velocity")));
+		Sig.Inputs.Add(FNiagaraVariable(FNiagaraTypeDefinition::GetVec3Def(), TEXT("Collision Normal")));
+		Sig.Inputs.Add(FNiagaraVariable(FNiagaraTypeDefinition::GetFloatDef(), TEXT("Static Friction")));
+		Sig.Inputs.Add(FNiagaraVariable(FNiagaraTypeDefinition::GetFloatDef(), TEXT("Kinetic Friction")));
+		Sig.Inputs.Add(FNiagaraVariable(FNiagaraTypeDefinition::GetFloatDef(), TEXT("Node Mass")));
+		Sig.Inputs.Add(FNiagaraVariable(FNiagaraTypeDefinition::GetFloatDef(), TEXT("Delta Time")));
+		Sig.Inputs.Add(FNiagaraVariable(FNiagaraTypeDefinition::GetVec3Def(), TEXT("Previous Position")));
+		Sig.Inputs.Add(FNiagaraVariable(FNiagaraTypeDefinition::GetVec3Def(), TEXT("Node Position")));
+		Sig.Outputs.Add(FNiagaraVariable(FNiagaraTypeDefinition::GetVec3Def(), TEXT("Node Position")));
+
+		OutFunctions.Add(Sig);
+	}
 	{
 		FNiagaraFunctionSignature Sig;
 		Sig.Name = ComputeRestDirectionName;
@@ -1709,6 +1829,10 @@ DEFINE_NDI_DIRECT_FUNC_BINDER(UNiagaraDataInterfaceHairStrands, ProjectBendRodMa
 DEFINE_NDI_DIRECT_FUNC_BINDER(UNiagaraDataInterfaceHairStrands, SolveStaticCollisionConstraint);
 DEFINE_NDI_DIRECT_FUNC_BINDER(UNiagaraDataInterfaceHairStrands, ProjectStaticCollisionConstraint);
 
+DEFINE_NDI_DIRECT_FUNC_BINDER(UNiagaraDataInterfaceHairStrands, SetupSoftCollisionConstraint);
+DEFINE_NDI_DIRECT_FUNC_BINDER(UNiagaraDataInterfaceHairStrands, ProjectSoftCollisionConstraint);
+DEFINE_NDI_DIRECT_FUNC_BINDER(UNiagaraDataInterfaceHairStrands, SolveSoftCollisionConstraint);
+
 DEFINE_NDI_DIRECT_FUNC_BINDER(UNiagaraDataInterfaceHairStrands, ComputeRestDirection);
 DEFINE_NDI_DIRECT_FUNC_BINDER(UNiagaraDataInterfaceHairStrands, UpdateNodeOrientation);
 DEFINE_NDI_DIRECT_FUNC_BINDER(UNiagaraDataInterfaceHairStrands, ComputeAirDragForce);
@@ -1853,7 +1977,7 @@ void UNiagaraDataInterfaceHairStrands::GetVMExternalFunction(const FVMExternalFu
 	}
 	else if (BindingInfo.Name == SolveDistanceSpringMaterialName)
 	{
-		check(BindingInfo.GetNumInputs() == 8 && BindingInfo.GetNumOutputs() == 1);
+		check(BindingInfo.GetNumInputs() == 9 && BindingInfo.GetNumOutputs() == 1);
 		NDI_FUNC_BINDER(UNiagaraDataInterfaceHairStrands, SolveDistanceSpringMaterial)::Bind(this, OutFunc);
 	}
 	else if (BindingInfo.Name == ProjectDistanceSpringMaterialName)
@@ -1868,7 +1992,7 @@ void UNiagaraDataInterfaceHairStrands::GetVMExternalFunction(const FVMExternalFu
 	}
 	else if (BindingInfo.Name == SolveAngularSpringMaterialName)
 	{
-		check(BindingInfo.GetNumInputs() == 12 && BindingInfo.GetNumOutputs() == 3);
+		check(BindingInfo.GetNumInputs() == 13 && BindingInfo.GetNumOutputs() == 3);
 		NDI_FUNC_BINDER(UNiagaraDataInterfaceHairStrands, SolveAngularSpringMaterial)::Bind(this, OutFunc);
 	}
 	else if (BindingInfo.Name == ProjectAngularSpringMaterialName)
@@ -1908,7 +2032,7 @@ void UNiagaraDataInterfaceHairStrands::GetVMExternalFunction(const FVMExternalFu
 	}
 	else if (BindingInfo.Name == SolveStaticCollisionConstraintName)
 	{
-		check(BindingInfo.GetNumInputs() == 17 && BindingInfo.GetNumOutputs() == 3);
+		check(BindingInfo.GetNumInputs() == 18 && BindingInfo.GetNumOutputs() == 3);
 		NDI_FUNC_BINDER(UNiagaraDataInterfaceHairStrands, SolveStaticCollisionConstraint)::Bind(this, OutFunc);
 	}
 	else if (BindingInfo.Name == ProjectStaticCollisionConstraintName)
@@ -1916,6 +2040,21 @@ void UNiagaraDataInterfaceHairStrands::GetVMExternalFunction(const FVMExternalFu
 		check(BindingInfo.GetNumInputs() == 20 && BindingInfo.GetNumOutputs() == 3);
 		NDI_FUNC_BINDER(UNiagaraDataInterfaceHairStrands, ProjectStaticCollisionConstraint)::Bind(this, OutFunc);
 	}
+	else if (BindingInfo.Name == SolveSoftCollisionConstraintName)
+	{
+		check(BindingInfo.GetNumInputs() == 21 && BindingInfo.GetNumOutputs() == 3);
+		NDI_FUNC_BINDER(UNiagaraDataInterfaceHairStrands, SolveSoftCollisionConstraint)::Bind(this, OutFunc);
+	}
+	else if (BindingInfo.Name == ProjectSoftCollisionConstraintName)
+	{
+		check(BindingInfo.GetNumInputs() == 22 && BindingInfo.GetNumOutputs() == 3);
+		NDI_FUNC_BINDER(UNiagaraDataInterfaceHairStrands, ProjectSoftCollisionConstraint)::Bind(this, OutFunc);
+	}
+	else if (BindingInfo.Name == SetupSoftCollisionConstraintName)
+	{
+		check(BindingInfo.GetNumInputs() == 4 && BindingInfo.GetNumOutputs() == 5);
+		NDI_FUNC_BINDER(UNiagaraDataInterfaceHairStrands, SetupSoftCollisionConstraint)::Bind(this, OutFunc);
+	}
 	else if (BindingInfo.Name == ComputeRestDirectionName)
 	{
 		check(BindingInfo.GetNumInputs() == 8 && BindingInfo.GetNumOutputs() == 3);
@@ -2225,6 +2364,20 @@ void UNiagaraDataInterfaceHairStrands::ProjectStaticCollisionConstraint(FVectorV
 	// @todo : implement function for cpu 
 }
 
+void UNiagaraDataInterfaceHairStrands::SolveSoftCollisionConstraint(FVectorVMContext& Context)
+{
+	// @todo : implement function for cpu 
+}
+
+void UNiagaraDataInterfaceHairStrands::ProjectSoftCollisionConstraint(FVectorVMContext& Context)
+{
+	// @todo : implement function for cpu 
+}
+void UNiagaraDataInterfaceHairStrands::SetupSoftCollisionConstraint(FVectorVMContext& Context)
+{
+	// @todo : implement function for cpu 
+}
+
 void UNiagaraDataInterfaceHairStrands::ComputeAirDragForce(FVectorVMContext& Context)
 {
 	// @todo : implement function for cpu 
@@ -2359,7 +2512,15 @@ bool UNiagaraDataInterfaceHairStrands::GetFunctionHLSL(const FNiagaraDataInterfa
 		static const TCHAR *FormatSample = TEXT(R"(
 			void {InstanceFunctionName} (in int NodeIndex, in float3 NodePosition, in int NodeOffset, out float OutEdgeLength)
 			{
-				{HairStrandsContextName} DIHairStrands_ComputeEdgeLength(DIContext,NodeIndex,NodePosition,NodeOffset,OutEdgeLength);
+				{HairStrandsContextName} 
+				if(NodeOffset == 2)
+				{
+					DIHairStrands_ComputeEdgeVolume(DIContext,NodeIndex,NodePosition,OutEdgeLength);
+				}
+				else
+				{
+					DIHairStrands_ComputeEdgeLength(DIContext,NodeIndex,NodePosition,NodeOffset,OutEdgeLength);
+				}
 			}
 			)");
 		OutHLSL += FString::Format(FormatSample, ArgsSample);
@@ -2564,7 +2725,8 @@ in float RestLength, in float DeltaTime, in int NodeOffset, in float MaterialDam
 					}
 					else if( NodeOffset == 2)
 					{
-						SetupTwistSpringMaterial(DIContext.StrandSize,YoungModulus,RodThickness,RestLength,DeltaTime,false,MaterialDamping,OutMaterialCompliance,OutMaterialWeight,OutMaterialMultiplier);
+						//SetupTwistSpringMaterial(DIContext.StrandSize,YoungModulus,RodThickness,RestLength,DeltaTime,false,MaterialDamping,OutMaterialCompliance,OutMaterialWeight,OutMaterialMultiplier);
+						SetupVolumeSpringMaterial(DIContext.StrandSize,YoungModulus,RestLength,DeltaTime,false,MaterialDamping,OutMaterialCompliance,OutMaterialWeight,OutMaterialMultiplier);
 					}
 				}
 				)");
@@ -2574,21 +2736,22 @@ in float RestLength, in float DeltaTime, in int NodeOffset, in float MaterialDam
 	else if (FunctionInfo.DefinitionName == SolveDistanceSpringMaterialName)
 	{
 		static const TCHAR *FormatSample = TEXT(R"(
-				void {InstanceFunctionName} (in float RestLength, in float DeltaTime, in int NodeOffset, in float MaterialDamping, 
+				void {InstanceFunctionName} (in bool EnableConstraint, in float RestLength, in float DeltaTime, in int NodeOffset, in float MaterialDamping, 
 		in float MaterialCompliance, in float MaterialWeight, in float MaterialMultiplier, out float OutMaterialMultiplier)
 				{
 					{HairStrandsContextName} 
 					if(NodeOffset == 0)
 					{
-						SolveStretchSpringMaterial(DIContext.StrandSize,RestLength,DeltaTime,MaterialDamping,MaterialCompliance,MaterialWeight,MaterialMultiplier,OutMaterialMultiplier);
+						SolveStretchSpringMaterial(EnableConstraint,DIContext.StrandSize,RestLength,DeltaTime,MaterialDamping,MaterialCompliance,MaterialWeight,MaterialMultiplier,OutMaterialMultiplier);
 					}
 					else if(NodeOffset == 1)
 					{
-						SolveBendSpringMaterial(DIContext.StrandSize,RestLength,DeltaTime,MaterialDamping,MaterialCompliance,MaterialWeight,MaterialMultiplier,OutMaterialMultiplier);
+						SolveBendSpringMaterial(EnableConstraint,DIContext.StrandSize,RestLength,DeltaTime,MaterialDamping,MaterialCompliance,MaterialWeight,MaterialMultiplier,OutMaterialMultiplier);
 					}
 					else if(NodeOffset == 2)
 					{
-						SolveTwistSpringMaterial(DIContext.StrandSize,RestLength,DeltaTime,MaterialDamping,MaterialCompliance,MaterialWeight,MaterialMultiplier,OutMaterialMultiplier);
+						//SolveTwistSpringMaterial(EnableConstraint,DIContext.StrandSize,RestLength,DeltaTime,MaterialDamping,MaterialCompliance,MaterialWeight,MaterialMultiplier,OutMaterialMultiplier);
+						SolveVolumeSpringMaterial(EnableConstraint,DIContext.StrandSize,RestLength,DeltaTime,MaterialDamping,MaterialCompliance,MaterialWeight,MaterialMultiplier,OutMaterialMultiplier);
 					}
 				}
 				)");
@@ -2633,10 +2796,10 @@ in float RestLength, in float DeltaTime, in int NodeOffset, in float MaterialDam
 	else if (FunctionInfo.DefinitionName == SolveAngularSpringMaterialName)
 	{
 		static const TCHAR *FormatSample = TEXT(R"(
-					void {InstanceFunctionName} (in float RestLength, in float3 RestDirection, in float DeltaTime, in float MaterialDamping, 
+					void {InstanceFunctionName} (in bool EnableConstraint, in float RestLength, in float3 RestDirection, in float DeltaTime, in float MaterialDamping, 
 			in float MaterialCompliance, in float MaterialWeight, in float3 MaterialMultiplier, out float3 OutMaterialMultiplier)
 					{
-						{HairStrandsContextName} SolveAngularSpringMaterial(DIContext.StrandSize,RestLength, RestDirection,DeltaTime,MaterialDamping,MaterialCompliance,MaterialWeight,MaterialMultiplier,OutMaterialMultiplier);
+						{HairStrandsContextName} SolveAngularSpringMaterial(EnableConstraint,DIContext.StrandSize,RestLength, RestDirection,DeltaTime,MaterialDamping,MaterialCompliance,MaterialWeight,MaterialMultiplier,OutMaterialMultiplier);
 					}
 					)");
 		OutHLSL += FString::Format(FormatSample, ArgsSample);
@@ -2726,11 +2889,11 @@ in float RestLength, in float DeltaTime, in int NodeOffset, in float MaterialDam
 	else if (FunctionInfo.DefinitionName == SolveStaticCollisionConstraintName)
 	{
 		static const TCHAR *FormatSample = TEXT(R"(
-							void {InstanceFunctionName} (in float PenetrationDepth, in float3 CollisionPosition, in float3 CollisionVelocity, in float3 CollisionNormal, 
+							void {InstanceFunctionName} (in bool EnableConstraint, in float PenetrationDepth, in float3 CollisionPosition, in float3 CollisionVelocity, in float3 CollisionNormal, 
 				in float StaticFriction, in float KineticFriction, in float DeltaTime, in float3 PreviousPosition, out float3 OutMaterialMultiplier )
 							{
 								OutMaterialMultiplier = float3(0,0,0);
-								{HairStrandsContextName} SolveStaticCollisionConstraint(DIContext.StrandSize,PenetrationDepth,
+								{HairStrandsContextName} SolveStaticCollisionConstraint(EnableConstraint,DIContext.StrandSize,PenetrationDepth,
 									CollisionPosition,CollisionVelocity,CollisionNormal,StaticFriction,KineticFriction,DeltaTime,false,PreviousPosition,SharedNodePosition[GGroupThreadId.x]);
 							}
 							)");
@@ -2744,13 +2907,52 @@ in float RestLength, in float DeltaTime, in int NodeOffset, in float MaterialDam
 			in float StaticFriction, in float KineticFriction, in float DeltaTime, in float3 PreviousPosition, in float3 NodePosition, out float3 OutNodePosition )
 						{
 							OutNodePosition = NodePosition;
-							{HairStrandsContextName} SolveStaticCollisionConstraint(DIContext.StrandSize,PenetrationDepth,
+							{HairStrandsContextName} SolveStaticCollisionConstraint(true,DIContext.StrandSize,PenetrationDepth,
 								CollisionPosition,CollisionVelocity,CollisionNormal,StaticFriction,KineticFriction,DeltaTime,true,PreviousPosition,OutNodePosition);
 						}
 						)");
 		OutHLSL += FString::Format(FormatSample, ArgsSample);
 		return true;
 	}
+	else if (FunctionInfo.DefinitionName == SolveSoftCollisionConstraintName)
+	{
+		static const TCHAR *FormatSample = TEXT(R"(
+								void {InstanceFunctionName} (in bool EnableConstraint, in float PenetrationDepth, in float3 CollisionPosition, in float3 CollisionVelocity, in float3 CollisionNormal, 
+					in float StaticFriction, in float KineticFriction, in float DeltaTime, in float MaterialDamping, 
+			in float MaterialCompliance, in float MaterialWeight, in float3 MaterialMultiplier, out float3 OutMaterialMultiplier )
+								{
+									{HairStrandsContextName} SolveSoftCollisionConstraint(EnableConstraint,DIContext.StrandSize,PenetrationDepth,
+										CollisionPosition,CollisionVelocity,CollisionNormal,StaticFriction,KineticFriction,DeltaTime,MaterialDamping,
+											MaterialCompliance,MaterialWeight,MaterialMultiplier,OutMaterialMultiplier);
+								}
+								)");
+		OutHLSL += FString::Format(FormatSample, ArgsSample);
+		return true;
+	}
+	else if (FunctionInfo.DefinitionName == ProjectSoftCollisionConstraintName)
+	{
+		static const TCHAR *FormatSample = TEXT(R"(
+							void {InstanceFunctionName} (in float ConstraintStiffness, in float PenetrationDepth, in float3 CollisionPosition, in float3 CollisionVelocity, in float3 CollisionNormal, 
+					in float StaticFriction, in float KineticFriction, in float NodeMass, in float DeltaTime, in float3 PreviousPosition, in float3 NodePosition, out float3 OutNodePosition  )
+							{
+								{HairStrandsContextName} ProjectSoftCollisionConstraint(DIContext.StrandSize,ConstraintStiffness,PenetrationDepth,
+									CollisionPosition,CollisionVelocity,CollisionNormal,StaticFriction,KineticFriction,NodeMass,DeltaTime,PreviousPosition,NodePosition,OutNodePosition);
+							}
+							)");
+		OutHLSL += FString::Format(FormatSample, ArgsSample);
+		return true;
+	}
+	else if (FunctionInfo.DefinitionName == SetupSoftCollisionConstraintName)
+	{
+		static const TCHAR *FormatSample = TEXT(R"(
+								void {InstanceFunctionName} (in float ConstraintStiffness, in float DeltaTime, in float MaterialDamping, out float OutMaterialCompliance, out float OutMaterialWeight, out float3 OutMaterialMultiplier )
+								{
+									{HairStrandsContextName} SetupSoftCollisionConstraint(DIContext.StrandSize,ConstraintStiffness,DeltaTime,MaterialDamping,OutMaterialCompliance,OutMaterialWeight,OutMaterialMultiplier);
+								}
+								)");
+		OutHLSL += FString::Format(FormatSample, ArgsSample);
+		return true;
+	}
 	else if (FunctionInfo.DefinitionName == ComputeRestDirectionName)
 	{
 		static const TCHAR *FormatSample = TEXT(R"(
@@ -2842,6 +3044,11 @@ UNiagaraDataInterfaceHairStrands::ConstructComputeParameters() const
 }
 
 void FNDIHairStrandsProxy::PreStage(FRHICommandList& RHICmdList, const FNiagaraDataInterfaceSetArgs& Context)
+{
+}
+	
+
+void FNDIHairStrandsProxy::PostStage(FRHICommandList& RHICmdList, const FNiagaraDataInterfaceSetArgs& Context)
 {
 	if (Context.ShaderStageIndex == 0)
 	{
@@ -2850,14 +3057,11 @@ void FNDIHairStrandsProxy::PreStage(FRHICommandList& RHICmdList, const FNiagaraD
 
 		if (ProxyData != nullptr && ProxyData->HairStrandsBuffer != nullptr)
 		{
-			ProxyData->HairStrandsBuffer->ClearBuffer(RHICmdList);
+			ClearBuffer(RHICmdList, ProxyData->HairStrandsBuffer);
 		}
 	}
 }
 
-void FNDIHairStrandsProxy::PostStage(FRHICommandList& RHICmdList, const FNiagaraDataInterfaceSetArgs& Context)
-{}
-
 void FNDIHairStrandsProxy::ResetData(FRHICommandList& RHICmdList, const FNiagaraDataInterfaceSetArgs& Context)
 {}
 
diff --git a/Engine/Plugins/Experimental/HairStrands/Source/HairStrandsNiagara/Private/NiagaraDataInterfacePhysicsAsset.cpp b/Engine/Plugins/Experimental/HairStrands/Source/HairStrandsNiagara/Private/NiagaraDataInterfacePhysicsAsset.cpp
--- a/Engine/Plugins/Experimental/HairStrands/Source/HairStrandsNiagara/Private/NiagaraDataInterfacePhysicsAsset.cpp
+++ b/Engine/Plugins/Experimental/HairStrands/Source/HairStrandsNiagara/Private/NiagaraDataInterfacePhysicsAsset.cpp
@@ -33,6 +33,7 @@ const FString UNiagaraDataInterfacePhysicsAsset::ElementOffsetsName(TEXT("Elemen
 
 const FString UNiagaraDataInterfacePhysicsAsset::CurrentTransformBufferName(TEXT("CurrentTransformBuffer_"));
 const FString UNiagaraDataInterfacePhysicsAsset::PreviousTransformBufferName(TEXT("PreviousTransformBuffer_"));
+const FString UNiagaraDataInterfacePhysicsAsset::PreviousInverseBufferName(TEXT("PreviousInverseBuffer_"));
 const FString UNiagaraDataInterfacePhysicsAsset::InverseTransformBufferName(TEXT("InverseTransformBuffer_"));
 const FString UNiagaraDataInterfacePhysicsAsset::RestTransformBufferName(TEXT("RestTransformBuffer_"));
 const FString UNiagaraDataInterfacePhysicsAsset::RestInverseBufferName(TEXT("RestInverseBuffer_"));
@@ -51,6 +52,7 @@ struct FNDIPhysicsAssetParametersName
 
 		CurrentTransformBufferName = UNiagaraDataInterfacePhysicsAsset::CurrentTransformBufferName + Suffix;
 		PreviousTransformBufferName = UNiagaraDataInterfacePhysicsAsset::PreviousTransformBufferName + Suffix;
+		PreviousInverseBufferName = UNiagaraDataInterfacePhysicsAsset::PreviousInverseBufferName + Suffix;
 		InverseTransformBufferName = UNiagaraDataInterfacePhysicsAsset::InverseTransformBufferName + Suffix;
 		RestTransformBufferName = UNiagaraDataInterfacePhysicsAsset::RestTransformBufferName + Suffix;
 		RestInverseBufferName = UNiagaraDataInterfacePhysicsAsset::RestInverseBufferName + Suffix;
@@ -64,6 +66,7 @@ struct FNDIPhysicsAssetParametersName
 	
 	FString CurrentTransformBufferName;
 	FString PreviousTransformBufferName;
+	FString PreviousInverseBufferName;
 	FString InverseTransformBufferName;
 	FString RestTransformBufferName;
 	FString RestInverseBufferName;
@@ -154,6 +157,7 @@ void CreateInternalArrays(const TWeakObjectPtr<UPhysicsAsset> PhysicsAsset, cons
 				OutAssetArrays->RestInverse.SetNum(NumTransforms);
 				OutAssetArrays->RestTransform.SetNum(NumTransforms);
 				OutAssetArrays->PreviousTransform.SetNum(NumTransforms);
+				OutAssetArrays->PreviousInverse.SetNum(NumTransforms);
 				OutAssetArrays->ElementExtent.SetNum(NumExtents);
 
 				uint32 ElementCount = 0;
@@ -201,6 +205,7 @@ void CreateInternalArrays(const TWeakObjectPtr<UPhysicsAsset> PhysicsAsset, cons
 					}
 				}
 				OutAssetArrays->PreviousTransform = OutAssetArrays->CurrentTransform;
+				OutAssetArrays->PreviousInverse = OutAssetArrays->InverseTransform;
 			}
 		}
 	}
@@ -226,6 +231,7 @@ void UpdateInternalArrays(const TWeakObjectPtr<UPhysicsAsset> PhysicsAsset, cons
 			const bool bHasMasterPoseComponent = (SkeletalMesh != nullptr) && SkeletalMesh->MasterPoseComponent.IsValid();
 
 			OutAssetArrays->PreviousTransform = OutAssetArrays->CurrentTransform;
+			OutAssetArrays->PreviousInverse= OutAssetArrays->InverseTransform;
 
 			uint32 ElementCount = 0;
 			for (const UBodySetup* BodySetup : PhysicsAsset->SkeletalBodySetups)
@@ -266,7 +272,7 @@ bool FNDIPhysicsAssetBuffer::IsValid() const
 	return (PhysicsAsset.IsValid() && PhysicsAsset.Get() != nullptr) && (AssetArrays.IsValid() && AssetArrays.Get() != nullptr);
 }
 
-void FNDIPhysicsAssetBuffer::SetupArrays(const TWeakObjectPtr<UPhysicsAsset> InPhysicsAsset, const TWeakObjectPtr<USkeletalMeshComponent> InSkeletalMesh, const FTransform& InWorldTransform)
+void FNDIPhysicsAssetBuffer::Initialize(const TWeakObjectPtr<UPhysicsAsset> InPhysicsAsset, const TWeakObjectPtr<USkeletalMeshComponent> InSkeletalMesh, const FTransform& InWorldTransform)
 {
 	PhysicsAsset = InPhysicsAsset;
 	SkeletalMesh = InSkeletalMesh;
@@ -280,7 +286,7 @@ void FNDIPhysicsAssetBuffer::SetupArrays(const TWeakObjectPtr<UPhysicsAsset> InP
 	}
 }
 
-void FNDIPhysicsAssetBuffer::UpdateBuffers()
+void FNDIPhysicsAssetBuffer::Update()
 {
 	if (IsValid())
 	{
@@ -293,6 +299,7 @@ void FNDIPhysicsAssetBuffer::UpdateBuffers()
 			CreateInternalBuffer<FVector4, FVector4, 1, EPixelFormat::PF_A32B32G32R32F, false>(ThisBuffer->AssetArrays->CurrentTransform.Num(), ThisBuffer->AssetArrays->CurrentTransform, ThisBuffer->CurrentTransformBuffer);
 			CreateInternalBuffer<FVector4, FVector4, 1, EPixelFormat::PF_A32B32G32R32F, false>(ThisBuffer->AssetArrays->PreviousTransform.Num(), ThisBuffer->AssetArrays->PreviousTransform, ThisBuffer->PreviousTransformBuffer);
 			CreateInternalBuffer<FVector4, FVector4, 1, EPixelFormat::PF_A32B32G32R32F, false>(ThisBuffer->AssetArrays->InverseTransform.Num(), ThisBuffer->AssetArrays->InverseTransform, ThisBuffer->InverseTransformBuffer);
+			CreateInternalBuffer<FVector4, FVector4, 1, EPixelFormat::PF_A32B32G32R32F, false>(ThisBuffer->AssetArrays->PreviousInverse.Num(), ThisBuffer->AssetArrays->PreviousInverse, ThisBuffer->PreviousInverseBuffer);
 		}
 		);
 	}
@@ -308,6 +315,7 @@ void FNDIPhysicsAssetBuffer::InitRHI()
 		CreateInternalBuffer<FVector4, FVector4, 1, EPixelFormat::PF_A32B32G32R32F, true>(AssetArrays->RestTransform.Num(), AssetArrays->RestTransform, RestTransformBuffer);
 		CreateInternalBuffer<FVector4, FVector4, 1, EPixelFormat::PF_A32B32G32R32F, true>(AssetArrays->RestInverse.Num(), AssetArrays->RestInverse, RestInverseBuffer);
 		CreateInternalBuffer<FVector4, FVector4, 1, EPixelFormat::PF_A32B32G32R32F, true>(AssetArrays->ElementExtent.Num(), AssetArrays->ElementExtent, ElementExtentBuffer);
+		CreateInternalBuffer<FVector4, FVector4, 1, EPixelFormat::PF_A32B32G32R32F, true>(AssetArrays->PreviousInverse.Num(), AssetArrays->PreviousInverse, PreviousInverseBuffer);
 
 		//UE_LOG(LogPhysicsAsset, Warning, TEXT("Num Capsules = %d | Num Spheres = %d | Num Boxes = %d"), AssetArrays->ElementOffsets.NumElements - AssetArrays->ElementOffsets.CapsuleOffset,
 		//	AssetArrays->ElementOffsets.CapsuleOffset - AssetArrays->ElementOffsets.SphereOffset, AssetArrays->ElementOffsets.SphereOffset - AssetArrays->ElementOffsets.BoxOffset);
@@ -318,6 +326,7 @@ void FNDIPhysicsAssetBuffer::ReleaseRHI()
 {
 	CurrentTransformBuffer.Release();
 	PreviousTransformBuffer.Release();
+	PreviousInverseBuffer.Release();
 	InverseTransformBuffer.Release();
 	RestTransformBuffer.Release();
 	RestInverseBuffer.Release();
@@ -326,6 +335,45 @@ void FNDIPhysicsAssetBuffer::ReleaseRHI()
 
 //------------------------------------------------------------------------------------------------------------
 
+void FNDIPhysicsAssetData::Release()
+{
+	if (PhysicsAssetBuffer)
+	{
+		BeginReleaseResource(PhysicsAssetBuffer);
+		ENQUEUE_RENDER_COMMAND(DeleteResource)(
+			[ParamPointerToRelease = PhysicsAssetBuffer](FRHICommandListImmediate& RHICmdList)
+		{
+			delete ParamPointerToRelease;
+		});
+		PhysicsAssetBuffer = nullptr;
+	}
+}
+
+bool FNDIPhysicsAssetData::Init(UNiagaraDataInterfacePhysicsAsset* Interface, FNiagaraSystemInstance* SystemInstance)
+{
+	PhysicsAssetBuffer = nullptr;
+
+	if (Interface != nullptr)
+	{
+		Interface->ExtractSourceComponent(SystemInstance);
+
+		const FTransform WorldTransform = (Interface->SourceComponent != nullptr) ? Interface->SourceComponent->GetComponentTransform() :
+			SystemInstance->GetComponent()->GetComponentTransform();
+
+		PhysicsAssetBuffer = new FNDIPhysicsAssetBuffer();
+		PhysicsAssetBuffer->Initialize(Interface->PhysicsAsset, Interface->SourceComponent, WorldTransform);
+		BeginInitResource(PhysicsAssetBuffer);
+
+		const bool HasPreviewMesh = Interface->PhysicsAsset != nullptr && Interface->PhysicsAsset->GetPreviewMesh() != nullptr;
+		BoxOrigin = HasPreviewMesh ? Interface->PhysicsAsset->GetPreviewMesh()->GetImportedBounds().Origin : FVector(0, 0, 0);
+		BoxExtent = HasPreviewMesh ? Interface->PhysicsAsset->GetPreviewMesh()->GetImportedBounds().BoxExtent : FVector(1, 1, 1);
+	} 
+
+	return true;
+}
+
+//------------------------------------------------------------------------------------------------------------
+
 struct FNDIPhysicsAssetParametersCS : public FNiagaraDataInterfaceParametersCS
 {
 	virtual void Bind(const FNiagaraDataInterfaceParamRef& ParamRef, const class FShaderParameterMap& ParameterMap) override
@@ -336,6 +384,7 @@ struct FNDIPhysicsAssetParametersCS : public FNiagaraDataInterfaceParametersCS
 
 		CurrentTransformBuffer.Bind(ParameterMap, *ParamNames.CurrentTransformBufferName);
 		PreviousTransformBuffer.Bind(ParameterMap, *ParamNames.PreviousTransformBufferName);
+		PreviousInverseBuffer.Bind(ParameterMap, *ParamNames.PreviousInverseBufferName);
 		InverseTransformBuffer.Bind(ParameterMap, *ParamNames.InverseTransformBufferName);
 		RestTransformBuffer.Bind(ParameterMap, *ParamNames.RestTransformBufferName);
 		RestInverseBuffer.Bind(ParameterMap, *ParamNames.RestInverseBufferName);
@@ -352,6 +401,10 @@ struct FNDIPhysicsAssetParametersCS : public FNiagaraDataInterfaceParametersCS
 		{
 			UE_LOG(LogPhysicsAsset, Warning, TEXT("Binding failed for FNDIPhysicsAssetParametersCS %s. Was it optimized out?"), *ParamNames.PreviousTransformBufferName)
 		}
+		if (!PreviousInverseBuffer.IsBound())
+		{
+			UE_LOG(LogPhysicsAsset, Warning, TEXT("Binding failed for FNDIPhysicsAssetParametersCS %s. Was it optimized out?"), *ParamNames.PreviousInverseBufferName)
+		}
 		if (!InverseTransformBuffer.IsBound())
 		{
 			UE_LOG(LogPhysicsAsset, Warning, TEXT("Binding failed for FNDIPhysicsAssetParametersCS %s. Was it optimized out?"), *ParamNames.InverseTransformBufferName)
@@ -375,6 +428,7 @@ struct FNDIPhysicsAssetParametersCS : public FNiagaraDataInterfaceParametersCS
 		Ar << ElementOffsets;
 		Ar << CurrentTransformBuffer;
 		Ar << PreviousTransformBuffer;
+		Ar << PreviousInverseBuffer;
 		Ar << InverseTransformBuffer;
 		Ar << RestTransformBuffer;
 		Ar << RestInverseBuffer;
@@ -394,11 +448,12 @@ struct FNDIPhysicsAssetParametersCS : public FNiagaraDataInterfaceParametersCS
 		FNDIPhysicsAssetData* ProxyData =
 			InterfaceProxy->SystemInstancesToProxyData.Find(Context.SystemInstance);
 
-		if (ProxyData != nullptr)
+		if (ProxyData != nullptr && ProxyData->PhysicsAssetBuffer && ProxyData->PhysicsAssetBuffer->IsInitialized() )
 		{
-			FNDIPhysicsAssetBuffer* AssetBuffer = ProxyData->AssetBuffer;
+			FNDIPhysicsAssetBuffer* AssetBuffer = ProxyData->PhysicsAssetBuffer;
 			SetSRVParameter(RHICmdList, ComputeShaderRHI, CurrentTransformBuffer, AssetBuffer->CurrentTransformBuffer.SRV);
 			SetSRVParameter(RHICmdList, ComputeShaderRHI, PreviousTransformBuffer, AssetBuffer->PreviousTransformBuffer.SRV);
+			SetSRVParameter(RHICmdList, ComputeShaderRHI, PreviousInverseBuffer, AssetBuffer->PreviousInverseBuffer.SRV);
 			SetSRVParameter(RHICmdList, ComputeShaderRHI, InverseTransformBuffer, AssetBuffer->InverseTransformBuffer.SRV);
 			SetSRVParameter(RHICmdList, ComputeShaderRHI, RestTransformBuffer, AssetBuffer->RestTransformBuffer.SRV);
 			SetSRVParameter(RHICmdList, ComputeShaderRHI, RestInverseBuffer, AssetBuffer->RestInverseBuffer.SRV);
@@ -412,6 +467,7 @@ struct FNDIPhysicsAssetParametersCS : public FNiagaraDataInterfaceParametersCS
 		{
 			SetSRVParameter(RHICmdList, ComputeShaderRHI, CurrentTransformBuffer, FNiagaraRenderer::GetDummyFloatBuffer().SRV);
 			SetSRVParameter(RHICmdList, ComputeShaderRHI, PreviousTransformBuffer, FNiagaraRenderer::GetDummyFloatBuffer().SRV);
+			SetSRVParameter(RHICmdList, ComputeShaderRHI, PreviousInverseBuffer, FNiagaraRenderer::GetDummyFloatBuffer().SRV);
 			SetSRVParameter(RHICmdList, ComputeShaderRHI, InverseTransformBuffer, FNiagaraRenderer::GetDummyFloatBuffer().SRV);
 			SetSRVParameter(RHICmdList, ComputeShaderRHI, RestTransformBuffer, FNiagaraRenderer::GetDummyFloatBuffer().SRV);
 			SetSRVParameter(RHICmdList, ComputeShaderRHI, RestInverseBuffer, FNiagaraRenderer::GetDummyFloatBuffer().SRV);
@@ -434,6 +490,7 @@ private:
 
 	FShaderResourceParameter CurrentTransformBuffer;
 	FShaderResourceParameter PreviousTransformBuffer;
+	FShaderResourceParameter PreviousInverseBuffer;
 	FShaderResourceParameter InverseTransformBuffer;
 	FShaderResourceParameter RestTransformBuffer;
 	FShaderResourceParameter RestInverseBuffer;
@@ -458,22 +515,34 @@ void FNDIPhysicsAssetProxy::DeferredDestroy()
 void FNDIPhysicsAssetProxy::ConsumePerInstanceDataFromGameThread(void* PerInstanceData, const FNiagaraSystemInstanceID& Instance)
 {
 	FNDIPhysicsAssetData* SourceData = static_cast<FNDIPhysicsAssetData*>(PerInstanceData);
-	FNDIPhysicsAssetData& TargetData = SystemInstancesToProxyData.FindOrAdd(Instance);
+	FNDIPhysicsAssetData* TargetData = &(SystemInstancesToProxyData.FindOrAdd(Instance));
 
-	TargetData.AssetBuffer = SourceData->AssetBuffer;
-	TargetData.BoxOrigin = SourceData->BoxOrigin;
-	TargetData.BoxExtent = SourceData->BoxExtent;
+	ensure(TargetData);
+	if (TargetData)
+	{
+		TargetData->PhysicsAssetBuffer = SourceData->PhysicsAssetBuffer;
+		TargetData->BoxOrigin = SourceData->BoxOrigin;
+		TargetData->BoxExtent = SourceData->BoxExtent;
+	}
+	else
+	{
+		UE_LOG(LogPhysicsAsset, Log, TEXT("ConsumePerInstanceDataFromGameThread() ... could not find %d"), Instance);
+	}
 }
 
-void FNDIPhysicsAssetProxy::InitializePerInstanceData(const FNiagaraSystemInstanceID& SystemInstance, FNDIPhysicsAssetBuffer* AssetBuffer, const FVector& BoxOrigin , const FVector& BoxExtent)
+void FNDIPhysicsAssetProxy::InitializePerInstanceData(const FNiagaraSystemInstanceID& SystemInstance)
 {
 	check(IsInRenderingThread());
 
-	FNDIPhysicsAssetData& TargetData = SystemInstancesToProxyData.FindOrAdd(SystemInstance);
-
-	TargetData.AssetBuffer = AssetBuffer;
-	TargetData.BoxOrigin = BoxOrigin;
-	TargetData.BoxExtent = BoxExtent;
+	FNDIPhysicsAssetData* TargetData = SystemInstancesToProxyData.Find(SystemInstance);
+	if (TargetData != nullptr)
+	{
+		DeferredDestroyList.Remove(SystemInstance);
+	}
+	else
+	{
+		TargetData = &SystemInstancesToProxyData.Add(SystemInstance);
+	}
 }
 
 void FNDIPhysicsAssetProxy::DestroyPerInstanceData(NiagaraEmitterInstanceBatcher* Batcher, const FNiagaraSystemInstanceID& SystemInstance)
@@ -496,10 +565,8 @@ UNiagaraDataInterfacePhysicsAsset::UNiagaraDataInterfacePhysicsAsset(FObjectInit
 	Proxy = MakeShared<FNDIPhysicsAssetProxy, ESPMode::ThreadSafe>();
 }
 
-bool UNiagaraDataInterfacePhysicsAsset::InitPerInstanceData(void* PerInstanceData, FNiagaraSystemInstance* SystemInstance)
+void UNiagaraDataInterfacePhysicsAsset::ExtractSourceComponent(FNiagaraSystemInstance* SystemInstance)
 {
-	FNDIPhysicsAssetData* InstanceData = new (PerInstanceData) FNDIPhysicsAssetData();
-
 	SourceComponent = nullptr;
 	if (SourceActor)
 	{
@@ -525,9 +592,12 @@ bool UNiagaraDataInterfacePhysicsAsset::InitPerInstanceData(void* PerInstanceDat
 			{
 				SourceComponent = OuterComp;
 			}
-			else if (AActor* Owner = SimComp->GetAttachmentRootActor())
+			else
 			{
-				for (UActorComponent* ActorComp : Owner->GetComponents())
+				TArray<USceneComponent*> SceneComponents;
+				SimComp->GetParentComponents(SceneComponents);
+
+				for (USceneComponent* ActorComp : SceneComponents)
 				{
 					USkeletalMeshComponent* SourceComp = Cast<USkeletalMeshComponent>(ActorComp);
 					if (SourceComp && SourceComp->SkeletalMesh)
@@ -536,75 +606,55 @@ bool UNiagaraDataInterfacePhysicsAsset::InitPerInstanceData(void* PerInstanceDat
 						break;
 					}
 				}
+
 			}
 		}
 	}
+	PhysicsAsset = (SourceComponent != nullptr) ? SourceComponent->GetPhysicsAsset() : (DefaultSource != nullptr) ? DefaultSource : nullptr;
+}
 
-	PhysicsAsset = (SourceComponent != nullptr) ? SourceComponent->GetPhysicsAsset() : ( DefaultSource != nullptr ) ? DefaultSource : nullptr;
-
-	//if (PhysicsAsset == nullptr )
-	//{
-	//	UE_LOG(LogPhysicsAsset, Log, TEXT("Physics Asset data interface has no valid asset. Failed InitPerInstanceData - %s %d %d"), *GetFullName(), PhysicsAsset.Get(), DefaultSource);
-	//	return false;
-	//}
-	//else
-	{
-		const FTransform WorldTransform = (SourceComponent != nullptr) ? SourceComponent->GetComponentTransform() : SystemInstance->GetComponent()->GetComponentTransform();
-		FNDIPhysicsAssetBuffer* AssetBuffer = new FNDIPhysicsAssetBuffer;
-		AssetBuffer->SetupArrays(PhysicsAsset, SourceComponent, WorldTransform);
-
-		const bool HasPreviewMesh = PhysicsAsset != nullptr && PhysicsAsset->GetPreviewMesh() != nullptr;
-		const FVector BoxOrigin = HasPreviewMesh ? PhysicsAsset->GetPreviewMesh()->GetImportedBounds().Origin : FVector(0, 0, 0);
-		const FVector BoxExtent = HasPreviewMesh ? PhysicsAsset->GetPreviewMesh()->GetImportedBounds().BoxExtent : FVector(1, 1, 1);
+bool UNiagaraDataInterfacePhysicsAsset::InitPerInstanceData(void* PerInstanceData, FNiagaraSystemInstance* SystemInstance)
+{
+	FNDIPhysicsAssetData* InstanceData = new (PerInstanceData) FNDIPhysicsAssetData();
 
-		// Push instance data to RT
-		{
-			InstanceData->AssetBuffer = AssetBuffer;
-			InstanceData->BoxOrigin = BoxOrigin;
-			InstanceData->BoxExtent = BoxExtent;
+	check(InstanceData);
 
-			FNDIPhysicsAssetProxy* ThisProxy = GetProxyAs<FNDIPhysicsAssetProxy>();
-			ENQUEUE_RENDER_COMMAND(FNiagaraDIPushInitialInstanceDataToRT) (
-				[ThisProxy, InstanceID = SystemInstance->GetId(), AssetBuffer, BoxOrigin, BoxExtent](FRHICommandListImmediate& CmdList)
-			{
-				AssetBuffer->InitResource();
-				ThisProxy->InitializePerInstanceData(InstanceID, AssetBuffer, BoxOrigin, BoxExtent);
-			}
-			);
-		}
+	/*FNDIPhysicsAssetProxy* ThisProxy = GetProxyAs<FNDIPhysicsAssetProxy>();
+	ENQUEUE_RENDER_COMMAND(FNiagaraDIPushInitialInstanceDataToRT) (
+		[ThisProxy, InstanceID = SystemInstance->GetId()](FRHICommandListImmediate& CmdList)
+	{
+		ThisProxy->InitializePerInstanceData(InstanceID);
 	}
-	return true;
+	);*/
+
+	return InstanceData->Init(this, SystemInstance);
 }
 
 void UNiagaraDataInterfacePhysicsAsset::DestroyPerInstanceData(void* PerInstanceData, FNiagaraSystemInstance* SystemInstance)
 {
 	FNDIPhysicsAssetData* InstanceData = static_cast<FNDIPhysicsAssetData*>(PerInstanceData);
+
+	InstanceData->Release();
 	InstanceData->~FNDIPhysicsAssetData();
 
-	if (InstanceData->AssetBuffer)
+	FNDIPhysicsAssetProxy* ThisProxy = GetProxyAs<FNDIPhysicsAssetProxy>();
+	ENQUEUE_RENDER_COMMAND(FNiagaraDIDestroyInstanceData) (
+		[ThisProxy, InstanceID = SystemInstance->GetId(), Batcher = SystemInstance->GetBatcher()](FRHICommandListImmediate& CmdList)
 	{
-		FNDIPhysicsAssetProxy* ThisProxy = GetProxyAs<FNDIPhysicsAssetProxy>();
-		FNDIPhysicsAssetBuffer* InBuffer = InstanceData->AssetBuffer;
-		ENQUEUE_RENDER_COMMAND(FNiagaraDIDestroyInstanceData) (
-			[ThisProxy, InBuffer,  InstanceID = SystemInstance->GetId(), Batcher = SystemInstance->GetBatcher()](FRHICommandListImmediate& CmdList)
-		{
-			InBuffer->ReleaseResource();
-			ThisProxy->SystemInstancesToProxyData.Remove(InstanceID);
-			delete InBuffer;
-		}
-		);
-		InstanceData->AssetBuffer = nullptr;
+		ThisProxy->SystemInstancesToProxyData.Remove(InstanceID);
+		//ThisProxy->DestroyPerInstanceData(Batcher, InstanceID);
 	}
+	);
 }
 
 bool UNiagaraDataInterfacePhysicsAsset::PerInstanceTick(void* PerInstanceData, FNiagaraSystemInstance* SystemInstance, float InDeltaSeconds)
 {
 	FNDIPhysicsAssetData* InstanceData = static_cast<FNDIPhysicsAssetData*>(PerInstanceData);
-	if (InstanceData->AssetBuffer)
+	if (InstanceData->PhysicsAssetBuffer)
 	{
-		InstanceData->AssetBuffer->WorldTransform = (InstanceData->AssetBuffer->SkeletalMesh != nullptr) ? InstanceData->AssetBuffer->SkeletalMesh->GetComponentTransform() :
+		InstanceData->PhysicsAssetBuffer->WorldTransform = (InstanceData->PhysicsAssetBuffer->SkeletalMesh != nullptr) ? InstanceData->PhysicsAssetBuffer->SkeletalMesh->GetComponentTransform() :
 			SystemInstance->GetComponent()->GetComponentTransform();
-		InstanceData->AssetBuffer->UpdateBuffers();
+		InstanceData->PhysicsAssetBuffer->Update();
 	}
 	return false;
 }
@@ -686,6 +736,7 @@ void UNiagaraDataInterfacePhysicsAsset::GetFunctions(TArray<FNiagaraFunctionSign
 		Sig.Inputs.Add(FNiagaraVariable(FNiagaraTypeDefinition(GetClass()), TEXT("Physics Asset")));
 		Sig.Inputs.Add(FNiagaraVariable(FNiagaraTypeDefinition::GetVec3Def(), TEXT("World Position")));
 		Sig.Inputs.Add(FNiagaraVariable(FNiagaraTypeDefinition::GetFloatDef(), TEXT("Delta Time")));
+		Sig.Inputs.Add(FNiagaraVariable(FNiagaraTypeDefinition::GetFloatDef(), TEXT("Time Fraction")));
 		Sig.Outputs.Add(FNiagaraVariable(FNiagaraTypeDefinition::GetVec3Def(), TEXT("Closest Position")));
 		Sig.Outputs.Add(FNiagaraVariable(FNiagaraTypeDefinition::GetVec3Def(), TEXT("Closest Normal")));
 		Sig.Outputs.Add(FNiagaraVariable(FNiagaraTypeDefinition::GetVec3Def(), TEXT("Closest Velocity")));
@@ -751,7 +802,7 @@ void UNiagaraDataInterfacePhysicsAsset::GetVMExternalFunction(const FVMExternalF
 	}
 	else if (BindingInfo.Name == GetClosestPointName)
 	{
-		check(BindingInfo.GetNumInputs() == 5 && BindingInfo.GetNumOutputs() == 10);
+		check(BindingInfo.GetNumInputs() == 6 && BindingInfo.GetNumOutputs() == 10);
 		NDI_FUNC_BINDER(UNiagaraDataInterfacePhysicsAsset, GetClosestPoint)::Bind(this, OutFunc);
 	}
 	else if (BindingInfo.Name == GetTexturePointName)
@@ -799,6 +850,7 @@ bool UNiagaraDataInterfacePhysicsAsset::GetFunctionHLSL(const FNiagaraDataInterf
 		{TEXT("ElementOffsetsName"), ParamNames.ElementOffsetsName},
 		{TEXT("CurrentTransformBufferName"), ParamNames.CurrentTransformBufferName},
 		{TEXT("PreviousTransformBufferName"), ParamNames.PreviousTransformBufferName},
+		{TEXT("PreviousInverseBufferName"), ParamNames.PreviousInverseBufferName},
 		{TEXT("InverseTransformBufferName"), ParamNames.InverseTransformBufferName},
 		{TEXT("ElementExtentBufferName"), ParamNames.ElementExtentBufferName},
 		{TEXT("PhysicsAssetContextName"), TEXT("DIPHYSICSASSET_MAKE_CONTEXT(") + ParamInfo.DataInterfaceHLSLSymbol + TEXT(")")},
@@ -843,10 +895,10 @@ bool UNiagaraDataInterfacePhysicsAsset::GetFunctionHLSL(const FNiagaraDataInterf
 	else if (FunctionInfo.DefinitionName == GetClosestPointName)
 	{
 		static const TCHAR *FormatSample = TEXT(R"(
-		void {InstanceFunctionName}(in float3 WorldPosition, in float DeltaTime, out float3 OutClosestPosition, 
+		void {InstanceFunctionName}(in float3 WorldPosition, in float DeltaTime, in float TimeFraction, out float3 OutClosestPosition, 
 							out float3 OutClosestNormal, out float3 OutClosestVelocity, out float OutClosestDistance)
 		{
-			{PhysicsAssetContextName} DIPhysicsAsset_GetClosestPoint(DIContext,WorldPosition,DeltaTime,
+			{PhysicsAssetContextName} DIPhysicsAsset_GetClosestPoint(DIContext,WorldPosition,DeltaTime,TimeFraction,
 				OutClosestPosition,OutClosestNormal,OutClosestVelocity,OutClosestDistance);
 		}
 		)");
@@ -899,7 +951,7 @@ void UNiagaraDataInterfacePhysicsAsset::ProvidePerInstanceDataForRenderThread(vo
 
 	if (GameThreadData != nullptr && RenderThreadData != nullptr)
 	{
-		RenderThreadData->AssetBuffer = GameThreadData->AssetBuffer;
+		RenderThreadData->PhysicsAssetBuffer = GameThreadData->PhysicsAssetBuffer;
 		RenderThreadData->BoxOrigin = GameThreadData->BoxOrigin;
 		RenderThreadData->BoxExtent = GameThreadData->BoxExtent;
 	}
