Copying //UE4/Dev-Physics to Main (//UE4/Main)

diff --git a/Engine/Plugins/Experimental/HairStrands/Source/HairStrandsNiagara/Private/NiagaraDataInterfacePressureGrid.cpp b/Engine/Plugins/Experimental/HairStrands/Source/HairStrandsNiagara/Private/NiagaraDataInterfacePressureGrid.cpp
--- a/Engine/Plugins/Experimental/HairStrands/Source/HairStrandsNiagara/Private/NiagaraDataInterfacePressureGrid.cpp
+++ b/Engine/Plugins/Experimental/HairStrands/Source/HairStrandsNiagara/Private/NiagaraDataInterfacePressureGrid.cpp
@@ -9,6 +9,10 @@
 
 #include "ShaderParameterUtils.h"
 #include "ClearQuad.h"
+#include "RenderGraphBuilder.h"
+#include "RenderGraphUtils.h"
+#include "ShaderParameterStruct.h"
+#include "GlobalShader.h"
 
 #define LOCTEXT_NAMESPACE "NiagaraDataInterfacePressureGrid"
 DEFINE_LOG_CATEGORY_STATIC(LogPressureGrid, Log, All);
@@ -18,12 +22,14 @@ DEFINE_LOG_CATEGORY_STATIC(LogPressureGrid, Log, All);
 static const FName BuildVelocityFieldName(TEXT("BuildVelocityField"));
 static const FName ProjectVelocityFieldName(TEXT("ProjectVelocityField"));
 static const FName SampleVelocityFieldName(TEXT("SampleVelocityField"));
-static const FName GetCellPositionName(TEXT("GetCellPosition"));
+static const FName GetNodePositionName(TEXT("GetNodePosition"));
 static const FName TransferCellVelocityName(TEXT("TransferCellVelocity"));
 static const FName SetSolidBoundaryName(TEXT("SetSolidBoundary"));
 static const FName ComputeBoundaryWeightsName(TEXT("ComputeBoundaryWeights"));
 static const FName BuildGridTopologyName(TEXT("BuildGridTopology"));
 static const FName UpdateGridTransformName(TEXT("UpdateGridTransform"));
+static const FName AddGridVelocityName(TEXT("AddGridVelocity"));
+static const FName GetGridVelocityName(TEXT("GetGridVelocity"));
 
 //------------------------------------------------------------------------------------------------------------
 
@@ -64,18 +70,18 @@ struct FNDIPressureGridParametersName
 
 //------------------------------------------------------------------------------------------------------------
 
-void FNDIPressureGridBuffer::SetGridSize(const FUintVector4 InGridSize)
+void FNDIPressureGridBuffer::Initialize(const FIntVector InGridSize)
 {
 	GridSize = InGridSize;
 }
 
 void FNDIPressureGridBuffer::InitRHI()
 {
-	if(GridSize.X != 0 && GridSize.Y != 0 && GridSize.Z != 0)
+	if (GridSize.X != 0 && GridSize.Y != 0 && GridSize.Z != 0)
 	{
-		static const uint32 NumComponents = 14;
-		GridDataBuffer.Initialize(sizeof(int32), GridSize.X*NumComponents, GridSize.Y,
-			GridSize.Z, EPixelFormat::PF_R32_SINT);
+		static const uint32 NumComponents = 17;
+		GridDataBuffer.Initialize(sizeof(int32), (GridSize.X + 1)*NumComponents, (GridSize.Y + 1),
+			(GridSize.Z + 1), EPixelFormat::PF_R32_SINT);
 	}
 }
 
@@ -84,20 +90,64 @@ void FNDIPressureGridBuffer::ReleaseRHI()
 	GridDataBuffer.Release();
 }
 
-void FNDIPressureGridBuffer::ClearBuffers(FRHICommandList& RHICmdList)
-{
-	ClearUAV(RHICmdList, GridDataBuffer, FLinearColor(0, 0, 0, 0));
-}
-
 //------------------------------------------------------------------------------------------------------------
 
-void FNDIPressureGridData::SwapBuffers()
+void FNDIPressureGridData::Swap()
 {
 	FNDIPressureGridBuffer* StoredBufferPointer = CurrentGridBuffer;
 	CurrentGridBuffer = DestinationGridBuffer;
 	DestinationGridBuffer = StoredBufferPointer;
 }
 
+void FNDIPressureGridData::Release()
+{
+	if (CurrentGridBuffer)
+	{
+		BeginReleaseResource(CurrentGridBuffer);
+		ENQUEUE_RENDER_COMMAND(DeleteResourceA)(
+			[ParamPointerToRelease = CurrentGridBuffer](FRHICommandListImmediate& RHICmdList)
+		{
+			delete ParamPointerToRelease;
+		});
+		CurrentGridBuffer = nullptr;
+	}
+	if (DestinationGridBuffer)
+	{
+		BeginReleaseResource(DestinationGridBuffer);
+		ENQUEUE_RENDER_COMMAND(DeleteResourceB)(
+			[ParamPointerToRelease = DestinationGridBuffer](FRHICommandListImmediate& RHICmdList)
+		{
+			delete ParamPointerToRelease;
+		});
+		DestinationGridBuffer = nullptr;
+	}
+}
+
+bool FNDIPressureGridData::Init(UNiagaraDataInterfacePressureGrid* Interface, FNiagaraSystemInstance* SystemInstance)
+{
+	CurrentGridBuffer = nullptr;
+	DestinationGridBuffer = nullptr;
+
+	GridOrigin = FVector4(0, 0, 0, 0);
+	GridSize = FIntVector(1, 1, 1);
+	WorldTransform = WorldInverse = FMatrix::Identity;
+
+	if (Interface != nullptr)
+	{
+		GridSize = Interface->GridSize;
+
+		CurrentGridBuffer = new FNDIPressureGridBuffer();
+		CurrentGridBuffer->Initialize(Interface->GridSize);
+		BeginInitResource(CurrentGridBuffer);
+
+		DestinationGridBuffer = new FNDIPressureGridBuffer();
+		DestinationGridBuffer->Initialize(Interface->GridSize);
+		BeginInitResource(DestinationGridBuffer);
+	}
+
+	return true;
+}
+
 //------------------------------------------------------------------------------------------------------------
 
 struct FNDIPressureGridParametersCS : public FNiagaraDataInterfaceParametersCS
@@ -149,7 +199,8 @@ struct FNDIPressureGridParametersCS : public FNiagaraDataInterfaceParametersCS
 		FNDIPressureGridData* ProxyData =
 			InterfaceProxy->SystemInstancesToProxyData.Find(Context.SystemInstance);
 
-		if (ProxyData != nullptr)
+		if (ProxyData != nullptr && ProxyData->CurrentGridBuffer != nullptr && ProxyData->DestinationGridBuffer != nullptr
+			&& ProxyData->CurrentGridBuffer->IsInitialized() && ProxyData->DestinationGridBuffer->IsInitialized())
 		{
 			FNDIPressureGridBuffer* CurrentGridBuffer = ProxyData->CurrentGridBuffer;
 			FNDIPressureGridBuffer* DestinationGridBuffer = ProxyData->DestinationGridBuffer;
@@ -160,13 +211,10 @@ struct FNDIPressureGridParametersCS : public FNiagaraDataInterfaceParametersCS
 			RHICmdList.TransitionResource(EResourceTransitionAccess::EReadable, EResourceTransitionPipeline::EComputeToCompute, CurrentGridBuffer->GridDataBuffer.UAV);
 			SetSRVParameter(RHICmdList, ComputeShaderRHI, GridCurrentBuffer, CurrentGridBuffer->GridDataBuffer.SRV);
 
-			//UE_LOG(LogPressureGrid, Warning, TEXT("Binding Grid Origin : %s"), *ProxyData->GridOrigin.ToString());
 			SetShaderValue(RHICmdList, ComputeShaderRHI, GridOrigin, ProxyData->GridOrigin);
 
-			//UE_LOG(LogPressureGrid, Warning, TEXT("Binding Grid Size : %d %d %d"), ProxyData->GridSize.X, ProxyData->GridSize.Y, ProxyData->GridSize.Z);
 			SetShaderValue(RHICmdList, ComputeShaderRHI, GridSize, ProxyData->GridSize);
 
-			//UE_LOG(LogPressureGrid, Warning, TEXT("Binding Grid Transform : %s"), *ProxyData->WorldTransform.ToString());
 			SetShaderValue(RHICmdList, ComputeShaderRHI, WorldTransform, ProxyData->WorldTransform);
 			SetShaderValue(RHICmdList, ComputeShaderRHI, WorldInverse, ProxyData->WorldTransform.Inverse());
 		}
@@ -175,8 +223,8 @@ struct FNDIPressureGridParametersCS : public FNiagaraDataInterfaceParametersCS
 			SetUAVParameter(RHICmdList, ComputeShaderRHI, GridDestinationBuffer, FNiagaraRenderer::GetDummyUIntBuffer().UAV);
 			SetSRVParameter(RHICmdList, ComputeShaderRHI, GridCurrentBuffer, FNiagaraRenderer::GetDummyUIntBuffer().SRV);
 
-			SetShaderValue(RHICmdList, ComputeShaderRHI, GridOrigin, FVector4(0,0,0,0));
-			SetShaderValue(RHICmdList, ComputeShaderRHI, GridSize, FUintVector4());
+			SetShaderValue(RHICmdList, ComputeShaderRHI, GridOrigin, FVector4(0, 0, 0, 0));
+			SetShaderValue(RHICmdList, ComputeShaderRHI, GridSize, FIntVector());
 
 			SetShaderValue(RHICmdList, ComputeShaderRHI, WorldTransform, FMatrix::Identity);
 			SetShaderValue(RHICmdList, ComputeShaderRHI, WorldInverse, FMatrix::Identity);
@@ -200,57 +248,17 @@ private:
 	FShaderParameter WorldInverse;
 };
 
-//------------------------------------------------------------------------------------------------------------
-
-void FNDIPressureGridProxy::DeferredDestroy()
-{
-	for (const FNiagaraSystemInstanceID& Sys : DeferredDestroyList)
-	{
-		SystemInstancesToProxyData.Remove(Sys);
-	}
-
-	DeferredDestroyList.Empty();
-}
-
-void FNDIPressureGridProxy::ConsumePerInstanceDataFromGameThread(void* PerInstanceData, const FNiagaraSystemInstanceID& Instance)
-{
-	FNDIPressureGridData* SourceData = static_cast<FNDIPressureGridData*>(PerInstanceData);
-	FNDIPressureGridData& TargetData = SystemInstancesToProxyData.FindOrAdd(Instance);
-
-	TargetData.WorldTransform = SourceData->WorldTransform;
-	TargetData.GridOrigin = SourceData->GridOrigin;
-}
-
-void FNDIPressureGridProxy::InitializePerInstanceData(const FNiagaraSystemInstanceID& SystemInstance, FNDIPressureGridBuffer* CurrentGridBuffer,
-FNDIPressureGridBuffer* DestinationGridBuffer, const FVector4& GridOrigin, const FUintVector4& GridSize)
-{
-	check(IsInRenderingThread());
-
-	FNDIPressureGridData& TargetData = SystemInstancesToProxyData.FindOrAdd(SystemInstance);
-	
-	TargetData.CurrentGridBuffer = CurrentGridBuffer;
-	TargetData.DestinationGridBuffer = DestinationGridBuffer;
-	TargetData.GridOrigin = GridOrigin;
-	TargetData.GridSize = GridSize;
-}
-
-void FNDIPressureGridProxy::DestroyPerInstanceData(NiagaraEmitterInstanceBatcher* Batcher, const FNiagaraSystemInstanceID& SystemInstance)
-{
-	check(IsInRenderingThread());
-
-	DeferredDestroyList.Add(SystemInstance);
-	Batcher->EnqueueDeferredDeletesForDI_RenderThread(this->AsShared());
-}
 
 //------------------------------------------------------------------------------------------------------------
 
 UNiagaraDataInterfacePressureGrid::UNiagaraDataInterfacePressureGrid(FObjectInitializer const& ObjectInitializer)
 	: Super(ObjectInitializer)
 	, GridSize(10)
-	, GridOrigin(0,0,0)
-	, GridLength(1)
+	, MinIteration(0)
+	, MaxIteration(0)
+	, MinOutput(0)
+	, MaxOutput(0)
 {
-
 	Proxy = MakeShared<FNDIPressureGridProxy, ESPMode::ThreadSafe>();
 }
 
@@ -258,69 +266,50 @@ bool UNiagaraDataInterfacePressureGrid::InitPerInstanceData(void* PerInstanceDat
 {
 	FNDIPressureGridData* InstanceData = new (PerInstanceData) FNDIPressureGridData();
 	check(InstanceData);
-	
-	FNDIPressureGridBuffer* CurrentGridBuffer = new FNDIPressureGridBuffer;
-	FNDIPressureGridBuffer* DestinationGridBuffer = new FNDIPressureGridBuffer;
-		
-
-	// Push instance data to RT
-	{
-		const FUintVector4 LocalGridSize(GridSize.X, GridSize.Y, GridSize.Z,0);
-		const FVector4 LocalGridOrigin(GridOrigin.X,GridOrigin.Y,GridOrigin.Z,GridLength);
-		CurrentGridBuffer->SetGridSize(LocalGridSize);
-		DestinationGridBuffer->SetGridSize(LocalGridSize);
-
-		TSet<int> RT_OutputShaderStages = OutputShaderStages;
-		TSet<int> RT_IterationShaderStages = IterationShaderStages;
-		
-		InstanceData->CurrentGridBuffer = CurrentGridBuffer;
-		InstanceData->DestinationGridBuffer = DestinationGridBuffer;
-		InstanceData->GridSize = LocalGridSize;
-		InstanceData->GridOrigin = LocalGridOrigin;
-
-		InstanceData->WorldTransform = FMatrix::Identity;// SystemInstance->GetComponent()->GetComponentToWorld().ToMatrixWithScale();
-
-		FNDIPressureGridProxy* ThisProxy = GetProxyAs<FNDIPressureGridProxy>();
-		ENQUEUE_RENDER_COMMAND(FNiagaraDIPushInitialInstanceDataToRT) (
-			[ThisProxy, RT_OutputShaderStages, RT_IterationShaderStages, InstanceID = SystemInstance->GetId(), CurrentGridBuffer, DestinationGridBuffer, LocalGridSize, LocalGridOrigin](FRHICommandListImmediate& CmdList)
-		{
-			ThisProxy->OutputShaderStages = RT_OutputShaderStages;
-			ThisProxy->IterationShaderStages = RT_IterationShaderStages;
-			ThisProxy->SetElementCount(LocalGridSize.X*LocalGridSize.Y*LocalGridSize.Z);
 
-			CurrentGridBuffer->InitResource();
-			DestinationGridBuffer->InitResource();
+	TSet<int> RT_OutputShaderStages;
+	TSet<int> RT_IterationShaderStages;
 
-			ThisProxy->InitializePerInstanceData(InstanceID, CurrentGridBuffer, DestinationGridBuffer, LocalGridOrigin, LocalGridSize);
-		}
-		);
+	for (int32 i = MinIteration; i < MaxIteration; ++i)
+	{
+		RT_IterationShaderStages.Add(i);
 	}
-	return true;
+	for (int32 i = MinOutput; i < MaxOutput; ++i)
+	{
+		RT_OutputShaderStages.Add(i);
+	}
+	const int32 ElementCount = (GridSize.X + 1)*(GridSize.Y + 1)*(GridSize.Z + 1);
+
+	FNDIPressureGridProxy* ThisProxy = GetProxyAs<FNDIPressureGridProxy>();
+	ENQUEUE_RENDER_COMMAND(FNiagaraDIPushInitialInstanceDataToRT) (
+		[ThisProxy, RT_OutputShaderStages, RT_IterationShaderStages, InstanceID = SystemInstance->GetId(), ElementCount](FRHICommandListImmediate& CmdList)
+	{
+		ThisProxy->OutputShaderStages = RT_OutputShaderStages;
+		ThisProxy->IterationShaderStages = RT_IterationShaderStages;
+		ThisProxy->SetElementCount(ElementCount);
+
+		//ThisProxy->InitializePerInstanceData(InstanceID);
+	}
+	);
+
+	return InstanceData->Init(this, SystemInstance);
 }
 
 void UNiagaraDataInterfacePressureGrid::DestroyPerInstanceData(void* PerInstanceData, FNiagaraSystemInstance* SystemInstance)
 {
 	FNDIPressureGridData* InstanceData = static_cast<FNDIPressureGridData*>(PerInstanceData);
+
+	InstanceData->Release();
 	InstanceData->~FNDIPressureGridData();
 
-	if (InstanceData->CurrentGridBuffer && InstanceData->DestinationGridBuffer)
+	FNDIPressureGridProxy* ThisProxy = GetProxyAs<FNDIPressureGridProxy>();
+	ENQUEUE_RENDER_COMMAND(FNiagaraDIDestroyInstanceData) (
+		[ThisProxy, InstanceID = SystemInstance->GetId(), Batcher = SystemInstance->GetBatcher()](FRHICommandListImmediate& CmdList)
 	{
-		FNDIPressureGridProxy* ThisProxy = GetProxyAs<FNDIPressureGridProxy>();
-		FNDIPressureGridBuffer* CurrentGridBuffer = InstanceData->CurrentGridBuffer;
-		FNDIPressureGridBuffer* DestinationGridBuffer = InstanceData->DestinationGridBuffer;
-		ENQUEUE_RENDER_COMMAND(FNiagaraDIDestroyInstanceData) (
-			[ThisProxy, CurrentGridBuffer, DestinationGridBuffer, InstanceID = SystemInstance->GetId(), Batcher = SystemInstance->GetBatcher()](FRHICommandListImmediate& CmdList)
-		{
-			CurrentGridBuffer->ReleaseResource();
-			DestinationGridBuffer->ReleaseResource();
-			ThisProxy->SystemInstancesToProxyData.Remove(InstanceID);
-			delete CurrentGridBuffer;
-			delete DestinationGridBuffer;
-		}
-		);
-		InstanceData->CurrentGridBuffer = nullptr;
-		InstanceData->DestinationGridBuffer = nullptr;
+		ThisProxy->SystemInstancesToProxyData.Remove(InstanceID);
+		//ThisProxy->DestroyPerInstanceData(Batcher, InstanceID);
 	}
+	);
 }
 
 bool UNiagaraDataInterfacePressureGrid::PerInstanceTick(void* PerInstanceData, FNiagaraSystemInstance* SystemInstance, float InDeltaSeconds)
@@ -332,7 +321,7 @@ bool UNiagaraDataInterfacePressureGrid::PerInstanceTick(void* PerInstanceData, F
 	{
 		InstanceData->WorldTransform = SystemInstance->GetComponent()->GetComponentToWorld().ToMatrixWithScale();
 	}
-	
+
 	return RequireReset;
 }
 
@@ -345,8 +334,10 @@ bool UNiagaraDataInterfacePressureGrid::CopyToInternal(UNiagaraDataInterface* De
 
 	UNiagaraDataInterfacePressureGrid* OtherTyped = CastChecked<UNiagaraDataInterfacePressureGrid>(Destination);
 	OtherTyped->GridSize = GridSize;
-	OtherTyped->GridOrigin = GridOrigin;
-	OtherTyped->GridLength= GridLength;
+	OtherTyped->MinIteration = MinIteration;
+	OtherTyped->MaxIteration = MaxIteration;
+	OtherTyped->MinOutput = MinOutput;
+	OtherTyped->MaxOutput = MaxOutput;
 
 	return true;
 }
@@ -359,7 +350,8 @@ bool UNiagaraDataInterfacePressureGrid::Equals(const UNiagaraDataInterface* Othe
 	}
 	const UNiagaraDataInterfacePressureGrid* OtherTyped = CastChecked<const UNiagaraDataInterfacePressureGrid>(Other);
 
-	return (OtherTyped->GridSize == GridSize) && (OtherTyped->GridOrigin == GridOrigin) && (OtherTyped->GridLength == GridLength);
+	return (OtherTyped->GridSize == GridSize) && (OtherTyped->MinIteration == MinIteration) && (OtherTyped->MaxIteration == MaxIteration) &&
+		(OtherTyped->MinOutput == MinOutput) && (OtherTyped->MaxOutput == MaxOutput);
 }
 
 void UNiagaraDataInterfacePressureGrid::PostInitProperties()
@@ -417,13 +409,13 @@ void UNiagaraDataInterfacePressureGrid::GetFunctions(TArray<FNiagaraFunctionSign
 		Sig.bRequiresContext = false;
 		Sig.Inputs.Add(FNiagaraVariable(FNiagaraTypeDefinition(GetClass()), TEXT("Pressure Grid")));
 		Sig.Inputs.Add(FNiagaraVariable(FNiagaraTypeDefinition::GetIntDef(), TEXT("Grid Cell")));
+		Sig.Inputs.Add(FNiagaraVariable(FNiagaraTypeDefinition::GetIntDef(), TEXT("Init Stage")));
 		Sig.Outputs.Add(FNiagaraVariable(FNiagaraTypeDefinition::GetBoolDef(), TEXT("Project Status")));
-		Sig.Outputs.Add(FNiagaraVariable(FNiagaraTypeDefinition::GetFloatDef(), TEXT("Cell Divergence")));
 		OutFunctions.Add(Sig);
 	}
 	{
 		FNiagaraFunctionSignature Sig;
-		Sig.Name = GetCellPositionName;
+		Sig.Name = GetNodePositionName;
 		Sig.bMemberFunction = true;
 		Sig.bRequiresContext = false;
 		Sig.Inputs.Add(FNiagaraVariable(FNiagaraTypeDefinition(GetClass()), TEXT("Pressure Grid")));
@@ -465,6 +457,7 @@ void UNiagaraDataInterfacePressureGrid::GetFunctions(TArray<FNiagaraFunctionSign
 		Sig.bRequiresContext = false;
 		Sig.Inputs.Add(FNiagaraVariable(FNiagaraTypeDefinition(GetClass()), TEXT("Pressure Grid")));
 		Sig.Inputs.Add(FNiagaraVariable(FNiagaraTypeDefinition::GetIntDef(), TEXT("Grid Cell")));
+		Sig.Inputs.Add(FNiagaraVariable(FNiagaraTypeDefinition::GetFloatDef(), TEXT("Grid Length")));
 		Sig.Outputs.Add(FNiagaraVariable(FNiagaraTypeDefinition::GetBoolDef(), TEXT("Transfer Status")));
 
 		OutFunctions.Add(Sig);
@@ -491,19 +484,44 @@ void UNiagaraDataInterfacePressureGrid::GetFunctions(TArray<FNiagaraFunctionSign
 		Sig.Inputs.Add(FNiagaraVariable(FNiagaraTypeDefinition::GetMatrix4Def(), TEXT("Grid Transform")));
 		Sig.Outputs.Add(FNiagaraVariable(FNiagaraTypeDefinition::GetBoolDef(), TEXT("Transform Status")));
 
+		OutFunctions.Add(Sig);
+	}
+	{
+		FNiagaraFunctionSignature Sig;
+		Sig.Name = AddGridVelocityName;
+		Sig.bMemberFunction = true;
+		Sig.bRequiresContext = false;
+		Sig.Inputs.Add(FNiagaraVariable(FNiagaraTypeDefinition(GetClass()), TEXT("Pressure Grid")));
+		Sig.Inputs.Add(FNiagaraVariable(FNiagaraTypeDefinition::GetIntDef(), TEXT("Grid Cell")));
+		Sig.Inputs.Add(FNiagaraVariable(FNiagaraTypeDefinition::GetVec3Def(), TEXT("Grid Velocity")));
+		Sig.Outputs.Add(FNiagaraVariable(FNiagaraTypeDefinition::GetBoolDef(), TEXT("Add Status")));
+
+		OutFunctions.Add(Sig);
+	}
+	{
+		FNiagaraFunctionSignature Sig;
+		Sig.Name = GetGridVelocityName;
+		Sig.bMemberFunction = true;
+		Sig.bRequiresContext = false;
+		Sig.Inputs.Add(FNiagaraVariable(FNiagaraTypeDefinition(GetClass()), TEXT("Pressure Grid")));
+		Sig.Inputs.Add(FNiagaraVariable(FNiagaraTypeDefinition::GetIntDef(), TEXT("Grid Cell")));
+		Sig.Outputs.Add(FNiagaraVariable(FNiagaraTypeDefinition::GetVec3Def(), TEXT("Grid Velocity")));
+
 		OutFunctions.Add(Sig);
 	}
 }
 
 DEFINE_NDI_DIRECT_FUNC_BINDER(UNiagaraDataInterfacePressureGrid, BuildVelocityField);
 DEFINE_NDI_DIRECT_FUNC_BINDER(UNiagaraDataInterfacePressureGrid, ProjectVelocityField);
-DEFINE_NDI_DIRECT_FUNC_BINDER(UNiagaraDataInterfacePressureGrid, GetCellPosition);
+DEFINE_NDI_DIRECT_FUNC_BINDER(UNiagaraDataInterfacePressureGrid, GetNodePosition);
 DEFINE_NDI_DIRECT_FUNC_BINDER(UNiagaraDataInterfacePressureGrid, SetSolidBoundary);
 DEFINE_NDI_DIRECT_FUNC_BINDER(UNiagaraDataInterfacePressureGrid, TransferCellVelocity);
 DEFINE_NDI_DIRECT_FUNC_BINDER(UNiagaraDataInterfacePressureGrid, ComputeBoundaryWeights);
 DEFINE_NDI_DIRECT_FUNC_BINDER(UNiagaraDataInterfacePressureGrid, SampleVelocityField);
 DEFINE_NDI_DIRECT_FUNC_BINDER(UNiagaraDataInterfacePressureGrid, BuildGridTopology);
 DEFINE_NDI_DIRECT_FUNC_BINDER(UNiagaraDataInterfacePressureGrid, UpdateGridTransform);
+DEFINE_NDI_DIRECT_FUNC_BINDER(UNiagaraDataInterfacePressureGrid, AddGridVelocity);
+DEFINE_NDI_DIRECT_FUNC_BINDER(UNiagaraDataInterfacePressureGrid, GetGridVelocity);
 
 void UNiagaraDataInterfacePressureGrid::GetVMExternalFunction(const FVMExternalFunctionBindingInfo& BindingInfo, void* InstanceData, FVMExternalFunction &OutFunc)
 {
@@ -514,13 +532,13 @@ void UNiagaraDataInterfacePressureGrid::GetVMExternalFunction(const FVMExternalF
 	}
 	else if (BindingInfo.Name == ProjectVelocityFieldName)
 	{
-		check(BindingInfo.GetNumInputs() == 2 && BindingInfo.GetNumOutputs() == 2);
+		check(BindingInfo.GetNumInputs() == 3 && BindingInfo.GetNumOutputs() == 1);
 		NDI_FUNC_BINDER(UNiagaraDataInterfacePressureGrid, ProjectVelocityField)::Bind(this, OutFunc);
 	}
-	else if (BindingInfo.Name == GetCellPositionName)
+	else if (BindingInfo.Name == GetNodePositionName)
 	{
 		check(BindingInfo.GetNumInputs() == 6 && BindingInfo.GetNumOutputs() == 3);
-		NDI_FUNC_BINDER(UNiagaraDataInterfacePressureGrid, GetCellPosition)::Bind(this, OutFunc);
+		NDI_FUNC_BINDER(UNiagaraDataInterfacePressureGrid, GetNodePosition)::Bind(this, OutFunc);
 	}
 	else if (BindingInfo.Name == SampleVelocityFieldName)
 	{
@@ -537,9 +555,9 @@ void UNiagaraDataInterfacePressureGrid::GetVMExternalFunction(const FVMExternalF
 		check(BindingInfo.GetNumInputs() == 2 && BindingInfo.GetNumOutputs() == 1);
 		NDI_FUNC_BINDER(UNiagaraDataInterfacePressureGrid, ComputeBoundaryWeights)::Bind(this, OutFunc);
 	}
-	else if (BindingInfo.Name ==TransferCellVelocityName)
+	else if (BindingInfo.Name == TransferCellVelocityName)
 	{
-		check(BindingInfo.GetNumInputs() == 2 && BindingInfo.GetNumOutputs() == 1);
+		check(BindingInfo.GetNumInputs() == 3 && BindingInfo.GetNumOutputs() == 1);
 		NDI_FUNC_BINDER(UNiagaraDataInterfacePressureGrid, TransferCellVelocity)::Bind(this, OutFunc);
 	}
 	else if (BindingInfo.Name == BuildGridTopologyName)
@@ -552,6 +570,16 @@ void UNiagaraDataInterfacePressureGrid::GetVMExternalFunction(const FVMExternalF
 		check(BindingInfo.GetNumInputs() == 17 && BindingInfo.GetNumOutputs() == 1);
 		NDI_FUNC_BINDER(UNiagaraDataInterfacePressureGrid, UpdateGridTransform)::Bind(this, OutFunc);
 	}
+	else if (BindingInfo.Name == AddGridVelocityName)
+	{
+		check(BindingInfo.GetNumInputs() == 5 && BindingInfo.GetNumOutputs() == 1);
+		NDI_FUNC_BINDER(UNiagaraDataInterfacePressureGrid, AddGridVelocity)::Bind(this, OutFunc);
+	}
+	else if (BindingInfo.Name == GetGridVelocityName)
+	{
+		check(BindingInfo.GetNumInputs() == 2 && BindingInfo.GetNumOutputs() == 3);
+		NDI_FUNC_BINDER(UNiagaraDataInterfacePressureGrid, GetGridVelocity)::Bind(this, OutFunc);
+	}
 }
 
 void UNiagaraDataInterfacePressureGrid::BuildVelocityField(FVectorVMContext& Context)
@@ -564,7 +592,7 @@ void UNiagaraDataInterfacePressureGrid::ProjectVelocityField(FVectorVMContext& C
 	// @todo : implement function for cpu 
 }
 
-void UNiagaraDataInterfacePressureGrid::GetCellPosition(FVectorVMContext& Context)
+void UNiagaraDataInterfacePressureGrid::GetNodePosition(FVectorVMContext& Context)
 {
 	// @todo : implement function for cpu 
 }
@@ -589,6 +617,16 @@ void UNiagaraDataInterfacePressureGrid::TransferCellVelocity(FVectorVMContext& C
 	// @todo : implement function for cpu 
 }
 
+void UNiagaraDataInterfacePressureGrid::AddGridVelocity(FVectorVMContext& Context)
+{
+	// @todo : implement function for cpu 
+}
+
+void UNiagaraDataInterfacePressureGrid::GetGridVelocity(FVectorVMContext& Context)
+{
+	// @todo : implement function for cpu 
+}
+
 void UNiagaraDataInterfacePressureGrid::BuildGridTopology(FVectorVMContext& Context)
 {
 	VectorVM::FExternalFuncInputHandler<float> CenterXParam(Context);
@@ -610,13 +648,13 @@ void UNiagaraDataInterfacePressureGrid::BuildGridTopology(FVectorVMContext& Cont
 		const FVector GridExtent(*ExtentXParam.GetDestAndAdvance(), *ExtentYParam.GetDestAndAdvance(), *ExtentZParam.GetDestAndAdvance());
 
 		const FVector GridLengths(2.0 * GridExtent.X / (InstData->GridSize.X - 1),
-								  2.0 * GridExtent.Y / (InstData->GridSize.Y - 1),
-								  2.0 * GridExtent.Z / (InstData->GridSize.Z - 1));
+			2.0 * GridExtent.Y / (InstData->GridSize.Y - 1),
+			2.0 * GridExtent.Z / (InstData->GridSize.Z - 1));
 		const float MaxLength = GridLengths.GetMax();
 
-		const FVector RegularExtent((InstData->GridSize.X - 1)*MaxLength, 
-								    (InstData->GridSize.Y - 1)*MaxLength, 
-									(InstData->GridSize.Z - 1)*MaxLength);
+		const FVector RegularExtent((InstData->GridSize.X - 1)*MaxLength,
+			(InstData->GridSize.Y - 1)*MaxLength,
+			(InstData->GridSize.Z - 1)*MaxLength);
 		const FVector BoxOrigin = GridCenter - 0.5 * RegularExtent;
 		InstData->GridOrigin = FVector4(BoxOrigin.X, BoxOrigin.Y, BoxOrigin.Z, MaxLength);
 
@@ -727,20 +765,20 @@ bool UNiagaraDataInterfacePressureGrid::GetFunctionHLSL(const FNiagaraDataInterf
 	else if (FunctionInfo.DefinitionName == ProjectVelocityFieldName)
 	{
 		static const TCHAR *FormatSample = TEXT(R"(
-				void {InstanceFunctionName} (in int GridCell, out bool OutProjectStatus, out float OutCellDivergence)
+				void {InstanceFunctionName} (in int GridCell, in int InitStage, out bool OutProjectStatus)
 				{
-					{PressureGridContextName} DIPressureGrid_ProjectVelocityField(DIContext,GridCell,OutProjectStatus,OutCellDivergence);
+					{PressureGridContextName} DIPressureGrid_ProjectVelocityField(DIContext,GridCell,InitStage,OutProjectStatus);
 				}
 				)");
 		OutHLSL += FString::Format(FormatSample, ArgsSample);
 		return true;
 	}
-	else if (FunctionInfo.DefinitionName == GetCellPositionName)
+	else if (FunctionInfo.DefinitionName == GetNodePositionName)
 	{
 		static const TCHAR *FormatSample = TEXT(R"(
 				void {InstanceFunctionName} (in int GridCell, in float3 GridOrigin, in float GridLength, out float3 OutGridPosition)
 				{
-					{PressureGridContextName} DIPressureGrid_GetCellPosition(DIContext,GridCell,GridOrigin,GridLength,OutGridPosition);
+					{PressureGridContextName} DIPressureGrid_GetNodePosition(DIContext,GridCell,GridOrigin,GridLength,OutGridPosition);
 				}
 				)");
 		OutHLSL += FString::Format(FormatSample, ArgsSample);
@@ -771,9 +809,9 @@ bool UNiagaraDataInterfacePressureGrid::GetFunctionHLSL(const FNiagaraDataInterf
 	else if (FunctionInfo.DefinitionName == TransferCellVelocityName)
 	{
 		static const TCHAR *FormatSample = TEXT(R"(
-				void {InstanceFunctionName} (in int GridCell, out bool OutTransferStatus)
+				void {InstanceFunctionName} (in int GridCell, in float GridLength, out bool OutTransferStatus)
 				{
-					{PressureGridContextName} DIPressureGrid_TransferCellVelocity(DIContext,GridCell,OutTransferStatus);
+					{PressureGridContextName} DIPressureGrid_TransferCellVelocity(DIContext,GridCell,GridLength,OutTransferStatus);
 				}
 				)");
 		OutHLSL += FString::Format(FormatSample, ArgsSample);
@@ -790,6 +828,28 @@ bool UNiagaraDataInterfacePressureGrid::GetFunctionHLSL(const FNiagaraDataInterf
 		OutHLSL += FString::Format(FormatSample, ArgsSample);
 		return true;
 	}
+	else if (FunctionInfo.DefinitionName == AddGridVelocityName)
+	{
+		static const TCHAR *FormatSample = TEXT(R"(
+				void {InstanceFunctionName} (in int GridCell, in float3 GridVelocity, out bool OutAddStatus)
+				{
+					{PressureGridContextName} DIPressureGrid_AddGridVelocity(DIContext,GridCell,GridVelocity,OutAddStatus);
+				}
+				)");
+		OutHLSL += FString::Format(FormatSample, ArgsSample);
+		return true;
+	}
+	else if (FunctionInfo.DefinitionName == GetGridVelocityName)
+	{
+		static const TCHAR *FormatSample = TEXT(R"(
+				void {InstanceFunctionName} (in int GridCell, out float3 OutGridVelocity)
+				{
+					{PressureGridContextName} DIPressureGrid_GetGridVelocity(DIContext,GridCell,OutGridVelocity);
+				}
+				)");
+		OutHLSL += FString::Format(FormatSample, ArgsSample);
+		return true;
+	}
 
 	OutHLSL += TEXT("\n");
 	return false;
@@ -811,7 +871,11 @@ void UNiagaraDataInterfacePressureGrid::ProvidePerInstanceDataForRenderThread(vo
 	FNDIPressureGridData* RenderThreadData = static_cast<FNDIPressureGridData*>(DataForRenderThread);
 
 	RenderThreadData->WorldTransform = GameThreadData->WorldTransform;
+	RenderThreadData->WorldInverse = GameThreadData->WorldInverse;
 	RenderThreadData->GridOrigin = GameThreadData->GridOrigin;
+	RenderThreadData->CurrentGridBuffer = GameThreadData->CurrentGridBuffer;
+	RenderThreadData->DestinationGridBuffer = GameThreadData->DestinationGridBuffer;
+	RenderThreadData->GridSize = GameThreadData->GridSize;
 }
 
 FNiagaraDataInterfaceParametersCS*
@@ -820,6 +884,150 @@ UNiagaraDataInterfacePressureGrid::ConstructComputeParameters() const
 	return new FNDIPressureGridParametersCS();
 }
 
+//------------------------------------------------------------------------------------------------------------
+
+#define NIAGARA_HAIR_STRANDS_THREAD_COUNT 64
+
+class FClearPressureGridCS : public FGlobalShader
+{
+	DECLARE_GLOBAL_SHADER(FClearPressureGridCS);
+	SHADER_USE_PARAMETER_STRUCT(FClearPressureGridCS, FGlobalShader);
+
+	BEGIN_SHADER_PARAMETER_STRUCT(FParameters, )
+		SHADER_PARAMETER(FIntVector, GridSize)
+		SHADER_PARAMETER(int32, CopyPressure)
+		SHADER_PARAMETER_SRV(Texture3D, GridCurrentBuffer)
+		SHADER_PARAMETER_UAV(RWTexture3D, GridDestinationBuffer)
+		END_SHADER_PARAMETER_STRUCT()
+
+public:
+	static bool ShouldCompilePermutation(const FGlobalShaderPermutationParameters& Parameters)
+	{
+		return RHISupportsComputeShaders(Parameters.Platform);
+	}
+
+	static void ModifyCompilationEnvironment(const FGlobalShaderPermutationParameters& Parameters, FShaderCompilerEnvironment& OutEnvironment)
+	{
+		FGlobalShader::ModifyCompilationEnvironment(Parameters, OutEnvironment);
+		OutEnvironment.SetDefine(TEXT("THREAD_COUNT"), NIAGARA_HAIR_STRANDS_THREAD_COUNT);
+	}
+};
+
+IMPLEMENT_GLOBAL_SHADER(FClearPressureGridCS, "/Plugin/Experimental/HairStrands/Private/NiagaraClearPressureGrid.usf", "MainCS", SF_Compute);
+
+static void AddClearPressureGridPass(
+	FRDGBuilder& GraphBuilder,
+	FRHIShaderResourceView* GridCurrentBuffer,
+	FRHIUnorderedAccessView* GridDestinationBuffer,
+	const FIntVector& GridSize, const bool CopyPressure)
+{
+	const uint32 GroupSize = NIAGARA_HAIR_STRANDS_THREAD_COUNT;
+	const uint32 NumElements = (GridSize.X + 1) * (GridSize.Y + 1) * (GridSize.Z + 1);
+
+	FClearPressureGridCS::FParameters* Parameters = GraphBuilder.AllocParameters<FClearPressureGridCS::FParameters>();
+	Parameters->GridCurrentBuffer = GridCurrentBuffer;
+	Parameters->GridDestinationBuffer = GridDestinationBuffer;
+	Parameters->GridSize = GridSize;
+	Parameters->CopyPressure = CopyPressure;
+
+	TShaderMap<FGlobalShaderType>* ShaderMap = GetGlobalShaderMap(ERHIFeatureLevel::SM5);
+
+	const uint32 DispatchCount = FMath::DivideAndRoundUp(NumElements, GroupSize);
+
+	TShaderMapRef<FClearPressureGridCS> ComputeShader(ShaderMap);
+	FComputeShaderUtils::AddPass(
+		GraphBuilder,
+		RDG_EVENT_NAME("ClearPressureGrid"),
+		*ComputeShader,
+		Parameters,
+		FIntVector(DispatchCount, 1, 1));
+}
+
+inline void ClearBuffer(FNDIPressureGridBuffer* CurrentGridBuffer, FNDIPressureGridBuffer* DestinationGridBuffer, const FIntVector& GridSize, const bool CopyPressure)
+{
+	FRHIUnorderedAccessView* DestinationGridBufferUAV = DestinationGridBuffer->GridDataBuffer.UAV;
+	FRHIShaderResourceView* CurrentGridBufferSRV = CurrentGridBuffer->GridDataBuffer.SRV;
+	FRHIUnorderedAccessView* CurrentGridBufferUAV = CurrentGridBuffer->GridDataBuffer.UAV;
+
+	if (DestinationGridBufferUAV != nullptr && CurrentGridBufferSRV != nullptr && CurrentGridBufferUAV != nullptr)
+	{
+		const FIntVector LocalGridSize = GridSize;
+		const bool LocalCopyPressure = CopyPressure;
+
+		ENQUEUE_RENDER_COMMAND(ClearPressureGrid)(
+			[DestinationGridBufferUAV, CurrentGridBufferSRV, CurrentGridBufferUAV, LocalGridSize, LocalCopyPressure]
+		(FRHICommandListImmediate& RHICmdListImm)
+		{
+			FRDGBuilder GraphBuilder(RHICmdListImm);
+
+			RHICmdListImm.TransitionResource(EResourceTransitionAccess::EWritable, EResourceTransitionPipeline::EComputeToCompute, DestinationGridBufferUAV);
+			RHICmdListImm.TransitionResource(EResourceTransitionAccess::EReadable, EResourceTransitionPipeline::EComputeToCompute, CurrentGridBufferUAV);
+
+			AddClearPressureGridPass(
+				GraphBuilder,
+				CurrentGridBufferSRV, DestinationGridBufferUAV, LocalGridSize, LocalCopyPressure);
+
+			GraphBuilder.Execute();
+		});
+	}
+}
+
+//------------------------------------------------------------------------------------------------------------
+
+void FNDIPressureGridProxy::DeferredDestroy()
+{
+	for (const FNiagaraSystemInstanceID& Sys : DeferredDestroyList)
+	{
+		SystemInstancesToProxyData.Remove(Sys);
+	}
+
+	DeferredDestroyList.Empty();
+}
+
+void FNDIPressureGridProxy::ConsumePerInstanceDataFromGameThread(void* PerInstanceData, const FNiagaraSystemInstanceID& Instance)
+{
+	FNDIPressureGridData* SourceData = static_cast<FNDIPressureGridData*>(PerInstanceData);
+	FNDIPressureGridData* TargetData = &(SystemInstancesToProxyData.FindOrAdd(Instance));
+
+	ensure(TargetData);
+	if (TargetData)
+	{
+		TargetData->WorldTransform = SourceData->WorldTransform;
+		TargetData->WorldInverse = SourceData->WorldInverse;
+		TargetData->GridOrigin = SourceData->GridOrigin;
+		TargetData->GridSize = SourceData->GridSize;
+		TargetData->DestinationGridBuffer = SourceData->DestinationGridBuffer;
+		TargetData->CurrentGridBuffer = SourceData->CurrentGridBuffer;
+	}
+	else
+	{
+		UE_LOG(LogPressureGrid, Log, TEXT("ConsumePerInstanceDataFromGameThread() ... could not find %s"), *FNiagaraUtilities::SystemInstanceIDToString(Instance));
+	}
+}
+
+void FNDIPressureGridProxy::InitializePerInstanceData(const FNiagaraSystemInstanceID& SystemInstance)
+{
+	check(IsInRenderingThread());
+
+	FNDIPressureGridData* TargetData = SystemInstancesToProxyData.Find(SystemInstance);
+	if (TargetData != nullptr)
+	{
+		DeferredDestroyList.Remove(SystemInstance);
+	}
+	else
+	{
+		TargetData = &SystemInstancesToProxyData.Add(SystemInstance);
+	}
+}
+
+void FNDIPressureGridProxy::DestroyPerInstanceData(NiagaraEmitterInstanceBatcher* Batcher, const FNiagaraSystemInstanceID& SystemInstance)
+{
+	check(IsInRenderingThread());
+
+	DeferredDestroyList.Add(SystemInstance);
+	Batcher->EnqueueDeferredDeletesForDI_RenderThread(this->AsShared());
+}
+
 void FNDIPressureGridProxy::PreStage(FRHICommandList& RHICmdList, const FNiagaraDataInterfaceSetArgs& Context)
 {
 	FNDIPressureGridData* ProxyData =
@@ -829,7 +1037,7 @@ void FNDIPressureGridProxy::PreStage(FRHICommandList& RHICmdList, const FNiagara
 	{
 		if (!Context.IsIterationStage)
 		{
-			ProxyData->DestinationGridBuffer->ClearBuffers(RHICmdList);
+			ClearBuffer(ProxyData->CurrentGridBuffer, ProxyData->DestinationGridBuffer, ProxyData->GridSize, true);
 		}
 		else
 		{
@@ -853,7 +1061,13 @@ void FNDIPressureGridProxy::PostStage(FRHICommandList& RHICmdList, const FNiagar
 	{
 		if (ProxyData != nullptr)
 		{
-			ProxyData->SwapBuffers();
+			ENQUEUE_RENDER_COMMAND(CopyPressureGrid)(
+				[ProxyData](FRHICommandListImmediate& RHICmdList)
+			{
+				FRHICopyTextureInfo CopyInfo;
+				RHICmdList.CopyTexture(ProxyData->DestinationGridBuffer->GridDataBuffer.Buffer,
+					ProxyData->CurrentGridBuffer->GridDataBuffer.Buffer, CopyInfo);
+			});
 		}
 	}
 }
@@ -862,10 +1076,10 @@ void FNDIPressureGridProxy::ResetData(FRHICommandList& RHICmdList, const FNiagar
 {
 	FNDIPressureGridData* ProxyData = SystemInstancesToProxyData.Find(Context.SystemInstance);
 
-	if (ProxyData != nullptr)
+	if (ProxyData != nullptr && ProxyData->DestinationGridBuffer != nullptr && ProxyData->CurrentGridBuffer != nullptr)
 	{
-		ProxyData->CurrentGridBuffer->ClearBuffers(RHICmdList);
-		ProxyData->DestinationGridBuffer->ClearBuffers(RHICmdList);
+		ClearBuffer(ProxyData->CurrentGridBuffer, ProxyData->DestinationGridBuffer, ProxyData->GridSize, false);
+		ClearBuffer(ProxyData->DestinationGridBuffer, ProxyData->CurrentGridBuffer, ProxyData->GridSize, false);
 	}
 }
 
diff --git a/Engine/Plugins/Experimental/HairStrands/Source/HairStrandsNiagara/Public/NiagaraDataInterfaceHairStrands.h b/Engine/Plugins/Experimental/HairStrands/Source/HairStrandsNiagara/Public/NiagaraDataInterfaceHairStrands.h
--- a/Engine/Plugins/Experimental/HairStrands/Source/HairStrandsNiagara/Public/NiagaraDataInterfaceHairStrands.h
+++ b/Engine/Plugins/Experimental/HairStrands/Source/HairStrandsNiagara/Public/NiagaraDataInterfaceHairStrands.h
@@ -26,7 +26,7 @@ enum class EHairStrandsSize : uint8
 struct FNDIHairStrandsBuffer : public FRenderResource
 {
 	/** Set the asset that will be used to affect the buffer */
-	void SetHairAsset(const FHairStrandsDatas*  HairStrandsDatas, const FHairStrandsRestResource*  HairStrandsRestResource, 
+	void Initialize(const FHairStrandsDatas*  HairStrandsDatas, const FHairStrandsRestResource*  HairStrandsRestResource, 
 		const FHairStrandsDeformedResource*  HairStrandsDeformedResource, const FHairStrandsRootResource* HairStrandsRootResource );
 
 	/** Init the buffer */
@@ -38,36 +38,12 @@ struct FNDIHairStrandsBuffer : public FRenderResource
 	/** Get the resource name */
 	virtual FString GetFriendlyName() const override { return TEXT("FNDIHairStrandsBuffer"); }
 
-	/** Clear the bounding box buffer */
-	void ClearBuffer(FRHICommandList& RHICmdList);
-
 	/** Strand curves point offset buffer */
 	FRWBuffer CurvesOffsetsBuffer;
 
 	/** Deformed position buffer in case no ressource are there */
 	FRWBuffer DeformedPositionBuffer;
 
-	/**Rest triangle position of vertex A*/
-	FRWBuffer RestTrianglePositionABuffer;
-
-	/**Rest triangle position of vertex B*/
-	FRWBuffer RestTrianglePositionBBuffer;
-
-	/**Rest triangle position of vertex C*/
-	FRWBuffer RestTrianglePositionCBuffer;
-
-	/**Deformed triangle position of vertex A*/
-	FRWBuffer DeformedTrianglePositionABuffer;
-
-	/**Deformed triangle position of vertex B*/
-	FRWBuffer DeformedTrianglePositionBBuffer;
-
-	/**Deformed triangle position of vertex C*/
-	FRWBuffer DeformedTrianglePositionCBuffer;
-
-	/**Root barycentric coordinates */
-	FRWBuffer RootBarycentricCoordinatesBuffer;
-
 	/** Bounding Box Buffer*/
 	FRWBuffer BoundingBoxBuffer;
 
@@ -90,6 +66,12 @@ struct FNDIHairStrandsBuffer : public FRenderResource
 /** Data stored per strand base instance*/
 struct FNDIHairStrandsData
 {
+	/** Initialize the buffers */
+	bool Init(class UNiagaraDataInterfaceHairStrands* Interface, FNiagaraSystemInstance* SystemInstance);
+
+	/** Release the buffers */
+	void Release();
+
 	/** Cached World transform. */
 	FTransform WorldTransform;
 
@@ -166,7 +148,7 @@ public:
 	virtual void GetCommonHLSL(FString& OutHLSL) override;
 
 	/** Update the source component */
-	void UpdateSourceComponent(FNiagaraSystemInstance* SystemInstance);
+	void ExtractSourceComponent(FNiagaraSystemInstance* SystemInstance);
 
 	/** Check if the component is Valid */
 	bool IsComponentValid() const;
@@ -289,6 +271,15 @@ public:
 	/** Project the static collision constraint */
 	void ProjectStaticCollisionConstraint(FVectorVMContext& Context);
 
+	/** Solve the soft collision constraint */
+	void SolveSoftCollisionConstraint(FVectorVMContext& Context);
+
+	/** Project the soft collision constraint */
+	void ProjectSoftCollisionConstraint(FVectorVMContext& Context);
+
+	/** Setup the soft collision constraint */
+	void SetupSoftCollisionConstraint(FVectorVMContext& Context);
+
 	/** Compute the rest direction*/
 	void ComputeRestDirection(FVectorVMContext& Context);
 
@@ -400,7 +391,7 @@ struct FNDIHairStrandsProxy : public FNiagaraDataInterfaceProxy
 	virtual void ConsumePerInstanceDataFromGameThread(void* PerInstanceData, const FNiagaraSystemInstanceID& Instance) override;
 
 	/** Initialize the Proxy data strands buffer */
-	void InitializePerInstanceData(const FNiagaraSystemInstanceID& SystemInstance, FNDIHairStrandsBuffer* StrandsBuffer, const uint32 NumStrands, const uint8 StrandSize, const FVector& BoxCenter, const FVector& BoxExtent, const FTransform& WorldTransform);
+	void InitializePerInstanceData(const FNiagaraSystemInstanceID& SystemInstance);
 
 	/** Destroy the proxy data if necessary */
 	void DestroyPerInstanceData(NiagaraEmitterInstanceBatcher* Batcher, const FNiagaraSystemInstanceID& SystemInstance);
diff --git a/Engine/Plugins/Experimental/HairStrands/Source/HairStrandsNiagara/Public/NiagaraDataInterfacePhysicsAsset.h b/Engine/Plugins/Experimental/HairStrands/Source/HairStrandsNiagara/Public/NiagaraDataInterfacePhysicsAsset.h
--- a/Engine/Plugins/Experimental/HairStrands/Source/HairStrandsNiagara/Public/NiagaraDataInterfacePhysicsAsset.h
+++ b/Engine/Plugins/Experimental/HairStrands/Source/HairStrandsNiagara/Public/NiagaraDataInterfacePhysicsAsset.h
@@ -31,6 +31,7 @@ struct FNDIPhysicsAssetArrays
 	TArray<FVector4> CurrentTransform;
 	TArray<FVector4> InverseTransform;
 	TArray<FVector4> PreviousTransform;
+	TArray<FVector4> PreviousInverse;
 	TArray<FVector4> RestTransform;
 	TArray<FVector4> RestInverse;
 	TArray<FVector4> ElementExtent;
@@ -43,10 +44,10 @@ struct FNDIPhysicsAssetBuffer : public FRenderResource
 	bool IsValid() const;
 
 	/** Set the assets that will be used to affect the buffer */
-	void SetupArrays(const TWeakObjectPtr<class UPhysicsAsset>  PhysicsAsset, const TWeakObjectPtr<class USkeletalMeshComponent>  SkeletalMesh, const FTransform& InWorldTransform);
+	void Initialize(const TWeakObjectPtr<class UPhysicsAsset>  PhysicsAsset, const TWeakObjectPtr<class USkeletalMeshComponent>  SkeletalMesh, const FTransform& InWorldTransform);
 
 	/** Update the buffers */
-	void UpdateBuffers();
+	void Update();
 
 	/** Init the buffer */
 	virtual void InitRHI() override;
@@ -63,6 +64,9 @@ struct FNDIPhysicsAssetBuffer : public FRenderResource
 	/** Previous transform buffer */
 	FRWBuffer PreviousTransformBuffer;
 
+	/** Previous inverse buffer */
+	FRWBuffer PreviousInverseBuffer;
+
 	/** Inverse transform buffer*/
 	FRWBuffer InverseTransformBuffer;
 
@@ -91,8 +95,14 @@ struct FNDIPhysicsAssetBuffer : public FRenderResource
 /** Data stored per physics asset instance*/
 struct FNDIPhysicsAssetData
 {
+	/** Initialize the buffers */
+	bool Init(class UNiagaraDataInterfacePhysicsAsset* Interface, FNiagaraSystemInstance* SystemInstance);
+
+	/** Release the buffers */
+	void Release();
+
 	/** Physics asset Gpu buffer */
-	FNDIPhysicsAssetBuffer* AssetBuffer;
+	FNDIPhysicsAssetBuffer* PhysicsAssetBuffer;
 
 	/** Bounding box center */
 	FVector BoxOrigin;
@@ -143,6 +153,9 @@ public:
 	virtual void ProvidePerInstanceDataForRenderThread(void* DataForRenderThread, void* PerInstanceData, const FNiagaraSystemInstanceID& SystemInstance) override;
 	virtual void GetCommonHLSL(FString& OutHLSL) override;
 
+	/** Extract the source component */
+	void ExtractSourceComponent(FNiagaraSystemInstance* SystemInstance);
+
 	/** Get the number of boxes*/
 	void GetNumBoxes(FVectorVMContext& Context);
 
@@ -170,6 +183,9 @@ public:
 	/** Name of the previous transform buffer */
 	static const FString PreviousTransformBufferName;
 
+	/** Name of the previous inverse buffer */
+	static const FString PreviousInverseBufferName;
+
 	/** Name of the inverse transform buffer */
 	static const FString InverseTransformBufferName;
 
@@ -206,7 +222,7 @@ struct FNDIPhysicsAssetProxy : public FNiagaraDataInterfaceProxy
 	virtual void ConsumePerInstanceDataFromGameThread(void* PerInstanceData, const FNiagaraSystemInstanceID& Instance) override;
 
 	/** Initialize the Proxy data strands buffer */
-	void InitializePerInstanceData(const FNiagaraSystemInstanceID& SystemInstance, FNDIPhysicsAssetBuffer* AssetBuffer, const FVector& BoxOrigin, const FVector& BoxExtent);
+	void InitializePerInstanceData(const FNiagaraSystemInstanceID& SystemInstance);
 
 	/** Destroy the proxy data if necessary */
 	void DestroyPerInstanceData(NiagaraEmitterInstanceBatcher* Batcher, const FNiagaraSystemInstanceID& SystemInstance);
diff --git a/Engine/Plugins/Experimental/HairStrands/Source/HairStrandsNiagara/Public/NiagaraDataInterfacePressureGrid.h b/Engine/Plugins/Experimental/HairStrands/Source/HairStrandsNiagara/Public/NiagaraDataInterfacePressureGrid.h
--- a/Engine/Plugins/Experimental/HairStrands/Source/HairStrandsNiagara/Public/NiagaraDataInterfacePressureGrid.h
+++ b/Engine/Plugins/Experimental/HairStrands/Source/HairStrandsNiagara/Public/NiagaraDataInterfacePressureGrid.h
@@ -13,7 +13,7 @@
 struct FNDIPressureGridBuffer : public FRenderResource
 {
 	/** Set the grid size */
-	void SetGridSize(const FUintVector4 GridSize);
+	void Initialize(const FIntVector GridSize);
 
 	/** Init the buffer */
 	virtual void InitRHI() override;
@@ -21,9 +21,6 @@ struct FNDIPressureGridBuffer : public FRenderResource
 	/** Release the buffer */
 	virtual void ReleaseRHI() override;
 
-	/** Clear all UAV*/
-	void ClearBuffers(FRHICommandList& RHICmdList);
-
 	/** Get the resource name */
 	virtual FString GetFriendlyName() const override { return TEXT("FNDIPressureGridBuffer"); }
 
@@ -31,21 +28,27 @@ struct FNDIPressureGridBuffer : public FRenderResource
 	FTextureRWBuffer3D GridDataBuffer;
 
 	/** Grid size that will be used for the collision*/
-	FUintVector4 GridSize;
+	FIntVector GridSize;
 };
 
 /** Data stored per strand base instance*/
 struct FNDIPressureGridData
 {
 	/** Swap the current and the destination data */
-	void SwapBuffers();
+	void Swap();
+
+	/** Initialize the buffers */
+	bool Init(class UNiagaraDataInterfacePressureGrid* Interface, FNiagaraSystemInstance* SystemInstance);
+
+	/** Release the buffers */
+	void Release();
 
 	/** Grid Origin */
 	FVector4 GridOrigin;
 
 	/** Grid Size */
-	FUintVector4 GridSize;
-	
+	FIntVector GridSize;
+
 	/** World Transform */
 	FMatrix WorldTransform;
 
@@ -69,15 +72,23 @@ public:
 
 	/** Grid size along the X axis. */
 	UPROPERTY(EditAnywhere, Category = "Spawn")
-	FIntVector GridSize;
+		FIntVector GridSize;
 
-	/** Grid size along the Y axis. */
-	UPROPERTY(EditAnywhere, Category = "Spawn")
-	FVector GridOrigin;
+	/** Min stage iteration number */
+	UPROPERTY(EditAnywhere, Category = "Stages")
+		int MinIteration;
 
-	/** Grid size along the Z axis. */
-	UPROPERTY(EditAnywhere, Category = "Spawn")
-	float GridLength;
+	/** Max stage iteration number  */
+	UPROPERTY(EditAnywhere, Category = "Stages")
+		int MaxIteration;
+
+	/** Min stage output number */
+	UPROPERTY(EditAnywhere, Category = "Stages")
+		int MinOutput;
+
+	/** Max stage output number  */
+	UPROPERTY(EditAnywhere, Category = "Stages")
+		int MaxOutput;
 
 	/** UObject Interface */
 	virtual void PostInitProperties() override;
@@ -109,7 +120,7 @@ public:
 	void ProjectVelocityField(FVectorVMContext& Context);
 
 	/** Compute the cell position*/
-	void GetCellPosition(FVectorVMContext& Context);
+	void GetNodePosition(FVectorVMContext& Context);
 
 	/** Transfer the cell velocity */
 	void TransferCellVelocity(FVectorVMContext& Context);
@@ -126,6 +137,12 @@ public:
 	/** Update the grid transform */
 	void UpdateGridTransform(FVectorVMContext& Context);
 
+	/** Get the grid velocity */
+	void GetGridVelocity(FVectorVMContext& Context);
+
+	/** Add velocity to the grid */
+	void AddGridVelocity(FVectorVMContext& Context);
+
 	/** Name of the grid current buffer */
 	static const FString GridCurrentBufferName;
 
@@ -161,12 +178,12 @@ struct FNDIPressureGridProxy : public FNiagaraDataInterfaceProxy
 	/** Get the data that will be passed to render*/
 	virtual void ConsumePerInstanceDataFromGameThread(void* PerInstanceData, const FNiagaraSystemInstanceID& Instance) override;
 
-	/** Initialize the Proxy data strands buffer */
-	void InitializePerInstanceData(const FNiagaraSystemInstanceID& SystemInstance, FNDIPressureGridBuffer* CurrentGridBuffer, FNDIPressureGridBuffer* DestinationGridBuffer, const FVector4& GridOrigin, const FUintVector4& GridSize);
-
 	/** Destroy the proxy data if necessary */
 	void DestroyPerInstanceData(NiagaraEmitterInstanceBatcher* Batcher, const FNiagaraSystemInstanceID& SystemInstance);
 
+	/** Initialize the proxy instance data */
+	void InitializePerInstanceData(const FNiagaraSystemInstanceID& SystemInstance);
+
 	/** Launch all pre stage functions */
 	virtual void PreStage(FRHICommandList& RHICmdList, const FNiagaraDataInterfaceSetArgs& Context) override;
 
diff --git a/Engine/Plugins/FX/Niagara/Source/NiagaraCore/Private/NiagaraCustomVersion.cpp b/Engine/Plugins/FX/Niagara/Source/NiagaraCore/Private/NiagaraCustomVersion.cpp
--- a/Engine/Plugins/FX/Niagara/Source/NiagaraCore/Private/NiagaraCustomVersion.cpp
+++ b/Engine/Plugins/FX/Niagara/Source/NiagaraCore/Private/NiagaraCustomVersion.cpp
@@ -8,4 +8,4 @@ const FGuid FNiagaraCustomVersion::GUID(0xFCF57AFA, 0x50764283, 0xB9A9E658, 0xFF
 // Register the custom version with core
 FCustomVersionRegistration GRegisterNiagaraCustomVersion(FNiagaraCustomVersion::GUID, FNiagaraCustomVersion::LatestVersion, TEXT("NiagaraVer"));
 
-const FGuid FNiagaraCustomVersion::LatestScriptCompileVersion(0x4953F984, 0x8FEC104D, 0xAF87DB5C, 0x44A7F79F);
+const FGuid FNiagaraCustomVersion::LatestScriptCompileVersion(0xA6B5A395, 0xF276467E, 0x9D03C2C2, 0x99CAE1E3);
diff --git a/Engine/Shaders/Private/DeferredLightPixelShaders.usf b/Engine/Shaders/Private/DeferredLightPixelShaders.usf
--- a/Engine/Shaders/Private/DeferredLightPixelShaders.usf
+++ b/Engine/Shaders/Private/DeferredLightPixelShaders.usf
@@ -9,6 +9,7 @@
 #include "Common.ush"
 #if USE_HAIR_LIGHTING
 #include "HairStrands/HairStrandsCommon.ush"
+#include "HairStrands/HairStrandsDebugCommon.ush"
 #endif
 
 #include "DeferredShadingCommon.ush"
@@ -106,12 +107,12 @@ StructuredBuffer<FNodeData> HairVisibilityNodeData;
 Texture2D<float4> ScreenShadowMaskSubPixelTexture;
 
 uint HairTransmittanceBufferMaxCount;
-StructuredBuffer<float4> HairTransmittanceBuffer;
+StructuredBuffer<FPackedHairTransmittanceMask> HairTransmittanceBuffer;
 Texture3D<float4> HairLUTTexture;
 SamplerState HairLUTSampler;
 uint HairShadowMaskValid;
 uint HairComponents; // R, TT, TRT, LocalScattering, GlobalScattering, MultipleScattering
-
+float HairDualScatteringRoughnessOverride;
 
 float Luma(float3 Color)
 {
@@ -172,7 +173,7 @@ void ComputeHairLighting(
 		HairScreenSpaceData.GBuffer.Metallic = 0;
 		HairScreenSpaceData.GBuffer.Depth = SampleDepth;
 		HairScreenSpaceData.GBuffer.GBufferAO = 1;
-		HairScreenSpaceData.GBuffer.CustomData = 0;
+		HairScreenSpaceData.GBuffer.CustomData = float4(HairDualScatteringRoughnessOverride, 0, 0, 0);;
 		HairScreenSpaceData.GBuffer.IndirectIrradiance = 0;
 		HairScreenSpaceData.GBuffer.PrecomputedShadowFactors = 1;
 
@@ -204,7 +205,8 @@ void ComputeHairLighting(
 
 		const float Dither = InterleavedGradientNoise(InputParams.PixelPos, View.StateFrameIndexMod8);
 		const uint3 PixelCoord = uint3(floor(InputParams.PixelPos), 0);
-		const float4 LightAttenuation = HairShadowMaskValid ? ScreenShadowMaskSubPixelTexture.Load(PixelCoord) : 1;
+		float4 LightAttenuation = HairShadowMaskValid ? ScreenShadowMaskSubPixelTexture.Load(PixelCoord) : 1;
+		LightAttenuation = min(LightAttenuation, LightData.HairTransmittance.OpaqueVisibility.xxxx);
 
 		FRectTexture RectTexture = InitRectTexture(DeferredLightUniforms.SourceTexture);
 		float SurfaceShadow = 1.0f;
diff --git a/Engine/Shaders/Private/HairStrands/HairScatter.usf b/Engine/Shaders/Private/HairStrands/HairScatter.usf
new file mode 100644
--- /dev/null
+++ b/Engine/Shaders/Private/HairStrands/HairScatter.usf
@@ -0,0 +1,467 @@
+// Copyright Epic Games, Inc. All Rights Reserved.
+
+#include "HairStrandsVoxelPageCommonStruct.ush"
+
+#include "../Common.ush"
+//#include "../CommonViewUniformBuffer.ush"
+
+#include "HairStrandsCommon.ush"
+#define SUPPORT_CONTACT_SHADOWS 0
+
+#include "../SceneTexturesCommon.ush"
+#include "HairStrandsDebugCommon.ush"
+#include "HairStrandsDeepShadowCommon.ush"
+#include "../DeferredShadingCommon.ush"
+#include "../ShadingModels.ush"
+#include "HairStrandsVisibilityCommon.ush"
+#include "HairStrandsDeepTransmittanceCommon.ush"
+#include "HairStrandsVoxelPageCommon.ush"
+
+#define DEBUG_ENABLE 0
+
+#define VOXEL_TRAVERSAL_DEBUG DEBUG_ENABLE
+#define VOXEL_TRAVERSAL_TYPE VOXEL_TRAVERSAL_LINEAR
+#include "HairStrandsVoxelPageTraversal.ush"
+#include "../HairBsdf.ush"
+#include "../ShaderDrawDebug.ush"
+
+////////////////////////////////////////////////////////////////////////////////////////////////////////////////////////////////////////////////
+// INLINE OF: #include "DeepTransmittanceDualScattering.ush"
+
+float Hair_g2(float B,float Theta)
+{
+	#if 0
+	return Hair_g(B,Theta);
+	#else
+	// Unnormalized Gaussian
+	return exp( -0.5 * Pow2( Theta ) / (B*B) );
+	#endif
+}
+
+// TODO; clean up the dual scattering file to remove the dependencies
+// Dual scattering computation are done here for faster iteration (i.e., does not invalidate tons of shaders)
+void ComputeDualScatteringTerms(
+	const FHairTransmittanceData TransmittanceData,
+	const float Roughness,
+	const float3 V,
+	const float3 L,
+	const float3 T,
+	const uint HairComponents,
+	inout float3 OutGlobalScattering,
+	inout float3 OutLocalScattering)
+{
+	const float SinThetaL = clamp(dot(T, L), -1, 1);
+	const float SinThetaV = clamp(dot(T, V), -1, 1);
+
+	const float MaxAverageScatteringValue = 0.99f;
+
+	// Straight implementation of the dual scattering paper 
+    const float3 af	 = min(MaxAverageScatteringValue.xxx, TransmittanceData.A_front);
+	const float3 af2 = Pow2(af);
+	const float3 ab  = min(MaxAverageScatteringValue.xxx, TransmittanceData.A_back);
+	const float3 ab2 = Pow2(ab);
+    const float3 OneMinusAf2 = 1 - af2;
+
+	const float3 A1 = ab * af2 / OneMinusAf2;
+	const float3 A3 = ab * ab2 * af2 / (OneMinusAf2*Pow2(OneMinusAf2));
+	const float3 Ab = A1 + A3;
+
+	// todo_hair: Should use the average roughness after k scattering event instead of the BSDF roughness
+	const float Beta_R	 = Pow2( Roughness );
+	const float Beta_TT	 = Pow2( Roughness / 2 ); 
+	const float Beta_TRT = Pow2( Roughness * 2 ); 
+
+	const float Shift     = 0.035;
+	const float Shift_R   =-0.035*2;
+	const float Shift_TT  = 0.035;
+	const float Shift_TRT = 0.035*4;
+
+	// Average density factor (This is the constant used in the original paper)
+	const float df = 0.7f;
+	const float db = 0.7f;
+
+	// This is a coarse approximation of eq. 13. Normally, Beta_f should be weighted by the 'normalized' 
+	// R, TT, and TRT terms 
+	const float3 af_weights = af / (af.r + af.g + af.b);
+    const float3 Beta_f  = dot(float3(Beta_R, Beta_TT, Beta_TRT), af_weights);
+	const float3 Beta_f2 = Beta_f*Beta_f;
+    const float3 sigma_f = Beta_f * max(1, TransmittanceData.HairCount);
+
+    const float Theta_h = (asin(SinThetaL) + asin(SinThetaV)) * 0.5f;
+
+	// Global scattering spread 'Sf'
+	// It is unclear if this is correct to use a normalized Gaussian as in the original paper, as with an increase 
+	// number of hair count the Gaussian will become more and more spread, and have less and less contribution. 
+	// This is correct from an energy point of view, as long as the scattering is integrated over multiple scattering 
+	// event, which is not our case here.
+	// This is why there is a non-normalized Gaussian version in case this does not fly properly with regular 
+	// normalized Gaussian
+	// const float3 Sf = 1;
+	// const float3 Sf = Hair_g_unnormalized(sigma_f, Theta_h) / PI; 
+    float3 Sf = float3(	Hair_g2(sigma_f.r, Theta_h), 
+						Hair_g2(sigma_f.g, Theta_h), 
+						Hair_g2(sigma_f.b, Theta_h)) / PI;
+	const float3 Tf = TransmittanceData.Transmittance;
+
+	// Overall shift due to the various local scatteing event (all shift & roughnesss should vary with color)
+	//const float3 shift_f = 0.33f * (Shift_R + Shift_TT + Shift_TRT);
+    const float3 shift_f = dot(float3(Shift_R, Shift_TT, Shift_TRT), af_weights);
+	const float3 shift_b = shift_f;
+    const float3 delta_b = 0; //shift_b * (1 - 2*ab2 / Pow2(1 - af2)) * shift_f * (2 * Pow2(1 - af2) + 4*af2*ab2)/Pow3(1-af2);
+
+	const float3 ab_weights = ab / (ab.r + ab.g + ab.b);
+	const float3 Beta_b  = dot(float3(Beta_R, Beta_TT, Beta_TRT), ab_weights);
+	const float3 Beta_b2 = Beta_f2;
+
+    const float3 sigma_b = (1 + db*af2) * (ab*sqrt(2*Beta_f2 + Beta_b2) + ab*ab2*sqrt(2*Beta_f2 + Beta_b2)) / (ab + ab*ab2*(2*Beta_f + 3*Beta_b));
+
+	// Local scattering Spread 'Sb'
+    float3 Sb = float3(	Hair_g2(Pow2(sigma_b.r), Theta_h - delta_b.r),
+						Hair_g2(Pow2(sigma_b.g), Theta_h - delta_b.g),
+						Hair_g2(Pow2(sigma_b.b), Theta_h - delta_b.b))  / PI;
+
+	// Override the spread function to cancel them
+	//Sf = 1;
+	//Sb = 1;
+
+    // Different variant for managing sefl-occlusion issue for global scattering
+	const bool bIsOccluded = TransmittanceData.HairCount > 0;
+	const float3 GlobalScattering = bIsOccluded ? Tf * Sf * df : 1;
+	//const float3 GlobalScattering = TransmittanceData.HairCount >= 1 ? Tf * Sf * df : 1;
+    //const float3 GlobalScattering = lerp(1, Tf * Sf * df, saturate(TransmittanceData.HairCount));
+	//const float3 GlobalScattering = Tf * Sf * df;
+
+	const float3 LocalScattering  = 2 * Ab * Sb * db;
+	
+	// Decode if the Local/Global component should be enabled or not (for debug purpose)
+	const float LocalWeights  = (HairComponents & HAIR_COMPONENT_LS) > 0 ? 1 : 0;
+
+	OutGlobalScattering = (HairComponents & HAIR_COMPONENT_GS) > 0 ? GlobalScattering : 1;
+	OutLocalScattering  = LocalScattering * LocalWeights;
+
+	// Final computation is done in ShadingModels.ush with the following formula
+	// GlobalScattering * (Fs + LocalScattering) * TransmittanceData.OpaqueVisibility;
+}
+
+////////////////////////////////////////////////////////////////////////////////////////////////////////////////////////////////////////////////
+
+float3 ReconstructWorldPositionFromDepth(float2 UV, float Depth)
+{
+	float2 ScreenPosition = (UV - View.ScreenPositionScaleBias.wz) / View.ScreenPositionScaleBias.xy;
+	float4 HomogeneousWorldPosition = mul(float4(ScreenPosition * Depth, Depth, 1), View.ScreenToWorld);
+	float3 WorldPosition = HomogeneousWorldPosition.xyz / HomogeneousWorldPosition.w;
+	return WorldPosition;
+}
+
+////////////////////////////////////////////////////////////////////////////////////////////////////////////////////////////////////////////////
+
+float2 						OutputResolution;
+float  						GatherWorldRadius;  // cm
+float						PixelRadiusAtDepth1;
+
+uint 						Enable;
+uint 						Debug;
+uint						SampleCount;
+uint 						HairComponents;
+
+Texture2D<float4> 			DiffusionInputTexture;
+
+Texture2D<uint4>  			CategorizationTexture;
+Texture2D<uint> 			VisibilityNodeIndex;
+StructuredBuffer<FNodeData> VisibilityNodeData;
+
+Texture3D<float4>			HairLUTTexture;
+Texture3D<float4>			HairEnergyLUTTexture;
+SamplerState 				LinearSampler;
+
+FHairTransmittanceMask ComputeTransmittanceVoxel(
+	const uint2 PixelCoord,
+	float3 WorldPosition,
+	float3 SamplePosition,
+	uint MacroGroupId,
+	FGBufferData GBuffer,
+	FTransmittanceSettings Settings,
+	Texture3D<float4>	HairLUTTexture,
+	SamplerState		HairLUTSampler,
+	bool bDebugEnabled)
+{
+	const float3 L = normalize(SamplePosition - WorldPosition);;
+	const float3 T = GBuffer.WorldNormal;
+	const float SinLightAngle = dot(L, T);
+
+	// Depth bias
+	// Origin is shifted voxels away towards the light + a constant bias of the size of the voxel
+	float MaxVoxelSize = VirtualVoxel.VoxelWorldSize;
+
+	const float3 SampleRandom = GetHairVoxelJitter(PixelCoord, View.StateFrameIndexMod8);
+	const float3 DepthBias = MaxVoxelSize * (Settings.DeepShadowDepthBiasScale*L + SampleRandom);
+
+	FVirtualVoxelCommonDesc CommonDesc;
+	CommonDesc.PageCountResolution = VirtualVoxel.PageCountResolution;
+	CommonDesc.VoxelWorldSize = VirtualVoxel.VoxelWorldSize;
+	CommonDesc.PageTextureResolution = VirtualVoxel.PageTextureResolution;
+	CommonDesc.PageResolution = VirtualVoxel.PageResolution;
+
+	const FPackedVirtualVoxelNodeDesc PackedNode = VirtualVoxel.NodeDescBuffer[MacroGroupId];
+	const FVirtualVoxelNodeDesc NodeDesc = UnpackVoxelNode(PackedNode, VirtualVoxel.PageResolution);
+
+	FHairTraversalSettings TraversalSettings = InitHairTraversalSettings();
+	TraversalSettings.DensityScale = VirtualVoxel.DensityScale;
+	TraversalSettings.CountThreshold = GetOpaqueVoxelValue();
+	TraversalSettings.DistanceThreshold = 100000;
+	TraversalSettings.bDebugEnabled = bDebugEnabled;
+	TraversalSettings.SteppingScale = VirtualVoxel.SteppingScale;
+	TraversalSettings.Random = SampleRandom;
+
+	FHairTraversalResult Result = ComputeHairCountVirtualVoxel(
+		WorldPosition + DepthBias,
+		SamplePosition,
+		CommonDesc,
+		NodeDesc,
+		VirtualVoxel.PageIndexBuffer,
+		VirtualVoxel.PageTexture,
+		TraversalSettings);
+
+	return ComputeTransmittanceFromHairCount(Result.HairCount, Result.Visibility, GBuffer, SinLightAngle, HairLUTTexture, HairLUTSampler);
+}
+
+FGBufferData GetGBufferData(uint2 PixelCoord)
+{
+	FGBufferData GBuffer;
+	GBuffer.ShadingModelID = SHADINGMODELID_HAIR;
+	GBuffer.WorldNormal = float3(1, 0, 0);
+	GBuffer.BaseColor = float3(0, 0, 0);
+	GBuffer.Roughness = 1;
+	GBuffer.Specular = 0;
+
+	uint MacroGroupIndex = 0;
+	const FNodeDesc NodeDesc = DecodeNodeDesc(VisibilityNodeIndex.Load(uint3(PixelCoord, 0)));
+	if (NodeDesc.Count > 0)
+	{
+		float ClosestDepth = 0; // Inverse-Z
+		[loop]
+		for (uint SampleIt = 0; SampleIt < NodeDesc.Count; ++SampleIt)
+		{
+			const uint LocalOffset = NodeDesc.Offset + SampleIt;
+			const FNodeData NodeData = VisibilityNodeData[LocalOffset];
+			float			SampleDepth;
+			FVisibilityData SampleVisData;
+			float3			SampleBaseColor;
+			float			SampleRoughness;
+			float			SampleSpecular;
+			DecodeNodeData(NodeData, SampleDepth, SampleVisData, SampleBaseColor, SampleRoughness, SampleSpecular);
+
+			if (SampleDepth >= ClosestDepth) // InverseZ
+			{
+				GBuffer.ShadingModelID = SHADINGMODELID_HAIR;
+				GBuffer.WorldNormal = SampleVisData.Tangent;
+				GBuffer.BaseColor = SampleBaseColor;
+				GBuffer.Roughness = SampleRoughness;
+				GBuffer.Specular = SampleSpecular;
+				ClosestDepth = SampleDepth;
+
+				// todo
+				MacroGroupIndex = 0; //clamp(SampleVisData.MacroGroupId, 0, MAX_HAIR_MACROGROUP_COUNT - 1);
+			}
+		}
+	}
+
+	return GBuffer;
+}
+
+void MainPS(
+	in FScreenVertexOutput Input,
+	out float4 OutColor  : SV_Target0,
+	out float4 OutSubPixelColor  : SV_Target1,
+	out float4 OutDiffusionTexture : SV_Target2)
+{
+	OutColor = 0;
+	OutSubPixelColor = 0;	
+	if (Enable == 0)
+		return;
+
+
+	const uint2 PixelCoord = floor(Input.Position.xy);
+	const float2 UV = Input.UV;
+	float4 HairColor = 0;
+	const bool bDebugEnabled = PixelCoord.x == uint(GetCursorPos().x) && PixelCoord.y == uint(GetCursorPos().y);
+	const float3 TransmittanceRandom = float3(InterleavedGradientNoise(PixelCoord.xy, 1), InterleavedGradientNoise(PixelCoord.xy, 2), InterleavedGradientNoise(PixelCoord.xy, 3));
+	const FCategorizationData CatData = DecodeCategorizationData(CategorizationTexture.Load(uint3(PixelCoord,0)));
+
+	const bool bIsCoveredByHair = CatData.PixelCoverage > 0;
+	if (!bIsCoveredByHair)
+		discard;
+
+	const float PixelDepth = ConvertFromDeviceZ(CatData.ClosestDepth);
+	const float3 WorldPosition = ReconstructWorldPositionFromDepth(UV, PixelDepth);
+	const float3 V = normalize(View.WorldCameraOrigin - WorldPosition);
+
+	FGBufferData GBuffer = GetGBufferData(PixelCoord);
+
+	const float3 T = GBuffer.WorldNormal;
+	const float3 RandOffset = float3(
+		InterleavedGradientNoise(PixelCoord.xy, 5), 
+		InterleavedGradientNoise(PixelCoord.xy, 6), 
+		InterleavedGradientNoise(PixelCoord.xy, 7));
+
+	const uint IndirectSampleCount = clamp(SampleCount, 0, 32);
+	const float GatherRadiusInPixel = clamp(GatherWorldRadius / (PixelRadiusAtDepth1 * PixelDepth), 0, 300) * 0.5f;
+
+	#if DEBUG_ENABLE
+	if (bDebugEnabled)
+	{
+		const float3 Extent = 0.5;
+		const float4 DebugColor = float4(1, 0, 1, 1);
+		AddAABB(WorldPosition - GatherWorldRadius, WorldPosition + GatherWorldRadius, DebugColor);
+	}
+	#endif
+
+	float4 AccIndirectLuminance = 0;
+	[loop]
+	for (uint IndirectSampleIt=0; IndirectSampleIt < IndirectSampleCount; ++IndirectSampleIt)
+	{	
+		// [0..1]^2
+		const float2 Rand2 = Hammersley(IndirectSampleIt, IndirectSampleCount, uint2(0, 0)) + RandOffset.xy;
+
+		const float2 SampleOffset	= frac(Rand2) * 2 - float2(1,1);
+		const uint2 SamplePixelPos	= clamp(PixelCoord + SampleOffset * GatherRadiusInPixel, 0, OutputResolution-1);
+		const float2 SampleUV		= SamplePixelPos / OutputResolution;
+
+		const float4 SampleColor = DiffusionInputTexture.Load(uint3(SamplePixelPos, 0));
+		if(SampleColor.w == 0)
+			continue;
+		const FCategorizationData SampleCatData = DecodeCategorizationData(CategorizationTexture.Load(uint3(SamplePixelPos,0)));
+		const float  SampleDepth    = ConvertFromDeviceZ(SampleCatData.ClosestDepth);
+		const float3 SampleWorldPos = ReconstructWorldPositionFromDepth(SampleUV, SampleDepth);
+		const float3 L 				= normalize(SampleWorldPos - WorldPosition);
+		const float SinLightAngle	= dot(L, T);
+
+		const float Distance = distance(SampleWorldPos, WorldPosition);
+		FGBufferData SampleGBuffer = GetGBufferData(SamplePixelPos);
+		const float3 SampleT = SampleGBuffer.WorldNormal;
+		const float SampleSinLightAngle = dot(L, SampleT);
+
+		#if DEBUG_ENABLE
+		if (bDebugEnabled)
+		{
+			const float3 Extent = 0.1;
+			const float4 DebugColor = float4(1, 1, 0, 1);
+			AddAABB(SampleWorldPos - Extent, SampleWorldPos + Extent, DebugColor);
+		}
+		#endif	
+
+		// Weighted by the BSDF and dual scattering
+		{
+			// Compute transmittance with voxel data
+			FHairTransmittanceMask TransmittanceMask = InitHairTransmittanceMask();
+			#if 1
+			{
+				FTransmittanceSettings Settings;
+				Settings.LightDirection = L;
+				Settings.LightPosition = float4(SampleWorldPos,1);
+				Settings.DeepShadowDensityScale = 1; // Voxel_DensityScale;
+				Settings.DeepShadowDepthBiasScale = 1; // Voxel_DepthBiasScale;
+				Settings.Random = TransmittanceRandom;
+
+				const uint MacroGroupId = 0;
+				TransmittanceMask = ComputeTransmittanceVoxel(
+					PixelCoord,
+					WorldPosition,
+					SampleWorldPos,
+					MacroGroupId,
+					GBuffer,
+					Settings,
+					HairLUTTexture,
+					LinearSampler,
+					bDebugEnabled && DEBUG_ENABLE);
+			}
+			#endif
+
+
+			// Use dual scattering
+			#if 0
+			FHairTransmittanceData TransmittanceData = InitHairTransmittanceData();
+			{
+				TransmittanceData.Transmittance = TransmittanceMask.Transmittance;
+				TransmittanceData.HairCount = TransmittanceMask.HairCount;
+				TransmittanceData.GlobalScattering = TransmittanceMask.Transmittance;
+				TransmittanceData.LocalScattering = 0;
+				TransmittanceData.ScatteringComponent = HairComponents | HAIR_COMPONENT_MULTISCATTER;
+				TransmittanceData.OpaqueVisibility = TransmittanceMask.Visibility;
+
+				{
+					// Compute the transmittance based on precompute Hair transmittance LUT
+					HairAverageScattering AverageScattering = SampleHairLUT(HairLUTTexture, LinearSampler, GBuffer.BaseColor, GBuffer.Roughness, SinLightAngle);
+					TransmittanceData.A_front = AverageScattering.A_front;
+					TransmittanceData.A_back  = AverageScattering.A_back;
+				}
+
+				ComputeDualScatteringTerms(
+					TransmittanceData, 
+					GBuffer.Roughness, 
+					V,
+					L,
+					T,
+					HairComponents,
+					TransmittanceData.GlobalScattering, 
+					TransmittanceData.LocalScattering);
+			
+				TransmittanceData.LocalScattering = 1;
+			}
+			#endif
+
+			// Use average energy
+			#if 1
+			FHairTransmittanceData TransmittanceData = InitHairTransmittanceData();
+			{
+				TransmittanceData.LocalScattering = 0;
+				TransmittanceData.OpaqueVisibility = 1;
+				TransmittanceData.ScatteringComponent = HairComponents | HAIR_COMPONENT_MULTISCATTER;
+
+				// Compute the transmittance based on precompute Hair transmittance LUT
+				HairAverageEnergy Energy	= SampleHairEnergyLUT(HairEnergyLUTTexture, LinearSampler, GBuffer.BaseColor, GBuffer.Roughness, SinLightAngle);
+				const float3 TotalEnergy	= Energy.A_R + Energy.A_TT + Energy.A_TRT;
+	
+				HairAverageEnergy SampleEnergy = SampleHairEnergyLUT(HairEnergyLUTTexture, LinearSampler, SampleGBuffer.BaseColor, SampleGBuffer.Roughness, SampleSinLightAngle);
+				const float3 SampleTotalEnergy = SampleEnergy.A_R + SampleEnergy.A_TT + SampleEnergy.A_TRT;
+
+				TransmittanceData.LocalScattering = TotalEnergy * SampleTotalEnergy;
+			}
+			#endif
+
+			// Nothing (passthrough)
+			#if 0
+			FHairTransmittanceData TransmittanceData = InitHairTransmittanceData();
+			{
+				TransmittanceData.LocalScattering = 0;
+				TransmittanceData.GlobalScattering = 0;
+				TransmittanceData.OpaqueVisibility = 1;
+				TransmittanceData.ScatteringComponent = HairComponents;
+			}
+			#endif
+
+			const float Area = 0;
+			const float3 TransmittedLight = TransmittanceMask.Transmittance * TransmittanceMask.Visibility * SampleColor.xyz;
+			float3 Luminance = HairShading(GBuffer, L, V, T, 1, TransmittanceData, 1, Area, uint2(0,0)) * TransmittedLight;
+			
+			AccIndirectLuminance.rgb += Luminance.xyz;
+			AccIndirectLuminance.w += 1;
+		}
+	}
+
+	AccIndirectLuminance /= max(AccIndirectLuminance.w,1);
+
+
+	if (CatData.PixelCoverage < 1)
+	{
+		OutSubPixelColor.xyz = AccIndirectLuminance.xyz * CatData.PixelCoverage;
+		OutSubPixelColor.a = 0;
+	}
+	else if (CatData.PixelCoverage >= 1)
+	{
+		OutColor.xyz = AccIndirectLuminance.xyz;
+		OutColor.a = 0;
+	}
+
+	OutDiffusionTexture = AccIndirectLuminance;
+}
diff --git a/Engine/Shaders/Private/HairStrands/HairScatterCompose.usf b/Engine/Shaders/Private/HairStrands/HairScatterCompose.usf
new file mode 100644
--- /dev/null
+++ b/Engine/Shaders/Private/HairStrands/HairScatterCompose.usf
@@ -0,0 +1,38 @@
+// Copyright Epic Games, Inc. All Rights Reserved.
+
+#include "../Common.ush"
+
+#include "../SceneTexturesCommon.ush"
+#include "HairStrandsDeepShadowCommon.ush"
+#include "../DeferredShadingCommon.ush"
+#include "HairStrandsVisibilityCommon.ush"
+
+Texture2D<float4> SceneColorTexture;
+Texture2D<float4> SceneColorSubPixelTexture;
+Texture2D<uint4>  CategorizationTexture;
+
+void MainPS(
+	in FScreenVertexOutput Input,
+	out float4 OutColor  : SV_Target0)
+{
+	OutColor = 0;
+
+	const uint2 PixelCoord = floor(Input.Position.xy);
+	const FCategorizationData CatData = DecodeCategorizationData(CategorizationTexture.Load(uint3(PixelCoord,0)));
+
+	const bool bIsCoveredByHair = CatData.PixelCoverage > 0;
+	if (bIsCoveredByHair)
+	{
+		if (CatData.PixelCoverage < 1)
+		{
+			const float4 Color = SceneColorSubPixelTexture.Load(uint3(PixelCoord, 0));
+			OutColor.xyz = Color.xyz / Color.w;
+			OutColor.w = Color.w;
+		}
+
+		if (CatData.PixelCoverage >= 1)
+		{
+			OutColor = SceneColorTexture.Load(uint3(PixelCoord, 0));
+		}
+	}
+}
diff --git a/Engine/Shaders/Private/HairStrands/HairStrandsAABB.usf b/Engine/Shaders/Private/HairStrands/HairStrandsAABB.usf
new file mode 100644
--- /dev/null
+++ b/Engine/Shaders/Private/HairStrands/HairStrandsAABB.usf
@@ -0,0 +1,53 @@
+// Copyright Epic Games, Inc. All Rights Reserved.
+
+#include "HairStrandsAABBCommon.ush"
+
+#if SHADER_AABBUPDATE
+
+uint bClearBuffer;
+uint MacroGroupId;
+uint MacroGroupValid;
+
+// Group data packed as {uint3 AABBMin, uint3 AABBMax}
+Buffer<int> InGroupAABBBuffer0;
+Buffer<int> InGroupAABBBuffer1;
+Buffer<int> InGroupAABBBuffer2;
+Buffer<int> InGroupAABBBuffer3;
+Buffer<int> InGroupAABBBuffer4;
+Buffer<int> InGroupAABBBuffer5;
+Buffer<int> InGroupAABBBuffer6;
+Buffer<int> InGroupAABBBuffer7;
+RWBuffer<int> OutMacroGroupAABBBuffer;
+
+void UpdateValue(inout FHairAABB B, in uint Index, in Buffer<int> InBuffer)
+{
+	FHairAABB A = ReadHairAABB(Index, InBuffer);
+	B = UnionHairAABB(A, B);
+}
+
+[numthreads(1, 1, 1)]
+void Main(uint2 DispatchThreadId : SV_DispatchThreadID)
+{
+	if (DispatchThreadId.x == 0)
+	{
+		FHairAABB Bound = InitHairAABB();
+		if (bClearBuffer == 0)
+		{
+			Bound = ReadHairAABB(MacroGroupId, OutMacroGroupAABBBuffer);
+		}
+
+		if (MacroGroupValid & 1)   { UpdateValue(Bound, 0, InGroupAABBBuffer0); }
+		if (MacroGroupValid & 2)   { UpdateValue(Bound, 0, InGroupAABBBuffer1); }
+		if (MacroGroupValid & 4)   { UpdateValue(Bound, 0, InGroupAABBBuffer2); }
+		if (MacroGroupValid & 8)   { UpdateValue(Bound, 0, InGroupAABBBuffer3); }
+		if (MacroGroupValid & 16)  { UpdateValue(Bound, 0, InGroupAABBBuffer4); }
+		if (MacroGroupValid & 32)  { UpdateValue(Bound, 0, InGroupAABBBuffer5); }
+		if (MacroGroupValid & 64)  { UpdateValue(Bound, 0, InGroupAABBBuffer6); }
+		if (MacroGroupValid & 128) { UpdateValue(Bound, 0, InGroupAABBBuffer7); }
+
+		WriteHairAABB(MacroGroupId, Bound, OutMacroGroupAABBBuffer);
+	}
+}
+
+#endif
+
diff --git a/Engine/Shaders/Private/HairStrands/HairStrandsAABBCommon.ush b/Engine/Shaders/Private/HairStrands/HairStrandsAABBCommon.ush
new file mode 100644
--- /dev/null
+++ b/Engine/Shaders/Private/HairStrands/HairStrandsAABBCommon.ush
@@ -0,0 +1,81 @@
+// Copyright Epic Games, Inc. All Rights Reserved.
+
+#define HAIR_AABB_PRECISION_SCALE 1000.f
+
+struct FHairAABB
+{
+	float3 Min;
+	float3 Max;
+};
+
+FHairAABB InitHairAABB()
+{
+	FHairAABB O;
+	O.Min = float3(2147483647, 2147483647, 2147483647);
+	O.Max = float3(-2147483647, -2147483647, -2147483647);
+	return O;
+}
+
+FHairAABB ReadHairAABB(uint Index, RWBuffer<int> InBuffer)
+{
+	const uint Index6 = Index * 6;
+	FHairAABB O;
+	O.Min.x = InBuffer[Index6 + 0] / HAIR_AABB_PRECISION_SCALE;
+	O.Min.y = InBuffer[Index6 + 1] / HAIR_AABB_PRECISION_SCALE;
+	O.Min.z = InBuffer[Index6 + 2] / HAIR_AABB_PRECISION_SCALE;
+
+	O.Max.x = InBuffer[Index6 + 3] / HAIR_AABB_PRECISION_SCALE;
+	O.Max.y = InBuffer[Index6 + 4] / HAIR_AABB_PRECISION_SCALE;
+	O.Max.z = InBuffer[Index6 + 5] / HAIR_AABB_PRECISION_SCALE;
+	return O;
+}
+
+FHairAABB ReadHairAABB(uint Index, Buffer<int> InBuffer)
+{
+	const uint Index6 = Index * 6;
+	FHairAABB O;
+	O.Min.x = InBuffer[Index6 + 0] / HAIR_AABB_PRECISION_SCALE;
+	O.Min.y = InBuffer[Index6 + 1] / HAIR_AABB_PRECISION_SCALE;
+	O.Min.z = InBuffer[Index6 + 2] / HAIR_AABB_PRECISION_SCALE;
+
+	O.Max.x = InBuffer[Index6 + 3] / HAIR_AABB_PRECISION_SCALE;
+	O.Max.y = InBuffer[Index6 + 4] / HAIR_AABB_PRECISION_SCALE;
+	O.Max.z = InBuffer[Index6 + 5] / HAIR_AABB_PRECISION_SCALE;
+	return O;
+}
+
+void WriteHairAABB(uint Index, FHairAABB In, RWBuffer<int> OutBuffer)
+{
+	const uint Index6 = Index * 6;
+	OutBuffer[Index6 + 0] = In.Min.x * HAIR_AABB_PRECISION_SCALE;
+	OutBuffer[Index6 + 1] = In.Min.y * HAIR_AABB_PRECISION_SCALE;
+	OutBuffer[Index6 + 2] = In.Min.z * HAIR_AABB_PRECISION_SCALE;
+
+	OutBuffer[Index6 + 3] = In.Max.x * HAIR_AABB_PRECISION_SCALE;
+	OutBuffer[Index6 + 4] = In.Max.y * HAIR_AABB_PRECISION_SCALE;
+	OutBuffer[Index6 + 5] = In.Max.z * HAIR_AABB_PRECISION_SCALE;
+}
+
+FHairAABB UnionHairAABB(FHairAABB A, FHairAABB B)
+{
+	FHairAABB O;
+	O.Min.x = min(A.Min.x, B.Min.x);
+	O.Min.y = min(A.Min.y, B.Min.y);
+	O.Min.z = min(A.Min.z, B.Min.z);
+
+	O.Max.x = max(A.Max.x, B.Max.x);
+	O.Max.y = max(A.Max.y, B.Max.y);
+	O.Max.z = max(A.Max.z, B.Max.z);
+	return O;
+}
+
+void InterlockHairAABB(uint Index, FHairAABB In, RWBuffer<int> OutBuffer)
+{
+	const uint Index6 = Index * 6;
+	InterlockedMin(OutBuffer[Index6 + 0], int(In.Min.x * HAIR_AABB_PRECISION_SCALE));
+	InterlockedMin(OutBuffer[Index6 + 1], int(In.Min.y * HAIR_AABB_PRECISION_SCALE));
+	InterlockedMin(OutBuffer[Index6 + 2], int(In.Min.z * HAIR_AABB_PRECISION_SCALE));
+	InterlockedMax(OutBuffer[Index6 + 3], int(In.Max.x * HAIR_AABB_PRECISION_SCALE));
+	InterlockedMax(OutBuffer[Index6 + 4], int(In.Max.y * HAIR_AABB_PRECISION_SCALE));
+	InterlockedMax(OutBuffer[Index6 + 5], int(In.Max.z * HAIR_AABB_PRECISION_SCALE));
+}
\ No newline at end of file
