Merging main @ cl 11028482

diff --git a/Engine/Plugins/Experimental/HairStrands/Source/HairStrandsCore/Private/GroomComponent.cpp b/Engine/Plugins/Experimental/HairStrands/Source/HairStrandsCore/Private/GroomComponent.cpp
--- a/Engine/Plugins/Experimental/HairStrands/Source/HairStrandsCore/Private/GroomComponent.cpp
+++ b/Engine/Plugins/Experimental/HairStrands/Source/HairStrandsCore/Private/GroomComponent.cpp
@@ -12,6 +12,7 @@
 #include "UObject/UObjectIterator.h"
 #include "GlobalShader.h"
 #include "Components/SkeletalMeshComponent.h"
+#include "Rendering/SkeletalMeshRenderData.h"
 #include "Engine/RendererSettings.h"
 #include "Animation/AnimationSettings.h"
 #include "Logging/MessageLog.h"
@@ -23,34 +24,11 @@ static float GHairClipLength = -1;
 static FAutoConsoleVariableRef CVarHairClipLength(TEXT("r.HairStrands.DebugClipLength"), GHairClipLength, TEXT("Clip hair strands which have a lenth larger than this value. (default is -1, no effect)"));
 float GetHairClipLength() { return GHairClipLength > 0 ? GHairClipLength : 100000;  }
 
-static int32 GHairStrandsMeshProjectionForceRefPoseEnable = 0;
-static int32 GHairStrandsMeshProjectionForceLOD = -1;
-static FAutoConsoleVariableRef CVarHairStrandsMeshProjectionForceRefPoseEnable(TEXT("r.HairStrands.MeshProjection.RefPose"), GHairStrandsMeshProjectionForceRefPoseEnable, TEXT("Enable/Disable reference pose"));
-static FAutoConsoleVariableRef CVarHairStrandsMeshProjectionForceLOD(TEXT("r.HairStrands.MeshProjection.LOD"), GHairStrandsMeshProjectionForceLOD, TEXT("Force a specific LOD"));
-
 static int32 GHairStrandsMeshUseRelativePosition = 1;
 static FAutoConsoleVariableRef CVarHairStrandsMeshUseRelativePosition(TEXT("r.HairStrands.MeshProjection.RelativePosition"), GHairStrandsMeshUseRelativePosition, TEXT("Enable/Disable relative triangle position for improving positions"));
 
-static int32 GHairStrandsMeshProjectionTickDelay = 1;
-static FAutoConsoleVariableRef CVarHairStrandsMeshProjectionTickDelay(TEXT("r.HairStrands.MeshProjection.TickDelay"), GHairStrandsMeshProjectionTickDelay, TEXT("Number of simulation tick to wait before projecting a groom onto a mesh"));
-
 #define LOCTEXT_NAMESPACE "GroomComponent"
 
-static float GetGroupMaxHairRadius(const FHairGroupDesc& GroupDesc, const FHairGroupData& GroupData)
-{
-	return GroupDesc.HairWidth > 0 ? GroupDesc.HairWidth * 0.5f : GroupData.HairRenderData.StrandsCurves.MaxRadius;
-}
-
-static float GetGroupHairShadowDensity(const FHairGroupDesc& GroupDesc, const FHairGroupData& GroupData)
-{
-	return GroupDesc.HairShadowDensity > 0 ? GroupDesc.HairShadowDensity : GroupData.HairRenderData.HairDensity;
-}
-
-static float GetGroupHairRaytracingRadiusScale(const FHairGroupDesc& GroupDesc, const FHairGroupData& GroupData)
-{
-	return GroupDesc.HairRaytracingRadiusScale;
-}
-
 /**
  * An material render proxy which overrides the debug mode parameter.
  */
@@ -155,6 +133,7 @@ public:
 		check(Component->GroomAsset);
 		check(Component->GroomAsset->GetNumHairGroups() > 0);
 		check(Component->InterpolationOutput);
+		ComponentId = Component->ComponentId.PrimIDValue;
 
 		FHairStrandsVertexFactory::FDataType VFData;
 		VFData.InterpolationOutput = Component->InterpolationOutput;
@@ -174,16 +153,9 @@ public:
 				Material = GEngine->HairDefaultMaterial;
 			}
 
-			FHairStrandsVertexFactory::FDataType::HairGroup& VFGroupData = VFData.HairGroups.AddDefaulted_GetRef();
-			VFGroupData.MinStrandRadius		= 0;
-			VFGroupData.MaxStrandRadius		= GetGroupMaxHairRadius(InGroupDesc, InGroupData);
-			VFGroupData.MaxStrandLength		= InGroupData.HairRenderData.StrandsCurves.MaxLength;
-			VFGroupData.HairDensity			= GetGroupHairShadowDensity(InGroupDesc, InGroupData);
-			VFGroupData.HairWorldOffset		= InGroupData.HairRenderData.BoundingBox.GetCenter();
-
 			#if RHI_RAYTRACING
 			FRayTracingGeometry* RayTracingGeometry = nullptr;
-			if (IsRayTracingEnabled() && GroupResources.RaytracingResources)
+			if (IsHairRayTracingEnabled() && GroupResources.RaytracingResources)
 			{
 				RayTracingGeometry = &GroupResources.RaytracingResources->RayTracingGeometry;
 			}
@@ -191,11 +163,11 @@ public:
 
 			HairGroup& OutGroupData = HairGroups.Add_GetRef(
 			{
-				GroupIt,
 				Material
 				#if RHI_RAYTRACING
 				, RayTracingGeometry
 				#endif
+				, GroupResources.HairGroupPublicDatas
 			});
 
 		}
@@ -220,7 +192,7 @@ public:
 
 	virtual void GetDynamicRayTracingInstances(FRayTracingMaterialGatheringContext & Context, TArray<struct FRayTracingInstance>& OutRayTracingInstances) override
 	{
-		if (!IsRayTracingEnabled() || HairGroups.Num() == 0)
+		if (!IsHairRayTracingEnabled() || HairGroups.Num() == 0)
 			return;
 
 		for (const HairGroup& GroupData : HairGroups)
@@ -239,9 +211,7 @@ public:
 
 	virtual void GetDynamicMeshElements(const TArray<const FSceneView*>& Views, const FSceneViewFamily& ViewFamily, uint32 VisibilityMap, FMeshElementCollector& Collector) const override
 	{
-		const uint32 GroupCount = HairGroups.Num();
-		check(GroupCount == VertexFactory.GetData().HairGroups.Num());
-		if (GroupCount == 0)
+		if (!VertexFactory.GetData().InterpolationOutput)
 			return;
 
 		bool bHasOneElementValid = false;
@@ -256,6 +226,8 @@ public:
 		if (!bHasOneElementValid) 
 			return;
 
+		const uint32 GroupCount = VertexFactory.GetData().InterpolationOutput->HairGroups.Num();
+
 		QUICK_SCOPE_CYCLE_COUNTER(STAT_HairStrandsSceneProxy_GetDynamicMeshElements);
 
 
@@ -277,19 +249,20 @@ public:
 			case EHairStrandsDebugMode::RenderHairRootUDIM			: DebugModeScalar = 6.f; break;
 			case EHairStrandsDebugMode::RenderHairBaseColor			: DebugModeScalar = 7.f; break;
 			case EHairStrandsDebugMode::RenderHairRoughness			: DebugModeScalar = 8.f; break;
+			case EHairStrandsDebugMode::RenderVisCluster			: DebugModeScalar = 0.f; break;
 			};
 
-			float HairMinRadius = FLT_MAX;
+
 			float HairMaxRadius = 0;
-			for (uint32 GroupIt = 0; GroupIt < GroupCount; ++GroupIt)
+			for (FHairStrandsInterpolationOutput::HairGroup& HairGroup : VertexFactory.GetData().InterpolationOutput->HairGroups)
 			{
-				HairMinRadius = FMath::Min(HairMinRadius, VertexFactory.GetMinStrandRadius(GroupIt));
-				HairMaxRadius = FMath::Max(HairMaxRadius, VertexFactory.GetMaxStrandRadius(GroupIt));
+				HairMaxRadius = FMath::Max(HairMaxRadius, HairGroup.VFInput.HairRadius);
 			}
+
 			const float HairClipLength = GetHairClipLength();
 			auto DebugMaterial = new FHairDebugModeMaterialRenderProxy(
 				GEngine->HairDebugMaterial ? GEngine->HairDebugMaterial->GetRenderProxy() : nullptr,
-				DebugModeScalar, HairMinRadius, HairMaxRadius, HairClipLength);
+				DebugModeScalar, 0, HairMaxRadius, HairClipLength);
 			Collector.RegisterOneFrameMaterialProxy(DebugMaterial);
 			MaterialProxy = DebugMaterial;
 		}
@@ -302,9 +275,7 @@ public:
 				{
 					const HairGroup& GroupData = HairGroups[GroupIt];
 					const uint32 HairVertexCount = VertexFactory.GetData().InterpolationOutput->HairGroups[GroupIt].VFInput.VertexCount;
-					const FHairStrandsVertexFactory::FDataType::HairGroup& HairGroup = VertexFactory.GetData().HairGroups[GroupIt];
 
-					const FSceneView* View = Views[ViewIndex];
 					// Draw the mesh.
 					FMeshBatch& Mesh = Collector.AllocateMesh();
 					FMeshBatchElement& BatchElement = Mesh.Elements[0];
@@ -322,11 +293,19 @@ public:
 					FDynamicPrimitiveUniformBuffer& DynamicPrimitiveUniformBuffer = Collector.AllocateOneFrameResource<FDynamicPrimitiveUniformBuffer>();
 					DynamicPrimitiveUniformBuffer.Set(GetLocalToWorld(), PreviousLocalToWorld, GetBounds(), GetLocalBounds(), true, false, bDrawVelocity, bOutputVelocity);
 					BatchElement.PrimitiveUniformBufferResource = &DynamicPrimitiveUniformBuffer.UniformBuffer;
+
 					BatchElement.FirstIndex = 0;
-					BatchElement.NumPrimitives = HairVertexCount * 2;
+					BatchElement.NumInstances = 1;
+					BatchElement.NumPrimitives = 0;
+					BatchElement.IndirectArgsBuffer = GroupData.PublicData->GetDrawIndirectBuffer().Buffer.GetReference();
+					BatchElement.IndirectArgsOffset = 0;
+					
+					// Setup our vertex factor custom data
+					BatchElement.VertexFactoryUserData = const_cast<void*>(reinterpret_cast<const void*>(GroupData.PublicData));
+
 					BatchElement.MinVertexIndex = 0;
 					BatchElement.MaxVertexIndex = HairVertexCount * 6;
-					BatchElement.UserData = reinterpret_cast<void*>(uint64(GroupIt));
+					BatchElement.UserData = reinterpret_cast<void*>(uint64(ComponentId));
 					Mesh.ReverseCulling = IsLocalToWorldDeterminantNegative();
 					Mesh.Type = PT_TriangleList;
 					Mesh.DepthPriorityGroup = SDPG_World;
@@ -383,15 +362,16 @@ public:
 	uint32 GetAllocatedSize(void) const { return(FPrimitiveSceneProxy::GetAllocatedSize()); }
 private:
 
+	uint32 ComponentId = 0;
 	FHairStrandsVertexFactory VertexFactory;
 	FMaterialRelevance MaterialRelevance;
 	struct HairGroup
 	{
-		uint32 GroupIndex = 0;
 		UMaterialInterface* Material = nullptr;
 	#if RHI_RAYTRACING
 		FRayTracingGeometry* RayTracingGeometry = nullptr;
 	#endif
+		FHairGroupPublicData* PublicData = nullptr;
 	};
 	TArray<HairGroup> HairGroups;
 };
@@ -404,7 +384,7 @@ static void UpdateHairGroupsDesc(UGroomAsset* GroomAsset, TArray<FHairGroupDesc>
 		return;
 	}
 
-	check(GroomAsset->HairGroupsInfo.Num() == GroomAsset->HairGroupsInfo.Num());
+	check(GroomAsset->HairGroupsInfo.Num() == GroomAsset->HairGroupsData.Num());
 
 	const uint32 GroupCount = GroomAsset->HairGroupsInfo.Num();
 	const bool bReinitOverride = GroupCount != GroomGroupsDesc.Num();
@@ -421,9 +401,11 @@ static void UpdateHairGroupsDesc(UGroomAsset* GroomAsset, TArray<FHairGroupDesc>
 		FHairGroupDesc& Desc = GroomGroupsDesc[GroupIt];
 		Desc.GuideCount	= GroupInfo.NumGuides;
 		Desc.HairCount	= GroupInfo.NumCurves;
+		Desc.HairLength = GroupData.HairRenderData.StrandsCurves.MaxLength;
+
 		if (bReinitOverride || Desc.HairWidth == 0)
 		{
-			Desc.HairWidth = GroupData.HairRenderData.StrandsCurves.MaxRadius * 0.5f;
+			Desc.HairWidth = GroupData.HairRenderData.StrandsCurves.MaxRadius * 2.0f;
 		}
 		if (bReinitOverride || Desc.HairShadowDensity == 0)
 		{
@@ -432,7 +414,7 @@ static void UpdateHairGroupsDesc(UGroomAsset* GroomAsset, TArray<FHairGroupDesc>
 
 		if (bReinitOverride)
 		{
-			Desc.HairRaytracingRadiusScale = 1;
+			Desc.ReInit();
 		}
 	}
 }
@@ -455,9 +437,6 @@ UGroomComponent::UGroomComponent(const FObjectInitializer& ObjectInitializer)
 	bBindGroomToSkeletalMesh = false;
 	InitializedResources = nullptr;
 	Mobility = EComponentMobility::Movable;
-	MeshProjectionTickDelay = 0;
-	MeshProjectionLODIndex = -1;
-	MeshProjectionState = EMeshProjectionState::Invalid;
 	bIsGroomAssetCallbackRegistered = false;
 
 	SetCollisionProfileName(UCollisionProfile::PhysicsActor_ProfileName);
@@ -567,16 +546,17 @@ static FHairStrandsProjectionHairData::HairGroup ToProjectionHairData(FHairStran
 		LODData.RootTriangleIndexBuffer = &MeshLODData.RootTriangleIndexBuffer;
 		LODData.RootTriangleBarycentricBuffer = &MeshLODData.RootTriangleBarycentricBuffer;
 		
-		LODData.RestPositionOffset					= MeshLODData.RestRootOffset;
+		LODData.RestPositionOffset				= &MeshLODData.RestRootOffset;
 		LODData.RestRootTrianglePosition0Buffer = &MeshLODData.RestRootTrianglePosition0Buffer;
 		LODData.RestRootTrianglePosition1Buffer = &MeshLODData.RestRootTrianglePosition1Buffer;
 		LODData.RestRootTrianglePosition2Buffer = &MeshLODData.RestRootTrianglePosition2Buffer;
 
-		LODData.DeformedPositionOffset					= MeshLODData.DeformedRootOffset;
+		LODData.DeformedPositionOffset				= &MeshLODData.DeformedRootOffset;
 		LODData.DeformedRootTrianglePosition0Buffer = &MeshLODData.DeformedRootTrianglePosition0Buffer;
 		LODData.DeformedRootTrianglePosition1Buffer = &MeshLODData.DeformedRootTrianglePosition1Buffer;
 		LODData.DeformedRootTrianglePosition2Buffer = &MeshLODData.DeformedRootTrianglePosition2Buffer;
-		LODData.bIsValid = MeshLODData.Status == FHairStrandsRootResource::FMeshProjectionLOD::EStatus::Completed;
+
+		LODData.Status = &MeshLODData.Status;
 	}
 
 	return Out;
@@ -738,6 +718,8 @@ void UGroomComponent::InitResources()
 	bResetSimulation = true;
 	//UE_LOG(LogHairStrands, Warning, TEXT("Groom Init = %d"), bResetSimulation);
 
+	UpdateHairGroupsDesc(GroomAsset, GroomGroupsDesc);
+
 	if (!GroomAsset || GroomAsset->GetNumHairGroups() == 0)
 		return;
 
@@ -749,20 +731,22 @@ void UGroomComponent::InitResources()
 
 	// Insure the ticking of the Groom component always happens after the skeletalMeshComponent.
 	USkeletalMeshComponent* SkeletalMeshComponent = bBindGroomToSkeletalMesh && GetAttachParent() ? Cast<USkeletalMeshComponent>(GetAttachParent()) : nullptr;
+	const bool bHasValidSketalMesh = SkeletalMeshComponent && SkeletalMeshComponent->GetSkeletalMeshRenderData();
+
+	// Always setup the callback, even if the SkeletalMeshData are not ready yet
 	if (SkeletalMeshComponent)
 	{
-		RegisteredSkeletalMeshComponent = SkeletalMeshComponent;
-		AddTickPrerequisiteComponent(SkeletalMeshComponent);
-
 		FSkeletalMeshObjectCallbackData CallbackData;
 		CallbackData.Run = CallbackMeshObjectCallback;
 		CallbackData.UserData = (uint64(LocalComponentId.PrimIDValue) & 0xFFFFFFFF) | (uint64(WorldType) << 32);
 		SkeletalMeshComponent->MeshObjectCallbackData = CallbackData;
+	}
 
-
-		{
-			SkeletalMeshComponent->OnBoneTransformsFinalized.AddDynamic(this, &UGroomComponent::ResetSimulation);
-		}
+	if (bHasValidSketalMesh)
+	{
+		RegisteredSkeletalMeshComponent = SkeletalMeshComponent;
+		AddTickPrerequisiteComponent(SkeletalMeshComponent);
+		SkeletalMeshComponent->OnBoneTransformsFinalized.AddDynamic(this, &UGroomComponent::ResetSimulation);
 	}
 
 	const bool bIsSimulationEnable = IsSimulationEnabled(this);
@@ -793,14 +777,14 @@ void UGroomComponent::InitResources()
 		BeginInitResource(Res.InterpolationResource);
 
 		#if RHI_RAYTRACING
-		if (IsRayTracingEnabled())
+		if (IsHairRayTracingEnabled())
 		{
 			Res.RaytracingResources = new FHairStrandsRaytracingResource(GroupData.HairRenderData);
 			BeginInitResource(Res.RaytracingResources);
 		}
 		#endif
 
-		if (SkeletalMeshComponent)
+		if (bHasValidSketalMesh)
 		{
 			const uint32 LODCount = SkeletalMeshComponent->GetNumLODs();
 			if (LODCount > 0)
@@ -817,9 +801,16 @@ void UGroomComponent::InitResources()
 
 		Res.RenderDeformedResources = new FHairStrandsDeformedResource(GroupData.HairRenderData.RenderData, false);
 		Res.SimDeformedResources = new FHairStrandsDeformedResource(GroupData.HairSimulationData.RenderData, true);
+		Res.ClusterCullingResources = new FHairStrandsClusterCullingResource(GroupData);
+
+		const uint32 HairVertexCount = GroupData.HairRenderData.GetNumPoints();
+		const uint32 GroupInstanceVertexCount = HairVertexCount * 6; // 6 vertex per point for a quad
+		Res.HairGroupPublicDatas = new FHairGroupPublicData(GroupIt, GroupInstanceVertexCount, Res.ClusterCullingResources->ClusterCount, HairVertexCount);
 
 		BeginInitResource(Res.RenderDeformedResources);
 		BeginInitResource(Res.SimDeformedResources);
+		BeginInitResource(Res.ClusterCullingResources);
+		BeginInitResource(Res.HairGroupPublicDatas);
 
 		const FVector RenderRestHairPositionOffset = Res.RenderRestResources->PositionOffset;
 		const FVector SimRestHairPositionOffset = Res.SimRestResources->PositionOffset;
@@ -833,9 +824,11 @@ void UGroomComponent::InitResources()
 		FHairStrandsInterpolationInput::FHairGroup& InterpolationInputGroup = InterpolationInput->HairGroups.AddDefaulted_GetRef();
 
 		check(GroupIt < GroomGroupsDesc.Num());
-		const FHairGroupDesc& InGroupDesc = GroomGroupsDesc[GroupIt];
-		InterpolationInputGroup.HairRadius = GetGroupMaxHairRadius(InGroupDesc, GroupData);
-		InterpolationInputGroup.HairRaytracingRadiusScale = GetGroupHairRaytracingRadiusScale(InGroupDesc, GroupData);
+		InterpolationInputGroup.GroupDesc = GroomGroupsDesc[GroupIt];
+
+		InterpolationOutputGroup.HairGroupPublicData = Res.HairGroupPublicDatas;
+
+		// Rest post offset
 		InterpolationInputGroup.InRenderHairPositionOffset = RenderRestHairPositionOffset;
 		InterpolationInputGroup.InSimHairPositionOffset = SimRestHairPositionOffset;
 
@@ -852,6 +845,7 @@ void UGroomComponent::InitResources()
 	Interpolation.Input  = InterpolationInput;
 	Interpolation.Output = InterpolationOutput;
 	Interpolation.Function = ComputeHairStrandsInterpolation;
+	Interpolation.ResetFunction = ResetHairStrandsInterpolation;
 
 	FHairGroupResources* LocalResources = &HairGroupResources;
 	const uint64 Id = LocalComponentId.PrimIDValue;
@@ -862,12 +856,15 @@ void UGroomComponent::InitResources()
 			LocalResources,
 			HairLocalToWorld, SkinLocalToWorld,
 			WorldType,
-			DebugGroupInfo
+			DebugGroupInfo,
+			bHasValidSketalMesh,
+			SkeletalMeshComponent
 		]
 		(FRHICommandListImmediate& RHICmdList)
 	{
 		FHairStrandsProjectionHairData RenProjectionDatas;
 		FHairStrandsProjectionHairData SimProjectionDatas;
+		FHairStrandsPrimitiveResources PrimitiveResources;
 		const uint32 GroupCount = LocalResources->Num();
 		for (uint32 GroupIt=0;GroupIt<GroupCount; ++GroupIt)
 		{
@@ -887,9 +884,15 @@ void UGroomComponent::InitResources()
 
 			InputGroup.Interpolation0Buffer			= &Res.InterpolationResource->Interpolation0Buffer;
 			InputGroup.Interpolation1Buffer			= &Res.InterpolationResource->Interpolation1Buffer;
+			
+			InputGroup.ClusterCount					= Res.ClusterCullingResources->ClusterCount;
+			InputGroup.ClusterVertexCount			= Res.ClusterCullingResources->VertexCount;
+			InputGroup.VertexToClusterIdBuffer		= &Res.ClusterCullingResources->VertexToClusterIdBuffer;
+			InputGroup.ClusterVertexIdBuffer		= &Res.ClusterCullingResources->ClusterVertexIdBuffer;
+			InputGroup.ClusterIndexRadiusScaleInfoBuffer= &Res.ClusterCullingResources->ClusterIndexRadiusScaleInfoBuffer;
 
 			#if RHI_RAYTRACING
-			if (IsRayTracingEnabled())
+			if (IsHairRayTracingEnabled())
 			{
 				InputGroup.RaytracingGeometry		= &Res.RaytracingResources->RayTracingGeometry;
 				InputGroup.RaytracingPositionBuffer	= &Res.RaytracingResources->PositionBuffer;
@@ -905,17 +908,68 @@ void UGroomComponent::InitResources()
 			OutputGroup.RenderMaterialBuffer			= &Res.RenderRestResources->MaterialBuffer;
 			OutputGroup.RenderTangentBuffer				= &Res.RenderDeformedResources->TangentBuffer;
 			OutputGroup.SimTangentBuffer				= &Res.SimDeformedResources->TangentBuffer;
+			
+			OutputGroup.RenderClusterAABBBuffer			= &Res.HairGroupPublicDatas->GetClusterAABBBuffer();
+			OutputGroup.RenderGroupAABBBuffer			= &Res.HairGroupPublicDatas->GetGroupAABBBuffer();
+			OutputGroup.ClusterInfoBuffer				= &Res.ClusterCullingResources->ClusterInfoBuffer;
 
 			RenProjectionDatas.HairGroups.Add(ToProjectionHairData(Res.RenRootResources));
 			SimProjectionDatas.HairGroups.Add(ToProjectionHairData(Res.SimRootResources));
+
+			FHairStrandsPrimitiveResources::FHairGroup& Group = PrimitiveResources.Groups.AddDefaulted_GetRef();
+			Group.ClusterAABBBuffer = &Res.HairGroupPublicDatas->GetClusterAABBBuffer();
+			Group.GroupAABBBuffer	= &Res.HairGroupPublicDatas->GetGroupAABBBuffer();
+			Group.ClusterCount		= Res.ClusterCullingResources->ClusterCount;
+		}
+
+		if (bHasValidSketalMesh)
+		{		
+			// Convert FSkeletalMeshRenderData into FHairStrandsProjectionMeshData for the initial projection
+			FHairStrandsProjectionMeshData MeshData;
+			FSkeletalMeshRenderData* RenderData = SkeletalMeshComponent->GetSkeletalMeshRenderData();
+
+			uint32 LODIndex = 0;
+			for (FSkeletalMeshLODRenderData& LODRenderData : RenderData->LODRenderData)
+			{
+				FHairStrandsProjectionMeshData::LOD& LOD = MeshData.LODs.AddDefaulted_GetRef();
+				uint32 SectionIndex = 0;
+				for (FSkelMeshRenderSection& InSection : LODRenderData.RenderSections)
+				{
+					FHairStrandsProjectionMeshData::Section& OutSection = LOD.Sections.AddDefaulted_GetRef();
+					
+					OutSection.PositionBuffer	= LODRenderData.StaticVertexBuffers.PositionVertexBuffer.GetSRV();
+					OutSection.IndexBuffer		= LODRenderData.MultiSizeIndexContainer.GetIndexBuffer()->GetSRV();
+					OutSection.TotalVertexCount	= LODRenderData.StaticVertexBuffers.PositionVertexBuffer.GetNumVertices();
+					OutSection.TotalIndexCount	= LODRenderData.MultiSizeIndexContainer.GetIndexBuffer()->Num();
+					OutSection.NumPrimitives	= InSection.NumTriangles;
+					OutSection.VertexBaseIndex	= InSection.BaseVertexIndex;
+					OutSection.IndexBaseIndex	= InSection.BaseIndex;
+					OutSection.SectionIndex		= SectionIndex;
+					OutSection.LODIndex			= LODIndex;
+
+					++SectionIndex;
+				}
+				++LODIndex;
+			}
+
+			// The offset is based on the center of the skeletal mesh (which is computed based on the physics capsules/boxes/...)
+			const FVector MeshPositionOffset = SkeletalMeshComponent->CalcBounds(FTransform::Identity).GetBox().GetCenter();
+			RunProjection(
+				RHICmdList,
+				HairLocalToWorld,
+				MeshPositionOffset,
+				MeshData, 
+				RenProjectionDatas,
+				SimProjectionDatas);
 		}
 
 		RegisterHairStrands(
 			Id,
 			WorldType,
 			Interpolation,
-			RenProjectionDatas, 
+			RenProjectionDatas,
 			SimProjectionDatas,
+			PrimitiveResources,
 			DebugGroupInfo);
 	});
 }
@@ -937,6 +991,8 @@ void UGroomComponent::ReleaseResources()
 		SafeRelease(Res.RenRootResources);
 		SafeRelease(Res.SimRootResources);
 		SafeRelease(Res.RenderDeformedResources);
+		SafeRelease(Res.ClusterCullingResources);
+		SafeRelease(Res.HairGroupPublicDatas);
 		SafeRelease(Res.SimDeformedResources);
 	#if RHI_RAYTRACING
 		SafeRelease(Res.RaytracingResources);
@@ -958,10 +1014,6 @@ void UGroomComponent::ReleaseResources()
 	InterpolationOutput = nullptr;
 	InitializedResources = nullptr;
 
-	MeshProjectionLODIndex = -1;
-	MeshProjectionTickDelay = 0;
-	MeshProjectionState = EMeshProjectionState::Invalid;
-
 	// Insure the ticking of the Groom component always happens after the skeletalMeshComponent.
 	if (RegisteredSkeletalMeshComponent)
 	{
@@ -990,7 +1042,6 @@ void UGroomComponent::PostLoad()
 		GroomAsset->ConditionalPostLoad();
 	}
 
-	UpdateHairGroupsDesc(GroomAsset, GroomGroupsDesc);
 	InitResources();
 
 #if WITH_EDITOR
@@ -1017,7 +1068,6 @@ void UGroomComponent::OnRegister()
 
 	if (!InitializedResources)
 	{
-		UpdateHairGroupsDesc(GroomAsset, GroomGroupsDesc);
 		InitResources();
 	}
 
@@ -1063,6 +1113,7 @@ void UGroomComponent::OnAttachmentChanged()
 	Super::OnAttachmentChanged();
 	if (GroomAsset && !IsBeingDestroyed() && HasBeenCreated())
 	{
+		FlushRenderingCommands();
 		InitResources();
 	}
 }
@@ -1076,180 +1127,33 @@ void UGroomComponent::TickComponent(float DeltaTime, enum ELevelTick TickType, F
 	const ERHIFeatureLevel::Type FeatureLevel = GetWorld() ? ERHIFeatureLevel::Type(GetWorld()->FeatureLevel) : ERHIFeatureLevel::Num;
 
 	FVector MeshPositionOffset = FVector::ZeroVector;
-	USkeletalMeshComponent* SkeletalMeshComponent = RegisteredSkeletalMeshComponent;	
-	if (SkeletalMeshComponent)
-	{
-		// The offset is based on the center of the skeletal mesh (which is computed based on the physics capsules/boxes/...)
-		MeshPositionOffset = SkeletalMeshComponent->CalcBounds(FTransform::Identity).GetBox().GetCenter();
+	USkeletalMeshComponent* SkeletalMeshComponent = RegisteredSkeletalMeshComponent;
 
-		const uint32 SkeletalLODCount = SkeletalMeshComponent->GetNumLODs();
-		const uint32 ResourceLODCount = HairGroupResources.Num() > 0 && HairGroupResources[0].RenRootResources ? HairGroupResources[0].RenRootResources->MeshProjectionLODs.Num() : 0;
-		if (SkeletalLODCount != ResourceLODCount)
+	if (USkeletalMeshComponent* ParentComp = Cast<USkeletalMeshComponent>(GetAttachParent()))
+	{
+		if (ParentComp->GetNumBones() > 0)
 		{
-			auto InitRootResource = [](FHairStrandsRootResource*& PersistentResources, FHairStrandsRootResource* LocalResources)
-			{
-				SafeRelease(PersistentResources);
-				PersistentResources = LocalResources;
-				if (PersistentResources)
-				{
-					BeginInitResource(PersistentResources);
-				}
-			};
+			const int32 BoneIndex = FMath::Min(1, ParentComp->GetNumBones()-1);
+			const FMatrix NextBoneMatrix = ParentComp->GetBoneMatrix(BoneIndex);
 
-			TArray<FHairStrandsRootResource*> GroupRenRootResources;
-			TArray<FHairStrandsRootResource*> GroupSimRootResources;
-			const uint32 GroupCount = GroomAsset ? GroomAsset->GetNumHairGroups() : 0;
-			check(GroupCount == HairGroupResources.Num());
-			for (uint32 GroupIndex=0; GroupIndex < GroupCount; ++GroupIndex)
+			const float BoneDistance = FVector::DistSquared(PrevBoneMatrix.GetOrigin(), NextBoneMatrix.GetOrigin());
+			if (ParentComp->GetTeleportDistanceThreshold() > 0.0 && BoneDistance >
+				ParentComp->GetTeleportDistanceThreshold() * ParentComp->GetTeleportDistanceThreshold())
 			{
-				const FHairGroupData* HairGroupsData = &GroomAsset->HairGroupsData[GroupIndex];
-				FHairGroupResource& Res = HairGroupResources[GroupIndex];
-				FHairStrandsRootResource* LocalRenRootResources = SkeletalLODCount > 0 ? new FHairStrandsRootResource(&HairGroupsData->HairRenderData, SkeletalLODCount) : nullptr;
-				FHairStrandsRootResource* LocalSimRootResources = SkeletalLODCount > 0 ? new FHairStrandsRootResource(&HairGroupsData->HairSimulationData, SkeletalLODCount) : nullptr;
-
-				InitRootResource(Res.RenRootResources, LocalRenRootResources);
-				InitRootResource(Res.SimRootResources, LocalSimRootResources);
-
-				GroupRenRootResources.Add(LocalRenRootResources);
-				GroupSimRootResources.Add(LocalSimRootResources);
+				bResetSimulation = true;
 			}
-
-			FTransform HairLocalToWorld = GetComponentTransform();
-			ENQUEUE_RENDER_COMMAND(FHairStrandsTick_LODUpdate)(
-				[Id, WorldType, GroupRenRootResources, GroupSimRootResources, HairLocalToWorld](FRHICommandListImmediate& RHICmdList)
-			{
-				FHairStrandsProjectionHairData RenProjectionDatas;
-				for (FHairStrandsRootResource* Res : GroupRenRootResources)
-				{
-					RenProjectionDatas.HairGroups.Add(ToProjectionHairData(Res));
-				}
-
-				FHairStrandsProjectionHairData SimProjectionDatas;
-				for (FHairStrandsRootResource* Res : GroupSimRootResources)
-				{
-					SimProjectionDatas.HairGroups.Add(ToProjectionHairData(Res));
-				}
-
-				UpdateHairStrands(Id, WorldType, HairLocalToWorld, RenProjectionDatas, SimProjectionDatas);
-			});
-		}
-
-		// State machine for binding/projecting hair onto skinned mesh. The state machine operates as follow:
-		// * For each mesh LOD
-		//   * Force the skeletal mesh in rest pose
-		//   * Wait for the skeletal mesh to tick/update the animation
-		//   * Issue a binding/projection query to the hair strands system
-		//   * Wait for the query to be completed
-		// * Put the skeletal mesh into its original pose
-		//
-		// In editor, the animation is not ticked, so we manually issue a re-initialization of the animation 
-		// to force the animation in rest pose or original pose
-		if (MeshProjectionState != EMeshProjectionState::Completed)
-		{
-			const uint32 TickDelay = FMath::Max(0, GHairStrandsMeshProjectionTickDelay);
-
-			if (MeshProjectionTickDelay == 0)
+			else
 			{
-				if (MeshProjectionState == EMeshProjectionState::Invalid)
-				{
-					MeshProjectionLODIndex = 0;
-					MeshProjectionState = EMeshProjectionState::InProgressBinding;
-				}
-
-				if (MeshProjectionLODIndex == SkeletalMeshComponent->GetNumLODs())
-				{
-					SkeletalMeshComponent->SetForceRefPose(false);
-					SkeletalMeshComponent->SetForcedLOD(0);
-					if (WorldType == EWorldType::Editor)
-					{
-						SkeletalMeshComponent->InitAnim(true);
-					}
-					MeshProjectionState = EMeshProjectionState::Completed;
-
-					FHairGroupResources* LocalResources = &HairGroupResources;
-					ENQUEUE_RENDER_COMMAND(FHairStrandsTick_UpdateProjectionStatus)(
-						[LocalResources](FRHICommandListImmediate& RHICmdList)
-					{
-						for (FHairGroupResource& Res : *LocalResources)
-						{
-							for (FHairStrandsRootResource::FMeshProjectionLOD& MeshLODData : Res.RenRootResources->MeshProjectionLODs)
-							{
-								MeshLODData.Status = FHairStrandsRootResource::FMeshProjectionLOD::EStatus::Completed;
-							}
-
-							for (FHairStrandsRootResource::FMeshProjectionLOD& MeshLODData : Res.SimRootResources->MeshProjectionLODs)
-							{
-								MeshLODData.Status = FHairStrandsRootResource::FMeshProjectionLOD::EStatus::Completed;
-							}
-						}
-					});
-				}
-
-				if (MeshProjectionLODIndex < SkeletalMeshComponent->GetNumLODs())
-				{
-					if (MeshProjectionState == EMeshProjectionState::InProgressBinding)
-					{
-						SkeletalMeshComponent->SetForceRefPose(true);
-						SkeletalMeshComponent->SetForcedLOD(MeshProjectionLODIndex + 1);
-						if (WorldType == EWorldType::Editor)
-						{
-							SkeletalMeshComponent->InitAnim(true);
-						}
-						MeshProjectionState = EMeshProjectionState::WaitForRestPose;
-						MeshProjectionTickDelay += TickDelay;
-					}
-
-					if (MeshProjectionState == EMeshProjectionState::WaitForRestPose && MeshProjectionTickDelay == 0)
-					{
-						FHairGroupResources* LocalResources = &HairGroupResources;
-
-						const FVector RestPositionOffset = MeshPositionOffset;
-						const uint32 LODIndex = MeshProjectionLODIndex;
-						ENQUEUE_RENDER_COMMAND(FHairStrandsTick_Projection)(
-							[Id, WorldType, FeatureLevel, LODIndex, RestPositionOffset, LocalResources](FRHICommandListImmediate& RHICmdList)
-						{
-							AddHairStrandsProjectionQuery(RHICmdList, Id, WorldType, LODIndex, RestPositionOffset);
-
-							// Update rest (sim/render) position offsets. This is used by Niagara.
-							for (FHairGroupResource& Res : *LocalResources)
-							{
-								if (Res.RenRootResources)
-								{
-									for (FHairStrandsRootResource::FMeshProjectionLOD& MeshProjectionLOD : Res.RenRootResources->MeshProjectionLODs)
-									{
-										MeshProjectionLOD.RestRootOffset = RestPositionOffset;
-									}
-								}
-
-								if (Res.SimRootResources)
-								{
-									for (FHairStrandsRootResource::FMeshProjectionLOD& MeshProjectionLOD : Res.SimRootResources->MeshProjectionLODs)
-									{
-										MeshProjectionLOD.RestRootOffset = RestPositionOffset;
-									}
-								}
-							}
-						});
-
-						MeshProjectionState = EMeshProjectionState::InProgressBinding;
-						MeshProjectionTickDelay += TickDelay;
-						++MeshProjectionLODIndex;
-					}
-				}
+				bResetSimulation = false;
 			}
+			PrevBoneMatrix = NextBoneMatrix;
 		}
-
-		FSkeletalMeshConfiguration CurrentState;
-		CurrentState.ForceRefPose = GHairStrandsMeshProjectionForceRefPoseEnable > 0;
-		CurrentState.ForceLOD     = GHairStrandsMeshProjectionForceLOD >= 0 ? FMath::Clamp(GHairStrandsMeshProjectionForceLOD, 0, SkeletalMeshComponent->GetNumLODs() - 1) : -1;
-		if (!FSkeletalMeshConfiguration::Equals(CurrentState, SkeletalMeshConfiguration) && MeshProjectionState == EMeshProjectionState::Completed)
-		{
-			SkeletalMeshComponent->SetForcedLOD(CurrentState.ForceLOD + 1);
-			SkeletalMeshComponent->SetForceRefPose(CurrentState.ForceRefPose);
-			SkeletalMeshConfiguration = CurrentState;
-		}
-
+	}
+	if (SkeletalMeshComponent)
+	{
+		// The offset is based on the center of the skeletal mesh (which is computed based on the physics capsules/boxes/...)
 		// For skinned mesh update the relative center of hair positions after deformation
+		MeshPositionOffset = SkeletalMeshComponent->CalcBounds(FTransform::Identity).GetBox().GetCenter();
 		{
 			FHairGroupResources* LocalResources = &HairGroupResources;
 
@@ -1283,12 +1187,6 @@ void UGroomComponent::TickComponent(float DeltaTime, enum ELevelTick TickType, F
 		// When a skeletal object with projection is enabled, activate the refresh of the bounding box to 
 		// insure the component/proxy bounding box always lies onto the actual skinned mesh
 		MarkRenderTransformDirty();
-
-	}
-
-	if (MeshProjectionTickDelay > 0)
-	{
-		MeshProjectionTickDelay--;
 	}
 
 	FHairGroupResources* LocalResources = &HairGroupResources;
@@ -1372,7 +1270,11 @@ void UGroomComponent::PostEditChangeProperty(FPropertyChangedEvent& PropertyChan
 
 	//  Init/release resources when setting the GroomAsset (or undoing)
 	const bool bAssetChanged = PropertyName == GET_MEMBER_NAME_CHECKED(UGroomComponent, GroomAsset);
-	if (bAssetChanged)
+	const bool bRecreateResources =
+		(bAssetChanged || PropertyThatChanged == nullptr) ||
+		PropertyName == GET_MEMBER_NAME_CHECKED(UGroomComponent, bBindGroomToSkeletalMesh);
+
+	if (bRecreateResources)
 	{
 		// Release the resources before Super::PostEditChangeProperty so that they get
 		// re-initialized in OnRegister
@@ -1381,11 +1283,7 @@ void UGroomComponent::PostEditChangeProperty(FPropertyChangedEvent& PropertyChan
 
 	Super::PostEditChangeProperty(PropertyChangedEvent);
 
-	const bool bRecreateResources =
-		(PropertyName == GET_MEMBER_NAME_CHECKED(UGroomComponent, GroomAsset) || PropertyThatChanged == nullptr) ||
-		PropertyName == GET_MEMBER_NAME_CHECKED(UGroomComponent, bBindGroomToSkeletalMesh);
-
-	const bool bSupportSkinProjection = GetDefault<URendererSettings>()->bSupportSkinCacheShaders && !UAnimationSettings::Get()->bTickAnimationOnSkeletalMeshInit;
+	const bool bSupportSkinProjection = GetDefault<URendererSettings>()->bSupportSkinCacheShaders;
 	if (!bSupportSkinProjection)
 	{
 		bBindGroomToSkeletalMesh = false;
@@ -1403,35 +1301,26 @@ void UGroomComponent::PostEditChangeProperty(FPropertyChangedEvent& PropertyChan
 	}
 #endif
 
-	if (bRecreateResources)
-	{
-		//if (GroomAsset != InitializedResources)
+	if (PropertyName == GET_MEMBER_NAME_CHECKED(FHairGroupDesc, HairWidth) || 
+		PropertyName == GET_MEMBER_NAME_CHECKED(FHairGroupDesc, HairLength) ||
+		PropertyName == GET_MEMBER_NAME_CHECKED(FHairGroupDesc, HairRootScale) ||
+		PropertyName == GET_MEMBER_NAME_CHECKED(FHairGroupDesc, HairTipScale) ||
+		PropertyName == GET_MEMBER_NAME_CHECKED(FHairGroupDesc, HairClipLength) ||
+		PropertyName == GET_MEMBER_NAME_CHECKED(FHairGroupDesc, HairShadowDensity) ||
+		PropertyName == GET_MEMBER_NAME_CHECKED(FHairGroupDesc, HairRaytracingRadiusScale) ||
+		PropertyName == GET_MEMBER_NAME_CHECKED(FHairGroupDesc, LodBias) ||
+		PropertyName == GET_MEMBER_NAME_CHECKED(FHairGroupDesc, LodAverageVertexPerPixel))
+	{	
+		UpdateHairGroupsDesc(GroomAsset, GroomGroupsDesc);
+
+		if (InterpolationInput)
 		{
-			if (GroomAsset)
-			{
-				UpdateHairGroupsDesc(GroomAsset, GroomGroupsDesc);
-				InitResources();
-			}
-			else
+			const int32 GroupCount = InterpolationInput->HairGroups.Num();
+			for (int32 GroupIt = 0; GroupIt < GroupCount; ++GroupIt)
 			{
-				ReleaseResources();
+				InterpolationInput->HairGroups[GroupIt].GroupDesc = GroomGroupsDesc[GroupIt];
 			}
 		}
-	}
-
-	if (PropertyName == GET_MEMBER_NAME_CHECKED(FHairGroupDesc, HairRaytracingRadiusScale) && InterpolationInput && InterpolationInput->HairGroups.Num() == GroomGroupsDesc.Num())
-	{
-		const int32 GroupCount = InterpolationInput->HairGroups.Num();
-		for (int32 GroupIt = 0; GroupIt < GroupCount; ++GroupIt)
-		{
-			InterpolationInput->HairGroups[GroupIt].HairRaytracingRadiusScale = GroomGroupsDesc[GroupIt].HairRaytracingRadiusScale;
-		}
-	}
-
-	if (PropertyName == GET_MEMBER_NAME_CHECKED(FHairGroupDesc, HairWidth) || 
-		PropertyName == GET_MEMBER_NAME_CHECKED(FHairGroupDesc, HairShadowDensity))
-	{	
-		UpdateHairGroupsDesc(GroomAsset, GroomGroupsDesc);
 		MarkRenderStateDirty();
 	}
 
@@ -1450,7 +1339,7 @@ bool UGroomComponent::CanEditChange(const FProperty* InProperty) const
 		{
 			bool bIsEditable = false;
 			#if RHI_RAYTRACING
-			if (IsRayTracingEnabled())
+			if (IsHairRayTracingEnabled())
 			{
 				bIsEditable = true;
 			}
@@ -1460,7 +1349,7 @@ bool UGroomComponent::CanEditChange(const FProperty* InProperty) const
 
 		if (FCString::Strcmp(*PropertyName, TEXT("bBindGroomToSkeletalMesh")) == 0)
 		{
-			return GetDefault<URendererSettings>()->bSupportSkinCacheShaders && !UAnimationSettings::Get()->bTickAnimationOnSkeletalMeshInit;
+			return GetDefault<URendererSettings>()->bSupportSkinCacheShaders;
 		}
 	}
 
diff --git a/Engine/Plugins/Experimental/HairStrands/Source/HairStrandsCore/Private/HairAttributes.cpp b/Engine/Plugins/Experimental/HairStrands/Source/HairStrandsCore/Private/HairAttributes.cpp
--- a/Engine/Plugins/Experimental/HairStrands/Source/HairStrandsCore/Private/HairAttributes.cpp
+++ b/Engine/Plugins/Experimental/HairStrands/Source/HairStrandsCore/Private/HairAttributes.cpp
@@ -15,6 +15,8 @@ namespace HairAttribute
 	const FName Strand::RootUV("groom_root_uv");
 	const FName Strand::VertexCount("vertexcount");
 	const FName Strand::Width("groom_width");
+	const FName Strand::ClosestGuides("groom_closest_guides");
+	const FName Strand::GuideWeights("groom_guide_weights");
 
 	const FName Groom::Color("groom_color");
 	const FName Groom::Width("groom_width");
diff --git a/Engine/Plugins/Experimental/HairStrands/Source/HairStrandsCore/Private/HairDescription.cpp b/Engine/Plugins/Experimental/HairStrands/Source/HairStrandsCore/Private/HairDescription.cpp
--- a/Engine/Plugins/Experimental/HairStrands/Source/HairStrandsCore/Private/HairDescription.cpp
+++ b/Engine/Plugins/Experimental/HairStrands/Source/HairStrandsCore/Private/HairDescription.cpp
@@ -1,6 +1,9 @@
 // Copyright Epic Games, Inc. All Rights Reserved.
 
 #include "HairDescription.h"
+#include "Misc/SecureHash.h"
+#include "Serialization/BulkDataReader.h"
+#include "Serialization/BulkDataWriter.h"
 
 const FStrandID FStrandID::Invalid(TNumericLimits<uint32>::Max());
 const FGroomID FGroomID::Invalid(TNumericLimits<uint32>::Max());
@@ -43,6 +46,16 @@ FStrandID FHairDescription::AddStrand()
 	return StrandID;
 }
 
+void FHairDescription::Reset()
+{
+	NumVertices = 0;
+	NumStrands = 0;
+
+	VertexAttributesSet.Initialize(0);
+	StrandAttributesSet.Initialize(0);
+	GroomAttributesSet.Initialize(0);
+}
+
 bool FHairDescription::IsValid() const
 {
 	return (NumStrands > 0) && (NumVertices > 0);
@@ -56,4 +69,82 @@ void FHairDescription::Serialize(FArchive& Ar)
 	Ar << VertexAttributesSet;
 	Ar << StrandAttributesSet;
 	Ar << GroomAttributesSet;
-}
\ No newline at end of file
+}
+
+#if WITH_EDITORONLY_DATA
+
+void FHairDescriptionBulkData::Serialize(FArchive& Ar, UObject* Owner)
+{
+	if (Ar.IsLoading())
+	{
+		// If loading, take the package custom version so it can be applied to the bulk data archive
+		// when unpacking HairDescription from it
+		CustomVersions = Ar.GetCustomVersions();
+	}
+
+	BulkData.Serialize(Ar, Owner);
+
+	Ar << Guid;
+}
+
+void FHairDescriptionBulkData::SaveHairDescription(FHairDescription& HairDescription)
+{
+	BulkData.RemoveBulkData();
+
+	if (HairDescription.IsValid())
+	{
+		FBulkDataWriter Ar(BulkData, /*bIsPersistent*/ true);
+		HairDescription.Serialize(Ar);
+
+		// Preserve CustomVersions at save time so we can reuse the same ones when reloading direct from memory
+		CustomVersions = Ar.GetCustomVersions();
+	}
+
+	// Use bulk data hash instead of guid to identify content to improve DDC cache hit
+	ComputeGuidFromHash();
+}
+
+
+void FHairDescriptionBulkData::LoadHairDescription(FHairDescription& HairDescription)
+{
+	HairDescription.Reset();
+
+	if (!IsEmpty())
+	{
+		FBulkDataReader Ar(BulkData, /*bIsPersistent*/ true);
+
+		// Propagate the custom version information from the package to the bulk data, so that the HairDescription
+		// is serialized with the same versioning
+		Ar.SetCustomVersions(CustomVersions);
+
+		HairDescription.Serialize(Ar);
+	}
+}
+
+void FHairDescriptionBulkData::Empty()
+{
+	BulkData.RemoveBulkData();
+}
+
+FString FHairDescriptionBulkData::GetIdString() const
+{
+	FString GuidString = Guid.ToString();
+	GuidString += TEXT("X");
+	return GuidString;
+}
+
+void FHairDescriptionBulkData::ComputeGuidFromHash()
+{
+	uint32 Hash[5] = {};
+
+	if (BulkData.GetBulkDataSize() > 0)
+	{
+		void* Buffer = BulkData.Lock(LOCK_READ_ONLY);
+		FSHA1::HashBuffer(Buffer, BulkData.GetBulkDataSize(), (uint8*)Hash);
+		BulkData.Unlock();
+	}
+
+	Guid = FGuid(Hash[0] ^ Hash[4], Hash[1], Hash[2], Hash[3]);
+}
+
+#endif // #if WITH_EDITORONLY_DATA
diff --git a/Engine/Plugins/Experimental/HairStrands/Source/HairStrandsCore/Private/HairStrandsRendering.cpp b/Engine/Plugins/Experimental/HairStrands/Source/HairStrandsCore/Private/HairStrandsRendering.cpp
--- a/Engine/Plugins/Experimental/HairStrands/Source/HairStrandsCore/Private/HairStrandsRendering.cpp
+++ b/Engine/Plugins/Experimental/HairStrands/Source/HairStrandsCore/Private/HairStrandsRendering.cpp
@@ -7,6 +7,10 @@
 #include "ShaderParameterStruct.h"
 #include "GlobalShader.h"
 #include "HairStrandsInterface.h"
+#include "SceneView.h"
+#include "Containers/ResourceArray.h"
+#include "HAL/ConsoleManager.h"
+#include "GpuDebugRendering.h"
 
 static int32 GHairDeformationType = 0;
 static FAutoConsoleVariableRef CVarHairDeformationType(TEXT("r.HairStrands.DeformationType"), GHairDeformationType, TEXT("Type of procedural deformation applied on hair strands (0:use simulation's output, 1:use rest strands, 2: use rest guides, 3:wave pattern, 4:follow root normal)"));
@@ -17,6 +21,13 @@ static FAutoConsoleVariableRef CVarHairRaytracingRadiusScale(TEXT("r.HairStrands
 static int32 GHairStrandsInterpolateSimulation = 1;
 static FAutoConsoleVariableRef CVarHairInterpolateSimulation(TEXT("r.HairStrands.InterpolateSimulation"), GHairStrandsInterpolateSimulation, TEXT("Enable/disable simulation output during the hair interpolation"));
 
+static float GStrandHairWidth = 0.0f;
+static FAutoConsoleVariableRef CVarStrandHairWidth(TEXT("r.HairStrands.StrandWidth"), GStrandHairWidth, TEXT("Width of hair strand"));
+
+static int32 GStrandHairInterpolationDebug = 0;
+static FAutoConsoleVariableRef CVarStrandHairInterpolationDebug(TEXT("r.HairStrands.Interpolation.Debug"), GStrandHairInterpolationDebug, TEXT("Enable debug rendering for hair interpolation"));
+
+
 static FIntVector ComputeDispatchCount(uint32 ItemCount, uint32 GroupSize)
 {
 	const uint32 BatchCount = FMath::DivideAndRoundUp(ItemCount, GroupSize);
@@ -29,6 +40,18 @@ static FIntVector ComputeDispatchCount(uint32 ItemCount, uint32 GroupSize)
 	return FIntVector(DispatchCountX, DispatchCountY, 1);
 }
 
+// Same as above but the group count is what matters and is preserved
+static FIntVector ComputeDispatchGroupCount2D(uint32 GroupCount)
+{
+	const uint32 DispatchCountX = FMath::FloorToInt(FMath::Sqrt(GroupCount));
+	const uint32 DispatchCountY = DispatchCountX + FMath::DivideAndRoundUp(GroupCount - DispatchCountX * DispatchCountX, DispatchCountX);
+
+	check(DispatchCountX <= 65535);
+	check(DispatchCountY <= 65535);
+	check(GroupCount <= DispatchCountX * DispatchCountY);
+	return FIntVector(DispatchCountX, DispatchCountY, 1);
+}
+
 inline uint32 ComputeGroupSize()
 {
 	const uint32 GroupSize = IsRHIDeviceAMD() ? 64 : (IsRHIDeviceNVIDIA() ? 32 : 64);
@@ -47,7 +70,8 @@ enum class EDeformationType : uint8
 	RestStrands,	// Use the rest strands position (no weighted interpolation)
 	RestGuide,		// Use the rest guide as input of the interpolation (no deformation), only weighted interpolation
 	Wave,			// Apply a wave pattern to deform the guides
-	NormalDirection // Apply a stretch pattern aligned with the guide root's normal
+	NormalDirection,// Apply a stretch pattern aligned with the guide root's normal
+	OffsetGuide		// Offset the guides
 };
 
 static EDeformationType GetDeformationType()
@@ -70,12 +94,15 @@ class FDeformGuideCS : public FGlobalShader
 	SHADER_USE_PARAMETER_STRUCT(FDeformGuideCS, FGlobalShader);
 
 	class FGroupSize : SHADER_PERMUTATION_INT("PERMUTATION_GROUP_SIZE", 2);
-	class FDeformationType : SHADER_PERMUTATION_INT("PERMUTATION_DEFORMATION", 3);
+	class FDeformationType : SHADER_PERMUTATION_INT("PERMUTATION_DEFORMATION", 4);
 	using FPermutationDomain = TShaderPermutationDomain<FGroupSize, FDeformationType>;
 
 	BEGIN_SHADER_PARAMETER_STRUCT(FParameters, )
 		SHADER_PARAMETER(uint32, VertexCount)
 		SHADER_PARAMETER(uint32, IterationCount)
+		SHADER_PARAMETER(FVector, SimRestOffset)
+		SHADER_PARAMETER(FVector, SimDeformedOffset)
+		SHADER_PARAMETER(uint32, DispatchCountX)
 		SHADER_PARAMETER_SRV(Buffer, SimRestPosePositionBuffer)
 		SHADER_PARAMETER_SRV(Buffer, SimRootIndexBuffer)
 		SHADER_PARAMETER_UAV(RWBuffer, OutSimDeformedPositionBuffer)
@@ -93,7 +120,9 @@ static void AddDeformSimHairStrandsPass(
 	uint32 VertexCount,
 	FRHIShaderResourceView* SimRestPosePositionBuffer,
 	FRHIShaderResourceView* SimRootIndexBuffer,
-	FRHIUnorderedAccessView* OutSimDeformedPositionBuffer)
+	FRHIUnorderedAccessView* OutSimDeformedPositionBuffer,
+	FVector& SimRestOffset,
+	FVector& SimDeformedOffset)
 {
 	static uint32 IterationCount = 0;
 	++IterationCount;
@@ -104,17 +133,25 @@ static void AddDeformSimHairStrandsPass(
 	case EDeformationType::RestGuide: InternalDeformationType = 0; break;
 	case EDeformationType::Wave: InternalDeformationType = 1; break;
 	case EDeformationType::NormalDirection: InternalDeformationType = 2; break;
+	case EDeformationType::OffsetGuide : InternalDeformationType = 3; break;
 	}
 
 	if (InternalDeformationType < 0) return;
 
+	const uint32 GroupSize = ComputeGroupSize();
+	const uint32 DispatchCount = FMath::DivideAndRoundUp(VertexCount, GroupSize);
+	const uint32 DispatchCountX = 16;
+	const uint32 DispatchCountY = FMath::DivideAndRoundUp(DispatchCount, DispatchCountX);
+
 	FDeformGuideCS::FParameters* Parameters = GraphBuilder.AllocParameters<FDeformGuideCS::FParameters>();
 	Parameters->SimRestPosePositionBuffer = SimRestPosePositionBuffer;
 	Parameters->SimRootIndexBuffer = SimRootIndexBuffer;
 	Parameters->OutSimDeformedPositionBuffer = OutSimDeformedPositionBuffer;
 	Parameters->VertexCount = VertexCount;
 	Parameters->IterationCount = IterationCount % 10000;
-	const uint32 GroupSize = ComputeGroupSize();
+	Parameters->SimDeformedOffset = SimDeformedOffset;
+	Parameters->SimRestOffset = SimRestOffset;
+	Parameters->DispatchCountX = DispatchCountX;
 
 	FDeformGuideCS::FPermutationDomain PermutationVector;
 	PermutationVector.Set<FDeformGuideCS::FGroupSize>(GetGroupSizePermutation(GroupSize));
@@ -122,54 +159,79 @@ static void AddDeformSimHairStrandsPass(
 
 	TShaderMap<FGlobalShaderType>* ShaderMap = GetGlobalShaderMap(ERHIFeatureLevel::SM5);
 
-	const uint32 DispatchCount = FMath::DivideAndRoundUp(VertexCount, GroupSize);
-	check(DispatchCount <= 65535);
-
 	TShaderMapRef<FDeformGuideCS> ComputeShader(ShaderMap, PermutationVector);
 	FComputeShaderUtils::AddPass(
 		GraphBuilder,
 		RDG_EVENT_NAME("DeformSimHairStrands"),
 		*ComputeShader,
 		Parameters,
-		FIntVector(DispatchCount, 1, 1));
+		FIntVector(DispatchCountX, DispatchCountY, 1));
 }
 
 ///////////////////////////////////////////////////////////////////////////////////////////////////
 
+struct FHairScaleAndClipDesc
+{
+	float InHairLength;
+	float InHairRadius;
+	float OutHairRadius;
+	float MaxOutHairRadius;
+	float HairRadiusRootScale;
+	float HairRadiusTipScale;
+	float HairLengthClip;
+};
+
 class FHairInterpolationCS : public FGlobalShader
 {
 	DECLARE_GLOBAL_SHADER(FHairInterpolationCS);
 	SHADER_USE_PARAMETER_STRUCT(FHairInterpolationCS, FGlobalShader);
 
 	class FGroupSize : SHADER_PERMUTATION_INT("PERMUTATION_GROUP_SIZE", 2);
-	class FDebug : SHADER_PERMUTATION_INT("PERMUTATION_DEBUG", 3);
+	class FDebug : SHADER_PERMUTATION_INT("PERMUTATION_DEBUG", 2);
 	class FDynamicGeometry : SHADER_PERMUTATION_INT("PERMUTATION_DYNAMIC_GEOMETRY", 2);
 	class FSimulation : SHADER_PERMUTATION_INT("PERMUTATION_SIMULATION", 2);
-	using FPermutationDomain = TShaderPermutationDomain<FGroupSize, FDebug, FDynamicGeometry, FSimulation>;
+	class FScaleAndClip : SHADER_PERMUTATION_INT("PERMUTATION_SCALE_AND_CLIP", 2);
+	using FPermutationDomain = TShaderPermutationDomain<FGroupSize, FDebug, FDynamicGeometry, FSimulation, FScaleAndClip>;
 
 	BEGIN_SHADER_PARAMETER_STRUCT(FParameters, )
+		SHADER_PARAMETER_STRUCT_INCLUDE(ShaderDrawDebug::FShaderDrawDebugParameters, ShaderDrawParameters)
 		SHADER_PARAMETER(uint32, VertexCount)
 		SHADER_PARAMETER(uint32, DispatchCountX)
+		SHADER_PARAMETER(uint32, HairDebugMode)
 		SHADER_PARAMETER(FVector, InRenderHairPositionOffset)
 		SHADER_PARAMETER(FVector, InSimHairPositionOffset)
 		SHADER_PARAMETER(FVector, OutHairPositionOffset)
 		SHADER_PARAMETER(FIntPoint, HairStrandsCullIndex)
 
+		SHADER_PARAMETER(float, InHairLength)
+		SHADER_PARAMETER(float, InHairRadius)
+		SHADER_PARAMETER(float, OutHairRadius)
+		SHADER_PARAMETER(float, MaxOutHairRadius)
+		SHADER_PARAMETER(float, HairRadiusRootScale)
+		SHADER_PARAMETER(float, HairRadiusTipScale)
+		SHADER_PARAMETER(float, HairLengthClip)
+
 		SHADER_PARAMETER(FVector, RestPositionOffset)
 		SHADER_PARAMETER(FVector, DeformedPositionOffset)
 
 		SHADER_PARAMETER(FVector, SimRestPositionOffset)
 		SHADER_PARAMETER(FVector, SimDeformedPositionOffset)
 
+		SHADER_PARAMETER(FMatrix, LocalToWorldMatrix)
+
 		SHADER_PARAMETER_SRV(Buffer, RenderRestPosePositionBuffer)
 		SHADER_PARAMETER_UAV(RWBuffer, OutRenderDeformedPositionBuffer)
 
+		SHADER_PARAMETER_UAV(RWBuffer, OutClusterBuffer)
+		SHADER_PARAMETER_SRV(Buffer, VertexToClusterIdBuffer)
+
 		SHADER_PARAMETER_SRV(Buffer, SimRestPosePositionBuffer)
 		SHADER_PARAMETER_SRV(Buffer, DeformedSimPositionBuffer)
 
 		SHADER_PARAMETER_SRV(Buffer, Interpolation0Buffer)
 		SHADER_PARAMETER_SRV(Buffer, Interpolation1Buffer)
 
+		SHADER_PARAMETER_SRV(Buffer, AttributeBuffer)
 		SHADER_PARAMETER_SRV(Buffer, SimAttributeBuffer)
 		SHADER_PARAMETER_UAV(RWBuffer, OutRenderAttributeBuffer)
 		SHADER_PARAMETER_SRV(Buffer<float4>, RestPosition0Buffer)
@@ -206,22 +268,27 @@ IMPLEMENT_GLOBAL_SHADER(FHairInterpolationCS, "/Engine/Private/HairStrands/HairS
 
 static void AddHairStrandsInterpolationPass(
 	FRDGBuilder& GraphBuilder,
+	const FShaderDrawDebugData* ShaderDrawData,
 	const FHairStrandsProjectionHairData::HairGroup& InRenHairData,
 	const FHairStrandsProjectionHairData::HairGroup& InSimHairData,
 	const FVector& InRenderHairWorldOffset,
 	const FVector& InSimHairWorldOffset,
 	const FVector& OutHairWorldOffset,
+	const FHairScaleAndClipDesc ScaleAndClipDesc,
 	const int32 LODIndex,
 	const bool bHasSimulationEnable,
+	const bool bPatchedAttributeBuffer, 
 	const uint32 VertexCount,
 	const FShaderResourceViewRHIRef& RenderRestPosePositionBuffer,
+	const FShaderResourceViewRHIRef& RenderAttributeBuffer,
 	const FShaderResourceViewRHIRef& Interpolation0Buffer,
 	const FShaderResourceViewRHIRef& Interpolation1Buffer,
 	const FShaderResourceViewRHIRef& SimRestPosePositionBuffer,
 	const FShaderResourceViewRHIRef& SimDeformedPositionBuffer,
 	const FShaderResourceViewRHIRef& SimAttributeBuffer,
 	const FUnorderedAccessViewRHIRef& OutRenderPositionBuffer,
-	const FUnorderedAccessViewRHIRef& OutRenderAttributeBuffer)
+	const FUnorderedAccessViewRHIRef& OutRenderAttributeBuffer,
+	const FShaderResourceViewRHIRef& VertexToClusterIdBuffer)
 {
 	const uint32 GroupSize = ComputeGroupSize();
 	const FIntVector DispatchCount = ComputeDispatchCount(VertexCount, GroupSize);
@@ -240,6 +307,23 @@ static void AddHairStrandsInterpolationPass(
 	Parameters->OutHairPositionOffset = OutHairWorldOffset;
 	Parameters->DispatchCountX = DispatchCount.X;
 	
+	const bool bNeedScaleOrClip = 
+		ScaleAndClipDesc.InHairRadius != ScaleAndClipDesc.OutHairRadius || 
+		ScaleAndClipDesc.HairRadiusRootScale != 1 || 
+		ScaleAndClipDesc.HairRadiusTipScale != 1 ||
+		ScaleAndClipDesc.HairLengthClip < 1;
+
+	Parameters->InHairLength = ScaleAndClipDesc.InHairLength;
+	Parameters->InHairRadius = ScaleAndClipDesc.InHairRadius;
+	Parameters->OutHairRadius = ScaleAndClipDesc.OutHairRadius;
+	Parameters->MaxOutHairRadius = ScaleAndClipDesc.MaxOutHairRadius;
+	Parameters->HairRadiusRootScale = ScaleAndClipDesc.HairRadiusRootScale;
+	Parameters->HairRadiusTipScale = ScaleAndClipDesc.HairRadiusTipScale;
+	Parameters->HairLengthClip = ScaleAndClipDesc.HairLengthClip * ScaleAndClipDesc.InHairLength; // HairLengthClip is the normalized length and we convert it to world length
+	if (bNeedScaleOrClip)
+	{
+		Parameters->AttributeBuffer = RenderAttributeBuffer;
+	}
 	const bool bIsVertexToCurveBuffersValid = InRenHairData.VertexToCurveIndexBuffer && InSimHairData.VertexToCurveIndexBuffer;
 	if (bIsVertexToCurveBuffersValid)
 	{
@@ -247,32 +331,56 @@ static void AddHairStrandsInterpolationPass(
 		Parameters->SimVertexToRootIndexBuffer = InSimHairData.VertexToCurveIndexBuffer->SRV;
 	}
 
-	const FHairCullInfo Info = GetHairStrandsCullInfo();
-	const bool bCullingEnable = Info.CullMode != EHairCullMode::None && bIsVertexToCurveBuffersValid;
-	if (bCullingEnable)
-	{
-		if (Info.CullMode == EHairCullMode::Sim)
-			Parameters->HairStrandsCullIndex.Y = Info.ExplicitIndex >= 0 ? Info.ExplicitIndex : FMath::Clamp(uint32(Info.NormalizedIndex * InSimHairData.RootCount), 0u, InSimHairData.RootCount - 1);
-		if (Info.CullMode == EHairCullMode::Render)
-			Parameters->HairStrandsCullIndex.X = Info.ExplicitIndex >= 0 ? Info.ExplicitIndex : FMath::Clamp(uint32(Info.NormalizedIndex * InRenHairData.RootCount), 0u, InRenHairData.RootCount - 1);
-	}
+	Parameters->VertexToClusterIdBuffer = VertexToClusterIdBuffer;
+	
+	Parameters->LocalToWorldMatrix = InRenHairData.LocalToWorld.ToMatrixWithScale();
 
-	const bool bCopySimAttributesToRenderAttributes = SimAttributeBuffer != nullptr && OutRenderAttributeBuffer != nullptr;
-	if (bCopySimAttributesToRenderAttributes)
+	// Debug rendering
+	Parameters->HairDebugMode = 0;
 	{
-		Parameters->SimAttributeBuffer = SimAttributeBuffer;
-		Parameters->OutRenderAttributeBuffer = OutRenderAttributeBuffer;
+		const FHairCullInfo Info = GetHairStrandsCullInfo();
+		const bool bCullingEnable = Info.CullMode != EHairCullMode::None && bIsVertexToCurveBuffersValid;
+
+		if (bPatchedAttributeBuffer)
+		{
+			Parameters->HairDebugMode = 1;
+			check(SimAttributeBuffer != nullptr);
+			check(OutRenderAttributeBuffer != nullptr);
+		}
+		else if (GStrandHairInterpolationDebug > 0)
+		{
+			Parameters->HairDebugMode = 4;
+		}
+		else if (GetHairStrandsDebugStrandsMode() == EHairStrandsDebugMode::RenderVisCluster)
+		{
+			Parameters->HairDebugMode = 3;
+		}
+		else if (bCullingEnable)
+		{
+			Parameters->HairDebugMode = 2;
+			
+			if (Info.CullMode == EHairCullMode::Sim)
+				Parameters->HairStrandsCullIndex.Y = Info.ExplicitIndex >= 0 ? Info.ExplicitIndex : FMath::Clamp(uint32(Info.NormalizedIndex * InSimHairData.RootCount), 0u, InSimHairData.RootCount - 1);
+			if (Info.CullMode == EHairCullMode::Render)
+				Parameters->HairStrandsCullIndex.X = Info.ExplicitIndex >= 0 ? Info.ExplicitIndex : FMath::Clamp(uint32(Info.NormalizedIndex * InRenHairData.RootCount), 0u, InRenHairData.RootCount - 1);
+		}
+
+		if (Parameters->HairDebugMode > 0)
+		{
+			Parameters->SimAttributeBuffer = SimAttributeBuffer;
+			Parameters->OutRenderAttributeBuffer = OutRenderAttributeBuffer;
+		}
 	}
 
 	const bool bSupportDynamicMesh = InRenHairData.RootCount > 0 && LODIndex >= 0 && LODIndex < InRenHairData.LODDatas.Num() && InRenHairData.LODDatas[LODIndex].bIsValid && bIsVertexToCurveBuffersValid;
 	if (bSupportDynamicMesh)
 	{
-		Parameters->RestPositionOffset = InRenHairData.LODDatas[LODIndex].RestPositionOffset;
+		Parameters->RestPositionOffset = *InRenHairData.LODDatas[LODIndex].RestPositionOffset;
 		Parameters->RestPosition0Buffer = InRenHairData.LODDatas[LODIndex].RestRootTrianglePosition0Buffer->SRV;
 		Parameters->RestPosition1Buffer = InRenHairData.LODDatas[LODIndex].RestRootTrianglePosition1Buffer->SRV;
 		Parameters->RestPosition2Buffer = InRenHairData.LODDatas[LODIndex].RestRootTrianglePosition2Buffer->SRV;
 
-		Parameters->DeformedPositionOffset = InRenHairData.LODDatas[LODIndex].DeformedPositionOffset;
+		Parameters->DeformedPositionOffset = *InRenHairData.LODDatas[LODIndex].DeformedPositionOffset;
 		Parameters->DeformedPosition0Buffer = InRenHairData.LODDatas[LODIndex].DeformedRootTrianglePosition0Buffer->SRV;
 		Parameters->DeformedPosition1Buffer = InRenHairData.LODDatas[LODIndex].DeformedRootTrianglePosition1Buffer->SRV;
 		Parameters->DeformedPosition2Buffer = InRenHairData.LODDatas[LODIndex].DeformedRootTrianglePosition2Buffer->SRV;
@@ -280,12 +388,12 @@ static void AddHairStrandsInterpolationPass(
 		Parameters->RootToTriangleIndex = InRenHairData.LODDatas[LODIndex].RootTriangleIndexBuffer->SRV;
 		Parameters->RootBarycentricBuffer = InRenHairData.LODDatas[LODIndex].RootTriangleBarycentricBuffer->SRV;
 
-		Parameters->SimRestPositionOffset = InSimHairData.LODDatas[LODIndex].RestPositionOffset;
+		Parameters->SimRestPositionOffset = *InSimHairData.LODDatas[LODIndex].RestPositionOffset;
 		Parameters->SimRestPosition0Buffer = InSimHairData.LODDatas[LODIndex].RestRootTrianglePosition0Buffer->SRV;
 		Parameters->SimRestPosition1Buffer = InSimHairData.LODDatas[LODIndex].RestRootTrianglePosition1Buffer->SRV;
 		Parameters->SimRestPosition2Buffer = InSimHairData.LODDatas[LODIndex].RestRootTrianglePosition2Buffer->SRV;
 
-		Parameters->SimDeformedPositionOffset = InSimHairData.LODDatas[LODIndex].DeformedPositionOffset;
+		Parameters->SimDeformedPositionOffset = *InSimHairData.LODDatas[LODIndex].DeformedPositionOffset;
 		Parameters->SimDeformedPosition0Buffer = InSimHairData.LODDatas[LODIndex].DeformedRootTrianglePosition0Buffer->SRV;
 		Parameters->SimDeformedPosition1Buffer = InSimHairData.LODDatas[LODIndex].DeformedRootTrianglePosition1Buffer->SRV;
 		Parameters->SimDeformedPosition2Buffer = InSimHairData.LODDatas[LODIndex].DeformedRootTrianglePosition2Buffer->SRV;
@@ -294,11 +402,17 @@ static void AddHairStrandsInterpolationPass(
 		Parameters->SimRootBarycentricBuffer = InSimHairData.LODDatas[LODIndex].RootTriangleBarycentricBuffer->SRV;
 	}
 
+	if (ShaderDrawDebug::IsShaderDrawDebugEnabled() && ShaderDrawData)
+	{
+		ShaderDrawDebug::SetParameters(GraphBuilder, *ShaderDrawData, Parameters->ShaderDrawParameters);
+	}
+
 	FHairInterpolationCS::FPermutationDomain PermutationVector;
 	PermutationVector.Set<FHairInterpolationCS::FGroupSize>(GetGroupSizePermutation(GroupSize));
-	PermutationVector.Set<FHairInterpolationCS::FDebug>(bCopySimAttributesToRenderAttributes ? 1 : (bCullingEnable ? 2 : 0));
+	PermutationVector.Set<FHairInterpolationCS::FDebug>(Parameters->HairDebugMode > 0 ? 1 : 0);
 	PermutationVector.Set<FHairInterpolationCS::FDynamicGeometry>(bSupportDynamicMesh ? 1 : 0);
 	PermutationVector.Set<FHairInterpolationCS::FSimulation>(bHasSimulationEnable ? 1 : 0);
+	PermutationVector.Set<FHairInterpolationCS::FScaleAndClip>(bNeedScaleOrClip ? 1 : 0);
 
 	TShaderMap<FGlobalShaderType>* ShaderMap = GetGlobalShaderMap(ERHIFeatureLevel::SM5);
 
@@ -313,6 +427,68 @@ static void AddHairStrandsInterpolationPass(
 
 ///////////////////////////////////////////////////////////////////////////////////////////////////
 
+
+class FHairClusterAABBCS : public FGlobalShader
+{
+	DECLARE_GLOBAL_SHADER(FHairClusterAABBCS);
+	SHADER_USE_PARAMETER_STRUCT(FHairClusterAABBCS, FGlobalShader);
+
+	class FGroupSize : SHADER_PERMUTATION_INT("PERMUTATION_GROUP_SIZE", 2);
+	using FPermutationDomain = TShaderPermutationDomain<FGroupSize>;
+
+	BEGIN_SHADER_PARAMETER_STRUCT(FParameters, )
+		SHADER_PARAMETER(uint32, DispatchCountX)
+		SHADER_PARAMETER(uint32, ClusterCount)
+		SHADER_PARAMETER(FVector, OutHairPositionOffset)
+		SHADER_PARAMETER(FMatrix, LocalToWorldMatrix)
+		SHADER_PARAMETER_SRV(Buffer, RenderDeformedPositionBuffer)
+		SHADER_PARAMETER_SRV(Buffer, ClusterVertexIdBuffer)
+		SHADER_PARAMETER_SRV(Buffer, ClusterInfoBuffer)
+		SHADER_PARAMETER_UAV(RWBuffer, OutClusterAABBBuffer)
+		SHADER_PARAMETER_UAV(RWBuffer, OutGroupAABBBuffer)
+	END_SHADER_PARAMETER_STRUCT()
+
+public:
+	static bool ShouldCompilePermutation(const FGlobalShaderPermutationParameters& Parameters) { return IsHairStrandsSupported(Parameters.Platform); }
+};
+
+IMPLEMENT_GLOBAL_SHADER(FHairClusterAABBCS, "/Engine/Private/HairStrands/HairStrandsInterpolation.usf", "ClusterAABBEvaluationCS", SF_Compute);
+
+static void AddHairClusterAABBPass(
+	FRDGBuilder& GraphBuilder,
+	const FHairStrandsProjectionHairData::HairGroup& InRenHairData,
+	const FVector& OutHairWorldOffset,
+	FHairStrandClusterData::FHairGroup& ClusterData,
+	const FShaderResourceViewRHIRef& RenderPositionBuffer)
+{
+	const uint32 GroupSize = ComputeGroupSize();
+	const FIntVector DispatchCount = ComputeDispatchGroupCount2D(ClusterData.ClusterCount);
+
+	FHairClusterAABBCS::FParameters* Parameters = GraphBuilder.AllocParameters<FHairClusterAABBCS::FParameters>();
+	Parameters->DispatchCountX = DispatchCount.X;
+	Parameters->ClusterCount = ClusterData.ClusterCount;
+	Parameters->LocalToWorldMatrix = InRenHairData.LocalToWorld.ToMatrixWithScale();
+	Parameters->OutHairPositionOffset = OutHairWorldOffset;
+	Parameters->RenderDeformedPositionBuffer = RenderPositionBuffer;
+	Parameters->ClusterVertexIdBuffer = ClusterData.ClusterVertexIdBuffer->SRV;
+	Parameters->ClusterInfoBuffer = ClusterData.ClusterInfoBuffer->SRV;
+	Parameters->OutClusterAABBBuffer = ClusterData.HairGroupPublicPtr->GetClusterAABBBuffer().UAV;
+	Parameters->OutGroupAABBBuffer = ClusterData.HairGroupPublicPtr->GetGroupAABBBuffer().UAV;
+
+	FHairClusterAABBCS::FPermutationDomain PermutationVector;
+	PermutationVector.Set<FHairClusterAABBCS::FGroupSize>(GetGroupSizePermutation(GroupSize));
+	TShaderMapRef<FHairClusterAABBCS> ComputeShader(GetGlobalShaderMap(ERHIFeatureLevel::SM5), PermutationVector);
+
+	FComputeShaderUtils::AddPass(
+		GraphBuilder,
+		RDG_EVENT_NAME("HairStrandsClusterAABB"),
+		*ComputeShader,
+		Parameters,
+		DispatchCount);
+}
+
+///////////////////////////////////////////////////////////////////////////////////////////////////
+
 class FHairTangentCS : public FGlobalShader
 {
 	DECLARE_GLOBAL_SHADER(FHairTangentCS);
@@ -423,6 +599,56 @@ static void AddGenerateRaytracingGeometryPass(
 
 ///////////////////////////////////////////////////////////////////////////////////////////////////
 
+class FClearClusterAABBCS : public FGlobalShader
+{
+	DECLARE_GLOBAL_SHADER(FClearClusterAABBCS);
+	SHADER_USE_PARAMETER_STRUCT(FClearClusterAABBCS, FGlobalShader);
+
+	BEGIN_SHADER_PARAMETER_STRUCT(FParameters, )
+		SHADER_PARAMETER_UAV(RWBuffer, OutClusterAABBBuffer)
+		SHADER_PARAMETER_UAV(RWBuffer, OutGroupAABBBuffer)
+		SHADER_PARAMETER(uint32, ClusterCount)
+	END_SHADER_PARAMETER_STRUCT()
+
+public:
+	static bool ShouldCompilePermutation(const FGlobalShaderPermutationParameters& Parameters) { return IsHairStrandsSupported(Parameters.Platform); }
+	static void ModifyCompilationEnvironment(const FGlobalShaderPermutationParameters& Parameters, FShaderCompilerEnvironment& OutEnvironment)
+	{
+		FGlobalShader::ModifyCompilationEnvironment(Parameters, OutEnvironment);
+		OutEnvironment.SetDefine(TEXT("SHADER_CLEARCLUSTERAABB"), 1);
+	}
+};
+
+IMPLEMENT_GLOBAL_SHADER(FClearClusterAABBCS, "/Engine/Private/HairStrands/HairStrandsClusterCulling.usf", "MainClearClusterAABBCS", SF_Compute);
+
+static void AddClearClusterAABBPass(
+	FRDGBuilder& GraphBuilder,
+	uint32 ClusterCount,
+	FRHIUnorderedAccessView* OutClusterAABBuffer,
+	FRHIUnorderedAccessView* OutGroupAABBuffer)
+{
+	check(OutClusterAABBuffer);
+
+	FClearClusterAABBCS::FParameters* Parameters = GraphBuilder.AllocParameters<FClearClusterAABBCS::FParameters>();
+	Parameters->ClusterCount = ClusterCount;
+	Parameters->OutClusterAABBBuffer = OutClusterAABBuffer;
+	Parameters->OutGroupAABBBuffer = OutGroupAABBuffer;
+
+	
+	TShaderMap<FGlobalShaderType>* ShaderMap = GetGlobalShaderMap(ERHIFeatureLevel::SM5);
+	TShaderMapRef<FClearClusterAABBCS> ComputeShader(ShaderMap);
+
+	const FIntVector DispatchCount = DispatchCount.DivideAndRoundUp(FIntVector(ClusterCount * 6, 1, 1) , FIntVector(64, 1, 1));
+	FComputeShaderUtils::AddPass(
+		GraphBuilder,
+		RDG_EVENT_NAME("HairStrandsClearClusterAABB"),
+		*ComputeShader,
+		Parameters,
+		DispatchCount);
+}
+
+///////////////////////////////////////////////////////////////////////////////////////////////////
+
 #if RHI_RAYTRACING
 static void UpdateHairAccelerationStructure(FRHICommandList& RHICmdList, FRayTracingGeometry* RayTracingGeometry)
 {
@@ -461,18 +687,21 @@ static void BuildHairAccelerationStructure(FRHICommandList& RHICmdList, uint32 R
 
 void ComputeHairStrandsInterpolation(
 	FRHICommandListImmediate& RHICmdList,
+	const FShaderDrawDebugData* ShaderDrawData,
 	FHairStrandsInterpolationInput* InInput,
 	FHairStrandsInterpolationOutput* InOutput,
 	FHairStrandsProjectionHairData& InRenHairDatas,
 	FHairStrandsProjectionHairData& InSimHairDatas,
-	int32 LODIndex)
+	int32 LODIndex,
+	FHairStrandClusterData* ClusterData)
 {
 	if (!InInput || !InOutput) return;
 
+
 	const uint32 GroupCount = InOutput->HairGroups.Num();
 	for (uint32 GroupIndex = 0; GroupIndex < GroupCount; ++GroupIndex)
 	{
-		FHairStrandsInterpolationInput::FHairGroup& Input   = InInput->HairGroups[GroupIndex];
+		FHairStrandsInterpolationInput::FHairGroup& Input  = InInput->HairGroups[GroupIndex];
 		FHairStrandsInterpolationOutput::HairGroup& Output = InOutput->HairGroups[GroupIndex];
 		Output.VFInput.Reset();
 
@@ -496,7 +725,8 @@ void ComputeHairStrandsInterpolation(
 				Input.SimVertexCount,
 				Input.SimRestPosePositionBuffer->SRV,
 				Input.SimRootPointIndexBuffer ? Input.SimRootPointIndexBuffer->SRV : nullptr,
-				Output.SimDeformedPositionBuffer[CurrIndex]->UAV);
+				Output.SimDeformedPositionBuffer[CurrIndex]->UAV, Input.InSimHairPositionOffset,
+				Input.OutHairPositionOffset );
 		}
 
 		// If the deformation is driven by the physics simulation, then the output is always the 0 index
@@ -527,40 +757,90 @@ void ComputeHairStrandsInterpolation(
 			Output.VFInput.HairPositionOffset = bHasSimulationEnabled ? Input.OutHairPositionOffset : Input.InSimHairPositionOffset;
 			Output.VFInput.HairPreviousPositionOffset = bHasSimulationEnabled ? Input.OutHairPreviousPositionOffset : Input.InSimHairPositionOffset;
 			Output.VFInput.VertexCount = Input.SimVertexCount;
+			Output.VFInput.HairRadius = (GStrandHairWidth > 0 ? GStrandHairWidth : Input.GroupDesc.HairWidth) * 0.5f;
+			Output.VFInput.HairLength = Input.GroupDesc.HairLength;
+			Output.VFInput.HairDensity = Input.GroupDesc.HairShadowDensity;
 		}
 		else
 		{
+			check(ClusterData);
+			const bool bDebugModePatchedAttributeBuffer = DebugMode == EHairStrandsDebugMode::RenderHairStrands || DebugMode == EHairStrandsDebugMode::RenderVisCluster;
+
 			const uint32 BufferSizeInBytes = Input.RenderVertexCount * FHairStrandsAttributeFormat::SizeInByte;
-			if (DebugMode == EHairStrandsDebugMode::RenderHairStrands && Output.RenderPatchedAttributeBuffer.NumBytes != BufferSizeInBytes)
+			if (bDebugModePatchedAttributeBuffer && Output.RenderPatchedAttributeBuffer.NumBytes != BufferSizeInBytes)
 			{
 				Output.RenderPatchedAttributeBuffer.Release();
 				Output.RenderPatchedAttributeBuffer.Initialize(FHairStrandsAttributeFormat::SizeInByte, Input.RenderVertexCount, FHairStrandsAttributeFormat::Format, BUF_Static);
 			}
 
+			AddClearClusterAABBPass(
+				GraphBuilder,
+				Input.ClusterCount,
+				Output.RenderClusterAABBBuffer->UAV,
+				Output.RenderGroupAABBBuffer->UAV);
+
+			FHairScaleAndClipDesc ScaleAndClipDesc;
+			ScaleAndClipDesc.InHairLength = Input.GroupDesc.HairLength;
+			ScaleAndClipDesc.InHairRadius = Input.GroupDesc.HairWidth * 0.5f;
+			ScaleAndClipDesc.OutHairRadius = (GStrandHairWidth > 0 ? GStrandHairWidth : Input.GroupDesc.HairWidth) * 0.5f;
+			ScaleAndClipDesc.MaxOutHairRadius = ScaleAndClipDesc.OutHairRadius * FMath::Max(1.f, FMath::Max(Input.GroupDesc.HairRootScale, Input.GroupDesc.HairTipScale));
+			ScaleAndClipDesc.HairRadiusRootScale = Input.GroupDesc.HairRootScale;
+			ScaleAndClipDesc.HairRadiusTipScale = Input.GroupDesc.HairTipScale;
+			ScaleAndClipDesc.HairLengthClip = Input.GroupDesc.HairClipLength;
+
 			const bool bHasSimulationEnabled = Input.bIsSimulationEnable && GHairStrandsInterpolateSimulation && DeformationType != EDeformationType::RestStrands;
 			check(GroupIndex < uint32(InRenHairDatas.HairGroups.Num()));
 			check(GroupIndex < uint32(InSimHairDatas.HairGroups.Num()));
 			AddHairStrandsInterpolationPass(
 				GraphBuilder,
+				ShaderDrawData,
 				InRenHairDatas.HairGroups[GroupIndex],
 				InSimHairDatas.HairGroups[GroupIndex],
 				Input.InRenderHairPositionOffset,
 				Input.InSimHairPositionOffset,
 				Input.OutHairPositionOffset,
+				ScaleAndClipDesc,
 				LODIndex,
 				bHasSimulationEnabled,
+				bDebugModePatchedAttributeBuffer,
 				Input.RenderVertexCount,
 				Input.RenderRestPosePositionBuffer->SRV,
+				Input.RenderAttributeBuffer->SRV,
 				Input.Interpolation0Buffer->SRV,
 				Input.Interpolation1Buffer->SRV,
 				Input.SimRestPosePositionBuffer->SRV,
 				Output.SimDeformedPositionBuffer[SimIndex]->SRV,
-				DebugMode == EHairStrandsDebugMode::RenderHairStrands ? Input.SimAttributeBuffer->SRV : nullptr,
+				Input.SimAttributeBuffer->SRV,
 				Output.RenderDeformedPositionBuffer[CurrIndex]->UAV,
-				DebugMode == EHairStrandsDebugMode::RenderHairStrands ? Output.RenderPatchedAttributeBuffer.UAV : nullptr);
-
+				Output.RenderPatchedAttributeBuffer.UAV,
+				Input.VertexToClusterIdBuffer->SRV);
+
+			// Initialize group cluster data for culling by the renderer
+			FHairStrandClusterData::FHairGroup& HairGroupCluster = ClusterData->HairGroups.Emplace_GetRef();
+			HairGroupCluster.ClusterCount = Input.ClusterCount;
+			HairGroupCluster.VertexCount = Input.ClusterVertexCount;
+			HairGroupCluster.GroupAABBBuffer = Output.RenderGroupAABBBuffer;
+			HairGroupCluster.ClusterAABBBuffer = Output.RenderClusterAABBBuffer;
+			HairGroupCluster.ClusterInfoBuffer = Output.ClusterInfoBuffer;
+			HairGroupCluster.VertexToClusterIdBuffer = Input.VertexToClusterIdBuffer;
+			HairGroupCluster.ClusterVertexIdBuffer = Input.ClusterVertexIdBuffer;
+			HairGroupCluster.ClusterIndexRadiusScaleInfoBuffer = Input.ClusterIndexRadiusScaleInfoBuffer;
+			HairGroupCluster.HairGroupPublicPtr = Output.HairGroupPublicData;
+			HairGroupCluster.LodBias = Input.GroupDesc.LodBias;
+			HairGroupCluster.LodAverageVertexPerPixel = Input.GroupDesc.LodAverageVertexPerPixel;
+
+			Output.VFInput.HairRadius = ScaleAndClipDesc.MaxOutHairRadius;
+			Output.VFInput.HairLength = Input.GroupDesc.HairLength;
+			Output.VFInput.HairDensity = Input.GroupDesc.HairShadowDensity;
 			Output.VFInput.HairPositionBuffer = Output.RenderDeformedPositionBuffer[CurrIndex]->SRV;
 			Output.VFInput.HairPreviousPositionBuffer = Output.RenderDeformedPositionBuffer[PrevIndex]->SRV;
+
+			AddHairClusterAABBPass(
+				GraphBuilder,
+				InRenHairDatas.HairGroups[GroupIndex],
+				Input.OutHairPositionOffset,
+				HairGroupCluster,
+				Output.RenderDeformedPositionBuffer[CurrIndex]->SRV);
 			
 			AddHairTangentPass(
 				GraphBuilder,
@@ -569,12 +849,13 @@ void ComputeHairStrandsInterpolation(
 				Output.RenderTangentBuffer->UAV);
 
 			#if RHI_RAYTRACING
-			if (IsRayTracingEnabled())
+			if (IsHairRayTracingEnabled())
 			{
+				const float HairRadiusScaleRT = (GHairRaytracingRadiusScale > 0 ? GHairRaytracingRadiusScale : Input.GroupDesc.HairRaytracingRadiusScale);
 				AddGenerateRaytracingGeometryPass(
 					GraphBuilder,
 					Input.RenderVertexCount,
-					Input.HairRadius * (GHairRaytracingRadiusScale > 0 ? GHairRaytracingRadiusScale : Input.HairRaytracingRadiusScale),
+					Output.VFInput.HairRadius * HairRadiusScaleRT,
 					Input.OutHairPositionOffset,
 					Output.VFInput.HairPositionBuffer,
 					Input.RaytracingPositionBuffer->UAV);
@@ -582,15 +863,15 @@ void ComputeHairStrandsInterpolation(
 			#endif
 			GraphBuilder.Execute();
 
-			Output.VFInput.HairTangentBuffer = Output.RenderTangentBuffer->SRV;
-			Output.VFInput.HairAttributeBuffer = DebugMode == EHairStrandsDebugMode::RenderHairStrands ? Output.RenderPatchedAttributeBuffer.SRV : Input.RenderAttributeBuffer->SRV;
-			Output.VFInput.HairMaterialBuffer = Output.RenderMaterialBuffer->SRV;
-			Output.VFInput.HairPositionOffset = Input.OutHairPositionOffset;
-			Output.VFInput.HairPreviousPositionOffset = Input.OutHairPreviousPositionOffset;
-			Output.VFInput.VertexCount = Input.RenderVertexCount;
+			Output.VFInput.HairTangentBuffer			= Output.RenderTangentBuffer->SRV;
+			Output.VFInput.HairAttributeBuffer			= bDebugModePatchedAttributeBuffer ? Output.RenderPatchedAttributeBuffer.SRV : Input.RenderAttributeBuffer->SRV;
+			Output.VFInput.HairMaterialBuffer			= Output.RenderMaterialBuffer->SRV;
+			Output.VFInput.HairPositionOffset			= Input.OutHairPositionOffset;
+			Output.VFInput.HairPreviousPositionOffset	= Input.OutHairPreviousPositionOffset;
+			Output.VFInput.VertexCount					= Input.RenderVertexCount;
 
 			#if RHI_RAYTRACING
-			if (IsRayTracingEnabled())
+			if (IsHairRayTracingEnabled())
 			{
 				FRHIUnorderedAccessView* UAV = Input.RaytracingPositionBuffer->UAV;
 				RHICmdList.TransitionResources(EResourceTransitionAccess::EReadable, EResourceTransitionPipeline::EComputeToGfx, &UAV, 1);
@@ -608,3 +889,36 @@ void ComputeHairStrandsInterpolation(
 		Output.CurrentIndex = PrevIndex;
 	}
 }
+
+void ResetHairStrandsInterpolation(
+	FRHICommandListImmediate& RHICmdList,
+	FHairStrandsInterpolationInput* InInput,
+	FHairStrandsInterpolationOutput* InOutput)
+{
+	if (!InInput || !InOutput) return;
+
+	const uint32 GroupCount = InOutput->HairGroups.Num();
+	for (uint32 GroupIndex = 0; GroupIndex < GroupCount; ++GroupIndex)
+	{
+		FHairStrandsInterpolationInput::FHairGroup& Input = InInput->HairGroups[GroupIndex];
+		FHairStrandsInterpolationOutput::HairGroup& Output = InOutput->HairGroups[GroupIndex];
+		if (!Input.IsValid() || !Output.IsValid()) return;
+
+		DECLARE_GPU_STAT(HairStrandsResetInterpolation);
+		SCOPED_DRAW_EVENT(RHICmdList, HairStrandsResetInterpolation);
+		SCOPED_GPU_STAT(RHICmdList, HairStrandsResetInterpolation);
+
+		const uint32 CurrIndex = Output.CurrentIndex;
+		const uint32 PrevIndex = (Output.CurrentIndex + 1) % 2;
+		FRDGBuilder GraphBuilder(RHICmdList);
+		AddDeformSimHairStrandsPass(
+			GraphBuilder,
+			EDeformationType::OffsetGuide,
+			Input.SimVertexCount,
+			Input.SimRestPosePositionBuffer->SRV,
+			Input.SimRootPointIndexBuffer ? Input.SimRootPointIndexBuffer->SRV : nullptr,
+			Output.SimDeformedPositionBuffer[CurrIndex]->UAV, Input.InSimHairPositionOffset,
+			Input.OutHairPositionOffset);
+		GraphBuilder.Execute();
+	}
+}
\ No newline at end of file
diff --git a/Engine/Plugins/Experimental/HairStrands/Source/HairStrandsCore/Private/HairStrandsVertexFactory.cpp b/Engine/Plugins/Experimental/HairStrands/Source/HairStrandsCore/Private/HairStrandsVertexFactory.cpp
--- a/Engine/Plugins/Experimental/HairStrands/Source/HairStrandsCore/Private/HairStrandsVertexFactory.cpp
+++ b/Engine/Plugins/Experimental/HairStrands/Source/HairStrandsCore/Private/HairStrandsVertexFactory.cpp
@@ -12,25 +12,34 @@
 #include "MeshMaterialShader.h"
 #include "HairStrandsInterface.h"
 
-static float GStrandHairWidth = 0.0f;
-static FAutoConsoleVariableRef CVarStrandHairWidth(TEXT("r.HairStrands.StrandWidth"), GStrandHairWidth, TEXT("Width of hair strand"));
-
-float FHairStrandsVertexFactory::GetMaxStrandRadius(uint32 GroupIndex) const
-{
-	return GStrandHairWidth > 0 ? GStrandHairWidth * 0.5f : Data.HairGroups[GroupIndex].MaxStrandRadius;
-}
-
-#define OPTIMIZE_OFF 0
-
-#if OPTIMIZE_OFF
-	#pragma optimize("", off)
-#endif
-
 /////////////////////////////////////////////////////////////////////////////////////////
 
 template<typename T> inline void BindParam(FMeshDrawSingleShaderBindings& ShaderBindings, const FShaderResourceParameter& Param, T* Value)	{ if (Param.IsBound() && Value) ShaderBindings.Add(Param, Value); }
 template<typename T> inline void BindParam(FMeshDrawSingleShaderBindings& ShaderBindings, const FShaderParameter& Param, const T& Value)	{ if (Param.IsBound()) ShaderBindings.Add(Param, Value); }
 
+class FDummyCulledDispatchVertexIdsBuffer : public FVertexBuffer
+{
+public:
+	FShaderResourceViewRHIRef SRVUint;
+	FShaderResourceViewRHIRef SRVFloat;
+
+	virtual void InitRHI() override
+	{
+		FRHIResourceCreateInfo CreateInfo;
+		void* BufferData = nullptr;
+		uint32 NumBytes = sizeof(uint32) * 4;
+		VertexBufferRHI = RHICreateAndLockVertexBuffer(NumBytes, BUF_Static | BUF_ShaderResource, CreateInfo, BufferData);
+		uint32* DummyContents = (uint32*)BufferData;
+		DummyContents[0] = DummyContents[1] = DummyContents[2] = DummyContents[3] = 0;
+		RHIUnlockVertexBuffer(VertexBufferRHI);
+
+		SRVUint = RHICreateShaderResourceView(VertexBufferRHI, sizeof(uint32), PF_R32_UINT);
+		SRVFloat = RHICreateShaderResourceView(VertexBufferRHI, sizeof(uint32), PF_R32_FLOAT);
+	}
+
+};
+TGlobalResource<FDummyCulledDispatchVertexIdsBuffer> GDummyCulledDispatchVertexIdsBuffer;
+
 class FHairStrandsVertexFactoryShaderParameters : public FVertexFactoryShaderParameters
 {
 public:
@@ -42,6 +51,7 @@ public:
 	FShaderParameter PositionOffset;
 	FShaderParameter PreviousPositionOffset;
 	FShaderParameter Density;
+	FShaderParameter Culling;
 
 	FShaderResourceParameter PositionBuffer;
 	FShaderResourceParameter PreviousPositionBuffer;
@@ -49,6 +59,9 @@ public:
 	FShaderResourceParameter MaterialBuffer;
 	FShaderResourceParameter TangentBuffer;
 
+	FShaderResourceParameter CulledVertexIdsBuffer;
+	FShaderResourceParameter CulledVertexRadiusScaleBuffer;
+
 	virtual void Bind(const FShaderParameterMap& ParameterMap) override
 	{
 		Radius.Bind(ParameterMap, TEXT("HairStrandsVF_Radius"));
@@ -56,12 +69,16 @@ public:
 		PositionOffset.Bind(ParameterMap, TEXT("HairStrandsVF_PositionOffset"));
 		PreviousPositionOffset.Bind(ParameterMap, TEXT("HairStrandsVF_PreviousPositionOffset"));
 		Density.Bind(ParameterMap, TEXT("HairStrandsVF_Density"));	
+		Culling.Bind(ParameterMap, TEXT("HairStrandsVF_CullingEnable"));
 
 		PositionBuffer.Bind(ParameterMap, TEXT("HairStrandsVF_PositionBuffer"));
 		PreviousPositionBuffer.Bind(ParameterMap, TEXT("HairStrandsVF_PreviousPositionBuffer"));
 		AttributeBuffer.Bind(ParameterMap, TEXT("HairStrandsVF_AttributeBuffer"));
 		MaterialBuffer.Bind(ParameterMap, TEXT("HairStrandsVF_MaterialBuffer"));
 		TangentBuffer.Bind(ParameterMap, TEXT("HairStrandsVF_TangentBuffer"));
+
+		CulledVertexIdsBuffer.Bind(ParameterMap, TEXT("CulledVertexIdsBuffer"));
+		CulledVertexRadiusScaleBuffer.Bind(ParameterMap, TEXT("CulledVertexRadiusScaleBuffer"));
 	}
 
 	virtual void Serialize(FArchive& Ar) override
@@ -73,12 +90,16 @@ public:
 		Ar << PositionOffset;
 		Ar << PreviousPositionOffset;
 		Ar << Density;
+		Ar << Culling;
 
 		Ar << PositionBuffer;
 		Ar << PreviousPositionBuffer;
 		Ar << AttributeBuffer;
 		Ar << MaterialBuffer;
 		Ar << TangentBuffer;
+
+		Ar << CulledVertexIdsBuffer;
+		Ar << CulledVertexRadiusScaleBuffer;
 	}
 
 	virtual void GetElementShaderBindings(
@@ -95,7 +116,9 @@ public:
 	{
 		const FHairStrandsVertexFactory* VF = static_cast<const FHairStrandsVertexFactory*>(VertexFactory);
 
-		const uint64 GroupIndex = reinterpret_cast<uint64>(BatchElement.UserData);
+		const FHairGroupPublicData* GroupPublicData = reinterpret_cast<const FHairGroupPublicData*>(BatchElement.VertexFactoryUserData);
+		check(GroupPublicData);
+		const uint64 GroupIndex = GroupPublicData->GetGroupIndex();
 		BindParam(ShaderBindings, PositionBuffer, VF->GetPositionSRV(GroupIndex));
 		BindParam(ShaderBindings, PreviousPositionBuffer, VF->GetPreviousPositionSRV(GroupIndex));
 		BindParam(ShaderBindings, AttributeBuffer, VF->GetAttributeSRV(GroupIndex));
@@ -106,6 +129,22 @@ public:
 		BindParam(ShaderBindings, PositionOffset, VF->GetPositionOffset(GroupIndex));
 		BindParam(ShaderBindings, PreviousPositionOffset, VF->GetPreviousPositionOffset(GroupIndex));
 		BindParam(ShaderBindings, Density, VF->GetHairDensity(GroupIndex));
+		
+		FShaderResourceViewRHIRef CulledDispatchVertexIdsSRV = GDummyCulledDispatchVertexIdsBuffer.SRVUint;
+		FShaderResourceViewRHIRef CulledCompactedRadiusScaleBufferSRV = GDummyCulledDispatchVertexIdsBuffer.SRVFloat;
+
+		FHairGroupPublicData* PublicGroupData = VF->GetHairGroupPublicData(GroupIndex);
+		
+		const bool bCulling = PublicGroupData ? PublicGroupData->GetCullingResultAvailable() : false;
+		if (bCulling && PublicGroupData)
+		{
+			CulledDispatchVertexIdsSRV = PublicGroupData->GetCulledVertexIdBuffer().SRV;
+			CulledCompactedRadiusScaleBufferSRV = PublicGroupData->GetCulledVertexRadiusScaleBuffer().SRV;
+		}
+
+		BindParam(ShaderBindings, Culling, bCulling ? 1 : 0);
+		ShaderBindings.Add(CulledVertexIdsBuffer, CulledDispatchVertexIdsSRV);
+		ShaderBindings.Add(CulledVertexRadiusScaleBuffer, CulledCompactedRadiusScaleBufferSRV);
 	}
 };
 
@@ -213,9 +252,4 @@ void FHairStrandsVertexFactory::ReleaseRHI()
 	FVertexFactory::ReleaseRHI();
 }
 
-
 IMPLEMENT_VERTEX_FACTORY_TYPE_EX(FHairStrandsVertexFactory,"/Engine/Private/HairStrands/HairStrandsVertexFactory.ush",true,false,true,true,true,true,true);
-
-#if OPTIMIZE_OFF
-	#pragma optimize("", on)
-#endif
\ No newline at end of file
diff --git a/Engine/Plugins/Experimental/HairStrands/Source/HairStrandsCore/Public/GroomAsset.h b/Engine/Plugins/Experimental/HairStrands/Source/HairStrandsCore/Public/GroomAsset.h
--- a/Engine/Plugins/Experimental/HairStrands/Source/HairStrandsCore/Public/GroomAsset.h
+++ b/Engine/Plugins/Experimental/HairStrands/Source/HairStrandsCore/Public/GroomAsset.h
@@ -8,12 +8,15 @@
 #include "HairDescription.h"
 #include "HairStrandsDatas.h"
 #include "RenderResource.h"
+#include "GroomSettings.h"
+#include "HairStrandsInterface.h"
 #include "Interfaces/Interface_AssetUserData.h"
 
 #include "GroomAsset.generated.h"
 
 class UMaterialInterface;
 class UNiagaraSystem;
+struct FHairGroupData;
 
 /* Render buffers for root deformation for dynamic meshes */
 struct FHairStrandsRootResource : public FRenderResource
@@ -35,8 +38,7 @@ struct FHairStrandsRootResource : public FRenderResource
 
 	struct FMeshProjectionLOD
 	{
-		enum class EStatus { Invalid, Initialized, Completed };
-		EStatus Status = EStatus::Invalid;
+		FHairStrandsProjectionHairData::LODData::EStatus Status = FHairStrandsProjectionHairData::LODData::EStatus::Invalid;
 		int32 LODIndex = -1;
 
 		/* Triangle on which a root is attached */
@@ -134,6 +136,50 @@ struct FHairStrandsDeformedResource : public FRenderResource
 	const bool bInitializedData;
 };
 
+struct FHairStrandsClusterCullingResource : public FRenderResource
+{
+	FHairStrandsClusterCullingResource(const FHairGroupData& GroupData);
+
+	/* Init the buffer */
+	virtual void InitRHI() override;
+
+	/* Release the buffer */
+	virtual void ReleaseRHI() override;
+
+	/* Get the resource name */
+	virtual FString GetFriendlyName() const override { return TEXT("FHairStrandsClusterResource"); }
+
+
+	struct FClusterInfo
+	{
+		uint32 FirstVertexId;
+		uint32 VertexIdCountLodHigh;
+		uint32 VertexIdCountLodLow;
+		uint32 UnusedUint;
+	};
+	TArray<FClusterInfo> ClusterInfoArray;
+	TArray<uint32> VertexToClusterIdArray;
+	TArray<uint32> ClusterVertexIdArray;
+	TArray<float> ClusterIndexRadiusScaleInfoArray;
+
+	/* Cluster info buffer */
+	FReadBuffer ClusterInfoBuffer;		// Contains: Start VertexId, VertexId Count (for VertexToClusterIdBuffer)
+
+	/* VertexId => ClusterId to know which AABB to contribute to*/
+	FReadBuffer VertexToClusterIdBuffer;
+
+	/* Concatenated data for each cluster: list of VertexId pointed to by ClusterInfoBuffer */
+	FReadBuffer ClusterVertexIdBuffer;
+
+	/* Contains information to recove the radius scale to apply per cluster when decimating vertex count */
+	FReadBuffer ClusterIndexRadiusScaleInfoBuffer;
+
+	/* Number of cluster  */
+	uint32 ClusterCount;
+	/* Number of vertex  */
+	uint32 VertexCount;
+};
+
 struct FHairStrandsInterpolationResource : public FRenderResource
 {
 	/** Build the hair strands resource */
@@ -194,8 +240,8 @@ struct HAIRSTRANDSCORE_API FHairGroupInfo
 	UPROPERTY(VisibleAnywhere, Category = "Info", meta = (DisplayName = "Guide Count"))
 	int32 NumGuides = 0;
 
-	UPROPERTY(VisibleAnywhere, Category = "Info", meta = (DisplayName = "Auto-generated Guides", ToolTip = "Checked when guides were auto-generated"))
-	bool bIsAutoGenerated = false;
+	UPROPERTY(VisibleAnywhere, Category = "Info", meta = (DisplayName = "Import options", ToolTip = "Show the options used at import time"))
+	FGroomBuildSettings ImportSettings;
 
 	UPROPERTY(EditAnywhere, Category="Rendering")
 	UMaterialInterface* Material = nullptr;
@@ -203,6 +249,8 @@ struct HAIRSTRANDSCORE_API FHairGroupInfo
 	// Currently hide the Nigara simulation slot as it is not used, and could confuse users
 	//UPROPERTY(EditAnywhere, Category="Simulation", meta = (DisplayName = "Niagara System Asset"))
 	//UNiagaraSystem* NiagaraAsset = nullptr;
+
+	friend FArchive& operator<<(FArchive& Ar, FHairGroupInfo& GroupInfo);
 };
 
 struct HAIRSTRANDSCORE_API FHairGroupData
@@ -289,9 +337,8 @@ public:
 
 	int32 GetNumHairGroups() const;
 
-#if WITH_EDITOR
-	FOnGroomAssetChanged OnGroomAssetChanged;
-#endif
+	/** Returns true if the asset has the HairDescription needed to recompute its groom data */
+	bool CanRebuildFromDescription() const;
 
 	//~ Begin IInterface_AssetUserData Interface
 	virtual void AddAssetUserData(UAssetUserData* InUserData) override;
@@ -300,12 +347,30 @@ public:
 	virtual const TArray<UAssetUserData*>* GetAssetUserDataArray() const override;
 	//~ End IInterface_AssetUserData Interface
 
-//private : 
+//private :
+#if WITH_EDITOR
+	FOnGroomAssetChanged OnGroomAssetChanged;
+#endif
+
+#if WITH_EDITORONLY_DATA
+public:
+	/** Commits a HairDescription to buffer for serialization */
+	void CommitHairDescription(FHairDescription&& HairDescription);
+
+	/** Caches the computed groom data with the given build settings from/to the Derived Data Cache, building it if needed */
+	bool CacheDerivedData(const struct FGroomBuildSettings* BuildSettings = nullptr);
 
+private:
+	FString BuildDerivedDataKeySuffix(const struct FGroomBuildSettings& BuildSettings);
 
 	TUniquePtr<FHairDescription> HairDescription;
+	TUniquePtr<FHairDescriptionBulkData> HairDescriptionBulkData;
 
 	/** Array of user data stored with the asset */
 	UPROPERTY(EditAnywhere, AdvancedDisplay, Instanced, Category = HairLab)
 	TArray<UAssetUserData*> AssetUserData;
+
+	UPROPERTY()
+	bool bIsCacheable = true;
+#endif
 };
diff --git a/Engine/Plugins/Experimental/HairStrands/Source/HairStrandsCore/Public/GroomAssetImportData.h b/Engine/Plugins/Experimental/HairStrands/Source/HairStrandsCore/Public/GroomAssetImportData.h
new file mode 100644
--- /dev/null
+++ b/Engine/Plugins/Experimental/HairStrands/Source/HairStrandsCore/Public/GroomAssetImportData.h
@@ -0,0 +1,19 @@
+// Copyright Epic Games, Inc. All Rights Reserved.
+
+#pragma once
+
+#include "CoreMinimal.h"
+
+#include "EditorFramework/AssetImportData.h"
+#include "GroomAssetImportData.generated.h"
+
+UCLASS()
+class HAIRSTRANDSCORE_API UGroomAssetImportData : public UAssetImportData
+{
+	GENERATED_UCLASS_BODY()
+
+public:
+
+	UPROPERTY()
+	class UGroomImportOptions* ImportOptions;
+};
diff --git a/Engine/Plugins/Experimental/HairStrands/Source/HairStrandsCore/Public/GroomBuilder.h b/Engine/Plugins/Experimental/HairStrands/Source/HairStrandsCore/Public/GroomBuilder.h
--- a/Engine/Plugins/Experimental/HairStrands/Source/HairStrandsCore/Public/GroomBuilder.h
+++ b/Engine/Plugins/Experimental/HairStrands/Source/HairStrandsCore/Public/GroomBuilder.h
@@ -3,6 +3,9 @@
 #pragma once
 
 #include "CoreMinimal.h"
+#include "UObject/Object.h"
+#include "UObject/ObjectMacros.h"
+#include "GroomSettings.h"
 
 struct FHairStrandsDatas;
 class UGroomAsset;
@@ -11,5 +14,5 @@ struct HAIRSTRANDSCORE_API FGroomBuilder
 {
 	static bool BuildGroom(const class FHairDescription& HairDescription, const struct FGroomBuildSettings& BuildSettings, UGroomAsset* GroomAsset);
 	static void GenerateGuides(const FHairStrandsDatas& InData, float DecimationPercentage, FHairStrandsDatas& OutData);
-	static void BuildData(UGroomAsset* GroomAsset, uint8 QualityLevel, uint8 WeightMethod, bool bRandomize, bool bUnique);
+	static void BuildData(UGroomAsset* GroomAsset, const FGroomBuildSettings& BuildSettings);
 };
\ No newline at end of file
diff --git a/Engine/Plugins/Experimental/HairStrands/Source/HairStrandsCore/Public/GroomComponent.h b/Engine/Plugins/Experimental/HairStrands/Source/HairStrandsCore/Public/GroomComponent.h
--- a/Engine/Plugins/Experimental/HairStrands/Source/HairStrandsCore/Public/GroomComponent.h
+++ b/Engine/Plugins/Experimental/HairStrands/Source/HairStrandsCore/Public/GroomComponent.h
@@ -8,35 +8,10 @@
 #include "Components/MeshComponent.h"
 #include "GroomAsset.h"
 #include "RHIDefinitions.h"
+#include "GroomDesc.h"
 
 #include "GroomComponent.generated.h"
 
-USTRUCT(BlueprintType)
-struct FHairGroupDesc
-{
-	GENERATED_USTRUCT_BODY()
-
-	/** Number of hairs within this hair group.  */
-	UPROPERTY(VisibleAnywhere, BlueprintReadOnly, Category = "Groom")
-	int32 HairCount;
-
-	/** Number of simulation guides within this hair group. */
-	UPROPERTY(VisibleAnywhere, BlueprintReadOnly, Category = "Groom")
-	int32 GuideCount;
-
-	/** Override the hair width (in centimeters) */
-	UPROPERTY(EditAnywhere, BlueprintReadWrite, Category = "Groom", meta = (ClampMin = "0.0001", UIMin = "0.001", UIMax = "1.0", SliderExponent = 6))
-	float HairWidth;
-
-	/** Override the hair shadow density factor (unit less).  */
-	UPROPERTY(EditAnywhere, BlueprintReadWrite, Category = "Groom", meta = (ClampMin = "0.0001", UIMin = "0.001", UIMax = "10.0", SliderExponent = 6))
-	float HairShadowDensity;
-
-	/** Scale the hair geometry radius for ray tracing effects (e.g. shadow) */
-	UPROPERTY(EditAnywhere, BlueprintReadWrite, Category = "Groom", meta = (ClampMin = "0.0001", UIMin = "0.001", UIMax = "10.0", SliderExponent = 6))
-	float HairRaytracingRadiusScale;
-};
-
 UCLASS(HideCategories = (Object, Physics, Activation, Mobility, "Components|Activation"), editinlinenew, meta = (BlueprintSpawnableComponent), ClassGroup = Rendering)
 class HAIRSTRANDSCORE_API UGroomComponent : public UMeshComponent
 {
@@ -60,6 +35,9 @@ public:
 	/** Boolean to check when animation has been loaded */
 	bool bResetSimulation;
 
+	/** Previous bone matrix to compare the difference and decide to reset or not the simulation */
+	FMatrix	PrevBoneMatrix;
+
 	/** Listen for the animation event to trigger the sim */
 	UFUNCTION()
 	void ResetSimulation();
@@ -124,6 +102,9 @@ public:
 		FHairStrandsDeformedResource* RenderDeformedResources = nullptr;
 		FHairStrandsDeformedResource* SimDeformedResources = nullptr;
 
+		FHairStrandsClusterCullingResource* ClusterCullingResources = nullptr; // TODO merge into FHairGroupPublicData
+		FHairGroupPublicData* HairGroupPublicDatas = nullptr;
+
 		// Rest resources, owned by the asset
 		FHairStrandsRestResource* RenderRestResources = nullptr;
 		FHairStrandsRestResource* SimRestResources = nullptr;
@@ -146,30 +127,9 @@ protected:
 private:
 	void* InitializedResources;
 
-	enum class EMeshProjectionState
-	{
-		Invalid,
-		InProgressBinding,
-		WaitForRestPose,
-		Completed
-	};
 	class USkeletalMeshComponent* RegisteredSkeletalMeshComponent;
 	FVector SkeletalPreviousPositionOffset;
-	int32 MeshProjectionLODIndex;
-	uint32 MeshProjectionTickDelay;
-	EMeshProjectionState MeshProjectionState;
-	bool bIsGroomAssetCallbackRegistered;
-	
-	struct FSkeletalMeshConfiguration
-	{
-		int32 ForceLOD = -1;
-		bool ForceRefPose = false;
-		static bool Equals(const FSkeletalMeshConfiguration& A, const FSkeletalMeshConfiguration& B)
-		{
-			return A.ForceLOD == B.ForceLOD && A.ForceRefPose == B.ForceRefPose;
-		}
-	};
-	FSkeletalMeshConfiguration SkeletalMeshConfiguration;
+	bool bIsGroomAssetCallbackRegistered;	
 
 	void InitResources();
 	void ReleaseResources();
diff --git a/Engine/Plugins/Experimental/HairStrands/Source/HairStrandsCore/Public/GroomDesc.h b/Engine/Plugins/Experimental/HairStrands/Source/HairStrandsCore/Public/GroomDesc.h
new file mode 100644
--- /dev/null
+++ b/Engine/Plugins/Experimental/HairStrands/Source/HairStrandsCore/Public/GroomDesc.h
@@ -0,0 +1,74 @@
+// Copyright Epic Games, Inc. All Rights Reserved.
+
+#pragma once
+
+#include "CoreMinimal.h"
+#include "UObject/ObjectMacros.h"
+#include "Engine/EngineTypes.h"
+#include "GroomDesc.generated.h"
+
+USTRUCT(BlueprintType)
+struct FHairGroupDesc
+{
+	GENERATED_USTRUCT_BODY()
+
+	/** Number of hairs within this hair group. */
+	UPROPERTY(VisibleAnywhere, BlueprintReadOnly, Category = "Groom")
+	int32 HairCount;
+
+	/** Number of simulation guides within this hair group. */
+	UPROPERTY(VisibleAnywhere, BlueprintReadOnly, Category = "Groom")
+	int32 GuideCount;
+
+	/** Length of the longest hair strands */
+	UPROPERTY(VisibleAnywhere, BlueprintReadOnly, Category = "Groom")
+	float HairLength;
+
+	/** Hair width (in centimeters) */
+	UPROPERTY(EditAnywhere, BlueprintReadWrite, Category = "Groom", meta = (ClampMin = "0.0001", UIMin = "0.001", UIMax = "1.0", SliderExponent = 6))
+	float HairWidth;
+
+	/** Scale the hair width at the root */
+	UPROPERTY(EditAnywhere, BlueprintReadWrite, Category = "Groom", AdvancedDisplay, meta = (ClampMin = "0.0001", UIMin = "0.001", UIMax = "2.0", SliderExponent = 6))
+	float HairRootScale;
+
+	/** Scale the hair with at the tip */
+	UPROPERTY(EditAnywhere, BlueprintReadWrite, Category = "Groom", AdvancedDisplay, meta = (ClampMin = "0.0001", UIMin = "0.001", UIMax = "2.0", SliderExponent = 6))
+	float HairTipScale;
+
+	/** Normalized hair clip length, i.e. at which length hair will be clipped. 1 means no clipping. 0 means hairs are fully clipped */
+	UPROPERTY(EditAnywhere, BlueprintReadWrite, Category = "Groom", AdvancedDisplay, meta = (ClampMin = "0.0", UIMin = "0.0", UIMax = "1.0", SliderExponent = 1))
+	float HairClipLength;
+
+	/** Override the hair shadow density factor (unit less). */
+	UPROPERTY(EditAnywhere, BlueprintReadWrite, Category = "Groom", AdvancedDisplay, meta = (ClampMin = "0.0001", UIMin = "0.001", UIMax = "10.0", SliderExponent = 6))
+	float HairShadowDensity;
+
+	/** Scale the hair geometry radius for ray tracing effects (e.g. shadow) */
+	UPROPERTY(EditAnywhere, BlueprintReadWrite, Category = "Groom", AdvancedDisplay, meta = (ClampMin = "0.0001", UIMin = "0.001", UIMax = "10.0", SliderExponent = 6))
+	float HairRaytracingRadiusScale;
+
+	/** The targeted budget of hair vertex per pixel. Cluster strands will be decimated based on that. Used when r.HairStrands.Cluster.Culling = 1. */
+	UPROPERTY(EditAnywhere, BlueprintReadWrite, Category = "Performance", AdvancedDisplay, meta = (ClampMin = "0.0", UIMin = "0.0", UIMax = "6.0", SliderExponent = 1))
+	float LodAverageVertexPerPixel;
+	/** Bias the selected LOD. A value >0 will progressively select lower detailed lods. Used when r.HairStrands.Cluster.Culling = 1. */
+	UPROPERTY(EditAnywhere, AdvancedDisplay, BlueprintReadWrite, Category = "Performance", AdvancedDisplay, meta = (ClampMin = "-1.0", ClampMax = "1.0", UIMin = "-1.0", UIMax = "1.0", SliderExponent = 1))
+	float LodBias;
+
+	// Contructor used to initialised new values
+	FHairGroupDesc()
+	{
+		ReInit();
+	}
+
+	void ReInit()
+	{
+		HairRaytracingRadiusScale = 1.0f;
+		HairRootScale = 1.0f;
+		HairTipScale = 1.0f;
+		HairClipLength = 1.0f;
+
+		LodAverageVertexPerPixel = 3.0f; // Good for quality on heavy assets
+		LodBias = 0.0f;
+	}
+};
diff --git a/Engine/Plugins/Experimental/HairStrands/Source/HairStrandsCore/Public/GroomImportOptions.h b/Engine/Plugins/Experimental/HairStrands/Source/HairStrandsCore/Public/GroomImportOptions.h
new file mode 100644
--- /dev/null
+++ b/Engine/Plugins/Experimental/HairStrands/Source/HairStrandsCore/Public/GroomImportOptions.h
@@ -0,0 +1,21 @@
+// Copyright Epic Games, Inc. All Rights Reserved.
+
+#pragma once
+
+#include "CoreMinimal.h"
+#include "GroomSettings.h"
+#include "UObject/Object.h"
+#include "GroomImportOptions.generated.h"
+
+UCLASS(BlueprintType, config = EditorPerProjectUserSettings, HideCategories = ("Hidden"))
+class HAIRSTRANDSCORE_API UGroomImportOptions : public UObject
+{
+	GENERATED_UCLASS_BODY()
+
+public:
+	UPROPERTY(config, EditAnywhere, BlueprintReadWrite, meta = (ShowOnlyInnerProperties), Category = Conversion)
+	FGroomConversionSettings ConversionSettings;
+
+	UPROPERTY(config, EditAnywhere, BlueprintReadWrite, meta = (ShowOnlyInnerProperties), Category = BuildSettings)
+	FGroomBuildSettings BuildSettings;
+};
diff --git a/Engine/Plugins/Experimental/HairStrands/Source/HairStrandsCore/Public/GroomSettings.h b/Engine/Plugins/Experimental/HairStrands/Source/HairStrandsCore/Public/GroomSettings.h
--- a/Engine/Plugins/Experimental/HairStrands/Source/HairStrandsCore/Public/GroomSettings.h
+++ b/Engine/Plugins/Experimental/HairStrands/Source/HairStrandsCore/Public/GroomSettings.h
@@ -31,6 +31,7 @@ enum class EGroomInterpolationQuality : uint8
 	Low		UMETA(DisplayName = "Low",		ToolTip = "Build interpolation data based on nearst neighbor search. Low quality interpolation data, but fast to build (takes a few minutes)"),
 	Medium	UMETA(DisplayName = "Medium",	ToolTip = "Build interpolation data using curve shape matching search but within a limited spatial range. This is a tradeoff between Low and high quality in term of quality & build time (can takes several dozen of minutes)"),
 	High	UMETA(DisplayName = "High",		ToolTip = "Build interpolation data using curve shape matching search. This result in high quality interpolation data, but is relatively slow to build (can takes several dozen of minutes)"),
+	Unknown	UMETA(Hidden),
 };
 
 UENUM()
@@ -39,6 +40,7 @@ enum class EGroomInterpolationWeight : uint8
 	Parametric	UMETA(DisplayName = "Parametric", ToolTip = "Build interpolation data based on curve parametric distance"),
 	Root		UMETA(DisplayName = "Root", ToolTip = "Build interpolation data based on distance between guide's root and strands's root"),
 	Index		UMETA(DisplayName = "Index", ToolTip = "Build interpolation data based on guide and strands vertex indices"),
+	Unknown		UMETA(Hidden),
 };
 
 USTRUCT(BlueprintType)
@@ -60,7 +62,7 @@ struct HAIRSTRANDSCORE_API FGroomBuildSettings
 	bool bOverrideGuides;
 
 	/** Density factor for converting hair into guide curve if no guides are provided. */
-	UPROPERTY(EditAnywhere, BlueprintReadWrite, Category = "BuildSettings", meta = (ClampMin = "0.01", UIMin = "0.01", UIMax = "1.0"))
+	UPROPERTY(EditAnywhere, BlueprintReadWrite, Category = "BuildSettings", meta = (ClampMin = "0", UIMin = "0", UIMax = "1.0"))
 	float HairToGuideDensity;
 
 	/** Interpolation data quality. */
@@ -78,4 +80,6 @@ struct HAIRSTRANDSCORE_API FGroomBuildSettings
 	/** Force a hair strand to be affected by a unique guide. */
 	UPROPERTY(EditAnywhere, BlueprintReadWrite, Category = "BuildSettings")
 	bool bUseUniqueGuide;
+
+	friend FArchive& operator<<(FArchive& Ar, FGroomBuildSettings& GroupInfo);
 };
diff --git a/Engine/Plugins/Experimental/HairStrands/Source/HairStrandsCore/Public/HairAttributes.h b/Engine/Plugins/Experimental/HairStrands/Source/HairStrandsCore/Public/HairAttributes.h
--- a/Engine/Plugins/Experimental/HairStrands/Source/HairStrandsCore/Public/HairAttributes.h
+++ b/Engine/Plugins/Experimental/HairStrands/Source/HairStrandsCore/Public/HairAttributes.h
@@ -17,11 +17,13 @@ namespace HairAttribute
 	{
 		extern HAIRSTRANDSCORE_API const FName Color;			// FVector
 		extern HAIRSTRANDSCORE_API const FName GroupID;			// int
-		extern HAIRSTRANDSCORE_API const FName Guide;			// bool
+		extern HAIRSTRANDSCORE_API const FName Guide;			// int
 		extern HAIRSTRANDSCORE_API const FName ID;				// int
 		extern HAIRSTRANDSCORE_API const FName RootUV;			// FVector2D
 		extern HAIRSTRANDSCORE_API const FName VertexCount;		// int
 		extern HAIRSTRANDSCORE_API const FName Width;			// float
+		extern HAIRSTRANDSCORE_API const FName ClosestGuides;	// FVector
+		extern HAIRSTRANDSCORE_API const FName GuideWeights;	// FVector
 	}
 
 	namespace Groom
