Fix Hair CIS

diff --git a/Engine/Plugins/Experimental/HairStrands/Source/HairStrandsCore/Private/GroomAsset.cpp b/Engine/Plugins/Experimental/HairStrands/Source/HairStrandsCore/Private/GroomAsset.cpp
--- a/Engine/Plugins/Experimental/HairStrands/Source/HairStrandsCore/Private/GroomAsset.cpp
+++ b/Engine/Plugins/Experimental/HairStrands/Source/HairStrandsCore/Private/GroomAsset.cpp
@@ -976,5 +976,3 @@ bool UGroomAsset::CacheDerivedData(const FGroomBuildSettings* InBuildSettings)
 }
 
 #endif // WITH_EDITORONLY_DATA
-
-#undef LOCTEXT_NAMESPACE
\ No newline at end of file
diff --git a/Engine/Plugins/Experimental/HairStrands/Source/HairStrandsCore/Public/GroomAsset.h b/Engine/Plugins/Experimental/HairStrands/Source/HairStrandsCore/Public/GroomAsset.h
--- a/Engine/Plugins/Experimental/HairStrands/Source/HairStrandsCore/Public/GroomAsset.h
+++ b/Engine/Plugins/Experimental/HairStrands/Source/HairStrandsCore/Public/GroomAsset.h
@@ -14,6 +14,7 @@
 
 #include "GroomAsset.generated.h"
 
+class UAssetUserData;
 class UMaterialInterface;
 class UNiagaraSystem;
 struct FHairGroupData;
