- Fix bad mege where some code was moved from header into source

#rnx

diff --git a/Engine/Plugins/Experimental/HairStrands/Source/HairStrandsNiagara/Private/NiagaraDataInterfacePressureGrid.cpp b/Engine/Plugins/Experimental/HairStrands/Source/HairStrandsNiagara/Private/NiagaraDataInterfacePressureGrid.cpp
--- a/Engine/Plugins/Experimental/HairStrands/Source/HairStrandsNiagara/Private/NiagaraDataInterfacePressureGrid.cpp
+++ b/Engine/Plugins/Experimental/HairStrands/Source/HairStrandsNiagara/Private/NiagaraDataInterfacePressureGrid.cpp
@@ -974,16 +974,6 @@ inline void ClearBuffer(FNDIPressureGridBuffer* CurrentGridBuffer, FNDIPressureG
 
 //------------------------------------------------------------------------------------------------------------
 
-void FNDIPressureGridProxy::DeferredDestroy()
-{
-	for (const FNiagaraSystemInstanceID& Sys : DeferredDestroyList)
-	{
-		SystemInstancesToProxyData.Remove(Sys);
-	}
-
-	DeferredDestroyList.Empty();
-}
-
 void FNDIPressureGridProxy::ConsumePerInstanceDataFromGameThread(void* PerInstanceData, const FNiagaraSystemInstanceID& Instance)
 {
 	FNDIPressureGridData* SourceData = static_cast<FNDIPressureGridData*>(PerInstanceData);
@@ -1010,22 +1000,13 @@ void FNDIPressureGridProxy::InitializePerInstanceData(const FNiagaraSystemInstan
 	check(IsInRenderingThread());
 
 	FNDIPressureGridData* TargetData = SystemInstancesToProxyData.Find(SystemInstance);
-	if (TargetData != nullptr)
-	{
-		DeferredDestroyList.Remove(SystemInstance);
-	}
-	else
-	{
-		TargetData = &SystemInstancesToProxyData.Add(SystemInstance);
-	}
+	TargetData = &SystemInstancesToProxyData.Add(SystemInstance);
 }
 
 void FNDIPressureGridProxy::DestroyPerInstanceData(NiagaraEmitterInstanceBatcher* Batcher, const FNiagaraSystemInstanceID& SystemInstance)
 {
 	check(IsInRenderingThread());
-
-	DeferredDestroyList.Add(SystemInstance);
-	Batcher->EnqueueDeferredDeletesForDI_RenderThread(this->AsShared());
+	SystemInstancesToProxyData.Remove(SystemInstance);
 }
 
 void FNDIPressureGridProxy::PreStage(FRHICommandList& RHICmdList, const FNiagaraDataInterfaceSetArgs& Context)
