Copying //UE4/Dev-Physics to Main (//UE4/Main)

diff --git a/Engine/Plugins/Experimental/HairStrands/Shaders/Private/NiagaraAngularSpringMaterial.ush b/Engine/Plugins/Experimental/HairStrands/Shaders/Private/NiagaraAngularSpringMaterial.ush
--- a/Engine/Plugins/Experimental/HairStrands/Shaders/Private/NiagaraAngularSpringMaterial.ush
+++ b/Engine/Plugins/Experimental/HairStrands/Shaders/Private/NiagaraAngularSpringMaterial.ush
@@ -2,273 +2,6 @@
 
 #pragma once
 
-/* -----------------------------------------------------------------
- * Stretch spring material
- * -----------------------------------------------------------------
- */
-
- void ResetStretchSpringMaterial(in int StrandsSize, in float YoungModulus, in float RodThickness, 
-	in float RestLength, in float DeltaTime, in bool ProjectConstraint, in float MaterialDamping, out float OutMaterialCompliance, out float OutMaterialWeight, out float OutMaterialMultiplier, in int NodeOffset)
-{
-	// Compliance = 1.0 / (k * dt * dt)
-	// with k = L * L * (Y * A / L) (the L*L term before is coming from the fact that our constraint is dL/L and not dL)
-	// A is the cross section area (Pi * R * R), L is the rest length and Y is the young modulus
-	OutMaterialCompliance = 4.0/(YoungModulus*PI*RodThickness*RestLength*RodThickness*DeltaTime*DeltaTime);
-
-	const float SumInverseMass = !ProjectConstraint ? ( SharedInverseMass[GGroupThreadId.x] + SharedInverseMass[GGroupThreadId.x-1-NodeOffset] ) / (RestLength*RestLength) :
-				SharedInverseMass[GGroupThreadId.x] / (RestLength*RestLength);
-
-	const float SchurDiagonal = ( (1.0  + MaterialDamping / DeltaTime ) * SumInverseMass + OutMaterialCompliance );
-	OutMaterialWeight = ( SchurDiagonal != 0.0 ) ? 1.0 / SchurDiagonal : 0.0;
-
-	OutMaterialMultiplier = 0;
-}
-
-
-void UpdateStretchSpringMultiplier(in float RestLength, in float DeltaTime, in bool ProjectConstraint, in float MaterialDamping, in float MaterialCompliance, in float MaterialWeight, inout float OutMaterialMultiplier, in int NodeOffset)
-{
-	const int PreviousIndex = GGroupThreadId.x-1-NodeOffset;
-	float3 EdgeDirection = SharedNodePosition[GGroupThreadId.x] - SharedNodePosition[PreviousIndex];
-	const float3 DeltaVelocity = (EdgeDirection - ( SharedPreviousPosition[GGroupThreadId.x] - SharedPreviousPosition[PreviousIndex] ) ) / DeltaTime; 
-
-	const float EdgeLength = length(EdgeDirection);
-	EdgeDirection /= EdgeLength;
-
-	// XPBD lagrange multiplier update : dL = -(C+compliance*L) / (dC * invM * dCt + alpha)
-	const float DeltaLambda = -((EdgeLength-RestLength)/RestLength + OutMaterialMultiplier * MaterialCompliance + MaterialDamping * dot(EdgeDirection,DeltaVelocity)/RestLength ) * MaterialWeight;
-
-	// L += dL
-	OutMaterialMultiplier += DeltaLambda;
-
-	// XPBD position update : dX = dL * dCt * invM
-	const float3 PositionDelta = EdgeDirection * DeltaLambda / RestLength;
-
-	// dX += dX
-	SharedNodePosition[GGroupThreadId.x] += PositionDelta * SharedInverseMass[GGroupThreadId.x];
-	if(!ProjectConstraint)
-	{
-		SharedNodePosition[PreviousIndex] -= PositionDelta * SharedInverseMass[PreviousIndex];
-	}
-}
-
-void SetupStretchSpringMaterial(in int StrandsSize, in float YoungModulus, in float RodThickness, 
-	in float RestLength, in float DeltaTime, in bool ProjectConstraint, in float MaterialDamping, out float OutMaterialCompliance, out float OutMaterialWeight, out float OutMaterialMultiplier)
-{
-	OutMaterialCompliance = 0.0;
-	OutMaterialWeight = 0.0;
-	OutMaterialMultiplier = 0.0;
-	const int LocalIndex = (GGroupThreadId.x % StrandsSize);
-	if( LocalIndex > 1 )
-	{
-		ResetStretchSpringMaterial(StrandsSize,YoungModulus,RodThickness,RestLength,DeltaTime,ProjectConstraint,MaterialDamping,OutMaterialCompliance,OutMaterialWeight,OutMaterialMultiplier,0);
-	}
-}
-
-void SolveStretchSpringMaterial(in int StrandsSize, in float RestLength, in float DeltaTime, in float MaterialDamping, 
-			in float MaterialCompliance, in float MaterialWeight, in float MaterialMultiplier, out float OutMaterialMultiplier)
-{
-	OutMaterialMultiplier = MaterialMultiplier;
-	const int LocalIndex = (GGroupThreadId.x % StrandsSize);
-	if( LocalIndex > 1)
-	{
-		const int IsRed = (GGroupThreadId.x % 2) == 0;
-		// Process all the red rods
-		if (IsRed)
-		{	
-			UpdateStretchSpringMultiplier(RestLength,DeltaTime,false,MaterialDamping,MaterialCompliance,MaterialWeight,OutMaterialMultiplier,0);
-		}
-		// Process all the black rods
-		GroupMemoryBarrier(); 
-		if (!IsRed)
-		{
-			UpdateStretchSpringMultiplier(RestLength,DeltaTime,false,MaterialDamping,MaterialCompliance,MaterialWeight,OutMaterialMultiplier,0);
-		}
-		GroupMemoryBarrier(); 
-	}
-} 
-
-void ProjectStretchSpringMaterial(in int StrandsSize, in float YoungModulus, in float RodThickness, in float RestLength, in float NodeMass, in float DeltaTime, in float3 NodePosition, out float3 OutNodePosition)
-{
-	const int LocalIndex = (GGroupThreadId.x % StrandsSize);
-
-	SharedNodePosition[GGroupThreadId.x] = NodePosition;
-	SharedPreviousPosition[GGroupThreadId.x] = NodePosition;
-	SharedInverseMass[GGroupThreadId.x] = (LocalIndex > 1 && NodeMass != 0.0) ? 1.0 / NodeMass : 0.0;
-	GroupMemoryBarrier();
-
-	if(DeltaTime != 0.0)
-	{
-		float MaterialCompliance = 0.0;
-		float MaterialWeight = 0.0;
-		float MaterialMultiplier = 0.0;
-
-		if( LocalIndex > 1 )
-		{
-			ResetStretchSpringMaterial(StrandsSize,YoungModulus,RodThickness,RestLength,DeltaTime,true,0.0,MaterialCompliance,MaterialWeight,MaterialMultiplier,0);
-
-			for(int i = 2; i < StrandsSize; ++i)
-			{
-				if( LocalIndex == i )
-				{
-					UpdateStretchSpringMultiplier(RestLength,DeltaTime,true,0.0,MaterialCompliance,MaterialWeight,MaterialMultiplier,0);
-				}
-				GroupMemoryBarrier(); 
-			}
-		}
-		
-	}
-	GroupMemoryBarrier();
-	OutNodePosition = SharedNodePosition[GGroupThreadId.x] ;
-}
-
-/* -----------------------------------------------------------------
- * Bend spring material
- * -----------------------------------------------------------------
- */
-
- void SetupBendSpringMaterial(in int StrandsSize, in float YoungModulus, in float RodThickness, 
-	in float RestLength, in float DeltaTime, in bool ProjectConstraint, in float MaterialDamping, out float OutMaterialCompliance, out float OutMaterialWeight, out float OutMaterialMultiplier)
-{
-	OutMaterialCompliance = 0.0;
-	OutMaterialWeight = 0.0;
-	OutMaterialMultiplier = 0.0;
-	const int LocalIndex = (GGroupThreadId.x % StrandsSize);
-	if( LocalIndex > 1 )
-	{
-		ResetStretchSpringMaterial(StrandsSize,YoungModulus,RodThickness,RestLength,DeltaTime,ProjectConstraint,MaterialDamping,OutMaterialCompliance,OutMaterialWeight,OutMaterialMultiplier,1);
-	}
-}
-
-void SolveBendSpringMaterial(in int StrandsSize, in float RestLength, in float DeltaTime, in float MaterialDamping, 
-			in float MaterialCompliance, in float MaterialWeight, in float MaterialMultiplier, out float OutMaterialMultiplier)
-{
-	OutMaterialMultiplier = MaterialMultiplier;
-	const int LocalIndex = (GGroupThreadId.x % StrandsSize);
-	if( LocalIndex > 1)
-	{
-		const int IsRed = ((GGroupThreadId.x/2) % 2) == 0;
-		// Process all the red rods
-		if (IsRed)
-		{	
-			UpdateStretchSpringMultiplier(RestLength,DeltaTime,false,MaterialDamping,MaterialCompliance,MaterialWeight,OutMaterialMultiplier,1);
-		}
-		// Process all the black rods
-		GroupMemoryBarrier(); 
-		if (!IsRed)
-		{
-			UpdateStretchSpringMultiplier(RestLength,DeltaTime,false,MaterialDamping,MaterialCompliance,MaterialWeight,OutMaterialMultiplier,1);
-		}
-		GroupMemoryBarrier(); 
-	}
-} 
-
-void ProjectBendSpringMaterial(in int StrandsSize, in float YoungModulus, in float RodThickness, in float RestLength, in float NodeMass, in float DeltaTime, in float3 NodePosition, out float3 OutNodePosition)
-{
-	const int LocalIndex = (GGroupThreadId.x % StrandsSize);
-
-	SharedNodePosition[GGroupThreadId.x] = NodePosition;
-	SharedPreviousPosition[GGroupThreadId.x] = NodePosition;
-	SharedInverseMass[GGroupThreadId.x] = (LocalIndex > 1 && NodeMass != 0.0) ? 1.0 / NodeMass : 0.0;
-	GroupMemoryBarrier();
-
-	if(DeltaTime != 0.0)
-	{
-		float MaterialCompliance = 0.0;
-		float MaterialWeight = 0.0;
-		float MaterialMultiplier = 0.0;
-
-		if( LocalIndex > 1 )
-		{
-			ResetStretchSpringMaterial(StrandsSize,YoungModulus,RodThickness,RestLength,DeltaTime,true,0.0,MaterialCompliance,MaterialWeight,MaterialMultiplier,1);
-
-			for(int i = 2; i < StrandsSize; ++i)
-			{
-				if( LocalIndex == i )
-				{
-					UpdateStretchSpringMultiplier(RestLength,DeltaTime,true,0.0,MaterialCompliance,MaterialWeight,MaterialMultiplier,1);
-				}
-				GroupMemoryBarrier(); 
-			}
-		}
-	}
-	GroupMemoryBarrier();
-	OutNodePosition = SharedNodePosition[GGroupThreadId.x];
-}
-
-/* -----------------------------------------------------------------
- * Twist spring material
- * -----------------------------------------------------------------
- */
-
- void SetupTwistSpringMaterial(in int StrandsSize, in float YoungModulus, in float RodThickness, 
-	in float RestLength, in float DeltaTime, in bool ProjectConstraint, in float MaterialDamping, out float OutMaterialCompliance, out float OutMaterialWeight, out float OutMaterialMultiplier)
-{
-	OutMaterialCompliance = 0.0;
-	OutMaterialWeight = 0.0;
-	OutMaterialMultiplier = 0.0;
-	const int LocalIndex = (GGroupThreadId.x % StrandsSize);
-	if( LocalIndex > 2 )
-	{
-		ResetStretchSpringMaterial(StrandsSize,YoungModulus,RodThickness,RestLength,DeltaTime,ProjectConstraint,MaterialDamping,OutMaterialCompliance,OutMaterialWeight,OutMaterialMultiplier,2);
-	}
-}
-
-void SolveTwistSpringMaterial(in int StrandsSize, in float RestLength, in float DeltaTime, in float MaterialDamping, 
-			in float MaterialCompliance, in float MaterialWeight, in float MaterialMultiplier, out float OutMaterialMultiplier)
-{
-	OutMaterialMultiplier = MaterialMultiplier;
-	const int LocalIndex = (GGroupThreadId.x % StrandsSize);
-	if( LocalIndex > 2)
-	{
-		const int IsRed = ((GGroupThreadId.x/3) % 2) == 0;
-		// Process all the red rods
-		if (IsRed)
-		{	
-			UpdateStretchSpringMultiplier(RestLength,DeltaTime,false,MaterialDamping,MaterialCompliance,MaterialWeight,OutMaterialMultiplier,2);
-		}
-		// Process all the black rods
-		GroupMemoryBarrier(); 
-		if (!IsRed)
-		{
-			UpdateStretchSpringMultiplier(RestLength,DeltaTime,false,MaterialDamping,MaterialCompliance,MaterialWeight,OutMaterialMultiplier,2);
-		}
-		GroupMemoryBarrier(); 
-	}
-} 
-
-void ProjectTwistSpringMaterial(in int StrandsSize, in float YoungModulus, in float RodThickness, in float RestLength, in float NodeMass, in float DeltaTime, in float3 NodePosition, out float3 OutNodePosition)
-{
-	const int LocalIndex = (GGroupThreadId.x % StrandsSize);
-
-	SharedNodePosition[GGroupThreadId.x] = NodePosition;
-	SharedPreviousPosition[GGroupThreadId.x] = NodePosition;
-	SharedInverseMass[GGroupThreadId.x] = (LocalIndex > 1 && NodeMass != 0.0) ? 1.0 / NodeMass : 0.0;
-	GroupMemoryBarrier();
-
-	if(DeltaTime != 0.0)
-	{
-		float MaterialCompliance = 0.0;
-		float MaterialWeight = 0.0;
-		float MaterialMultiplier = 0.0;
-
-		if( LocalIndex > 2 )
-		{
-			ResetStretchSpringMaterial(StrandsSize,YoungModulus,RodThickness,RestLength,DeltaTime,true,0.0,MaterialCompliance,MaterialWeight,MaterialMultiplier,2);
-
-			for(int i = 3; i < StrandsSize; ++i)
-			{
-				if( LocalIndex == i )
-				{
-					UpdateStretchSpringMultiplier(RestLength,DeltaTime,true,0.0,MaterialCompliance,MaterialWeight,MaterialMultiplier,2);
-				}
-				GroupMemoryBarrier(); 
-			}
-		}
-	}
-	GroupMemoryBarrier();
-	OutNodePosition = SharedNodePosition[GGroupThreadId.x];
-}
-
 /* -----------------------------------------------------------------
  * Compute Bishop frame
  * -----------------------------------------------------------------
@@ -300,26 +33,21 @@ void ComputeMaterialFrame(in int StrandsSize)
 void UpdateMaterialFrame(in int StrandsSize)
 {
 	const int LocalIndex = (GGroupThreadId.x % StrandsSize);
-	if( LocalIndex == 0 )
+	if( (LocalIndex > 0) && (LocalIndex < StrandsSize-1) )
 	{
-		float4 NodeQuaternion = SharedPreviousOrientation[GGroupThreadId.x];
-		float3 TangentPrev = normalize(SharedPreviousPosition[GGroupThreadId.x+1] - SharedPreviousPosition[GGroupThreadId.x]);
-		float3 TangentNext = TangentPrev;
-
-		for (int EdgeIndex = 1, EdgeEnd = StrandsSize-1; EdgeIndex < EdgeEnd; ++EdgeIndex)
-		{
-			TangentPrev = TangentNext;
-			TangentNext = normalize(SharedPreviousPosition[GGroupThreadId.x+EdgeIndex+1] - SharedPreviousPosition[GGroupThreadId.x+EdgeIndex]);
+		float4 NodeQuaternion = SharedNodeOrientation[GGroupThreadId.x];
+		float3 TangentPrev = RotateVectorByQuat( float3(0,0,1), NodeQuaternion);
+		float3 TangentNext = normalize(SharedNodePosition[GGroupThreadId.x+1] - SharedNodePosition[GGroupThreadId.x]);
 
-			NodeQuaternion = NormalizeQuat( MultiplyQuat( FindQuatBetweenNormals(TangentPrev,TangentNext), NodeQuaternion ) );
-			SharedNodeOrientation[GGroupThreadId.x+EdgeIndex] = NodeQuaternion;
-			GroupMemoryBarrier();
-		}
-		GroupMemoryBarrier();
+		SharedNodeOrientation[GGroupThreadId.x] = NormalizeQuat( MultiplyQuat( FindQuatBetweenNormals(TangentPrev,TangentNext), NodeQuaternion) );
 	}
 	GroupMemoryBarrier();
 }
 
+/* -----------------------------------------------------------------
+ * Alternative version
+ * -----------------------------------------------------------------
+ */
 
 /*void UpdateMaterialFrame(in int StrandsSize)
 {
@@ -339,7 +67,7 @@ void UpdateMaterialFrame(in int StrandsSize)
 }*/
 
 /* -----------------------------------------------------------------
- * Bend spring material
+ * Angular spring material
  * -----------------------------------------------------------------
  */
 
@@ -362,7 +90,7 @@ void UpdateAngularSpringMultiplier( in float RestLength, in float3 RestDirection
 	const float4 ParentFrame = SlerpQuat(SharedNodeOrientation[GGroupThreadId.x-2],SharedNodeOrientation[GGroupThreadId.x-2],0.5);
 	const float3 TargetDirection = RotateVectorByQuat(RestDirection,ParentFrame);
 
-	// XPBD lagrange multiplier update : dL = -(C+compliance*L) / (dC * invM * dCt + alpha)
+	// XPBD lagrang : dL = -(C+compliance*L) / (dC * invM * dCt + alpha)
 	const float3 DeltaLambda = -((EdgeDirection - TargetDirection)/RestLength + 
 					OutMaterialMultiplier * InMaterialCompliance  + MaterialDamping * DeltaVelocity /RestLength) * InMaterialWeight;
 
@@ -380,125 +108,44 @@ void UpdateAngularSpringMultiplier( in float RestLength, in float3 RestDirection
 	}
 }
 
-void SolveAngularSpringMaterial(in int StrandsSize, in float RestLength, in float3 RestDirection, in float DeltaTime, in float MaterialDamping,
+void SolveAngularSpringMaterial(in bool EnableConstraint, in int StrandsSize, in float RestLength, in float3 RestDirection, in float DeltaTime, in float MaterialDamping,
 			in float MaterialCompliance, in float MaterialWeight, in float3 MaterialMultiplier, out float3 OutMaterialMultiplier)
 {
-	
 	OutMaterialMultiplier = MaterialMultiplier;
-	const int LocalIndex = (GGroupThreadId.x % StrandsSize);
-	if( LocalIndex > 1)
-	{
-		const int IsRed = (GGroupThreadId.x % 2) == 0;
-		// Process all the red rods
-		if (!IsRed)
-		{	
-			UpdateAngularSpringMultiplier(RestLength,RestDirection,DeltaTime,false,MaterialDamping,MaterialCompliance,MaterialWeight,OutMaterialMultiplier);
-		}
-		GroupMemoryBarrier(); 
-
-		if (IsRed)
-		{
-			UpdateAngularSpringMultiplier(RestLength,RestDirection,DeltaTime,false,MaterialDamping,MaterialCompliance,MaterialWeight,OutMaterialMultiplier);
-		}
-		GroupMemoryBarrier(); 
-	}
-}
-/*
-groupshared float4 SharedRestDatas[THREADGROUP_SIZE];
-void SolveAngularSpringMaterial(in int StrandsSize, in float RestLength, in float3 RestDirection, in float DeltaTime, in float MaterialDamping,
-			in float MaterialCompliance, in float MaterialWeight, in float3 MaterialMultiplier, out float3 OutMaterialMultiplier)
-{
-	SharedRestDatas[GGroupThreadId.x] = float4(RestDirection,RestLength);
-	GroupMemoryBarrier(); 
 
-	OutMaterialMultiplier = MaterialMultiplier;
-	const int LocalIndex = (GGroupThreadId.x % StrandsSize);
-	if( LocalIndex == 0)
+	if(EnableConstraint)
 	{
-		for(int i = GGroupThreadId.x+2; i < GGroupThreadId.x+StrandsSize; ++i)
+		const int LocalIndex = (GGroupThreadId.x % StrandsSize);
+		if( LocalIndex > 1)
 		{
-			const float3 EdgeDirection = SharedNodePosition[i] - SharedNodePosition[i-1];
-
-			const float4 ParentFrame = SharedNodeOrientation[i-2];
-			const float3 TargetDirection = RotateVectorByQuat(SharedRestDatas[i].xyz,ParentFrame);
-
-			// XPBD lagrange multiplier update : dL = -(C+compliance*L) / (dC * invM * dCt + alpha)
-			const float3 DeltaLambda = -(EdgeDirection - TargetDirection) / (SharedInverseMass[i]+SharedInverseMass[i-1]);
-
-			// XPBD position update : dX = dL * dCt * invM
-			const float3 PositionDelta  = DeltaLambda;
-
-			SharedDeltaPosition[i] = PositionDelta;
-			GroupMemoryBarrier();
-
-			// dX += dX
-			SharedNodePosition[i] += PositionDelta * SharedInverseMass[i];
-			GroupMemoryBarrier();
-			SharedNodePosition[i-1] -= PositionDelta * SharedInverseMass[i-1];
-			GroupMemoryBarrier();
-
-			if( i > GGroupThreadId.x+2 )
-			{
-				const float3 TangentInit = normalize(SharedNodePosition[i-2] - SharedNodePosition[i-3]);
-				const float3 TangentPrev = normalize(SharedNodePosition[i-1] - SharedNodePosition[i-2]);
-				const float3 TangentNext = normalize(SharedNodePosition[i] - SharedNodePosition[i-1]);
-
-				SharedNodeOrientation[i-2] = NormalizeQuat( MultiplyQuat( FindQuatBetweenNormals(TangentInit,TangentPrev), SharedNodeOrientation[i-3] ) );
-				GroupMemoryBarrier();
-				SharedNodeOrientation[i-1] = NormalizeQuat( MultiplyQuat( FindQuatBetweenNormals(TangentPrev,TangentNext), SharedNodeOrientation[i-2] ) );
-				GroupMemoryBarrier();
+			const int IsRed = (GGroupThreadId.x % 2) == 0;
+			// Process all the red rods
+			if (!IsRed)
+			{	
+				UpdateAngularSpringMultiplier(RestLength,RestDirection,DeltaTime,false,MaterialDamping,MaterialCompliance,MaterialWeight,OutMaterialMultiplier);
 			}
-			else
-			{
-				const float3 TangentPrev = normalize(SharedNodePosition[i-1] - SharedNodePosition[i-2]);
-				const float3 TangentNext = normalize(SharedNodePosition[i] - SharedNodePosition[i-1]);
+			GroupMemoryBarrier(); 
 
-				SharedNodeOrientation[i-1] = NormalizeQuat( MultiplyQuat( FindQuatBetweenNormals(TangentPrev,TangentNext), SharedNodeOrientation[i-2] ) );
-				GroupMemoryBarrier();
+			if (IsRed)
+			{
+				UpdateAngularSpringMultiplier(RestLength,RestDirection,DeltaTime,false,MaterialDamping,MaterialCompliance,MaterialWeight,OutMaterialMultiplier);
 			}
-			GroupMemoryBarrier();
-		}
-		GroupMemoryBarrier();
-		for(int i = GGroupThreadId.x+2; i < GGroupThreadId.x+StrandsSize; ++i)
-		{
-			SharedNodePosition[i] += SharedDeltaPosition[i] * SharedInverseMass[i];
-			SharedNodePosition[i-1] -= SharedDeltaPosition[i] * SharedInverseMass[i-1];
-			GroupMemoryBarrier();
+			GroupMemoryBarrier(); 
 		}
-		GroupMemoryBarrier();
-		for(int i = GGroupThreadId.x+2; i < GGroupThreadId.x+StrandsSize; ++i)
-		{
-			const float3 TangentPrev = normalize(SharedNodePosition[i-1] - SharedNodePosition[i-2]);
-			const float3 TangentNext = normalize(SharedNodePosition[i] - SharedNodePosition[i-1]);
-
-			SharedNodeOrientation[i-1] = NormalizeQuat( MultiplyQuat( FindQuatBetweenNormals(TangentPrev,TangentNext), SharedNodeOrientation[i-2] ) );
-			GroupMemoryBarrier();
-		}
-		GroupMemoryBarrier();
-	} 
-}*/
+	}
+}
 
-void ProjectAngularSpringMaterial(in int StrandsSize, in float YoungModulus, in float RodThickness, in float RestLength, in float3 RestDirection, in float NodeMass, in float DeltaTime,
-			 in float3 NodePosition, out float3 OutNodePosition)
+void ProjectAngularSpringMaterial(in bool EnableConstraint, in int StrandsSize, in float YoungModulus, in float RodThickness, in float RestLength, in float3 RestDirection, in float DeltaTime, out float3 OutNodePosition)
 {
 	const int LocalIndex = (GGroupThreadId.x % StrandsSize);
-
-	SharedNodePosition[GGroupThreadId.x] = NodePosition;
-	SharedPreviousPosition[GGroupThreadId.x] = NodePosition;
-	SharedInverseMass[GGroupThreadId.x] = (LocalIndex > 1 && NodeMass != 0.0) ? 1.0 / NodeMass : 0.0;
-	GroupMemoryBarrier();
-
-	if(DeltaTime != 0.0)
+	if(DeltaTime != 0.0 && EnableConstraint)
 	{
-
 		float MaterialCompliance = 0.0;
 		float MaterialWeight = 0.0;
 		float3 MaterialMultiplier = float3(0,0,0);
 
-		//ComputeMaterialFrame(StrandsSize);
 		SetupAngularSpringMaterial(StrandsSize, YoungModulus, RodThickness, RestLength, DeltaTime, true, 0.0, MaterialCompliance, MaterialWeight, MaterialMultiplier);
 
-		//if( LocalIndex > 1 )
 		{
 			for(int i = 2; i < StrandsSize; ++i)
 			{
@@ -517,7 +164,7 @@ void ProjectAngularSpringMaterial(in int StrandsSize, in float YoungModulus, in
 }
 
 /* -----------------------------------------------------------------
- * Compute Rest Direction and length
+ * Compute Rest Direction 
  * -----------------------------------------------------------------
  */
 
diff --git a/Engine/Plugins/Experimental/HairStrands/Shaders/Private/NiagaraCosseratRodMaterial.ush b/Engine/Plugins/Experimental/HairStrands/Shaders/Private/NiagaraCosseratRodMaterial.ush
--- a/Engine/Plugins/Experimental/HairStrands/Shaders/Private/NiagaraCosseratRodMaterial.ush
+++ b/Engine/Plugins/Experimental/HairStrands/Shaders/Private/NiagaraCosseratRodMaterial.ush
@@ -3,11 +3,11 @@
 #pragma once
 
 /* -----------------------------------------------------------------
- * Stretch rod material
+ * Align rod material
  * -----------------------------------------------------------------
  */
 
-void SetupStretchRodMaterial(in int StrandsSize, in float YoungModulus, in float RodThickness, 
+ void SetupAlignRodMaterial(in int StrandsSize, in float YoungModulus, in float RodThickness, 
 	in float RestLength, in float DeltaTime, in bool ProjectConstraint, in float MaterialDamping, out float OutMaterialCompliance, out float OutMaterialWeight, out float3 OutMaterialMultiplier)
 {
 	const int LocalIndex = (GGroupThreadId.x % StrandsSize);
@@ -18,10 +18,7 @@ void SetupStretchRodMaterial(in int StrandsSize, in float YoungModulus, in float
 		// A is the cross section area (Pi * R * R), L is the rest length and Y is the young modulus
 		OutMaterialCompliance = 4.0/(YoungModulus*PI*RodThickness*RestLength*RodThickness*DeltaTime*DeltaTime);
 
-		const float SumInverseMass = !ProjectConstraint ? ( SharedInverseMass[GGroupThreadId.x] + SharedInverseMass[GGroupThreadId.x-1] ) / (RestLength*RestLength) + 
-					SharedInverseInertia[GGroupThreadId.x-1]*4.0 : SharedInverseMass[GGroupThreadId.x] / (RestLength*RestLength) ;
-
-		//const float SumInverseMass = !ProjectConstraint ? ( SharedInverseMass[GGroupThreadId.x] + SharedInverseMass[GGroupThreadId.x-1] ) / (RestLength*RestLength) : SharedInverseMass[GGroupThreadId.x] / (RestLength*RestLength) ;
+		const float SumInverseMass = SharedInverseInertia[GGroupThreadId.x-1]*4.0;
 
 		const float SchurDiagonal = ( (1.0  + MaterialDamping / DeltaTime ) * SumInverseMass + OutMaterialCompliance );
 		OutMaterialWeight = ( SchurDiagonal != 0.0 ) ? 1.0 / SchurDiagonal : 0.0;
@@ -30,7 +27,7 @@ void SetupStretchRodMaterial(in int StrandsSize, in float YoungModulus, in float
 	}
 }
 
-void UpdateStretchRodMultiplier( in float RestLength, in float DeltaTime, in bool ProjectConstraint, in float MaterialDamping, in float MaterialCompliance, in float MaterialWeight, inout float3 OutMaterialMultiplier)
+void UpdateAlignRodMultiplier( in float RestLength, in float DeltaTime, in bool ProjectConstraint, in float MaterialDamping, in float MaterialCompliance, in float MaterialWeight, inout float3 OutMaterialMultiplier)
 {
 	float3 TargetDirection;
 	float4 q0 = SharedNodeOrientation[GGroupThreadId.x-1];
@@ -49,7 +46,116 @@ void UpdateStretchRodMultiplier( in float RestLength, in float DeltaTime, in boo
 
 	// XPBD lagrange multiplier update : dL = -(C+compliance*L) / (dC * invM * dCt + alpha)
 	const float3 DeltaLambda = -( EdgeDirection / RestLength - TargetDirection + 
-					OutMaterialMultiplier * MaterialCompliance + MaterialDamping * (DeltaVelocity/RestLength + QuatDamping) ) * MaterialWeight;
+					OutMaterialMultiplier * MaterialCompliance + MaterialDamping *  QuatDamping ) * MaterialWeight;
+
+	// L += dL
+	OutMaterialMultiplier += DeltaLambda;
+
+	// XPBD position update : dX = dL * dCt * invM
+	//const float QuaternionDelta = 0.5 * MultiplyQuat(float4(cross( TargetDirection, DeltaLambda),0), q0 );
+	const float4 QuaternionDelta = -2.0 * float4( DeltaLambda.xyz * qebar.w + cross(DeltaLambda.xyz, qebar.xyz),
+					- dot(DeltaLambda.xyz, qebar.xyz) );
+
+	// dX += dX
+	q0 += QuaternionDelta * SharedInverseInertia[GGroupThreadId.x-1];
+	SharedNodeOrientation[GGroupThreadId.x-1] = NormalizeQuat(q0);
+}
+
+void SolveAlignRodMaterial(in bool EnableConstraint, in int StrandsSize, in float RestLength, in float DeltaTime, in float MaterialDamping, 
+			in float MaterialCompliance, in float MaterialWeight, in float3 MaterialMultiplier, out float3 OutMaterialMultiplier)
+{
+	OutMaterialMultiplier = MaterialMultiplier;
+	if(EnableConstraint)
+	{
+		const int LocalIndex = (GGroupThreadId.x % StrandsSize);
+		if( LocalIndex > 1)
+		{
+			const int IsRed = (GGroupThreadId.x % 2) == 0;
+			// Process all the red rods
+			if (IsRed)
+			{	
+				UpdateAlignRodMultiplier(RestLength,DeltaTime,false,MaterialDamping,MaterialCompliance,MaterialWeight,OutMaterialMultiplier);
+			}
+			// Process all the black rods
+			GroupMemoryBarrier(); 
+			if (!IsRed)
+			{
+				UpdateAlignRodMultiplier(RestLength,DeltaTime,false,MaterialDamping,MaterialCompliance,MaterialWeight,OutMaterialMultiplier);
+			}
+			GroupMemoryBarrier(); 
+		}
+	}
+} 
+
+void ProjectAlignRodMaterial(in bool EnableConstraint, in int StrandsSize, in float YoungModulus, in float RodThickness, in float RestLength, in float DeltaTime, out float3 OutNodePosition)
+{
+	const int LocalIndex = (GGroupThreadId.x % StrandsSize);
+	if(DeltaTime != 0.0 && EnableConstraint)
+	{
+		float MaterialCompliance = 0.0;
+		float MaterialWeight = 0.0;
+		float3 MaterialMultiplier = float3(0,0,0);
+		
+		SetupAlignRodMaterial(StrandsSize, YoungModulus, RodThickness, RestLength, DeltaTime, true, 0.0, MaterialCompliance, MaterialWeight, MaterialMultiplier);
+
+		if( LocalIndex > 1 )
+		{
+			for(int i = 2; i < StrandsSize; ++i)
+			{
+				if( LocalIndex == i)
+				{
+					UpdateAlignRodMultiplier(RestLength,DeltaTime,true,0.0,MaterialCompliance,MaterialWeight,MaterialMultiplier);
+				}
+				GroupMemoryBarrier(); 
+			}
+		}
+	}
+	GroupMemoryBarrier();
+	OutNodePosition = SharedNodePosition[GGroupThreadId.x];
+}
+
+/* -----------------------------------------------------------------
+ * Stretch rod material
+ * -----------------------------------------------------------------
+ */
+
+void SetupStretchRodMaterial(in int StrandsSize, in float YoungModulus, in float RodThickness, 
+	in float RestLength, in float DeltaTime, in bool ProjectConstraint, in float MaterialDamping, out float OutMaterialCompliance, out float OutMaterialWeight, out float3 OutMaterialMultiplier)
+{
+	const int LocalIndex = (GGroupThreadId.x % StrandsSize);
+	if( LocalIndex > 1 )
+	{
+		// Compliance = 1.0 / (k * dt * dt)
+		// with k = L * L * (Y * A / L) (the L*L term before is coming from the fact that our constraint is dL/L and not dL)
+		// A is the cross section area (Pi * R * R), L is the rest length and Y is the young modulus
+		OutMaterialCompliance = 4.0/(YoungModulus*PI*RodThickness*RestLength*RodThickness*DeltaTime*DeltaTime);
+
+		const float SumInverseMass = !ProjectConstraint ? ( SharedInverseMass[GGroupThreadId.x] + SharedInverseMass[GGroupThreadId.x-1] ) / (RestLength*RestLength) 
+			: SharedInverseMass[GGroupThreadId.x] / (RestLength*RestLength) ;
+
+		const float SchurDiagonal = ( (1.0  + MaterialDamping / DeltaTime ) * SumInverseMass + OutMaterialCompliance );
+		OutMaterialWeight = ( SchurDiagonal != 0.0 ) ? 1.0 / SchurDiagonal : 0.0;
+
+		OutMaterialMultiplier = float3(0,0,0);
+	}
+}
+
+void UpdateStretchRodMultiplier( in float RestLength, in float DeltaTime, in bool ProjectConstraint, in float MaterialDamping, in float MaterialCompliance, in float MaterialWeight, inout float3 OutMaterialMultiplier)
+{
+	float3 TargetDirection;
+	float4 q0 = SharedNodeOrientation[GGroupThreadId.x-1];
+	TargetDirection[0] = 2.0 * (q0.x * q0.z + q0.w * q0.y);
+	TargetDirection[1] = 2.0 * (q0.y * q0.z - q0.w * q0.x);
+	TargetDirection[2] = q0.w * q0.w - q0.x * q0.x - q0.y * q0.y + q0.z * q0.z;
+
+	const float4 qebar = float4(-q0.y, q0.x, -q0.w, q0.z);	
+	const float3 EdgeDirection = SharedNodePosition[GGroupThreadId.x] - SharedNodePosition[GGroupThreadId.x-1];
+
+	const float3 DeltaVelocity = (EdgeDirection - ( SharedPreviousPosition[GGroupThreadId.x] - SharedPreviousPosition[GGroupThreadId.x-1] ) ) / DeltaTime; 
+
+	// XPBD lagrange multiplier update : dL = -(C+compliance*L) / (dC * invM * dCt + alpha)
+	const float3 DeltaLambda = -( EdgeDirection / RestLength - TargetDirection + 
+					OutMaterialMultiplier * MaterialCompliance + MaterialDamping * DeltaVelocity/RestLength ) * MaterialWeight;
 
 	// L += dL
 	OutMaterialMultiplier += DeltaLambda;
@@ -62,60 +168,60 @@ void UpdateStretchRodMultiplier( in float RestLength, in float DeltaTime, in boo
 	if(!ProjectConstraint)
 	{
 		SharedNodePosition[GGroupThreadId.x-1] -= PositionDelta * SharedInverseMass[GGroupThreadId.x-1];
-
-		//const float QuaternionDelta = 0.5 * MultiplyQuat(float4(cross( TargetDirection, DeltaLambda),0), q0 );
-
-		const float4 QuaternionDelta = -2.0 * float4( DeltaLambda.xyz * qebar.w + cross(DeltaLambda.xyz, qebar.xyz),
-						- dot(DeltaLambda.xyz, qebar.xyz) );
-
-		q0 += QuaternionDelta * SharedInverseInertia[GGroupThreadId.x-1];
-		SharedNodeOrientation[GGroupThreadId.x-1] = NormalizeQuat(q0);
 	}
 }
 
-void SolveStretchRodMaterial(in int StrandsSize, in float RestLength, in float DeltaTime, in float MaterialDamping, 
+void SolveStretchRodMaterial(in bool EnableConstraint, in int StrandsSize, in float RestLength, in float DeltaTime, in float MaterialDamping, 
 			in float MaterialCompliance, in float MaterialWeight, in float3 MaterialMultiplier, out float3 OutMaterialMultiplier)
 {
 	OutMaterialMultiplier = MaterialMultiplier;
-	const int LocalIndex = (GGroupThreadId.x % StrandsSize);
-	if( LocalIndex > 1)
+	if(EnableConstraint)
 	{
-		const int IsRed = (GGroupThreadId.x % 2) == 0;
-		// Process all the red rods
-		if (IsRed)
-		{	
-			UpdateStretchRodMultiplier(RestLength,DeltaTime,false,MaterialDamping,MaterialCompliance,MaterialWeight,OutMaterialMultiplier);
-		}
-		// Process all the black rods
-		GroupMemoryBarrier(); 
-		if (!IsRed)
+		const int LocalIndex = (GGroupThreadId.x % StrandsSize);
+		if( LocalIndex > 1)
 		{
-			UpdateStretchRodMultiplier(RestLength,DeltaTime,false,MaterialDamping,MaterialCompliance,MaterialWeight,OutMaterialMultiplier);
+			const int IsRed = (GGroupThreadId.x % 2) == 0;
+			// Process all the red rods
+			if (IsRed)
+			{	
+				UpdateStretchRodMultiplier(RestLength,DeltaTime,false,MaterialDamping,MaterialCompliance,MaterialWeight,OutMaterialMultiplier);
+			}
+			// Process all the black rods
+			GroupMemoryBarrier(); 
+			if (!IsRed)
+			{
+				UpdateStretchRodMultiplier(RestLength,DeltaTime,false,MaterialDamping,MaterialCompliance,MaterialWeight,OutMaterialMultiplier);
+			}
+			GroupMemoryBarrier(); 
 		}
-		GroupMemoryBarrier(); 
 	}
 } 
 
-void ProjectStretchRodMaterial(in int StrandsSize, in float YoungModulus, in float RodThickness, in float RestLength, in float DeltaTime, out float3 OutMaterialMultiplier)
+void ProjectStretchRodMaterial(in bool EnableConstraint, in int StrandsSize, in float YoungModulus, in float RodThickness, in float RestLength, in float DeltaTime, out float3 OutNodePosition)
 {
-	float MaterialCompliance = 0.0;
-	float MaterialWeight = 0.0;
-	OutMaterialMultiplier = float3(0,0,0);
-
-	SetupStretchRodMaterial(StrandsSize, YoungModulus, RodThickness, RestLength, DeltaTime, true, 0.0, MaterialCompliance, MaterialWeight, OutMaterialMultiplier);
-
 	const int LocalIndex = (GGroupThreadId.x % StrandsSize);
-	if( LocalIndex > 1 )
+	if(DeltaTime != 0.0 && EnableConstraint)
 	{
-		for(int i = 2; i < StrandsSize; ++i)
+		float MaterialCompliance = 0.0;
+		float MaterialWeight = 0.0;
+		float3 MaterialMultiplier = float3(0,0,0);
+
+		SetupStretchRodMaterial(StrandsSize, YoungModulus, RodThickness, RestLength, DeltaTime, true, 0.0, MaterialCompliance, MaterialWeight, MaterialMultiplier);
+
+		if( LocalIndex > 1 )
 		{
-			if( LocalIndex == i)
+			for(int i = 2; i < StrandsSize; ++i)
 			{
-				UpdateStretchRodMultiplier(RestLength,DeltaTime,true,0.0,MaterialCompliance,MaterialWeight,OutMaterialMultiplier);
+				if( LocalIndex == i)
+				{
+					UpdateStretchRodMultiplier(RestLength,DeltaTime,true,0.0,MaterialCompliance,MaterialWeight,MaterialMultiplier);
+				}
+				GroupMemoryBarrier(); 
 			}
-			GroupMemoryBarrier(); 
 		}
 	}
+	GroupMemoryBarrier();
+	OutNodePosition = SharedNodePosition[GGroupThreadId.x];
 }
 
 /* -----------------------------------------------------------------
@@ -132,7 +238,10 @@ void ProjectStretchRodMaterial(in int StrandsSize, in float YoungModulus, in flo
 		// Compliance = 1.0 / (k * dt * dt)
 		// with k = (Y * I * L) 
 		// I is the polar moment of inertia (Pi * R * R * R * R / 4), L is the rest length and Y is the young modulus
-		OutMaterialCompliance = 64.0/(YoungModulus*PI*RestLength*RodThickness*RodThickness*RodThickness*RodThickness*DeltaTime*DeltaTime);
+		const float InverseMoment = SharedInverseInertia[GGroupThreadId.x];
+		OutMaterialCompliance = InverseMoment / (YoungModulus*DeltaTime*DeltaTime);
+
+		//OutMaterialCompliance = 64.0 /(YoungModulus*PI*RestLength*RodThickness*RodThickness*RodThickness*RodThickness*DeltaTime*DeltaTime);
 
 		const float SumInverseMass = !ProjectConstraint ? ( SharedInverseInertia[GGroupThreadId.x] + SharedInverseInertia[GGroupThreadId.x-1] ) * 4.0 / (RestLength*RestLength) 
 							: SharedInverseInertia[GGroupThreadId.x]  * 4.0 / (RestLength*RestLength);
@@ -181,52 +290,67 @@ void UpdateBendRodMultiplier(in float RestLength, in float4 RestDarboux, in floa
 	}
 }
 
-void SolveBendRodMaterial(in int StrandsSize, in float RestLength, in float4 RestDarboux, in float DeltaTime, in float MaterialDamping, 
+void SolveBendRodMaterial(in bool EnableConstraint, in int StrandsSize, in float RestLength, in float4 RestDarboux, in float DeltaTime, in float MaterialDamping, 
 			in float MaterialCompliance, in float MaterialWeight, in float3 MaterialMultiplier, out float3 OutMaterialMultiplier)
 {
 	OutMaterialMultiplier = MaterialMultiplier;
-	const int LocalIndex = (GGroupThreadId.x % StrandsSize);
-	if( LocalIndex > 0)
+
+	if(EnableConstraint)
 	{
-		const int IsRed = (GGroupThreadId.x % 2) == 0;
-		// Process all the red rods
-		if (IsRed)
-		{	
-			UpdateBendRodMultiplier(RestLength,RestDarboux,DeltaTime,false,MaterialDamping,MaterialCompliance,MaterialWeight,OutMaterialMultiplier);
-		}
-		// Process all the black rods
-		GroupMemoryBarrier(); 
-		if (!IsRed)
+		const int LocalIndex = (GGroupThreadId.x % StrandsSize);
+		if( LocalIndex > 0)
 		{
-			UpdateBendRodMultiplier(RestLength,RestDarboux,DeltaTime,false,MaterialDamping,MaterialCompliance,MaterialWeight,OutMaterialMultiplier);
+			const int IsRed = (GGroupThreadId.x % 2) == 0;
+			// Process all the red rods
+			if (IsRed)
+			{	
+				UpdateBendRodMultiplier(RestLength,RestDarboux,DeltaTime,false,MaterialDamping,MaterialCompliance,MaterialWeight,OutMaterialMultiplier);
+			}
+			// Process all the black rods
+			GroupMemoryBarrier(); 
+			if (!IsRed)
+			{
+				UpdateBendRodMultiplier(RestLength,RestDarboux,DeltaTime,false,MaterialDamping,MaterialCompliance,MaterialWeight,OutMaterialMultiplier);
+			}
+			GroupMemoryBarrier(); 
 		}
-		GroupMemoryBarrier(); 
 	}
 } 
 
-void ProjectBendRodMaterial(in int StrandsSize, in float YoungModulus, in float RodThickness, in float RestLength, in float4 RestDarboux, in float DeltaTime, out float3 OutMaterialMultiplier)
+void ProjectBendRodMaterial(in bool EnableConstraint, in int StrandsSize, in float YoungModulus, in float RodThickness, in float RestLength, in float4 RestDarboux, in float DeltaTime, out float3 OutNodePosition)
 {
-	float MaterialCompliance = 0.0;
-	float MaterialWeight = 0.0;
-	OutMaterialMultiplier = float3(0,0,0);
-
-	SetupBendRodMaterial(StrandsSize, YoungModulus, RodThickness, RestLength, DeltaTime, true, 0.0, MaterialCompliance, MaterialWeight, OutMaterialMultiplier);
-
 	const int LocalIndex = (GGroupThreadId.x % StrandsSize);
-	if( LocalIndex > 0 )
+	if(DeltaTime != 0.0 && EnableConstraint)
 	{
-		for(int i = 1; i < StrandsSize; ++i)
+		float MaterialCompliance = 0.0;
+		float MaterialWeight = 0.0;
+		float3 MaterialMultiplier = float3(0,0,0);
+
+		SetupBendRodMaterial(StrandsSize, YoungModulus, RodThickness, RestLength, DeltaTime, true, 0.0, MaterialCompliance, MaterialWeight, MaterialMultiplier);
+
+		if( LocalIndex > 0 )
 		{
-			if( LocalIndex == i)
+			for(int i = 1; i < StrandsSize; ++i)
 			{
-				UpdateBendRodMultiplier(RestLength,RestDarboux,DeltaTime,true,0.0,MaterialCompliance,MaterialWeight,OutMaterialMultiplier);
+				if( LocalIndex == i)
+				{
+					UpdateBendRodMultiplier(RestLength,RestDarboux,DeltaTime,true,0.0,MaterialCompliance,MaterialWeight,MaterialMultiplier);
+				}
+				GroupMemoryBarrier(); 
 			}
-			GroupMemoryBarrier(); 
 		}
 	}
+	GroupMemoryBarrier();
+	OutNodePosition = SharedNodePosition[GGroupThreadId.x];
 }
 
-/*void UpdateBendTwistMultiplier( in float RestLength, in float4 RestDarboux, in float DeltaTime, in bool ProjectConstraint, in float MaterialDamping, in float MaterialCompliance, in float MaterialWeight, inout float3 OutMaterialMultiplier)
+/* -----------------------------------------------------------------
+ * Alternative of the Bend rod material
+ * -----------------------------------------------------------------
+ */
+
+/*void UpdateBendTwistMultiplier( in float RestLength, in float4 RestDarboux, in float DeltaTime, 
+				in bool ProjectConstraint, in float MaterialDamping, in float MaterialCompliance, in float MaterialWeight, inout float3 OutMaterialMultiplier)
 {
 	float4 q0 = SharedNodeOrientation[GGroupThreadId.x-1];
 	float4 q1 = SharedNodeOrientation[GGroupThreadId.x];
diff --git a/Engine/Plugins/Experimental/HairStrands/Shaders/Private/NiagaraDataInterfacePhysicsAsset.ush b/Engine/Plugins/Experimental/HairStrands/Shaders/Private/NiagaraDataInterfacePhysicsAsset.ush
--- a/Engine/Plugins/Experimental/HairStrands/Shaders/Private/NiagaraDataInterfacePhysicsAsset.ush
+++ b/Engine/Plugins/Experimental/HairStrands/Shaders/Private/NiagaraDataInterfacePhysicsAsset.ush
@@ -93,6 +93,27 @@ float3 DIPhysicsAsset_GetPreviousPosition(in FDIPhysicsAssetContext DIContext, i
     return float3(DIContext.PreviousTransformBuffer[BufferOffset].w, DIContext.PreviousTransformBuffer[BufferOffset+1].w, DIContext.PreviousTransformBuffer[BufferOffset+2].w);
 }
 
+float4 DIPhysicsAsset_GetCurrentOrientation(in FDIPhysicsAssetContext DIContext, in int ElementIndex)
+{
+	const int BufferOffset = 3 * ElementIndex;
+	const float3 OrientationMatrix[3] = {float3(DIContext.CurrentTransformBuffer[BufferOffset].x, DIContext.CurrentTransformBuffer[BufferOffset+1].x, DIContext.CurrentTransformBuffer[BufferOffset+2].x),
+									     float3(DIContext.CurrentTransformBuffer[BufferOffset].y, DIContext.CurrentTransformBuffer[BufferOffset+1].y, DIContext.CurrentTransformBuffer[BufferOffset+2].y),
+										 float3(DIContext.CurrentTransformBuffer[BufferOffset].z, DIContext.CurrentTransformBuffer[BufferOffset+1].z, DIContext.CurrentTransformBuffer[BufferOffset+2].z)};
+
+
+	return QuatFromMatrix(OrientationMatrix);
+}
+
+float4 DIPhysicsAsset_GetPreviousOrientation(in FDIPhysicsAssetContext DIContext, in int ElementIndex)
+{
+	const int BufferOffset = 3 * ElementIndex;
+	const float3 OrientationMatrix[3] = {float3(DIContext.PreviousTransformBuffer[BufferOffset].x, DIContext.PreviousTransformBuffer[BufferOffset+1].x, DIContext.PreviousTransformBuffer[BufferOffset+2].x),
+										 float3(DIContext.PreviousTransformBuffer[BufferOffset].y, DIContext.PreviousTransformBuffer[BufferOffset+1].y, DIContext.PreviousTransformBuffer[BufferOffset+2].y),
+										 float3(DIContext.PreviousTransformBuffer[BufferOffset].z, DIContext.PreviousTransformBuffer[BufferOffset+1].z, DIContext.PreviousTransformBuffer[BufferOffset+2].z)};
+
+	return QuatFromMatrix(OrientationMatrix);
+}
+
 float3x4 DIPhysicsAsset_GetCurrentTransform(in FDIPhysicsAssetContext DIContext, in int ElementIndex)
 {
 	const int BufferOffset = 3 * ElementIndex;
@@ -201,9 +222,6 @@ float DIPhysicsAsset_GetBoxProjection(in float3 LocalPosition, in float3 BoxExte
 				OutClosestPosition.z = NegativeSide ? -HalfExtent.z : HalfExtent.z;
 				OutClosestNormal.z = NegativeSide ? -1.0 : 1.0;
 			}
-
-			//OutClosestPosition[ClosestAxis] = NegativeSide ? -HalfExtent[ClosestAxis] : HalfExtent[ClosestAxis];
-			//OutClosestNormal[ClosestAxis] = NegativeSide ? -1.0 : 1.0;
 		}
 		else
 		{
@@ -220,6 +238,22 @@ float DIPhysicsAsset_GetBoxProjection(in float3 LocalPosition, in float3 BoxExte
  * -----------------------------------------------------------------
  */
 
+/*float3 DIPhysicsAsset_GetLocalPosition(in FDIPhysicsAssetContext DIContext, in float3 NodePosition, in int ElementIndex, in float TimeFraction )
+{
+	//return mul(DIPhysicsAsset_GetInverseTransform(DIContext,ElementIndex), float4(NodePosition,1.0)).xyz;
+
+	const float3 CurrentPosition = DIPhysicsAsset_GetCurrentPosition(DIContext,ElementIndex);
+	const float4 CurrentOrientation = DIPhysicsAsset_GetCurrentOrientation(DIContext,ElementIndex);
+
+	const float3 PreviousPosition = DIPhysicsAsset_GetPreviousPosition(DIContext,ElementIndex);
+	const float4 PreviousOrientation = DIPhysicsAsset_GetPreviousOrientation(DIContext,ElementIndex);
+
+	const float4 SlerpedOrientation = InverseQuat(SlerpQuat(PreviousOrientation,CurrentOrientation,TimeFraction));
+	const float3 SlerpedPosition = PreviousPosition + TimeFraction * (CurrentPosition-PreviousPosition);
+
+	return RotateVectorByQuat( NodePosition - SlerpedPosition, SlerpedOrientation );
+}*/
+
 float3 DIPhysicsAsset_GetLocalPosition(in FDIPhysicsAssetContext DIContext, in float3 NodePosition, in int ElementIndex, in float TimeFraction )
 {
 	const float3 CurrentLocal = mul(DIPhysicsAsset_GetInverseTransform(DIContext,ElementIndex), float4(NodePosition,1.0)).xyz;
diff --git a/Engine/Plugins/Experimental/HairStrands/Shaders/Private/NiagaraHookeSpringMaterial.ush b/Engine/Plugins/Experimental/HairStrands/Shaders/Private/NiagaraHookeSpringMaterial.ush
--- a/Engine/Plugins/Experimental/HairStrands/Shaders/Private/NiagaraHookeSpringMaterial.ush
+++ b/Engine/Plugins/Experimental/HairStrands/Shaders/Private/NiagaraHookeSpringMaterial.ush
@@ -90,16 +90,10 @@ void SolveStretchSpringMaterial(in bool EnableConstraint, in int StrandsSize, in
 	}
 } 
 
-void ProjectStretchSpringMaterial(in int StrandsSize, in float YoungModulus, in float RodThickness, in float RestLength, in float NodeMass, in float DeltaTime, in float3 NodePosition, out float3 OutNodePosition)
+void ProjectStretchSpringMaterial(in bool EnableConstraint, in int StrandsSize, in float YoungModulus, in float RodThickness, in float RestLength, in float DeltaTime, out float3 OutNodePosition)
 {
 	const int LocalIndex = (GGroupThreadId.x % StrandsSize);
-
-	SharedNodePosition[GGroupThreadId.x] = NodePosition;
-	SharedPreviousPosition[GGroupThreadId.x] = NodePosition;
-	SharedInverseMass[GGroupThreadId.x] = (LocalIndex > 1 && NodeMass != 0.0) ? 1.0 / NodeMass : 0.0;
-	GroupMemoryBarrier();
-
-	if(DeltaTime != 0.0)
+	if(DeltaTime != 0.0 && EnableConstraint)
 	{
 		float MaterialCompliance = 0.0;
 		float MaterialWeight = 0.0;
@@ -168,16 +162,10 @@ void SolveBendSpringMaterial(in bool EnableConstraint, in int StrandsSize, in fl
 	}
 } 
 
-void ProjectBendSpringMaterial(in int StrandsSize, in float YoungModulus, in float RodThickness, in float RestLength, in float NodeMass, in float DeltaTime, in float3 NodePosition, out float3 OutNodePosition)
+void ProjectBendSpringMaterial(in bool EnableConstraint, in int StrandsSize, in float YoungModulus, in float RodThickness, in float RestLength, in float DeltaTime, out float3 OutNodePosition)
 {
 	const int LocalIndex = (GGroupThreadId.x % StrandsSize);
-
-	SharedNodePosition[GGroupThreadId.x] = NodePosition;
-	SharedPreviousPosition[GGroupThreadId.x] = NodePosition;
-	SharedInverseMass[GGroupThreadId.x] = (LocalIndex > 1 && NodeMass != 0.0) ? 1.0 / NodeMass : 0.0;
-	GroupMemoryBarrier();
-
-	if(DeltaTime != 0.0)
+	if(DeltaTime != 0.0 && EnableConstraint)
 	{
 		float MaterialCompliance = 0.0;
 		float MaterialWeight = 0.0;
@@ -201,102 +189,6 @@ void ProjectBendSpringMaterial(in int StrandsSize, in float YoungModulus, in flo
 	OutNodePosition = SharedNodePosition[GGroupThreadId.x];
 }
 
-/* -----------------------------------------------------------------
- * Volume spring material
- * -----------------------------------------------------------------
- */
-
-
- void SetupVolumeSpringMaterial(in int StrandsSize, in float YoungModulus,
-	in float RestVolume, in float DeltaTime, in bool ProjectConstraint, in float MaterialDamping, out float OutMaterialCompliance, out float OutMaterialWeight, out float OutMaterialMultiplier)
-{
-	OutMaterialCompliance = 0.0;
-	OutMaterialWeight = 0.0;
-	OutMaterialMultiplier = 0.0;
-
-	const int LocalIndex = (GGroupThreadId.x % StrandsSize);
-	if( LocalIndex > 2 )
-	{
-		// Compliance = 1.0 / (k * dt * dt)
-		// with k = (Y / V) 
-		// V is the constraint inital volume and Y is the young modulus
-		OutMaterialCompliance = RestVolume/(YoungModulus*DeltaTime*DeltaTime);
-	}
-}
- 
-void UpdateVolumeSpringMultiplier(in float RestVolume, in float DeltaTime, in bool ProjectConstraint, in float MaterialDamping, in float MaterialCompliance, in float MaterialWeight, inout float OutMaterialMultiplier)
-{
-	const int NodeIndexA = GGroupThreadId.x;
-	const int NodeIndexB = GGroupThreadId.x-1;
-	const int NodeIndexC = GGroupThreadId.x-2;
-	const int NodeIndexD = GGroupThreadId.x-3;
-
-	const float3 EdgeVectorA = SharedNodePosition[NodeIndexB] - SharedNodePosition[NodeIndexA];
-	const float3 EdgeVectorB = SharedNodePosition[NodeIndexC] - SharedNodePosition[NodeIndexA];
-	const float3 EdgeVectorC = SharedNodePosition[NodeIndexD] - SharedNodePosition[NodeIndexA];
-
-	const float NodesVolume = dot(cross(EdgeVectorB,EdgeVectorC),EdgeVectorA);
-
-	const float3 ConstraintGradB = cross(EdgeVectorB,EdgeVectorC);
-	const float3 ConstraintGradC = cross(EdgeVectorC,EdgeVectorA);
-	const float3 ConstraintGradD = cross(EdgeVectorA,EdgeVectorB);
-	const float3 ConstraintGradA = -ConstraintGradB-ConstraintGradC-ConstraintGradD;
-
-	const float3 NodeVelocityA = ( SharedNodePosition[NodeIndexA] - SharedPreviousPosition[NodeIndexA] ) / DeltaTime;
-	const float3 NodeVelocityB = ( SharedNodePosition[NodeIndexB] - SharedPreviousPosition[NodeIndexB] ) / DeltaTime;
-	const float3 NodeVelocityC = ( SharedNodePosition[NodeIndexC] - SharedPreviousPosition[NodeIndexC] ) / DeltaTime;
-	const float3 NodeVelocityD = ( SharedNodePosition[NodeIndexD] - SharedPreviousPosition[NodeIndexD] ) / DeltaTime;
-
-	const float SumInverseMass = !ProjectConstraint ? ( dot(ConstraintGradA,ConstraintGradA) * SharedInverseMass[NodeIndexA] +
-														dot(ConstraintGradB,ConstraintGradB) * SharedInverseMass[NodeIndexB] +
-													    dot(ConstraintGradC,ConstraintGradC) * SharedInverseMass[NodeIndexC] +
-														dot(ConstraintGradD,ConstraintGradD) * SharedInverseMass[NodeIndexD] ) :
-														dot(ConstraintGradA,ConstraintGradA) * SharedInverseMass[NodeIndexA] ;
-
-	const float SchurDiagonal = (1.0 + MaterialDamping) * SumInverseMass + MaterialCompliance; 
-	const float ConstraintWeight = ( SchurDiagonal != 0.0 ) ? 1.0 / SchurDiagonal : 0.0;
-
-	// XPBD lagrange multiplier update : dL = -(C+compliance*L) / (dC * invM * dCt + alpha)
-	const float DeltaLambda = -((NodesVolume - RestVolume) + OutMaterialMultiplier * MaterialCompliance + MaterialDamping * 
-		(dot(ConstraintGradA,NodeVelocityA)+dot(ConstraintGradB,NodeVelocityB)+dot(ConstraintGradC,NodeVelocityC)+dot(ConstraintGradD,NodeVelocityD) ) ) * ConstraintWeight;
-
-	// L += dL
-	//OutMaterialMultiplier += DeltaLambda;
-	OutMaterialMultiplier = NodesVolume;
-
-	// XPBD position update : dX += dL * dCt * invM
-	SharedNodePosition[NodeIndexA] += ConstraintGradA * DeltaLambda * SharedInverseMass[NodeIndexA];
-	if(!ProjectConstraint)
-	{
-		SharedNodePosition[NodeIndexB] += ConstraintGradB * DeltaLambda * SharedInverseMass[NodeIndexB]; 
-		SharedNodePosition[NodeIndexC] += ConstraintGradC * DeltaLambda * SharedInverseMass[NodeIndexC];
-		SharedNodePosition[NodeIndexD] += ConstraintGradD * DeltaLambda * SharedInverseMass[NodeIndexD];
-	}
-}
-
-void SolveVolumeSpringMaterial(in bool EnableConstraint, in int StrandsSize, in float RestVolume, in float DeltaTime, in float MaterialDamping, 
-			in float MaterialCompliance, in float MaterialWeight, in float MaterialMultiplier, out float OutMaterialMultiplier)
-{
-	OutMaterialMultiplier = MaterialMultiplier;
-	if(EnableConstraint)
-	{
-		const int LocalIndex = (GGroupThreadId.x % StrandsSize);
-		if( LocalIndex > 2)
-		{
-			for(int i = 3; i < StrandsSize; ++i)
-			{
-				if( LocalIndex == i )
-				{
-					//UpdateVolumeSpringMultiplier(RestVolume,DeltaTime,false,MaterialDamping,MaterialCompliance,MaterialWeight,OutMaterialMultiplier);
-					UpdateVolumeSpringMultiplier(RestVolume,DeltaTime,false,0.0,0.0,MaterialWeight,OutMaterialMultiplier);
-				}
-				GroupMemoryBarrier(); 
-			}
-		}
-	}
-	GroupMemoryBarrier();
-} 
-
 /* -----------------------------------------------------------------
  * Twist spring material
  * -----------------------------------------------------------------
@@ -341,16 +233,11 @@ void SolveTwistSpringMaterial(in bool EnableConstraint, in int StrandsSize, in f
 	}
 } 
 
-void ProjectTwistSpringMaterial(in int StrandsSize, in float YoungModulus, in float RodThickness, in float RestLength, in float NodeMass, in float DeltaTime, in float3 NodePosition, out float3 OutNodePosition)
+void ProjectTwistSpringMaterial(in bool EnableConstraint, in int StrandsSize, in float YoungModulus, in float RodThickness, in float RestLength, in float DeltaTime, out float3 OutNodePosition)
 {
 	const int LocalIndex = (GGroupThreadId.x % StrandsSize);
 
-	SharedNodePosition[GGroupThreadId.x] = NodePosition;
-	SharedPreviousPosition[GGroupThreadId.x] = NodePosition;
-	SharedInverseMass[GGroupThreadId.x] = (LocalIndex > 1 && NodeMass != 0.0) ? 1.0 / NodeMass : 0.0;
-	GroupMemoryBarrier();
-
-	if(DeltaTime != 0.0)
+	if(DeltaTime != 0.0 && EnableConstraint)
 	{
 		float MaterialCompliance = 0.0;
 		float MaterialWeight = 0.0;
@@ -373,278 +260,3 @@ void ProjectTwistSpringMaterial(in int StrandsSize, in float YoungModulus, in fl
 	GroupMemoryBarrier();
 	OutNodePosition = SharedNodePosition[GGroupThreadId.x];
 }
-
-/* -----------------------------------------------------------------
- * Compute Bishop frame
- * -----------------------------------------------------------------
- */
-
-void ComputeMaterialFrame(in int StrandsSize)
-{
-	const int LocalIndex = (GGroupThreadId.x % StrandsSize);
-	if( LocalIndex == 0 )
-	{
-		float4 NodeQuaternion = SharedNodeOrientation[GGroupThreadId.x];
-		float3 TangentPrev = normalize(SharedNodePosition[GGroupThreadId.x+1] - SharedNodePosition[GGroupThreadId.x]);
-		float3 TangentNext = TangentPrev;
-
-		for (int EdgeIndex = 1, EdgeEnd = StrandsSize-1; EdgeIndex < EdgeEnd; ++EdgeIndex)
-		{
-			TangentPrev = TangentNext;
-			TangentNext = normalize(SharedNodePosition[GGroupThreadId.x+EdgeIndex+1] - SharedNodePosition[GGroupThreadId.x+EdgeIndex]);
-
-			NodeQuaternion = NormalizeQuat( MultiplyQuat( FindQuatBetweenNormals(TangentPrev,TangentNext), NodeQuaternion ) );
-			SharedNodeOrientation[GGroupThreadId.x+EdgeIndex] = NodeQuaternion;
-			GroupMemoryBarrier();
-		}
-		GroupMemoryBarrier();
-	}
-	GroupMemoryBarrier();
-}
-
-void UpdateMaterialFrame(in int StrandsSize)
-{
-	const int LocalIndex = (GGroupThreadId.x % StrandsSize);
-	if( LocalIndex == 0 )
-	{
-		float4 NodeQuaternion = SharedPreviousOrientation[GGroupThreadId.x];
-		float3 TangentPrev = normalize(SharedPreviousPosition[GGroupThreadId.x+1] - SharedPreviousPosition[GGroupThreadId.x]);
-		float3 TangentNext = TangentPrev;
-
-		for (int EdgeIndex = 1, EdgeEnd = StrandsSize-1; EdgeIndex < EdgeEnd; ++EdgeIndex)
-		{
-			TangentPrev = TangentNext;
-			TangentNext = normalize(SharedPreviousPosition[GGroupThreadId.x+EdgeIndex+1] - SharedPreviousPosition[GGroupThreadId.x+EdgeIndex]);
-
-			NodeQuaternion = NormalizeQuat( MultiplyQuat( FindQuatBetweenNormals(TangentPrev,TangentNext), NodeQuaternion ) );
-			SharedNodeOrientation[GGroupThreadId.x+EdgeIndex] = NodeQuaternion;
-			GroupMemoryBarrier();
-		}
-		GroupMemoryBarrier();
-	}
-	GroupMemoryBarrier();
-}
-
-
-/*void UpdateMaterialFrame(in int StrandsSize)
-{
-	const int LocalIndex = (GGroupThreadId.x % StrandsSize);
-	if( (LocalIndex > 0) && (LocalIndex < StrandsSize-1) )
-	{
-		float4 NodeQuaternion = SharedPreviousOrientation[GGroupThreadId.x];
-		//float3 TangentPrev = RotateVectorByQuat( float3(0,0,1), NodeQuaternion);
-		//float3 TangentNext = normalize(SharedPreviousPosition[GGroupThreadId.x+1] - SharedPreviousPosition[GGroupThreadId.x]);
-
-		float3 TangentPrev = normalize(SharedPreviousPosition[GGroupThreadId.x+1] - SharedPreviousPosition[GGroupThreadId.x]);
-		float3 TangentNext = normalize(SharedNodePosition[GGroupThreadId.x+1] - SharedNodePosition[GGroupThreadId.x]);
-
-		SharedNodeOrientation[GGroupThreadId.x] = NormalizeQuat( MultiplyQuat( FindQuatBetweenNormals(TangentPrev,TangentNext), NodeQuaternion) );
-	}
-	GroupMemoryBarrier();
-}*/
-
-/* -----------------------------------------------------------------
- * Bend spring material
- * -----------------------------------------------------------------
- */
-
- void SetupAngularSpringMaterial(in int StrandsSize, in float YoungModulus, in float RodThickness, 
-	in float RestLength, in float DeltaTime, in bool ProjectConstraint, in float MaterialDamping, out float OutMaterialCompliance, out float OutMaterialWeight, out float3 OutMaterialMultiplier)
-{
-	OutMaterialCompliance = 0.0;
-	OutMaterialWeight = 0.0;
-	OutMaterialMultiplier = float3(0,0,0);
-
-	float MaterialMultiplier = 0;
-	SetupStretchSpringMaterial(StrandsSize,YoungModulus,RodThickness,RestLength,DeltaTime,ProjectConstraint,MaterialDamping,OutMaterialCompliance,OutMaterialWeight,MaterialMultiplier);
-}
-
-void UpdateAngularSpringMultiplier( in float RestLength, in float3 RestDirection, in float DeltaTime, in bool ProjectConstraint, in float MaterialDamping, in float InMaterialCompliance, in float InMaterialWeight, inout float3 OutMaterialMultiplier)
-{
-	const float3 EdgeDirection = SharedNodePosition[GGroupThreadId.x] - SharedNodePosition[GGroupThreadId.x-1];
-	const float3 DeltaVelocity = (EdgeDirection - ( SharedPreviousPosition[GGroupThreadId.x] - SharedPreviousPosition[GGroupThreadId.x-1] ) ) / DeltaTime; 
-
-	const float4 ParentFrame = SlerpQuat(SharedNodeOrientation[GGroupThreadId.x-2],SharedNodeOrientation[GGroupThreadId.x-2],0.5);
-	const float3 TargetDirection = RotateVectorByQuat(RestDirection,ParentFrame);
-
-	// XPBD lagrang : dL = -(C+compliance*L) / (dC * invM * dCt + alpha)
-	const float3 DeltaLambda = -((EdgeDirection - TargetDirection)/RestLength + 
-					OutMaterialMultiplier * InMaterialCompliance  + MaterialDamping * DeltaVelocity /RestLength) * InMaterialWeight;
-
-	// L += dL
-	OutMaterialMultiplier += DeltaLambda;
-
-	// XPBD position update : dX = dL * dCt * invM
-	const float3 PositionDelta  = DeltaLambda/RestLength;
-
-	// dX += dX
-	SharedNodePosition[GGroupThreadId.x] += PositionDelta * SharedInverseMass[GGroupThreadId.x];
-	if(!ProjectConstraint)
-	{
-		SharedNodePosition[GGroupThreadId.x-1] -= PositionDelta * SharedInverseMass[GGroupThreadId.x-1];
-	}
-}
-
-void SolveAngularSpringMaterial(in bool EnableConstraint, in int StrandsSize, in float RestLength, in float3 RestDirection, in float DeltaTime, in float MaterialDamping,
-			in float MaterialCompliance, in float MaterialWeight, in float3 MaterialMultiplier, out float3 OutMaterialMultiplier)
-{
-	OutMaterialMultiplier = MaterialMultiplier;
-
-	if(EnableConstraint)
-	{
-		const int LocalIndex = (GGroupThreadId.x % StrandsSize);
-		if( LocalIndex > 1)
-		{
-			const int IsRed = (GGroupThreadId.x % 2) == 0;
-			// Process all the red rods
-			if (!IsRed)
-			{	
-				UpdateAngularSpringMultiplier(RestLength,RestDirection,DeltaTime,false,MaterialDamping,MaterialCompliance,MaterialWeight,OutMaterialMultiplier);
-			}
-			GroupMemoryBarrier(); 
-
-			if (IsRed)
-			{
-				UpdateAngularSpringMultiplier(RestLength,RestDirection,DeltaTime,false,MaterialDamping,MaterialCompliance,MaterialWeight,OutMaterialMultiplier);
-			}
-			GroupMemoryBarrier(); 
-		}
-	}
-}
-/*
-groupshared float4 SharedRestDatas[THREADGROUP_SIZE];
-void SolveAngularSpringMaterial(in int StrandsSize, in float RestLength, in float3 RestDirection, in float DeltaTime, in float MaterialDamping,
-			in float MaterialCompliance, in float MaterialWeight, in float3 MaterialMultiplier, out float3 OutMaterialMultiplier)
-{
-	SharedRestDatas[GGroupThreadId.x] = float4(RestDirection,RestLength);
-	GroupMemoryBarrier(); 
-
-	OutMaterialMultiplier = MaterialMultiplier;
-	const int LocalIndex = (GGroupThreadId.x % StrandsSize);
-	if( LocalIndex == 0)
-	{
-		for(int i = GGroupThreadId.x+2; i < GGroupThreadId.x+StrandsSize; ++i)
-		{
-			const float3 EdgeDirection = SharedNodePosition[i] - SharedNodePosition[i-1];
-
-			const float4 ParentFrame = SharedNodeOrientation[i-2];
-			const float3 TargetDirection = RotateVectorByQuat(SharedRestDatas[i].xyz,ParentFrame);
-
-			// XPBD lagrange multiplier update : dL = -(C+compliance*L) / (dC * invM * dCt + alpha)
-			const float3 DeltaLambda = -(EdgeDirection - TargetDirection) / (SharedInverseMass[i]+SharedInverseMass[i-1]);
-
-			// XPBD position update : dX = dL * dCt * invM
-			const float3 PositionDelta  = DeltaLambda;
-
-			SharedDeltaPosition[i] = PositionDelta;
-			GroupMemoryBarrier();
-
-			// dX += dX
-			SharedNodePosition[i] += PositionDelta * SharedInverseMass[i];
-			GroupMemoryBarrier();
-			SharedNodePosition[i-1] -= PositionDelta * SharedInverseMass[i-1];
-			GroupMemoryBarrier();
-
-			if( i > GGroupThreadId.x+2 )
-			{
-				const float3 TangentInit = normalize(SharedNodePosition[i-2] - SharedNodePosition[i-3]);
-				const float3 TangentPrev = normalize(SharedNodePosition[i-1] - SharedNodePosition[i-2]);
-				const float3 TangentNext = normalize(SharedNodePosition[i] - SharedNodePosition[i-1]);
-
-				SharedNodeOrientation[i-2] = NormalizeQuat( MultiplyQuat( FindQuatBetweenNormals(TangentInit,TangentPrev), SharedNodeOrientation[i-3] ) );
-				GroupMemoryBarrier();
-				SharedNodeOrientation[i-1] = NormalizeQuat( MultiplyQuat( FindQuatBetweenNormals(TangentPrev,TangentNext), SharedNodeOrientation[i-2] ) );
-				GroupMemoryBarrier();
-			}
-			else
-			{
-				const float3 TangentPrev = normalize(SharedNodePosition[i-1] - SharedNodePosition[i-2]);
-				const float3 TangentNext = normalize(SharedNodePosition[i] - SharedNodePosition[i-1]);
-
-				SharedNodeOrientation[i-1] = NormalizeQuat( MultiplyQuat( FindQuatBetweenNormals(TangentPrev,TangentNext), SharedNodeOrientation[i-2] ) );
-				GroupMemoryBarrier();
-			}
-			GroupMemoryBarrier();
-		}
-		GroupMemoryBarrier();
-		for(int i = GGroupThreadId.x+2; i < GGroupThreadId.x+StrandsSize; ++i)
-		{
-			SharedNodePosition[i] += SharedDeltaPosition[i] * SharedInverseMass[i];
-			SharedNodePosition[i-1] -= SharedDeltaPosition[i] * SharedInverseMass[i-1];
-			GroupMemoryBarrier();
-		}
-		GroupMemoryBarrier();
-		for(int i = GGroupThreadId.x+2; i < GGroupThreadId.x+StrandsSize; ++i)
-		{
-			const float3 TangentPrev = normalize(SharedNodePosition[i-1] - SharedNodePosition[i-2]);
-			const float3 TangentNext = normalize(SharedNodePosition[i] - SharedNodePosition[i-1]);
-
-			SharedNodeOrientation[i-1] = NormalizeQuat( MultiplyQuat( FindQuatBetweenNormals(TangentPrev,TangentNext), SharedNodeOrientation[i-2] ) );
-			GroupMemoryBarrier();
-		}
-		GroupMemoryBarrier();
-	} 
-}*/
-
-void ProjectAngularSpringMaterial(in int StrandsSize, in float YoungModulus, in float RodThickness, in float RestLength, in float3 RestDirection, in float NodeMass, in float DeltaTime,
-			 in float3 NodePosition, out float3 OutNodePosition)
-{
-	const int LocalIndex = (GGroupThreadId.x % StrandsSize);
-
-	SharedNodePosition[GGroupThreadId.x] = NodePosition;
-	SharedPreviousPosition[GGroupThreadId.x] = NodePosition;
-	SharedInverseMass[GGroupThreadId.x] = (LocalIndex > 1 && NodeMass != 0.0) ? 1.0 / NodeMass : 0.0;
-	GroupMemoryBarrier();
-
-	if(DeltaTime != 0.0)
-	{
-
-		float MaterialCompliance = 0.0;
-		float MaterialWeight = 0.0;
-		float3 MaterialMultiplier = float3(0,0,0);
-
-		//ComputeMaterialFrame(StrandsSize);
-		SetupAngularSpringMaterial(StrandsSize, YoungModulus, RodThickness, RestLength, DeltaTime, true, 0.0, MaterialCompliance, MaterialWeight, MaterialMultiplier);
-
-		//if( LocalIndex > 1 )
-		{
-			for(int i = 2; i < StrandsSize; ++i)
-			{
-				ComputeMaterialFrame(StrandsSize);
-				if( LocalIndex == i )
-				{
-					UpdateAngularSpringMultiplier(RestLength,RestDirection,DeltaTime,true,0.0,MaterialCompliance,MaterialWeight,MaterialMultiplier);
-				}
-				GroupMemoryBarrier(); 
-			}
-		}
-	}
-
-	GroupMemoryBarrier();
-	OutNodePosition = SharedNodePosition[GGroupThreadId.x];
-}
-
-/* -----------------------------------------------------------------
- * Compute Rest Direction and length
- * -----------------------------------------------------------------
- */
-
-void ComputeRestDirection(in int StrandsSize, in float3 NodePosition, in float4 NodeOrientation, out float3 OutRestDirection)
-{
-	SharedNodePosition[GGroupThreadId.x] = NodePosition;
-	SharedNodeOrientation[GGroupThreadId.x] = NodeOrientation;
-	GroupMemoryBarrier(); 
-
-	ComputeMaterialFrame(StrandsSize);
-
-	OutRestDirection = float3(0,0,0);
-	const int LocalIndex = (GGroupThreadId.x % StrandsSize);
-	if( LocalIndex > 1 )
-	{
-		const float3 EdgeDirection = SharedNodePosition[GGroupThreadId.x] - SharedNodePosition[GGroupThreadId.x-1];
-
-		const float4 ParentFrame = SlerpQuat(SharedNodeOrientation[GGroupThreadId.x-2],SharedNodeOrientation[GGroupThreadId.x-2],0.5);
-		OutRestDirection = UnrotateVectorByQuat(EdgeDirection,ParentFrame);
-	}
-	GroupMemoryBarrier(); 
-}
diff --git a/Engine/Plugins/Experimental/HairStrands/Shaders/Private/NiagaraStaticCollisionConstraint.ush b/Engine/Plugins/Experimental/HairStrands/Shaders/Private/NiagaraStaticCollisionConstraint.ush
--- a/Engine/Plugins/Experimental/HairStrands/Shaders/Private/NiagaraStaticCollisionConstraint.ush
+++ b/Engine/Plugins/Experimental/HairStrands/Shaders/Private/NiagaraStaticCollisionConstraint.ush
@@ -3,28 +3,30 @@
 #pragma once
 
 /* -----------------------------------------------------------------
- * Static hard collision constraint
+ * Collision Utilities
  * -----------------------------------------------------------------
  */
 
- float3 ComputeStaticCollisionDisplacement( in float3 DeltaPosition, in float3 DeltaVelocity, in float3 CollisionNormal, in float PenetrationDepth,
-							  in float StaticFriction, in float KineticFriction, in bool IsProjection)
+  bool ComputeCollisionDatas( in float3 DeltaPosition, in float3 DeltaVelocity, in float3 CollisionNormal, in float PenetrationDepth,
+							  in float StaticFriction, in float KineticFriction, in bool IsProjection, inout float DeltaFNormal, 
+							  inout float AlphaTangentA, inout float DeltaVTangentA, inout float3 CollisionTangentA, 
+							  inout float AlphaTangentB, inout float DeltaVTangentB, inout float3 CollisionTangentB)
  {
 	const float DeltaVNormal = dot(DeltaVelocity,CollisionNormal);
 	const float DeltaPNormal = dot(DeltaPosition,CollisionNormal)-PenetrationDepth;
 
 	if (DeltaPNormal < 0.0)
 	{
-		const float DeltaFNormal = IsProjection ? DeltaPNormal : DeltaVNormal;
+		DeltaFNormal = IsProjection ? DeltaPNormal : DeltaVNormal;
 		const float NormDeltaFN = abs(DeltaFNormal);
 
 		const float3 CollisionTangent = DeltaVelocity - DeltaVNormal * CollisionNormal;
 		const float TangentLength = length(CollisionTangent);
 
-		float3 CollisionTangentA = (TangentLength > 0.0) ? CollisionTangent/TangentLength : (abs(CollisionNormal.x-1.0f) > 0.0) ? float3(1,0,0) : 
+		CollisionTangentA = (TangentLength > 0.0) ? CollisionTangent/TangentLength : (abs(CollisionNormal.x-1.0f) > 0.0) ? float3(1,0,0) : 
 																					        (abs(CollisionNormal.y-1.0f) > 0.0) ? float3(0,1,0) :
 																						     float3(0,0,1);
-		float3 CollisionTangentB = cross(CollisionNormal,CollisionTangentA);
+		CollisionTangentB = cross(CollisionNormal,CollisionTangentA);
 		CollisionTangentA = cross(CollisionTangentB,CollisionNormal);
 
 		const float TangentLengthA = length(CollisionTangentA);
@@ -33,76 +35,97 @@
 		CollisionTangentA = (TangentLengthA > 0.0) ? CollisionTangentA/TangentLengthA : float3(0,0,0);
 		CollisionTangentB = (TangentLengthB > 0.0) ? CollisionTangentB/TangentLengthB : float3(0,0,0);
 
-		const float DeltaVTangentA = dot( DeltaVelocity, CollisionTangentA);
-		const float DeltaVTangentB = dot( DeltaVelocity, CollisionTangentB);
+		DeltaVTangentA = dot( DeltaVelocity, CollisionTangentA);
+		DeltaVTangentB = dot( DeltaVelocity, CollisionTangentB);
 
 		const float NormDeltaVTA = abs(DeltaVTangentA);
 		const float NormDeltaVTB = abs(DeltaVTangentB);
 
-		const float AlphaTangentA = ( NormDeltaVTA < StaticFriction * NormDeltaFN ) ? 1.0 : (KineticFriction > 0.0) ? min(KineticFriction*NormDeltaFN/NormDeltaVTA, 1.0) : 0.0;
-		const float AlphaTangentB = ( NormDeltaVTB < StaticFriction * NormDeltaFN ) ? 1.0 : (KineticFriction > 0.0) ? min(KineticFriction*NormDeltaFN/NormDeltaVTB, 1.0) : 0.0;
+		AlphaTangentA = ( NormDeltaVTA < StaticFriction * NormDeltaFN ) ? 1.0 : (KineticFriction > 0.0) ? min(KineticFriction*NormDeltaFN/NormDeltaVTA, 1.0) : 0.0;
+		AlphaTangentB = ( NormDeltaVTB < StaticFriction * NormDeltaFN ) ? 1.0 : (KineticFriction > 0.0) ? min(KineticFriction*NormDeltaFN/NormDeltaVTB, 1.0) : 0.0;
+		return true;
+	}
+	return false;
+ }
 
-		return DeltaFNormal * CollisionNormal + AlphaTangentA * CollisionTangentA * DeltaVTangentA + AlphaTangentB * CollisionTangentB * DeltaVTangentB;
+/* -----------------------------------------------------------------
+ * Static hard collision constraint
+ * -----------------------------------------------------------------
+ */
+
+ void UpdateHardCollisionConstraint( in float3 DeltaPosition, in float3 DeltaVelocity, in float3 CollisionNormal, in float PenetrationDepth,
+							  in float StaticFriction, in float KineticFriction, in bool IsProjection)
+ {
+	float DeltaFNormal = 0.0; 
+
+	float AlphaTangentA = 0.0;
+	float DeltaVTangentA = 0.0;
+	float3 CollisionTangentA = float3(0,0,0);
+	
+	float AlphaTangentB = 0.0;
+	float DeltaVTangentB = 0.0;
+	float3 CollisionTangentB = float3(0,0,0);
+
+	if( ComputeCollisionDatas(DeltaPosition,DeltaVelocity,CollisionNormal,PenetrationDepth,StaticFriction,KineticFriction,IsProjection,
+			DeltaFNormal,AlphaTangentA,DeltaVTangentA,CollisionTangentA,AlphaTangentB,DeltaVTangentB,CollisionTangentB))
+	{
+		SharedNodePosition[GGroupThreadId.x] -= DeltaFNormal * CollisionNormal + AlphaTangentA * CollisionTangentA * DeltaVTangentA + AlphaTangentB * CollisionTangentB * DeltaVTangentB;
 	}
-	return float3(0,0,0);
  }
 
- void SolveStaticCollisionConstraint(in bool EnableConstraint, in int StrandsSize, in float PenetrationDepth, in float3 CollisionPosition, in float3 CollisionVelocity, in float3 CollisionNormal, 
-				in float StaticFriction, in float KineticFriction, in float DeltaTime, in bool IsProjection, in float3 PreviousPosition, inout float3 OutNodePosition )
+ void SolveHardCollisionConstraint(in bool EnableConstraint, in int StrandsSize, in float PenetrationDepth, in float3 CollisionPosition, in float3 CollisionVelocity, in float3 CollisionNormal, 
+				in float StaticFriction, in float KineticFriction, in bool IsProjection, in float DeltaTime )
 {
 	if( GGroupThreadId.x % StrandsSize > 1 && EnableConstraint)
 	{
-		const float3 NodeVelocity = (OutNodePosition - PreviousPosition)/DeltaTime;
-		const float3 DeltaPosition = OutNodePosition - CollisionPosition;
+		const float3 NodeVelocity = (SharedNodePosition[GGroupThreadId.x] - SharedPreviousPosition[GGroupThreadId.x])/DeltaTime;
+		const float3 DeltaPosition = SharedNodePosition[GGroupThreadId.x] - CollisionPosition;
 		const float3 DeltaVelocity = (NodeVelocity - CollisionVelocity) * DeltaTime;
 
-		const float3 PositionDisplacement = ComputeStaticCollisionDisplacement(DeltaPosition,DeltaVelocity,CollisionNormal,PenetrationDepth,StaticFriction,KineticFriction,IsProjection);
-		OutNodePosition -= PositionDisplacement;
+		UpdateHardCollisionConstraint(DeltaPosition,DeltaVelocity,CollisionNormal,PenetrationDepth,StaticFriction,KineticFriction,IsProjection);
+		GroupMemoryBarrier();
 	}
 }
 
+void ProjectHardCollisionConstraint(in bool EnableConstraint, in int StrandsSize, in float PenetrationDepth, in float3 CollisionPosition, in float3 CollisionVelocity, in float3 CollisionNormal, 
+				in float StaticFriction, in float KineticFriction, in float DeltaTime, out float3 OutNodePosition  )
+{
+	if(DeltaTime != 0.0)
+	{
+		SolveHardCollisionConstraint(EnableConstraint,StrandsSize,PenetrationDepth,CollisionPosition,CollisionVelocity,
+					CollisionNormal,StaticFriction,KineticFriction,true,DeltaTime);
+	}
+
+	GroupMemoryBarrier();
+	OutNodePosition = SharedNodePosition[GGroupThreadId.x];
+}
+
+
 /* -----------------------------------------------------------------
  * Static soft collision constraint
  * -----------------------------------------------------------------
  */
 
  void UpdateSoftCollisionConstraint(in float3 DeltaPosition, in float3 DeltaVelocity, in float3 CollisionNormal, in float PenetrationDepth,
-							  in float StaticFriction, in float KineticFriction, in float DeltaTime, in bool ProjectConstraint, in float MaterialDamping, 
-			in float MaterialCompliance, in float MaterialWeight, inout float3 OutMaterialMultiplier )
+							  in float StaticFriction, in float KineticFriction, in bool IsProjection, in float DeltaTime, in float MaterialDamping, 
+							  in float MaterialCompliance, in float MaterialWeight, inout float3 OutMaterialMultiplier )
 {
-	const float DeltaVNormal = dot(DeltaVelocity,CollisionNormal);
-	const float DeltaPNormal = dot(DeltaPosition,CollisionNormal)-PenetrationDepth;
 
-	if (DeltaPNormal < 0.0)
-	{
-		const float LocalCompliance = ((PenetrationDepth <= 0.0) || ( PenetrationDepth > 0.0 && DeltaPNormal < -PenetrationDepth)) ? 0.0 : MaterialCompliance;// * max(0.0,min(1.0,(1.0 + DeltaPNormal / PenetrationDepth)));
+	float DeltaFNormal = 0.0; 
 
-		const float DeltaFNormal = ProjectConstraint ? DeltaPNormal : DeltaVNormal;
-		const float NormDeltaFN = abs(DeltaFNormal);
+	float AlphaTangentA = 0.0;
+	float DeltaVTangentA = 0.0;
+	float3 CollisionTangentA = float3(0,0,0);
+	
+	float AlphaTangentB = 0.0;
+	float DeltaVTangentB = 0.0;
+	float3 CollisionTangentB = float3(0,0,0);
 
-		const float3 CollisionTangent = DeltaVelocity - DeltaVNormal * CollisionNormal;
-		const float TangentLength = length(CollisionTangent);
-
-		float3 CollisionTangentA = (TangentLength > 0.0) ? CollisionTangent/TangentLength : (abs(CollisionNormal.x-1.0f) > 0.0) ? float3(1,0,0) : 
-																					        (abs(CollisionNormal.y-1.0f) > 0.0) ? float3(0,1,0) :
-																						     float3(0,0,1);
-		float3 CollisionTangentB = cross(CollisionNormal,CollisionTangentA);
-		CollisionTangentA = cross(CollisionTangentB,CollisionNormal);
-
-		const float TangentLengthA = length(CollisionTangentA);
-		const float TangentLengthB = length(CollisionTangentB);
-
-		CollisionTangentA = (TangentLengthA > 0.0) ? CollisionTangentA/TangentLengthA : float3(0,0,0);
-		CollisionTangentB = (TangentLengthB > 0.0) ? CollisionTangentB/TangentLengthB : float3(0,0,0);
-
-		const float DeltaVTangentA = dot( DeltaVelocity, CollisionTangentA);
-		const float DeltaVTangentB = dot( DeltaVelocity, CollisionTangentB);
-
-		const float NormDeltaVTA = abs(DeltaVTangentA);
-		const float NormDeltaVTB = abs(DeltaVTangentB);
-
-		const float AlphaTangentA = ( NormDeltaVTA < StaticFriction * NormDeltaFN ) ? 1.0 : (KineticFriction > 0.0) ? max(0.0,min(KineticFriction*NormDeltaFN/NormDeltaVTA, 1.0)) : 0.0;
-		const float AlphaTangentB = ( NormDeltaVTB < StaticFriction * NormDeltaFN ) ? 1.0 : (KineticFriction > 0.0) ? max(0.0,min(KineticFriction*NormDeltaFN/NormDeltaVTB, 1.0)) : 0.0;
+	if( ComputeCollisionDatas(DeltaPosition,DeltaVelocity,CollisionNormal,PenetrationDepth,StaticFriction,KineticFriction,IsProjection,
+			DeltaFNormal,AlphaTangentA,DeltaVTangentA,CollisionTangentA,AlphaTangentB,DeltaVTangentB,CollisionTangentB))
+	{
+		const float LocalCompliance = ((PenetrationDepth <= 0.0) || ( PenetrationDepth > 0.0 && dot(DeltaPosition,CollisionNormal) < 0.0)) ? 0.0 : 
+							MaterialCompliance;// * max(0.0,min(1.0,(1.0 + DeltaPNormal / PenetrationDepth)));
 
 		const float ScaledInverseMass = ( (1.0  + MaterialDamping / DeltaTime ) * SharedInverseMass[GGroupThreadId.x] );
 
@@ -112,7 +135,7 @@
 		const float MaterialWeightB = ( SchurDiagonal.z != 0.0 ) ? 1.0 / SchurDiagonal.z : 0.0;
 
 		const float3 DeltaLambda = float3(
-		    -( DeltaFNormal + OutMaterialMultiplier.x * LocalCompliance + MaterialDamping * DeltaVNormal / DeltaTime ) * MaterialWeightN,
+		    -( DeltaFNormal + OutMaterialMultiplier.x * LocalCompliance + MaterialDamping * dot(DeltaVelocity,CollisionNormal) / DeltaTime ) * MaterialWeightN,
 			-( DeltaVTangentA + OutMaterialMultiplier.y * LocalCompliance + MaterialDamping * DeltaVTangentA / DeltaTime ) * MaterialWeightA * AlphaTangentA,
 			-( DeltaVTangentB + OutMaterialMultiplier.z * LocalCompliance + MaterialDamping * DeltaVTangentB / DeltaTime ) * MaterialWeightB * AlphaTangentB);
 
@@ -137,65 +160,44 @@
 	if( LocalIndex > 1 )
 	{
 		OutMaterialCompliance = (ConstraintStiffness > 0.0) ? 1.0/(ConstraintStiffness*DeltaTime*DeltaTime) : 0.0;
-
-		//const float SumInverseMass = SharedInverseMass[GGroupThreadId.x];
-		//const float SchurDiagonal = ( (1.0  + MaterialDamping / DeltaTime ) * SumInverseMass + OutMaterialCompliance );
-		//OutMaterialWeight = ( SchurDiagonal != 0.0 ) ? 1.0 / SchurDiagonal : 0.0;
 	}
 }
 
 void SolveSoftCollisionConstraint(in bool EnableConstraint, in int StrandsSize, in float PenetrationDepth, in float3 CollisionPosition, in float3 CollisionVelocity, in float3 CollisionNormal, 
-				in float StaticFriction, in float KineticFriction, in float DeltaTime, in float MaterialDamping, 
+				in float StaticFriction, in float KineticFriction, in bool IsProjection, in float DeltaTime, in float MaterialDamping, 
 			in float MaterialCompliance, in float MaterialWeight, in float3 MaterialMultiplier, out float3 OutMaterialMultiplier )
 {
 	OutMaterialMultiplier = MaterialMultiplier;
-	if(EnableConstraint)
+	if( GGroupThreadId.x % StrandsSize > 1 && EnableConstraint)
 	{
-		const int LocalIndex = (GGroupThreadId.x % StrandsSize);
-		if( LocalIndex > 1)
-		{
-			const float3 NodeVelocity = (SharedNodePosition[GGroupThreadId.x] - SharedPreviousPosition[GGroupThreadId.x])/DeltaTime;
-			const float3 DeltaPosition = SharedNodePosition[GGroupThreadId.x] - CollisionPosition;
-			const float3 DeltaVelocity = (NodeVelocity - CollisionVelocity) * DeltaTime;
-
-			UpdateSoftCollisionConstraint(DeltaPosition,DeltaVelocity,CollisionNormal,PenetrationDepth,
-					StaticFriction,KineticFriction,DeltaTime,false,MaterialDamping,MaterialCompliance,MaterialWeight,OutMaterialMultiplier);
-
-			GroupMemoryBarrier();
-		}
+		const float3 NodeVelocity = (SharedNodePosition[GGroupThreadId.x] - SharedPreviousPosition[GGroupThreadId.x])/DeltaTime;
+		const float3 DeltaPosition = SharedNodePosition[GGroupThreadId.x] - CollisionPosition;
+		const float3 DeltaVelocity = (NodeVelocity - CollisionVelocity) * DeltaTime;
+
+		const float LocalDamping = IsProjection ? 0.0 : MaterialDamping;
+		UpdateSoftCollisionConstraint(DeltaPosition,DeltaVelocity,CollisionNormal,PenetrationDepth,
+				StaticFriction,KineticFriction,IsProjection,DeltaTime,LocalDamping,MaterialCompliance,MaterialWeight,OutMaterialMultiplier);
+
+		GroupMemoryBarrier();
 	}
 }
 
-void ProjectSoftCollisionConstraint(in int StrandsSize, in float ConstraintStiffness, in float PenetrationDepth, in float3 CollisionPosition, in float3 CollisionVelocity, in float3 CollisionNormal, 
-				in float StaticFriction, in float KineticFriction, in float NodeMass, in float DeltaTime, in float3 PreviousPosition, in float3 NodePosition, out float3 OutNodePosition  )
+void ProjectSoftCollisionConstraint(in bool EnableConstraint, in int StrandsSize, in float ConstraintStiffness, in float PenetrationDepth, in float3 CollisionPosition, in float3 CollisionVelocity, in float3 CollisionNormal, 
+				in float StaticFriction, in float KineticFriction, in float DeltaTime, in float NodeMass, out float3 OutNodePosition  )
 {
-	const int LocalIndex = (GGroupThreadId.x % StrandsSize);
-
-	SharedNodePosition[GGroupThreadId.x] = NodePosition;
-	SharedPreviousPosition[GGroupThreadId.x] = PreviousPosition;
-	SharedInverseMass[GGroupThreadId.x] = (LocalIndex > 1 && NodeMass != 0.0) ? 1.0 / NodeMass : 0.0;
-	GroupMemoryBarrier();
+	float MaterialCompliance = 0.0;
+	float MaterialWeight = 0.0;
+	float3 MaterialMultiplier = float3(0.0,0.0,0.0);
+	float3 OutMaterialMultiplier = float3(0.0,0.0,0.0);
 
 	if(DeltaTime != 0.0)
 	{
-		float MaterialCompliance = 0.0;
-		float MaterialWeight = 0.0;
-		float3 MaterialMultiplier = float3(0.0,0.0,0.0);
-
-		if( LocalIndex > 1)
-		{
-			SetupSoftCollisionConstraint(StrandsSize,ConstraintStiffness,DeltaTime,0.0,MaterialCompliance,MaterialWeight,MaterialMultiplier);
+		SetupSoftCollisionConstraint(StrandsSize,ConstraintStiffness,DeltaTime,0.0,MaterialCompliance,MaterialWeight,MaterialMultiplier);
 
-			const float3 NodeVelocity = (SharedNodePosition[GGroupThreadId.x] - SharedPreviousPosition[GGroupThreadId.x])/DeltaTime;
-			const float3 DeltaPosition = SharedNodePosition[GGroupThreadId.x] - CollisionPosition;
-			const float3 DeltaVelocity = (NodeVelocity - CollisionVelocity) * DeltaTime;
-
-			UpdateSoftCollisionConstraint(DeltaPosition,DeltaVelocity,CollisionNormal,PenetrationDepth,
-					StaticFriction,KineticFriction,DeltaTime,true,0.0,MaterialCompliance,MaterialWeight,MaterialMultiplier);
-
-			GroupMemoryBarrier();
-		}
+		SolveSoftCollisionConstraint(EnableConstraint,StrandsSize,PenetrationDepth,CollisionPosition,CollisionVelocity,
+			CollisionNormal,StaticFriction,KineticFriction,true,DeltaTime,0.0,MaterialCompliance,MaterialWeight,MaterialMultiplier,OutMaterialMultiplier);
 	}
+
 	GroupMemoryBarrier();
 	OutNodePosition = SharedNodePosition[GGroupThreadId.x];
 }
diff --git a/Engine/Plugins/Experimental/HairStrands/Source/HairStrandsCore/Private/GroomAsset.cpp b/Engine/Plugins/Experimental/HairStrands/Source/HairStrandsCore/Private/GroomAsset.cpp
--- a/Engine/Plugins/Experimental/HairStrands/Source/HairStrandsCore/Private/GroomAsset.cpp
+++ b/Engine/Plugins/Experimental/HairStrands/Source/HairStrandsCore/Private/GroomAsset.cpp
@@ -16,6 +16,7 @@
 #include "Serialization/MemoryWriter.h"
 #include "UObject/PhysicsObjectVersion.h"
 #include "UObject/ReleaseObjectVersion.h"
+#include "NiagaraSystem.h"
 
 #if WITH_EDITORONLY_DATA
 #include "DerivedDataCacheInterface.h"
@@ -357,16 +358,126 @@ void FHairStrandsClusterCullingResource::ReleaseRHI()
 
 /////////////////////////////////////////////////////////////////////////////////////////
 
+FHairStrandsRootResource::FHairStrandsRootResource()
+{
+
+}
+
+FHairStrandsRootResource::FHairStrandsRootResource(const FHairStrandsRootData& InRootData): 
+RootData(InRootData)
+{
+	PopulateFromRootData();
+}
+
 FHairStrandsRootResource::FHairStrandsRootResource(const FHairStrandsDatas* HairStrandsDatas, uint32 LODCount):
+	RootData(HairStrandsDatas, LODCount)
+{
+	PopulateFromRootData();
+}
+
+void FHairStrandsRootResource::PopulateFromRootData()
+{
+	uint32 LODIndex = 0;
+	for (FHairStrandsRootData::FMeshProjectionLOD& MeshProjectionLOD : RootData.MeshProjectionLODs)
+	{
+		FMeshProjectionLOD& LOD = MeshProjectionLODs.AddDefaulted_GetRef();
+
+		LOD.LODIndex = MeshProjectionLOD.LODIndex;
+		LOD.Status = FHairStrandsProjectionHairData::LODData::EStatus::Invalid;
+	}
+}
+
+void FHairStrandsRootResource::InitRHI()
+{
+	if (RootData.VertexToCurveIndexBuffer.Num() > 0)
+	{
+		CreateBuffer<FHairStrandsIndexFormat>(RootData.VertexToCurveIndexBuffer, VertexToCurveIndexBuffer);
+		CreateBuffer<FHairStrandsRootPositionFormat>(RootData.RootPositionBuffer, RootPositionBuffer);
+		CreateBuffer<FHairStrandsRootNormalFormat>(RootData.RootNormalBuffer, RootNormalBuffer);
+		
+		check(MeshProjectionLODs.Num() == RootData.MeshProjectionLODs.Num());
+		for (uint32 LODIt=0, LODCount = MeshProjectionLODs.Num(); LODIt<LODCount; ++LODIt)
+		{
+			FMeshProjectionLOD& GPUData = MeshProjectionLODs[LODIt];
+			const FHairStrandsRootData::FMeshProjectionLOD& CPUData = RootData.MeshProjectionLODs[LODIt];
+
+			const bool bHasValidCPUData = CPUData.RootTriangleBarycentricBuffer.Num() > 0;
+			if (bHasValidCPUData)
+			{
+				GPUData.Status = FHairStrandsProjectionHairData::LODData::EStatus::Completed;
+
+				check(CPUData.RootTriangleBarycentricBuffer.Num() > 0);
+				CreateBuffer<FHairStrandsCurveTriangleBarycentricFormat>(CPUData.RootTriangleBarycentricBuffer, GPUData.RootTriangleBarycentricBuffer);
+
+				check(CPUData.RootTriangleIndexBuffer.Num() > 0);
+				CreateBuffer<FHairStrandsCurveTriangleIndexFormat>(CPUData.RootTriangleIndexBuffer, GPUData.RootTriangleIndexBuffer);
+
+				check(CPUData.RestRootTrianglePosition0Buffer.Num() > 0);
+				check(CPUData.RestRootTrianglePosition1Buffer.Num() > 0);
+				check(CPUData.RestRootTrianglePosition2Buffer.Num() > 0);
+				CreateBuffer<FHairStrandsMeshTrianglePositionFormat>(CPUData.RestRootTrianglePosition0Buffer, GPUData.RestRootTrianglePosition0Buffer);
+				CreateBuffer<FHairStrandsMeshTrianglePositionFormat>(CPUData.RestRootTrianglePosition1Buffer, GPUData.RestRootTrianglePosition1Buffer);
+				CreateBuffer<FHairStrandsMeshTrianglePositionFormat>(CPUData.RestRootTrianglePosition2Buffer, GPUData.RestRootTrianglePosition2Buffer);
+			}
+			else
+			{
+				GPUData.Status = FHairStrandsProjectionHairData::LODData::EStatus::Initialized;
+
+				CreateBuffer<FHairStrandsCurveTriangleBarycentricFormat>(RootData.RootCount, GPUData.RootTriangleBarycentricBuffer);
+				CreateBuffer<FHairStrandsCurveTriangleIndexFormat>(RootData.RootCount, GPUData.RootTriangleIndexBuffer);
+
+				// Create buffers. Initialization will be done by render passes
+				CreateBuffer<FHairStrandsMeshTrianglePositionFormat>(RootData.RootCount, GPUData.RestRootTrianglePosition0Buffer);
+				CreateBuffer<FHairStrandsMeshTrianglePositionFormat>(RootData.RootCount, GPUData.RestRootTrianglePosition1Buffer);
+				CreateBuffer<FHairStrandsMeshTrianglePositionFormat>(RootData.RootCount, GPUData.RestRootTrianglePosition2Buffer);
+			}
+
+			// Strand hair roots translation and rotation in triangle-deformed position relative to the bound triangle 
+			CreateBuffer<FHairStrandsMeshTrianglePositionFormat>(RootData.RootCount, GPUData.DeformedRootTrianglePosition0Buffer);
+			CreateBuffer<FHairStrandsMeshTrianglePositionFormat>(RootData.RootCount, GPUData.DeformedRootTrianglePosition1Buffer);
+			CreateBuffer<FHairStrandsMeshTrianglePositionFormat>(RootData.RootCount, GPUData.DeformedRootTrianglePosition2Buffer);
+		}
+	}
+}
+
+void FHairStrandsRootResource::ReleaseRHI()
+{
+	RootPositionBuffer.Release();
+	RootNormalBuffer.Release();
+	VertexToCurveIndexBuffer.Release();
+
+	for (FMeshProjectionLOD& GPUData : MeshProjectionLODs)
+	{
+		GPUData.Status = FHairStrandsProjectionHairData::LODData::EStatus::Invalid;
+		GPUData.RootTriangleIndexBuffer.Release();
+		GPUData.RootTriangleBarycentricBuffer.Release();
+		GPUData.RestRootTrianglePosition0Buffer.Release();
+		GPUData.RestRootTrianglePosition1Buffer.Release();
+		GPUData.RestRootTrianglePosition2Buffer.Release();
+		GPUData.DeformedRootTrianglePosition0Buffer.Release();
+		GPUData.DeformedRootTrianglePosition1Buffer.Release();
+		GPUData.DeformedRootTrianglePosition2Buffer.Release();
+	}
+	MeshProjectionLODs.Empty();
+}
+
+/////////////////////////////////////////////////////////////////////////////////////////
+
+FHairStrandsRootData::FHairStrandsRootData()
+{
+
+}
+
+FHairStrandsRootData::FHairStrandsRootData(const FHairStrandsDatas* HairStrandsDatas, uint32 LODCount):
 	RootCount(HairStrandsDatas ? HairStrandsDatas->GetNumCurves() : 0)
 {
 	if (!HairStrandsDatas)
 		return;
 
 	const uint32 CurveCount = HairStrandsDatas->GetNumCurves();
-	CurveIndices.SetNum(HairStrandsDatas->GetNumPoints());
-	RootPositions.SetNum(RootCount);
-	RootNormals.SetNum(RootCount);
+	VertexToCurveIndexBuffer.SetNum(HairStrandsDatas->GetNumPoints());
+	RootPositionBuffer.SetNum(RootCount);
+	RootNormalBuffer.SetNum(RootCount);
 
 	for (uint32 CurveIndex = 0; CurveIndex < CurveCount; ++CurveIndex)
 	{
@@ -374,13 +485,13 @@ FHairStrandsRootResource::FHairStrandsRootResource(const FHairStrandsDatas* Hair
 		const uint32 PointCount = HairStrandsDatas->StrandsCurves.CurvesCount[CurveIndex];
 		for (uint32 PointIndex = 0; PointIndex < PointCount; ++PointIndex)
 		{
-			CurveIndices[RootIndex + PointIndex] = CurveIndex; // RootIndex;
+			VertexToCurveIndexBuffer[RootIndex + PointIndex] = CurveIndex; // RootIndex;
 		}
 
 		check(PointCount > 1);
 
 		const FVector P0 = HairStrandsDatas->StrandsPoints.PointsPosition[RootIndex];
-		const FVector P1 = HairStrandsDatas->StrandsPoints.PointsPosition[RootIndex+1];
+		const FVector P1 = HairStrandsDatas->StrandsPoints.PointsPosition[RootIndex + 1];
 		FVector N0 = (P1 - P0).GetSafeNormal();
 
 		// Fallback in case the initial points are too close (this happens on certain assets)
@@ -401,64 +512,65 @@ FHairStrandsRootResource::FHairStrandsRootResource(const FHairStrandsDatas* Hair
 		N.Z = N0.Z;
 		N.W = 0;
 
-		RootPositions[CurveIndex] = P;
-		RootNormals[CurveIndex] = N;
+		RootPositionBuffer[CurveIndex] = P;
+		RootNormalBuffer[CurveIndex] = N;
 	}
-	
+
 	MeshProjectionLODs.SetNum(LODCount);
 	uint32 LODIndex = 0;
 	for (FMeshProjectionLOD& MeshProjectionLOD : MeshProjectionLODs)
 	{
-		MeshProjectionLOD.Status = FHairStrandsProjectionHairData::LODData::EStatus::Invalid;
-		MeshProjectionLOD.RestRootOffset = HairStrandsDatas->BoundingBox.GetCenter();
 		MeshProjectionLOD.LODIndex = LODIndex++;
 	}
 }
 
-void FHairStrandsRootResource::InitRHI()
+bool FHairStrandsRootData::HasProjectionData() const
 {
-	if (CurveIndices.Num() > 0)
+	bool bIsValid = MeshProjectionLODs.Num() > 0;
+	for (const FMeshProjectionLOD& LOD : MeshProjectionLODs)
 	{
-		CreateBuffer<FHairStrandsIndexFormat>(CurveIndices, VertexToCurveIndexBuffer);
-		CreateBuffer<FHairStrandsRootPositionFormat>(RootPositions, RootPositionBuffer);
-		CreateBuffer<FHairStrandsRootNormalFormat>(RootNormals, RootNormalBuffer);	
-		for (FMeshProjectionLOD& LOD : MeshProjectionLODs)
+		const bool bHasValidCPUData = LOD.RootTriangleBarycentricBuffer.Num() > 0;
+		if (bHasValidCPUData)
 		{
-			LOD.Status = FHairStrandsProjectionHairData::LODData::EStatus::Initialized;
-
-			// Create buffers. Initialization will be done by render passes
-			CreateBuffer<FHairStrandsCurveTriangleIndexFormat>(RootCount, LOD.RootTriangleIndexBuffer);
-			CreateBuffer<FHairStrandsCurveTriangleBarycentricFormat>(RootCount, LOD.RootTriangleBarycentricBuffer);
-			CreateBuffer<FHairStrandsMeshTrianglePositionFormat>(RootCount, LOD.RestRootTrianglePosition0Buffer);
-			CreateBuffer<FHairStrandsMeshTrianglePositionFormat>(RootCount, LOD.RestRootTrianglePosition1Buffer);
-			CreateBuffer<FHairStrandsMeshTrianglePositionFormat>(RootCount, LOD.RestRootTrianglePosition2Buffer);
+			bIsValid = bIsValid && LOD.RootTriangleBarycentricBuffer.Num() > 0;
+			bIsValid = bIsValid && LOD.RootTriangleIndexBuffer.Num() > 0;
+			bIsValid = bIsValid && LOD.RestRootTrianglePosition0Buffer.Num() > 0;
+			bIsValid = bIsValid && LOD.RestRootTrianglePosition1Buffer.Num() > 0;
+			bIsValid = bIsValid && LOD.RestRootTrianglePosition2Buffer.Num() > 0;
 
-			// Strand hair roots translation and rotation in triangle-deformed position relative to the bound triangle 
-			CreateBuffer<FHairStrandsMeshTrianglePositionFormat>(RootCount, LOD.DeformedRootTrianglePosition0Buffer);
-			CreateBuffer<FHairStrandsMeshTrianglePositionFormat>(RootCount, LOD.DeformedRootTrianglePosition1Buffer);
-			CreateBuffer<FHairStrandsMeshTrianglePositionFormat>(RootCount, LOD.DeformedRootTrianglePosition2Buffer);
+			if (!bIsValid) break;
 		}
 	}
+
+	return bIsValid;
 }
 
-void FHairStrandsRootResource::ReleaseRHI()
+FArchive& operator<<(FArchive& Ar, FHairStrandsRootData::FMeshProjectionLOD& LOD)
 {
-	RootPositionBuffer.Release();
-	RootNormalBuffer.Release();
-	VertexToCurveIndexBuffer.Release();
+	Ar << LOD.LODIndex;
+	Ar << LOD.RootTriangleIndexBuffer;
+	Ar << LOD.RootTriangleBarycentricBuffer;
+	Ar << LOD.RestRootTrianglePosition0Buffer;
+	Ar << LOD.RestRootTrianglePosition1Buffer;
+	Ar << LOD.RestRootTrianglePosition2Buffer;
+	return Ar;
+}
 
-	for (FMeshProjectionLOD& LOD : MeshProjectionLODs)
-	{
-		LOD.Status = FHairStrandsProjectionHairData::LODData::EStatus::Invalid;
-		LOD.RootTriangleIndexBuffer.Release();
-		LOD.RootTriangleBarycentricBuffer.Release();
-		LOD.RestRootTrianglePosition0Buffer.Release();
-		LOD.RestRootTrianglePosition1Buffer.Release();
-		LOD.RestRootTrianglePosition2Buffer.Release();
-		LOD.DeformedRootTrianglePosition0Buffer.Release();
-		LOD.DeformedRootTrianglePosition1Buffer.Release();
-		LOD.DeformedRootTrianglePosition2Buffer.Release();
-	}
+void FHairStrandsRootData::Serialize(FArchive& Ar)
+{
+	Ar << RootCount;
+	Ar << VertexToCurveIndexBuffer;
+	Ar << RootPositionBuffer;
+	Ar << RootNormalBuffer;
+	Ar << MeshProjectionLODs;
+}
+
+void FHairStrandsRootData::Reset()
+{
+	RootCount = 0;
+	VertexToCurveIndexBuffer.Empty();
+	RootPositionBuffer.Empty();
+	RootNormalBuffer.Empty();
 	MeshProjectionLODs.Empty();
 }
 
@@ -563,6 +675,57 @@ void UGroomAsset::Serialize(FArchive& Ar)
 	}
 }
 
+UGroomAsset::UGroomAsset(const FObjectInitializer& ObjectInitializer)
+	: Super(ObjectInitializer)
+{
+	HairToGuideDensity = 0.1f;
+
+	SimulatedGroup = 0;
+	EnableSimulation = false;
+	NiagaraSolver = EGroomNiagaraSolvers::AngularSprings;
+
+	SubSteps = 5;
+	IterationCount = 20;
+
+	GravityVector = FVector(0.0, 0.0, -981.0);
+	AirDrag = 0.1;
+	AirVelocity = FVector(0, 0, 0);
+
+	SolveBend = true;
+	ProjectBend = false;
+	BendDamping = 0.01;
+	BendStiffness = 0.01;
+
+	BendScale.GetRichCurve()->SetKeyInterpMode(BendScale.GetRichCurve()->AddKey(0.f, 1.f), ERichCurveInterpMode::RCIM_Cubic);
+	BendScale.GetRichCurve()->SetKeyInterpMode(BendScale.GetRichCurve()->AddKey(1.f, 1.f), ERichCurveInterpMode::RCIM_Cubic);
+
+	SolveStretch = true;
+	ProjectStretch = false;
+	StretchDamping = 0.01;
+	StretchStiffness = 1.0;
+
+	StretchScale.GetRichCurve()->SetKeyInterpMode(StretchScale.GetRichCurve()->AddKey(0.f, 1.f), ERichCurveInterpMode::RCIM_Cubic);
+	StretchScale.GetRichCurve()->SetKeyInterpMode(StretchScale.GetRichCurve()->AddKey(1.f, 1.f), ERichCurveInterpMode::RCIM_Cubic);
+
+	SolveCollision = true;
+	ProjectCollision = true;
+	KineticFriction = 0.1;
+	StaticFriction = 0.1;
+	StrandsViscosity = 1.0;
+	CollisionRadius = 1.0;
+
+	RadiusScale.GetRichCurve()->SetKeyInterpMode(RadiusScale.GetRichCurve()->AddKey(0.f, 1.0f), ERichCurveInterpMode::RCIM_Cubic);
+	RadiusScale.GetRichCurve()->SetKeyInterpMode(RadiusScale.GetRichCurve()->AddKey(1.f, 0.1f), ERichCurveInterpMode::RCIM_Cubic);
+
+	StrandsSize = EGroomStrandsSize::Size16;
+	StrandsDensity = 1.0;
+	StrandsSmoothing = 0.1;
+	StrandsThickness = 0.01;
+
+	ThicknessScale.GetRichCurve()->SetKeyInterpMode(ThicknessScale.GetRichCurve()->AddKey(0.f, 1.0f), ERichCurveInterpMode::RCIM_Cubic);
+	ThicknessScale.GetRichCurve()->SetKeyInterpMode(ThicknessScale.GetRichCurve()->AddKey(1.f, 1.0f), ERichCurveInterpMode::RCIM_Cubic);
+}
+
 void UGroomAsset::InitResource()
 {
 	for (FHairGroupData& GroupData : HairGroupsData)
@@ -976,3 +1139,130 @@ bool UGroomAsset::CacheDerivedData(const FGroomBuildSettings* InBuildSettings)
 }
 
 #endif // WITH_EDITORONLY_DATA
+
+/////////////////////////////////////////////////////////////////////////////////////////
+
+static void InitializeGroupInfos(
+	const TArray<UGroomBindingAsset::FHairGroupData>& HairGroupDatas,
+	TArray<FGoomBindingGroupInfo>& GroupInfos)
+{
+	GroupInfos.Empty();
+	for (const UGroomBindingAsset::FHairGroupData& Data : HairGroupDatas)
+	{
+		FGoomBindingGroupInfo& Info = GroupInfos.AddDefaulted_GetRef();
+		Info.SimRootCount	= Data.SimRootData.RootCount;
+		Info.SimLODCount	= Data.SimRootData.MeshProjectionLODs.Num();
+		Info.RenRootCount	= Data.RenRootData.RootCount;
+		Info.RenLODCount	= Data.RenRootData.MeshProjectionLODs.Num();
+	}
+}
+
+FArchive& operator<<(FArchive& Ar, UGroomBindingAsset::FHairGroupData& GroupData)
+{
+	GroupData.SimRootData.Serialize(Ar);
+	GroupData.RenRootData.Serialize(Ar);
+	return Ar;
+}
+
+void UGroomBindingAsset::Serialize(FArchive& Ar)
+{
+	Super::Serialize(Ar);
+	Ar << HairGroupDatas;
+}
+
+void UGroomBindingAsset::InitResource()
+{
+	for (FHairGroupData& Data : HairGroupDatas)
+	{
+		FHairGroupResource& Resource = HairGroupResources.AddDefaulted_GetRef();
+		Resource.SimRootResources = new FHairStrandsRootResource(Data.SimRootData);
+		Resource.RenRootResources = new FHairStrandsRootResource(Data.RenRootData);
+
+		BeginInitResource(Resource.SimRootResources);
+		BeginInitResource(Resource.RenRootResources);
+	}
+}
+
+void UGroomBindingAsset::UpdateResource()
+{
+	for (FHairGroupResource& Resource : HairGroupResources)
+	{
+		BeginUpdateResourceRHI(Resource.SimRootResources);
+		BeginUpdateResourceRHI(Resource.RenRootResources);
+	}
+}
+
+void UGroomBindingAsset::ReleaseResource()
+{
+	// Delay destruction to insure that the rendering thread is done with all resources usage
+	if (HairGroupResources.Num() > 0)
+	{
+		for (FHairGroupResource& Resource : HairGroupResources)
+		{
+			FHairStrandsRootResource* InSimRootResources = Resource.SimRootResources;
+			FHairStrandsRootResource* InRenRootResources = Resource.RenRootResources;
+			ENQUEUE_RENDER_COMMAND(ReleaseHairStrandsResourceCommand)(
+				[InSimRootResources, InRenRootResources](FRHICommandList& RHICmdList)
+			{
+				InSimRootResources->ReleaseResource();
+				InRenRootResources->ReleaseResource();
+				delete InSimRootResources;
+				delete InRenRootResources;
+			});
+			Resource.SimRootResources = nullptr;
+			Resource.RenRootResources = nullptr;
+		}
+		HairGroupResources.Empty();
+	}
+}
+
+void UGroomBindingAsset::Reset()
+{
+	ReleaseResource();
+	for (FHairGroupData& Data : HairGroupDatas)
+	{
+		Data.SimRootData.Reset();
+		Data.RenRootData.Reset();
+	}
+}
+
+void UGroomBindingAsset::PostLoad()
+{
+	Super::PostLoad();
+
+	InitializeGroupInfos(HairGroupDatas, GroupInfos);
+
+	if (!IsTemplate())
+	{
+		InitResource();
+	}
+}
+
+void UGroomBindingAsset::PreSave(const class ITargetPlatform* TargetPlatform)
+{
+	Super::PreSave(TargetPlatform);
+	InitializeGroupInfos(HairGroupDatas, GroupInfos);
+}
+
+void UGroomBindingAsset::PostSaveRoot(bool bCleanupIsRequired)
+{
+	Super::PostSaveRoot(bCleanupIsRequired);
+	InitializeGroupInfos(HairGroupDatas, GroupInfos);
+}
+
+void UGroomBindingAsset::BeginDestroy()
+{
+	ReleaseResource();
+	Super::BeginDestroy();
+}
+
+#if WITH_EDITOR
+void UGroomBindingAsset::PostEditChangeProperty(FPropertyChangedEvent& PropertyChangedEvent)
+{
+	Super::PostEditChangeProperty(PropertyChangedEvent);
+	UpdateResource();
+	OnGroomBindingAssetChanged.Broadcast();
+}
+#endif // WITH_EDITOR
+
+
