Copying //UE4/Dev-Physics to Main (//UE4/Main)

diff --git a/Engine/Plugins/Experimental/HairStrands/Source/HairStrandsCore/Private/GroomComponent.cpp b/Engine/Plugins/Experimental/HairStrands/Source/HairStrandsCore/Private/GroomComponent.cpp
--- a/Engine/Plugins/Experimental/HairStrands/Source/HairStrandsCore/Private/GroomComponent.cpp
+++ b/Engine/Plugins/Experimental/HairStrands/Source/HairStrandsCore/Private/GroomComponent.cpp
@@ -24,9 +24,6 @@ static float GHairClipLength = -1;
 static FAutoConsoleVariableRef CVarHairClipLength(TEXT("r.HairStrands.DebugClipLength"), GHairClipLength, TEXT("Clip hair strands which have a lenth larger than this value. (default is -1, no effect)"));
 float GetHairClipLength() { return GHairClipLength > 0 ? GHairClipLength : 100000;  }
 
-static int32 GHairStrandsMeshUseRelativePosition = 1;
-static FAutoConsoleVariableRef CVarHairStrandsMeshUseRelativePosition(TEXT("r.HairStrands.MeshProjection.RelativePosition"), GHairStrandsMeshUseRelativePosition, TEXT("Enable/Disable relative triangle position for improving positions"));
-
 #define LOCTEXT_NAMESPACE "GroomComponent"
 
 /**
@@ -108,6 +105,64 @@ public:
 	}
 };
 
+static bool IsCompatible(UGroomAsset* InGroom, UGroomBindingAsset* InBinding)
+{
+	if (InBinding && InGroom)
+	{
+		if (InGroom->HairGroupsInfo.Num() != InBinding->GroupInfos.Num())
+		{
+			UE_LOG(LogHairStrands, Warning, TEXT("[Groom] The GroomBinding asset (Groom:%s) does not contains the same number of groups. It is incompatible with the groom asset assigned to this component (Groom:%s). The binding asset will be set to null."), *InBinding->Groom->GetName(), *InGroom->GetName());
+			return false;
+		}
+		for (uint32 GroupIt=0, GroupCount=InGroom->HairGroupsInfo.Num(); GroupIt < GroupCount; ++GroupIt)
+		{
+			if (InGroom->HairGroupsInfo[GroupIt].NumCurves != InBinding->GroupInfos[GroupIt].RenRootCount ||
+				InGroom->HairGroupsInfo[GroupIt].NumGuides != InBinding->GroupInfos[GroupIt].SimRootCount)
+			{
+				UE_LOG(LogHairStrands, Warning, TEXT("[Groom] The GroomBinding asset (Groom:%s) does not contains the same curves/guides in group %d. It is incompatible with the groom asset assigned to this component (Groom:%s). The binding asset will be set to null."), *InBinding->Groom->GetName(), GroupIt, *InGroom->GetName());
+				return false;
+			}
+		}
+
+		if (InBinding->Groom != InGroom)
+		{
+			UE_LOG(LogHairStrands, Warning, TEXT("[Groom] The GroomBinding asset (Groom:%s) is incompatible with the groom asset assigned to this component (Groom:%s). The binding asset will be set to null."), *InBinding->Groom->GetName(), *InGroom->GetName());
+			return false;
+		}
+	}
+	return true;
+}
+
+static bool IsBindingAssetValid(UGroomBindingAsset* InBinding)
+{
+	if (InBinding)
+	{
+		if (const UPackage* Package = InBinding->GetOutermost())
+		{
+			if (Package->IsDirty())
+			{
+				UE_LOG(LogHairStrands, Warning, TEXT("[Groom] The binding asset (%s) is not saved and will be considered as invalid. Falling back onto non-binding version."), *InBinding->GetName());
+				return false;
+			}
+		}
+
+		if (InBinding->GroupInfos.Num() == 0)
+		{
+			UE_LOG(LogHairStrands, Warning, TEXT("[Groom] The GroomBinding asset (Groom:%s) does not contain any groups. It is invalid and can't be assigned. The binding asset will be set to null."), *InBinding->Groom->GetName());
+			return false;
+		}
+
+		for (FGoomBindingGroupInfo& Info : InBinding->GroupInfos)
+		{
+			if (Info.RenRootCount == 0 || Info.SimRootCount == 0)
+			{
+				UE_LOG(LogHairStrands, Warning, TEXT("[Groom] The GroomBinding asset (Groom:%s) has group with 0 curves/guides. It is invalid and can't be assigned. The binding asset will be set to null."), *InBinding->Groom->GetName());
+				return false;
+			}
+		}
+	}
+	return true;
+}
 
 /////////////////////////////////////////////////////////////////////////////////////////
 /////////////////////////////////////////////////////////////////////////////////////////
@@ -138,14 +193,17 @@ public:
 		FHairStrandsVertexFactory::FDataType VFData;
 		VFData.InterpolationOutput = Component->InterpolationOutput;
 
+		check(Component->HairGroupResources);
+		check(Component->HairGroupResources->HairGroups.Num());
+
 		const uint32 GroupCount = Component->GroomAsset->GetNumHairGroups();
-		check(Component->GroomAsset->HairGroupsData.Num() == Component->HairGroupResources.Num());
+		check(Component->GroomAsset->HairGroupsData.Num() == Component->HairGroupResources->HairGroups.Num());
 		for (uint32 GroupIt=0;GroupIt<GroupCount; GroupIt++)
 		{		
 			const FHairGroupData& InGroupData = Component->GroomAsset->HairGroupsData[GroupIt];
 			const FHairGroupDesc& InGroupDesc = Component->GroomGroupsDesc[GroupIt];
 
-			const UGroomComponent::FHairGroupResource& GroupResources = Component->HairGroupResources[GroupIt];
+			const UGroomComponent::FHairGroupResource& GroupResources = Component->HairGroupResources->HairGroups[GroupIt];
 
 			UMaterialInterface* Material = Component->GetMaterial(GroupIt);
 			if (Material == nullptr || !Material->GetMaterialResource(GetScene().GetFeatureLevel())->IsUsedWithHairStrands())
@@ -269,6 +327,10 @@ public:
 		
 		for (int32 ViewIndex = 0; ViewIndex < Views.Num(); ViewIndex++)
 		{
+			const FSceneView* View = Views[ViewIndex];
+			if (View->bIsReflectionCapture || View->bIsPlanarReflection)
+				continue;
+
 			if (VisibilityMap & (1 << ViewIndex))
 			{
 				for (uint32 GroupIt = 0; GroupIt < GroupCount; ++GroupIt)
@@ -376,7 +438,33 @@ private:
 	TArray<HairGroup> HairGroups;
 };
 
-static void UpdateHairGroupsDesc(UGroomAsset* GroomAsset, TArray<FHairGroupDesc>& GroomGroupsDesc)
+/////////////////////////////////////////////////////////////////////////////////////////
+/////////////////////////////////////////////////////////////////////////////////////////
+// UComponent
+//////////////////////////////////////////////////////////////////////////
+//////////////////////////////////////////////////////////////////////////
+
+UGroomComponent::UGroomComponent(const FObjectInitializer& ObjectInitializer)
+	: Super(ObjectInitializer)
+{
+	PrimaryComponentTick.bCanEverTick = true;
+	bTickInEditor = true;
+	bAutoActivate = true;
+	bSelectable = true;
+	RegisteredSkeletalMeshComponent = nullptr;
+	SkeletalPreviousPositionOffset = FVector::ZeroVector;
+	bBindGroomToSkeletalMesh = false;
+	bCreateNiagaraComponent = false;
+	InitializedResources = nullptr;
+	Mobility = EComponentMobility::Movable;
+	bIsGroomAssetCallbackRegistered = false;
+	SourceSkeletalMesh = nullptr; 
+	NiagaraComponent = nullptr;
+
+	SetCollisionProfileName(UCollisionProfile::PhysicsActor_ProfileName);
+}
+
+void UGroomComponent::UpdateHairGroupsDesc()
 {
 	if (!GroomAsset)
 	{
@@ -399,8 +487,8 @@ static void UpdateHairGroupsDesc(UGroomAsset* GroomAsset, TArray<FHairGroupDesc>
 		FHairGroupData& GroupData = GroomAsset->HairGroupsData[GroupIt];
 
 		FHairGroupDesc& Desc = GroomGroupsDesc[GroupIt];
-		Desc.GuideCount	= GroupInfo.NumGuides;
-		Desc.HairCount	= GroupInfo.NumCurves;
+		Desc.GuideCount = GroupInfo.NumGuides;
+		Desc.HairCount = GroupInfo.NumCurves;
 		Desc.HairLength = GroupData.HairRenderData.StrandsCurves.MaxLength;
 
 		if (bReinitOverride || Desc.HairWidth == 0)
@@ -417,29 +505,140 @@ static void UpdateHairGroupsDesc(UGroomAsset* GroomAsset, TArray<FHairGroupDesc>
 			Desc.ReInit();
 		}
 	}
+	UpdateNiagaraComponent();
 }
 
-/////////////////////////////////////////////////////////////////////////////////////////
-/////////////////////////////////////////////////////////////////////////////////////////
-// UComponent
-//////////////////////////////////////////////////////////////////////////
-//////////////////////////////////////////////////////////////////////////
+void UGroomComponent::UpdateNiagaraComponent()
+{
+	static UNiagaraSystem* CosseratRodsSystem = LoadObject<UNiagaraSystem>(nullptr, TEXT("/HairStrands/Emitters/GroomAssetSystem.GroomAssetSystem"));
+	static UNiagaraSystem* AngularSpringsSystem = LoadObject<UNiagaraSystem>(nullptr, TEXT("/HairStrands/Emitters/GroomAssetSystem.GroomAssetSystem"));
 
-UGroomComponent::UGroomComponent(const FObjectInitializer& ObjectInitializer)
-	: Super(ObjectInitializer)
+	if (GroomAsset && GroomAsset->EnableSimulation && bCreateNiagaraComponent)
+	{
+		if (!NiagaraComponent)
+		{
+			NiagaraComponent = NewObject<UNiagaraComponent>(this);
+			NiagaraComponent->AttachToComponent(this, FAttachmentTransformRules::KeepRelativeTransform);
+			NiagaraComponent->RegisterComponent();
+			NiagaraComponent->SetVisibleFlag(false);
+			NiagaraComponent->Activate(true);
+		}
+		if (GroomAsset->NiagaraSolver == EGroomNiagaraSolvers::AngularSprings)
+		{
+			NiagaraComponent->SetAsset(AngularSpringsSystem);
+		}
+		else if (GroomAsset->NiagaraSolver == EGroomNiagaraSolvers::CosseratRods)
+		{
+			NiagaraComponent->SetAsset(CosseratRodsSystem);
+		}
+		if (NiagaraComponent->GetSystemInstance())
+		{
+			NiagaraComponent->GetSystemInstance()->Reset(FNiagaraSystemInstance::EResetMode::ReInit);
+			NiagaraComponent->GetSystemInstance()->Reset(FNiagaraSystemInstance::EResetMode::ResetAll);
+		}
+		EnableSimulatedGroups();
+	}
+	else if (NiagaraComponent)
+	{
+		NiagaraComponent->DetachFromComponent(FDetachmentTransformRules::KeepRelativeTransform);
+		NiagaraComponent->UnregisterComponent();
+		NiagaraComponent->DestroyComponent();
+		NiagaraComponent = nullptr;
+		DisableSimulatedGroups();
+	}
+}
+
+void UGroomComponent::SetGroomAsset(UGroomAsset* Asset)
 {
-	PrimaryComponentTick.bCanEverTick = true;
-	bTickInEditor = true;
-	bAutoActivate = true;
-	bSelectable = true;
-	RegisteredSkeletalMeshComponent = nullptr;
-	SkeletalPreviousPositionOffset = FVector::ZeroVector;
-	bBindGroomToSkeletalMesh = false;
-	InitializedResources = nullptr;
-	Mobility = EComponentMobility::Movable;
-	bIsGroomAssetCallbackRegistered = false;
+	ReleaseResources();
+	GroomAsset = Asset;
+	if (!IsBindingAssetValid(BindingAsset) || !IsCompatible(GroomAsset, BindingAsset))
+	{
+		BindingAsset = nullptr;
+	}
 
-	SetCollisionProfileName(UCollisionProfile::PhysicsActor_ProfileName);
+	UpdateHairGroupsDesc();
+	if (!Asset)
+		return;
+	InitResources();
+}
+
+void UGroomComponent::SetGroomAsset(UGroomAsset* Asset, UGroomBindingAsset* InBinding)
+{
+	ReleaseResources();
+	GroomAsset = Asset;
+	BindingAsset = InBinding;
+	if (!IsBindingAssetValid(BindingAsset) || !IsCompatible(GroomAsset, BindingAsset))
+	{
+		BindingAsset = nullptr;
+	}
+
+	if (BindingAsset)
+	{
+		bBindGroomToSkeletalMesh = true;
+	}
+
+	UpdateHairGroupsDesc();
+	if (!Asset)
+		return;
+	InitResources();
+}
+
+void UGroomComponent::SetHairLengthScale(float Scale) 
+{ 
+	Scale = FMath::Clamp(Scale, 0.f, 1.f);
+	for (FHairGroupDesc& HairDesc : GroomGroupsDesc)
+	{
+		HairDesc.HairClipLength = Scale;
+	}
+	UpdateHairGroupsDescAndInvalidateRenderState();
+}
+
+void UGroomComponent::SetHairWidth(float HairWidth)
+{
+	for (FHairGroupDesc& HairDesc : GroomGroupsDesc)
+	{
+		HairDesc.HairWidth = HairWidth;
+	}
+	InitResources();
+}
+
+void UGroomComponent::SetBinding(bool bBind)
+{
+	if (bBind != bBindGroomToSkeletalMesh)
+	{
+		bBindGroomToSkeletalMesh = bBind;
+		InitResources();
+	}
+}
+
+void UGroomComponent::SetBinding(UGroomBindingAsset* InBinding)
+{
+	if (BindingAsset != InBinding)
+	{
+		const bool bIsValid = InBinding != nullptr ? IsBindingAssetValid(BindingAsset) : true;
+		if (bIsValid && IsCompatible(GroomAsset, InBinding))
+		{
+			BindingAsset = InBinding;
+			bBindGroomToSkeletalMesh = InBinding != nullptr;
+			InitResources();
+		}
+	}
+}
+
+void UGroomComponent::UpdateHairGroupsDescAndInvalidateRenderState()
+{
+	UpdateHairGroupsDesc();
+
+	if (InterpolationInput)
+	{
+		const int32 GroupCount = InterpolationInput->HairGroups.Num();
+		for (int32 GroupIt = 0; GroupIt < GroupCount; ++GroupIt)
+		{
+			InterpolationInput->HairGroups[GroupIt].GroupDesc = GroomGroupsDesc[GroupIt];
+		}
+	}
+	MarkRenderStateDirty();
 }
 
 FPrimitiveSceneProxy* UGroomComponent::CreateSceneProxy()
@@ -526,42 +725,6 @@ UMaterialInterface* UGroomComponent::GetMaterial(int32 ElementIndex) const
 	return OverrideMaterial;
 }
 
-static FHairStrandsProjectionHairData::HairGroup ToProjectionHairData(FHairStrandsRootResource* In)
-{
-	check(IsInRenderingThread());
-
-	FHairStrandsProjectionHairData::HairGroup Out = {};
-	if (!In)
-		return Out;
-
-	Out.RootCount = In->RootCount;
-	Out.RootPositionBuffer = In->RootPositionBuffer.SRV;
-	Out.RootNormalBuffer = In->RootNormalBuffer.SRV;
-	Out.VertexToCurveIndexBuffer = &In->VertexToCurveIndexBuffer;
-
-	for (FHairStrandsRootResource::FMeshProjectionLOD& MeshLODData : In->MeshProjectionLODs)
-	{
-		FHairStrandsProjectionHairData::LODData& LODData = Out.LODDatas.AddDefaulted_GetRef();
-		LODData.LODIndex = MeshLODData.LODIndex;
-		LODData.RootTriangleIndexBuffer = &MeshLODData.RootTriangleIndexBuffer;
-		LODData.RootTriangleBarycentricBuffer = &MeshLODData.RootTriangleBarycentricBuffer;
-		
-		LODData.RestPositionOffset				= &MeshLODData.RestRootOffset;
-		LODData.RestRootTrianglePosition0Buffer = &MeshLODData.RestRootTrianglePosition0Buffer;
-		LODData.RestRootTrianglePosition1Buffer = &MeshLODData.RestRootTrianglePosition1Buffer;
-		LODData.RestRootTrianglePosition2Buffer = &MeshLODData.RestRootTrianglePosition2Buffer;
-
-		LODData.DeformedPositionOffset				= &MeshLODData.DeformedRootOffset;
-		LODData.DeformedRootTrianglePosition0Buffer = &MeshLODData.DeformedRootTrianglePosition0Buffer;
-		LODData.DeformedRootTrianglePosition1Buffer = &MeshLODData.DeformedRootTrianglePosition1Buffer;
-		LODData.DeformedRootTrianglePosition2Buffer = &MeshLODData.DeformedRootTrianglePosition2Buffer;
-
-		LODData.Status = &MeshLODData.Status;
-	}
-
-	return Out;
-}
-
 FHairStrandsDatas* UGroomComponent::GetGuideStrandsDatas(uint32 GroupIndex)
 {
 	if (!GroomAsset || GroupIndex >= uint32(GroomAsset->GetNumHairGroups()))
@@ -584,22 +747,22 @@ FHairStrandsRestResource* UGroomComponent::GetGuideStrandsRestResource(uint32 Gr
 
 FHairStrandsDeformedResource* UGroomComponent::GetGuideStrandsDeformedResource(uint32 GroupIndex)
 {
-	if (GroupIndex >= uint32(HairGroupResources.Num()))
+	if (!HairGroupResources || GroupIndex >= uint32(HairGroupResources->HairGroups.Num()))
 	{
 		return nullptr;
 	}
 
-	return HairGroupResources[GroupIndex].SimDeformedResources;
+	return HairGroupResources->HairGroups[GroupIndex].SimDeformedResources;
 }
 
 FHairStrandsRootResource* UGroomComponent::GetGuideStrandsRootResource(uint32 GroupIndex)
 {
-	if (GroupIndex >= uint32(HairGroupResources.Num()))
+	if (!HairGroupResources || GroupIndex >= uint32(HairGroupResources->HairGroups.Num()))
 	{
 		return nullptr;
 	}
 
-	return HairGroupResources[GroupIndex].SimRootResources;
+	return HairGroupResources->HairGroups[GroupIndex].SimRootResources;
 }
 
 template<typename T> void SafeDelete(T*& Data) 
@@ -672,31 +835,52 @@ static bool IsSimulationEnabled(const USceneComponent* Component)
 	return bHasNiagaraSimulationComponent;
 }
 
-void UGroomComponent::OnChildAttached(USceneComponent* ChildComponent)
+static void ResetSimulation(const USceneComponent* Component)
+{
+	check(Component);
+
+	for (int32 ChildIt = 0, ChildCount = Component->GetNumChildrenComponents(); ChildIt < ChildCount; ++ChildIt)
+	{
+		const USceneComponent* ChildComponent = Component->GetChildComponent(ChildIt);
+		const UNiagaraComponent* NiagaraComponent = Cast<UNiagaraComponent>(ChildComponent);
+		if (NiagaraComponent != nullptr)
+		{
+			if (NiagaraComponent->GetSystemInstance())
+			{
+				NiagaraComponent->GetSystemInstance()->Reset(FNiagaraSystemInstance::EResetMode::ReInit);
+				NiagaraComponent->GetSystemInstance()->Reset(FNiagaraSystemInstance::EResetMode::ResetAll);
+			}
+		}
+	}
+}
+
+void UGroomComponent::EnableSimulatedGroups()
 {
-	const UNiagaraComponent* NiagaraComponent = Cast<UNiagaraComponent>(ChildComponent);
-	if (NiagaraComponent && InterpolationInput)
+	if (InterpolationInput)
 	{
 		FHairStrandsInterpolationInput* LocalInterpolationInput = InterpolationInput;
-		ENQUEUE_RENDER_COMMAND(FHairStrandsTick_UpdateSimulationEnable)(
-			[LocalInterpolationInput](FRHICommandListImmediate& RHICmdList)
+		UGroomAsset* LocalGroomAsset = GroomAsset;
+		ENQUEUE_RENDER_COMMAND(FHairStrandsTick_UEnableSimulatedGroups)(
+			[LocalInterpolationInput, LocalGroomAsset](FRHICommandListImmediate& RHICmdList)
 		{
+			int32 GroupIt = 0;
 			for (FHairStrandsInterpolationInput::FHairGroup& HairGroup : LocalInterpolationInput->HairGroups)
 			{
-				HairGroup.bIsSimulationEnable = true;
+				HairGroup.bIsSimulationEnable = LocalGroomAsset->EnableSimulation && (LocalGroomAsset->SimulatedGroup == GroupIt);
+				++GroupIt;
 			}
 		});
 	}
 }
 
-void UGroomComponent::OnChildDetached(USceneComponent* ChildComponent)
+void UGroomComponent::DisableSimulatedGroups()
 {
-	const UNiagaraComponent* NiagaraComponent = Cast<UNiagaraComponent>(ChildComponent);
-	if (NiagaraComponent && InterpolationInput)
+	if (InterpolationInput)
 	{
 		FHairStrandsInterpolationInput* LocalInterpolationInput = InterpolationInput;
-		ENQUEUE_RENDER_COMMAND(FHairStrandsTick_UpdateSimulationDisable)(
-			[LocalInterpolationInput](FRHICommandListImmediate& RHICmdList)
+		UGroomAsset* LocalGroomAsset = GroomAsset;
+		ENQUEUE_RENDER_COMMAND(FHairStrandsTick_DisableSimulatedGroups)(
+			[LocalInterpolationInput, LocalGroomAsset](FRHICommandListImmediate& RHICmdList)
 		{
 			for (FHairStrandsInterpolationInput::FHairGroup& HairGroup : LocalInterpolationInput->HairGroups)
 			{
@@ -706,19 +890,29 @@ void UGroomComponent::OnChildDetached(USceneComponent* ChildComponent)
 	}
 }
 
-void UGroomComponent::ResetSimulation()
+void UGroomComponent::OnChildDetached(USceneComponent* ChildComponent)
 {
-	bResetSimulation = false;
-	//UE_LOG(LogHairStrands, Warning, TEXT("Groom Reset = %d"), bResetSimulation);
+	if (Cast<UNiagaraComponent>(ChildComponent) && InterpolationInput)
+	{
+		DisableSimulatedGroups();
+	}
+}
+
+
+void UGroomComponent::OnChildAttached(USceneComponent* ChildComponent)
+{
+	if (Cast<UNiagaraComponent>(ChildComponent) && InterpolationInput)
+	{
+		EnableSimulatedGroups();
+	}
 }
 
 void UGroomComponent::InitResources()
 {
 	ReleaseResources();
 	bResetSimulation = true;
-	//UE_LOG(LogHairStrands, Warning, TEXT("Groom Init = %d"), bResetSimulation);
 
-	UpdateHairGroupsDesc(GroomAsset, GroomGroupsDesc);
+	UpdateHairGroupsDesc();
 
 	if (!GroomAsset || GroomAsset->GetNumHairGroups() == 0)
 		return;
@@ -746,14 +940,16 @@ void UGroomComponent::InitResources()
 	{
 		RegisteredSkeletalMeshComponent = SkeletalMeshComponent;
 		AddTickPrerequisiteComponent(SkeletalMeshComponent);
-		SkeletalMeshComponent->OnBoneTransformsFinalized.AddDynamic(this, &UGroomComponent::ResetSimulation);
 	}
 
+	const bool bIsBindingCompatible = IsCompatible(GroomAsset, BindingAsset) && IsBindingAssetValid(BindingAsset);
 	const bool bIsSimulationEnable = IsSimulationEnabled(this);
 
 	FTransform HairLocalToWorld = GetComponentTransform();
 	FTransform SkinLocalToWorld = bBindGroomToSkeletalMesh && SkeletalMeshComponent ? SkeletalMeshComponent->GetComponentTransform() : FTransform::Identity;
 	
+	check(HairGroupResources == nullptr);
+	HairGroupResources = new FHairGroupResources();
 	InterpolationOutput = new FHairStrandsInterpolationOutput();
 	InterpolationInput = new FHairStrandsInterpolationInput();
 
@@ -772,7 +968,7 @@ void UGroomComponent::InitResources()
 			DebugHairGroup.CurveCount	= GroupData.HairRenderData.GetNumCurves();
 		}
 
-		FHairGroupResource& Res = HairGroupResources.AddDefaulted_GetRef();
+		FHairGroupResource& Res = HairGroupResources->HairGroups.AddDefaulted_GetRef();
 		Res.InterpolationResource = new FHairStrandsInterpolationResource(GroupData.HairInterpolationData.RenderData, GroupData.HairSimulationData);
 		BeginInitResource(Res.InterpolationResource);
 
@@ -786,13 +982,26 @@ void UGroomComponent::InitResources()
 
 		if (bHasValidSketalMesh)
 		{
-			const uint32 LODCount = SkeletalMeshComponent->GetNumLODs();
-			if (LODCount > 0)
+			if (BindingAsset && bIsBindingCompatible)
 			{
-				Res.RenRootResources = new FHairStrandsRootResource(&GroupData.HairRenderData, LODCount);
-				Res.SimRootResources = new FHairStrandsRootResource(&GroupData.HairSimulationData, LODCount);
-				BeginInitResource(Res.RenRootResources);
-				BeginInitResource(Res.SimRootResources);
+				check(GroupIt < BindingAsset->HairGroupResources.Num());
+				check(SkeletalMeshComponent->GetNumLODs() == BindingAsset->HairGroupResources[GroupIt].RenRootResources->RootData.MeshProjectionLODs.Num());
+				check(SkeletalMeshComponent->GetNumLODs() == BindingAsset->HairGroupResources[GroupIt].SimRootResources->RootData.MeshProjectionLODs.Num());
+				Res.bOwnRootResourceAllocation = false;
+				Res.RenRootResources = BindingAsset->HairGroupResources[GroupIt].RenRootResources;
+				Res.SimRootResources = BindingAsset->HairGroupResources[GroupIt].SimRootResources;
+			}
+			else
+			{
+				const uint32 LODCount = SkeletalMeshComponent->GetNumLODs();
+				if (LODCount > 0)
+				{
+					Res.bOwnRootResourceAllocation = true;
+					Res.RenRootResources = new FHairStrandsRootResource(&GroupData.HairRenderData, LODCount);
+					Res.SimRootResources = new FHairStrandsRootResource(&GroupData.HairSimulationData, LODCount);
+					BeginInitResource(Res.RenRootResources);
+					BeginInitResource(Res.SimRootResources);
+				}
 			}
 		}
 		
@@ -836,7 +1045,7 @@ void UGroomComponent::InitResources()
 		// Deformed sim & render are expressed within the referential (unlike rest pose)
 		InterpolationInputGroup.OutHairPositionOffset = RenderRestHairPositionOffset;
 		InterpolationInputGroup.OutHairPreviousPositionOffset = RenderRestHairPositionOffset;
-		InterpolationInputGroup.bIsSimulationEnable = bIsSimulationEnable;
+		InterpolationInputGroup.bIsSimulationEnable = bIsSimulationEnable && GroomAsset->EnableSimulation && (GroomAsset->SimulatedGroup == GroupIt);
 
 		GroupIt++;
 	}
@@ -847,7 +1056,9 @@ void UGroomComponent::InitResources()
 	Interpolation.Function = ComputeHairStrandsInterpolation;
 	Interpolation.ResetFunction = ResetHairStrandsInterpolation;
 
-	FHairGroupResources* LocalResources = &HairGroupResources;
+	USkeletalMesh* InSourceSkeletalMesh = SourceSkeletalMesh;
+	const bool bRunMeshProjection = bHasValidSketalMesh && (!BindingAsset || !bIsBindingCompatible);
+	FHairGroupResources* LocalResources = HairGroupResources;
 	const uint64 Id = LocalComponentId.PrimIDValue;
 	ENQUEUE_RENDER_COMMAND(FHairStrandsBuffers)(
 		[
@@ -857,18 +1068,19 @@ void UGroomComponent::InitResources()
 			HairLocalToWorld, SkinLocalToWorld,
 			WorldType,
 			DebugGroupInfo,
-			bHasValidSketalMesh,
-			SkeletalMeshComponent
+			bRunMeshProjection,
+			SkeletalMeshComponent,
+			InSourceSkeletalMesh
 		]
 		(FRHICommandListImmediate& RHICmdList)
 	{
 		FHairStrandsProjectionHairData RenProjectionDatas;
 		FHairStrandsProjectionHairData SimProjectionDatas;
 		FHairStrandsPrimitiveResources PrimitiveResources;
-		const uint32 GroupCount = LocalResources->Num();
+		const uint32 GroupCount = LocalResources->HairGroups.Num();
 		for (uint32 GroupIt=0;GroupIt<GroupCount; ++GroupIt)
 		{
-			FHairGroupResource& Res = (*LocalResources)[GroupIt];
+			FHairGroupResource& Res = LocalResources->HairGroups[GroupIt];
 
 			FHairStrandsInterpolationInput::FHairGroup& InputGroup 	= Interpolation.Input->HairGroups[GroupIt];
 			FHairStrandsInterpolationOutput::HairGroup& OutputGroup	= Interpolation.Output->HairGroups[GroupIt];
@@ -922,43 +1134,45 @@ void UGroomComponent::InitResources()
 			Group.ClusterCount		= Res.ClusterCullingResources->ClusterCount;
 		}
 
-		if (bHasValidSketalMesh)
-		{		
-			// Convert FSkeletalMeshRenderData into FHairStrandsProjectionMeshData for the initial projection
-			FHairStrandsProjectionMeshData MeshData;
-			FSkeletalMeshRenderData* RenderData = SkeletalMeshComponent->GetSkeletalMeshRenderData();
+		FHairStrandsProjectionDebugInfo HairProjectionDebugInfo;
+		if (bRunMeshProjection)
+		{				
+			FSkeletalMeshRenderData* TargetRenderData = SkeletalMeshComponent->GetSkeletalMeshRenderData();
+			FHairStrandsProjectionMeshData TargetMeshData = ExtractMeshData(TargetRenderData);
 
-			uint32 LODIndex = 0;
-			for (FSkeletalMeshLODRenderData& LODRenderData : RenderData->LODRenderData)
+			// Create mapping between the source & target using their UV
+			// The lifetime of 'TransferredPositions' needs to encompass RunProjection
+			TArray<FRWBuffer> TransferredPositions;
+			
+			if (FSkeletalMeshRenderData* SourceRenderData = InSourceSkeletalMesh ? InSourceSkeletalMesh->GetResourceForRendering() : nullptr)
 			{
-				FHairStrandsProjectionMeshData::LOD& LOD = MeshData.LODs.AddDefaulted_GetRef();
-				uint32 SectionIndex = 0;
-				for (FSkelMeshRenderSection& InSection : LODRenderData.RenderSections)
+				FHairStrandsProjectionMeshData SourceMeshData = ExtractMeshData(SourceRenderData);
+				RunMeshTransfer(
+					RHICmdList,
+					SourceMeshData,
+					TargetMeshData,
+					TransferredPositions);
+
+				const uint32 LODCount = TargetMeshData.LODs.Num();
+				for (uint32 LODIndex = 0; LODIndex < LODCount; ++LODIndex)
 				{
-					FHairStrandsProjectionMeshData::Section& OutSection = LOD.Sections.AddDefaulted_GetRef();
-					
-					OutSection.PositionBuffer	= LODRenderData.StaticVertexBuffers.PositionVertexBuffer.GetSRV();
-					OutSection.IndexBuffer		= LODRenderData.MultiSizeIndexContainer.GetIndexBuffer()->GetSRV();
-					OutSection.TotalVertexCount	= LODRenderData.StaticVertexBuffers.PositionVertexBuffer.GetNumVertices();
-					OutSection.TotalIndexCount	= LODRenderData.MultiSizeIndexContainer.GetIndexBuffer()->Num();
-					OutSection.NumPrimitives	= InSection.NumTriangles;
-					OutSection.VertexBaseIndex	= InSection.BaseVertexIndex;
-					OutSection.IndexBaseIndex	= InSection.BaseIndex;
-					OutSection.SectionIndex		= SectionIndex;
-					OutSection.LODIndex			= LODIndex;
-
-					++SectionIndex;
+					for (FHairStrandsProjectionMeshData::Section& Section : TargetMeshData.LODs[LODIndex].Sections)
+					{
+						Section.PositionBuffer = TransferredPositions[LODIndex].SRV;
+					}
 				}
-				++LODIndex;
+			#if WITH_EDITOR
+				HairProjectionDebugInfo.SourceMeshData = SourceMeshData;
+				HairProjectionDebugInfo.TargetMeshData = TargetMeshData;
+				HairProjectionDebugInfo.TransferredPositions = TransferredPositions;
+			#endif
 			}
 
 			// The offset is based on the center of the skeletal mesh (which is computed based on the physics capsules/boxes/...)
-			const FVector MeshPositionOffset = SkeletalMeshComponent->CalcBounds(FTransform::Identity).GetBox().GetCenter();
 			RunProjection(
 				RHICmdList,
 				HairLocalToWorld,
-				MeshPositionOffset,
-				MeshData, 
+				TargetMeshData,
 				RenProjectionDatas,
 				SimProjectionDatas);
 		}
@@ -970,10 +1184,47 @@ void UGroomComponent::InitResources()
 			RenProjectionDatas,
 			SimProjectionDatas,
 			PrimitiveResources,
-			DebugGroupInfo);
+			DebugGroupInfo,
+			HairProjectionDebugInfo);
 	});
 }
 
+void UGroomComponent::DeleteHairGroupResources(FHairGroupResources*& InHairGroupResources)
+{
+	if (!InHairGroupResources)
+	{
+		return;
+	}
+
+	FHairGroupResources* Local = InHairGroupResources;
+	ENQUEUE_RENDER_COMMAND(StaticMeshVertexBuffersLegacyInit)(
+		[Local](FRHICommandListImmediate& RHICmdList)
+	{
+		for (FHairGroupResource& Res : Local->HairGroups)
+		{
+			SafeRelease(Res.InterpolationResource);
+
+			// Release the root resources only if they have been created internally (vs. being created by external asset)
+			if (Res.bOwnRootResourceAllocation)
+			{
+				SafeRelease(Res.RenRootResources);
+				SafeRelease(Res.SimRootResources);
+			}
+
+			SafeRelease(Res.RenderDeformedResources);
+			SafeRelease(Res.ClusterCullingResources);
+			SafeRelease(Res.HairGroupPublicDatas);
+			SafeRelease(Res.SimDeformedResources);
+			#if RHI_RAYTRACING
+			SafeRelease(Res.RaytracingResources);
+			#endif
+		}
+		delete Local;
+	});
+
+	InHairGroupResources = nullptr;
+}
+
 void UGroomComponent::ReleaseResources()
 {
 	// Unregister component interpolation resources
@@ -985,20 +1236,7 @@ void UGroomComponent::ReleaseResources()
 		UnregisterHairStrands(Id);
 	});
 
-	for (FHairGroupResource& Res : HairGroupResources)
-	{
-		SafeRelease(Res.InterpolationResource);
-		SafeRelease(Res.RenRootResources);
-		SafeRelease(Res.SimRootResources);
-		SafeRelease(Res.RenderDeformedResources);
-		SafeRelease(Res.ClusterCullingResources);
-		SafeRelease(Res.HairGroupPublicDatas);
-		SafeRelease(Res.SimDeformedResources);
-	#if RHI_RAYTRACING
-		SafeRelease(Res.RaytracingResources);
-	#endif
-	}
-	HairGroupResources.Empty();
+	DeleteHairGroupResources(HairGroupResources);
 
 	// Delay destruction as resources reference by the interpolation 
 	// structs are used on the rendering thread, 
@@ -1022,13 +1260,6 @@ void UGroomComponent::ReleaseResources()
 	SkeletalPreviousPositionOffset = FVector::ZeroVector;
 	RegisteredSkeletalMeshComponent = nullptr;
 
-	USkeletalMeshComponent* SkeletalMeshComponent = GetAttachParent() ? Cast<USkeletalMeshComponent>(GetAttachParent()) : nullptr;
-	if (SkeletalMeshComponent)
-	{
-		SkeletalMeshComponent->OnBoneTransformsFinalized.RemoveDynamic(this, &UGroomComponent::ResetSimulation);
-		bResetSimulation = true;
-	}
-
 	MarkRenderStateDirty();
 }
 
@@ -1042,6 +1273,11 @@ void UGroomComponent::PostLoad()
 		GroomAsset->ConditionalPostLoad();
 	}
 
+	if (BindingAsset)
+	{
+		// Make sure that the asset initialized its resources first since the component needs them to initialize its own resources
+		BindingAsset->ConditionalPostLoad();
+	}
 	InitResources();
 
 #if WITH_EDITOR
@@ -1057,6 +1293,7 @@ void UGroomComponent::PostLoad()
 #if WITH_EDITOR
 void UGroomComponent::Invalidate()
 {
+	UpdateNiagaraComponent();
 	MarkRenderStateDirty();
 	ValidateMaterials(false);
 }
@@ -1073,10 +1310,6 @@ void UGroomComponent::OnRegister()
 
 	// Insure the ticking of the Groom component always happens after the skeletalMeshComponent.
 	USkeletalMeshComponent* SkeletalMeshComponent = GetAttachParent() ? Cast<USkeletalMeshComponent>(GetAttachParent()) : nullptr;
-	/*if (SkeletalMeshComponent)
-	{
-		SkeletalMeshComponent->OnBoneTransformsFinalized.AddDynamic(this, &UGroomComponent::ResetSimulation);
-	}*/
 
 	const EWorldType::Type WorldType = GetWorld() ? EWorldType::Type(GetWorld()->WorldType) : EWorldType::None;
 	const uint64 Id = ComponentId.PrimIDValue;
@@ -1149,13 +1382,23 @@ void UGroomComponent::TickComponent(float DeltaTime, enum ELevelTick TickType, F
 			PrevBoneMatrix = NextBoneMatrix;
 		}
 	}
+	else
+	{
+		bResetSimulation = false;
+	}
+
+	if (!HairGroupResources)
+	{
+		return;
+	}
+
 	if (SkeletalMeshComponent)
 	{
 		// The offset is based on the center of the skeletal mesh (which is computed based on the physics capsules/boxes/...)
 		// For skinned mesh update the relative center of hair positions after deformation
 		MeshPositionOffset = SkeletalMeshComponent->CalcBounds(FTransform::Identity).GetBox().GetCenter();
 		{
-			FHairGroupResources* LocalResources = &HairGroupResources;
+			FHairGroupResources* LocalResources = HairGroupResources;
 
 			const FVector OutHairPositionOffset = MeshPositionOffset;
 			const FVector OutHairPreviousPositionOffset = SkeletalPreviousPositionOffset;
@@ -1170,7 +1413,7 @@ void UGroomComponent::TickComponent(float DeltaTime, enum ELevelTick TickType, F
 				}
 
 				// Update deformed (sim/render) hair position offsets. This is used by Niagara.
-				for (FHairGroupResource& Res : *LocalResources)
+				for (FHairGroupResource& Res : LocalResources->HairGroups)
 				{
 					if (Res.RenderDeformedResources)
 						Res.RenderDeformedResources->PositionOffset = OutHairPositionOffset;
@@ -1189,36 +1432,16 @@ void UGroomComponent::TickComponent(float DeltaTime, enum ELevelTick TickType, F
 		MarkRenderTransformDirty();
 	}
 
-	FHairGroupResources* LocalResources = &HairGroupResources;
-	const FVector DeformedPositionCenter = GHairStrandsMeshUseRelativePosition > 0 ? MeshPositionOffset : FVector::ZeroVector;
+	FHairGroupResources* LocalResources = HairGroupResources;
 	const FTransform SkinLocalToWorld = SkeletalMeshComponent ? SkeletalMeshComponent->GetComponentTransform() : FTransform();
 	const FTransform HairLocalToWorld = GetComponentTransform();
 	ENQUEUE_RENDER_COMMAND(FHairStrandsTick_TransformUpdate)(
-		[Id, WorldType, HairLocalToWorld, SkinLocalToWorld, DeformedPositionCenter, FeatureLevel, LocalResources](FRHICommandListImmediate& RHICmdList)
+		[Id, WorldType, HairLocalToWorld, SkinLocalToWorld, FeatureLevel, LocalResources](FRHICommandListImmediate& RHICmdList)
 	{		
 		if (ERHIFeatureLevel::Num == FeatureLevel)
 			return;
 
-		UpdateHairStrands(Id, WorldType, HairLocalToWorld, SkinLocalToWorld, DeformedPositionCenter);
-
-		// Update deformed (sim/render) triangles position offsets. This is used by Niagara.
-		for (FHairGroupResource& Res : *LocalResources)
-		{
-			if (Res.RenRootResources)
-			{
-				for (FHairStrandsRootResource::FMeshProjectionLOD& MeshProjectionLOD : Res.RenRootResources->MeshProjectionLODs)
-				{
-					MeshProjectionLOD.DeformedRootOffset = DeformedPositionCenter;
-				}
-			}
-			if (Res.SimRootResources)
-			{
-				for (FHairStrandsRootResource::FMeshProjectionLOD& MeshProjectionLOD : Res.SimRootResources->MeshProjectionLODs)
-				{
-					MeshProjectionLOD.DeformedRootOffset = DeformedPositionCenter;
-				}
-			}
-		}
+		UpdateHairStrands(Id, WorldType, HairLocalToWorld, SkinLocalToWorld);
 	});
 }
 
@@ -1270,10 +1493,20 @@ void UGroomComponent::PostEditChangeProperty(FPropertyChangedEvent& PropertyChan
 
 	//  Init/release resources when setting the GroomAsset (or undoing)
 	const bool bAssetChanged = PropertyName == GET_MEMBER_NAME_CHECKED(UGroomComponent, GroomAsset);
-	const bool bRecreateResources =
-		(bAssetChanged || PropertyThatChanged == nullptr) ||
-		PropertyName == GET_MEMBER_NAME_CHECKED(UGroomComponent, bBindGroomToSkeletalMesh);
+	const bool bBindingAssetChanged = PropertyName == GET_MEMBER_NAME_CHECKED(UGroomComponent, BindingAsset);
+	const bool bBindToSkeletalChanged = PropertyName == GET_MEMBER_NAME_CHECKED(UGroomComponent, bBindGroomToSkeletalMesh);
+	const bool bDefaultNiagaraChanged = PropertyName == GET_MEMBER_NAME_CHECKED(UGroomComponent, bCreateNiagaraComponent);
+	const bool bIsBindingCompatible = IsCompatible(GroomAsset, BindingAsset);
+	if (!bIsBindingCompatible || !IsBindingAssetValid(BindingAsset))
+	{
+		BindingAsset = nullptr;
+	}
+	if (BindingAsset)
+	{
+		bBindGroomToSkeletalMesh = true;
+	}
 
+	const bool bRecreateResources = bAssetChanged || bBindingAssetChanged || PropertyThatChanged == nullptr || bBindToSkeletalChanged || bDefaultNiagaraChanged;
 	if (bRecreateResources)
 	{
 		// Release the resources before Super::PostEditChangeProperty so that they get
@@ -1294,6 +1527,7 @@ void UGroomComponent::PostEditChangeProperty(FPropertyChangedEvent& PropertyChan
 	{
 		if (GroomAsset)
 		{
+			ResetSimulation(this);
 			// Set the callback on the new GroomAsset being assigned
 			GroomAsset->GetOnGroomAssetChanged().AddUObject(this, &UGroomComponent::Invalidate);
 			bIsGroomAssetCallbackRegistered = true;
@@ -1311,17 +1545,7 @@ void UGroomComponent::PostEditChangeProperty(FPropertyChangedEvent& PropertyChan
 		PropertyName == GET_MEMBER_NAME_CHECKED(FHairGroupDesc, LodBias) ||
 		PropertyName == GET_MEMBER_NAME_CHECKED(FHairGroupDesc, LodAverageVertexPerPixel))
 	{	
-		UpdateHairGroupsDesc(GroomAsset, GroomGroupsDesc);
-
-		if (InterpolationInput)
-		{
-			const int32 GroupCount = InterpolationInput->HairGroups.Num();
-			for (int32 GroupIt = 0; GroupIt < GroupCount; ++GroupIt)
-			{
-				InterpolationInput->HairGroups[GroupIt].GroupDesc = GroomGroupsDesc[GroupIt];
-			}
-		}
-		MarkRenderStateDirty();
+		UpdateHairGroupsDescAndInvalidateRenderState();
 	}
 
 #if WITH_EDITOR
diff --git a/Engine/Plugins/Experimental/HairStrands/Source/HairStrandsCore/Private/HairStrandsDatas.cpp b/Engine/Plugins/Experimental/HairStrands/Source/HairStrandsCore/Private/HairStrandsDatas.cpp
--- a/Engine/Plugins/Experimental/HairStrands/Source/HairStrandsCore/Private/HairStrandsDatas.cpp
+++ b/Engine/Plugins/Experimental/HairStrands/Source/HairStrandsCore/Private/HairStrandsDatas.cpp
@@ -51,6 +51,15 @@ void FHairStrandsPoints::Reset()
 	PointsCoordU.Reset();
 }
 
+FArchive& operator<<(FArchive& Ar, FVector4_16& Vertex)
+{
+	Ar << Vertex.X;
+	Ar << Vertex.Y;
+	Ar << Vertex.Z;
+	Ar << Vertex.W;
+	return Ar;
+}
+
 FArchive& operator<<(FArchive& Ar, FPackedHairVertex& Vertex)
 {
 	Ar << Vertex.X;
diff --git a/Engine/Plugins/Experimental/HairStrands/Source/HairStrandsCore/Private/HairStrandsRendering.cpp b/Engine/Plugins/Experimental/HairStrands/Source/HairStrandsCore/Private/HairStrandsRendering.cpp
--- a/Engine/Plugins/Experimental/HairStrands/Source/HairStrandsCore/Private/HairStrandsRendering.cpp
+++ b/Engine/Plugins/Experimental/HairStrands/Source/HairStrandsCore/Private/HairStrandsRendering.cpp
@@ -6,9 +6,11 @@
 #include "RenderGraphUtils.h"
 #include "ShaderParameterStruct.h"
 #include "GlobalShader.h"
+#include "GroomAsset.h"
 #include "HairStrandsInterface.h"
 #include "SceneView.h"
 #include "Containers/ResourceArray.h"
+#include "Rendering\SkeletalMeshRenderData.h"
 #include "HAL/ConsoleManager.h"
 #include "GpuDebugRendering.h"
 
@@ -94,7 +96,7 @@ class FDeformGuideCS : public FGlobalShader
 	SHADER_USE_PARAMETER_STRUCT(FDeformGuideCS, FGlobalShader);
 
 	class FGroupSize : SHADER_PERMUTATION_INT("PERMUTATION_GROUP_SIZE", 2);
-	class FDeformationType : SHADER_PERMUTATION_INT("PERMUTATION_DEFORMATION", 4);
+	class FDeformationType : SHADER_PERMUTATION_INT("PERMUTATION_DEFORMATION", 5);
 	using FPermutationDomain = TShaderPermutationDomain<FGroupSize, FDeformationType>;
 
 	BEGIN_SHADER_PARAMETER_STRUCT(FParameters, )
@@ -103,6 +105,18 @@ class FDeformGuideCS : public FGlobalShader
 		SHADER_PARAMETER(FVector, SimRestOffset)
 		SHADER_PARAMETER(FVector, SimDeformedOffset)
 		SHADER_PARAMETER(uint32, DispatchCountX)
+
+		SHADER_PARAMETER_SRV(Buffer<float4>, SimRestPosition0Buffer)
+		SHADER_PARAMETER_SRV(Buffer<float4>, SimRestPosition1Buffer)
+		SHADER_PARAMETER_SRV(Buffer<float4>, SimRestPosition2Buffer)
+
+		SHADER_PARAMETER_SRV(Buffer<float4>, SimDeformedPosition0Buffer)
+		SHADER_PARAMETER_SRV(Buffer<float4>, SimDeformedPosition1Buffer)
+		SHADER_PARAMETER_SRV(Buffer<float4>, SimDeformedPosition2Buffer)
+
+		SHADER_PARAMETER_SRV(Buffer<uint>, SimRootBarycentricBuffer)
+		SHADER_PARAMETER_SRV(Buffer<uint>, SimVertexToRootIndexBuffer)
+
 		SHADER_PARAMETER_SRV(Buffer, SimRestPosePositionBuffer)
 		SHADER_PARAMETER_SRV(Buffer, SimRootIndexBuffer)
 		SHADER_PARAMETER_UAV(RWBuffer, OutSimDeformedPositionBuffer)
@@ -118,6 +132,8 @@ static void AddDeformSimHairStrandsPass(
 	FRDGBuilder& GraphBuilder,
 	EDeformationType DeformationType,
 	uint32 VertexCount,
+	const uint32 LODIndex,
+	const FHairStrandsProjectionHairData::HairGroup& InSimHairData,
 	FRHIShaderResourceView* SimRestPosePositionBuffer,
 	FRHIShaderResourceView* SimRootIndexBuffer,
 	FRHIUnorderedAccessView* OutSimDeformedPositionBuffer,
@@ -153,6 +169,30 @@ static void AddDeformSimHairStrandsPass(
 	Parameters->SimRestOffset = SimRestOffset;
 	Parameters->DispatchCountX = DispatchCountX;
 
+	if (DeformationType == EDeformationType::OffsetGuide)
+	{
+		const bool bIsVertexToCurveBuffersValid = InSimHairData.VertexToCurveIndexBuffer != nullptr;
+		if (bIsVertexToCurveBuffersValid)
+		{
+			Parameters->SimVertexToRootIndexBuffer = InSimHairData.VertexToCurveIndexBuffer->SRV;
+		}
+
+		const bool bSupportDynamicMesh = InSimHairData.RootCount > 0 && LODIndex >= 0 && LODIndex < uint32(InSimHairData.LODDatas.Num()) && InSimHairData.LODDatas[LODIndex].bIsValid && bIsVertexToCurveBuffersValid;
+		if (bSupportDynamicMesh)
+		{
+			InternalDeformationType = 4;
+			Parameters->SimRestPosition0Buffer = InSimHairData.LODDatas[LODIndex].RestRootTrianglePosition0Buffer->SRV;
+			Parameters->SimRestPosition1Buffer = InSimHairData.LODDatas[LODIndex].RestRootTrianglePosition1Buffer->SRV;
+			Parameters->SimRestPosition2Buffer = InSimHairData.LODDatas[LODIndex].RestRootTrianglePosition2Buffer->SRV;
+
+			Parameters->SimDeformedPosition0Buffer = InSimHairData.LODDatas[LODIndex].DeformedRootTrianglePosition0Buffer->SRV;
+			Parameters->SimDeformedPosition1Buffer = InSimHairData.LODDatas[LODIndex].DeformedRootTrianglePosition1Buffer->SRV;
+			Parameters->SimDeformedPosition2Buffer = InSimHairData.LODDatas[LODIndex].DeformedRootTrianglePosition2Buffer->SRV;
+
+			Parameters->SimRootBarycentricBuffer = InSimHairData.LODDatas[LODIndex].RootTriangleBarycentricBuffer->SRV;
+		}
+	}
+
 	FDeformGuideCS::FPermutationDomain PermutationVector;
 	PermutationVector.Set<FDeformGuideCS::FGroupSize>(GetGroupSizePermutation(GroupSize));
 	PermutationVector.Set<FDeformGuideCS::FDeformationType>(InternalDeformationType);
@@ -211,12 +251,6 @@ class FHairInterpolationCS : public FGlobalShader
 		SHADER_PARAMETER(float, HairRadiusTipScale)
 		SHADER_PARAMETER(float, HairLengthClip)
 
-		SHADER_PARAMETER(FVector, RestPositionOffset)
-		SHADER_PARAMETER(FVector, DeformedPositionOffset)
-
-		SHADER_PARAMETER(FVector, SimRestPositionOffset)
-		SHADER_PARAMETER(FVector, SimDeformedPositionOffset)
-
 		SHADER_PARAMETER(FMatrix, LocalToWorldMatrix)
 
 		SHADER_PARAMETER_SRV(Buffer, RenderRestPosePositionBuffer)
@@ -243,7 +277,6 @@ class FHairInterpolationCS : public FGlobalShader
 		SHADER_PARAMETER_SRV(Buffer<float4>, DeformedPosition2Buffer)
 
 		SHADER_PARAMETER_SRV(Buffer<uint>, RootBarycentricBuffer)
-		SHADER_PARAMETER_SRV(Buffer<uint>, RootToTriangleIndex)
 		SHADER_PARAMETER_SRV(Buffer<uint>, RenVertexToRootIndexBuffer)
 
 		SHADER_PARAMETER_SRV(Buffer<float4>, SimRestPosition0Buffer)
@@ -255,7 +288,6 @@ class FHairInterpolationCS : public FGlobalShader
 		SHADER_PARAMETER_SRV(Buffer<float4>, SimDeformedPosition2Buffer)
 
 		SHADER_PARAMETER_SRV(Buffer<uint>, SimRootBarycentricBuffer)
-		SHADER_PARAMETER_SRV(Buffer<uint>, SimRootToTriangleIndex)
 		SHADER_PARAMETER_SRV(Buffer<uint>, SimVertexToRootIndexBuffer)
 
 		END_SHADER_PARAMETER_STRUCT()
@@ -375,30 +407,24 @@ static void AddHairStrandsInterpolationPass(
 	const bool bSupportDynamicMesh = InRenHairData.RootCount > 0 && LODIndex >= 0 && LODIndex < InRenHairData.LODDatas.Num() && InRenHairData.LODDatas[LODIndex].bIsValid && bIsVertexToCurveBuffersValid;
 	if (bSupportDynamicMesh)
 	{
-		Parameters->RestPositionOffset = *InRenHairData.LODDatas[LODIndex].RestPositionOffset;
 		Parameters->RestPosition0Buffer = InRenHairData.LODDatas[LODIndex].RestRootTrianglePosition0Buffer->SRV;
 		Parameters->RestPosition1Buffer = InRenHairData.LODDatas[LODIndex].RestRootTrianglePosition1Buffer->SRV;
 		Parameters->RestPosition2Buffer = InRenHairData.LODDatas[LODIndex].RestRootTrianglePosition2Buffer->SRV;
 
-		Parameters->DeformedPositionOffset = *InRenHairData.LODDatas[LODIndex].DeformedPositionOffset;
 		Parameters->DeformedPosition0Buffer = InRenHairData.LODDatas[LODIndex].DeformedRootTrianglePosition0Buffer->SRV;
 		Parameters->DeformedPosition1Buffer = InRenHairData.LODDatas[LODIndex].DeformedRootTrianglePosition1Buffer->SRV;
 		Parameters->DeformedPosition2Buffer = InRenHairData.LODDatas[LODIndex].DeformedRootTrianglePosition2Buffer->SRV;
 
-		Parameters->RootToTriangleIndex = InRenHairData.LODDatas[LODIndex].RootTriangleIndexBuffer->SRV;
 		Parameters->RootBarycentricBuffer = InRenHairData.LODDatas[LODIndex].RootTriangleBarycentricBuffer->SRV;
 
-		Parameters->SimRestPositionOffset = *InSimHairData.LODDatas[LODIndex].RestPositionOffset;
 		Parameters->SimRestPosition0Buffer = InSimHairData.LODDatas[LODIndex].RestRootTrianglePosition0Buffer->SRV;
 		Parameters->SimRestPosition1Buffer = InSimHairData.LODDatas[LODIndex].RestRootTrianglePosition1Buffer->SRV;
 		Parameters->SimRestPosition2Buffer = InSimHairData.LODDatas[LODIndex].RestRootTrianglePosition2Buffer->SRV;
 
-		Parameters->SimDeformedPositionOffset = *InSimHairData.LODDatas[LODIndex].DeformedPositionOffset;
 		Parameters->SimDeformedPosition0Buffer = InSimHairData.LODDatas[LODIndex].DeformedRootTrianglePosition0Buffer->SRV;
 		Parameters->SimDeformedPosition1Buffer = InSimHairData.LODDatas[LODIndex].DeformedRootTrianglePosition1Buffer->SRV;
 		Parameters->SimDeformedPosition2Buffer = InSimHairData.LODDatas[LODIndex].DeformedRootTrianglePosition2Buffer->SRV;
 
-		Parameters->SimRootToTriangleIndex = InSimHairData.LODDatas[LODIndex].RootTriangleIndexBuffer->SRV;
 		Parameters->SimRootBarycentricBuffer = InSimHairData.LODDatas[LODIndex].RootTriangleBarycentricBuffer->SRV;
 	}
 
@@ -723,6 +749,8 @@ void ComputeHairStrandsInterpolation(
 				GraphBuilder,
 				DeformationType,
 				Input.SimVertexCount,
+				LODIndex,
+				InSimHairDatas.HairGroups[GroupIndex],
 				Input.SimRestPosePositionBuffer->SRV,
 				Input.SimRootPointIndexBuffer ? Input.SimRootPointIndexBuffer->SRV : nullptr,
 				Output.SimDeformedPositionBuffer[CurrIndex]->UAV, Input.InSimHairPositionOffset,
@@ -893,7 +921,9 @@ void ComputeHairStrandsInterpolation(
 void ResetHairStrandsInterpolation(
 	FRHICommandListImmediate& RHICmdList,
 	FHairStrandsInterpolationInput* InInput,
-	FHairStrandsInterpolationOutput* InOutput)
+	FHairStrandsInterpolationOutput* InOutput,
+	FHairStrandsProjectionHairData& InSimHairDatas,
+	int32 LODIndex)
 {
 	if (!InInput || !InOutput) return;
 
@@ -904,21 +934,197 @@ void ResetHairStrandsInterpolation(
 		FHairStrandsInterpolationOutput::HairGroup& Output = InOutput->HairGroups[GroupIndex];
 		if (!Input.IsValid() || !Output.IsValid()) return;
 
-		DECLARE_GPU_STAT(HairStrandsResetInterpolation);
-		SCOPED_DRAW_EVENT(RHICmdList, HairStrandsResetInterpolation);
-		SCOPED_GPU_STAT(RHICmdList, HairStrandsResetInterpolation);
+		if (!Input.bIsSimulationEnable)
+		{
+			DECLARE_GPU_STAT(HairStrandsResetInterpolation);
+			SCOPED_DRAW_EVENT(RHICmdList, HairStrandsResetInterpolation);
+			SCOPED_GPU_STAT(RHICmdList, HairStrandsResetInterpolation);
 
-		const uint32 CurrIndex = Output.CurrentIndex;
-		const uint32 PrevIndex = (Output.CurrentIndex + 1) % 2;
-		FRDGBuilder GraphBuilder(RHICmdList);
-		AddDeformSimHairStrandsPass(
-			GraphBuilder,
-			EDeformationType::OffsetGuide,
-			Input.SimVertexCount,
-			Input.SimRestPosePositionBuffer->SRV,
-			Input.SimRootPointIndexBuffer ? Input.SimRootPointIndexBuffer->SRV : nullptr,
-			Output.SimDeformedPositionBuffer[CurrIndex]->UAV, Input.InSimHairPositionOffset,
-			Input.OutHairPositionOffset);
-		GraphBuilder.Execute();
+			const uint32 CurrIndex = Output.CurrentIndex;
+			const uint32 PrevIndex = (Output.CurrentIndex + 1) % 2;
+			FRDGBuilder GraphBuilder(RHICmdList);
+			AddDeformSimHairStrandsPass(
+				GraphBuilder,
+				EDeformationType::OffsetGuide,
+				Input.SimVertexCount,
+				LODIndex,
+				InSimHairDatas.HairGroups[GroupIndex],
+				Input.SimRestPosePositionBuffer->SRV,
+				Input.SimRootPointIndexBuffer ? Input.SimRootPointIndexBuffer->SRV : nullptr,
+				Output.SimDeformedPositionBuffer[CurrIndex]->UAV, Input.InSimHairPositionOffset,
+				Input.OutHairPositionOffset);
+			GraphBuilder.Execute();
+		}
 	}
+}
+
+///////////////////////////////////////////////////////////////////////////////////////////////////
+
+template<typename ReadBackType>
+void ReadbackBuffer(TArray<ReadBackType>& OutData, FRWBuffer& InBuffer)
+{
+	ReadBackType* Data = (ReadBackType*)RHILockVertexBuffer(InBuffer.Buffer, 0, InBuffer.Buffer->GetSize(), RLM_ReadOnly);
+	const uint32 ElementCount = InBuffer.Buffer->GetSize() / sizeof(ReadBackType);
+	OutData.SetNum(ElementCount);
+	for (uint32 ElementIt = 0; ElementIt < ElementCount; ++ElementIt)
+	{
+		OutData[ElementIt] = Data[ElementIt];
+	}
+	RHIUnlockVertexBuffer(InBuffer.Buffer);
+}
+
+static void ReadbackGroupData(
+	FRHICommandListImmediate& RHICmdList, 
+	FHairStrandsRootData& OutCPUData,
+	FHairStrandsRootResource* InGPUData)
+{
+	if (!InGPUData)
+	{
+		return;
+	}
+
+	check(InGPUData->MeshProjectionLODs.Num() == OutCPUData.MeshProjectionLODs.Num());
+
+	const uint32 LODCount = InGPUData->MeshProjectionLODs.Num();
+	for (uint32 LODIt = 0; LODIt < LODCount; ++LODIt)
+	{
+		FHairStrandsRootData::FMeshProjectionLOD& CPULOD = OutCPUData.MeshProjectionLODs[LODIt];
+		FHairStrandsRootResource::FMeshProjectionLOD& GPULOD = InGPUData->MeshProjectionLODs[LODIt];
+		check(CPULOD.LODIndex == GPULOD.LODIndex);
+
+		ReadbackBuffer(CPULOD.RootTriangleIndexBuffer, GPULOD.RootTriangleIndexBuffer);
+		ReadbackBuffer(CPULOD.RootTriangleBarycentricBuffer, GPULOD.RootTriangleBarycentricBuffer);
+		ReadbackBuffer(CPULOD.RestRootTrianglePosition0Buffer, GPULOD.RestRootTrianglePosition0Buffer);
+		ReadbackBuffer(CPULOD.RestRootTrianglePosition1Buffer, GPULOD.RestRootTrianglePosition1Buffer);
+		ReadbackBuffer(CPULOD.RestRootTrianglePosition2Buffer, GPULOD.RestRootTrianglePosition2Buffer);
+	}
+}
+
+static void InternalProcessGroomBindingTask(FRHICommandListImmediate& RHICmdList, void* Asset)
+{
+	UGroomBindingAsset* BindingAsset = (UGroomBindingAsset*)Asset;
+	if (!BindingAsset ||
+		!BindingAsset->Groom ||
+		!BindingAsset->TargetSkeletalMesh ||
+		 BindingAsset->Groom->GetNumHairGroups() == 0)
+	{
+		return;
+	}
+
+	UGroomAsset* GroomAsset				= BindingAsset->Groom;
+	USkeletalMesh* SourceSkeletalMesh	= BindingAsset->SourceSkeletalMesh;
+	USkeletalMesh* TargetSkeletalMesh	= BindingAsset->TargetSkeletalMesh;
+
+	const uint32 LODCount = BindingAsset->TargetSkeletalMesh->GetLODNum();
+	UGroomBindingAsset::FHairGroupDatas& OutHairGroupDatas = BindingAsset->HairGroupDatas;
+	for (const FHairGroupData& GroupData : GroomAsset->HairGroupsData)
+	{
+		UGroomBindingAsset::FHairGroupData& Data = OutHairGroupDatas.AddDefaulted_GetRef();
+		Data.RenRootData = FHairStrandsRootData(&GroupData.HairRenderData, LODCount);
+		Data.SimRootData = FHairStrandsRootData(&GroupData.HairSimulationData, LODCount);
+	}
+
+	UGroomBindingAsset::FHairGroupResources& OutHairGroupResources = BindingAsset->HairGroupResources;
+	check(OutHairGroupResources.Num() == 0);
+	for (UGroomBindingAsset::FHairGroupData& GroupData : OutHairGroupDatas)
+	{
+		UGroomBindingAsset::FHairGroupResource& Resource = OutHairGroupResources.AddDefaulted_GetRef();
+		Resource.SimRootResources = new FHairStrandsRootResource(GroupData.SimRootData);
+		Resource.RenRootResources = new FHairStrandsRootResource(GroupData.RenRootData);
+
+		Resource.SimRootResources->InitRHI();
+		Resource.RenRootResources->InitRHI();
+	}
+
+	FHairStrandsProjectionHairData RenProjectionDatas;
+	FHairStrandsProjectionHairData SimProjectionDatas;
+	const uint32 GroupCount = OutHairGroupResources.Num();
+	for (const UGroomBindingAsset::FHairGroupResource& GroupResources : OutHairGroupResources)
+	{
+		RenProjectionDatas.HairGroups.Add(ToProjectionHairData(GroupResources.RenRootResources));
+		SimProjectionDatas.HairGroups.Add(ToProjectionHairData(GroupResources.SimRootResources));
+	}
+
+	FSkeletalMeshRenderData* TargetRenderData = TargetSkeletalMesh->GetResourceForRendering();
+	FHairStrandsProjectionMeshData TargetMeshData = ExtractMeshData(TargetRenderData);
+
+	// Create mapping between the source & target using their UV
+	// The lifetime of 'TransferredPositions' needs to encompass RunProjection
+	TArray<FRWBuffer> TransferredPositions;
+
+	if (FSkeletalMeshRenderData* SourceRenderData = SourceSkeletalMesh ? SourceSkeletalMesh->GetResourceForRendering() : nullptr)
+	{
+		FHairStrandsProjectionMeshData SourceMeshData = ExtractMeshData(SourceRenderData);
+		RunMeshTransfer(
+			RHICmdList,
+			SourceMeshData,
+			TargetMeshData,
+			TransferredPositions);
+
+		for (uint32 LODIndex = 0; LODIndex < LODCount; ++LODIndex)
+		{
+			for (FHairStrandsProjectionMeshData::Section& Section : TargetMeshData.LODs[LODIndex].Sections)
+			{
+				Section.PositionBuffer = TransferredPositions[LODIndex].SRV;
+			}
+		}
+	}
+
+	RunProjection(
+		RHICmdList,
+		FTransform::Identity,
+		TargetMeshData,
+		RenProjectionDatas,
+		SimProjectionDatas);
+
+
+	// Enforce GPU sync to read back data on CPU
+	GDynamicRHI->RHISubmitCommandsAndFlushGPU();
+	GDynamicRHI->RHIBlockUntilGPUIdle();
+
+	for (uint32 GroupIt = 0; GroupIt < GroupCount; ++GroupIt)
+	{
+		ReadbackGroupData(RHICmdList, OutHairGroupDatas[GroupIt].SimRootData, OutHairGroupResources[GroupIt].SimRootResources);
+		ReadbackGroupData(RHICmdList, OutHairGroupDatas[GroupIt].RenRootData, OutHairGroupResources[GroupIt].RenRootResources);
+	}
+}
+
+void AddGroomBindingTask(UGroomBindingAsset* BindingAsset)
+{
+	EnqueueGroomBindingQuery((void*)BindingAsset, InternalProcessGroomBindingTask);
+}
+
+FHairStrandsProjectionHairData::HairGroup ToProjectionHairData(FHairStrandsRootResource* In)
+{
+	check(IsInRenderingThread());
+
+	FHairStrandsProjectionHairData::HairGroup Out = {};
+	if (!In)
+		return Out;
+
+	Out.RootCount = In->RootData.RootCount;
+	Out.RootPositionBuffer = In->RootPositionBuffer.SRV;
+	Out.RootNormalBuffer = In->RootNormalBuffer.SRV;
+	Out.VertexToCurveIndexBuffer = &In->VertexToCurveIndexBuffer;
+
+	for (FHairStrandsRootResource::FMeshProjectionLOD& MeshLODData : In->MeshProjectionLODs)
+	{
+		FHairStrandsProjectionHairData::LODData& LODData = Out.LODDatas.AddDefaulted_GetRef();
+		LODData.LODIndex = MeshLODData.LODIndex;
+		LODData.RootTriangleIndexBuffer = &MeshLODData.RootTriangleIndexBuffer;
+		LODData.RootTriangleBarycentricBuffer = &MeshLODData.RootTriangleBarycentricBuffer;
+
+		LODData.RestRootTrianglePosition0Buffer = &MeshLODData.RestRootTrianglePosition0Buffer;
+		LODData.RestRootTrianglePosition1Buffer = &MeshLODData.RestRootTrianglePosition1Buffer;
+		LODData.RestRootTrianglePosition2Buffer = &MeshLODData.RestRootTrianglePosition2Buffer;
+
+		LODData.DeformedRootTrianglePosition0Buffer = &MeshLODData.DeformedRootTrianglePosition0Buffer;
+		LODData.DeformedRootTrianglePosition1Buffer = &MeshLODData.DeformedRootTrianglePosition1Buffer;
+		LODData.DeformedRootTrianglePosition2Buffer = &MeshLODData.DeformedRootTrianglePosition2Buffer;
+
+		LODData.Status = &MeshLODData.Status;
+		LODData.bIsValid = MeshLODData.Status == FHairStrandsProjectionHairData::LODData::EStatus::Completed;
+	}
+
+	return Out;
 }
\ No newline at end of file
diff --git a/Engine/Plugins/Experimental/HairStrands/Source/HairStrandsCore/Public/GroomAsset.h b/Engine/Plugins/Experimental/HairStrands/Source/HairStrandsCore/Public/GroomAsset.h
--- a/Engine/Plugins/Experimental/HairStrands/Source/HairStrandsCore/Public/GroomAsset.h
+++ b/Engine/Plugins/Experimental/HairStrands/Source/HairStrandsCore/Public/GroomAsset.h
@@ -9,20 +9,88 @@
 #include "HairStrandsDatas.h"
 #include "RenderResource.h"
 #include "GroomSettings.h"
+#include "Curves/CurveFloat.h"
 #include "HairStrandsInterface.h"
+#include "Engine/SkeletalMesh.h"
 #include "Interfaces/Interface_AssetUserData.h"
 
 #include "GroomAsset.generated.h"
 
+
+/** List of niagara solvers */
+UENUM(BlueprintType)
+enum class EGroomNiagaraSolvers : uint8
+{
+	None = 0 UMETA(Hidden),
+	CosseratRods = 0x02 UMETA(DisplatName = "CosseratRods"),
+	AngularSprings = 0x04 UMETA(DisplatName = "AngularSprings")
+};
+
+/** Size of each strands*/
+UENUM(BlueprintType)
+enum class EGroomStrandsSize : uint8
+{
+	None = 0 UMETA(Hidden),
+	Size2 = 0x02 UMETA(DisplatName = "2"),
+	Size4 = 0x04 UMETA(DisplatName = "4"),
+	Size8 = 0x08 UMETA(DisplatName = "8"),
+	Size16 = 0x10 UMETA(DisplatName = "16"),
+	Size32 = 0x20 UMETA(DisplatName = "32")
+};
+
 class UAssetUserData;
 class UMaterialInterface;
 class UNiagaraSystem;
 struct FHairGroupData;
 
+/* Source/CPU data for root resources (GPU resources are stored into FHairStrandsRootResources) */
+struct FHairStrandsRootData
+{
+	/** Build the hair strands resource */
+	FHairStrandsRootData();
+	FHairStrandsRootData(const FHairStrandsDatas* HairStrandsDatas, uint32 LODCount);
+	void Serialize(FArchive& Ar);
+	void Reset();
+	bool HasProjectionData() const;
+
+	struct FMeshProjectionLOD
+	{
+		int32 LODIndex = -1;
+
+		/* Triangle on which a root is attached */
+		/* When the projection is done with source to target mesh transfer, the projection indices does not match.
+		   In this case we need to separate index computation. The barycentric coords remain the same however. */
+		TArray<FHairStrandsCurveTriangleIndexFormat::Type> RootTriangleIndexBuffer;
+		TArray<FHairStrandsCurveTriangleBarycentricFormat::Type> RootTriangleBarycentricBuffer;
+
+		/* Strand hair roots translation and rotation in rest position relative to the bound triangle. Positions are relative to the rest root center */
+		TArray< FHairStrandsMeshTrianglePositionFormat::Type> RestRootTrianglePosition0Buffer;
+		TArray< FHairStrandsMeshTrianglePositionFormat::Type> RestRootTrianglePosition1Buffer;
+		TArray< FHairStrandsMeshTrianglePositionFormat::Type> RestRootTrianglePosition2Buffer;
+	};
+
+	/* Number of roots */
+	uint32 RootCount;
+
+	/* Curve index for every vertices */
+	TArray<FHairStrandsIndexFormat::Type> VertexToCurveIndexBuffer;
+
+	/* Curve root's positions */
+	TArray<FHairStrandsRootPositionFormat::Type> RootPositionBuffer;
+
+	/* Curve root's normal orientation */
+	TArray<FHairStrandsRootNormalFormat::Type> RootNormalBuffer;
+
+	/* Store the hair projection information for each mesh LOD */
+	TArray<FMeshProjectionLOD> MeshProjectionLODs;
+};
+
 /* Render buffers for root deformation for dynamic meshes */
 struct FHairStrandsRootResource : public FRenderResource
 {
 	/** Build the hair strands resource */
+	FHairStrandsRootResource();
+	FHairStrandsRootResource(const FHairStrandsRootData& RootData);
 	FHairStrandsRootResource(const FHairStrandsDatas* HairStrandsDatas, uint32 LODCount);
 
 	/* Init the buffer */
@@ -34,8 +102,12 @@ struct FHairStrandsRootResource : public FRenderResource
 	/* Get the resource name */
 	virtual FString GetFriendlyName() const override { return TEXT("FHairStrandsRootResource"); }
 	
+	/* Populate GPU LOD data from RootData (this function doesn't initialize resources) */
+	void PopulateFromRootData();
+
 	FRWBuffer RootPositionBuffer;
 	FRWBuffer RootNormalBuffer;
+	FRWBuffer VertexToCurveIndexBuffer;
 
 	struct FMeshProjectionLOD
 	{
@@ -43,17 +115,17 @@ struct FHairStrandsRootResource : public FRenderResource
 		int32 LODIndex = -1;
 
 		/* Triangle on which a root is attached */
+		/* When the projection is done with source to target mesh transfer, the projection indices does not match. 
+		   In this case we need to separate index computation. The barycentric coords remain the same however. */
 		FRWBuffer RootTriangleIndexBuffer;
 		FRWBuffer RootTriangleBarycentricBuffer;
 	
 		/* Strand hair roots translation and rotation in rest position relative to the bound triangle. Positions are relative to the rest root center */
-		FVector	  RestRootOffset = FVector::ZeroVector;
 		FRWBuffer RestRootTrianglePosition0Buffer;
 		FRWBuffer RestRootTrianglePosition1Buffer;
 		FRWBuffer RestRootTrianglePosition2Buffer;
 
 		/* Strand hair roots translation and rotation in triangle-deformed position relative to the bound triangle. Positions are relative the deformed root center*/
-		FVector   DeformedRootOffset = FVector::ZeroVector;
 		FRWBuffer DeformedRootTrianglePosition0Buffer;
 		FRWBuffer DeformedRootTrianglePosition1Buffer;
 		FRWBuffer DeformedRootTrianglePosition2Buffer;
@@ -62,18 +134,8 @@ struct FHairStrandsRootResource : public FRenderResource
 	/* Store the hair projection information for each mesh LOD */
 	TArray<FMeshProjectionLOD> MeshProjectionLODs;
 
-	/* Strand hair vertex to curve index */
-	FRWBuffer VertexToCurveIndexBuffer;
-
-	const uint32 RootCount;
-	/* Curve index for every vertices */
-	TArray<FHairStrandsIndexFormat::Type> CurveIndices;
-
-	/* Curve root's positions */
-	TArray<FHairStrandsRootPositionFormat::Type> RootPositions;
-
-	/* Curve root's normal orientation */
-	TArray<FHairStrandsRootNormalFormat::Type> RootNormals;
+	/* Store CPU data for root info & root binding */
+	FHairStrandsRootData RootData;
 };
 
 /* Render buffers that will be used for rendering */
@@ -247,10 +309,6 @@ struct HAIRSTRANDSCORE_API FHairGroupInfo
 	UPROPERTY(EditAnywhere, Category="Rendering")
 	UMaterialInterface* Material = nullptr;
 
-	// Currently hide the Nigara simulation slot as it is not used, and could confuse users
-	//UPROPERTY(EditAnywhere, Category="Simulation", meta = (DisplayName = "Niagara System Asset"))
-	//UNiagaraSystem* NiagaraAsset = nullptr;
-
 	friend FArchive& operator<<(FArchive& Ar, FHairGroupInfo& GroupInfo);
 };
 
@@ -275,7 +333,7 @@ struct HAIRSTRANDSCORE_API FHairGroupData
 UCLASS(BlueprintType, hidecategories = (Object))
 class HAIRSTRANDSCORE_API UGroomAsset : public UObject, public IInterface_AssetUserData
 {
-	GENERATED_BODY()
+	GENERATED_UCLASS_BODY()
 
 #if WITH_EDITOR
 	/** Notification when anything changed */
@@ -298,6 +356,126 @@ public:
 	UPROPERTY(VisibleAnywhere, BlueprintReadWrite, Category = "BuildSettings", meta = (ClampMin = "0.01", UIMin = "0.01", UIMax = "1.0"))
 	float HairToGuideDensity = 0.1f;
 
+	/** Group Id to be simulated */
+	UPROPERTY(EditAnywhere, Category = "Hair Physics|Group Setup", meta = (ToolTip = "Group ID that will be simulated with niagara"))
+	int32 SimulatedGroup;
+
+	/** Enable the simulation for the group */
+	UPROPERTY(EditAnywhere, Category = "Hair Physics|Group Setup", meta = (ToolTip = "Enable the simulation on that group"))
+	bool EnableSimulation;
+
+	/** Enable the simulation for the group */
+	UPROPERTY(EditAnywhere, Category = "Hair Physics|Group Setup", meta = (ToolTip = "Niagara solver to be used for simulation"))
+	EGroomNiagaraSolvers NiagaraSolver;
+
+	/** Number of substeps to be used */
+	UPROPERTY(EditAnywhere, Category = "Hair Physics|Niagara Solver", meta = (ToolTip = "Number of sub steps to be done per frame. The actual solver calls are done at 24 fps"))
+	int32 SubSteps;
+
+	/** Number of iterations for the constraint solver  */
+	UPROPERTY(EditAnywhere, Category = "Hair Physics|Niagara Solver", meta = (ToolTip = "Number of iterations to solve the constraints with the xpbd solver"))
+	int32 IterationCount;
+
+	/** Acceleration vector in cm/s2 to be used for the gravity*/
+	UPROPERTY(EditAnywhere, Category = "Hair Physics|External Forces", meta = (ToolTip = "Acceleration vector in cm/s2 to be used for the gravity"))
+	FVector GravityVector;
+
+	/** Coefficient between 0 and 1 to be used for the air drag */
+	UPROPERTY(EditAnywhere, Category = "Hair Physics|External Forces", meta = (ToolTip = "Coefficient between 0 and 1 to be used for the air drag"))
+	float AirDrag;
+
+	/** Velocity of the surrounding air in cm/s */
+	UPROPERTY(EditAnywhere, Category = "Hair Physics|External Forces", meta = (ToolTip = "Velocity of the surrounding air in cm/s"))
+	FVector AirVelocity;
+
+	/** Enable the solve of the bend constraint during the xpbd loop */
+	UPROPERTY(EditAnywhere, Category = "Hair Physics|Solver Constraints|Bend Constraint", meta = (ToolTip = "Enable the solve of the bend constraint during the xpbd loop"))
+	bool SolveBend;
+
+	/** Enable ther projection of the bend constraint after the xpbd loop */
+	UPROPERTY(EditAnywhere, Category = "Hair Physics|Solver Constraints|Bend Constraint", meta = (ToolTip = "Enable ther projection of the bend constraint after the xpbd loop"))
+	bool ProjectBend;
+
+	/** Damping for the bend constraint between 0 and 1 */
+	UPROPERTY(EditAnywhere, Category = "Hair Physics|Solver Constraints|Bend Constraint", meta = (ToolTip = "Damping for the bend constraint between 0 and 1"))
+	float BendDamping;
+
+	/** Stiffness for the bend constraint in GPa */
+	UPROPERTY(EditAnywhere, Category = "Hair Physics|Solver Constraints|Bend Constraint", meta = (ToolTip = "Stiffness for the bend constraint in GPa"))
+	float BendStiffness;
+
+	/** Stiffness scale along the strand */
+	UPROPERTY(EditAnywhere, Category = "Hair Physics|Solver Constraints|Bend Constraint", meta = (ViewMinInput = "0.0", ViewMaxInput = "1.0", ViewMinOutput = "0.0", ViewMaxOutput = "1.0", TimeLineLength = "1.0", XAxisName = "Strand Coordinate (0,1)", YAxisName = "Bend Scale", ToolTip = "This curve determines how much the bend stiffness will be scaled along each strand. \n The X axis range is [0,1], 0 mapping the root and 1 the tip"))
+	FRuntimeFloatCurve BendScale;
+
+	/** Enable the solve of the stretch constraint during the xpbd loop */
+	UPROPERTY(EditAnywhere, Category = "Hair Physics|Solver Constraints|Stretch Constraint", meta = (ToolTip = "Enable the solve of the stretch constraint during the xpbd loop"))
+	bool SolveStretch;
+
+	/** Enable the projection of the stretch constraint after the xpbd loop */
+	UPROPERTY(EditAnywhere, Category = "Hair Physics|Solver Constraints|Stretch Constraint", meta = (ToolTip = "Enable ther projection of the stretch constraint after the xpbd loop"))
+	bool ProjectStretch;
+
+	/** Damping for the stretch constraint between 0 and 1 */
+	UPROPERTY(EditAnywhere, Category = "Hair Physics|Solver Constraints|Stretch Constraint", meta = (ToolTip = "Damping for the stretch constraint between 0 and 1"))
+	float StretchDamping;
+
+	/** Stiffness for the stretch constraint in GPa */
+	UPROPERTY(EditAnywhere, Category = "Hair Physics|Solver Constraints|Stretch Constraint", meta = (ToolTip = "Stiffness for the stretch constraint in GPa"))
+	float StretchStiffness;
+
+	/** Stretch scale along the strand  */
+	UPROPERTY(EditAnywhere, Category = "Hair Physics|Solver Constraints|Stretch Constraint", meta = (ViewMinInput = "0.0", ViewMaxInput = "1.0", ViewMinOutput = "0.0", ViewMaxOutput = "1.0", TimeLineLength = "1.0", XAxisName = "Strand Coordinate (0,1)", YAxisName = "Stretch Scale", ToolTip = "This curve determines how much the stretch stiffness will be scaled along each strand. \n The X axis range is [0,1], 0 mapping the root and 1 the tip"))
+	FRuntimeFloatCurve StretchScale;
+
+	/** Enable the solve of the collision constraint during the xpbd loop  */
+	UPROPERTY(EditAnywhere, Category = "Hair Physics|Solver Constraints|Collision Constraint", meta = (ToolTip = "Enable the solve of the collision constraint during the xpbd loop"))
+	bool SolveCollision;
+
+	/** Enable ther projection of the collision constraint after the xpbd loop */
+	UPROPERTY(EditAnywhere, Category = "Hair Physics|Solver Constraints|Collision Constraint", meta = (ToolTip = "Enable ther projection of the collision constraint after the xpbd loop"))
+	bool ProjectCollision;
+
+	/** Static friction used for collision against the physics asset */
+	UPROPERTY(EditAnywhere, Category = "Hair Physics|Solver Constraints|Collision Constraint", meta = (ToolTip = "Static friction used for collision against the physics asset"))
+	float StaticFriction;
+
+	/** Kinetic friction used for collision against the physics asset*/
+	UPROPERTY(EditAnywhere, Category = "Hair Physics|Solver Constraints|Collision Constraint", meta = (ToolTip = "Kinetic friction used for collision against the physics asset"))
+	float KineticFriction;
+
+	/** Viscosity parameter between 0 and 1 that will be used for self collision */
+	UPROPERTY(EditAnywhere, Category = "Hair Physics|Solver Constraints|Collision Constraint", meta = (ToolTip = "Viscosity parameter between 0 and 1 that will be used for self collision"))
+	float StrandsViscosity;
+
+	/** Radius that will be used for the collision detection against the physics asset */
+	UPROPERTY(EditAnywhere, Category = "Hair Physics|Solver Constraints|Collision Constraint", meta = (ToolTip = "Radius that will be used for the collision detection against the physics asset"))
+	float CollisionRadius;
+
+	/** Radius scale along the strand */
+	UPROPERTY(EditAnywhere, Category = "Hair Physics|Solver Constraints|Collision Constraint", meta = (ViewMinInput = "0.0", ViewMaxInput = "1.0", ViewMinOutput = "0.0", ViewMaxOutput = "1.0", TimeLineLength = "1.0", XAxisName = "Strand Coordinate (0,1)", YAxisName = "Collision Radius", ToolTip = "This curve determines how much the collision radius will be scaled along each strand. \n The X axis range is [0,1], 0 mapping the root and 1 the tip"))
+	FRuntimeFloatCurve RadiusScale;
+
+	/** Number of particles per guide that will be used for simulation*/
+	UPROPERTY(EditAnywhere, Category = "Hair Physics|Strands Parameters", meta = (ToolTip = "Number of particles per guide that will be used for simulation"))
+	EGroomStrandsSize StrandsSize;
+
+	/** Density of the strands in g/cm3 */
+	UPROPERTY(EditAnywhere, Category = "Hair Physics|Strands Parameters", meta = (ToolTip = "Density of the strands in g/cm3"))
+	float StrandsDensity;
+
+	/** Smoothing between 0 and 1 of the incoming guides curves for better stability */
+	UPROPERTY(EditAnywhere, Category = "Hair Physics|Strands Parameters", meta = (ToolTip = "Smoothing between 0 and 1 of the incoming guides curves for better stability"))
+	float StrandsSmoothing;
+
+	/** Strands thickness in cm that will be used for mass and inertia computation */
+	UPROPERTY(EditAnywhere, Category = "Hair Physics|Strands Parameters", meta = (ToolTip = "Strands thickness in cm that will be used for mass and inertia computation"))
+	float StrandsThickness;
+
+	/** Thickness scale along the curve */
+	UPROPERTY(EditAnywhere, Category = "Hair Physics|Strands Parameters", meta = (ViewMinInput = "0.0", ViewMaxInput = "1.0", ViewMinOutput = "0.0", ViewMaxOutput = "1.0", TimeLineLength = "1.0", XAxisName = "Strand Coordinate (0,1)", YAxisName = "Strands Thickness", ToolTip = "This curve determines how much the strands thickness will be scaled along each strand. \n The X axis range is [0,1], 0 mapping the root and 1 the tip"))
+	FRuntimeFloatCurve ThicknessScale;
+
 #if WITH_EDITOR
 	FOnGroomAssetChanged& GetOnGroomAssetChanged() { return OnGroomAssetChanged;  }
 
@@ -329,11 +507,6 @@ public:
 	/** Release the hair strands resource. */
 	void ReleaseResource();
 
-	/**
-	 * Initializes an instance for use with this vector field.
-	 */
-	void InitInstance(class FVectorFieldInstance* Instance, bool bPreviewInstance);
-
 	void Reset();
 
 	int32 GetNumHairGroups() const;
@@ -375,3 +548,101 @@ private:
 	bool bIsCacheable = true;
 #endif
 };
+
+USTRUCT(BlueprintType)
+struct HAIRSTRANDSCORE_API FGoomBindingGroupInfo
+{
+	GENERATED_BODY()
+
+	UPROPERTY(VisibleAnywhere, Category = "Info", meta = (DisplayName = "Curve Count"))
+	int32 RenRootCount = 0;
+
+	UPROPERTY(VisibleAnywhere, Category = "Info", meta = (DisplayName = "Curve LOD"))
+	int32 RenLODCount = 0;
+
+	UPROPERTY(VisibleAnywhere, Category = "Info", meta = (DisplayName = "Guide Count"))
+	int32 SimRootCount = 0;
+
+	UPROPERTY(VisibleAnywhere, Category = "Info", meta = (DisplayName = "Guide LOD"))
+	int32 SimLODCount = 0;
+};
+
+/**
+ * Implements an asset that can be used to store binding information between a groom and a skeletal mesh
+ */
+UCLASS(BlueprintType, hidecategories = (Object))
+class HAIRSTRANDSCORE_API UGroomBindingAsset : public UObject
+{
+	GENERATED_BODY()
+
+#if WITH_EDITOR
+	/** Notification when anything changed */
+	DECLARE_MULTICAST_DELEGATE(FOnGroomBindingAssetChanged);
+#endif
+
+public:
+
+	/** Groom to bind. */
+	UPROPERTY(VisibleAnywhere, BlueprintReadWrite, Category = "BuildSettings", meta = (ClampMin = "0.01", UIMin = "0.01", UIMax = "1.0"))
+	UGroomAsset* Groom;
+
+	/** Skeletal mesh on which the groom has been authored. This is optional, and used only if the hair
+		binding is done a different mesh than the one which it has been authored */
+	UPROPERTY(VisibleAnywhere, BlueprintReadWrite, Category = "BuildSettings", meta = (ClampMin = "0.01", UIMin = "0.01", UIMax = "1.0"))
+	USkeletalMesh* SourceSkeletalMesh;
+
+	/** Skeletal mesh on which the groom is attached to. */
+	UPROPERTY(VisibleAnywhere, BlueprintReadWrite, Category = "BuildSettings", meta = (ClampMin = "0.01", UIMin = "0.01", UIMax = "1.0"))
+	USkeletalMesh* TargetSkeletalMesh;
+
+	UPROPERTY(EditAnywhere, EditFixedSize, BlueprintReadWrite, Category = "HairGroups", meta = (DisplayName = "Group"))
+	TArray<FGoomBindingGroupInfo> GroupInfos;
+
+	/** GPU and CPU binding data for both simulation and rendering. */
+	struct FHairGroupResource
+	{
+		FHairStrandsRootResource* SimRootResources = nullptr;
+		FHairStrandsRootResource* RenRootResources = nullptr;
+	};
+	typedef TArray<FHairGroupResource> FHairGroupResources;
+	FHairGroupResources HairGroupResources;
+
+	struct FHairGroupData
+	{
+		FHairStrandsRootData SimRootData;
+		FHairStrandsRootData RenRootData;
+	};
+	typedef TArray<FHairGroupData> FHairGroupDatas;
+	FHairGroupDatas HairGroupDatas;
+
+	//~ Begin UObject Interface.
+	virtual void PostLoad() override;
+	virtual void PreSave(const class ITargetPlatform* TargetPlatform) override;
+	virtual void PostSaveRoot(bool bCleanupIsRequired) override;
+	virtual void BeginDestroy() override;
+	virtual void Serialize(FArchive& Ar) override;
+
+#if WITH_EDITOR
+	FOnGroomBindingAssetChanged& GetOnGroomBindingAssetChanged() { return OnGroomBindingAssetChanged; }
+
+	/**  Part of UObject interface  */
+	virtual void PostEditChangeProperty(FPropertyChangedEvent& PropertyChangedEvent) override;
+
+#endif // WITH_EDITOR
+
+	/** Initialize resources. */
+	void InitResource();
+
+	/** Update resources. */
+	void UpdateResource();
+
+	/** Release the hair strands resource. */
+	void ReleaseResource();
+
+	void Reset();
+
+	//private :
+#if WITH_EDITOR
+	FOnGroomBindingAssetChanged OnGroomBindingAssetChanged;
+#endif
+};
diff --git a/Engine/Plugins/Experimental/HairStrands/Source/HairStrandsCore/Public/GroomComponent.h b/Engine/Plugins/Experimental/HairStrands/Source/HairStrandsCore/Public/GroomComponent.h
--- a/Engine/Plugins/Experimental/HairStrands/Source/HairStrandsCore/Public/GroomComponent.h
+++ b/Engine/Plugins/Experimental/HairStrands/Source/HairStrandsCore/Public/GroomComponent.h
@@ -23,6 +23,10 @@ public:
 	UPROPERTY(EditAnywhere, BlueprintReadWrite, Category = "Groom")
 	UGroomAsset* GroomAsset;
 
+	/** Niagara component that will be attached to the system*/
+	//UPROPERTY()
+	class UNiagaraComponent* NiagaraComponent;
+
 	/** 
 	 * When activated, the groom will be attached and skinned onto the skeletal mesh, if the groom component is a child of a skeletal/skinned component.
 	 * This requires the following projection settings: 
@@ -32,15 +36,35 @@ public:
 	UPROPERTY(EditAnywhere, BlueprintReadOnly, Category = "Groom")
 	bool bBindGroomToSkeletalMesh;
 
+	/** Skeletal mesh on which the groom has been authored. If not provided, the skeletal mesh on which the groom component is attached will be used. If provided, both skeletal mesh needs to share the same topology. */
+	UPROPERTY(EditAnywhere, BlueprintReadOnly, Category = "Groom")
+	class USkeletalMesh* SourceSkeletalMesh;
+
+	/** Optional binding asset for binding a groom onto a skeletal mesh. If the binding asset is not provided the projection is done at runtime, which implies a large GPU cost at startup time. */
+	UPROPERTY(EditAnywhere, BlueprintReadOnly, Category = "Groom")
+	class UGroomBindingAsset* BindingAsset;
+
+	/** Create default Niagara component for simulation */
+	UPROPERTY(EditAnywhere, BlueprintReadOnly, Category = "Groom")
+	bool	bCreateNiagaraComponent;
+
 	/** Boolean to check when animation has been loaded */
 	bool bResetSimulation;
 
 	/** Previous bone matrix to compare the difference and decide to reset or not the simulation */
 	FMatrix	PrevBoneMatrix;
 
-	/** Listen for the animation event to trigger the sim */
-	UFUNCTION()
-	void ResetSimulation();
+	/** Update Niagara component */
+	void UpdateNiagaraComponent();
+
+	/** Update Group Description */
+	void UpdateHairGroupsDesc();
+
+	/** Enable simulated groups */
+	void EnableSimulatedGroups();
+
+	/** Disable Simulated groups */
+	void DisableSimulatedGroups();
 
 	//~ Begin UActorComponent Interface.
 	virtual void OnRegister() override;
@@ -86,12 +110,26 @@ public:
 	void Invalidate();
 #endif
 
+	void SetGroomAsset(UGroomAsset* Asset);
+	void SetGroomAsset(UGroomAsset* Asset, UGroomBindingAsset* InBinding);
+	void SetHairLengthScale(float Scale);
+	void SetHairWidth(float HairWidth);
+	void SetBinding(bool bBind);
+	void SetBinding(UGroomBindingAsset* InBinding);
+
+	/** Groom's groups info. */
+	UPROPERTY(EditAnywhere, BlueprintReadWrite, Category = "Groom")
+		TArray<FHairGroupDesc> GroomGroupsDesc;
+
+private:
+
 	struct FHairGroupResource
 	{
 		// Sim to rendering interpolation resources
 		FHairStrandsInterpolationResource* InterpolationResource = nullptr;
 
 		// Projection resources
+		bool bOwnRootResourceAllocation = true;
 		struct FHairStrandsRootResource* RenRootResources = nullptr;
 		struct FHairStrandsRootResource* SimRootResources = nullptr;
 	#if RHI_RAYTRACING
@@ -109,13 +147,14 @@ public:
 		FHairStrandsRestResource* RenderRestResources = nullptr;
 		FHairStrandsRestResource* SimRestResources = nullptr;
 	};
-	typedef TArray<FHairGroupResource> FHairGroupResources;
 
-	/** Groom's groups info. */
-	UPROPERTY(EditAnywhere, BlueprintReadWrite, Category = "Groom")
-	TArray<FHairGroupDesc> GroomGroupsDesc;
+	struct FHairGroupResources
+	{
+		TArray<FHairGroupResource> HairGroups;
+	};
+	static void DeleteHairGroupResources(FHairGroupResources*& InHairGroupResources);
 
-	FHairGroupResources HairGroupResources;
+	FHairGroupResources* HairGroupResources = nullptr;
 	struct FHairStrandsInterpolationOutput* InterpolationOutput = nullptr;
 	struct FHairStrandsInterpolationInput* InterpolationInput = nullptr;
 
@@ -133,10 +172,12 @@ private:
 
 	void InitResources();
 	void ReleaseResources();
+	void UpdateHairGroupsDescAndInvalidateRenderState();
 
 	virtual void GetUsedMaterials(TArray<UMaterialInterface*>& OutMaterials, bool bGetDebugMaterials) const override;
 
 	friend class FGroomComponentRecreateRenderStateContext;
+	friend class FHairStrandsSceneProxy;
 };
 
 /** Used to recreate render context for all GroomComponents that use a given GroomAsset */
diff --git a/Engine/Plugins/Experimental/HairStrands/Source/HairStrandsCore/Public/GroomCreateBindingOptions.h b/Engine/Plugins/Experimental/HairStrands/Source/HairStrandsCore/Public/GroomCreateBindingOptions.h
new file mode 100644
--- /dev/null
+++ b/Engine/Plugins/Experimental/HairStrands/Source/HairStrandsCore/Public/GroomCreateBindingOptions.h
@@ -0,0 +1,26 @@
+// Copyright Epic Games, Inc. All Rights Reserved.
+
+#pragma once
+
+#include "CoreMinimal.h"
+#include "GroomSettings.h"
+#include "UObject/Object.h"
+#include "Engine/SkeletalMesh.h"
+#include "GroomAsset.h"
+#include "GroomCreateBindingOptions.generated.h"
+
+UCLASS(BlueprintType, config = EditorPerProjectUserSettings, HideCategories = ("Hidden"))
+class HAIRSTRANDSCORE_API UGroomCreateBindingOptions : public UObject
+{
+	GENERATED_UCLASS_BODY()
+
+public:
+	/** Skeletal mesh on which the groom has been authored. This is optional, and used only if the hair
+		binding is done a different mesh than the one which it has been authored */
+	UPROPERTY(EditAnywhere, BlueprintReadWrite, meta = (ShowOnlyInnerProperties), Category = Conversion)
+	USkeletalMesh* SourceSkeletalMesh;
+
+	/** Skeletal mesh on which the groom is attached to. */
+	UPROPERTY(EditAnywhere, BlueprintReadWrite, meta = (ShowOnlyInnerProperties), Category = BuildSettings)
+	USkeletalMesh* TargetSkeletalMesh;
+};
diff --git a/Engine/Plugins/Experimental/HairStrands/Source/HairStrandsCore/Public/HairStrandsDatas.h b/Engine/Plugins/Experimental/HairStrands/Source/HairStrandsCore/Public/HairStrandsDatas.h
--- a/Engine/Plugins/Experimental/HairStrands/Source/HairStrandsCore/Public/HairStrandsDatas.h
+++ b/Engine/Plugins/Experimental/HairStrands/Source/HairStrandsCore/Public/HairStrandsDatas.h
@@ -72,6 +72,7 @@ struct FVector4_16
 	FFloat16 Z;
 	FFloat16 W;
 }; 
+FArchive& operator<<(FArchive& Ar, FVector4_16& Vertex);
 
 struct FHairStrandsPositionFormat
 {
@@ -173,26 +174,14 @@ struct FHairStrandsWeightFormat
 
 /** 
  * Skinned mesh triangle vertex position format
- * Two precision options are available: 4x16bits or 4x32bits
- * Triangle vertices are relative to their bounding box in order to preserve precision, 
- * however this is sometime not enough for large asset, this is why by default the format 
- * use 32bits precision
-*/
+ */
 struct FHairStrandsMeshTrianglePositionFormat
 {
-#if 1
 	using Type = FVector4;
 	static const uint32 ComponentCount = 1;
 	static const uint32 SizeInByte = sizeof(Type);
 	static const EVertexElementType VertexElementType = VET_Float4;
 	static const EPixelFormat Format = PF_A32B32G32R32F;
-#else
-	using Type = FVector4_16;
-	static const uint32 ComponentCount = 1;
-	static const uint32 SizeInByte = sizeof(Type);
-	static const EVertexElementType VertexElementType = VET_Float4;
-	static const EPixelFormat Format = PF_FloatRGBA;
-#endif
 };
 
 struct FHairStrandsCurveTriangleIndexFormat
diff --git a/Engine/Plugins/Experimental/HairStrands/Source/HairStrandsCore/Public/HairStrandsRendering.h b/Engine/Plugins/Experimental/HairStrands/Source/HairStrandsCore/Public/HairStrandsRendering.h
--- a/Engine/Plugins/Experimental/HairStrands/Source/HairStrandsCore/Public/HairStrandsRendering.h
+++ b/Engine/Plugins/Experimental/HairStrands/Source/HairStrandsCore/Public/HairStrandsRendering.h
@@ -145,7 +145,9 @@ struct FHairStrandsInterpolationOutput
 void ResetHairStrandsInterpolation(
 	FRHICommandListImmediate& RHICmdList,
 	FHairStrandsInterpolationInput* InInput,
-	FHairStrandsInterpolationOutput* InOutput);
+	FHairStrandsInterpolationOutput* InOutput,
+	struct FHairStrandsProjectionHairData& SimHairDatas,
+	int32 LODIndex);
 
 void ComputeHairStrandsInterpolation(
 	FRHICommandListImmediate& RHICmdList,
@@ -156,3 +158,5 @@ void ComputeHairStrandsInterpolation(
 	struct FHairStrandsProjectionHairData& SimHairDatas,
 	int32 LODIndex,
 	struct FHairStrandClusterData* ClusterData);
+
+HAIRSTRANDSCORE_API void AddGroomBindingTask(class UGroomBindingAsset* BindingAsset);
\ No newline at end of file
diff --git a/Engine/Plugins/Experimental/HairStrands/Source/HairStrandsEditor/Private/GroomActions.cpp b/Engine/Plugins/Experimental/HairStrands/Source/HairStrandsEditor/Private/GroomActions.cpp
--- a/Engine/Plugins/Experimental/HairStrands/Source/HairStrandsEditor/Private/GroomActions.cpp
+++ b/Engine/Plugins/Experimental/HairStrands/Source/HairStrandsEditor/Private/GroomActions.cpp
@@ -4,9 +4,12 @@
 #include "GroomAsset.h"
 
 #include "EditorFramework/AssetImportData.h"
+#include "HairStrandsRendering.h"
 #include "GroomAssetImportData.h"
 #include "GroomImportOptions.h"
 #include "GroomImportOptionsWindow.h"
+#include "GroomCreateBindingOptions.h"
+#include "GroomCreateBindingOptionsWindow.h"
 #include "Toolkits/SimpleAssetEditor.h"
 #include "ToolMenuSection.h"
 
@@ -36,6 +39,17 @@ void FGroomActions::GetActions(const TArray<UObject*>& InObjects, FToolMenuSecti
 			FCanExecuteAction::CreateSP(this, &FGroomActions::CanRebuild, GroomAssets)
 		)
 	);
+
+	Section.AddMenuEntry(
+		"CreateBindingAsset",
+		LOCTEXT("CreateBindingAsset", "Create Binding"),
+		LOCTEXT("CreateBindingAssetTooltip", "Create a binding asset between a skeletal mesh and a groom asset"),
+		FSlateIcon(FEditorStyle::GetStyleSetName(), "ContentBrowser.AssetActions"),
+		FUIAction(
+			FExecuteAction::CreateSP(this, &FGroomActions::ExecuteCreateBindingAsset, GroomAssets),
+			FCanExecuteAction::CreateSP(this, &FGroomActions::CanCreateBindingAsset, GroomAssets)
+		)
+	);
 }
 
 uint32 FGroomActions::GetCategories()
@@ -126,4 +140,45 @@ void FGroomActions::ExecuteRebuild(TArray<TWeakObjectPtr<UGroomAsset>> Objects)
 	}
 }
 
+bool FGroomActions::CanCreateBindingAsset(TArray<TWeakObjectPtr<UGroomAsset>> Objects) const
+{
+	for (TWeakObjectPtr<UGroomAsset> GroomAsset : Objects)
+	{
+		if (GroomAsset.IsValid())
+		{
+			return true;
+		}
+	}
+	return false;
+}
+
+void FGroomActions::ExecuteCreateBindingAsset(TArray<TWeakObjectPtr<UGroomAsset>> Objects) const
+{
+	for (TWeakObjectPtr<UGroomAsset> GroomAsset : Objects)
+	{
+		if (GroomAsset.IsValid())
+		{
+			{
+
+				// Duplicate the options to prevent dirtying the asset when they are modified but the rebuild is cancelled
+				UGroomCreateBindingOptions* CurrentOptions = NewObject<UGroomCreateBindingOptions>();
+				TSharedPtr<SGroomCreateBindingOptionsWindow> GroomOptionWindow = SGroomCreateBindingOptionsWindow::DisplayCreateBindingOptions(CurrentOptions);
+
+				if (!GroomOptionWindow->ShouldCreate())
+				{
+					continue;
+				}
+				else if (GroomAsset.Get() && CurrentOptions && CurrentOptions->TargetSkeletalMesh)
+				{
+					UGroomBindingAsset* BindingAsset = CreateGroomBindinAsset(GroomAsset.Get(), CurrentOptions->SourceSkeletalMesh, CurrentOptions->TargetSkeletalMesh);
+
+					// The binding task will generate and set the binding value back to the binding asset.
+					// This code is not thread safe.
+					AddGroomBindingTask(BindingAsset);
+				}
+			}
+		}
+	}
+}
+
 #undef LOCTEXT_NAMESPACE
\ No newline at end of file
diff --git a/Engine/Plugins/Experimental/HairStrands/Source/HairStrandsEditor/Private/GroomCreateBindingOptionsWindow.cpp b/Engine/Plugins/Experimental/HairStrands/Source/HairStrandsEditor/Private/GroomCreateBindingOptionsWindow.cpp
new file mode 100644
--- /dev/null
+++ b/Engine/Plugins/Experimental/HairStrands/Source/HairStrandsEditor/Private/GroomCreateBindingOptionsWindow.cpp
@@ -0,0 +1,233 @@
+// Copyright Epic Games, Inc. All Rights Reserved.
+
+#include "GroomCreateBindingOptionsWindow.h"
+
+#include "EditorStyleSet.h"
+#include "Framework/Application/SlateApplication.h"
+#include "GroomCreateBindingOptions.h"
+#include "IDetailsView.h"
+#include "Interfaces/IMainFrameModule.h"
+#include "Modules/ModuleManager.h"
+#include "PropertyEditorModule.h"
+#include "AssetToolsModule.h"
+#include "AssetRegistryModule.h"
+#include "ContentBrowserModule.h"
+#include "Subsystems/AssetEditorSubsystem.h"
+#include "IContentBrowserSingleton.h"
+#include "Misc/MessageDialog.h"
+#include "Widgets/Input/SButton.h"
+#include "Widgets/Layout/SBorder.h"
+#include "Widgets/Layout/SUniformGridPanel.h"
+#include "Widgets/SBoxPanel.h"
+#include "Widgets/Text/STextBlock.h"
+#include "Editor.h"
+
+#define LOCTEXT_NAMESPACE "GroomOptions"
+
+void SGroomCreateBindingOptionsWindow::Construct(const FArguments& InArgs)
+{
+	BindingOptions = InArgs._BindingOptions;
+	WidgetWindow = InArgs._WidgetWindow;
+
+	FPropertyEditorModule& PropertyEditorModule = FModuleManager::GetModuleChecked<FPropertyEditorModule>("PropertyEditor");
+	FDetailsViewArgs DetailsViewArgs;
+	DetailsViewArgs.bAllowSearch = false;
+	DetailsViewArgs.NameAreaSettings = FDetailsViewArgs::HideNameArea;
+	DetailsView = PropertyEditorModule.CreateDetailView(DetailsViewArgs);
+	DetailsView->SetObject(BindingOptions);
+
+	this->ChildSlot
+	[
+		SNew(SVerticalBox)
+
+		+ SVerticalBox::Slot()
+		.AutoHeight()
+		.Padding(2)
+		[
+			SNew(SBorder)
+			.Padding(FMargin(3))
+			.BorderImage(FEditorStyle::GetBrush("ToolPanel.GroupBorder"))
+			[
+				SNew(SHorizontalBox)
+				+ SHorizontalBox::Slot()
+				.AutoWidth()
+				[
+					SNew(STextBlock)
+					.Font(FEditorStyle::GetFontStyle("CurveEd.LabelFont"))
+					.Text(LOCTEXT("GroomBindingOptionsWindow_CurrentFile", "Current File: "))
+				]
+				+ SHorizontalBox::Slot()
+				.Padding(5, 0, 0, 0)
+				.AutoWidth()
+				.VAlign(VAlign_Center)
+				[
+					SNew(STextBlock)
+					.Font(FEditorStyle::GetFontStyle("CurveEd.InfoFont"))
+					.Text(InArgs._FullPath)
+				]
+			]
+		]
+
+		+ SVerticalBox::Slot()
+		.Padding(2)
+		.MaxHeight(500.0f)
+		[
+			DetailsView->AsShared()
+		]
+
+		+ SVerticalBox::Slot()
+		.AutoHeight()
+		.HAlign(HAlign_Right)
+		.Padding(2)
+		[
+			SNew(SUniformGridPanel)
+			.SlotPadding(2)
+			+ SUniformGridPanel::Slot(0, 0)
+			[
+				SAssignNew(ImportButton, SButton)
+				.HAlign(HAlign_Center)
+				.Text(InArgs._ButtonLabel)
+				.IsEnabled(this, &SGroomCreateBindingOptionsWindow::CanCreateBinding)
+				.OnClicked(this, &SGroomCreateBindingOptionsWindow::OnCreateBinding)
+			]
+			+ SUniformGridPanel::Slot(1, 0)
+			[
+				SNew(SButton)
+				.HAlign(HAlign_Center)
+				.Text(LOCTEXT("GroomBindingOptionsWindow_Cancel", "Cancel"))
+				.OnClicked(this, &SGroomCreateBindingOptionsWindow::OnCancel)
+			]
+		]
+	];
+}
+
+bool SGroomCreateBindingOptionsWindow::CanCreateBinding()  const
+{
+	return true;
+}
+
+enum class EGroomBindingOptionsVisibility : uint8
+{
+	None = 0x00,
+	ConversionOptions = 0x01,
+	BuildOptions = 0x02,
+	All = ConversionOptions | BuildOptions
+};
+
+ENUM_CLASS_FLAGS(EGroomBindingOptionsVisibility);
+
+TSharedPtr<SGroomCreateBindingOptionsWindow> DisplayOptions(UGroomCreateBindingOptions* BindingOptions, EGroomBindingOptionsVisibility VisibilityFlag, FText WindowTitle, FText InButtonLabel)
+{
+	TSharedRef<SWindow> Window = SNew(SWindow)
+		.Title(WindowTitle)
+		.SizingRule(ESizingRule::Autosized);
+
+	TSharedPtr<SGroomCreateBindingOptionsWindow> OptionsWindow;
+
+	Window->SetContent
+	(
+		SAssignNew(OptionsWindow, SGroomCreateBindingOptionsWindow)
+		.BindingOptions(BindingOptions)
+		.WidgetWindow(Window)
+//		.FullPath(FText::FromString(FileName))
+		.ButtonLabel(InButtonLabel)
+	);
+
+	TSharedPtr<SWindow> ParentWindow;
+
+	if (FModuleManager::Get().IsModuleLoaded("MainFrame"))
+	{
+		IMainFrameModule& MainFrame = FModuleManager::LoadModuleChecked<IMainFrameModule>("MainFrame");
+		ParentWindow = MainFrame.GetParentWindow();
+	}
+
+	FSlateApplication::Get().AddModalWindow(Window, ParentWindow, false);
+
+	return OptionsWindow;
+}
+
+TSharedPtr<SGroomCreateBindingOptionsWindow> SGroomCreateBindingOptionsWindow::DisplayCreateBindingOptions(UGroomCreateBindingOptions* BindingOptions)
+{
+	return DisplayOptions(BindingOptions, EGroomBindingOptionsVisibility::BuildOptions, LOCTEXT("GroomBuildOptionsWindow_WindowTitle", "Groom Binding Options"), LOCTEXT("GroomBindingOptionsWindow_Build", "Create"));
+}
+
+static UObject* InternalCreateNewBindAsset(FName InAssetName, UObject* InParent, UGroomAsset* GroomAsset, USkeletalMesh* SourceSkelMesh, USkeletalMesh* TargetSkelMesh)
+{
+	if (!TargetSkelMesh)
+	{
+		return nullptr;
+	}
+
+	FString Name = InAssetName.ToString();
+	FString PackageName;
+
+	if (InAssetName == NAME_None)
+	{
+		FString Suffix;
+		if (SourceSkelMesh)
+		{
+			Suffix += TEXT("_") + SourceSkelMesh->GetName();
+		}
+		if (TargetSkelMesh)
+		{
+			Suffix += TEXT("_") + TargetSkelMesh->GetName();
+		}
+		Suffix += TEXT("_Binding");
+
+		// Get a unique package and asset name
+		FAssetToolsModule& AssetToolsModule = FModuleManager::Get().LoadModuleChecked<FAssetToolsModule>("AssetTools");
+		AssetToolsModule.Get().CreateUniqueAssetName(GroomAsset->GetOutermost()->GetName(), Suffix, PackageName, Name);
+	}
+
+	UPackage* Package = Cast<UPackage>(InParent);
+	if (InParent == nullptr && !PackageName.IsEmpty())
+	{
+		// Then find/create it.
+		Package = CreatePackage(nullptr, *PackageName);
+		if (!ensure(Package))
+		{
+			// There was a problem creating the package
+			return nullptr;
+		}
+	}
+	  
+	if (UGroomBindingAsset* NewAsset = NewObject<UGroomBindingAsset>(Package, *Name, RF_Public | RF_Standalone | RF_Transactional))
+	{
+		NewAsset->Groom = GroomAsset;
+		NewAsset->SourceSkeletalMesh = SourceSkelMesh;
+		NewAsset->TargetSkeletalMesh = TargetSkelMesh;
+		NewAsset->MarkPackageDirty();
+
+		// Notify the asset registry
+		FAssetRegistryModule::AssetCreated(NewAsset);
+		return NewAsset;
+	}
+
+	return nullptr;
+}
+
+UGroomBindingAsset* CreateGroomBindinAsset(UGroomAsset* GroomAsset, USkeletalMesh* SourceSkelMesh, USkeletalMesh* TargetSkelMesh)
+{
+	if (!GroomAsset || !TargetSkelMesh)
+	{
+		return nullptr;
+	}
+
+	UObject* BindingAsset = InternalCreateNewBindAsset(NAME_None, nullptr, GroomAsset, SourceSkelMesh, TargetSkelMesh);
+
+	if (BindingAsset)
+	{
+		TArray<UObject*> CreatedObjects;
+		CreatedObjects.Add(BindingAsset);
+
+		FContentBrowserModule& ContentBrowserModule = FModuleManager::LoadModuleChecked<FContentBrowserModule>("ContentBrowser");
+		ContentBrowserModule.Get().SyncBrowserToAssets(CreatedObjects);
+	#if WITH_EDITOR
+		GEditor->GetEditorSubsystem<UAssetEditorSubsystem>()->OpenEditorForAssets(CreatedObjects);
+	#endif
+	}
+
+	return (UGroomBindingAsset*) BindingAsset;
+}
+
+#undef LOCTEXT_NAMESPACE
diff --git a/Engine/Plugins/Experimental/HairStrands/Source/HairStrandsEditor/Private/HairStrandsEditor.cpp b/Engine/Plugins/Experimental/HairStrands/Source/HairStrandsEditor/Private/HairStrandsEditor.cpp
--- a/Engine/Plugins/Experimental/HairStrands/Source/HairStrandsEditor/Private/HairStrandsEditor.cpp
+++ b/Engine/Plugins/Experimental/HairStrands/Source/HairStrandsEditor/Private/HairStrandsEditor.cpp
@@ -41,7 +41,10 @@ void FHairStrandsEditor::StartupModule()
 
 		StyleSet->Set("ClassIcon.GroomAsset", new FSlateImageBrush(HairStrandsContent + "/Icons/S_Groom_16.png", Icon16x16));
 		StyleSet->Set("ClassThumbnail.GroomAsset", new FSlateImageBrush(HairStrandsContent + "/Icons/S_Groom_64.png", Icon64x64));
-		
+
+		StyleSet->Set("ClassIcon.GroomBindingAsset", new FSlateImageBrush(HairStrandsContent + "/Icons/S_GroomBinding_16.png", Icon16x16));
+		StyleSet->Set("ClassThumbnail.GroomBindingAsset", new FSlateImageBrush(HairStrandsContent + "/Icons/S_GroomBinding_64.png", Icon64x64));
+
 		FSlateStyleRegistry::RegisterSlateStyle(*StyleSet.Get());
 	}
 }
diff --git a/Engine/Plugins/Experimental/HairStrands/Source/HairStrandsEditor/Public/GroomActions.h b/Engine/Plugins/Experimental/HairStrands/Source/HairStrandsEditor/Public/GroomActions.h
--- a/Engine/Plugins/Experimental/HairStrands/Source/HairStrandsEditor/Public/GroomActions.h
+++ b/Engine/Plugins/Experimental/HairStrands/Source/HairStrandsEditor/Public/GroomActions.h
@@ -42,4 +42,8 @@ private:
 	/** Callback for Rebuild groom action */
 	bool CanRebuild(TArray<TWeakObjectPtr<UGroomAsset>> Objects) const;
 	void ExecuteRebuild(TArray<TWeakObjectPtr<UGroomAsset>> Objects) const;
+
+	/** Callback for Rebuild groom action */
+	bool CanCreateBindingAsset(TArray<TWeakObjectPtr<UGroomAsset>> Objects) const;
+	void ExecuteCreateBindingAsset(TArray<TWeakObjectPtr<UGroomAsset>> Objects) const;
 };
