Copying //UE4/Dev-Physics to Main (//UE4/Main)

diff --git a/Engine/Plugins/Experimental/HairStrands/Source/HairStrandsNiagara/Private/NiagaraDataInterfaceHairStrands.cpp b/Engine/Plugins/Experimental/HairStrands/Source/HairStrandsNiagara/Private/NiagaraDataInterfaceHairStrands.cpp
--- a/Engine/Plugins/Experimental/HairStrands/Source/HairStrandsNiagara/Private/NiagaraDataInterfaceHairStrands.cpp
+++ b/Engine/Plugins/Experimental/HairStrands/Source/HairStrandsNiagara/Private/NiagaraDataInterfaceHairStrands.cpp
@@ -19,15 +19,49 @@
 
 #define LOCTEXT_NAMESPACE "NiagaraDataInterfaceHairStrands"
 
-static const int32 MaxDelay = 5;
-
 //------------------------------------------------------------------------------------------------------------
 
+static const FName GetPointPositionName(TEXT("GetPointPosition"));
+
 static const FName GetStrandSizeName(TEXT("GetStrandSize"));
 static const FName GetNumStrandsName(TEXT("GetNumStrands"));
+
 static const FName GetWorldTransformName(TEXT("GetWorldTransform"));
 static const FName GetWorldInverseName(TEXT("GetWorldInverse"));
-static const FName GetPointPositionName(TEXT("GetPointPosition"));
+static const FName GetBoxCenterName(TEXT("GetBoxCenter"));
+static const FName GetBoxExtentName(TEXT("GetBoxExtent"));
+
+static const FName GetSubStepsName("GetSubSteps");
+static const FName GetIterationCountName("GetIterationCount");
+
+static const FName GetGravityVectorName("GetGravityVector");
+static const FName GetAirDragName("GetAirDrag");
+static const FName GetAirVelocityName("GetAirVelocity");
+
+static const FName GetSolveBendName("GetSolveBend");
+static const FName GetProjectBendName("GetProjectBend");
+static const FName GetBendDampingName("GetBendDamping");
+static const FName GetBendStiffnessName("GetBendStiffness");
+static const FName GetBendScaleName("GetBendScale");
+
+static const FName GetSolveStretchName("GetSolveStretch");
+static const FName GetProjectStretchName("GetProjectStretch");
+static const FName GetStretchDampingName("GetStretchDamping");
+static const FName GetStretchStiffnessName("GetStretchStiffness");
+static const FName GetStretchScaleName("GetStretchScale");
+
+static const FName GetSolveCollisionName("GetSolveCollision");
+static const FName GetProjectCollisionName("GetProjectCollision");
+static const FName GetStaticFrictionName("GetStaticFriction");
+static const FName GetKineticFrictionName("GetKineticFriction");
+static const FName GetStrandsViscosityName("GetStrandsViscosity");
+static const FName GetCollisionRadiusName("GetCollisionRadius");
+static const FName GetRadiusScaleName("GetRadiusScale");
+
+static const FName GetStrandsDensityName("GetStrandsDensity");
+static const FName GetStrandsSmoothingName("GetStrandsSmoothing");
+static const FName GetStrandsThicknessName("GetStrandsThickness");
+static const FName GetThicknessScaleName("GetThicknessScale");
 
 //------------------------------------------------------------------------------------------------------------
 
@@ -64,8 +98,6 @@ static const FName ResetPointPositionName(TEXT("ResetPointPosition"));
 
 //------------------------------------------------------------------------------------------------------------
 
-static const FName GetBoxCenterName(TEXT("GetBoxCenter"));
-static const FName GetBoxExtentName(TEXT("GetBoxExtent"));
 static const FName BuildBoundingBoxName(TEXT("BuildBoundingBox"));
 
 //------------------------------------------------------------------------------------------------------------
@@ -91,10 +123,11 @@ static const FName ProjectStretchRodMaterialName(TEXT("ProjectStretchRodMaterial
 static const FName SetupBendRodMaterialName(TEXT("SetupBendRodMaterial"));
 static const FName SolveBendRodMaterialName(TEXT("SolveBendRodMaterial"));
 static const FName ProjectBendRodMaterialName(TEXT("ProjectBendRodMaterial"));
+
 //------------------------------------------------------------------------------------------------------------
 
-static const FName SolveStaticCollisionConstraintName(TEXT("SolveStaticCollisionConstraint"));
-static const FName ProjectStaticCollisionConstraintName(TEXT("ProjectStaticCollisionConstraint"));
+static const FName SolveHardCollisionConstraintName(TEXT("SolveHardCollisionConstraint"));
+static const FName ProjectHardCollisionConstraintName(TEXT("ProjectHardCollisionConstraint"));
 
 static const FName SetupSoftCollisionConstraintName(TEXT("SetupSoftCollisionConstraint"));
 static const FName SolveSoftCollisionConstraintName(TEXT("SolveSoftCollisionConstraint"));
@@ -103,7 +136,8 @@ static const FName ProjectSoftCollisionConstraintName(TEXT("ProjectSoftCollision
 //------------------------------------------------------------------------------------------------------------
 
 static const FName ComputeRestDirectionName(TEXT("ComputeRestDirection"));
-static const FName UpdateNodeOrientationName(TEXT("UpdateNodeOrientation"));
+static const FName UpdateMaterialFrameName(TEXT("UpdateMaterialFrame"));
+static const FName ComputeMaterialFrameName(TEXT("ComputeMaterialFrame"));
 
 //------------------------------------------------------------------------------------------------------------
 
@@ -512,9 +546,10 @@ bool FNDIHairStrandsData::Init(UNiagaraDataInterfaceHairStrands* Interface, FNia
 		FHairStrandsRestResource* StrandsRestResource = nullptr;
 		FHairStrandsDeformedResource* StrandsDeformedResource = nullptr;
 		FHairStrandsRootResource* StrandsRootResource = nullptr;
+		UGroomAsset* GroomAsset = nullptr;
 
 		{
-			Interface->ExtractDatasAndResources(SystemInstance, StrandsDatas, StrandsRestResource, StrandsDeformedResource, StrandsRootResource);
+			Interface->ExtractDatasAndResources(SystemInstance, StrandsDatas, StrandsRestResource, StrandsDeformedResource, StrandsRootResource, GroomAsset);
 
 			HairStrandsBuffer = new FNDIHairStrandsBuffer();
 			HairStrandsBuffer->Initialize(StrandsDatas, StrandsRestResource, StrandsDeformedResource, StrandsRootResource);
@@ -522,10 +557,49 @@ bool FNDIHairStrandsData::Init(UNiagaraDataInterfaceHairStrands* Interface, FNia
 
 			WorldTransform = Interface->IsComponentValid() ? Interface->SourceComponent->GetComponentToWorld() :
 				SystemInstance->GetComponent()->GetComponentToWorld();
-			StrandSize = static_cast<uint8>(Interface->StrandSize);
 
-			if (StrandsDatas != nullptr)
+			if (StrandsDatas != nullptr && GroomAsset != nullptr)
 			{
+				StrandsSize = static_cast<uint8>(GroomAsset->StrandsSize);
+
+				SubSteps = GroomAsset->SubSteps;
+				IterationCount = GroomAsset->IterationCount;
+
+				GravityVector = GroomAsset->GravityVector;
+				AirDrag = GroomAsset->AirDrag;
+				AirVelocity = GroomAsset->AirVelocity;
+
+				SolveBend = GroomAsset->SolveBend;
+				ProjectBend = GroomAsset->ProjectBend;
+				BendDamping = GroomAsset->BendDamping;
+				BendStiffness = GroomAsset->BendStiffness;
+
+				SolveStretch = GroomAsset->SolveStretch;
+				ProjectStretch = GroomAsset->ProjectStretch;
+				StretchDamping = GroomAsset->StretchDamping;
+				StretchStiffness = GroomAsset->StretchStiffness;
+
+				SolveCollision = GroomAsset->SolveCollision;
+				ProjectCollision = GroomAsset->ProjectCollision;
+				StaticFriction = GroomAsset->StaticFriction;
+				KineticFriction = GroomAsset->KineticFriction;
+				StrandsViscosity = GroomAsset->StrandsViscosity;
+				CollisionRadius = GroomAsset->CollisionRadius;
+
+				StrandsDensity = GroomAsset->StrandsDensity;
+				StrandsSmoothing = GroomAsset->StrandsSmoothing;
+				StrandsThickness = GroomAsset->StrandsThickness;
+
+				for (int32 i = 0; i < 4; ++i)
+				{
+					const float VertexCoord = static_cast<float>(i) / (3);
+
+					BendScale[i] = GroomAsset->BendScale.GetRichCurve()->Eval(VertexCoord);
+					StretchScale[i] = GroomAsset->StretchScale.GetRichCurve()->Eval(VertexCoord);
+					RadiusScale[i] = GroomAsset->RadiusScale.GetRichCurve()->Eval(VertexCoord);
+					ThicknessScale[i] = GroomAsset->ThicknessScale.GetRichCurve()->Eval(VertexCoord);
+				}
+
 				const FBox& StrandsBox = StrandsDatas->BoundingBox;
 
 				NumStrands = StrandsDatas->GetNumCurves();
@@ -534,9 +608,7 @@ bool FNDIHairStrandsData::Init(UNiagaraDataInterfaceHairStrands* Interface, FNia
 			}
 			else
 			{
-				NumStrands = 0;
-				BoxCenter = FVector(0, 0, 0);
-				BoxExtent = FVector(0, 0, 0);
+				ResetDatas();
 			}
 			TickCount = 0;
 			ForceReset = true;
@@ -723,11 +795,11 @@ struct FNDIHairStrandsParametersCS : public FNiagaraDataInterfaceParametersCS
 				MeshProjection->RootTriangleBarycentricBuffer.SRV : FNiagaraRenderer::GetDummyFloatBuffer().SRV;
 
 			int32 bNeedSimReset = (ProxyData->TickCount <= ProxyData->ResetTick ? 1 : 0);
-			FVector RestRootOffsetValue = (HasRootAttachedValue == 1) ? MeshProjection->RestRootOffset : FVector(0, 0, 0);
-			FVector DeformedRootOffsetValue = (HasRootAttachedValue == 1) ? MeshProjection->DeformedRootOffset : FVector(0, 0, 0);
-			FVector DeformedPositionOffsetValue = IsDeformedValid ? HairStrandsBuffer->SourceDeformedResources->PositionOffset :
-																	HairStrandsBuffer->SourceRestResources->PositionOffset;
-			FVector RestPositionOffsetValue = HairStrandsBuffer->SourceRestResources->PositionOffset;
+			FVector RestRootOffsetValue = FVector::ZeroVector;
+			FVector DeformedRootOffsetValue = FVector::ZeroVector;
+			FVector DeformedPositionOffsetValue =  IsDeformedValid ? HairStrandsBuffer->SourceDeformedResources->PositionOffset :
+																	 HairStrandsBuffer->SourceRestResources->PositionOffset;
+			FVector RestPositionOffsetValue = ProxyData->HairStrandsBuffer->SourceRestResources->PositionOffset;
 
 			SetUAVParameter(RHICmdList, ComputeShaderRHI, DeformedPositionBuffer, PointPositionsUAV);
 			RHICmdList.TransitionResource(EResourceTransitionAccess::ERWBarrier, EResourceTransitionPipeline::EComputeToCompute, HairStrandsBuffer->BoundingBoxBuffer.UAV);
@@ -741,7 +813,7 @@ struct FNDIHairStrandsParametersCS : public FNiagaraDataInterfaceParametersCS
 			SetShaderValue(RHICmdList, ComputeShaderRHI, WorldInverse, ProxyData->WorldTransform.ToMatrixWithScale().Inverse());
 			SetShaderValue(RHICmdList, ComputeShaderRHI, WorldRotation, ProxyData->WorldTransform.ToMatrixNoScale().ToQuat());
 			SetShaderValue(RHICmdList, ComputeShaderRHI, NumStrands, ProxyData->NumStrands);
-			SetShaderValue(RHICmdList, ComputeShaderRHI, StrandSize, ProxyData->StrandSize);
+			SetShaderValue(RHICmdList, ComputeShaderRHI, StrandSize, ProxyData->StrandsSize);
 			SetShaderValue(RHICmdList, ComputeShaderRHI, BoxCenter, ProxyData->BoxCenter);
 			SetShaderValue(RHICmdList, ComputeShaderRHI, BoxExtent, ProxyData->BoxExtent);
 
@@ -851,15 +923,7 @@ void FNDIHairStrandsProxy::ConsumePerInstanceDataFromGameThread(void* PerInstanc
 	ensure(TargetData);
 	if (TargetData)
 	{
-		TargetData->WorldTransform = SourceData->WorldTransform;
-		TargetData->NumStrands = SourceData->NumStrands;
-		TargetData->StrandSize = SourceData->StrandSize;
-		TargetData->BoxCenter = SourceData->BoxCenter;
-		TargetData->BoxExtent = SourceData->BoxExtent;
-		TargetData->TickCount = SourceData->TickCount;
-		TargetData->ResetTick = SourceData->ResetTick;
-		TargetData->ForceReset = SourceData->ForceReset;
-		TargetData->HairStrandsBuffer = SourceData->HairStrandsBuffer;
+		TargetData->CopyDatas(SourceData);
 	}
 	else
 	{
@@ -886,11 +950,9 @@ void FNDIHairStrandsProxy::DestroyPerInstanceData(NiagaraEmitterInstanceBatcher*
 
 UNiagaraDataInterfaceHairStrands::UNiagaraDataInterfaceHairStrands(FObjectInitializer const& ObjectInitializer)
 	: Super(ObjectInitializer)
-	, StrandSize(EHairStrandsSize::Size4)
 	, DefaultSource(nullptr)
 	, SourceActor(nullptr)
 	, SourceComponent(nullptr)
-	, GroupIndex(0)
 {
 
 	Proxy.Reset(new FNDIHairStrandsProxy());
@@ -946,7 +1008,7 @@ void UNiagaraDataInterfaceHairStrands::ExtractSourceComponent(FNiagaraSystemInst
 }
 
 void UNiagaraDataInterfaceHairStrands::ExtractDatasAndResources(FNiagaraSystemInstance* SystemInstance,FHairStrandsDatas*& OutStrandsDatas, 
-	FHairStrandsRestResource*& OutStrandsRestResource, FHairStrandsDeformedResource*& OutStrandsDeformedResource, FHairStrandsRootResource*& OutStrandsRootResource)
+	FHairStrandsRestResource*& OutStrandsRestResource, FHairStrandsDeformedResource*& OutStrandsDeformedResource, FHairStrandsRootResource*& OutStrandsRootResource, UGroomAsset*& OutGroomAsset)
 {
 	ExtractSourceComponent(SystemInstance);
 
@@ -955,17 +1017,22 @@ void UNiagaraDataInterfaceHairStrands::ExtractDatasAndResources(FNiagaraSystemIn
 	OutStrandsDeformedResource = nullptr;
 	OutStrandsRootResource = nullptr;
 
+	const int32 GroupIndex = (DefaultSource != nullptr && (DefaultSource->SimulatedGroup < 
+		DefaultSource->GetNumHairGroups())) ? DefaultSource->SimulatedGroup : 0;
+
 	if (IsComponentValid())
 	{
 		OutStrandsDatas = SourceComponent->GetGuideStrandsDatas(GroupIndex);
 		OutStrandsRestResource = SourceComponent->GetGuideStrandsRestResource(GroupIndex);
 		OutStrandsDeformedResource = SourceComponent->GetGuideStrandsDeformedResource(GroupIndex);
 		OutStrandsRootResource = SourceComponent->GetGuideStrandsRootResource(GroupIndex);
+		OutGroomAsset = SourceComponent->GroomAsset;
 	}
 	else if (DefaultSource != nullptr && (GroupIndex < DefaultSource->GetNumHairGroups()))
 	{
 		OutStrandsDatas = &DefaultSource->HairGroupsData[GroupIndex].HairRenderData;
 		OutStrandsRestResource = DefaultSource->HairGroupsData[GroupIndex].HairStrandsRestResource;
+		OutGroomAsset = DefaultSource;
 	}
 }
 
@@ -974,17 +1041,6 @@ bool UNiagaraDataInterfaceHairStrands::InitPerInstanceData(void* PerInstanceData
 	FNDIHairStrandsData* InstanceData = new (PerInstanceData) FNDIHairStrandsData();
 	check(InstanceData);
 
-	/*FNDIHairStrandsProxy* ThisProxy = GetProxyAs<FNDIHairStrandsProxy>();
-	ENQUEUE_RENDER_COMMAND(FNiagaraDIPushInitialInstanceDataToRT) (
-		[ThisProxy, InstanceID = SystemInstance->GetId()](FRHICommandListImmediate& CmdList)
-	{
-		ThisProxy->InitializePerInstanceData(InstanceID);
-	}
-	);*/
-
-	//UE_LOG(LogHairStrands, Warning, TEXT("Init Strand Size : %d"), StrandSize);
-
-
 	return InstanceData->Init(this, SystemInstance);
 }
 
@@ -1000,7 +1056,6 @@ void UNiagaraDataInterfaceHairStrands::DestroyPerInstanceData(void* PerInstanceD
 		[ThisProxy, InstanceID = SystemInstance->GetId(), Batcher = SystemInstance->GetBatcher()](FRHICommandListImmediate& CmdList)
 	{
 		ThisProxy->SystemInstancesToProxyData.Remove(InstanceID);
-		//ThisProxy->DestroyPerInstanceData(Batcher, InstanceID);
 	}
 	);
 }
@@ -1013,9 +1068,11 @@ bool UNiagaraDataInterfaceHairStrands::PerInstanceTick(void* PerInstanceData, FN
 	FHairStrandsRestResource* StrandsRestResource = nullptr;
 	FHairStrandsDeformedResource* StrandsDeformedResource = nullptr;
 	FHairStrandsRootResource* StrandsRootResource = nullptr;
+	UGroomAsset* GroomAsset = nullptr;
+
 	InstanceData->TickCount = FMath::Min(MaxDelay+1,InstanceData->TickCount+1);
 
-	ExtractDatasAndResources(SystemInstance, StrandsDatas, StrandsRestResource, StrandsDeformedResource, StrandsRootResource);
+	ExtractDatasAndResources(SystemInstance, StrandsDatas, StrandsRestResource, StrandsDeformedResource, StrandsRootResource, GroomAsset);
 	InstanceData->HairStrandsBuffer->Initialize(StrandsDatas, StrandsRestResource, StrandsDeformedResource, StrandsRootResource);
 
 	if (SourceComponent != nullptr)
@@ -1049,12 +1106,10 @@ bool UNiagaraDataInterfaceHairStrands::PerInstanceTick(void* PerInstanceData, FN
 		if (IsComponentValid())
 		{
 			InstanceData->WorldTransform = SourceComponent->GetComponentToWorld();
-			//UE_LOG(LogHairStrands, Warning, TEXT("Get Component Transform : %s"), *InstanceData->WorldTransform.ToString());
 		}
 		else
 		{
 			InstanceData->WorldTransform = SystemInstance->GetComponent()->GetComponentToWorld();
-			//UE_LOG(LogHairStrands, Warning, TEXT("Get Instance Transform : %s"), *InstanceData->WorldTransform.ToString());
 		}
 	}
 	return RequireReset;
@@ -1068,11 +1123,9 @@ bool UNiagaraDataInterfaceHairStrands::CopyToInternal(UNiagaraDataInterface* Des
 	}
 
 	UNiagaraDataInterfaceHairStrands* OtherTyped = CastChecked<UNiagaraDataInterfaceHairStrands>(Destination);
-	OtherTyped->StrandSize = StrandSize;
 	OtherTyped->SourceActor= SourceActor;
 	OtherTyped->SourceComponent = SourceComponent;
 	OtherTyped->DefaultSource = DefaultSource;
-	OtherTyped->GroupIndex = GroupIndex;
 
 	return true;
 }
@@ -1085,9 +1138,8 @@ bool UNiagaraDataInterfaceHairStrands::Equals(const UNiagaraDataInterface* Other
 	}
 	const UNiagaraDataInterfaceHairStrands* OtherTyped = CastChecked<const UNiagaraDataInterfaceHairStrands>(Other);
 
-	return  (OtherTyped->StrandSize == StrandSize)
-		&& (OtherTyped->SourceActor == SourceActor) && (OtherTyped->SourceComponent == SourceComponent)
-		&& (OtherTyped->DefaultSource == DefaultSource) && (OtherTyped->GroupIndex == GroupIndex);
+	return  (OtherTyped->SourceActor == SourceActor) && (OtherTyped->SourceComponent == SourceComponent)
+		&& (OtherTyped->DefaultSource == DefaultSource);
 }
 
 void UNiagaraDataInterfaceHairStrands::PostInitProperties()
@@ -1122,6 +1174,266 @@ void UNiagaraDataInterfaceHairStrands::GetFunctions(TArray<FNiagaraFunctionSigna
 
 		OutFunctions.Add(Sig);
 	}
+	{
+		FNiagaraFunctionSignature Sig;
+		Sig.Name = GetSubStepsName;
+		Sig.bMemberFunction = true;
+		Sig.bRequiresContext = false;
+		Sig.Inputs.Add(FNiagaraVariable(FNiagaraTypeDefinition(GetClass()), TEXT("Hair Strands")));
+		Sig.Outputs.Add(FNiagaraVariable(FNiagaraTypeDefinition::GetIntDef(), TEXT("Sub Steps")));
+
+		OutFunctions.Add(Sig);
+	}
+	{
+		FNiagaraFunctionSignature Sig;
+		Sig.Name = GetIterationCountName;
+		Sig.bMemberFunction = true;
+		Sig.bRequiresContext = false;
+		Sig.Inputs.Add(FNiagaraVariable(FNiagaraTypeDefinition(GetClass()), TEXT("Hair Strands")));
+		Sig.Outputs.Add(FNiagaraVariable(FNiagaraTypeDefinition::GetIntDef(), TEXT("Iteration Count")));
+
+		OutFunctions.Add(Sig);
+	}
+	{
+		FNiagaraFunctionSignature Sig;
+		Sig.Name = GetGravityVectorName;
+		Sig.bMemberFunction = true;
+		Sig.bRequiresContext = false;
+		Sig.Inputs.Add(FNiagaraVariable(FNiagaraTypeDefinition(GetClass()), TEXT("Hair Strands")));
+		Sig.Outputs.Add(FNiagaraVariable(FNiagaraTypeDefinition::GetVec3Def(), TEXT("Gravity Vector")));
+
+		OutFunctions.Add(Sig);
+	}
+	{
+		FNiagaraFunctionSignature Sig;
+		Sig.Name = GetAirDragName;
+		Sig.bMemberFunction = true;
+		Sig.bRequiresContext = false;
+		Sig.Inputs.Add(FNiagaraVariable(FNiagaraTypeDefinition(GetClass()), TEXT("Hair Strands")));
+		Sig.Outputs.Add(FNiagaraVariable(FNiagaraTypeDefinition::GetFloatDef(), TEXT("Air Drag")));
+
+		OutFunctions.Add(Sig);
+	}
+	{
+		FNiagaraFunctionSignature Sig;
+		Sig.Name = GetAirVelocityName;
+		Sig.bMemberFunction = true;
+		Sig.bRequiresContext = false;
+		Sig.Inputs.Add(FNiagaraVariable(FNiagaraTypeDefinition(GetClass()), TEXT("Hair Strands")));
+		Sig.Outputs.Add(FNiagaraVariable(FNiagaraTypeDefinition::GetVec3Def(), TEXT("Air Velocity")));
+
+		OutFunctions.Add(Sig);
+	}
+	{
+		FNiagaraFunctionSignature Sig;
+		Sig.Name = GetSolveBendName;
+		Sig.bMemberFunction = true;
+		Sig.bRequiresContext = false;
+		Sig.Inputs.Add(FNiagaraVariable(FNiagaraTypeDefinition(GetClass()), TEXT("Hair Strands")));
+		Sig.Outputs.Add(FNiagaraVariable(FNiagaraTypeDefinition::GetBoolDef(), TEXT("Solve Bend")));
+
+		OutFunctions.Add(Sig);
+	}
+	{
+		FNiagaraFunctionSignature Sig;
+		Sig.Name = GetProjectBendName;
+		Sig.bMemberFunction = true;
+		Sig.bRequiresContext = false;
+		Sig.Inputs.Add(FNiagaraVariable(FNiagaraTypeDefinition(GetClass()), TEXT("Hair Strands")));
+		Sig.Outputs.Add(FNiagaraVariable(FNiagaraTypeDefinition::GetBoolDef(), TEXT("Project Bend")));
+
+		OutFunctions.Add(Sig);
+	}
+	{
+		FNiagaraFunctionSignature Sig;
+		Sig.Name = GetBendDampingName;
+		Sig.bMemberFunction = true;
+		Sig.bRequiresContext = false;
+		Sig.Inputs.Add(FNiagaraVariable(FNiagaraTypeDefinition(GetClass()), TEXT("Hair Strands")));
+		Sig.Outputs.Add(FNiagaraVariable(FNiagaraTypeDefinition::GetFloatDef(), TEXT("Bend Damping")));
+
+		OutFunctions.Add(Sig);
+	}
+	{
+		FNiagaraFunctionSignature Sig;
+		Sig.Name = GetBendStiffnessName;
+		Sig.bMemberFunction = true;
+		Sig.bRequiresContext = false;
+		Sig.Inputs.Add(FNiagaraVariable(FNiagaraTypeDefinition(GetClass()), TEXT("Hair Strands")));
+		Sig.Outputs.Add(FNiagaraVariable(FNiagaraTypeDefinition::GetFloatDef(), TEXT("Bend Stiffness")));
+
+		OutFunctions.Add(Sig);
+	}
+	{
+		FNiagaraFunctionSignature Sig;
+		Sig.Name = GetBendScaleName;
+		Sig.bMemberFunction = true;
+		Sig.bRequiresContext = false;
+		Sig.Inputs.Add(FNiagaraVariable(FNiagaraTypeDefinition(GetClass()), TEXT("Hair Strands")));
+		Sig.Outputs.Add(FNiagaraVariable(FNiagaraTypeDefinition::GetVec4Def(), TEXT("Bend Scale")));
+
+		OutFunctions.Add(Sig);
+	}
+	{
+		FNiagaraFunctionSignature Sig;
+		Sig.Name = GetSolveStretchName;
+		Sig.bMemberFunction = true;
+		Sig.bRequiresContext = false;
+		Sig.Inputs.Add(FNiagaraVariable(FNiagaraTypeDefinition(GetClass()), TEXT("Hair Strands")));
+		Sig.Outputs.Add(FNiagaraVariable(FNiagaraTypeDefinition::GetBoolDef(), TEXT("Solve Stretch")));
+
+		OutFunctions.Add(Sig);
+	}
+	{
+		FNiagaraFunctionSignature Sig;
+		Sig.Name = GetProjectStretchName;
+		Sig.bMemberFunction = true;
+		Sig.bRequiresContext = false;
+		Sig.Inputs.Add(FNiagaraVariable(FNiagaraTypeDefinition(GetClass()), TEXT("Hair Strands")));
+		Sig.Outputs.Add(FNiagaraVariable(FNiagaraTypeDefinition::GetBoolDef(), TEXT("Project Stretch")));
+
+		OutFunctions.Add(Sig);
+	}
+	{
+		FNiagaraFunctionSignature Sig;
+		Sig.Name = GetStretchDampingName;
+		Sig.bMemberFunction = true;
+		Sig.bRequiresContext = false;
+		Sig.Inputs.Add(FNiagaraVariable(FNiagaraTypeDefinition(GetClass()), TEXT("Hair Strands")));
+		Sig.Outputs.Add(FNiagaraVariable(FNiagaraTypeDefinition::GetFloatDef(), TEXT("Stretch Damping")));
+
+		OutFunctions.Add(Sig);
+	}
+	{
+		FNiagaraFunctionSignature Sig;
+		Sig.Name = GetStretchStiffnessName;
+		Sig.bMemberFunction = true;
+		Sig.bRequiresContext = false;
+		Sig.Inputs.Add(FNiagaraVariable(FNiagaraTypeDefinition(GetClass()), TEXT("Hair Strands")));
+		Sig.Outputs.Add(FNiagaraVariable(FNiagaraTypeDefinition::GetFloatDef(), TEXT("Stretch Stiffness")));
+
+		OutFunctions.Add(Sig);
+	}
+	{
+		FNiagaraFunctionSignature Sig;
+		Sig.Name = GetStretchScaleName;
+		Sig.bMemberFunction = true;
+		Sig.bRequiresContext = false;
+		Sig.Inputs.Add(FNiagaraVariable(FNiagaraTypeDefinition(GetClass()), TEXT("Hair Strands")));
+		Sig.Outputs.Add(FNiagaraVariable(FNiagaraTypeDefinition::GetVec4Def(), TEXT("Stretch Scale")));
+
+		OutFunctions.Add(Sig);
+	}
+	{
+		FNiagaraFunctionSignature Sig;
+		Sig.Name = GetSolveCollisionName;
+		Sig.bMemberFunction = true;
+		Sig.bRequiresContext = false;
+		Sig.Inputs.Add(FNiagaraVariable(FNiagaraTypeDefinition(GetClass()), TEXT("Hair Strands")));
+		Sig.Outputs.Add(FNiagaraVariable(FNiagaraTypeDefinition::GetBoolDef(), TEXT("Solve Collision")));
+
+		OutFunctions.Add(Sig);
+	}
+	{
+		FNiagaraFunctionSignature Sig;
+		Sig.Name = GetProjectCollisionName;
+		Sig.bMemberFunction = true;
+		Sig.bRequiresContext = false;
+		Sig.Inputs.Add(FNiagaraVariable(FNiagaraTypeDefinition(GetClass()), TEXT("Hair Strands")));
+		Sig.Outputs.Add(FNiagaraVariable(FNiagaraTypeDefinition::GetBoolDef(), TEXT("Project Collision")));
+
+		OutFunctions.Add(Sig);
+	}
+	{
+		FNiagaraFunctionSignature Sig;
+		Sig.Name = GetStaticFrictionName;
+		Sig.bMemberFunction = true;
+		Sig.bRequiresContext = false;
+		Sig.Inputs.Add(FNiagaraVariable(FNiagaraTypeDefinition(GetClass()), TEXT("Hair Strands")));
+		Sig.Outputs.Add(FNiagaraVariable(FNiagaraTypeDefinition::GetFloatDef(), TEXT("Static Fraction")));
+
+		OutFunctions.Add(Sig);
+	}
+	{
+		FNiagaraFunctionSignature Sig;
+		Sig.Name = GetKineticFrictionName;
+		Sig.bMemberFunction = true;
+		Sig.bRequiresContext = false;
+		Sig.Inputs.Add(FNiagaraVariable(FNiagaraTypeDefinition(GetClass()), TEXT("Hair Strands")));
+		Sig.Outputs.Add(FNiagaraVariable(FNiagaraTypeDefinition::GetFloatDef(), TEXT("Kinetic Friction")));
+
+		OutFunctions.Add(Sig);
+	}
+	{
+		FNiagaraFunctionSignature Sig;
+		Sig.Name = GetStrandsViscosityName;
+		Sig.bMemberFunction = true;
+		Sig.bRequiresContext = false;
+		Sig.Inputs.Add(FNiagaraVariable(FNiagaraTypeDefinition(GetClass()), TEXT("Hair Strands")));
+		Sig.Outputs.Add(FNiagaraVariable(FNiagaraTypeDefinition::GetFloatDef(), TEXT("Strands Viscosity")));
+
+		OutFunctions.Add(Sig);
+	}
+	{
+		FNiagaraFunctionSignature Sig;
+		Sig.Name = GetCollisionRadiusName;
+		Sig.bMemberFunction = true;
+		Sig.bRequiresContext = false;
+		Sig.Inputs.Add(FNiagaraVariable(FNiagaraTypeDefinition(GetClass()), TEXT("Hair Strands")));
+		Sig.Outputs.Add(FNiagaraVariable(FNiagaraTypeDefinition::GetFloatDef(), TEXT("Collision Radius")));
+
+		OutFunctions.Add(Sig);
+	}
+	{
+		FNiagaraFunctionSignature Sig;
+		Sig.Name = GetRadiusScaleName;
+		Sig.bMemberFunction = true;
+		Sig.bRequiresContext = false;
+		Sig.Inputs.Add(FNiagaraVariable(FNiagaraTypeDefinition(GetClass()), TEXT("Hair Strands")));
+		Sig.Outputs.Add(FNiagaraVariable(FNiagaraTypeDefinition::GetVec4Def(), TEXT("Radius Scale")));
+
+		OutFunctions.Add(Sig);
+	}
+	{
+		FNiagaraFunctionSignature Sig;
+		Sig.Name = GetStrandsDensityName;
+		Sig.bMemberFunction = true;
+		Sig.bRequiresContext = false;
+		Sig.Inputs.Add(FNiagaraVariable(FNiagaraTypeDefinition(GetClass()), TEXT("Hair Strands")));
+		Sig.Outputs.Add(FNiagaraVariable(FNiagaraTypeDefinition::GetFloatDef(), TEXT("Strands Density")));
+
+		OutFunctions.Add(Sig);
+	}
+	{
+		FNiagaraFunctionSignature Sig;
+		Sig.Name = GetStrandsSmoothingName;
+		Sig.bMemberFunction = true;
+		Sig.bRequiresContext = false;
+		Sig.Inputs.Add(FNiagaraVariable(FNiagaraTypeDefinition(GetClass()), TEXT("Hair Strands")));
+		Sig.Outputs.Add(FNiagaraVariable(FNiagaraTypeDefinition::GetFloatDef(), TEXT("Strands Smoothing")));
+
+		OutFunctions.Add(Sig);
+	}
+	{
+		FNiagaraFunctionSignature Sig;
+		Sig.Name = GetStrandsThicknessName;
+		Sig.bMemberFunction = true;
+		Sig.bRequiresContext = false;
+		Sig.Inputs.Add(FNiagaraVariable(FNiagaraTypeDefinition(GetClass()), TEXT("Hair Strands")));
+		Sig.Outputs.Add(FNiagaraVariable(FNiagaraTypeDefinition::GetFloatDef(), TEXT("Strands Thickness")));
+
+		OutFunctions.Add(Sig);
+	}
+	{
+		FNiagaraFunctionSignature Sig;
+		Sig.Name = GetThicknessScaleName;
+		Sig.bMemberFunction = true;
+		Sig.bRequiresContext = false;
+		Sig.Inputs.Add(FNiagaraVariable(FNiagaraTypeDefinition(GetClass()), TEXT("Hair Strands")));
+		Sig.Outputs.Add(FNiagaraVariable(FNiagaraTypeDefinition::GetVec4Def(), TEXT("Thickness Scale")));
+
+		OutFunctions.Add(Sig);
+	}
 	{
 		FNiagaraFunctionSignature Sig;
 		Sig.Name = GetWorldTransformName;
@@ -1460,13 +1772,12 @@ void UNiagaraDataInterfaceHairStrands::GetFunctions(TArray<FNiagaraFunctionSigna
 		Sig.bMemberFunction = true;
 		Sig.bRequiresContext = false;
 		Sig.Inputs.Add(FNiagaraVariable(FNiagaraTypeDefinition(GetClass()), TEXT("Hair Strands")));
+		Sig.Inputs.Add(FNiagaraVariable(FNiagaraTypeDefinition::GetBoolDef(), TEXT("Enable Constraint")));
 		Sig.Inputs.Add(FNiagaraVariable(FNiagaraTypeDefinition::GetFloatDef(), TEXT("Stretch Stiffness")));
 		Sig.Inputs.Add(FNiagaraVariable(FNiagaraTypeDefinition::GetFloatDef(), TEXT("Node Thickness")));
 		Sig.Inputs.Add(FNiagaraVariable(FNiagaraTypeDefinition::GetFloatDef(), TEXT("Rest Length")));
-		Sig.Inputs.Add(FNiagaraVariable(FNiagaraTypeDefinition::GetFloatDef(), TEXT("Node Mass")));
 		Sig.Inputs.Add(FNiagaraVariable(FNiagaraTypeDefinition::GetFloatDef(), TEXT("Delta Time")));
 		Sig.Inputs.Add(FNiagaraVariable(FNiagaraTypeDefinition::GetIntDef(), TEXT("Node Offset")));
-		Sig.Inputs.Add(FNiagaraVariable(FNiagaraTypeDefinition::GetVec3Def(), TEXT("Node Position")));
 		Sig.Outputs.Add(FNiagaraVariable(FNiagaraTypeDefinition::GetVec3Def(), TEXT("Node Position")));
 
 		OutFunctions.Add(Sig);
@@ -1510,13 +1821,12 @@ void UNiagaraDataInterfaceHairStrands::GetFunctions(TArray<FNiagaraFunctionSigna
 		Sig.bMemberFunction = true;
 		Sig.bRequiresContext = false;
 		Sig.Inputs.Add(FNiagaraVariable(FNiagaraTypeDefinition(GetClass()), TEXT("Hair Strands")));
+		Sig.Inputs.Add(FNiagaraVariable(FNiagaraTypeDefinition::GetBoolDef(), TEXT("Enable Constraint")));
 		Sig.Inputs.Add(FNiagaraVariable(FNiagaraTypeDefinition::GetFloatDef(), TEXT("Bend Stiffness")));
 		Sig.Inputs.Add(FNiagaraVariable(FNiagaraTypeDefinition::GetFloatDef(), TEXT("Node Thickness")));
 		Sig.Inputs.Add(FNiagaraVariable(FNiagaraTypeDefinition::GetFloatDef(), TEXT("Rest Length")));
 		Sig.Inputs.Add(FNiagaraVariable(FNiagaraTypeDefinition::GetVec3Def(), TEXT("Rest Direction")));
-		Sig.Inputs.Add(FNiagaraVariable(FNiagaraTypeDefinition::GetFloatDef(), TEXT("Node Mass")));
 		Sig.Inputs.Add(FNiagaraVariable(FNiagaraTypeDefinition::GetFloatDef(), TEXT("Delta Time")));
-		Sig.Inputs.Add(FNiagaraVariable(FNiagaraTypeDefinition::GetVec3Def(), TEXT("Node Position")));
 		Sig.Outputs.Add(FNiagaraVariable(FNiagaraTypeDefinition::GetVec3Def(), TEXT("Node Position")));
 
 		OutFunctions.Add(Sig);
@@ -1544,6 +1854,7 @@ void UNiagaraDataInterfaceHairStrands::GetFunctions(TArray<FNiagaraFunctionSigna
 		Sig.bMemberFunction = true;
 		Sig.bRequiresContext = false;
 		Sig.Inputs.Add(FNiagaraVariable(FNiagaraTypeDefinition(GetClass()), TEXT("Hair Strands")));
+		Sig.Inputs.Add(FNiagaraVariable(FNiagaraTypeDefinition::GetBoolDef(), TEXT("Enable Constraint")));
 		Sig.Inputs.Add(FNiagaraVariable(FNiagaraTypeDefinition::GetFloatDef(), TEXT("Rest Length")));
 		Sig.Inputs.Add(FNiagaraVariable(FNiagaraTypeDefinition::GetFloatDef(), TEXT("Delta Time")));
 		Sig.Inputs.Add(FNiagaraVariable(FNiagaraTypeDefinition::GetFloatDef(), TEXT("Material Damping")));
@@ -1559,11 +1870,12 @@ void UNiagaraDataInterfaceHairStrands::GetFunctions(TArray<FNiagaraFunctionSigna
 		Sig.bMemberFunction = true;
 		Sig.bRequiresContext = false;
 		Sig.Inputs.Add(FNiagaraVariable(FNiagaraTypeDefinition(GetClass()), TEXT("Hair Strands")));
+		Sig.Inputs.Add(FNiagaraVariable(FNiagaraTypeDefinition::GetBoolDef(), TEXT("Enable Constraint")));
 		Sig.Inputs.Add(FNiagaraVariable(FNiagaraTypeDefinition::GetFloatDef(), TEXT("Stretch Stiffness")));
 		Sig.Inputs.Add(FNiagaraVariable(FNiagaraTypeDefinition::GetFloatDef(), TEXT("Node Thickness")));
 		Sig.Inputs.Add(FNiagaraVariable(FNiagaraTypeDefinition::GetFloatDef(), TEXT("Rest Length")));
 		Sig.Inputs.Add(FNiagaraVariable(FNiagaraTypeDefinition::GetFloatDef(), TEXT("Delta Time")));
-		Sig.Outputs.Add(FNiagaraVariable(FNiagaraTypeDefinition::GetVec3Def(), TEXT("Material Multiplier")));
+		Sig.Outputs.Add(FNiagaraVariable(FNiagaraTypeDefinition::GetVec3Def(), TEXT("Node Position")));
 
 		OutFunctions.Add(Sig);
 	}
@@ -1590,6 +1902,7 @@ void UNiagaraDataInterfaceHairStrands::GetFunctions(TArray<FNiagaraFunctionSigna
 		Sig.bMemberFunction = true;
 		Sig.bRequiresContext = false;
 		Sig.Inputs.Add(FNiagaraVariable(FNiagaraTypeDefinition(GetClass()), TEXT("Hair Strands")));
+		Sig.Inputs.Add(FNiagaraVariable(FNiagaraTypeDefinition::GetBoolDef(), TEXT("Enable Constraint")));
 		Sig.Inputs.Add(FNiagaraVariable(FNiagaraTypeDefinition::GetFloatDef(), TEXT("Rest Length")));
 		Sig.Inputs.Add(FNiagaraVariable(FNiagaraTypeDefinition::GetQuatDef(), TEXT("Rest Darboux")));
 		Sig.Inputs.Add(FNiagaraVariable(FNiagaraTypeDefinition::GetFloatDef(), TEXT("Delta Time")));
@@ -1606,18 +1919,19 @@ void UNiagaraDataInterfaceHairStrands::GetFunctions(TArray<FNiagaraFunctionSigna
 		Sig.bMemberFunction = true;
 		Sig.bRequiresContext = false;
 		Sig.Inputs.Add(FNiagaraVariable(FNiagaraTypeDefinition(GetClass()), TEXT("Hair Strands")));
-		Sig.Inputs.Add(FNiagaraVariable(FNiagaraTypeDefinition::GetFloatDef(), TEXT("Stretch Stiffness")));
+		Sig.Inputs.Add(FNiagaraVariable(FNiagaraTypeDefinition::GetBoolDef(), TEXT("Enable Constraint")));
+		Sig.Inputs.Add(FNiagaraVariable(FNiagaraTypeDefinition::GetFloatDef(), TEXT("Bend Stiffness")));
 		Sig.Inputs.Add(FNiagaraVariable(FNiagaraTypeDefinition::GetFloatDef(), TEXT("Node Thickness")));
 		Sig.Inputs.Add(FNiagaraVariable(FNiagaraTypeDefinition::GetFloatDef(), TEXT("Rest Length")));
 		Sig.Inputs.Add(FNiagaraVariable(FNiagaraTypeDefinition::GetQuatDef(), TEXT("Rest Darboux")));
 		Sig.Inputs.Add(FNiagaraVariable(FNiagaraTypeDefinition::GetFloatDef(), TEXT("Delta Time")));
-		Sig.Outputs.Add(FNiagaraVariable(FNiagaraTypeDefinition::GetVec3Def(), TEXT("Material Multiplier")));
+		Sig.Outputs.Add(FNiagaraVariable(FNiagaraTypeDefinition::GetVec3Def(), TEXT("Node Position")));
 
 		OutFunctions.Add(Sig);
 	}
 	{
 		FNiagaraFunctionSignature Sig;
-		Sig.Name = SolveStaticCollisionConstraintName;
+		Sig.Name = SolveHardCollisionConstraintName;
 		Sig.bMemberFunction = true;
 		Sig.bRequiresContext = false;
 		Sig.Inputs.Add(FNiagaraVariable(FNiagaraTypeDefinition(GetClass()), TEXT("Hair Strands")));
@@ -1629,17 +1943,17 @@ void UNiagaraDataInterfaceHairStrands::GetFunctions(TArray<FNiagaraFunctionSigna
 		Sig.Inputs.Add(FNiagaraVariable(FNiagaraTypeDefinition::GetFloatDef(), TEXT("Static Friction")));
 		Sig.Inputs.Add(FNiagaraVariable(FNiagaraTypeDefinition::GetFloatDef(), TEXT("Kinetic Friction")));
 		Sig.Inputs.Add(FNiagaraVariable(FNiagaraTypeDefinition::GetFloatDef(), TEXT("Delta Time")));
-		Sig.Inputs.Add(FNiagaraVariable(FNiagaraTypeDefinition::GetVec3Def(), TEXT("Previous Position")));
 		Sig.Outputs.Add(FNiagaraVariable(FNiagaraTypeDefinition::GetVec3Def(), TEXT("Constraint Multiplier")));
 
 		OutFunctions.Add(Sig);
 	}
 	{
 		FNiagaraFunctionSignature Sig;
-		Sig.Name = ProjectStaticCollisionConstraintName;
+		Sig.Name = ProjectHardCollisionConstraintName;
 		Sig.bMemberFunction = true;
 		Sig.bRequiresContext = false;
 		Sig.Inputs.Add(FNiagaraVariable(FNiagaraTypeDefinition(GetClass()), TEXT("Hair Strands")));
+		Sig.Inputs.Add(FNiagaraVariable(FNiagaraTypeDefinition::GetBoolDef(), TEXT("Enable Constraint")));
 		Sig.Inputs.Add(FNiagaraVariable(FNiagaraTypeDefinition::GetFloatDef(), TEXT("Penetration Depth")));
 		Sig.Inputs.Add(FNiagaraVariable(FNiagaraTypeDefinition::GetVec3Def(), TEXT("Collision Position")));
 		Sig.Inputs.Add(FNiagaraVariable(FNiagaraTypeDefinition::GetVec3Def(), TEXT("Collision Velocity")));
@@ -1647,8 +1961,6 @@ void UNiagaraDataInterfaceHairStrands::GetFunctions(TArray<FNiagaraFunctionSigna
 		Sig.Inputs.Add(FNiagaraVariable(FNiagaraTypeDefinition::GetFloatDef(), TEXT("Static Friction")));
 		Sig.Inputs.Add(FNiagaraVariable(FNiagaraTypeDefinition::GetFloatDef(), TEXT("Kinetic Friction")));
 		Sig.Inputs.Add(FNiagaraVariable(FNiagaraTypeDefinition::GetFloatDef(), TEXT("Delta Time")));
-		Sig.Inputs.Add(FNiagaraVariable(FNiagaraTypeDefinition::GetVec3Def(), TEXT("Previous Position")));
-		Sig.Inputs.Add(FNiagaraVariable(FNiagaraTypeDefinition::GetVec3Def(), TEXT("Node Position")));
 		Sig.Outputs.Add(FNiagaraVariable(FNiagaraTypeDefinition::GetVec3Def(), TEXT("Node Position")));
 
 		OutFunctions.Add(Sig);
@@ -1696,6 +2008,7 @@ void UNiagaraDataInterfaceHairStrands::GetFunctions(TArray<FNiagaraFunctionSigna
 		Sig.bMemberFunction = true;
 		Sig.bRequiresContext = false;
 		Sig.Inputs.Add(FNiagaraVariable(FNiagaraTypeDefinition(GetClass()), TEXT("Hair Strands")));
+		Sig.Inputs.Add(FNiagaraVariable(FNiagaraTypeDefinition::GetBoolDef(), TEXT("Enable Constraint")));
 		Sig.Inputs.Add(FNiagaraVariable(FNiagaraTypeDefinition::GetFloatDef(), TEXT("Collision Stiffness")));
 		Sig.Inputs.Add(FNiagaraVariable(FNiagaraTypeDefinition::GetFloatDef(), TEXT("Penetration Depth")));
 		Sig.Inputs.Add(FNiagaraVariable(FNiagaraTypeDefinition::GetVec3Def(), TEXT("Collision Position")));
@@ -1703,10 +2016,7 @@ void UNiagaraDataInterfaceHairStrands::GetFunctions(TArray<FNiagaraFunctionSigna
 		Sig.Inputs.Add(FNiagaraVariable(FNiagaraTypeDefinition::GetVec3Def(), TEXT("Collision Normal")));
 		Sig.Inputs.Add(FNiagaraVariable(FNiagaraTypeDefinition::GetFloatDef(), TEXT("Static Friction")));
 		Sig.Inputs.Add(FNiagaraVariable(FNiagaraTypeDefinition::GetFloatDef(), TEXT("Kinetic Friction")));
-		Sig.Inputs.Add(FNiagaraVariable(FNiagaraTypeDefinition::GetFloatDef(), TEXT("Node Mass")));
 		Sig.Inputs.Add(FNiagaraVariable(FNiagaraTypeDefinition::GetFloatDef(), TEXT("Delta Time")));
-		Sig.Inputs.Add(FNiagaraVariable(FNiagaraTypeDefinition::GetVec3Def(), TEXT("Previous Position")));
-		Sig.Inputs.Add(FNiagaraVariable(FNiagaraTypeDefinition::GetVec3Def(), TEXT("Node Position")));
 		Sig.Outputs.Add(FNiagaraVariable(FNiagaraTypeDefinition::GetVec3Def(), TEXT("Node Position")));
 
 		OutFunctions.Add(Sig);
@@ -1725,7 +2035,7 @@ void UNiagaraDataInterfaceHairStrands::GetFunctions(TArray<FNiagaraFunctionSigna
 	}
 	{
 		FNiagaraFunctionSignature Sig;
-		Sig.Name = UpdateNodeOrientationName;
+		Sig.Name = UpdateMaterialFrameName;
 		Sig.bMemberFunction = true;
 		Sig.bRequiresContext = false;
 		Sig.Inputs.Add(FNiagaraVariable(FNiagaraTypeDefinition(GetClass()), TEXT("Hair Strands")));
@@ -1733,6 +2043,16 @@ void UNiagaraDataInterfaceHairStrands::GetFunctions(TArray<FNiagaraFunctionSigna
 
 		OutFunctions.Add(Sig);
 	}
+	{
+		FNiagaraFunctionSignature Sig;
+		Sig.Name = ComputeMaterialFrameName;
+		Sig.bMemberFunction = true;
+		Sig.bRequiresContext = false;
+		Sig.Inputs.Add(FNiagaraVariable(FNiagaraTypeDefinition(GetClass()), TEXT("Hair Strands")));
+		Sig.Outputs.Add(FNiagaraVariable(FNiagaraTypeDefinition::GetBoolDef(), TEXT("Compute Status")));
+
+		OutFunctions.Add(Sig);
+	}
 	{
 		FNiagaraFunctionSignature Sig;
 		Sig.Name = ComputeAirDragForceName;
@@ -1765,6 +2085,39 @@ void UNiagaraDataInterfaceHairStrands::GetFunctions(TArray<FNiagaraFunctionSigna
 
 DEFINE_NDI_DIRECT_FUNC_BINDER(UNiagaraDataInterfaceHairStrands, GetNumStrands);
 DEFINE_NDI_DIRECT_FUNC_BINDER(UNiagaraDataInterfaceHairStrands, GetStrandSize);
+
+DEFINE_NDI_DIRECT_FUNC_BINDER(UNiagaraDataInterfaceHairStrands, GetSubSteps);
+DEFINE_NDI_DIRECT_FUNC_BINDER(UNiagaraDataInterfaceHairStrands, GetIterationCount);
+
+DEFINE_NDI_DIRECT_FUNC_BINDER(UNiagaraDataInterfaceHairStrands, GetGravityVector);
+DEFINE_NDI_DIRECT_FUNC_BINDER(UNiagaraDataInterfaceHairStrands, GetAirDrag);
+DEFINE_NDI_DIRECT_FUNC_BINDER(UNiagaraDataInterfaceHairStrands, GetAirVelocity);
+
+DEFINE_NDI_DIRECT_FUNC_BINDER(UNiagaraDataInterfaceHairStrands, GetSolveBend);
+DEFINE_NDI_DIRECT_FUNC_BINDER(UNiagaraDataInterfaceHairStrands, GetProjectBend);
+DEFINE_NDI_DIRECT_FUNC_BINDER(UNiagaraDataInterfaceHairStrands, GetBendDamping);
+DEFINE_NDI_DIRECT_FUNC_BINDER(UNiagaraDataInterfaceHairStrands, GetBendStiffness);
+DEFINE_NDI_DIRECT_FUNC_BINDER(UNiagaraDataInterfaceHairStrands, GetBendScale);
+
+DEFINE_NDI_DIRECT_FUNC_BINDER(UNiagaraDataInterfaceHairStrands, GetSolveStretch);
+DEFINE_NDI_DIRECT_FUNC_BINDER(UNiagaraDataInterfaceHairStrands, GetProjectStretch);
+DEFINE_NDI_DIRECT_FUNC_BINDER(UNiagaraDataInterfaceHairStrands, GetStretchDamping);
+DEFINE_NDI_DIRECT_FUNC_BINDER(UNiagaraDataInterfaceHairStrands, GetStretchStiffness);
+DEFINE_NDI_DIRECT_FUNC_BINDER(UNiagaraDataInterfaceHairStrands, GetStretchScale);
+
+DEFINE_NDI_DIRECT_FUNC_BINDER(UNiagaraDataInterfaceHairStrands, GetSolveCollision);
+DEFINE_NDI_DIRECT_FUNC_BINDER(UNiagaraDataInterfaceHairStrands, GetProjectCollision);
+DEFINE_NDI_DIRECT_FUNC_BINDER(UNiagaraDataInterfaceHairStrands, GetStaticFriction);
+DEFINE_NDI_DIRECT_FUNC_BINDER(UNiagaraDataInterfaceHairStrands, GetKineticFriction);
+DEFINE_NDI_DIRECT_FUNC_BINDER(UNiagaraDataInterfaceHairStrands, GetStrandsViscosity);
+DEFINE_NDI_DIRECT_FUNC_BINDER(UNiagaraDataInterfaceHairStrands, GetCollisionRadius);
+DEFINE_NDI_DIRECT_FUNC_BINDER(UNiagaraDataInterfaceHairStrands, GetRadiusScale);
+
+DEFINE_NDI_DIRECT_FUNC_BINDER(UNiagaraDataInterfaceHairStrands, GetStrandsDensity);
+DEFINE_NDI_DIRECT_FUNC_BINDER(UNiagaraDataInterfaceHairStrands, GetStrandsSmoothing);
+DEFINE_NDI_DIRECT_FUNC_BINDER(UNiagaraDataInterfaceHairStrands, GetStrandsThickness);
+DEFINE_NDI_DIRECT_FUNC_BINDER(UNiagaraDataInterfaceHairStrands, GetThicknessScale);
+
 DEFINE_NDI_DIRECT_FUNC_BINDER(UNiagaraDataInterfaceHairStrands, GetWorldTransform);
 DEFINE_NDI_DIRECT_FUNC_BINDER(UNiagaraDataInterfaceHairStrands, GetWorldInverse);
 DEFINE_NDI_DIRECT_FUNC_BINDER(UNiagaraDataInterfaceHairStrands, GetPointPosition);
@@ -1808,15 +2161,16 @@ DEFINE_NDI_DIRECT_FUNC_BINDER(UNiagaraDataInterfaceHairStrands, SetupBendRodMate
 DEFINE_NDI_DIRECT_FUNC_BINDER(UNiagaraDataInterfaceHairStrands, SolveBendRodMaterial);
 DEFINE_NDI_DIRECT_FUNC_BINDER(UNiagaraDataInterfaceHairStrands, ProjectBendRodMaterial);
 
-DEFINE_NDI_DIRECT_FUNC_BINDER(UNiagaraDataInterfaceHairStrands, SolveStaticCollisionConstraint);
-DEFINE_NDI_DIRECT_FUNC_BINDER(UNiagaraDataInterfaceHairStrands, ProjectStaticCollisionConstraint);
+DEFINE_NDI_DIRECT_FUNC_BINDER(UNiagaraDataInterfaceHairStrands, SolveHardCollisionConstraint);
+DEFINE_NDI_DIRECT_FUNC_BINDER(UNiagaraDataInterfaceHairStrands, ProjectHardCollisionConstraint);
 
 DEFINE_NDI_DIRECT_FUNC_BINDER(UNiagaraDataInterfaceHairStrands, SetupSoftCollisionConstraint);
 DEFINE_NDI_DIRECT_FUNC_BINDER(UNiagaraDataInterfaceHairStrands, ProjectSoftCollisionConstraint);
 DEFINE_NDI_DIRECT_FUNC_BINDER(UNiagaraDataInterfaceHairStrands, SolveSoftCollisionConstraint);
 
 DEFINE_NDI_DIRECT_FUNC_BINDER(UNiagaraDataInterfaceHairStrands, ComputeRestDirection);
-DEFINE_NDI_DIRECT_FUNC_BINDER(UNiagaraDataInterfaceHairStrands, UpdateNodeOrientation);
+DEFINE_NDI_DIRECT_FUNC_BINDER(UNiagaraDataInterfaceHairStrands, UpdateMaterialFrame);
+DEFINE_NDI_DIRECT_FUNC_BINDER(UNiagaraDataInterfaceHairStrands, ComputeMaterialFrame);
 DEFINE_NDI_DIRECT_FUNC_BINDER(UNiagaraDataInterfaceHairStrands, ComputeAirDragForce);
 DEFINE_NDI_DIRECT_FUNC_BINDER(UNiagaraDataInterfaceHairStrands, NeedSimulationReset);
 
@@ -1832,6 +2186,136 @@ void UNiagaraDataInterfaceHairStrands::GetVMExternalFunction(const FVMExternalFu
 		check(BindingInfo.GetNumInputs() == 1 && BindingInfo.GetNumOutputs() == 1);
 		NDI_FUNC_BINDER(UNiagaraDataInterfaceHairStrands, GetStrandSize)::Bind(this, OutFunc);
 	}
+	else if (BindingInfo.Name == GetSubStepsName)
+	{
+		check(BindingInfo.GetNumInputs() == 1 && BindingInfo.GetNumOutputs() == 1);
+		NDI_FUNC_BINDER(UNiagaraDataInterfaceHairStrands, GetSubSteps)::Bind(this, OutFunc);
+	}
+	else if (BindingInfo.Name == GetIterationCountName)
+	{
+		check(BindingInfo.GetNumInputs() == 1 && BindingInfo.GetNumOutputs() == 1);
+		NDI_FUNC_BINDER(UNiagaraDataInterfaceHairStrands, GetIterationCount)::Bind(this, OutFunc);
+	}
+	else if (BindingInfo.Name == GetGravityVectorName)
+	{
+		check(BindingInfo.GetNumInputs() == 1 && BindingInfo.GetNumOutputs() == 3);
+		NDI_FUNC_BINDER(UNiagaraDataInterfaceHairStrands, GetGravityVector)::Bind(this, OutFunc);
+	}
+	else if (BindingInfo.Name == GetAirDragName)
+	{
+		check(BindingInfo.GetNumInputs() == 1 && BindingInfo.GetNumOutputs() == 1);
+		NDI_FUNC_BINDER(UNiagaraDataInterfaceHairStrands, GetAirDrag)::Bind(this, OutFunc);
+	}
+	else if (BindingInfo.Name == GetAirVelocityName)
+	{
+		check(BindingInfo.GetNumInputs() == 1 && BindingInfo.GetNumOutputs() == 3);
+		NDI_FUNC_BINDER(UNiagaraDataInterfaceHairStrands, GetAirVelocity)::Bind(this, OutFunc);
+	}
+	else if (BindingInfo.Name == GetSolveBendName)
+	{
+		check(BindingInfo.GetNumInputs() == 1 && BindingInfo.GetNumOutputs() == 1);
+		NDI_FUNC_BINDER(UNiagaraDataInterfaceHairStrands, GetSolveBend)::Bind(this, OutFunc);
+	}
+	else if (BindingInfo.Name == GetProjectBendName)
+	{
+		check(BindingInfo.GetNumInputs() == 1 && BindingInfo.GetNumOutputs() == 1);
+		NDI_FUNC_BINDER(UNiagaraDataInterfaceHairStrands, GetProjectBend)::Bind(this, OutFunc);
+	}
+	else if (BindingInfo.Name == GetBendDampingName)
+	{
+		check(BindingInfo.GetNumInputs() == 1 && BindingInfo.GetNumOutputs() == 1);
+		NDI_FUNC_BINDER(UNiagaraDataInterfaceHairStrands, GetBendDamping)::Bind(this, OutFunc);
+	}
+	else if (BindingInfo.Name == GetBendStiffnessName)
+	{
+		check(BindingInfo.GetNumInputs() == 1 && BindingInfo.GetNumOutputs() == 1);
+		NDI_FUNC_BINDER(UNiagaraDataInterfaceHairStrands, GetBendStiffness)::Bind(this, OutFunc);
+	}
+	else if (BindingInfo.Name == GetBendScaleName)
+	{
+		check(BindingInfo.GetNumInputs() == 1 && BindingInfo.GetNumOutputs() == 4);
+		NDI_FUNC_BINDER(UNiagaraDataInterfaceHairStrands, GetBendScale)::Bind(this, OutFunc);
+	}
+	else if (BindingInfo.Name == GetSolveStretchName)
+	{
+		check(BindingInfo.GetNumInputs() == 1 && BindingInfo.GetNumOutputs() == 1);
+		NDI_FUNC_BINDER(UNiagaraDataInterfaceHairStrands, GetSolveStretch)::Bind(this, OutFunc);
+	}
+	else if (BindingInfo.Name == GetProjectStretchName)
+	{
+		check(BindingInfo.GetNumInputs() == 1 && BindingInfo.GetNumOutputs() == 1);
+		NDI_FUNC_BINDER(UNiagaraDataInterfaceHairStrands, GetProjectStretch)::Bind(this, OutFunc);
+	}
+	else if (BindingInfo.Name == GetStretchDampingName)
+	{
+		check(BindingInfo.GetNumInputs() == 1 && BindingInfo.GetNumOutputs() == 1);
+		NDI_FUNC_BINDER(UNiagaraDataInterfaceHairStrands, GetStretchDamping)::Bind(this, OutFunc);
+	}
+	else if (BindingInfo.Name == GetStretchStiffnessName)
+	{
+		check(BindingInfo.GetNumInputs() == 1 && BindingInfo.GetNumOutputs() == 1);
+		NDI_FUNC_BINDER(UNiagaraDataInterfaceHairStrands, GetStretchStiffness)::Bind(this, OutFunc);
+	}
+	else if (BindingInfo.Name == GetStretchScaleName)
+	{
+		check(BindingInfo.GetNumInputs() == 1 && BindingInfo.GetNumOutputs() == 4);
+		NDI_FUNC_BINDER(UNiagaraDataInterfaceHairStrands, GetStretchScale)::Bind(this, OutFunc);
+	}
+	else if (BindingInfo.Name == GetSolveCollisionName)
+	{
+		check(BindingInfo.GetNumInputs() == 1 && BindingInfo.GetNumOutputs() == 1);
+		NDI_FUNC_BINDER(UNiagaraDataInterfaceHairStrands, GetSolveCollision)::Bind(this, OutFunc);
+	}
+	else if (BindingInfo.Name == GetProjectCollisionName)
+	{
+		check(BindingInfo.GetNumInputs() == 1 && BindingInfo.GetNumOutputs() == 1);
+		NDI_FUNC_BINDER(UNiagaraDataInterfaceHairStrands, GetProjectCollision)::Bind(this, OutFunc);
+	}
+	else if (BindingInfo.Name == GetStaticFrictionName)
+	{
+		check(BindingInfo.GetNumInputs() == 1 && BindingInfo.GetNumOutputs() == 1);
+		NDI_FUNC_BINDER(UNiagaraDataInterfaceHairStrands, GetStaticFriction)::Bind(this, OutFunc);
+	}
+	else if (BindingInfo.Name == GetKineticFrictionName)
+	{
+		check(BindingInfo.GetNumInputs() == 1 && BindingInfo.GetNumOutputs() == 1);
+		NDI_FUNC_BINDER(UNiagaraDataInterfaceHairStrands, GetKineticFriction)::Bind(this, OutFunc);
+	}
+	else if (BindingInfo.Name == GetStrandsViscosityName)
+	{
+		check(BindingInfo.GetNumInputs() == 1 && BindingInfo.GetNumOutputs() == 1);
+		NDI_FUNC_BINDER(UNiagaraDataInterfaceHairStrands, GetStrandsViscosity)::Bind(this, OutFunc);
+	}
+	else if (BindingInfo.Name == GetCollisionRadiusName)
+	{
+		check(BindingInfo.GetNumInputs() == 1 && BindingInfo.GetNumOutputs() == 1);
+		NDI_FUNC_BINDER(UNiagaraDataInterfaceHairStrands, GetCollisionRadius)::Bind(this, OutFunc);
+	}
+	else if (BindingInfo.Name == GetRadiusScaleName)
+	{
+		check(BindingInfo.GetNumInputs() == 1 && BindingInfo.GetNumOutputs() == 4);
+		NDI_FUNC_BINDER(UNiagaraDataInterfaceHairStrands, GetRadiusScale)::Bind(this, OutFunc);
+	}
+	else if (BindingInfo.Name == GetStrandsDensityName)
+	{
+		check(BindingInfo.GetNumInputs() == 1 && BindingInfo.GetNumOutputs() == 1);
+		NDI_FUNC_BINDER(UNiagaraDataInterfaceHairStrands, GetStrandsDensity)::Bind(this, OutFunc);
+	}
+	else if (BindingInfo.Name == GetStrandsSmoothingName)
+	{
+		check(BindingInfo.GetNumInputs() == 1 && BindingInfo.GetNumOutputs() == 1);
+		NDI_FUNC_BINDER(UNiagaraDataInterfaceHairStrands, GetStrandsSmoothing)::Bind(this, OutFunc);
+	}
+	else if (BindingInfo.Name == GetStrandsThicknessName)
+	{
+		check(BindingInfo.GetNumInputs() == 1 && BindingInfo.GetNumOutputs() == 1);
+		NDI_FUNC_BINDER(UNiagaraDataInterfaceHairStrands, GetStrandsThickness)::Bind(this, OutFunc);
+	}
+	else if (BindingInfo.Name == GetThicknessScaleName)
+	{
+		check(BindingInfo.GetNumInputs() == 1 && BindingInfo.GetNumOutputs() == 4);
+		NDI_FUNC_BINDER(UNiagaraDataInterfaceHairStrands, GetThicknessScale)::Bind(this, OutFunc);
+	}
 	else if (BindingInfo.Name == GetWorldTransformName)
 	{
 		check(BindingInfo.GetNumInputs() == 1 && BindingInfo.GetNumOutputs() == 16);
@@ -1964,7 +2448,7 @@ void UNiagaraDataInterfaceHairStrands::GetVMExternalFunction(const FVMExternalFu
 	}
 	else if (BindingInfo.Name == ProjectDistanceSpringMaterialName)
 	{
-		check(BindingInfo.GetNumInputs() == 10 && BindingInfo.GetNumOutputs() == 3);
+		check(BindingInfo.GetNumInputs() == 7 && BindingInfo.GetNumOutputs() == 3);
 		NDI_FUNC_BINDER(UNiagaraDataInterfaceHairStrands, ProjectDistanceSpringMaterial)::Bind(this, OutFunc);
 	}
 	else if (BindingInfo.Name == SetupAngularSpringMaterialName)
@@ -1979,7 +2463,7 @@ void UNiagaraDataInterfaceHairStrands::GetVMExternalFunction(const FVMExternalFu
 	}
 	else if (BindingInfo.Name == ProjectAngularSpringMaterialName)
 	{
-		check(BindingInfo.GetNumInputs() == 12 && BindingInfo.GetNumOutputs() == 3);
+		check(BindingInfo.GetNumInputs() == 9 && BindingInfo.GetNumOutputs() == 3);
 		NDI_FUNC_BINDER(UNiagaraDataInterfaceHairStrands, ProjectAngularSpringMaterial)::Bind(this, OutFunc);
 	}
 	else if (BindingInfo.Name == SetupStretchRodMaterialName)
@@ -1989,12 +2473,12 @@ void UNiagaraDataInterfaceHairStrands::GetVMExternalFunction(const FVMExternalFu
 	}
 	else if (BindingInfo.Name == SolveStretchRodMaterialName)
 	{
-		check(BindingInfo.GetNumInputs() == 9 && BindingInfo.GetNumOutputs() == 3);
+		check(BindingInfo.GetNumInputs() == 10 && BindingInfo.GetNumOutputs() == 3);
 		NDI_FUNC_BINDER(UNiagaraDataInterfaceHairStrands, SolveStretchRodMaterial)::Bind(this, OutFunc);
 	}
 	else if (BindingInfo.Name == ProjectStretchRodMaterialName)
 	{
-		check(BindingInfo.GetNumInputs() == 5 && BindingInfo.GetNumOutputs() == 3);
+		check(BindingInfo.GetNumInputs() == 6 && BindingInfo.GetNumOutputs() == 3);
 		NDI_FUNC_BINDER(UNiagaraDataInterfaceHairStrands, ProjectStretchRodMaterial)::Bind(this, OutFunc);
 	}
 	else if (BindingInfo.Name == SetupBendRodMaterialName)
@@ -2004,23 +2488,23 @@ void UNiagaraDataInterfaceHairStrands::GetVMExternalFunction(const FVMExternalFu
 	}
 	else if (BindingInfo.Name == SolveBendRodMaterialName)
 	{
-		check(BindingInfo.GetNumInputs() == 13 && BindingInfo.GetNumOutputs() == 3);
+		check(BindingInfo.GetNumInputs() == 14 && BindingInfo.GetNumOutputs() == 3);
 		NDI_FUNC_BINDER(UNiagaraDataInterfaceHairStrands, SolveBendRodMaterial)::Bind(this, OutFunc);
 	}
 	else if (BindingInfo.Name == ProjectBendRodMaterialName)
 	{
-		check(BindingInfo.GetNumInputs() == 9 && BindingInfo.GetNumOutputs() == 3);
+		check(BindingInfo.GetNumInputs() == 10 && BindingInfo.GetNumOutputs() == 3);
 		NDI_FUNC_BINDER(UNiagaraDataInterfaceHairStrands, ProjectBendRodMaterial)::Bind(this, OutFunc);
 	}
-	else if (BindingInfo.Name == SolveStaticCollisionConstraintName)
+	else if (BindingInfo.Name == SolveHardCollisionConstraintName)
 	{
-		check(BindingInfo.GetNumInputs() == 18 && BindingInfo.GetNumOutputs() == 3);
-		NDI_FUNC_BINDER(UNiagaraDataInterfaceHairStrands, SolveStaticCollisionConstraint)::Bind(this, OutFunc);
+		check(BindingInfo.GetNumInputs() == 15 && BindingInfo.GetNumOutputs() == 3);
+		NDI_FUNC_BINDER(UNiagaraDataInterfaceHairStrands, SolveHardCollisionConstraint)::Bind(this, OutFunc);
 	}
-	else if (BindingInfo.Name == ProjectStaticCollisionConstraintName)
+	else if (BindingInfo.Name == ProjectHardCollisionConstraintName)
 	{
-		check(BindingInfo.GetNumInputs() == 20 && BindingInfo.GetNumOutputs() == 3);
-		NDI_FUNC_BINDER(UNiagaraDataInterfaceHairStrands, ProjectStaticCollisionConstraint)::Bind(this, OutFunc);
+		check(BindingInfo.GetNumInputs() == 15 && BindingInfo.GetNumOutputs() == 3);
+		NDI_FUNC_BINDER(UNiagaraDataInterfaceHairStrands, ProjectHardCollisionConstraint)::Bind(this, OutFunc);
 	}
 	else if (BindingInfo.Name == SolveSoftCollisionConstraintName)
 	{
@@ -2029,7 +2513,7 @@ void UNiagaraDataInterfaceHairStrands::GetVMExternalFunction(const FVMExternalFu
 	}
 	else if (BindingInfo.Name == ProjectSoftCollisionConstraintName)
 	{
-		check(BindingInfo.GetNumInputs() == 22 && BindingInfo.GetNumOutputs() == 3);
+		check(BindingInfo.GetNumInputs() == 16 && BindingInfo.GetNumOutputs() == 3);
 		NDI_FUNC_BINDER(UNiagaraDataInterfaceHairStrands, ProjectSoftCollisionConstraint)::Bind(this, OutFunc);
 	}
 	else if (BindingInfo.Name == SetupSoftCollisionConstraintName)
@@ -2042,10 +2526,15 @@ void UNiagaraDataInterfaceHairStrands::GetVMExternalFunction(const FVMExternalFu
 		check(BindingInfo.GetNumInputs() == 8 && BindingInfo.GetNumOutputs() == 3);
 		NDI_FUNC_BINDER(UNiagaraDataInterfaceHairStrands, ComputeRestDirection)::Bind(this, OutFunc);
 	}
-	else if (BindingInfo.Name == UpdateNodeOrientationName)
+	else if (BindingInfo.Name == UpdateMaterialFrameName)
+	{
+		check(BindingInfo.GetNumInputs() == 1 && BindingInfo.GetNumOutputs() == 1);
+		NDI_FUNC_BINDER(UNiagaraDataInterfaceHairStrands, UpdateMaterialFrame)::Bind(this, OutFunc);
+	}
+	else if (BindingInfo.Name == ComputeMaterialFrameName)
 	{
 		check(BindingInfo.GetNumInputs() == 1 && BindingInfo.GetNumOutputs() == 1);
-		NDI_FUNC_BINDER(UNiagaraDataInterfaceHairStrands, UpdateNodeOrientation)::Bind(this, OutFunc);
+		NDI_FUNC_BINDER(UNiagaraDataInterfaceHairStrands, ComputeMaterialFrame)::Bind(this, OutFunc);
 	}
 	else if (BindingInfo.Name == ComputeAirDragForceName)
 	{
@@ -2117,7 +2606,326 @@ void UNiagaraDataInterfaceHairStrands::GetStrandSize(FVectorVMContext& Context)
 
 	for (int32 InstanceIdx = 0; InstanceIdx < Context.NumInstances; ++InstanceIdx)
 	{
-		*OutStrandSize.GetDestAndAdvance() = InstData->StrandSize;
+		*OutStrandSize.GetDestAndAdvance() = InstData->StrandsSize;
+	}
+}
+
+void UNiagaraDataInterfaceHairStrands::GetSubSteps(FVectorVMContext& Context)
+{
+	VectorVM::FUserPtrHandler<FNDIHairStrandsData> InstData(Context);
+	VectorVM::FExternalFuncRegisterHandler<int32> OutSubSteps(Context);
+
+	for (int32 InstanceIdx = 0; InstanceIdx < Context.NumInstances; ++InstanceIdx)
+	{
+		*OutSubSteps.GetDestAndAdvance() = InstData->SubSteps;
+	}
+}
+
+void UNiagaraDataInterfaceHairStrands::GetIterationCount(FVectorVMContext& Context)
+{
+	VectorVM::FUserPtrHandler<FNDIHairStrandsData> InstData(Context);
+	VectorVM::FExternalFuncRegisterHandler<int32> OutIterationCount(Context);
+
+	for (int32 InstanceIdx = 0; InstanceIdx < Context.NumInstances; ++InstanceIdx)
+	{
+		*OutIterationCount.GetDestAndAdvance() = InstData->IterationCount;
+	}
+}
+
+void UNiagaraDataInterfaceHairStrands::GetGravityVector(FVectorVMContext& Context)
+{
+	VectorVM::FUserPtrHandler<FNDIHairStrandsData> InstData(Context);
+	VectorVM::FExternalFuncRegisterHandler<float> OutGravityVectorX(Context);
+	VectorVM::FExternalFuncRegisterHandler<float> OutGravityVectorY(Context);
+	VectorVM::FExternalFuncRegisterHandler<float> OutGravityVectorZ(Context);
+
+	for (int32 InstanceIdx = 0; InstanceIdx < Context.NumInstances; ++InstanceIdx)
+	{
+		*OutGravityVectorX.GetDestAndAdvance() = InstData->GravityVector.X;
+		*OutGravityVectorY.GetDestAndAdvance() = InstData->GravityVector.Y;
+		*OutGravityVectorZ.GetDestAndAdvance() = InstData->GravityVector.Z;
+	}
+}
+
+void UNiagaraDataInterfaceHairStrands::GetAirDrag(FVectorVMContext& Context)
+{
+	VectorVM::FUserPtrHandler<FNDIHairStrandsData> InstData(Context);
+	VectorVM::FExternalFuncRegisterHandler<float> OutAirDrag(Context);
+
+	for (int32 InstanceIdx = 0; InstanceIdx < Context.NumInstances; ++InstanceIdx)
+	{
+		*OutAirDrag.GetDestAndAdvance() = InstData->AirDrag;
+	}
+}
+
+void UNiagaraDataInterfaceHairStrands::GetAirVelocity(FVectorVMContext& Context)
+{
+	VectorVM::FUserPtrHandler<FNDIHairStrandsData> InstData(Context);
+	VectorVM::FExternalFuncRegisterHandler<float> OutAirVelocityX(Context);
+	VectorVM::FExternalFuncRegisterHandler<float> OutAirVelocityY(Context);
+	VectorVM::FExternalFuncRegisterHandler<float> OutAirVelocityZ(Context);
+
+	for (int32 InstanceIdx = 0; InstanceIdx < Context.NumInstances; ++InstanceIdx)
+	{
+		*OutAirVelocityX.GetDestAndAdvance() = InstData->AirVelocity.X;
+		*OutAirVelocityY.GetDestAndAdvance() = InstData->AirVelocity.Y;
+		*OutAirVelocityZ.GetDestAndAdvance() = InstData->AirVelocity.Z;
+	}
+}
+
+void UNiagaraDataInterfaceHairStrands::GetSolveBend(FVectorVMContext& Context)
+{
+	VectorVM::FUserPtrHandler<FNDIHairStrandsData> InstData(Context);
+	VectorVM::FExternalFuncRegisterHandler<int32> OutSolveBend(Context);
+
+	for (int32 InstanceIdx = 0; InstanceIdx < Context.NumInstances; ++InstanceIdx)
+	{
+		*OutSolveBend.GetDestAndAdvance() = InstData->SolveBend;
+	}
+}
+
+void UNiagaraDataInterfaceHairStrands::GetProjectBend(FVectorVMContext& Context)
+{
+	VectorVM::FUserPtrHandler<FNDIHairStrandsData> InstData(Context);
+	VectorVM::FExternalFuncRegisterHandler<int32> OutProjectBend(Context);
+
+	for (int32 InstanceIdx = 0; InstanceIdx < Context.NumInstances; ++InstanceIdx)
+	{
+		*OutProjectBend.GetDestAndAdvance() = InstData->ProjectBend;
+	}
+}
+
+void UNiagaraDataInterfaceHairStrands::GetBendDamping(FVectorVMContext& Context)
+{
+	VectorVM::FUserPtrHandler<FNDIHairStrandsData> InstData(Context);
+	VectorVM::FExternalFuncRegisterHandler<float> OutBendDamping(Context);
+
+	for (int32 InstanceIdx = 0; InstanceIdx < Context.NumInstances; ++InstanceIdx)
+	{
+		*OutBendDamping.GetDestAndAdvance() = InstData->BendDamping;
+	}
+}
+
+void UNiagaraDataInterfaceHairStrands::GetBendStiffness(FVectorVMContext& Context)
+{
+	VectorVM::FUserPtrHandler<FNDIHairStrandsData> InstData(Context);
+	VectorVM::FExternalFuncRegisterHandler<float> OutBendStiffness(Context);
+
+	for (int32 InstanceIdx = 0; InstanceIdx < Context.NumInstances; ++InstanceIdx)
+	{
+		*OutBendStiffness.GetDestAndAdvance() = InstData->BendStiffness;
+	}
+}
+
+void UNiagaraDataInterfaceHairStrands::GetBendScale(FVectorVMContext& Context)
+{
+	VectorVM::FUserPtrHandler<FNDIHairStrandsData> InstData(Context);
+	VectorVM::FExternalFuncRegisterHandler<float> OutBendScaleX(Context);
+	VectorVM::FExternalFuncRegisterHandler<float> OutBendScaleY(Context);
+	VectorVM::FExternalFuncRegisterHandler<float> OutBendScaleZ(Context);
+	VectorVM::FExternalFuncRegisterHandler<float> OutBendScaleW(Context);
+
+	for (int32 InstanceIdx = 0; InstanceIdx < Context.NumInstances; ++InstanceIdx)
+	{
+		*OutBendScaleX.GetDestAndAdvance() = InstData->BendScale[0];
+		*OutBendScaleY.GetDestAndAdvance() = InstData->BendScale[1];
+		*OutBendScaleZ.GetDestAndAdvance() = InstData->BendScale[2];
+		*OutBendScaleW.GetDestAndAdvance() = InstData->BendScale[3];
+	}
+}
+
+void UNiagaraDataInterfaceHairStrands::GetSolveStretch(FVectorVMContext& Context)
+{
+	VectorVM::FUserPtrHandler<FNDIHairStrandsData> InstData(Context);
+	VectorVM::FExternalFuncRegisterHandler<int32> OutSolveStretch(Context);
+
+	for (int32 InstanceIdx = 0; InstanceIdx < Context.NumInstances; ++InstanceIdx)
+	{
+		*OutSolveStretch.GetDestAndAdvance() = InstData->SolveStretch;
+	}
+}
+
+void UNiagaraDataInterfaceHairStrands::GetProjectStretch(FVectorVMContext& Context)
+{
+	VectorVM::FUserPtrHandler<FNDIHairStrandsData> InstData(Context);
+	VectorVM::FExternalFuncRegisterHandler<int32> OutProjectStretch(Context);
+
+	for (int32 InstanceIdx = 0; InstanceIdx < Context.NumInstances; ++InstanceIdx)
+	{
+		*OutProjectStretch.GetDestAndAdvance() = InstData->ProjectStretch;
+	}
+}
+
+void UNiagaraDataInterfaceHairStrands::GetStretchDamping(FVectorVMContext& Context)
+{
+	VectorVM::FUserPtrHandler<FNDIHairStrandsData> InstData(Context);
+	VectorVM::FExternalFuncRegisterHandler<float> OutStretchDamping(Context);
+
+	for (int32 InstanceIdx = 0; InstanceIdx < Context.NumInstances; ++InstanceIdx)
+	{
+		*OutStretchDamping.GetDestAndAdvance() = InstData->StretchDamping;
+	}
+}
+
+void UNiagaraDataInterfaceHairStrands::GetStretchStiffness(FVectorVMContext& Context)
+{
+	VectorVM::FUserPtrHandler<FNDIHairStrandsData> InstData(Context);
+	VectorVM::FExternalFuncRegisterHandler<float> OutStretchStiffness(Context);
+
+	for (int32 InstanceIdx = 0; InstanceIdx < Context.NumInstances; ++InstanceIdx)
+	{
+		*OutStretchStiffness.GetDestAndAdvance() = InstData->StretchStiffness;
+	}
+}
+
+void UNiagaraDataInterfaceHairStrands::GetStretchScale(FVectorVMContext& Context)
+{
+	VectorVM::FUserPtrHandler<FNDIHairStrandsData> InstData(Context);
+	VectorVM::FExternalFuncRegisterHandler<float> OutStretchScaleX(Context);
+	VectorVM::FExternalFuncRegisterHandler<float> OutStretchScaleY(Context);
+	VectorVM::FExternalFuncRegisterHandler<float> OutStretchScaleZ(Context);
+	VectorVM::FExternalFuncRegisterHandler<float> OutStretchScaleW(Context);
+
+	for (int32 InstanceIdx = 0; InstanceIdx < Context.NumInstances; ++InstanceIdx)
+	{
+		*OutStretchScaleX.GetDestAndAdvance() = InstData->StretchScale[0];
+		*OutStretchScaleY.GetDestAndAdvance() = InstData->StretchScale[1];
+		*OutStretchScaleZ.GetDestAndAdvance() = InstData->StretchScale[2];
+		*OutStretchScaleW.GetDestAndAdvance() = InstData->StretchScale[3];
+	}
+}
+
+
+void UNiagaraDataInterfaceHairStrands::GetSolveCollision(FVectorVMContext& Context)
+{
+	VectorVM::FUserPtrHandler<FNDIHairStrandsData> InstData(Context);
+	VectorVM::FExternalFuncRegisterHandler<int32> OutSolveCollision(Context);
+
+	for (int32 InstanceIdx = 0; InstanceIdx < Context.NumInstances; ++InstanceIdx)
+	{
+		*OutSolveCollision.GetDestAndAdvance() = InstData->SolveCollision;
+	}
+}
+
+void UNiagaraDataInterfaceHairStrands::GetProjectCollision(FVectorVMContext& Context)
+{
+	VectorVM::FUserPtrHandler<FNDIHairStrandsData> InstData(Context);
+	VectorVM::FExternalFuncRegisterHandler<int32> OutProjectCollision(Context);
+
+	for (int32 InstanceIdx = 0; InstanceIdx < Context.NumInstances; ++InstanceIdx)
+	{
+		*OutProjectCollision.GetDestAndAdvance() = InstData->ProjectCollision;
+	}
+}
+
+void UNiagaraDataInterfaceHairStrands::GetStaticFriction(FVectorVMContext& Context)
+{
+	VectorVM::FUserPtrHandler<FNDIHairStrandsData> InstData(Context);
+	VectorVM::FExternalFuncRegisterHandler<float> OutStaticFriction(Context);
+
+	for (int32 InstanceIdx = 0; InstanceIdx < Context.NumInstances; ++InstanceIdx)
+	{
+		*OutStaticFriction.GetDestAndAdvance() = InstData->StaticFriction;
+	}
+}
+
+void UNiagaraDataInterfaceHairStrands::GetKineticFriction(FVectorVMContext& Context)
+{
+	VectorVM::FUserPtrHandler<FNDIHairStrandsData> InstData(Context);
+	VectorVM::FExternalFuncRegisterHandler<float> OutKineticFriction(Context);
+
+	for (int32 InstanceIdx = 0; InstanceIdx < Context.NumInstances; ++InstanceIdx)
+	{
+		*OutKineticFriction.GetDestAndAdvance() = InstData->KineticFriction;
+	}
+}
+
+void UNiagaraDataInterfaceHairStrands::GetStrandsViscosity(FVectorVMContext& Context)
+{
+	VectorVM::FUserPtrHandler<FNDIHairStrandsData> InstData(Context);
+	VectorVM::FExternalFuncRegisterHandler<float> OutStrandsViscosity(Context);
+
+	for (int32 InstanceIdx = 0; InstanceIdx < Context.NumInstances; ++InstanceIdx)
+	{
+		*OutStrandsViscosity.GetDestAndAdvance() = InstData->StrandsViscosity;
+	}
+}
+
+void UNiagaraDataInterfaceHairStrands::GetCollisionRadius(FVectorVMContext& Context)
+{
+	VectorVM::FUserPtrHandler<FNDIHairStrandsData> InstData(Context);
+	VectorVM::FExternalFuncRegisterHandler<float> OutCollisionRadius(Context);
+
+	for (int32 InstanceIdx = 0; InstanceIdx < Context.NumInstances; ++InstanceIdx)
+	{
+		*OutCollisionRadius.GetDestAndAdvance() = InstData->CollisionRadius;
+	}
+}
+
+void UNiagaraDataInterfaceHairStrands::GetRadiusScale(FVectorVMContext& Context)
+{
+	VectorVM::FUserPtrHandler<FNDIHairStrandsData> InstData(Context);
+	VectorVM::FExternalFuncRegisterHandler<float> OutRadiusScaleX(Context);
+	VectorVM::FExternalFuncRegisterHandler<float> OutRadiusScaleY(Context);
+	VectorVM::FExternalFuncRegisterHandler<float> OutRadiusScaleZ(Context);
+	VectorVM::FExternalFuncRegisterHandler<float> OutRadiusScaleW(Context);
+
+	for (int32 InstanceIdx = 0; InstanceIdx < Context.NumInstances; ++InstanceIdx)
+	{
+		*OutRadiusScaleX.GetDestAndAdvance() = InstData->RadiusScale[0];
+		*OutRadiusScaleY.GetDestAndAdvance() = InstData->RadiusScale[1];
+		*OutRadiusScaleZ.GetDestAndAdvance() = InstData->RadiusScale[2];
+		*OutRadiusScaleW.GetDestAndAdvance() = InstData->RadiusScale[3];
+	}
+}
+
+void UNiagaraDataInterfaceHairStrands::GetStrandsSmoothing(FVectorVMContext& Context)
+{
+	VectorVM::FUserPtrHandler<FNDIHairStrandsData> InstData(Context);
+	VectorVM::FExternalFuncRegisterHandler<float> OutStrandsSmoothing(Context);
+
+	for (int32 InstanceIdx = 0; InstanceIdx < Context.NumInstances; ++InstanceIdx)
+	{
+		*OutStrandsSmoothing.GetDestAndAdvance() = InstData->StrandsSmoothing;
+	}
+}
+
+void UNiagaraDataInterfaceHairStrands::GetStrandsDensity(FVectorVMContext& Context)
+{
+	VectorVM::FUserPtrHandler<FNDIHairStrandsData> InstData(Context);
+	VectorVM::FExternalFuncRegisterHandler<float> OutStrandsDensity(Context);
+
+	for (int32 InstanceIdx = 0; InstanceIdx < Context.NumInstances; ++InstanceIdx)
+	{
+		*OutStrandsDensity.GetDestAndAdvance() = InstData->StrandsDensity;
+	}
+}
+
+void UNiagaraDataInterfaceHairStrands::GetStrandsThickness(FVectorVMContext& Context)
+{
+	VectorVM::FUserPtrHandler<FNDIHairStrandsData> InstData(Context);
+	VectorVM::FExternalFuncRegisterHandler<float> OutStrandsThickness(Context);
+
+	for (int32 InstanceIdx = 0; InstanceIdx < Context.NumInstances; ++InstanceIdx)
+	{
+		*OutStrandsThickness.GetDestAndAdvance() = InstData->StrandsThickness;
+	}
+}
+
+void UNiagaraDataInterfaceHairStrands::GetThicknessScale(FVectorVMContext& Context)
+{
+	VectorVM::FUserPtrHandler<FNDIHairStrandsData> InstData(Context);
+	VectorVM::FExternalFuncRegisterHandler<float> OutThicknessScaleX(Context);
+	VectorVM::FExternalFuncRegisterHandler<float> OutThicknessScaleY(Context);
+	VectorVM::FExternalFuncRegisterHandler<float> OutThicknessScaleZ(Context);
+	VectorVM::FExternalFuncRegisterHandler<float> OutThicknessScaleW(Context);
+
+	for (int32 InstanceIdx = 0; InstanceIdx < Context.NumInstances; ++InstanceIdx)
+	{
+		*OutThicknessScaleX.GetDestAndAdvance() = InstData->ThicknessScale[0];
+		*OutThicknessScaleY.GetDestAndAdvance() = InstData->ThicknessScale[1];
+		*OutThicknessScaleZ.GetDestAndAdvance() = InstData->ThicknessScale[2];
+		*OutThicknessScaleW.GetDestAndAdvance() = InstData->ThicknessScale[3];
 	}
 }
 
@@ -2137,9 +2945,6 @@ void UNiagaraDataInterfaceHairStrands::GetWorldInverse(FVectorVMContext& Context
 	WriteTransform(WorldInverse, Context);
 }
 
-void UNiagaraDataInterfaceHairStrands::BuildBoundingBox(FVectorVMContext& Context)
-{}
-
 void UNiagaraDataInterfaceHairStrands::GetBoxCenter(FVectorVMContext& Context)
 {
 	VectorVM::FUserPtrHandler<FNDIHairStrandsData> InstData(Context);
@@ -2170,11 +2975,15 @@ void UNiagaraDataInterfaceHairStrands::GetBoxExtent(FVectorVMContext& Context)
 	}
 }
 
+
 void UNiagaraDataInterfaceHairStrands::GetPointPosition(FVectorVMContext& Context)
 {
 	// @todo : implement function for cpu 
 }
 
+void UNiagaraDataInterfaceHairStrands::BuildBoundingBox(FVectorVMContext& Context)
+{}
+
 void UNiagaraDataInterfaceHairStrands::ComputeNodePosition(FVectorVMContext& Context)
 {
 	// @todo : implement function for cpu 
@@ -2331,17 +3140,22 @@ void UNiagaraDataInterfaceHairStrands::ComputeRestDirection(FVectorVMContext& Co
 	// @todo : implement function for cpu 
 }
 
-void UNiagaraDataInterfaceHairStrands::UpdateNodeOrientation(FVectorVMContext& Context)
+void UNiagaraDataInterfaceHairStrands::UpdateMaterialFrame(FVectorVMContext& Context)
+{
+	// @todo : implement function for cpu 
+}
+
+void UNiagaraDataInterfaceHairStrands::ComputeMaterialFrame(FVectorVMContext& Context)
 {
 	// @todo : implement function for cpu 
 }
 
-void UNiagaraDataInterfaceHairStrands::SolveStaticCollisionConstraint(FVectorVMContext& Context)
+void UNiagaraDataInterfaceHairStrands::SolveHardCollisionConstraint(FVectorVMContext& Context)
 {
 	// @todo : implement function for cpu 
 }
 
-void UNiagaraDataInterfaceHairStrands::ProjectStaticCollisionConstraint(FVectorVMContext& Context)
+void UNiagaraDataInterfaceHairStrands::ProjectHardCollisionConstraint(FVectorVMContext& Context)
 {
 	// @todo : implement function for cpu 
 }
@@ -2707,8 +3521,7 @@ in float RestLength, in float DeltaTime, in int NodeOffset, in float MaterialDam
 					}
 					else if( NodeOffset == 2)
 					{
-						//SetupTwistSpringMaterial(DIContext.StrandSize,YoungModulus,RodThickness,RestLength,DeltaTime,false,MaterialDamping,OutMaterialCompliance,OutMaterialWeight,OutMaterialMultiplier);
-						SetupVolumeSpringMaterial(DIContext.StrandSize,YoungModulus,RestLength,DeltaTime,false,MaterialDamping,OutMaterialCompliance,OutMaterialWeight,OutMaterialMultiplier);
+						SetupTwistSpringMaterial(DIContext.StrandSize,YoungModulus,RodThickness,RestLength,DeltaTime,false,MaterialDamping,OutMaterialCompliance,OutMaterialWeight,OutMaterialMultiplier);
 					}
 				}
 				)");
@@ -2732,8 +3545,7 @@ in float RestLength, in float DeltaTime, in int NodeOffset, in float MaterialDam
 					}
 					else if(NodeOffset == 2)
 					{
-						//SolveTwistSpringMaterial(EnableConstraint,DIContext.StrandSize,RestLength,DeltaTime,MaterialDamping,MaterialCompliance,MaterialWeight,MaterialMultiplier,OutMaterialMultiplier);
-						SolveVolumeSpringMaterial(EnableConstraint,DIContext.StrandSize,RestLength,DeltaTime,MaterialDamping,MaterialCompliance,MaterialWeight,MaterialMultiplier,OutMaterialMultiplier);
+						SolveTwistSpringMaterial(EnableConstraint,DIContext.StrandSize,RestLength,DeltaTime,MaterialDamping,MaterialCompliance,MaterialWeight,MaterialMultiplier,OutMaterialMultiplier);
 					}
 				}
 				)");
@@ -2743,20 +3555,20 @@ in float RestLength, in float DeltaTime, in int NodeOffset, in float MaterialDam
 	else if (FunctionInfo.DefinitionName == ProjectDistanceSpringMaterialName)
 	{
 		static const TCHAR *FormatSample = TEXT(R"(
-				void {InstanceFunctionName} (in float YoungModulus, in float RodThickness, in float RestLength, in float NodeMass, in float DeltaTime, in int NodeOffset, in float3 NodePosition, out float3 OutNodePosition)
+				void {InstanceFunctionName} (in bool EnableConstraint, in float YoungModulus, in float RodThickness, in float RestLength, in float DeltaTime, in int NodeOffset, out float3 OutNodePosition)
 				{
 					{HairStrandsContextName} 
 					if(NodeOffset == 0)
 					{
-						ProjectStretchSpringMaterial(DIContext.StrandSize,YoungModulus,RodThickness,RestLength,NodeMass,DeltaTime,NodePosition,OutNodePosition);
+						ProjectStretchSpringMaterial(EnableConstraint,DIContext.StrandSize,YoungModulus,RodThickness,RestLength,DeltaTime,OutNodePosition);
 					}
 					if(NodeOffset == 1)
 					{
-						ProjectBendSpringMaterial(DIContext.StrandSize,YoungModulus,RodThickness,RestLength,NodeMass,DeltaTime,NodePosition,OutNodePosition);
+						ProjectBendSpringMaterial(EnableConstraint,DIContext.StrandSize,YoungModulus,RodThickness,RestLength,DeltaTime,OutNodePosition);
 					}
 					if(NodeOffset == 2)
 					{
-						ProjectTwistSpringMaterial(DIContext.StrandSize,YoungModulus,RodThickness,RestLength,NodeMass,DeltaTime,NodePosition,OutNodePosition);
+						ProjectTwistSpringMaterial(EnableConstraint,DIContext.StrandSize,YoungModulus,RodThickness,RestLength,DeltaTime,OutNodePosition);
 					}
 				}
 				)");
@@ -2790,9 +3602,9 @@ in float RestLength, in float DeltaTime, in int NodeOffset, in float MaterialDam
 	else if (FunctionInfo.DefinitionName == ProjectAngularSpringMaterialName)
 	{
 		static const TCHAR *FormatSample = TEXT(R"(
-					void {InstanceFunctionName} (in float YoungModulus, in float RodThickness, in float RestLength, in float3 RestDirection, in float NodeMass, in float DeltaTime, in float3 NodePosition, out float3 OutNodePosition)
+					void {InstanceFunctionName} (in bool EnableConstraint, in float YoungModulus, in float RodThickness, in float RestLength, in float3 RestDirection, in float DeltaTime, out float3 OutNodePosition)
 					{
-						{HairStrandsContextName} ProjectAngularSpringMaterial(DIContext.StrandSize,YoungModulus,RodThickness,RestLength,RestDirection,NodeMass,DeltaTime,NodePosition,OutNodePosition);
+						{HairStrandsContextName} ProjectAngularSpringMaterial(EnableConstraint,DIContext.StrandSize,YoungModulus,RodThickness,RestLength,RestDirection,DeltaTime,OutNodePosition);
 					}
 					)");
 		OutHLSL += FString::Format(FormatSample, ArgsSample);
@@ -2813,10 +3625,10 @@ in float RestLength, in float DeltaTime, in int NodeOffset, in float MaterialDam
 	else if (FunctionInfo.DefinitionName == SolveStretchRodMaterialName)
 	{
 		static const TCHAR *FormatSample = TEXT(R"(
-					void {InstanceFunctionName} (in float RestLength, in float DeltaTime, in float MaterialDamping, 
+					void {InstanceFunctionName} (in bool EnableConstraint, in float RestLength, in float DeltaTime, in float MaterialDamping, 
 			in float MaterialCompliance, in float MaterialWeight, in float3 MaterialMultiplier, out float3 OutMaterialMultiplier)
 					{
-						{HairStrandsContextName} SolveStretchRodMaterial(DIContext.StrandSize,RestLength,DeltaTime,MaterialDamping,MaterialCompliance,MaterialWeight,MaterialMultiplier,OutMaterialMultiplier);
+						{HairStrandsContextName} SolveStretchRodMaterial(EnableConstraint,DIContext.StrandSize,RestLength,DeltaTime,MaterialDamping,MaterialCompliance,MaterialWeight,MaterialMultiplier,OutMaterialMultiplier);
 					}
 					)");
 		OutHLSL += FString::Format(FormatSample, ArgsSample);
@@ -2825,9 +3637,9 @@ in float RestLength, in float DeltaTime, in int NodeOffset, in float MaterialDam
 	else if (FunctionInfo.DefinitionName == ProjectStretchRodMaterialName)
 	{
 		static const TCHAR *FormatSample = TEXT(R"(
-					void {InstanceFunctionName} (in float YoungModulus, in float RodThickness, in float RestLength, in float DeltaTime, out float3 OutMaterialMultiplier)
+					void {InstanceFunctionName} (in bool EnableConstraint, in float YoungModulus, in float RodThickness, in float RestLength, in float DeltaTime, out float3 OutNodePosition)
 					{
-						{HairStrandsContextName} ProjectStretchRodMaterial(DIContext.StrandSize,YoungModulus,RodThickness,RestLength,DeltaTime,OutMaterialMultiplier);
+						{HairStrandsContextName} ProjectStretchRodMaterial(EnableConstraint,DIContext.StrandSize,YoungModulus,RodThickness,RestLength,DeltaTime,OutNodePosition);
 					}
 					)");
 		OutHLSL += FString::Format(FormatSample, ArgsSample);
@@ -2848,10 +3660,10 @@ in float RestLength, in float DeltaTime, in int NodeOffset, in float MaterialDam
 	else if (FunctionInfo.DefinitionName == SolveBendRodMaterialName)
 	{
 		static const TCHAR *FormatSample = TEXT(R"(
-						void {InstanceFunctionName} (in float RestLength, in float4 RestDarboux, in float DeltaTime, in float MaterialDamping, 
+						void {InstanceFunctionName} (in bool EnableConstraint, in float RestLength, in float4 RestDarboux, in float DeltaTime, in float MaterialDamping, 
 				in float MaterialCompliance, in float MaterialWeight, in float3 MaterialMultiplier, out float3 OutMaterialMultiplier)
 						{
-							{HairStrandsContextName} SolveBendRodMaterial(DIContext.StrandSize,RestLength,RestDarboux,DeltaTime,MaterialDamping,MaterialCompliance,MaterialWeight,MaterialMultiplier,OutMaterialMultiplier);
+							{HairStrandsContextName} SolveBendRodMaterial(EnableConstraint,DIContext.StrandSize,RestLength,RestDarboux,DeltaTime,MaterialDamping,MaterialCompliance,MaterialWeight,MaterialMultiplier,OutMaterialMultiplier);
 						}
 						)");
 		OutHLSL += FString::Format(FormatSample, ArgsSample);
@@ -2860,37 +3672,36 @@ in float RestLength, in float DeltaTime, in int NodeOffset, in float MaterialDam
 	else if (FunctionInfo.DefinitionName == ProjectBendRodMaterialName)
 	{
 		static const TCHAR *FormatSample = TEXT(R"(
-						void {InstanceFunctionName} (in float YoungModulus, in float RodThickness, in float RestLength, in float4 RestDarboux, in float DeltaTime, out float3 OutMaterialMultiplier)
+						void {InstanceFunctionName} (in bool EnableConstraint, in float YoungModulus, in float RodThickness, in float RestLength, in float4 RestDarboux, in float DeltaTime, out float3 OutNodePosition)
 						{
-							{HairStrandsContextName} ProjectBendRodMaterial(DIContext.StrandSize,YoungModulus,RodThickness,RestLength,RestDarboux,DeltaTime,OutMaterialMultiplier);
+							{HairStrandsContextName} ProjectBendRodMaterial(EnableConstraint,DIContext.StrandSize,YoungModulus,RodThickness,RestLength,RestDarboux,DeltaTime,OutNodePosition);
 						}
 						)");
 		OutHLSL += FString::Format(FormatSample, ArgsSample);
 		return true;
 	}
-	else if (FunctionInfo.DefinitionName == SolveStaticCollisionConstraintName)
+	else if (FunctionInfo.DefinitionName == SolveHardCollisionConstraintName)
 	{
 		static const TCHAR *FormatSample = TEXT(R"(
 							void {InstanceFunctionName} (in bool EnableConstraint, in float PenetrationDepth, in float3 CollisionPosition, in float3 CollisionVelocity, in float3 CollisionNormal, 
-				in float StaticFriction, in float KineticFriction, in float DeltaTime, in float3 PreviousPosition, out float3 OutMaterialMultiplier )
+				in float StaticFriction, in float KineticFriction, in float DeltaTime, out float3 OutMaterialMultiplier )
 							{
 								OutMaterialMultiplier = float3(0,0,0);
-								{HairStrandsContextName} SolveStaticCollisionConstraint(EnableConstraint,DIContext.StrandSize,PenetrationDepth,
-									CollisionPosition,CollisionVelocity,CollisionNormal,StaticFriction,KineticFriction,DeltaTime,false,PreviousPosition,SharedNodePosition[GGroupThreadId.x]);
+								{HairStrandsContextName} SolveHardCollisionConstraint(EnableConstraint,DIContext.StrandSize,PenetrationDepth,
+									CollisionPosition,CollisionVelocity,CollisionNormal,StaticFriction,KineticFriction,false,DeltaTime);
 							}
 							)");
 		OutHLSL += FString::Format(FormatSample, ArgsSample);
 		return true;
 	}
-	else if (FunctionInfo.DefinitionName == ProjectStaticCollisionConstraintName)
+	else if (FunctionInfo.DefinitionName == ProjectHardCollisionConstraintName)
 	{
 			static const TCHAR *FormatSample = TEXT(R"(
-						void {InstanceFunctionName} (in float PenetrationDepth, in float3 CollisionPosition, in float3 CollisionVelocity, in float3 CollisionNormal, 
-			in float StaticFriction, in float KineticFriction, in float DeltaTime, in float3 PreviousPosition, in float3 NodePosition, out float3 OutNodePosition )
+						void {InstanceFunctionName} (in bool EnableConstraint, in float PenetrationDepth, in float3 CollisionPosition, in float3 CollisionVelocity, in float3 CollisionNormal, 
+			in float StaticFriction, in float KineticFriction, in float DeltaTime, out float3 OutNodePosition )
 						{
-							OutNodePosition = NodePosition;
-							{HairStrandsContextName} SolveStaticCollisionConstraint(true,DIContext.StrandSize,PenetrationDepth,
-								CollisionPosition,CollisionVelocity,CollisionNormal,StaticFriction,KineticFriction,DeltaTime,true,PreviousPosition,OutNodePosition);
+							{HairStrandsContextName} ProjectHardCollisionConstraint(EnableConstraint,DIContext.StrandSize,PenetrationDepth,
+								CollisionPosition,CollisionVelocity,CollisionNormal,StaticFriction,KineticFriction,DeltaTime,OutNodePosition);
 						}
 						)");
 		OutHLSL += FString::Format(FormatSample, ArgsSample);
@@ -2904,7 +3715,7 @@ in float RestLength, in float DeltaTime, in int NodeOffset, in float MaterialDam
 			in float MaterialCompliance, in float MaterialWeight, in float3 MaterialMultiplier, out float3 OutMaterialMultiplier )
 								{
 									{HairStrandsContextName} SolveSoftCollisionConstraint(EnableConstraint,DIContext.StrandSize,PenetrationDepth,
-										CollisionPosition,CollisionVelocity,CollisionNormal,StaticFriction,KineticFriction,DeltaTime,MaterialDamping,
+										CollisionPosition,CollisionVelocity,CollisionNormal,StaticFriction,KineticFriction,false,DeltaTime,MaterialDamping,
 											MaterialCompliance,MaterialWeight,MaterialMultiplier,OutMaterialMultiplier);
 								}
 								)");
@@ -2914,11 +3725,11 @@ in float RestLength, in float DeltaTime, in int NodeOffset, in float MaterialDam
 	else if (FunctionInfo.DefinitionName == ProjectSoftCollisionConstraintName)
 	{
 		static const TCHAR *FormatSample = TEXT(R"(
-							void {InstanceFunctionName} (in float ConstraintStiffness, in float PenetrationDepth, in float3 CollisionPosition, in float3 CollisionVelocity, in float3 CollisionNormal, 
-					in float StaticFriction, in float KineticFriction, in float NodeMass, in float DeltaTime, in float3 PreviousPosition, in float3 NodePosition, out float3 OutNodePosition  )
+							void {InstanceFunctionName} (in bool EnableConstraint, in float ConstraintStiffness, in float PenetrationDepth, in float3 CollisionPosition, in float3 CollisionVelocity, in float3 CollisionNormal, 
+					in float StaticFriction, in float KineticFriction, in float DeltaTime, in float NodeMass, out float3 OutNodePosition )
 							{
-								{HairStrandsContextName} ProjectSoftCollisionConstraint(DIContext.StrandSize,ConstraintStiffness,PenetrationDepth,
-									CollisionPosition,CollisionVelocity,CollisionNormal,StaticFriction,KineticFriction,NodeMass,DeltaTime,PreviousPosition,NodePosition,OutNodePosition);
+								{HairStrandsContextName} ProjectSoftCollisionConstraint(EnableConstraint,DIContext.StrandSize,ConstraintStiffness,PenetrationDepth,
+									CollisionPosition,CollisionVelocity,CollisionNormal,StaticFriction,KineticFriction,DeltaTime,NodeMass,OutNodePosition);
 							}
 							)");
 		OutHLSL += FString::Format(FormatSample, ArgsSample);
@@ -2946,7 +3757,7 @@ in float RestLength, in float DeltaTime, in int NodeOffset, in float MaterialDam
 		OutHLSL += FString::Format(FormatSample, ArgsSample);
 		return true;
 	}
-	else if (FunctionInfo.DefinitionName == UpdateNodeOrientationName)
+	else if (FunctionInfo.DefinitionName == UpdateMaterialFrameName)
 	{
 		static const TCHAR *FormatSample = TEXT(R"(
 					void {InstanceFunctionName} ( out bool OutUpdateStatus)
@@ -2958,6 +3769,18 @@ in float RestLength, in float DeltaTime, in int NodeOffset, in float MaterialDam
 		OutHLSL += FString::Format(FormatSample, ArgsSample);
 		return true;
 	}
+	else if (FunctionInfo.DefinitionName == ComputeMaterialFrameName)
+	{
+		static const TCHAR *FormatSample = TEXT(R"(
+						void {InstanceFunctionName} ( out bool OutUpdateStatus)
+						{
+							{HairStrandsContextName} ComputeMaterialFrame(DIContext.StrandSize);
+							OutUpdateStatus = true;
+						}
+						)");
+		OutHLSL += FString::Format(FormatSample, ArgsSample);
+		return true;
+	}
 	else if (FunctionInfo.DefinitionName == ComputeAirDragForceName)
 	{
 		static const TCHAR *FormatSample = TEXT(R"(
@@ -2991,6 +3814,8 @@ void UNiagaraDataInterfaceHairStrands::GetCommonHLSL(FString& OutHLSL)
 	OutHLSL += TEXT("#include \"/Plugin/Experimental/HairStrands/Private/NiagaraQuaternionUtils.ush\"\n");
 	OutHLSL += TEXT("#include \"/Plugin/Experimental/HairStrands/Private/NiagaraStrandsExternalForce.ush\"\n");
 	OutHLSL += TEXT("#include \"/Plugin/Experimental/HairStrands/Private/NiagaraHookeSpringMaterial.ush\"\n");
+	OutHLSL += TEXT("#include \"/Plugin/Experimental/HairStrands/Private/NiagaraAngularSpringMaterial.ush\"\n");
+	OutHLSL += TEXT("#include \"/Plugin/Experimental/HairStrands/Private/NiagaraConstantVolumeMaterial.ush\"\n");
 	OutHLSL += TEXT("#include \"/Plugin/Experimental/HairStrands/Private/NiagaraCosseratRodMaterial.ush\"\n");
 	OutHLSL += TEXT("#include \"/Plugin/Experimental/HairStrands/Private/NiagaraStaticCollisionConstraint.ush\"\n");
 	OutHLSL += TEXT("#include \"/Plugin/Experimental/HairStrands/Private/NiagaraDataInterfaceHairStrands.ush\"\n");
@@ -3008,15 +3833,10 @@ void UNiagaraDataInterfaceHairStrands::ProvidePerInstanceDataForRenderThread(voi
 	FNDIHairStrandsData* GameThreadData = static_cast<FNDIHairStrandsData*>(PerInstanceData);
 	FNDIHairStrandsData* RenderThreadData = static_cast<FNDIHairStrandsData*>(DataForRenderThread);
 
-	RenderThreadData->WorldTransform = GameThreadData->WorldTransform;
-	RenderThreadData->HairStrandsBuffer = GameThreadData->HairStrandsBuffer;
-	RenderThreadData->NumStrands = GameThreadData->NumStrands;
-	RenderThreadData->StrandSize = GameThreadData->StrandSize;
-	RenderThreadData->BoxCenter = GameThreadData->BoxCenter;
-	RenderThreadData->BoxExtent = GameThreadData->BoxExtent;
-	RenderThreadData->TickCount = GameThreadData->TickCount;
-	RenderThreadData->ResetTick = GameThreadData->ResetTick;
-	RenderThreadData->ForceReset = GameThreadData->ForceReset;
+	if (GameThreadData && RenderThreadData)
+	{
+		RenderThreadData->CopyDatas(GameThreadData);
+	}
 }
 
 FNiagaraDataInterfaceParametersCS*
@@ -3028,7 +3848,6 @@ UNiagaraDataInterfaceHairStrands::ConstructComputeParameters() const
 void FNDIHairStrandsProxy::PreStage(FRHICommandList& RHICmdList, const FNiagaraDataInterfaceSetArgs& Context)
 {
 }
-	
 
 void FNDIHairStrandsProxy::PostStage(FRHICommandList& RHICmdList, const FNiagaraDataInterfaceSetArgs& Context)
 {
diff --git a/Engine/Plugins/Experimental/HairStrands/Source/HairStrandsNiagara/Private/NiagaraDataInterfacePhysicsAsset.cpp b/Engine/Plugins/Experimental/HairStrands/Source/HairStrandsNiagara/Private/NiagaraDataInterfacePhysicsAsset.cpp
--- a/Engine/Plugins/Experimental/HairStrands/Source/HairStrandsNiagara/Private/NiagaraDataInterfacePhysicsAsset.cpp
+++ b/Engine/Plugins/Experimental/HairStrands/Source/HairStrandsNiagara/Private/NiagaraDataInterfacePhysicsAsset.cpp
@@ -3,6 +3,7 @@
 #include "NiagaraDataInterfacePhysicsAsset.h"
 #include "Animation/SkeletalMeshActor.h"
 #include "Components/SkeletalMeshComponent.h"
+#include "SkeletalRenderPublic.h"
 #include "SkeletalMeshTypes.h"
 #include "AnimationRuntime.h"
 #include "NiagaraShader.h"
@@ -305,6 +306,28 @@ void FNDIPhysicsAssetBuffer::Update()
 	}
 }
 
+struct FPhysicsAssetManager
+{
+	TMap<int32,FSkeletalMeshObject*> Elements;
+};
+
+FPhysicsAssetManager GHairManager;
+
+void MeshObjectCallback(
+	FSkeletalMeshObjectCallbackData::EEventType Event,
+	FSkeletalMeshObject* MeshObject,
+	uint64 UserData)
+{
+	ENQUEUE_RENDER_COMMAND(FPhysicsAssetUpdate)(
+		[Event, MeshObject, UserData](FRHICommandListImmediate& RHICmdList)
+	{
+		if (Event == FSkeletalMeshObjectCallbackData::EEventType::Register || Event == FSkeletalMeshObjectCallbackData::EEventType::Update)
+		{
+			GHairManager.Elements.Add(UserData,MeshObject);
+		}
+	});
+}
+
 void FNDIPhysicsAssetBuffer::InitRHI()
 {
 	if (IsValid())
@@ -320,6 +343,25 @@ void FNDIPhysicsAssetBuffer::InitRHI()
 		//UE_LOG(LogPhysicsAsset, Warning, TEXT("Num Capsules = %d | Num Spheres = %d | Num Boxes = %d"), AssetArrays->ElementOffsets.NumElements - AssetArrays->ElementOffsets.CapsuleOffset,
 		//	AssetArrays->ElementOffsets.CapsuleOffset - AssetArrays->ElementOffsets.SphereOffset, AssetArrays->ElementOffsets.SphereOffset - AssetArrays->ElementOffsets.BoxOffset);
 	}
+
+	// /*const FPrimitiveComponentId LocalComponentId = SkeletalMesh->ComponentId;
+	//EWorldType::Type WorldType = SkeletalMesh->GetWorld() ? EWorldType::Type(SkeletalMesh->GetWorld()->WorldType) : EWorldType::None;
+	//WorldType = WorldType == EWorldType::Inactive ? EWorldType::Editor : WorldType;*/
+
+	//// Always setup the callback, even if the SkeletalMeshData are not ready yet
+	//if (SkeletalMesh.IsValid() && SkeletalMesh.Get() != nullptr)
+	//{
+	//	/*FSkeletalMeshObjectCallbackData CallbackData;
+	//	CallbackData.Run = MeshObjectCallback;
+	//	CallbackData.UserData = (uint64(LocalComponentId.PrimIDValue) & 0xFFFFFFFF) | (uint64(WorldType) << 32);
+	//	SkeletalMesh->MeshObjectCallbackData = CallbackData;*/
+
+	//	if (SkeletalMesh->MeshObject != nullptr)
+	//	{
+	//		UE_LOG(LogPhysicsAsset, Log, TEXT("Skin Cache there : %d %d %d"), SkeletalMesh->MeshObject, SkeletalMesh->MeshObject->GetCachedGeometry().LODIndex, 
+	//				SkeletalMesh->MeshObject->GetCachedGeometry().Sections.Num());
+	//	}
+	//}
 }
 
 void FNDIPhysicsAssetBuffer::ReleaseRHI()
@@ -535,6 +577,28 @@ void FNDIPhysicsAssetProxy::DestroyPerInstanceData(NiagaraEmitterInstanceBatcher
 	SystemInstancesToProxyData.Remove(SystemInstance);
 }
 
+void FNDIPhysicsAssetProxy::PreStage(FRHICommandList& RHICmdList, const FNiagaraDataInterfaceSetArgs& Context)
+{
+	if (Context.ShaderStageIndex == 0)
+	{
+		FNDIPhysicsAssetData* ProxyData =
+			SystemInstancesToProxyData.Find(Context.SystemInstance);
+
+		if (ProxyData != nullptr)
+		{
+
+		}
+	}
+}
+
+void FNDIPhysicsAssetProxy::PostStage(FRHICommandList& RHICmdList, const FNiagaraDataInterfaceSetArgs& Context)
+{
+}
+
+void FNDIPhysicsAssetProxy::ResetData(FRHICommandList& RHICmdList, const FNiagaraDataInterfaceSetArgs& Context)
+{
+}
+
 //------------------------------------------------------------------------------------------------------------
 
 UNiagaraDataInterfacePhysicsAsset::UNiagaraDataInterfacePhysicsAsset(FObjectInitializer const& ObjectInitializer)
@@ -946,5 +1010,70 @@ UNiagaraDataInterfacePhysicsAsset::ConstructComputeParameters() const
 	return new FNDIPhysicsAssetParametersCS();
 }
 
+//------------------------------------------------------------------------------------------------------------
+//
+//class FResetCS : public FGlobalShader
+//{
+//	DECLARE_GLOBAL_SHADER(FCopyBoundingBoxCS);
+//	SHADER_USE_PARAMETER_STRUCT(FCopyBoundingBoxCS, FGlobalShader);
+//
+//	BEGIN_SHADER_PARAMETER_STRUCT(FParameters, )
+//		SHADER_PARAMETER(uint32, NumElements)
+//		SHADER_PARAMETER_UAV(RWBuffer, BoundingBoxBuffer)
+//		SHADER_PARAMETER_UAV(RWBuffer, OutNodeBoundBuffer)
+//		END_SHADER_PARAMETER_STRUCT()
+//
+//public:
+//	static bool ShouldCompilePermutation(const FGlobalShaderPermutationParameters& Parameters)
+//	{
+//		return RHISupportsComputeShaders(Parameters.Platform);
+//	}
+//
+//	static void ModifyCompilationEnvironment(const FGlobalShaderPermutationParameters& Parameters, FShaderCompilerEnvironment& OutEnvironment)
+//	{
+//		FGlobalShader::ModifyCompilationEnvironment(Parameters, OutEnvironment);
+//		OutEnvironment.SetDefine(TEXT("THREAD_COUNT"), NIAGARA_HAIR_STRANDS_THREAD_COUNT);
+//	}
+//};
+//
+//IMPLEMENT_GLOBAL_SHADER(FCopyBoundingBoxCS, "/Plugin/Experimental/HairStrands/Private/NiagaraCopyBoundingBox.usf", "MainCS", SF_Compute);
+//
+//static void AddSkinCacheRasterPass(
+//	FRDGBuilder& GraphBuilder,
+//	FRHIUnorderedAccessView* BoundingBoxBuffer,
+//	FRHIUnorderedAccessView* OutNodeBoundBuffer)
+//{
+//	const uint32 GroupSize = NIAGARA_HAIR_STRANDS_THREAD_COUNT;
+//	const uint32 NumElements = 1;
+//
+//	FCopyBoundingBoxCS::FParameters* Parameters = GraphBuilder.AllocParameters<FCopyBoundingBoxCS::FParameters>();
+//	Parameters->BoundingBoxBuffer = BoundingBoxBuffer;
+//	Parameters->OutNodeBoundBuffer = OutNodeBoundBuffer;
+//	Parameters->NumElements = NumElements;
+//
+//	TShaderMap<FGlobalShaderType>* ShaderMap = GetGlobalShaderMap(ERHIFeatureLevel::SM5);
+//
+//	const uint32 DispatchCount = FMath::DivideAndRoundUp(NumElements, GroupSize);
+//
+//	TShaderMapRef<FCopyBoundingBoxCS> ComputeShader(ShaderMap);
+//	FComputeShaderUtils::AddPass(
+//		GraphBuilder,
+//		RDG_EVENT_NAME("CopyBoundingBox"),
+//		*ComputeShader,
+//		Parameters,
+//		FIntVector(DispatchCount, 1, 1));
+//
+//	GraphBuilder.AddPass(
+//		RDG_EVENT_NAME("SkinCacheRaster"),
+//		PassParameters,
+//		ERDGPassFlags::Raster,
+//		[PassParameters, Scene = Scene, ViewInfo, RasterPassType, &PrimitiveSceneInfos, ViewportRect, HairRenderInfo, RasterDirection](FRHICommandListImmediate& RHICmdList)
+//	{
+//	});
+//}
+
+//------------------------------------------------------------------------------------------------------------
+
+
 
 #undef LOCTEXT_NAMESPACE
\ No newline at end of file
