Add debug context to DDC requests that take the key as a parameter

Fix new deprecation warnings after merging the API change from another stream.

#rnx

diff --git a/Engine/Plugins/Experimental/HairStrands/Source/HairStrandsCore/Private/GroomAsset.cpp b/Engine/Plugins/Experimental/HairStrands/Source/HairStrandsCore/Private/GroomAsset.cpp
--- a/Engine/Plugins/Experimental/HairStrands/Source/HairStrandsCore/Private/GroomAsset.cpp
+++ b/Engine/Plugins/Experimental/HairStrands/Source/HairStrandsCore/Private/GroomAsset.cpp
@@ -1071,7 +1071,7 @@ bool UGroomAsset::CacheDerivedData(const FGroomBuildSettings* InBuildSettings)
 	const FString DerivedDataKey = GroomDerivedDataCacheUtils::BuildGroomDerivedDataKey(KeySuffix);
 
 	TArray<uint8> DerivedData;
-	if (GetDerivedDataCacheRef().GetSynchronous(*DerivedDataKey, DerivedData))
+	if (GetDerivedDataCacheRef().GetSynchronous(*DerivedDataKey, DerivedData, GetPathName()))
 	{
 		FMemoryReader Ar(DerivedData, /*bIsPersistent=*/ true);
 
@@ -1131,7 +1131,7 @@ bool UGroomAsset::CacheDerivedData(const FGroomBuildSettings* InBuildSettings)
 				CompressedArchive << UncompressedSize; // needed for allocating decompression buffer
 				CompressedArchive.SerializeCompressed(LargeMemWriter.GetData(), UncompressedSize, NAME_Zlib);
 
-				GetDerivedDataCacheRef().Put(*DerivedDataKey, DerivedData);
+				GetDerivedDataCacheRef().Put(*DerivedDataKey, DerivedData, GetPathName());
 			}
 		}
 	}
