<Opened by P4VS>
Initial Checkin
#rnx

diff --git a/Engine/Plugins/Experimental/HairStrands/Source/HairStrandsCore/Private/GroomAsset.cpp b/Engine/Plugins/Experimental/HairStrands/Source/HairStrandsCore/Private/GroomAsset.cpp
--- a/Engine/Plugins/Experimental/HairStrands/Source/HairStrandsCore/Private/GroomAsset.cpp
+++ b/Engine/Plugins/Experimental/HairStrands/Source/HairStrandsCore/Private/GroomAsset.cpp
@@ -397,3 +397,47 @@ FArchive& operator<<(FArchive& Ar, FHairGroupData& GroupData)
 
 	return Ar;
 }
+
+void UGroomAsset::AddAssetUserData(UAssetUserData* InUserData)
+{
+	if (InUserData != nullptr)
+	{
+		UAssetUserData* ExistingData = GetAssetUserDataOfClass(InUserData->GetClass());
+		if (ExistingData != nullptr)
+		{
+			AssetUserData.Remove(ExistingData);
+		}
+		AssetUserData.Add(InUserData);
+	}
+}
+
+UAssetUserData* UGroomAsset::GetAssetUserDataOfClass(TSubclassOf<UAssetUserData> InUserDataClass)
+{
+	for (int32 DataIdx = 0; DataIdx < AssetUserData.Num(); DataIdx++)
+	{
+		UAssetUserData* Datum = AssetUserData[DataIdx];
+		if (Datum != nullptr && Datum->IsA(InUserDataClass))
+		{
+			return Datum;
+		}
+	}
+	return NULL;
+}
+
+void UGroomAsset::RemoveUserDataOfClass(TSubclassOf<UAssetUserData> InUserDataClass)
+{
+	for (int32 DataIdx = 0; DataIdx < AssetUserData.Num(); DataIdx++)
+	{
+		UAssetUserData* Datum = AssetUserData[DataIdx];
+		if (Datum != nullptr && Datum->IsA(InUserDataClass))
+		{
+			AssetUserData.RemoveAt(DataIdx);
+			return;
+		}
+	}
+}
+
+const TArray<UAssetUserData*>* UGroomAsset::GetAssetUserDataArray() const
+{
+	return &AssetUserData;
+}
\ No newline at end of file
diff --git a/Engine/Plugins/Experimental/HairStrands/Source/HairStrandsCore/Public/GroomAsset.h b/Engine/Plugins/Experimental/HairStrands/Source/HairStrandsCore/Public/GroomAsset.h
--- a/Engine/Plugins/Experimental/HairStrands/Source/HairStrandsCore/Public/GroomAsset.h
+++ b/Engine/Plugins/Experimental/HairStrands/Source/HairStrandsCore/Public/GroomAsset.h
@@ -8,6 +8,7 @@
 #include "HairDescription.h"
 #include "HairStrandsDatas.h"
 #include "RenderResource.h"
+#include "Interfaces/Interface_AssetUserData.h"
 
 #include "GroomAsset.generated.h"
 
@@ -223,7 +224,7 @@ struct HAIRSTRANDSCORE_API FHairGroupData
  * Implements an asset that can be used to store hair strands
  */
 UCLASS(BlueprintType, hidecategories = (Object))
-class HAIRSTRANDSCORE_API UGroomAsset : public UObject
+class HAIRSTRANDSCORE_API UGroomAsset : public UObject, public IInterface_AssetUserData
 {
 	GENERATED_BODY()
 
@@ -288,9 +289,23 @@ public:
 
 	int32 GetNumHairGroups() const;
 
-//private :
 #if WITH_EDITOR
 	FOnGroomAssetChanged OnGroomAssetChanged;
 #endif
+
+	//~ Begin IInterface_AssetUserData Interface
+	virtual void AddAssetUserData(UAssetUserData* InUserData) override;
+	virtual void RemoveUserDataOfClass(TSubclassOf<UAssetUserData> InUserDataClass) override;
+	virtual UAssetUserData* GetAssetUserDataOfClass(TSubclassOf<UAssetUserData> InUserDataClass) override;
+	virtual const TArray<UAssetUserData*>* GetAssetUserDataArray() const override;
+	//~ End IInterface_AssetUserData Interface
+
+//private : 
+
+
 	TUniquePtr<FHairDescription> HairDescription;
+
+	/** Array of user data stored with the asset */
+	UPROPERTY(EditAnywhere, AdvancedDisplay, Instanced, Category = HairLab)
+	TArray<UAssetUserData*> AssetUserData;
 };
