Attempt to fix my last broken commit.
#rnx

diff --git a/Engine/Plugins/Experimental/HairStrands/Source/HairStrandsCore/Private/GroomAsset.cpp b/Engine/Plugins/Experimental/HairStrands/Source/HairStrandsCore/Private/GroomAsset.cpp
--- a/Engine/Plugins/Experimental/HairStrands/Source/HairStrandsCore/Private/GroomAsset.cpp
+++ b/Engine/Plugins/Experimental/HairStrands/Source/HairStrandsCore/Private/GroomAsset.cpp
@@ -3,6 +3,7 @@
 #include "GroomAsset.h"
 #include "GroomBuilder.h"
 #include "RenderingThread.h"
+#include "Engine/AssetUserData.h"
 #include "HairStrandsVertexFactory.h"
 #include "Misc/Paths.h"
 
