Merging //UE4/Dev-Tools-Staging to Dev-VirtualProduction (//UE4/Dev-VirtualProduction)
CL 11033125
#rnx

diff --git a/Engine/Plugins/Experimental/AlembicHairImporter/Source/AlembicHairTranslator/Private/AlembicHairTranslator.cpp b/Engine/Plugins/Experimental/AlembicHairImporter/Source/AlembicHairTranslator/Private/AlembicHairTranslator.cpp
--- a/Engine/Plugins/Experimental/AlembicHairImporter/Source/AlembicHairTranslator/Private/AlembicHairTranslator.cpp
+++ b/Engine/Plugins/Experimental/AlembicHairImporter/Source/AlembicHairTranslator/Private/AlembicHairTranslator.cpp
@@ -264,7 +264,7 @@ namespace AlembicHairTranslatorUtils
 			TStrandAttributesRef<AttributeType> StrandAttributeRef = HairDescription.StrandAttributes().GetAttributesRef<AttributeType>(AttributeName);
 			if (!StrandAttributeRef.IsValid())
 			{
-				HairDescription.StrandAttributes().RegisterAttribute<AttributeType>(AttributeName);
+				HairDescription.StrandAttributes().RegisterAttribute<AttributeType>(AttributeName, 1, AttributeType::ZeroVector);
 				StrandAttributeRef = HairDescription.StrandAttributes().GetAttributesRef<AttributeType>(AttributeName);
 			}
 
@@ -345,7 +345,15 @@ namespace AlembicHairTranslatorUtils
 				break;
 				case Alembic::Util::kInt32POD:
 				{
-					ConvertAlembicAttribute<Alembic::AbcGeom::IInt32GeomParam, Alembic::Abc::Int32ArraySamplePtr, int>(HairDescription, StartStrandID, NumStrands, StartVertexID, NumVertices, Parameters, PropName);
+					switch (Extent)
+					{
+					case 1:
+						ConvertAlembicAttribute<Alembic::AbcGeom::IInt32GeomParam, Alembic::Abc::Int32ArraySamplePtr, int>(HairDescription, StartStrandID, NumStrands, StartVertexID, NumVertices, Parameters, PropName);
+						break;
+					case 3:
+						ConvertAlembicAttribute<Alembic::AbcGeom::IV3iGeomParam, Alembic::Abc::V3iArraySamplePtr, FVector>(HairDescription, StartStrandID, NumStrands, StartVertexID, NumVertices, Parameters, PropName, DataType.getExtent());
+						break;
+					}
 				}
 				break;
 				case Alembic::Util::kFloat32POD:
@@ -420,8 +428,8 @@ static void ParseObject(const Alembic::Abc::IObject& InObject, FHairDescription&
 		Alembic::Abc::Int32ArraySamplePtr NumVertices = Sample.getCurvesNumVertices();
 
 		const int32 NumWidths = Widths ? Widths->size() : 0;
-		const uint32 NumPoints = Positions ? Positions->size() : 0;
-		const uint32 NumCurves = NumVertices->size(); // equivalent to Sample.getNumCurves()
+		uint32 NumPoints = Positions ? Positions->size() : 0;
+		uint32 NumCurves = NumVertices ? NumVertices->size() : 0; // equivalent to Sample.getNumCurves()
 
 		// Get the starting strand and vertex IDs for this group of ICurves
 		int32 StartStrandID = HairDescription.GetNumStrands();
@@ -429,10 +437,24 @@ static void ParseObject(const Alembic::Abc::IObject& InObject, FHairDescription&
 
 		FMatrix ConvertedMatrix = ParentMatrix * ConversionMatrix;
 		uint32 GlobalIndex = 0;
+		uint32 TotalVertices = 0;
 		for (uint32 CurveIndex = 0; CurveIndex < NumCurves; ++CurveIndex)
 		{
 			const uint32 CurveNumVertices = (*NumVertices)[CurveIndex];
 
+			// Check the running total number of vertices and skip processing the rest of the node if there is mismatch with the number of points
+			TotalVertices += CurveNumVertices;
+			if (TotalVertices > NumPoints)
+			{
+				UE_LOG(LogAlembicHairImporter, Warning, TEXT("Curve %u of %u has %u vertices which causes total vertices (%u) to exceed the expected vertices (%u) in ICurves node. This curve and the remaining ones in the node will be skipped."),
+					CurveIndex + 1, NumCurves, CurveNumVertices, TotalVertices, NumPoints);
+
+				// Adjust the number of curves and points to those processed so far
+				NumCurves = HairDescription.GetNumStrands() - StartStrandID;
+				NumPoints = HairDescription.GetNumVertices() - StartVertexID;
+				break;
+			}
+
 			FStrandID StrandID = HairDescription.AddStrand();
 
 			SetHairStrandAttribute(HairDescription, StrandID, HairAttribute::Strand::VertexCount, (int) CurveNumVertices);
diff --git a/Engine/Plugins/Experimental/HairStrands/Source/HairStrandsCore/Private/GroomAsset.cpp b/Engine/Plugins/Experimental/HairStrands/Source/HairStrandsCore/Private/GroomAsset.cpp
--- a/Engine/Plugins/Experimental/HairStrands/Source/HairStrandsCore/Private/GroomAsset.cpp
+++ b/Engine/Plugins/Experimental/HairStrands/Source/HairStrandsCore/Private/GroomAsset.cpp
@@ -1,16 +1,35 @@
 // Copyright Epic Games, Inc. All Rights Reserved.
 
 #include "GroomAsset.h"
+#include "EngineUtils.h"
+#include "GroomAssetImportData.h"
 #include "GroomBuilder.h"
+#include "GroomImportOptions.h"
+#include "GroomSettings.h"
 #include "RenderingThread.h"
 #include "Engine/AssetUserData.h"
 #include "HairStrandsVertexFactory.h"
 #include "Misc/Paths.h"
+#include "Serialization/LargeMemoryReader.h"
+#include "Serialization/LargeMemoryWriter.h"
+#include "Serialization/MemoryReader.h"
+#include "Serialization/MemoryWriter.h"
+#include "UObject/PhysicsObjectVersion.h"
+#include "UObject/ReleaseObjectVersion.h"
 
 #if WITH_EDITORONLY_DATA
+#include "DerivedDataCacheInterface.h"
 #include "EditorFramework/AssetImportData.h"
 #endif
 
+template<typename BufferType>
+void UploadDataToBuffer(BufferType& OutBuffer, uint32 DataSizeInBytes, void* InCpuData)
+{
+	void* BufferData = RHILockVertexBuffer(OutBuffer.Buffer, 0, DataSizeInBytes, RLM_WriteOnly);
+	FMemory::Memcpy(BufferData, InCpuData, DataSizeInBytes);
+	RHIUnlockVertexBuffer(OutBuffer.Buffer);
+}
+
 template<typename FormatType>
 void CreateBuffer(const TArray<typename FormatType::Type>& InData, FRWBuffer& OutBuffer)
 {
@@ -95,6 +114,249 @@ void FHairStrandsDeformedResource::ReleaseRHI()
 
 /////////////////////////////////////////////////////////////////////////////////////////
 
+struct FClusterGrid
+{
+	struct FRenderingCurve
+	{
+		uint32 RenOffset;
+		uint32 RendPointCount;
+		float Area;
+	};
+	struct FCluster
+	{
+		TArray<FRenderingCurve> ClusterCurves;
+	};
+
+	FClusterGrid(float VoxelCentimeterSize, const FVector& InMinBound, const FVector& InMaxBound)
+	{
+		MinBound = InMinBound;
+		MaxBound = InMaxBound;
+
+		FVector VoxelCount = (MaxBound - MinBound) * FVector(1.0f / VoxelCentimeterSize);
+		GridResolution = FIntVector(FMath::CeilToInt(VoxelCount.X), FMath::CeilToInt(VoxelCount.Y), FMath::CeilToInt(VoxelCount.Z));
+
+		ClusterInfo.SetNum(GridResolution.X * GridResolution.Y * GridResolution.Z);
+	}
+
+	FORCEINLINE bool IsValid(const FIntVector& P) const
+	{
+		return	0 <= P.X && P.X < GridResolution.X &&
+			0 <= P.Y && P.Y < GridResolution.Y &&
+			0 <= P.Z && P.Z < GridResolution.Z;
+	}
+
+	FORCEINLINE FIntVector ClampToVolume(const FIntVector& CellCoord, bool& bIsValid) const
+	{
+		bIsValid = IsValid(CellCoord);
+		return FIntVector(
+			FMath::Clamp(CellCoord.X, 0, GridResolution.X - 1),
+			FMath::Clamp(CellCoord.Y, 0, GridResolution.Y - 1),
+			FMath::Clamp(CellCoord.Z, 0, GridResolution.Z - 1));
+	}
+
+	FORCEINLINE FIntVector ToCellCoord(const FVector& P) const
+	{
+		bool bIsValid = false;
+		const FVector F = ((P - MinBound) / (MaxBound - MinBound));
+		const FIntVector CellCoord = FIntVector(FMath::FloorToInt(F.X * GridResolution.X), FMath::FloorToInt(F.Y * GridResolution.Y), FMath::FloorToInt(F.Z * GridResolution.Z));
+		return ClampToVolume(CellCoord, bIsValid);
+	}
+
+	uint32 ToIndex(const FIntVector& CellCoord) const
+	{
+		uint32 CellIndex = CellCoord.X + CellCoord.Y * GridResolution.X + CellCoord.Z * GridResolution.X * GridResolution.Y;
+		check(CellIndex < uint32(ClusterInfo.Num()));
+		return CellIndex;
+	}
+
+	void InsertRenderingCurve(FRenderingCurve& Curve, const FVector& Root)
+	{
+		FIntVector CellCoord = ToCellCoord(Root);
+		uint32 Index = ToIndex(CellCoord);
+		FCluster& Cluster = ClusterInfo[Index];
+		Cluster.ClusterCurves.Add(Curve);
+	}
+
+	TArray<FCluster> GetAllClusters() { return ClusterInfo; }
+
+private:
+
+	FVector MinBound;
+	FVector MaxBound;
+
+	FIntVector GridResolution;
+
+	TArray<FCluster> ClusterInfo;
+};
+
+FHairStrandsClusterCullingResource::FHairStrandsClusterCullingResource(const FHairGroupData& GroupData)
+{
+	const FHairStrandsDatas& SimStrandsData = GroupData.HairSimulationData;
+	const FHairStrandsDatas& RenStrandsData = GroupData.HairRenderData;
+	const FHairStrandsInterpolationDatas& InterpolationData = GroupData.HairInterpolationData;
+	const uint32 RenCurveCount = RenStrandsData.GetNumCurves();
+	const uint32 SimCurveCount = SimStrandsData.GetNumCurves();
+	const uint32 PointCount = RenStrandsData.GetNumPoints();
+
+
+	// Allocate look up arrays for as many hair vertices as needed.
+	VertexCount = PointCount;
+	VertexToClusterIdArray.SetNum(PointCount);
+	ClusterVertexIdArray.SetNum(PointCount);
+
+	// Allocate cluster per voxel containing contains >=1 render curve root
+	const uint32 VoxelcountAlongLargerSide = 256; // TODO expose in UI
+	const float VoxelCentimeterSize = GroupData.HairRenderData.BoundingBox.GetSize().GetAbsMax() / VoxelcountAlongLargerSide;
+	FClusterGrid ClusterGrid(VoxelCentimeterSize, GroupData.HairRenderData.BoundingBox.Min, GroupData.HairRenderData.BoundingBox.Max);
+
+	for (uint32 RenCurveIndex = 0; RenCurveIndex < RenCurveCount; ++RenCurveIndex)
+	{
+		FClusterGrid::FRenderingCurve RCurve;
+		RCurve.RendPointCount = RenStrandsData.StrandsCurves.CurvesCount[RenCurveIndex];
+		RCurve.RenOffset = RenStrandsData.StrandsCurves.CurvesOffset[RenCurveIndex];
+
+		RCurve.Area = 0.0f;
+		for (uint32 RenPointIndex = 0; RenPointIndex < RCurve.RendPointCount; ++RenPointIndex)
+		{
+			uint32 PointGlobalIndex = RenPointIndex + RCurve.RenOffset;
+			const FVector& V0 = RenStrandsData.StrandsPoints.PointsPosition[PointGlobalIndex];
+			if (RenPointIndex > 0)
+			{
+				const FVector& V1 = RenStrandsData.StrandsPoints.PointsPosition[PointGlobalIndex-1];
+				FVector OutDir;
+				float OutLength;
+				(V1 - V0).ToDirectionAndLength(OutDir, OutLength);
+				RCurve.Area += RenStrandsData.StrandsPoints.PointsRadius[PointGlobalIndex] * OutLength;
+			}
+		}
+
+		const FVector Root = RenStrandsData.StrandsPoints.PointsPosition[RCurve.RenOffset];
+		ClusterGrid.InsertRenderingCurve(RCurve, Root);
+	}
+
+	// Count Cluster
+	ClusterCount = 0;
+	for (auto& Cluster : ClusterGrid.GetAllClusters())
+	{
+		ClusterCount += Cluster.ClusterCurves.Num() > 0 ? 1 : 0;
+	}
+	ClusterInfoArray.SetNum(ClusterCount);
+	ClusterIndexRadiusScaleInfoArray.SetNum(ClusterCount);
+
+	TArray<TArray<uint32>> ClusterToVertexIndexLod0;	// List of indices per cluster
+	ClusterToVertexIndexLod0.SetNum(ClusterCount);
+	TArray<uint32> ClusterVertexIndexCountLod1;			// Index count per cluster
+	ClusterVertexIndexCountLod1.SetNum(ClusterCount);
+
+	// Write out cluster information
+	uint32 Index = 0;
+	for (auto& Cluster : ClusterGrid.GetAllClusters())
+	{
+		if (Cluster.ClusterCurves.Num() == 0)
+		{
+			continue;
+		}
+
+		uint32 ClusterCurveCount = Cluster.ClusterCurves.Num();
+		uint32 Lod1CurveCount = FMath::Max(ClusterCurveCount / 4, 1u);		// a fourth of the curves for the lowest lod
+		uint32 Lod0VertexCount = 0;
+		uint32 Lod1VertexCount = 0;
+
+		// Prepare data to reconstruct vertex/index count to radius scale
+		float Lod0StrandArea = 0.0f;
+		float Lod1StrandArea = 0.0f;
+
+		// Sort curve to have largest area first, so that lower area curves with less influence are removed first.
+		// This also helps the radius scaling to not explode.
+		Cluster.ClusterCurves.Sort([](const FClusterGrid::FRenderingCurve& A, const FClusterGrid::FRenderingCurve& B) -> bool
+		{
+				return A.Area > B.Area;
+		});
+
+		uint32 CurveIndex = 0;
+		for (auto& ClusterCurve : Cluster.ClusterCurves)
+		{
+			for (uint32 RenPointIndex = 0; RenPointIndex < ClusterCurve.RendPointCount; ++RenPointIndex)
+			{
+				uint32 PointGlobalIndex = RenPointIndex + ClusterCurve.RenOffset;
+				VertexToClusterIdArray[PointGlobalIndex] = Index;
+				ClusterToVertexIndexLod0[Index].Add(PointGlobalIndex);
+			}
+			Lod0VertexCount += ClusterCurve.RendPointCount;
+			Lod0StrandArea += ClusterCurve.Area;
+			if (CurveIndex < Lod1CurveCount)
+			{
+				Lod1VertexCount += ClusterCurve.RendPointCount;
+				Lod1StrandArea += ClusterCurve.Area;
+			}
+			CurveIndex++;
+		}
+		if (Cluster.ClusterCurves.Num() > 0)
+		{
+			ClusterVertexIndexCountLod1[Index] = Lod1VertexCount;
+
+			const float ClusterRadiusScaleLod1 = Lod0StrandArea / Lod1StrandArea; // Could use coverage ratio (relative to cluster 0&2 bounding boxes) instead of hair area.
+			const float Numer = ClusterRadiusScaleLod1 - 1.0f;
+			const float Denom = float(Lod1VertexCount) - float(Lod0VertexCount);
+			const float AValue = Denom == 0.0f ? 0.0f : Numer / Denom;
+			ClusterIndexRadiusScaleInfoArray[Index] = AValue;
+			// ClusterIndexRadiusScaleInfoArray contains the 'a' value from the linear equation y=ax+b. Radius scale can then be recovered as ( 1.0 + a * (VertexCount - Lod0VertexCount) )
+		}
+
+		Index += Cluster.ClusterCurves.Num() > 0 ? 1 : 0;
+	}
+
+	// Now compute the cluster/strand information
+	uint32 VertexCountSum = 0;
+	for (uint32 ClusterId = 0; ClusterId < ClusterCount; ++ClusterId)
+	{
+		const uint32 ClusterVertexCountLod0 = ClusterToVertexIndexLod0[ClusterId].Num();
+		ClusterInfoArray[ClusterId].FirstVertexId = VertexCountSum;
+		ClusterInfoArray[ClusterId].VertexIdCountLodHigh = ClusterVertexCountLod0;
+		ClusterInfoArray[ClusterId].VertexIdCountLodLow = ClusterVertexIndexCountLod1[ClusterId];
+		ClusterInfoArray[ClusterId].UnusedUint = 0;
+		VertexCountSum += ClusterVertexCountLod0;
+	}
+	check(ClusterVertexIdArray.Num() == VertexCountSum);
+
+	// Flatten global VertexId for each cluster
+	for (uint32 ClusterId = 0; ClusterId < ClusterCount; ++ClusterId)
+	{
+		uint32 ClusterFirstVertexId = ClusterInfoArray[ClusterId].FirstVertexId;
+		uint32 ClusterVertexIdCount = ClusterInfoArray[ClusterId].VertexIdCountLodHigh;
+		TArray<uint32>& ClusterVertices = ClusterToVertexIndexLod0[ClusterId];
+		for (uint32 ClusterVertexId = 0; ClusterVertexId < ClusterVertexIdCount; ++ClusterVertexId)
+		{
+			ClusterVertexIdArray[ClusterFirstVertexId + ClusterVertexId] = ClusterVertices[ClusterVertexId];
+		}
+	}
+}
+
+void FHairStrandsClusterCullingResource::InitRHI()
+{
+	ClusterInfoBuffer.Initialize(sizeof(uint32) * 4, ClusterCount, EPixelFormat::PF_R32G32B32A32_UINT, BUF_Static);
+	UploadDataToBuffer(ClusterInfoBuffer, sizeof(uint32) * 4 * ClusterCount, ClusterInfoArray.GetData());
+
+	VertexToClusterIdBuffer.Initialize(sizeof(uint32), VertexToClusterIdArray.Num(), EPixelFormat::PF_R32_UINT, BUF_Static);
+	UploadDataToBuffer(VertexToClusterIdBuffer, sizeof(uint32) * VertexToClusterIdArray.Num(), VertexToClusterIdArray.GetData());
+
+	ClusterVertexIdBuffer.Initialize(sizeof(uint32), ClusterVertexIdArray.Num(), EPixelFormat::PF_R32_UINT, BUF_Static);
+	UploadDataToBuffer(ClusterVertexIdBuffer, sizeof(uint32) * ClusterVertexIdArray.Num(), ClusterVertexIdArray.GetData());
+
+	ClusterIndexRadiusScaleInfoBuffer.Initialize(sizeof(float), ClusterIndexRadiusScaleInfoArray.Num(), EPixelFormat::PF_R32_FLOAT, BUF_Static); // TODO use fp16
+	UploadDataToBuffer(ClusterIndexRadiusScaleInfoBuffer, sizeof(float) * ClusterIndexRadiusScaleInfoArray.Num(), ClusterIndexRadiusScaleInfoArray.GetData());
+}
+
+void FHairStrandsClusterCullingResource::ReleaseRHI()
+{
+	ClusterInfoBuffer.Release();
+	VertexToClusterIdBuffer.Release();
+	ClusterVertexIdBuffer.Release();
+	ClusterIndexRadiusScaleInfoBuffer.Release();
+}
+
+/////////////////////////////////////////////////////////////////////////////////////////
+
 FHairStrandsRootResource::FHairStrandsRootResource(const FHairStrandsDatas* HairStrandsDatas, uint32 LODCount):
 	RootCount(HairStrandsDatas ? HairStrandsDatas->GetNumCurves() : 0)
 {
@@ -147,7 +409,7 @@ FHairStrandsRootResource::FHairStrandsRootResource(const FHairStrandsDatas* Hair
 	uint32 LODIndex = 0;
 	for (FMeshProjectionLOD& MeshProjectionLOD : MeshProjectionLODs)
 	{
-		MeshProjectionLOD.Status = FMeshProjectionLOD::EStatus::Invalid;
+		MeshProjectionLOD.Status = FHairStrandsProjectionHairData::LODData::EStatus::Invalid;
 		MeshProjectionLOD.RestRootOffset = HairStrandsDatas->BoundingBox.GetCenter();
 		MeshProjectionLOD.LODIndex = LODIndex++;
 	}
@@ -162,7 +424,7 @@ void FHairStrandsRootResource::InitRHI()
 		CreateBuffer<FHairStrandsRootNormalFormat>(RootNormals, RootNormalBuffer);	
 		for (FMeshProjectionLOD& LOD : MeshProjectionLODs)
 		{
-			LOD.Status = FHairStrandsRootResource::FMeshProjectionLOD::EStatus::Initialized;
+			LOD.Status = FHairStrandsProjectionHairData::LODData::EStatus::Initialized;
 
 			// Create buffers. Initialization will be done by render passes
 			CreateBuffer<FHairStrandsCurveTriangleIndexFormat>(RootCount, LOD.RootTriangleIndexBuffer);
@@ -187,7 +449,7 @@ void FHairStrandsRootResource::ReleaseRHI()
 
 	for (FMeshProjectionLOD& LOD : MeshProjectionLODs)
 	{
-		LOD.Status = FHairStrandsRootResource::FMeshProjectionLOD::EStatus::Invalid;
+		LOD.Status = FHairStrandsProjectionHairData::LODData::EStatus::Invalid;
 		LOD.RootTriangleIndexBuffer.Release();
 		LOD.RootTriangleBarycentricBuffer.Release();
 		LOD.RestRootTrianglePosition0Buffer.Release();
@@ -261,7 +523,44 @@ void UGroomAsset::Serialize(FArchive& Ar)
 {
 	Super::Serialize(Ar);
 
-	Ar << HairGroupsData;
+	Ar.UsingCustomVersion(FPhysicsObjectVersion::GUID);
+	Ar.UsingCustomVersion(FReleaseObjectVersion::GUID); // Needed to support MeshDescription AttributesSet serialization
+
+	if (Ar.CustomVer(FPhysicsObjectVersion::GUID) >= FPhysicsObjectVersion::GroomWithDescription)
+	{
+		FStripDataFlags StripFlags(Ar);
+		if (StripFlags.IsEditorDataStripped() || (Ar.IsSaving() && !CanRebuildFromDescription()))
+		{
+			// When cooking data or serializing old format to new format,
+			// serialize the computed groom data
+			Ar << HairGroupsData;
+		}
+#if WITH_EDITORONLY_DATA
+		else
+		{
+			// When serializing data for editor, serialize the HairDescription as bulk data
+			// The computed groom data is fetched from the Derived Data Cache
+			if (!HairDescriptionBulkData)
+			{
+				// When loading, bulk data can be null so instantiate a new one to serialize into
+				HairDescriptionBulkData = MakeUnique<FHairDescriptionBulkData>();
+			}
+
+			HairDescriptionBulkData->Serialize(Ar, this);
+
+			// Serialize the HairGroupsData directly into the asset if it couldn't be cached in the DDC
+			if (!bIsCacheable)
+			{
+				Ar << HairGroupsData;
+			}
+		}
+#endif // WITH_EDITORONLY_DATA
+	}
+	else
+	{
+		// Old format serialized the computed groom data directly
+		Ar << HairGroupsData;
+	}
 }
 
 void UGroomAsset::InitResource()
@@ -336,10 +635,24 @@ void UGroomAsset::PostLoad()
 {
 	Super::PostLoad();
 
+#if WITH_EDITORONLY_DATA
+	if (bIsCacheable)
+	{
+		CacheDerivedData();
+	}
+#endif
+
 	check(HairGroupsData.Num() > 0);
 	if (HairGroupsData[0].HairSimulationData.GetNumCurves() == 0 || HairGroupsData[0].HairInterpolationData.Num() == 0)
 	{
-		FGroomBuilder::BuildData(this, 2, 0, false, false);
+		FGroomBuildSettings Settings;
+		Settings.bOverrideGuides = false;
+		Settings.HairToGuideDensity = 0.1f;
+		Settings.InterpolationQuality = EGroomInterpolationQuality::High;
+		Settings.InterpolationDistance = EGroomInterpolationWeight::Parametric;
+		Settings.bRandomizeGuide = false;
+		Settings.bUseUniqueGuide = false;
+		FGroomBuilder::BuildData(this, Settings);
 	}
 
 	if (!IsTemplate())
@@ -441,4 +754,227 @@ void UGroomAsset::RemoveUserDataOfClass(TSubclassOf<UAssetUserData> InUserDataCl
 const TArray<UAssetUserData*>* UGroomAsset::GetAssetUserDataArray() const
 {
 	return &AssetUserData;
-}
\ No newline at end of file
+}
+
+FArchive& operator<<(FArchive& Ar, FHairGroupInfo& GroupInfo)
+{
+	Ar.UsingCustomVersion(FPhysicsObjectVersion::GUID);
+
+	Ar << GroupInfo.GroupID;
+	Ar << GroupInfo.NumCurves;
+	Ar << GroupInfo.NumGuides;
+
+	if (Ar.CustomVer(FPhysicsObjectVersion::GUID) >= FPhysicsObjectVersion::GroomWithImportSettings)
+	{
+		Ar << GroupInfo.ImportSettings;
+	}
+	else
+	{
+		if (Ar.IsLoading())
+		{
+			bool bIsAutoGenerated = false;
+			Ar << bIsAutoGenerated;
+
+			GroupInfo.ImportSettings = FGroomBuildSettings();
+			GroupInfo.ImportSettings.InterpolationQuality = EGroomInterpolationQuality::Unknown;
+			GroupInfo.ImportSettings.InterpolationDistance = EGroomInterpolationWeight::Unknown;
+			GroupInfo.ImportSettings.bOverrideGuides = bIsAutoGenerated;
+		}
+		else
+		{
+			bool bIsAutoGenerated = GroupInfo.ImportSettings.bOverrideGuides;
+			Ar << bIsAutoGenerated;
+		}
+	}
+
+	// Ignoring Material since it's null
+
+	return Ar;
+}
+
+bool UGroomAsset::CanRebuildFromDescription() const
+{
+#if WITH_EDITORONLY_DATA
+	return HairDescriptionBulkData.IsValid() && !HairDescriptionBulkData->IsEmpty();
+#else
+	return false;
+#endif
+}
+
+// If groom derived data needs to be rebuilt (new format, serialization
+// differences, etc.) replace the version GUID below with a new one.
+// In case of merge conflicts with DDC versions, you MUST generate a new GUID
+// and set this new GUID as the version.
+#define GROOM_DERIVED_DATA_VERSION TEXT("42406925F1934F9794A08A8EE51BA05D")
+
+#if WITH_EDITORONLY_DATA
+
+namespace GroomDerivedDataCacheUtils
+{
+	const FString& GetGroomDerivedDataVersion()
+	{
+		static FString CachedVersionString(GROOM_DERIVED_DATA_VERSION);
+		return CachedVersionString;
+	}
+
+	FString BuildGroomDerivedDataKey(const FString& KeySuffix)
+	{
+		return FDerivedDataCacheInterface::BuildCacheKey(TEXT("GROOM"), *GetGroomDerivedDataVersion(), *KeySuffix);
+	}
+
+	void SerializeGroomBuildSettingsForDDC(FArchive& Ar, FGroomBuildSettings& BuildSettings)
+	{
+		// Note: this serializer is only used to build the groom DDC key, no versioning is required
+		Ar << BuildSettings.bOverrideGuides;
+		Ar << BuildSettings.HairToGuideDensity;
+		Ar << BuildSettings.InterpolationQuality;
+		Ar << BuildSettings.InterpolationDistance;
+		Ar << BuildSettings.bRandomizeGuide;
+		Ar << BuildSettings.bUseUniqueGuide;
+	}
+}
+
+FString UGroomAsset::BuildDerivedDataKeySuffix(const FGroomBuildSettings& BuildSettings)
+{
+	// Serialize the build settings into a temporary array
+	// The archive is flagged as persistent so that machines of different endianness produce identical binary results.
+	TArray<uint8> TempBytes;
+	TempBytes.Reserve(64);
+	FMemoryWriter Ar(TempBytes, /*bIsPersistent=*/ true);
+
+	GroomDerivedDataCacheUtils::SerializeGroomBuildSettingsForDDC(Ar, const_cast<FGroomBuildSettings&>(BuildSettings));
+
+	FString KeySuffix;
+	if (HairDescriptionBulkData)
+	{
+		// Reserve twice the size of TempBytes because of ByteToHex below + 3 for "ID" and \0
+		KeySuffix.Reserve(HairDescriptionBulkData->GetIdString().Len() + TempBytes.Num() * 2 + 3);
+		KeySuffix += TEXT("ID");
+		KeySuffix += HairDescriptionBulkData->GetIdString();
+	}
+	else
+	{
+		KeySuffix.Reserve(TempBytes.Num() * 2 + 1);
+	}
+
+	// Now convert the raw bytes to a string
+	const uint8* SettingsAsBytes = TempBytes.GetData();
+	for (int32 ByteIndex = 0; ByteIndex < TempBytes.Num(); ++ByteIndex)
+	{
+		ByteToHex(SettingsAsBytes[ByteIndex], KeySuffix);
+	}
+
+	return KeySuffix;
+}
+
+void UGroomAsset::CommitHairDescription(FHairDescription&& InHairDescription)
+{
+	HairDescription = MakeUnique<FHairDescription>(InHairDescription);
+
+	if (!HairDescriptionBulkData)
+	{
+		HairDescriptionBulkData = MakeUnique<FHairDescriptionBulkData>();
+	}
+	HairDescriptionBulkData->SaveHairDescription(*HairDescription);
+}
+
+bool UGroomAsset::CacheDerivedData(const FGroomBuildSettings* InBuildSettings)
+{
+	if (!HairDescriptionBulkData)
+	{
+		return false;
+	}
+
+	static FGroomBuildSettings DefaultBuildSettings;
+	const FGroomBuildSettings* BuildSettings = &DefaultBuildSettings;
+	if (InBuildSettings)
+	{
+		// Use the provided build settings since a new build is requested
+		BuildSettings = InBuildSettings;
+	}
+	else if (AssetImportData)
+	{
+		// Otherwise, use the build settings that were used previously at import/re-import
+		UGroomAssetImportData* GroomImportData = Cast<UGroomAssetImportData>(AssetImportData);
+		if (GroomImportData && GroomImportData->ImportOptions)
+		{
+			BuildSettings = &GroomImportData->ImportOptions->BuildSettings;
+		}
+	}
+
+	bool bSuccess = true;
+
+	const FString KeySuffix = BuildDerivedDataKeySuffix(*BuildSettings);
+	const FString DerivedDataKey = GroomDerivedDataCacheUtils::BuildGroomDerivedDataKey(KeySuffix);
+
+	TArray<uint8> DerivedData;
+	if (GetDerivedDataCacheRef().GetSynchronous(*DerivedDataKey, DerivedData))
+	{
+		FMemoryReader Ar(DerivedData, /*bIsPersistent=*/ true);
+
+		int64 UncompressedSize = 0;
+		Ar << UncompressedSize;
+
+		uint8* DecompressionBuffer = reinterpret_cast<uint8*>(FMemory::Malloc(UncompressedSize));
+		Ar.SerializeCompressed(DecompressionBuffer, 0, NAME_Zlib);
+
+		FLargeMemoryReader LargeMemReader(DecompressionBuffer, UncompressedSize, ELargeMemoryReaderFlags::Persistent | ELargeMemoryReaderFlags::TakeOwnership);
+		if (HairGroupsInfo.Num() > 0)
+		{
+			// HairGroupsInfo serialized from the asset, so ignore the one from the DDC
+			TArray<FHairGroupInfo> TempHairGroupsInfo;
+			LargeMemReader << TempHairGroupsInfo;
+		}
+		else
+		{
+			// This is the case where a groom with the same build settings was previously pushed to DDC,
+			// but this is a new asset so HairGroupsInfo isn't serialized yet
+			LargeMemReader << HairGroupsInfo;
+		}
+
+		LargeMemReader << HairGroupsData;
+	}
+	else
+	{
+		// Load the HairDescription from the bulk data if needed
+		if (!HairDescription)
+		{
+			HairDescription = MakeUnique<FHairDescription>();
+			HairDescriptionBulkData->LoadHairDescription(*HairDescription);
+		}
+
+		// Build groom data with the new build settings
+		bSuccess = FGroomBuilder::BuildGroom(*HairDescription, *BuildSettings, this);
+		if (bSuccess)
+		{
+			// Using a LargeMemoryWriter for serialization since the data can be bigger than 2 GB
+			FLargeMemoryWriter LargeMemWriter(0, /*bIsPersistent=*/ true);
+			LargeMemWriter << HairGroupsInfo;
+			LargeMemWriter << HairGroupsData;
+
+			int64 UncompressedSize = LargeMemWriter.TotalSize();
+
+			// Then the content of the LargeMemWriter is compressed into a MemoryWriter
+			// Compression ratio can reach about 5:2 depending on the data
+			// Since the DDC doesn't support data bigger than 2 GB
+			// we can compute a size threshold to skip the caching when
+			// the uncompressed size exceeds the threshold
+			static constexpr const int64 SizeThreshold =  (int64) MAX_int32 * 2.5;
+			bIsCacheable = UncompressedSize < SizeThreshold;
+			if (bIsCacheable)
+			{
+				FMemoryWriter CompressedArchive(DerivedData, true);
+
+				CompressedArchive << UncompressedSize; // needed for allocating decompression buffer
+				CompressedArchive.SerializeCompressed(LargeMemWriter.GetData(), UncompressedSize, NAME_Zlib);
+
+				GetDerivedDataCacheRef().Put(*DerivedDataKey, DerivedData);
+			}
+		}
+	}
+	return bSuccess;
+}
+
+#endif // WITH_EDITORONLY_DATA
+
+#undef LOCTEXT_NAMESPACE
\ No newline at end of file
diff --git a/Engine/Plugins/Experimental/HairStrands/Source/HairStrandsCore/Public/GroomAsset.h b/Engine/Plugins/Experimental/HairStrands/Source/HairStrandsCore/Public/GroomAsset.h
--- a/Engine/Plugins/Experimental/HairStrands/Source/HairStrandsCore/Public/GroomAsset.h
+++ b/Engine/Plugins/Experimental/HairStrands/Source/HairStrandsCore/Public/GroomAsset.h
@@ -8,12 +8,15 @@
 #include "HairDescription.h"
 #include "HairStrandsDatas.h"
 #include "RenderResource.h"
+#include "GroomSettings.h"
+#include "HairStrandsInterface.h"
 #include "Interfaces/Interface_AssetUserData.h"
 
 #include "GroomAsset.generated.h"
 
 class UMaterialInterface;
 class UNiagaraSystem;
+struct FHairGroupData;
 
 /* Render buffers for root deformation for dynamic meshes */
 struct FHairStrandsRootResource : public FRenderResource
@@ -35,8 +38,7 @@ struct FHairStrandsRootResource : public FRenderResource
 
 	struct FMeshProjectionLOD
 	{
-		enum class EStatus { Invalid, Initialized, Completed };
-		EStatus Status = EStatus::Invalid;
+		FHairStrandsProjectionHairData::LODData::EStatus Status = FHairStrandsProjectionHairData::LODData::EStatus::Invalid;
 		int32 LODIndex = -1;
 
 		/* Triangle on which a root is attached */
@@ -134,6 +136,50 @@ struct FHairStrandsDeformedResource : public FRenderResource
 	const bool bInitializedData;
 };
 
+struct FHairStrandsClusterCullingResource : public FRenderResource
+{
+	FHairStrandsClusterCullingResource(const FHairGroupData& GroupData);
+
+	/* Init the buffer */
+	virtual void InitRHI() override;
+
+	/* Release the buffer */
+	virtual void ReleaseRHI() override;
+
+	/* Get the resource name */
+	virtual FString GetFriendlyName() const override { return TEXT("FHairStrandsClusterResource"); }
+
+
+	struct FClusterInfo
+	{
+		uint32 FirstVertexId;
+		uint32 VertexIdCountLodHigh;
+		uint32 VertexIdCountLodLow;
+		uint32 UnusedUint;
+	};
+	TArray<FClusterInfo> ClusterInfoArray;
+	TArray<uint32> VertexToClusterIdArray;
+	TArray<uint32> ClusterVertexIdArray;
+	TArray<float> ClusterIndexRadiusScaleInfoArray;
+
+	/* Cluster info buffer */
+	FReadBuffer ClusterInfoBuffer;		// Contains: Start VertexId, VertexId Count (for VertexToClusterIdBuffer)
+
+	/* VertexId => ClusterId to know which AABB to contribute to*/
+	FReadBuffer VertexToClusterIdBuffer;
+
+	/* Concatenated data for each cluster: list of VertexId pointed to by ClusterInfoBuffer */
+	FReadBuffer ClusterVertexIdBuffer;
+
+	/* Contains information to recove the radius scale to apply per cluster when decimating vertex count */
+	FReadBuffer ClusterIndexRadiusScaleInfoBuffer;
+
+	/* Number of cluster  */
+	uint32 ClusterCount;
+	/* Number of vertex  */
+	uint32 VertexCount;
+};
+
 struct FHairStrandsInterpolationResource : public FRenderResource
 {
 	/** Build the hair strands resource */
@@ -194,8 +240,8 @@ struct HAIRSTRANDSCORE_API FHairGroupInfo
 	UPROPERTY(VisibleAnywhere, Category = "Info", meta = (DisplayName = "Guide Count"))
 	int32 NumGuides = 0;
 
-	UPROPERTY(VisibleAnywhere, Category = "Info", meta = (DisplayName = "Auto-generated Guides", ToolTip = "Checked when guides were auto-generated"))
-	bool bIsAutoGenerated = false;
+	UPROPERTY(VisibleAnywhere, Category = "Info", meta = (DisplayName = "Import options", ToolTip = "Show the options used at import time"))
+	FGroomBuildSettings ImportSettings;
 
 	UPROPERTY(EditAnywhere, Category="Rendering")
 	UMaterialInterface* Material = nullptr;
@@ -203,6 +249,8 @@ struct HAIRSTRANDSCORE_API FHairGroupInfo
 	// Currently hide the Nigara simulation slot as it is not used, and could confuse users
 	//UPROPERTY(EditAnywhere, Category="Simulation", meta = (DisplayName = "Niagara System Asset"))
 	//UNiagaraSystem* NiagaraAsset = nullptr;
+
+	friend FArchive& operator<<(FArchive& Ar, FHairGroupInfo& GroupInfo);
 };
 
 struct HAIRSTRANDSCORE_API FHairGroupData
@@ -289,9 +337,8 @@ public:
 
 	int32 GetNumHairGroups() const;
 
-#if WITH_EDITOR
-	FOnGroomAssetChanged OnGroomAssetChanged;
-#endif
+	/** Returns true if the asset has the HairDescription needed to recompute its groom data */
+	bool CanRebuildFromDescription() const;
 
 	//~ Begin IInterface_AssetUserData Interface
 	virtual void AddAssetUserData(UAssetUserData* InUserData) override;
@@ -300,12 +347,30 @@ public:
 	virtual const TArray<UAssetUserData*>* GetAssetUserDataArray() const override;
 	//~ End IInterface_AssetUserData Interface
 
-//private : 
+//private :
+#if WITH_EDITOR
+	FOnGroomAssetChanged OnGroomAssetChanged;
+#endif
+
+#if WITH_EDITORONLY_DATA
+public:
+	/** Commits a HairDescription to buffer for serialization */
+	void CommitHairDescription(FHairDescription&& HairDescription);
+
+	/** Caches the computed groom data with the given build settings from/to the Derived Data Cache, building it if needed */
+	bool CacheDerivedData(const struct FGroomBuildSettings* BuildSettings = nullptr);
 
+private:
+	FString BuildDerivedDataKeySuffix(const struct FGroomBuildSettings& BuildSettings);
 
 	TUniquePtr<FHairDescription> HairDescription;
+	TUniquePtr<FHairDescriptionBulkData> HairDescriptionBulkData;
 
 	/** Array of user data stored with the asset */
 	UPROPERTY(EditAnywhere, AdvancedDisplay, Instanced, Category = HairLab)
 	TArray<UAssetUserData*> AssetUserData;
+
+	UPROPERTY()
+	bool bIsCacheable = true;
+#endif
 };
diff --git a/Engine/Plugins/FX/Niagara/Source/NiagaraCore/Private/NiagaraCustomVersion.cpp b/Engine/Plugins/FX/Niagara/Source/NiagaraCore/Private/NiagaraCustomVersion.cpp
--- a/Engine/Plugins/FX/Niagara/Source/NiagaraCore/Private/NiagaraCustomVersion.cpp
+++ b/Engine/Plugins/FX/Niagara/Source/NiagaraCore/Private/NiagaraCustomVersion.cpp
@@ -8,4 +8,4 @@ const FGuid FNiagaraCustomVersion::GUID(0xFCF57AFA, 0x50764283, 0xB9A9E658, 0xFF
 // Register the custom version with core
 FCustomVersionRegistration GRegisterNiagaraCustomVersion(FNiagaraCustomVersion::GUID, FNiagaraCustomVersion::LatestVersion, TEXT("NiagaraVer"));
 
-const FGuid FNiagaraCustomVersion::LatestScriptCompileVersion(0x64FCB871, 0xB2954D47, 0x8721FF75, 0x20496D93);
+const FGuid FNiagaraCustomVersion::LatestScriptCompileVersion(0xA6B5A395, 0xF276467E, 0x9D03C2C2, 0x99CAE1E3);
diff --git a/Engine/Source/Runtime/Renderer/Private/DeferredShadingRenderer.cpp b/Engine/Source/Runtime/Renderer/Private/DeferredShadingRenderer.cpp
--- a/Engine/Source/Runtime/Renderer/Private/DeferredShadingRenderer.cpp
+++ b/Engine/Source/Runtime/Renderer/Private/DeferredShadingRenderer.cpp
@@ -39,6 +39,7 @@
 #include "RayTracingDefinitions.h"
 #include "RayTracingInstance.h"
 #include "ShaderPrint.h"
+#include "GpuDebugRendering.h"
 #include "HairStrands/HairStrandsRendering.h"
 
 static TAutoConsoleVariable<int32> CVarStencilForLODDither(
@@ -454,6 +455,16 @@ static FORCEINLINE bool NeedsPrePass(const FDeferredShadingSceneRenderer* Render
 		(Renderer->EarlyZPassMode != DDM_None || Renderer->bEarlyZPassMovable != 0);
 }
 
+static bool DoesHairStrandsRequestHzb(EShaderPlatform Platform)
+{
+	auto HzbRequested = [&]()
+	{
+		IConsoleVariable* CVar = IConsoleManager::Get().FindConsoleVariable(TEXT("HairStrands.Cluster.CullingUsesHzb"));
+		return CVar && CVar->GetInt() > 0;
+	};
+	return IsHairStrandsEnable(Platform) && HzbRequested();
+}
+
 bool FDeferredShadingSceneRenderer::RenderHzb(FRHICommandListImmediate& RHICmdList)
 {
 	FSceneRenderTargets& SceneContext = FSceneRenderTargets::Get(RHICmdList);
@@ -475,8 +486,9 @@ bool FDeferredShadingSceneRenderer::RenderHzb(FRHICommandListImmediate& RHICmdLi
 		const bool bSSR  = ShouldRenderScreenSpaceReflections(View);
 		const bool bSSAO = ShouldRenderScreenSpaceAmbientOcclusion(View);
 		const bool bSSGI = ShouldRenderScreenSpaceDiffuseIndirect(View);
+		const bool bHair = DoesHairStrandsRequestHzb(Scene->GetShaderPlatform());
 
-		if (bSSAO || bHZBOcclusion || bSSR || bSSGI)
+		if (bSSAO || bHZBOcclusion || bSSR || bSSGI || bHair)
 		{
 			FRDGBuilder GraphBuilder(RHICmdList);
 
@@ -552,10 +564,17 @@ static int32 GetCustomDepthPassLocation()
 
 void FDeferredShadingSceneRenderer::PrepareDistanceFieldScene(FRHICommandListImmediate& RHICmdList, bool bSplitDispatch)
 {
+	CSV_SCOPED_TIMING_STAT_EXCLUSIVE(RenderDFAO);
+	SCOPE_CYCLE_COUNTER(STAT_FDeferredShadingSceneRenderer_DistanceFieldAO_Init);
+
+	if (ShouldPrepareHeightFieldScene())
+	{
+		GHeightFieldTextureAtlas.UpdateAllocations(RHICmdList, FeatureLevel);
+		UpdateGlobalHeightFieldObjectBuffers(RHICmdList);
+	}
+
 	if (ShouldPrepareDistanceFieldScene())
 	{
-		CSV_SCOPED_TIMING_STAT_EXCLUSIVE(RenderDFAO);
-		SCOPE_CYCLE_COUNTER(STAT_FDeferredShadingSceneRenderer_DistanceFieldAO_Init);
 		GDistanceFieldVolumeTextureAtlas.UpdateAllocations(RHICmdList, FeatureLevel);
 		UpdateGlobalDistanceFieldObjectBuffers(RHICmdList);
 		if (bSplitDispatch)
@@ -1152,6 +1171,7 @@ void FDeferredShadingSceneRenderer::Render(FRHICommandListImmediate& RHICmdList)
 			FViewInfo& View = Views[ViewIndex];
 			SCOPED_GPU_MASK(RHICmdList, View.GPUMask);
 			ShaderPrint::BeginView(RHICmdList, View);
+			ShaderDrawDebug::BeginView(RHICmdList, View);
 		}
 
 		for (int32 ViewIndex = 0; ViewIndex < Views.Num(); ViewIndex++)
@@ -1386,11 +1406,13 @@ void FDeferredShadingSceneRenderer::Render(FRHICommandListImmediate& RHICmdList)
 
 	// Interpolation needs to happen after the skin cache run as there is a dependency 
 	// on the skin cache output.
-	if (IsHairStrandsEnable(Scene->GetShaderPlatform()) && Views.Num() > 0)
+	const bool bRunHairStrands = IsHairStrandsEnable(Scene->GetShaderPlatform()) && Views.Num() > 0;
+	FHairStrandClusterData HairClusterData;
+	if (bRunHairStrands)
 	{
 		const EWorldType::Type WorldType = Views[0].Family->Scene->GetWorld()->WorldType;
 		auto ShaderMap = GetGlobalShaderMap(FeatureLevel);
-		RunHairStrandsInterpolation(RHICmdList, WorldType, ShaderMap, EHairStrandsInterpolationType::RenderStrands);
+		RunHairStrandsInterpolation(RHICmdList, WorldType, &Views[0].ShaderDrawData, ShaderMap, EHairStrandsInterpolationType::RenderStrands, &HairClusterData); // Send data to full up with culling
 	}
 
 	// Before starting the render, all async task for the Custom data must be completed
@@ -1930,10 +1952,30 @@ void FDeferredShadingSceneRenderer::Render(FRHICommandListImmediate& RHICmdList)
 	if (IsHairStrandsEnable(Scene->GetShaderPlatform()) && bIsViewCompatible)
 	{
 		SCOPED_GPU_STAT(RHICmdList, HairRendering);
-		HairDatasStorage.HairClusterPerViews = CreateHairStrandsClusters(RHICmdList, Scene, Views);
-		VoxelizeHairStrands(RHICmdList, Scene, Views, HairDatasStorage.HairClusterPerViews);
-		HairDatasStorage.DeepShadowViews = RenderHairStrandsDeepShadows(RHICmdList, Scene, Views, HairDatasStorage.HairClusterPerViews);
-		HairDatasStorage.HairVisibilityViews = RenderHairStrandsVisibilityBuffer(RHICmdList, Scene, Views, SceneContext.GBufferB, SceneContext.GetSceneColor(), SceneContext.SceneDepthZ, SceneContext.SceneVelocity, HairDatasStorage.HairClusterPerViews);
+		HairDatasStorage.MacroGroupsPerViews = CreateHairStrandsMacroGroups(RHICmdList, Scene, Views);
+
+		// Culling/LOD pass for DOM and Voxelisation altogether
+		FHairCullingParams CullingParams;
+		CullingParams.bShadowViewMode = true;
+		CullingParams.bCullingProcessSkipped = false;
+		ComputeHairStrandsClustersCulling(RHICmdList, *GetGlobalShaderMap(FeatureLevel), Views, CullingParams, HairClusterData);
+
+		ServiceLocalQueue();
+
+		// Voxelization and Deep Opacity Maps
+		VoxelizeHairStrands(RHICmdList, Scene, Views, HairDatasStorage.MacroGroupsPerViews);
+		HairDatasStorage.DeepShadowViews = RenderHairStrandsDeepShadows(RHICmdList, Scene, Views, HairDatasStorage.MacroGroupsPerViews);
+
+		ServiceLocalQueue();
+
+		// Culling/LOD pass for visibility (must be done after HZB is generated)
+		CullingParams.bShadowViewMode = false;
+		ComputeHairStrandsClustersCulling(RHICmdList, *GetGlobalShaderMap(FeatureLevel), Views, CullingParams, HairClusterData);
+		// Hair visibility pass
+		HairDatasStorage.HairVisibilityViews = RenderHairStrandsVisibilityBuffer(RHICmdList, Scene, Views, SceneContext.GBufferB, SceneContext.GetSceneColor(), SceneContext.SceneDepthZ, SceneContext.SceneVelocity, HairDatasStorage.MacroGroupsPerViews);
+		// Reset indirect draw buffer
+		ResetHairStrandsClusterToLOD0(RHICmdList, *GetGlobalShaderMap(FeatureLevel), HairClusterData);
+
 		ServiceLocalQueue();
 
 		HairDatas = &HairDatasStorage;
@@ -2313,7 +2355,7 @@ void FDeferredShadingSceneRenderer::Render(FRHICommandListImmediate& RHICmdList)
 
 		if (IsHairStrandsEnable(Scene->GetShaderPlatform()))
 		{
-			RenderHairStrandsDebugInfo(RHICmdList, Views, HairDatas);
+			RenderHairStrandsDebugInfo(RHICmdList, Views, HairDatas, HairClusterData);
 		}
 	}
 
@@ -2486,6 +2528,7 @@ void FDeferredShadingSceneRenderer::Render(FRHICommandListImmediate& RHICmdList)
 		for (int32 ViewIndex = 0; ViewIndex < Views.Num(); ViewIndex++)
 		{
 			ShaderPrint::EndView(Views[ViewIndex]);
+			ShaderDrawDebug::EndView(Views[ViewIndex]);
 		}
 
 #if WITH_MGPU
diff --git a/Engine/Source/Runtime/Renderer/Private/PostProcess/SceneRenderTargets.h b/Engine/Source/Runtime/Renderer/Private/PostProcess/SceneRenderTargets.h
--- a/Engine/Source/Runtime/Renderer/Private/PostProcess/SceneRenderTargets.h
+++ b/Engine/Source/Runtime/Renderer/Private/PostProcess/SceneRenderTargets.h
@@ -117,7 +117,7 @@ static const int32 NumTranslucencyShadowSurfaces = 2;
 
 #define STENCIL_LIGHTING_CHANNELS_MASK(Value) uint8((Value & 0x7) << STENCIL_LIGHTING_CHANNELS_BIT_ID)
 
-#define PREVENT_RENDERTARGET_SIZE_THRASHING (PLATFORM_DESKTOP || PLATFORM_XBOXONE || PLATFORM_PS4 || PLATFORM_ANDROID || PLATFORM_IOS || PLATFORM_SWITCH)
+#define PREVENT_RENDERTARGET_SIZE_THRASHING (PLATFORM_DESKTOP || PLATFORM_XBOXONE || PLATFORM_PS4 || PLATFORM_ANDROID || PLATFORM_IOS || PLATFORM_SWITCH || PLATFORM_UNIX)
 
 enum class ESceneColorFormatType
 {
@@ -490,7 +490,7 @@ public:
 
 	void AllocateDebugViewModeTargets(FRHICommandList& RHICmdList);
 
-	void AllocateScreenShadowMask(FRHICommandList& RHICmdList, TRefCountPtr<IPooledRenderTarget>& ScreenShadowMaskTexture);
+	void AllocateScreenShadowMask(FRHICommandList& RHICmdList, TRefCountPtr<IPooledRenderTarget>& ScreenShadowMaskTexture, bool bCreateShaderResourceView = false);
 	void AllocateScreenTransmittanceMask(FRHICommandList& RHICmdList, TRefCountPtr<IPooledRenderTarget>& ScreenTransmittanceMaskTexture);
 
 	TRefCountPtr<IPooledRenderTarget>& GetReflectionBrightnessTarget();
