Copying Private-LoadTimes-4.24 stream to Main. Biggest changes are in Materials/Shader memory freezing.

diff --git a/Engine/Plugins/Experimental/HairStrands/Source/HairStrandsCore/Private/HairStrandsVertexFactory.cpp b/Engine/Plugins/Experimental/HairStrands/Source/HairStrandsCore/Private/HairStrandsVertexFactory.cpp
--- a/Engine/Plugins/Experimental/HairStrands/Source/HairStrandsCore/Private/HairStrandsVertexFactory.cpp
+++ b/Engine/Plugins/Experimental/HairStrands/Source/HairStrandsCore/Private/HairStrandsVertexFactory.cpp
@@ -42,27 +42,28 @@ TGlobalResource<FDummyCulledDispatchVertexIdsBuffer> GDummyCulledDispatchVertexI
 
 class FHairStrandsVertexFactoryShaderParameters : public FVertexFactoryShaderParameters
 {
+	DECLARE_TYPE_LAYOUT(FHairStrandsVertexFactoryShaderParameters, NonVirtual);
 public:
 
-	FShaderParameter Radius;
-	FShaderParameter Length;
-	FShaderParameter RadiusAtDepth1_Primary;	// unused
-	FShaderParameter RadiusAtDepth1_Velocity;	// unused
-	FShaderParameter PositionOffset;
-	FShaderParameter PreviousPositionOffset;
-	FShaderParameter Density;
-	FShaderParameter Culling;
-
-	FShaderResourceParameter PositionBuffer;
-	FShaderResourceParameter PreviousPositionBuffer;
-	FShaderResourceParameter AttributeBuffer;
-	FShaderResourceParameter MaterialBuffer;
-	FShaderResourceParameter TangentBuffer;
-
-	FShaderResourceParameter CulledVertexIdsBuffer;
-	FShaderResourceParameter CulledVertexRadiusScaleBuffer;
-
-	virtual void Bind(const FShaderParameterMap& ParameterMap) override
+	LAYOUT_FIELD(FShaderParameter, Radius);
+	LAYOUT_FIELD(FShaderParameter, Length);
+	LAYOUT_FIELD(FShaderParameter, RadiusAtDepth1_Primary);	// unused
+	LAYOUT_FIELD(FShaderParameter, RadiusAtDepth1_Velocity);	// unused
+	LAYOUT_FIELD(FShaderParameter, PositionOffset);
+	LAYOUT_FIELD(FShaderParameter, PreviousPositionOffset);
+	LAYOUT_FIELD(FShaderParameter, Density);
+	LAYOUT_FIELD(FShaderParameter, Culling);
+
+	LAYOUT_FIELD(FShaderResourceParameter, PositionBuffer);
+	LAYOUT_FIELD(FShaderResourceParameter, PreviousPositionBuffer);
+	LAYOUT_FIELD(FShaderResourceParameter, AttributeBuffer);
+	LAYOUT_FIELD(FShaderResourceParameter, MaterialBuffer);
+	LAYOUT_FIELD(FShaderResourceParameter, TangentBuffer);
+
+	LAYOUT_FIELD(FShaderResourceParameter, CulledVertexIdsBuffer);
+	LAYOUT_FIELD(FShaderResourceParameter, CulledVertexRadiusScaleBuffer);
+
+	void Bind(const FShaderParameterMap& ParameterMap)
 	{
 		Radius.Bind(ParameterMap, TEXT("HairStrandsVF_Radius"));
 		Length.Bind(ParameterMap, TEXT("HairStrandsVF_Length"));
@@ -81,28 +82,7 @@ public:
 		CulledVertexRadiusScaleBuffer.Bind(ParameterMap, TEXT("CulledVertexRadiusScaleBuffer"));
 	}
 
-	virtual void Serialize(FArchive& Ar) override
-	{
-		Ar << Radius;
-		Ar << Length;
-		Ar << RadiusAtDepth1_Primary;	// unused
-		Ar << RadiusAtDepth1_Velocity;	// unused
-		Ar << PositionOffset;
-		Ar << PreviousPositionOffset;
-		Ar << Density;
-		Ar << Culling;
-
-		Ar << PositionBuffer;
-		Ar << PreviousPositionBuffer;
-		Ar << AttributeBuffer;
-		Ar << MaterialBuffer;
-		Ar << TangentBuffer;
-
-		Ar << CulledVertexIdsBuffer;
-		Ar << CulledVertexRadiusScaleBuffer;
-	}
-
-	virtual void GetElementShaderBindings(
+	void GetElementShaderBindings(
 		const FSceneInterface* Scene,
 		const FSceneView* View,
 		const FMeshMaterialShader* Shader,
@@ -112,7 +92,7 @@ public:
 		const FMeshBatchElement& BatchElement,
 		FMeshDrawSingleShaderBindings& ShaderBindings,
 		FVertexInputStreamArray& VertexStreams
-	) const override
+	) const
 	{
 		const FHairStrandsVertexFactory* VF = static_cast<const FHairStrandsVertexFactory*>(VertexFactory);
 
@@ -148,20 +128,22 @@ public:
 	}
 };
 
+IMPLEMENT_TYPE_LAYOUT(FHairStrandsVertexFactoryShaderParameters);
+
 /**
  * Should we cache the material's shadertype on this platform with this vertex factory? 
  */
-bool FHairStrandsVertexFactory::ShouldCompilePermutation(EShaderPlatform Platform, const class FMaterial* Material, const class FShaderType* ShaderType)
+bool FHairStrandsVertexFactory::ShouldCompilePermutation(const FVertexFactoryShaderPermutationParameters& Parameters)
 {
-	return (Material->GetMaterialDomain() == MD_Surface && Material->IsUsedWithHairStrands() && Platform == EShaderPlatform::SP_PCD3D_SM5) || Material->IsSpecialEngineMaterial();
+	return (Parameters.MaterialParameters.MaterialDomain == MD_Surface && Parameters.MaterialParameters.bIsUsedWithHairStrands && Parameters.Platform == EShaderPlatform::SP_PCD3D_SM5) || Parameters.MaterialParameters.bIsSpecialEngineMaterial;
 }
 
-void FHairStrandsVertexFactory::ModifyCompilationEnvironment(const FVertexFactoryType* Type, EShaderPlatform Platform, const FMaterial* Material, FShaderCompilerEnvironment& OutEnvironment)
+void FHairStrandsVertexFactory::ModifyCompilationEnvironment(const FVertexFactoryShaderPermutationParameters& Parameters, FShaderCompilerEnvironment& OutEnvironment)
 {
-	const bool bUseGPUSceneAndPrimitiveIdStream = Type->SupportsPrimitiveIdStream() && UseGPUScene(Platform, GetMaxSupportedFeatureLevel(Platform));
+	const bool bUseGPUSceneAndPrimitiveIdStream = Parameters.VertexFactoryType->SupportsPrimitiveIdStream() && UseGPUScene(Parameters.Platform, GetMaxSupportedFeatureLevel(Parameters.Platform));
 	OutEnvironment.SetDefine(TEXT("VF_SUPPORTS_PRIMITIVE_SCENE_DATA"), bUseGPUSceneAndPrimitiveIdStream);
 	OutEnvironment.SetDefine(TEXT("VF_STRAND_HAIR"), TEXT("1"));
-	OutEnvironment.SetDefine(TEXT("VF_GPU_SCENE_TEXTURE"), bUseGPUSceneAndPrimitiveIdStream && GPUSceneUseTexture2D(Platform));
+	OutEnvironment.SetDefine(TEXT("VF_GPU_SCENE_TEXTURE"), bUseGPUSceneAndPrimitiveIdStream && GPUSceneUseTexture2D(Parameters.Platform));
 }
 
 void FHairStrandsVertexFactory::ValidateCompiledResult(const FVertexFactoryType* Type, EShaderPlatform Platform, const FShaderParameterMap& ParameterMap, TArray<FString>& OutErrors)
diff --git a/Engine/Plugins/Experimental/HairStrands/Source/HairStrandsCore/Public/HairStrandsVertexFactory.h b/Engine/Plugins/Experimental/HairStrands/Source/HairStrandsCore/Public/HairStrandsVertexFactory.h
--- a/Engine/Plugins/Experimental/HairStrands/Source/HairStrandsCore/Public/HairStrandsVertexFactory.h
+++ b/Engine/Plugins/Experimental/HairStrands/Source/HairStrandsCore/Public/HairStrandsVertexFactory.h
@@ -39,8 +39,8 @@ public:
 	/**
 	 * Should we cache the material's shadertype on this platform with this vertex factory? 
 	 */
-	static bool ShouldCompilePermutation(EShaderPlatform Platform, const class FMaterial* Material, const class FShaderType* ShaderType);
-	static void ModifyCompilationEnvironment(const FVertexFactoryType* Type, EShaderPlatform Platform, const FMaterial* Material, FShaderCompilerEnvironment& OutEnvironment);
+	static bool ShouldCompilePermutation(const FVertexFactoryShaderPermutationParameters& Parameters);
+	static void ModifyCompilationEnvironment(const FVertexFactoryShaderPermutationParameters& Parameters, FShaderCompilerEnvironment& OutEnvironment);
 	static void ValidateCompiledResult(const FVertexFactoryType* Type, EShaderPlatform Platform, const FShaderParameterMap& ParameterMap, TArray<FString>& OutErrors);
 
 	/**
diff --git a/Engine/Plugins/Experimental/HairStrands/Source/HairStrandsNiagara/Private/NiagaraDataInterfaceHairStrands.cpp b/Engine/Plugins/Experimental/HairStrands/Source/HairStrandsNiagara/Private/NiagaraDataInterfaceHairStrands.cpp
--- a/Engine/Plugins/Experimental/HairStrands/Source/HairStrandsNiagara/Private/NiagaraDataInterfaceHairStrands.cpp
+++ b/Engine/Plugins/Experimental/HairStrands/Source/HairStrandsNiagara/Private/NiagaraDataInterfaceHairStrands.cpp
@@ -296,7 +296,7 @@ static void AddCopyBoundingBoxPass(
 	Parameters->OutNodeBoundBuffer = OutNodeBoundBuffer;
 	Parameters->NumElements = NumElements;
 
-	TShaderMap<FGlobalShaderType>* ShaderMap = GetGlobalShaderMap(ERHIFeatureLevel::SM5);
+	FGlobalShaderMap* ShaderMap = GetGlobalShaderMap(ERHIFeatureLevel::SM5);
 
 	const uint32 DispatchCount = FMath::DivideAndRoundUp(NumElements, GroupSize);
 
@@ -304,7 +304,7 @@ static void AddCopyBoundingBoxPass(
 	FComputeShaderUtils::AddPass(
 		GraphBuilder,
 		RDG_EVENT_NAME("CopyBoundingBox"),
-		*ComputeShader,
+		ComputeShader,
 		Parameters,
 		FIntVector(DispatchCount, 1, 1));
 }
@@ -334,9 +334,6 @@ public:
 	/** Initialization constructor. */
 	explicit FDirectCopyBoundingBoxCS(const ShaderMetaType::CompiledShaderInitializerType& Initializer);
 
-	/** Serialization. */
-	virtual bool Serialize(FArchive& Ar) override;
-
 	/**
 	 * Set parameters.
 	 */
@@ -368,22 +365,13 @@ FDirectCopyBoundingBoxCS::FDirectCopyBoundingBoxCS(const ShaderMetaType::Compile
 	OutputBoundingBox.Bind(Initializer.ParameterMap, TEXT("OutNodeBoundBuffer"));
 }
 
-bool FDirectCopyBoundingBoxCS::Serialize(FArchive& Ar)
-{
-	bool bShaderHasOutdatedParameters = FGlobalShader::Serialize(Ar);
-	Ar << NumElements;
-	Ar << InputBoundingBox;
-	Ar << OutputBoundingBox;
-	return bShaderHasOutdatedParameters;
-}
-
 void FDirectCopyBoundingBoxCS::SetParameters(
 	FRHICommandList& RHICmdList,
 	FRHIUnorderedAccessView* InInputBoundingBox,
 	FRHIUnorderedAccessView* InOutputBoundingBox,
 	int32 InNumElements)
 {
-	FRHIComputeShader* ComputeShaderRHI = GetComputeShader();
+	FRHIComputeShader* ComputeShaderRHI = RHICmdList.GetBoundComputeShader();
 
 	RHICmdList.TransitionResource(EResourceTransitionAccess::ERWBarrier, EResourceTransitionPipeline::EComputeToCompute, InInputBoundingBox);
 	SetUAVParameter(RHICmdList, ComputeShaderRHI, InputBoundingBox, InInputBoundingBox);
@@ -399,7 +387,7 @@ void FDirectCopyBoundingBoxCS::SetParameters(
 
 void FDirectCopyBoundingBoxCS::UnbindBuffers(FRHICommandList& RHICmdList)
 {
-	FRHIComputeShader* ComputeShaderRHI = GetComputeShader();
+	FRHIComputeShader* ComputeShaderRHI = RHICmdList.GetBoundComputeShader();
 	SetUAVParameter(RHICmdList, ComputeShaderRHI, InputBoundingBox, nullptr);
 	SetUAVParameter(RHICmdList, ComputeShaderRHI, OutputBoundingBox, nullptr);
 }
@@ -412,15 +400,15 @@ static void AddDirectCopyBoundingBoxPass(
 	const uint32 GroupSize = NIAGARA_HAIR_STRANDS_THREAD_COUNT;
 	const uint32 NumElements = 1;
 
-	TShaderMap<FGlobalShaderType>* ShaderMap = GetGlobalShaderMap(ERHIFeatureLevel::SM5);
+	FGlobalShaderMap* ShaderMap = GetGlobalShaderMap(ERHIFeatureLevel::SM5);
 	const uint32 DispatchCount = FMath::DivideAndRoundUp(NumElements, GroupSize);
 
 	TShaderMapRef<FDirectCopyBoundingBoxCS> ComputeShader(ShaderMap);
-	RHICmdList.SetComputeShader(ComputeShader->GetComputeShader());
+	RHICmdList.SetComputeShader(ComputeShader.GetComputeShader());
 
 	ComputeShader->SetParameters(RHICmdList, BoundingBoxBuffer, OutNodeBoundBuffer, NumElements);
 
-	DispatchComputeShader(RHICmdList, *ComputeShader, DispatchCount, 1, 1);
+	DispatchComputeShader(RHICmdList, ComputeShader, DispatchCount, 1, 1);
 	ComputeShader->UnbindBuffers(RHICmdList);
 }
 
@@ -623,9 +611,11 @@ bool FNDIHairStrandsData::Init(UNiagaraDataInterfaceHairStrands* Interface, FNia
 
 struct FNDIHairStrandsParametersCS : public FNiagaraDataInterfaceParametersCS
 {
-	virtual void Bind(const FNiagaraDataInterfaceParamRef& ParamRef, const class FShaderParameterMap& ParameterMap) override
+	DECLARE_TYPE_LAYOUT(FNDIHairStrandsParametersCS, NonVirtual);
+public:
+	void Bind(const FNiagaraDataInterfaceGPUParamInfo& ParameterInfo, const class FShaderParameterMap& ParameterMap)
 	{
-		FNDIHairStrandsParametersName ParamNames(ParamRef.ParameterInfo.DataInterfaceHLSLSymbol);
+		FNDIHairStrandsParametersName ParamNames(*ParameterInfo.DataInterfaceHLSLSymbol);
 
 		WorldTransform.Bind(ParameterMap, *ParamNames.WorldTransformName);
 		WorldInverse.Bind(ParameterMap, *ParamNames.WorldInverseName);
@@ -713,47 +703,11 @@ struct FNDIHairStrandsParametersCS : public FNiagaraDataInterfaceParametersCS
 		}
 	}
 
-	virtual void Serialize(FArchive& Ar) override
-	{
-		Ar << WorldTransform;
-		Ar << WorldInverse;
-		Ar << WorldRotation;
-		Ar << NumStrands;
-		Ar << StrandSize;
-		Ar << BoxCenter;
-		Ar << BoxExtent;
-
-		Ar << DeformedPositionBuffer;
-		Ar << CurvesOffsetsBuffer;
-		Ar << RestPositionBuffer;
-
-		Ar << ResetSimulation;
-		Ar << HasRootAttached;
-		Ar << RestRootOffset;
-		Ar << DeformedRootOffset;
-
-		Ar << RootBarycentricCoordinatesBuffer;
-
-		Ar << RestTrianglePositionABuffer;
-		Ar << RestTrianglePositionBBuffer;
-		Ar << RestTrianglePositionCBuffer;
-
-		Ar << DeformedTrianglePositionABuffer;
-		Ar << DeformedTrianglePositionBBuffer;
-		Ar << DeformedTrianglePositionCBuffer;
-
-		Ar << RestPositionOffset;
-		Ar << DeformedPositionOffset;
-
-		Ar << BoundingBoxBuffer;
-		Ar << NodeBoundBuffer;
-	}
-
-	virtual void Set(FRHICommandList& RHICmdList, const FNiagaraDataInterfaceSetArgs& Context) const override
+	void Set(FRHICommandList& RHICmdList, const FNiagaraDataInterfaceSetArgs& Context) const
 	{
-		check(IsInRenderingThread()); 
+		check(IsInRenderingThread());
 
-		FRHIComputeShader* ComputeShaderRHI = Context.Shader->GetComputeShader();
+		FRHIComputeShader* ComputeShaderRHI = RHICmdList.GetBoundComputeShader();
 
 		FNDIHairStrandsProxy* InterfaceProxy =
 			static_cast<FNDIHairStrandsProxy*>(Context.DataInterface);
@@ -871,48 +825,53 @@ struct FNDIHairStrandsParametersCS : public FNiagaraDataInterfaceParametersCS
 		}
 	}
 
-	virtual void Unset(FRHICommandList& RHICmdList, const FNiagaraDataInterfaceSetArgs& Context) const override
+	void Unset(FRHICommandList& RHICmdList, const FNiagaraDataInterfaceSetArgs& Context) const
 	{
-		SetUAVParameter(RHICmdList, Context.Shader->GetComputeShader(), DeformedPositionBuffer, nullptr);
-		SetUAVParameter(RHICmdList, Context.Shader->GetComputeShader(), BoundingBoxBuffer, nullptr);
+		FRHIComputeShader* ShaderRHI = RHICmdList.GetBoundComputeShader();
+		SetUAVParameter(RHICmdList, ShaderRHI, DeformedPositionBuffer, nullptr);
+		SetUAVParameter(RHICmdList, ShaderRHI, BoundingBoxBuffer, nullptr);
 	}
 
 private:
 
-	FShaderParameter WorldTransform;
-	FShaderParameter WorldInverse;
-	FShaderParameter WorldRotation;
-	FShaderParameter NumStrands;
-	FShaderParameter StrandSize;
-	FShaderParameter BoxCenter;
-	FShaderParameter BoxExtent;
+	LAYOUT_FIELD(FShaderParameter, WorldTransform);
+	LAYOUT_FIELD(FShaderParameter, WorldInverse);
+	LAYOUT_FIELD(FShaderParameter, WorldRotation);
+	LAYOUT_FIELD(FShaderParameter, NumStrands);
+	LAYOUT_FIELD(FShaderParameter, StrandSize);
+	LAYOUT_FIELD(FShaderParameter, BoxCenter);
+	LAYOUT_FIELD(FShaderParameter, BoxExtent);
 
-	FShaderResourceParameter DeformedPositionBuffer;
-	FShaderResourceParameter CurvesOffsetsBuffer;
-	FShaderResourceParameter RestPositionBuffer;
+	LAYOUT_FIELD(FShaderResourceParameter, DeformedPositionBuffer);
+	LAYOUT_FIELD(FShaderResourceParameter, CurvesOffsetsBuffer);
+	LAYOUT_FIELD(FShaderResourceParameter, RestPositionBuffer);
 
-	FShaderParameter ResetSimulation;
-	FShaderParameter HasRootAttached;
-	FShaderParameter RestRootOffset;
-	FShaderParameter DeformedRootOffset;
+	LAYOUT_FIELD(FShaderParameter, ResetSimulation);
+	LAYOUT_FIELD(FShaderParameter, HasRootAttached);
+	LAYOUT_FIELD(FShaderParameter, RestRootOffset);
+	LAYOUT_FIELD(FShaderParameter, DeformedRootOffset);
 
-	FShaderResourceParameter RootBarycentricCoordinatesBuffer;
+	LAYOUT_FIELD(FShaderResourceParameter, RootBarycentricCoordinatesBuffer);
 
-	FShaderResourceParameter RestTrianglePositionABuffer;
-	FShaderResourceParameter RestTrianglePositionBBuffer;
-	FShaderResourceParameter RestTrianglePositionCBuffer;
+	LAYOUT_FIELD(FShaderResourceParameter, RestTrianglePositionABuffer);
+	LAYOUT_FIELD(FShaderResourceParameter, RestTrianglePositionBBuffer);
+	LAYOUT_FIELD(FShaderResourceParameter, RestTrianglePositionCBuffer);
 
-	FShaderResourceParameter DeformedTrianglePositionABuffer;
-	FShaderResourceParameter DeformedTrianglePositionBBuffer;
-	FShaderResourceParameter DeformedTrianglePositionCBuffer;
+	LAYOUT_FIELD(FShaderResourceParameter, DeformedTrianglePositionABuffer);
+	LAYOUT_FIELD(FShaderResourceParameter, DeformedTrianglePositionBBuffer);
+	LAYOUT_FIELD(FShaderResourceParameter, DeformedTrianglePositionCBuffer);
 
-	FShaderParameter RestPositionOffset;
-	FShaderParameter DeformedPositionOffset;
+	LAYOUT_FIELD(FShaderParameter, RestPositionOffset);
+	LAYOUT_FIELD(FShaderParameter, DeformedPositionOffset);
 
-	FShaderResourceParameter BoundingBoxBuffer;
-	FShaderResourceParameter NodeBoundBuffer;
+	LAYOUT_FIELD(FShaderResourceParameter, BoundingBoxBuffer);
+	LAYOUT_FIELD(FShaderResourceParameter, NodeBoundBuffer);
 };
 
+IMPLEMENT_TYPE_LAYOUT(FNDIHairStrandsParametersCS);
+
+IMPLEMENT_NIAGARA_DI_PARAMETER(UNiagaraDataInterfaceHairStrands, FNDIHairStrandsParametersCS);
+
 //------------------------------------------------------------------------------------------------------------
 
 void FNDIHairStrandsProxy::ConsumePerInstanceDataFromGameThread(void* PerInstanceData, const FNiagaraSystemInstanceID& Instance)
@@ -3398,7 +3357,7 @@ bool UNiagaraDataInterfaceHairStrands::GetFunctionHLSL(const FNiagaraDataInterfa
 				)");
 		OutHLSL += FString::Format(FormatSample, ArgsSample);
 		return true;
-	}
+	}	
 	else if (FunctionInfo.DefinitionName == UpdatePointPositionName)
 	{
 		static const TCHAR *FormatSample = TEXT(R"(
@@ -3839,12 +3798,6 @@ void UNiagaraDataInterfaceHairStrands::ProvidePerInstanceDataForRenderThread(voi
 	}
 }
 
-FNiagaraDataInterfaceParametersCS*
-UNiagaraDataInterfaceHairStrands::ConstructComputeParameters() const
-{
-	return new FNDIHairStrandsParametersCS();
-}
-
 void FNDIHairStrandsProxy::PreStage(FRHICommandList& RHICmdList, const FNiagaraDataInterfaceSetArgs& Context)
 {
 }
diff --git a/Engine/Plugins/Experimental/HairStrands/Source/HairStrandsNiagara/Private/NiagaraDataInterfacePhysicsAsset.cpp b/Engine/Plugins/Experimental/HairStrands/Source/HairStrandsNiagara/Private/NiagaraDataInterfacePhysicsAsset.cpp
--- a/Engine/Plugins/Experimental/HairStrands/Source/HairStrandsNiagara/Private/NiagaraDataInterfacePhysicsAsset.cpp
+++ b/Engine/Plugins/Experimental/HairStrands/Source/HairStrandsNiagara/Private/NiagaraDataInterfacePhysicsAsset.cpp
@@ -47,24 +47,7 @@ const FString UNiagaraDataInterfacePhysicsAsset::BoxExtentName(TEXT("BoxExtent_"
 
 struct FNDIPhysicsAssetParametersName
 {
-	FNDIPhysicsAssetParametersName(const FString& Suffix)
-	{
-		ElementOffsetsName = UNiagaraDataInterfacePhysicsAsset::ElementOffsetsName + Suffix;
-
-		CurrentTransformBufferName = UNiagaraDataInterfacePhysicsAsset::CurrentTransformBufferName + Suffix;
-		PreviousTransformBufferName = UNiagaraDataInterfacePhysicsAsset::PreviousTransformBufferName + Suffix;
-		PreviousInverseBufferName = UNiagaraDataInterfacePhysicsAsset::PreviousInverseBufferName + Suffix;
-		InverseTransformBufferName = UNiagaraDataInterfacePhysicsAsset::InverseTransformBufferName + Suffix;
-		RestTransformBufferName = UNiagaraDataInterfacePhysicsAsset::RestTransformBufferName + Suffix;
-		RestInverseBufferName = UNiagaraDataInterfacePhysicsAsset::RestInverseBufferName + Suffix;
-		ElementExtentBufferName = UNiagaraDataInterfacePhysicsAsset::ElementExtentBufferName + Suffix;
-
-		BoxOriginName = UNiagaraDataInterfacePhysicsAsset::BoxOriginName + Suffix;
-		BoxExtentName = UNiagaraDataInterfacePhysicsAsset::BoxExtentName + Suffix;
-	}
-
 	FString ElementOffsetsName;
-	
 	FString CurrentTransformBufferName;
 	FString PreviousTransformBufferName;
 	FString PreviousInverseBufferName;
@@ -77,6 +60,20 @@ struct FNDIPhysicsAssetParametersName
 	FString BoxExtentName;
 };
 
+static void GetNiagaraDataInterfaceParametersName(FNDIPhysicsAssetParametersName& Names, const FString& Suffix)
+{
+	Names.ElementOffsetsName = UNiagaraDataInterfacePhysicsAsset::ElementOffsetsName + Suffix;
+	Names.CurrentTransformBufferName = UNiagaraDataInterfacePhysicsAsset::CurrentTransformBufferName + Suffix;
+	Names.PreviousTransformBufferName = UNiagaraDataInterfacePhysicsAsset::PreviousTransformBufferName + Suffix;
+	Names.InverseTransformBufferName = UNiagaraDataInterfacePhysicsAsset::InverseTransformBufferName + Suffix;
+	Names.RestTransformBufferName = UNiagaraDataInterfacePhysicsAsset::RestTransformBufferName + Suffix;
+	Names.RestInverseBufferName = UNiagaraDataInterfacePhysicsAsset::RestInverseBufferName + Suffix;
+	Names.ElementExtentBufferName = UNiagaraDataInterfacePhysicsAsset::ElementExtentBufferName + Suffix;
+
+	Names.BoxOriginName = UNiagaraDataInterfacePhysicsAsset::BoxOriginName + Suffix;
+	Names.BoxExtentName = UNiagaraDataInterfacePhysicsAsset::BoxExtentName + Suffix;
+}
+
 //------------------------------------------------------------------------------------------------------------
 
 template<typename BufferType, typename DataType, int ElementSize, EPixelFormat PixelFormat, bool InitBuffer>
@@ -418,9 +415,12 @@ bool FNDIPhysicsAssetData::Init(UNiagaraDataInterfacePhysicsAsset* Interface, FN
 
 struct FNDIPhysicsAssetParametersCS : public FNiagaraDataInterfaceParametersCS
 {
-	virtual void Bind(const FNiagaraDataInterfaceParamRef& ParamRef, const class FShaderParameterMap& ParameterMap) override
+	DECLARE_TYPE_LAYOUT(FNDIPhysicsAssetParametersCS, NonVirtual);
+public:
+	void Bind(const FNiagaraDataInterfaceGPUParamInfo& ParameterInfo, const class FShaderParameterMap& ParameterMap)
 	{
-		FNDIPhysicsAssetParametersName ParamNames(ParamRef.ParameterInfo.DataInterfaceHLSLSymbol);
+		FNDIPhysicsAssetParametersName ParamNames;
+		GetNiagaraDataInterfaceParametersName(ParamNames, *ParameterInfo.DataInterfaceHLSLSymbol);
 
 		ElementOffsets.Bind(ParameterMap, *ParamNames.ElementOffsetsName);
 
@@ -465,25 +465,11 @@ struct FNDIPhysicsAssetParametersCS : public FNiagaraDataInterfaceParametersCS
 		}
 	}
 
-	virtual void Serialize(FArchive& Ar) override
-	{
-		Ar << ElementOffsets;
-		Ar << CurrentTransformBuffer;
-		Ar << PreviousTransformBuffer;
-		Ar << PreviousInverseBuffer;
-		Ar << InverseTransformBuffer;
-		Ar << RestTransformBuffer;
-		Ar << RestInverseBuffer;
-		Ar << ElementExtentBuffer;
-		Ar << BoxOrigin;
-		Ar << BoxExtent;
-	}
-
-	virtual void Set(FRHICommandList& RHICmdList, const FNiagaraDataInterfaceSetArgs& Context) const override
+	void Set(FRHICommandList& RHICmdList, const FNiagaraDataInterfaceSetArgs& Context) const
 	{
 		check(IsInRenderingThread());
 
-		FRHIComputeShader* ComputeShaderRHI = Context.Shader->GetComputeShader();
+		FRHIComputeShader* ComputeShaderRHI = RHICmdList.GetBoundComputeShader();
 
 		FNDIPhysicsAssetProxy* InterfaceProxy =
 			static_cast<FNDIPhysicsAssetProxy*>(Context.DataInterface);
@@ -522,26 +508,30 @@ struct FNDIPhysicsAssetParametersCS : public FNiagaraDataInterfaceParametersCS
 		}
 	}
 
-	virtual void Unset(FRHICommandList& RHICmdList, const FNiagaraDataInterfaceSetArgs& Context) const override
+	void Unset(FRHICommandList& RHICmdList, const FNiagaraDataInterfaceSetArgs& Context) const
 	{
 	}
 
 private:
 
-	FShaderParameter ElementOffsets;
+	LAYOUT_FIELD(FShaderParameter, ElementOffsets);
 
-	FShaderResourceParameter CurrentTransformBuffer;
-	FShaderResourceParameter PreviousTransformBuffer;
-	FShaderResourceParameter PreviousInverseBuffer;
-	FShaderResourceParameter InverseTransformBuffer;
-	FShaderResourceParameter RestTransformBuffer;
-	FShaderResourceParameter RestInverseBuffer;
-	FShaderResourceParameter ElementExtentBuffer;
+	LAYOUT_FIELD(FShaderResourceParameter, CurrentTransformBuffer);
+	LAYOUT_FIELD(FShaderResourceParameter, PreviousTransformBuffer);
+	LAYOUT_FIELD(FShaderResourceParameter, PreviousInverseBuffer);
+	LAYOUT_FIELD(FShaderResourceParameter, InverseTransformBuffer);
+	LAYOUT_FIELD(FShaderResourceParameter, RestTransformBuffer);
+	LAYOUT_FIELD(FShaderResourceParameter, RestInverseBuffer);
+	LAYOUT_FIELD(FShaderResourceParameter, ElementExtentBuffer);
 
-	FShaderParameter BoxOrigin;
-	FShaderParameter BoxExtent;
+	LAYOUT_FIELD(FShaderParameter, BoxOrigin);
+	LAYOUT_FIELD(FShaderParameter, BoxExtent);
 };
 
+IMPLEMENT_TYPE_LAYOUT(FNDIPhysicsAssetParametersCS);
+
+IMPLEMENT_NIAGARA_DI_PARAMETER(UNiagaraDataInterfacePhysicsAsset, FNDIPhysicsAssetParametersCS);
+
 //------------------------------------------------------------------------------------------------------------
 
 void FNDIPhysicsAssetProxy::ConsumePerInstanceDataFromGameThread(void* PerInstanceData, const FNiagaraSystemInstanceID& Instance)
@@ -889,7 +879,8 @@ void UNiagaraDataInterfacePhysicsAsset::GetProjectionPoint(FVectorVMContext& Con
 
 bool UNiagaraDataInterfacePhysicsAsset::GetFunctionHLSL(const FNiagaraDataInterfaceGPUParamInfo& ParamInfo, const FNiagaraDataInterfaceGeneratedFunction& FunctionInfo, int FunctionInstanceIndex, FString& OutHLSL)
 {
-	FNDIPhysicsAssetParametersName ParamNames(ParamInfo.DataInterfaceHLSLSymbol);
+	FNDIPhysicsAssetParametersName ParamNames;
+	GetNiagaraDataInterfaceParametersName(ParamNames, ParamInfo.DataInterfaceHLSLSymbol);
 
 	TMap<FString, FStringFormatArg> ArgsSample = {
 		{TEXT("InstanceFunctionName"), FunctionInfo.InstanceName},
@@ -1004,76 +995,4 @@ void UNiagaraDataInterfacePhysicsAsset::ProvidePerInstanceDataForRenderThread(vo
 	check(Proxy);
 }
 
-FNiagaraDataInterfaceParametersCS*
-UNiagaraDataInterfacePhysicsAsset::ConstructComputeParameters() const
-{
-	return new FNDIPhysicsAssetParametersCS();
-}
-
-//------------------------------------------------------------------------------------------------------------
-//
-//class FResetCS : public FGlobalShader
-//{
-//	DECLARE_GLOBAL_SHADER(FCopyBoundingBoxCS);
-//	SHADER_USE_PARAMETER_STRUCT(FCopyBoundingBoxCS, FGlobalShader);
-//
-//	BEGIN_SHADER_PARAMETER_STRUCT(FParameters, )
-//		SHADER_PARAMETER(uint32, NumElements)
-//		SHADER_PARAMETER_UAV(RWBuffer, BoundingBoxBuffer)
-//		SHADER_PARAMETER_UAV(RWBuffer, OutNodeBoundBuffer)
-//		END_SHADER_PARAMETER_STRUCT()
-//
-//public:
-//	static bool ShouldCompilePermutation(const FGlobalShaderPermutationParameters& Parameters)
-//	{
-//		return RHISupportsComputeShaders(Parameters.Platform);
-//	}
-//
-//	static void ModifyCompilationEnvironment(const FGlobalShaderPermutationParameters& Parameters, FShaderCompilerEnvironment& OutEnvironment)
-//	{
-//		FGlobalShader::ModifyCompilationEnvironment(Parameters, OutEnvironment);
-//		OutEnvironment.SetDefine(TEXT("THREAD_COUNT"), NIAGARA_HAIR_STRANDS_THREAD_COUNT);
-//	}
-//};
-//
-//IMPLEMENT_GLOBAL_SHADER(FCopyBoundingBoxCS, "/Plugin/Experimental/HairStrands/Private/NiagaraCopyBoundingBox.usf", "MainCS", SF_Compute);
-//
-//static void AddSkinCacheRasterPass(
-//	FRDGBuilder& GraphBuilder,
-//	FRHIUnorderedAccessView* BoundingBoxBuffer,
-//	FRHIUnorderedAccessView* OutNodeBoundBuffer)
-//{
-//	const uint32 GroupSize = NIAGARA_HAIR_STRANDS_THREAD_COUNT;
-//	const uint32 NumElements = 1;
-//
-//	FCopyBoundingBoxCS::FParameters* Parameters = GraphBuilder.AllocParameters<FCopyBoundingBoxCS::FParameters>();
-//	Parameters->BoundingBoxBuffer = BoundingBoxBuffer;
-//	Parameters->OutNodeBoundBuffer = OutNodeBoundBuffer;
-//	Parameters->NumElements = NumElements;
-//
-//	TShaderMap<FGlobalShaderType>* ShaderMap = GetGlobalShaderMap(ERHIFeatureLevel::SM5);
-//
-//	const uint32 DispatchCount = FMath::DivideAndRoundUp(NumElements, GroupSize);
-//
-//	TShaderMapRef<FCopyBoundingBoxCS> ComputeShader(ShaderMap);
-//	FComputeShaderUtils::AddPass(
-//		GraphBuilder,
-//		RDG_EVENT_NAME("CopyBoundingBox"),
-//		*ComputeShader,
-//		Parameters,
-//		FIntVector(DispatchCount, 1, 1));
-//
-//	GraphBuilder.AddPass(
-//		RDG_EVENT_NAME("SkinCacheRaster"),
-//		PassParameters,
-//		ERDGPassFlags::Raster,
-//		[PassParameters, Scene = Scene, ViewInfo, RasterPassType, &PrimitiveSceneInfos, ViewportRect, HairRenderInfo, RasterDirection](FRHICommandListImmediate& RHICmdList)
-//	{
-//	});
-//}
-
-//------------------------------------------------------------------------------------------------------------
-
-
-
 #undef LOCTEXT_NAMESPACE
\ No newline at end of file
diff --git a/Engine/Plugins/Experimental/HairStrands/Source/HairStrandsNiagara/Private/NiagaraDataInterfacePressureGrid.cpp b/Engine/Plugins/Experimental/HairStrands/Source/HairStrandsNiagara/Private/NiagaraDataInterfacePressureGrid.cpp
--- a/Engine/Plugins/Experimental/HairStrands/Source/HairStrandsNiagara/Private/NiagaraDataInterfacePressureGrid.cpp
+++ b/Engine/Plugins/Experimental/HairStrands/Source/HairStrandsNiagara/Private/NiagaraDataInterfacePressureGrid.cpp
@@ -46,18 +46,6 @@ const FString UNiagaraDataInterfacePressureGrid::WorldInverseName(TEXT("WorldInv
 
 struct FNDIPressureGridParametersName
 {
-	FNDIPressureGridParametersName(const FString& Suffix)
-	{
-		GridCurrentBufferName = UNiagaraDataInterfacePressureGrid::GridCurrentBufferName + Suffix;
-		GridDestinationBufferName = UNiagaraDataInterfacePressureGrid::GridDestinationBufferName + Suffix;
-
-		GridSizeName = UNiagaraDataInterfacePressureGrid::GridSizeName + Suffix;
-		GridOriginName = UNiagaraDataInterfacePressureGrid::GridOriginName + Suffix;
-
-		WorldTransformName = UNiagaraDataInterfacePressureGrid::WorldTransformName + Suffix;
-		WorldInverseName = UNiagaraDataInterfacePressureGrid::WorldInverseName + Suffix;
-	}
-
 	FString GridCurrentBufferName;
 	FString GridDestinationBufferName;
 
@@ -68,6 +56,16 @@ struct FNDIPressureGridParametersName
 	FString WorldInverseName;
 };
 
+static void GetNiagaraDataInterfaceParametersName(FNDIPressureGridParametersName& Names, const FString& Suffix)
+{
+	Names.GridCurrentBufferName = UNiagaraDataInterfacePressureGrid::GridCurrentBufferName + Suffix;
+	Names.GridDestinationBufferName = UNiagaraDataInterfacePressureGrid::GridDestinationBufferName + Suffix;
+	Names.GridSizeName = UNiagaraDataInterfacePressureGrid::GridSizeName + Suffix;
+	Names.GridOriginName = UNiagaraDataInterfacePressureGrid::GridOriginName + Suffix;
+	Names.WorldTransformName = UNiagaraDataInterfacePressureGrid::WorldTransformName + Suffix;
+	Names.WorldInverseName = UNiagaraDataInterfacePressureGrid::WorldInverseName + Suffix;
+}
+
 //------------------------------------------------------------------------------------------------------------
 
 void FNDIPressureGridBuffer::Initialize(const FIntVector InGridSize)
@@ -152,9 +150,12 @@ bool FNDIPressureGridData::Init(UNiagaraDataInterfacePressureGrid* Interface, FN
 
 struct FNDIPressureGridParametersCS : public FNiagaraDataInterfaceParametersCS
 {
-	virtual void Bind(const FNiagaraDataInterfaceParamRef& ParamRef, const class FShaderParameterMap& ParameterMap) override
+	DECLARE_TYPE_LAYOUT(FNDIPressureGridParametersCS, NonVirtual);
+
+	void Bind(const FNiagaraDataInterfaceGPUParamInfo& ParameterInfo, const class FShaderParameterMap& ParameterMap)
 	{
-		FNDIPressureGridParametersName ParamNames(ParamRef.ParameterInfo.DataInterfaceHLSLSymbol);
+		FNDIPressureGridParametersName ParamNames;
+		GetNiagaraDataInterfaceParametersName(ParamNames, *ParameterInfo.DataInterfaceHLSLSymbol);
 
 		GridCurrentBuffer.Bind(ParameterMap, *ParamNames.GridCurrentBufferName);
 		GridDestinationBuffer.Bind(ParameterMap, *ParamNames.GridDestinationBufferName);
@@ -176,23 +177,11 @@ struct FNDIPressureGridParametersCS : public FNiagaraDataInterfaceParametersCS
 		}
 	}
 
-	virtual void Serialize(FArchive& Ar) override
-	{
-		Ar << GridCurrentBuffer;
-		Ar << GridDestinationBuffer;
-
-		Ar << GridOrigin;
-		Ar << GridSize;
-
-		Ar << WorldTransform;
-		Ar << WorldInverse;
-	}
-
-	virtual void Set(FRHICommandList& RHICmdList, const FNiagaraDataInterfaceSetArgs& Context) const override
+	void Set(FRHICommandList& RHICmdList, const FNiagaraDataInterfaceSetArgs& Context) const
 	{
 		check(IsInRenderingThread());
 
-		FRHIComputeShader* ComputeShaderRHI = Context.Shader->GetComputeShader();
+		FRHIComputeShader* ComputeShaderRHI = RHICmdList.GetBoundComputeShader();
 
 		FNDIPressureGridProxy* InterfaceProxy =
 			static_cast<FNDIPressureGridProxy*>(Context.DataInterface);
@@ -231,23 +220,27 @@ struct FNDIPressureGridParametersCS : public FNiagaraDataInterfaceParametersCS
 		}
 	}
 
-	virtual void Unset(FRHICommandList& RHICmdList, const FNiagaraDataInterfaceSetArgs& Context) const override
+	void Unset(FRHICommandList& RHICmdList, const FNiagaraDataInterfaceSetArgs& Context) const
 	{
-		SetUAVParameter(RHICmdList, Context.Shader->GetComputeShader(), GridDestinationBuffer, nullptr);
+		FRHIComputeShader* ShaderRHI = RHICmdList.GetBoundComputeShader();
+		SetUAVParameter(RHICmdList, ShaderRHI, GridDestinationBuffer, nullptr);
 	}
 
 private:
 
-	FShaderResourceParameter GridCurrentBuffer;
-	FShaderResourceParameter GridDestinationBuffer;
+	LAYOUT_FIELD(FShaderResourceParameter, GridCurrentBuffer);
+	LAYOUT_FIELD(FShaderResourceParameter, GridDestinationBuffer);
 
-	FShaderParameter GridSize;
-	FShaderParameter GridOrigin;
+	LAYOUT_FIELD(FShaderParameter, GridSize);
+	LAYOUT_FIELD(FShaderParameter, GridOrigin);
 
-	FShaderParameter WorldTransform;
-	FShaderParameter WorldInverse;
+	LAYOUT_FIELD(FShaderParameter, WorldTransform);
+	LAYOUT_FIELD(FShaderParameter, WorldInverse);
 };
 
+IMPLEMENT_TYPE_LAYOUT(FNDIPressureGridParametersCS);
+
+IMPLEMENT_NIAGARA_DI_PARAMETER(UNiagaraDataInterfacePressureGrid, FNDIPressureGridParametersCS);
 
 //------------------------------------------------------------------------------------------------------------
 
@@ -726,7 +719,8 @@ void UNiagaraDataInterfacePressureGrid::UpdateGridTransform(FVectorVMContext& Co
 
 bool UNiagaraDataInterfacePressureGrid::GetFunctionHLSL(const FNiagaraDataInterfaceGPUParamInfo& ParamInfo, const FNiagaraDataInterfaceGeneratedFunction& FunctionInfo, int FunctionInstanceIndex, FString& OutHLSL)
 {
-	FNDIPressureGridParametersName ParamNames(ParamInfo.DataInterfaceHLSLSymbol);
+	FNDIPressureGridParametersName ParamNames;
+	GetNiagaraDataInterfaceParametersName(ParamNames, ParamInfo.DataInterfaceHLSLSymbol);
 
 	TMap<FString, FStringFormatArg> ArgsSample = {
 		{TEXT("InstanceFunctionName"), FunctionInfo.InstanceName},
@@ -878,12 +872,6 @@ void UNiagaraDataInterfacePressureGrid::ProvidePerInstanceDataForRenderThread(vo
 	RenderThreadData->GridSize = GameThreadData->GridSize;
 }
 
-FNiagaraDataInterfaceParametersCS*
-UNiagaraDataInterfacePressureGrid::ConstructComputeParameters() const
-{
-	return new FNDIPressureGridParametersCS();
-}
-
 //------------------------------------------------------------------------------------------------------------
 
 #define NIAGARA_HAIR_STRANDS_THREAD_COUNT 64
@@ -930,7 +918,7 @@ static void AddClearPressureGridPass(
 	Parameters->GridSize = GridSize;
 	Parameters->CopyPressure = CopyPressure;
 
-	TShaderMap<FGlobalShaderType>* ShaderMap = GetGlobalShaderMap(ERHIFeatureLevel::SM5);
+	FGlobalShaderMap* ShaderMap = GetGlobalShaderMap(ERHIFeatureLevel::SM5);
 
 	const uint32 DispatchCount = FMath::DivideAndRoundUp(NumElements, GroupSize);
 
@@ -938,7 +926,7 @@ static void AddClearPressureGridPass(
 	FComputeShaderUtils::AddPass(
 		GraphBuilder,
 		RDG_EVENT_NAME("ClearPressureGrid"),
-		*ComputeShader,
+		ComputeShader,
 		Parameters,
 		FIntVector(DispatchCount, 1, 1));
 }
diff --git a/Engine/Plugins/Experimental/HairStrands/Source/HairStrandsNiagara/Public/NiagaraDataInterfaceHairStrands.h b/Engine/Plugins/Experimental/HairStrands/Source/HairStrandsNiagara/Public/NiagaraDataInterfaceHairStrands.h
--- a/Engine/Plugins/Experimental/HairStrands/Source/HairStrandsNiagara/Public/NiagaraDataInterfaceHairStrands.h
+++ b/Engine/Plugins/Experimental/HairStrands/Source/HairStrandsNiagara/Public/NiagaraDataInterfaceHairStrands.h
@@ -273,6 +273,8 @@ class HAIRSTRANDSNIAGARA_API UNiagaraDataInterfaceHairStrands : public UNiagaraD
 
 public:
 
+	DECLARE_NIAGARA_DI_PARAMETER();
+
 	/** Hair Strands Asset used to sample from when not overridden by a source actor from the scene. Also useful for previewing in the editor. */
 	UPROPERTY(EditAnywhere, Category = "Source")
 	UGroomAsset* DefaultSource;
@@ -300,7 +302,6 @@ public:
 	/** GPU simulation  functionality */
 	virtual void GetParameterDefinitionHLSL(const FNiagaraDataInterfaceGPUParamInfo& ParamInfo, FString& OutHLSL) override;
 	virtual bool GetFunctionHLSL(const FNiagaraDataInterfaceGPUParamInfo& ParamInfo, const FNiagaraDataInterfaceGeneratedFunction& FunctionInfo, int FunctionInstanceIndex, FString& OutHLSL) override;
-	virtual FNiagaraDataInterfaceParametersCS* ConstructComputeParameters() const override;
 	virtual void ProvidePerInstanceDataForRenderThread(void* DataForRenderThread, void* PerInstanceData, const FNiagaraSystemInstanceID& SystemInstance) override;
 	virtual void GetCommonHLSL(FString& OutHLSL) override;
 
diff --git a/Engine/Plugins/Experimental/HairStrands/Source/HairStrandsNiagara/Public/NiagaraDataInterfacePhysicsAsset.h b/Engine/Plugins/Experimental/HairStrands/Source/HairStrandsNiagara/Public/NiagaraDataInterfacePhysicsAsset.h
--- a/Engine/Plugins/Experimental/HairStrands/Source/HairStrandsNiagara/Public/NiagaraDataInterfacePhysicsAsset.h
+++ b/Engine/Plugins/Experimental/HairStrands/Source/HairStrandsNiagara/Public/NiagaraDataInterfacePhysicsAsset.h
@@ -119,6 +119,8 @@ class HAIRSTRANDSNIAGARA_API UNiagaraDataInterfacePhysicsAsset : public UNiagara
 
 public:
 
+	DECLARE_NIAGARA_DI_PARAMETER();
+
 	/** Skeletal Mesh from which the Physics Asset will be found. */
 	UPROPERTY(EditAnywhere, Category = "Source")
 	UPhysicsAsset* DefaultSource;
@@ -149,7 +151,6 @@ public:
 	/** GPU simulation  functionality */
 	virtual void GetParameterDefinitionHLSL(const FNiagaraDataInterfaceGPUParamInfo& ParamInfo, FString& OutHLSL) override;
 	virtual bool GetFunctionHLSL(const FNiagaraDataInterfaceGPUParamInfo& ParamInfo, const FNiagaraDataInterfaceGeneratedFunction& FunctionInfo, int FunctionInstanceIndex, FString& OutHLSL) override;
-	virtual FNiagaraDataInterfaceParametersCS* ConstructComputeParameters() const override;
 	virtual void ProvidePerInstanceDataForRenderThread(void* DataForRenderThread, void* PerInstanceData, const FNiagaraSystemInstanceID& SystemInstance) override;
 	virtual void GetCommonHLSL(FString& OutHLSL) override;
 
diff --git a/Engine/Plugins/Experimental/RemoteSession/Source/RemoteSession/Private/Channels/RemoteSessionARCameraChannel.cpp b/Engine/Plugins/Experimental/RemoteSession/Source/RemoteSession/Private/Channels/RemoteSessionARCameraChannel.cpp
--- a/Engine/Plugins/Experimental/RemoteSession/Source/RemoteSession/Private/Channels/RemoteSessionARCameraChannel.cpp
+++ b/Engine/Plugins/Experimental/RemoteSession/Source/RemoteSession/Private/Channels/RemoteSessionARCameraChannel.cpp
@@ -67,15 +67,15 @@ public:
 
 	static bool ShouldCompilePermutation(const FMaterialShaderPermutationParameters& Parameters)
 	{
-		return Parameters.Material->GetMaterialDomain() == MD_PostProcess && !IsMobilePlatform(Parameters.Platform);
+		return Parameters.MaterialParameters.MaterialDomain == MD_PostProcess && !IsMobilePlatform(Parameters.Platform);
 	}
 
 	static void ModifyCompilationEnvironment(const FMaterialShaderPermutationParameters& Parameters, FShaderCompilerEnvironment& OutEnvironment)
 	{
-		FMaterialShader::ModifyCompilationEnvironment(Parameters.Platform, OutEnvironment);
+		FMaterialShader::ModifyCompilationEnvironment(Parameters, OutEnvironment);
 		OutEnvironment.SetDefine(TEXT("POST_PROCESS_MATERIAL"), 1);
 		OutEnvironment.SetDefine(TEXT("POST_PROCESS_MATERIAL_MOBILE"), 0);
-		OutEnvironment.SetDefine(TEXT("POST_PROCESS_MATERIAL_BEFORE_TONEMAP"), (Parameters.Material->GetBlendableLocation() != BL_AfterTonemapping) ? 1 : 0);
+		OutEnvironment.SetDefine(TEXT("POST_PROCESS_MATERIAL_BEFORE_TONEMAP"), (Parameters.MaterialParameters.BlendableLocation != BL_AfterTonemapping) ? 1 : 0);
 	}
 };
 
@@ -84,7 +84,7 @@ class FRemoteSessionARCameraVS :
 	public FPostProcessMaterialShader
 {
 public:
-	DECLARE_MATERIAL_SHADER(FRemoteSessionARCameraVS);
+	DECLARE_SHADER_TYPE(FRemoteSessionARCameraVS, Material);
 
 	static void ModifyCompilationEnvironment(const FMaterialShaderPermutationParameters& Parameters, FShaderCompilerEnvironment& OutEnvironment)
 	{
@@ -99,18 +99,18 @@ public:
 
 	void SetParameters(FRHICommandList& RHICmdList, const FSceneView& View)
 	{
-		FRHIVertexShader* ShaderRHI = GetVertexShader();
+		FRHIVertexShader* ShaderRHI = RHICmdList.GetBoundVertexShader();
 		FMaterialShader::SetViewParameters(RHICmdList, ShaderRHI, View, View.ViewUniformBuffer);
 	}
 };
 
-IMPLEMENT_MATERIAL_SHADER(FRemoteSessionARCameraVS, "/Engine/Private/PostProcessMaterialShaders.usf", "MainVS_VideoOverlay", SF_Vertex);
+IMPLEMENT_SHADER_TYPE(,FRemoteSessionARCameraVS, TEXT("/Engine/Private/PostProcessMaterialShaders.usf"), TEXT("MainVS_VideoOverlay"), SF_Vertex);
 
 class FRemoteSessionARCameraPS :
 	public FPostProcessMaterialShader
 {
 public:
-	DECLARE_MATERIAL_SHADER(FRemoteSessionARCameraPS);
+	DECLARE_SHADER_TYPE(FRemoteSessionARCameraPS, Material);
 
 	static void ModifyCompilationEnvironment(const FMaterialShaderPermutationParameters& Parameters, FShaderCompilerEnvironment& OutEnvironment)
 	{
@@ -125,12 +125,12 @@ public:
 
 	void SetParameters(FRHICommandList& RHICmdList, const FSceneView& View, const FMaterialRenderProxy* Material)
 	{
-		FRHIPixelShader* ShaderRHI = GetPixelShader();
+		FRHIPixelShader* ShaderRHI = RHICmdList.GetBoundPixelShader();
 		FMaterialShader::SetParameters(RHICmdList, ShaderRHI, Material, *Material->GetMaterial(View.GetFeatureLevel()), View, View.ViewUniformBuffer, ESceneTextureSetupMode::None);
 	}
 };
 
-IMPLEMENT_MATERIAL_SHADER(FRemoteSessionARCameraPS, "/Engine/Private/PostProcessMaterialShaders.usf", "MainPS_VideoOverlay", SF_Pixel);
+IMPLEMENT_SHADER_TYPE(,FRemoteSessionARCameraPS, TEXT("/Engine/Private/PostProcessMaterialShaders.usf"), TEXT("MainPS_VideoOverlay"), SF_Pixel);
 
 class FARCameraSceneViewExtension :
 	public FSceneViewExtensionBase
@@ -264,12 +264,10 @@ void FARCameraSceneViewExtension::RenderARCamera_RenderThread(FRHICommandListImm
 	GraphicsPSOInit.PrimitiveType = PT_TriangleList;
 	GraphicsPSOInit.BoundShaderState.VertexDeclarationRHI = GFilterVertexDeclaration.VertexDeclarationRHI;
 
-	FRemoteSessionARCameraVS* VertexShader = MaterialShaderMap->GetShader<FRemoteSessionARCameraVS>();
-	FRemoteSessionARCameraPS* PixelShader = MaterialShaderMap->GetShader<FRemoteSessionARCameraPS>();
-	check(PixelShader != nullptr && VertexShader != nullptr);
-
-	GraphicsPSOInit.BoundShaderState.VertexShaderRHI = GETSAFERHISHADER_VERTEX(VertexShader);
-	GraphicsPSOInit.BoundShaderState.PixelShaderRHI = GETSAFERHISHADER_PIXEL(PixelShader);
+	TShaderRef<FRemoteSessionARCameraVS> VertexShader = MaterialShaderMap->GetShader<FRemoteSessionARCameraVS>();
+	TShaderRef<FRemoteSessionARCameraPS> PixelShader = MaterialShaderMap->GetShader<FRemoteSessionARCameraPS>();
+	GraphicsPSOInit.BoundShaderState.VertexShaderRHI = VertexShader.GetVertexShader();
+	GraphicsPSOInit.BoundShaderState.PixelShaderRHI = PixelShader.GetPixelShader();
 
 	SetGraphicsPipelineState(RHICmdList, GraphicsPSOInit);
 
@@ -281,7 +279,7 @@ void FARCameraSceneViewExtension::RenderARCamera_RenderThread(FRHICommandListImm
 			1.0f / ViewSize.X, 1.0f / ViewSize.Y,
 			1.0f, 1.0f);
 
-	SetUniformBufferParameterImmediate(RHICmdList, VertexShader->GetVertexShader(), VertexShader->GetUniformBufferParameter<FDrawRectangleParameters>(), Parameters);
+	SetUniformBufferParameterImmediate(RHICmdList, VertexShader.GetVertexShader(), VertexShader->GetUniformBufferParameter<FDrawRectangleParameters>(), Parameters);
 	VertexShader->SetParameters(RHICmdList, InView);
 	PixelShader->SetParameters(RHICmdList, InView, PPMaterial->GetRenderProxy());
 
diff --git a/Engine/Plugins/FX/HoudiniNiagara/Source/HoudiniNiagara/Private/NiagaraDataInterfaceHoudiniCSV.cpp b/Engine/Plugins/FX/HoudiniNiagara/Source/HoudiniNiagara/Private/NiagaraDataInterfaceHoudiniCSV.cpp
--- a/Engine/Plugins/FX/HoudiniNiagara/Source/HoudiniNiagara/Private/NiagaraDataInterfaceHoudiniCSV.cpp
+++ b/Engine/Plugins/FX/HoudiniNiagara/Source/HoudiniNiagara/Private/NiagaraDataInterfaceHoudiniCSV.cpp
@@ -2211,54 +2211,34 @@ void FNiagaraDataInterfaceProxyHoudiniCSV::AcceptStaticDataUpdate(FNiagaraDIHoud
 // Parameters used for GPU sim compatibility
 struct FNiagaraDataInterfaceParametersCS_HoudiniCSV : public FNiagaraDataInterfaceParametersCS
 {
-	virtual void Bind(const FNiagaraDataInterfaceParamRef& ParamRef, const class FShaderParameterMap& ParameterMap) override
+	DECLARE_TYPE_LAYOUT(FNiagaraDataInterfaceParametersCS_HoudiniCSV, NonVirtual);
+public:
+	void Bind(const FNiagaraDataInterfaceGPUParamInfo& ParameterInfo, const class FShaderParameterMap& ParameterMap)
 	{
-		NumberOfRows.Bind(ParameterMap, *(UNiagaraDataInterfaceHoudiniCSV::NumberOfRowsBaseName + ParamRef.ParameterInfo.DataInterfaceHLSLSymbol));
-		NumberOfColumns.Bind(ParameterMap, *(UNiagaraDataInterfaceHoudiniCSV::NumberOfColumnsBaseName + ParamRef.ParameterInfo.DataInterfaceHLSLSymbol));
-		NumberOfPoints.Bind(ParameterMap, *(UNiagaraDataInterfaceHoudiniCSV::NumberOfPointsBaseName + ParamRef.ParameterInfo.DataInterfaceHLSLSymbol));
+		NumberOfRows.Bind(ParameterMap, *(UNiagaraDataInterfaceHoudiniCSV::NumberOfRowsBaseName + ParameterInfo.DataInterfaceHLSLSymbol));
+		NumberOfColumns.Bind(ParameterMap, *(UNiagaraDataInterfaceHoudiniCSV::NumberOfColumnsBaseName + ParameterInfo.DataInterfaceHLSLSymbol));
+		NumberOfPoints.Bind(ParameterMap, *(UNiagaraDataInterfaceHoudiniCSV::NumberOfPointsBaseName + ParameterInfo.DataInterfaceHLSLSymbol));
 		
-		FloatValuesBuffer.Bind(ParameterMap, *(UNiagaraDataInterfaceHoudiniCSV::FloatValuesBufferBaseName + ParamRef.ParameterInfo.DataInterfaceHLSLSymbol));
+		FloatValuesBuffer.Bind(ParameterMap, *(UNiagaraDataInterfaceHoudiniCSV::FloatValuesBufferBaseName + ParameterInfo.DataInterfaceHLSLSymbol));
 
-		SpecialAttributesColumnIndexesBuffer.Bind(ParameterMap, *(UNiagaraDataInterfaceHoudiniCSV::SpecialAttributesColumnIndexesBufferBaseName + ParamRef.ParameterInfo.DataInterfaceHLSLSymbol));		
+		SpecialAttributesColumnIndexesBuffer.Bind(ParameterMap, *(UNiagaraDataInterfaceHoudiniCSV::SpecialAttributesColumnIndexesBufferBaseName + ParameterInfo.DataInterfaceHLSLSymbol));		
 
-		SpawnTimesBuffer.Bind(ParameterMap, *(UNiagaraDataInterfaceHoudiniCSV::SpawnTimesBufferBaseName + ParamRef.ParameterInfo.DataInterfaceHLSLSymbol));
-		LifeValuesBuffer.Bind(ParameterMap, *(UNiagaraDataInterfaceHoudiniCSV::LifeValuesBufferBaseName + ParamRef.ParameterInfo.DataInterfaceHLSLSymbol));
-		PointTypesBuffer.Bind(ParameterMap, *(UNiagaraDataInterfaceHoudiniCSV::PointTypesBufferBaseName + ParamRef.ParameterInfo.DataInterfaceHLSLSymbol));
+		SpawnTimesBuffer.Bind(ParameterMap, *(UNiagaraDataInterfaceHoudiniCSV::SpawnTimesBufferBaseName + ParameterInfo.DataInterfaceHLSLSymbol));
+		LifeValuesBuffer.Bind(ParameterMap, *(UNiagaraDataInterfaceHoudiniCSV::LifeValuesBufferBaseName + ParameterInfo.DataInterfaceHLSLSymbol));
+		PointTypesBuffer.Bind(ParameterMap, *(UNiagaraDataInterfaceHoudiniCSV::PointTypesBufferBaseName + ParameterInfo.DataInterfaceHLSLSymbol));
 
-		MaxNumberOfIndexesPerPoint.Bind(ParameterMap, *(UNiagaraDataInterfaceHoudiniCSV::MaxNumberOfIndexesPerPointBaseName + ParamRef.ParameterInfo.DataInterfaceHLSLSymbol));
-		PointValueIndexesBuffer.Bind(ParameterMap, *(UNiagaraDataInterfaceHoudiniCSV::PointValueIndexesBufferBaseName + ParamRef.ParameterInfo.DataInterfaceHLSLSymbol));
+		MaxNumberOfIndexesPerPoint.Bind(ParameterMap, *(UNiagaraDataInterfaceHoudiniCSV::MaxNumberOfIndexesPerPointBaseName + ParameterInfo.DataInterfaceHLSLSymbol));
+		PointValueIndexesBuffer.Bind(ParameterMap, *(UNiagaraDataInterfaceHoudiniCSV::PointValueIndexesBufferBaseName + ParameterInfo.DataInterfaceHLSLSymbol));
 
-		LastSpawnedPointId.Bind(ParameterMap, *(UNiagaraDataInterfaceHoudiniCSV::LastSpawnedPointIdBaseName + ParamRef.ParameterInfo.DataInterfaceHLSLSymbol));
-		LastSpawnTime.Bind(ParameterMap, *(UNiagaraDataInterfaceHoudiniCSV::LastSpawnTimeBaseName + ParamRef.ParameterInfo.DataInterfaceHLSLSymbol));
+		LastSpawnedPointId.Bind(ParameterMap, *(UNiagaraDataInterfaceHoudiniCSV::LastSpawnedPointIdBaseName + ParameterInfo.DataInterfaceHLSLSymbol));
+		LastSpawnTime.Bind(ParameterMap, *(UNiagaraDataInterfaceHoudiniCSV::LastSpawnTimeBaseName + ParameterInfo.DataInterfaceHLSLSymbol));
 	}
 
-	virtual void Serialize(FArchive& Ar)override
-	{
-		Ar << Version;
-
-		Ar << NumberOfRows;
-		Ar << NumberOfColumns;
-		Ar << NumberOfPoints;
-
-		Ar << FloatValuesBuffer;
-		Ar << SpecialAttributesColumnIndexesBuffer;
-
-		Ar << SpawnTimesBuffer;
-		Ar << LifeValuesBuffer;
-		Ar << PointTypesBuffer;
-
-		Ar << MaxNumberOfIndexesPerPoint;
-		Ar << PointValueIndexesBuffer;
-
-		Ar << LastSpawnedPointId;
-		Ar << LastSpawnTime;
-	}
-
-	virtual void Set(FRHICommandList& RHICmdList, const FNiagaraDataInterfaceSetArgs& Context) const override
+	void Set(FRHICommandList& RHICmdList, const FNiagaraDataInterfaceSetArgs& Context) const
 	{
 		check( IsInRenderingThread() );
 
-		FRHIComputeShader* ComputeShaderRHI = Context.Shader->GetComputeShader();
+		FRHIComputeShader* ComputeShaderRHI = Context.Shader.GetComputeShader();
 		FNiagaraDataInterfaceProxyHoudiniCSV* HoudiniDI = static_cast<FNiagaraDataInterfaceProxyHoudiniCSV*>(Context.DataInterface);
 		if ( !HoudiniDI )
 		{
@@ -2284,30 +2264,28 @@ struct FNiagaraDataInterfaceParametersCS_HoudiniCSV : public FNiagaraDataInterfa
 	}
 
 private:
+	LAYOUT_FIELD(FShaderParameter, NumberOfRows);
+	LAYOUT_FIELD(FShaderParameter, NumberOfColumns);
+	LAYOUT_FIELD(FShaderParameter, NumberOfPoints);
 
-	FShaderParameter NumberOfRows;
-	FShaderParameter NumberOfColumns;
-	FShaderParameter NumberOfPoints;
-
-	FShaderResourceParameter FloatValuesBuffer;
-	FShaderResourceParameter SpecialAttributesColumnIndexesBuffer;
+	LAYOUT_FIELD(FShaderResourceParameter, FloatValuesBuffer);
+	LAYOUT_FIELD(FShaderResourceParameter, SpecialAttributesColumnIndexesBuffer);
 
-	FShaderResourceParameter SpawnTimesBuffer;
-	FShaderResourceParameter LifeValuesBuffer;
-	FShaderResourceParameter PointTypesBuffer;
+	LAYOUT_FIELD(FShaderResourceParameter, SpawnTimesBuffer);
+	LAYOUT_FIELD(FShaderResourceParameter, LifeValuesBuffer);
+	LAYOUT_FIELD(FShaderResourceParameter, PointTypesBuffer);
 
-	FShaderParameter MaxNumberOfIndexesPerPoint;
-	FShaderResourceParameter PointValueIndexesBuffer;
+	LAYOUT_FIELD(FShaderParameter, MaxNumberOfIndexesPerPoint);
+	LAYOUT_FIELD(FShaderResourceParameter, PointValueIndexesBuffer);
 
-	FShaderParameter LastSpawnedPointId;
-	FShaderParameter LastSpawnTime;
+	LAYOUT_FIELD(FShaderParameter, LastSpawnedPointId);
+	LAYOUT_FIELD(FShaderParameter, LastSpawnTime);
 
-	uint32 Version = 1;
+	LAYOUT_FIELD_INITIALIZED(uint32, Version, 1);
 };
 
-FNiagaraDataInterfaceParametersCS* UNiagaraDataInterfaceHoudiniCSV::ConstructComputeParameters()const
-{
-	return new FNiagaraDataInterfaceParametersCS_HoudiniCSV();
-}
+IMPLEMENT_TYPE_LAYOUT(FNiagaraDataInterfaceParametersCS_HoudiniCSV);
+
+IMPLEMENT_NIAGARA_DI_PARAMETER(UNiagaraDataInterfaceHoudiniCSV, FNiagaraDataInterfaceParametersCS_HoudiniCSV);
 
 #undef LOCTEXT_NAMESPACE
diff --git a/Engine/Plugins/FX/Niagara/Source/Niagara/Classes/NiagaraDataInterfaceCamera.h b/Engine/Plugins/FX/Niagara/Source/Niagara/Classes/NiagaraDataInterfaceCamera.h
--- a/Engine/Plugins/FX/Niagara/Source/Niagara/Classes/NiagaraDataInterfaceCamera.h
+++ b/Engine/Plugins/FX/Niagara/Source/Niagara/Classes/NiagaraDataInterfaceCamera.h
@@ -20,6 +20,9 @@ class NIAGARA_API UNiagaraDataInterfaceCamera : public UNiagaraDataInterface
 	GENERATED_UCLASS_BODY()
 
 public:
+
+	DECLARE_NIAGARA_DI_PARAMETER();
+
 	/** This is used to determine which camera position to query for cpu emitters. If no valid index is supplied, the first controller is used as camera reference. */
 	UPROPERTY(EditAnywhere, Category = "Camera")
 	int32 PlayerControllerIndex = 0;
@@ -62,4 +65,17 @@ struct FNiagaraDataIntefaceProxyCameraQuery : public FNiagaraDataInterfaceProxy
 	{
 		return 0;
 	}
-};
\ No newline at end of file
+};
+
+struct FNiagaraDataInterfaceParametersCS_CameraQuery : public FNiagaraDataInterfaceParametersCS
+{
+	DECLARE_TYPE_LAYOUT(FNiagaraDataInterfaceParametersCS_CameraQuery, NonVirtual);
+
+	void Bind(const FNiagaraDataInterfaceGPUParamInfo& ParameterInfo, const class FShaderParameterMap& ParameterMap);
+	void Set(FRHICommandList& RHICmdList, const FNiagaraDataInterfaceSetArgs& Context) const;
+
+private:
+
+	/** The SceneDepthTexture parameter for depth buffer query. */
+	LAYOUT_FIELD(FShaderUniformBufferParameter, PassUniformBuffer);
+};
diff --git a/Engine/Plugins/FX/Niagara/Source/Niagara/Classes/NiagaraDataInterfaceCurveBase.h b/Engine/Plugins/FX/Niagara/Source/Niagara/Classes/NiagaraDataInterfaceCurveBase.h
--- a/Engine/Plugins/FX/Niagara/Source/Niagara/Classes/NiagaraDataInterfaceCurveBase.h
+++ b/Engine/Plugins/FX/Niagara/Source/Niagara/Classes/NiagaraDataInterfaceCurveBase.h
@@ -71,6 +71,8 @@ protected:
 	}
 
 public:
+	DECLARE_NIAGARA_DI_PARAMETER();
+
 	UPROPERTY(EditAnywhere, AdvancedDisplay, Category = "Curve")
 	uint32 bUseLUT : 1;
 
@@ -152,7 +154,6 @@ public:
 	virtual void GetCurveData(TArray<FCurveData>& OutCurveData) { }
 
 	virtual void GetParameterDefinitionHLSL(const FNiagaraDataInterfaceGPUParamInfo& ParamInfo, FString& OutHLSL) override;
-	virtual FNiagaraDataInterfaceParametersCS* ConstructComputeParameters() const override;
 
 	void SetDefaultLUT();
 #if WITH_EDITORONLY_DATA
@@ -198,3 +199,17 @@ struct TCurveUseLUTBinder
 		}
 	}
 };
+
+struct FNiagaraDataInterfaceParametersCS_Curve : public FNiagaraDataInterfaceParametersCS
+{
+	DECLARE_TYPE_LAYOUT(FNiagaraDataInterfaceParametersCS_Curve, NonVirtual);
+
+	void Bind(const FNiagaraDataInterfaceGPUParamInfo& ParameterInfo, const class FShaderParameterMap& ParameterMap);
+	void Set(FRHICommandList& RHICmdList, const FNiagaraDataInterfaceSetArgs& Context) const;
+
+	LAYOUT_FIELD(FShaderParameter, MinTime);
+	LAYOUT_FIELD(FShaderParameter, MaxTime);
+	LAYOUT_FIELD(FShaderParameter, InvTimeRange);
+	LAYOUT_FIELD(FShaderParameter, CurveLUTNumMinusOne);
+	LAYOUT_FIELD(FShaderResourceParameter, CurveLUT);
+};
diff --git a/Engine/Plugins/FX/Niagara/Source/Niagara/Classes/NiagaraDataInterfaceGrid2DCollection.h b/Engine/Plugins/FX/Niagara/Source/Niagara/Classes/NiagaraDataInterfaceGrid2DCollection.h
--- a/Engine/Plugins/FX/Niagara/Source/Niagara/Classes/NiagaraDataInterfaceGrid2DCollection.h
+++ b/Engine/Plugins/FX/Niagara/Source/Niagara/Classes/NiagaraDataInterfaceGrid2DCollection.h
@@ -102,6 +102,8 @@ class NIAGARA_API UNiagaraDataInterfaceGrid2DCollection : public UNiagaraDataInt
 	GENERATED_UCLASS_BODY()
 
 public:
+
+	DECLARE_NIAGARA_DI_PARAMETER();
 	
 	virtual void PostInitProperties() override;
 	
@@ -115,8 +117,7 @@ public:
 	// GPU sim functionality
 	virtual void GetParameterDefinitionHLSL(const FNiagaraDataInterfaceGPUParamInfo& ParamInfo, FString& OutHLSL) override;
 	virtual bool GetFunctionHLSL(const FNiagaraDataInterfaceGPUParamInfo& ParamInfo, const FNiagaraDataInterfaceGeneratedFunction& FunctionInfo, int FunctionInstanceIndex, FString& OutHLSL) override;
-	virtual FNiagaraDataInterfaceParametersCS* ConstructComputeParameters() const override;
-	
+
 	virtual void ProvidePerInstanceDataForRenderThread(void* DataForRenderThread, void* PerInstanceData, const FNiagaraSystemInstanceID& SystemInstance) override {}
 	virtual bool InitPerInstanceData(void* PerInstanceData, FNiagaraSystemInstance* SystemInstance) override;
 	virtual void DestroyPerInstanceData(void* PerInstanceData, FNiagaraSystemInstance* SystemInstance) override;
diff --git a/Engine/Plugins/FX/Niagara/Source/Niagara/Classes/NiagaraDataInterfaceNeighborGrid3D.h b/Engine/Plugins/FX/Niagara/Source/Niagara/Classes/NiagaraDataInterfaceNeighborGrid3D.h
--- a/Engine/Plugins/FX/Niagara/Source/Niagara/Classes/NiagaraDataInterfaceNeighborGrid3D.h
+++ b/Engine/Plugins/FX/Niagara/Source/Niagara/Classes/NiagaraDataInterfaceNeighborGrid3D.h
@@ -52,6 +52,8 @@ public:
 
 public:
 
+	DECLARE_NIAGARA_DI_PARAMETER();
+
 	virtual void PostInitProperties() override
 	{
 		Super::PostInitProperties();
@@ -73,8 +75,7 @@ public:
 	// GPU sim functionality
 	virtual void GetParameterDefinitionHLSL(const FNiagaraDataInterfaceGPUParamInfo& ParamInfo, FString& OutHLSL) override;
 	virtual bool GetFunctionHLSL(const FNiagaraDataInterfaceGPUParamInfo& ParamInfo, const FNiagaraDataInterfaceGeneratedFunction& FunctionInfo, int FunctionInstanceIndex, FString& OutHLSL) override;
-	virtual FNiagaraDataInterfaceParametersCS* ConstructComputeParameters() const override;
-		
+
 	virtual void ProvidePerInstanceDataForRenderThread(void* DataForRenderThread, void* PerInstanceData, const FNiagaraSystemInstanceID& SystemInstance) override {}
 	virtual bool InitPerInstanceData(void* PerInstanceData, FNiagaraSystemInstance* SystemInstance) override;
 	virtual void DestroyPerInstanceData(void* PerInstanceData, FNiagaraSystemInstance* SystemInstance) override;
diff --git a/Engine/Plugins/FX/Niagara/Source/Niagara/Classes/NiagaraDataInterfaceOcclusion.h b/Engine/Plugins/FX/Niagara/Source/Niagara/Classes/NiagaraDataInterfaceOcclusion.h
--- a/Engine/Plugins/FX/Niagara/Source/Niagara/Classes/NiagaraDataInterfaceOcclusion.h
+++ b/Engine/Plugins/FX/Niagara/Source/Niagara/Classes/NiagaraDataInterfaceOcclusion.h
@@ -12,6 +12,9 @@ class NIAGARA_API UNiagaraDataInterfaceOcclusion : public UNiagaraDataInterface
 	GENERATED_UCLASS_BODY()
 
 public:
+
+	DECLARE_NIAGARA_DI_PARAMETER();
+
 	//UObject Interface
 	virtual void PostInitProperties() override;
 	//UObject Interface End
@@ -21,7 +24,6 @@ public:
 	virtual void GetVMExternalFunction(const FVMExternalFunctionBindingInfo& BindingInfo, void* InstanceData, FVMExternalFunction &OutFunc) override;
 	virtual bool GetFunctionHLSL(const FNiagaraDataInterfaceGPUParamInfo& ParamInfo, const FNiagaraDataInterfaceGeneratedFunction& FunctionInfo, int FunctionInstanceIndex, FString& OutHLSL) override;
 	virtual bool CanExecuteOnTarget(ENiagaraSimTarget Target) const override { return Target == ENiagaraSimTarget::GPUComputeSim; }
-	virtual FNiagaraDataInterfaceParametersCS* ConstructComputeParameters() const override;
 	virtual bool RequiresDepthBuffer() const override { return true; }
 	//UNiagaraDataInterface Interface
 
diff --git a/Engine/Plugins/FX/Niagara/Source/Niagara/Classes/NiagaraDataInterfaceParticleRead.h b/Engine/Plugins/FX/Niagara/Source/Niagara/Classes/NiagaraDataInterfaceParticleRead.h
--- a/Engine/Plugins/FX/Niagara/Source/Niagara/Classes/NiagaraDataInterfaceParticleRead.h
+++ b/Engine/Plugins/FX/Niagara/Source/Niagara/Classes/NiagaraDataInterfaceParticleRead.h
@@ -29,7 +29,6 @@ public:
 	virtual bool Equals(const UNiagaraDataInterface* Other) const override;
 	virtual void GetParameterDefinitionHLSL(const FNiagaraDataInterfaceGPUParamInfo& ParamInfo, FString& OutHLSL) override;
 	virtual bool GetFunctionHLSL(const FNiagaraDataInterfaceGPUParamInfo& ParamInfo, const FNiagaraDataInterfaceGeneratedFunction& FunctionInfo, int FunctionInstanceIndex, FString& OutHLSL) override;
-	virtual FNiagaraDataInterfaceParametersCS* ConstructComputeParameters() const override;
 	virtual void ProvidePerInstanceDataForRenderThread(void* DataForRenderThread, void* PerInstanceData, const FNiagaraSystemInstanceID& SystemInstance) override;
 	virtual void GetEmitterDependencies(void* PerInstanceData, FNiagaraSystemInstance* SystemInstance, TArray<FNiagaraEmitterInstance*>& Dependencies) const override;
 	//UNiagaraDataInterface Interface End
diff --git a/Engine/Plugins/FX/Niagara/Source/Niagara/Classes/NiagaraDataInterfaceRW.h b/Engine/Plugins/FX/Niagara/Source/Niagara/Classes/NiagaraDataInterfaceRW.h
--- a/Engine/Plugins/FX/Niagara/Source/Niagara/Classes/NiagaraDataInterfaceRW.h
+++ b/Engine/Plugins/FX/Niagara/Source/Niagara/Classes/NiagaraDataInterfaceRW.h
@@ -142,6 +142,7 @@ public:
 
 
 public:
+
 	//~ UNiagaraDataInterface interface
 	// VM functionality
 	virtual void GetFunctions(TArray<FNiagaraFunctionSignature>& OutFunctions) override;
@@ -153,7 +154,6 @@ public:
 	// GPU sim functionality
 	virtual void GetParameterDefinitionHLSL(const FNiagaraDataInterfaceGPUParamInfo& ParamInfo, FString& OutHLSL) override;
 	virtual bool GetFunctionHLSL(const FNiagaraDataInterfaceGPUParamInfo& ParamInfo, const FNiagaraDataInterfaceGeneratedFunction& FunctionInfo, int FunctionInstanceIndex, FString& OutHLSL) override;
-	virtual FNiagaraDataInterfaceParametersCS* ConstructComputeParameters() const override;
 	//~ UNiagaraDataInterface interface END
 
 
@@ -202,7 +202,6 @@ public:
 	// GPU sim functionality
 	virtual void GetParameterDefinitionHLSL(const FNiagaraDataInterfaceGPUParamInfo& ParamInfo, FString& OutHLSL) override;
 	virtual bool GetFunctionHLSL(const FNiagaraDataInterfaceGPUParamInfo& ParamInfo, const FNiagaraDataInterfaceGeneratedFunction& FunctionInfo, int FunctionInstanceIndex, FString& OutHLSL) override;
-	virtual FNiagaraDataInterfaceParametersCS* ConstructComputeParameters() const override;
 	//~ UNiagaraDataInterface interface END
 
 
diff --git a/Engine/Plugins/FX/Niagara/Source/Niagara/Classes/NiagaraDataInterfaceSkeletalMesh.h b/Engine/Plugins/FX/Niagara/Source/Niagara/Classes/NiagaraDataInterfaceSkeletalMesh.h
--- a/Engine/Plugins/FX/Niagara/Source/Niagara/Classes/NiagaraDataInterfaceSkeletalMesh.h
+++ b/Engine/Plugins/FX/Niagara/Source/Niagara/Classes/NiagaraDataInterfaceSkeletalMesh.h
@@ -476,6 +476,8 @@ class NIAGARA_API UNiagaraDataInterfaceSkeletalMesh : public UNiagaraDataInterfa
 	GENERATED_UCLASS_BODY()
 
 public:
+
+	DECLARE_NIAGARA_DI_PARAMETER();
 	
 #if WITH_EDITORONLY_DATA
 	/** Mesh used to sample from when not overridden by a source actor from the scene. Only available in editor for previewing. This is removed in cooked builds. */
@@ -555,7 +557,6 @@ public:
 	virtual void GetCommonHLSL(FString& OutHLSL) override;
 	virtual void GetParameterDefinitionHLSL(const FNiagaraDataInterfaceGPUParamInfo& ParamInfo, FString& OutHLSL) override;
 	virtual bool GetFunctionHLSL(const FNiagaraDataInterfaceGPUParamInfo& ParamInfo, const FNiagaraDataInterfaceGeneratedFunction& FunctionInfo, int FunctionInstanceIndex, FString& OutHLSL) override;
-	virtual FNiagaraDataInterfaceParametersCS* ConstructComputeParameters() const override;
 
 	virtual void ProvidePerInstanceDataForRenderThread(void* DataForRenderThread, void* PerInstanceData, const FNiagaraSystemInstanceID& SystemInstance) override;
 
diff --git a/Engine/Plugins/FX/Niagara/Source/Niagara/Classes/NiagaraDataInterfaceStaticMesh.h b/Engine/Plugins/FX/Niagara/Source/Niagara/Classes/NiagaraDataInterfaceStaticMesh.h
--- a/Engine/Plugins/FX/Niagara/Source/Niagara/Classes/NiagaraDataInterfaceStaticMesh.h
+++ b/Engine/Plugins/FX/Niagara/Source/Niagara/Classes/NiagaraDataInterfaceStaticMesh.h
@@ -161,6 +161,8 @@ class NIAGARA_API UNiagaraDataInterfaceStaticMesh : public UNiagaraDataInterface
 	GENERATED_UCLASS_BODY()
 
 public:
+
+	DECLARE_NIAGARA_DI_PARAMETER();
 	
 	/** Mesh used to sample from when not overridden by a source actor from the scene. Also useful for previewing in the editor. */
 	UPROPERTY(EditAnywhere, Category = "Mesh")
@@ -209,7 +211,6 @@ public:
 
 	virtual void GetParameterDefinitionHLSL(const FNiagaraDataInterfaceGPUParamInfo& ParamInfo, FString& OutHLSL) override;
 	virtual bool GetFunctionHLSL(const FNiagaraDataInterfaceGPUParamInfo& ParamInfo, const FNiagaraDataInterfaceGeneratedFunction& FunctionInfo, int FunctionInstanceIndex, FString& OutHLSL) override;
-	virtual FNiagaraDataInterfaceParametersCS* ConstructComputeParameters()const override;
 
 	virtual void ProvidePerInstanceDataForRenderThread(void* DataForRenderThread, void* PerInstanceData, const FNiagaraSystemInstanceID& SystemInstance) override;
 
diff --git a/Engine/Plugins/FX/Niagara/Source/Niagara/Classes/NiagaraDataInterfaceTexture.h b/Engine/Plugins/FX/Niagara/Source/Niagara/Classes/NiagaraDataInterfaceTexture.h
--- a/Engine/Plugins/FX/Niagara/Source/Niagara/Classes/NiagaraDataInterfaceTexture.h
+++ b/Engine/Plugins/FX/Niagara/Source/Niagara/Classes/NiagaraDataInterfaceTexture.h
@@ -14,6 +14,8 @@ class NIAGARA_API UNiagaraDataInterfaceTexture : public UNiagaraDataInterface
 	GENERATED_UCLASS_BODY()
 public:
 
+	DECLARE_NIAGARA_DI_PARAMETER();
+
 	UPROPERTY(EditAnywhere, Category = "Texture")
 	UTexture* Texture;
 
@@ -41,7 +43,6 @@ public:
 	// GPU sim functionality
 	virtual void GetParameterDefinitionHLSL(const FNiagaraDataInterfaceGPUParamInfo& ParamInfo, FString& OutHLSL) override;
 	virtual bool GetFunctionHLSL(const FNiagaraDataInterfaceGPUParamInfo& ParamInfo, const FNiagaraDataInterfaceGeneratedFunction& FunctionInfo, int FunctionInstanceIndex, FString& OutHLSL) override;
-	virtual FNiagaraDataInterfaceParametersCS* ConstructComputeParameters()const override;
 
 	void SetTexture(UTexture* InTexture);
 
diff --git a/Engine/Plugins/FX/Niagara/Source/Niagara/Classes/NiagaraDataInterfaceVectorField.h b/Engine/Plugins/FX/Niagara/Source/Niagara/Classes/NiagaraDataInterfaceVectorField.h
--- a/Engine/Plugins/FX/Niagara/Source/Niagara/Classes/NiagaraDataInterfaceVectorField.h
+++ b/Engine/Plugins/FX/Niagara/Source/Niagara/Classes/NiagaraDataInterfaceVectorField.h
@@ -27,6 +27,8 @@ public:
 public:
 	//~ UObject interface
 
+	DECLARE_NIAGARA_DI_PARAMETER();
+
 	virtual void PostInitProperties() override;
 	virtual void PostLoad() override; 
 #if WITH_EDITOR
@@ -50,7 +52,6 @@ public:
 	// GPU sim functionality
 	virtual void GetParameterDefinitionHLSL(const FNiagaraDataInterfaceGPUParamInfo& ParamInfo, FString& OutHLSL) override;
 	virtual bool GetFunctionHLSL(const FNiagaraDataInterfaceGPUParamInfo& ParamInfo, const FNiagaraDataInterfaceGeneratedFunction& FunctionInfo, int FunctionInstanceIndex, FString& OutHLSL) override;
-	virtual FNiagaraDataInterfaceParametersCS* ConstructComputeParameters() const override;
 	//~ UNiagaraDataInterface interface END
 
 	// VM functions
@@ -92,4 +93,20 @@ struct FNiagaraDataInterfaceProxyVectorField : public FNiagaraDataInterfaceProxy
 	{
 		return 0;
 	}
-};
\ No newline at end of file
+};
+
+struct FNiagaraDataInterfaceParametersCS_VectorField : public FNiagaraDataInterfaceParametersCS
+{
+	DECLARE_TYPE_LAYOUT(FNiagaraDataInterfaceParametersCS_VectorField, NonVirtual);
+public:
+	void Bind(const FNiagaraDataInterfaceGPUParamInfo& ParameterInfo, const class FShaderParameterMap& ParameterMap);
+	void Set(FRHICommandList& RHICmdList, const FNiagaraDataInterfaceSetArgs& Context) const;
+
+private:
+	LAYOUT_FIELD(FShaderResourceParameter, VectorFieldSampler);
+	LAYOUT_FIELD(FShaderResourceParameter, VectorFieldTexture);
+	LAYOUT_FIELD(FShaderParameter, TilingAxes);
+	LAYOUT_FIELD(FShaderParameter, Dimensions);
+	LAYOUT_FIELD(FShaderParameter, MinBounds);
+	LAYOUT_FIELD(FShaderParameter, MaxBounds);
+};
diff --git a/Engine/Plugins/FX/Niagara/Source/Niagara/Classes/NiagaraEmitterInstanceBatcher.h b/Engine/Plugins/FX/Niagara/Source/Niagara/Classes/NiagaraEmitterInstanceBatcher.h
--- a/Engine/Plugins/FX/Niagara/Source/Niagara/Classes/NiagaraEmitterInstanceBatcher.h
+++ b/Engine/Plugins/FX/Niagara/Source/Niagara/Classes/NiagaraEmitterInstanceBatcher.h
@@ -108,14 +108,14 @@ public:
 
 	void ProcessDebugInfo(FRHICommandList &RHICmdLis, const FNiagaraComputeExecutionContext* Context) const;
 
-	void SetDataInterfaceParameters(const TArray<FNiagaraDataInterfaceProxy*>& DataInterfaceProxies, FNiagaraShader* Shader, FRHICommandList &RHICmdList, const FNiagaraComputeInstanceData* Instance, const FNiagaraGPUSystemTick& Tick, uint32 ShaderStageIndex) const;
-	void UnsetDataInterfaceParameters(const TArray<FNiagaraDataInterfaceProxy*>& DataInterfaceProxies, FNiagaraShader* Shader, FRHICommandList& RHICmdList, const FNiagaraComputeInstanceData* Instance, const FNiagaraGPUSystemTick& Tick) const;
+	void SetDataInterfaceParameters(const TArray<FNiagaraDataInterfaceProxy*>& DataInterfaceProxies, const FNiagaraShaderRef& Shader, FRHICommandList &RHICmdList, const FNiagaraComputeInstanceData* Instance, const FNiagaraGPUSystemTick& Tick, uint32 ShaderStageIndex) const;
+	void UnsetDataInterfaceParameters(const TArray<FNiagaraDataInterfaceProxy*>& DataInterfaceProxies, const FNiagaraShaderRef& Shader, FRHICommandList& RHICmdList, const FNiagaraComputeInstanceData* Instance, const FNiagaraGPUSystemTick& Tick) const;
 
 	void Run(	const FNiagaraGPUSystemTick& Tick,
 				const FNiagaraComputeInstanceData* Instance, 
 				uint32 UpdateStartInstance, 
 				const uint32 TotalNumInstances, 
-				FNiagaraShader* Shader,
+				const FNiagaraShaderRef& Shader,
 				FRHICommandList &RHICmdList, 
 				FRHIUniformBuffer* ViewUniformBuffer, 
 				const FNiagaraGpuSpawnInfo& SpawnInfo,
@@ -134,19 +134,19 @@ public:
 	FORCEINLINE ERHIFeatureLevel::Type GetFeatureLevel() const { return FeatureLevel; }
 
 	/** Reset the data interfaces and check if the spawn stages are valid */
-	bool ResetDataInterfaces(const FNiagaraGPUSystemTick& Tick, FNiagaraComputeInstanceData *Instance, FRHICommandList &RHICmdList, FNiagaraShader* ComputeShader ) const;
+	bool ResetDataInterfaces(const FNiagaraGPUSystemTick& Tick, FNiagaraComputeInstanceData *Instance, FRHICommandList &RHICmdList, const FNiagaraShaderRef& ComputeShader ) const;
 
 	/** Given a shader stage index, find the corresponding data interface */
 	FNiagaraDataInterfaceProxy* FindIterationInterface(FNiagaraComputeInstanceData *Instance, const uint32 ShaderStageIndex) const;
 
 	/** Loop over all the data interfaces and call the prestage methods */
-	void PreStageInterface(const FNiagaraGPUSystemTick& Tick, FNiagaraComputeInstanceData *Instance, FRHICommandList &RHICmdList, FNiagaraShader* ComputeShader, const uint32 ShaderStageIndex) const;
+	void PreStageInterface(const FNiagaraGPUSystemTick& Tick, FNiagaraComputeInstanceData *Instance, FRHICommandList &RHICmdList, const FNiagaraShaderRef& ComputeShader, const uint32 ShaderStageIndex) const;
 
 	/** Loop over all the data interfaces and call the poststage methods */
-	void PostStageInterface(const FNiagaraGPUSystemTick& Tick, FNiagaraComputeInstanceData *Instance, FRHICommandList &RHICmdList, FNiagaraShader* ComputeShader, const uint32 ShaderStageIndex) const;
+	void PostStageInterface(const FNiagaraGPUSystemTick& Tick, FNiagaraComputeInstanceData *Instance, FRHICommandList &RHICmdList, const FNiagaraShaderRef& ComputeShader, const uint32 ShaderStageIndex) const;
 
 	/** Run the dispatch over multiple stages */
-	void DispatchMultipleStages(const FNiagaraGPUSystemTick& Tick, FNiagaraComputeInstanceData *Instance, FRHICommandList &RHICmdList, FRHIUniformBuffer* ViewUniformBuffer, FNiagaraShader* ComputeShader);
+	void DispatchMultipleStages(const FNiagaraGPUSystemTick& Tick, FNiagaraComputeInstanceData *Instance, FRHICommandList &RHICmdList, FRHIUniformBuffer* ViewUniformBuffer, const FNiagaraShaderRef& ComputeShader);
 
 	/** Get the shared SortManager, used in the rendering loop to call FGPUSortManager::OnPreRender() and FGPUSortManager::OnPostRenderOpaque() */
 	virtual FGPUSortManager* GetGPUSortManager() const override;
diff --git a/Engine/Plugins/FX/Niagara/Source/Niagara/Classes/NiagaraScript.h b/Engine/Plugins/FX/Niagara/Source/Niagara/Classes/NiagaraScript.h
--- a/Engine/Plugins/FX/Niagara/Source/Niagara/Classes/NiagaraScript.h
+++ b/Engine/Plugins/FX/Niagara/Source/Niagara/Classes/NiagaraScript.h
@@ -547,7 +547,7 @@ public:
 	{
 		if (!ScriptShader)
 		{
-			ScriptShader = ScriptResource.GetShader()->GetComputeShader();	// NIAGARATODO: need to put this caching somewhere else, as it wont' know when we update the resource
+			ScriptShader = ScriptResource.GetShader().GetComputeShader();	// NIAGARATODO: need to put this caching somewhere else, as it wont' know when we update the resource
 		}
 		return ScriptShader; 
 	}
@@ -556,7 +556,7 @@ public:
 	{
 		if (!ScriptShader)
 		{
-			ScriptShader = ScriptResource.GetShaderGameThread()->GetComputeShader();	// NIAGARATODO: need to put this caching somewhere else, as it wont' know when we update the resource
+			ScriptShader = ScriptResource.GetShaderGameThread().GetComputeShader();	// NIAGARATODO: need to put this caching somewhere else, as it wont' know when we update the resource
 		}
 		return ScriptShader;
 	}
diff --git a/Engine/Plugins/FX/Niagara/Source/Niagara/Classes/NiagaraScriptExecutionContext.h b/Engine/Plugins/FX/Niagara/Source/Niagara/Classes/NiagaraScriptExecutionContext.h
--- a/Engine/Plugins/FX/Niagara/Source/Niagara/Classes/NiagaraScriptExecutionContext.h
+++ b/Engine/Plugins/FX/Niagara/Source/Niagara/Classes/NiagaraScriptExecutionContext.h
@@ -247,7 +247,7 @@ public:
 	//TArray<uint8, TAlignedHeapAllocator<16>> ParamData_RT;		// RT side copy of the parameter data
 	FNiagaraScriptExecutionParameterStore CombinedParamStore;
 #if DO_CHECK
-	TArray< FNiagaraDataInterfaceGPUParamInfo >  DIParamInfo;
+	TArray< FString >  DIClassNames;
 #endif
 
 	TArray<FNiagaraDataInterfaceProxy*> DataInterfaceProxies;
diff --git a/Engine/Plugins/FX/Niagara/Source/Niagara/Private/NiagaraComponent.cpp b/Engine/Plugins/FX/Niagara/Source/Niagara/Private/NiagaraComponent.cpp
--- a/Engine/Plugins/FX/Niagara/Source/Niagara/Private/NiagaraComponent.cpp
+++ b/Engine/Plugins/FX/Niagara/Source/Niagara/Private/NiagaraComponent.cpp
@@ -127,8 +127,6 @@ FNiagaraSceneProxy::FNiagaraSceneProxy(const UNiagaraComponent* InComponent)
 		, PerfAsset(InComponent->GetAsset())
 #endif
 {
-	// In this case only, update the System renderers on the game thread.
-	check(IsInGameThread());
 	FNiagaraSystemInstance* SystemInst = InComponent->GetSystemInstance();
 	if (SystemInst)
 	{
@@ -1187,9 +1185,9 @@ void UNiagaraComponent::OnEndOfFrameUpdateDuringTick()
 	}
 }
 
-void UNiagaraComponent::CreateRenderState_Concurrent()
+void UNiagaraComponent::CreateRenderState_Concurrent(FRegisterComponentContext* Context)
 {
-	Super::CreateRenderState_Concurrent();
+	Super::CreateRenderState_Concurrent(Context);
 	// The emitter instance may not tick again next frame so we send the dynamic data here so that the current state
 	// renders.  This can happen when while editing, or any time the age update mode is set to desired age.
 	SendRenderDynamicData_Concurrent();
diff --git a/Engine/Plugins/FX/Niagara/Source/Niagara/Private/NiagaraDataInterfaceAudioOscilloscope.cpp b/Engine/Plugins/FX/Niagara/Source/Niagara/Private/NiagaraDataInterfaceAudioOscilloscope.cpp
--- a/Engine/Plugins/FX/Niagara/Source/Niagara/Private/NiagaraDataInterfaceAudioOscilloscope.cpp
+++ b/Engine/Plugins/FX/Niagara/Source/Niagara/Private/NiagaraDataInterfaceAudioOscilloscope.cpp
@@ -316,23 +316,12 @@ struct FNiagaraDataInterfaceParametersCS_AudioOscilloscope : public FNiagaraData
 	}
 
 	virtual ~FNiagaraDataInterfaceParametersCS_AudioOscilloscope() {}
-	virtual void Bind(const FNiagaraDataInterfaceParamRef& ParamRef, const class FShaderParameterMap& ParameterMap) override
-	{
-		NumChannels.Bind(ParameterMap, *(NumChannelsName + ParamRef.ParameterInfo.DataInterfaceHLSLSymbol));
-		AudioBuffer.Bind(ParameterMap, *(AudioBufferName + ParamRef.ParameterInfo.DataInterfaceHLSLSymbol));
-	}
-
-	virtual void Serialize(FArchive& Ar) override
-	{
-		Ar << NumChannels;
-		Ar << AudioBuffer;
-	}
 
-	virtual void Set(FRHICommandList& RHICmdList, const FNiagaraDataInterfaceSetArgs& Context) const override
+	void Set(FRHICommandList& RHICmdList, const FNiagaraDataInterfaceSetArgs& Context) const
 	{
 		check(IsInRenderingThread());
 
-		FRHIComputeShader* ComputeShaderRHI = Context.Shader->GetComputeShader();
+		FRHIComputeShader* ComputeShaderRHI = Context.Shader.GetComputeShader();
 
 		FNiagaraDataInterfaceProxyOscilloscope* NDI = (FNiagaraDataInterfaceProxyOscilloscope*)Context.DataInterface;
 		FReadBuffer& AudioBufferSRV = NDI->ComputeAndPostSRV();
diff --git a/Engine/Plugins/FX/Niagara/Source/Niagara/Private/NiagaraDataInterfaceCamera.cpp b/Engine/Plugins/FX/Niagara/Source/Niagara/Private/NiagaraDataInterfaceCamera.cpp
--- a/Engine/Plugins/FX/Niagara/Source/Niagara/Private/NiagaraDataInterfaceCamera.cpp
+++ b/Engine/Plugins/FX/Niagara/Source/Niagara/Private/NiagaraDataInterfaceCamera.cpp
@@ -6,12 +6,31 @@
 #include "Internationalization/Internationalization.h"
 #include "NiagaraSystemInstance.h"
 #include "GameFramework/PlayerController.h"
+#include "ShaderParameterUtils.h"
 
 #if WITH_EDITORONLY_DATA
 #include "EditorViewportClient.h"
 #include "LevelEditorViewport.h"
 #endif
 
+
+IMPLEMENT_TYPE_LAYOUT(FNiagaraDataInterfaceParametersCS_CameraQuery);
+
+void FNiagaraDataInterfaceParametersCS_CameraQuery::Bind(const FNiagaraDataInterfaceGPUParamInfo& ParameterInfo, const class FShaderParameterMap& ParameterMap)
+{
+	PassUniformBuffer.Bind(ParameterMap, FSceneTexturesUniformParameters::StaticStructMetadata.GetShaderVariableName());
+}
+
+void FNiagaraDataInterfaceParametersCS_CameraQuery::Set(FRHICommandList& RHICmdList, const FNiagaraDataInterfaceSetArgs& Context) const
+{
+	check(IsInRenderingThread());
+	FRHIComputeShader* ComputeShaderRHI = RHICmdList.GetBoundComputeShader();
+
+	TUniformBufferRef<FSceneTexturesUniformParameters> SceneTextureUniformParams = GNiagaraViewDataManager.GetSceneTextureUniformParameters();
+	SetUniformBufferParameter(RHICmdList, ComputeShaderRHI, PassUniformBuffer, SceneTextureUniformParams);
+}
+
+
 const FName UNiagaraDataInterfaceCamera::GetViewPropertiesName(TEXT("GetViewPropertiesGPU"));
 const FName UNiagaraDataInterfaceCamera::GetClipSpaceTransformsName(TEXT("GetClipSpaceTransformsGPU"));
 const FName UNiagaraDataInterfaceCamera::GetViewSpaceTransformsName(TEXT("GetViewSpaceTransformsGPU"));
@@ -417,3 +436,4 @@ void UNiagaraDataInterfaceCamera::GetViewSpaceTransformsGPU(FVectorVMContext& Co
 	}
 }
 
+IMPLEMENT_NIAGARA_DI_PARAMETER(UNiagaraDataInterfaceCamera, FNiagaraDataInterfaceParametersCS_CameraQuery);
diff --git a/Engine/Plugins/FX/Niagara/Source/Niagara/Private/NiagaraDataInterfaceCollisionQuery.cpp b/Engine/Plugins/FX/Niagara/Source/Niagara/Private/NiagaraDataInterfaceCollisionQuery.cpp
--- a/Engine/Plugins/FX/Niagara/Source/Niagara/Private/NiagaraDataInterfaceCollisionQuery.cpp
+++ b/Engine/Plugins/FX/Niagara/Source/Niagara/Private/NiagaraDataInterfaceCollisionQuery.cpp
@@ -852,23 +852,19 @@ bool UNiagaraDataInterfaceCollisionQuery::PerInstanceTickPostSimulate(void* PerI
 
 struct FNiagaraDataInterfaceParametersCS_CollisionQuery : public FNiagaraDataInterfaceParametersCS
 {
-	virtual void Bind(const FNiagaraDataInterfaceParamRef& ParamRef, const class FShaderParameterMap& ParameterMap) override
+	DECLARE_TYPE_LAYOUT(FNiagaraDataInterfaceParametersCS_CollisionQuery, NonVirtual);
+public:
+	void Bind(const FNiagaraDataInterfaceGPUParamInfo& ParameterInfo, const class FShaderParameterMap& ParameterMap)
 	{
 		PassUniformBuffer.Bind(ParameterMap, FSceneTexturesUniformParameters::StaticStructMetadata.GetShaderVariableName());
 		GlobalDistanceFieldParameters.Bind(ParameterMap);
 	}
 
-	virtual void Serialize(FArchive& Ar) override
-	{
-		Ar << PassUniformBuffer;
-		Ar << GlobalDistanceFieldParameters;
-	}
-
-	virtual void Set(FRHICommandList& RHICmdList, const FNiagaraDataInterfaceSetArgs& Context) const override
+	void Set(FRHICommandList& RHICmdList, const FNiagaraDataInterfaceSetArgs& Context) const
 	{
 		check(IsInRenderingThread());
 
-		FRHIComputeShader* ComputeShaderRHI = Context.Shader->GetComputeShader();
+		FRHIComputeShader* ComputeShaderRHI = Context.Shader.GetComputeShader();
 		
 		TUniformBufferRef<FSceneTexturesUniformParameters> SceneTextureUniformParams = GNiagaraViewDataManager.GetSceneTextureUniformParameters();
 		SetUniformBufferParameter(RHICmdList, ComputeShaderRHI, PassUniformBuffer/*Shader->GetUniformBufferParameter(SceneTexturesUniformBufferStruct)*/, SceneTextureUniformParams);
@@ -879,14 +875,12 @@ struct FNiagaraDataInterfaceParametersCS_CollisionQuery : public FNiagaraDataInt
 	}
 
 private:
-
 	/** The SceneDepthTexture parameter for depth buffer collision. */
-	FShaderUniformBufferParameter PassUniformBuffer;
+	LAYOUT_FIELD(FShaderUniformBufferParameter, PassUniformBuffer);
 
-	FGlobalDistanceFieldParameters GlobalDistanceFieldParameters;
+	LAYOUT_FIELD(FGlobalDistanceFieldParameters, GlobalDistanceFieldParameters);
 };
 
-FNiagaraDataInterfaceParametersCS* UNiagaraDataInterfaceCollisionQuery::ConstructComputeParameters() const
-{
-	return new FNiagaraDataInterfaceParametersCS_CollisionQuery();
-}
\ No newline at end of file
+IMPLEMENT_TYPE_LAYOUT(FNiagaraDataInterfaceParametersCS_CollisionQuery);
+
+IMPLEMENT_NIAGARA_DI_PARAMETER(UNiagaraDataInterfaceCollisionQuery, FNiagaraDataInterfaceParametersCS_CollisionQuery);
diff --git a/Engine/Plugins/FX/Niagara/Source/Niagara/Private/NiagaraDataInterfaceCurlNoise.cpp b/Engine/Plugins/FX/Niagara/Source/Niagara/Private/NiagaraDataInterfaceCurlNoise.cpp
--- a/Engine/Plugins/FX/Niagara/Source/Niagara/Private/NiagaraDataInterfaceCurlNoise.cpp
+++ b/Engine/Plugins/FX/Niagara/Source/Niagara/Private/NiagaraDataInterfaceCurlNoise.cpp
@@ -467,22 +467,19 @@ void UNiagaraDataInterfaceCurlNoise::PushToRenderThread()
 
 struct FNiagaraDataInterfaceParametersCS_CurlNoise : public FNiagaraDataInterfaceParametersCS
 {
-	virtual void Bind(const FNiagaraDataInterfaceParamRef& ParamRef, const class FShaderParameterMap& ParameterMap) override
+	DECLARE_TYPE_LAYOUT(FNiagaraDataInterfaceParametersCS_CurlNoise, NonVirtual);
+public:
+	void Bind(const FNiagaraDataInterfaceGPUParamInfo& ParameterInfo, const class FShaderParameterMap& ParameterMap)
 	{
-		OffsetFromSeed.Bind(ParameterMap, *(OffsetFromSeedBaseName + ParamRef.ParameterInfo.DataInterfaceHLSLSymbol));
+		OffsetFromSeed.Bind(ParameterMap, *(OffsetFromSeedBaseName + ParameterInfo.DataInterfaceHLSLSymbol));
 	}
 
-	virtual void Serialize(FArchive& Ar)override
-	{
-		Ar << OffsetFromSeed;
-	}
-
-	virtual void Set(FRHICommandList& RHICmdList, const FNiagaraDataInterfaceSetArgs& Context) const override
+	void Set(FRHICommandList& RHICmdList, const FNiagaraDataInterfaceSetArgs& Context) const
 	{
 		check(IsInRenderingThread());
 
 		// Get shader and DI
-		FRHIComputeShader* ComputeShaderRHI = Context.Shader->GetComputeShader();
+		FRHIComputeShader* ComputeShaderRHI = Context.Shader.GetComputeShader();
 		FNiagaraDataInterfaceProxyCurlNoise* CNDI = static_cast<FNiagaraDataInterfaceProxyCurlNoise*>(Context.DataInterface);
 
 		// Set parameters
@@ -490,10 +487,9 @@ struct FNiagaraDataInterfaceParametersCS_CurlNoise : public FNiagaraDataInterfac
 	}
 
 private:
-	FShaderParameter OffsetFromSeed;
+	LAYOUT_FIELD(FShaderParameter, OffsetFromSeed);
 };
 
-FNiagaraDataInterfaceParametersCS* UNiagaraDataInterfaceCurlNoise::ConstructComputeParameters()const
-{
-	return new FNiagaraDataInterfaceParametersCS_CurlNoise();
-}
\ No newline at end of file
+IMPLEMENT_TYPE_LAYOUT(FNiagaraDataInterfaceParametersCS_CurlNoise);
+
+IMPLEMENT_NIAGARA_DI_PARAMETER(UNiagaraDataInterfaceCurlNoise, FNiagaraDataInterfaceParametersCS_CurlNoise);
diff --git a/Engine/Plugins/FX/Niagara/Source/Niagara/Private/NiagaraDataInterfaceCurveBase.cpp b/Engine/Plugins/FX/Niagara/Source/Niagara/Private/NiagaraDataInterfaceCurveBase.cpp
--- a/Engine/Plugins/FX/Niagara/Source/Niagara/Private/NiagaraDataInterfaceCurveBase.cpp
+++ b/Engine/Plugins/FX/Niagara/Source/Niagara/Private/NiagaraDataInterfaceCurveBase.cpp
@@ -9,6 +9,36 @@
 #include "NiagaraShader.h"
 #include "NiagaraCustomVersion.h"
 
+IMPLEMENT_TYPE_LAYOUT(FNiagaraDataInterfaceParametersCS_Curve);
+
+IMPLEMENT_NIAGARA_DI_PARAMETER(UNiagaraDataInterfaceCurveBase, FNiagaraDataInterfaceParametersCS_Curve);
+
+void FNiagaraDataInterfaceParametersCS_Curve::Bind(const FNiagaraDataInterfaceGPUParamInfo& ParameterInfo, const class FShaderParameterMap& ParameterMap)
+{
+	MinTime.Bind(ParameterMap, *(TEXT("MinTime_") + ParameterInfo.DataInterfaceHLSLSymbol));
+	MaxTime.Bind(ParameterMap, *(TEXT("MaxTime_") + ParameterInfo.DataInterfaceHLSLSymbol));
+	InvTimeRange.Bind(ParameterMap, *(TEXT("InvTimeRange_") + ParameterInfo.DataInterfaceHLSLSymbol));
+	CurveLUTNumMinusOne.Bind(ParameterMap, *(TEXT("CurveLUTNumMinusOne_") + ParameterInfo.DataInterfaceHLSLSymbol));
+	CurveLUT.Bind(ParameterMap, *(TEXT("CurveLUT_") + ParameterInfo.DataInterfaceHLSLSymbol));
+}
+
+void FNiagaraDataInterfaceParametersCS_Curve::Set(FRHICommandList& RHICmdList, const FNiagaraDataInterfaceSetArgs& Context) const
+{
+	check(IsInRenderingThread());
+
+	FRHIComputeShader* ComputeShaderRHI = RHICmdList.GetBoundComputeShader();
+	FNiagaraDataInterfaceProxyCurveBase* CurveDI = (FNiagaraDataInterfaceProxyCurveBase*)Context.DataInterface;
+	FReadBuffer& CurveLUTBuffer = CurveDI->CurveLUT;
+
+	SetShaderValue(RHICmdList, ComputeShaderRHI, MinTime, CurveDI->LUTMinTime);
+	SetShaderValue(RHICmdList, ComputeShaderRHI, MaxTime, CurveDI->LUTMaxTime);
+	SetShaderValue(RHICmdList, ComputeShaderRHI, InvTimeRange, CurveDI->LUTInvTimeRange);
+	SetShaderValue(RHICmdList, ComputeShaderRHI, CurveLUTNumMinusOne, CurveDI->CurveLUTNumMinusOne);
+	SetSRVParameter(RHICmdList, ComputeShaderRHI, CurveLUT, CurveLUTBuffer.SRV);
+}
+
+
+
 float GNiagaraLUTOptimizeThreshold = UNiagaraDataInterfaceCurveBase::DefaultOptimizeThreshold;
 static FAutoConsoleVariableRef CVarNiagaraLUTOptimizeThreshold(
 	TEXT("fx.Niagara.LUT.OptimizeThreshold"),
@@ -307,51 +337,3 @@ void UNiagaraDataInterfaceCurveBase::PushToRenderThread()
 		RHIUnlockVertexBuffer(RT_Proxy->CurveLUT.Buffer);
 	});
 }
-
-struct FNiagaraDataInterfaceParametersCS_Curve : public FNiagaraDataInterfaceParametersCS
-{
-	virtual ~FNiagaraDataInterfaceParametersCS_Curve() {}
-	virtual void Bind(const FNiagaraDataInterfaceParamRef& ParamRef, const class FShaderParameterMap& ParameterMap) override
-	{
-		MinTime.Bind(ParameterMap, *(TEXT("MinTime_") + ParamRef.ParameterInfo.DataInterfaceHLSLSymbol));
-		MaxTime.Bind(ParameterMap, *(TEXT("MaxTime_") + ParamRef.ParameterInfo.DataInterfaceHLSLSymbol));
-		InvTimeRange.Bind(ParameterMap, *(TEXT("InvTimeRange_") + ParamRef.ParameterInfo.DataInterfaceHLSLSymbol));
-		CurveLUTNumMinusOne.Bind(ParameterMap, *(TEXT("CurveLUTNumMinusOne_") + ParamRef.ParameterInfo.DataInterfaceHLSLSymbol));
-		CurveLUT.Bind(ParameterMap, *(TEXT("CurveLUT_") + ParamRef.ParameterInfo.DataInterfaceHLSLSymbol));
-	}
-
-	virtual void Serialize(FArchive& Ar) override
-	{
-		Ar << MinTime;
-		Ar << MaxTime;
-		Ar << InvTimeRange;
-		Ar << CurveLUTNumMinusOne;
-		Ar << CurveLUT;
-	}
-
-	virtual void Set(FRHICommandList& RHICmdList, const FNiagaraDataInterfaceSetArgs& Context) const override
-	{
-		check(IsInRenderingThread());
-
-		FRHIComputeShader* ComputeShaderRHI = Context.Shader->GetComputeShader();
-		FNiagaraDataInterfaceProxyCurveBase* CurveDI = (FNiagaraDataInterfaceProxyCurveBase*)Context.DataInterface;
-		FReadBuffer& CurveLUTBuffer = CurveDI->CurveLUT;
-
-		SetShaderValue(RHICmdList, ComputeShaderRHI, MinTime, CurveDI->LUTMinTime);
-		SetShaderValue(RHICmdList, ComputeShaderRHI, MaxTime, CurveDI->LUTMaxTime);
-		SetShaderValue(RHICmdList, ComputeShaderRHI, InvTimeRange, CurveDI->LUTInvTimeRange);
-		SetShaderValue(RHICmdList, ComputeShaderRHI, CurveLUTNumMinusOne, CurveDI->CurveLUTNumMinusOne);
-		SetSRVParameter(RHICmdList, ComputeShaderRHI, CurveLUT, CurveLUTBuffer.SRV);
-	}
-
-	FShaderParameter MinTime;
-	FShaderParameter MaxTime;
-	FShaderParameter InvTimeRange;
-	FShaderParameter CurveLUTNumMinusOne;
-	FShaderResourceParameter CurveLUT;
-};
-
-FNiagaraDataInterfaceParametersCS* UNiagaraDataInterfaceCurveBase::ConstructComputeParameters()const
-{
-	return new FNiagaraDataInterfaceParametersCS_Curve();
-}
diff --git a/Engine/Plugins/FX/Niagara/Source/Niagara/Private/NiagaraDataInterfaceGrid2DCollection.cpp b/Engine/Plugins/FX/Niagara/Source/Niagara/Private/NiagaraDataInterfaceGrid2DCollection.cpp
--- a/Engine/Plugins/FX/Niagara/Source/Niagara/Private/NiagaraDataInterfaceGrid2DCollection.cpp
+++ b/Engine/Plugins/FX/Niagara/Source/Niagara/Private/NiagaraDataInterfaceGrid2DCollection.cpp
@@ -28,40 +28,29 @@ static const FName SampleGridFunctionName("SampleGrid");
 /*--------------------------------------------------------------------------------------------------------------------------*/
 struct FNiagaraDataInterfaceParametersCS_Grid2DCollection : public FNiagaraDataInterfaceParametersCS
 {
-	virtual void Bind(const FNiagaraDataInterfaceParamRef& ParamRef, const class FShaderParameterMap& ParameterMap) override
+	DECLARE_TYPE_LAYOUT(FNiagaraDataInterfaceParametersCS_Grid2DCollection, NonVirtual);
+public:
+	void Bind(const FNiagaraDataInterfaceGPUParamInfo& ParameterInfo, const class FShaderParameterMap& ParameterMap)
 	{			
-		NumCellsParam.Bind(ParameterMap, *(NumCellsName + ParamRef.ParameterInfo.DataInterfaceHLSLSymbol));
-		NumTilesParam.Bind(ParameterMap, *(NumTilesName + ParamRef.ParameterInfo.DataInterfaceHLSLSymbol));
+		NumCellsParam.Bind(ParameterMap, *(NumCellsName + ParameterInfo.DataInterfaceHLSLSymbol));
+		NumTilesParam.Bind(ParameterMap, *(NumTilesName + ParameterInfo.DataInterfaceHLSLSymbol));
 
-		CellSizeParam.Bind(ParameterMap, *(CellSizeName + ParamRef.ParameterInfo.DataInterfaceHLSLSymbol));
-		
-		WorldBBoxSizeParam.Bind(ParameterMap, *(WorldBBoxSizeName + ParamRef.ParameterInfo.DataInterfaceHLSLSymbol));
-
-		GridParam.Bind(ParameterMap, *(GridName + ParamRef.ParameterInfo.DataInterfaceHLSLSymbol));
-		OutputGridParam.Bind(ParameterMap, *(OutputGridName + ParamRef.ParameterInfo.DataInterfaceHLSLSymbol));
+		CellSizeParam.Bind(ParameterMap, *(CellSizeName + ParameterInfo.DataInterfaceHLSLSymbol));
 
-		SamplerParam.Bind(ParameterMap, *(SamplerName + ParamRef.ParameterInfo.DataInterfaceHLSLSymbol));
-	}
+		WorldBBoxSizeParam.Bind(ParameterMap, *(WorldBBoxSizeName + ParameterInfo.DataInterfaceHLSLSymbol));
 
-	virtual void Serialize(FArchive& Ar)override
-	{		
-		Ar << NumCellsParam;		
-		Ar << NumTilesParam;	
-		Ar << CellSizeParam;			
-		Ar << WorldBBoxSizeParam;
+		GridParam.Bind(ParameterMap, *(GridName + ParameterInfo.DataInterfaceHLSLSymbol));
+		OutputGridParam.Bind(ParameterMap, *(OutputGridName + ParameterInfo.DataInterfaceHLSLSymbol));
 
-		Ar << GridParam;
-		Ar << OutputGridParam;		
-
-		Ar << SamplerParam;
+		SamplerParam.Bind(ParameterMap, *(SamplerName + ParameterInfo.DataInterfaceHLSLSymbol));
 	}
 
-	virtual void Set(FRHICommandList& RHICmdList, const FNiagaraDataInterfaceSetArgs& Context) const override
+	void Set(FRHICommandList& RHICmdList, const FNiagaraDataInterfaceSetArgs& Context) const
 	{
 		check(IsInRenderingThread());
 
 		// Get shader and DI
-		FRHIComputeShader* ComputeShaderRHI = Context.Shader->GetComputeShader();
+		FRHIComputeShader* ComputeShaderRHI = Context.Shader.GetComputeShader();
 		FNiagaraDataInterfaceProxyGrid2DCollection* VFDI = static_cast<FNiagaraDataInterfaceProxyGrid2DCollection*>(Context.DataInterface);
 		
 		Grid2DCollectionRWInstanceData* ProxyData = VFDI->SystemInstancesToProxyData.Find(Context.SystemInstance);
@@ -70,8 +59,8 @@ struct FNiagaraDataInterfaceParametersCS_Grid2DCollection : public FNiagaraDataI
 		int NumCellsTmp[2];
 		NumCellsTmp[0] = ProxyData->NumCellsX;
 		NumCellsTmp[1] = ProxyData->NumCellsY;
-		SetShaderValue(RHICmdList, ComputeShaderRHI, NumCellsParam, NumCellsTmp);	
-
+		SetShaderValue(RHICmdList, ComputeShaderRHI, NumCellsParam, NumCellsTmp);		
+		
 		int NumTilesTmp[2];
 		NumTilesTmp[0] = ProxyData->NumTilesX;
 		NumTilesTmp[1] = ProxyData->NumTilesY;
@@ -87,39 +76,42 @@ struct FNiagaraDataInterfaceParametersCS_Grid2DCollection : public FNiagaraDataI
 		if (GridParam.IsBound() && ProxyData->GetCurrentData() != NULL)
 		{
 			RHICmdList.TransitionResource(EResourceTransitionAccess::EReadable, EResourceTransitionPipeline::EComputeToCompute, ProxyData->GetCurrentData()->GridBuffer.UAV);
-			SetSRVParameter(RHICmdList, Context.Shader->GetComputeShader(), GridParam, ProxyData->GetCurrentData()->GridBuffer.SRV);
+			SetSRVParameter(RHICmdList, Context.Shader.GetComputeShader(), GridParam, ProxyData->GetCurrentData()->GridBuffer.SRV);
 		}
 
 		if (Context.IsOutputStage && OutputGridParam.IsBound() && ProxyData->GetDestinationData() != NULL)
 		{
 			const FTextureRWBuffer2D& TextureBuffer = ProxyData->GetDestinationData()->GridBuffer;
 			RHICmdList.TransitionResource(EResourceTransitionAccess::EWritable, EResourceTransitionPipeline::EComputeToCompute, ProxyData->GetDestinationData()->GridBuffer.UAV);
-			OutputGridParam.SetTexture(RHICmdList, Context.Shader->GetComputeShader(), TextureBuffer.Buffer, TextureBuffer.UAV);
+			OutputGridParam.SetTexture(RHICmdList, Context.Shader.GetComputeShader(), TextureBuffer.Buffer, TextureBuffer.UAV);
 		}
 	}
 
-	virtual void Unset(FRHICommandList& RHICmdList, const FNiagaraDataInterfaceSetArgs& Context) const 
+	void Unset(FRHICommandList& RHICmdList, const FNiagaraDataInterfaceSetArgs& Context) const 
 	{
 		if (OutputGridParam.IsBound())
 		{
-			OutputGridParam.UnsetUAV(RHICmdList, Context.Shader->GetComputeShader());
+			OutputGridParam.UnsetUAV(RHICmdList, Context.Shader.GetComputeShader());
 		}
 	}
 
 private:
 
-	FShaderParameter NumCellsParam;
-	FShaderParameter NumTilesParam;
-	FShaderParameter CellSizeParam;	
-	FShaderParameter WorldBBoxSizeParam;
-
-	FShaderResourceParameter GridParam;
-	FRWShaderParameter OutputGridParam;
+	LAYOUT_FIELD(FShaderParameter, NumCellsParam);
+	LAYOUT_FIELD(FShaderParameter, NumTilesParam);
+	LAYOUT_FIELD(FShaderParameter, CellSizeParam);
+	LAYOUT_FIELD(FShaderParameter, WorldBBoxSizeParam);
 
-	FShaderResourceParameter SamplerParam;
+	LAYOUT_FIELD(FShaderResourceParameter, GridParam);
+	LAYOUT_FIELD(FRWShaderParameter, OutputGridParam);
 	
+	LAYOUT_FIELD(FShaderResourceParameter, SamplerParam);
 };
 
+IMPLEMENT_TYPE_LAYOUT(FNiagaraDataInterfaceParametersCS_Grid2DCollection);
+
+IMPLEMENT_NIAGARA_DI_PARAMETER(UNiagaraDataInterfaceGrid2DCollection, FNiagaraDataInterfaceParametersCS_Grid2DCollection);
+
 
 UNiagaraDataInterfaceGrid2DCollection::UNiagaraDataInterfaceGrid2DCollection(FObjectInitializer const& ObjectInitializer)
 	: Super(ObjectInitializer)
@@ -179,7 +171,7 @@ void UNiagaraDataInterfaceGrid2DCollection::GetFunctions(TArray<FNiagaraFunction
 		Sig.Name = SampleGridFunctionName;
 		Sig.Inputs.Add(FNiagaraVariable(FNiagaraTypeDefinition(GetClass()), TEXT("Grid")));
 		Sig.Inputs.Add(FNiagaraVariable(FNiagaraTypeDefinition::GetFloatDef(), TEXT("UnitX")));
-		Sig.Inputs.Add(FNiagaraVariable(FNiagaraTypeDefinition::GetFloatDef(), TEXT("UnitY")));
+		Sig.Inputs.Add(FNiagaraVariable(FNiagaraTypeDefinition::GetFloatDef(), TEXT("UnitY")));		
 		Sig.Inputs.Add(FNiagaraVariable(FNiagaraTypeDefinition::GetIntDef(), TEXT("AttributeIndex")));
 		Sig.Outputs.Add(FNiagaraVariable(FNiagaraTypeDefinition::GetFloatDef(), TEXT("Value")));
 
@@ -237,7 +229,7 @@ void UNiagaraDataInterfaceGrid2DCollection::GetParameterDefinitionHLSL(const FNi
 		RWTexture2D<float> RW{OutputGridName};
 		int2 {NumTiles};
 		SamplerState {SamplerName};
-	
+
 	)");
 	TMap<FString, FStringFormatArg> ArgsDeclarations = {				
 		{ TEXT("GridName"),    GridName + ParamInfo.DataInterfaceHLSLSymbol },
@@ -279,7 +271,7 @@ bool UNiagaraDataInterfaceGrid2DCollection::GetFunctionHLSL(const FNiagaraDataIn
 	{
 		static const TCHAR *FormatBounds = TEXT(R"(
 			void {FunctionName}(int In_IndexX, int In_IndexY, int In_AttributeIndex, float In_Value, out int val)
-			{			
+			{				
 				int TileIndexX = In_AttributeIndex % {NumTiles}.x;
 				int TileIndexY = In_AttributeIndex / {NumTiles}.x;
 	
@@ -320,13 +312,6 @@ bool UNiagaraDataInterfaceGrid2DCollection::GetFunctionHLSL(const FNiagaraDataIn
 	return false;
 }
 
-FNiagaraDataInterfaceParametersCS* UNiagaraDataInterfaceGrid2DCollection::ConstructComputeParameters()const
-{
-	return new FNiagaraDataInterfaceParametersCS_Grid2DCollection();
-}
-
-
-
 bool UNiagaraDataInterfaceGrid2DCollection::CopyToInternal(UNiagaraDataInterface* Destination) const
 {
 	if (!Super::CopyToInternal(Destination))
@@ -390,7 +375,7 @@ bool UNiagaraDataInterfaceGrid2DCollection::InitPerInstanceData(void* PerInstanc
 
 	TSet<int> RT_OutputShaderStages = OutputShaderStages;
 	TSet<int> RT_IterationShaderStages = IterationShaderStages;
-	
+
 	FVector2D RT_CellSize = RT_WorldBBoxSize / FVector2D(RT_NumCellsX, RT_NumCellsY);
 	InstanceData->CellSize = RT_CellSize;
 	InstanceData->WorldBBoxSize = RT_WorldBBoxSize;
@@ -404,7 +389,7 @@ bool UNiagaraDataInterfaceGrid2DCollection::InitPerInstanceData(void* PerInstanc
 
 		TargetData->NumCellsX = RT_NumCellsX;
 		TargetData->NumCellsY = RT_NumCellsY;
-		
+
 		TargetData->NumTilesX = RT_NumTilesX;
 		TargetData->NumTilesY = RT_NumTilesY;
 
@@ -668,7 +653,7 @@ void FNiagaraDataInterfaceProxyGrid2DCollection::PreStage(FRHICommandList& RHICm
 		// we start.  If a user wants to access data from the previous stage, then they can read from the current data.
 
 		// #todo(dmp): we might want to expose an option where we have buffers that are write only and need a clear (ie: no buffering like the neighbor grid).  They would be considered transient perhaps?  It'd be more
-		// memory efficient since it would theoretically not require any double buffering.		
+		// memory efficient since it would theoretically not require any double buffering.
 		if (!Context.IsIterationStage)
 		{
 			RHICmdList.TransitionResource(EResourceTransitionAccess::ERWBarrier, EResourceTransitionPipeline::EComputeToCompute, ProxyData->DestinationData->GridBuffer.UAV);
