Fix shader compilation related to hair simulation.

diff --git a/Engine/Plugins/Experimental/HairStrands/Shaders/Private/NiagaraClearPressureGrid.usf b/Engine/Plugins/Experimental/HairStrands/Shaders/Private/NiagaraClearPressureGrid.usf
--- a/Engine/Plugins/Experimental/HairStrands/Shaders/Private/NiagaraClearPressureGrid.usf
+++ b/Engine/Plugins/Experimental/HairStrands/Shaders/Private/NiagaraClearPressureGrid.usf
@@ -1,5 +1,7 @@
 // Copyright Epic Games, Inc. All Rights Reserved.
 
+#include "/Engine/Private/Common.ush"
+
 static const int FLUID_PRESSURE_OFFSET = 6;
 static const int FLUID_DISTANCE_OFFSET = 14;
 static const int NUM_ATTRIBUTES = 17;
diff --git a/Engine/Plugins/Experimental/HairStrands/Shaders/Private/NiagaraCopyBoundingBox.usf b/Engine/Plugins/Experimental/HairStrands/Shaders/Private/NiagaraCopyBoundingBox.usf
--- a/Engine/Plugins/Experimental/HairStrands/Shaders/Private/NiagaraCopyBoundingBox.usf
+++ b/Engine/Plugins/Experimental/HairStrands/Shaders/Private/NiagaraCopyBoundingBox.usf
@@ -1,5 +1,7 @@
 // Copyright Epic Games, Inc. All Rights Reserved.
 
+#include "/Engine/Private/Common.ush"
+
 uint NumElements;
 RWBuffer<uint> BoundingBoxBuffer;
 RWBuffer<uint> OutNodeBoundBuffer;
