Fix various vertex factory types to use new system for creating parameter data

diff --git a/Engine/Plugins/Enterprise/LidarPointCloud/Source/LidarPointCloudRuntime/Private/Rendering/LidarPointCloudRenderBuffers.cpp b/Engine/Plugins/Enterprise/LidarPointCloud/Source/LidarPointCloudRuntime/Private/Rendering/LidarPointCloudRenderBuffers.cpp
--- a/Engine/Plugins/Enterprise/LidarPointCloud/Source/LidarPointCloudRuntime/Private/Rendering/LidarPointCloudRenderBuffers.cpp
+++ b/Engine/Plugins/Enterprise/LidarPointCloud/Source/LidarPointCloudRuntime/Private/Rendering/LidarPointCloudRenderBuffers.cpp
@@ -415,16 +415,6 @@ void FLidarPointCloudVertexFactoryShaderParameters::GetElementShaderBindings(con
 	SETPARAM(ClassificationColors);
 }
 
-FVertexFactoryShaderParameters* FLidarPointCloudVertexFactory::ConstructShaderParameters(EShaderFrequency ShaderFrequency)
-{
-	if (ShaderFrequency == SF_Vertex)
-	{
-		return new FLidarPointCloudVertexFactoryShaderParameters();
-	}
-
-	return nullptr;
-}
-
 FLidarPointCloudVertexFactory::~FLidarPointCloudVertexFactory()
 {
 	FRenderCommandFence Fence;
@@ -462,6 +452,8 @@ void FLidarPointCloudVertexFactory::ReleaseRHI()
 	VertexBuffer.ReleaseResource();
 }
 
+IMPLEMENT_VERTEX_FACTORY_PARAMETER_TYPE(FLidarPointCloudVertexFactory, SF_Vertex, FLidarPointCloudVertexFactoryShaderParameters);
+
 IMPLEMENT_VERTEX_FACTORY_TYPE(FLidarPointCloudVertexFactory, "/Plugin/LidarPointCloud/Private/LidarPointCloudVertexFactory.ush", /* bUsedWithMaterials */ true, /* bSupportsStaticLighting */ false, /* bSupportsDynamicLighting */ true, /* bPrecisePrevWorldPos */ false, /* bSupportsPositionOnly */ true);
 
 #undef BINDPARAM
diff --git a/Engine/Plugins/Enterprise/LidarPointCloud/Source/LidarPointCloudRuntime/Private/Rendering/LidarPointCloudRenderBuffers.h b/Engine/Plugins/Enterprise/LidarPointCloud/Source/LidarPointCloudRuntime/Private/Rendering/LidarPointCloudRenderBuffers.h
--- a/Engine/Plugins/Enterprise/LidarPointCloud/Source/LidarPointCloudRuntime/Private/Rendering/LidarPointCloudRenderBuffers.h
+++ b/Engine/Plugins/Enterprise/LidarPointCloud/Source/LidarPointCloudRuntime/Private/Rendering/LidarPointCloudRenderBuffers.h
@@ -209,8 +209,6 @@ class FLidarPointCloudVertexFactory : public FVertexFactory
 	DECLARE_VERTEX_FACTORY_TYPE(FLidarPointCloudVertexFactory);
 
 public:
-	static FVertexFactoryShaderParameters* ConstructShaderParameters(EShaderFrequency ShaderFrequency);
-
 	static bool ShouldCompilePermutation(const FVertexFactoryShaderPermutationParameters& Parameters);
 	static void ModifyCompilationEnvironment(const FVertexFactoryShaderPermutationParameters& Parameters, FShaderCompilerEnvironment& OutEnvironment) {}
 
diff --git a/Engine/Plugins/Experimental/HairStrands/Source/HairStrandsCore/Private/HairStrandsVertexFactory.cpp b/Engine/Plugins/Experimental/HairStrands/Source/HairStrandsCore/Private/HairStrandsVertexFactory.cpp
--- a/Engine/Plugins/Experimental/HairStrands/Source/HairStrandsCore/Private/HairStrandsVertexFactory.cpp
+++ b/Engine/Plugins/Experimental/HairStrands/Source/HairStrandsCore/Private/HairStrandsVertexFactory.cpp
@@ -207,31 +207,15 @@ void FHairStrandsVertexFactory::InitRHI()
 	check(IsValidRef(GetDeclaration()));
 }
 
-FVertexFactoryShaderParameters* FHairStrandsVertexFactory::ConstructShaderParameters(EShaderFrequency ShaderFrequency)
-{
-	if (ShaderFrequency == SF_Vertex)
-	{
-		return new FHairStrandsVertexFactoryShaderParameters();
-	}
-
-	if (ShaderFrequency == SF_Pixel)
-	{
-		return new FHairStrandsVertexFactoryShaderParameters();
-	}
-
-#if RHI_RAYTRACING
-	if (ShaderFrequency == SF_RayHitGroup)
-	{
-		return new FHairStrandsVertexFactoryShaderParameters();
-	}
-#endif // RHI_RAYTRACING
-
-	return NULL;
-}
-
 void FHairStrandsVertexFactory::ReleaseRHI()
 {
 	FVertexFactory::ReleaseRHI();
 }
 
+IMPLEMENT_VERTEX_FACTORY_PARAMETER_TYPE(FHairStrandsVertexFactory, SF_Vertex, FHairStrandsVertexFactoryShaderParameters);
+IMPLEMENT_VERTEX_FACTORY_PARAMETER_TYPE(FHairStrandsVertexFactory, SF_Pixel, FHairStrandsVertexFactoryShaderParameters);
+#if RHI_RAYTRACING
+IMPLEMENT_VERTEX_FACTORY_PARAMETER_TYPE(FHairStrandsVertexFactory, SF_RayHitGroup, FHairStrandsVertexFactoryShaderParameters);
+#endif // RHI_RAYTRACING
+
 IMPLEMENT_VERTEX_FACTORY_TYPE_EX(FHairStrandsVertexFactory,"/Engine/Private/HairStrands/HairStrandsVertexFactory.ush",true,false,true,true,true,true,true);
diff --git a/Engine/Plugins/Experimental/HairStrands/Source/HairStrandsCore/Public/HairStrandsVertexFactory.h b/Engine/Plugins/Experimental/HairStrands/Source/HairStrandsCore/Public/HairStrandsVertexFactory.h
--- a/Engine/Plugins/Experimental/HairStrands/Source/HairStrandsCore/Public/HairStrandsVertexFactory.h
+++ b/Engine/Plugins/Experimental/HairStrands/Source/HairStrandsCore/Public/HairStrandsVertexFactory.h
@@ -60,8 +60,6 @@ public:
 
 	static bool SupportsTessellationShaders() { return false; }
 
-	static FVertexFactoryShaderParameters* ConstructShaderParameters(EShaderFrequency ShaderFrequency);
-
 	inline FRHIShaderResourceView* GetPositionSRV(uint32 GroupIndex) const			{ check(Data.InterpolationOutput); return Data.InterpolationOutput->HairGroups[GroupIndex].VFInput.HairPositionBuffer; };
 	inline FRHIShaderResourceView* GetPreviousPositionSRV(uint32 GroupIndex) const	{ check(Data.InterpolationOutput); return Data.InterpolationOutput->HairGroups[GroupIndex].VFInput.HairPreviousPositionBuffer; }
 	inline FRHIShaderResourceView* GetAttributeSRV(uint32 GroupIndex) const			{ check(Data.InterpolationOutput); return Data.InterpolationOutput->HairGroups[GroupIndex].VFInput.HairAttributeBuffer; }
diff --git a/Engine/Source/Runtime/Engine/Private/InstancedStaticMesh.cpp b/Engine/Source/Runtime/Engine/Private/InstancedStaticMesh.cpp
--- a/Engine/Source/Runtime/Engine/Private/InstancedStaticMesh.cpp
+++ b/Engine/Source/Runtime/Engine/Private/InstancedStaticMesh.cpp
@@ -778,18 +778,7 @@ void FInstancedStaticMeshVertexFactory::InitRHI()
 }
 
 IMPLEMENT_VERTEX_FACTORY_PARAMETER_TYPE(FInstancedStaticMeshVertexFactory, SF_Vertex, FInstancedStaticMeshVertexFactoryShaderParameters);
-
-FVertexFactoryShaderParameters* FInstancedStaticMeshVertexFactory::ConstructShaderParameters(EShaderFrequency ShaderFrequency)
-{
-	switch (ShaderFrequency)
-	{
-	case SF_Vertex:
-	case SF_RayHitGroup:
-		return new FInstancedStaticMeshVertexFactoryShaderParameters();
-	default:
-		return nullptr;
-	}
-}
+IMPLEMENT_VERTEX_FACTORY_PARAMETER_TYPE(FInstancedStaticMeshVertexFactory, SF_RayHitGroup, FInstancedStaticMeshVertexFactoryShaderParameters);
 
 IMPLEMENT_VERTEX_FACTORY_TYPE_EX(FInstancedStaticMeshVertexFactory,"/Engine/Private/LocalVertexFactory.ush",true,true,true,true,true,true,false);
 
diff --git a/Engine/Source/Runtime/Engine/Private/InstancedStaticMesh.h b/Engine/Source/Runtime/Engine/Private/InstancedStaticMesh.h
--- a/Engine/Source/Runtime/Engine/Private/InstancedStaticMesh.h
+++ b/Engine/Source/Runtime/Engine/Private/InstancedStaticMesh.h
@@ -239,8 +239,6 @@ public:
 	// FRenderResource interface.
 	virtual void InitRHI() override;
 
-	static FVertexFactoryShaderParameters* ConstructShaderParameters(EShaderFrequency ShaderFrequency);
-
 	/** Make sure we account for changes in the signature of GetStaticBatchElementVisibility() */
 	static CONSTEXPR uint32 NumBitsForVisibilityMask()
 	{		
