Fixed localization conflicts

#rnx

diff --git a/Engine/Plugins/Developer/Concert/DisasterRecoveryClient/Source/DisasterRecoveryClient/Private/DisasterRecoverySessionManager.cpp b/Engine/Plugins/Developer/Concert/DisasterRecoveryClient/Source/DisasterRecoveryClient/Private/DisasterRecoverySessionManager.cpp
--- a/Engine/Plugins/Developer/Concert/DisasterRecoveryClient/Source/DisasterRecoveryClient/Private/DisasterRecoverySessionManager.cpp
+++ b/Engine/Plugins/Developer/Concert/DisasterRecoveryClient/Source/DisasterRecoveryClient/Private/DisasterRecoverySessionManager.cpp
@@ -566,7 +566,7 @@ TFuture<TVariant<TSharedPtr<FConcertActivityStream>, FText>> FDisasterRecoverySe
 		if (!ArchivedSessionId.IsValid()) // Not found?
 		{
 			OnSessionNotFoundInternal(Session->RepositoryId); // The archive was likely deleted. Mark it to be discarded to be cleaned up at the next start up.
-			OnErrorFn(LOCTEXT("ArchivedSessionNotFound", "Session not found."));
+			OnErrorFn(LOCTEXT("SessionNotFound", "Session not found."));
 			return; // 3b - No session.
 		}
 
diff --git a/Engine/Plugins/Developer/TraceDataFiltering/Source/TraceDataFiltering/Private/SFilterPresetList.cpp b/Engine/Plugins/Developer/TraceDataFiltering/Source/TraceDataFiltering/Private/SFilterPresetList.cpp
--- a/Engine/Plugins/Developer/TraceDataFiltering/Source/TraceDataFiltering/Private/SFilterPresetList.cpp
+++ b/Engine/Plugins/Developer/TraceDataFiltering/Source/TraceDataFiltering/Private/SFilterPresetList.cpp
@@ -352,7 +352,7 @@ void SFilterPresetList::GenerateCommonPresetEntries(FMenuBuilder& InSubMenuBuild
 		, FText::GetEmpty()
 	);
 	
-	const FText DisplayText = UserPreset->IsLocal() ? LOCTEXT("MakeSharedPresetLabel", "Make Shared Preset") : LOCTEXT("MakeSharedPresetLabel", "Make Local Preset");
+	const FText DisplayText = UserPreset->IsLocal() ? LOCTEXT("MakeSharedPresetLabel", "Make Shared Preset") : LOCTEXT("MakeLocalPresetLabel", "Make Local Preset");
 	const FText TooltipText = UserPreset->IsLocal() ? LOCTEXT("MakeSharedPresetToolTip", "Makes this preset a Shared User Preset (Config INI file has to be writable)") : LOCTEXT("MakeLocalPresetToolTip", "Makes this preset a Local Preset");
 
 	InSubMenuBuilder.AddMenuEntry(DisplayText,
diff --git a/Engine/Plugins/Developer/TraceSourceFiltering/Source/SourceFilteringEditor/Private/EditorSessionSourceFilterService.cpp b/Engine/Plugins/Developer/TraceSourceFiltering/Source/SourceFilteringEditor/Private/EditorSessionSourceFilterService.cpp
--- a/Engine/Plugins/Developer/TraceSourceFiltering/Source/SourceFilteringEditor/Private/EditorSessionSourceFilterService.cpp
+++ b/Engine/Plugins/Developer/TraceSourceFiltering/Source/SourceFilteringEditor/Private/EditorSessionSourceFilterService.cpp
@@ -87,7 +87,7 @@ void FEditorSessionSourceFilterService::AddFilterToSet(TSharedRef<const IFilterO
 	FSoftClassPath ClassPath(FilterClassName);
 	if (UClass* Class = ClassPath.TryLoadClass<UDataSourceFilter>())
 	{
-		const FScopedTransaction Transaction(*FEditorSessionSourceFilterService::TransactionContext, LOCTEXT("AddFilter", "Add Filter to Filter Set"), FilterCollection);
+		const FScopedTransaction Transaction(*FEditorSessionSourceFilterService::TransactionContext, LOCTEXT("AddFilterToSet", "Add Filter to Filter Set"), FilterCollection);
 		FilterCollection->Modify();
 
 		FilterCollection->AddFilterOfClassToSet(Class, CastChecked<UDataSourceFilterSet>(FilterSet->GetFilter()));
diff --git a/Engine/Plugins/Editor/GeometryMode/Source/GeometryMode/Private/GeometryModifiers.cpp b/Engine/Plugins/Editor/GeometryMode/Source/GeometryMode/Private/GeometryModifiers.cpp
--- a/Engine/Plugins/Editor/GeometryMode/Source/GeometryMode/Private/GeometryModifiers.cpp
+++ b/Engine/Plugins/Editor/GeometryMode/Source/GeometryMode/Private/GeometryModifiers.cpp
@@ -2070,7 +2070,7 @@ static ABrush* ClipBrushAgainstPlane( const FPlane& InPlane, ABrush* InBrush)
 UGeomModifier_Clip::UGeomModifier_Clip(const FObjectInitializer& ObjectInitializer)
 	: Super(ObjectInitializer)
 {
-	Description = NSLOCTEXT("UnrealEd", "BrushClip", "Clip");
+	Description = NSLOCTEXT("UnrealEd", "Clip", "Clip");
 	Tooltip = NSLOCTEXT("UnrealEd.GeomModifier_Clip", "Tooltip", "Given a dividing plane, cut the geometry into two pieces, optionally discarding one of them. This operation only works in an orthographic viewport.  Define the vertices of the dividing plane with the space bar, and press Enter to apply.");
 	bFlipNormal = false;
 	bSplit = false;
diff --git a/Engine/Plugins/Enterprise/DataprepEditor/Source/DataprepEditor/Private/SchemaActions/DataprepDragDropOp.cpp b/Engine/Plugins/Enterprise/DataprepEditor/Source/DataprepEditor/Private/SchemaActions/DataprepDragDropOp.cpp
--- a/Engine/Plugins/Enterprise/DataprepEditor/Source/DataprepEditor/Private/SchemaActions/DataprepDragDropOp.cpp
+++ b/Engine/Plugins/Enterprise/DataprepEditor/Source/DataprepEditor/Private/SchemaActions/DataprepDragDropOp.cpp
@@ -109,17 +109,17 @@ void FDataprepDragDropOp::HoverTargetChanged()
 	else if(Cast<UDataprepGraphActionStepNode>(GetHoveredNode()) != nullptr)
 	{
 		bDropTargetValid = true;
-		DrapDropText = LOCTEXT("CopyDataprepActionStepNode", "Add/Insert step to location");
+		DrapDropText = LOCTEXT("TargetIsDataprepActionNode", "Add/Insert step to location");
 	}
 	else if(Cast<UDataprepGraphActionNode>(GetHoveredNode()) != nullptr)
 	{
 		bDropTargetValid = true;
-		DrapDropText = LOCTEXT("CopyDataprepActionAssetNode", "Add/Insert step to location");
+		DrapDropText = LOCTEXT("TargetIsDataprepActionNode", "Add/Insert step to location");
 	}
 	else if(Cast<UDataprepGraphRecipeNode>(GetHoveredNode()) != nullptr)
 	{
 		bDropTargetValid = true;
-		DrapDropText = LOCTEXT("InsertDataprepActionAssetNode", "Insert action to location");
+		DrapDropText = LOCTEXT("TargetIsDataprepRecipeNode", "Insert action to location");
 	}
 	else if ( UEdGraph* EdGraph = GetHoveredGraph() )
 	{
diff --git a/Engine/Plugins/Enterprise/LidarPointCloud/Source/LidarPointCloudEditor/Private/LidarPointCloudImportUI.cpp b/Engine/Plugins/Enterprise/LidarPointCloud/Source/LidarPointCloudEditor/Private/LidarPointCloudImportUI.cpp
--- a/Engine/Plugins/Enterprise/LidarPointCloud/Source/LidarPointCloudEditor/Private/LidarPointCloudImportUI.cpp
+++ b/Engine/Plugins/Enterprise/LidarPointCloud/Source/LidarPointCloudEditor/Private/LidarPointCloudImportUI.cpp
@@ -36,7 +36,7 @@ void SLidarPointCloudOptionWindow::Construct(const FArguments& InArgs)
 
 	ImportButton = SNew(SButton)
 					.HAlign(HAlign_Center)
-					.Text(bIsReimport ? LOCTEXT("LidarPointCloudOptionWindow_Import", "Reimport") : LOCTEXT("LidarPointCloudOptionWindow_Import", "Import"))
+					.Text(bIsReimport ? LOCTEXT("LidarPointCloudOptionWindow_Reimport", "Reimport") : LOCTEXT("LidarPointCloudOptionWindow_Import", "Import"))
 					.IsEnabled(this, &SLidarPointCloudOptionWindow::CanImport)
 					.OnClicked(this, &SLidarPointCloudOptionWindow::OnImport);
 
diff --git a/Engine/Plugins/Experimental/HairStrands/Source/HairStrandsEditor/Private/GroomCreateBindingOptionsWindow.cpp b/Engine/Plugins/Experimental/HairStrands/Source/HairStrandsEditor/Private/GroomCreateBindingOptionsWindow.cpp
--- a/Engine/Plugins/Experimental/HairStrands/Source/HairStrandsEditor/Private/GroomCreateBindingOptionsWindow.cpp
+++ b/Engine/Plugins/Experimental/HairStrands/Source/HairStrandsEditor/Private/GroomCreateBindingOptionsWindow.cpp
@@ -22,7 +22,7 @@
 #include "Widgets/Text/STextBlock.h"
 #include "Editor.h"
 
-#define LOCTEXT_NAMESPACE "GroomOptions"
+#define LOCTEXT_NAMESPACE "GroomCreateBindingOptionsWindow"
 
 void SGroomCreateBindingOptionsWindow::Construct(const FArguments& InArgs)
 {
@@ -54,7 +54,7 @@ void SGroomCreateBindingOptionsWindow::Construct(const FArguments& InArgs)
 				[
 					SNew(STextBlock)
 					.Font(FEditorStyle::GetFontStyle("CurveEd.LabelFont"))
-					.Text(LOCTEXT("GroomBindingOptionsWindow_CurrentFile", "Current File: "))
+					.Text(LOCTEXT("CurrentFile", "Current File: "))
 				]
 				+ SHorizontalBox::Slot()
 				.Padding(5, 0, 0, 0)
@@ -94,7 +94,7 @@ void SGroomCreateBindingOptionsWindow::Construct(const FArguments& InArgs)
 			[
 				SNew(SButton)
 				.HAlign(HAlign_Center)
-				.Text(LOCTEXT("GroomBindingOptionsWindow_Cancel", "Cancel"))
+				.Text(LOCTEXT("Cancel", "Cancel"))
 				.OnClicked(this, &SGroomCreateBindingOptionsWindow::OnCancel)
 			]
 		]
@@ -148,7 +148,7 @@ TSharedPtr<SGroomCreateBindingOptionsWindow> DisplayOptions(UGroomCreateBindingO
 
 TSharedPtr<SGroomCreateBindingOptionsWindow> SGroomCreateBindingOptionsWindow::DisplayCreateBindingOptions(UGroomCreateBindingOptions* BindingOptions)
 {
-	return DisplayOptions(BindingOptions, EGroomBindingOptionsVisibility::BuildOptions, LOCTEXT("GroomBuildOptionsWindow_WindowTitle", "Groom Binding Options"), LOCTEXT("GroomBindingOptionsWindow_Build", "Create"));
+	return DisplayOptions(BindingOptions, EGroomBindingOptionsVisibility::BuildOptions, LOCTEXT("WindowTitle", "Groom Binding Options"), LOCTEXT("Build", "Create"));
 }
 
 static UObject* InternalCreateNewBindAsset(FName InAssetName, UObject* InParent, UGroomAsset* GroomAsset, USkeletalMesh* SourceSkelMesh, USkeletalMesh* TargetSkelMesh)
diff --git a/Engine/Plugins/Experimental/HairStrands/Source/HairStrandsEditor/Private/GroomImportOptionsWindow.cpp b/Engine/Plugins/Experimental/HairStrands/Source/HairStrandsEditor/Private/GroomImportOptionsWindow.cpp
--- a/Engine/Plugins/Experimental/HairStrands/Source/HairStrandsEditor/Private/GroomImportOptionsWindow.cpp
+++ b/Engine/Plugins/Experimental/HairStrands/Source/HairStrandsEditor/Private/GroomImportOptionsWindow.cpp
@@ -15,7 +15,7 @@
 #include "Widgets/SBoxPanel.h"
 #include "Widgets/Text/STextBlock.h"
 
-#define LOCTEXT_NAMESPACE "GroomOptions"
+#define LOCTEXT_NAMESPACE "GroomImportOptionsWindow"
 
 void SGroomImportOptionsWindow::Construct(const FArguments& InArgs)
 {
@@ -47,7 +47,7 @@ void SGroomImportOptionsWindow::Construct(const FArguments& InArgs)
 				[
 					SNew(STextBlock)
 					.Font(FEditorStyle::GetFontStyle("CurveEd.LabelFont"))
-					.Text(LOCTEXT("GroomOptionsWindow_CurrentFile", "Current File: "))
+					.Text(LOCTEXT("CurrentFile", "Current File: "))
 				]
 				+ SHorizontalBox::Slot()
 				.Padding(5, 0, 0, 0)
@@ -87,7 +87,7 @@ void SGroomImportOptionsWindow::Construct(const FArguments& InArgs)
 			[
 				SNew(SButton)
 				.HAlign(HAlign_Center)
-				.Text(LOCTEXT("GroomOptionsWindow_Cancel", "Cancel"))
+				.Text(LOCTEXT("Cancel", "Cancel"))
 				.OnClicked(this, &SGroomImportOptionsWindow::OnCancel)
 			]
 		]
@@ -159,12 +159,12 @@ TSharedPtr<SGroomImportOptionsWindow> DisplayOptions(UGroomImportOptions* Import
 
 TSharedPtr<SGroomImportOptionsWindow> SGroomImportOptionsWindow::DisplayImportOptions(UGroomImportOptions* ImportOptions, const FString& FilePath)
 {
-	return DisplayOptions(ImportOptions, FilePath, EGroomOptionsVisibility::All, LOCTEXT("GroomOptionsWindow_WindowTitle", "Groom Import Options"), LOCTEXT("GroomOptionsWindow_Import", "Import"));
+	return DisplayOptions(ImportOptions, FilePath, EGroomOptionsVisibility::All, LOCTEXT("WindowTitle", "Groom Import Options"), LOCTEXT("Import", "Import"));
 }
 
 TSharedPtr<SGroomImportOptionsWindow> SGroomImportOptionsWindow::DisplayRebuildOptions(UGroomImportOptions* ImportOptions, const FString& FilePath)
 {
-	return DisplayOptions(ImportOptions, FilePath, EGroomOptionsVisibility::BuildOptions, LOCTEXT("GroomBuildOptionsWindow_WindowTitle", "Groom Build Options"), LOCTEXT("GroomOptionsWindow_Build", "Build"));
+	return DisplayOptions(ImportOptions, FilePath, EGroomOptionsVisibility::BuildOptions, LOCTEXT("WindowTitle", "Groom Build Options"), LOCTEXT("Build", "Build"));
 }
 
 
diff --git a/Engine/Plugins/FX/Niagara/Source/NiagaraEditor/Private/NiagaraNodeFunctionCall.cpp b/Engine/Plugins/FX/Niagara/Source/NiagaraEditor/Private/NiagaraNodeFunctionCall.cpp
--- a/Engine/Plugins/FX/Niagara/Source/NiagaraEditor/Private/NiagaraNodeFunctionCall.cpp
+++ b/Engine/Plugins/FX/Niagara/Source/NiagaraEditor/Private/NiagaraNodeFunctionCall.cpp
@@ -681,7 +681,7 @@ void UNiagaraNodeFunctionCall::UpdateNodeErrorMessage()
 			{
 				FFormatNamedArguments Args;
 				Args.Add(TEXT("Message"), Signature.ExperimentalMessage);
-				UEdGraphNode::NodeUpgradeMessage = FText::Format(LOCTEXT("FunctionExperimentalReason", "This function is marked as experimental, reason: {Message}"), Args);
+				UEdGraphNode::NodeUpgradeMessage = FText::Format(LOCTEXT("FunctionExperimentalReason", "This function is marked as experimental, reason: {Message}."), Args);
 			}
 		}
 	}
diff --git a/Engine/Plugins/Media/BlackmagicMedia/Source/BlackmagicMediaOutput/Private/BlackmagicMediaOutput.cpp b/Engine/Plugins/Media/BlackmagicMedia/Source/BlackmagicMediaOutput/Private/BlackmagicMediaOutput.cpp
--- a/Engine/Plugins/Media/BlackmagicMedia/Source/BlackmagicMediaOutput/Private/BlackmagicMediaOutput.cpp
+++ b/Engine/Plugins/Media/BlackmagicMedia/Source/BlackmagicMediaOutput/Private/BlackmagicMediaOutput.cpp
@@ -209,3 +209,4 @@ void UBlackmagicMediaOutput::PostEditChangeChainProperty(struct FPropertyChanged
 
 #endif //WITH_EDITOR
 
+#undef LOCTEXT_NAMESPACE
diff --git a/Engine/Plugins/VirtualProduction/TimedDataMonitor/Source/TimedDataMonitorEditor/Private/STimedDataListView.cpp b/Engine/Plugins/VirtualProduction/TimedDataMonitor/Source/TimedDataMonitorEditor/Private/STimedDataListView.cpp
--- a/Engine/Plugins/VirtualProduction/TimedDataMonitor/Source/TimedDataMonitorEditor/Private/STimedDataListView.cpp
+++ b/Engine/Plugins/VirtualProduction/TimedDataMonitor/Source/TimedDataMonitorEditor/Private/STimedDataListView.cpp
@@ -214,8 +214,8 @@ TSharedRef<SWidget> STimedDataInputTableRow::GenerateWidgetForColumn(const FName
 	if (TimedDataListView::HeaderIdName_Enable == ColumnName)
 	{
 		const FText Tooltip = Item->bIsInput
-			? LOCTEXT("EnabledChannelToolTip", "Toggles all channels from this input.")
-			: LOCTEXT("EnabledChannelToolTip", "Toggles whether this channel will collect stats and be used when calibrating.");
+			? LOCTEXT("ToggleAllChannelsToolTip", "Toggles all channels from this input.")
+			: LOCTEXT("ToggleChannelToolTip", "Toggles whether this channel will collect stats and be used when calibrating.");
 		return SNew(SCheckBox)
 			.Style(FTimedDataMonitorStyle::Get(), "CheckBox.Enable")
 			.ToolTipText(Tooltip)
diff --git a/Engine/Source/Developer/AssetTools/Private/AssetTypeActions/AssetTypeActions_PhysicalMaterialMask.cpp b/Engine/Source/Developer/AssetTools/Private/AssetTypeActions/AssetTypeActions_PhysicalMaterialMask.cpp
--- a/Engine/Source/Developer/AssetTools/Private/AssetTypeActions/AssetTypeActions_PhysicalMaterialMask.cpp
+++ b/Engine/Source/Developer/AssetTools/Private/AssetTypeActions/AssetTypeActions_PhysicalMaterialMask.cpp
@@ -153,3 +153,5 @@ void FAssetTypeActions_PhysicalMaterialMask::ExecuteDebug(TWeakObjectPtr<UPhysic
 		PhysMatMask->DumpMaskData();
 	}
 }
+
+#undef LOCTEXT_NAMESPACE
diff --git a/Engine/Source/Editor/KismetWidgets/Private/CreateBlueprintFromActorDialog.cpp b/Engine/Source/Editor/KismetWidgets/Private/CreateBlueprintFromActorDialog.cpp
--- a/Engine/Source/Editor/KismetWidgets/Private/CreateBlueprintFromActorDialog.cpp
+++ b/Engine/Source/Editor/KismetWidgets/Private/CreateBlueprintFromActorDialog.cpp
@@ -456,7 +456,7 @@ void SSCreateBlueprintPicker::Construct(const FArguments& InArgs)
 					SNew(SExpandableArea)
 					.MaxHeight(320.f)
 					.InitiallyCollapsed(false)
-					.AreaTitle(NSLOCTEXT("SClassPickerDialog", "AllClassesAreaTitle", "Parent Class"))
+					.AreaTitle(NSLOCTEXT("SClassPickerDialog", "ParentClassAreaTitle", "Parent Class"))
 					.OnAreaExpansionChanged(this, &SSCreateBlueprintPicker::OnCustomAreaExpansionChanged)
 					.BodyContent()
 					[
diff --git a/Engine/Source/Editor/MovieSceneTools/Private/Channels/BuiltInChannelEditors.cpp b/Engine/Source/Editor/MovieSceneTools/Private/Channels/BuiltInChannelEditors.cpp
--- a/Engine/Source/Editor/MovieSceneTools/Private/Channels/BuiltInChannelEditors.cpp
+++ b/Engine/Source/Editor/MovieSceneTools/Private/Channels/BuiltInChannelEditors.cpp
@@ -242,7 +242,7 @@ TSharedRef<SWidget> CreateKeyEditor(const TMovieSceneChannelHandle<FMovieSceneOb
 
 		auto OnSetObjectLambda = [KeyEditor](const FAssetData& Asset) mutable
 		{
-			FScopedTransaction Transaction(LOCTEXT("SetKey", "Set Enum Key Value"));
+			FScopedTransaction Transaction(LOCTEXT("SetEnumKey", "Set Enum Key Value"));
 			KeyEditor.SetValueWithNotify(Asset.GetAsset(), EMovieSceneDataChangeType::TrackValueChangedRefreshImmediately);
 		};
 
@@ -379,7 +379,7 @@ TSharedRef<SWidget> CreateKeyEditor(const TMovieSceneChannelHandle<FMovieSceneAc
 
 	auto OnSetCurrentValueLambda = [KeyEditor](FMovieSceneActorReferenceKey& ActorKey) mutable
 	{
-		FScopedTransaction Transaction(LOCTEXT("SetKey", "Set Actor Reference Key Value"));
+		FScopedTransaction Transaction(LOCTEXT("SetActorReferenceKey", "Set Actor Reference Key Value"));
 		KeyEditor.SetValueWithNotify(ActorKey, EMovieSceneDataChangeType::TrackValueChangedRefreshImmediately);
 
 		// Look for components to choose
diff --git a/Engine/Source/Editor/StaticMeshEditor/Private/StaticMeshEditorTools.cpp b/Engine/Source/Editor/StaticMeshEditor/Private/StaticMeshEditorTools.cpp
--- a/Engine/Source/Editor/StaticMeshEditor/Private/StaticMeshEditorTools.cpp
+++ b/Engine/Source/Editor/StaticMeshEditor/Private/StaticMeshEditorTools.cpp
@@ -2143,7 +2143,7 @@ void FMeshSectionSettingsLayout::OnSectionForceOpaqueFlagChanged( ECheckBoxState
 	FText TransactionTest = LOCTEXT("StaticMeshEditorSetForceOpaqueSectionFlag", "Staticmesh editor: Set Force Opaque For section, the section will be considered opaque in ray tracing effects");
 	if (NewState == ECheckBoxState::Unchecked)
 	{
-		TransactionTest = LOCTEXT("StaticMeshEditorClearShadowCastingSectionFlag", "Staticmesh editor: Clear Force Opaque For section");
+		TransactionTest = LOCTEXT("StaticMeshEditorClearForceOpaqueSectionFlag", "Staticmesh editor: Clear Force Opaque For section");
 	}
 	FScopedTransaction Transaction(TransactionTest);
 
