Compile fixes for PressureGrid and PhysicsAsset in CIS

diff --git a/Engine/Plugins/Experimental/HairStrands/Source/HairStrandsNiagara/Private/NiagaraDataInterfacePhysicsAsset.cpp b/Engine/Plugins/Experimental/HairStrands/Source/HairStrandsNiagara/Private/NiagaraDataInterfacePhysicsAsset.cpp
--- a/Engine/Plugins/Experimental/HairStrands/Source/HairStrandsNiagara/Private/NiagaraDataInterfacePhysicsAsset.cpp
+++ b/Engine/Plugins/Experimental/HairStrands/Source/HairStrandsNiagara/Private/NiagaraDataInterfacePhysicsAsset.cpp
@@ -569,7 +569,7 @@ void FNDIPhysicsAssetProxy::DestroyPerInstanceData(NiagaraEmitterInstanceBatcher
 
 void FNDIPhysicsAssetProxy::PreStage(FRHICommandList& RHICmdList, const FNiagaraDataInterfaceSetArgs& Context)
 {
-	if (Context.ShaderStageIndex == 0)
+	if (Context.SimulationStageIndex == 0)
 	{
 		FNDIPhysicsAssetData* ProxyData =
 			SystemInstancesToProxyData.Find(Context.SystemInstance);
diff --git a/Engine/Plugins/Experimental/HairStrands/Source/HairStrandsNiagara/Private/NiagaraDataInterfacePressureGrid.cpp b/Engine/Plugins/Experimental/HairStrands/Source/HairStrandsNiagara/Private/NiagaraDataInterfacePressureGrid.cpp
--- a/Engine/Plugins/Experimental/HairStrands/Source/HairStrandsNiagara/Private/NiagaraDataInterfacePressureGrid.cpp
+++ b/Engine/Plugins/Experimental/HairStrands/Source/HairStrandsNiagara/Private/NiagaraDataInterfacePressureGrid.cpp
@@ -277,8 +277,8 @@ bool UNiagaraDataInterfacePressureGrid::InitPerInstanceData(void* PerInstanceDat
 	ENQUEUE_RENDER_COMMAND(FNiagaraDIPushInitialInstanceDataToRT) (
 		[ThisProxy, RT_OutputShaderStages, RT_IterationShaderStages, InstanceID = SystemInstance->GetId(), ElementCount](FRHICommandListImmediate& CmdList)
 	{
-		ThisProxy->OutputShaderStages = RT_OutputShaderStages;
-		ThisProxy->IterationShaderStages = RT_IterationShaderStages;
+		ThisProxy->OutputSimulationStages_DEPRECATED = RT_OutputShaderStages;
+		ThisProxy->IterationSimulationStages_DEPRECATED = RT_IterationShaderStages;
 		ThisProxy->SetElementCount(ElementCount);
 
 		//ThisProxy->InitializePerInstanceData(InstanceID);
