More build fixes

diff --git a/Engine/Plugins/Experimental/HairStrands/Source/HairStrandsNiagara/Private/NiagaraDataInterfaceHairStrands.cpp b/Engine/Plugins/Experimental/HairStrands/Source/HairStrandsNiagara/Private/NiagaraDataInterfaceHairStrands.cpp
--- a/Engine/Plugins/Experimental/HairStrands/Source/HairStrandsNiagara/Private/NiagaraDataInterfaceHairStrands.cpp
+++ b/Engine/Plugins/Experimental/HairStrands/Source/HairStrandsNiagara/Private/NiagaraDataInterfaceHairStrands.cpp
@@ -731,22 +731,22 @@ public:
 
 			const FHairStrandsRootResource::FMeshProjectionLOD* MeshProjection = (HasRootAttachedValue == 1) ? &(HairStrandsBuffer->SourceRootResources->MeshProjectionLODs[0]) : nullptr;
 
-			FShaderResourceViewRHIRef RestTrianglePositionASRV = (HasRootAttachedValue == 1) ?
-				MeshProjection->RestRootTrianglePosition0Buffer.SRV : FNiagaraRenderer::GetDummyFloatBuffer().SRV;
-			FShaderResourceViewRHIRef RestTrianglePositionBSRV = (HasRootAttachedValue == 1) ?
-				MeshProjection->RestRootTrianglePosition1Buffer.SRV : FNiagaraRenderer::GetDummyFloatBuffer().SRV;
-			FShaderResourceViewRHIRef RestTrianglePositionCSRV = (HasRootAttachedValue == 1) ?
-				MeshProjection->RestRootTrianglePosition2Buffer.SRV : FNiagaraRenderer::GetDummyFloatBuffer().SRV;
-
-			FShaderResourceViewRHIRef DeformedTrianglePositionASRV = (HasRootAttachedValue == 1) ?
-				MeshProjection->DeformedRootTrianglePosition0Buffer.SRV : FNiagaraRenderer::GetDummyFloatBuffer().SRV;
-			FShaderResourceViewRHIRef DeformedTrianglePositionBSRV = (HasRootAttachedValue == 1) ?
-				MeshProjection->DeformedRootTrianglePosition1Buffer.SRV : FNiagaraRenderer::GetDummyFloatBuffer().SRV;
-			FShaderResourceViewRHIRef DeformedTrianglePositionCSRV = (HasRootAttachedValue == 1) ?
-				MeshProjection->DeformedRootTrianglePosition2Buffer.SRV : FNiagaraRenderer::GetDummyFloatBuffer().SRV;
-
-			FShaderResourceViewRHIRef RootBarycentricCoordinatesSRV = (HasRootAttachedValue == 1) ? 
-				MeshProjection->RootTriangleBarycentricBuffer.SRV : FNiagaraRenderer::GetDummyFloatBuffer().SRV;
+			FRHIShaderResourceView* RestTrianglePositionASRV = (HasRootAttachedValue == 1) ?
+				MeshProjection->RestRootTrianglePosition0Buffer.SRV.GetReference() : FNiagaraRenderer::GetDummyFloatBuffer();
+			FRHIShaderResourceView* RestTrianglePositionBSRV = (HasRootAttachedValue == 1) ?
+				MeshProjection->RestRootTrianglePosition1Buffer.SRV.GetReference() : FNiagaraRenderer::GetDummyFloatBuffer();
+			FRHIShaderResourceView* RestTrianglePositionCSRV = (HasRootAttachedValue == 1) ?
+				MeshProjection->RestRootTrianglePosition2Buffer.SRV.GetReference() : FNiagaraRenderer::GetDummyFloatBuffer();
+
+			FRHIShaderResourceView* DeformedTrianglePositionASRV = (HasRootAttachedValue == 1) ?
+				MeshProjection->DeformedRootTrianglePosition0Buffer.SRV.GetReference() : FNiagaraRenderer::GetDummyFloatBuffer();
+			FRHIShaderResourceView* DeformedTrianglePositionBSRV = (HasRootAttachedValue == 1) ?
+				MeshProjection->DeformedRootTrianglePosition1Buffer.SRV.GetReference() : FNiagaraRenderer::GetDummyFloatBuffer();
+			FRHIShaderResourceView* DeformedTrianglePositionCSRV = (HasRootAttachedValue == 1) ?
+				MeshProjection->DeformedRootTrianglePosition2Buffer.SRV.GetReference() : FNiagaraRenderer::GetDummyFloatBuffer();
+
+			FRHIShaderResourceView* RootBarycentricCoordinatesSRV = (HasRootAttachedValue == 1) ?
+				MeshProjection->RootTriangleBarycentricBuffer.SRV.GetReference() : FNiagaraRenderer::GetDummyFloatBuffer();
 
 			int32 bNeedSimReset = (ProxyData->TickCount <= ProxyData->ResetTick ? 1 : 0);
 			FVector RestRootOffsetValue = FVector::ZeroVector;
@@ -791,11 +791,11 @@ public:
 		}
 		else
 		{
-			SetUAVParameter(RHICmdList, ComputeShaderRHI, DeformedPositionBuffer, FNiagaraRenderer::GetDummyFloatBuffer().UAV);
-			SetUAVParameter(RHICmdList, ComputeShaderRHI, BoundingBoxBuffer, FNiagaraRenderer::GetDummyUIntBuffer().UAV);
-			SetSRVParameter(RHICmdList, ComputeShaderRHI, NodeBoundBuffer, FNiagaraRenderer::GetDummyUIntBuffer().SRV);
-			SetSRVParameter(RHICmdList, ComputeShaderRHI, CurvesOffsetsBuffer, FNiagaraRenderer::GetDummyUIntBuffer().SRV);
-			SetSRVParameter(RHICmdList, ComputeShaderRHI, RestPositionBuffer, FNiagaraRenderer::GetDummyFloatBuffer().SRV);
+			SetUAVParameter(RHICmdList, ComputeShaderRHI, DeformedPositionBuffer, Context.Batcher->GetEmptyRWBufferFromPool(RHICmdList, PF_R32_FLOAT));
+			SetUAVParameter(RHICmdList, ComputeShaderRHI, BoundingBoxBuffer, Context.Batcher->GetEmptyRWBufferFromPool(RHICmdList, PF_R32_UINT));
+			SetSRVParameter(RHICmdList, ComputeShaderRHI, NodeBoundBuffer, FNiagaraRenderer::GetDummyUIntBuffer());
+			SetSRVParameter(RHICmdList, ComputeShaderRHI, CurvesOffsetsBuffer, FNiagaraRenderer::GetDummyUIntBuffer());
+			SetSRVParameter(RHICmdList, ComputeShaderRHI, RestPositionBuffer, FNiagaraRenderer::GetDummyFloatBuffer());
 
 			SetShaderValue(RHICmdList, ComputeShaderRHI, WorldTransform, FMatrix::Identity);
 			SetShaderValue(RHICmdList, ComputeShaderRHI, WorldInverse, FMatrix::Identity);
@@ -813,15 +813,15 @@ public:
 			SetShaderValue(RHICmdList, ComputeShaderRHI, RestPositionOffset, FVector(0,0,0));
 			SetShaderValue(RHICmdList, ComputeShaderRHI, DeformedPositionOffset, FVector(0,0,0));
 
-			SetSRVParameter(RHICmdList, ComputeShaderRHI, RestTrianglePositionABuffer, FNiagaraRenderer::GetDummyFloatBuffer().SRV);
-			SetSRVParameter(RHICmdList, ComputeShaderRHI, RestTrianglePositionBBuffer, FNiagaraRenderer::GetDummyFloatBuffer().SRV);
-			SetSRVParameter(RHICmdList, ComputeShaderRHI, RestTrianglePositionCBuffer, FNiagaraRenderer::GetDummyFloatBuffer().SRV);
+			SetSRVParameter(RHICmdList, ComputeShaderRHI, RestTrianglePositionABuffer, FNiagaraRenderer::GetDummyFloatBuffer());
+			SetSRVParameter(RHICmdList, ComputeShaderRHI, RestTrianglePositionBBuffer, FNiagaraRenderer::GetDummyFloatBuffer());
+			SetSRVParameter(RHICmdList, ComputeShaderRHI, RestTrianglePositionCBuffer, FNiagaraRenderer::GetDummyFloatBuffer());
 
-			SetSRVParameter(RHICmdList, ComputeShaderRHI, DeformedTrianglePositionABuffer, FNiagaraRenderer::GetDummyFloatBuffer().SRV);
-			SetSRVParameter(RHICmdList, ComputeShaderRHI, DeformedTrianglePositionBBuffer, FNiagaraRenderer::GetDummyFloatBuffer().SRV);
-			SetSRVParameter(RHICmdList, ComputeShaderRHI, DeformedTrianglePositionCBuffer, FNiagaraRenderer::GetDummyFloatBuffer().SRV);
+			SetSRVParameter(RHICmdList, ComputeShaderRHI, DeformedTrianglePositionABuffer, FNiagaraRenderer::GetDummyFloatBuffer());
+			SetSRVParameter(RHICmdList, ComputeShaderRHI, DeformedTrianglePositionBBuffer, FNiagaraRenderer::GetDummyFloatBuffer());
+			SetSRVParameter(RHICmdList, ComputeShaderRHI, DeformedTrianglePositionCBuffer, FNiagaraRenderer::GetDummyFloatBuffer());
 
-			SetSRVParameter(RHICmdList, ComputeShaderRHI, RootBarycentricCoordinatesBuffer, FNiagaraRenderer::GetDummyFloatBuffer().SRV);
+			SetSRVParameter(RHICmdList, ComputeShaderRHI, RootBarycentricCoordinatesBuffer, FNiagaraRenderer::GetDummyFloatBuffer());
 		}
 	}
 
diff --git a/Engine/Plugins/Experimental/HairStrands/Source/HairStrandsNiagara/Private/NiagaraDataInterfacePhysicsAsset.cpp b/Engine/Plugins/Experimental/HairStrands/Source/HairStrandsNiagara/Private/NiagaraDataInterfacePhysicsAsset.cpp
--- a/Engine/Plugins/Experimental/HairStrands/Source/HairStrandsNiagara/Private/NiagaraDataInterfacePhysicsAsset.cpp
+++ b/Engine/Plugins/Experimental/HairStrands/Source/HairStrandsNiagara/Private/NiagaraDataInterfacePhysicsAsset.cpp
@@ -493,13 +493,13 @@ public:
 		}
 		else
 		{
-			SetSRVParameter(RHICmdList, ComputeShaderRHI, CurrentTransformBuffer, FNiagaraRenderer::GetDummyFloatBuffer().SRV);
-			SetSRVParameter(RHICmdList, ComputeShaderRHI, PreviousTransformBuffer, FNiagaraRenderer::GetDummyFloatBuffer().SRV);
-			SetSRVParameter(RHICmdList, ComputeShaderRHI, PreviousInverseBuffer, FNiagaraRenderer::GetDummyFloatBuffer().SRV);
-			SetSRVParameter(RHICmdList, ComputeShaderRHI, InverseTransformBuffer, FNiagaraRenderer::GetDummyFloatBuffer().SRV);
-			SetSRVParameter(RHICmdList, ComputeShaderRHI, RestTransformBuffer, FNiagaraRenderer::GetDummyFloatBuffer().SRV);
-			SetSRVParameter(RHICmdList, ComputeShaderRHI, RestInverseBuffer, FNiagaraRenderer::GetDummyFloatBuffer().SRV);
-			SetSRVParameter(RHICmdList, ComputeShaderRHI, ElementExtentBuffer, FNiagaraRenderer::GetDummyFloatBuffer().SRV);
+			SetSRVParameter(RHICmdList, ComputeShaderRHI, CurrentTransformBuffer, FNiagaraRenderer::GetDummyFloatBuffer());
+			SetSRVParameter(RHICmdList, ComputeShaderRHI, PreviousTransformBuffer, FNiagaraRenderer::GetDummyFloatBuffer());
+			SetSRVParameter(RHICmdList, ComputeShaderRHI, PreviousInverseBuffer, FNiagaraRenderer::GetDummyFloatBuffer());
+			SetSRVParameter(RHICmdList, ComputeShaderRHI, InverseTransformBuffer, FNiagaraRenderer::GetDummyFloatBuffer());
+			SetSRVParameter(RHICmdList, ComputeShaderRHI, RestTransformBuffer, FNiagaraRenderer::GetDummyFloatBuffer());
+			SetSRVParameter(RHICmdList, ComputeShaderRHI, RestInverseBuffer, FNiagaraRenderer::GetDummyFloatBuffer());
+			SetSRVParameter(RHICmdList, ComputeShaderRHI, ElementExtentBuffer, FNiagaraRenderer::GetDummyFloatBuffer());
 
 			static const FElementOffset DummyOffsets(0,0,0,0);
 			SetShaderValue(RHICmdList, ComputeShaderRHI, ElementOffsets, DummyOffsets);
diff --git a/Engine/Plugins/Experimental/HairStrands/Source/HairStrandsNiagara/Private/NiagaraDataInterfacePressureGrid.cpp b/Engine/Plugins/Experimental/HairStrands/Source/HairStrandsNiagara/Private/NiagaraDataInterfacePressureGrid.cpp
--- a/Engine/Plugins/Experimental/HairStrands/Source/HairStrandsNiagara/Private/NiagaraDataInterfacePressureGrid.cpp
+++ b/Engine/Plugins/Experimental/HairStrands/Source/HairStrandsNiagara/Private/NiagaraDataInterfacePressureGrid.cpp
@@ -209,8 +209,8 @@ struct FNDIPressureGridParametersCS : public FNiagaraDataInterfaceParametersCS
 		}
 		else
 		{
-			SetUAVParameter(RHICmdList, ComputeShaderRHI, GridDestinationBuffer, FNiagaraRenderer::GetDummyUIntBuffer().UAV);
-			SetSRVParameter(RHICmdList, ComputeShaderRHI, GridCurrentBuffer, FNiagaraRenderer::GetDummyUIntBuffer().SRV);
+			SetUAVParameter(RHICmdList, ComputeShaderRHI, GridDestinationBuffer, Context.Batcher->GetEmptyRWBufferFromPool(RHICmdList, PF_R32_UINT));
+			SetSRVParameter(RHICmdList, ComputeShaderRHI, GridCurrentBuffer, FNiagaraRenderer::GetDummyUIntBuffer());
 
 			SetShaderValue(RHICmdList, ComputeShaderRHI, GridOrigin, FVector4(0, 0, 0, 0));
 			SetShaderValue(RHICmdList, ComputeShaderRHI, GridSize, FIntVector());
diff --git a/Engine/Plugins/FX/Niagara/Source/Niagara/Classes/NiagaraEmitterInstanceBatcher.h b/Engine/Plugins/FX/Niagara/Source/Niagara/Classes/NiagaraEmitterInstanceBatcher.h
--- a/Engine/Plugins/FX/Niagara/Source/Niagara/Classes/NiagaraEmitterInstanceBatcher.h
+++ b/Engine/Plugins/FX/Niagara/Source/Niagara/Classes/NiagaraEmitterInstanceBatcher.h
@@ -145,8 +145,8 @@ public:
 	/** Loop over all the data interfaces and call the poststage methods */
 	void PostStageInterface(const FNiagaraGPUSystemTick& Tick, FNiagaraComputeInstanceData *Instance, FRHICommandList &RHICmdList, const FNiagaraShaderRef& ComputeShader, const uint32 SimulationStageIndex) const;
 
-	FRHIUnorderedAccessView* GetEmptyRWBufferFromPool(FRHICommandList& RHICmdList, EPixelFormat Format) const { return GetEmptyUAVFromPool(RHICmdList, Format, false); }
-	FRHIUnorderedAccessView* GetEmptyRWTextureFromPool(FRHICommandList& RHICmdList, EPixelFormat Format) const { return GetEmptyUAVFromPool(RHICmdList, Format, true); }
+	NIAGARA_API FRHIUnorderedAccessView* GetEmptyRWBufferFromPool(FRHICommandList& RHICmdList, EPixelFormat Format) const { return GetEmptyUAVFromPool(RHICmdList, Format, false); }
+	NIAGARA_API FRHIUnorderedAccessView* GetEmptyRWTextureFromPool(FRHICommandList& RHICmdList, EPixelFormat Format) const { return GetEmptyUAVFromPool(RHICmdList, Format, true); }
 
 	/** Get the shared SortManager, used in the rendering loop to call FGPUSortManager::OnPreRender() and FGPUSortManager::OnPostRenderOpaque() */
 	virtual FGPUSortManager* GetGPUSortManager() const override;
diff --git a/Engine/Plugins/FX/Niagara/Source/Niagara/Private/NiagaraDataSet.cpp b/Engine/Plugins/FX/Niagara/Source/Niagara/Private/NiagaraDataSet.cpp
--- a/Engine/Plugins/FX/Niagara/Source/Niagara/Private/NiagaraDataSet.cpp
+++ b/Engine/Plugins/FX/Niagara/Source/Niagara/Private/NiagaraDataSet.cpp
@@ -993,8 +993,8 @@ void FNiagaraDataBuffer::SetShaderParams(FNiagaraShader *Shader, FRHICommandList
 	{
 		const bool InstancesAllocated = GetNumInstancesAllocated() > 0;
 
-		SetSRVParameter(CommandList, ComputeShader, Shader->FloatInputBufferParam, InstancesAllocated ? GetGPUBufferFloat().SRV : FNiagaraRenderer::GetDummyFloatBuffer());
-		SetSRVParameter(CommandList, ComputeShader, Shader->IntInputBufferParam, InstancesAllocated ? GetGPUBufferInt().SRV : FNiagaraRenderer::GetDummyIntBuffer());
+		SetSRVParameter(CommandList, ComputeShader, Shader->FloatInputBufferParam, InstancesAllocated ? GetGPUBufferFloat().SRV.GetReference() : FNiagaraRenderer::GetDummyFloatBuffer());
+		SetSRVParameter(CommandList, ComputeShader, Shader->IntInputBufferParam, InstancesAllocated ? GetGPUBufferInt().SRV.GetReference() : FNiagaraRenderer::GetDummyIntBuffer());
 		SetShaderValue(CommandList, ComputeShader, Shader->ComponentBufferSizeReadParam, SafeBufferSize);
 	}
 	else
diff --git a/Engine/Plugins/FX/Niagara/Source/Niagara/Private/NiagaraEmitterInstanceBatcher.cpp b/Engine/Plugins/FX/Niagara/Source/Niagara/Private/NiagaraEmitterInstanceBatcher.cpp
--- a/Engine/Plugins/FX/Niagara/Source/Niagara/Private/NiagaraEmitterInstanceBatcher.cpp
+++ b/Engine/Plugins/FX/Niagara/Source/Niagara/Private/NiagaraEmitterInstanceBatcher.cpp
@@ -1295,7 +1295,7 @@ void NiagaraEmitterInstanceBatcher::Run(const FNiagaraGPUSystemTick& Tick, const
 	// set the shader and data set params 
 	//
 	const bool bRequiresPersistentIDs = Context->MainDataSet->RequiresPersistentIDs();
-	SetSRVParameter(RHICmdList, Shader.GetComputeShader(), Shader->FreeIDBufferParam, bRequiresPersistentIDs ? Context->MainDataSet->GetGPUFreeIDs().SRV : FNiagaraRenderer::GetDummyIntBuffer());
+	SetSRVParameter(RHICmdList, Shader.GetComputeShader(), Shader->FreeIDBufferParam, bRequiresPersistentIDs ? Context->MainDataSet->GetGPUFreeIDs().SRV.GetReference() : FNiagaraRenderer::GetDummyIntBuffer());
 	CurrentData.SetShaderParams(Shader.GetShader(), RHICmdList, true);
 	DestinationData.SetShaderParams(Shader.GetShader(), RHICmdList, false);
 
