Fixed crash on startup when the Groom plug-in is enabled.

FDirectCopyBoundingBoxCS had escaped the conversion to the new shader serialization system.

diff --git a/Engine/Plugins/Experimental/HairStrands/Source/HairStrandsNiagara/Private/NiagaraDataInterfaceHairStrands.cpp b/Engine/Plugins/Experimental/HairStrands/Source/HairStrandsNiagara/Private/NiagaraDataInterfaceHairStrands.cpp
--- a/Engine/Plugins/Experimental/HairStrands/Source/HairStrandsNiagara/Private/NiagaraDataInterfaceHairStrands.cpp
+++ b/Engine/Plugins/Experimental/HairStrands/Source/HairStrandsNiagara/Private/NiagaraDataInterfaceHairStrands.cpp
@@ -266,7 +266,7 @@ class FCopyBoundingBoxCS : public FGlobalShader
 		SHADER_PARAMETER(uint32, NumElements)
 		SHADER_PARAMETER_UAV(RWBuffer, BoundingBoxBuffer)
 		SHADER_PARAMETER_UAV(RWBuffer, OutNodeBoundBuffer)
-		END_SHADER_PARAMETER_STRUCT()
+	END_SHADER_PARAMETER_STRUCT()
 
 public:
 	static bool ShouldCompilePermutation(const FGlobalShaderPermutationParameters& Parameters)
@@ -350,9 +350,9 @@ public:
 
 private:
 
-	FShaderParameter NumElements;
-	FShaderResourceParameter InputBoundingBox;
-	FShaderResourceParameter OutputBoundingBox;
+	LAYOUT_FIELD(FShaderParameter, NumElements);
+	LAYOUT_FIELD(FShaderResourceParameter, InputBoundingBox);
+	LAYOUT_FIELD(FShaderResourceParameter, OutputBoundingBox);
 };
 
 IMPLEMENT_SHADER_TYPE(, FDirectCopyBoundingBoxCS, TEXT("/Plugin/Experimental/HairStrands/Private/NiagaraCopyBoundingBox.usf"), TEXT("MainCS"), SF_Compute);
