Integrate hair changes (rendering+simulation)

diff --git a/Engine/Plugins/Experimental/HairStrands/Shaders/Private/NiagaraCopyBoundingBox.usf b/Engine/Plugins/Experimental/HairStrands/Shaders/Private/NiagaraCopyBoundingBox.usf
--- a/Engine/Plugins/Experimental/HairStrands/Shaders/Private/NiagaraCopyBoundingBox.usf
+++ b/Engine/Plugins/Experimental/HairStrands/Shaders/Private/NiagaraCopyBoundingBox.usf
@@ -3,8 +3,8 @@
 #include "/Engine/Private/Common.ush"
 
 uint NumElements;
-RWBuffer<uint> BoundingBoxBuffer;
-RWBuffer<uint> OutNodeBoundBuffer;
+RWBuffer<uint> BoundingBoxBufferA;
+RWBuffer<uint> OutBoundingBoxBufferB;
 
 uint FloatToUint( in float FloatValue)
 {
@@ -23,20 +23,20 @@ void MainCS(uint2 DispatchThreadId : SV_DispatchThreadID)
 	const uint ElementIndex = DispatchThreadId.x;
 	if (ElementIndex < NumElements)
 	{
-		OutNodeBoundBuffer[0] = BoundingBoxBuffer[0];
-		OutNodeBoundBuffer[1] = BoundingBoxBuffer[1];
-		OutNodeBoundBuffer[2] = BoundingBoxBuffer[2];
+		OutBoundingBoxBufferB[0] = BoundingBoxBufferA[0];
+		OutBoundingBoxBufferB[1] = BoundingBoxBufferA[1];
+		OutBoundingBoxBufferB[2] = BoundingBoxBufferA[2];
 
-		OutNodeBoundBuffer[3] = BoundingBoxBuffer[3];
-		OutNodeBoundBuffer[4] = BoundingBoxBuffer[4];
-		OutNodeBoundBuffer[5] = BoundingBoxBuffer[5];
+		OutBoundingBoxBufferB[3] = BoundingBoxBufferA[3];
+		OutBoundingBoxBufferB[4] = BoundingBoxBufferA[4];
+		OutBoundingBoxBufferB[5] = BoundingBoxBufferA[5];
 
-		BoundingBoxBuffer[0] = UINT_MAX;
-		BoundingBoxBuffer[1] = UINT_MAX;
-		BoundingBoxBuffer[2] = UINT_MAX;
+		BoundingBoxBufferA[0] = UINT_MAX;
+		BoundingBoxBufferA[1] = UINT_MAX;
+		BoundingBoxBufferA[2] = UINT_MAX;
 
-		BoundingBoxBuffer[3] = UINT_MIN;
-		BoundingBoxBuffer[4] = UINT_MIN;
-		BoundingBoxBuffer[5] = UINT_MIN;
+		BoundingBoxBufferA[3] = UINT_MIN;
+		BoundingBoxBufferA[4] = UINT_MIN;
+		BoundingBoxBufferA[5] = UINT_MIN;
 	}
 }
\ No newline at end of file
diff --git a/Engine/Plugins/Experimental/HairStrands/Shaders/Private/NiagaraDataInterfaceHairStrands.ush b/Engine/Plugins/Experimental/HairStrands/Shaders/Private/NiagaraDataInterfaceHairStrands.ush
--- a/Engine/Plugins/Experimental/HairStrands/Shaders/Private/NiagaraDataInterfaceHairStrands.ush
+++ b/Engine/Plugins/Experimental/HairStrands/Shaders/Private/NiagaraDataInterfaceHairStrands.ush
@@ -39,8 +39,11 @@ struct FDIHairStrandsContext
 	float3				RestPositionOffset;
 	float3				DeformedPositionOffset;
 	bool				ResetSimulation;
-	RWBuffer<uint>		BoundingBoxBuffer;
-	Buffer<uint>		NodeBoundBuffer;
+	RWBuffer<uint>		BoundingBoxBufferA;
+	Buffer<uint>		BoundingBoxBufferB;
+	int					SampleCount;
+	Buffer<float4>		RestSamplePositionsBuffer;
+	Buffer<float4>		MeshSampleWeightsBuffer;
 }; 
 
 #define DIHAIRSTRANDS_DECLARE_CONSTANTS(NAME)\
@@ -67,8 +70,11 @@ float3				DeformedRootOffset_##NAME;\
 float3				RestPositionOffset_##NAME;\
 float3				DeformedPositionOffset_##NAME;\
 bool				ResetSimulation_##NAME;\
-RWBuffer<uint>	    BoundingBoxBuffer_##NAME;\
-Buffer<uint>		NodeBoundBuffer_##NAME;\
+RWBuffer<uint>	    BoundingBoxBufferA_##NAME;\
+Buffer<uint>		BoundingBoxBufferB_##NAME;\
+int					SampleCount_##NAME;\
+Buffer<float4>		RestSamplePositionsBuffer_##NAME;\
+Buffer<float4>		MeshSampleWeightsBuffer_##NAME;\
 
 #define DIHAIRSTRANDS_MAKE_CONTEXT(NAME)\
 FDIHairStrandsContext DIContext;\
@@ -95,8 +101,11 @@ DIContext.DeformedRootOffset = DeformedRootOffset_##NAME;\
 DIContext.RestPositionOffset =	RestPositionOffset_##NAME;\
 DIContext.DeformedPositionOffset = DeformedPositionOffset_##NAME;\
 DIContext.ResetSimulation = ResetSimulation_##NAME;\
-DIContext.BoundingBoxBuffer = BoundingBoxBuffer_##NAME;\
-DIContext.NodeBoundBuffer = NodeBoundBuffer_##NAME;\
+DIContext.BoundingBoxBufferA = BoundingBoxBufferA_##NAME;\
+DIContext.BoundingBoxBufferB = BoundingBoxBufferB_##NAME;\
+DIContext.SampleCount = SampleCount_##NAME;\
+DIContext.RestSamplePositionsBuffer = RestSamplePositionsBuffer_##NAME;\
+DIContext.MeshSampleWeightsBuffer = MeshSampleWeightsBuffer_##NAME;\
 
 /* -----------------------------------------------------------------
  * Vector utilities
@@ -131,8 +140,8 @@ DIContext.NodeBoundBuffer = NodeBoundBuffer_##NAME;\
 // Return the bounding box center
 void DIHairStrands_GetBoxCenter(in FDIHairStrandsContext DIContext, out float3 OutBoxCenter)
 {
-	const float3 BoxMin = Uint3ToFloat3(uint3(DIContext.NodeBoundBuffer[0],DIContext.NodeBoundBuffer[1],DIContext.NodeBoundBuffer[2]));
-	const float3 BoxMax = Uint3ToFloat3(uint3(DIContext.NodeBoundBuffer[3],DIContext.NodeBoundBuffer[4],DIContext.NodeBoundBuffer[5]));
+	const float3 BoxMin = Uint3ToFloat3(uint3(DIContext.BoundingBoxBufferB[0],DIContext.BoundingBoxBufferB[1],DIContext.BoundingBoxBufferB[2]));
+	const float3 BoxMax = Uint3ToFloat3(uint3(DIContext.BoundingBoxBufferB[3],DIContext.BoundingBoxBufferB[4],DIContext.BoundingBoxBufferB[5]));
 
 	OutBoxCenter = 0.5 * (BoxMin+BoxMax);
 }
@@ -140,8 +149,8 @@ void DIHairStrands_GetBoxCenter(in FDIHairStrandsContext DIContext, out float3 O
 // Return the bounding box extent
 void DIHairStrands_GetBoxExtent(in FDIHairStrandsContext DIContext, out float3 OutBoxExtent)
 {
-	const float3 BoxMin = Uint3ToFloat3(uint3(DIContext.NodeBoundBuffer[0],DIContext.NodeBoundBuffer[1],DIContext.NodeBoundBuffer[2]));
-	const float3 BoxMax = Uint3ToFloat3(uint3(DIContext.NodeBoundBuffer[3],DIContext.NodeBoundBuffer[4],DIContext.NodeBoundBuffer[5]));
+	const float3 BoxMin = Uint3ToFloat3(uint3(DIContext.BoundingBoxBufferB[0],DIContext.BoundingBoxBufferB[1],DIContext.BoundingBoxBufferB[2]));
+	const float3 BoxMax = Uint3ToFloat3(uint3(DIContext.BoundingBoxBufferB[3],DIContext.BoundingBoxBufferB[4],DIContext.BoundingBoxBufferB[5]));
 
 	//OutBoxExtent =  0.5 * (BoxMax-BoxMin);
 	OutBoxExtent =   (BoxMax-BoxMin);
@@ -208,26 +217,6 @@ float DIHairStrands_ComputeStrandLength(in FDIHairStrandsContext DIContext, in i
 
  #define WRAP_SIZE 32
 
- /*
- const uint3 LocalPosition = Float3ToUint3(NodePosition);
-
-if(isfinite(NodePosition))
-{
-	InterlockedMin(DIContext.BoundingBoxBuffer[0],LocalPosition.x);
-	InterlockedMin(DIContext.BoundingBoxBuffer[1],LocalPosition.y);
-	InterlockedMin(DIContext.BoundingBoxBuffer[2],LocalPosition.z);
-
-	InterlockedMax(DIContext.BoundingBoxBuffer[3],LocalPosition.x);
-	InterlockedMax(DIContext.BoundingBoxBuffer[4],LocalPosition.y);
-	InterlockedMax(DIContext.BoundingBoxBuffer[5],LocalPosition.z);
-}
-
-AllMemoryBarrier();
-
-OutBoundMin = Uint3ToFloat3(uint3(DIContext.NodeBoundBuffer[0],DIContext.NodeBoundBuffer[1],DIContext.NodeBoundBuffer[2]));
-OutBoundMax = Uint3ToFloat3(uint3(DIContext.NodeBoundBuffer[3],DIContext.NodeBoundBuffer[4],DIContext.NodeBoundBuffer[5])); 
-*/
-
 void DIHairStrands_BuildBoundingBox(in FDIHairStrandsContext DIContext, in float3 NodePosition, out bool BuildStatus, out float3 OutBoundMin, out float3 OutBoundMax)
 {
 	BuildStatus = true;
@@ -236,62 +225,33 @@ void DIHairStrands_BuildBoundingBox(in FDIHairStrandsContext DIContext, in float
 
 	if(isfinite(NodePosition.x) && isfinite(NodePosition.y) && isfinite(NodePosition.z))
 	{
-		InterlockedMin(DIContext.BoundingBoxBuffer[0],LocalPosition.x);
-		InterlockedMin(DIContext.BoundingBoxBuffer[1],LocalPosition.y);
-		InterlockedMin(DIContext.BoundingBoxBuffer[2],LocalPosition.z);
+		InterlockedMin(DIContext.BoundingBoxBufferA[0],LocalPosition.x);
+		InterlockedMin(DIContext.BoundingBoxBufferA[1],LocalPosition.y);
+		InterlockedMin(DIContext.BoundingBoxBufferA[2],LocalPosition.z);
 
-		InterlockedMax(DIContext.BoundingBoxBuffer[3],LocalPosition.x);
-		InterlockedMax(DIContext.BoundingBoxBuffer[4],LocalPosition.y);
-		InterlockedMax(DIContext.BoundingBoxBuffer[5],LocalPosition.z);
+		InterlockedMax(DIContext.BoundingBoxBufferA[3],LocalPosition.x);
+		InterlockedMax(DIContext.BoundingBoxBufferA[4],LocalPosition.y);
+		InterlockedMax(DIContext.BoundingBoxBufferA[5],LocalPosition.z);
 	}
 
 	DeviceMemoryBarrier();
 
-	OutBoundMin = Uint3ToFloat3(uint3(DIContext.NodeBoundBuffer[0],DIContext.NodeBoundBuffer[1],DIContext.NodeBoundBuffer[2]));
-	OutBoundMax = Uint3ToFloat3(uint3(DIContext.NodeBoundBuffer[3],DIContext.NodeBoundBuffer[4],DIContext.NodeBoundBuffer[5])); 
-
-	return;
+	OutBoundMin = Uint3ToFloat3(uint3(DIContext.BoundingBoxBufferB[0],DIContext.BoundingBoxBufferB[1],DIContext.BoundingBoxBufferB[2]));
+	OutBoundMax = Uint3ToFloat3(uint3(DIContext.BoundingBoxBufferB[3],DIContext.BoundingBoxBufferB[4],DIContext.BoundingBoxBufferB[5])); 
+}
 
-	SharedNodePosition[GGroupThreadId.x] = NodePosition;
-	GroupMemoryBarrier();
+// Compute the node mass
+void DIHairStrands_EvalSkinnedPosition(in FDIHairStrandsContext DIContext, in float3 RestPosition, out float3 OutSkinnedPosition)
+{
+	const float3 RestSkinnedPosition = RestPosition;
 
-	const int LocalIndex = (GGroupThreadId.x % DIContext.StrandSize);
-	if(LocalIndex == 0)
+	OutSkinnedPosition = RestPosition;
+	for(int i = 0; i < DIContext.SampleCount; ++i)
 	{
-		const float FLOAT_MAX = 1e+10;
-
-		float3 BoundingBoxMin = float3(FLOAT_MAX,FLOAT_MAX,FLOAT_MAX);
-		float3 BoundingBoxMax = -BoundingBoxMin;
-
-		[unroll]
-		for( int i = GGroupThreadId.x; i < GGroupThreadId.x+DIContext.StrandSize; ++i)
-		{
-			BoundingBoxMin = min(BoundingBoxMin,SharedNodePosition[i]);
-			BoundingBoxMax = max(BoundingBoxMax,SharedNodePosition[i]);
-		}
-
-		if(isfinite(BoundingBoxMin.x) && isfinite(BoundingBoxMin.y) && isfinite(BoundingBoxMin.z))
-		{
-			const uint3 MinPosition = Float3ToUint3(BoundingBoxMin);
-
-			InterlockedMin(DIContext.BoundingBoxBuffer[0],MinPosition.x);
-			InterlockedMin(DIContext.BoundingBoxBuffer[1],MinPosition.y);
-			InterlockedMin(DIContext.BoundingBoxBuffer[2],MinPosition.z);
-		}
-		if(isfinite(BoundingBoxMax.x) && isfinite(BoundingBoxMax.y) && isfinite(BoundingBoxMax.z))
-		{
-			const uint3 MaxPosition = Float3ToUint3(BoundingBoxMax);
-
-			InterlockedMax(DIContext.BoundingBoxBuffer[3],MaxPosition.x);
-			InterlockedMax(DIContext.BoundingBoxBuffer[4],MaxPosition.y);
-			InterlockedMax(DIContext.BoundingBoxBuffer[5],MaxPosition.z);
-		}
+		const float3 PositionDelta = RestSkinnedPosition - DIContext.RestSamplePositionsBuffer[i];
+		const float FunctionValue = sqrt( dot(PositionDelta,PositionDelta)+1);
+		OutSkinnedPosition += FunctionValue * DIContext.MeshSampleWeightsBuffer[i].xyz;
 	}
-
-	AllMemoryBarrier();
-
-	OutBoundMin = Uint3ToFloat3(uint3(DIContext.NodeBoundBuffer[0],DIContext.NodeBoundBuffer[1],DIContext.NodeBoundBuffer[2]));
-	OutBoundMax = Uint3ToFloat3(uint3(DIContext.NodeBoundBuffer[3],DIContext.NodeBoundBuffer[4],DIContext.NodeBoundBuffer[5]));
 }
 
 /* -----------------------------------------------------------------
diff --git a/Engine/Plugins/Experimental/HairStrands/Source/HairStrandsCore/HairStrandsCore.Build.cs b/Engine/Plugins/Experimental/HairStrands/Source/HairStrandsCore/HairStrandsCore.Build.cs
--- a/Engine/Plugins/Experimental/HairStrands/Source/HairStrandsCore/HairStrandsCore.Build.cs
+++ b/Engine/Plugins/Experimental/HairStrands/Source/HairStrandsCore/HairStrandsCore.Build.cs
@@ -21,7 +21,8 @@ namespace UnrealBuildTool.Rules
 					"Renderer",
 					"RHI",
 					"ChaosCore",
-					"Chaos"
+					"Chaos",
+					"Eigen"
 				});
 
 			if (Target.bBuildEditor == true)
diff --git a/Engine/Plugins/Experimental/HairStrands/Source/HairStrandsCore/Private/GroomAsset.cpp b/Engine/Plugins/Experimental/HairStrands/Source/HairStrandsCore/Private/GroomAsset.cpp
--- a/Engine/Plugins/Experimental/HairStrands/Source/HairStrandsCore/Private/GroomAsset.cpp
+++ b/Engine/Plugins/Experimental/HairStrands/Source/HairStrandsCore/Private/GroomAsset.cpp
@@ -23,6 +23,14 @@
 #include "EditorFramework/AssetImportData.h"
 #endif
 
+static int32 GHairStrandsLoadAsset = 1;
+static FAutoConsoleVariableRef CVarHairStrandsLoadAsset(TEXT("r.HairStrands.LoadAsset"), GHairStrandsLoadAsset, TEXT("Allow groom asset to be loaded"));
+
+bool IsHairStrandsAssetLoadingEnable()
+{
+	return GHairStrandsLoadAsset > 0;
+}
+
 template<typename BufferType>
 void UploadDataToBuffer(BufferType& OutBuffer, uint32 DataSizeInBytes, void* InCpuData)
 {
@@ -190,16 +198,12 @@ private:
 	TArray<FCluster> ClusterInfo;
 };
 
-FHairStrandsClusterCullingResource::FHairStrandsClusterCullingResource(const FHairGroupData& GroupData)
+FHairStrandsClusterCullingResource::FHairStrandsClusterCullingResource(const FHairStrandsDatas& RenStrandsData)
 {
-	const FHairStrandsDatas& SimStrandsData = GroupData.HairSimulationData;
-	const FHairStrandsDatas& RenStrandsData = GroupData.HairRenderData;
-	const FHairStrandsInterpolationDatas& InterpolationData = GroupData.HairInterpolationData;
 	const uint32 RenCurveCount = RenStrandsData.GetNumCurves();
-	const uint32 SimCurveCount = SimStrandsData.GetNumCurves();
 	const uint32 PointCount = RenStrandsData.GetNumPoints();
 
-
+	check(PointCount);
 	// Allocate look up arrays for as many hair vertices as needed.
 	VertexCount = PointCount;
 	VertexToClusterIdArray.SetNum(PointCount);
@@ -207,8 +211,8 @@ FHairStrandsClusterCullingResource::FHairStrandsClusterCullingResource(const FHa
 
 	// Allocate cluster per voxel containing contains >=1 render curve root
 	const uint32 VoxelcountAlongLargerSide = 256; // TODO expose in UI
-	const float VoxelCentimeterSize = GroupData.HairRenderData.BoundingBox.GetSize().GetAbsMax() / VoxelcountAlongLargerSide;
-	FClusterGrid ClusterGrid(VoxelCentimeterSize, GroupData.HairRenderData.BoundingBox.Min, GroupData.HairRenderData.BoundingBox.Max);
+	const float VoxelCentimeterSize = RenStrandsData.BoundingBox.GetSize().GetAbsMax() / VoxelcountAlongLargerSide;
+	FClusterGrid ClusterGrid(VoxelCentimeterSize, RenStrandsData.BoundingBox.Min, RenStrandsData.BoundingBox.Max);
 
 	for (uint32 RenCurveIndex = 0; RenCurveIndex < RenCurveCount; ++RenCurveIndex)
 	{
@@ -369,8 +373,8 @@ RootData(InRootData)
 	PopulateFromRootData();
 }
 
-FHairStrandsRootResource::FHairStrandsRootResource(const FHairStrandsDatas* HairStrandsDatas, uint32 LODCount):
-	RootData(HairStrandsDatas, LODCount)
+FHairStrandsRootResource::FHairStrandsRootResource(const FHairStrandsDatas* HairStrandsDatas, uint32 LODCount, const TArray<uint32>& NumSamples):
+	RootData(HairStrandsDatas, LODCount, NumSamples)
 {
 	PopulateFromRootData();
 }
@@ -384,6 +388,7 @@ void FHairStrandsRootResource::PopulateFromRootData()
 
 		LOD.LODIndex = MeshProjectionLOD.LODIndex;
 		LOD.Status = FHairStrandsProjectionHairData::LODData::EStatus::Invalid;
+		LOD.SampleCount = MeshProjectionLOD.SampleCount;
 	}
 }
 
@@ -432,6 +437,28 @@ void FHairStrandsRootResource::InitRHI()
 				CreateBuffer<FHairStrandsMeshTrianglePositionFormat>(RootData.RootCount, GPUData.RestRootTrianglePosition2Buffer);
 			}
 
+			GPUData.SampleCount = CPUData.SampleCount;
+			const bool bHasValidCPUWeights = CPUData.MeshSampleIndicesBuffer.Num() > 0;
+			if(bHasValidCPUWeights)
+			{
+				check(CPUData.MeshInterpolationWeightsBuffer.Num() == CPUData.SampleCount * CPUData.SampleCount);
+				check(CPUData.MeshSampleIndicesBuffer.Num() == CPUData.SampleCount);
+				check(CPUData.RestSamplePositionsBuffer.Num() == CPUData.SampleCount);
+
+				CreateBuffer<FHairStrandsWeightFormat>(CPUData.MeshInterpolationWeightsBuffer, GPUData.MeshInterpolationWeightsBuffer);
+				CreateBuffer<FHairStrandsIndexFormat>(CPUData.MeshSampleIndicesBuffer, GPUData.MeshSampleIndicesBuffer);
+				CreateBuffer<FHairStrandsMeshTrianglePositionFormat>(CPUData.RestSamplePositionsBuffer, GPUData.RestSamplePositionsBuffer);
+			}
+			else
+			{
+				CreateBuffer<FHairStrandsWeightFormat>(CPUData.SampleCount * CPUData.SampleCount, GPUData.MeshInterpolationWeightsBuffer);
+				CreateBuffer<FHairStrandsIndexFormat>(CPUData.SampleCount, GPUData.MeshSampleIndicesBuffer);
+				CreateBuffer<FHairStrandsMeshTrianglePositionFormat>(CPUData.SampleCount, GPUData.RestSamplePositionsBuffer);
+			}
+			CreateBuffer<FHairStrandsMeshTrianglePositionFormat>(CPUData.SampleCount, GPUData.DeformedSamplePositionsBuffer);
+			CreateBuffer<FHairStrandsMeshTrianglePositionFormat>(CPUData.SampleCount, GPUData.MeshSampleWeightsBuffer);
+
+
 			// Strand hair roots translation and rotation in triangle-deformed position relative to the bound triangle 
 			CreateBuffer<FHairStrandsMeshTrianglePositionFormat>(RootData.RootCount, GPUData.DeformedRootTrianglePosition0Buffer);
 			CreateBuffer<FHairStrandsMeshTrianglePositionFormat>(RootData.RootCount, GPUData.DeformedRootTrianglePosition1Buffer);
@@ -445,7 +472,7 @@ void FHairStrandsRootResource::ReleaseRHI()
 	RootPositionBuffer.Release();
 	RootNormalBuffer.Release();
 	VertexToCurveIndexBuffer.Release();
-
+	
 	for (FMeshProjectionLOD& GPUData : MeshProjectionLODs)
 	{
 		GPUData.Status = FHairStrandsProjectionHairData::LODData::EStatus::Invalid;
@@ -457,6 +484,12 @@ void FHairStrandsRootResource::ReleaseRHI()
 		GPUData.DeformedRootTrianglePosition0Buffer.Release();
 		GPUData.DeformedRootTrianglePosition1Buffer.Release();
 		GPUData.DeformedRootTrianglePosition2Buffer.Release();
+		GPUData.SampleCount = 0;
+		GPUData.MeshInterpolationWeightsBuffer.Release();
+		GPUData.MeshSampleIndicesBuffer.Release();
+		GPUData.RestSamplePositionsBuffer.Release();
+		GPUData.DeformedSamplePositionsBuffer.Release();
+		GPUData.MeshSampleWeightsBuffer.Release();
 	}
 	MeshProjectionLODs.Empty();
 }
@@ -468,7 +501,7 @@ FHairStrandsRootData::FHairStrandsRootData()
 
 }
 
-FHairStrandsRootData::FHairStrandsRootData(const FHairStrandsDatas* HairStrandsDatas, uint32 LODCount):
+FHairStrandsRootData::FHairStrandsRootData(const FHairStrandsDatas* HairStrandsDatas, uint32 LODCount, const TArray<uint32>& NumSamples):
 	RootCount(HairStrandsDatas ? HairStrandsDatas->GetNumCurves() : 0)
 {
 	if (!HairStrandsDatas)
@@ -515,12 +548,17 @@ FHairStrandsRootData::FHairStrandsRootData(const FHairStrandsDatas* HairStrandsD
 		RootPositionBuffer[CurveIndex] = P;
 		RootNormalBuffer[CurveIndex] = N;
 	}
+	check(NumSamples.Num() == LODCount);
 
 	MeshProjectionLODs.SetNum(LODCount);
 	uint32 LODIndex = 0;
 	for (FMeshProjectionLOD& MeshProjectionLOD : MeshProjectionLODs)
 	{
+		MeshProjectionLOD.SampleCount = NumSamples[LODIndex];
 		MeshProjectionLOD.LODIndex = LODIndex++;
+		MeshProjectionLOD.MeshInterpolationWeightsBuffer.Empty();
+		MeshProjectionLOD.MeshSampleIndicesBuffer.Empty();
+		MeshProjectionLOD.RestSamplePositionsBuffer.Empty();
 	}
 }
 
@@ -553,6 +591,20 @@ FArchive& operator<<(FArchive& Ar, FHairStrandsRootData::FMeshProjectionLOD& LOD
 	Ar << LOD.RestRootTrianglePosition0Buffer;
 	Ar << LOD.RestRootTrianglePosition1Buffer;
 	Ar << LOD.RestRootTrianglePosition2Buffer;
+
+	Ar << LOD.SampleCount;
+	Ar << LOD.MeshInterpolationWeightsBuffer;
+	Ar << LOD.MeshSampleIndicesBuffer;
+	Ar << LOD.RestSamplePositionsBuffer;
+
+	//Ar.UsingCustomVersion(FAnimObjectVersion::GUID);
+	/*if (Ar.CustomVer(FAnimObjectVersion::GUID) >= FAnimObjectVersion::GroomAssetWithRbfWeights)
+	{
+		Ar << LOD.SampleCount;
+		Ar << LOD.MeshInterpolationWeightsBuffer;
+		Ar << LOD.MeshSampleIndicesBuffer;
+		Ar << LOD.RestSamplePositionsBuffer;
+	}*/
 	return Ar;
 }
 
@@ -678,6 +730,8 @@ void UGroomAsset::Serialize(FArchive& Ar)
 UGroomAsset::UGroomAsset(const FObjectInitializer& ObjectInitializer)
 	: Super(ObjectInitializer)
 {
+	bIsInitialized = false;
+
 	HairToGuideDensity = 0.1f;
 
 	SimulatedGroup = 0;
@@ -728,6 +782,8 @@ UGroomAsset::UGroomAsset(const FObjectInitializer& ObjectInitializer)
 
 void UGroomAsset::InitResource()
 {
+	bIsInitialized = true;
+
 	for (FHairGroupData& GroupData : HairGroupsData)
 	{
 		GroupData.HairStrandsRestResource = new FHairStrandsRestResource(GroupData.HairRenderData.RenderData, GroupData.HairRenderData.BoundingBox.GetCenter());
@@ -737,6 +793,35 @@ void UGroomAsset::InitResource()
 		GroupData.HairSimulationRestResource = new FHairStrandsRestResource(GroupData.HairSimulationData.RenderData, GroupData.HairSimulationData.BoundingBox.GetCenter());
 
 		BeginInitResource(GroupData.HairSimulationRestResource);
+
+		GroupData.ClusterCullingResource = new FHairStrandsClusterCullingResource(GroupData.HairRenderData);
+		
+		BeginInitResource(GroupData.ClusterCullingResource);
+
+		GroupData.HairInterpolationResource = new FHairStrandsInterpolationResource(GroupData.HairInterpolationData.RenderData, GroupData.HairSimulationData);
+
+		BeginInitResource(GroupData.HairInterpolationResource); 
+	}
+
+	// Sanity check
+	const uint32 GroupCount = HairGroupsData.Num();
+	if (HairGroupsInfo.Num() == GroupCount)
+	{
+		for (uint32 GroupIt = 0; GroupIt < GroupCount; ++GroupIt)
+		{
+			if (HairGroupsData[GroupIt].HairRenderData.GetNumCurves() != HairGroupsInfo[GroupIt].NumCurves)
+			{
+				UE_LOG(LogHairStrands, Warning, TEXT("[Groom] The groom asset (%s) has inconsistent data. Curves count differs between data and description (%d vs. %d). Please reimport it."), *GetName(), HairGroupsData[GroupIt].HairRenderData.GetNumCurves(), HairGroupsInfo[GroupIt].NumCurves);
+			}
+			if (HairGroupsData[GroupIt].HairSimulationData.GetNumCurves() != HairGroupsInfo[GroupIt].NumGuides)
+			{
+				UE_LOG(LogHairStrands, Warning, TEXT("[Groom] The groom asset (%s) has inconsistent data. Guides count different between data and description (%d vs. %d). Please reimport it."), *GetName(), HairGroupsData[GroupIt].HairSimulationData.GetNumCurves(), HairGroupsInfo[GroupIt].NumGuides);
+			}
+		}
+	}
+	else
+	{
+		UE_LOG(LogHairStrands, Warning, TEXT("[Groom] The groom asset (%s) has inconsistent data. Group count differs between data and description (%d vs. %d). Please reimport it."), *GetName(), HairGroupsData.Num(), HairGroupsInfo.Num());
 	}
 }
 
@@ -753,11 +838,22 @@ void UGroomAsset::UpdateResource()
 		{
 			BeginUpdateResourceRHI(GroupData.HairSimulationRestResource);
 		}
+
+		if (GroupData.ClusterCullingResource)
+		{
+			BeginUpdateResourceRHI(GroupData.ClusterCullingResource);
+		}
+
+		if (GroupData.HairInterpolationResource)
+		{
+			BeginUpdateResourceRHI(GroupData.HairInterpolationResource);
+		}
 	}
 }
 
 void UGroomAsset::ReleaseResource()
 {
+	bIsInitialized = false;
 	for (FHairGroupData& GroupData : HairGroupsData)
 	{
 		if (GroupData.HairStrandsRestResource)
@@ -783,6 +879,30 @@ void UGroomAsset::ReleaseResource()
 			});
 			GroupData.HairSimulationRestResource = nullptr;
 		}
+
+		if (GroupData.ClusterCullingResource)
+		{
+			FHairStrandsClusterCullingResource* InResource = GroupData.ClusterCullingResource;
+			ENQUEUE_RENDER_COMMAND(ReleaseHairStrandsResourceCommand)(
+				[InResource](FRHICommandList& RHICmdList)
+			{
+				InResource->ReleaseResource();
+				delete InResource;
+			});
+			GroupData.ClusterCullingResource = nullptr;
+		}
+
+		if (GroupData.HairInterpolationResource)
+		{
+			FHairStrandsInterpolationResource* InResource = GroupData.HairInterpolationResource;
+			ENQUEUE_RENDER_COMMAND(ReleaseHairStrandsResourceCommand)(
+				[InResource](FRHICommandList& RHICmdList)
+			{
+				InResource->ReleaseResource();
+				delete InResource;
+			});
+			GroupData.HairInterpolationResource = nullptr;
+		}
 	}
 }
 
@@ -818,7 +938,7 @@ void UGroomAsset::PostLoad()
 		FGroomBuilder::BuildData(this, Settings);
 	}
 
-	if (!IsTemplate())
+	if (!IsTemplate() && IsHairStrandsAssetLoadingEnable())
 	{
 		InitResource();
 	}
@@ -1167,6 +1287,7 @@ FArchive& operator<<(FArchive& Ar, UGroomBindingAsset::FHairGroupData& GroupData
 void UGroomBindingAsset::Serialize(FArchive& Ar)
 {
 	Super::Serialize(Ar);
+
 	Ar << HairGroupDatas;
 }
 
@@ -1214,6 +1335,22 @@ void UGroomBindingAsset::ReleaseResource()
 		}
 		HairGroupResources.Empty();
 	}
+
+	// Process resources to be deleted (should happen only in editor)
+	FHairGroupResource ResourceToDelete;
+	while (HairGroupResourcesToDelete.Dequeue(ResourceToDelete))
+	{
+		FHairStrandsRootResource* InSimRootResources = ResourceToDelete.SimRootResources;
+		FHairStrandsRootResource* InRenRootResources = ResourceToDelete.RenRootResources;
+		ENQUEUE_RENDER_COMMAND(ReleaseHairStrandsResourceCommand)(
+			[InSimRootResources, InRenRootResources](FRHICommandList& RHICmdList)
+		{
+			InSimRootResources->ReleaseResource();
+			InRenRootResources->ReleaseResource();
+			delete InSimRootResources;
+			delete InRenRootResources;
+		});
+	}
 }
 
 void UGroomBindingAsset::Reset()
@@ -1232,6 +1369,15 @@ void UGroomBindingAsset::PostLoad()
 
 	InitializeGroupInfos(HairGroupDatas, GroupInfos);
 
+	if (Groom)
+	{
+		// Make sure that the asset initialized its resources first since the component needs them to initialize its own resources
+		Groom->ConditionalPostLoad();
+
+		// Sanity check. This function will report back warnings/issues into the log for user.
+		UGroomBindingAsset::IsCompatible(Groom, this);
+	}
+
 	if (!IsTemplate())
 	{
 		InitResource();
@@ -1240,14 +1386,26 @@ void UGroomBindingAsset::PostLoad()
 
 void UGroomBindingAsset::PreSave(const class ITargetPlatform* TargetPlatform)
 {
+#if WITH_EDITOR
+	while (QueryStatus == EQueryStatus::Submitted)
+	{
+		FPlatformProcess::Sleep(1);
+	}
+#endif
 	Super::PreSave(TargetPlatform);
 	InitializeGroupInfos(HairGroupDatas, GroupInfos);
+#if WITH_EDITOR
+	OnGroomBindingAssetChanged.Broadcast();
+#endif
 }
 
 void UGroomBindingAsset::PostSaveRoot(bool bCleanupIsRequired)
 {
 	Super::PostSaveRoot(bCleanupIsRequired);
 	InitializeGroupInfos(HairGroupDatas, GroupInfos);
+#if WITH_EDITOR
+	OnGroomBindingAssetChanged.Broadcast();
+#endif
 }
 
 void UGroomBindingAsset::BeginDestroy()
@@ -1256,6 +1414,143 @@ void UGroomBindingAsset::BeginDestroy()
 	Super::BeginDestroy();
 }
 
+bool UGroomBindingAsset::IsCompatible(const USkeletalMesh* InSkeletalMesh, const UGroomBindingAsset* InBinding)
+{
+	if (InBinding && InSkeletalMesh)
+	{
+		if (!InBinding->TargetSkeletalMesh)
+		{
+			UE_LOG(LogHairStrands, Warning, TEXT("[Groom] The binding asset (%s) does not have a target skeletal mesh. Falling back onto non-binding version."), *InBinding->GetName());
+			return false;
+		}
+		
+		// TODO: need something better to assess that skeletal meshes match. In the mean time, string comparison. 
+		// Since they can be several instances of a skeletalMesh asset (??), a numerical suffix is sometime added to the name (e.g., SkeletalName_0).
+		// This is why we are using substring comparison.
+		//if (InSkeletalMesh->GetPrimaryAssetId() != InBinding->TargetSkeletalMesh->GetPrimaryAssetId())
+		if (!InSkeletalMesh->GetName().Contains(InBinding->TargetSkeletalMesh->GetName()))
+		{
+			UE_LOG(LogHairStrands, Warning, TEXT("[Groom] The Groom binding (%s) does not reference the same skeletal asset (BindingAsset's skeletal:%s vs. Attached skeletal:%s). The binding asset will not be used."),
+				*InBinding->GetName(),
+				*InBinding->TargetSkeletalMesh->GetName(),
+				*InSkeletalMesh->GetName());
+			return false;
+		}
+	}
+
+	return true;
+}
+
+bool UGroomBindingAsset::IsCompatible(const UGroomAsset* InGroom, const UGroomBindingAsset* InBinding)
+{
+	if (InBinding && InGroom)
+	{
+		if (!InBinding->Groom)
+		{
+			UE_LOG(LogHairStrands, Warning, TEXT("[Groom] The binding asset (%s) does not reference a groom. Falling back onto non-binding version."), *InBinding->GetName());
+			return false;
+		}
+
+		if (InGroom->GetPrimaryAssetId() != InBinding->Groom->GetPrimaryAssetId())
+		{
+			UE_LOG(LogHairStrands, Warning, TEXT("[Groom] The Groom binding (%) does not reference the same groom asset (BindingAsset's groom:%s vs. Groom:%s). The binding asset will not be used."), 
+				*InBinding->GetName(),
+				*InBinding->Groom->GetName(),
+				*InGroom->GetName());
+			return false;
+		}
+
+		if (InGroom->HairGroupsInfo.Num() != InBinding->GroupInfos.Num())
+		{
+			UE_LOG(LogHairStrands, Warning, TEXT("[Groom] The GroomBinding asset (%s) does not contains the same number of groups (%d vs. %d) than the groom (%s). The binding asset will not be used."), 
+				*InBinding->GetName(), 
+				InGroom->HairGroupsInfo.Num(), 
+				InBinding->GroupInfos.Num(), 
+				*InGroom->GetName());
+			return false;
+		}
+
+		for (uint32 GroupIt = 0, GroupCount = InGroom->HairGroupsInfo.Num(); GroupIt < GroupCount; ++GroupIt)
+		{
+			{
+				const uint32 GroomCount = InGroom->HairGroupsInfo[GroupIt].NumGuides;
+				const uint32 BindingCount = InBinding->GroupInfos[GroupIt].SimRootCount;
+
+				if (GroomCount != BindingCount)
+				{
+					UE_LOG(LogHairStrands, Warning, TEXT("[Groom] The GroomBinding asset (%s) does not contains the same guides in group %d (%d vs. %d) than the groom (%s). The binding asset will not be used."), 
+						*InBinding->GetName(), 
+						GroupIt,
+						GroomCount,
+						BindingCount,
+						*InGroom->GetName());
+					return false;
+				}
+			}
+
+			{
+				const uint32 GroomCount = InGroom->HairGroupsInfo[GroupIt].NumCurves;
+				const uint32 BindingCount = InBinding->GroupInfos[GroupIt].RenRootCount;
+
+				if (GroomCount != BindingCount)
+				{
+					UE_LOG(LogHairStrands, Warning, TEXT("[Groom] The GroomBinding asset (%s) does not contains the same curves in group %d (%d vs. %d) than the groom (%s). The binding asset will not be used."), 
+						*InBinding->GetName(), 
+						GroupIt, 
+						GroomCount,
+						BindingCount,
+						*InGroom->GetName());
+					return false;
+				}
+			}
+		}
+	}
+	return true;
+}
+
+bool UGroomBindingAsset::IsBindingAssetValid(const UGroomBindingAsset* InBinding, bool bIsBindingReloading)
+{
+	if (InBinding)
+	{
+		if (!InBinding->Groom)
+		{
+			UE_LOG(LogHairStrands, Warning, TEXT("[Groom] The binding asset (%s) does not reference a groom. Falling back onto non-binding version."), *InBinding->GetName());
+			return false;
+		}
+
+		if (const UPackage* Package = InBinding->GetOutermost())
+		{
+			if (Package->IsDirty() && !bIsBindingReloading)
+			{
+				UE_LOG(LogHairStrands, Warning, TEXT("[Groom] The binding asset (%s) is not saved and will be considered as invalid. Falling back onto non-binding version."), *InBinding->GetName());
+				return false;
+			}
+		}
+
+		if (InBinding->GroupInfos.Num() == 0)
+		{
+			UE_LOG(LogHairStrands, Warning, TEXT("[Groom] The GroomBinding asset (Groom:%s) does not contain any groups. It is invalid and can't be assigned. The binding asset will not be used."), *InBinding->Groom->GetName());
+			return false;
+		}
+
+		for (const FGoomBindingGroupInfo& Info : InBinding->GroupInfos)
+		{
+			if (Info.SimRootCount == 0)
+			{
+				UE_LOG(LogHairStrands, Warning, TEXT("[Groom] The GroomBinding asset (Groom:%s) has group with 0 guides. It is invalid and can't be assigned. The binding asset will not be used."), *InBinding->Groom->GetName());
+				return false;
+			}
+
+			if (Info.RenRootCount == 0)
+			{
+				UE_LOG(LogHairStrands, Warning, TEXT("[Groom] The GroomBinding asset (Groom:%s) has group with 0 curves. It is invalid and can't be assigned. The binding asset will not be used."), *InBinding->Groom->GetName());
+				return false;
+			}
+		}
+	}
+	return true;
+}
+
 #if WITH_EDITOR
 void UGroomBindingAsset::PostEditChangeProperty(FPropertyChangedEvent& PropertyChangedEvent)
 {
@@ -1265,4 +1560,21 @@ void UGroomBindingAsset::PostEditChangeProperty(FPropertyChangedEvent& PropertyC
 }
 #endif // WITH_EDITOR
 
+bool FProcessedHairDescription::IsValid() const
+{
+	for (TPair<int32, FProcessedHairDescription::FHairGroup> HairGroupIt : HairGroups)
+	{
+		const FProcessedHairDescription::FHairGroup& Group = HairGroupIt.Value;
+		const FHairGroupInfo& GroupInfo = Group.Key;
+		if (GroupInfo.NumCurves == 0)
+		{
+			return false;
+		}
+	}
+
+	return 	HairGroups.Num() > 0;
+}
+
+
+#undef LOCTEXT_NAMESPACE
 
diff --git a/Engine/Plugins/Experimental/HairStrands/Source/HairStrandsCore/Private/GroomBuilder.cpp b/Engine/Plugins/Experimental/HairStrands/Source/HairStrandsCore/Private/GroomBuilder.cpp
--- a/Engine/Plugins/Experimental/HairStrands/Source/HairStrandsCore/Private/GroomBuilder.cpp
+++ b/Engine/Plugins/Experimental/HairStrands/Source/HairStrandsCore/Private/GroomBuilder.cpp
@@ -961,8 +961,10 @@ namespace HairInterpolationBuilder
 		FScopedSlowTask SlowTask(RenCurveCount, LOCTEXT("BuildInterpolationData", "Building groom simulation data"));
 		SlowTask.MakeDialog();
 
+		const FDateTime StartTime = FDateTime::UtcNow();
 		ParallelFor(RenCurveCount, 
 		[
+			StartTime,
 			Settings,
 			RenCurveCount, &RenRoots, &RenStrandsData,
 			SimCurveCount, &SimRoots, &SimStrandsData, 
@@ -981,8 +983,14 @@ namespace HairInterpolationBuilder
 
 			if (IsInGameThread())
 			{
-				uint32 CurrentCompletedTasks = CompletedTasks.Exchange(0);
-				SlowTask.EnterProgressFrame(CurrentCompletedTasks, LOCTEXT("ComputeGuidesAndWeights", "Computing closest guides and weights"));
+				const uint32 CurrentCompletedTasks = CompletedTasks.Exchange(0);
+				const float RemainingTasks = FMath::Clamp(SlowTask.TotalAmountOfWork - SlowTask.CompletedWork, 1.f, SlowTask.TotalAmountOfWork);
+				const FTimespan ElaspedTime = FDateTime::UtcNow() - StartTime;
+				const double RemainingTimeInSeconds = RemainingTasks * double(ElaspedTime.GetSeconds() / SlowTask.CompletedWork);
+
+				FTextBuilder TextBuilder;
+				TextBuilder.AppendLineFormat(LOCTEXT("ComputeGuidesAndWeights", "Computing closest guides and weights ({0})"), FText::AsTimespan(FTimespan::FromSeconds(RemainingTimeInSeconds)));
+				SlowTask.EnterProgressFrame(CurrentCompletedTasks, TextBuilder.ToText());
 			}
 
 			const FHairRoot& StrandRoot = RenRoots[RenCurveIndex];
@@ -1250,13 +1258,8 @@ private:
 	TArray<FIntVector> GuideIndices;
 };
 
-bool FGroomBuilder::BuildGroom(const FHairDescription& HairDescription, const FGroomBuildSettings& BuildSettings, UGroomAsset* GroomAsset)
+bool FGroomBuilder::ProcessHairDescription(const FHairDescription& HairDescription, const FGroomBuildSettings& BuildSettings, FProcessedHairDescription& Out)
 {
-	if (!GroomAsset)
-	{
-		return false;
-	}
-
 	TRACE_CPUPROFILER_EVENT_SCOPE(FGroomBuilder::BuildGroom);
 
 	// Convert HairDescription to HairStrandsDatas
@@ -1305,7 +1308,7 @@ bool FGroomBuilder::BuildGroom(const FHairDescription& HairDescription, const FG
 	TStrandAttributesConstRef<float> StrandWidths = HairDescription.StrandAttributes().GetAttributesRef<float>(HairAttribute::Strand::Width);
 
 	TStrandAttributesConstRef<FVector2D> StrandRootUV = HairDescription.StrandAttributes().GetAttributesRef<FVector2D>(HairAttribute::Strand::RootUV);
-	bool bHasUVData = StrandRootUV.IsValid();
+	Out.bHasUVData = StrandRootUV.IsValid();
 
 	TStrandAttributesConstRef<int> StrandGuides = HairDescription.StrandAttributes().GetAttributesRef<int>(HairAttribute::Strand::Guide);
 	TStrandAttributesConstRef<int> GroupIDs = HairDescription.StrandAttributes().GetAttributesRef<int>(HairAttribute::Strand::GroupID);
@@ -1318,16 +1321,9 @@ bool FGroomBuilder::BuildGroom(const FHairDescription& HairDescription, const FG
 
 	// To use ClosestGuides and GuideWeights attributes, guides must be imported from HairDescription and
 	// must include StrandID attribute since ClosestGuides references those IDs
-	bool bCanUseClosestGuidesAndWeights = bImportGuides && StrandIDs.IsValid() && ClosestGuides.IsValid() && GuideWeights.IsValid();
-	
-	FHairStrandsDatas* CurrentHairStrandsDatas = nullptr;
-
-	typedef TPair<FHairGroupInfo, FHairGroupData> FHairGroup;
-	TMap<int32, FHairGroup> HairGroups;
+	Out.bCanUseClosestGuidesAndWeights = bImportGuides && StrandIDs.IsValid() && ClosestGuides.IsValid() && GuideWeights.IsValid();
 
 	int32 GlobalVertexIndex = 0;
-	int32 NumHairCurves = 0;
-	int32 NumGuideCurves = 0;
 	int32 NumHairPoints = 0;
 	int32 NumGuidePoints = 0;
 	for (int32 CurveIndex = 0; CurveIndex < NumCurves; ++CurveIndex)
@@ -1349,13 +1345,13 @@ bool FGroomBuilder::BuildGroom(const FHairDescription& HairDescription, const FG
 			GroupID = GroupIDs[StrandID];
 		}
 
-		FHairGroup& Group = HairGroups.FindOrAdd(GroupID);
+		FHairStrandsDatas* CurrentHairStrandsDatas = nullptr;
+		FProcessedHairDescription::FHairGroup& Group = Out.HairGroups.FindOrAdd(GroupID);
 		FHairGroupInfo& GroupInfo = Group.Key;
 		FHairGroupData& GroupData = Group.Value;
 		GroupInfo.GroupID = GroupID;
 		if (!bIsGuide)
 		{
-			++NumHairCurves;
 			NumHairPoints += CurveNumVertices;
 			CurrentHairStrandsDatas = &GroupData.HairRenderData;
 
@@ -1363,7 +1359,6 @@ bool FGroomBuilder::BuildGroom(const FHairDescription& HairDescription, const FG
 		}
 		else if (bImportGuides)
 		{
-			++NumGuideCurves;
 			NumGuidePoints += CurveNumVertices;
 			CurrentHairStrandsDatas = &GroupData.HairSimulationData;
 
@@ -1376,9 +1371,12 @@ bool FGroomBuilder::BuildGroom(const FHairDescription& HairDescription, const FG
 			continue;
 		}
 
+		if (!CurrentHairStrandsDatas)
+			continue;
+
 		CurrentHairStrandsDatas->StrandsCurves.CurvesCount.Add(CurveNumVertices);
 
-		if (bCanUseClosestGuidesAndWeights)
+		if (Out.bCanUseClosestGuidesAndWeights)
 		{
 			// ClosesGuides needs mapping of StrandID (the strand index in the HairDescription)
 			CurrentHairStrandsDatas->StrandsCurves.StrandIDs.Add(CurveIndex);
@@ -1389,7 +1387,7 @@ bool FGroomBuilder::BuildGroom(const FHairDescription& HairDescription, const FG
 			CurrentHairStrandsDatas->StrandsCurves.GroomIDToIndex.Add(ImportedGroomID, StrandCurveIndex);
 		}
 
-		if (bHasUVData)
+		if (Out.bHasUVData)
 		{
 			CurrentHairStrandsDatas->StrandsCurves.CurvesRootUV.Add(StrandRootUV[StrandID]);
 		}
@@ -1424,50 +1422,77 @@ bool FGroomBuilder::BuildGroom(const FHairDescription& HairDescription, const FG
 		}
 	}
 
+	return true;
+}
+
+bool FGroomBuilder::BuildGroom(const FHairDescription& HairDescription, const FGroomBuildSettings& BuildSettings, UGroomAsset* GroomAsset)
+{
+	if (!GroomAsset)
+	{
+		return false;
+	}
+
+	TRACE_CPUPROFILER_EVENT_SCOPE(FGroomBuilder::BuildGroom);
+
+	FProcessedHairDescription ProcessedHairDescription;
+	if (!ProcessHairDescription(HairDescription, BuildSettings, ProcessedHairDescription))
+	{
+		return false;
+	}
+
 	FGroomComponentRecreateRenderStateContext RecreateRenderContext(GroomAsset);
 
-	for (TPair<int32, FHairGroup>& HairGroupIt : HairGroups)
+	for (TPair<int32, FProcessedHairDescription::FHairGroup>& HairGroupIt : ProcessedHairDescription.HairGroups)
 	{
 		int32 GroupID = HairGroupIt.Key;
-		FHairGroup& Group = HairGroupIt.Value;
+		FProcessedHairDescription::FHairGroup& Group = HairGroupIt.Value;
 		FHairGroupInfo& GroupInfo = Group.Key;
 		FHairGroupData& GroupData = Group.Value;
 
-		FHairStrandsDatas& HairRenderData = GroupData.HairRenderData;
-		int32 GroupNumCurves = HairRenderData.StrandsCurves.Num();
-		HairRenderData.StrandsCurves.SetNum(GroupNumCurves);
-		GroupInfo.NumCurves = GroupNumCurves;
+		GroupInfo.ImportSettings = BuildSettings;
 
-		int32 GroupNumPoints = HairRenderData.StrandsPoints.Num();
-		HairRenderData.StrandsPoints.SetNum(GroupNumPoints);
+		// Rendering data
+		{
+			FHairStrandsDatas& RenData = GroupData.HairRenderData;
 
-		HairStrandsBuilder::BuildInternalData(HairRenderData, !bHasUVData);
+			const int32 RenGroupNumCurves = RenData.StrandsCurves.Num();
+			RenData.StrandsCurves.SetNum(RenGroupNumCurves);
+			GroupInfo.NumCurves = RenGroupNumCurves;
 
-		FHairStrandsDatas& HairSimulationData = GroupData.HairSimulationData;
-		GroupNumCurves = GroupData.HairSimulationData.StrandsCurves.Num();
+			const int32 RenGroupNumPoints = RenData.StrandsPoints.Num();
+			RenData.StrandsPoints.SetNum(RenGroupNumPoints);
+
+			HairStrandsBuilder::BuildInternalData(RenData, !ProcessedHairDescription.bHasUVData);
+		}
 
-		if (GroupNumCurves > 0)
+		// Simulation data
 		{
-			GroupInfo.NumGuides = GroupNumCurves;
-			HairSimulationData.StrandsCurves.SetNum(GroupNumCurves);
+			FHairStrandsDatas& SimData = GroupData.HairSimulationData;
+			const int32 SimGroupNumGuides = SimData.StrandsCurves.Num();
 
-			GroupNumPoints = HairSimulationData.StrandsPoints.Num();
-			HairSimulationData.StrandsPoints.SetNum(GroupNumPoints);
+			if (SimGroupNumGuides > 0)
+			{
+				GroupInfo.NumGuides = SimGroupNumGuides;
+				SimData.StrandsCurves.SetNum(SimGroupNumGuides);
 
-			HairStrandsBuilder::BuildInternalData(HairSimulationData, true); // Imported guides don't currently have root UVs so force computing them
-		}
-		else
-		{
-			GroomAsset->HairToGuideDensity = FMath::Clamp(BuildSettings.HairToGuideDensity, 0.f, 1.f);
+				const int32 GroupNumPoints = SimData.StrandsPoints.Num();
+				SimData.StrandsPoints.SetNum(GroupNumPoints);
+
+				HairStrandsBuilder::BuildInternalData(SimData, true); // Imported guides don't currently have root UVs so force computing them
+			}
+			else
+			{
+				GroomAsset->HairToGuideDensity = FMath::Clamp(BuildSettings.HairToGuideDensity, 0.f, 1.f);
+			}
 		}
 	}
 
 	GroomAsset->HairGroupsInfo.Reset();
 	GroomAsset->HairGroupsData.Reset();
 
-	for (TPair<int32, FHairGroup> HairGroupIt : HairGroups)
+	for (TPair<int32, FProcessedHairDescription::FHairGroup> HairGroupIt : ProcessedHairDescription.HairGroups)
 	{
-		FHairGroup& Group = HairGroupIt.Value;
+		FProcessedHairDescription::FHairGroup& Group = HairGroupIt.Value;
 		FHairGroupInfo& GroupInfo = Group.Key;
 		FHairGroupData& GroupData = Group.Value;
 
@@ -1477,7 +1502,7 @@ bool FGroomBuilder::BuildGroom(const FHairDescription& HairDescription, const FG
 
 	// If there's usable closest guides and guide weights attributes, fill them into the asset
 	// This step requires the HairSimulationData (guides) to be filled prior to this
-	if (bCanUseClosestGuidesAndWeights)
+	if (ProcessedHairDescription.bCanUseClosestGuidesAndWeights)
 	{
 		HairInterpolationBuilder::FillInterpolationData(GroomAsset, HairDescription);
 	}
@@ -1505,13 +1530,14 @@ void FGroomBuilder::BuildData(UGroomAsset* GroomAsset, const FGroomBuildSettings
 		FHairStrandsDatas& HairRenderData = GroupData.HairRenderData;
 		FHairStrandsDatas& HairSimulationData = GroupData.HairSimulationData;
 
+		FHairGroupInfo& GroupInfo = GroomAsset->HairGroupsInfo[Index];
+		GroupInfo.ImportSettings = BuildSettings;
+
 		if (HairSimulationData.GetNumCurves() == 0)
 		{
-			FHairGroupInfo& GroupInfo = GroomAsset->HairGroupsInfo[Index];
 
 			const float GuideDensity = FMath::Clamp(GroomAsset->HairToGuideDensity, 0.f, 1.f);
 			GenerateGuides(HairRenderData, GuideDensity, HairSimulationData);
-			GroupInfo.ImportSettings = BuildSettings;
 			GroupInfo.NumGuides = HairSimulationData.GetNumCurves();
 		}
 
diff --git a/Engine/Plugins/Experimental/HairStrands/Source/HairStrandsCore/Private/GroomComponent.cpp b/Engine/Plugins/Experimental/HairStrands/Source/HairStrandsCore/Private/GroomComponent.cpp
--- a/Engine/Plugins/Experimental/HairStrands/Source/HairStrandsCore/Private/GroomComponent.cpp
+++ b/Engine/Plugins/Experimental/HairStrands/Source/HairStrandsCore/Private/GroomComponent.cpp
@@ -105,65 +105,6 @@ public:
 	}
 };
 
-static bool IsCompatible(UGroomAsset* InGroom, UGroomBindingAsset* InBinding)
-{
-	if (InBinding && InGroom)
-	{
-		if (InGroom->HairGroupsInfo.Num() != InBinding->GroupInfos.Num())
-		{
-			UE_LOG(LogHairStrands, Warning, TEXT("[Groom] The GroomBinding asset (Groom:%s) does not contains the same number of groups. It is incompatible with the groom asset assigned to this component (Groom:%s). The binding asset will be set to null."), *InBinding->Groom->GetName(), *InGroom->GetName());
-			return false;
-		}
-		for (uint32 GroupIt=0, GroupCount=InGroom->HairGroupsInfo.Num(); GroupIt < GroupCount; ++GroupIt)
-		{
-			if (InGroom->HairGroupsInfo[GroupIt].NumCurves != InBinding->GroupInfos[GroupIt].RenRootCount ||
-				InGroom->HairGroupsInfo[GroupIt].NumGuides != InBinding->GroupInfos[GroupIt].SimRootCount)
-			{
-				UE_LOG(LogHairStrands, Warning, TEXT("[Groom] The GroomBinding asset (Groom:%s) does not contains the same curves/guides in group %d. It is incompatible with the groom asset assigned to this component (Groom:%s). The binding asset will be set to null."), *InBinding->Groom->GetName(), GroupIt, *InGroom->GetName());
-				return false;
-			}
-		}
-
-		if (InBinding->Groom != InGroom)
-		{
-			UE_LOG(LogHairStrands, Warning, TEXT("[Groom] The GroomBinding asset (Groom:%s) is incompatible with the groom asset assigned to this component (Groom:%s). The binding asset will be set to null."), *InBinding->Groom->GetName(), *InGroom->GetName());
-			return false;
-		}
-	}
-	return true;
-}
-
-static bool IsBindingAssetValid(UGroomBindingAsset* InBinding)
-{
-	if (InBinding)
-	{
-		if (const UPackage* Package = InBinding->GetOutermost())
-		{
-			if (Package->IsDirty())
-			{
-				UE_LOG(LogHairStrands, Warning, TEXT("[Groom] The binding asset (%s) is not saved and will be considered as invalid. Falling back onto non-binding version."), *InBinding->GetName());
-				return false;
-			}
-		}
-
-		if (InBinding->GroupInfos.Num() == 0)
-		{
-			UE_LOG(LogHairStrands, Warning, TEXT("[Groom] The GroomBinding asset (Groom:%s) does not contain any groups. It is invalid and can't be assigned. The binding asset will be set to null."), *InBinding->Groom->GetName());
-			return false;
-		}
-
-		for (FGoomBindingGroupInfo& Info : InBinding->GroupInfos)
-		{
-			if (Info.RenRootCount == 0 || Info.SimRootCount == 0)
-			{
-				UE_LOG(LogHairStrands, Warning, TEXT("[Groom] The GroomBinding asset (Groom:%s) has group with 0 curves/guides. It is invalid and can't be assigned. The binding asset will be set to null."), *InBinding->Groom->GetName());
-				return false;
-			}
-		}
-	}
-	return true;
-}
-
 /////////////////////////////////////////////////////////////////////////////////////////
 /////////////////////////////////////////////////////////////////////////////////////////
 //  FStrandHairSceneProxy
@@ -458,6 +399,7 @@ UGroomComponent::UGroomComponent(const FObjectInitializer& ObjectInitializer)
 	InitializedResources = nullptr;
 	Mobility = EComponentMobility::Movable;
 	bIsGroomAssetCallbackRegistered = false;
+	bIsGroomBindingAssetCallbackRegistered = false;
 	SourceSkeletalMesh = nullptr; 
 	NiagaraComponent = nullptr;
 
@@ -483,12 +425,16 @@ void UGroomComponent::UpdateHairGroupsDesc()
 
 	for (uint32 GroupIt = 0; GroupIt < GroupCount; ++GroupIt)
 	{
-		FHairGroupInfo& GroupInfo = GroomAsset->HairGroupsInfo[GroupIt];
-		FHairGroupData& GroupData = GroomAsset->HairGroupsData[GroupIt];
+		const FHairGroupInfo& GroupInfo = GroomAsset->HairGroupsInfo[GroupIt];
+		const FHairGroupData& GroupData = GroomAsset->HairGroupsData[GroupIt];
 
 		FHairGroupDesc& Desc = GroomGroupsDesc[GroupIt];
+		if (bReinitOverride)
+		{
+			Desc.ReInit();
+		}
 		Desc.GuideCount = GroupInfo.NumGuides;
-		Desc.HairCount = GroupInfo.NumCurves;
+		Desc.HairCount  = GroupInfo.NumCurves;
 		Desc.HairLength = GroupData.HairRenderData.StrandsCurves.MaxLength;
 
 		if (bReinitOverride || Desc.HairWidth == 0)
@@ -499,11 +445,6 @@ void UGroomComponent::UpdateHairGroupsDesc()
 		{
 			Desc.HairShadowDensity = GroupData.HairRenderData.HairDensity;
 		}
-
-		if (bReinitOverride)
-		{
-			Desc.ReInit();
-		}
 	}
 	UpdateNiagaraComponent();
 }
@@ -551,8 +492,12 @@ void UGroomComponent::UpdateNiagaraComponent()
 void UGroomComponent::SetGroomAsset(UGroomAsset* Asset)
 {
 	ReleaseResources();
-	GroomAsset = Asset;
-	if (!IsBindingAssetValid(BindingAsset) || !IsCompatible(GroomAsset, BindingAsset))
+	if (Asset && Asset->IsValid())
+	{
+		GroomAsset = Asset;
+	}
+
+	if (!UGroomBindingAsset::IsBindingAssetValid(BindingAsset) || !UGroomBindingAsset::IsCompatible(GroomAsset, BindingAsset))
 	{
 		BindingAsset = nullptr;
 	}
@@ -563,12 +508,16 @@ void UGroomComponent::SetGroomAsset(UGroomAsset* Asset)
 	InitResources();
 }
 
+static void ResetSimulation(const USceneComponent* Component);
 void UGroomComponent::SetGroomAsset(UGroomAsset* Asset, UGroomBindingAsset* InBinding)
 {
 	ReleaseResources();
-	GroomAsset = Asset;
+	if (Asset && Asset->IsValid())
+	{
+		GroomAsset = Asset;
+	}
 	BindingAsset = InBinding;
-	if (!IsBindingAssetValid(BindingAsset) || !IsCompatible(GroomAsset, BindingAsset))
+	if (!UGroomBindingAsset::IsBindingAssetValid(BindingAsset) || !UGroomBindingAsset::IsCompatible(GroomAsset, BindingAsset))
 	{
 		BindingAsset = nullptr;
 	}
@@ -582,6 +531,16 @@ void UGroomComponent::SetGroomAsset(UGroomAsset* Asset, UGroomBindingAsset* InBi
 	if (!Asset)
 		return;
 	InitResources();
+	ResetSimulation(this);
+}
+
+void UGroomComponent::SetStableRasterization(bool bEnable)
+{
+	for (FHairGroupDesc& HairDesc : GroomGroupsDesc)
+	{
+		HairDesc.bUseStableRasterization = bEnable;
+	}
+	UpdateHairGroupsDescAndInvalidateRenderState();
 }
 
 void UGroomComponent::SetHairLengthScale(float Scale) 
@@ -594,13 +553,32 @@ void UGroomComponent::SetHairLengthScale(float Scale)
 	UpdateHairGroupsDescAndInvalidateRenderState();
 }
 
+void UGroomComponent::SetHairRootScale(float Scale)
+{
+	Scale = FMath::Clamp(Scale, 0.f, 10.f);
+	for (FHairGroupDesc& HairDesc : GroomGroupsDesc)
+	{
+		HairDesc.HairRootScale = Scale;
+	}
+	UpdateHairGroupsDescAndInvalidateRenderState();
+}
+
 void UGroomComponent::SetHairWidth(float HairWidth)
 {
 	for (FHairGroupDesc& HairDesc : GroomGroupsDesc)
 	{
 		HairDesc.HairWidth = HairWidth;
 	}
-	InitResources();
+	UpdateHairGroupsDescAndInvalidateRenderState();
+}
+
+void UGroomComponent::SetScatterSceneLighting(bool Enable)
+{
+	for (FHairGroupDesc& HairDesc : GroomGroupsDesc)
+	{
+		HairDesc.bScatterSceneLighting = Enable;
+	}
+	UpdateHairGroupsDescAndInvalidateRenderState();
 }
 
 void UGroomComponent::SetBinding(bool bBind)
@@ -616,8 +594,8 @@ void UGroomComponent::SetBinding(UGroomBindingAsset* InBinding)
 {
 	if (BindingAsset != InBinding)
 	{
-		const bool bIsValid = InBinding != nullptr ? IsBindingAssetValid(BindingAsset) : true;
-		if (bIsValid && IsCompatible(GroomAsset, InBinding))
+		const bool bIsValid = InBinding != nullptr ? UGroomBindingAsset::IsBindingAssetValid(BindingAsset) : true;
+		if (bIsValid && UGroomBindingAsset::IsCompatible(GroomAsset, InBinding))
 		{
 			BindingAsset = InBinding;
 			bBindGroomToSkeletalMesh = InBinding != nullptr;
@@ -652,28 +630,16 @@ FPrimitiveSceneProxy* UGroomComponent::CreateSceneProxy()
 
 FBoxSphereBounds UGroomComponent::CalcBounds(const FTransform& InLocalToWorld) const
 {
-	FBox HairBox(ForceInit);
 	if (GroomAsset && GroomAsset->GetNumHairGroups() > 0)
 	{
 		if (RegisteredSkeletalMeshComponent)
 		{
-			FBox LocalBound(EForceInit::ForceInit);
-			for (const FHairGroupData& GroupData : GroomAsset->HairGroupsData)
-			{
-				LocalBound += GroupData.HairRenderData.BoundingBox;
-			}
-			FBox WorldBound = LocalBound.TransformBy(InLocalToWorld);
-
-			const FVector MeshTranslation = RegisteredSkeletalMeshComponent->CalcBounds(InLocalToWorld).GetBox().GetCenter();
-			const FVector LocalAnimationTranslation = MeshTranslation - WorldBound.GetCenter();
-
-			WorldBound.Max += LocalAnimationTranslation;
-			WorldBound.Min += LocalAnimationTranslation;
-			return FBoxSphereBounds(WorldBound);
+			const FBox WorldSkeletalBound = RegisteredSkeletalMeshComponent->CalcBounds(InLocalToWorld).GetBox();
+			return FBoxSphereBounds(WorldSkeletalBound);
 		}
 		else
 		{
-			FBox LocalBounds(EForceInit::ForceInit);
+			FBox LocalBounds(EForceInit::ForceInitToZero);
 			for (const FHairGroupData& GroupData : GroomAsset->HairGroupsData)
 			{
 				LocalBounds += GroupData.HairRenderData.BoundingBox;
@@ -683,7 +649,7 @@ FBoxSphereBounds UGroomComponent::CalcBounds(const FTransform& InLocalToWorld) c
 	}
 	else
 	{
-		return FBoxSphereBounds();
+		return FBoxSphereBounds(EForceInit::ForceInitToZero);
 	}
 }
 
@@ -717,7 +683,7 @@ UMaterialInterface* UGroomComponent::GetMaterial(int32 ElementIndex) const
 		}
 	}
 
-	if (OverrideMaterial == nullptr || !OverrideMaterial->GetMaterialResource(FeatureLevel)->IsUsedWithHairStrands())
+	if (OverrideMaterial == nullptr || OverrideMaterial->GetMaterialResource(FeatureLevel) == nullptr || !OverrideMaterial->GetMaterialResource(FeatureLevel)->IsUsedWithHairStrands())
 	{
 		OverrideMaterial = GEngine->HairDefaultMaterial;
 	}
@@ -794,27 +760,6 @@ template<typename T> void SafeRelease(T*& Data)
 	}
 }
 
-void CallbackMeshObjectCallback(
-	FSkeletalMeshObjectCallbackData::EEventType Event,
-	FSkeletalMeshObject* MeshObject,
-	uint64 UserData) 
-{
-	ENQUEUE_RENDER_COMMAND(FHairStrandsMeshObjectUpdate)(
-		[Event, MeshObject, UserData](FRHICommandListImmediate& RHICmdList)
-	{
-		const uint64 ComponentId = uint64(UserData & 0xFFFFFFFF);
-		const EWorldType::Type WorldType = EWorldType::Type((UserData>>32) & 0xFFFFFFFF);
-		if (Event == FSkeletalMeshObjectCallbackData::EEventType::Register || Event == FSkeletalMeshObjectCallbackData::EEventType::Update)
-		{
-			UpdateHairStrands(ComponentId, WorldType, MeshObject);
-		}
-		else
-		{
-			UpdateHairStrands(ComponentId, WorldType, nullptr);
-		}
-	});
-}
-
 static bool IsSimulationEnabled(const USceneComponent* Component)
 {
 	check(Component);
@@ -854,19 +799,30 @@ static void ResetSimulation(const USceneComponent* Component)
 	}
 }
 
+EWorldType::Type UGroomComponent::GetWorldType() const
+{
+	EWorldType::Type WorldType = GetWorld() ? EWorldType::Type(GetWorld()->WorldType) : EWorldType::None;
+	return WorldType == EWorldType::Inactive ? EWorldType::Editor : WorldType;
+}
+
 void UGroomComponent::EnableSimulatedGroups()
 {
 	if (InterpolationInput)
 	{
+		const uint32 Id = ComponentId.PrimIDValue;
+		const EWorldType::Type WorldType = GetWorldType();
+
 		FHairStrandsInterpolationInput* LocalInterpolationInput = InterpolationInput;
 		UGroomAsset* LocalGroomAsset = GroomAsset;
 		ENQUEUE_RENDER_COMMAND(FHairStrandsTick_UEnableSimulatedGroups)(
-			[LocalInterpolationInput, LocalGroomAsset](FRHICommandListImmediate& RHICmdList)
+			[LocalInterpolationInput, LocalGroomAsset, Id, WorldType](FRHICommandListImmediate& RHICmdList)
 		{
 			int32 GroupIt = 0;
 			for (FHairStrandsInterpolationInput::FHairGroup& HairGroup : LocalInterpolationInput->HairGroups)
 			{
-				HairGroup.bIsSimulationEnable = LocalGroomAsset->EnableSimulation && (LocalGroomAsset->SimulatedGroup == GroupIt);
+				const bool bIsSimulationEnable = LocalGroomAsset->EnableSimulation && (LocalGroomAsset->SimulatedGroup == GroupIt);
+				HairGroup.bIsSimulationEnable = bIsSimulationEnable;
+				UpdateHairStrandsDebugInfo(Id, WorldType, GroupIt, bIsSimulationEnable);
 				++GroupIt;
 			}
 		});
@@ -877,14 +833,20 @@ void UGroomComponent::DisableSimulatedGroups()
 {
 	if (InterpolationInput)
 	{
+		const uint32 Id = ComponentId.PrimIDValue;
+		const EWorldType::Type WorldType = GetWorldType();
+
 		FHairStrandsInterpolationInput* LocalInterpolationInput = InterpolationInput;
 		UGroomAsset* LocalGroomAsset = GroomAsset;
 		ENQUEUE_RENDER_COMMAND(FHairStrandsTick_DisableSimulatedGroups)(
-			[LocalInterpolationInput, LocalGroomAsset](FRHICommandListImmediate& RHICmdList)
+			[LocalInterpolationInput, LocalGroomAsset, Id, WorldType](FRHICommandListImmediate& RHICmdList)
 		{
+			int32 GroupIt = 0;
 			for (FHairStrandsInterpolationInput::FHairGroup& HairGroup : LocalInterpolationInput->HairGroups)
 			{
 				HairGroup.bIsSimulationEnable = false;
+				UpdateHairStrandsDebugInfo(Id, WorldType, GroupIt, false);
+				++GroupIt;
 			}
 		});
 	}
@@ -907,9 +869,10 @@ void UGroomComponent::OnChildAttached(USceneComponent* ChildComponent)
 	}
 }
 
-void UGroomComponent::InitResources()
+void UGroomComponent::InitResources(bool bIsBindingReloading)
 {
 	ReleaseResources();
+	bInitSimulation = true;
 	bResetSimulation = true;
 
 	UpdateHairGroupsDesc();
@@ -920,29 +883,24 @@ void UGroomComponent::InitResources()
 	InitializedResources = GroomAsset;
 
 	const FPrimitiveComponentId LocalComponentId = ComponentId;
-	EWorldType::Type WorldType = GetWorld() ? EWorldType::Type(GetWorld()->WorldType) : EWorldType::None;
-	WorldType = WorldType == EWorldType::Inactive ? EWorldType::Editor : WorldType;
+	const EWorldType::Type WorldType = GetWorldType();
 
 	// Insure the ticking of the Groom component always happens after the skeletalMeshComponent.
 	USkeletalMeshComponent* SkeletalMeshComponent = bBindGroomToSkeletalMesh && GetAttachParent() ? Cast<USkeletalMeshComponent>(GetAttachParent()) : nullptr;
 	const bool bHasValidSketalMesh = SkeletalMeshComponent && SkeletalMeshComponent->GetSkeletalMeshRenderData();
 
-	// Always setup the callback, even if the SkeletalMeshData are not ready yet
-	if (SkeletalMeshComponent)
-	{
-		FSkeletalMeshObjectCallbackData CallbackData;
-		CallbackData.Run = CallbackMeshObjectCallback;
-		CallbackData.UserData = (uint64(LocalComponentId.PrimIDValue) & 0xFFFFFFFF) | (uint64(WorldType) << 32);
-		SkeletalMeshComponent->MeshObjectCallbackData = CallbackData;
-	}
-
+	uint32 SkeletalComponentId = ~0;
 	if (bHasValidSketalMesh)
 	{
 		RegisteredSkeletalMeshComponent = SkeletalMeshComponent;
+		SkeletalComponentId = SkeletalMeshComponent->ComponentId.PrimIDValue;
 		AddTickPrerequisiteComponent(SkeletalMeshComponent);
 	}
 
-	const bool bIsBindingCompatible = IsCompatible(GroomAsset, BindingAsset) && IsBindingAssetValid(BindingAsset);
+	const bool bIsBindingCompatible = 
+		UGroomBindingAsset::IsCompatible(SkeletalMeshComponent ? SkeletalMeshComponent->SkeletalMesh : nullptr, BindingAsset) &&
+		UGroomBindingAsset::IsCompatible(GroomAsset, BindingAsset) && 
+		UGroomBindingAsset::IsBindingAssetValid(BindingAsset, bIsBindingReloading);
 	const bool bIsSimulationEnable = IsSimulationEnabled(this);
 
 	FTransform HairLocalToWorld = GetComponentTransform();
@@ -969,8 +927,7 @@ void UGroomComponent::InitResources()
 		}
 
 		FHairGroupResource& Res = HairGroupResources->HairGroups.AddDefaulted_GetRef();
-		Res.InterpolationResource = new FHairStrandsInterpolationResource(GroupData.HairInterpolationData.RenderData, GroupData.HairSimulationData);
-		BeginInitResource(Res.InterpolationResource);
+		Res.InterpolationResource = GroupData.HairInterpolationResource;
 
 		#if RHI_RAYTRACING
 		if (IsHairRayTracingEnabled())
@@ -990,6 +947,7 @@ void UGroomComponent::InitResources()
 				Res.bOwnRootResourceAllocation = false;
 				Res.RenRootResources = BindingAsset->HairGroupResources[GroupIt].RenRootResources;
 				Res.SimRootResources = BindingAsset->HairGroupResources[GroupIt].SimRootResources;
+				DebugHairGroup.bHasBinding = true;
 			}
 			else
 			{
@@ -997,8 +955,9 @@ void UGroomComponent::InitResources()
 				if (LODCount > 0)
 				{
 					Res.bOwnRootResourceAllocation = true;
-					Res.RenRootResources = new FHairStrandsRootResource(&GroupData.HairRenderData, LODCount);
-					Res.SimRootResources = new FHairStrandsRootResource(&GroupData.HairSimulationData, LODCount);
+					TArray<uint32> NumSamples; NumSamples.Init(0, LODCount);
+					Res.RenRootResources = new FHairStrandsRootResource(&GroupData.HairRenderData, LODCount, NumSamples);
+					Res.SimRootResources = new FHairStrandsRootResource(&GroupData.HairSimulationData, LODCount, NumSamples);
 					BeginInitResource(Res.RenRootResources);
 					BeginInitResource(Res.SimRootResources);
 				}
@@ -1010,7 +969,7 @@ void UGroomComponent::InitResources()
 
 		Res.RenderDeformedResources = new FHairStrandsDeformedResource(GroupData.HairRenderData.RenderData, false);
 		Res.SimDeformedResources = new FHairStrandsDeformedResource(GroupData.HairSimulationData.RenderData, true);
-		Res.ClusterCullingResources = new FHairStrandsClusterCullingResource(GroupData);
+		Res.ClusterCullingResources = GroupData.ClusterCullingResource;
 
 		const uint32 HairVertexCount = GroupData.HairRenderData.GetNumPoints();
 		const uint32 GroupInstanceVertexCount = HairVertexCount * 6; // 6 vertex per point for a quad
@@ -1018,7 +977,6 @@ void UGroomComponent::InitResources()
 
 		BeginInitResource(Res.RenderDeformedResources);
 		BeginInitResource(Res.SimDeformedResources);
-		BeginInitResource(Res.ClusterCullingResources);
 		BeginInitResource(Res.HairGroupPublicDatas);
 
 		const FVector RenderRestHairPositionOffset = Res.RenderRestResources->PositionOffset;
@@ -1046,6 +1004,7 @@ void UGroomComponent::InitResources()
 		InterpolationInputGroup.OutHairPositionOffset = RenderRestHairPositionOffset;
 		InterpolationInputGroup.OutHairPreviousPositionOffset = RenderRestHairPositionOffset;
 		InterpolationInputGroup.bIsSimulationEnable = bIsSimulationEnable && GroomAsset->EnableSimulation && (GroomAsset->SimulatedGroup == GroupIt);
+		DebugHairGroup.bHasSimulation = InterpolationInputGroup.bIsSimulationEnable;
 
 		GroupIt++;
 	}
@@ -1059,10 +1018,11 @@ void UGroomComponent::InitResources()
 	USkeletalMesh* InSourceSkeletalMesh = SourceSkeletalMesh;
 	const bool bRunMeshProjection = bHasValidSketalMesh && (!BindingAsset || !bIsBindingCompatible);
 	FHairGroupResources* LocalResources = HairGroupResources;
-	const uint64 Id = LocalComponentId.PrimIDValue;
+	const uint32 Id = LocalComponentId.PrimIDValue;
 	ENQUEUE_RENDER_COMMAND(FHairStrandsBuffers)(
 		[
 			Id,
+			SkeletalComponentId,
 			Interpolation,
 			LocalResources,
 			HairLocalToWorld, SkinLocalToWorld,
@@ -1120,6 +1080,7 @@ void UGroomComponent::InitResources()
 			OutputGroup.RenderMaterialBuffer			= &Res.RenderRestResources->MaterialBuffer;
 			OutputGroup.RenderTangentBuffer				= &Res.RenderDeformedResources->TangentBuffer;
 			OutputGroup.SimTangentBuffer				= &Res.SimDeformedResources->TangentBuffer;
+			OutputGroup.CurrentIndex					= &Res.SimDeformedResources->CurrentIndex;
 			
 			OutputGroup.RenderClusterAABBBuffer			= &Res.HairGroupPublicDatas->GetClusterAABBBuffer();
 			OutputGroup.RenderGroupAABBBuffer			= &Res.HairGroupPublicDatas->GetGroupAABBBuffer();
@@ -1175,10 +1136,12 @@ void UGroomComponent::InitResources()
 				TargetMeshData,
 				RenProjectionDatas,
 				SimProjectionDatas);
+
 		}
 
 		RegisterHairStrands(
 			Id,
+			SkeletalComponentId,
 			WorldType,
 			Interpolation,
 			RenProjectionDatas,
@@ -1202,8 +1165,6 @@ void UGroomComponent::DeleteHairGroupResources(FHairGroupResources*& InHairGroup
 	{
 		for (FHairGroupResource& Res : Local->HairGroups)
 		{
-			SafeRelease(Res.InterpolationResource);
-
 			// Release the root resources only if they have been created internally (vs. being created by external asset)
 			if (Res.bOwnRootResourceAllocation)
 			{
@@ -1212,7 +1173,6 @@ void UGroomComponent::DeleteHairGroupResources(FHairGroupResources*& InHairGroup
 			}
 
 			SafeRelease(Res.RenderDeformedResources);
-			SafeRelease(Res.ClusterCullingResources);
 			SafeRelease(Res.HairGroupPublicDatas);
 			SafeRelease(Res.SimDeformedResources);
 			#if RHI_RAYTRACING
@@ -1229,7 +1189,7 @@ void UGroomComponent::ReleaseResources()
 {
 	// Unregister component interpolation resources
 	const FPrimitiveComponentId LocalComponentId = ComponentId;
-	const uint64 Id = LocalComponentId.PrimIDValue;
+	const uint32 Id = LocalComponentId.PrimIDValue;
 	ENQUEUE_RENDER_COMMAND(StaticMeshVertexBuffersLegacyInit)(
 		[Id](FRHICommandListImmediate& RHICmdList)
 	{
@@ -1273,6 +1233,11 @@ void UGroomComponent::PostLoad()
 		GroomAsset->ConditionalPostLoad();
 	}
 
+	if (GroomAsset && !GroomAsset->IsValid())
+	{
+		GroomAsset = nullptr;
+	}
+
 	if (BindingAsset)
 	{
 		// Make sure that the asset initialized its resources first since the component needs them to initialize its own resources
@@ -1286,6 +1251,12 @@ void UGroomComponent::PostLoad()
 		GroomAsset->GetOnGroomAssetChanged().AddUObject(this, &UGroomComponent::Invalidate);
 		bIsGroomAssetCallbackRegistered = true;
 	}
+
+	if (BindingAsset && !bIsGroomBindingAssetCallbackRegistered)
+	{
+		BindingAsset->GetOnGroomBindingAssetChanged().AddUObject(this, &UGroomComponent::InvalidateAndRecreate);
+		bIsGroomBindingAssetCallbackRegistered = true;
+	}
 	ValidateMaterials(false);
 #endif
 }
@@ -1297,6 +1268,12 @@ void UGroomComponent::Invalidate()
 	MarkRenderStateDirty();
 	ValidateMaterials(false);
 }
+
+void UGroomComponent::InvalidateAndRecreate()
+{
+	InitResources(true);
+	MarkRenderStateDirty();
+}
 #endif
 
 void UGroomComponent::OnRegister()
@@ -1311,8 +1288,8 @@ void UGroomComponent::OnRegister()
 	// Insure the ticking of the Groom component always happens after the skeletalMeshComponent.
 	USkeletalMeshComponent* SkeletalMeshComponent = GetAttachParent() ? Cast<USkeletalMeshComponent>(GetAttachParent()) : nullptr;
 
-	const EWorldType::Type WorldType = GetWorld() ? EWorldType::Type(GetWorld()->WorldType) : EWorldType::None;
-	const uint64 Id = ComponentId.PrimIDValue;
+	const EWorldType::Type WorldType = GetWorldType();
+	const uint32 Id = ComponentId.PrimIDValue;
 
 	ENQUEUE_RENDER_COMMAND(FHairStrandsRegister)(
 		[Id, WorldType](FRHICommandListImmediate& RHICmdList)
@@ -1336,6 +1313,11 @@ void UGroomComponent::OnComponentDestroyed(bool bDestroyingHierarchy)
 		GroomAsset->GetOnGroomAssetChanged().RemoveAll(this);
 		bIsGroomAssetCallbackRegistered = false;
 	}
+	if (bIsGroomBindingAssetCallbackRegistered)
+	{
+		BindingAsset->GetOnGroomBindingAssetChanged().RemoveAll(this);
+		bIsGroomBindingAssetCallbackRegistered = false;
+	}
 #endif
 
 	Super::OnComponentDestroyed(bDestroyingHierarchy);
@@ -1355,37 +1337,34 @@ void UGroomComponent::TickComponent(float DeltaTime, enum ELevelTick TickType, F
 {
 	Super::TickComponent(DeltaTime, TickType, ThisTickFunction);	
 	
-	const EWorldType::Type WorldType = GetWorld() ? EWorldType::Type(GetWorld()->WorldType) : EWorldType::None;
-	const uint64 Id = ComponentId.PrimIDValue;
+	const EWorldType::Type WorldType = GetWorldType();
+	const uint32 Id = ComponentId.PrimIDValue;
 	const ERHIFeatureLevel::Type FeatureLevel = GetWorld() ? ERHIFeatureLevel::Type(GetWorld()->FeatureLevel) : ERHIFeatureLevel::Num;
 
 	FVector MeshPositionOffset = FVector::ZeroVector;
 	USkeletalMeshComponent* SkeletalMeshComponent = RegisteredSkeletalMeshComponent;
 
-	if (USkeletalMeshComponent* ParentComp = Cast<USkeletalMeshComponent>(GetAttachParent()))
+	bResetSimulation = bInitSimulation;
+	if (!bInitSimulation)
 	{
-		if (ParentComp->GetNumBones() > 0)
+		if (USkeletalMeshComponent* ParentComp = Cast<USkeletalMeshComponent>(GetAttachParent()))
 		{
-			const int32 BoneIndex = FMath::Min(1, ParentComp->GetNumBones()-1);
-			const FMatrix NextBoneMatrix = ParentComp->GetBoneMatrix(BoneIndex);
-
-			const float BoneDistance = FVector::DistSquared(PrevBoneMatrix.GetOrigin(), NextBoneMatrix.GetOrigin());
-			if (ParentComp->GetTeleportDistanceThreshold() > 0.0 && BoneDistance >
-				ParentComp->GetTeleportDistanceThreshold() * ParentComp->GetTeleportDistanceThreshold())
+			if (ParentComp->GetNumBones() > 0)
 			{
-				bResetSimulation = true;
-			}
-			else
-			{
-				bResetSimulation = false;
+				const int32 BoneIndex = FMath::Min(1, ParentComp->GetNumBones() - 1);
+				const FMatrix NextBoneMatrix = ParentComp->GetBoneMatrix(BoneIndex);
+
+				const float BoneDistance = FVector::DistSquared(PrevBoneMatrix.GetOrigin(), NextBoneMatrix.GetOrigin());
+				if (ParentComp->GetTeleportDistanceThreshold() > 0.0 && BoneDistance >
+					ParentComp->GetTeleportDistanceThreshold() * ParentComp->GetTeleportDistanceThreshold())
+				{
+					bResetSimulation = true;
+				}
+				PrevBoneMatrix = NextBoneMatrix;
 			}
-			PrevBoneMatrix = NextBoneMatrix;
 		}
 	}
-	else
-	{
-		bResetSimulation = false;
-	}
+	bInitSimulation = false;
 
 	if (!HairGroupResources)
 	{
@@ -1474,8 +1453,7 @@ void UGroomComponent::PreEditChange(FProperty* PropertyAboutToChange)
 	Super::PreEditChange(PropertyAboutToChange);
 
 	const FName PropertyName = PropertyAboutToChange ? PropertyAboutToChange->GetFName() : NAME_None;
-	const bool bAssetAboutToChanged = PropertyName == GET_MEMBER_NAME_CHECKED(UGroomComponent, GroomAsset);
-	if (bAssetAboutToChanged)
+	if (PropertyName == GET_MEMBER_NAME_CHECKED(UGroomComponent, GroomAsset))
 	{
 		// Remove the callback on the GroomAsset about to be replaced
 		if (bIsGroomAssetCallbackRegistered && GroomAsset)
@@ -1484,6 +1462,16 @@ void UGroomComponent::PreEditChange(FProperty* PropertyAboutToChange)
 		}
 		bIsGroomAssetCallbackRegistered = false;
 	}
+
+	if (PropertyName == GET_MEMBER_NAME_CHECKED(UGroomComponent, BindingAsset))
+	{
+		// Remove the callback on the GroomAsset about to be replaced
+		if (bIsGroomBindingAssetCallbackRegistered && BindingAsset)
+		{
+			BindingAsset->GetOnGroomBindingAssetChanged().RemoveAll(this);
+		}
+		bIsGroomBindingAssetCallbackRegistered = false;
+	}
 }
 
 void UGroomComponent::PostEditChangeProperty(FPropertyChangedEvent& PropertyChangedEvent) 
@@ -1493,11 +1481,12 @@ void UGroomComponent::PostEditChangeProperty(FPropertyChangedEvent& PropertyChan
 
 	//  Init/release resources when setting the GroomAsset (or undoing)
 	const bool bAssetChanged = PropertyName == GET_MEMBER_NAME_CHECKED(UGroomComponent, GroomAsset);
+	const bool bSourceSkeletalMeshChanged = PropertyName == GET_MEMBER_NAME_CHECKED(UGroomComponent, SourceSkeletalMesh);
 	const bool bBindingAssetChanged = PropertyName == GET_MEMBER_NAME_CHECKED(UGroomComponent, BindingAsset);
 	const bool bBindToSkeletalChanged = PropertyName == GET_MEMBER_NAME_CHECKED(UGroomComponent, bBindGroomToSkeletalMesh);
 	const bool bDefaultNiagaraChanged = PropertyName == GET_MEMBER_NAME_CHECKED(UGroomComponent, bCreateNiagaraComponent);
-	const bool bIsBindingCompatible = IsCompatible(GroomAsset, BindingAsset);
-	if (!bIsBindingCompatible || !IsBindingAssetValid(BindingAsset))
+	const bool bIsBindingCompatible = UGroomBindingAsset::IsCompatible(GroomAsset, BindingAsset);
+	if (!bIsBindingCompatible || !UGroomBindingAsset::IsBindingAssetValid(BindingAsset))
 	{
 		BindingAsset = nullptr;
 	}
@@ -1506,7 +1495,12 @@ void UGroomComponent::PostEditChangeProperty(FPropertyChangedEvent& PropertyChan
 		bBindGroomToSkeletalMesh = true;
 	}
 
-	const bool bRecreateResources = bAssetChanged || bBindingAssetChanged || PropertyThatChanged == nullptr || bBindToSkeletalChanged || bDefaultNiagaraChanged;
+	if (GroomAsset && !GroomAsset->IsValid())
+	{
+		GroomAsset = nullptr;
+	}
+
+	const bool bRecreateResources = bAssetChanged || bBindingAssetChanged || PropertyThatChanged == nullptr || bBindToSkeletalChanged || bDefaultNiagaraChanged || bSourceSkeletalMeshChanged;
 	if (bRecreateResources)
 	{
 		// Release the resources before Super::PostEditChangeProperty so that they get
@@ -1533,6 +1527,17 @@ void UGroomComponent::PostEditChangeProperty(FPropertyChangedEvent& PropertyChan
 			bIsGroomAssetCallbackRegistered = true;
 		}
 	}
+
+	if (bBindingAssetChanged)
+	{
+		if (BindingAsset)
+		{
+			ResetSimulation(this);
+			// Set the callback on the new GroomAsset being assigned
+			BindingAsset->GetOnGroomBindingAssetChanged().AddUObject(this, &UGroomComponent::InvalidateAndRecreate);
+			bIsGroomBindingAssetCallbackRegistered = true;
+		}
+	}
 #endif
 
 	if (PropertyName == GET_MEMBER_NAME_CHECKED(FHairGroupDesc, HairWidth) || 
@@ -1543,7 +1548,9 @@ void UGroomComponent::PostEditChangeProperty(FPropertyChangedEvent& PropertyChan
 		PropertyName == GET_MEMBER_NAME_CHECKED(FHairGroupDesc, HairShadowDensity) ||
 		PropertyName == GET_MEMBER_NAME_CHECKED(FHairGroupDesc, HairRaytracingRadiusScale) ||
 		PropertyName == GET_MEMBER_NAME_CHECKED(FHairGroupDesc, LodBias) ||
-		PropertyName == GET_MEMBER_NAME_CHECKED(FHairGroupDesc, LodAverageVertexPerPixel))
+		PropertyName == GET_MEMBER_NAME_CHECKED(FHairGroupDesc, LodAverageVertexPerPixel) ||
+		PropertyName == GET_MEMBER_NAME_CHECKED(FHairGroupDesc, bUseStableRasterization) ||
+		PropertyName == GET_MEMBER_NAME_CHECKED(FHairGroupDesc, bScatterSceneLighting))
 	{	
 		UpdateHairGroupsDescAndInvalidateRenderState();
 	}
diff --git a/Engine/Plugins/Experimental/HairStrands/Source/HairStrandsCore/Private/GroomCreateBindingOptions.cpp b/Engine/Plugins/Experimental/HairStrands/Source/HairStrandsCore/Private/GroomCreateBindingOptions.cpp
--- a/Engine/Plugins/Experimental/HairStrands/Source/HairStrandsCore/Private/GroomCreateBindingOptions.cpp
+++ b/Engine/Plugins/Experimental/HairStrands/Source/HairStrandsCore/Private/GroomCreateBindingOptions.cpp
@@ -7,4 +7,5 @@ UGroomCreateBindingOptions::UGroomCreateBindingOptions(const FObjectInitializer&
 {
 	SourceSkeletalMesh = nullptr;
 	TargetSkeletalMesh = nullptr;
+	NumInterpolationPoints = 100;
 }
diff --git a/Engine/Plugins/Experimental/HairStrands/Source/HairStrandsCore/Private/GroomImportOptions.cpp b/Engine/Plugins/Experimental/HairStrands/Source/HairStrandsCore/Private/GroomImportOptions.cpp
--- a/Engine/Plugins/Experimental/HairStrands/Source/HairStrandsCore/Private/GroomImportOptions.cpp
+++ b/Engine/Plugins/Experimental/HairStrands/Source/HairStrandsCore/Private/GroomImportOptions.cpp
@@ -6,3 +6,8 @@ UGroomImportOptions::UGroomImportOptions(const FObjectInitializer& ObjectInitial
 : Super(ObjectInitializer)
 {
 }
+
+UGroomHairGroupsPreview::UGroomHairGroupsPreview(const FObjectInitializer& ObjectInitializer)
+	: Super(ObjectInitializer)
+{
+}
