Integrate hair changes (rendering+simulation)

diff --git a/Engine/Plugins/Experimental/HairStrands/Source/HairStrandsCore/Private/HairStrandsRendering.cpp b/Engine/Plugins/Experimental/HairStrands/Source/HairStrandsCore/Private/HairStrandsRendering.cpp
--- a/Engine/Plugins/Experimental/HairStrands/Source/HairStrandsCore/Private/HairStrandsRendering.cpp
+++ b/Engine/Plugins/Experimental/HairStrands/Source/HairStrandsCore/Private/HairStrandsRendering.cpp
@@ -13,6 +13,29 @@
 #include "Rendering\SkeletalMeshRenderData.h"
 #include "HAL/ConsoleManager.h"
 #include "GpuDebugRendering.h"
+#include "Async/ParallelFor.h"
+
+
+// Just to be sure, also added this in Eigen.Build.cs
+#ifndef EIGEN_MPL2_ONLY
+#define EIGEN_MPL2_ONLY
+#endif
+
+#if defined(_MSC_VER) && USING_CODE_ANALYSIS
+#pragma warning(push)
+#pragma warning(disable:6294) // Ill-defined for-loop:  initial condition does not satisfy test.  Loop body not executed.
+#endif
+PRAGMA_DEFAULT_VISIBILITY_START
+THIRD_PARTY_INCLUDES_START
+#include <Eigen/Core>
+#include <Eigen/Dense>
+#include <Eigen/Sparse>
+#include <Eigen/SparseLU>
+THIRD_PARTY_INCLUDES_END
+PRAGMA_DEFAULT_VISIBILITY_END
+#if defined(_MSC_VER) && USING_CODE_ANALYSIS
+#pragma warning(pop)
+#endif
 
 static int32 GHairDeformationType = 0;
 static FAutoConsoleVariableRef CVarHairDeformationType(TEXT("r.HairStrands.DeformationType"), GHairDeformationType, TEXT("Type of procedural deformation applied on hair strands (0:use simulation's output, 1:use rest strands, 2: use rest guides, 3:wave pattern, 4:follow root normal)"));
@@ -29,6 +52,11 @@ static FAutoConsoleVariableRef CVarStrandHairWidth(TEXT("r.HairStrands.StrandWid
 static int32 GStrandHairInterpolationDebug = 0;
 static FAutoConsoleVariableRef CVarStrandHairInterpolationDebug(TEXT("r.HairStrands.Interpolation.Debug"), GStrandHairInterpolationDebug, TEXT("Enable debug rendering for hair interpolation"));
 
+static int32 GStrandHairInterpolationGlobal = 0;
+static FAutoConsoleVariableRef CVarStrandHairInterpolationGlobal(TEXT("r.HairStrands.Interpolation.Global"), GStrandHairInterpolationGlobal, TEXT("Enable global radial basis function based hair interpolation"));
+
+static int32 GStrandHairInterpolationTarget = 0;
+static FAutoConsoleVariableRef CVarStrandHairInterpolationTarget(TEXT("r.HairStrands.Interpolation.Target"), GStrandHairInterpolationTarget, TEXT("Set the global interpolation as the kinematics target if simulation is enabled"));
 
 static FIntVector ComputeDispatchCount(uint32 ItemCount, uint32 GroupSize)
 {
@@ -96,7 +124,7 @@ class FDeformGuideCS : public FGlobalShader
 	SHADER_USE_PARAMETER_STRUCT(FDeformGuideCS, FGlobalShader);
 
 	class FGroupSize : SHADER_PERMUTATION_INT("PERMUTATION_GROUP_SIZE", 2);
-	class FDeformationType : SHADER_PERMUTATION_INT("PERMUTATION_DEFORMATION", 5);
+	class FDeformationType : SHADER_PERMUTATION_INT("PERMUTATION_DEFORMATION", 6);
 	using FPermutationDomain = TShaderPermutationDomain<FGroupSize, FDeformationType>;
 
 	BEGIN_SHADER_PARAMETER_STRUCT(FParameters, )
@@ -120,6 +148,10 @@ class FDeformGuideCS : public FGlobalShader
 		SHADER_PARAMETER_SRV(Buffer, SimRestPosePositionBuffer)
 		SHADER_PARAMETER_SRV(Buffer, SimRootIndexBuffer)
 		SHADER_PARAMETER_UAV(RWBuffer, OutSimDeformedPositionBuffer)
+
+		SHADER_PARAMETER(uint32, SampleCount)
+		SHADER_PARAMETER_SRV(Buffer, RestSamplePositionsBuffer)
+		SHADER_PARAMETER_SRV(Buffer, MeshSampleWeightsBuffer)
 	END_SHADER_PARAMETER_STRUCT()
 
 public:
@@ -178,18 +210,30 @@ static void AddDeformSimHairStrandsPass(
 		}
 
 		const bool bSupportDynamicMesh = InSimHairData.RootCount > 0 && LODIndex >= 0 && LODIndex < uint32(InSimHairData.LODDatas.Num()) && InSimHairData.LODDatas[LODIndex].bIsValid && bIsVertexToCurveBuffersValid;
+		
+		bool bSupportGlobalInterpolation = false;
 		if (bSupportDynamicMesh)
 		{
-			InternalDeformationType = 4;
-			Parameters->SimRestPosition0Buffer = InSimHairData.LODDatas[LODIndex].RestRootTrianglePosition0Buffer->SRV;
-			Parameters->SimRestPosition1Buffer = InSimHairData.LODDatas[LODIndex].RestRootTrianglePosition1Buffer->SRV;
-			Parameters->SimRestPosition2Buffer = InSimHairData.LODDatas[LODIndex].RestRootTrianglePosition2Buffer->SRV;
-
-			Parameters->SimDeformedPosition0Buffer = InSimHairData.LODDatas[LODIndex].DeformedRootTrianglePosition0Buffer->SRV;
-			Parameters->SimDeformedPosition1Buffer = InSimHairData.LODDatas[LODIndex].DeformedRootTrianglePosition1Buffer->SRV;
-			Parameters->SimDeformedPosition2Buffer = InSimHairData.LODDatas[LODIndex].DeformedRootTrianglePosition2Buffer->SRV;
-
-			Parameters->SimRootBarycentricBuffer = InSimHairData.LODDatas[LODIndex].RootTriangleBarycentricBuffer->SRV;
+			bSupportGlobalInterpolation = (GStrandHairInterpolationGlobal == 1) && (InSimHairData.LODDatas[LODIndex].SampleCount > 0);
+			if (!bSupportGlobalInterpolation) {
+				InternalDeformationType = 4;
+				Parameters->SimRestPosition0Buffer = InSimHairData.LODDatas[LODIndex].RestRootTrianglePosition0Buffer->SRV;
+				Parameters->SimRestPosition1Buffer = InSimHairData.LODDatas[LODIndex].RestRootTrianglePosition1Buffer->SRV;
+				Parameters->SimRestPosition2Buffer = InSimHairData.LODDatas[LODIndex].RestRootTrianglePosition2Buffer->SRV;
+
+				Parameters->SimDeformedPosition0Buffer = InSimHairData.LODDatas[LODIndex].DeformedRootTrianglePosition0Buffer->SRV;
+				Parameters->SimDeformedPosition1Buffer = InSimHairData.LODDatas[LODIndex].DeformedRootTrianglePosition1Buffer->SRV;
+				Parameters->SimDeformedPosition2Buffer = InSimHairData.LODDatas[LODIndex].DeformedRootTrianglePosition2Buffer->SRV;
+
+				Parameters->SimRootBarycentricBuffer = InSimHairData.LODDatas[LODIndex].RootTriangleBarycentricBuffer->SRV;
+			}
+			else
+			{
+				InternalDeformationType = 5;
+				Parameters->MeshSampleWeightsBuffer = InSimHairData.LODDatas[LODIndex].MeshSampleWeightsBuffer->SRV;
+				Parameters->RestSamplePositionsBuffer = InSimHairData.LODDatas[LODIndex].RestSamplePositionsBuffer->SRV;
+				Parameters->SampleCount = InSimHairData.LODDatas[LODIndex].SampleCount;
+			}
 		}
 	}
 
@@ -228,7 +272,7 @@ class FHairInterpolationCS : public FGlobalShader
 
 	class FGroupSize : SHADER_PERMUTATION_INT("PERMUTATION_GROUP_SIZE", 2);
 	class FDebug : SHADER_PERMUTATION_INT("PERMUTATION_DEBUG", 2);
-	class FDynamicGeometry : SHADER_PERMUTATION_INT("PERMUTATION_DYNAMIC_GEOMETRY", 2);
+	class FDynamicGeometry : SHADER_PERMUTATION_INT("PERMUTATION_DYNAMIC_GEOMETRY", 3);
 	class FSimulation : SHADER_PERMUTATION_INT("PERMUTATION_SIMULATION", 2);
 	class FScaleAndClip : SHADER_PERMUTATION_INT("PERMUTATION_SCALE_AND_CLIP", 2);
 	using FPermutationDomain = TShaderPermutationDomain<FGroupSize, FDebug, FDynamicGeometry, FSimulation, FScaleAndClip>;
@@ -290,6 +334,8 @@ class FHairInterpolationCS : public FGlobalShader
 		SHADER_PARAMETER_SRV(Buffer<uint>, SimRootBarycentricBuffer)
 		SHADER_PARAMETER_SRV(Buffer<uint>, SimVertexToRootIndexBuffer)
 
+		SHADER_PARAMETER_SRV(Buffer<uint>, SimRootPointIndexBuffer)
+
 		END_SHADER_PARAMETER_STRUCT()
 
 public:
@@ -320,7 +366,8 @@ static void AddHairStrandsInterpolationPass(
 	const FShaderResourceViewRHIRef& SimAttributeBuffer,
 	const FUnorderedAccessViewRHIRef& OutRenderPositionBuffer,
 	const FUnorderedAccessViewRHIRef& OutRenderAttributeBuffer,
-	const FShaderResourceViewRHIRef& VertexToClusterIdBuffer)
+	const FShaderResourceViewRHIRef& VertexToClusterIdBuffer,
+	const FShaderResourceViewRHIRef& SimRootPointIndexBuffer )
 {
 	const uint32 GroupSize = ComputeGroupSize();
 	const FIntVector DispatchCount = ComputeDispatchCount(VertexCount, GroupSize);
@@ -338,6 +385,7 @@ static void AddHairStrandsInterpolationPass(
 	Parameters->InSimHairPositionOffset = InSimHairWorldOffset;
 	Parameters->OutHairPositionOffset = OutHairWorldOffset;
 	Parameters->DispatchCountX = DispatchCount.X;
+	Parameters->SimRootPointIndexBuffer = SimRootPointIndexBuffer;
 	
 	const bool bNeedScaleOrClip = 
 		ScaleAndClipDesc.InHairRadius != ScaleAndClipDesc.OutHairRadius || 
@@ -405,27 +453,31 @@ static void AddHairStrandsInterpolationPass(
 	}
 
 	const bool bSupportDynamicMesh = InRenHairData.RootCount > 0 && LODIndex >= 0 && LODIndex < InRenHairData.LODDatas.Num() && InRenHairData.LODDatas[LODIndex].bIsValid && bIsVertexToCurveBuffersValid;
+	bool bSupportGlobalInterpolation = false;
 	if (bSupportDynamicMesh)
 	{
-		Parameters->RestPosition0Buffer = InRenHairData.LODDatas[LODIndex].RestRootTrianglePosition0Buffer->SRV;
-		Parameters->RestPosition1Buffer = InRenHairData.LODDatas[LODIndex].RestRootTrianglePosition1Buffer->SRV;
-		Parameters->RestPosition2Buffer = InRenHairData.LODDatas[LODIndex].RestRootTrianglePosition2Buffer->SRV;
+		bSupportGlobalInterpolation = (GStrandHairInterpolationGlobal == 1) && (InSimHairData.LODDatas[LODIndex].SampleCount > 0);
+		{
+			Parameters->RestPosition0Buffer = InRenHairData.LODDatas[LODIndex].RestRootTrianglePosition0Buffer->SRV;
+			Parameters->RestPosition1Buffer = InRenHairData.LODDatas[LODIndex].RestRootTrianglePosition1Buffer->SRV;
+			Parameters->RestPosition2Buffer = InRenHairData.LODDatas[LODIndex].RestRootTrianglePosition2Buffer->SRV;
 
-		Parameters->DeformedPosition0Buffer = InRenHairData.LODDatas[LODIndex].DeformedRootTrianglePosition0Buffer->SRV;
-		Parameters->DeformedPosition1Buffer = InRenHairData.LODDatas[LODIndex].DeformedRootTrianglePosition1Buffer->SRV;
-		Parameters->DeformedPosition2Buffer = InRenHairData.LODDatas[LODIndex].DeformedRootTrianglePosition2Buffer->SRV;
+			Parameters->DeformedPosition0Buffer = InRenHairData.LODDatas[LODIndex].DeformedRootTrianglePosition0Buffer->SRV;
+			Parameters->DeformedPosition1Buffer = InRenHairData.LODDatas[LODIndex].DeformedRootTrianglePosition1Buffer->SRV;
+			Parameters->DeformedPosition2Buffer = InRenHairData.LODDatas[LODIndex].DeformedRootTrianglePosition2Buffer->SRV;
 
-		Parameters->RootBarycentricBuffer = InRenHairData.LODDatas[LODIndex].RootTriangleBarycentricBuffer->SRV;
+			Parameters->RootBarycentricBuffer = InRenHairData.LODDatas[LODIndex].RootTriangleBarycentricBuffer->SRV;
 
-		Parameters->SimRestPosition0Buffer = InSimHairData.LODDatas[LODIndex].RestRootTrianglePosition0Buffer->SRV;
-		Parameters->SimRestPosition1Buffer = InSimHairData.LODDatas[LODIndex].RestRootTrianglePosition1Buffer->SRV;
-		Parameters->SimRestPosition2Buffer = InSimHairData.LODDatas[LODIndex].RestRootTrianglePosition2Buffer->SRV;
+			Parameters->SimRestPosition0Buffer = InSimHairData.LODDatas[LODIndex].RestRootTrianglePosition0Buffer->SRV;
+			Parameters->SimRestPosition1Buffer = InSimHairData.LODDatas[LODIndex].RestRootTrianglePosition1Buffer->SRV;
+			Parameters->SimRestPosition2Buffer = InSimHairData.LODDatas[LODIndex].RestRootTrianglePosition2Buffer->SRV;
 
-		Parameters->SimDeformedPosition0Buffer = InSimHairData.LODDatas[LODIndex].DeformedRootTrianglePosition0Buffer->SRV;
-		Parameters->SimDeformedPosition1Buffer = InSimHairData.LODDatas[LODIndex].DeformedRootTrianglePosition1Buffer->SRV;
-		Parameters->SimDeformedPosition2Buffer = InSimHairData.LODDatas[LODIndex].DeformedRootTrianglePosition2Buffer->SRV;
+			Parameters->SimDeformedPosition0Buffer = InSimHairData.LODDatas[LODIndex].DeformedRootTrianglePosition0Buffer->SRV;
+			Parameters->SimDeformedPosition1Buffer = InSimHairData.LODDatas[LODIndex].DeformedRootTrianglePosition1Buffer->SRV;
+			Parameters->SimDeformedPosition2Buffer = InSimHairData.LODDatas[LODIndex].DeformedRootTrianglePosition2Buffer->SRV;
 
-		Parameters->SimRootBarycentricBuffer = InSimHairData.LODDatas[LODIndex].RootTriangleBarycentricBuffer->SRV;
+			Parameters->SimRootBarycentricBuffer = InSimHairData.LODDatas[LODIndex].RootTriangleBarycentricBuffer->SRV;
+		}
 	}
 
 	if (ShaderDrawDebug::IsShaderDrawDebugEnabled() && ShaderDrawData)
@@ -436,8 +488,9 @@ static void AddHairStrandsInterpolationPass(
 	FHairInterpolationCS::FPermutationDomain PermutationVector;
 	PermutationVector.Set<FHairInterpolationCS::FGroupSize>(GetGroupSizePermutation(GroupSize));
 	PermutationVector.Set<FHairInterpolationCS::FDebug>(Parameters->HairDebugMode > 0 ? 1 : 0);
-	PermutationVector.Set<FHairInterpolationCS::FDynamicGeometry>(bSupportDynamicMesh ? 1 : 0);
-	PermutationVector.Set<FHairInterpolationCS::FSimulation>(bHasSimulationEnable ? 1 : 0);
+	PermutationVector.Set<FHairInterpolationCS::FDynamicGeometry>((bSupportGlobalInterpolation && (GStrandHairInterpolationGlobal == 1)) ? 2 : 
+								bSupportDynamicMesh ? 1 : 0);
+	PermutationVector.Set<FHairInterpolationCS::FSimulation>((bHasSimulationEnable || bSupportGlobalInterpolation) ? 1 : 0);
 	PermutationVector.Set<FHairInterpolationCS::FScaleAndClip>(bNeedScaleOrClip ? 1 : 0);
 
 	FGlobalShaderMap* ShaderMap = GetGlobalShaderMap(ERHIFeatureLevel::SM5);
@@ -714,6 +767,7 @@ static void BuildHairAccelerationStructure(FRHICommandList& RHICmdList, uint32 R
 void ComputeHairStrandsInterpolation(
 	FRHICommandListImmediate& RHICmdList,
 	const FShaderDrawDebugData* ShaderDrawData,
+	const FTransform& LocalToWorld,
 	FHairStrandsInterpolationInput* InInput,
 	FHairStrandsInterpolationOutput* InOutput,
 	FHairStrandsProjectionHairData& InRenHairDatas,
@@ -739,8 +793,8 @@ void ComputeHairStrandsInterpolation(
 
 		FRDGBuilder GraphBuilder(RHICmdList);
 
-		const uint32 CurrIndex = Output.CurrentIndex;
-		const uint32 PrevIndex = (Output.CurrentIndex + 1) % 2;
+		const uint32 CurrIndex = *Output.CurrentIndex;
+		const uint32 PrevIndex = (CurrIndex + 1) % 2;
 
 		const EDeformationType DeformationType = GetDeformationType();
 		if (DeformationType != EDeformationType::RestStrands && DeformationType != EDeformationType::Simulation)
@@ -758,7 +812,7 @@ void ComputeHairStrandsInterpolation(
 		}
 
 		// If the deformation is driven by the physics simulation, then the output is always the 0 index
-		const uint32 SimIndex = GHairDeformationType == 0 ? 0 : CurrIndex;
+		const uint32 SimIndex = CurrIndex;// GHairDeformationType == 0 ? 0 : CurrIndex;
 
 		// Debug mode:
 		// * None	: Display hair normally
@@ -782,12 +836,14 @@ void ComputeHairStrandsInterpolation(
 			Output.VFInput.HairTangentBuffer = Output.SimTangentBuffer->SRV;
 			Output.VFInput.HairAttributeBuffer = Input.SimAttributeBuffer->SRV;
 			Output.VFInput.HairMaterialBuffer = Output.RenderMaterialBuffer->SRV;
-			Output.VFInput.HairPositionOffset = bHasSimulationEnabled ? Input.OutHairPositionOffset : Input.InSimHairPositionOffset;
-			Output.VFInput.HairPreviousPositionOffset = bHasSimulationEnabled ? Input.OutHairPreviousPositionOffset : Input.InSimHairPositionOffset;
+			Output.VFInput.HairPositionOffset = Input.OutHairPositionOffset;// bHasSimulationEnabled ? Input.OutHairPositionOffset : Input.InSimHairPositionOffset;
+			Output.VFInput.HairPreviousPositionOffset = Input.OutHairPreviousPositionOffset;// bHasSimulationEnabled ? Input.OutHairPreviousPositionOffset : Input.InSimHairPositionOffset;
 			Output.VFInput.VertexCount = Input.SimVertexCount;
 			Output.VFInput.HairRadius = (GStrandHairWidth > 0 ? GStrandHairWidth : Input.GroupDesc.HairWidth) * 0.5f;
 			Output.VFInput.HairLength = Input.GroupDesc.HairLength;
 			Output.VFInput.HairDensity = Input.GroupDesc.HairShadowDensity;
+			Output.VFInput.bUseStableRasterization = Input.GroupDesc.bUseStableRasterization;
+			Output.VFInput.bScatterSceneLighting = Input.GroupDesc.bScatterSceneLighting;
 		}
 		else
 		{
@@ -841,7 +897,8 @@ void ComputeHairStrandsInterpolation(
 				Input.SimAttributeBuffer->SRV,
 				Output.RenderDeformedPositionBuffer[CurrIndex]->UAV,
 				Output.RenderPatchedAttributeBuffer.UAV,
-				Input.VertexToClusterIdBuffer->SRV);
+				Input.VertexToClusterIdBuffer->SRV, 
+				Input.SimRootPointIndexBuffer->SRV);
 
 			// Initialize group cluster data for culling by the renderer
 			FHairStrandClusterData::FHairGroup& HairGroupCluster = ClusterData->HairGroups.Emplace_GetRef();
@@ -862,6 +919,8 @@ void ComputeHairStrandsInterpolation(
 			Output.VFInput.HairDensity = Input.GroupDesc.HairShadowDensity;
 			Output.VFInput.HairPositionBuffer = Output.RenderDeformedPositionBuffer[CurrIndex]->SRV;
 			Output.VFInput.HairPreviousPositionBuffer = Output.RenderDeformedPositionBuffer[PrevIndex]->SRV;
+			Output.VFInput.bUseStableRasterization = Input.GroupDesc.bUseStableRasterization;
+			Output.VFInput.bScatterSceneLighting = Input.GroupDesc.bScatterSceneLighting;
 
 			AddHairClusterAABBPass(
 				GraphBuilder,
@@ -898,6 +957,16 @@ void ComputeHairStrandsInterpolation(
 			Output.VFInput.HairPreviousPositionOffset	= Input.OutHairPreviousPositionOffset;
 			Output.VFInput.VertexCount					= Input.RenderVertexCount;
 
+			// TODO: find a more robust way to handle parameters passing to compute raster.
+			// At the moment there is a loose compling which will break if the vertex factor change.
+			Output.HairGroupPublicData->VFInput.HairPositionBuffer = Output.VFInput.HairPositionBuffer;
+			Output.HairGroupPublicData->VFInput.HairPositionOffset = Output.VFInput.HairPositionOffset;
+			Output.HairGroupPublicData->VFInput.VertexCount = Output.VFInput.VertexCount;
+			Output.HairGroupPublicData->VFInput.HairRadius = Output.VFInput.HairRadius;
+			Output.HairGroupPublicData->VFInput.HairLength = Output.VFInput.HairLength;
+			Output.HairGroupPublicData->VFInput.HairDensity = Output.VFInput.HairDensity;
+			Output.HairGroupPublicData->VFInput.LocalToWorldTransform = LocalToWorld;
+
 			#if RHI_RAYTRACING
 			if (IsHairRayTracingEnabled())
 			{
@@ -914,7 +983,7 @@ void ComputeHairStrandsInterpolation(
 			#endif
 		}
 
-		Output.CurrentIndex = PrevIndex;
+		*Output.CurrentIndex = PrevIndex;
 	}
 }
 
@@ -940,8 +1009,8 @@ void ResetHairStrandsInterpolation(
 			SCOPED_DRAW_EVENT(RHICmdList, HairStrandsResetInterpolation);
 			SCOPED_GPU_STAT(RHICmdList, HairStrandsResetInterpolation);
 
-			const uint32 CurrIndex = Output.CurrentIndex;
-			const uint32 PrevIndex = (Output.CurrentIndex + 1) % 2;
+			const uint32 CurrIndex = *Output.CurrentIndex;
+			const uint32 PrevIndex = (CurrIndex + 1) % 2;
 			FRDGBuilder GraphBuilder(RHICmdList);
 			AddDeformSimHairStrandsPass(
 				GraphBuilder,
@@ -973,8 +1042,18 @@ void ReadbackBuffer(TArray<ReadBackType>& OutData, FRWBuffer& InBuffer)
 	RHIUnlockVertexBuffer(InBuffer.Buffer);
 }
 
+template<typename WriteBackType>
+void WritebackBuffer(TArray<WriteBackType>& InData, FRWBuffer& OutBuffer)
+{
+	const uint32 DataSize = sizeof(WriteBackType) * InData.Num();
+	check(DataSize == OutBuffer.Buffer->GetSize() );
+
+	WriteBackType* Data = (WriteBackType*)RHILockVertexBuffer(OutBuffer.Buffer, 0, DataSize, RLM_WriteOnly);
+	FMemory::Memcpy(Data, InData.GetData(), DataSize);
+	RHIUnlockVertexBuffer(OutBuffer.Buffer);
+}
+
 static void ReadbackGroupData(
-	FRHICommandListImmediate& RHICmdList, 
 	FHairStrandsRootData& OutCPUData,
 	FHairStrandsRootResource* InGPUData)
 {
@@ -997,6 +1076,197 @@ static void ReadbackGroupData(
 		ReadbackBuffer(CPULOD.RestRootTrianglePosition0Buffer, GPULOD.RestRootTrianglePosition0Buffer);
 		ReadbackBuffer(CPULOD.RestRootTrianglePosition1Buffer, GPULOD.RestRootTrianglePosition1Buffer);
 		ReadbackBuffer(CPULOD.RestRootTrianglePosition2Buffer, GPULOD.RestRootTrianglePosition2Buffer);
+
+		InGPUData->RootData.MeshProjectionLODs[LODIt].RootTriangleIndexBuffer = CPULOD.RootTriangleIndexBuffer;
+		InGPUData->RootData.MeshProjectionLODs[LODIt].RootTriangleBarycentricBuffer = CPULOD.RootTriangleBarycentricBuffer;
+		InGPUData->RootData.MeshProjectionLODs[LODIt].RestRootTrianglePosition0Buffer = CPULOD.RestRootTrianglePosition0Buffer;
+		InGPUData->RootData.MeshProjectionLODs[LODIt].RestRootTrianglePosition1Buffer = CPULOD.RestRootTrianglePosition1Buffer;
+		InGPUData->RootData.MeshProjectionLODs[LODIt].RestRootTrianglePosition2Buffer = CPULOD.RestRootTrianglePosition2Buffer;
+	}
+}
+
+void WritebackGroupData(FHairStrandsRootData& InCPUData,
+	FHairStrandsRootResource* OutGPUData)
+{
+	if (!OutGPUData)
+	{
+		return;
+	}
+	check(OutGPUData->MeshProjectionLODs.Num() == InCPUData.MeshProjectionLODs.Num());
+	const uint32 LODCount = OutGPUData->MeshProjectionLODs.Num();
+	for (uint32 LODIt = 0; LODIt < LODCount; ++LODIt)
+	{
+		FHairStrandsRootData::FMeshProjectionLOD& CPULOD = InCPUData.MeshProjectionLODs[LODIt];
+		FHairStrandsRootResource::FMeshProjectionLOD& GPULOD = OutGPUData->MeshProjectionLODs[LODIt];
+		check(CPULOD.LODIndex == GPULOD.LODIndex);
+
+		if (CPULOD.SampleCount > 0)
+		{
+			WritebackBuffer(CPULOD.MeshInterpolationWeightsBuffer, GPULOD.MeshInterpolationWeightsBuffer);
+			WritebackBuffer(CPULOD.MeshSampleIndicesBuffer, GPULOD.MeshSampleIndicesBuffer);
+			WritebackBuffer(CPULOD.RestSamplePositionsBuffer, GPULOD.RestSamplePositionsBuffer);
+
+			OutGPUData->RootData.MeshProjectionLODs[LODIt].SampleCount = CPULOD.SampleCount;
+			OutGPUData->RootData.MeshProjectionLODs[LODIt].MeshInterpolationWeightsBuffer = CPULOD.MeshInterpolationWeightsBuffer;
+			OutGPUData->RootData.MeshProjectionLODs[LODIt].RestSamplePositionsBuffer = CPULOD.RestSamplePositionsBuffer;
+			OutGPUData->RootData.MeshProjectionLODs[LODIt].MeshSampleIndicesBuffer = CPULOD.MeshSampleIndicesBuffer;
+		}
+	}
+}
+
+struct FPointsSampler
+{
+	FPointsSampler(TArray<bool>& ValidPoints, const FVector* PointPositions, const int32 NumSamples);
+
+	/** Build the sample position from the sample indices */
+	void BuildPositions(const FVector* PointPositions);
+
+	/** Compute the furthest point */
+	void FurthestPoint(const int32 NumPoints, const FVector* PointPositions, const uint32 SampleIndex, TArray<bool>& ValidPoints, TArray<float>& PointsDistance);
+
+	/** Compute the starting point */
+	int32 StartingPoint(const TArray<bool>& ValidPoints, int32& NumPoints) const;
+
+	/** List of sampled points */
+	TArray<uint32> SampleIndices;
+
+	/** List of sampled positions */
+	TArray<FVector> SamplePositions;
+};
+
+int32 FPointsSampler::StartingPoint(const TArray<bool>& ValidPoints, int32& NumPoints) const
+{
+	int32 StartIndex = -1;
+	NumPoints = 0;
+	for (int32 i = 0; i < ValidPoints.Num(); ++i)
+	{
+		if (ValidPoints[i])
+		{
+			++NumPoints;
+			if (StartIndex == -1)
+			{
+				StartIndex = i;
+			}
+		}
+	}
+	return StartIndex;
+}
+
+void FPointsSampler::BuildPositions(const FVector* PointPositions)
+{
+	SamplePositions.SetNum(SampleIndices.Num());
+	for (int32 i = 0; i < SampleIndices.Num(); ++i)
+	{
+		SamplePositions[i] = PointPositions[SampleIndices[i]];
+	}
+}
+
+void FPointsSampler::FurthestPoint(const int32 NumPoints, const FVector* PointPositions, const uint32 SampleIndex, TArray<bool>& ValidPoints, TArray<float>& PointsDistance)
+{
+	float FurthestDistance = 0.0;
+	uint32 PointIndex = 0;
+	for (int32 j = 0; j < NumPoints; ++j)
+	{
+		if (ValidPoints[j])
+		{
+			PointsDistance[j] = FMath::Min((PointPositions[SampleIndices[SampleIndex - 1]] - PointPositions[j]).Size(), PointsDistance[j]);
+			if (PointsDistance[j] >= FurthestDistance)
+			{
+				PointIndex = j;
+				FurthestDistance = PointsDistance[j];
+			}
+		}
+	}
+	ValidPoints[PointIndex] = false;
+	SampleIndices[SampleIndex] = PointIndex;
+}
+
+FPointsSampler::FPointsSampler(TArray<bool>& ValidPoints, const FVector* PointPositions, const int32 NumSamples)
+{
+	int32 NumPoints = 0;
+	int32 StartIndex = StartingPoint(ValidPoints, NumPoints);
+
+	const int32 SamplesCount = FMath::Min(NumPoints, NumSamples);
+	if (SamplesCount != 0)
+	{
+		SampleIndices.SetNum(SamplesCount);
+		SampleIndices[0] = StartIndex;
+		ValidPoints[StartIndex] = false;
+
+		TArray<float> PointsDistance;
+		PointsDistance.Init(MAX_FLT, ValidPoints.Num());
+
+		for (int32 i = 1; i < SamplesCount; ++i)
+		{
+			FurthestPoint(ValidPoints.Num(), PointPositions, i, ValidPoints, PointsDistance);
+		}
+		BuildPositions(PointPositions);
+	}
+}
+
+struct FWeightsBuilder
+{
+	FWeightsBuilder(const uint32 NumRows, const uint32 NumColumns,
+		const FVector* SourcePositions, const FVector* TargetPositions);
+
+	using EigenMatrix = Eigen::Map<Eigen::Matrix<float, Eigen::Dynamic, Eigen::Dynamic, Eigen::RowMajor>>;
+
+	/** Compute the weights by inverting the matrix*/
+	void ComputeWeights(const uint32 NumRows, const uint32 NumColumns);
+
+	/** Entries in the dense structure */
+	TArray<float> MatrixEntries;
+
+	/** Entries of the matrix inverse */
+	TArray<float> InverseEntries;
+};
+
+FWeightsBuilder::FWeightsBuilder(const uint32 NumRows, const uint32 NumColumns,
+	const FVector* SourcePositions, const FVector* TargetPositions)
+{
+	MatrixEntries.SetNum(NumRows * NumColumns);  
+	InverseEntries.SetNum(NumRows * NumColumns);
+	TArray<float>& LocalEntries = MatrixEntries;
+	ParallelFor(NumRows,
+		[
+			NumRows,
+			NumColumns,
+			SourcePositions,
+			TargetPositions,
+			&LocalEntries
+		] (uint32 RowIndex)
+	{
+		int32 EntryIndex = RowIndex * NumColumns;
+		for (uint32 j = 0; j < NumColumns; ++j)
+		{
+			const float FunctionScale = (SourcePositions[RowIndex] - TargetPositions[j]).Size();
+			LocalEntries[EntryIndex++] = FMath::Sqrt(FunctionScale*FunctionScale + 1.0);
+		}
+	});
+	ComputeWeights(NumRows, NumColumns);
+}
+
+void FWeightsBuilder::ComputeWeights(const uint32 NumRows, const uint32 NumColumns)
+{
+	EigenMatrix WeightsMatrix(MatrixEntries.GetData(), NumRows, NumColumns);
+	EigenMatrix WeightsInverse(InverseEntries.GetData(), NumColumns, NumRows);
+
+	WeightsInverse = WeightsMatrix.inverse();
+}
+
+void UpdateInterpolationWeights(const FWeightsBuilder& InterpolationWeights, const FPointsSampler& PointsSampler, const uint32 LODIndex, FHairStrandsRootData& RootDatas)
+{
+	FHairStrandsRootData::FMeshProjectionLOD& CPULOD = RootDatas.MeshProjectionLODs[LODIndex];
+	CPULOD.MeshSampleIndicesBuffer.SetNum(PointsSampler.SampleIndices.Num());
+	CPULOD.MeshInterpolationWeightsBuffer.SetNum(InterpolationWeights.InverseEntries.Num());
+	CPULOD.RestSamplePositionsBuffer.SetNum(PointsSampler.SampleIndices.Num());
+
+	CPULOD.SampleCount = PointsSampler.SampleIndices.Num();
+	CPULOD.MeshSampleIndicesBuffer = PointsSampler.SampleIndices;
+	CPULOD.MeshInterpolationWeightsBuffer = InterpolationWeights.InverseEntries;
+	for (int32 i = 0; i < PointsSampler.SamplePositions.Num(); ++i)
+	{
+		CPULOD.RestSamplePositionsBuffer[i] = FVector4(PointsSampler.SamplePositions[i], 1.0f );
 	}
 }
 
@@ -1008,23 +1278,37 @@ static void InternalProcessGroomBindingTask(FRHICommandListImmediate& RHICmdList
 		!BindingAsset->TargetSkeletalMesh ||
 		 BindingAsset->Groom->GetNumHairGroups() == 0)
 	{
+		UE_LOG(LogHairStrands, Warning, TEXT("[Groom] Error - Binding asset can be created/rebuilt."));
 		return;
 	}
 
+	const int32 NumInterpolationPoints	= BindingAsset->NumInterpolationPoints;
 	UGroomAsset* GroomAsset				= BindingAsset->Groom;
 	USkeletalMesh* SourceSkeletalMesh	= BindingAsset->SourceSkeletalMesh;
 	USkeletalMesh* TargetSkeletalMesh	= BindingAsset->TargetSkeletalMesh;
 
 	const uint32 LODCount = BindingAsset->TargetSkeletalMesh->GetLODNum();
 	UGroomBindingAsset::FHairGroupDatas& OutHairGroupDatas = BindingAsset->HairGroupDatas;
+	OutHairGroupDatas.Empty();
+	TArray<uint32> NumSamples;
+	NumSamples.Init(NumInterpolationPoints, LODCount);
 	for (const FHairGroupData& GroupData : GroomAsset->HairGroupsData)
 	{
 		UGroomBindingAsset::FHairGroupData& Data = OutHairGroupDatas.AddDefaulted_GetRef();
-		Data.RenRootData = FHairStrandsRootData(&GroupData.HairRenderData, LODCount);
-		Data.SimRootData = FHairStrandsRootData(&GroupData.HairSimulationData, LODCount);
+		Data.RenRootData = FHairStrandsRootData(&GroupData.HairRenderData, LODCount, NumSamples);
+		Data.SimRootData = FHairStrandsRootData(&GroupData.HairSimulationData, LODCount, NumSamples);
 	}
 
 	UGroomBindingAsset::FHairGroupResources& OutHairGroupResources = BindingAsset->HairGroupResources;
+	if (BindingAsset->HairGroupResources.Num() > 0)
+	{
+		for (UGroomBindingAsset::FHairGroupResource& GroupResrouces : OutHairGroupResources)
+		{
+			BindingAsset->HairGroupResourcesToDelete.Enqueue(GroupResrouces);
+		}
+		OutHairGroupResources.Empty();
+	}
+
 	check(OutHairGroupResources.Num() == 0);
 	for (UGroomBindingAsset::FHairGroupData& GroupData : OutHairGroupDatas)
 	{
@@ -1077,20 +1361,162 @@ static void InternalProcessGroomBindingTask(FRHICommandListImmediate& RHICmdList
 		RenProjectionDatas,
 		SimProjectionDatas);
 
+	ComputeInterpolationWeights(BindingAsset, TargetRenderData, TransferredPositions);
+	BindingAsset->QueryStatus = UGroomBindingAsset::EQueryStatus::Completed;
+}
+
+void FillLocalValidPoints(FSkeletalMeshLODRenderData& LODRenderData, const uint32 TargetSection, 
+	const FHairStrandsRootData::FMeshProjectionLOD& ProjectionLOD, TArray<bool>& ValidPoints )
+{
+	TArray<uint32> TriangleIndices; TriangleIndices.SetNum(LODRenderData.MultiSizeIndexContainer.GetIndexBuffer()->Num());
+	LODRenderData.MultiSizeIndexContainer.GetIndexBuffer(TriangleIndices);
+
+	ValidPoints.Init(false, LODRenderData.StaticVertexBuffers.PositionVertexBuffer.GetNumVertices());
+
+	const TArray<uint32>& RootBuffers = ProjectionLOD.RootTriangleIndexBuffer;
+	for (int32 RootIt = 0; RootIt < RootBuffers.Num(); ++RootIt)
+	{
+		const uint32 SectionIndex = (RootBuffers[RootIt] >> 28) & 0xF;
+		const uint32 TriangleIndex = RootBuffers[RootIt] & 0xFFFFFFF;
+		if (SectionIndex == TargetSection)
+		{
+			for (uint32 VertexIt = 0; VertexIt < 3; ++VertexIt)
+			{
+				const uint32 VertexIndex = TriangleIndices[LODRenderData.RenderSections[SectionIndex].BaseIndex + 3 * TriangleIndex + VertexIt];
+				ValidPoints[VertexIndex] = true;
+			}
+		}
+	}
+}
+
+void FillGlobalValidPoints(FSkeletalMeshLODRenderData& LODRenderData, const uint32 TargetSection, TArray<bool>& ValidPoints)
+{
+	TArray<uint32> TriangleIndices; TriangleIndices.SetNum(LODRenderData.MultiSizeIndexContainer.GetIndexBuffer()->Num());
+	LODRenderData.MultiSizeIndexContainer.GetIndexBuffer(TriangleIndices);
 
+	ValidPoints.Init(false, LODRenderData.StaticVertexBuffers.PositionVertexBuffer.GetNumVertices());
+	
+	for (uint32 TriangleIt = 0; TriangleIt < LODRenderData.RenderSections[TargetSection].NumTriangles; ++TriangleIt)
+	{
+		for (uint32 VertexIt = 0; VertexIt < 3; ++VertexIt)
+		{
+			const uint32 VertexIndex = TriangleIndices[LODRenderData.RenderSections[TargetSection].BaseIndex + 3 * TriangleIt + VertexIt];
+			ValidPoints[VertexIndex] = true;
+		}
+	}
+}
+
+void ComputeInterpolationWeights(UGroomBindingAsset* BindingAsset, FSkeletalMeshRenderData* TargetRenderData, TArray<FRWBuffer>& TransferedPositions)
+{
+	UGroomAsset* GroomAsset = BindingAsset->Groom;
 	// Enforce GPU sync to read back data on CPU
 	GDynamicRHI->RHISubmitCommandsAndFlushGPU();
 	GDynamicRHI->RHIBlockUntilGPUIdle();
 
+	UGroomBindingAsset::FHairGroupDatas& OutHairGroupDatas = BindingAsset->HairGroupDatas;
+	UGroomBindingAsset::FHairGroupResources& OutHairGroupResources = BindingAsset->HairGroupResources;
+
+	const uint32 GroupCount = OutHairGroupResources.Num();
+	const uint32 LODCount = BindingAsset->TargetSkeletalMesh->GetLODNum();
+	const uint32 MaxSamples = BindingAsset->NumInterpolationPoints;
+
+	for (uint32 GroupIt = 0; GroupIt < GroupCount; ++GroupIt)
+	{
+		ReadbackGroupData(OutHairGroupDatas[GroupIt].SimRootData, OutHairGroupResources[GroupIt].SimRootResources);
+		ReadbackGroupData(OutHairGroupDatas[GroupIt].RenRootData, OutHairGroupResources[GroupIt].RenRootResources);
+	}
+
+	const uint32 TargetSection = 0;
+	const bool LocalSamples = false;
+	for (uint32 LODIndex = 0; LODIndex < LODCount; ++LODIndex)
+	{
+		FSkeletalMeshLODRenderData& LODRenderData = TargetRenderData->LODRenderData[LODIndex];
+
+		TArray<FSkelMeshRenderSection>& RenderSections = LODRenderData.RenderSections;
+		const int32 NumVertices = (RenderSections.Num() > TargetSection) ? RenderSections[TargetSection].NumVertices : 0;
+
+		TArray<FVector> SourcePositions;
+		FVector* PositionsPointer = nullptr;
+		if (TransferedPositions.Num() == LODCount)
+		{
+			ReadbackBuffer(SourcePositions, TransferedPositions[LODIndex]);
+			PositionsPointer = SourcePositions.GetData();
+		}
+		else
+		{
+			FPositionVertexBuffer& VertexBuffer = LODRenderData.StaticVertexBuffers.PositionVertexBuffer;
+			PositionsPointer = static_cast<FVector*>(VertexBuffer.GetVertexData());
+		}
+
+		if (LocalSamples)
+		{
+			TArray<bool> ValidPoints;
+			for (uint32 GroupIt = 0; GroupIt < GroupCount; ++GroupIt)
+			{
+				FillLocalValidPoints(LODRenderData, TargetSection, OutHairGroupDatas[GroupIt].RenRootData.MeshProjectionLODs[LODIndex], ValidPoints);
+
+				FPointsSampler PointsSampler(ValidPoints, PositionsPointer, MaxSamples);
+				const uint32 SampleCount = PointsSampler.SamplePositions.Num();
+
+				FWeightsBuilder InterpolationWeights(SampleCount, SampleCount,
+					PointsSampler.SamplePositions.GetData(), PointsSampler.SamplePositions.GetData());
+
+				//const FVector Displace(0.0,0.0,10.0);
+				//TArray<FVector> Deltas; Deltas.Init(FVector::ZeroVector, SampleCount );
+				//for (uint32 i = 0; i < SampleCount; ++i)
+				//{
+				//	Deltas[i] = FVector(0,0,0);
+				//	for (uint32 j = 0; j < SampleCount; ++j)
+				//	{
+				//		Deltas[i] += InterpolationWeights.InverseEntries[SampleCount * i + j] * Displace;
+				//	}
+				//	UE_LOG(LogHairStrands, Log, TEXT("[Groom] Sample Deltas[%d] = %s"), i, *Deltas[i].ToString());
+				//}
+				//for (uint32 i = 0; i < GroomAsset->HairGroupsData[GroupIt].HairSimulationData.StrandsPoints.Num(); ++i)
+				//{
+				//	FVector Offset(0,0,0);
+				//	for (uint32 j = 0; j < SampleCount; ++j)
+				//	{
+				//		const FVector DeltaPosition = GroomAsset->HairGroupsData[GroupIt].HairSimulationData.StrandsPoints.PointsPosition[i] - PointsSampler.SamplePositions[j];
+				//		const float FunctionValue = FMath::Sqrt(FVector::DotProduct(DeltaPosition, DeltaPosition)+1.0);
+				//		Offset += FunctionValue * Deltas[j];
+				//	}
+				//	UE_LOG(LogHairStrands, Log, TEXT("[Groom] Sample Displace[%d] = %s"), i, *Offset.ToString());
+				//}
+
+				UpdateInterpolationWeights(InterpolationWeights, PointsSampler, LODIndex, OutHairGroupDatas[GroupIt].SimRootData);
+				UpdateInterpolationWeights(InterpolationWeights, PointsSampler, LODIndex, OutHairGroupDatas[GroupIt].RenRootData);
+			}
+		}
+		else
+		{
+			TArray<bool> ValidPoints;
+
+			FillGlobalValidPoints(LODRenderData, TargetSection, ValidPoints);
+
+			FPointsSampler PointsSampler(ValidPoints, PositionsPointer, MaxSamples);
+			const uint32 SampleCount = PointsSampler.SamplePositions.Num();
+
+			FWeightsBuilder InterpolationWeights(SampleCount, SampleCount,
+				PointsSampler.SamplePositions.GetData(), PointsSampler.SamplePositions.GetData());
+
+			for (uint32 GroupIt = 0; GroupIt < GroupCount; ++GroupIt)
+			{
+				UpdateInterpolationWeights(InterpolationWeights, PointsSampler, LODIndex, OutHairGroupDatas[GroupIt].SimRootData);
+				UpdateInterpolationWeights(InterpolationWeights, PointsSampler, LODIndex, OutHairGroupDatas[GroupIt].RenRootData);
+			}
+		}
+	}
 	for (uint32 GroupIt = 0; GroupIt < GroupCount; ++GroupIt)
 	{
-		ReadbackGroupData(RHICmdList, OutHairGroupDatas[GroupIt].SimRootData, OutHairGroupResources[GroupIt].SimRootResources);
-		ReadbackGroupData(RHICmdList, OutHairGroupDatas[GroupIt].RenRootData, OutHairGroupResources[GroupIt].RenRootResources);
+		WritebackGroupData(OutHairGroupDatas[GroupIt].SimRootData, OutHairGroupResources[GroupIt].SimRootResources);
+		WritebackGroupData(OutHairGroupDatas[GroupIt].RenRootData, OutHairGroupResources[GroupIt].RenRootResources);
 	}
 }
 
 void AddGroomBindingTask(UGroomBindingAsset* BindingAsset)
 {
+	BindingAsset->QueryStatus = UGroomBindingAsset::EQueryStatus::Submitted;
 	EnqueueGroomBindingQuery((void*)BindingAsset, InternalProcessGroomBindingTask);
 }
 
@@ -1123,7 +1549,14 @@ FHairStrandsProjectionHairData::HairGroup ToProjectionHairData(FHairStrandsRootR
 		LODData.DeformedRootTrianglePosition2Buffer = &MeshLODData.DeformedRootTrianglePosition2Buffer;
 
 		LODData.Status = &MeshLODData.Status;
-		LODData.bIsValid = MeshLODData.Status == FHairStrandsProjectionHairData::LODData::EStatus::Completed;
+		LODData.bIsValid = MeshLODData.Status == FHairStrandsProjectionHairData::LODData::EStatus::Completed || MeshLODData.Status == FHairStrandsProjectionHairData::LODData::EStatus::Initialized;
+
+		LODData.SampleCount = MeshLODData.SampleCount;
+		LODData.MeshInterpolationWeightsBuffer = &MeshLODData.MeshInterpolationWeightsBuffer;
+		LODData.MeshSampleIndicesBuffer = &MeshLODData.MeshSampleIndicesBuffer;
+		LODData.RestSamplePositionsBuffer = &MeshLODData.RestSamplePositionsBuffer;
+		LODData.DeformedSamplePositionsBuffer = &MeshLODData.DeformedSamplePositionsBuffer;
+		LODData.MeshSampleWeightsBuffer = &MeshLODData.MeshSampleWeightsBuffer;
 	}
 
 	return Out;
diff --git a/Engine/Plugins/Experimental/HairStrands/Source/HairStrandsCore/Private/HairStrandsVertexFactory.cpp b/Engine/Plugins/Experimental/HairStrands/Source/HairStrandsCore/Private/HairStrandsVertexFactory.cpp
--- a/Engine/Plugins/Experimental/HairStrands/Source/HairStrandsCore/Private/HairStrandsVertexFactory.cpp
+++ b/Engine/Plugins/Experimental/HairStrands/Source/HairStrandsCore/Private/HairStrandsVertexFactory.cpp
@@ -53,6 +53,8 @@ public:
 	LAYOUT_FIELD(FShaderParameter, PreviousPositionOffset);
 	LAYOUT_FIELD(FShaderParameter, Density);
 	LAYOUT_FIELD(FShaderParameter, Culling);
+	LAYOUT_FIELD(FShaderParameter, StableRasterization);
+	LAYOUT_FIELD(FShaderParameter, ScatterSceneLighing);
 
 	LAYOUT_FIELD(FShaderResourceParameter, PositionBuffer);
 	LAYOUT_FIELD(FShaderResourceParameter, PreviousPositionBuffer);
@@ -71,6 +73,8 @@ public:
 		PreviousPositionOffset.Bind(ParameterMap, TEXT("HairStrandsVF_PreviousPositionOffset"));
 		Density.Bind(ParameterMap, TEXT("HairStrandsVF_Density"));	
 		Culling.Bind(ParameterMap, TEXT("HairStrandsVF_CullingEnable"));
+		StableRasterization.Bind(ParameterMap, TEXT("HairStrandsVF_bUseStableRasterization"));
+		ScatterSceneLighing.Bind(ParameterMap, TEXT("HairStrandsVF_bScatterSceneLighing"));
 
 		PositionBuffer.Bind(ParameterMap, TEXT("HairStrandsVF_PositionBuffer"));
 		PreviousPositionBuffer.Bind(ParameterMap, TEXT("HairStrandsVF_PreviousPositionBuffer"));
@@ -109,6 +113,8 @@ public:
 		BindParam(ShaderBindings, PositionOffset, VF->GetPositionOffset(GroupIndex));
 		BindParam(ShaderBindings, PreviousPositionOffset, VF->GetPreviousPositionOffset(GroupIndex));
 		BindParam(ShaderBindings, Density, VF->GetHairDensity(GroupIndex));
+		BindParam(ShaderBindings, StableRasterization, VF->UseStableRasterization(GroupIndex) ? 1u : 0u);
+		BindParam(ShaderBindings, ScatterSceneLighing, VF->UseScatterSceneLighting(GroupIndex) ? 1u : 0u);
 		
 		FShaderResourceViewRHIRef CulledDispatchVertexIdsSRV = GDummyCulledDispatchVertexIdsBuffer.SRVUint;
 		FShaderResourceViewRHIRef CulledCompactedRadiusScaleBufferSRV = GDummyCulledDispatchVertexIdsBuffer.SRVFloat;
@@ -207,15 +213,39 @@ void FHairStrandsVertexFactory::InitRHI()
 	check(IsValidRef(GetDeclaration()));
 }
 
-void FHairStrandsVertexFactory::ReleaseRHI()
+IMPLEMENT_VERTEX_FACTORY_PARAMETER_TYPE(FHairStrandsVertexFactory, SF_Vertex,		FHairStrandsVertexFactoryShaderParameters);
+IMPLEMENT_VERTEX_FACTORY_PARAMETER_TYPE(FHairStrandsVertexFactory, SF_Pixel,		FHairStrandsVertexFactoryShaderParameters);
+IMPLEMENT_VERTEX_FACTORY_PARAMETER_TYPE(FHairStrandsVertexFactory, SF_Compute,		FHairStrandsVertexFactoryShaderParameters);
+#if RHI_RAYTRACING
+IMPLEMENT_VERTEX_FACTORY_PARAMETER_TYPE(FHairStrandsVertexFactory, SF_RayHitGroup,	FHairStrandsVertexFactoryShaderParameters);
+#endif
+
+FVertexFactoryShaderParameters* FHairStrandsVertexFactory::ConstructShaderParameters(EShaderFrequency ShaderFrequency)
 {
-	FVertexFactory::ReleaseRHI();
-}
+	if (ShaderFrequency == SF_Vertex)
+	{
+		return new FHairStrandsVertexFactoryShaderParameters();
+	}
+
+	if (ShaderFrequency == SF_Pixel)
+	{
+		return new FHairStrandsVertexFactoryShaderParameters();
+	}
 
-IMPLEMENT_VERTEX_FACTORY_PARAMETER_TYPE(FHairStrandsVertexFactory, SF_Vertex, FHairStrandsVertexFactoryShaderParameters);
-IMPLEMENT_VERTEX_FACTORY_PARAMETER_TYPE(FHairStrandsVertexFactory, SF_Pixel, FHairStrandsVertexFactoryShaderParameters);
 #if RHI_RAYTRACING
-IMPLEMENT_VERTEX_FACTORY_PARAMETER_TYPE(FHairStrandsVertexFactory, SF_RayHitGroup, FHairStrandsVertexFactoryShaderParameters);
+	if (ShaderFrequency == SF_RayHitGroup)
+	{
+		return new FHairStrandsVertexFactoryShaderParameters();
+	}
 #endif // RHI_RAYTRACING
 
+	return NULL;
+}
+
+void FHairStrandsVertexFactory::ReleaseRHI()
+{
+	FVertexFactory::ReleaseRHI();
+}
+
 IMPLEMENT_VERTEX_FACTORY_TYPE_EX(FHairStrandsVertexFactory,"/Engine/Private/HairStrands/HairStrandsVertexFactory.ush",true,false,true,true,true,true,true);
+
diff --git a/Engine/Plugins/Experimental/HairStrands/Source/HairStrandsCore/Public/GroomAsset.h b/Engine/Plugins/Experimental/HairStrands/Source/HairStrandsCore/Public/GroomAsset.h
--- a/Engine/Plugins/Experimental/HairStrands/Source/HairStrandsCore/Public/GroomAsset.h
+++ b/Engine/Plugins/Experimental/HairStrands/Source/HairStrandsCore/Public/GroomAsset.h
@@ -11,6 +11,7 @@
 #include "GroomSettings.h"
 #include "Curves/CurveFloat.h"
 #include "HairStrandsInterface.h"
+#include "Interfaces/Interface_AssetUserData.h"
 #include "Engine/SkeletalMesh.h"
 #include "Interfaces/Interface_AssetUserData.h"
 
@@ -48,7 +49,7 @@ struct FHairStrandsRootData
 {
 	/** Build the hair strands resource */
 	FHairStrandsRootData();
-	FHairStrandsRootData(const FHairStrandsDatas* HairStrandsDatas, uint32 LODCount);
+	FHairStrandsRootData(const FHairStrandsDatas* HairStrandsDatas, uint32 LODCount, const TArray<uint32>& NumSamples);
 	void Serialize(FArchive& Ar);
 	void Reset();
 	bool HasProjectionData() const;
@@ -64,9 +65,21 @@ struct FHairStrandsRootData
 		TArray<FHairStrandsCurveTriangleBarycentricFormat::Type> RootTriangleBarycentricBuffer;
 
 		/* Strand hair roots translation and rotation in rest position relative to the bound triangle. Positions are relative to the rest root center */
-		TArray< FHairStrandsMeshTrianglePositionFormat::Type> RestRootTrianglePosition0Buffer;
-		TArray< FHairStrandsMeshTrianglePositionFormat::Type> RestRootTrianglePosition1Buffer;
-		TArray< FHairStrandsMeshTrianglePositionFormat::Type> RestRootTrianglePosition2Buffer;
+		TArray<FHairStrandsMeshTrianglePositionFormat::Type> RestRootTrianglePosition0Buffer;
+		TArray<FHairStrandsMeshTrianglePositionFormat::Type> RestRootTrianglePosition1Buffer;
+		TArray<FHairStrandsMeshTrianglePositionFormat::Type> RestRootTrianglePosition2Buffer;
+
+		/* Number of samples used for the mesh interpolation */
+		uint32 SampleCount = 0;
+
+		/* Store the hair interpolation weights | Size = SamplesCount * SamplesCount */
+		TArray<FHairStrandsWeightFormat::Type> MeshInterpolationWeightsBuffer;
+
+		/* Store the samples vertex indices */
+		TArray<FHairStrandsIndexFormat::Type> MeshSampleIndicesBuffer;
+
+		/* Store the samples rest positions */
+		TArray<FHairStrandsMeshTrianglePositionFormat::Type> RestSamplePositionsBuffer;
 	};
 
 	/* Number of roots */
@@ -91,7 +104,7 @@ struct FHairStrandsRootResource : public FRenderResource
 	/** Build the hair strands resource */
 	FHairStrandsRootResource();
 	FHairStrandsRootResource(const FHairStrandsRootData& RootData);
-	FHairStrandsRootResource(const FHairStrandsDatas* HairStrandsDatas, uint32 LODCount);
+	FHairStrandsRootResource(const FHairStrandsDatas* HairStrandsDatas, uint32 LODCount, const TArray<uint32>& NumSamples);
 
 	/* Init the buffer */
 	virtual void InitRHI() override;
@@ -101,7 +114,7 @@ struct FHairStrandsRootResource : public FRenderResource
 
 	/* Get the resource name */
 	virtual FString GetFriendlyName() const override { return TEXT("FHairStrandsRootResource"); }
-	
+
 	/* Populate GPU LOD data from RootData (this function doesn't initialize resources) */
 	void PopulateFromRootData();
 
@@ -115,11 +128,11 @@ struct FHairStrandsRootResource : public FRenderResource
 		int32 LODIndex = -1;
 
 		/* Triangle on which a root is attached */
-		/* When the projection is done with source to target mesh transfer, the projection indices does not match. 
+		/* When the projection is done with source to target mesh transfer, the projection indices does not match.
 		   In this case we need to separate index computation. The barycentric coords remain the same however. */
 		FRWBuffer RootTriangleIndexBuffer;
 		FRWBuffer RootTriangleBarycentricBuffer;
-	
+
 		/* Strand hair roots translation and rotation in rest position relative to the bound triangle. Positions are relative to the rest root center */
 		FRWBuffer RestRootTrianglePosition0Buffer;
 		FRWBuffer RestRootTrianglePosition1Buffer;
@@ -129,6 +142,14 @@ struct FHairStrandsRootResource : public FRenderResource
 		FRWBuffer DeformedRootTrianglePosition0Buffer;
 		FRWBuffer DeformedRootTrianglePosition1Buffer;
 		FRWBuffer DeformedRootTrianglePosition2Buffer;
+
+		/* Strand hair mesh interpolation matrix and sample indices */
+		uint32 SampleCount = 0;
+		FRWBuffer MeshInterpolationWeightsBuffer;
+		FRWBuffer MeshSampleIndicesBuffer;
+		FRWBuffer RestSamplePositionsBuffer;
+		FRWBuffer DeformedSamplePositionsBuffer;
+		FRWBuffer MeshSampleWeightsBuffer;
 	};
 
 	/* Store the hair projection information for each mesh LOD */
@@ -197,11 +218,14 @@ struct FHairStrandsDeformedResource : public FRenderResource
 
 	/* Whether the GPU data should be initialized with the asset data or not */
 	const bool bInitializedData;
+
+	/* Whether the GPU data should be initialized with the asset data or not */
+	uint32 CurrentIndex = 0;
 };
 
 struct FHairStrandsClusterCullingResource : public FRenderResource
 {
-	FHairStrandsClusterCullingResource(const FHairGroupData& GroupData);
+	FHairStrandsClusterCullingResource(const FHairStrandsDatas& RenStrandsData);
 
 	/* Init the buffer */
 	virtual void InitRHI() override;
@@ -324,6 +348,12 @@ struct HAIRSTRANDSCORE_API FHairGroupData
 	/** Guide render resource to be allocated */
 	FHairStrandsRestResource* HairSimulationRestResource = nullptr;
 
+	/** Interpolation resource to be allocated */
+	FHairStrandsInterpolationResource* HairInterpolationResource = nullptr;
+
+	/** Cluster culling resource to be allocated */
+	FHairStrandsClusterCullingResource* ClusterCullingResource = nullptr;
+
 	friend FArchive& operator<<(FArchive& Ar, FHairGroupData& GroupData);
 };
 
@@ -353,8 +383,8 @@ public:
 	virtual void Serialize(FArchive& Ar) override;
 
 	/** Density factor for converting hair into guide curve if no guides are provided. */
-	UPROPERTY(VisibleAnywhere, BlueprintReadWrite, Category = "BuildSettings", meta = (ClampMin = "0.01", UIMin = "0.01", UIMax = "1.0"))
-	float HairToGuideDensity = 0.1f;
+	UPROPERTY(BlueprintReadWrite, Category = "BuildSettings", meta = (ClampMin = "0.01", UIMin = "0.01", UIMax = "1.0"))
+	float HairToGuideDensity = 0.1f; // To remove, as this is now stored into the import settings
 
 	/** Group Id to be simulated */
 	UPROPERTY(EditAnywhere, Category = "Hair Physics|Group Setup", meta = (ToolTip = "Group ID that will be simulated with niagara"))
@@ -498,6 +528,8 @@ public:
 
 #endif // WITH_EDITORONLY_DATA
 
+	bool IsValid() const { return bIsInitialized; }
+
 	/** Initialize resources. */
 	void InitResource();
 
@@ -543,6 +575,7 @@ private:
 
 	TUniquePtr<FHairDescription> HairDescription;
 	TUniquePtr<FHairDescriptionBulkData> HairDescriptionBulkData;
+	bool bIsInitialized;
 
 	UPROPERTY()
 	bool bIsCacheable = true;
@@ -595,6 +628,10 @@ public:
 	UPROPERTY(VisibleAnywhere, BlueprintReadWrite, Category = "BuildSettings", meta = (ClampMin = "0.01", UIMin = "0.01", UIMax = "1.0"))
 	USkeletalMesh* TargetSkeletalMesh;
 
+	/** Number of points used for the rbf interpolation */
+	UPROPERTY(VisibleAnywhere, BlueprintReadWrite, Category = "BuildSettings", meta = (ClampMin = "0.01", UIMin = "0.01", UIMax = "1.0"))
+	int32 NumInterpolationPoints = 100;
+
 	UPROPERTY(EditAnywhere, EditFixedSize, BlueprintReadWrite, Category = "HairGroups", meta = (DisplayName = "Group"))
 	TArray<FGoomBindingGroupInfo> GroupInfos;
 
@@ -607,6 +644,10 @@ public:
 	typedef TArray<FHairGroupResource> FHairGroupResources;
 	FHairGroupResources HairGroupResources;
 
+	/** Queue of resources which needs to be deleted. This queue is needed for keeping valid pointer on the group resources 
+	   when the binding asset is recomputed */
+	TQueue<FHairGroupResource> HairGroupResourcesToDelete;
+
 	struct FHairGroupData
 	{
 		FHairStrandsRootData SimRootData;
@@ -622,12 +663,23 @@ public:
 	virtual void BeginDestroy() override;
 	virtual void Serialize(FArchive& Ar) override;
 
+	static bool IsCompatible(const USkeletalMesh* InSkeletalMesh, const UGroomBindingAsset* InBinding);
+	static bool IsCompatible(const UGroomAsset* InGroom, const UGroomBindingAsset* InBinding);
+	static bool IsBindingAssetValid(const UGroomBindingAsset* InBinding, bool bIsBindingReloading=false);
+
 #if WITH_EDITOR
 	FOnGroomBindingAssetChanged& GetOnGroomBindingAssetChanged() { return OnGroomBindingAssetChanged; }
 
 	/**  Part of UObject interface  */
 	virtual void PostEditChangeProperty(FPropertyChangedEvent& PropertyChangedEvent) override;
 
+	enum class EQueryStatus
+	{
+		None,
+		Submitted,
+		Completed
+	};
+	EQueryStatus QueryStatus = EQueryStatus::None;
 #endif // WITH_EDITOR
 
 	/** Initialize resources. */
@@ -646,3 +698,13 @@ public:
 	FOnGroomBindingAssetChanged OnGroomBindingAssetChanged;
 #endif
 };
+
+struct FProcessedHairDescription
+{
+	typedef TPair<FHairGroupInfo, FHairGroupData> FHairGroup;
+	typedef TMap<int32, FHairGroup> FHairGroups;
+	FHairGroups HairGroups;
+	bool bCanUseClosestGuidesAndWeights = false;
+	bool bHasUVData = false;
+	bool IsValid() const;
+};
diff --git a/Engine/Plugins/Experimental/HairStrands/Source/HairStrandsCore/Public/GroomBuilder.h b/Engine/Plugins/Experimental/HairStrands/Source/HairStrandsCore/Public/GroomBuilder.h
--- a/Engine/Plugins/Experimental/HairStrands/Source/HairStrandsCore/Public/GroomBuilder.h
+++ b/Engine/Plugins/Experimental/HairStrands/Source/HairStrandsCore/Public/GroomBuilder.h
@@ -8,6 +8,7 @@
 #include "GroomSettings.h"
 
 struct FHairStrandsDatas;
+struct FProcessedHairDescription;
 class UGroomAsset;
 
 struct HAIRSTRANDSCORE_API FGroomBuilder
@@ -15,4 +16,5 @@ struct HAIRSTRANDSCORE_API FGroomBuilder
 	static bool BuildGroom(const class FHairDescription& HairDescription, const struct FGroomBuildSettings& BuildSettings, UGroomAsset* GroomAsset);
 	static void GenerateGuides(const FHairStrandsDatas& InData, float DecimationPercentage, FHairStrandsDatas& OutData);
 	static void BuildData(UGroomAsset* GroomAsset, const FGroomBuildSettings& BuildSettings);
+	static bool ProcessHairDescription(const class FHairDescription& HairDescription, const FGroomBuildSettings& BuildSettings, FProcessedHairDescription& Out);
 };
\ No newline at end of file
diff --git a/Engine/Plugins/Experimental/HairStrands/Source/HairStrandsCore/Public/GroomComponent.h b/Engine/Plugins/Experimental/HairStrands/Source/HairStrandsCore/Public/GroomComponent.h
--- a/Engine/Plugins/Experimental/HairStrands/Source/HairStrandsCore/Public/GroomComponent.h
+++ b/Engine/Plugins/Experimental/HairStrands/Source/HairStrandsCore/Public/GroomComponent.h
@@ -48,9 +48,12 @@ public:
 	UPROPERTY(EditAnywhere, BlueprintReadOnly, Category = "Groom")
 	bool	bCreateNiagaraComponent;
 
-	/** Boolean to check when animation has been loaded */
+	/** Boolean to check when the simulation should be reset */
 	bool bResetSimulation;
 
+	/** Boolean to check when the simulation should be initialized */
+	bool bInitSimulation;
+
 	/** Previous bone matrix to compare the difference and decide to reset or not the simulation */
 	FMatrix	PrevBoneMatrix;
 
@@ -108,12 +111,16 @@ public:
 	virtual bool CanEditChange(const FProperty* InProperty) const override;
 	void ValidateMaterials(bool bMapCheck) const;
 	void Invalidate();
+	void InvalidateAndRecreate();
 #endif
 
+	void SetStableRasterization(bool bEnable);
 	void SetGroomAsset(UGroomAsset* Asset);
 	void SetGroomAsset(UGroomAsset* Asset, UGroomBindingAsset* InBinding);
 	void SetHairLengthScale(float Scale);
+	void SetHairRootScale(float Scale);
 	void SetHairWidth(float HairWidth);
+	void SetScatterSceneLighting(bool Enable);
 	void SetBinding(bool bBind);
 	void SetBinding(UGroomBindingAsset* InBinding);
 
@@ -168,9 +175,11 @@ private:
 
 	class USkeletalMeshComponent* RegisteredSkeletalMeshComponent;
 	FVector SkeletalPreviousPositionOffset;
-	bool bIsGroomAssetCallbackRegistered;	
+	bool bIsGroomAssetCallbackRegistered;
+	bool bIsGroomBindingAssetCallbackRegistered;
 
-	void InitResources();
+	EWorldType::Type GetWorldType() const; 
+	void InitResources(bool bIsBindingReloading=false);
 	void ReleaseResources();
 	void UpdateHairGroupsDescAndInvalidateRenderState();
 
diff --git a/Engine/Plugins/Experimental/HairStrands/Source/HairStrandsCore/Public/GroomCreateBindingOptions.h b/Engine/Plugins/Experimental/HairStrands/Source/HairStrandsCore/Public/GroomCreateBindingOptions.h
--- a/Engine/Plugins/Experimental/HairStrands/Source/HairStrandsCore/Public/GroomCreateBindingOptions.h
+++ b/Engine/Plugins/Experimental/HairStrands/Source/HairStrandsCore/Public/GroomCreateBindingOptions.h
@@ -23,4 +23,8 @@ public:
 	/** Skeletal mesh on which the groom is attached to. */
 	UPROPERTY(EditAnywhere, BlueprintReadWrite, meta = (ShowOnlyInnerProperties), Category = BuildSettings)
 	USkeletalMesh* TargetSkeletalMesh;
+
+	/** Number of points used for the rbf interpolation */
+	UPROPERTY(EditAnywhere, BlueprintReadWrite, meta = (ShowOnlyInnerProperties), Category = BuildSettings)
+	int32 NumInterpolationPoints;
 };
diff --git a/Engine/Plugins/Experimental/HairStrands/Source/HairStrandsCore/Public/GroomDesc.h b/Engine/Plugins/Experimental/HairStrands/Source/HairStrandsCore/Public/GroomDesc.h
--- a/Engine/Plugins/Experimental/HairStrands/Source/HairStrandsCore/Public/GroomDesc.h
+++ b/Engine/Plugins/Experimental/HairStrands/Source/HairStrandsCore/Public/GroomDesc.h
@@ -51,10 +51,19 @@ struct FHairGroupDesc
 	/** The targeted budget of hair vertex per pixel. Cluster strands will be decimated based on that. Used when r.HairStrands.Cluster.Culling = 1. */
 	UPROPERTY(EditAnywhere, BlueprintReadWrite, Category = "Performance", AdvancedDisplay, meta = (ClampMin = "0.0", UIMin = "0.0", UIMax = "6.0", SliderExponent = 1))
 	float LodAverageVertexPerPixel;
+
 	/** Bias the selected LOD. A value >0 will progressively select lower detailed lods. Used when r.HairStrands.Cluster.Culling = 1. */
 	UPROPERTY(EditAnywhere, AdvancedDisplay, BlueprintReadWrite, Category = "Performance", AdvancedDisplay, meta = (ClampMin = "-1.0", ClampMax = "1.0", UIMin = "-1.0", UIMax = "1.0", SliderExponent = 1))
 	float LodBias;
 
+	/** Insure the hair does not alias. When enable, group of hairs might appear thicker. Isolated hair should remain thin. */
+	UPROPERTY(EditAnywhere, BlueprintReadOnly, Category = "Groom")
+	bool bUseStableRasterization;
+
+	/** Light hair with the scene color. This is used for vellus/short hair to bring light from the surrounding surface, like skin. */
+	UPROPERTY(EditAnywhere, BlueprintReadOnly, Category = "Groom")
+	bool bScatterSceneLighting;
+
 	// Contructor used to initialised new values
 	FHairGroupDesc()
 	{
@@ -63,10 +72,19 @@ struct FHairGroupDesc
 
 	void ReInit()
 	{
+		// Filled in by groom asset
+		HairCount = 0;	// Read only
+		GuideCount = 0;	// Read only
+		HairLength = 0;	// Read only
+		HairWidth = 0;	// Can be overridden by user
+
 		HairRaytracingRadiusScale = 1.0f;
 		HairRootScale = 1.0f;
 		HairTipScale = 1.0f;
 		HairClipLength = 1.0f;
+		HairShadowDensity = 1.0f;
+		bUseStableRasterization = false;
+		bScatterSceneLighting = false;
 
 		LodAverageVertexPerPixel = 3.0f; // Good for quality on heavy assets
 		LodBias = 0.0f;
diff --git a/Engine/Plugins/Experimental/HairStrands/Source/HairStrandsCore/Public/GroomImportOptions.h b/Engine/Plugins/Experimental/HairStrands/Source/HairStrandsCore/Public/GroomImportOptions.h
--- a/Engine/Plugins/Experimental/HairStrands/Source/HairStrandsCore/Public/GroomImportOptions.h
+++ b/Engine/Plugins/Experimental/HairStrands/Source/HairStrandsCore/Public/GroomImportOptions.h
@@ -19,3 +19,34 @@ public:
 	UPROPERTY(config, EditAnywhere, BlueprintReadWrite, meta = (ShowOnlyInnerProperties), Category = BuildSettings)
 	FGroomBuildSettings BuildSettings;
 };
+
+USTRUCT(BlueprintType)
+struct HAIRSTRANDSCORE_API FGroomHairGroupPreview
+{
+	GENERATED_USTRUCT_BODY()
+
+	FGroomHairGroupPreview()
+	: GroupID(0)
+	, CurveCount(0)
+	, GuideCount(0)
+	{}
+
+	UPROPERTY(VisibleAnywhere, BlueprintReadWrite, Category = Preview)
+	int32 GroupID;
+
+	UPROPERTY(VisibleAnywhere, BlueprintReadWrite, Category = Preview)
+	int32 CurveCount;
+
+	UPROPERTY(VisibleAnywhere, BlueprintReadWrite, Category = Preview)
+	int32 GuideCount;
+};
+
+UCLASS(BlueprintType, config = EditorPerProjectUserSettings, HideCategories = ("Hidden"))
+class HAIRSTRANDSCORE_API UGroomHairGroupsPreview : public UObject
+{
+	GENERATED_UCLASS_BODY()
+
+public:
+	UPROPERTY(config, VisibleAnywhere, BlueprintReadWrite, meta = (ShowOnlyInnerProperties), Category = Preview)
+	TArray<FGroomHairGroupPreview> Groups;
+};
\ No newline at end of file
diff --git a/Engine/Plugins/Experimental/HairStrands/Source/HairStrandsCore/Public/HairStrandsRendering.h b/Engine/Plugins/Experimental/HairStrands/Source/HairStrandsCore/Public/HairStrandsRendering.h
--- a/Engine/Plugins/Experimental/HairStrands/Source/HairStrandsCore/Public/HairStrandsRendering.h
+++ b/Engine/Plugins/Experimental/HairStrands/Source/HairStrandsCore/Public/HairStrandsRendering.h
@@ -83,6 +83,8 @@ struct FHairStrandsInterpolationOutput
 		float HairRadius = 0;
 		float HairLength = 0;
 		float HairDensity = 0;
+		bool bUseStableRasterization = false;
+		bool bScatterSceneLighting = false;
 
 		inline void Reset()
 		{
@@ -97,6 +99,8 @@ struct FHairStrandsInterpolationOutput
 			HairRadius = 0;
 			HairLength = 0;
 			HairDensity = 0;
+			bUseStableRasterization = false;
+			bScatterSceneLighting = false;
 		}
 	};
 
@@ -104,7 +108,7 @@ struct FHairStrandsInterpolationOutput
 	{
 		FRWBuffer* SimDeformedPositionBuffer[2] = { nullptr, nullptr };
 		FRWBuffer* RenderDeformedPositionBuffer[2] = { nullptr, nullptr };
-		uint32 CurrentIndex = 0;
+		uint32*	   CurrentIndex = nullptr;
 
 		FRWBuffer* RenderTangentBuffer = nullptr;
 		FRWBuffer* RenderAttributeBuffer = nullptr;
@@ -152,6 +156,7 @@ void ResetHairStrandsInterpolation(
 void ComputeHairStrandsInterpolation(
 	FRHICommandListImmediate& RHICmdList,
 	const struct FShaderDrawDebugData* DebugShaderData,
+	const FTransform& LocalToWorld,
 	FHairStrandsInterpolationInput* Input,
 	FHairStrandsInterpolationOutput* Output, 
 	struct FHairStrandsProjectionHairData& RenHairDatas,
@@ -159,4 +164,6 @@ void ComputeHairStrandsInterpolation(
 	int32 LODIndex,
 	struct FHairStrandClusterData* ClusterData);
 
+HAIRSTRANDSCORE_API void ComputeInterpolationWeights(class UGroomBindingAsset* BindingAsset, class FSkeletalMeshRenderData* TargetRenderData, TArray<FRWBuffer>& TransferedPositions);
+
 HAIRSTRANDSCORE_API void AddGroomBindingTask(class UGroomBindingAsset* BindingAsset);
\ No newline at end of file
diff --git a/Engine/Plugins/Experimental/HairStrands/Source/HairStrandsCore/Public/HairStrandsVertexFactory.h b/Engine/Plugins/Experimental/HairStrands/Source/HairStrandsCore/Public/HairStrandsVertexFactory.h
--- a/Engine/Plugins/Experimental/HairStrands/Source/HairStrandsCore/Public/HairStrandsVertexFactory.h
+++ b/Engine/Plugins/Experimental/HairStrands/Source/HairStrandsCore/Public/HairStrandsVertexFactory.h
@@ -66,6 +66,8 @@ public:
 	inline FRHIShaderResourceView* GetMaterialSRV(uint32 GroupIndex) const			{ check(Data.InterpolationOutput); return Data.InterpolationOutput->HairGroups[GroupIndex].VFInput.HairMaterialBuffer; }
 	inline FRHIShaderResourceView* GetTangentSRV(uint32 GroupIndex) const			{ check(Data.InterpolationOutput); return Data.InterpolationOutput->HairGroups[GroupIndex].VFInput.HairTangentBuffer; }
 
+	inline bool  UseStableRasterization(uint32 GroupIndex) const					{ check(Data.InterpolationOutput); return Data.InterpolationOutput->HairGroups[GroupIndex].VFInput.bUseStableRasterization; };
+	inline bool  UseScatterSceneLighting(uint32 GroupIndex) const					{ check(Data.InterpolationOutput); return Data.InterpolationOutput->HairGroups[GroupIndex].VFInput.bScatterSceneLighting; };
 	inline float GetMaxStrandRadius(uint32 GroupIndex) const						{ check(Data.InterpolationOutput); return Data.InterpolationOutput->HairGroups[GroupIndex].VFInput.HairRadius; };
 	inline float GetMaxStrandLength(uint32 GroupIndex) const						{ check(Data.InterpolationOutput); return Data.InterpolationOutput->HairGroups[GroupIndex].VFInput.HairLength; };
 	inline float GetHairDensity(uint32 GroupIndex) const							{ check(Data.InterpolationOutput); return Data.InterpolationOutput->HairGroups[GroupIndex].VFInput.HairDensity; };
diff --git a/Engine/Plugins/Experimental/HairStrands/Source/HairStrandsEditor/Private/GroomActions.cpp b/Engine/Plugins/Experimental/HairStrands/Source/HairStrandsEditor/Private/GroomActions.cpp
--- a/Engine/Plugins/Experimental/HairStrands/Source/HairStrandsEditor/Private/GroomActions.cpp
+++ b/Engine/Plugins/Experimental/HairStrands/Source/HairStrandsEditor/Private/GroomActions.cpp
@@ -170,7 +170,14 @@ void FGroomActions::ExecuteCreateBindingAsset(TArray<TWeakObjectPtr<UGroomAsset>
 				}
 				else if (GroomAsset.Get() && CurrentOptions && CurrentOptions->TargetSkeletalMesh)
 				{
-					UGroomBindingAsset* BindingAsset = CreateGroomBindinAsset(GroomAsset.Get(), CurrentOptions->SourceSkeletalMesh, CurrentOptions->TargetSkeletalMesh);
+					GroomAsset->ConditionalPostLoad();
+					if (CurrentOptions->SourceSkeletalMesh)
+					{
+						CurrentOptions->SourceSkeletalMesh->ConditionalPostLoad();
+					}
+					CurrentOptions->TargetSkeletalMesh->ConditionalPostLoad();
+
+					UGroomBindingAsset* BindingAsset = CreateGroomBindinAsset(GroomAsset.Get(), CurrentOptions->SourceSkeletalMesh, CurrentOptions->TargetSkeletalMesh, CurrentOptions->NumInterpolationPoints);
 
 					// The binding task will generate and set the binding value back to the binding asset.
 					// This code is not thread safe.
diff --git a/Engine/Plugins/Experimental/HairStrands/Source/HairStrandsEditor/Private/GroomBindingActions.cpp b/Engine/Plugins/Experimental/HairStrands/Source/HairStrandsEditor/Private/GroomBindingActions.cpp
new file mode 100644
--- /dev/null
+++ b/Engine/Plugins/Experimental/HairStrands/Source/HairStrandsEditor/Private/GroomBindingActions.cpp
@@ -0,0 +1,100 @@
+// Copyright Epic Games, Inc. All Rights Reserved.
+
+#include "GroomBindingActions.h"
+#include "GroomAsset.h"
+
+#include "EditorFramework/AssetImportData.h"
+#include "HairStrandsRendering.h"
+#include "Toolkits/SimpleAssetEditor.h"
+#include "ToolMenuSection.h"
+
+#define LOCTEXT_NAMESPACE "AssetTypeActions"
+
+FGroomBindingActions::FGroomBindingActions()
+{}
+
+bool FGroomBindingActions::CanFilter()
+{
+	return true;
+}
+
+void FGroomBindingActions::GetActions(const TArray<UObject*>& InObjects, FToolMenuSection& Section)
+{
+	FAssetTypeActions_Base::GetActions(InObjects, Section);
+
+	TArray<TWeakObjectPtr<UGroomBindingAsset>> GroomBindingAssets = GetTypedWeakObjectPtrs<UGroomBindingAsset>(InObjects);
+
+	Section.AddMenuEntry(
+		"RebuildGroom",
+		LOCTEXT("RebuildGroom", "Rebuild"),
+		LOCTEXT("RebuildGroomTooltip", "Rebuild the groom binding"),
+		FSlateIcon(FEditorStyle::GetStyleSetName(), "ContentBrowser.AssetActions"),
+		FUIAction(
+			FExecuteAction::CreateSP(this, &FGroomBindingActions::ExecuteRebuildBindingAsset, GroomBindingAssets),
+			FCanExecuteAction::CreateSP(this, &FGroomBindingActions::CanRebuildBindingAsset, GroomBindingAssets)
+		)
+	);
+}
+
+uint32 FGroomBindingActions::GetCategories()
+{
+	return EAssetTypeCategories::Misc;
+}
+
+FText FGroomBindingActions::GetName() const
+{
+	return NSLOCTEXT("AssetTypeActions", "AssetTypeActions_GroomBinding", "GroomBinding");
+}
+
+UClass* FGroomBindingActions::GetSupportedClass() const
+{
+	return UGroomBindingAsset::StaticClass();
+}
+
+FColor FGroomBindingActions::GetTypeColor() const
+{
+	return FColor::White;
+}
+
+bool FGroomBindingActions::HasActions(const TArray<UObject*>& InObjects) const
+{
+	return true;
+}
+
+void FGroomBindingActions::OpenAssetEditor(const TArray<UObject*>& InObjects, TSharedPtr<IToolkitHost> EditWithinLevelEditor)
+{
+	// #ueent_todo: Will need a custom editor at some point, for now just use the Properties editor
+	FSimpleAssetEditor::CreateEditor(EToolkitMode::Standalone, EditWithinLevelEditor, InObjects);
+}
+
+bool FGroomBindingActions::CanRebuildBindingAsset(TArray<TWeakObjectPtr<UGroomBindingAsset>> Objects) const
+{
+	for (TWeakObjectPtr<UGroomBindingAsset> BindingAsset : Objects)
+	{
+		if (BindingAsset.IsValid())
+		{
+			return true;
+		}
+	}
+	return false;
+}
+
+void FGroomBindingActions::ExecuteRebuildBindingAsset(TArray<TWeakObjectPtr<UGroomBindingAsset>> Objects) const
+{
+	for (TWeakObjectPtr<UGroomBindingAsset> BindingAsset : Objects)
+	{
+		if (BindingAsset.IsValid() && BindingAsset->Groom && BindingAsset->TargetSkeletalMesh)
+		{
+			BindingAsset->Groom->ConditionalPostLoad();
+			BindingAsset->TargetSkeletalMesh->ConditionalPostLoad();
+			if (BindingAsset->SourceSkeletalMesh)
+			{
+				BindingAsset->SourceSkeletalMesh->ConditionalPostLoad();
+			}
+			AddGroomBindingTask(BindingAsset.Get());
+			BindingAsset->GetOutermost()->MarkPackageDirty();
+		}
+	}
+}
+
+#undef LOCTEXT_NAMESPACE
\ No newline at end of file
diff --git a/Engine/Plugins/Experimental/HairStrands/Source/HairStrandsEditor/Private/GroomCreateBindingOptionsWindow.cpp b/Engine/Plugins/Experimental/HairStrands/Source/HairStrandsEditor/Private/GroomCreateBindingOptionsWindow.cpp
--- a/Engine/Plugins/Experimental/HairStrands/Source/HairStrandsEditor/Private/GroomCreateBindingOptionsWindow.cpp
+++ b/Engine/Plugins/Experimental/HairStrands/Source/HairStrandsEditor/Private/GroomCreateBindingOptionsWindow.cpp
@@ -151,7 +151,7 @@ TSharedPtr<SGroomCreateBindingOptionsWindow> SGroomCreateBindingOptionsWindow::D
 	return DisplayOptions(BindingOptions, EGroomBindingOptionsVisibility::BuildOptions, LOCTEXT("WindowTitle", "Groom Binding Options"), LOCTEXT("Build", "Create"));
 }
 
-static UObject* InternalCreateNewBindAsset(FName InAssetName, UObject* InParent, UGroomAsset* GroomAsset, USkeletalMesh* SourceSkelMesh, USkeletalMesh* TargetSkelMesh)
+static UObject* InternalCreateNewBindAsset(FName InAssetName, UObject* InParent, UGroomAsset* GroomAsset, USkeletalMesh* SourceSkelMesh, USkeletalMesh* TargetSkelMesh, const int32 NumInterpolationPoints)
 {
 	if (!TargetSkelMesh)
 	{
@@ -196,6 +196,8 @@ static UObject* InternalCreateNewBindAsset(FName InAssetName, UObject* InParent,
 		NewAsset->Groom = GroomAsset;
 		NewAsset->SourceSkeletalMesh = SourceSkelMesh;
 		NewAsset->TargetSkeletalMesh = TargetSkelMesh;
+		NewAsset->HairGroupDatas.Reserve(GroomAsset->HairGroupsData.Num());
+		NewAsset->NumInterpolationPoints = NumInterpolationPoints;
 		NewAsset->MarkPackageDirty();
 
 		// Notify the asset registry
@@ -206,14 +208,14 @@ static UObject* InternalCreateNewBindAsset(FName InAssetName, UObject* InParent,
 	return nullptr;
 }
 
-UGroomBindingAsset* CreateGroomBindinAsset(UGroomAsset* GroomAsset, USkeletalMesh* SourceSkelMesh, USkeletalMesh* TargetSkelMesh)
+UGroomBindingAsset* CreateGroomBindinAsset(UGroomAsset* GroomAsset, USkeletalMesh* SourceSkelMesh, USkeletalMesh* TargetSkelMesh, const int32 NumInterpolationPoints)
 {
 	if (!GroomAsset || !TargetSkelMesh)
 	{
 		return nullptr;
 	}
 
-	UObject* BindingAsset = InternalCreateNewBindAsset(NAME_None, nullptr, GroomAsset, SourceSkelMesh, TargetSkelMesh);
+	UObject* BindingAsset = InternalCreateNewBindAsset(NAME_None, nullptr, GroomAsset, SourceSkelMesh, TargetSkelMesh, NumInterpolationPoints);
 
 	if (BindingAsset)
 	{
diff --git a/Engine/Plugins/Experimental/HairStrands/Source/HairStrandsEditor/Private/GroomCreateBindingOptionsWindow.h b/Engine/Plugins/Experimental/HairStrands/Source/HairStrandsEditor/Private/GroomCreateBindingOptionsWindow.h
--- a/Engine/Plugins/Experimental/HairStrands/Source/HairStrandsEditor/Private/GroomCreateBindingOptionsWindow.h
+++ b/Engine/Plugins/Experimental/HairStrands/Source/HairStrandsEditor/Private/GroomCreateBindingOptionsWindow.h
@@ -86,4 +86,4 @@ private:
 	bool bShouldCreate;
 };
 
-class UGroomBindingAsset* CreateGroomBindinAsset(class UGroomAsset* GroomAsset, class USkeletalMesh* SourceSkelMesh, class USkeletalMesh* TargetSkelMesh);
\ No newline at end of file
+class UGroomBindingAsset* CreateGroomBindinAsset(class UGroomAsset* GroomAsset, class USkeletalMesh* SourceSkelMesh, class USkeletalMesh* TargetSkelMesh, const int32 NumInterpolationPoints);
\ No newline at end of file
diff --git a/Engine/Plugins/Experimental/HairStrands/Source/HairStrandsEditor/Private/GroomImportOptionsWindow.cpp b/Engine/Plugins/Experimental/HairStrands/Source/HairStrandsEditor/Private/GroomImportOptionsWindow.cpp
--- a/Engine/Plugins/Experimental/HairStrands/Source/HairStrandsEditor/Private/GroomImportOptionsWindow.cpp
+++ b/Engine/Plugins/Experimental/HairStrands/Source/HairStrandsEditor/Private/GroomImportOptionsWindow.cpp
@@ -14,20 +14,85 @@
 #include "Widgets/Layout/SUniformGridPanel.h"
 #include "Widgets/SBoxPanel.h"
 #include "Widgets/Text/STextBlock.h"
+#include "GroomAsset.h"
 
 #define LOCTEXT_NAMESPACE "GroomImportOptionsWindow"
 
+static int32 GHairStrandsImportValidation = 1;
+static FAutoConsoleVariableRef CVarHairStrandsImportValidation(TEXT("r.HairStrands.ImportValidation"), GHairStrandsImportValidation, TEXT("Enable/Disable hair strands validation at import time"));
+
+bool RunGroomAssetValidation()
+{
+	return GHairStrandsImportValidation > 0;
+}
+
+enum class EHairDescriptionStatus
+{
+	Valid,
+	NoGroup,
+	NoCurve,
+	Unknown
+};
+
+static EHairDescriptionStatus GetStatus(UGroomHairGroupsPreview* Description)
+{
+	if (Description)
+	{
+		for (const FGroomHairGroupPreview& Group : Description->Groups)
+		{
+			if (Group.CurveCount == 0)
+				return EHairDescriptionStatus::NoCurve;
+		}
+
+		if (Description->Groups.Num() == 0)
+		{
+			return EHairDescriptionStatus::NoGroup;
+		}
+		return EHairDescriptionStatus::Valid;
+	}
+	else
+	{
+		return EHairDescriptionStatus::Unknown;
+	}
+}
+
 void SGroomImportOptionsWindow::Construct(const FArguments& InArgs)
 {
 	ImportOptions = InArgs._ImportOptions;
+	GroupsPreview = InArgs._GroupsPreview;
 	WidgetWindow = InArgs._WidgetWindow;
 
 	FPropertyEditorModule& PropertyEditorModule = FModuleManager::GetModuleChecked<FPropertyEditorModule>("PropertyEditor");
 	FDetailsViewArgs DetailsViewArgs;
 	DetailsViewArgs.bAllowSearch = false;
 	DetailsViewArgs.NameAreaSettings = FDetailsViewArgs::HideNameArea;
+
 	DetailsView = PropertyEditorModule.CreateDetailView(DetailsViewArgs);
 	DetailsView->SetObject(ImportOptions);
+	
+	DetailsView2 = PropertyEditorModule.CreateDetailView(DetailsViewArgs);
+	DetailsView2->SetObject(GroupsPreview);
+
+	
+	const EHairDescriptionStatus Status = GetStatus(GroupsPreview);
+
+	FText ValidationText = LOCTEXT("GroomOptionsWindow_ValidationText", "Unknown");
+	FLinearColor ValidationColor(1,1,1);
+	if (Status == EHairDescriptionStatus::Valid)
+	{
+		ValidationText = LOCTEXT("GroomOptionsWindow_ValidationText", "Valid");
+		ValidationColor = FLinearColor(0, 0.80f, 0, 1);
+	}
+	else if (Status == EHairDescriptionStatus::NoCurve)
+	{
+		ValidationText = LOCTEXT("GroomOptionsWindow_ValidationText", "Invalid. Some groups have 0 curves.");
+		ValidationColor = FLinearColor(0.80f, 0, 0, 1);
+	}
+	else if (Status == EHairDescriptionStatus::NoGroup)
+	{
+		ValidationText = LOCTEXT("GroomOptionsWindow_ValidationText", "Invalid. The groom does not contain any group.");
+		ValidationColor = FLinearColor(1, 0, 0, 1);
+	}
 
 	this->ChildSlot
 	[
@@ -61,6 +126,35 @@ void SGroomImportOptionsWindow::Construct(const FArguments& InArgs)
 			]
 		]
 
+		+ SVerticalBox::Slot()
+		.AutoHeight()
+		.Padding(2)
+		[
+			SNew(SBorder)
+			.Padding(FMargin(3))
+			.BorderImage(FEditorStyle::GetBrush("ToolPanel.GroupBorder"))
+			[
+				SNew(SHorizontalBox)
+				+ SHorizontalBox::Slot()
+				.AutoWidth()
+				[
+					SNew(STextBlock)
+					.Font(FEditorStyle::GetFontStyle("CurveEd.LabelFont"))
+					.Text(LOCTEXT("GroomOptionsWindow_StatusFile", "Status File: "))
+				]
+				+ SHorizontalBox::Slot()
+				.Padding(5, 0, 0, 0)
+				.AutoWidth()
+				.VAlign(VAlign_Center)
+				[
+					SNew(STextBlock)
+					.Font(FEditorStyle::GetFontStyle("CurveEd.InfoFont"))
+					.Text(ValidationText)
+					.ColorAndOpacity(ValidationColor)
+				]
+			]
+		]
+
 		+ SVerticalBox::Slot()
 		.Padding(2)
 		.MaxHeight(500.0f)
@@ -68,6 +162,13 @@ void SGroomImportOptionsWindow::Construct(const FArguments& InArgs)
 			DetailsView->AsShared()
 		]
 
+		+ SVerticalBox::Slot()
+		.AutoHeight()
+		.Padding(2)		
+		[
+			DetailsView2->AsShared()
+		]
+
 		+ SVerticalBox::Slot()
 		.AutoHeight()
 		.HAlign(HAlign_Right)
@@ -96,6 +197,19 @@ void SGroomImportOptionsWindow::Construct(const FArguments& InArgs)
 
 bool SGroomImportOptionsWindow::CanImport()  const
 {
+	if (GroupsPreview)
+	{
+		for (const FGroomHairGroupPreview& Group : GroupsPreview->Groups)
+		{
+			if (Group.CurveCount == 0)
+				return false;
+		}
+
+		if (GroupsPreview->Groups.Num() == 0)
+		{
+			return false;
+		}
+	}
 	return true;
 }
 
@@ -109,7 +223,7 @@ enum class EGroomOptionsVisibility : uint8
 
 ENUM_CLASS_FLAGS(EGroomOptionsVisibility);
 
-TSharedPtr<SGroomImportOptionsWindow> DisplayOptions(UGroomImportOptions* ImportOptions, const FString& FilePath, EGroomOptionsVisibility VisibilityFlag, FText WindowTitle, FText InButtonLabel)
+TSharedPtr<SGroomImportOptionsWindow> DisplayOptions(UGroomImportOptions* ImportOptions, const FString& FilePath, const FProcessedHairDescription* ProcessedDescription, EGroomOptionsVisibility VisibilityFlag, FText WindowTitle, FText InButtonLabel)
 {
 	TSharedRef<SWindow> Window = SNew(SWindow)
 		.Title(WindowTitle)
@@ -134,11 +248,29 @@ TSharedPtr<SGroomImportOptionsWindow> DisplayOptions(UGroomImportOptions* Import
 		}
 	}
 
+	// Convert the process hair description into hair groups
+	UGroomHairGroupsPreview* GroupsPreview = nullptr;
+	if (ProcessedDescription)
+	{
+		GroupsPreview = NewObject<UGroomHairGroupsPreview>();
+		for (TPair<int32, FProcessedHairDescription::FHairGroup> HairGroupIt : ProcessedDescription->HairGroups)
+		{
+			const FProcessedHairDescription::FHairGroup& Group = HairGroupIt.Value;
+			const FHairGroupInfo& GroupInfo = Group.Key;
+
+			FGroomHairGroupPreview& OutGroup = GroupsPreview->Groups.AddDefaulted_GetRef();
+			OutGroup.GroupID = GroupInfo.GroupID;
+			OutGroup.CurveCount = GroupInfo.NumCurves;
+			OutGroup.GuideCount = GroupInfo.NumGuides;
+		}
+	}
+
 	FString FileName = FPaths::GetCleanFilename(FilePath);
 	Window->SetContent
 	(
 		SAssignNew(OptionsWindow, SGroomImportOptionsWindow)
 		.ImportOptions(ImportOptions)
+		.GroupsPreview(GroupsPreview)
 		.WidgetWindow(Window)
 		.FullPath(FText::FromString(FileName))
 		.ButtonLabel(InButtonLabel)
@@ -157,14 +289,14 @@ TSharedPtr<SGroomImportOptionsWindow> DisplayOptions(UGroomImportOptions* Import
 	return OptionsWindow;
 }
 
-TSharedPtr<SGroomImportOptionsWindow> SGroomImportOptionsWindow::DisplayImportOptions(UGroomImportOptions* ImportOptions, const FString& FilePath)
+TSharedPtr<SGroomImportOptionsWindow> SGroomImportOptionsWindow::DisplayImportOptions(UGroomImportOptions* ImportOptions, const FString& FilePath, const FProcessedHairDescription* ProcessedDescription)
 {
-	return DisplayOptions(ImportOptions, FilePath, EGroomOptionsVisibility::All, LOCTEXT("ImportOptions", "Groom Import Options"), LOCTEXT("Import", "Import"));
+	return DisplayOptions(ImportOptions, FilePath, ProcessedDescription, EGroomOptionsVisibility::All, LOCTEXT("WindowTitle", "Groom Import Options"), LOCTEXT("Import", "Import"));
 }
 
 TSharedPtr<SGroomImportOptionsWindow> SGroomImportOptionsWindow::DisplayRebuildOptions(UGroomImportOptions* ImportOptions, const FString& FilePath)
 {
-	return DisplayOptions(ImportOptions, FilePath, EGroomOptionsVisibility::BuildOptions, LOCTEXT("RebuildOptions", "Groom Build Options"), LOCTEXT("Build", "Build"));
+	return DisplayOptions(ImportOptions, FilePath, nullptr, EGroomOptionsVisibility::BuildOptions, LOCTEXT("WindowTitle", "Groom Build Options"), LOCTEXT("Build", "Build"));
 }
 
 
diff --git a/Engine/Plugins/Experimental/HairStrands/Source/HairStrandsEditor/Private/GroomImportOptionsWindow.h b/Engine/Plugins/Experimental/HairStrands/Source/HairStrandsEditor/Private/GroomImportOptionsWindow.h
--- a/Engine/Plugins/Experimental/HairStrands/Source/HairStrandsEditor/Private/GroomImportOptionsWindow.h
+++ b/Engine/Plugins/Experimental/HairStrands/Source/HairStrandsEditor/Private/GroomImportOptionsWindow.h
@@ -11,18 +11,24 @@
 
 class SButton;
 class UGroomImportOptions;
+class UGroomHairGroupsPreview;
+struct FHairGroupInfo;
+
+bool RunGroomAssetValidation();
 
 class SGroomImportOptionsWindow : public SCompoundWidget
 {
 public:
 	SLATE_BEGIN_ARGS(SGroomImportOptionsWindow)
 		: _ImportOptions(nullptr)
+		, _GroupsPreview(nullptr)
 		, _WidgetWindow()
 		, _FullPath()
 		, _ButtonLabel()
 	{}
 
 	SLATE_ARGUMENT(UGroomImportOptions*, ImportOptions)
+	SLATE_ARGUMENT(UGroomHairGroupsPreview*, GroupsPreview)
 	SLATE_ARGUMENT(TSharedPtr<SWindow>, WidgetWindow)
 	SLATE_ARGUMENT(FText, FullPath)
 	SLATE_ARGUMENT(FText, ButtonLabel)
@@ -32,7 +38,7 @@ public:
 	void Construct(const FArguments& InArgs);
 	virtual bool SupportsKeyboardFocus() const override { return true; }
 
-	static TSharedPtr<SGroomImportOptionsWindow> DisplayImportOptions(UGroomImportOptions* ImportOptions, const FString& FilePath);
+	static TSharedPtr<SGroomImportOptionsWindow> DisplayImportOptions(UGroomImportOptions* ImportOptions, const FString& FilePath, const struct FProcessedHairDescription* ProcessedDescription);
 	static TSharedPtr<SGroomImportOptionsWindow> DisplayRebuildOptions(UGroomImportOptions* ImportOptions, const FString& FilePath);
 
 	FReply OnImport()
@@ -73,6 +79,7 @@ public:
 	SGroomImportOptionsWindow() 
 		: ImportOptions(nullptr)
 		, bShouldImport(false)
+		, GroupsPreview(nullptr)
 	{}
 
 private:
@@ -82,7 +89,10 @@ private:
 private:
 	UGroomImportOptions* ImportOptions;
 	TSharedPtr<class IDetailsView> DetailsView;
+	TSharedPtr<class IDetailsView> DetailsView2;
 	TWeakPtr<SWindow> WidgetWindow;
 	TSharedPtr<SButton> ImportButton;
 	bool bShouldImport;
+public:
+	UGroomHairGroupsPreview* GroupsPreview;
 };
diff --git a/Engine/Plugins/Experimental/HairStrands/Source/HairStrandsEditor/Private/HairStrandsEditor.cpp b/Engine/Plugins/Experimental/HairStrands/Source/HairStrandsEditor/Private/HairStrandsEditor.cpp
--- a/Engine/Plugins/Experimental/HairStrands/Source/HairStrandsEditor/Private/HairStrandsEditor.cpp
+++ b/Engine/Plugins/Experimental/HairStrands/Source/HairStrandsEditor/Private/HairStrandsEditor.cpp
@@ -2,6 +2,7 @@
 
 #include "HairStrandsEditor.h"
 #include "GroomActions.h"
+#include "GroomBindingActions.h"
 
 #include "Styling/SlateStyleRegistry.h"
 #include "Styling/SlateTypes.h"
@@ -15,10 +16,13 @@ IMPLEMENT_MODULE(FHairStrandsEditor, HairStrandsEditor);
 void FHairStrandsEditor::StartupModule()
 {
 	IAssetTools& AssetTools = FModuleManager::LoadModuleChecked<FAssetToolsModule>("AssetTools").Get();
-	TSharedRef<IAssetTypeActions> AssetActions = MakeShareable(new FGroomActions());
+	TSharedRef<IAssetTypeActions> GroomAssetActions = MakeShareable(new FGroomActions());
+	TSharedRef<IAssetTypeActions> BindingAssetActions = MakeShareable(new FGroomBindingActions());
 
-	AssetTools.RegisterAssetTypeActions(AssetActions);
-	RegisteredAssetTypeActions.Add(AssetActions);
+	AssetTools.RegisterAssetTypeActions(GroomAssetActions);
+	AssetTools.RegisterAssetTypeActions(BindingAssetActions);
+	RegisteredAssetTypeActions.Add(GroomAssetActions);
+	RegisteredAssetTypeActions.Add(BindingAssetActions);
 
 	RegisterHairTranslator<FFbxHairTranslator>();
 
diff --git a/Engine/Plugins/Experimental/HairStrands/Source/HairStrandsEditor/Private/HairStrandsFactory.cpp b/Engine/Plugins/Experimental/HairStrands/Source/HairStrandsEditor/Private/HairStrandsFactory.cpp
--- a/Engine/Plugins/Experimental/HairStrands/Source/HairStrandsEditor/Private/HairStrandsFactory.cpp
+++ b/Engine/Plugins/Experimental/HairStrands/Source/HairStrandsEditor/Private/HairStrandsFactory.cpp
@@ -6,6 +6,7 @@
 #include "GroomAsset.h"
 #include "GroomAssetImportData.h"
 #include "GroomImportOptions.h"
+#include "GroomBuilder.h"
 #include "GroomImportOptionsWindow.h"
 #include "HairDescription.h"
 #include "HairStrandsEditor.h"
@@ -74,8 +75,25 @@ UObject* UHairStrandsFactory::FactoryCreateFile(UClass* InClass, UObject* InPare
 
 	if (!GIsRunningUnattendedScript && !IsAutomatedImport())
 	{
+		// Load the alembic file upfront to preview & report any potential issue
+		const bool bRunValidation = RunGroomAssetValidation();
+		FProcessedHairDescription OutDescription;
+		if (bRunValidation)
+		{
+			FScopedSlowTask Progress((float)1, LOCTEXT("ImportHairAsset", "Importing hair asset for preview..."), true);
+			Progress.MakeDialog(true);
+
+			FHairDescription HairDescription;
+			if (!SelectedTranslator->Translate(Filename, HairDescription, ImportOptions->ConversionSettings))
+			{
+				return nullptr;
+			}
+
+			FGroomBuilder::ProcessHairDescription(HairDescription, ImportOptions->BuildSettings, OutDescription);
+		}
+
 		// Display import options and handle user cancellation
-		TSharedPtr<SGroomImportOptionsWindow> GroomOptionWindow = SGroomImportOptionsWindow::DisplayImportOptions(ImportOptions, Filename);
+		TSharedPtr<SGroomImportOptionsWindow> GroomOptionWindow = SGroomImportOptionsWindow::DisplayImportOptions(ImportOptions, Filename, bRunValidation ? &OutDescription : nullptr);
 		if (!GroomOptionWindow->ShouldImport())
 		{
 			bOutOperationCanceled = true;
@@ -86,8 +104,6 @@ UObject* UHairStrandsFactory::FactoryCreateFile(UClass* InClass, UObject* InPare
 		ImportOptions->SaveConfig();
 	}
 
-	//GEditor->GetEditorSubsystem<UImportSubsystem>()->BroadcastAssetPreImport(this, InClass, InParent, InName, FileType);
-
 	FScopedSlowTask Progress((float) 1, LOCTEXT("ImportHairAsset", "Importing hair asset..."), true);
 	Progress.MakeDialog(true);
 
diff --git a/Engine/Plugins/Experimental/HairStrands/Source/HairStrandsEditor/Private/ReimportHairStrandsFactory.cpp b/Engine/Plugins/Experimental/HairStrands/Source/HairStrandsEditor/Private/ReimportHairStrandsFactory.cpp
--- a/Engine/Plugins/Experimental/HairStrands/Source/HairStrandsEditor/Private/ReimportHairStrandsFactory.cpp
+++ b/Engine/Plugins/Experimental/HairStrands/Source/HairStrandsEditor/Private/ReimportHairStrandsFactory.cpp
@@ -5,6 +5,7 @@
 #include "EditorFramework/AssetImportData.h"
 #include "HairDescription.h"
 #include "GroomAsset.h"
+#include "GroomBuilder.h"
 #include "GroomAssetImportData.h"
 #include "GroomImportOptions.h"
 #include "GroomImportOptionsWindow.h"
@@ -95,9 +96,32 @@ EReimportResult::Type UReimportHairStrandsFactory::Reimport(UObject* Obj)
 			CurrentOptions = NewObject<UGroomImportOptions>();
 		}
 
+		TSharedPtr<IHairStrandsTranslator> SelectedTranslator = GetTranslator(CurrentFilename);
+		if (!SelectedTranslator.IsValid())
+		{
+			return EReimportResult::Failed;
+		}
+
+		// Load the alembic file upfront to preview & report any potential issue
+		FProcessedHairDescription OutDescription;
+		const bool bRunValidation = RunGroomAssetValidation();
+		if (bRunValidation)
+		{
+			FScopedSlowTask Progress((float)1, LOCTEXT("ReimportHairAsset", "Reimporting hair asset for preview..."), true);
+			Progress.MakeDialog(true);
+
+			FHairDescription HairDescription;
+			if (!SelectedTranslator->Translate(CurrentFilename, HairDescription, CurrentOptions->ConversionSettings))
+			{
+				return EReimportResult::Failed;
+			}
+
+			FGroomBuilder::ProcessHairDescription(HairDescription, CurrentOptions->BuildSettings, OutDescription);
+		}
+
 		if (!GIsRunningUnattendedScript && !IsAutomatedImport())
 		{
-			TSharedPtr<SGroomImportOptionsWindow> GroomOptionWindow = SGroomImportOptionsWindow::DisplayImportOptions(CurrentOptions, CurrentFilename);
+			TSharedPtr<SGroomImportOptionsWindow> GroomOptionWindow = SGroomImportOptionsWindow::DisplayImportOptions(CurrentOptions, CurrentFilename, bRunValidation ? &OutDescription : nullptr);
 
 			if (!GroomOptionWindow->ShouldImport())
 			{
@@ -109,12 +133,6 @@ EReimportResult::Type UReimportHairStrandsFactory::Reimport(UObject* Obj)
 			GroomAssetImportData->ImportOptions = CurrentOptions;
 		}
 
-		TSharedPtr<IHairStrandsTranslator> SelectedTranslator = GetTranslator(CurrentFilename);
-		if (!SelectedTranslator.IsValid())
-		{
-			return EReimportResult::Failed;
-		}
-
 		FHairDescription HairDescription;
 		if (!SelectedTranslator->Translate(CurrentFilename, HairDescription, CurrentOptions->ConversionSettings))
 		{
diff --git a/Engine/Plugins/Experimental/HairStrands/Source/HairStrandsEditor/Public/GroomActions.h b/Engine/Plugins/Experimental/HairStrands/Source/HairStrandsEditor/Public/GroomActions.h
--- a/Engine/Plugins/Experimental/HairStrands/Source/HairStrandsEditor/Public/GroomActions.h
+++ b/Engine/Plugins/Experimental/HairStrands/Source/HairStrandsEditor/Public/GroomActions.h
@@ -43,7 +43,7 @@ private:
 	bool CanRebuild(TArray<TWeakObjectPtr<UGroomAsset>> Objects) const;
 	void ExecuteRebuild(TArray<TWeakObjectPtr<UGroomAsset>> Objects) const;
 
-	/** Callback for Rebuild groom action */
+	/** Callback for create binding asset action */
 	bool CanCreateBindingAsset(TArray<TWeakObjectPtr<UGroomAsset>> Objects) const;
 	void ExecuteCreateBindingAsset(TArray<TWeakObjectPtr<UGroomAsset>> Objects) const;
 };
diff --git a/Engine/Plugins/Experimental/HairStrands/Source/HairStrandsEditor/Public/GroomBindingActions.h b/Engine/Plugins/Experimental/HairStrands/Source/HairStrandsEditor/Public/GroomBindingActions.h
new file mode 100644
--- /dev/null
+++ b/Engine/Plugins/Experimental/HairStrands/Source/HairStrandsEditor/Public/GroomBindingActions.h
@@ -0,0 +1,43 @@
+// Copyright Epic Games, Inc. All Rights Reserved.
+
+#pragma once
+
+#include "AssetTypeActions_Base.h"
+#include "Templates/SharedPointer.h"
+
+class ISlateStyle;
+class UGroomBindingAsset;
+
+/**
+ * Implements an action for groom binding assets.
+ */
+class FGroomBindingActions : public FAssetTypeActions_Base
+{
+public:
+
+	/**
+	 * Creates and initializes a new instance.
+	 *
+	 * @param InStyle The style set to use for asset editor toolkits.
+	 */
+	FGroomBindingActions();
+
+public:
+
+	//~ FAssetTypeActions_Base overrides
+
+	virtual bool CanFilter() override;
+	virtual void GetActions(const TArray<UObject*>& InObjects, FToolMenuSection& Section) override;
+	virtual uint32 GetCategories() override;
+	virtual FText GetName() const override;
+	virtual UClass* GetSupportedClass() const override;
+	virtual FColor GetTypeColor() const override;
+	virtual bool HasActions(const TArray<UObject*>& InObjects) const override;
+	virtual void OpenAssetEditor(const TArray<UObject*>& InObjects, TSharedPtr<IToolkitHost> EditWithinLevelEditor = TSharedPtr<IToolkitHost>()) override;
+
+private:
+
+	/** Callback for rebuild binding asset action */
+	bool CanRebuildBindingAsset(TArray<TWeakObjectPtr<UGroomBindingAsset>> Objects) const;
+	void ExecuteRebuildBindingAsset(TArray<TWeakObjectPtr<UGroomBindingAsset>> Objects) const;
+};
