Integrate hair changes (rendering+simulation)

diff --git a/Engine/Plugins/Experimental/HairStrands/Source/HairStrandsNiagara/Private/NiagaraDataInterfaceHairStrands.cpp b/Engine/Plugins/Experimental/HairStrands/Source/HairStrandsNiagara/Private/NiagaraDataInterfaceHairStrands.cpp
--- a/Engine/Plugins/Experimental/HairStrands/Source/HairStrandsNiagara/Private/NiagaraDataInterfaceHairStrands.cpp
+++ b/Engine/Plugins/Experimental/HairStrands/Source/HairStrandsNiagara/Private/NiagaraDataInterfaceHairStrands.cpp
@@ -93,6 +93,10 @@ static const FName AttachNodeStateName(TEXT("AttachNodeState"));
 
 //------------------------------------------------------------------------------------------------------------
 
+static const FName EvalSkinnedPositionName(TEXT("EvalSkinnedPosition"));
+
+//------------------------------------------------------------------------------------------------------------
+
 static const FName UpdatePointPositionName(TEXT("UpdatePointPosition"));
 static const FName ResetPointPositionName(TEXT("ResetPointPosition"));
 
@@ -146,6 +150,8 @@ static const FName ComputeAirDragForceName(TEXT("ComputeAirDragForce"));
 //------------------------------------------------------------------------------------------------------------
 
 static const FName NeedSimulationResetName(TEXT("NeedSimulationReset"));
+static const FName HasGlobalInterpolationName(TEXT("HasGlobalInterpolation"));
+static const FName HasKinematicsTargetName(TEXT("HasKinematicsTarget"));
 
 //------------------------------------------------------------------------------------------------------------
 
@@ -175,11 +181,15 @@ const FString UNiagaraDataInterfaceHairStrands::DeformedTrianglePositionAName(TE
 const FString UNiagaraDataInterfaceHairStrands::DeformedTrianglePositionBName(TEXT("DeformedTrianglePositionBBuffer_"));
 const FString UNiagaraDataInterfaceHairStrands::DeformedTrianglePositionCName(TEXT("DeformedTrianglePositionCBuffer_"));
 
+const FString UNiagaraDataInterfaceHairStrands::SampleCountName(TEXT("SampleCount_"));
+const FString UNiagaraDataInterfaceHairStrands::RestSamplePositionsName(TEXT("RestSamplePositionsBuffer_"));
+const FString UNiagaraDataInterfaceHairStrands::MeshSampleWeightsName(TEXT("MeshSampleWeightsBuffer_"));
+
 const FString UNiagaraDataInterfaceHairStrands::RestPositionOffsetName(TEXT("RestPositionOffset_"));
 const FString UNiagaraDataInterfaceHairStrands::DeformedPositionOffsetName(TEXT("DeformedPositionOffset_"));
 
-const FString UNiagaraDataInterfaceHairStrands::BoundingBoxBufferName(TEXT("BoundingBoxBuffer_"));
-const FString UNiagaraDataInterfaceHairStrands::NodeBoundBufferName(TEXT("NodeBoundBuffer_"));
+const FString UNiagaraDataInterfaceHairStrands::BoundingBoxBufferAName(TEXT("BoundingBoxBufferA_"));
+const FString UNiagaraDataInterfaceHairStrands::BoundingBoxBufferBName(TEXT("BoundingBoxBufferB_"));
 
 //------------------------------------------------------------------------------------------------------------
 
@@ -213,11 +223,15 @@ struct FNDIHairStrandsParametersName
 		DeformedTrianglePositionBName = UNiagaraDataInterfaceHairStrands::DeformedTrianglePositionBName + Suffix;
 		DeformedTrianglePositionCName = UNiagaraDataInterfaceHairStrands::DeformedTrianglePositionCName + Suffix;
 
+		SampleCountName = UNiagaraDataInterfaceHairStrands::SampleCountName + Suffix;
+		RestSamplePositionsName = UNiagaraDataInterfaceHairStrands::RestSamplePositionsName + Suffix;
+		MeshSampleWeightsName = UNiagaraDataInterfaceHairStrands::MeshSampleWeightsName + Suffix;
+
 		RestPositionOffsetName = UNiagaraDataInterfaceHairStrands::RestPositionOffsetName + Suffix;
 		DeformedPositionOffsetName = UNiagaraDataInterfaceHairStrands::DeformedPositionOffsetName + Suffix;
 
-		BoundingBoxBufferName = UNiagaraDataInterfaceHairStrands::BoundingBoxBufferName + Suffix;
-		NodeBoundBufferName = UNiagaraDataInterfaceHairStrands::NodeBoundBufferName + Suffix;
+		BoundingBoxBufferAName = UNiagaraDataInterfaceHairStrands::BoundingBoxBufferAName + Suffix;
+		BoundingBoxBufferBName = UNiagaraDataInterfaceHairStrands::BoundingBoxBufferBName + Suffix;
 	}
 
 	FString NumStrandsName;
@@ -246,74 +260,26 @@ struct FNDIHairStrandsParametersName
 	FString DeformedTrianglePositionBName;
 	FString DeformedTrianglePositionCName;
 
+	FString SampleCountName;
+	FString RestSamplePositionsName;
+	FString MeshSampleWeightsName;
+
 	FString RestPositionOffsetName;
 	FString DeformedPositionOffsetName;
 
-	FString BoundingBoxBufferName;
-	FString NodeBoundBufferName;
+	FString BoundingBoxBufferAName;
+	FString BoundingBoxBufferBName;
 };
 
 //------------------------------------------------------------------------------------------------------------
 
 #define NIAGARA_HAIR_STRANDS_THREAD_COUNT 64
 
-class FCopyBoundingBoxCS : public FGlobalShader
-{
-	DECLARE_GLOBAL_SHADER(FCopyBoundingBoxCS);
-	SHADER_USE_PARAMETER_STRUCT(FCopyBoundingBoxCS, FGlobalShader);
-
-	BEGIN_SHADER_PARAMETER_STRUCT(FParameters, )
-		SHADER_PARAMETER(uint32, NumElements)
-		SHADER_PARAMETER_UAV(RWBuffer, BoundingBoxBuffer)
-		SHADER_PARAMETER_UAV(RWBuffer, OutNodeBoundBuffer)
-	END_SHADER_PARAMETER_STRUCT()
-
-public:
-	static bool ShouldCompilePermutation(const FGlobalShaderPermutationParameters& Parameters)
-	{
-		return RHISupportsComputeShaders(Parameters.Platform);
-	}
-
-	static void ModifyCompilationEnvironment(const FGlobalShaderPermutationParameters& Parameters, FShaderCompilerEnvironment& OutEnvironment)
-	{
-		FGlobalShader::ModifyCompilationEnvironment(Parameters, OutEnvironment);
-		OutEnvironment.SetDefine(TEXT("THREAD_COUNT"), NIAGARA_HAIR_STRANDS_THREAD_COUNT);
-	}
-};
-
-IMPLEMENT_GLOBAL_SHADER(FCopyBoundingBoxCS, "/Plugin/Experimental/HairStrands/Private/NiagaraCopyBoundingBox.usf", "MainCS", SF_Compute);
-
-static void AddCopyBoundingBoxPass(
-	FRDGBuilder& GraphBuilder,
-	FRHIUnorderedAccessView* BoundingBoxBuffer,
-	FRHIUnorderedAccessView* OutNodeBoundBuffer)
-{
-	const uint32 GroupSize = NIAGARA_HAIR_STRANDS_THREAD_COUNT;
-	const uint32 NumElements = 1;
-
-	FCopyBoundingBoxCS::FParameters* Parameters = GraphBuilder.AllocParameters<FCopyBoundingBoxCS::FParameters>();
-	Parameters->BoundingBoxBuffer = BoundingBoxBuffer;
-	Parameters->OutNodeBoundBuffer = OutNodeBoundBuffer;
-	Parameters->NumElements = NumElements;
-
-	FGlobalShaderMap* ShaderMap = GetGlobalShaderMap(ERHIFeatureLevel::SM5);
-
-	const uint32 DispatchCount = FMath::DivideAndRoundUp(NumElements, GroupSize);
-
-	TShaderMapRef<FCopyBoundingBoxCS> ComputeShader(ShaderMap);
-	FComputeShaderUtils::AddPass(
-		GraphBuilder,
-		RDG_EVENT_NAME("CopyBoundingBox"),
-		ComputeShader,
-		Parameters,
-		FIntVector(DispatchCount, 1, 1));
-}
-
 //------------------------------------------------------------------------------------------------------------
 
-class FDirectCopyBoundingBoxCS : public FGlobalShader
+class FCopyBoundingBoxCS : public FGlobalShader
 {
-	DECLARE_SHADER_TYPE(FDirectCopyBoundingBoxCS, Global);
+	DECLARE_SHADER_TYPE(FCopyBoundingBoxCS, Global);
 
 public:
 
@@ -329,10 +295,10 @@ public:
 	}
 
 	/** Default constructor. */
-	FDirectCopyBoundingBoxCS() {}
+	FCopyBoundingBoxCS() {}
 
 	/** Initialization constructor. */
-	explicit FDirectCopyBoundingBoxCS(const ShaderMetaType::CompiledShaderInitializerType& Initializer);
+	explicit FCopyBoundingBoxCS(const ShaderMetaType::CompiledShaderInitializerType& Initializer);
 
 	/**
 	 * Set parameters.
@@ -355,17 +321,17 @@ private:
 	LAYOUT_FIELD(FShaderResourceParameter, OutputBoundingBox);
 };
 
-IMPLEMENT_SHADER_TYPE(, FDirectCopyBoundingBoxCS, TEXT("/Plugin/Experimental/HairStrands/Private/NiagaraCopyBoundingBox.usf"), TEXT("MainCS"), SF_Compute);
+IMPLEMENT_SHADER_TYPE(, FCopyBoundingBoxCS, TEXT("/Plugin/Experimental/HairStrands/Private/NiagaraCopyBoundingBox.usf"), TEXT("MainCS"), SF_Compute);
 
-FDirectCopyBoundingBoxCS::FDirectCopyBoundingBoxCS(const ShaderMetaType::CompiledShaderInitializerType& Initializer)
+FCopyBoundingBoxCS::FCopyBoundingBoxCS(const ShaderMetaType::CompiledShaderInitializerType& Initializer)
 	: FGlobalShader(Initializer)
 {
 	NumElements.Bind(Initializer.ParameterMap, TEXT("NumElements"));
-	InputBoundingBox.Bind(Initializer.ParameterMap, TEXT("BoundingBoxBuffer"));
-	OutputBoundingBox.Bind(Initializer.ParameterMap, TEXT("OutNodeBoundBuffer"));
+	InputBoundingBox.Bind(Initializer.ParameterMap, TEXT("BoundingBoxBufferA"));
+	OutputBoundingBox.Bind(Initializer.ParameterMap, TEXT("OutBoundingBoxBufferB"));
 }
 
-void FDirectCopyBoundingBoxCS::SetParameters(
+void FCopyBoundingBoxCS::SetParameters(
 	FRHICommandList& RHICmdList,
 	FRHIUnorderedAccessView* InInputBoundingBox,
 	FRHIUnorderedAccessView* InOutputBoundingBox,
@@ -385,17 +351,17 @@ void FDirectCopyBoundingBoxCS::SetParameters(
 	RHICmdList.TransitionResource(EResourceTransitionAccess::ERWBarrier, EResourceTransitionPipeline::EComputeToCompute, InOutputBoundingBox);
 }
 
-void FDirectCopyBoundingBoxCS::UnbindBuffers(FRHICommandList& RHICmdList)
+void FCopyBoundingBoxCS::UnbindBuffers(FRHICommandList& RHICmdList)
 {
 	FRHIComputeShader* ComputeShaderRHI = RHICmdList.GetBoundComputeShader();
 	SetUAVParameter(RHICmdList, ComputeShaderRHI, InputBoundingBox, nullptr);
 	SetUAVParameter(RHICmdList, ComputeShaderRHI, OutputBoundingBox, nullptr);
 }
 
-static void AddDirectCopyBoundingBoxPass(
+static void AddCopyBoundingBoxPass(
 	FRHICommandListImmediate& RHICmdList,
-	FRHIUnorderedAccessView* BoundingBoxBuffer,
-	FRHIUnorderedAccessView* OutNodeBoundBuffer)
+	FRHIUnorderedAccessView* BoundingBoxBufferA,
+	FRHIUnorderedAccessView* OutBoundingBoxBufferB)
 {
 	const uint32 GroupSize = NIAGARA_HAIR_STRANDS_THREAD_COUNT;
 	const uint32 NumElements = 1;
@@ -403,10 +369,10 @@ static void AddDirectCopyBoundingBoxPass(
 	FGlobalShaderMap* ShaderMap = GetGlobalShaderMap(ERHIFeatureLevel::SM5);
 	const uint32 DispatchCount = FMath::DivideAndRoundUp(NumElements, GroupSize);
 
-	TShaderMapRef<FDirectCopyBoundingBoxCS> ComputeShader(ShaderMap);
+	TShaderMapRef<FCopyBoundingBoxCS> ComputeShader(ShaderMap);
 	RHICmdList.SetComputeShader(ComputeShader.GetComputeShader());
 
-	ComputeShader->SetParameters(RHICmdList, BoundingBoxBuffer, OutNodeBoundBuffer, NumElements);
+	ComputeShader->SetParameters(RHICmdList, BoundingBoxBufferA, OutBoundingBoxBufferB, NumElements);
 
 	DispatchComputeShader(RHICmdList, ComputeShader, DispatchCount, 1, 1);
 	ComputeShader->UnbindBuffers(RHICmdList);
@@ -416,28 +382,17 @@ static void AddDirectCopyBoundingBoxPass(
 
 inline void ClearBuffer(FRHICommandList& RHICmdList, FNDIHairStrandsBuffer* HairStrandsBuffer)
 {
-	FRHIUnorderedAccessView* BoundingBoxBufferUAV = HairStrandsBuffer->BoundingBoxBuffer.UAV;
-	FRHIShaderResourceView* BoundingBoxBufferSRV = HairStrandsBuffer->BoundingBoxBuffer.SRV;
-	FRHIUnorderedAccessView* NodeBoundBufferUAV = HairStrandsBuffer->NodeBoundBuffer.UAV;
+	FRHIUnorderedAccessView* BoundingBoxBufferAUAV = HairStrandsBuffer->BoundingBoxBufferA.UAV;
+	FRHIShaderResourceView* BoundingBoxBufferASRV = HairStrandsBuffer->BoundingBoxBufferA.SRV;
+	FRHIUnorderedAccessView* BoundingBoxBufferBUAV = HairStrandsBuffer->BoundingBoxBufferB.UAV;
 
-	if (BoundingBoxBufferUAV != nullptr && BoundingBoxBufferSRV != nullptr && NodeBoundBufferUAV != nullptr)
+	if (BoundingBoxBufferAUAV != nullptr && BoundingBoxBufferASRV != nullptr && BoundingBoxBufferBUAV != nullptr)
 	{
 		ENQUEUE_RENDER_COMMAND(FNDIHairStrandsCopyBoundingBox)(
-			[BoundingBoxBufferUAV, NodeBoundBufferUAV, BoundingBoxBufferSRV, HairStrandsBuffer]
+			[BoundingBoxBufferAUAV, BoundingBoxBufferBUAV, BoundingBoxBufferASRV, HairStrandsBuffer]
 		(FRHICommandListImmediate& RHICmdListImm)
 		{
-		/*	FRDGBuilder GraphBuilder(RHICmdListImm);
-
-			RHICmdListImm.TransitionResource(EResourceTransitionAccess::ERWBarrier, EResourceTransitionPipeline::EComputeToCompute, BoundingBoxBufferUAV);
-			RHICmdListImm.TransitionResource(EResourceTransitionAccess::ERWBarrier, EResourceTransitionPipeline::EComputeToCompute, NodeBoundBufferUAV);
-
-			AddCopyBoundingBoxPass(
-				GraphBuilder,
-				BoundingBoxBufferUAV, NodeBoundBufferUAV);
-
-			GraphBuilder.Execute();*/
-
-			AddDirectCopyBoundingBoxPass(RHICmdListImm, HairStrandsBuffer->BoundingBoxBuffer.UAV, HairStrandsBuffer->NodeBoundBuffer.UAV);
+			AddCopyBoundingBoxPass(RHICmdListImm, HairStrandsBuffer->BoundingBoxBufferA.UAV, HairStrandsBuffer->BoundingBoxBufferB.UAV);
 		});
 	}
 }
@@ -473,24 +428,23 @@ void FNDIHairStrandsBuffer::InitRHI()
 			const uint32 BoundCount = 6;
 			const uint32 BoundBytes = sizeof(uint32)*BoundCount;
 
-			BoundingBoxBuffer.Initialize(sizeof(uint32), BoundCount, EPixelFormat::PF_R32_UINT, BUF_Static);
-			void* BoundBufferData = RHILockVertexBuffer(BoundingBoxBuffer.Buffer, 0, BoundBytes, RLM_WriteOnly);
+			BoundingBoxBufferA.Initialize(sizeof(uint32), BoundCount, EPixelFormat::PF_R32_UINT, BUF_Static);
+			void* BoundBufferData = RHILockVertexBuffer(BoundingBoxBufferA.Buffer, 0, BoundBytes, RLM_WriteOnly);
 			
 			FMemory::Memcpy(BoundBufferData, ZeroData.GetData(), BoundBytes);
-			RHIUnlockVertexBuffer(BoundingBoxBuffer.Buffer);
+			RHIUnlockVertexBuffer(BoundingBoxBufferA.Buffer);
 		}
-
 		{
 			static const TArray<uint32> ZeroData = { UINT_MAX,UINT_MAX,UINT_MAX,0,0,0 };
 
 			const uint32 BoundCount = 6;
 			const uint32 BoundBytes = sizeof(uint32)*BoundCount;
 
-			NodeBoundBuffer.Initialize(sizeof(uint32), BoundCount, EPixelFormat::PF_R32_UINT, BUF_Static);
-			void* BoundBufferData = RHILockVertexBuffer(NodeBoundBuffer.Buffer, 0, BoundBytes, RLM_WriteOnly);
+			BoundingBoxBufferB.Initialize(sizeof(uint32), BoundCount, EPixelFormat::PF_R32_UINT, BUF_Static);
+			void* BoundBufferData = RHILockVertexBuffer(BoundingBoxBufferB.Buffer, 0, BoundBytes, RLM_WriteOnly);
 
 			FMemory::Memcpy(BoundBufferData, ZeroData.GetData(), BoundBytes);
-			RHIUnlockVertexBuffer(NodeBoundBuffer.Buffer);
+			RHIUnlockVertexBuffer(BoundingBoxBufferB.Buffer);
 		}
 		if (SourceDeformedResources == nullptr)
 		{
@@ -504,8 +458,8 @@ void FNDIHairStrandsBuffer::ReleaseRHI()
 {
 	CurvesOffsetsBuffer.Release();
 	DeformedPositionBuffer.Release();
-	BoundingBoxBuffer.Release();
-	NodeBoundBuffer.Release();
+	BoundingBoxBufferA.Release();
+	BoundingBoxBufferB.Release();
 }
 
 //------------------------------------------------------------------------------------------------------------
@@ -546,7 +500,7 @@ bool FNDIHairStrandsData::Init(UNiagaraDataInterfaceHairStrands* Interface, FNia
 			WorldTransform = Interface->IsComponentValid() ? Interface->SourceComponent->GetComponentToWorld() :
 				SystemInstance->GetComponent()->GetComponentToWorld();
 
-			if (StrandsDatas != nullptr && GroomAsset != nullptr)
+			if (StrandsDatas != nullptr && GroomAsset != nullptr && GroomAsset->EnableSimulation)
 			{
 				StrandsSize = static_cast<uint8>(GroomAsset->StrandsSize);
 
@@ -600,7 +554,6 @@ bool FNDIHairStrandsData::Init(UNiagaraDataInterfaceHairStrands* Interface, FNia
 			}
 			TickCount = 0;
 			ForceReset = true;
-			ResetTick = MaxDelay;
 		}
 	}
 
@@ -612,7 +565,6 @@ bool FNDIHairStrandsData::Init(UNiagaraDataInterfaceHairStrands* Interface, FNia
 struct FNDIHairStrandsParametersCS : public FNiagaraDataInterfaceParametersCS
 {
 	DECLARE_TYPE_LAYOUT(FNDIHairStrandsParametersCS, NonVirtual);
-public:
 	void Bind(const FNiagaraDataInterfaceGPUParamInfo& ParameterInfo, const class FShaderParameterMap& ParameterMap)
 	{
 		FNDIHairStrandsParametersName ParamNames(*ParameterInfo.DataInterfaceHLSLSymbol);
@@ -647,8 +599,12 @@ public:
 		DeformedTrianglePositionBBuffer.Bind(ParameterMap, *ParamNames.DeformedTrianglePositionBName);
 		DeformedTrianglePositionCBuffer.Bind(ParameterMap, *ParamNames.DeformedTrianglePositionCName);
 
-		BoundingBoxBuffer.Bind(ParameterMap, *ParamNames.BoundingBoxBufferName);
-		NodeBoundBuffer.Bind(ParameterMap, *ParamNames.NodeBoundBufferName);
+		SampleCount.Bind(ParameterMap, *ParamNames.SampleCountName);
+		RestSamplePositionsBuffer.Bind(ParameterMap, *ParamNames.RestSamplePositionsName);
+		MeshSampleWeightsBuffer.Bind(ParameterMap, *ParamNames.MeshSampleWeightsName);
+
+		BoundingBoxBufferA.Bind(ParameterMap, *ParamNames.BoundingBoxBufferAName);
+		BoundingBoxBufferB.Bind(ParameterMap, *ParamNames.BoundingBoxBufferBName);
 
 		if (!DeformedPositionBuffer.IsBound())
 		{
@@ -693,19 +649,19 @@ public:
 		{
 			UE_LOG(LogHairStrands, Warning, TEXT("Binding failed for FNDIHairStrandsParametersCS %s. Was it optimized out?"), *ParamNames.DeformedTrianglePositionCName)
 		}
-		if (!BoundingBoxBuffer.IsBound())
+		if (!BoundingBoxBufferA.IsBound())
 		{
-			UE_LOG(LogHairStrands, Warning, TEXT("Binding failed for FNDIHairStrandsParametersCS %s. Was it optimized out?"), *ParamNames.BoundingBoxBufferName)
+			UE_LOG(LogHairStrands, Warning, TEXT("Binding failed for FNDIHairStrandsParametersCS %s. Was it optimized out?"), *ParamNames.BoundingBoxBufferAName)
 		}
-		if (!NodeBoundBuffer.IsBound())
+		if (!BoundingBoxBufferB.IsBound())
 		{
-			UE_LOG(LogHairStrands, Warning, TEXT("Binding failed for FNDIHairStrandsParametersCS %s. Was it optimized out?"), *ParamNames.BoundingBoxBufferName)
+			UE_LOG(LogHairStrands, Warning, TEXT("Binding failed for FNDIHairStrandsParametersCS %s. Was it optimized out?"), *ParamNames.BoundingBoxBufferAName)
 		}
 	}
 
 	void Set(FRHICommandList& RHICmdList, const FNiagaraDataInterfaceSetArgs& Context) const
 	{
-		check(IsInRenderingThread());
+		check(IsInRenderingThread()); 
 
 		FRHIComputeShader* ComputeShaderRHI = RHICmdList.GetBoundComputeShader();
 
@@ -724,7 +680,7 @@ public:
 			FNDIHairStrandsBuffer* HairStrandsBuffer = ProxyData->HairStrandsBuffer;
 
 			FUnorderedAccessViewRHIRef PointPositionsUAV = IsDeformedValid ?
-				HairStrandsBuffer->SourceDeformedResources->DeformedPositionBuffer[0].UAV : HairStrandsBuffer->DeformedPositionBuffer.UAV;
+				HairStrandsBuffer->SourceDeformedResources->DeformedPositionBuffer[HairStrandsBuffer->SourceDeformedResources->CurrentIndex].UAV : HairStrandsBuffer->DeformedPositionBuffer.UAV;
 
 			const int32 HasRootAttachedValue = (IsRootValid && HairStrandsBuffer->SourceRootResources->MeshProjectionLODs.Num() > 0 && 
 				HairStrandsBuffer->SourceRootResources->MeshProjectionLODs[0].Status == FHairStrandsProjectionHairData::LODData::EStatus::Completed);
@@ -745,10 +701,19 @@ public:
 			FRHIShaderResourceView* DeformedTrianglePositionCSRV = (HasRootAttachedValue == 1) ?
 				MeshProjection->DeformedRootTrianglePosition2Buffer.SRV.GetReference() : FNiagaraRenderer::GetDummyFloatBuffer();
 
+			const int32 SampleCountValue = (MeshProjection != nullptr) ? MeshProjection->SampleCount : 0;
+			FShaderResourceViewRHIRef RestSamplePositionsBufferSRV = (MeshProjection != nullptr && MeshProjection->SampleCount > 0) ?
+				MeshProjection->RestSamplePositionsBuffer.SRV : FNiagaraRenderer::GetDummyFloatBuffer();
+			FShaderResourceViewRHIRef MeshSampleWeightsBufferSRV = (MeshProjection != nullptr && MeshProjection->SampleCount > 0) ?
+				MeshProjection->MeshSampleWeightsBuffer.SRV : FNiagaraRenderer::GetDummyFloatBuffer();
+
 			FRHIShaderResourceView* RootBarycentricCoordinatesSRV = (HasRootAttachedValue == 1) ?
 				MeshProjection->RootTriangleBarycentricBuffer.SRV.GetReference() : FNiagaraRenderer::GetDummyFloatBuffer();
 
-			int32 bNeedSimReset = (ProxyData->TickCount <= ProxyData->ResetTick ? 1 : 0);
+
+			int32 bNeedSimReset = (ProxyData->TickCount <= MaxDelay ? 1 : 0);
+
+			//UE_LOG(LogHairStrands, Log, TEXT("Shader Reset : %d %d %d"), bNeedSimReset, IsRootValid, HasRootAttachedValue);
 			FVector RestRootOffsetValue = FVector::ZeroVector;
 			FVector DeformedRootOffsetValue = FVector::ZeroVector;
 			FVector DeformedPositionOffsetValue =  IsDeformedValid ? HairStrandsBuffer->SourceDeformedResources->PositionOffset :
@@ -756,10 +721,10 @@ public:
 			FVector RestPositionOffsetValue = ProxyData->HairStrandsBuffer->SourceRestResources->PositionOffset;
 
 			SetUAVParameter(RHICmdList, ComputeShaderRHI, DeformedPositionBuffer, PointPositionsUAV);
-			RHICmdList.TransitionResource(EResourceTransitionAccess::ERWBarrier, EResourceTransitionPipeline::EComputeToCompute, HairStrandsBuffer->BoundingBoxBuffer.UAV);
-			SetUAVParameter(RHICmdList, ComputeShaderRHI, BoundingBoxBuffer, HairStrandsBuffer->BoundingBoxBuffer.UAV);
-			RHICmdList.TransitionResource(EResourceTransitionAccess::ERWBarrier, EResourceTransitionPipeline::EComputeToCompute, HairStrandsBuffer->NodeBoundBuffer.UAV);
-			SetSRVParameter(RHICmdList, ComputeShaderRHI, NodeBoundBuffer, HairStrandsBuffer->NodeBoundBuffer.SRV);
+			RHICmdList.TransitionResource(EResourceTransitionAccess::ERWBarrier, EResourceTransitionPipeline::EComputeToCompute, HairStrandsBuffer->BoundingBoxBufferA.UAV);
+			SetUAVParameter(RHICmdList, ComputeShaderRHI, BoundingBoxBufferA, HairStrandsBuffer->BoundingBoxBufferA.UAV);
+			RHICmdList.TransitionResource(EResourceTransitionAccess::ERWBarrier, EResourceTransitionPipeline::EComputeToCompute, HairStrandsBuffer->BoundingBoxBufferB.UAV);
+			SetSRVParameter(RHICmdList, ComputeShaderRHI, BoundingBoxBufferB, HairStrandsBuffer->BoundingBoxBufferB.SRV);
 			SetSRVParameter(RHICmdList, ComputeShaderRHI, CurvesOffsetsBuffer, HairStrandsBuffer->CurvesOffsetsBuffer.SRV);
 			SetSRVParameter(RHICmdList, ComputeShaderRHI, RestPositionBuffer, HairStrandsBuffer->SourceRestResources->RestPositionBuffer.SRV);
 
@@ -787,13 +752,17 @@ public:
 			SetSRVParameter(RHICmdList, ComputeShaderRHI, DeformedTrianglePositionBBuffer, DeformedTrianglePositionBSRV);
 			SetSRVParameter(RHICmdList, ComputeShaderRHI, DeformedTrianglePositionCBuffer, DeformedTrianglePositionCSRV);
 
+			SetShaderValue(RHICmdList, ComputeShaderRHI, SampleCount, SampleCountValue);
+			SetSRVParameter(RHICmdList, ComputeShaderRHI, RestSamplePositionsBuffer, RestSamplePositionsBufferSRV);
+			SetSRVParameter(RHICmdList, ComputeShaderRHI, MeshSampleWeightsBuffer, MeshSampleWeightsBufferSRV);
+
 			SetSRVParameter(RHICmdList, ComputeShaderRHI, RootBarycentricCoordinatesBuffer, RootBarycentricCoordinatesSRV);
 		}
 		else
 		{
 			SetUAVParameter(RHICmdList, ComputeShaderRHI, DeformedPositionBuffer, Context.Batcher->GetEmptyRWBufferFromPool(RHICmdList, PF_R32_FLOAT));
-			SetUAVParameter(RHICmdList, ComputeShaderRHI, BoundingBoxBuffer, Context.Batcher->GetEmptyRWBufferFromPool(RHICmdList, PF_R32_UINT));
-			SetSRVParameter(RHICmdList, ComputeShaderRHI, NodeBoundBuffer, FNiagaraRenderer::GetDummyUIntBuffer());
+			SetUAVParameter(RHICmdList, ComputeShaderRHI, BoundingBoxBufferA, Context.Batcher->GetEmptyRWBufferFromPool(RHICmdList, PF_R32_UINT));
+			SetSRVParameter(RHICmdList, ComputeShaderRHI, BoundingBoxBufferB, FNiagaraRenderer::GetDummyUIntBuffer());
 			SetSRVParameter(RHICmdList, ComputeShaderRHI, CurvesOffsetsBuffer, FNiagaraRenderer::GetDummyUIntBuffer());
 			SetSRVParameter(RHICmdList, ComputeShaderRHI, RestPositionBuffer, FNiagaraRenderer::GetDummyFloatBuffer());
 
@@ -821,6 +790,10 @@ public:
 			SetSRVParameter(RHICmdList, ComputeShaderRHI, DeformedTrianglePositionBBuffer, FNiagaraRenderer::GetDummyFloatBuffer());
 			SetSRVParameter(RHICmdList, ComputeShaderRHI, DeformedTrianglePositionCBuffer, FNiagaraRenderer::GetDummyFloatBuffer());
 
+			SetShaderValue(RHICmdList, ComputeShaderRHI, SampleCount, 0);
+			SetSRVParameter(RHICmdList, ComputeShaderRHI, RestSamplePositionsBuffer, FNiagaraRenderer::GetDummyFloatBuffer());
+			SetSRVParameter(RHICmdList, ComputeShaderRHI, MeshSampleWeightsBuffer, FNiagaraRenderer::GetDummyFloatBuffer());
+
 			SetSRVParameter(RHICmdList, ComputeShaderRHI, RootBarycentricCoordinatesBuffer, FNiagaraRenderer::GetDummyFloatBuffer());
 		}
 	}
@@ -829,7 +802,7 @@ public:
 	{
 		FRHIComputeShader* ShaderRHI = RHICmdList.GetBoundComputeShader();
 		SetUAVParameter(RHICmdList, ShaderRHI, DeformedPositionBuffer, nullptr);
-		SetUAVParameter(RHICmdList, ShaderRHI, BoundingBoxBuffer, nullptr);
+		SetUAVParameter(RHICmdList, ShaderRHI, BoundingBoxBufferA, nullptr);
 	}
 
 private:
@@ -861,11 +834,15 @@ private:
 	LAYOUT_FIELD(FShaderResourceParameter, DeformedTrianglePositionBBuffer);
 	LAYOUT_FIELD(FShaderResourceParameter, DeformedTrianglePositionCBuffer);
 
+	LAYOUT_FIELD(FShaderParameter, SampleCount);
+	LAYOUT_FIELD(FShaderResourceParameter, RestSamplePositionsBuffer);
+	LAYOUT_FIELD(FShaderResourceParameter, MeshSampleWeightsBuffer);
+
 	LAYOUT_FIELD(FShaderParameter, RestPositionOffset);
 	LAYOUT_FIELD(FShaderParameter, DeformedPositionOffset);
 
-	LAYOUT_FIELD(FShaderResourceParameter, BoundingBoxBuffer);
-	LAYOUT_FIELD(FShaderResourceParameter, NodeBoundBuffer);
+	LAYOUT_FIELD(FShaderResourceParameter, BoundingBoxBufferA);
+	LAYOUT_FIELD(FShaderResourceParameter, BoundingBoxBufferB);
 };
 
 IMPLEMENT_TYPE_LAYOUT(FNDIHairStrandsParametersCS);
@@ -1036,22 +1013,10 @@ bool UNiagaraDataInterfaceHairStrands::PerInstanceTick(void* PerInstanceData, FN
 
 	if (SourceComponent != nullptr)
 	{
-		if (!InstanceData->ForceReset && !SourceComponent->bResetSimulation && (InstanceData->ResetTick == MaxDelay))
-		{
-			InstanceData->ResetTick = FMath::Min(MaxDelay, InstanceData->TickCount+1);
-		}
-
-		//if (!InstanceData->ForceReset && !SourceComponent->bResetSimulation)
-		if (!SourceComponent->bResetSimulation)
-		{
-			InstanceData->TickCount = FMath::Min(MaxDelay + 1, InstanceData->TickCount + 1);
-		}
-		else
+		if (SourceComponent->bResetSimulation)
 		{
 			InstanceData->TickCount = 0;
 		}
-		InstanceData->ResetTick = MaxDelay;
-		//UE_LOG(LogHairStrands, Warning, TEXT("Reset Simulation : %d %d %d %d"), InstanceData->ForceReset, SourceComponent->bResetSimulation, InstanceData->TickCount, InstanceData->ResetTick);
 		InstanceData->ForceReset = SourceComponent->bResetSimulation;
 	}
 
@@ -1635,6 +1600,17 @@ void UNiagaraDataInterfaceHairStrands::GetFunctions(TArray<FNiagaraFunctionSigna
 
 		OutFunctions.Add(Sig);
 	}
+	{
+		FNiagaraFunctionSignature Sig;
+		Sig.Name = EvalSkinnedPositionName;
+		Sig.bMemberFunction = true;
+		Sig.bRequiresContext = false;
+		Sig.Inputs.Add(FNiagaraVariable(FNiagaraTypeDefinition(GetClass()), TEXT("Hair Strands")));
+		Sig.Inputs.Add(FNiagaraVariable(FNiagaraTypeDefinition::GetVec3Def(), TEXT("Rest Position")));
+		Sig.Outputs.Add(FNiagaraVariable(FNiagaraTypeDefinition::GetVec3Def(), TEXT("Skinned Position")));
+
+		OutFunctions.Add(Sig);
+	}
 	{
 		FNiagaraFunctionSignature Sig;
 		Sig.Name = UpdatePointPositionName;
@@ -2040,6 +2016,27 @@ void UNiagaraDataInterfaceHairStrands::GetFunctions(TArray<FNiagaraFunctionSigna
 
 		OutFunctions.Add(Sig);
 	}
+	{
+		FNiagaraFunctionSignature Sig;
+		Sig.Name = HasGlobalInterpolationName;
+		Sig.bMemberFunction = true;
+		Sig.bRequiresContext = false;
+		Sig.Inputs.Add(FNiagaraVariable(FNiagaraTypeDefinition(GetClass()), TEXT("Hair Strands")));
+		Sig.Outputs.Add(FNiagaraVariable(FNiagaraTypeDefinition::GetBoolDef(), TEXT("Globa lInterpolation")));
+
+		OutFunctions.Add(Sig);
+	}
+	{
+		FNiagaraFunctionSignature Sig;
+		Sig.Name = HasKinematicsTargetName;
+		Sig.bMemberFunction = true;
+		Sig.bRequiresContext = false;
+		Sig.Inputs.Add(FNiagaraVariable(FNiagaraTypeDefinition(GetClass()), TEXT("Hair Strands")));
+		Sig.Outputs.Add(FNiagaraVariable(FNiagaraTypeDefinition::GetBoolDef(), TEXT("Kinematics Target")));
+
+		OutFunctions.Add(Sig);
+	}
+
 }
 
 DEFINE_NDI_DIRECT_FUNC_BINDER(UNiagaraDataInterfaceHairStrands, GetNumStrands);
@@ -2094,6 +2091,7 @@ DEFINE_NDI_DIRECT_FUNC_BINDER(UNiagaraDataInterfaceHairStrands, AttachNodeOrient
 DEFINE_NDI_DIRECT_FUNC_BINDER(UNiagaraDataInterfaceHairStrands, AttachNodeState);
 DEFINE_NDI_DIRECT_FUNC_BINDER(UNiagaraDataInterfaceHairStrands, UpdatePointPosition);
 DEFINE_NDI_DIRECT_FUNC_BINDER(UNiagaraDataInterfaceHairStrands, ResetPointPosition);
+DEFINE_NDI_DIRECT_FUNC_BINDER(UNiagaraDataInterfaceHairStrands, EvalSkinnedPosition);
 
 DEFINE_NDI_DIRECT_FUNC_BINDER(UNiagaraDataInterfaceHairStrands, GetBoxExtent);
 DEFINE_NDI_DIRECT_FUNC_BINDER(UNiagaraDataInterfaceHairStrands, GetBoxCenter);
@@ -2132,6 +2130,8 @@ DEFINE_NDI_DIRECT_FUNC_BINDER(UNiagaraDataInterfaceHairStrands, UpdateMaterialFr
 DEFINE_NDI_DIRECT_FUNC_BINDER(UNiagaraDataInterfaceHairStrands, ComputeMaterialFrame);
 DEFINE_NDI_DIRECT_FUNC_BINDER(UNiagaraDataInterfaceHairStrands, ComputeAirDragForce);
 DEFINE_NDI_DIRECT_FUNC_BINDER(UNiagaraDataInterfaceHairStrands, NeedSimulationReset);
+DEFINE_NDI_DIRECT_FUNC_BINDER(UNiagaraDataInterfaceHairStrands, HasGlobalInterpolation);
+DEFINE_NDI_DIRECT_FUNC_BINDER(UNiagaraDataInterfaceHairStrands, HasKinematicsTarget);
 
 void UNiagaraDataInterfaceHairStrands::GetVMExternalFunction(const FVMExternalFunctionBindingInfo& BindingInfo, void* InstanceData, FVMExternalFunction &OutFunc)
 {
@@ -2350,6 +2350,11 @@ void UNiagaraDataInterfaceHairStrands::GetVMExternalFunction(const FVMExternalFu
 		check(BindingInfo.GetNumInputs() == 9 && BindingInfo.GetNumOutputs() == 7);
 		NDI_FUNC_BINDER(UNiagaraDataInterfaceHairStrands, AttachNodeState)::Bind(this, OutFunc);
 	}
+	else if (BindingInfo.Name == EvalSkinnedPositionName)
+	{
+		check(BindingInfo.GetNumInputs() == 4 && BindingInfo.GetNumOutputs() == 3);
+		NDI_FUNC_BINDER(UNiagaraDataInterfaceHairStrands, EvalSkinnedPosition)::Bind(this, OutFunc);
+	}
 	else if (BindingInfo.Name == UpdatePointPositionName)
 	{
 		check(BindingInfo.GetNumInputs() == 5 && BindingInfo.GetNumOutputs() == 1);
@@ -2505,6 +2510,16 @@ void UNiagaraDataInterfaceHairStrands::GetVMExternalFunction(const FVMExternalFu
 		check(BindingInfo.GetNumInputs() == 1 && BindingInfo.GetNumOutputs() == 1);
 		NDI_FUNC_BINDER(UNiagaraDataInterfaceHairStrands, NeedSimulationReset)::Bind(this, OutFunc);
 	}
+	else if (BindingInfo.Name == HasGlobalInterpolationName)
+	{
+		check(BindingInfo.GetNumInputs() == 1 && BindingInfo.GetNumOutputs() == 1);
+		NDI_FUNC_BINDER(UNiagaraDataInterfaceHairStrands, HasGlobalInterpolation)::Bind(this, OutFunc);
+	}
+	else if (BindingInfo.Name == HasKinematicsTargetName)
+	{
+		check(BindingInfo.GetNumInputs() == 1 && BindingInfo.GetNumOutputs() == 1);
+		NDI_FUNC_BINDER(UNiagaraDataInterfaceHairStrands, HasKinematicsTarget)::Bind(this, OutFunc);
+	}
 }
 
 void WriteTransform(const FMatrix& ToWrite, FVectorVMContext& Context)
@@ -3003,6 +3018,11 @@ void UNiagaraDataInterfaceHairStrands::AttachNodePosition(FVectorVMContext& Cont
 	// @todo : implement function for cpu 
 }
 
+void UNiagaraDataInterfaceHairStrands::EvalSkinnedPosition(FVectorVMContext& Context)
+{
+	// @todo : implement function for cpu 
+}
+
 void UNiagaraDataInterfaceHairStrands::AttachNodeOrientation(FVectorVMContext& Context)
 {
 	// @todo : implement function for cpu 
@@ -3143,6 +3163,34 @@ void UNiagaraDataInterfaceHairStrands::NeedSimulationReset(FVectorVMContext& Con
 	// @todo : implement function for cpu 
 }
 
+void UNiagaraDataInterfaceHairStrands::HasGlobalInterpolation(FVectorVMContext& Context)
+{
+	static const auto GlobalInterpolationCVar = IConsoleManager::Get().FindConsoleVariable(TEXT("r.HairStrands.Interpolation.Global"));
+	const int32 GlobalInterpolation = GlobalInterpolationCVar->GetInt();
+
+	VectorVM::FUserPtrHandler<FNDIHairStrandsData> InstData(Context);
+	VectorVM::FExternalFuncRegisterHandler<int32> OutGlobalInterpolation(Context);
+
+	for (int32 InstanceIdx = 0; InstanceIdx < Context.NumInstances; ++InstanceIdx)
+	{
+		*OutGlobalInterpolation.GetDestAndAdvance() = GlobalInterpolation;
+	}
+}
+
+void UNiagaraDataInterfaceHairStrands::HasKinematicsTarget(FVectorVMContext& Context)
+{
+	static const auto KinematicsTargetCVar = IConsoleManager::Get().FindConsoleVariable(TEXT("r.HairStrands.Interpolation.Target"));
+	const int32 KinematicsTarget = KinematicsTargetCVar->GetInt();
+
+	VectorVM::FUserPtrHandler<FNDIHairStrandsData> InstData(Context);
+	VectorVM::FExternalFuncRegisterHandler<int32> OutKinematicsTarget(Context);
+
+	for (int32 InstanceIdx = 0; InstanceIdx < Context.NumInstances; ++InstanceIdx)
+	{
+		*OutKinematicsTarget.GetDestAndAdvance() = KinematicsTarget;
+	}
+}
+
 bool UNiagaraDataInterfaceHairStrands::GetFunctionHLSL(const FNiagaraDataInterfaceGPUParamInfo& ParamInfo, const FNiagaraDataInterfaceGeneratedFunction& FunctionInfo, int FunctionInstanceIndex, FString& OutHLSL)
 {
 	FNDIHairStrandsParametersName ParamNames(ParamInfo.DataInterfaceHLSLSymbol);
@@ -3336,6 +3384,17 @@ bool UNiagaraDataInterfaceHairStrands::GetFunctionHLSL(const FNiagaraDataInterfa
 		OutHLSL += FString::Format(FormatSample, ArgsSample);
 		return true;
 	}
+	else if (FunctionInfo.DefinitionName == EvalSkinnedPositionName)
+	{
+		static const TCHAR* FormatSample = TEXT(R"(
+			void {InstanceFunctionName} (in float3 RestPosition, out float3 SkinnedPosition)
+			{
+				{HairStrandsContextName} DIHairStrands_EvalSkinnedPosition(DIContext,RestPosition,SkinnedPosition);
+			}
+			)");
+		OutHLSL += FString::Format(FormatSample, ArgsSample);
+		return true;
+	}
 	else if (FunctionInfo.DefinitionName == AttachNodeOrientationName)
 	{
 		static const TCHAR *FormatSample = TEXT(R"(
@@ -3357,7 +3416,7 @@ bool UNiagaraDataInterfaceHairStrands::GetFunctionHLSL(const FNiagaraDataInterfa
 				)");
 		OutHLSL += FString::Format(FormatSample, ArgsSample);
 		return true;
-	}	
+	}
 	else if (FunctionInfo.DefinitionName == UpdatePointPositionName)
 	{
 		static const TCHAR *FormatSample = TEXT(R"(
diff --git a/Engine/Plugins/Experimental/HairStrands/Source/HairStrandsNiagara/Private/NiagaraDataInterfacePhysicsAsset.cpp b/Engine/Plugins/Experimental/HairStrands/Source/HairStrandsNiagara/Private/NiagaraDataInterfacePhysicsAsset.cpp
--- a/Engine/Plugins/Experimental/HairStrands/Source/HairStrandsNiagara/Private/NiagaraDataInterfacePhysicsAsset.cpp
+++ b/Engine/Plugins/Experimental/HairStrands/Source/HairStrandsNiagara/Private/NiagaraDataInterfacePhysicsAsset.cpp
@@ -47,7 +47,24 @@ const FString UNiagaraDataInterfacePhysicsAsset::BoxExtentName(TEXT("BoxExtent_"
 
 struct FNDIPhysicsAssetParametersName
 {
+	FNDIPhysicsAssetParametersName(const FString& Suffix)
+	{
+		ElementOffsetsName = UNiagaraDataInterfacePhysicsAsset::ElementOffsetsName + Suffix;
+
+		CurrentTransformBufferName = UNiagaraDataInterfacePhysicsAsset::CurrentTransformBufferName + Suffix;
+		PreviousTransformBufferName = UNiagaraDataInterfacePhysicsAsset::PreviousTransformBufferName + Suffix;
+		PreviousInverseBufferName = UNiagaraDataInterfacePhysicsAsset::PreviousInverseBufferName + Suffix;
+		InverseTransformBufferName = UNiagaraDataInterfacePhysicsAsset::InverseTransformBufferName + Suffix;
+		RestTransformBufferName = UNiagaraDataInterfacePhysicsAsset::RestTransformBufferName + Suffix;
+		RestInverseBufferName = UNiagaraDataInterfacePhysicsAsset::RestInverseBufferName + Suffix;
+		ElementExtentBufferName = UNiagaraDataInterfacePhysicsAsset::ElementExtentBufferName + Suffix;
+
+		BoxOriginName = UNiagaraDataInterfacePhysicsAsset::BoxOriginName + Suffix;
+		BoxExtentName = UNiagaraDataInterfacePhysicsAsset::BoxExtentName + Suffix;
+	}
+
 	FString ElementOffsetsName;
+	
 	FString CurrentTransformBufferName;
 	FString PreviousTransformBufferName;
 	FString PreviousInverseBufferName;
@@ -60,20 +77,6 @@ struct FNDIPhysicsAssetParametersName
 	FString BoxExtentName;
 };
 
-static void GetNiagaraDataInterfaceParametersName(FNDIPhysicsAssetParametersName& Names, const FString& Suffix)
-{
-	Names.ElementOffsetsName = UNiagaraDataInterfacePhysicsAsset::ElementOffsetsName + Suffix;
-	Names.CurrentTransformBufferName = UNiagaraDataInterfacePhysicsAsset::CurrentTransformBufferName + Suffix;
-	Names.PreviousTransformBufferName = UNiagaraDataInterfacePhysicsAsset::PreviousTransformBufferName + Suffix;
-	Names.InverseTransformBufferName = UNiagaraDataInterfacePhysicsAsset::InverseTransformBufferName + Suffix;
-	Names.RestTransformBufferName = UNiagaraDataInterfacePhysicsAsset::RestTransformBufferName + Suffix;
-	Names.RestInverseBufferName = UNiagaraDataInterfacePhysicsAsset::RestInverseBufferName + Suffix;
-	Names.ElementExtentBufferName = UNiagaraDataInterfacePhysicsAsset::ElementExtentBufferName + Suffix;
-
-	Names.BoxOriginName = UNiagaraDataInterfacePhysicsAsset::BoxOriginName + Suffix;
-	Names.BoxExtentName = UNiagaraDataInterfacePhysicsAsset::BoxExtentName + Suffix;
-}
-
 //------------------------------------------------------------------------------------------------------------
 
 template<typename BufferType, typename DataType, int ElementSize, EPixelFormat PixelFormat, bool InitBuffer>
@@ -303,28 +306,6 @@ void FNDIPhysicsAssetBuffer::Update()
 	}
 }
 
-struct FPhysicsAssetManager
-{
-	TMap<int32,FSkeletalMeshObject*> Elements;
-};
-
-FPhysicsAssetManager GHairManager;
-
-void MeshObjectCallback(
-	FSkeletalMeshObjectCallbackData::EEventType Event,
-	FSkeletalMeshObject* MeshObject,
-	uint64 UserData)
-{
-	ENQUEUE_RENDER_COMMAND(FPhysicsAssetUpdate)(
-		[Event, MeshObject, UserData](FRHICommandListImmediate& RHICmdList)
-	{
-		if (Event == FSkeletalMeshObjectCallbackData::EEventType::Register || Event == FSkeletalMeshObjectCallbackData::EEventType::Update)
-		{
-			GHairManager.Elements.Add(UserData,MeshObject);
-		}
-	});
-}
-
 void FNDIPhysicsAssetBuffer::InitRHI()
 {
 	if (IsValid())
@@ -337,8 +318,8 @@ void FNDIPhysicsAssetBuffer::InitRHI()
 		CreateInternalBuffer<FVector4, FVector4, 1, EPixelFormat::PF_A32B32G32R32F, true>(AssetArrays->ElementExtent.Num(), AssetArrays->ElementExtent, ElementExtentBuffer);
 		CreateInternalBuffer<FVector4, FVector4, 1, EPixelFormat::PF_A32B32G32R32F, true>(AssetArrays->PreviousInverse.Num(), AssetArrays->PreviousInverse, PreviousInverseBuffer);
 
-		//UE_LOG(LogPhysicsAsset, Warning, TEXT("Num Capsules = %d | Num Spheres = %d | Num Boxes = %d"), AssetArrays->ElementOffsets.NumElements - AssetArrays->ElementOffsets.CapsuleOffset,
-		//	AssetArrays->ElementOffsets.CapsuleOffset - AssetArrays->ElementOffsets.SphereOffset, AssetArrays->ElementOffsets.SphereOffset - AssetArrays->ElementOffsets.BoxOffset);
+		UE_LOG(LogPhysicsAsset, Warning, TEXT("Num Capsules = %d | Num Spheres = %d | Num Boxes = %d"), AssetArrays->ElementOffsets.NumElements - AssetArrays->ElementOffsets.CapsuleOffset,
+			AssetArrays->ElementOffsets.CapsuleOffset - AssetArrays->ElementOffsets.SphereOffset, AssetArrays->ElementOffsets.SphereOffset - AssetArrays->ElementOffsets.BoxOffset);
 	}
 
 	// /*const FPrimitiveComponentId LocalComponentId = SkeletalMesh->ComponentId;
@@ -419,8 +400,7 @@ struct FNDIPhysicsAssetParametersCS : public FNiagaraDataInterfaceParametersCS
 public:
 	void Bind(const FNiagaraDataInterfaceGPUParamInfo& ParameterInfo, const class FShaderParameterMap& ParameterMap)
 	{
-		FNDIPhysicsAssetParametersName ParamNames;
-		GetNiagaraDataInterfaceParametersName(ParamNames, *ParameterInfo.DataInterfaceHLSLSymbol);
+		FNDIPhysicsAssetParametersName ParamNames(*ParameterInfo.DataInterfaceHLSLSymbol);
 
 		ElementOffsets.Bind(ParameterMap, *ParamNames.ElementOffsetsName);
 
@@ -532,6 +512,7 @@ IMPLEMENT_TYPE_LAYOUT(FNDIPhysicsAssetParametersCS);
 
 IMPLEMENT_NIAGARA_DI_PARAMETER(UNiagaraDataInterfacePhysicsAsset, FNDIPhysicsAssetParametersCS);
 
+
 //------------------------------------------------------------------------------------------------------------
 
 void FNDIPhysicsAssetProxy::ConsumePerInstanceDataFromGameThread(void* PerInstanceData, const FNiagaraSystemInstanceID& Instance)
@@ -879,8 +860,7 @@ void UNiagaraDataInterfacePhysicsAsset::GetProjectionPoint(FVectorVMContext& Con
 
 bool UNiagaraDataInterfacePhysicsAsset::GetFunctionHLSL(const FNiagaraDataInterfaceGPUParamInfo& ParamInfo, const FNiagaraDataInterfaceGeneratedFunction& FunctionInfo, int FunctionInstanceIndex, FString& OutHLSL)
 {
-	FNDIPhysicsAssetParametersName ParamNames;
-	GetNiagaraDataInterfaceParametersName(ParamNames, ParamInfo.DataInterfaceHLSLSymbol);
+	FNDIPhysicsAssetParametersName ParamNames(ParamInfo.DataInterfaceHLSLSymbol);
 
 	TMap<FString, FStringFormatArg> ArgsSample = {
 		{TEXT("InstanceFunctionName"), FunctionInfo.InstanceName},
diff --git a/Engine/Plugins/Experimental/HairStrands/Source/HairStrandsNiagara/Private/NiagaraDataInterfacePressureGrid.cpp b/Engine/Plugins/Experimental/HairStrands/Source/HairStrandsNiagara/Private/NiagaraDataInterfacePressureGrid.cpp
--- a/Engine/Plugins/Experimental/HairStrands/Source/HairStrandsNiagara/Private/NiagaraDataInterfacePressureGrid.cpp
+++ b/Engine/Plugins/Experimental/HairStrands/Source/HairStrandsNiagara/Private/NiagaraDataInterfacePressureGrid.cpp
@@ -46,6 +46,18 @@ const FString UNiagaraDataInterfacePressureGrid::WorldInverseName(TEXT("WorldInv
 
 struct FNDIPressureGridParametersName
 {
+	FNDIPressureGridParametersName(const FString& Suffix)
+	{
+		GridCurrentBufferName = UNiagaraDataInterfacePressureGrid::GridCurrentBufferName + Suffix;
+		GridDestinationBufferName = UNiagaraDataInterfacePressureGrid::GridDestinationBufferName + Suffix;
+
+		GridSizeName = UNiagaraDataInterfacePressureGrid::GridSizeName + Suffix;
+		GridOriginName = UNiagaraDataInterfacePressureGrid::GridOriginName + Suffix;
+
+		WorldTransformName = UNiagaraDataInterfacePressureGrid::WorldTransformName + Suffix;
+		WorldInverseName = UNiagaraDataInterfacePressureGrid::WorldInverseName + Suffix;
+	}
+
 	FString GridCurrentBufferName;
 	FString GridDestinationBufferName;
 
@@ -56,16 +68,6 @@ struct FNDIPressureGridParametersName
 	FString WorldInverseName;
 };
 
-static void GetNiagaraDataInterfaceParametersName(FNDIPressureGridParametersName& Names, const FString& Suffix)
-{
-	Names.GridCurrentBufferName = UNiagaraDataInterfacePressureGrid::GridCurrentBufferName + Suffix;
-	Names.GridDestinationBufferName = UNiagaraDataInterfacePressureGrid::GridDestinationBufferName + Suffix;
-	Names.GridSizeName = UNiagaraDataInterfacePressureGrid::GridSizeName + Suffix;
-	Names.GridOriginName = UNiagaraDataInterfacePressureGrid::GridOriginName + Suffix;
-	Names.WorldTransformName = UNiagaraDataInterfacePressureGrid::WorldTransformName + Suffix;
-	Names.WorldInverseName = UNiagaraDataInterfacePressureGrid::WorldInverseName + Suffix;
-}
-
 //------------------------------------------------------------------------------------------------------------
 
 void FNDIPressureGridBuffer::Initialize(const FIntVector InGridSize)
@@ -154,8 +156,7 @@ struct FNDIPressureGridParametersCS : public FNiagaraDataInterfaceParametersCS
 
 	void Bind(const FNiagaraDataInterfaceGPUParamInfo& ParameterInfo, const class FShaderParameterMap& ParameterMap)
 	{
-		FNDIPressureGridParametersName ParamNames;
-		GetNiagaraDataInterfaceParametersName(ParamNames, *ParameterInfo.DataInterfaceHLSLSymbol);
+		FNDIPressureGridParametersName ParamNames(ParameterInfo.DataInterfaceHLSLSymbol);
 
 		GridCurrentBuffer.Bind(ParameterMap, *ParamNames.GridCurrentBufferName);
 		GridDestinationBuffer.Bind(ParameterMap, *ParamNames.GridDestinationBufferName);
@@ -719,8 +720,7 @@ void UNiagaraDataInterfacePressureGrid::UpdateGridTransform(FVectorVMContext& Co
 
 bool UNiagaraDataInterfacePressureGrid::GetFunctionHLSL(const FNiagaraDataInterfaceGPUParamInfo& ParamInfo, const FNiagaraDataInterfaceGeneratedFunction& FunctionInfo, int FunctionInstanceIndex, FString& OutHLSL)
 {
-	FNDIPressureGridParametersName ParamNames;
-	GetNiagaraDataInterfaceParametersName(ParamNames, ParamInfo.DataInterfaceHLSLSymbol);
+	FNDIPressureGridParametersName ParamNames(ParamInfo.DataInterfaceHLSLSymbol);
 
 	TMap<FString, FStringFormatArg> ArgsSample = {
 		{TEXT("InstanceFunctionName"), FunctionInfo.InstanceName},
@@ -946,16 +946,18 @@ inline void ClearBuffer(FNDIPressureGridBuffer* CurrentGridBuffer, FNDIPressureG
 			[DestinationGridBufferUAV, CurrentGridBufferSRV, CurrentGridBufferUAV, LocalGridSize, LocalCopyPressure]
 		(FRHICommandListImmediate& RHICmdListImm)
 		{
-			FRDGBuilder GraphBuilder(RHICmdListImm);
-
 			RHICmdListImm.TransitionResource(EResourceTransitionAccess::EWritable, EResourceTransitionPipeline::EComputeToCompute, DestinationGridBufferUAV);
 			RHICmdListImm.TransitionResource(EResourceTransitionAccess::EReadable, EResourceTransitionPipeline::EComputeToCompute, CurrentGridBufferUAV);
 
+			RHICmdListImm.ClearUAVUint(DestinationGridBufferUAV, FUintVector4(0,0,0,0));
+
+		/*	FRDGBuilder GraphBuilder(RHICmdListImm);
+
 			AddClearPressureGridPass(
 				GraphBuilder,
 				CurrentGridBufferSRV, DestinationGridBufferUAV, LocalGridSize, LocalCopyPressure);
 
-			GraphBuilder.Execute();
+			GraphBuilder.Execute();*/
 		});
 	}
 }
diff --git a/Engine/Plugins/Experimental/HairStrands/Source/HairStrandsNiagara/Public/NiagaraDataInterfaceHairStrands.h b/Engine/Plugins/Experimental/HairStrands/Source/HairStrandsNiagara/Public/NiagaraDataInterfaceHairStrands.h
--- a/Engine/Plugins/Experimental/HairStrands/Source/HairStrandsNiagara/Public/NiagaraDataInterfaceHairStrands.h
+++ b/Engine/Plugins/Experimental/HairStrands/Source/HairStrandsNiagara/Public/NiagaraDataInterfaceHairStrands.h
@@ -10,7 +10,7 @@
 #include "GroomActor.h"
 #include "NiagaraDataInterfaceHairStrands.generated.h"
 
-static const int32 MaxDelay = 5;
+static const int32 MaxDelay = 2;
 
 /** Render buffers that will be used in hlsl functions */
 struct FNDIHairStrandsBuffer : public FRenderResource
@@ -34,11 +34,11 @@ struct FNDIHairStrandsBuffer : public FRenderResource
 	/** Deformed position buffer in case no ressource are there */
 	FRWBuffer DeformedPositionBuffer;
 
-	/** Bounding Box Buffer*/
-	FRWBuffer BoundingBoxBuffer;
+	/** Bounding Box Buffer A*/
+	FRWBuffer BoundingBoxBufferA;
 
-	/** Node Bound Buffer*/
-	FRWBuffer NodeBoundBuffer;
+	/** Bounding Box Buffer B*/
+	FRWBuffer BoundingBoxBufferB;
 
 	/** The strand asset datas from which to sample */
 	const FHairStrandsDatas* SourceDatas;
@@ -74,7 +74,6 @@ struct FNDIHairStrandsData
 
 		TickCount = 0;
 		ForceReset = true;
-		ResetTick = MaxDelay;
 
 		NumStrands = 0;
 		StrandsSize = 0;
@@ -120,7 +119,6 @@ struct FNDIHairStrandsData
 
 			TickCount = OtherDatas->TickCount;
 			ForceReset = OtherDatas->ForceReset;
-			ResetTick = OtherDatas->ResetTick;
 
 			NumStrands = OtherDatas->NumStrands;
 			StrandsSize = OtherDatas->StrandsSize;
@@ -180,9 +178,6 @@ struct FNDIHairStrandsData
 	/** Force reset simulation */
 	bool ForceReset;
 
-	/** Reset Tick*/
-	int32 ResetTick;
-
 	/** Strands Gpu buffer */
 	FNDIHairStrandsBuffer* HairStrandsBuffer;
 	
@@ -318,7 +313,7 @@ public:
 	/** Get the number of strands */
 	void GetNumStrands(FVectorVMContext& Context);
 
-	/** Get the grrom asset datas  */
+	/** Get the groom asset datas  */
 	void GetStrandSize(FVectorVMContext& Context);
 
 	void GetSubSteps(FVectorVMContext& Context);
@@ -511,6 +506,15 @@ public:
 	/** Check if we need or not a simulation reset*/
 	void NeedSimulationReset(FVectorVMContext& Context);
 
+	/** Check if we need or not a simulation reset*/
+	void HasGlobalInterpolation(FVectorVMContext& Context);
+
+	/** Check if we need or not a simulation reset*/
+	void HasKinematicsTarget(FVectorVMContext& Context);
+
+	/** Eval the skinned position given a rest position*/
+	void EvalSkinnedPosition(FVectorVMContext& Context);
+
 	/** Name of the world transform */
 	static const FString WorldTransformName;
 
@@ -533,10 +537,10 @@ public:
 	static const FString CurvesOffsetsBufferName;
 
 	/** Name of bounding box buffer */
-	static const FString BoundingBoxBufferName;
+	static const FString BoundingBoxBufferAName;
 
 	/** Name of node bound buffer */
-	static const FString NodeBoundBufferName;
+	static const FString BoundingBoxBufferBName;
 
 	/** Name of the nodes positions buffer */
 	static const FString RestPositionBufferName;
@@ -586,6 +590,15 @@ public:
 	/** Deformed center of all the position */
 	static const FString DeformedPositionOffsetName;
 
+	/** Number of samples for rbf interpolation */
+	static const FString SampleCountName;
+
+	/** Rbf sample weights */
+	static const FString MeshSampleWeightsName;
+
+	/** Rbf Sample rest positions */
+	static const FString RestSamplePositionsName;
+
 protected:
 	/** Copy one niagara DI to this */
 	virtual bool CopyToInternal(UNiagaraDataInterface* Destination) const override;
diff --git a/Engine/Shaders/Private/DeferredLightPixelShaders.usf b/Engine/Shaders/Private/DeferredLightPixelShaders.usf
--- a/Engine/Shaders/Private/DeferredLightPixelShaders.usf
+++ b/Engine/Shaders/Private/DeferredLightPixelShaders.usf
@@ -106,7 +106,7 @@ uint GetLightingChannelMask(float2 UV)
 
 Texture2D<uint4> HairCategorizationTexture;
 Texture2D<uint> HairVisibilityNodeOffsetAndCount;
-StructuredBuffer<FNodeData> HairVisibilityNodeData;
+StructuredBuffer<FPackedHairSample> HairVisibilityNodeData;
 
 Texture2D<float4> ScreenShadowMaskSubPixelTexture;
 
@@ -155,37 +155,31 @@ void ComputeHairLighting(
 	for (uint SampleIt = 0; SampleIt < NodeDesc.Count; SampleIt++)
 	{
 		const uint LocalOffset = NodeDesc.Offset + SampleIt;
-		const FNodeData NodeData = HairVisibilityNodeData[LocalOffset];
-		float			SampleDepth;
-		FVisibilityData SampleVisData;
-		float3			SampleBaseColor;
-		float			SampleRoughness;
-		float			SampleSpecular;
-		uint			SampleLightChannelMask;
-		DecodeNodeData(NodeData, SampleDepth, SampleVisData, SampleBaseColor, SampleRoughness, SampleSpecular, SampleLightChannelMask);
+		const FPackedHairSample NodeData = HairVisibilityNodeData[LocalOffset];
+		const FHairSample Sample = UnpackHairSample(NodeData);
 
 	#if USE_LIGHTING_CHANNELS
-		if (!(SampleLightChannelMask & DeferredLightUniforms.LightingChannelMask))
+		if (!(Sample.LightChannelMask & DeferredLightUniforms.LightingChannelMask))
 		{
 			continue;
 		}
 	#endif
 
 		// Inject material data from the visibility/mini-gbuffer for hair
-		const float SceneDepth = ConvertFromDeviceZ(SampleDepth);
+		const float SceneDepth = ConvertFromDeviceZ(Sample.Depth);
 		FScreenSpaceData HairScreenSpaceData;
 		HairScreenSpaceData.AmbientOcclusion = 1;
-		HairScreenSpaceData.GBuffer.WorldNormal = SampleVisData.Tangent;
-		HairScreenSpaceData.GBuffer.BaseColor = SampleBaseColor;
-		HairScreenSpaceData.GBuffer.Roughness = SampleRoughness;
+		HairScreenSpaceData.GBuffer.WorldNormal = Sample.Tangent;
+		HairScreenSpaceData.GBuffer.BaseColor = Sample.BaseColor;
+		HairScreenSpaceData.GBuffer.Roughness = Sample.Roughness;
 		HairScreenSpaceData.GBuffer.ShadingModelID = SHADINGMODELID_HAIR;
 		HairScreenSpaceData.GBuffer.DiffuseColor = 0;
 		HairScreenSpaceData.GBuffer.SpecularColor = 0;
-		HairScreenSpaceData.GBuffer.Specular = SampleSpecular;
+		HairScreenSpaceData.GBuffer.Specular = Sample.Specular;
 		HairScreenSpaceData.GBuffer.Metallic = 0;
-		HairScreenSpaceData.GBuffer.Depth = SampleDepth;
+		HairScreenSpaceData.GBuffer.Depth = SceneDepth;
 		HairScreenSpaceData.GBuffer.GBufferAO = 1;
-		HairScreenSpaceData.GBuffer.CustomData = float4(HairDualScatteringRoughnessOverride, 0, 0, 0);;
+		HairScreenSpaceData.GBuffer.CustomData = float4(HairDualScatteringRoughnessOverride, 0, Sample.Backlit, 0);
 		HairScreenSpaceData.GBuffer.IndirectIrradiance = 0;
 		HairScreenSpaceData.GBuffer.PrecomputedShadowFactors = 1;
 
@@ -227,7 +221,7 @@ void ComputeHairLighting(
 
 		const float4 LinearColor = Radiance * Attenuation;
 		const float4 PercepualColor = LinearToPerceptual(LinearColor);
-		const float LocalCoverage = From8bitCoverage(SampleVisData.Coverage8bit);
+		const float LocalCoverage = From8bitCoverage(Sample.Coverage8bit);
 		LocalColor += PercepualColor * LocalCoverage;
 	}
 	const float PixelCoverage = saturate(CatData.PixelCoverage);
diff --git a/Engine/Shaders/Private/HairBsdf.ush b/Engine/Shaders/Private/HairBsdf.ush
--- a/Engine/Shaders/Private/HairBsdf.ush
+++ b/Engine/Shaders/Private/HairBsdf.ush
@@ -98,7 +98,7 @@ float3 EvaluateHairMultipleScattering(
 	return TransmittanceData.GlobalScattering * (Fs + TransmittanceData.LocalScattering) * TransmittanceData.OpaqueVisibility;
 }
 
-float3 HairShading( FGBufferData GBuffer, float3 L, float3 V, half3 N, float Shadow, FHairTransmittanceData HairTransmittance, float Backlit, float Area, uint2 Random )
+float3 HairShading( FGBufferData GBuffer, float3 L, float3 V, half3 N, float Shadow, FHairTransmittanceData HairTransmittance, float InBacklit, float Area, uint2 Random )
 {
 	// Strand tangent is oriented from root to tip, while the existing shading code expect tip to root orientation
 	//N = -N;
@@ -110,7 +110,7 @@ float3 HairShading( FGBufferData GBuffer, float3 L, float3 V, half3 N, float Sha
 
 	//const float3 DiffuseN	= OctahedronToUnitVector( GBuffer.CustomData.xy * 2 - 1 );
 	//const float Backlit	= GBuffer.CustomData.z;
-
+	const float Backlit		= min(InBacklit, GBuffer.CustomData.z);
 #if HAIR_REFERENCE
 	// todo: ClampedRoughness is missing for this code path
 	float3 S = HairShadingRef( GBuffer, L, V, N, Random );
diff --git a/Engine/Shaders/Private/HairStrands/HairScatter.usf b/Engine/Shaders/Private/HairStrands/HairScatter.usf
--- a/Engine/Shaders/Private/HairStrands/HairScatter.usf
+++ b/Engine/Shaders/Private/HairStrands/HairScatter.usf
@@ -168,7 +168,7 @@ Texture2D<float4> 			DiffusionInputTexture;
 
 Texture2D<uint4>  			CategorizationTexture;
 Texture2D<uint> 			VisibilityNodeIndex;
-StructuredBuffer<FNodeData> VisibilityNodeData;
+StructuredBuffer<FPackedHairSample> VisibilityNodeData;
 
 Texture3D<float4>			HairLUTTexture;
 Texture3D<float4>			HairEnergyLUTTexture;
@@ -232,6 +232,7 @@ FGBufferData GetGBufferData(uint2 PixelCoord)
 	GBuffer.WorldNormal = float3(1, 0, 0);
 	GBuffer.BaseColor = float3(0, 0, 0);
 	GBuffer.Roughness = 1;
+	GBuffer.CustomData = float4(0,0,1,0);
 	GBuffer.Specular = 0;
 
 	uint MacroGroupIndex = 0;
@@ -243,23 +244,17 @@ FGBufferData GetGBufferData(uint2 PixelCoord)
 		for (uint SampleIt = 0; SampleIt < NodeDesc.Count; ++SampleIt)
 		{
 			const uint LocalOffset = NodeDesc.Offset + SampleIt;
-			const FNodeData NodeData = VisibilityNodeData[LocalOffset];
-			float			SampleDepth;
-			FVisibilityData SampleVisData;
-			float3			SampleBaseColor;
-			float			SampleRoughness;
-			float			SampleSpecular;
-			uint			SampleLigthChannelMask;
-			DecodeNodeData(NodeData, SampleDepth, SampleVisData, SampleBaseColor, SampleRoughness, SampleSpecular, SampleLigthChannelMask);
-
-			if (SampleDepth >= ClosestDepth) // InverseZ
+			const FPackedHairSample NodeData = VisibilityNodeData[LocalOffset];
+			const FHairSample Sample = UnpackHairSample(NodeData);
+
+			if (Sample.Depth >= ClosestDepth) // InverseZ
 			{
 				GBuffer.ShadingModelID = SHADINGMODELID_HAIR;
-				GBuffer.WorldNormal = SampleVisData.Tangent;
-				GBuffer.BaseColor = SampleBaseColor;
-				GBuffer.Roughness = SampleRoughness;
-				GBuffer.Specular = SampleSpecular;
-				ClosestDepth = SampleDepth;
+				GBuffer.WorldNormal = Sample.Tangent;
+				GBuffer.BaseColor = Sample.BaseColor;
+				GBuffer.Roughness = Sample.Roughness;
+				GBuffer.Specular = Sample.Specular;
+				ClosestDepth = Sample.Depth;
 
 				// todo
 				MacroGroupIndex = 0; //clamp(SampleVisData.MacroGroupId, 0, MAX_HAIR_MACROGROUP_COUNT - 1);
diff --git a/Engine/Shaders/Private/HairStrands/HairStrandsDebug.usf b/Engine/Shaders/Private/HairStrands/HairStrandsDebug.usf
--- a/Engine/Shaders/Private/HairStrands/HairStrandsDebug.usf
+++ b/Engine/Shaders/Private/HairStrands/HairStrandsDebug.usf
@@ -8,16 +8,17 @@ float2 OutputResolution;
 uint FastResolveMask;
 uint DebugMode;
 int SampleIndex;
+uint MaxSampleCount;
 
 Texture2D<uint4> CategorizationTexture;
 Texture2D<float> ViewTransmittanceTexture;
 Texture2D<uint2> DepthStencilTexture;
 Texture2D<uint> NodeIndex;
-StructuredBuffer<FNodeData> NodeData;
+StructuredBuffer<FPackedHairSample> NodeData;
 SamplerState LinearSampler;
 
 uint TileSize;
-Texture2D<uint2> TileIndexTexture;
+Texture2D<uint> TileIndexTexture;
 
 #define DEBUG_MODE_SAMPLE_COUNT 0 
 #define DEBUG_MODE_COVERAGE_TYPE 1
@@ -46,7 +47,7 @@ void MainPS(
 	float Alpha = 1;
 	if (DebugMode == DEBUG_MODE_SAMPLE_COUNT)
 	{
-		Color = GetColorCode(CatData.SampleCount/float(HAIR_VISIBILITY_MAXSAMPLECOUNT));
+		Color = GetColorCode(CatData.SampleCount/float(MaxSampleCount));
 	}
 
 	if (DebugMode == DEBUG_MODE_COVERAGE_TYPE)
@@ -70,7 +71,7 @@ void MainPS(
 	if (DebugMode == DEBUG_MODE_TILE)
 	{
 		const uint2 TileCoord = PixelPos / TileSize;
-		const uint TileOffset = TileIndexTexture[TileCoord];
+		const uint  TileOffset = TileIndexTexture[TileCoord];
 		Color = TileOffset != INVALID_TILE_OFFSET ? float3(0.25, 0.85, 0.25) : float3(0.85, 0.25, 0.25);
 		Alpha = 0.5f;
 	}
@@ -97,30 +98,24 @@ void MainPS(
 		for (uint NodeIt = 0; NodeIt < NodeDesc.Count; ++NodeIt)
 		{
 			const uint NodeOffset = NodeDesc.Offset + NodeIt;
-			const FNodeData Data = NodeData.Load(NodeOffset);
-			float			SampleDepth;
-			FVisibilityData SampleVisData;
-			float3			SampleBaseColor;
-			float			SampleRoughness;
-			float			SampleSpecular;
-			uint			SampleLigthChannelMask;
-			DecodeNodeData(Data, SampleDepth, SampleVisData, SampleBaseColor, SampleRoughness, SampleSpecular, SampleLigthChannelMask);
+			const FPackedHairSample Data = NodeData.Load(NodeOffset);
+			const FHairSample Sample = UnpackHairSample(Data);
 
 			if (bAverageAllSamples)
 			{
-				Depth		+= SampleDepth;
-				Tangent		+= SampleVisData.Tangent;
-				BaseColor	+= SampleBaseColor;
-				Roughness	+= SampleRoughness;
-				Specular	+= SampleSpecular;
+				Depth		+= Sample.Depth;
+				Tangent		+= Sample.Tangent;
+				BaseColor	+= Sample.BaseColor;
+				Roughness	+= Sample.Roughness;
+				Specular	+= Sample.Specular;
 			}
 			else if (SelectedIndex == NodeIt)
 			{
-				Depth		= SampleDepth;
-				Tangent		= SampleVisData.Tangent;
-				BaseColor	= SampleBaseColor;
-				Roughness	= SampleRoughness;
-				Specular	= SampleSpecular;
+				Depth		= Sample.Depth;
+				Tangent		= Sample.Tangent;
+				BaseColor	= Sample.BaseColor;
+				Roughness	= Sample.Roughness;
+				Specular	= Sample.Specular;
 			}
 		}
 
diff --git a/Engine/Shaders/Private/HairStrands/HairStrandsDebugPrint.usf b/Engine/Shaders/Private/HairStrands/HairStrandsDebugPrint.usf
--- a/Engine/Shaders/Private/HairStrands/HairStrandsDebugPrint.usf
+++ b/Engine/Shaders/Private/HairStrands/HairStrandsDebugPrint.usf
@@ -11,17 +11,19 @@ int2 PixelCoord;
 int2 MaxResolution;
 uint FastResolveMask;
 uint HairMacroGroupCount;
+uint MaxSampleCount;
 uint2 GroupSize; 
 
+Texture2D<float> HairCountTexture;
 Texture2D<uint4> CategorizationTexture;
 Texture2D<uint2> DepthStencilTexture;
 SamplerState LinearSampler;
 Texture2D<uint> HairVisibilityNodeOffsetAndCount;
-StructuredBuffer<FNodeData> HairVisibilityNodeData;
+StructuredBuffer<FPackedHairSample> HairVisibilityNodeData;
 
 Buffer<int> HairMacroGroupAABBBuffer;
 
-#define MATERIAL_DATA 0
+#define MATERIAL_DATA 1
 #define DETAIL_DATA 1
 
 [numthreads(1, 1, 1)]
@@ -49,6 +51,8 @@ void MainCS(uint GroupIndex : SV_GroupIndex, uint3 DispatchThreadId : SV_Dispatc
 	if (!bIsValid)
 		return;
 
+	const float ViewHairCount = HairCountTexture.Load(uint3(PixelPos, 0));
+
 	const float2 OriginalPos = PixelPos / float2(MaxResolution);
 	float2 Pos = OriginalPos;
 
@@ -57,15 +61,9 @@ void MainCS(uint GroupIndex : SV_GroupIndex, uint3 DispatchThreadId : SV_Dispatc
 	for (uint SampleIt0 = 0; SampleIt0 < NodeDesc.Count; SampleIt0++)
 	{
 		const uint LocalOffset = NodeDesc.Offset + SampleIt0;
-		const FNodeData NodeData = HairVisibilityNodeData[LocalOffset];
-		float			SampleDepth;
-		FVisibilityData SampleVisData;
-		float3			SampleBaseColor;
-		float			SampleRoughness;
-		float			SampleSpecular;
-		uint			SampleLightChannelMask;
-		DecodeNodeData(NodeData, SampleDepth, SampleVisData, SampleBaseColor, SampleRoughness, SampleSpecular, SampleLightChannelMask);
-		Total8BitCoverage += SampleVisData.Coverage8bit;
+		const FPackedHairSample NodeData = HairVisibilityNodeData[LocalOffset];
+		const FHairSample Sample = UnpackHairSample(NodeData);
+		Total8BitCoverage += Sample.Coverage8bit;
 	}
 
 	// Pixel coord
@@ -110,7 +108,7 @@ void MainCS(uint GroupIndex : SV_GroupIndex, uint3 DispatchThreadId : SV_Dispatc
 	Pos = ShaderPrintSymbol(Pos, _L_);
 	Pos = ShaderPrintSymbol(Pos, _E_);
 	Pos = ShaderPrintSymbol(Pos, _SPC_);
-	Pos = ShaderPrint(Pos, uint(HAIR_VISIBILITY_MAXSAMPLECOUNT));
+	Pos = ShaderPrint(Pos, MaxSampleCount);
 
 	Pos = ShaderPrintNewline(Pos);
 	Pos.x = OriginalPos.x;
@@ -128,6 +126,23 @@ void MainCS(uint GroupIndex : SV_GroupIndex, uint3 DispatchThreadId : SV_Dispatc
 	Pos = ShaderPrintNewline(Pos);
 	Pos.x = OriginalPos.x;
 
+	// Accurate coverage (if coverage pass is activated)
+	Pos = ShaderPrintSymbol(Pos, _H_);
+	Pos = ShaderPrintSymbol(Pos, _A_);
+	Pos = ShaderPrintSymbol(Pos, _I_);
+	Pos = ShaderPrintSymbol(Pos, _R_);
+	Pos = ShaderPrintSymbol(Pos, _SPC_);
+	Pos = ShaderPrintSymbol(Pos, _C_);
+	Pos = ShaderPrintSymbol(Pos, _O_);
+	Pos = ShaderPrintSymbol(Pos, _U_);
+	Pos = ShaderPrintSymbol(Pos, _N_);
+	Pos = ShaderPrintSymbol(Pos, _T_);
+	Pos = ShaderPrintSymbol(Pos, _SPC_);
+	Pos = ShaderPrint(Pos, ViewHairCount);
+
+	Pos = ShaderPrintNewline(Pos);
+	Pos.x = OriginalPos.x;
+
 	// Accurate coverage (if coverage pass is activated)
 	Pos = ShaderPrintSymbol(Pos, _C_);
 	Pos = ShaderPrintSymbol(Pos, _O_);
@@ -229,14 +244,8 @@ void MainCS(uint GroupIndex : SV_GroupIndex, uint3 DispatchThreadId : SV_Dispatc
 	for (uint SampleIt = 0; SampleIt < NodeDesc.Count; SampleIt++)
 	{
 		const uint LocalOffset = NodeDesc.Offset + SampleIt;
-		const FNodeData NodeData = HairVisibilityNodeData[LocalOffset];
-		float			SampleDepth;
-		FVisibilityData SampleVisData;
-		float3			SampleBaseColor;
-		float			SampleRoughness;
-		float			SampleSpecular;
-		uint			SampleLightChannelMask;
-		DecodeNodeData(NodeData, SampleDepth, SampleVisData, SampleBaseColor, SampleRoughness, SampleSpecular, SampleLightChannelMask);
+		const FPackedHairSample NodeData = HairVisibilityNodeData[LocalOffset];		
+		const FHairSample Sample = UnpackHairSample(NodeData);
 
 		// Base Color
 		Pos = ShaderPrintSymbol(Pos, _MINUS_);
@@ -277,7 +286,7 @@ void MainCS(uint GroupIndex : SV_GroupIndex, uint3 DispatchThreadId : SV_Dispatc
 		Pos = ShaderPrintSymbol(Pos, _O_);
 		Pos = ShaderPrintSymbol(Pos, _R_);
 		Pos = ShaderPrintSymbol(Pos, _SPC_);
-		Pos = ShaderPrint(Pos, SampleBaseColor);
+		Pos = ShaderPrint(Pos, Sample.BaseColor);
 		
 		Pos = ShaderPrintNewline(Pos);
 		Pos.x = OriginalPos.x;
@@ -294,11 +303,41 @@ void MainCS(uint GroupIndex : SV_GroupIndex, uint3 DispatchThreadId : SV_Dispatc
 		Pos = ShaderPrintSymbol(Pos, _S_);
 		Pos = ShaderPrintSymbol(Pos, _S_);
 		Pos = ShaderPrintSymbol(Pos, _SPC_);
-		Pos = ShaderPrint(Pos, SampleRoughness);
+		Pos = ShaderPrint(Pos, Sample.Roughness);
 
 		Pos = ShaderPrintNewline(Pos);
 		Pos.x = OriginalPos.x;
 
+		// Specular
+		Pos = ShaderPrintSymbol(Pos, _DOT_);
+		Pos = ShaderPrintSymbol(Pos, _S_);
+		Pos = ShaderPrintSymbol(Pos, _P_);
+		Pos = ShaderPrintSymbol(Pos, _E_);
+		Pos = ShaderPrintSymbol(Pos, _C_);
+		Pos = ShaderPrintSymbol(Pos, _U_);
+		Pos = ShaderPrintSymbol(Pos, _L_);
+		Pos = ShaderPrintSymbol(Pos, _A_);
+		Pos = ShaderPrintSymbol(Pos, _R_);
+		Pos = ShaderPrintSymbol(Pos, _SPC_);
+		Pos = ShaderPrint(Pos, Sample.Specular);
+
+		Pos = ShaderPrintNewline(Pos);
+		Pos.x = OriginalPos.x;
+
+		// Backlit
+		Pos = ShaderPrintSymbol(Pos, _DOT_);
+		Pos = ShaderPrintSymbol(Pos, _B_);
+		Pos = ShaderPrintSymbol(Pos, _A_);
+		Pos = ShaderPrintSymbol(Pos, _C_);
+		Pos = ShaderPrintSymbol(Pos, _K_);
+		Pos = ShaderPrintSymbol(Pos, _L_);
+		Pos = ShaderPrintSymbol(Pos, _I_);
+		Pos = ShaderPrintSymbol(Pos, _T_);
+		Pos = ShaderPrintSymbol(Pos, _SPC_);
+		Pos = ShaderPrint(Pos, Sample.Backlit);
+
+		Pos = ShaderPrintNewline(Pos);
+		Pos.x = OriginalPos.x;
 
 		// Tangent
 		Pos = ShaderPrintSymbol(Pos, _DOT_);
@@ -310,7 +349,7 @@ void MainCS(uint GroupIndex : SV_GroupIndex, uint3 DispatchThreadId : SV_Dispatc
 		Pos = ShaderPrintSymbol(Pos, _N_);
 		Pos = ShaderPrintSymbol(Pos, _T_);
 		Pos = ShaderPrintSymbol(Pos, _SPC_);
-		Pos = ShaderPrint(Pos, SampleVisData.Tangent);
+		Pos = ShaderPrint(Pos, Sample.Tangent);
 
 		Pos = ShaderPrintNewline(Pos);
 		Pos.x = OriginalPos.x;
@@ -329,7 +368,7 @@ void MainCS(uint GroupIndex : SV_GroupIndex, uint3 DispatchThreadId : SV_Dispatc
 		Pos = ShaderPrintSymbol(Pos, _V_);
 		Pos = ShaderPrintSymbol(Pos, _E_);
 		Pos = ShaderPrintSymbol(Pos, _SPC_);
-		Pos = ShaderPrint(Pos, SampleVisData.PrimitiveId);
+		Pos = ShaderPrint(Pos, Sample.PrimitiveId);
 
 		Pos = ShaderPrintNewline(Pos);
 		Pos.x = OriginalPos.x;
@@ -351,13 +390,12 @@ void MainCS(uint GroupIndex : SV_GroupIndex, uint3 DispatchThreadId : SV_Dispatc
 		Pos = ShaderPrintSymbol(Pos, _I_);
 		Pos = ShaderPrintSymbol(Pos, _D_);
 		Pos = ShaderPrintSymbol(Pos, _SPC_);
-		Pos = ShaderPrint(Pos, SampleVisData.MacroGroupId);
+		Pos = ShaderPrint(Pos, Sample.MacroGroupId);
 
 		Pos = ShaderPrintNewline(Pos);
 		Pos.x = OriginalPos.x;		
 
-
-		// Ligtht channel mask
+		// Light channel mask
 		Pos = ShaderPrintSymbol(Pos, _DOT_);
 		Pos = ShaderPrintSymbol(Pos, _L_);
 		Pos = ShaderPrintSymbol(Pos, _I_);
@@ -374,9 +412,30 @@ void MainCS(uint GroupIndex : SV_GroupIndex, uint3 DispatchThreadId : SV_Dispatc
 		Pos = ShaderPrintSymbol(Pos, _L_);
 		Pos = ShaderPrintSymbol(Pos, _S_);
 		Pos = ShaderPrintSymbol(Pos, _SPC_);
-		Pos = ShaderPrint(Pos, (SampleLightChannelMask & 0x1) ? 1u : 0u);
-		Pos = ShaderPrint(Pos, (SampleLightChannelMask & 0x2) ? 1u : 0u);
-		Pos = ShaderPrint(Pos, (SampleLightChannelMask & 0x4) ? 1u : 0u);
+		Pos = ShaderPrint(Pos, (Sample.LightChannelMask & 0x1) ? 1u : 0u);
+		Pos = ShaderPrint(Pos, (Sample.LightChannelMask & 0x2) ? 1u : 0u);
+		Pos = ShaderPrint(Pos, (Sample.LightChannelMask & 0x4) ? 1u : 0u);
+
+		Pos = ShaderPrintNewline(Pos);
+		Pos.x = OriginalPos.x;
+
+		// Scene lighting scatter
+		Pos = ShaderPrintSymbol(Pos, _DOT_);
+		Pos = ShaderPrintSymbol(Pos, _S_);
+		Pos = ShaderPrintSymbol(Pos, _C_);
+		Pos = ShaderPrintSymbol(Pos, _A_);
+		Pos = ShaderPrintSymbol(Pos, _T_);
+		Pos = ShaderPrintSymbol(Pos, _T_);
+		Pos = ShaderPrintSymbol(Pos, _E_);
+		Pos = ShaderPrintSymbol(Pos, _R_);
+		Pos = ShaderPrintSymbol(Pos, _SPC_);
+		Pos = ShaderPrintSymbol(Pos, _S_);
+		Pos = ShaderPrintSymbol(Pos, _C_);
+		Pos = ShaderPrintSymbol(Pos, _E_);
+		Pos = ShaderPrintSymbol(Pos, _N_);
+		Pos = ShaderPrintSymbol(Pos, _E_);
+		Pos = ShaderPrintSymbol(Pos, _SPC_);
+		Pos = ShaderPrint(Pos, Sample.bScatterSceneLighting ? 1u : 0u);
 
 		Pos = ShaderPrintNewline(Pos);
 		Pos.x = OriginalPos.x;
@@ -393,7 +452,7 @@ void MainCS(uint GroupIndex : SV_GroupIndex, uint3 DispatchThreadId : SV_Dispatc
 		Pos = ShaderPrintSymbol(Pos, _G_);
 		Pos = ShaderPrintSymbol(Pos, _E_);
 		Pos = ShaderPrintSymbol(Pos, _SPC_);
-		Pos = ShaderPrint(Pos, SampleVisData.Coverage8bit);
+		Pos = ShaderPrint(Pos, Sample.Coverage8bit);
 
 		Pos = ShaderPrintNewline(Pos);
 		Pos.x = OriginalPos.x;
diff --git a/Engine/Shaders/Private/HairStrands/HairStrandsDeepShadowMask.usf b/Engine/Shaders/Private/HairStrands/HairStrandsDeepShadowMask.usf
--- a/Engine/Shaders/Private/HairStrands/HairStrandsDeepShadowMask.usf
+++ b/Engine/Shaders/Private/HairStrands/HairStrandsDeepShadowMask.usf
@@ -18,7 +18,7 @@
   #if PERMUTATION_OPAQUEMASK_TYPE == OPAQUE_MASK_TYPE_VOXEL
 	#define VOXEL_TRAVERSAL_TYPE VOXEL_TRAVERSAL_LEGACY
   #elif PERMUTATION_OPAQUEMASK_TYPE == OPAQUE_MASK_TYPE_VIRTUAL_VOXEL
-	#define VOXEL_TRAVERSAL_TYPE VOXEL_TRAVERSAL_LINEAR
+	#define VOXEL_TRAVERSAL_TYPE VOXEL_TRAVERSAL_LINEAR_MIPMAP
   #endif
   #include "HairStrandsVoxelPageTraversal.ush"
 #endif
diff --git a/Engine/Shaders/Private/HairStrands/HairStrandsDeepShadowPS.usf b/Engine/Shaders/Private/HairStrands/HairStrandsDeepShadowPS.usf
--- a/Engine/Shaders/Private/HairStrands/HairStrandsDeepShadowPS.usf
+++ b/Engine/Shaders/Private/HairStrands/HairStrandsDeepShadowPS.usf
@@ -55,66 +55,26 @@ void MainDom(
 
 #define DEBUG_OUTPUT 0
 
-#if DEBUG_OUTPUT == 1
-	#define REGISTER_OUT_0 u1
-	#define REGISTER_OUT_1 u2
-	#define REGISTER_OUT_2 u3
-	#define REGISTER_OUT_3 u4
-	#define REGISTER_OUT_4 u5
-#else
-	#define REGISTER_OUT_0 u0
-	#define REGISTER_OUT_1 u1
-	#define REGISTER_OUT_2 u2
-	#define REGISTER_OUT_3 u3
-	#define REGISTER_OUT_4 u4
-#endif
-
-#if VF_STRAND_HAIR
-  #if SM5_PROFILE && COMPILER_HLSL
-	RWTexture3D<uint> DensityTexture  : register(REGISTER_OUT_0);
-
-	#if SUPPORT_TANGENT_PROPERTY == 1
-	RWTexture3D<uint> TangentXTexture : register(REGISTER_OUT_1);
-	RWTexture3D<uint> TangentYTexture : register(REGISTER_OUT_2);
-	RWTexture3D<uint> TangentZTexture : register(REGISTER_OUT_3);
-	#endif
-
-	#if SUPPORT_MATERIAL_PROPERTY == 1
-	RWTexture3D<uint> MaterialTexture : register(REGISTER_OUT_4);
-	#endif
-  #else
-	RWTexture3D<uint> DensityTexture;
-
-	#if SUPPORT_TANGENT_PROPERTY == 1
-	RWTexture3D<uint> TangentXTexture;
-	RWTexture3D<uint> TangentYTexture;
-	RWTexture3D<uint> TangentZTexture;
-	#endif
-
-	#if SUPPORT_MATERIAL_PROPERTY == 1
-	RWTexture3D<uint> MaterialTexture;
-	#endif
-  #endif
-#endif 
-
 void MainVoxel(
 #if SUPPORT_TANGENT_PROPERTY == 1 || SUPPORT_MATERIAL_PROPERTY == 1 
 	FVertexFactoryInterpolantsVSToPS Interpolants,
 #endif
 	in float4 SvPosition : SV_Position,
-	in float4 PixelPosition : TEXCOORD6,
-	in float HairCoverage: CUSTOM_COVERAGE
+	in float4 WorldPositon_Coverage : TEXCOORD6
 #if DEBUG_OUTPUT
 ,	out uint OutColor : SV_Target0 // For debug purpose only
 #endif
 )
 {
+	const float3 WorldPosition = WorldPositon_Coverage.xyz;
+	const float HairCoverage = WorldPositon_Coverage.w;
+
 #if VF_STRAND_HAIR  
   #if SUPPORT_SPARSE	
 	const FPackedVirtualVoxelNodeDesc PackedNode = VoxelRasterPass.VirtualVoxel.NodeDescBuffer[VoxelRasterPass.MacroGroupId];
 	const FVirtualVoxelNodeDesc NodeDesc = UnpackVoxelNode(PackedNode, VoxelRasterPass.VirtualVoxel.PageResolution);
 
-	const int3 VoxelCoord = PositionToCoord(PixelPosition.xyz, NodeDesc.MinAABB, NodeDesc.MaxAABB, NodeDesc.VirtualResolution);
+	const int3 VoxelCoord = PositionToCoord(WorldPosition.xyz, NodeDesc.MinAABB, NodeDesc.MaxAABB, NodeDesc.VirtualResolution);
 	if (VoxelCoord.x >= 0 || VoxelCoord.y >= 0 || VoxelCoord.z >= 0 ||
 		VoxelCoord.x < int(NodeDesc.VirtualResolution.x) || 
 		VoxelCoord.y < int(NodeDesc.VirtualResolution.y) || 
@@ -135,7 +95,7 @@ void MainVoxel(
 
 				const float HairCount			= HairCoverage;
 				const float VoxelFixPointScale	= GetVoxelDensityFixPointScale();
-				InterlockedAdd(DensityTexture[PageVoxelCoord], HairCount * VoxelFixPointScale);
+				InterlockedAdd(VoxelRasterPass.DensityTexture[PageVoxelCoord], HairCount * VoxelFixPointScale);
 
 				#if DEBUG_OUTPUT
 				OutColor = 0xFF;
@@ -152,7 +112,7 @@ void MainVoxel(
 		}
 	}
   #else
-	const int3 VoxelCoord = (PixelPosition.xyz - VoxelRasterPass.VoxelMinAABB) / (VoxelRasterPass.VoxelMaxAABB - VoxelRasterPass.VoxelMinAABB) * VoxelRasterPass.VoxelResolution;
+	const int3 VoxelCoord = (WorldPosition.xyz - VoxelRasterPass.VoxelMinAABB) / (VoxelRasterPass.VoxelMaxAABB - VoxelRasterPass.VoxelMinAABB) * VoxelRasterPass.VoxelResolution;
 		 
 	if (VoxelCoord.x >= 0 || VoxelCoord.y >= 0 || VoxelCoord.z >= 0 ||
 		VoxelCoord.x < int(VoxelRasterPass.VoxelResolution.x) || 
@@ -161,15 +121,15 @@ void MainVoxel(
 	{
 		const float HairCount = HairCoverage;
 		const float VoxelFixPointScale = GetVoxelDensityFixPointScale();
-		InterlockedAdd(DensityTexture[VoxelCoord], HairCount * VoxelFixPointScale);
+		InterlockedAdd(VoxelRasterPass.DensityTexture[VoxelCoord], HairCount * VoxelFixPointScale);
 
 		#if SUPPORT_TANGENT_PROPERTY == 1
 		const half3 Tangent = GetTangentToWorld2(Interpolants).xyz;
 		const uint3 Tangent8bits = QuantizeTo8Bits(Tangent);
 
-		InterlockedAdd(TangentXTexture[VoxelCoord], Tangent8bits.x);
-		InterlockedAdd(TangentYTexture[VoxelCoord], Tangent8bits.y);
-		InterlockedAdd(TangentZTexture[VoxelCoord], Tangent8bits.z);
+		InterlockedAdd(VoxelRasterPass.TangentXTexture[VoxelCoord], Tangent8bits.x);
+		InterlockedAdd(VoxelRasterPass.TangentYTexture[VoxelCoord], Tangent8bits.y);
+		InterlockedAdd(VoxelRasterPass.TangentZTexture[VoxelCoord], Tangent8bits.z);
 		#endif	
 
 		#if SUPPORT_MATERIAL_PROPERTY == 1
@@ -187,7 +147,7 @@ void MainVoxel(
 			const float3 BaseColor = GetMaterialBaseColor(PixelMaterialInputs);
 			const float Roughness = GetMaterialRoughness(PixelMaterialInputs);
 			const uint EncodedBaseColorAndRoughness = EncodeBaseColorAndRoughness(BaseColor, Roughness);
-			InterlockedMax(MaterialTexture[VoxelCoord], EncodedBaseColorAndRoughness);
+			InterlockedMax(VoxelRasterPass.MaterialTexture[VoxelCoord], EncodedBaseColorAndRoughness);
 		}
 		#endif
 
diff --git a/Engine/Shaders/Private/HairStrands/HairStrandsDeepShadowVS.usf b/Engine/Shaders/Private/HairStrands/HairStrandsDeepShadowVS.usf
--- a/Engine/Shaders/Private/HairStrands/HairStrandsDeepShadowVS.usf
+++ b/Engine/Shaders/Private/HairStrands/HairStrandsDeepShadowVS.usf
@@ -11,6 +11,14 @@
 #define MESH_RENDER_DOM   1
 #define MESH_RENDER_VOXEL 2
 
+#ifndef SUPPORT_TANGENT_PROPERTY
+#define SUPPORT_TANGENT_PROPERTY 0
+#endif
+
+#ifndef SUPPORT_MATERIAL_PROPERTY
+#define SUPPORT_MATERIAL_PROPERTY 0
+#endif
+
 struct FDeepShadowAccumulateVSToPS
 {
 #if SUPPORT_TANGENT_PROPERTY == 1 || SUPPORT_MATERIAL_PROPERTY == 1
@@ -18,9 +26,9 @@ struct FDeepShadowAccumulateVSToPS
 #endif
 	float4 Position : SV_POSITION;
 #if MESH_RENDER_MODE == MESH_RENDER_VOXEL
-	float4 PixelPosition : TEXCOORD6;
+	float4 WorldPositon_Coverage : TEXCOORD6;
 #endif
-#if MESH_RENDER_MODE == MESH_RENDER_DOM || MESH_RENDER_MODE == MESH_RENDER_VOXEL
+#if MESH_RENDER_MODE == MESH_RENDER_DOM
 	float HairCoverage : CUSTOM_COVERAGE;
 #endif
 };
@@ -35,7 +43,50 @@ void Main(
 	FVertexFactoryInput Input,
 	OPTIONAL_VertexID
 	out VS_OUTPUT_TYPE Output
-	)
+)
+#if defined(VF_STRAND_HAIR) && MESH_RENDER_MODE == MESH_RENDER_VOXEL && SUPPORT_TANGENT_PROPERTY == 0 && SUPPORT_MATERIAL_PROPERTY == 0
+{
+	ResolvedView = ResolveView();
+	
+	const FHairRenderInfo HairRenderInfo = GetHairRenderInfo(ResolvedView.HairRenderInfo, ResolvedView.HairRenderInfoBits);
+
+	FHairViewInfo HairViewInfo;
+	HairViewInfo.TranslatedWorldCameraOrigin = ResolvedView.TranslatedWorldCameraOrigin;
+	HairViewInfo.RadiusAtDepth1 = HairRenderInfo.RadiusAtDepth1Primary;
+	HairViewInfo.bIsOrthoView = true;
+	HairViewInfo.ViewForward = ResolvedView.ViewForward;
+
+	float2 ViewportScale;
+	float4x4 WorldToClipMatrix = VoxelRasterPass.WorldToClipMatrix;
+	if (HairRenderInfo.bIsGPUDriven)
+	{
+		const uint MacroGroupId = VoxelRasterPass.MacroGroupId;
+		FVoxelizationViewInfo VoxelizationViewInfo = VoxelRasterPass.VoxelizationViewInfoBuffer.Load(MacroGroupId);
+		WorldToClipMatrix = VoxelizationViewInfo.WorldToClip;
+		ViewportScale = VoxelizationViewInfo.RasterResolution / float2(VoxelRasterPass.ViewportResolution);
+		HairViewInfo.ViewForward = VoxelizationViewInfo.ViewForward;
+	}
+
+	float HairWorldRadius = 0;
+	float HairDensity = 0;
+	float4 WorldPosition = VertexFactoryGetWorldPosition_Voxelization(Input, HairViewInfo, HairWorldRadius, HairDensity);
+
+	WorldPosition.xyz -= ResolvedView.PreViewTranslation.xyz; // Remove the camera relative position (todo move this into matrix transform)
+
+	Output.Position = mul(WorldPosition, WorldToClipMatrix);
+
+	if (HairRenderInfo.bIsGPUDriven)
+	{
+		Output.Position.xy += float2(1, 1);
+		Output.Position.xy *= ViewportScale;
+		Output.Position.xy -= float2(1, 1);
+	}
+
+	const float PixelRadius = HairRenderInfo.RadiusAtDepth1Primary;
+	const float Coverage = HairWorldRadius / max(HairWorldRadius, PixelRadius);
+	Output.WorldPositon_Coverage = float4(WorldPosition.xyz, Coverage * HairDensity);
+}
+#else // SPARSE VOXELIZAITON SPECIALIZATION
 {
 	ResolvedView = ResolveView();
 
@@ -48,7 +99,7 @@ void Main(
 #endif
 
 #ifdef VF_STRAND_HAIR	
-	const FHairRenderInfo HairRenderInfo = GetHairRenderInfo(ResolvedView.HairRenderInfo);
+	const FHairRenderInfo HairRenderInfo = GetHairRenderInfo(ResolvedView.HairRenderInfo, ResolvedView.HairRenderInfoBits);
 
 	FHairViewInfo HairViewInfo;
 	HairViewInfo.TranslatedWorldCameraOrigin = ResolvedView.TranslatedWorldCameraOrigin;
@@ -84,7 +135,8 @@ void Main(
 		Output.Position.xy -= float2(1, 1);
 	}
 
-	Output.PixelPosition = WorldPosition;
+	Output.WorldPositon_Coverage.xyz = WorldPosition;
+
   #if SUPPORT_TANGENT_PROPERTY == 1 || SUPPORT_MATERIAL_PROPERTY == 1
 	float3x3 TangentToLocal = VertexFactoryGetTangentToLocal(Input, VFIntermediates);
 	FMaterialVertexParameters VertexParameters = GetMaterialVertexParameters(Input, VFIntermediates, WorldPosition.xyz, TangentToLocal);
@@ -92,19 +144,25 @@ void Main(
   #endif
 #endif
 
-#if MESH_RENDER_MODE == MESH_RENDER_DOM || MESH_RENDER_MODE == MESH_RENDER_VOXEL
+	float HairCoverage = 1;
   #ifdef VF_STRAND_HAIR	
 	const float CurrentDepth = Output.Position.z / Output.Position.w;
 	const float PixelRadius = HairRenderInfo.bIsOrthoView ? HairRenderInfo.RadiusAtDepth1Primary : (CurrentDepth * HairRenderInfo.RadiusAtDepth1Primary);
 	const float StrandRealRadius = VFIntermediates.HairDimensions.y;
 	const float Coverage = StrandRealRadius / max(StrandRealRadius, PixelRadius);
-	Output.HairCoverage = Coverage * VFIntermediates.HairDensity;
-  #else
-	Output.HairCoverage = 1;
+	HairCoverage = Coverage * VFIntermediates.HairDensity;
+  #endif
+
+  #if MESH_RENDER_MODE == MESH_RENDER_DOM
+	Output.HairCoverage = HairCoverage;
+  #endif
+
+  #if MESH_RENDER_MODE == MESH_RENDER_VOXEL
+	Output.WorldPositon_Coverage.w = HairCoverage;
   #endif
-#endif
 
 	OutputVertexID( Output );
 }
+#endif // SPARSE VOXELIZAITON SPECIALIZATION
 
 #endif // VERTEXSHADER
diff --git a/Engine/Shaders/Private/HairStrands/HairStrandsDeepTransmittanceDualScattering.ush b/Engine/Shaders/Private/HairStrands/HairStrandsDeepTransmittanceDualScattering.ush
--- a/Engine/Shaders/Private/HairStrands/HairStrandsDeepTransmittanceDualScattering.ush
+++ b/Engine/Shaders/Private/HairStrands/HairStrandsDeepTransmittanceDualScattering.ush
@@ -22,6 +22,7 @@ float Hair_g2(float B,float Theta)
 void ComputeDualScatteringTerms(
 	const FHairTransmittanceData TransmittanceData,
 	const float Roughness,
+	const float Backlit,
 	const float3 V,
 	const float3 L,
 	const float3 T,
@@ -110,7 +111,9 @@ void ComputeDualScatteringTerms(
     //const float3 GlobalScattering = lerp(1, Tf * Sf * df, saturate(TransmittanceData.HairCount));
 	//const float3 GlobalScattering = Tf * Sf * df;
 
-	const float3 LocalScattering  = 2 * Ab * Sb * db;
+	const float Attenuation = lerp(1, Backlit, saturate(-dot(V,L)));
+
+	const float3 LocalScattering  = 2 * Ab * Sb * db * Attenuation;
 	
 	// Decode if the Local/Global component should be enabled or not (for debug purpose)
 	const float LocalWeights  = (HairComponents & HAIR_COMPONENT_LS) > 0 ? 1 : 0;
@@ -141,6 +144,7 @@ FHairTransmittanceData GetHairTransmittance(
 
 	// Hack: Override the actual roughness, with a different value to achieve a specific look
 	const float Roughness = GBuffer.CustomData.x > 0 ? GBuffer.CustomData.x : GBuffer.Roughness;
+	const float Backlit = GBuffer.CustomData.z;
 
 	const float3 T = GBuffer.WorldNormal;
 	{
@@ -155,6 +159,7 @@ FHairTransmittanceData GetHairTransmittance(
 	ComputeDualScatteringTerms(
 		TransmittanceData,
 		Roughness,
+		Backlit,
 		V,
 		L,
 		T,
diff --git a/Engine/Shaders/Private/HairStrands/HairStrandsDeepTransmittanceMask.usf b/Engine/Shaders/Private/HairStrands/HairStrandsDeepTransmittanceMask.usf
--- a/Engine/Shaders/Private/HairStrands/HairStrandsDeepTransmittanceMask.usf
+++ b/Engine/Shaders/Private/HairStrands/HairStrandsDeepTransmittanceMask.usf
@@ -79,7 +79,7 @@ Texture3D<uint> 	Voxel_DensityTexture7;
 
 Texture2D<float4> 	RayMarchMaskTexture;
 
-StructuredBuffer<FNodeData> HairVisibilityNodeData;
+StructuredBuffer<FPackedHairSample> HairVisibilityNodeData;
 StructuredBuffer<uint>		HairVisibilityNodeCoord;
 Buffer<uint4>				IndirectArgsBuffer;
 
@@ -114,16 +114,10 @@ void MainCS(uint2 DispatchThreadId : SV_DispatchThreadID)
 	const float3 TransmittanceRandom = float3(InterleavedGradientNoise(PixelCoord, 1), InterleavedGradientNoise(PixelCoord, 2), InterleavedGradientNoise(PixelCoord, 3));
 
 	{
-		const FNodeData NodeData = HairVisibilityNodeData[LocalOffset];
-		float			SampleDepth;
-		FVisibilityData SampleVisData;
-		float3			SampleBaseColor;
-		float			SampleRoughness;
-		float			SampleSpecular;
-		uint			SampleLightChannelMask;
-		DecodeNodeData(NodeData, SampleDepth, SampleVisData, SampleBaseColor, SampleRoughness, SampleSpecular, SampleLightChannelMask);
-
-		if ((LightChannelMask & SampleLightChannelMask) == 0)
+		const FPackedHairSample NodeData = HairVisibilityNodeData[LocalOffset];
+		const FHairSample Sample = UnpackHairSample(NodeData);
+
+		if ((LightChannelMask & Sample.LightChannelMask) == 0)
 		{
 			OutputColor[LocalOffset] = InitNullPackedHairTransmittanceMask();
 			return;
@@ -131,12 +125,12 @@ void MainCS(uint2 DispatchThreadId : SV_DispatchThreadID)
 
 		FGBufferData GBuffer;
 		GBuffer.ShadingModelID = SHADINGMODELID_HAIR;
-		GBuffer.WorldNormal = SampleVisData.Tangent;
-		GBuffer.BaseColor = SampleBaseColor;
-		GBuffer.Roughness = SampleRoughness;
+		GBuffer.WorldNormal = Sample.Tangent;
+		GBuffer.BaseColor = Sample.BaseColor;
+		GBuffer.Roughness = Sample.Roughness;
 
-		const uint MacroGroupIndex = clamp(SampleVisData.MacroGroupId, 0, MAX_HAIR_MACROGROUP_COUNT - 1);
-		const int2 DeepShadowOffset = DeepShadow_AtlasSlotOffsets[MacroGroupIndex];
+		const uint MacroGroupIndex = clamp(Sample.MacroGroupId, 0, MAX_HAIR_MACROGROUP_COUNT - 1);
+		const int2 DeepShadowOffset = DeepShadow_AtlasSlotOffsets[MacroGroupIndex].xy;
 		const float4x4 WorldToLightTransform = DeepShadow_WorldToLightTransforms[MacroGroupIndex];
 		
 		FTransmittanceSettings Settings;
@@ -152,7 +146,7 @@ void MainCS(uint2 DispatchThreadId : SV_DispatchThreadID)
 		Settings.Random = TransmittanceRandom;
 		Settings.DebugMode = DeepShadow_DebugMode;
 
-		const float SceneDepth = ConvertFromDeviceZ(SampleDepth);
+		const float SceneDepth = ConvertFromDeviceZ(Sample.Depth);
 		const float3 WorldPosition = ToWorldPosition(UV, SceneDepth);
 		const FHairTransmittanceMask TransmittanceMask = ComputeTransmittance(WorldPosition, GBuffer, Settings, DeepShadow_FrontDepthTexture, DeepShadow_DomTexture, HairLUTTexture, LinearSampler);
 
@@ -232,7 +226,7 @@ FHairTransmittanceMask ComputeTransmittanceVoxel(
 	float MaxVoxelSize;
 	{
 	#if PERMUTATION_TRANSMITTANCE_TYPE == PERMUTATION_TRANSMITTANCE_VOXEL
-		const float3 VoxelSize = (Voxel_MaxAABBs[MacroGroupId] - Voxel_MinAABBs[MacroGroupId]) / Voxel_Resolution;
+		const float3 VoxelSize = (Voxel_MaxAABBs[MacroGroupId].xyz - Voxel_MinAABBs[MacroGroupId].xyz) / Voxel_Resolution;
 		MaxVoxelSize = max(VoxelSize.x, max(VoxelSize.y, VoxelSize.z));
 	#elif PERMUTATION_TRANSMITTANCE_TYPE == PERMUTATION_TRANSMITTANCE_VIRTUAL_VOXEL
 		MaxVoxelSize = VirtualVoxel.VoxelWorldSize;
@@ -261,14 +255,14 @@ FHairTransmittanceMask ComputeTransmittanceVoxel(
 
 		#if PERMUTATION_TRANSMITTANCE_TYPE == PERMUTATION_TRANSMITTANCE_VOXEL
 		FHairTraversalResult Result = InitHairTraversalResult();
-		if (MacroGroupId == 0) Result = ComputeHairCountVoxel(SampleWorldPosition, LightPosition, Voxel_MinAABBs[0], Voxel_MaxAABBs[0], Voxel_Resolution, Voxel_DensityTexture0, Settings.DeepShadowDensityScale, 0, DistanceThreshold);
-		if (MacroGroupId == 1) Result = ComputeHairCountVoxel(SampleWorldPosition, LightPosition, Voxel_MinAABBs[1], Voxel_MaxAABBs[1], Voxel_Resolution, Voxel_DensityTexture1, Settings.DeepShadowDensityScale, 0, DistanceThreshold);
-		if (MacroGroupId == 2) Result = ComputeHairCountVoxel(SampleWorldPosition, LightPosition, Voxel_MinAABBs[2], Voxel_MaxAABBs[2], Voxel_Resolution, Voxel_DensityTexture2, Settings.DeepShadowDensityScale, 0, DistanceThreshold);
-		if (MacroGroupId == 3) Result = ComputeHairCountVoxel(SampleWorldPosition, LightPosition, Voxel_MinAABBs[3], Voxel_MaxAABBs[3], Voxel_Resolution, Voxel_DensityTexture3, Settings.DeepShadowDensityScale, 0, DistanceThreshold);
-		if (MacroGroupId == 4) Result = ComputeHairCountVoxel(SampleWorldPosition, LightPosition, Voxel_MinAABBs[4], Voxel_MaxAABBs[4], Voxel_Resolution, Voxel_DensityTexture4, Settings.DeepShadowDensityScale, 0, DistanceThreshold);
-		if (MacroGroupId == 5) Result = ComputeHairCountVoxel(SampleWorldPosition, LightPosition, Voxel_MinAABBs[5], Voxel_MaxAABBs[5], Voxel_Resolution, Voxel_DensityTexture5, Settings.DeepShadowDensityScale, 0, DistanceThreshold);
-		if (MacroGroupId == 6) Result = ComputeHairCountVoxel(SampleWorldPosition, LightPosition, Voxel_MinAABBs[6], Voxel_MaxAABBs[6], Voxel_Resolution, Voxel_DensityTexture6, Settings.DeepShadowDensityScale, 0, DistanceThreshold);
-		if (MacroGroupId == 7) Result = ComputeHairCountVoxel(SampleWorldPosition, LightPosition, Voxel_MinAABBs[7], Voxel_MaxAABBs[7], Voxel_Resolution, Voxel_DensityTexture7, Settings.DeepShadowDensityScale, 0, DistanceThreshold);
+		if (MacroGroupId == 0) Result = ComputeHairCountVoxel(SampleWorldPosition, LightPosition, Voxel_MinAABBs[0].xyz, Voxel_MaxAABBs[0].xyz, Voxel_Resolution, Voxel_DensityTexture0, Settings.DeepShadowDensityScale, 0, DistanceThreshold);
+		if (MacroGroupId == 1) Result = ComputeHairCountVoxel(SampleWorldPosition, LightPosition, Voxel_MinAABBs[1].xyz, Voxel_MaxAABBs[1].xyz, Voxel_Resolution, Voxel_DensityTexture1, Settings.DeepShadowDensityScale, 0, DistanceThreshold);
+		if (MacroGroupId == 2) Result = ComputeHairCountVoxel(SampleWorldPosition, LightPosition, Voxel_MinAABBs[2].xyz, Voxel_MaxAABBs[2].xyz, Voxel_Resolution, Voxel_DensityTexture2, Settings.DeepShadowDensityScale, 0, DistanceThreshold);
+		if (MacroGroupId == 3) Result = ComputeHairCountVoxel(SampleWorldPosition, LightPosition, Voxel_MinAABBs[3].xyz, Voxel_MaxAABBs[3].xyz, Voxel_Resolution, Voxel_DensityTexture3, Settings.DeepShadowDensityScale, 0, DistanceThreshold);
+		if (MacroGroupId == 4) Result = ComputeHairCountVoxel(SampleWorldPosition, LightPosition, Voxel_MinAABBs[4].xyz, Voxel_MaxAABBs[4].xyz, Voxel_Resolution, Voxel_DensityTexture4, Settings.DeepShadowDensityScale, 0, DistanceThreshold);
+		if (MacroGroupId == 5) Result = ComputeHairCountVoxel(SampleWorldPosition, LightPosition, Voxel_MinAABBs[5].xyz, Voxel_MaxAABBs[5].xyz, Voxel_Resolution, Voxel_DensityTexture5, Settings.DeepShadowDensityScale, 0, DistanceThreshold);
+		if (MacroGroupId == 6) Result = ComputeHairCountVoxel(SampleWorldPosition, LightPosition, Voxel_MinAABBs[6].xyz, Voxel_MaxAABBs[6].xyz, Voxel_Resolution, Voxel_DensityTexture6, Settings.DeepShadowDensityScale, 0, DistanceThreshold);
+		if (MacroGroupId == 7) Result = ComputeHairCountVoxel(SampleWorldPosition, LightPosition, Voxel_MinAABBs[7].xyz, Voxel_MaxAABBs[7].xyz, Voxel_Resolution, Voxel_DensityTexture7, Settings.DeepShadowDensityScale, 0, DistanceThreshold);
 		#endif
 
 		#if PERMUTATION_TRANSMITTANCE_TYPE == PERMUTATION_TRANSMITTANCE_VIRTUAL_VOXEL
@@ -342,16 +336,10 @@ void MainCS(uint2 DispatchThreadId : SV_DispatchThreadID)
 	#endif	
 	const float2 UV = (PixelCoord.xy + float2(0.5f, 0.5f)) / float2(View.BufferSizeAndInvSize.xy);
 	{
-		const FNodeData NodeData = HairVisibilityNodeData[LocalOffset];
-		float			SampleDepth;
-		FVisibilityData SampleVisData;
-		float3			SampleBaseColor;
-		float			SampleRoughness;
-		float			SampleSpecular;
-		uint			SampleLightChannelMask;
-		DecodeNodeData(NodeData, SampleDepth, SampleVisData, SampleBaseColor, SampleRoughness, SampleSpecular, SampleLightChannelMask);
-
-		if ((LightChannelMask & SampleLightChannelMask) == 0)
+		const FPackedHairSample NodeData = HairVisibilityNodeData[LocalOffset];
+		const FHairSample Sample = UnpackHairSample(NodeData);
+
+		if ((LightChannelMask & Sample.LightChannelMask) == 0)
 		{
 			OutputColor[LocalOffset] = InitNullPackedHairTransmittanceMask();
 			return;
@@ -359,12 +347,12 @@ void MainCS(uint2 DispatchThreadId : SV_DispatchThreadID)
 
 		FGBufferData GBuffer;
 		GBuffer.ShadingModelID = SHADINGMODELID_HAIR;
-		GBuffer.WorldNormal = SampleVisData.Tangent;
-		GBuffer.BaseColor = SampleBaseColor;
-		GBuffer.Roughness = SampleRoughness;
-		GBuffer.Specular = SampleSpecular;
+		GBuffer.WorldNormal = Sample.Tangent;
+		GBuffer.BaseColor = Sample.BaseColor;
+		GBuffer.Roughness = Sample.Roughness;
+		GBuffer.Specular = Sample.Specular;
 
-		const uint MacroGroupIndex = clamp(SampleVisData.MacroGroupId, 0, MAX_HAIR_MACROGROUP_COUNT - 1);
+		const uint MacroGroupIndex = clamp(Sample.MacroGroupId, 0, MAX_HAIR_MACROGROUP_COUNT - 1);
 
 		FTransmittanceSettings Settings;
 		Settings.LightDirection = LightDirection;
@@ -373,7 +361,7 @@ void MainCS(uint2 DispatchThreadId : SV_DispatchThreadID)
 		Settings.DeepShadowDepthBiasScale = DepthBiasScale;
 		Settings.Random = GetHairVoxelJitter(PixelCoord.xy, View.StateFrameIndexMod8);
 
-		const float SceneDepth = ConvertFromDeviceZ(SampleDepth);
+		const float SceneDepth = ConvertFromDeviceZ(Sample.Depth);
 		const float3 WorldPosition = ToWorldPosition(UV, SceneDepth);
 	
 		const FHairTransmittanceMask TransmittanceMask = ComputeTransmittanceVoxel(PixelCoord, WorldPosition, MacroGroupIndex, GBuffer, Settings, HairLUTTexture, LinearSampler, bDebugEnabled);
diff --git a/Engine/Shaders/Private/HairStrands/HairStrandsEnvironmentLighting.usf b/Engine/Shaders/Private/HairStrands/HairStrandsEnvironmentLighting.usf
--- a/Engine/Shaders/Private/HairStrands/HairStrandsEnvironmentLighting.usf
+++ b/Engine/Shaders/Private/HairStrands/HairStrandsEnvironmentLighting.usf
@@ -82,7 +82,7 @@ float3 GatherRadiance(float3 RayDirection, float Roughness, uint2 PixelPosition,
 		EyeIndex);
 }
 
-StructuredBuffer<FNodeData> HairVisibilityNodeData;
+StructuredBuffer<FPackedHairSample> HairVisibilityNodeData;
 
 Texture3D<float4>			HairScatteringLUTTexture;
 Texture3D<float4>			HairEnergyLUTTexture;
@@ -487,8 +487,9 @@ float3 ReflectionEnvironment(FGBufferData GBuffer, float2 ScreenPosition, float2
 		// Average weighting
 		#if 1
 		{
+			const float Backlit = GBuffer.CustomData.z;
 			HairAverageEnergy Energy = SampleHairEnergyLUT(HairEnergyLUTTexture, HairLUTSampler, GBuffer.BaseColor, GBuffer.Roughness, SinViewAngle);
-			ColorTT.rgb *= Energy.A_TT * TransmittanceMask.Transmittance * TransmittanceMask.Visibility;
+			ColorTT.rgb *= Energy.A_TT * TransmittanceMask.Transmittance * TransmittanceMask.Visibility * Backlit;
 		}
 		#endif	
 
@@ -514,9 +515,44 @@ float3 ReflectionEnvironment(FGBufferData GBuffer, float2 ScreenPosition, float2
 	return -min(-Color.rgb, 0.0);
 }
 
+#if USE_SCENE_COLOR
+Texture2D<float4> SceneColorTexture;
+Texture2D<uint4> HairCategorizationTexture;
+
+float3 ComputeThinTransmission(FGBufferData GBuffer, float2 ScreenPosition, float2 SvPosition, uint PrimitiveId)
+{
+#if USE_PREEXPOSURE
+	const float PreExposure = View.PreExposure;
+#else
+	const float PreExposure = 1.f;
+#endif
+	const float3 WorldPosition = mul(float4(ScreenPosition * GBuffer.Depth, GBuffer.Depth, 1), View.ScreenToWorld).xyz;
+	const float3 CameraToPixel = normalize(WorldPosition - View.WorldCameraOrigin);
+	const uint2  PixelPosition = SvPosition - ResolvedView.ViewRectMin.xy;
+
+	const float3 V = -CameraToPixel;
+	const float3 T = GBuffer.WorldNormal;
+
+	const FCategorizationData CatData = DecodeCategorizationData(HairCategorizationTexture.Load(uint3(PixelPosition, 0)));
+	const float PixelCoverage = min(CatData.PixelCoverage, 1);
+	
+	// The scene color is attenuated propertionnaly to the hair coverage. In order to retrieve the actual lighting buffer 
+	// value, we divide by the hair coverage
+	const float4 OccludedLuminance = SceneColorTexture.Load(uint3(PixelPosition, 0));
+	const float4 Luminance = OccludedLuminance / PixelCoverage;
+	
+	const float SinViewAngle = dot(T, V);
+	const HairAverageScattering Energy = SampleHairLUT(HairScatteringLUTTexture, HairLUTSampler, GBuffer.BaseColor, GBuffer.Roughness, SinViewAngle);
+	float3 Color = Energy.A_front * Luminance * PixelCoverage;
+	Color *= PreExposure;
+
+	return -min(-Color, 0.0);
+}
+
+#endif
+
 ////////////////////////////////////////////////////////////////////////////////////////////////////////////////////////////////////////////
 
-#if COMPUTESHADER
 uint MaxVisibilityNodeCount;
 StructuredBuffer<uint>	HairVisibilityNodeCoord;
 Buffer<uint4>			IndirectArgsBuffer;
@@ -546,31 +582,33 @@ void MainCS(uint3 DispatchThreadId : SV_DispatchThreadID)
 	if (LocalOffset > MaxVisibilityNodeCount)
 		return;
 
-	const FNodeData NodeData = HairVisibilityNodeData[LocalOffset];
-	float			SampleDepth;
-	FVisibilityData SampleVisData;
-	float3			SampleBaseColor;
-	float			SampleRoughness;
-	float			SampleSpecular;
-	uint			SampleLigthChannelMask;
-	DecodeNodeData(NodeData, SampleDepth, SampleVisData, SampleBaseColor, SampleRoughness, SampleSpecular, SampleLigthChannelMask);
+	const FPackedHairSample PackedSample = HairVisibilityNodeData[LocalOffset];
+	const FHairSample Sample = UnpackHairSample(PackedSample);
 
 	FGBufferData GBuffer;
 	GBuffer.ShadingModelID = SHADINGMODELID_HAIR;
-	GBuffer.WorldNormal = SampleVisData.Tangent;
-	GBuffer.BaseColor = SampleBaseColor;
-	GBuffer.Roughness = SampleRoughness;
-	GBuffer.Specular = SampleSpecular;
-	GBuffer.Depth = ConvertFromDeviceZ(SampleDepth);
+	GBuffer.WorldNormal = Sample.Tangent;
+	GBuffer.BaseColor = Sample.BaseColor;
+	GBuffer.Roughness = Sample.Roughness;
+	GBuffer.Specular = Sample.Specular;
+	GBuffer.Depth = ConvertFromDeviceZ(Sample.Depth);
 	GBuffer.GBufferAO = 1;
-	GBuffer.CustomData = float4(HairDualScatteringRoughnessOverride,0,0,0);
+	GBuffer.CustomData = float4(HairDualScatteringRoughnessOverride,0,Sample.Backlit,0);
 	GBuffer.IndirectIrradiance = 1000000;
 
-	const bool bIsMacroGroupCompatible = SampleVisData.MacroGroupId == Voxel_MacroGroupId;
-	if (bIsMacroGroupCompatible)
+	bool bEnabled = true;
+	const bool bIsMacroGroupCompatible = Sample.MacroGroupId == Voxel_MacroGroupId;
+#if USE_SCENE_COLOR
+	bEnabled = Sample.bScatterSceneLighting;
+#endif
+	if (bIsMacroGroupCompatible && bEnabled)
 	{
-		const float SampleCoverage = saturate(From8bitCoverage(SampleVisData.Coverage8bit));
-		const float3 SkyLighting = ReflectionEnvironment(GBuffer, ScreenPosition, SvPosition, SampleVisData.PrimitiveId);
+		const float SampleCoverage = saturate(From8bitCoverage(Sample.Coverage8bit));
+	#if USE_SCENE_COLOR
+		const float3 SkyLighting = ComputeThinTransmission(GBuffer, ScreenPosition, SvPosition, Sample.PrimitiveId);
+	#else
+		const float3 SkyLighting = ReflectionEnvironment(GBuffer, ScreenPosition, SvPosition, Sample.PrimitiveId);
+	#endif
 
 		FLightAccumulator LightAccumulator = (FLightAccumulator)0;
 		LightAccumulator_Add(LightAccumulator, SkyLighting, SkyLighting, 1.0f, false);
@@ -585,4 +623,3 @@ void MainCS(uint3 DispatchThreadId : SV_DispatchThreadID)
 		OutLightingBuffer[LocalOffset] = 0;
 	}
 }
-#endif
diff --git a/Engine/Shaders/Private/HairStrands/HairStrandsGuideDeform.usf b/Engine/Shaders/Private/HairStrands/HairStrandsGuideDeform.usf
--- a/Engine/Shaders/Private/HairStrands/HairStrandsGuideDeform.usf
+++ b/Engine/Shaders/Private/HairStrands/HairStrandsGuideDeform.usf
@@ -28,6 +28,7 @@ RWBuffer<uint4> OutSimDeformedPositionBuffer;
 #define DEFORMATION_NORMAL_DIRECTION 2
 #define DEFORMATION_OFFSET_GUIDE 3
 #define DEFORMATION_OFFSET_GUIDE_DYNAMIC 4
+#define DEFORMATION_OFFSET_GUIDE_GLOBAL 5
 
 #if PERMUTATION_DEFORMATION == DEFORMATION_OFFSET_GUIDE_DYNAMIC
 
@@ -58,6 +59,12 @@ float3 DeformGuideByTriangle(
 
 #endif
 
+#if PERMUTATION_DEFORMATION == DEFORMATION_OFFSET_GUIDE_GLOBAL
+uint SampleCount;
+Buffer<float4>  RestSamplePositionsBuffer;
+Buffer<float4>  MeshSampleWeightsBuffer;
+#endif
+
 float3 DisplacePosition(uint VertexIndex, float3 Pos, float Seed)
 {
 // By pass
@@ -100,6 +107,21 @@ float3 DisplacePosition(uint VertexIndex, float3 Pos, float Seed)
 
 	return DeformGuideByTriangle(VertexIndex, RestTriangle, DeformedTriangle, RootBarycentric) - SimDeformedOffset;
 #endif
+
+#if PERMUTATION_DEFORMATION == DEFORMATION_OFFSET_GUIDE_GLOBAL
+	const float3 RestControlPoint = Pos + SimRestOffset;
+	float3 ControlPoint = RestControlPoint;
+
+	// Apply rbf interpolation from the samples set
+	for(int i = 0; i < SampleCount; ++i)
+	{
+		const float3 PositionDelta = RestControlPoint - RestSamplePositionsBuffer[i];
+		const float FunctionValue = sqrt( dot(PositionDelta,PositionDelta)+1);
+		ControlPoint += FunctionValue * MeshSampleWeightsBuffer[i].xyz;
+	}
+	return ControlPoint - SimDeformedOffset;
+#endif
+
 }
 
 [numthreads(GROUP_SIZE, 1, 1)]
