Fix compilation issue

diff --git a/Engine/Plugins/Experimental/HairStrands/Source/HairStrandsCore/Private/HairStrandsVertexFactory.cpp b/Engine/Plugins/Experimental/HairStrands/Source/HairStrandsCore/Private/HairStrandsVertexFactory.cpp
--- a/Engine/Plugins/Experimental/HairStrands/Source/HairStrandsCore/Private/HairStrandsVertexFactory.cpp
+++ b/Engine/Plugins/Experimental/HairStrands/Source/HairStrandsCore/Private/HairStrandsVertexFactory.cpp
@@ -220,28 +220,6 @@ IMPLEMENT_VERTEX_FACTORY_PARAMETER_TYPE(FHairStrandsVertexFactory, SF_Compute,
 IMPLEMENT_VERTEX_FACTORY_PARAMETER_TYPE(FHairStrandsVertexFactory, SF_RayHitGroup,	FHairStrandsVertexFactoryShaderParameters);
 #endif
 
-FVertexFactoryShaderParameters* FHairStrandsVertexFactory::ConstructShaderParameters(EShaderFrequency ShaderFrequency)
-{
-	if (ShaderFrequency == SF_Vertex)
-	{
-		return new FHairStrandsVertexFactoryShaderParameters();
-	}
-
-	if (ShaderFrequency == SF_Pixel)
-	{
-		return new FHairStrandsVertexFactoryShaderParameters();
-	}
-
-#if RHI_RAYTRACING
-	if (ShaderFrequency == SF_RayHitGroup)
-	{
-		return new FHairStrandsVertexFactoryShaderParameters();
-	}
-#endif // RHI_RAYTRACING
-
-	return NULL;
-}
-
 void FHairStrandsVertexFactory::ReleaseRHI()
 {
 	FVertexFactory::ReleaseRHI();
