Fix niagara ref

diff --git a/Engine/Plugins/Experimental/HairStrands/Source/HairStrandsNiagara/Private/NiagaraDataInterfaceHairStrands.cpp b/Engine/Plugins/Experimental/HairStrands/Source/HairStrandsNiagara/Private/NiagaraDataInterfaceHairStrands.cpp
--- a/Engine/Plugins/Experimental/HairStrands/Source/HairStrandsNiagara/Private/NiagaraDataInterfaceHairStrands.cpp
+++ b/Engine/Plugins/Experimental/HairStrands/Source/HairStrandsNiagara/Private/NiagaraDataInterfaceHairStrands.cpp
@@ -703,14 +703,13 @@ struct FNDIHairStrandsParametersCS : public FNiagaraDataInterfaceParametersCS
 
 			const int32 SampleCountValue = (MeshProjection != nullptr) ? MeshProjection->SampleCount : 0;
 			FShaderResourceViewRHIRef RestSamplePositionsBufferSRV = (MeshProjection != nullptr && MeshProjection->SampleCount > 0) ?
-				MeshProjection->RestSamplePositionsBuffer.SRV : FNiagaraRenderer::GetDummyFloatBuffer();
+				MeshProjection->RestSamplePositionsBuffer.SRV.GetReference() : FNiagaraRenderer::GetDummyFloatBuffer();
 			FShaderResourceViewRHIRef MeshSampleWeightsBufferSRV = (MeshProjection != nullptr && MeshProjection->SampleCount > 0) ?
-				MeshProjection->MeshSampleWeightsBuffer.SRV : FNiagaraRenderer::GetDummyFloatBuffer();
+				MeshProjection->MeshSampleWeightsBuffer.SRV.GetReference() : FNiagaraRenderer::GetDummyFloatBuffer();
 
 			FRHIShaderResourceView* RootBarycentricCoordinatesSRV = (HasRootAttachedValue == 1) ?
 				MeshProjection->RootTriangleBarycentricBuffer.SRV.GetReference() : FNiagaraRenderer::GetDummyFloatBuffer();
 
-
 			int32 bNeedSimReset = (ProxyData->TickCount <= MaxDelay ? 1 : 0);
 
 			//UE_LOG(LogHairStrands, Log, TEXT("Shader Reset : %d %d %d"), bNeedSimReset, IsRootValid, HasRootAttachedValue);
