Fix hair CIS

diff --git a/Engine/Plugins/Experimental/HairStrands/Source/HairStrandsCore/Private/GroomAsset.cpp b/Engine/Plugins/Experimental/HairStrands/Source/HairStrandsCore/Private/GroomAsset.cpp
--- a/Engine/Plugins/Experimental/HairStrands/Source/HairStrandsCore/Private/GroomAsset.cpp
+++ b/Engine/Plugins/Experimental/HairStrands/Source/HairStrandsCore/Private/GroomAsset.cpp
@@ -1576,5 +1576,3 @@ bool FProcessedHairDescription::IsValid() const
 }
 
 
-#undef LOCTEXT_NAMESPACE
-
diff --git a/Engine/Plugins/Experimental/HairStrands/Source/HairStrandsCore/Public/GroomAsset.h b/Engine/Plugins/Experimental/HairStrands/Source/HairStrandsCore/Public/GroomAsset.h
--- a/Engine/Plugins/Experimental/HairStrands/Source/HairStrandsCore/Public/GroomAsset.h
+++ b/Engine/Plugins/Experimental/HairStrands/Source/HairStrandsCore/Public/GroomAsset.h
@@ -217,7 +217,7 @@ struct FHairStrandsDeformedResource : public FRenderResource
 	const FHairStrandsDatas::FRenderData& RenderData;
 
 	/* Whether the GPU data should be initialized with the asset data or not */
-	const bool bInitializedData;
+	const bool bInitializedData = false;
 
 	/* Whether the GPU data should be initialized with the asset data or not */
 	uint32 CurrentIndex = 0;
@@ -575,11 +575,11 @@ private:
 
 	TUniquePtr<FHairDescription> HairDescription;
 	TUniquePtr<FHairDescriptionBulkData> HairDescriptionBulkData;
-	bool bIsInitialized;
 
 	UPROPERTY()
 	bool bIsCacheable = true;
 #endif
+	bool bIsInitialized;
 };
 
 USTRUCT(BlueprintType)
@@ -672,6 +672,8 @@ public:
 
 	/**  Part of UObject interface  */
 	virtual void PostEditChangeProperty(FPropertyChangedEvent& PropertyChangedEvent) override;
+	
+#endif // WITH_EDITOR
 
 	enum class EQueryStatus
 	{
@@ -680,7 +682,6 @@ public:
 		Completed
 	};
 	EQueryStatus QueryStatus = EQueryStatus::None;
-#endif // WITH_EDITOR
 
 	/** Initialize resources. */
 	void InitResource();
diff --git a/Engine/Plugins/Experimental/HairStrands/Source/HairStrandsEditor/Private/GroomBindingActions.cpp b/Engine/Plugins/Experimental/HairStrands/Source/HairStrandsEditor/Private/GroomBindingActions.cpp
--- a/Engine/Plugins/Experimental/HairStrands/Source/HairStrandsEditor/Private/GroomBindingActions.cpp
+++ b/Engine/Plugins/Experimental/HairStrands/Source/HairStrandsEditor/Private/GroomBindingActions.cpp
@@ -25,9 +25,9 @@ void FGroomBindingActions::GetActions(const TArray<UObject*>& InObjects, FToolMe
 	TArray<TWeakObjectPtr<UGroomBindingAsset>> GroomBindingAssets = GetTypedWeakObjectPtrs<UGroomBindingAsset>(InObjects);
 
 	Section.AddMenuEntry(
-		"RebuildGroom",
-		LOCTEXT("RebuildGroom", "Rebuild"),
-		LOCTEXT("RebuildGroomTooltip", "Rebuild the groom binding"),
+		"RebuildGroomBinding",
+		LOCTEXT("RebuildGroomBinding", "Rebuild"),
+		LOCTEXT("RebuildGroomBindingTooltip", "Rebuild the groom binding"),
 		FSlateIcon(FEditorStyle::GetStyleSetName(), "ContentBrowser.AssetActions"),
 		FUIAction(
 			FExecuteAction::CreateSP(this, &FGroomBindingActions::ExecuteRebuildBindingAsset, GroomBindingAssets),
diff --git a/Engine/Plugins/Experimental/HairStrands/Source/HairStrandsEditor/Private/GroomCreateBindingOptionsWindow.cpp b/Engine/Plugins/Experimental/HairStrands/Source/HairStrandsEditor/Private/GroomCreateBindingOptionsWindow.cpp
--- a/Engine/Plugins/Experimental/HairStrands/Source/HairStrandsEditor/Private/GroomCreateBindingOptionsWindow.cpp
+++ b/Engine/Plugins/Experimental/HairStrands/Source/HairStrandsEditor/Private/GroomCreateBindingOptionsWindow.cpp
@@ -148,7 +148,7 @@ TSharedPtr<SGroomCreateBindingOptionsWindow> DisplayOptions(UGroomCreateBindingO
 
 TSharedPtr<SGroomCreateBindingOptionsWindow> SGroomCreateBindingOptionsWindow::DisplayCreateBindingOptions(UGroomCreateBindingOptions* BindingOptions)
 {
-	return DisplayOptions(BindingOptions, EGroomBindingOptionsVisibility::BuildOptions, LOCTEXT("WindowTitle", "Groom Binding Options"), LOCTEXT("Build", "Create"));
+	return DisplayOptions(BindingOptions, EGroomBindingOptionsVisibility::BuildOptions, LOCTEXT("GroomBindingRebuildWindowTitle", "Groom Binding Options"), LOCTEXT("Build", "Create"));
 }
 
 static UObject* InternalCreateNewBindAsset(FName InAssetName, UObject* InParent, UGroomAsset* GroomAsset, USkeletalMesh* SourceSkelMesh, USkeletalMesh* TargetSkelMesh, const int32 NumInterpolationPoints)
diff --git a/Engine/Plugins/Experimental/HairStrands/Source/HairStrandsEditor/Private/GroomImportOptionsWindow.cpp b/Engine/Plugins/Experimental/HairStrands/Source/HairStrandsEditor/Private/GroomImportOptionsWindow.cpp
--- a/Engine/Plugins/Experimental/HairStrands/Source/HairStrandsEditor/Private/GroomImportOptionsWindow.cpp
+++ b/Engine/Plugins/Experimental/HairStrands/Source/HairStrandsEditor/Private/GroomImportOptionsWindow.cpp
@@ -76,7 +76,7 @@ void SGroomImportOptionsWindow::Construct(const FArguments& InArgs)
 	
 	const EHairDescriptionStatus Status = GetStatus(GroupsPreview);
 
-	FText ValidationText = LOCTEXT("GroomOptionsWindow_ValidationText", "Unknown");
+	FText ValidationText;
 	FLinearColor ValidationColor(1,1,1);
 	if (Status == EHairDescriptionStatus::Valid)
 	{
@@ -93,6 +93,10 @@ void SGroomImportOptionsWindow::Construct(const FArguments& InArgs)
 		ValidationText = LOCTEXT("GroomOptionsWindow_ValidationText", "Invalid. The groom does not contain any group.");
 		ValidationColor = FLinearColor(1, 0, 0, 1);
 	}
+	else
+	{
+		ValidationText = LOCTEXT("GroomOptionsWindow_ValidationText", "Unknown");
+	}
 
 	this->ChildSlot
 	[
@@ -291,12 +295,12 @@ TSharedPtr<SGroomImportOptionsWindow> DisplayOptions(UGroomImportOptions* Import
 
 TSharedPtr<SGroomImportOptionsWindow> SGroomImportOptionsWindow::DisplayImportOptions(UGroomImportOptions* ImportOptions, const FString& FilePath, const FProcessedHairDescription* ProcessedDescription)
 {
-	return DisplayOptions(ImportOptions, FilePath, ProcessedDescription, EGroomOptionsVisibility::All, LOCTEXT("WindowTitle", "Groom Import Options"), LOCTEXT("Import", "Import"));
+	return DisplayOptions(ImportOptions, FilePath, ProcessedDescription, EGroomOptionsVisibility::All, LOCTEXT("GroomImportWindowTitle", "Groom Import Options"), LOCTEXT("Import", "Import"));
 }
 
 TSharedPtr<SGroomImportOptionsWindow> SGroomImportOptionsWindow::DisplayRebuildOptions(UGroomImportOptions* ImportOptions, const FString& FilePath)
 {
-	return DisplayOptions(ImportOptions, FilePath, nullptr, EGroomOptionsVisibility::BuildOptions, LOCTEXT("WindowTitle", "Groom Build Options"), LOCTEXT("Build", "Build"));
+	return DisplayOptions(ImportOptions, FilePath, nullptr, EGroomOptionsVisibility::BuildOptions, LOCTEXT("GroomRebuildWindowTitle ", "Groom Build Options"), LOCTEXT("Build", "Build"));
 }
 
 
diff --git a/Engine/Plugins/Experimental/HairStrands/Source/HairStrandsEditor/Private/HairStrandsFactory.cpp b/Engine/Plugins/Experimental/HairStrands/Source/HairStrandsEditor/Private/HairStrandsFactory.cpp
--- a/Engine/Plugins/Experimental/HairStrands/Source/HairStrandsEditor/Private/HairStrandsFactory.cpp
+++ b/Engine/Plugins/Experimental/HairStrands/Source/HairStrandsEditor/Private/HairStrandsFactory.cpp
@@ -80,7 +80,7 @@ UObject* UHairStrandsFactory::FactoryCreateFile(UClass* InClass, UObject* InPare
 		FProcessedHairDescription OutDescription;
 		if (bRunValidation)
 		{
-			FScopedSlowTask Progress((float)1, LOCTEXT("ImportHairAsset", "Importing hair asset for preview..."), true);
+			FScopedSlowTask Progress((float)1, LOCTEXT("ImportHairAssetForPreview", "Importing hair asset for preview..."), true);
 			Progress.MakeDialog(true);
 
 			FHairDescription HairDescription;
diff --git a/Engine/Source/Runtime/Renderer/Private/HairStrands/HairStrandsDebug.cpp b/Engine/Source/Runtime/Renderer/Private/HairStrands/HairStrandsDebug.cpp
--- a/Engine/Source/Runtime/Renderer/Private/HairStrands/HairStrandsDebug.cpp
+++ b/Engine/Source/Runtime/Renderer/Private/HairStrands/HairStrandsDebug.cpp
@@ -1645,13 +1645,13 @@ void RenderHairStrandsDebugInfo(
 			if (GHairDebugMeshProjection_SkinCacheMesh > 0)
 			{
 				FViewInfo* LocalView = &View;
-				auto RenderMeshProjection = [&bClearDepth, WorldType, LocalView, SkinCache, &SceneColorTexture, &DepthTexture](FRDGBuilder& GraphBuilder, EHairStrandsProjectionMeshType MeshType)
+				auto RenderMeshProjection = [&bClearDepth, WorldType, LocalView, SkinCache, &SceneColorTexture, &DepthTexture](FRDGBuilder& LocalGraphBuilder, EHairStrandsProjectionMeshType MeshType)
 				{
 					FHairStrandsProjectionMeshData::LOD MeshProjectionLODData;
 					GetGroomInterpolationData(WorldType, MeshType, SkinCache, MeshProjectionLODData);
 					for (FHairStrandsProjectionMeshData::Section& Section : MeshProjectionLODData.Sections)
 					{
-						AddDebugProjectionMeshPass(GraphBuilder, LocalView, MeshType, bClearDepth, Section, SceneColorTexture, DepthTexture);
+						AddDebugProjectionMeshPass(LocalGraphBuilder, LocalView, MeshType, bClearDepth, Section, SceneColorTexture, DepthTexture);
 						bClearDepth = false;
 					}
 				};
