Fix CIS issue.

diff --git a/Engine/Plugins/Experimental/HairStrands/Source/HairStrandsEditor/Private/GroomImportOptionsWindow.cpp b/Engine/Plugins/Experimental/HairStrands/Source/HairStrandsEditor/Private/GroomImportOptionsWindow.cpp
--- a/Engine/Plugins/Experimental/HairStrands/Source/HairStrandsEditor/Private/GroomImportOptionsWindow.cpp
+++ b/Engine/Plugins/Experimental/HairStrands/Source/HairStrandsEditor/Private/GroomImportOptionsWindow.cpp
@@ -80,22 +80,22 @@ void SGroomImportOptionsWindow::Construct(const FArguments& InArgs)
 	FLinearColor ValidationColor(1,1,1);
 	if (Status == EHairDescriptionStatus::Valid)
 	{
-		ValidationText = LOCTEXT("GroomOptionsWindow_ValidationText", "Valid");
+		ValidationText = LOCTEXT("GroomOptionsWindow_ValidationText0", "Valid");
 		ValidationColor = FLinearColor(0, 0.80f, 0, 1);
 	}
 	else if (Status == EHairDescriptionStatus::NoCurve)
 	{
-		ValidationText = LOCTEXT("GroomOptionsWindow_ValidationText", "Invalid. Some groups have 0 curves.");
+		ValidationText = LOCTEXT("GroomOptionsWindow_ValidationText1", "Invalid. Some groups have 0 curves.");
 		ValidationColor = FLinearColor(0.80f, 0, 0, 1);
 	}
 	else if (Status == EHairDescriptionStatus::NoGroup)
 	{
-		ValidationText = LOCTEXT("GroomOptionsWindow_ValidationText", "Invalid. The groom does not contain any group.");
+		ValidationText = LOCTEXT("GroomOptionsWindow_ValidationText2", "Invalid. The groom does not contain any group.");
 		ValidationColor = FLinearColor(1, 0, 0, 1);
 	}
 	else
 	{
-		ValidationText = LOCTEXT("GroomOptionsWindow_ValidationText", "Unknown");
+		ValidationText = LOCTEXT("GroomOptionsWindow_ValidationText3", "Unknown");
 	}
 
 	this->ChildSlot
