Fix CIS

diff --git a/Engine/Plugins/Experimental/HairStrands/Source/HairStrandsNiagara/Private/NiagaraDataInterfaceHairStrands.cpp b/Engine/Plugins/Experimental/HairStrands/Source/HairStrandsNiagara/Private/NiagaraDataInterfaceHairStrands.cpp
--- a/Engine/Plugins/Experimental/HairStrands/Source/HairStrandsNiagara/Private/NiagaraDataInterfaceHairStrands.cpp
+++ b/Engine/Plugins/Experimental/HairStrands/Source/HairStrandsNiagara/Private/NiagaraDataInterfaceHairStrands.cpp
@@ -687,18 +687,18 @@ struct FNDIHairStrandsParametersCS : public FNiagaraDataInterfaceParametersCS
 
 			const FHairStrandsRootResource::FMeshProjectionLOD* MeshProjection = (HasRootAttachedValue == 1) ? &(HairStrandsBuffer->SourceRootResources->MeshProjectionLODs[0]) : nullptr;
 
-			FRHIShaderResourceView* RestTrianglePositionASRV = (HasRootAttachedValue == 1) ?
+			FRHIShaderResourceView* RestTrianglePositionASRV = (HasRootAttachedValue == 1 && MeshProjection != nullptr) ?
 				MeshProjection->RestRootTrianglePosition0Buffer.SRV.GetReference() : FNiagaraRenderer::GetDummyFloatBuffer();
-			FRHIShaderResourceView* RestTrianglePositionBSRV = (HasRootAttachedValue == 1) ?
+			FRHIShaderResourceView* RestTrianglePositionBSRV = (HasRootAttachedValue == 1 && MeshProjection != nullptr) ?
 				MeshProjection->RestRootTrianglePosition1Buffer.SRV.GetReference() : FNiagaraRenderer::GetDummyFloatBuffer();
-			FRHIShaderResourceView* RestTrianglePositionCSRV = (HasRootAttachedValue == 1) ?
+			FRHIShaderResourceView* RestTrianglePositionCSRV = (HasRootAttachedValue == 1 && MeshProjection != nullptr) ?
 				MeshProjection->RestRootTrianglePosition2Buffer.SRV.GetReference() : FNiagaraRenderer::GetDummyFloatBuffer();
 
-			FRHIShaderResourceView* DeformedTrianglePositionASRV = (HasRootAttachedValue == 1) ?
+			FRHIShaderResourceView* DeformedTrianglePositionASRV = (HasRootAttachedValue == 1 && MeshProjection != nullptr) ?
 				MeshProjection->DeformedRootTrianglePosition0Buffer.SRV.GetReference() : FNiagaraRenderer::GetDummyFloatBuffer();
-			FRHIShaderResourceView* DeformedTrianglePositionBSRV = (HasRootAttachedValue == 1) ?
+			FRHIShaderResourceView* DeformedTrianglePositionBSRV = (HasRootAttachedValue == 1 && MeshProjection != nullptr) ?
 				MeshProjection->DeformedRootTrianglePosition1Buffer.SRV.GetReference() : FNiagaraRenderer::GetDummyFloatBuffer();
-			FRHIShaderResourceView* DeformedTrianglePositionCSRV = (HasRootAttachedValue == 1) ?
+			FRHIShaderResourceView* DeformedTrianglePositionCSRV = (HasRootAttachedValue == 1 && MeshProjection != nullptr) ?
 				MeshProjection->DeformedRootTrianglePosition2Buffer.SRV.GetReference() : FNiagaraRenderer::GetDummyFloatBuffer();
 
 			const int32 SampleCountValue = (MeshProjection != nullptr) ? MeshProjection->SampleCount : 0;
@@ -707,7 +707,7 @@ struct FNDIHairStrandsParametersCS : public FNiagaraDataInterfaceParametersCS
 			FShaderResourceViewRHIRef MeshSampleWeightsBufferSRV = (MeshProjection != nullptr && MeshProjection->SampleCount > 0) ?
 				MeshProjection->MeshSampleWeightsBuffer.SRV.GetReference() : FNiagaraRenderer::GetDummyFloatBuffer();
 
-			FRHIShaderResourceView* RootBarycentricCoordinatesSRV = (HasRootAttachedValue == 1) ?
+			FRHIShaderResourceView* RootBarycentricCoordinatesSRV = (HasRootAttachedValue == 1 && MeshProjection != nullptr) ?
 				MeshProjection->RootTriangleBarycentricBuffer.SRV.GetReference() : FNiagaraRenderer::GetDummyFloatBuffer();
 
 			int32 bNeedSimReset = (ProxyData->TickCount <= MaxDelay ? 1 : 0);
