Fix CIS

diff --git a/Engine/Plugins/Experimental/HairStrands/Source/HairStrandsCore/Public/GroomAsset.h b/Engine/Plugins/Experimental/HairStrands/Source/HairStrandsCore/Public/GroomAsset.h
--- a/Engine/Plugins/Experimental/HairStrands/Source/HairStrandsCore/Public/GroomAsset.h
+++ b/Engine/Plugins/Experimental/HairStrands/Source/HairStrandsCore/Public/GroomAsset.h
@@ -681,7 +681,7 @@ public:
 		Submitted,
 		Completed
 	};
-	EQueryStatus QueryStatus = EQueryStatus::None;
+	volatile EQueryStatus QueryStatus = EQueryStatus::None;
 
 	/** Initialize resources. */
 	void InitResource();
