Piped stabelized rasterization to the hair public VF data.

diff --git a/Engine/Plugins/Experimental/HairStrands/Source/HairStrandsCore/Private/HairStrandsRendering.cpp b/Engine/Plugins/Experimental/HairStrands/Source/HairStrandsCore/Private/HairStrandsRendering.cpp
--- a/Engine/Plugins/Experimental/HairStrands/Source/HairStrandsCore/Private/HairStrandsRendering.cpp
+++ b/Engine/Plugins/Experimental/HairStrands/Source/HairStrandsCore/Private/HairStrandsRendering.cpp
@@ -964,6 +964,7 @@ void ComputeHairStrandsInterpolation(
 			Output.HairGroupPublicData->VFInput.VertexCount = Output.VFInput.VertexCount;
 			Output.HairGroupPublicData->VFInput.HairRadius = Output.VFInput.HairRadius;
 			Output.HairGroupPublicData->VFInput.HairLength = Output.VFInput.HairLength;
+			Output.HairGroupPublicData->VFInput.bUseStableRasterization = Output.VFInput.bUseStableRasterization;
 			Output.HairGroupPublicData->VFInput.HairDensity = Output.VFInput.HairDensity;
 			Output.HairGroupPublicData->VFInput.LocalToWorldTransform = LocalToWorld;
 
diff --git a/Engine/Source/Runtime/Renderer/Public/HairStrandsInterface.h b/Engine/Source/Runtime/Renderer/Public/HairStrandsInterface.h
--- a/Engine/Source/Runtime/Renderer/Public/HairStrandsInterface.h
+++ b/Engine/Source/Runtime/Renderer/Public/HairStrandsInterface.h
@@ -93,6 +93,7 @@ public:
 		float HairRadius = 0;
 		float HairLength = 0;
 		float HairDensity = 0;
+		bool bUseStableRasterization = false;
 		FTransform LocalToWorldTransform;
 	};
 	VertexFactoryInput VFInput;
