Remove hair description populate function when loading binding asset as these information are already serialized.

diff --git a/Engine/Plugins/Experimental/HairStrands/Source/HairStrandsCore/Private/GroomAsset.cpp b/Engine/Plugins/Experimental/HairStrands/Source/HairStrandsCore/Private/GroomAsset.cpp
--- a/Engine/Plugins/Experimental/HairStrands/Source/HairStrandsCore/Private/GroomAsset.cpp
+++ b/Engine/Plugins/Experimental/HairStrands/Source/HairStrandsCore/Private/GroomAsset.cpp
@@ -1262,21 +1262,6 @@ bool UGroomAsset::CacheDerivedData(const FGroomBuildSettings* InBuildSettings)
 
 /////////////////////////////////////////////////////////////////////////////////////////
 
-static void InitializeGroupInfos(
-	const TArray<UGroomBindingAsset::FHairGroupData>& HairGroupDatas,
-	TArray<FGoomBindingGroupInfo>& GroupInfos)
-{
-	GroupInfos.Empty();
-	for (const UGroomBindingAsset::FHairGroupData& Data : HairGroupDatas)
-	{
-		FGoomBindingGroupInfo& Info = GroupInfos.AddDefaulted_GetRef();
-		Info.SimRootCount	= Data.SimRootData.RootCount;
-		Info.SimLODCount	= Data.SimRootData.MeshProjectionLODs.Num();
-		Info.RenRootCount	= Data.RenRootData.RootCount;
-		Info.RenLODCount	= Data.RenRootData.MeshProjectionLODs.Num();
-	}
-}
-
 FArchive& operator<<(FArchive& Ar, UGroomBindingAsset::FHairGroupData& GroupData)
 {
 	GroupData.SimRootData.Serialize(Ar);
@@ -1367,8 +1352,6 @@ void UGroomBindingAsset::PostLoad()
 {
 	Super::PostLoad();
 
-	InitializeGroupInfos(HairGroupDatas, GroupInfos);
-
 	if (Groom)
 	{
 		// Make sure that the asset initialized its resources first since the component needs them to initialize its own resources
@@ -1393,7 +1376,6 @@ void UGroomBindingAsset::PreSave(const class ITargetPlatform* TargetPlatform)
 	}
 #endif
 	Super::PreSave(TargetPlatform);
-	InitializeGroupInfos(HairGroupDatas, GroupInfos);
 #if WITH_EDITOR
 	OnGroomBindingAssetChanged.Broadcast();
 #endif
@@ -1402,7 +1384,6 @@ void UGroomBindingAsset::PreSave(const class ITargetPlatform* TargetPlatform)
 void UGroomBindingAsset::PostSaveRoot(bool bCleanupIsRequired)
 {
 	Super::PostSaveRoot(bCleanupIsRequired);
-	InitializeGroupInfos(HairGroupDatas, GroupInfos);
 #if WITH_EDITOR
 	OnGroomBindingAssetChanged.Broadcast();
 #endif
diff --git a/Engine/Plugins/Experimental/HairStrands/Source/HairStrandsCore/Private/HairStrandsRendering.cpp b/Engine/Plugins/Experimental/HairStrands/Source/HairStrandsCore/Private/HairStrandsRendering.cpp
--- a/Engine/Plugins/Experimental/HairStrands/Source/HairStrandsCore/Private/HairStrandsRendering.cpp
+++ b/Engine/Plugins/Experimental/HairStrands/Source/HairStrandsCore/Private/HairStrandsRendering.cpp
@@ -1321,6 +1321,17 @@ static void InternalProcessGroomBindingTask(FRHICommandListImmediate& RHICmdList
 		Resource.RenRootResources->InitRHI();
 	}
 
+	TArray<FGoomBindingGroupInfo>& OutGroupInfos = BindingAsset->GroupInfos;
+	OutGroupInfos.Empty();
+	for (const UGroomBindingAsset::FHairGroupData& Data : OutHairGroupDatas)
+	{
+		FGoomBindingGroupInfo& Info = OutGroupInfos.AddDefaulted_GetRef();
+		Info.SimRootCount	= Data.SimRootData.RootCount;
+		Info.SimLODCount	= Data.SimRootData.MeshProjectionLODs.Num();
+		Info.RenRootCount	= Data.RenRootData.RootCount;
+		Info.RenLODCount	= Data.RenRootData.MeshProjectionLODs.Num();
+	}
+
 	FHairStrandsProjectionHairData RenProjectionDatas;
 	FHairStrandsProjectionHairData SimProjectionDatas;
 	const uint32 GroupCount = OutHairGroupResources.Num();
