Change hair default material and hair debug material to not use the special material flag and reduce minimal shader being required to compile the engine.

diff --git a/Engine/Config/BaseEngine.ini b/Engine/Config/BaseEngine.ini
--- a/Engine/Config/BaseEngine.ini
+++ b/Engine/Config/BaseEngine.ini
@@ -80,8 +80,6 @@ UnrealEdEngine=/Script/UnrealEd.UnrealEdEngine
 ;LargeFontName=/Engine/EngineFonts/Roboto.Roboto
 ;SubtitleFontName=/Engine/EngineFonts/Roboto.Roboto
 WireframeMaterialName=/Engine/EngineDebugMaterials/WireframeMaterial.WireframeMaterial
-HairDefaultMaterialName=/Engine/EngineMaterials/HairDefaultMaterial.HairDefaultMaterial
-HairDebugMaterialName=/Engine/EngineMaterials/HairDebugMaterial.HairDebugMaterial
 DefaultMaterialName=/Engine/EngineMaterials/WorldGridMaterial.WorldGridMaterial
 DefaultLightFunctionMaterialName=/Engine/EngineMaterials/DefaultLightFunctionMaterial.DefaultLightFunctionMaterial
 DefaultTextureName=/Engine/EngineResources/DefaultTexture.DefaultTexture
@@ -842,8 +840,6 @@ bFullyCompressStartupPackages=false
 +Package=/Engine/EngineDebugMaterials/VertexColorViewMode_GreenOnly
 +Package=/Engine/EngineDebugMaterials/VertexColorViewMode_RedOnly
 +Package=/Engine/EngineDebugMaterials/WireframeMaterial
-+Package=/Engine/EngineDebugMaterials/HairDefaultMaterial
-+Package=/Engine/EngineDebugMaterials/HairDebugMaterial
 
 +Package=/Engine/EngineSounds/WhiteNoise
 
diff --git a/Engine/Plugins/Experimental/HairStrands/Source/HairStrandsCore/Private/GroomComponent.cpp b/Engine/Plugins/Experimental/HairStrands/Source/HairStrandsCore/Private/GroomComponent.cpp
--- a/Engine/Plugins/Experimental/HairStrands/Source/HairStrandsCore/Private/GroomComponent.cpp
+++ b/Engine/Plugins/Experimental/HairStrands/Source/HairStrandsCore/Private/GroomComponent.cpp
@@ -105,6 +105,72 @@ public:
 	}
 };
 
+enum class EHairMaterialCompatibility : uint8
+{
+	Valid,
+	Invalid_UsedWithHairStrands,
+	Invalid_ShadingModel,
+	Invalid_BlendMode,
+	Invalid_IsNull
+};
+
+static EHairMaterialCompatibility IsHairMaterialCompatible(UMaterialInterface* Material, ERHIFeatureLevel::Type FeatureLevel)
+{
+	if (Material)
+	{
+		FMaterialResource* MaterialResources = Material->GetMaterialResource(FeatureLevel);
+		if (!MaterialResources)
+		{
+			return EHairMaterialCompatibility::Invalid_IsNull;
+		}
+		if (!MaterialResources->IsUsedWithHairStrands())
+		{
+			return EHairMaterialCompatibility::Invalid_UsedWithHairStrands;
+		}
+		if (!MaterialResources->GetShadingModels().HasShadingModel(MSM_Hair))
+		{
+			return EHairMaterialCompatibility::Invalid_ShadingModel;
+		}
+		if (MaterialResources->GetBlendMode() != BLEND_Opaque)
+		{
+			return EHairMaterialCompatibility::Invalid_BlendMode;
+		}
+	}
+	else
+	{
+		return EHairMaterialCompatibility::Invalid_IsNull;
+	}
+
+	return EHairMaterialCompatibility::Valid;
+}
+
+static UMaterialInterface* GetHairDefaultMaterial()
+{
+#if WITH_EDITOR
+	static UMaterialInterface* HairDefaultMaterial = nullptr;
+	if (!HairDefaultMaterial)
+	{
+		HairDefaultMaterial = LoadObject<UMaterialInterface>(NULL, TEXT("/HairStrands/Materials/HairDefaultMaterial.HairDefaultMaterial"));
+	}
+	return HairDefaultMaterial;
+#else
+	return nullptr;
+#endif
+}
+
+static UMaterialInterface* GetHairDebugMaterial()
+{
+#if WITH_EDITOR
+	static UMaterialInterface* HairDebugMaterial = nullptr;
+	if (!HairDebugMaterial)
+	{
+		HairDebugMaterial = LoadObject<UMaterialInterface>(NULL, TEXT("/HairStrands/Materials/HairDebugMaterial.HairDebugMaterial"));
+	}
+	return HairDebugMaterial;
+#else
+	return nullptr;
+#endif
+}
 /////////////////////////////////////////////////////////////////////////////////////////
 /////////////////////////////////////////////////////////////////////////////////////////
 //  FStrandHairSceneProxy
@@ -149,7 +215,7 @@ public:
 			UMaterialInterface* Material = Component->GetMaterial(GroupIt);
 			if (Material == nullptr || !Material->GetMaterialResource(GetScene().GetFeatureLevel())->IsUsedWithHairStrands())
 			{
-				Material = GEngine->HairDefaultMaterial;
+				Material =  UMaterial::GetDefaultMaterial(MD_Surface);
 			}
 
 			#if RHI_RAYTRACING
@@ -211,7 +277,9 @@ public:
 	virtual void GetDynamicMeshElements(const TArray<const FSceneView*>& Views, const FSceneViewFamily& ViewFamily, uint32 VisibilityMap, FMeshElementCollector& Collector) const override
 	{
 		if (!VertexFactory.GetData().InterpolationOutput)
+		{
 			return;
+		}
 
 		bool bHasOneElementValid = false;
 		for (FHairStrandsInterpolationOutput::HairGroup& HairGroup : VertexFactory.GetData().InterpolationOutput->HairGroups)
@@ -222,8 +290,10 @@ public:
 				break;
 			}
 		}
-		if (!bHasOneElementValid) 
+		if (!bHasOneElementValid)
+		{
 			return;
+		}
 
 		const uint32 GroupCount = VertexFactory.GetData().InterpolationOutput->HairGroups.Num();
 
@@ -258,10 +328,9 @@ public:
 				HairMaxRadius = FMath::Max(HairMaxRadius, HairGroup.VFInput.HairRadius);
 			}
 
+			UMaterialInterface* HairDebugMaterial = GetHairDebugMaterial();
 			const float HairClipLength = GetHairClipLength();
-			auto DebugMaterial = new FHairDebugModeMaterialRenderProxy(
-				GEngine->HairDebugMaterial ? GEngine->HairDebugMaterial->GetRenderProxy() : nullptr,
-				DebugModeScalar, 0, HairMaxRadius, HairClipLength);
+			auto DebugMaterial = new FHairDebugModeMaterialRenderProxy(HairDebugMaterial ? HairDebugMaterial->GetRenderProxy() : nullptr, DebugModeScalar, 0, HairMaxRadius, HairClipLength);
 			Collector.RegisterOneFrameMaterialProxy(DebugMaterial);
 			MaterialProxy = DebugMaterial;
 		}
@@ -270,7 +339,9 @@ public:
 		{
 			const FSceneView* View = Views[ViewIndex];
 			if (View->bIsReflectionCapture || View->bIsPlanarReflection)
+			{
 				continue;
+			}
 
 			if (VisibilityMap & (1 << ViewIndex))
 			{
@@ -279,13 +350,19 @@ public:
 					const HairGroup& GroupData = HairGroups[GroupIt];
 					const uint32 HairVertexCount = VertexFactory.GetData().InterpolationOutput->HairGroups[GroupIt].VFInput.VertexCount;
 
+					FMaterialRenderProxy* MaterialRenderProxy = MaterialProxy == nullptr ? GroupData.Material->GetRenderProxy() : MaterialProxy;
+					if (MaterialRenderProxy == nullptr)
+					{
+						continue;
+					}
+
 					// Draw the mesh.
 					FMeshBatch& Mesh = Collector.AllocateMesh();
 					FMeshBatchElement& BatchElement = Mesh.Elements[0];
 					BatchElement.IndexBuffer = nullptr;
 					Mesh.bWireframe = false;
 					Mesh.VertexFactory = &VertexFactory;
-					Mesh.MaterialRenderProxy = MaterialProxy == nullptr ? GroupData.Material->GetRenderProxy() : MaterialProxy;
+					Mesh.MaterialRenderProxy = MaterialRenderProxy;
 					bool bHasPrecomputedVolumetricLightmap;
 					FMatrix PreviousLocalToWorld;
 					int32 SingleCaptureIndex;
@@ -665,27 +742,30 @@ int32 UGroomComponent::GetNumMaterials() const
 UMaterialInterface* UGroomComponent::GetMaterial(int32 ElementIndex) const
 {
 	UMaterialInterface* OverrideMaterial = Super::GetMaterial(ElementIndex);
+	
+	bool bUseHairDefaultMaterial = false;
 
 	const ERHIFeatureLevel::Type FeatureLevel = GetScene() ? GetScene()->GetFeatureLevel() : ERHIFeatureLevel::Num;
 	if (!OverrideMaterial && GroomAsset && ElementIndex < GroomAsset->GetNumHairGroups() && FeatureLevel != ERHIFeatureLevel::Num)
 	{
 		if (UMaterialInterface* Material = GroomAsset->HairGroupsInfo[ElementIndex].Material)
 		{
-			if (!Material->GetMaterialResource(FeatureLevel)->IsUsedWithHairStrands())
-			{
-				Material = GEngine->HairDefaultMaterial;
-			}
-			return Material;
+			OverrideMaterial = Material;
 		}
 		else
 		{
-			return GEngine->HairDefaultMaterial;
+			bUseHairDefaultMaterial = true;
 		}
 	}
 
-	if (OverrideMaterial == nullptr || OverrideMaterial->GetMaterialResource(FeatureLevel) == nullptr || !OverrideMaterial->GetMaterialResource(FeatureLevel)->IsUsedWithHairStrands())
+	if (IsHairMaterialCompatible(OverrideMaterial, FeatureLevel) != EHairMaterialCompatibility::Valid)
+	{
+		bUseHairDefaultMaterial = true;
+	}
+
+	if (bUseHairDefaultMaterial)
 	{
-		OverrideMaterial = GEngine->HairDefaultMaterial;
+		OverrideMaterial = GetHairDefaultMaterial();
 	}
 
 	return OverrideMaterial;
@@ -1442,9 +1522,28 @@ void UGroomComponent::GetUsedMaterials(TArray<UMaterialInterface*>& OutMaterials
 #if WITH_EDITOR
 	UMeshComponent::GetUsedMaterials(OutMaterials, bGetDebugMaterials);
 	if (bGetDebugMaterials)
-		OutMaterials.Add(GEngine->HairDebugMaterial);
+	{
+		UMaterialInterface* HairDebugMaterial = GetHairDebugMaterial();
+		OutMaterials.Add(HairDebugMaterial);
+	}
+
+	const ERHIFeatureLevel::Type FeatureLevel = GetScene() ? GetScene()->GetFeatureLevel() : ERHIFeatureLevel::Num;
+	bool bRegisterDefaultMaterial = false;
+	for (UMaterialInterface* Material : OutMaterials)
+	{
+		if (IsHairMaterialCompatible(Material, FeatureLevel) != EHairMaterialCompatibility::Valid)
+		{
+			bRegisterDefaultMaterial = true;
+			break;
+		}
+	}
+
+	if (bRegisterDefaultMaterial)
+	{
+		UMaterialInterface* HairDefaultMaterial = GetHairDefaultMaterial();
+		OutMaterials.Add(HairDefaultMaterial);		
+	}
 #endif
-	OutMaterials.Add(GEngine->HairDefaultMaterial);
 }
 
 #if WITH_EDITOR
@@ -1589,6 +1688,7 @@ bool UGroomComponent::CanEditChange(const FProperty* InProperty) const
 #endif
 
 #if WITH_EDITOR
+
 void UGroomComponent::ValidateMaterials(bool bMapCheck) const
 {
 	if (!GroomAsset)
@@ -1605,76 +1705,70 @@ void UGroomComponent::ValidateMaterials(bool bMapCheck) const
 	for (uint32 MaterialIt = 0, MaterialCount = GetNumMaterials(); MaterialIt < MaterialCount; ++MaterialIt)
 	{
 		UMaterialInterface* OverrideMaterial = Super::GetMaterial(MaterialIt);
-
-		FMaterialResource* Material = nullptr;
-
-		if (OverrideMaterial)
+		if (!OverrideMaterial && MaterialIt < uint32(GroomAsset->HairGroupsInfo.Num()) && GroomAsset->HairGroupsInfo[MaterialIt].Material)
 		{
-			Material = OverrideMaterial->GetMaterialResource(FeatureLevel);
-		}
-		else if (MaterialIt < uint32(GroomAsset->HairGroupsInfo.Num()) && GroomAsset->HairGroupsInfo[MaterialIt].Material)
-		{
-			Material = GroomAsset->HairGroupsInfo[MaterialIt].Material->GetMaterialResource(FeatureLevel);
+			OverrideMaterial = GroomAsset->HairGroupsInfo[MaterialIt].Material;
 		}
 
-		if (Material)
+		const EHairMaterialCompatibility Result = IsHairMaterialCompatible(OverrideMaterial, FeatureLevel);
+		switch (Result)
 		{
-			if (!Material->IsUsedWithHairStrands())
+			case EHairMaterialCompatibility::Invalid_UsedWithHairStrands:
 			{
 				if (bMapCheck)
 				{
 					FMessageLog("MapCheck").Warning()
 						->AddToken(FUObjectToken::Create(GroomAsset))
-						->AddToken(FTextToken::Create(LOCTEXT("MapCheck_Message_HairStrandsMissingUseHairStrands", "Groom's material needs to enable the UseHairStrands option. Groom's material will be replaced with default hair strands shader.")))
+						->AddToken(FTextToken::Create(LOCTEXT("MapCheck_Message_HairStrandsMissingUseHairStrands", "Groom's material needs to enable the UseHairStrands option. Groom's material will be replaced with default hair strands shader in editor.")))
 						->AddToken(FMapErrorToken::Create(FMapErrors::InvalidHairStrandsMaterial));
 				}
 				else
 				{
-					UE_LOG(LogHairStrands, Warning, TEXT("[Groom] %s - Groom's material needs to enable the UseHairStrands option. Groom's material will be replaced with default hair strands shader."), *Name);
+					UE_LOG(LogHairStrands, Warning, TEXT("[Groom] %s - Groom's material needs to enable the UseHairStrands option. Groom's material will be replaced with default hair strands shader in editor."), *Name);
 				}
-			}
-			if (!Material->GetShadingModels().HasShadingModel(MSM_Hair))
+			} break;
+			case EHairMaterialCompatibility::Invalid_ShadingModel:
 			{
 				if (bMapCheck)
 				{
 					FMessageLog("MapCheck").Warning()
 						->AddToken(FUObjectToken::Create(GroomAsset))
-						->AddToken(FTextToken::Create(LOCTEXT("MapCheck_Message_HairStrandsInvalidShadingModel", "Groom's material needs to have Hair shading model. Groom's material will be replaced with default hair strands shader.")))
+						->AddToken(FTextToken::Create(LOCTEXT("MapCheck_Message_HairStrandsInvalidShadingModel", "Groom's material needs to have Hair shading model. Groom's material will be replaced with default hair strands shader in editor.")))
 						->AddToken(FMapErrorToken::Create(FMapErrors::InvalidHairStrandsMaterial));
 				}
 				else
 				{
-					UE_LOG(LogHairStrands, Warning, TEXT("[Groom] %s - Groom's material needs to have Hair shading model. Groom's material will be replaced with default hair strands shader."), *Name);
+					UE_LOG(LogHairStrands, Warning, TEXT("[Groom] %s - Groom's material needs to have Hair shading model. Groom's material will be replaced with default hair strands shader in editor."), *Name);
 				}
-			}
-			if (Material->GetBlendMode() != BLEND_Opaque)
+			}break;
+			case EHairMaterialCompatibility::Invalid_BlendMode:
 			{
 				if (bMapCheck)
 				{
 					FMessageLog("MapCheck").Warning()
 						->AddToken(FUObjectToken::Create(GroomAsset))
-						->AddToken(FTextToken::Create(LOCTEXT("MapCheck_Message_HairStrandsInvalidBlendMode", "Groom's material needs to have Opaque blend mode. Groom's material will be replaced with default hair strands shader.")))
+						->AddToken(FTextToken::Create(LOCTEXT("MapCheck_Message_HairStrandsInvalidBlendMode", "Groom's material needs to have Opaque blend mode. Groom's material will be replaced with default hair strands shader in editor.")))
 						->AddToken(FMapErrorToken::Create(FMapErrors::InvalidHairStrandsMaterial));
 				}
 				else
 				{
-					UE_LOG(LogHairStrands, Warning, TEXT("[Groom] %s - Groom's material needs to have Opaque blend mode. Groom's material will be replaced with default hair strands shader."), *Name);
+					UE_LOG(LogHairStrands, Warning, TEXT("[Groom] %s - Groom's material needs to have Opaque blend mode. Groom's material will be replaced with default hair strands shader in editor."), *Name);
 				}
-			}
-		}
-		else
-		{
-			if (bMapCheck)
-			{
-				FMessageLog("MapCheck").Info()
-					->AddToken(FUObjectToken::Create(GroomAsset))
-					->AddToken(FTextToken::Create(LOCTEXT("MapCheck_Message_HairStrandsMissingMaterial", "Groom's material is not set and will fallback on default hair strands shader.")))
-					->AddToken(FMapErrorToken::Create(FMapErrors::InvalidHairStrandsMaterial));
-			}
-			else
+			}break;
+			case EHairMaterialCompatibility::Invalid_IsNull:
 			{
-				UE_LOG(LogHairStrands, Warning, TEXT("[Groom] %s - Groom's material is not set and will fallback on default hair strands shader."), *Name);
-			}
+				if (bMapCheck)
+				{
+					FMessageLog("MapCheck").Info()
+						->AddToken(FUObjectToken::Create(GroomAsset))
+						->AddToken(FTextToken::Create(LOCTEXT("MapCheck_Message_HairStrandsMissingMaterial", "Groom's material is not set and will fallback on default hair strands shader in editor.")))
+						->AddToken(FMapErrorToken::Create(FMapErrors::InvalidHairStrandsMaterial));
+				}
+				else
+				{
+					UE_LOG(LogHairStrands, Warning, TEXT("[Groom] %s - Groom's material is not set and will fallback on default hair strands shader in editor."), *Name);
+				}
+			}break;
 		}
 	}
 }
diff --git a/Engine/Source/Runtime/Engine/Classes/Engine/Engine.h b/Engine/Source/Runtime/Engine/Classes/Engine/Engine.h
--- a/Engine/Source/Runtime/Engine/Classes/Engine/Engine.h
+++ b/Engine/Source/Runtime/Engine/Classes/Engine/Engine.h
@@ -862,22 +862,6 @@ public:
 	UPROPERTY(globalconfig)
 	FString WireframeMaterialName;
 
-	/** The material used to default hair meshes. */
-	UPROPERTY()
-	class UMaterial* HairDefaultMaterial;
-
-	/** @todo document */
-	UPROPERTY(globalconfig)
-	FString HairDefaultMaterialName;
-
-	/** The material used to debug hair meshes. */
-	UPROPERTY()
-	class UMaterial* HairDebugMaterial;
-
-	/** @todo document */
-	UPROPERTY(globalconfig)
-	FString HairDebugMaterialName;
-
 #if WITH_EDITORONLY_DATA
 	/** A translucent material used to render things in geometry mode. */
 	UPROPERTY()
diff --git a/Engine/Source/Runtime/Engine/Private/UnrealEngine.cpp b/Engine/Source/Runtime/Engine/Private/UnrealEngine.cpp
--- a/Engine/Source/Runtime/Engine/Private/UnrealEngine.cpp
+++ b/Engine/Source/Runtime/Engine/Private/UnrealEngine.cpp
@@ -2465,7 +2465,6 @@ void UEngine::InitializeObjectReferences()
 	{
 		// Materials that are needed in-game if debug viewmodes are allowed
 		LoadSpecialMaterial(TEXT("WireframeMaterialName"), WireframeMaterialName, WireframeMaterial, true);
-		LoadSpecialMaterial(TEXT("HairDebugMaterialName"), HairDebugMaterialName, HairDebugMaterial, true);
 		LoadSpecialMaterial(TEXT("LevelColorationLitMaterialName"), LevelColorationLitMaterialName, LevelColorationLitMaterial, true);
 		LoadSpecialMaterial(TEXT("LevelColorationUnlitMaterialName"), LevelColorationUnlitMaterialName, LevelColorationUnlitMaterial, true);
 		LoadSpecialMaterial(TEXT("LightingTexelDensityName"), LightingTexelDensityName, LightingTexelDensityMaterial, false);
@@ -2483,7 +2482,6 @@ void UEngine::InitializeObjectReferences()
 	LoadSpecialMaterial(TEXT("RemoveSurfaceMaterialName"), RemoveSurfaceMaterialName.ToString(), RemoveSurfaceMaterial, false);
 
 	// these one's are needed both editor and standalone 
-	LoadSpecialMaterial(TEXT("HairDefaultMaterialName"), HairDefaultMaterialName, HairDefaultMaterial, false);
 	LoadSpecialMaterial(TEXT("DebugMeshMaterialName"), DebugMeshMaterialName.ToString(), DebugMeshMaterial, false);
 	LoadSpecialMaterial(TEXT("EmissiveMeshMaterialName"), EmissiveMeshMaterialName.ToString(), EmissiveMeshMaterial, false);
 	LoadSpecialMaterial(TEXT("InvalidLightmapSettingsMaterialName"), InvalidLightmapSettingsMaterialName.ToString(), InvalidLightmapSettingsMaterial, false);
