Change hair interpolation passes for resources transitions.

diff --git a/Engine/Plugins/Experimental/HairStrands/Source/HairStrandsCore/Private/HairStrandsRendering.cpp b/Engine/Plugins/Experimental/HairStrands/Source/HairStrandsCore/Private/HairStrandsRendering.cpp
--- a/Engine/Plugins/Experimental/HairStrands/Source/HairStrandsCore/Private/HairStrandsRendering.cpp
+++ b/Engine/Plugins/Experimental/HairStrands/Source/HairStrandsCore/Private/HairStrandsRendering.cpp
@@ -498,7 +498,7 @@ static void AddHairStrandsInterpolationPass(
 	TShaderMapRef<FHairInterpolationCS> ComputeShader(ShaderMap, PermutationVector);
 	FComputeShaderUtils::AddPass(
 		GraphBuilder,
-		RDG_EVENT_NAME("HairStrandsInterpolation"),
+		RDG_EVENT_NAME("HairStrandsInterpolationDirect"),
 		ComputeShader,
 		Parameters,
 		DispatchCount);
@@ -775,30 +775,67 @@ void ComputeHairStrandsInterpolation(
 	int32 LODIndex,
 	FHairStrandClusterData* ClusterData)
 {
-	if (!InInput || !InOutput) return;
+	// Note: We are breaking this code up into several, larger for loops. In the previous version, the typical code path was:
+	// for each group:
+	//     AddClearClusterAABBPass()
+	//     AddHairStrandsInterpolationPass()
+	//     AddHairClusterAABBPass()
+	//     AddHairTangentPass()
+	//     AddGenerateRaytracingGeometryPass()
+	// 
+	// The problem is that it creates bubbles in the GPU, since each pass was dependent on the previous one. So it has been
+	// modified to be of the form:
+	//
+	// for each group:
+	//     AddClearClusterAABBPass()
+	// for each group:
+	//     AddHairStrandsInterpolationPass()
+	// ...
 
+	if (!InInput || !InOutput) return;
 
-	const uint32 GroupCount = InOutput->HairGroups.Num();
-	for (uint32 GroupIndex = 0; GroupIndex < GroupCount; ++GroupIndex)
+	// The previous loop would return if both Input and Output were not valid. Instead, count the number
+	// of valid groups first.
+	const uint32 ExpectedGroupCount = InOutput->HairGroups.Num();
+	uint32 GroupCount = ExpectedGroupCount;
+	for (uint32 GroupIndex = 0; GroupIndex < ExpectedGroupCount; ++GroupIndex)
 	{
 		FHairStrandsInterpolationInput::FHairGroup& Input  = InInput->HairGroups[GroupIndex];
 		FHairStrandsInterpolationOutput::HairGroup& Output = InOutput->HairGroups[GroupIndex];
 		Output.VFInput.Reset();
 
-		if (!Input.IsValid() || !Output.IsValid()) return;
+		if (!Input.IsValid() || !Output.IsValid())
+		{
+			GroupCount = FMath::Min(ExpectedGroupCount, GroupIndex);
+			break;
+		}
+	}
 
-		DECLARE_GPU_STAT(HairStrandsInterpolation);
-		SCOPED_DRAW_EVENT(RHICmdList, HairStrandsInterpolation);
-		SCOPED_GPU_STAT(RHICmdList, HairStrandsInterpolation);
+	FRDGBuilder GraphBuilder(RHICmdList);
 
-		FRDGBuilder GraphBuilder(RHICmdList);
+	DECLARE_GPU_STAT(HairStrandsInterpolationCluster);
+	SCOPED_DRAW_EVENT(RHICmdList, HairStrandsInterpolationCluster);
+	SCOPED_GPU_STAT(RHICmdList, HairStrandsInterpolationCluster);
 
-		const uint32 CurrIndex = *Output.CurrentIndex;
-		const uint32 PrevIndex = (CurrIndex + 1) % 2;
+	const EDeformationType DeformationType = GetDeformationType();
+
+	// Debug mode:
+	// * None	: Display hair normally
+	// * Sim	: Show sim strands
+	// * Render : Show rendering strands with sim color influence
+	const EHairStrandsDebugMode DebugMode = GetHairStrandsDebugStrandsMode();
+	const bool bDebugModePatchedAttributeBuffer = DebugMode == EHairStrandsDebugMode::RenderHairStrands || DebugMode == EHairStrandsDebugMode::RenderVisCluster;
 
-		const EDeformationType DeformationType = GetDeformationType();
-		if (DeformationType != EDeformationType::RestStrands && DeformationType != EDeformationType::Simulation)
+	if (DeformationType != EDeformationType::RestStrands && DeformationType != EDeformationType::Simulation)
+	{
+		for (uint32 GroupIndex = 0; GroupIndex < GroupCount; ++GroupIndex)
 		{
+			FHairStrandsInterpolationInput::FHairGroup& Input = InInput->HairGroups[GroupIndex];
+			FHairStrandsInterpolationOutput::HairGroup& Output = InOutput->HairGroups[GroupIndex];
+
+			const uint32 CurrIndex = *Output.CurrentIndex;
+			const uint32 PrevIndex = (CurrIndex + 1) % 2;
+
 			AddDeformSimHairStrandsPass(
 				GraphBuilder,
 				DeformationType,
@@ -808,31 +845,30 @@ void ComputeHairStrandsInterpolation(
 				Input.SimRestPosePositionBuffer->SRV,
 				Input.SimRootPointIndexBuffer ? Input.SimRootPointIndexBuffer->SRV : nullptr,
 				Output.SimDeformedPositionBuffer[CurrIndex]->UAV, Input.InSimHairPositionOffset,
-				Input.OutHairPositionOffset );
+				Input.OutHairPositionOffset);
 		}
+	}
 
-		// If the deformation is driven by the physics simulation, then the output is always the 0 index
-		const uint32 SimIndex = CurrIndex;// GHairDeformationType == 0 ? 0 : CurrIndex;
-
-		// Debug mode:
-		// * None	: Display hair normally
-		// * Sim	: Show sim strands
-		// * Render : Show rendering strands with sim color influence
-		const EHairStrandsDebugMode DebugMode = GetHairStrandsDebugStrandsMode();
-		if (DebugMode == EHairStrandsDebugMode::SimHairStrands)
+	if (DebugMode == EHairStrandsDebugMode::SimHairStrands)
+	{
+		for (uint32 GroupIndex = 0; GroupIndex < GroupCount; ++GroupIndex)
 		{
+			FHairStrandsInterpolationInput::FHairGroup& Input = InInput->HairGroups[GroupIndex];
+			FHairStrandsInterpolationOutput::HairGroup& Output = InOutput->HairGroups[GroupIndex];
+
+			const uint32 CurrIndex = *Output.CurrentIndex;
+			const uint32 PrevIndex = (CurrIndex + 1) % 2;
+
 			AddHairTangentPass(
 				GraphBuilder,
 				Input.SimVertexCount,
-				Output.SimDeformedPositionBuffer[SimIndex]->SRV,
+				Output.SimDeformedPositionBuffer[CurrIndex]->SRV,
 				Output.SimTangentBuffer->UAV);
 
-			GraphBuilder.Execute();
-
 			const bool bHasSimulationEnabled = Input.bIsSimulationEnable && GHairStrandsInterpolateSimulation && DeformationType != EDeformationType::RestStrands;
 
-			Output.VFInput.HairPositionBuffer = Output.SimDeformedPositionBuffer[SimIndex]->SRV;
-			Output.VFInput.HairPreviousPositionBuffer = Output.SimDeformedPositionBuffer[SimIndex]->SRV;
+			Output.VFInput.HairPositionBuffer = Output.SimDeformedPositionBuffer[CurrIndex]->SRV;
+			Output.VFInput.HairPreviousPositionBuffer = Output.SimDeformedPositionBuffer[CurrIndex]->SRV;
 			Output.VFInput.HairTangentBuffer = Output.SimTangentBuffer->SRV;
 			Output.VFInput.HairAttributeBuffer = Input.SimAttributeBuffer->SRV;
 			Output.VFInput.HairMaterialBuffer = Output.RenderMaterialBuffer->SRV;
@@ -845,10 +881,22 @@ void ComputeHairStrandsInterpolation(
 			Output.VFInput.bUseStableRasterization = Input.GroupDesc.bUseStableRasterization;
 			Output.VFInput.bScatterSceneLighting = Input.GroupDesc.bScatterSceneLighting;
 		}
-		else
+	}
+	else
+	{
+		for (uint32 GroupIndex = 0; GroupIndex < GroupCount; ++GroupIndex)
 		{
+			FHairStrandsInterpolationInput::FHairGroup& Input = InInput->HairGroups[GroupIndex];
+			FHairStrandsInterpolationOutput::HairGroup& Output = InOutput->HairGroups[GroupIndex];
+
+			const uint32 CurrIndex = *Output.CurrentIndex;
+			const uint32 PrevIndex = (CurrIndex + 1) % 2;
+
+			// If the deformation is driven by the physics simulation, then the output is always the 0 index
+			const uint32 SimIndex = CurrIndex;// GHairDeformationType == 0 ? 0 : CurrIndex;
+
+
 			check(ClusterData);
-			const bool bDebugModePatchedAttributeBuffer = DebugMode == EHairStrandsDebugMode::RenderHairStrands || DebugMode == EHairStrandsDebugMode::RenderVisCluster;
 
 			const uint32 BufferSizeInBytes = Input.RenderVertexCount * FHairStrandsAttributeFormat::SizeInByte;
 			if (bDebugModePatchedAttributeBuffer && Output.RenderPatchedAttributeBuffer.NumBytes != BufferSizeInBytes)
@@ -862,6 +910,18 @@ void ComputeHairStrandsInterpolation(
 				Input.ClusterCount,
 				Output.RenderClusterAABBBuffer->UAV,
 				Output.RenderGroupAABBBuffer->UAV);
+		}
+
+		for (uint32 GroupIndex = 0; GroupIndex < GroupCount; ++GroupIndex)
+		{
+			FHairStrandsInterpolationInput::FHairGroup& Input = InInput->HairGroups[GroupIndex];
+			FHairStrandsInterpolationOutput::HairGroup& Output = InOutput->HairGroups[GroupIndex];
+
+			const uint32 CurrIndex = *Output.CurrentIndex;
+			const uint32 PrevIndex = (CurrIndex + 1) % 2;
+
+			// If the deformation is driven by the physics simulation, then the output is always the 0 index
+			const uint32 SimIndex = CurrIndex;// GHairDeformationType == 0 ? 0 : CurrIndex;
 
 			FHairScaleAndClipDesc ScaleAndClipDesc;
 			ScaleAndClipDesc.InHairLength = Input.GroupDesc.HairLength;
@@ -897,9 +957,23 @@ void ComputeHairStrandsInterpolation(
 				Input.SimAttributeBuffer->SRV,
 				Output.RenderDeformedPositionBuffer[CurrIndex]->UAV,
 				Output.RenderPatchedAttributeBuffer.UAV,
-				Input.VertexToClusterIdBuffer->SRV, 
+				Input.VertexToClusterIdBuffer->SRV,
 				Input.SimRootPointIndexBuffer->SRV);
 
+		}
+
+
+		for (uint32 GroupIndex = 0; GroupIndex < GroupCount; ++GroupIndex)
+		{
+			FHairStrandsInterpolationInput::FHairGroup& Input = InInput->HairGroups[GroupIndex];
+			FHairStrandsInterpolationOutput::HairGroup& Output = InOutput->HairGroups[GroupIndex];
+
+			const uint32 CurrIndex = *Output.CurrentIndex;
+			const uint32 PrevIndex = (CurrIndex + 1) % 2;
+
+			// If the deformation is driven by the physics simulation, then the output is always the 0 index
+			const uint32 SimIndex = CurrIndex;// GHairDeformationType == 0 ? 0 : CurrIndex;
+
 			// Initialize group cluster data for culling by the renderer
 			FHairStrandClusterData::FHairGroup& HairGroupCluster = ClusterData->HairGroups.Emplace_GetRef();
 			HairGroupCluster.ClusterCount = Input.ClusterCount;
@@ -914,7 +988,11 @@ void ComputeHairStrandsInterpolation(
 			HairGroupCluster.LodBias = Input.GroupDesc.LodBias;
 			HairGroupCluster.LodAverageVertexPerPixel = Input.GroupDesc.LodAverageVertexPerPixel;
 
-			Output.VFInput.HairRadius = ScaleAndClipDesc.MaxOutHairRadius;
+			// Note: This code needs to exactly match the values FHairScaleAndClipDesc set int the previous loop.
+			const float OutHairRadius = (GStrandHairWidth > 0 ? GStrandHairWidth : Input.GroupDesc.HairWidth) * 0.5f;
+			const float MaxOutHairRadius = OutHairRadius * FMath::Max(1.f, FMath::Max(Input.GroupDesc.HairRootScale, Input.GroupDesc.HairTipScale));
+
+			Output.VFInput.HairRadius = MaxOutHairRadius;
 			Output.VFInput.HairLength = Input.GroupDesc.HairLength;
 			Output.VFInput.HairDensity = Input.GroupDesc.HairShadowDensity;
 			Output.VFInput.HairPositionBuffer = Output.RenderDeformedPositionBuffer[CurrIndex]->SRV;
@@ -928,14 +1006,32 @@ void ComputeHairStrandsInterpolation(
 				Input.OutHairPositionOffset,
 				HairGroupCluster,
 				Output.RenderDeformedPositionBuffer[CurrIndex]->SRV);
-			
+		}
+
+		for (uint32 GroupIndex = 0; GroupIndex < GroupCount; ++GroupIndex)
+		{
+			FHairStrandsInterpolationInput::FHairGroup& Input = InInput->HairGroups[GroupIndex];
+			FHairStrandsInterpolationOutput::HairGroup& Output = InOutput->HairGroups[GroupIndex];
+
+			const uint32 CurrIndex = *Output.CurrentIndex;
+			const uint32 PrevIndex = (CurrIndex + 1) % 2;
+
 			AddHairTangentPass(
 				GraphBuilder,
 				Input.RenderVertexCount,
 				Output.VFInput.HairPositionBuffer,
 				Output.RenderTangentBuffer->UAV);
+		}
+
+		for (uint32 GroupIndex = 0; GroupIndex < GroupCount; ++GroupIndex)
+		{
+			FHairStrandsInterpolationInput::FHairGroup& Input = InInput->HairGroups[GroupIndex];
+			FHairStrandsInterpolationOutput::HairGroup& Output = InOutput->HairGroups[GroupIndex];
+
+			const uint32 CurrIndex = *Output.CurrentIndex;
+			const uint32 PrevIndex = (CurrIndex + 1) % 2;
 
-			#if RHI_RAYTRACING
+#if RHI_RAYTRACING
 			if (IsHairRayTracingEnabled())
 			{
 				const float HairRadiusScaleRT = (GHairRaytracingRadiusScale > 0 ? GHairRaytracingRadiusScale : Input.GroupDesc.HairRaytracingRadiusScale);
@@ -947,15 +1043,23 @@ void ComputeHairStrandsInterpolation(
 					Output.VFInput.HairPositionBuffer,
 					Input.RaytracingPositionBuffer->UAV);
 			}
-			#endif
-			GraphBuilder.Execute();
+#endif
 
-			Output.VFInput.HairTangentBuffer			= Output.RenderTangentBuffer->SRV;
-			Output.VFInput.HairAttributeBuffer			= bDebugModePatchedAttributeBuffer ? Output.RenderPatchedAttributeBuffer.SRV : Input.RenderAttributeBuffer->SRV;
-			Output.VFInput.HairMaterialBuffer			= Output.RenderMaterialBuffer->SRV;
-			Output.VFInput.HairPositionOffset			= Input.OutHairPositionOffset;
-			Output.VFInput.HairPreviousPositionOffset	= Input.OutHairPreviousPositionOffset;
-			Output.VFInput.VertexCount					= Input.RenderVertexCount;
+			Output.VFInput.HairTangentBuffer = Output.RenderTangentBuffer->SRV;
+			Output.VFInput.HairAttributeBuffer = bDebugModePatchedAttributeBuffer ? Output.RenderPatchedAttributeBuffer.SRV : Input.RenderAttributeBuffer->SRV;
+			Output.VFInput.HairMaterialBuffer = Output.RenderMaterialBuffer->SRV;
+			Output.VFInput.HairPositionOffset = Input.OutHairPositionOffset;
+			Output.VFInput.HairPreviousPositionOffset = Input.OutHairPreviousPositionOffset;
+			Output.VFInput.VertexCount = Input.RenderVertexCount;
+		}
+
+		for (uint32 GroupIndex = 0; GroupIndex < GroupCount; ++GroupIndex)
+		{
+			FHairStrandsInterpolationInput::FHairGroup& Input = InInput->HairGroups[GroupIndex];
+			FHairStrandsInterpolationOutput::HairGroup& Output = InOutput->HairGroups[GroupIndex];
+
+			const uint32 CurrIndex = *Output.CurrentIndex;
+			const uint32 PrevIndex = (CurrIndex + 1) % 2;
 
 			// TODO: find a more robust way to handle parameters passing to compute raster.
 			// At the moment there is a loose compling which will break if the vertex factor change.
@@ -969,7 +1073,7 @@ void ComputeHairStrandsInterpolation(
 			Output.HairGroupPublicData->VFInput.HairDensity = Output.VFInput.HairDensity;
 			Output.HairGroupPublicData->VFInput.LocalToWorldTransform = LocalToWorld;
 
-			#if RHI_RAYTRACING
+#if RHI_RAYTRACING
 			if (IsHairRayTracingEnabled())
 			{
 				FRHIUnorderedAccessView* UAV = Input.RaytracingPositionBuffer->UAV;
@@ -982,8 +1086,20 @@ void ComputeHairStrandsInterpolation(
 					UpdateHairAccelerationStructure(RHICmdList, Input.RaytracingGeometry);
 				Input.bIsRTGeometryInitialized = true;
 			}
-			#endif
+#endif
 		}
+	}
+
+	GraphBuilder.Execute();
+
+	// update the current index
+	for (uint32 GroupIndex = 0; GroupIndex < GroupCount; ++GroupIndex)
+	{
+		FHairStrandsInterpolationInput::FHairGroup& Input  = InInput->HairGroups[GroupIndex];
+		FHairStrandsInterpolationOutput::HairGroup& Output = InOutput->HairGroups[GroupIndex];
+
+		const uint32 CurrIndex = *Output.CurrentIndex;
+		const uint32 PrevIndex = (CurrIndex + 1) % 2;
 
 		*Output.CurrentIndex = PrevIndex;
 	}
diff --git a/Engine/Source/Runtime/Renderer/Private/HairStrands/HairStrandsInterface.cpp b/Engine/Source/Runtime/Renderer/Private/HairStrands/HairStrandsInterface.cpp
--- a/Engine/Source/Runtime/Renderer/Private/HairStrands/HairStrandsInterface.cpp
+++ b/Engine/Source/Runtime/Renderer/Private/HairStrands/HairStrandsInterface.cpp
@@ -323,9 +323,9 @@ void RunHairStrandsInterpolation(
 {
 	check(IsInRenderingThread());
 
-	DECLARE_GPU_STAT(HairStrandsInterpolation);
-	SCOPED_DRAW_EVENT(RHICmdList, HairStrandsInterpolation);
-	SCOPED_GPU_STAT(RHICmdList, HairStrandsInterpolation);
+	DECLARE_GPU_STAT(HairStrandsInterpolationGrouped);
+	SCOPED_DRAW_EVENT(RHICmdList, HairStrandsInterpolationGrouped);
+	SCOPED_GPU_STAT(RHICmdList, HairStrandsInterpolationGrouped);
 
 	// Update dynamic mesh triangles
 	for (FHairStrandsManager::Element& E : GHairManager.Elements)
