Fix groom asset not being set to nullptr.

diff --git a/Engine/Plugins/Experimental/HairStrands/Source/HairStrandsCore/Private/GroomComponent.cpp b/Engine/Plugins/Experimental/HairStrands/Source/HairStrandsCore/Private/GroomComponent.cpp
--- a/Engine/Plugins/Experimental/HairStrands/Source/HairStrandsCore/Private/GroomComponent.cpp
+++ b/Engine/Plugins/Experimental/HairStrands/Source/HairStrandsCore/Private/GroomComponent.cpp
@@ -573,6 +573,10 @@ void UGroomComponent::SetGroomAsset(UGroomAsset* Asset)
 	{
 		GroomAsset = Asset;
 	}
+	else
+	{
+		GroomAsset = nullptr;
+	}
 
 	if (!UGroomBindingAsset::IsBindingAssetValid(BindingAsset) || !UGroomBindingAsset::IsCompatible(GroomAsset, BindingAsset))
 	{
@@ -580,7 +584,7 @@ void UGroomComponent::SetGroomAsset(UGroomAsset* Asset)
 	}
 
 	UpdateHairGroupsDesc();
-	if (!Asset)
+	if (!GroomAsset)
 		return;
 	InitResources();
 }
@@ -593,6 +597,10 @@ void UGroomComponent::SetGroomAsset(UGroomAsset* Asset, UGroomBindingAsset* InBi
 	{
 		GroomAsset = Asset;
 	}
+	else
+	{
+		GroomAsset = nullptr;
+	}
 	BindingAsset = InBinding;
 	if (!UGroomBindingAsset::IsBindingAssetValid(BindingAsset) || !UGroomBindingAsset::IsCompatible(GroomAsset, BindingAsset))
 	{
@@ -605,7 +613,7 @@ void UGroomComponent::SetGroomAsset(UGroomAsset* Asset, UGroomBindingAsset* InBi
 	}
 
 	UpdateHairGroupsDesc();
-	if (!Asset)
+	if (!GroomAsset)
 		return;
 	InitResources();
 	ResetSimulation(this);
