Multiple Group support for hair strands
Polynomial in rbf interp
Curves Params for groom asset
Automatic niagara component creation forgroom asset

diff --git a/Engine/Plugins/Experimental/HairStrands/Shaders/Private/NiagaraAngularSpringMaterial.ush b/Engine/Plugins/Experimental/HairStrands/Shaders/Private/NiagaraAngularSpringMaterial.ush
--- a/Engine/Plugins/Experimental/HairStrands/Shaders/Private/NiagaraAngularSpringMaterial.ush
+++ b/Engine/Plugins/Experimental/HairStrands/Shaders/Private/NiagaraAngularSpringMaterial.ush
@@ -162,28 +162,3 @@ void ProjectAngularSpringMaterial(in bool EnableConstraint, in int StrandsSize,
 	GroupMemoryBarrier();
 	OutNodePosition = SharedNodePosition[GGroupThreadId.x];
 }
-
-/* -----------------------------------------------------------------
- * Compute Rest Direction 
- * -----------------------------------------------------------------
- */
-
-void ComputeRestDirection(in int StrandsSize, in float3 NodePosition, in float4 NodeOrientation, out float3 OutRestDirection)
-{
-	SharedNodePosition[GGroupThreadId.x] = NodePosition;
-	SharedNodeOrientation[GGroupThreadId.x] = NodeOrientation;
-	GroupMemoryBarrier(); 
-
-	ComputeMaterialFrame(StrandsSize);
-
-	OutRestDirection = float3(0,0,0);
-	const int LocalIndex = (GGroupThreadId.x % StrandsSize);
-	if( LocalIndex > 1 )
-	{
-		const float3 EdgeDirection = SharedNodePosition[GGroupThreadId.x] - SharedNodePosition[GGroupThreadId.x-1];
-
-		const float4 ParentFrame = SlerpQuat(SharedNodeOrientation[GGroupThreadId.x-2],SharedNodeOrientation[GGroupThreadId.x-2],0.5);
-		OutRestDirection = UnrotateVectorByQuat(EdgeDirection,ParentFrame);
-	}
-	GroupMemoryBarrier(); 
-}
diff --git a/Engine/Plugins/Experimental/HairStrands/Shaders/Private/NiagaraDataInterfaceHairStrands.ush b/Engine/Plugins/Experimental/HairStrands/Shaders/Private/NiagaraDataInterfaceHairStrands.ush
--- a/Engine/Plugins/Experimental/HairStrands/Shaders/Private/NiagaraDataInterfaceHairStrands.ush
+++ b/Engine/Plugins/Experimental/HairStrands/Shaders/Private/NiagaraDataInterfaceHairStrands.ush
@@ -33,7 +33,7 @@ struct FDIHairStrandsContext
 	int					StrandSize;
 	float3				BoxCenter;
 	float3				BoxExtent;
-	int					HasRootAttached;
+	int					InterpolationMode;
 	float3				RestRootOffset;
 	float3				DeformedRootOffset;
 	float3				RestPositionOffset;
@@ -44,6 +44,7 @@ struct FDIHairStrandsContext
 	int					SampleCount;
 	Buffer<float4>		RestSamplePositionsBuffer;
 	Buffer<float4>		MeshSampleWeightsBuffer;
+	Buffer<float>		ParamsScaleBuffer;
 }; 
 
 #define DIHAIRSTRANDS_DECLARE_CONSTANTS(NAME)\
@@ -64,7 +65,7 @@ int					NumStrands_##NAME;\
 int					StrandSize_##NAME;\
 float3				BoxCenter_##NAME;\
 float3				BoxExtent_##NAME;\
-int					HasRootAttached_##NAME;\
+int					InterpolationMode_##NAME;\
 float3				RestRootOffset_##NAME;\
 float3				DeformedRootOffset_##NAME;\
 float3				RestPositionOffset_##NAME;\
@@ -75,6 +76,7 @@ Buffer<uint>		BoundingBoxBufferB_##NAME;\
 int					SampleCount_##NAME;\
 Buffer<float4>		RestSamplePositionsBuffer_##NAME;\
 Buffer<float4>		MeshSampleWeightsBuffer_##NAME;\
+Buffer<float>		ParamsScaleBuffer_##NAME;\
 
 #define DIHAIRSTRANDS_MAKE_CONTEXT(NAME)\
 FDIHairStrandsContext DIContext;\
@@ -95,7 +97,7 @@ DIContext.NumStrands = NumStrands_##NAME;\
 DIContext.StrandSize = StrandSize_##NAME;\
 DIContext.BoxCenter = BoxCenter_##NAME;\
 DIContext.BoxExtent = BoxExtent_##NAME;\
-DIContext.HasRootAttached =	HasRootAttached_##NAME;\
+DIContext.InterpolationMode =	InterpolationMode_##NAME;\
 DIContext.RestRootOffset =	RestRootOffset_##NAME;\
 DIContext.DeformedRootOffset = DeformedRootOffset_##NAME;\
 DIContext.RestPositionOffset =	RestPositionOffset_##NAME;\
@@ -106,6 +108,7 @@ DIContext.BoundingBoxBufferB = BoundingBoxBufferB_##NAME;\
 DIContext.SampleCount = SampleCount_##NAME;\
 DIContext.RestSamplePositionsBuffer = RestSamplePositionsBuffer_##NAME;\
 DIContext.MeshSampleWeightsBuffer = MeshSampleWeightsBuffer_##NAME;\
+DIContext.ParamsScaleBuffer = ParamsScaleBuffer_##NAME;\
 
 /* -----------------------------------------------------------------
  * Vector utilities
@@ -152,10 +155,34 @@ void DIHairStrands_GetBoxExtent(in FDIHairStrandsContext DIContext, out float3 O
 	const float3 BoxMin = Uint3ToFloat3(uint3(DIContext.BoundingBoxBufferB[0],DIContext.BoundingBoxBufferB[1],DIContext.BoundingBoxBufferB[2]));
 	const float3 BoxMax = Uint3ToFloat3(uint3(DIContext.BoundingBoxBufferB[3],DIContext.BoundingBoxBufferB[4],DIContext.BoundingBoxBufferB[5]));
 
-	//OutBoxExtent =  0.5 * (BoxMax-BoxMin);
 	OutBoxExtent =   (BoxMax-BoxMin);
 }
 
+ #define WRAP_SIZE 32
+
+void DIHairStrands_BuildBoundingBox(in FDIHairStrandsContext DIContext, in float3 NodePosition, out bool BuildStatus, out float3 OutBoundMin, out float3 OutBoundMax)
+{
+	BuildStatus = true;
+
+	const uint3 LocalPosition = Float3ToUint3(NodePosition);
+
+	if(isfinite(NodePosition.x) && isfinite(NodePosition.y) && isfinite(NodePosition.z))
+	{
+		InterlockedMin(DIContext.BoundingBoxBufferA[0],LocalPosition.x);
+		InterlockedMin(DIContext.BoundingBoxBufferA[1],LocalPosition.y);
+		InterlockedMin(DIContext.BoundingBoxBufferA[2],LocalPosition.z);
+
+		InterlockedMax(DIContext.BoundingBoxBufferA[3],LocalPosition.x);
+		InterlockedMax(DIContext.BoundingBoxBufferA[4],LocalPosition.y);
+		InterlockedMax(DIContext.BoundingBoxBufferA[5],LocalPosition.z);
+	}
+
+	DeviceMemoryBarrier();
+
+	OutBoundMin = Uint3ToFloat3(uint3(DIContext.BoundingBoxBufferB[0],DIContext.BoundingBoxBufferB[1],DIContext.BoundingBoxBufferB[2]));
+	OutBoundMax = Uint3ToFloat3(uint3(DIContext.BoundingBoxBufferB[3],DIContext.BoundingBoxBufferB[4],DIContext.BoundingBoxBufferB[5])); 
+}
+
 /* -----------------------------------------------------------------
  * Utilities regarding strands indexing
  * -----------------------------------------------------------------
@@ -167,26 +194,10 @@ int DIHairStrands_GetNumNodes(in FDIHairStrandsContext DIContext)
 	return DIContext.NumStrands * DIContext.StrandSize;
 }
 
-// Given a node index return the strand index and the local index within the strand
-bool DIHairStrands_IsRootNode(in FDIHairStrandsContext DIContext, in int NodeIndex)
-{
-	return (NodeIndex % DIContext.StrandSize) != 0;
-}
-
-// Given a node index return the strand index and the local index within the strand
-void DIHairStrands_ComputeNodeIndices(in FDIHairStrandsContext DIContext, in int NodeIndex, out int StrandIndex, out int LocalIndex)
-{
-	if (NodeIndex < DIHairStrands_GetNumNodes(DIContext))
-	{
-		StrandIndex = (DIContext.StrandSize>0) ? NodeIndex / DIContext.StrandSize : 0;
-		LocalIndex = NodeIndex - StrandIndex * DIContext.StrandSize;
-	}
-	else
-	{
-		StrandIndex = 0;
-		LocalIndex = 0;
-	}
-}
+/* -----------------------------------------------------------------
+ * Strands properties
+ * -----------------------------------------------------------------
+ */
 
 // Given a strand index return the corresponding strand length
 float DIHairStrands_ComputeStrandLength(in FDIHairStrandsContext DIContext, in int StrandIndex)
@@ -210,133 +221,53 @@ float DIHairStrands_ComputeStrandLength(in FDIHairStrandsContext DIContext, in i
 	return StrandLength;
 }
 
-/* -----------------------------------------------------------------
- * Bounding box generation
- * -----------------------------------------------------------------
- */
-
- #define WRAP_SIZE 32
-
-void DIHairStrands_BuildBoundingBox(in FDIHairStrandsContext DIContext, in float3 NodePosition, out bool BuildStatus, out float3 OutBoundMin, out float3 OutBoundMax)
-{
-	BuildStatus = true;
-
-	const uint3 LocalPosition = Float3ToUint3(NodePosition);
-
-	if(isfinite(NodePosition.x) && isfinite(NodePosition.y) && isfinite(NodePosition.z))
-	{
-		InterlockedMin(DIContext.BoundingBoxBufferA[0],LocalPosition.x);
-		InterlockedMin(DIContext.BoundingBoxBufferA[1],LocalPosition.y);
-		InterlockedMin(DIContext.BoundingBoxBufferA[2],LocalPosition.z);
-
-		InterlockedMax(DIContext.BoundingBoxBufferA[3],LocalPosition.x);
-		InterlockedMax(DIContext.BoundingBoxBufferA[4],LocalPosition.y);
-		InterlockedMax(DIContext.BoundingBoxBufferA[5],LocalPosition.z);
-	}
-
-	DeviceMemoryBarrier();
-
-	OutBoundMin = Uint3ToFloat3(uint3(DIContext.BoundingBoxBufferB[0],DIContext.BoundingBoxBufferB[1],DIContext.BoundingBoxBufferB[2]));
-	OutBoundMax = Uint3ToFloat3(uint3(DIContext.BoundingBoxBufferB[3],DIContext.BoundingBoxBufferB[4],DIContext.BoundingBoxBufferB[5])); 
-}
-
-// Compute the node mass
-void DIHairStrands_EvalSkinnedPosition(in FDIHairStrandsContext DIContext, in float3 RestPosition, out float3 OutSkinnedPosition)
-{
-	const float3 RestSkinnedPosition = RestPosition;
-
-	OutSkinnedPosition = RestPosition;
-	for(int i = 0; i < DIContext.SampleCount; ++i)
-	{
-		const float3 PositionDelta = RestSkinnedPosition - DIContext.RestSamplePositionsBuffer[i];
-		const float FunctionValue = sqrt( dot(PositionDelta,PositionDelta)+1);
-		OutSkinnedPosition += FunctionValue * DIContext.MeshSampleWeightsBuffer[i].xyz;
-	}
-}
-
 /* -----------------------------------------------------------------
  * Node Mass/Inertia/Position/Orientation computation
  * -----------------------------------------------------------------
  */
 
-// Compute the node mass
-void DIHairStrands_ComputeNodeMass(in FDIHairStrandsContext DIContext, in int NodeIndex, in float StrandsDensity, in float RootThickness, in float TipThickness, out float OutNodeMass)
+ // Compute the node mass
+void DIHairStrands_ComputeNodeMass(in FDIHairStrandsContext DIContext, in float StrandsDensity, in float NodeThickness, out float OutNodeMass)
 {
 	const int ReducedSize =  DIContext.StrandSize-1;
 	const float CoordScale = 1.0 / (ReducedSize-1.0);
 
-	int StrandIndex = 0, LocalIndex = 0;
-	DIHairStrands_ComputeNodeIndices(DIContext,NodeIndex,StrandIndex,LocalIndex);
-	const float EdgeLength = DIHairStrands_ComputeStrandLength(DIContext,StrandIndex) * CoordScale;
+	const int StrandIndex = ExecIndex() / DIContext.StrandSize;
 
-	LocalIndex = max(0,LocalIndex - 1);
-
-	const bool IsFirst = (LocalIndex == 0);
-	const bool IsLast  = (LocalIndex == (ReducedSize-1));
+	const float StrandLength = DIHairStrands_ComputeStrandLength(DIContext,StrandIndex);
 
-	const float EdgeScale = (IsFirst || IsLast) ? 0.5 : 1.0;
-	
-	const float PrevCoord = IsFirst ? 0.0 : (LocalIndex-0.5) * CoordScale;
-	const float NextCoord = IsLast  ? 1.0 : (LocalIndex+0.5) * CoordScale;
+	int  LocalIndex = GGroupThreadId.x % DIContext.StrandSize;
+	const float EdgeLength = StrandLength * CoordScale;
+	const float MeanRadius = 0.5 * NodeThickness;
 
-	const float PrevThick = RootThickness * ( 1.0 - PrevCoord ) + TipThickness * PrevCoord;
-	const float NextThick = RootThickness * ( 1.0 - NextCoord ) + TipThickness * NextCoord;
-	const float MeanRadius = 0.25 * ( PrevThick + NextThick );
+	LocalIndex = max(0,LocalIndex - 1);
+	const float EdgeScale = (LocalIndex == 0 || LocalIndex == (ReducedSize-1)) ? 0.5 : 1.0;
 
 	OutNodeMass = StrandsDensity * EdgeLength * EdgeScale * M_PI * MeanRadius * MeanRadius;
-	//OutNodeMass = StrandsDensity * EdgeLength * EdgeScale * M_PI * (PrevThick * PrevThick + NextThick * PrevThick + NextThick * NextThick) / 12.0;
 }
 
 // Compute the node inertia
-void DIHairStrands_ComputeNodeInertia(in FDIHairStrandsContext DIContext, in int NodeIndex, in float StrandsDensity, in float RootThickness, in float TipThickness,  out float3 OutNodeInertia)
+void DIHairStrands_ComputeNodeInertia(in FDIHairStrandsContext DIContext, in float StrandsDensity, in float NodeThickness, out float3 OutNodeInertia)
 {
 	const int ReducedSize =  DIContext.StrandSize-1;
 	const float CoordScale = 1.0 / (ReducedSize-1.0);
 
-	int StrandIndex = 0, LocalIndex = 0;
-	DIHairStrands_ComputeNodeIndices(DIContext,NodeIndex,StrandIndex,LocalIndex);
-	const float EdgeLength = DIHairStrands_ComputeStrandLength(DIContext,StrandIndex) * CoordScale;
+	const int StrandIndex = ExecIndex() / DIContext.StrandSize;
 
-	LocalIndex = min(max(0,LocalIndex - 1),ReducedSize-2);
+	const float StrandLength = DIHairStrands_ComputeStrandLength(DIContext,StrandIndex);
 
-	const float PrevCoord = LocalIndex * CoordScale;
-	const float NextCoord = (LocalIndex+1) * CoordScale;
+	const float EdgeLength = StrandLength * CoordScale;
+	const float MeanRadius = 0.5 * NodeThickness;
 
-	const float PrevThick = RootThickness * ( 1.0 - PrevCoord ) + TipThickness * PrevCoord;
-	const float NextThick = RootThickness * ( 1.0 - NextCoord ) + TipThickness * NextCoord;
-	const float MeanRadius = 0.25 * ( PrevThick + NextThick );
-
-	//const float EdgeMass = StrandsDensity * EdgeLength * M_PI * (PrevThick * PrevThick + NextThick * PrevThick + NextThick * NextThick) / 12.0;
 	const float EdgeMass = StrandsDensity * EdgeLength * M_PI * MeanRadius * MeanRadius;
 	const float RadialInertia = EdgeMass * ( 3.0 * MeanRadius * MeanRadius + EdgeLength * EdgeLength) / 12.0;
 	OutNodeInertia = float3( RadialInertia, RadialInertia, EdgeMass * MeanRadius * MeanRadius * 0.5 );
-
-	//const float EdgeMass = StrandsDensity * EdgeLength * M_PI * MeanRadius * MeanRadius;
-	//const float RadialInertia = EdgeMass * ( MeanRadius * MeanRadius) / 4.0;
-	//OutNodeInertia = float3( RadialInertia, RadialInertia, 2.0 * RadialInertia );
 }
 
-// Compute the node position by linear interpolation over the points
-void DIHairStrands_ComputeNodePosition(in FDIHairStrandsContext DIContext, in int NodeIndex, out float3 OutNodePosition)
+// Set the node position by linear interpolation over the points
+void DIHairStrands_SetNodePosition(in FDIHairStrandsContext DIContext, in bool RootNode, const float NodePoint, const float PointCount, 
+		const float CoordScale, const int PointPrev, const int PointNext, const int PointOffset, in int StrandIndex, out float3 OutNodePosition)
 {
-	const int ReducedSize =  DIContext.StrandSize-1;
-	const float CoordScale = 1.0 / (ReducedSize-1.0);
-
-	int StrandIndex = 0, LocalIndex = 0;
-	DIHairStrands_ComputeNodeIndices(DIContext,NodeIndex,StrandIndex,LocalIndex);
-
-	const bool RootNode = (LocalIndex == 0);
-	LocalIndex = max(0,LocalIndex - 1);
-
-	const int PointOffset = DIContext.CurvesOffsetsBuffer[StrandIndex];
-
-	const float NodeCoord = (float)(LocalIndex) * CoordScale;
-	const float PointCount = DIContext.CurvesOffsetsBuffer[StrandIndex+1]-PointOffset-1;
-
-	const float NodePoint = NodeCoord * PointCount;
-	const int PointPrev = (LocalIndex==0) ? 0 : (LocalIndex==(ReducedSize-1)) ? PointCount-1 : floor(NodePoint);
-	const int PointNext = PointPrev+1;
-
 	if(RootNode)
 	{
 		const float FirstPoint = CoordScale * PointCount;
@@ -350,8 +281,6 @@ void DIHairStrands_ComputeNodePosition(in FDIHairStrandsContext DIContext, in in
 		const float3 EdgeDirection =  normalize(DiffPosition);
 
 		const float EdgeLength = DIHairStrands_ComputeStrandLength(DIContext,StrandIndex) * CoordScale;
-		//float3 EdgeDirection = f16tof32(DIContext.RestPositionBuffer[PointOffset+PointNext].xyz) - 
-		//							 f16tof32(DIContext.RestPositionBuffer[PointOffset+PointPrev].xyz);
 		OutNodePosition = 	f16tof32(DIContext.RestPositionBuffer[PointOffset+PointPrev].xyz) - EdgeLength * normalize(EdgeDirection);					 
 	}
 	else
@@ -364,7 +293,31 @@ void DIHairStrands_ComputeNodePosition(in FDIHairStrandsContext DIContext, in in
 }
 
 // Compute the node position by linear interpolation over the points
-void DIHairStrands_SmoothNodePosition(in FDIHairStrandsContext DIContext, in int NodeIndex, in float SmoothingFilter, inout float3 OutNodePosition)
+void DIHairStrands_ComputeNodePosition(in FDIHairStrandsContext DIContext, out float3 OutNodePosition)
+{
+	const int ReducedSize =  DIContext.StrandSize-1;
+	const float CoordScale = 1.0 / (ReducedSize-1.0);
+
+	int LocalIndex = GGroupThreadId.x % DIContext.StrandSize;
+	int StrandIndex = ExecIndex() / DIContext.StrandSize;
+
+	const bool RootNode = (LocalIndex == 0);
+	LocalIndex = max(0,LocalIndex - 1);
+
+	const int PointOffset = DIContext.CurvesOffsetsBuffer[StrandIndex];
+
+	const float NodeCoord = (float)(LocalIndex) * CoordScale;
+	const float PointCount = DIContext.CurvesOffsetsBuffer[StrandIndex+1]-PointOffset-1;
+
+	const float NodePoint = NodeCoord * PointCount;
+	const int PointPrev = (LocalIndex==0) ? 0 : (LocalIndex==(ReducedSize-1)) ? PointCount-1 : floor(NodePoint);
+	const int PointNext = PointPrev+1;
+
+	DIHairStrands_SetNodePosition( DIContext, RootNode, NodePoint, PointCount, CoordScale, PointPrev, PointNext, PointOffset, StrandIndex, OutNodePosition);
+}
+
+// Compute the node position by linear interpolation over the points
+void DIHairStrands_SmoothNodePosition(in FDIHairStrandsContext DIContext, in float SmoothingFilter, inout float3 OutNodePosition)
 {
 	SharedNodePosition[GGroupThreadId.x] = OutNodePosition;
 	GroupMemoryBarrier();
@@ -399,54 +352,66 @@ void DIHairStrands_SmoothNodePosition(in FDIHairStrandsContext DIContext, in int
 	OutNodePosition = SharedPreviousPosition[GGroupThreadId.x];
 }
 
-// Compute the node orientation
-void DIHairStrands_ComputeNodeOrientation(in FDIHairStrandsContext DIContext, in int NodeIndex, in float3 NodePosition, out float4 OutNodeOrientation)
+// Compute the root orientation
+void DIHairStrands_ComputeRootOrientation(in FDIHairStrandsContext DIContext)
 {
-	SharedNodePosition[GGroupThreadId.x] = NodePosition;
-	SharedNodeOrientation[GGroupThreadId.x] = QUATERNION_IDENTITY;
-	GroupMemoryBarrier();
-
 	const int LocalIndex = (GGroupThreadId.x % DIContext.StrandSize);
 	if(LocalIndex == 0)
 	{
-		float4 NodeQuaternion = QUATERNION_IDENTITY;
-		float3 TangentNext = float3(0,0,1);
-		float3 TangentPrev = TangentNext;
+		const float3 EdgeDirection = normalize(SharedNodePosition[GGroupThreadId.x+1] - SharedNodePosition[GGroupThreadId.x]);
+		const float4 RootQuaternion = FindQuatBetweenNormals(float3(0,0,1),EdgeDirection);
 
-		float3 PositionNext = SharedNodePosition[GGroupThreadId.x];
-		float3 PositionPrev = PositionNext;
+		//const float3 TangentPrev = normalize(RotateVectorByQuat( float3(1,0,0), RootQuaternion));
+		//const float3 TangentNext = normalize(cross( normalize(cross(EdgeDirection,float3(0,0,1))), EdgeDirection));
 
-		for (int EdgeIndex = 0, EdgeEnd = DIContext.StrandSize-1; EdgeIndex < EdgeEnd; ++EdgeIndex)
-		{
-			PositionPrev = PositionNext;
-			PositionNext = SharedNodePosition[GGroupThreadId.x+EdgeIndex+1];
+		SharedNodeOrientation[GGroupThreadId.x] = RootQuaternion;
 
-			TangentPrev = TangentNext;
-			TangentNext = normalize(PositionNext - PositionPrev);
-			NodeQuaternion = NormalizeQuat( MultiplyQuat( FindQuatBetweenNormals(TangentPrev,TangentNext), NodeQuaternion ) );
+		//SharedNodeOrientation[GGroupThreadId.x] = NormalizeQuat( MultiplyQuat( RootQuaternion, FindQuatBetweenNormals(TangentPrev,TangentNext) ) );
+	}
+	GroupMemoryBarrier();
+}
 
-			SharedNodeOrientation[GGroupThreadId.x+EdgeIndex] = NodeQuaternion;
-			//GroupMemoryBarrier();
-		}
-		//GroupMemoryBarrier();
+// Update the root orientation
+void DIHairStrands_UpdateRootOrientation(in FDIHairStrandsContext DIContext)
+{
+	const int LocalIndex = (GGroupThreadId.x % DIContext.StrandSize);
+	if(LocalIndex == 0)
+	{
+		float4 NodeQuaternion = SharedNodeOrientation[GGroupThreadId.x];
+		float3 TangentPrev = RotateVectorByQuat( float3(0,0,1), NodeQuaternion);
+		float3 TangentNext = normalize(SharedNodePosition[GGroupThreadId.x+1] - SharedNodePosition[GGroupThreadId.x]);
+
+		SharedNodeOrientation[GGroupThreadId.x] = NormalizeQuat( MultiplyQuat( FindQuatBetweenNormals(TangentPrev,TangentNext), NodeQuaternion) );
 	}
 	GroupMemoryBarrier();
+}
+
+// Compute the node orientation
+void DIHairStrands_ComputeNodeOrientation(in FDIHairStrandsContext DIContext, in float3 NodePosition, out float4 OutNodeOrientation)
+{
+	SharedNodePosition[GGroupThreadId.x] = NodePosition;
+	SharedNodeOrientation[GGroupThreadId.x] = QUATERNION_IDENTITY;
+	GroupMemoryBarrier();
+
+	DIHairStrands_ComputeRootOrientation(DIContext);
+
+	ComputeMaterialFrame(DIContext.StrandSize);
 	OutNodeOrientation = 	SharedNodeOrientation[GGroupThreadId.x];
 }
 
+
 /* -----------------------------------------------------------------
- * Edge delta position and delta orientation
+ * Edge volume, length, rotation, direction
  * -----------------------------------------------------------------
  */
 
 // Compute the edge volume value
-void DIHairStrands_ComputeEdgeVolume(in FDIHairStrandsContext DIContext, in int NodeIndex, in float3 NodePosition, out float OutEdgeVolume)
+void DIHairStrands_ComputeEdgeVolume(in FDIHairStrandsContext DIContext, in float3 NodePosition, out float OutEdgeVolume)
 {
 	SharedNodePosition[GGroupThreadId.x] = NodePosition;
 	GroupMemoryBarrier(); 
 
-	int StrandIndex = 0, LocalIndex = 0;
-	DIHairStrands_ComputeNodeIndices(DIContext,NodeIndex,StrandIndex,LocalIndex);
+	const int LocalIndex = (GGroupThreadId.x % DIContext.StrandSize);
 
 	// L = ||P1-P0||
 	OutEdgeVolume = 0.0;
@@ -466,26 +431,24 @@ void DIHairStrands_ComputeEdgeVolume(in FDIHairStrandsContext DIContext, in int
 }
 
 // Compute the edge length value
-void DIHairStrands_ComputeEdgeLength(in FDIHairStrandsContext DIContext, in int NodeIndex, in float3 NodePosition, in int NodeOffset, out float OutEdgeLength)
+void DIHairStrands_ComputeEdgeLength(in FDIHairStrandsContext DIContext, in float3 NodePosition, in int NodeOffset, out float OutEdgeLength)
 {
 	SharedNodePosition[GGroupThreadId.x] = NodePosition;
 	GroupMemoryBarrier(); 
 
-	int StrandIndex = 0, LocalIndex = 0;
-	DIHairStrands_ComputeNodeIndices(DIContext,NodeIndex,StrandIndex,LocalIndex);
+	const int LocalIndex = (GGroupThreadId.x % DIContext.StrandSize);
 
 	// L = ||P1-P0||
 	OutEdgeLength = (LocalIndex>NodeOffset) ? length(SharedNodePosition[GGroupThreadId.x]- SharedNodePosition[GGroupThreadId.x-1-NodeOffset]) : 0.0;
 }
 
 // Compute the edge darboux vector (diff between consecutive edge orientations)
-void DIHairStrands_ComputeEdgeRotation(in FDIHairStrandsContext DIContext, in int NodeIndex, in float4 NodeOrientation, out float4 OutEdgeRotation)
+void DIHairStrands_ComputeEdgeRotation(in FDIHairStrandsContext DIContext, in float4 NodeOrientation, out float4 OutEdgeRotation)
 {
 	SharedNodeOrientation[GGroupThreadId.x] = NodeOrientation;
 	GroupMemoryBarrier(); 
 
-	int StrandIndex = 0, LocalIndex = 0;
-	DIHairStrands_ComputeNodeIndices(DIContext,NodeIndex,StrandIndex,LocalIndex);
+	const int LocalIndex = (GGroupThreadId.x % DIContext.StrandSize);
 
 	// D = Q0^-1 * Q1
 	if(LocalIndex>0)
@@ -506,17 +469,38 @@ void DIHairStrands_ComputeEdgeRotation(in FDIHairStrandsContext DIContext, in in
 	}
 }
 
+void DIHairStrands_ComputeEdgeDirection(in FDIHairStrandsContext DIContext, in float3 NodePosition, in float4 NodeOrientation, out float3 OutRestDirection)
+{
+	SharedNodePosition[GGroupThreadId.x] = NodePosition;
+	SharedNodeOrientation[GGroupThreadId.x] = NodeOrientation;
+	GroupMemoryBarrier(); 
+
+	//ComputeMaterialFrame(StrandsSize);
+
+	OutRestDirection = float3(0,0,0);
+	const int LocalIndex = (GGroupThreadId.x % DIContext.StrandSize);
+	if( LocalIndex > 1 )
+	{
+		const float3 EdgeDirection = SharedNodePosition[GGroupThreadId.x] - SharedNodePosition[GGroupThreadId.x-1];
+
+		const float4 ParentFrame = SlerpQuat(SharedNodeOrientation[GGroupThreadId.x-2],SharedNodeOrientation[GGroupThreadId.x-2],0.5);
+		OutRestDirection = UnrotateVectorByQuat(EdgeDirection,ParentFrame);
+	}
+	GroupMemoryBarrier(); 
+}
+
+
 /* -----------------------------------------------------------------
  * Points position update
  * -----------------------------------------------------------------
  */
 
 // Reset the deformed points position to the rest ones
-void DIHairStrands_ResetPointPosition(in FDIHairStrandsContext DIContext, in int NodeIndex, out bool ReportStatus) 
+void DIHairStrands_ResetPointPosition(in FDIHairStrandsContext DIContext, out bool ReportStatus) 
 {
-	int StrandIndex = 0;
-	int LocalIndex = 0;
-	DIHairStrands_ComputeNodeIndices(DIContext,NodeIndex,StrandIndex,LocalIndex);
+	int LocalIndex = GGroupThreadId.x % DIContext.StrandSize;
+	int StrandIndex = ExecIndex() / DIContext.StrandSize;
+
 	ReportStatus = false;
 
 	if(LocalIndex == 0)
@@ -536,16 +520,16 @@ void DIHairStrands_ResetPointPosition(in FDIHairStrandsContext DIContext, in int
 }
 
 // Report interpolated nodes displacements onto the points positions 
-void DIHairStrands_UpdatePointPosition(in FDIHairStrandsContext DIContext, in int NodeIndex, in float3 NodeDisplace, out bool ReportStatus) 
+void DIHairStrands_UpdatePointPosition(in FDIHairStrandsContext DIContext, in float3 NodeDisplace, out bool ReportStatus) 
 {
 	const int ReducedSize =  DIContext.StrandSize-1;
 	SharedNodePosition[GGroupThreadId.x] = NodeDisplace;
 	GroupMemoryBarrier(); 
 
-	int StrandIndex = 0, LocalIndex = 0;
-	DIHairStrands_ComputeNodeIndices(DIContext,NodeIndex,StrandIndex,LocalIndex);
-	LocalIndex = LocalIndex - 1;
+	int LocalIndex = GGroupThreadId.x % DIContext.StrandSize;
+	int StrandIndex = ExecIndex() / DIContext.StrandSize;
 
+	LocalIndex = LocalIndex - 1;
 	ReportStatus = false;
 
 	const int PointBegin = DIContext.CurvesOffsetsBuffer[StrandIndex];
@@ -576,7 +560,7 @@ void DIHairStrands_UpdatePointPosition(in FDIHairStrandsContext DIContext, in in
 // Get the point position
 void DIHairStrands_GetPointPosition(in FDIHairStrandsContext DIContext, in int PointIndex, out float3 OutPointPosition)
 {
-	OutPointPosition = mul(float4( f16tof32(DIContext.DeformedPositionBuffer[PointIndex].xyz+DIContext.DeformedPositionOffset), 1.0), DIContext.WorldTransform).xyz;
+	OutPointPosition = mul(float4( f16tof32(DIContext.DeformedPositionBuffer.Load(PointIndex).xyz)+DIContext.DeformedPositionOffset, 1.0), DIContext.WorldTransform).xyz;
 }
 
 /* -----------------------------------------------------------------
@@ -673,30 +657,61 @@ void DIHairStrands_AttachNodeOrientation(in FDIHairStrandsContext DIContext, in
  * -----------------------------------------------------------------
  */
 
+ // Compute the global position
+void DIHairStrands_EvalGlobalPosition(in FDIHairStrandsContext DIContext, in float3 RestPosition, out float3 OutGlobalPosition)
+{
+	const float3 RestSkinnedPosition = RestPosition;
+
+	OutGlobalPosition = RestPosition;
+	for(int i = 0; i < DIContext.SampleCount; ++i)
+	{
+		const float3 PositionDelta = RestSkinnedPosition - DIContext.RestSamplePositionsBuffer[i].xyz;
+		const float FunctionValue = sqrt( dot(PositionDelta,PositionDelta)+1);
+		OutGlobalPosition += FunctionValue * DIContext.MeshSampleWeightsBuffer[i].xyz;
+	}
+	OutGlobalPosition += DIContext.MeshSampleWeightsBuffer[DIContext.SampleCount].xyz;
+	OutGlobalPosition += DIContext.MeshSampleWeightsBuffer[DIContext.SampleCount+1].xyz * RestSkinnedPosition.x;
+	OutGlobalPosition += DIContext.MeshSampleWeightsBuffer[DIContext.SampleCount+2].xyz * RestSkinnedPosition.y;
+	OutGlobalPosition += DIContext.MeshSampleWeightsBuffer[DIContext.SampleCount+3].xyz * RestSkinnedPosition.z;
+}
+
+// Compute the global position
+void DIHairStrands_EvalGlobalOrientation(in FDIHairStrandsContext DIContext, in float3 GlobalPosition, inout float4 OutGlobalOrientation)
+{
+	SharedNodePosition[GGroupThreadId.x] = GlobalPosition;
+	SharedNodeOrientation[GGroupThreadId.x] = OutGlobalOrientation;
+	GroupMemoryBarrier(); 
+
+	DIHairStrands_UpdateRootOrientation(DIContext);
+
+	ComputeMaterialFrame(DIContext.StrandSize);
+	OutGlobalOrientation = 	SharedNodeOrientation[GGroupThreadId.x];
+}
+
 float2 ExtractBarycentrics(in uint PackedCoordinates )
 {
 	return float2(f16tof32(PackedCoordinates & 0xFFFF), f16tof32((PackedCoordinates>>16) & 0xFFFF));
 }
 
 // Compuite rest position and orientation
-void DIHairStrands_ComputeLocalState(in FDIHairStrandsContext DIContext, in int NodeIndex, in float3 RestPosition, in float4 RestOrientation, 
+void DIHairStrands_ComputeLocalState(in FDIHairStrandsContext DIContext, in float3 RestPosition, in float4 RestOrientation, 
 												out float3 OutLocalPosition, out float4 OutLocalOrientation)
 {
 	OutLocalPosition = RestPosition;
 	OutLocalOrientation = RestOrientation;
 
-	if( DIContext.HasRootAttached == 1)
+	if( DIContext.InterpolationMode >= 1)
 	{
-		const int CurveIndex = NodeIndex / DIContext.StrandSize;
-		const float2 RootUV = ExtractBarycentrics(DIContext.RootBarycentricCoordinatesBuffer[CurveIndex]);
+		const int StrandIndex = ExecIndex() / DIContext.StrandSize;
+		const float2 RootUV = ExtractBarycentrics(DIContext.RootBarycentricCoordinatesBuffer[StrandIndex]);
 
-		const float3 RootProjection = DIContext.RestTrianglePositionABuffer[CurveIndex].xyz * RootUV.x + 
-									  DIContext.RestTrianglePositionBBuffer[CurveIndex].xyz * RootUV.y +
-									  DIContext.RestTrianglePositionCBuffer[CurveIndex].xyz * (1.0 - RootUV.x - RootUV.y) + DIContext.RestRootOffset;
+		const float3 RootProjection = DIContext.RestTrianglePositionABuffer[StrandIndex].xyz * RootUV.x + 
+									  DIContext.RestTrianglePositionBBuffer[StrandIndex].xyz * RootUV.y +
+									  DIContext.RestTrianglePositionCBuffer[StrandIndex].xyz * (1.0 - RootUV.x - RootUV.y) + DIContext.RestRootOffset;
 		const float3 DiffPosition = RestPosition - RootProjection;
 
-		float3 TangentU = normalize(DIContext.RestTrianglePositionBBuffer[CurveIndex].xyz - DIContext.RestTrianglePositionABuffer[CurveIndex].xyz);
-		float3 TangentV = normalize(DIContext.RestTrianglePositionCBuffer[CurveIndex].xyz - DIContext.RestTrianglePositionABuffer[CurveIndex].xyz);
+		float3 TangentU = normalize(DIContext.RestTrianglePositionBBuffer[StrandIndex].xyz - DIContext.RestTrianglePositionABuffer[StrandIndex].xyz);
+		float3 TangentV = normalize(DIContext.RestTrianglePositionCBuffer[StrandIndex].xyz - DIContext.RestTrianglePositionABuffer[StrandIndex].xyz);
 
 		const float3 Normal = normalize(cross(TangentU,TangentV));
 		TangentU = normalize(cross(TangentV,Normal));
@@ -713,21 +728,22 @@ void DIHairStrands_ComputeLocalState(in FDIHairStrandsContext DIContext, in int
 }
 
 // Update Node Position and orientation
-void DIHairStrands_AttachNodeState(in FDIHairStrandsContext DIContext, in int NodeIndex, in float3 LocalPosition, in float4 LocalOrientation, out float3 OutNodePosition, out float4 OutNodeOrientation)
+void DIHairStrands_AttachNodeState(in FDIHairStrandsContext DIContext, in float3 LocalPosition, in float4 LocalOrientation, out float3 OutNodePosition, out float4 OutNodeOrientation)
 {	
-	float3 RestPosition =  LocalPosition;
-	float4 RestOrientation = LocalOrientation;
-	if( DIContext.HasRootAttached == 1)
+	float3 NodePosition =  LocalPosition;
+	float4 NodeOrientation = LocalOrientation;
+
+	if( DIContext.InterpolationMode >= 1)
 	{
-		const int CurveIndex = NodeIndex / DIContext.StrandSize;
-		const float2 RootUV = ExtractBarycentrics(DIContext.RootBarycentricCoordinatesBuffer[CurveIndex]);
+		const int StrandIndex = ExecIndex() / DIContext.StrandSize;
+		const float2 RootUV = ExtractBarycentrics(DIContext.RootBarycentricCoordinatesBuffer[StrandIndex]);
 
-		const float3 RootProjection = DIContext.DeformedTrianglePositionABuffer[CurveIndex].xyz * RootUV.x + 
-									  DIContext.DeformedTrianglePositionBBuffer[CurveIndex].xyz * RootUV.y +
-									  DIContext.DeformedTrianglePositionCBuffer[CurveIndex].xyz * (1.0 - RootUV.x - RootUV.y) + DIContext.DeformedRootOffset;
+		const float3 RootProjection = DIContext.DeformedTrianglePositionABuffer[StrandIndex].xyz * RootUV.x + 
+									  DIContext.DeformedTrianglePositionBBuffer[StrandIndex].xyz * RootUV.y +
+									  DIContext.DeformedTrianglePositionCBuffer[StrandIndex].xyz * (1.0 - RootUV.x - RootUV.y) + DIContext.DeformedRootOffset;
 
-		float3 TangentU = normalize(DIContext.DeformedTrianglePositionBBuffer[CurveIndex].xyz - DIContext.DeformedTrianglePositionABuffer[CurveIndex].xyz);
-		float3 TangentV = normalize(DIContext.DeformedTrianglePositionCBuffer[CurveIndex].xyz - DIContext.DeformedTrianglePositionABuffer[CurveIndex].xyz);
+		float3 TangentU = normalize(DIContext.DeformedTrianglePositionBBuffer[StrandIndex].xyz - DIContext.DeformedTrianglePositionABuffer[StrandIndex].xyz);
+		float3 TangentV = normalize(DIContext.DeformedTrianglePositionCBuffer[StrandIndex].xyz - DIContext.DeformedTrianglePositionABuffer[StrandIndex].xyz);
 
 		const float3 Normal = normalize(cross(TangentU,TangentV));
 		TangentU = normalize(cross(TangentV,Normal));
@@ -738,13 +754,31 @@ void DIHairStrands_AttachNodeState(in FDIHairStrandsContext DIContext, in int No
 		const float3 DiffPosition = RotateVectorByQuat(LocalPosition,TriangleOrientation);  
 		//const float3 DiffPosition = LocalPosition.x * TangentU + LocalPosition.y * TangentV + LocalPosition.z * Normal;
 
-		RestPosition = RootProjection + DiffPosition;
-		RestOrientation = NormalizeQuat( MultiplyQuat(TriangleOrientation,LocalOrientation) );
+		NodePosition = RootProjection + DiffPosition;
+		NodeOrientation = NormalizeQuat( MultiplyQuat(TriangleOrientation,LocalOrientation) );
 	}
-	DIHairStrands_AttachNodePosition(DIContext,RestPosition,OutNodePosition);
-	DIHairStrands_AttachNodeOrientation(DIContext,RestOrientation,OutNodeOrientation);
+	OutNodePosition = NodePosition;
+	OutNodeOrientation = NodeOrientation;
+	//DIHairStrands_AttachNodePosition(DIContext,NodePosition,OutNodePosition);
+	//DIHairStrands_AttachNodeOrientation(DIContext,NodeOrientation,OutNodeOrientation);
 }
 
+void DIHairStrands_UpdateNodeState(in FDIHairStrandsContext DIContext, in float3 RestPosition, in float3 NodePosition, in float4 NodeOrientation, out float3 OutNodePosition, out float4 OutNodeOrientation)
+{
+	float3 GlobalPosition =  NodePosition;
+	float4 GlobalOrientation = NodeOrientation;
+
+	if(DIContext.InterpolationMode == 2)
+	{
+		DIHairStrands_EvalGlobalPosition(DIContext, RestPosition, GlobalPosition);
+		DIHairStrands_EvalGlobalOrientation(DIContext,GlobalPosition,GlobalOrientation);
+	}
+	DIHairStrands_AttachNodePosition(DIContext,GlobalPosition,OutNodePosition);
+	DIHairStrands_AttachNodeOrientation(DIContext,GlobalOrientation,OutNodeOrientation);
+}
+
+
+
 
 
 
diff --git a/Engine/Plugins/Experimental/HairStrands/Shaders/Private/NiagaraDataInterfacePhysicsAsset.ush b/Engine/Plugins/Experimental/HairStrands/Shaders/Private/NiagaraDataInterfacePhysicsAsset.ush
--- a/Engine/Plugins/Experimental/HairStrands/Shaders/Private/NiagaraDataInterfacePhysicsAsset.ush
+++ b/Engine/Plugins/Experimental/HairStrands/Shaders/Private/NiagaraDataInterfacePhysicsAsset.ush
@@ -311,9 +311,99 @@ int DIPhysicsAsset_GetClosestElement(in FDIPhysicsAssetContext DIContext, in flo
  * -----------------------------------------------------------------
  */
 
-// Given a world space position (NodePosition) compute the sphere closest point (position,normal,velocity)
+// Given a world space position (NodePosition) and an element index compute the physics asset closest point (position,normal,velocity)
+void DIPhysicsAsset_GetElementPoint(in FDIPhysicsAssetContext DIContext, in float3 NodePosition, in float DeltaTime, in float TimeFraction, in int ElementIndex, out float3 OutClosestPosition, 
+										  out float3 OutClosestNormal, out float3 OutClosestVelocity)
+{	
+	const float3 LocalPosition = DIPhysicsAsset_GetLocalPosition(DIContext,NodePosition,ElementIndex,TimeFraction);
+
+	float MinDistance = MAX_DISTANCE;
+	int OutputIndex = -1;
+	float3 CollisionPosition = float3(0,0,0);
+	float3 CollisionNormal = float3(0,0,0);
+
+	if( ElementIndex >= DIContext.ElementOffsets[SPHERE_INDEX] && ElementIndex < DIContext.ElementOffsets[SPHERE_INDEX+1] )
+	{
+		DIPhysicsAsset_GetSphereProjection(LocalPosition, float3(0,0,0), DIContext.ElementExtentBuffer[ElementIndex].x, ElementIndex,
+				CollisionPosition, CollisionNormal, OutputIndex, MinDistance);
+	}
+	else if ( ElementIndex >= DIContext.ElementOffsets[CAPSULE_INDEX] && ElementIndex < DIContext.ElementOffsets[CAPSULE_INDEX+1] )
+	{
+		DIPhysicsAsset_GetCapsuleProjection(LocalPosition, DIContext.ElementExtentBuffer[ElementIndex].xy, ElementIndex,
+				CollisionPosition, CollisionNormal, OutputIndex, MinDistance);
+	}
+	else if ( ElementIndex >= DIContext.ElementOffsets[BOX_INDEX] && ElementIndex < DIContext.ElementOffsets[BOX_INDEX+1] )
+	{
+		DIPhysicsAsset_GetBoxProjection(LocalPosition, DIContext.ElementExtentBuffer[ElementIndex].xyz, ElementIndex,
+				CollisionPosition, CollisionNormal, OutputIndex, MinDistance);
+	}
+
+	OutClosestNormal = float3(0,0,0);
+	OutClosestPosition = float3(0,0,0);
+	OutClosestVelocity = float3(0,0,0);
+
+	if (ElementIndex != -1)
+	{
+		const float3 PreviousPosition = mul(DIPhysicsAsset_GetPreviousTransform(DIContext,ElementIndex), float4(CollisionPosition,1.0)).xyz;
+		const float3 CurrentPosition = mul(DIPhysicsAsset_GetCurrentTransform(DIContext,ElementIndex), float4(CollisionPosition,1.0)).xyz;
+
+		const float3 PreviousNormal = mul(DIPhysicsAsset_GetPreviousTransform(DIContext,ElementIndex), float4(CollisionNormal,0.0)).xyz;
+		const float3 CurrentNormal = mul(DIPhysicsAsset_GetCurrentTransform(DIContext,ElementIndex), float4(CollisionNormal,0.0)).xyz;
+
+		OutClosestPosition = PreviousPosition + TimeFraction * (CurrentPosition-PreviousPosition);
+		OutClosestNormal = normalize(PreviousNormal + TimeFraction * (CurrentNormal-PreviousNormal));
+
+		OutClosestVelocity = ( CurrentPosition - PreviousPosition ) / DeltaTime;
+	}
+}
+
+// Given a world space position (NodePosition) and an element index compute the physics asset closest distance
+void DIPhysicsAsset_GetElementDistance(in FDIPhysicsAssetContext DIContext, in float3 NodePosition, in float TimeFraction, in int ElementIndex, out float OutClosestDistance)
+{	
+	const float3 LocalPosition = DIPhysicsAsset_GetLocalPosition(DIContext,NodePosition,ElementIndex,TimeFraction);
+
+	float MinDistance = MAX_DISTANCE;
+	int OutputIndex = -1;
+	float3 CollisionPosition = float3(0,0,0);
+	float3 CollisionNormal = float3(0,0,0);
+
+	if( ElementIndex >= DIContext.ElementOffsets[SPHERE_INDEX] && ElementIndex < DIContext.ElementOffsets[SPHERE_INDEX+1] )
+	{
+		DIPhysicsAsset_GetSphereProjection(LocalPosition, float3(0,0,0), DIContext.ElementExtentBuffer[ElementIndex].x, ElementIndex,
+				CollisionPosition, CollisionNormal, OutputIndex, MinDistance);
+	}
+	else if ( ElementIndex >= DIContext.ElementOffsets[CAPSULE_INDEX] && ElementIndex < DIContext.ElementOffsets[CAPSULE_INDEX+1] )
+	{
+		DIPhysicsAsset_GetCapsuleProjection(LocalPosition, DIContext.ElementExtentBuffer[ElementIndex].xy, ElementIndex,
+				CollisionPosition, CollisionNormal, OutputIndex, MinDistance);
+	}
+	else if ( ElementIndex >= DIContext.ElementOffsets[BOX_INDEX] && ElementIndex < DIContext.ElementOffsets[BOX_INDEX+1] )
+	{
+		DIPhysicsAsset_GetBoxProjection(LocalPosition, DIContext.ElementExtentBuffer[ElementIndex].xyz, ElementIndex,
+				CollisionPosition, CollisionNormal, OutputIndex, MinDistance);
+	}
+	
+	OutClosestDistance = 0.0;
+
+	if (ElementIndex != -1)
+	{
+		OutClosestDistance = MinDistance;
+	}
+}
+
+ // Given a world space position (NodePosition) compute the physics asset closest element
+void DIPhysicsAsset_GetClosestElement(in FDIPhysicsAssetContext DIContext, in float3 NodePosition, in float TimeFraction, out int OutElementIndex)
+{	
+	float3 CollisionPosition = float3(0,0,0);
+	float3 CollisionNormal = float3(0,0,0);
+	float MinDistance = 0.0;
+
+	OutElementIndex = DIPhysicsAsset_GetClosestElement(DIContext,NodePosition,CollisionPosition,CollisionNormal,MinDistance,TimeFraction);
+}
+
+// Given a world space position (NodePosition) compute the physics asset closest point (position,normal,velocity)
 void DIPhysicsAsset_GetClosestPoint(in FDIPhysicsAssetContext DIContext, in float3 NodePosition, in float DeltaTime, in float TimeFraction, out float3 OutClosestPosition, 
-										  out float3 OutClosestNormal, out float3 OutClosestVelocity, out float OutClosestDistance)
+										  out float3 OutClosestNormal, out float3 OutClosestVelocity)
 {	
 	float3 CollisionPosition = float3(0,0,0);
 	float3 CollisionNormal = float3(0,0,0);
@@ -324,7 +414,6 @@ void DIPhysicsAsset_GetClosestPoint(in FDIPhysicsAssetContext DIContext, in floa
 	OutClosestNormal = float3(0,0,0);
 	OutClosestPosition = float3(0,0,0);
 	OutClosestVelocity = float3(0,0,0);
-	OutClosestDistance = 0.0;
 
 	if (ElementIndex != -1)
 	{
@@ -338,12 +427,31 @@ void DIPhysicsAsset_GetClosestPoint(in FDIPhysicsAssetContext DIContext, in floa
 		OutClosestNormal = normalize(PreviousNormal + TimeFraction * (CurrentNormal-PreviousNormal));
 
 		OutClosestVelocity = ( CurrentPosition - PreviousPosition ) / DeltaTime;
+	}
+}
+
+// Given a world space position (NodePosition) compute the physics asset closest distance
+void DIPhysicsAsset_GetClosestDistance(in FDIPhysicsAssetContext DIContext, in float3 NodePosition, in float TimeFraction, out float OutClosestDistance)
+{	
+	float3 CollisionPosition = float3(0,0,0);
+	float3 CollisionNormal = float3(0,0,0);
+	float MinDistance = 0.0;
+
+	OutClosestDistance = 0.0;
+
+	const int ElementIndex = DIPhysicsAsset_GetClosestElement(DIContext,NodePosition,CollisionPosition,CollisionNormal,MinDistance,TimeFraction);
+
+	OutClosestDistance = 0.0;
+
+	if (ElementIndex != -1)
+	{
 		OutClosestDistance = MinDistance;
 	}
 }
 
+
 /* -----------------------------------------------------------------
- * Get the closest point to the physics asset 
+ * Get the texture point to the physics asset 
  * -----------------------------------------------------------------
  */
 
diff --git a/Engine/Plugins/Experimental/HairStrands/Source/HairStrandsCore/Private/GroomAsset.cpp b/Engine/Plugins/Experimental/HairStrands/Source/HairStrandsCore/Private/GroomAsset.cpp
--- a/Engine/Plugins/Experimental/HairStrands/Source/HairStrandsCore/Private/GroomAsset.cpp
+++ b/Engine/Plugins/Experimental/HairStrands/Source/HairStrandsCore/Private/GroomAsset.cpp
@@ -441,7 +441,7 @@ void FHairStrandsRootResource::InitRHI()
 			const bool bHasValidCPUWeights = CPUData.MeshSampleIndicesBuffer.Num() > 0;
 			if(bHasValidCPUWeights)
 			{
-				check(CPUData.MeshInterpolationWeightsBuffer.Num() == CPUData.SampleCount * CPUData.SampleCount);
+				//check(CPUData.MeshInterpolationWeightsBuffer.Num() == (CPUData.SampleCount+4) * (CPUData.SampleCount+4));
 				check(CPUData.MeshSampleIndicesBuffer.Num() == CPUData.SampleCount);
 				check(CPUData.RestSamplePositionsBuffer.Num() == CPUData.SampleCount);
 
@@ -451,12 +451,12 @@ void FHairStrandsRootResource::InitRHI()
 			}
 			else
 			{
-				CreateBuffer<FHairStrandsWeightFormat>(CPUData.SampleCount * CPUData.SampleCount, GPUData.MeshInterpolationWeightsBuffer);
+				CreateBuffer<FHairStrandsWeightFormat>((CPUData.SampleCount+4) * (CPUData.SampleCount+4), GPUData.MeshInterpolationWeightsBuffer);
 				CreateBuffer<FHairStrandsIndexFormat>(CPUData.SampleCount, GPUData.MeshSampleIndicesBuffer);
 				CreateBuffer<FHairStrandsMeshTrianglePositionFormat>(CPUData.SampleCount, GPUData.RestSamplePositionsBuffer);
 			}
 			CreateBuffer<FHairStrandsMeshTrianglePositionFormat>(CPUData.SampleCount, GPUData.DeformedSamplePositionsBuffer);
-			CreateBuffer<FHairStrandsMeshTrianglePositionFormat>(CPUData.SampleCount, GPUData.MeshSampleWeightsBuffer);
+			CreateBuffer<FHairStrandsMeshTrianglePositionFormat>(CPUData.SampleCount+4, GPUData.MeshSampleWeightsBuffer);
 
 
 			// Strand hair roots translation and rotation in triangle-deformed position relative to the bound triangle 
@@ -649,7 +649,7 @@ void FHairStrandsInterpolationResource::InitRHI()
 	CreateBuffer<FHairStrandsInterpolation0Format>(RenderData.Interpolation0, Interpolation0Buffer);
 	CreateBuffer<FHairStrandsInterpolation1Format>(RenderData.Interpolation1, Interpolation1Buffer);
 	CreateBuffer<FHairStrandsRootIndexFormat>(SimRootPointIndex, SimRootPointIndexBuffer);
-	SimRootPointIndex.SetNum(0);
+	//SimRootPointIndex.SetNum(0);
 }
 
 void FHairStrandsInterpolationResource::ReleaseRHI()
@@ -727,51 +727,67 @@ void UGroomAsset::Serialize(FArchive& Ar)
 	}
 }
 
-UGroomAsset::UGroomAsset(const FObjectInitializer& ObjectInitializer)
-	: Super(ObjectInitializer)
+FHairSolverSettings::FHairSolverSettings()
 {
-	bIsInitialized = false;
-
-	HairToGuideDensity = 0.1f;
-
-	SimulatedGroup = 0;
 	EnableSimulation = false;
 	NiagaraSolver = EGroomNiagaraSolvers::AngularSprings;
-
+	KinematicsTarget = false;
 	SubSteps = 5;
-	IterationCount = 20;
+	IterationCount = 5;
+}
 
+FHairExternalForces::FHairExternalForces()
+{
 	GravityVector = FVector(0.0, 0.0, -981.0);
 	AirDrag = 0.1;
 	AirVelocity = FVector(0, 0, 0);
+}
 
+FHairBendConstraint::FHairBendConstraint()
+{
 	SolveBend = true;
 	ProjectBend = false;
-	BendDamping = 0.01;
+	BendDamping = 0.005;
 	BendStiffness = 0.01;
 
 	BendScale.GetRichCurve()->SetKeyInterpMode(BendScale.GetRichCurve()->AddKey(0.f, 1.f), ERichCurveInterpMode::RCIM_Cubic);
 	BendScale.GetRichCurve()->SetKeyInterpMode(BendScale.GetRichCurve()->AddKey(1.f, 1.f), ERichCurveInterpMode::RCIM_Cubic);
+}
 
+FHairStretchConstraint::FHairStretchConstraint()
+{
 	SolveStretch = true;
 	ProjectStretch = false;
-	StretchDamping = 0.01;
+	StretchDamping = 0.005;
 	StretchStiffness = 1.0;
 
 	StretchScale.GetRichCurve()->SetKeyInterpMode(StretchScale.GetRichCurve()->AddKey(0.f, 1.f), ERichCurveInterpMode::RCIM_Cubic);
 	StretchScale.GetRichCurve()->SetKeyInterpMode(StretchScale.GetRichCurve()->AddKey(1.f, 1.f), ERichCurveInterpMode::RCIM_Cubic);
+}
 
+FHairCollisionConstraint::FHairCollisionConstraint()
+{
 	SolveCollision = true;
 	ProjectCollision = true;
 	KineticFriction = 0.1;
 	StaticFriction = 0.1;
 	StrandsViscosity = 1.0;
-	CollisionRadius = 1.0;
+	CollisionRadius = 0.1;
 
 	RadiusScale.GetRichCurve()->SetKeyInterpMode(RadiusScale.GetRichCurve()->AddKey(0.f, 1.0f), ERichCurveInterpMode::RCIM_Cubic);
 	RadiusScale.GetRichCurve()->SetKeyInterpMode(RadiusScale.GetRichCurve()->AddKey(1.f, 0.1f), ERichCurveInterpMode::RCIM_Cubic);
+}
+
+FHairMaterialConstraints::FHairMaterialConstraints()
+{
+	BendConstraint = FHairBendConstraint();
+	StretchConstraint = FHairStretchConstraint();
+	CollisionConstraint = FHairCollisionConstraint();
+}
 
-	StrandsSize = EGroomStrandsSize::Size16;
+FHairStrandsParameters::FHairStrandsParameters()
+{
+	StrandsSize = EGroomStrandsSize::Size8;
 	StrandsDensity = 1.0;
 	StrandsSmoothing = 0.1;
 	StrandsThickness = 0.01;
@@ -780,6 +796,25 @@ UGroomAsset::UGroomAsset(const FObjectInitializer& ObjectInitializer)
 	ThicknessScale.GetRichCurve()->SetKeyInterpMode(ThicknessScale.GetRichCurve()->AddKey(1.f, 1.0f), ERichCurveInterpMode::RCIM_Cubic);
 }
 
+FHairGroupsPhysics::FHairGroupsPhysics()
+{
+	SolverSettings = FHairSolverSettings();
+
+	ExternalForces = FHairExternalForces();
+
+	MaterialConstraints = FHairMaterialConstraints();
+
+	StrandsParameters = FHairStrandsParameters();
+}
+
+UGroomAsset::UGroomAsset(const FObjectInitializer& ObjectInitializer)
+	: Super(ObjectInitializer)
+{
+	bIsInitialized = false;
+
+	HairToGuideDensity = 0.1f;
+}
+
 void UGroomAsset::InitResource()
 {
 	bIsInitialized = true;
@@ -912,6 +947,7 @@ void UGroomAsset::Reset()
 
 	HairGroupsInfo.Reset();
 	HairGroupsData.Reset();
+	HairGroupsPhysics.Reset();
 }
 
 void UGroomAsset::PostLoad()
@@ -942,6 +978,10 @@ void UGroomAsset::PostLoad()
 	{
 		InitResource();
 	}
+	if (HairGroupsPhysics.Num() != HairGroupsInfo.Num())
+	{
+		HairGroupsPhysics.Init(FHairGroupsPhysics(), HairGroupsInfo.Num());
+	}
 }
 
 void UGroomAsset::BeginDestroy()
@@ -1072,7 +1112,7 @@ FArchive& operator<<(FArchive& Ar, FHairGroupInfo& GroupInfo)
 
 	// Ignoring Material since it's null
 
-	return Ar;
+	return Ar;  
 }
 
 bool UGroomAsset::CanRebuildFromDescription() const
diff --git a/Engine/Plugins/Experimental/HairStrands/Source/HairStrandsCore/Private/GroomBuilder.cpp b/Engine/Plugins/Experimental/HairStrands/Source/HairStrandsCore/Private/GroomBuilder.cpp
--- a/Engine/Plugins/Experimental/HairStrands/Source/HairStrandsCore/Private/GroomBuilder.cpp
+++ b/Engine/Plugins/Experimental/HairStrands/Source/HairStrandsCore/Private/GroomBuilder.cpp
@@ -1489,6 +1489,7 @@ bool FGroomBuilder::BuildGroom(const FHairDescription& HairDescription, const FG
 
 	GroomAsset->HairGroupsInfo.Reset();
 	GroomAsset->HairGroupsData.Reset();
+	GroomAsset->HairGroupsPhysics.Reset();
 
 	for (TPair<int32, FProcessedHairDescription::FHairGroup> HairGroupIt : ProcessedHairDescription.HairGroups)
 	{
@@ -1498,6 +1499,7 @@ bool FGroomBuilder::BuildGroom(const FHairDescription& HairDescription, const FG
 
 		GroomAsset->HairGroupsInfo.Add(MoveTemp(GroupInfo));
 		GroomAsset->HairGroupsData.Add(MoveTemp(GroupData));
+		GroomAsset->HairGroupsPhysics.Add(FHairGroupsPhysics());
 	}
 
 	// If there's usable closest guides and guide weights attributes, fill them into the asset
diff --git a/Engine/Plugins/Experimental/HairStrands/Source/HairStrandsCore/Private/GroomComponent.cpp b/Engine/Plugins/Experimental/HairStrands/Source/HairStrandsCore/Private/GroomComponent.cpp
--- a/Engine/Plugins/Experimental/HairStrands/Source/HairStrandsCore/Private/GroomComponent.cpp
+++ b/Engine/Plugins/Experimental/HairStrands/Source/HairStrandsCore/Private/GroomComponent.cpp
@@ -472,13 +472,12 @@ UGroomComponent::UGroomComponent(const FObjectInitializer& ObjectInitializer)
 	RegisteredSkeletalMeshComponent = nullptr;
 	SkeletalPreviousPositionOffset = FVector::ZeroVector;
 	bBindGroomToSkeletalMesh = false;
-	bCreateNiagaraComponent = false;
 	InitializedResources = nullptr;
 	Mobility = EComponentMobility::Movable;
 	bIsGroomAssetCallbackRegistered = false;
 	bIsGroomBindingAssetCallbackRegistered = false;
 	SourceSkeletalMesh = nullptr; 
-	NiagaraComponent = nullptr;
+	NiagaraComponents.Empty();
 
 	SetCollisionProfileName(UCollisionProfile::PhysicsActor_ProfileName);
 }
@@ -523,47 +522,93 @@ void UGroomComponent::UpdateHairGroupsDesc()
 			Desc.HairShadowDensity = GroupData.HairRenderData.HairDensity;
 		}
 	}
-	UpdateNiagaraComponent();
 }
 
-void UGroomComponent::UpdateNiagaraComponent()
+void UGroomComponent::ReleaseHairSimulation()
 {
-	static UNiagaraSystem* CosseratRodsSystem = LoadObject<UNiagaraSystem>(nullptr, TEXT("/HairStrands/Emitters/GroomAssetSystem.GroomAssetSystem"));
-	static UNiagaraSystem* AngularSpringsSystem = LoadObject<UNiagaraSystem>(nullptr, TEXT("/HairStrands/Emitters/GroomAssetSystem.GroomAssetSystem"));
-
-	if (GroomAsset && GroomAsset->EnableSimulation && bCreateNiagaraComponent)
+	for (int32 i = 0; i < NiagaraComponents.Num(); ++i)
 	{
-		if (!NiagaraComponent)
+		if (NiagaraComponents[i] && !NiagaraComponents[i]->IsBeingDestroyed())
 		{
-			NiagaraComponent = NewObject<UNiagaraComponent>(this);
-			NiagaraComponent->AttachToComponent(this, FAttachmentTransformRules::KeepRelativeTransform);
-			NiagaraComponent->RegisterComponent();
-			NiagaraComponent->SetVisibleFlag(false);
-			NiagaraComponent->Activate(true);
+			if (GetWorld())
+			{
+				NiagaraComponents[i]->DetachFromComponent(FDetachmentTransformRules::KeepRelativeTransform);
+				NiagaraComponents[i]->UnregisterComponent();
+			}
+			NiagaraComponents[i]->DestroyComponent();
+			NiagaraComponents[i] = nullptr;
 		}
-		if (GroomAsset->NiagaraSolver == EGroomNiagaraSolvers::AngularSprings)
+	}
+	NiagaraComponents.Empty();
+}
+
+void UGroomComponent::UpdateHairSimulation()
+{
+	static UNiagaraSystem* CosseratRodsSystem = LoadObject<UNiagaraSystem>(nullptr, TEXT("/HairStrands/Emitters/GroomRodsSystem.GroomRodsSystem"));
+	static UNiagaraSystem* AngularSpringsSystem = LoadObject<UNiagaraSystem>(nullptr, TEXT("/HairStrands/Emitters/GroomSpringsSystem.GroomSpringsSystem"));
+
+	const int32 NumGroups = GroomAsset ? GroomAsset->HairGroupsPhysics.Num() : 0;
+	const int32 NumComponents = FMath::Max(NumGroups, NiagaraComponents.Num());
+
+	TArray<bool> ValidComponents;
+	ValidComponents.Init(false, NumComponents);
+
+	if (GroomAsset)
+	{
+		for (int32 i = 0; i < NumGroups; ++i)
 		{
-			NiagaraComponent->SetAsset(AngularSpringsSystem);
+			ValidComponents[i] = GroomAsset->HairGroupsPhysics[i].SolverSettings.EnableSimulation;
 		}
-		else if (GroomAsset->NiagaraSolver == EGroomNiagaraSolvers::CosseratRods)
+	}
+	NiagaraComponents.SetNumZeroed(NumComponents);
+	for (int32 i = 0; i < NumComponents; ++i)
+	{
+		UNiagaraComponent*& NiagaraComponent = NiagaraComponents[i];
+		if (ValidComponents[i])
 		{
-			NiagaraComponent->SetAsset(CosseratRodsSystem);
+			if (!NiagaraComponent)
+			{
+				NiagaraComponent = NewObject<UNiagaraComponent>(this, NAME_None, RF_Transient);
+				if (GetWorld())
+				{
+					NiagaraComponent->AttachToComponent(this, FAttachmentTransformRules::KeepRelativeTransform);
+					NiagaraComponent->RegisterComponent();
+				}
+				else
+				{
+					NiagaraComponent->SetupAttachment(this);
+				}
+				NiagaraComponent->SetVisibleFlag(false);
+			}
+			if (GroomAsset->HairGroupsPhysics[i].SolverSettings.NiagaraSolver == EGroomNiagaraSolvers::AngularSprings)
+			{
+				NiagaraComponent->SetAsset(AngularSpringsSystem);
+				NiagaraComponent->Activate(true);
+			}
+			else if (GroomAsset->HairGroupsPhysics[i].SolverSettings.NiagaraSolver == EGroomNiagaraSolvers::CosseratRods)
+			{
+				NiagaraComponent->SetAsset(CosseratRodsSystem);
+				NiagaraComponent->Activate(true);
+			}
+			if (NiagaraComponent->GetSystemInstance())
+			{
+				NiagaraComponent->GetSystemInstance()->Reset(FNiagaraSystemInstance::EResetMode::ReInit);
+				NiagaraComponent->GetSystemInstance()->Reset(FNiagaraSystemInstance::EResetMode::ResetAll);
+			}
 		}
-		if (NiagaraComponent->GetSystemInstance())
+		else if (NiagaraComponent && !NiagaraComponent->IsBeingDestroyed())
 		{
-			NiagaraComponent->GetSystemInstance()->Reset(FNiagaraSystemInstance::EResetMode::ReInit);
-			NiagaraComponent->GetSystemInstance()->Reset(FNiagaraSystemInstance::EResetMode::ResetAll);
+			if (GetWorld())
+			{
+				NiagaraComponent->DetachFromComponent(FDetachmentTransformRules::KeepRelativeTransform);
+				NiagaraComponent->UnregisterComponent();
+			}
+			NiagaraComponent->DestroyComponent();
+			NiagaraComponent = nullptr;
 		}
-		EnableSimulatedGroups();
-	}
-	else if (NiagaraComponent)
-	{
-		NiagaraComponent->DetachFromComponent(FDetachmentTransformRules::KeepRelativeTransform);
-		NiagaraComponent->UnregisterComponent();
-		NiagaraComponent->DestroyComponent();
-		NiagaraComponent = nullptr;
-		DisableSimulatedGroups();
 	}
+	NiagaraComponents.SetNum(NumGroups);
+	UpdateSimulatedGroups();
 }
 
 void UGroomComponent::SetGroomAsset(UGroomAsset* Asset)
@@ -584,12 +629,12 @@ void UGroomComponent::SetGroomAsset(UGroomAsset* Asset)
 	}
 
 	UpdateHairGroupsDesc();
+	UpdateHairSimulation();
 	if (!GroomAsset)
 		return;
 	InitResources();
 }
 
-static void ResetSimulation(const USceneComponent* Component);
 void UGroomComponent::SetGroomAsset(UGroomAsset* Asset, UGroomBindingAsset* InBinding)
 {
 	ReleaseResources();
@@ -613,10 +658,10 @@ void UGroomComponent::SetGroomAsset(UGroomAsset* Asset, UGroomBindingAsset* InBi
 	}
 
 	UpdateHairGroupsDesc();
+	UpdateHairSimulation();
 	if (!GroomAsset)
 		return;
 	InitResources();
-	ResetSimulation(this);
 }
 
 void UGroomComponent::SetStableRasterization(bool bEnable)
@@ -848,52 +893,13 @@ template<typename T> void SafeRelease(T*& Data)
 	}
 }
 
-static bool IsSimulationEnabled(const USceneComponent* Component)
-{
-	check(Component);
-
-	// If the groom component has an Niagara component attached, we assume it has simulation capabilities
-	bool bHasNiagaraSimulationComponent = false;
-	for (int32 ChildIt = 0, ChildCount = Component->GetNumChildrenComponents(); ChildIt < ChildCount; ++ChildIt)
-	{
-		const USceneComponent* ChildComponent = Component->GetChildComponent(ChildIt);
-		const UNiagaraComponent* NiagaraComponent = Cast<UNiagaraComponent>(ChildComponent);
-		if (NiagaraComponent != nullptr)
-		{
-			bHasNiagaraSimulationComponent = true;
-			break;
-		}
-	}
-
-	return bHasNiagaraSimulationComponent;
-}
-
-static void ResetSimulation(const USceneComponent* Component)
-{
-	check(Component);
-
-	for (int32 ChildIt = 0, ChildCount = Component->GetNumChildrenComponents(); ChildIt < ChildCount; ++ChildIt)
-	{
-		const USceneComponent* ChildComponent = Component->GetChildComponent(ChildIt);
-		const UNiagaraComponent* NiagaraComponent = Cast<UNiagaraComponent>(ChildComponent);
-		if (NiagaraComponent != nullptr)
-		{
-			if (NiagaraComponent->GetSystemInstance())
-			{
-				NiagaraComponent->GetSystemInstance()->Reset(FNiagaraSystemInstance::EResetMode::ReInit);
-				NiagaraComponent->GetSystemInstance()->Reset(FNiagaraSystemInstance::EResetMode::ResetAll);
-			}
-		}
-	}
-}
-
 EWorldType::Type UGroomComponent::GetWorldType() const
 {
 	EWorldType::Type WorldType = GetWorld() ? EWorldType::Type(GetWorld()->WorldType) : EWorldType::None;
 	return WorldType == EWorldType::Inactive ? EWorldType::Editor : WorldType;
 }
 
-void UGroomComponent::EnableSimulatedGroups()
+void UGroomComponent::UpdateSimulatedGroups()
 {
 	if (InterpolationInput)
 	{
@@ -908,7 +914,8 @@ void UGroomComponent::EnableSimulatedGroups()
 			int32 GroupIt = 0;
 			for (FHairStrandsInterpolationInput::FHairGroup& HairGroup : LocalInterpolationInput->HairGroups)
 			{
-				const bool bIsSimulationEnable = LocalGroomAsset->EnableSimulation && (LocalGroomAsset->SimulatedGroup == GroupIt);
+				const bool bIsSimulationEnable = (LocalGroomAsset && GroupIt < LocalGroomAsset->HairGroupsPhysics.Num()) ? 
+					LocalGroomAsset->HairGroupsPhysics[GroupIt].SolverSettings.EnableSimulation : false;
 				HairGroup.bIsSimulationEnable = bIsSimulationEnable;
 				UpdateHairStrandsDebugInfo(Id, WorldType, GroupIt, bIsSimulationEnable);
 				++GroupIt;
@@ -917,44 +924,12 @@ void UGroomComponent::EnableSimulatedGroups()
 	}
 }
 
-void UGroomComponent::DisableSimulatedGroups()
-{
-	if (InterpolationInput)
-	{
-		const uint32 Id = ComponentId.PrimIDValue;
-		const EWorldType::Type WorldType = GetWorldType();
-
-		FHairStrandsInterpolationInput* LocalInterpolationInput = InterpolationInput;
-		UGroomAsset* LocalGroomAsset = GroomAsset;
-		ENQUEUE_RENDER_COMMAND(FHairStrandsTick_DisableSimulatedGroups)(
-			[LocalInterpolationInput, LocalGroomAsset, Id, WorldType](FRHICommandListImmediate& RHICmdList)
-		{
-			int32 GroupIt = 0;
-			for (FHairStrandsInterpolationInput::FHairGroup& HairGroup : LocalInterpolationInput->HairGroups)
-			{
-				HairGroup.bIsSimulationEnable = false;
-				UpdateHairStrandsDebugInfo(Id, WorldType, GroupIt, false);
-				++GroupIt;
-			}
-		});
-	}
-}
-
 void UGroomComponent::OnChildDetached(USceneComponent* ChildComponent)
-{
-	if (Cast<UNiagaraComponent>(ChildComponent) && InterpolationInput)
-	{
-		DisableSimulatedGroups();
-	}
-}
-
+{}
 
 void UGroomComponent::OnChildAttached(USceneComponent* ChildComponent)
 {
-	if (Cast<UNiagaraComponent>(ChildComponent) && InterpolationInput)
-	{
-		EnableSimulatedGroups();
-	}
+
 }
 
 void UGroomComponent::InitResources(bool bIsBindingReloading)
@@ -989,7 +964,6 @@ void UGroomComponent::InitResources(bool bIsBindingReloading)
 		UGroomBindingAsset::IsCompatible(SkeletalMeshComponent ? SkeletalMeshComponent->SkeletalMesh : nullptr, BindingAsset) &&
 		UGroomBindingAsset::IsCompatible(GroomAsset, BindingAsset) && 
 		UGroomBindingAsset::IsBindingAssetValid(BindingAsset, bIsBindingReloading);
-	const bool bIsSimulationEnable = IsSimulationEnabled(this);
 
 	FTransform HairLocalToWorld = GetComponentTransform();
 	FTransform SkinLocalToWorld = bBindGroomToSkeletalMesh && SkeletalMeshComponent ? SkeletalMeshComponent->GetComponentTransform() : FTransform::Identity;
@@ -1087,11 +1061,14 @@ void UGroomComponent::InitResources(bool bIsBindingReloading)
 		InterpolationInputGroup.InRenderHairPositionOffset = RenderRestHairPositionOffset;
 		InterpolationInputGroup.InSimHairPositionOffset = SimRestHairPositionOffset;
 
+		const bool bIsSimulationEnable = (GroomAsset && GroupIt < GroomAsset->HairGroupsPhysics.Num()) ?
+			GroomAsset->HairGroupsPhysics[GroupIt].SolverSettings.EnableSimulation : false;
+
 		// For skinned groom, these value will be updated during TickComponent() call
 		// Deformed sim & render are expressed within the referential (unlike rest pose)
 		InterpolationInputGroup.OutHairPositionOffset = RenderRestHairPositionOffset;
 		InterpolationInputGroup.OutHairPreviousPositionOffset = RenderRestHairPositionOffset;
-		InterpolationInputGroup.bIsSimulationEnable = bIsSimulationEnable && GroomAsset->EnableSimulation && (GroomAsset->SimulatedGroup == GroupIt);
+		InterpolationInputGroup.bIsSimulationEnable = bIsSimulationEnable;
 		DebugHairGroup.bHasSimulation = InterpolationInputGroup.bIsSimulationEnable;
 
 		GroupIt++;
@@ -1362,7 +1339,7 @@ void UGroomComponent::PostLoad()
 #if WITH_EDITOR
 void UGroomComponent::Invalidate()
 {
-	UpdateNiagaraComponent();
+	UpdateHairSimulation();
 	MarkRenderStateDirty();
 	ValidateMaterials(false);
 }
@@ -1376,6 +1353,7 @@ void UGroomComponent::InvalidateAndRecreate()
 
 void UGroomComponent::OnRegister()
 {
+	UpdateHairSimulation();
 	Super::OnRegister();
 
 	if (!InitializedResources)
@@ -1399,6 +1377,7 @@ void UGroomComponent::OnRegister()
 void UGroomComponent::OnUnregister()
 {
 	Super::OnUnregister();
+	ReleaseHairSimulation();
 }
 
 void UGroomComponent::OnComponentDestroyed(bool bDestroyingHierarchy)
@@ -1601,7 +1580,6 @@ void UGroomComponent::PostEditChangeProperty(FPropertyChangedEvent& PropertyChan
 	const bool bSourceSkeletalMeshChanged = PropertyName == GET_MEMBER_NAME_CHECKED(UGroomComponent, SourceSkeletalMesh);
 	const bool bBindingAssetChanged = PropertyName == GET_MEMBER_NAME_CHECKED(UGroomComponent, BindingAsset);
 	const bool bBindToSkeletalChanged = PropertyName == GET_MEMBER_NAME_CHECKED(UGroomComponent, bBindGroomToSkeletalMesh);
-	const bool bDefaultNiagaraChanged = PropertyName == GET_MEMBER_NAME_CHECKED(UGroomComponent, bCreateNiagaraComponent);
 	const bool bIsBindingCompatible = UGroomBindingAsset::IsCompatible(GroomAsset, BindingAsset);
 	if (!bIsBindingCompatible || !UGroomBindingAsset::IsBindingAssetValid(BindingAsset))
 	{
@@ -1617,7 +1595,7 @@ void UGroomComponent::PostEditChangeProperty(FPropertyChangedEvent& PropertyChan
 		GroomAsset = nullptr;
 	}
 
-	const bool bRecreateResources = bAssetChanged || bBindingAssetChanged || PropertyThatChanged == nullptr || bBindToSkeletalChanged || bDefaultNiagaraChanged || bSourceSkeletalMeshChanged;
+	const bool bRecreateResources = bAssetChanged || bBindingAssetChanged || PropertyThatChanged == nullptr || bBindToSkeletalChanged || bSourceSkeletalMeshChanged;
 	if (bRecreateResources)
 	{
 		// Release the resources before Super::PostEditChangeProperty so that they get
@@ -1638,7 +1616,6 @@ void UGroomComponent::PostEditChangeProperty(FPropertyChangedEvent& PropertyChan
 	{
 		if (GroomAsset)
 		{
-			ResetSimulation(this);
 			// Set the callback on the new GroomAsset being assigned
 			GroomAsset->GetOnGroomAssetChanged().AddUObject(this, &UGroomComponent::Invalidate);
 			bIsGroomAssetCallbackRegistered = true;
@@ -1649,7 +1626,6 @@ void UGroomComponent::PostEditChangeProperty(FPropertyChangedEvent& PropertyChan
 	{
 		if (BindingAsset)
 		{
-			ResetSimulation(this);
 			// Set the callback on the new GroomAsset being assigned
 			BindingAsset->GetOnGroomBindingAssetChanged().AddUObject(this, &UGroomComponent::InvalidateAndRecreate);
 			bIsGroomBindingAssetCallbackRegistered = true;
diff --git a/Engine/Plugins/Experimental/HairStrands/Source/HairStrandsCore/Private/HairStrandsRendering.cpp b/Engine/Plugins/Experimental/HairStrands/Source/HairStrandsCore/Private/HairStrandsRendering.cpp
--- a/Engine/Plugins/Experimental/HairStrands/Source/HairStrandsCore/Private/HairStrandsRendering.cpp
+++ b/Engine/Plugins/Experimental/HairStrands/Source/HairStrandsCore/Private/HairStrandsRendering.cpp
@@ -1342,26 +1342,61 @@ struct FWeightsBuilder
 FWeightsBuilder::FWeightsBuilder(const uint32 NumRows, const uint32 NumColumns,
 	const FVector* SourcePositions, const FVector* TargetPositions)
 {
-	MatrixEntries.SetNum(NumRows * NumColumns);  
-	InverseEntries.SetNum(NumRows * NumColumns);
+	const uint32 PolyRows = NumRows + 4;
+	const uint32 PolyColumns = NumColumns + 4;
+
+	MatrixEntries.Init(0.0, PolyRows * PolyColumns);
+	InverseEntries.Init(0.0, PolyRows * PolyColumns);
 	TArray<float>& LocalEntries = MatrixEntries;
 	ParallelFor(NumRows,
 		[
 			NumRows,
 			NumColumns,
+			PolyRows,
+			PolyColumns,
 			SourcePositions,
 			TargetPositions,
 			&LocalEntries
 		] (uint32 RowIndex)
 	{
-		int32 EntryIndex = RowIndex * NumColumns;
+		int32 EntryIndex = RowIndex * PolyColumns;
 		for (uint32 j = 0; j < NumColumns; ++j)
 		{
 			const float FunctionScale = (SourcePositions[RowIndex] - TargetPositions[j]).Size();
 			LocalEntries[EntryIndex++] = FMath::Sqrt(FunctionScale*FunctionScale + 1.0);
 		}
+		LocalEntries[EntryIndex++] = 1.0;
+		LocalEntries[EntryIndex++] = SourcePositions[RowIndex].X;
+		LocalEntries[EntryIndex++] = SourcePositions[RowIndex].Y;
+		LocalEntries[EntryIndex++] = SourcePositions[RowIndex].Z;
+
+		EntryIndex = NumRows* PolyColumns + RowIndex;
+		LocalEntries[EntryIndex] = 1.0;
+
+		EntryIndex += PolyColumns;
+		LocalEntries[EntryIndex] = SourcePositions[RowIndex].X;
+
+		EntryIndex += PolyColumns;
+		LocalEntries[EntryIndex] = SourcePositions[RowIndex].Y;
+
+		EntryIndex += PolyColumns;
+		LocalEntries[EntryIndex] = SourcePositions[RowIndex].Z;
+
+		const float REGUL_VALUE = 1e-4;
+		EntryIndex = NumRows * PolyColumns + NumColumns;
+		LocalEntries[EntryIndex] = REGUL_VALUE;
+
+		EntryIndex += PolyColumns+1;
+		LocalEntries[EntryIndex] = REGUL_VALUE;
+
+		EntryIndex += PolyColumns+1;
+		LocalEntries[EntryIndex] = REGUL_VALUE;
+
+		EntryIndex += PolyColumns+1;
+		LocalEntries[EntryIndex] = REGUL_VALUE;
+
 	});
-	ComputeWeights(NumRows, NumColumns);
+	ComputeWeights(PolyRows, PolyColumns);
 }
 
 void FWeightsBuilder::ComputeWeights(const uint32 NumRows, const uint32 NumColumns)
diff --git a/Engine/Plugins/Experimental/HairStrands/Source/HairStrandsCore/Public/GroomAsset.h b/Engine/Plugins/Experimental/HairStrands/Source/HairStrandsCore/Public/GroomAsset.h
--- a/Engine/Plugins/Experimental/HairStrands/Source/HairStrandsCore/Public/GroomAsset.h
+++ b/Engine/Plugins/Experimental/HairStrands/Source/HairStrandsCore/Public/GroomAsset.h
@@ -23,8 +23,8 @@
 enum class EGroomNiagaraSolvers : uint8
 {
 	None = 0 UMETA(Hidden),
-	CosseratRods = 0x02 UMETA(DisplatName = "CosseratRods"),
-	AngularSprings = 0x04 UMETA(DisplatName = "AngularSprings")
+	CosseratRods = 0x02 UMETA(DisplatName = "GroomRods"),
+	AngularSprings = 0x04 UMETA(DisplatName = "GroomSprings")
 };
 
 /** Size of each strands*/
@@ -313,198 +313,296 @@
 };
 #endif
 
+
 USTRUCT(BlueprintType)
-struct HAIRSTRANDSCORE_API FHairGroupInfo
+struct HAIRSTRANDSCORE_API FHairSolverSettings
 {
 	GENERATED_BODY()
 
-	UPROPERTY(VisibleAnywhere, Category = "Info")
-	int32 GroupID = 0;
-
-	UPROPERTY(VisibleAnywhere, Category = "Info", meta = (DisplayName = "Curve Count"))
-	int32 NumCurves = 0;
-
-	UPROPERTY(VisibleAnywhere, Category = "Info", meta = (DisplayName = "Guide Count"))
-	int32 NumGuides = 0;
-
-	UPROPERTY(VisibleAnywhere, Category = "Info", meta = (DisplayName = "Import options", ToolTip = "Show the options used at import time"))
-	FGroomBuildSettings ImportSettings;
-
-	UPROPERTY(EditAnywhere, Category="Rendering")
-	UMaterialInterface* Material = nullptr;
-
-	friend FArchive& operator<<(FArchive& Ar, FHairGroupInfo& GroupInfo);
-};
-
-struct HAIRSTRANDSCORE_API FHairGroupData
-{
-	FHairStrandsDatas HairRenderData;
-	FHairStrandsDatas HairSimulationData;
-	FHairStrandsInterpolationDatas HairInterpolationData;
-
-	/** Interpolated hair render resource to be allocated */
-	FHairStrandsRestResource* HairStrandsRestResource = nullptr;
-
-	/** Guide render resource to be allocated */
-	FHairStrandsRestResource* HairSimulationRestResource = nullptr;
-
-	/** Interpolation resource to be allocated */
-	FHairStrandsInterpolationResource* HairInterpolationResource = nullptr;
-
-	/** Cluster culling resource to be allocated */
-	FHairStrandsClusterCullingResource* ClusterCullingResource = nullptr;
-
-	friend FArchive& operator<<(FArchive& Ar, FHairGroupData& GroupData);
-};
-
-/**
- * Implements an asset that can be used to store hair strands
- */
-UCLASS(BlueprintType, hidecategories = (Object))
-class HAIRSTRANDSCORE_API UGroomAsset : public UObject, public IInterface_AssetUserData
-{
-	GENERATED_UCLASS_BODY()
-
-#if WITH_EDITOR
-	/** Notification when anything changed */
-	DECLARE_MULTICAST_DELEGATE(FOnGroomAssetChanged);
-#endif
-
-public:
-
-	UPROPERTY(EditAnywhere, EditFixedSize, BlueprintReadWrite, Category = "HairGroups", meta = (DisplayName = "Group"))
-	TArray<FHairGroupInfo> HairGroupsInfo;
-
-	TArray<FHairGroupData> HairGroupsData;
-
-	//~ Begin UObject Interface.
-	virtual void PostLoad() override;
-	virtual void BeginDestroy() override;
-	virtual void Serialize(FArchive& Ar) override;
-
-	/** Density factor for converting hair into guide curve if no guides are provided. */
-	UPROPERTY(BlueprintReadWrite, Category = "BuildSettings", meta = (ClampMin = "0.01", UIMin = "0.01", UIMax = "1.0"))
-	float HairToGuideDensity = 0.1f; // To remove, as this is now stored into the import settings
-
-	/** Group Id to be simulated */
-	UPROPERTY(EditAnywhere, Category = "Hair Physics|Group Setup", meta = (ToolTip = "Group ID that will be simulated with niagara"))
-	int32 SimulatedGroup;
+	FHairSolverSettings();
 
 	/** Enable the simulation for the group */
-	UPROPERTY(EditAnywhere, Category = "Hair Physics|Group Setup", meta = (ToolTip = "Enable the simulation on that group"))
+	UPROPERTY(EditAnywhere, Category = "SolverSettings", meta = (ToolTip = "Enable the simulation on that group"))
 	bool EnableSimulation;
 
-	/** Enable the simulation for the group */
-	UPROPERTY(EditAnywhere, Category = "Hair Physics|Group Setup", meta = (ToolTip = "Niagara solver to be used for simulation"))
+	/** Allow the simulation to follow a kinematics target*/
+	UPROPERTY(EditAnywhere, Category = "SolverSettings", meta = (ToolTip = "Allow the simulation to follow a kinematics target"))
+	bool KinematicsTarget;
+
+	/** Niagara solver to be used for simulation */
+	UPROPERTY(EditAnywhere, Category = "SolverSettings", meta = (ToolTip = "Niagara solver to be used for simulation"))
 	EGroomNiagaraSolvers NiagaraSolver;
 
 	/** Number of substeps to be used */
-	UPROPERTY(EditAnywhere, Category = "Hair Physics|Niagara Solver", meta = (ToolTip = "Number of sub steps to be done per frame. The actual solver calls are done at 24 fps"))
+	UPROPERTY(EditAnywhere, Category = "SolverSettings", meta = (ToolTip = "Number of sub steps to be done per frame. The actual solver calls are done at 24 fps"))
 	int32 SubSteps;
 
 	/** Number of iterations for the constraint solver  */
-	UPROPERTY(EditAnywhere, Category = "Hair Physics|Niagara Solver", meta = (ToolTip = "Number of iterations to solve the constraints with the xpbd solver"))
+	UPROPERTY(EditAnywhere, Category = "SolverSettings", meta = (ToolTip = "Number of iterations to solve the constraints with the xpbd solver"))
 	int32 IterationCount;
+};
 
+USTRUCT(BlueprintType)
+struct HAIRSTRANDSCORE_API FHairExternalForces
+{
+	GENERATED_BODY()
+
+	FHairExternalForces();
+
 	/** Acceleration vector in cm/s2 to be used for the gravity*/
-	UPROPERTY(EditAnywhere, Category = "Hair Physics|External Forces", meta = (ToolTip = "Acceleration vector in cm/s2 to be used for the gravity"))
+	UPROPERTY(EditAnywhere, Category = "ExternalForces", meta = (ToolTip = "Acceleration vector in cm/s2 to be used for the gravity"))
 	FVector GravityVector;
 
 	/** Coefficient between 0 and 1 to be used for the air drag */
-	UPROPERTY(EditAnywhere, Category = "Hair Physics|External Forces", meta = (ToolTip = "Coefficient between 0 and 1 to be used for the air drag"))
+	UPROPERTY(EditAnywhere, Category = "ExternalForces", meta = (ToolTip = "Coefficient between 0 and 1 to be used for the air drag"))
 	float AirDrag;
 
 	/** Velocity of the surrounding air in cm/s */
-	UPROPERTY(EditAnywhere, Category = "Hair Physics|External Forces", meta = (ToolTip = "Velocity of the surrounding air in cm/s"))
+	UPROPERTY(EditAnywhere, Category = "ExternalForces", meta = (ToolTip = "Velocity of the surrounding air in cm/s"))
 	FVector AirVelocity;
+};
 
+USTRUCT(BlueprintType)
+struct HAIRSTRANDSCORE_API FHairBendConstraint
+{
+	GENERATED_BODY()
+
+	FHairBendConstraint();
+
 	/** Enable the solve of the bend constraint during the xpbd loop */
-	UPROPERTY(EditAnywhere, Category = "Hair Physics|Solver Constraints|Bend Constraint", meta = (ToolTip = "Enable the solve of the bend constraint during the xpbd loop"))
+	UPROPERTY(EditAnywhere, Category = "BendConstraint", meta = (ToolTip = "Enable the solve of the bend constraint during the xpbd loop"))
 	bool SolveBend;
 
-	/** Enable ther projection of the bend constraint after the xpbd loop */
-	UPROPERTY(EditAnywhere, Category = "Hair Physics|Solver Constraints|Bend Constraint", meta = (ToolTip = "Enable ther projection of the bend constraint after the xpbd loop"))
+	/** Enable the projection of the bend constraint after the xpbd loop */
+	UPROPERTY(EditAnywhere, Category = "BendConstraint", meta = (ToolTip = "Enable ther projection of the bend constraint after the xpbd loop"))
 	bool ProjectBend;
 
 	/** Damping for the bend constraint between 0 and 1 */
-	UPROPERTY(EditAnywhere, Category = "Hair Physics|Solver Constraints|Bend Constraint", meta = (ToolTip = "Damping for the bend constraint between 0 and 1"))
+	UPROPERTY(EditAnywhere, Category = "BendConstraint", meta = (ToolTip = "Damping for the bend constraint between 0 and 1"))
 	float BendDamping;
 
 	/** Stiffness for the bend constraint in GPa */
-	UPROPERTY(EditAnywhere, Category = "Hair Physics|Solver Constraints|Bend Constraint", meta = (ToolTip = "Stiffness for the bend constraint in GPa"))
+	UPROPERTY(EditAnywhere, Category = "BendConstraint", meta = (ToolTip = "Stiffness for the bend constraint in GPa"))
 	float BendStiffness;
 
 	/** Stiffness scale along the strand */
-	UPROPERTY(EditAnywhere, Category = "Hair Physics|Solver Constraints|Bend Constraint", meta = (ViewMinInput = "0.0", ViewMaxInput = "1.0", ViewMinOutput = "0.0", ViewMaxOutput = "1.0", TimeLineLength = "1.0", XAxisName = "Strand Coordinate (0,1)", YAxisName = "Bend Scale", ToolTip = "This curve determines how much the bend stiffness will be scaled along each strand. \n The X axis range is [0,1], 0 mapping the root and 1 the tip"))
+	UPROPERTY(EditAnywhere, Category = "BendConstraint", meta = (ViewMinInput = "0.0", ViewMaxInput = "1.0", ViewMinOutput = "0.0", ViewMaxOutput = "1.0", TimeLineLength = "1.0", XAxisName = "Strand Coordinate (0,1)", YAxisName = "Bend Scale", ToolTip = "This curve determines how much the bend stiffness will be scaled along each strand. \n The X axis range is [0,1], 0 mapping the root and 1 the tip"))
 	FRuntimeFloatCurve BendScale;
+};
 
+
+USTRUCT(BlueprintType)
+struct HAIRSTRANDSCORE_API FHairStretchConstraint
+{
+	GENERATED_BODY()
+
+	FHairStretchConstraint();
+
 	/** Enable the solve of the stretch constraint during the xpbd loop */
-	UPROPERTY(EditAnywhere, Category = "Hair Physics|Solver Constraints|Stretch Constraint", meta = (ToolTip = "Enable the solve of the stretch constraint during the xpbd loop"))
+	UPROPERTY(EditAnywhere, Category = "StretchConstraint", meta = (ToolTip = "Enable the solve of the stretch constraint during the xpbd loop"))
 	bool SolveStretch;
 
 	/** Enable the projection of the stretch constraint after the xpbd loop */
-	UPROPERTY(EditAnywhere, Category = "Hair Physics|Solver Constraints|Stretch Constraint", meta = (ToolTip = "Enable ther projection of the stretch constraint after the xpbd loop"))
+	UPROPERTY(EditAnywhere, Category = "StretchConstraint", meta = (ToolTip = "Enable ther projection of the stretch constraint after the xpbd loop"))
 	bool ProjectStretch;
 
 	/** Damping for the stretch constraint between 0 and 1 */
-	UPROPERTY(EditAnywhere, Category = "Hair Physics|Solver Constraints|Stretch Constraint", meta = (ToolTip = "Damping for the stretch constraint between 0 and 1"))
+	UPROPERTY(EditAnywhere, Category = "StretchConstraint", meta = (ToolTip = "Damping for the stretch constraint between 0 and 1"))
 	float StretchDamping;
 
 	/** Stiffness for the stretch constraint in GPa */
-	UPROPERTY(EditAnywhere, Category = "Hair Physics|Solver Constraints|Stretch Constraint", meta = (ToolTip = "Stiffness for the stretch constraint in GPa"))
+	UPROPERTY(EditAnywhere, Category = "StretchConstraint", meta = (ToolTip = "Stiffness for the stretch constraint in GPa"))
 	float StretchStiffness;
 
 	/** Stretch scale along the strand  */
-	UPROPERTY(EditAnywhere, Category = "Hair Physics|Solver Constraints|Stretch Constraint", meta = (ViewMinInput = "0.0", ViewMaxInput = "1.0", ViewMinOutput = "0.0", ViewMaxOutput = "1.0", TimeLineLength = "1.0", XAxisName = "Strand Coordinate (0,1)", YAxisName = "Stretch Scale", ToolTip = "This curve determines how much the stretch stiffness will be scaled along each strand. \n The X axis range is [0,1], 0 mapping the root and 1 the tip"))
+	UPROPERTY(EditAnywhere, Category = "StretchConstraint", meta = (ViewMinInput = "0.0", ViewMaxInput = "1.0", ViewMinOutput = "0.0", ViewMaxOutput = "1.0", TimeLineLength = "1.0", XAxisName = "Strand Coordinate (0,1)", YAxisName = "Stretch Scale", ToolTip = "This curve determines how much the stretch stiffness will be scaled along each strand. \n The X axis range is [0,1], 0 mapping the root and 1 the tip"))
 	FRuntimeFloatCurve StretchScale;
+};
 
+USTRUCT(BlueprintType)
+struct HAIRSTRANDSCORE_API FHairCollisionConstraint
+{
+	GENERATED_BODY()
+
+	FHairCollisionConstraint();
+
 	/** Enable the solve of the collision constraint during the xpbd loop  */
-	UPROPERTY(EditAnywhere, Category = "Hair Physics|Solver Constraints|Collision Constraint", meta = (ToolTip = "Enable the solve of the collision constraint during the xpbd loop"))
+	UPROPERTY(EditAnywhere, Category = "CollisionConstraint", meta = (ToolTip = "Enable the solve of the collision constraint during the xpbd loop"))
 	bool SolveCollision;
 
 	/** Enable ther projection of the collision constraint after the xpbd loop */
-	UPROPERTY(EditAnywhere, Category = "Hair Physics|Solver Constraints|Collision Constraint", meta = (ToolTip = "Enable ther projection of the collision constraint after the xpbd loop"))
+	UPROPERTY(EditAnywhere, Category = "CollisionConstraint", meta = (ToolTip = "Enable ther projection of the collision constraint after the xpbd loop"))
 	bool ProjectCollision;
 
 	/** Static friction used for collision against the physics asset */
-	UPROPERTY(EditAnywhere, Category = "Hair Physics|Solver Constraints|Collision Constraint", meta = (ToolTip = "Static friction used for collision against the physics asset"))
+	UPROPERTY(EditAnywhere, Category = "Collision Constraint", meta = (ToolTip = "Static friction used for collision against the physics asset"))
 	float StaticFriction;
 
 	/** Kinetic friction used for collision against the physics asset*/
-	UPROPERTY(EditAnywhere, Category = "Hair Physics|Solver Constraints|Collision Constraint", meta = (ToolTip = "Kinetic friction used for collision against the physics asset"))
+	UPROPERTY(EditAnywhere, Category = "CollisionConstraint", meta = (ToolTip = "Kinetic friction used for collision against the physics asset"))
 	float KineticFriction;
 
 	/** Viscosity parameter between 0 and 1 that will be used for self collision */
-	UPROPERTY(EditAnywhere, Category = "Hair Physics|Solver Constraints|Collision Constraint", meta = (ToolTip = "Viscosity parameter between 0 and 1 that will be used for self collision"))
+	UPROPERTY(EditAnywhere, Category = "CollisionConstraint", meta = (ToolTip = "Viscosity parameter between 0 and 1 that will be used for self collision"))
 	float StrandsViscosity;
 
 	/** Radius that will be used for the collision detection against the physics asset */
-	UPROPERTY(EditAnywhere, Category = "Hair Physics|Solver Constraints|Collision Constraint", meta = (ToolTip = "Radius that will be used for the collision detection against the physics asset"))
+	UPROPERTY(EditAnywhere, Category = "CollisionConstraint", meta = (ToolTip = "Radius that will be used for the collision detection against the physics asset"))
 	float CollisionRadius;
 
 	/** Radius scale along the strand */
-	UPROPERTY(EditAnywhere, Category = "Hair Physics|Solver Constraints|Collision Constraint", meta = (ViewMinInput = "0.0", ViewMaxInput = "1.0", ViewMinOutput = "0.0", ViewMaxOutput = "1.0", TimeLineLength = "1.0", XAxisName = "Strand Coordinate (0,1)", YAxisName = "Collision Radius", ToolTip = "This curve determines how much the collision radius will be scaled along each strand. \n The X axis range is [0,1], 0 mapping the root and 1 the tip"))
+	UPROPERTY(EditAnywhere, Category = "CollisionConstraint", meta = (ViewMinInput = "0.0", ViewMaxInput = "1.0", ViewMinOutput = "0.0", ViewMaxOutput = "1.0", TimeLineLength = "1.0", XAxisName = "Strand Coordinate (0,1)", YAxisName = "Collision Radius", ToolTip = "This curve determines how much the collision radius will be scaled along each strand. \n The X axis range is [0,1], 0 mapping the root and 1 the tip"))
 	FRuntimeFloatCurve RadiusScale;
+};
 
+
+USTRUCT(BlueprintType)
+struct HAIRSTRANDSCORE_API FHairMaterialConstraints
+{
+	GENERATED_BODY()
+
+	FHairMaterialConstraints();
+
+	/** Radius scale along the strand */
+	UPROPERTY(EditAnywhere, Category = "MaterialConstraints", meta = (ToolTip = "Bend constraint settings to be applied to the hair strands"))
+	FHairBendConstraint BendConstraint;
+
+	/** Radius scale along the strand */
+	UPROPERTY(EditAnywhere, Category = "MaterialConstraints", meta = (ToolTip = "Stretch constraint settings to be applied to the hair strands"))
+	FHairStretchConstraint StretchConstraint;
+
+	/** Radius scale along the strand */
+	UPROPERTY(EditAnywhere, Category = "MaterialConstraints", meta = (ToolTip = "Collision constraint settings to be applied to the hair strands"))
+	FHairCollisionConstraint CollisionConstraint;
+};
+
+USTRUCT(BlueprintType)
+struct HAIRSTRANDSCORE_API FHairStrandsParameters
+{
+	GENERATED_BODY()
+
+		FHairStrandsParameters();
+
 	/** Number of particles per guide that will be used for simulation*/
-	UPROPERTY(EditAnywhere, Category = "Hair Physics|Strands Parameters", meta = (ToolTip = "Number of particles per guide that will be used for simulation"))
-	EGroomStrandsSize StrandsSize;
+	UPROPERTY(EditAnywhere, Category = "StrandsParameters", meta = (ToolTip = "Number of particles per guide that will be used for simulation"))
+		EGroomStrandsSize StrandsSize;
 
 	/** Density of the strands in g/cm3 */
-	UPROPERTY(EditAnywhere, Category = "Hair Physics|Strands Parameters", meta = (ToolTip = "Density of the strands in g/cm3"))
-	float StrandsDensity;
+	UPROPERTY(EditAnywhere, Category = "StrandsParameters", meta = (ToolTip = "Density of the strands in g/cm3"))
+		float StrandsDensity;
 
 	/** Smoothing between 0 and 1 of the incoming guides curves for better stability */
-	UPROPERTY(EditAnywhere, Category = "Hair Physics|Strands Parameters", meta = (ToolTip = "Smoothing between 0 and 1 of the incoming guides curves for better stability"))
-	float StrandsSmoothing;
+	UPROPERTY(EditAnywhere, Category = "StrandsParameters", meta = (ToolTip = "Smoothing between 0 and 1 of the incoming guides curves for better stability"))
+		float StrandsSmoothing;
 
 	/** Strands thickness in cm that will be used for mass and inertia computation */
-	UPROPERTY(EditAnywhere, Category = "Hair Physics|Strands Parameters", meta = (ToolTip = "Strands thickness in cm that will be used for mass and inertia computation"))
-	float StrandsThickness;
+	UPROPERTY(EditAnywhere, Category = "StrandsParameters", meta = (ToolTip = "Strands thickness in cm that will be used for mass and inertia computation"))
+		float StrandsThickness;
 
 	/** Thickness scale along the curve */
-	UPROPERTY(EditAnywhere, Category = "Hair Physics|Strands Parameters", meta = (ViewMinInput = "0.0", ViewMaxInput = "1.0", ViewMinOutput = "0.0", ViewMaxOutput = "1.0", TimeLineLength = "1.0", XAxisName = "Strand Coordinate (0,1)", YAxisName = "Strands Thickness", ToolTip = "This curve determines how much the strands thickness will be scaled along each strand. \n The X axis range is [0,1], 0 mapping the root and 1 the tip"))
-	FRuntimeFloatCurve ThicknessScale;
+	UPROPERTY(EditAnywhere, Category = "StrandsParameters", meta = (ViewMinInput = "0.0", ViewMaxInput = "1.0", ViewMinOutput = "0.0", ViewMaxOutput = "1.0", TimeLineLength = "1.0", XAxisName = "Strand Coordinate (0,1)", YAxisName = "Strands Thickness", ToolTip = "This curve determines how much the strands thickness will be scaled along each strand. \n The X axis range is [0,1], 0 mapping the root and 1 the tip"))
+		FRuntimeFloatCurve ThicknessScale;
+};
+
+USTRUCT(BlueprintType)
+struct HAIRSTRANDSCORE_API FHairGroupsPhysics
+{
+	GENERATED_BODY()
+
+	FHairGroupsPhysics();
+
+	/** Enable the simulation for the group */
+	UPROPERTY(EditAnywhere, Category = "GroupsPhysics", meta = (ToolTip = "Solver Settings for the hair physics"))
+	FHairSolverSettings SolverSettings;
+
+	/** Enable the simulation for the group */
+	UPROPERTY(EditAnywhere, Category = "GroupsPhysics", meta = (ToolTip = "External Forces for the hair physics"))
+	FHairExternalForces ExternalForces;
+
+	/** Enable the simulation for the group */
+	UPROPERTY(EditAnywhere, Category = "GroupsPhysics", meta = (ToolTip = "Material Constraints for the hair physics"))
+	FHairMaterialConstraints MaterialConstraints;
+
+	/** Enable the simulation for the group */
+	UPROPERTY(EditAnywhere, Category = "GroupsPhysics", meta = (ToolTip = "Strands Parameters for the hair physics"))
+	FHairStrandsParameters StrandsParameters;
+};
+
+USTRUCT(BlueprintType)
+struct HAIRSTRANDSCORE_API FHairGroupInfo
+{
+	GENERATED_BODY()
+
+	UPROPERTY(VisibleAnywhere, Category = "Info")
+	int32 GroupID = 0;
+
+	UPROPERTY(VisibleAnywhere, Category = "Info", meta = (DisplayName = "Curve Count"))
+	int32 NumCurves = 0;
+
+	UPROPERTY(VisibleAnywhere, Category = "Info", meta = (DisplayName = "Guide Count"))
+	int32 NumGuides = 0;
+
+	UPROPERTY(VisibleAnywhere, Category = "Info", meta = (DisplayName = "Import options", ToolTip = "Show the options used at import time"))
+	FGroomBuildSettings ImportSettings;
+
+	UPROPERTY(EditAnywhere, Category = "Rendering")
+	UMaterialInterface* Material = nullptr;
+
+	friend FArchive& operator<<(FArchive& Ar, FHairGroupInfo& GroupInfo);
+};
+
+struct HAIRSTRANDSCORE_API FHairGroupData
+{
+	FHairStrandsDatas HairRenderData;
+	FHairStrandsDatas HairSimulationData;
+	FHairStrandsInterpolationDatas HairInterpolationData;
+
+	/** Interpolated hair render resource to be allocated */
+	FHairStrandsRestResource* HairStrandsRestResource = nullptr;
+
+	/** Guide render resource to be allocated */
+	FHairStrandsRestResource* HairSimulationRestResource = nullptr;
+
+	/** Interpolation resource to be allocated */
+	FHairStrandsInterpolationResource* HairInterpolationResource = nullptr;
+
+	/** Cluster culling resource to be allocated */
+	FHairStrandsClusterCullingResource* ClusterCullingResource = nullptr;
+
+	friend FArchive& operator<<(FArchive& Ar, FHairGroupData& GroupData);
+};
+
+/**
+ * Implements an asset that can be used to store hair strands
+ */
+UCLASS(BlueprintType, hidecategories = (Object))
+class HAIRSTRANDSCORE_API UGroomAsset : public UObject, public IInterface_AssetUserData
+{
+	GENERATED_UCLASS_BODY()
 
 #if WITH_EDITOR
+	/** Notification when anything changed */
+	DECLARE_MULTICAST_DELEGATE(FOnGroomAssetChanged);
+#endif
+
+public:
+
+	UPROPERTY(EditAnywhere, EditFixedSize, BlueprintReadWrite, Category = "HairGroups", meta = (DisplayName = "Group"))
+	TArray<FHairGroupInfo> HairGroupsInfo;
+
+	UPROPERTY(EditAnywhere, EditFixedSize, BlueprintReadWrite, Category = "HairPhysics", meta = (DisplayName = "Group"))
+	TArray<FHairGroupsPhysics> HairGroupsPhysics;
+
+	TArray<FHairGroupData> HairGroupsData;
+
+	//~ Begin UObject Interface.
+	virtual void PostLoad() override;
+	virtual void BeginDestroy() override;
+	virtual void Serialize(FArchive& Ar) override;
+
+	/** Density factor for converting hair into guide curve if no guides are provided. */
+	UPROPERTY(BlueprintReadWrite, Category = "BuildSettings", meta = (ClampMin = "0.01", UIMin = "0.01", UIMax = "1.0"))
+	float HairToGuideDensity = 0.1f; // To remove, as this is now stored into the import settings
+
+#if WITH_EDITOR
 	FOnGroomAssetChanged& GetOnGroomAssetChanged() { return OnGroomAssetChanged;  }
diff --git a/Engine/Plugins/Experimental/HairStrands/Source/HairStrandsCore/Public/GroomComponent.h b/Engine/Plugins/Experimental/HairStrands/Source/HairStrandsCore/Public/GroomComponent.h
--- a/Engine/Plugins/Experimental/HairStrands/Source/HairStrandsCore/Public/GroomComponent.h
+++ b/Engine/Plugins/Experimental/HairStrands/Source/HairStrandsCore/Public/GroomComponent.h
@@ -23,9 +23,9 @@ public:
 	UPROPERTY(EditAnywhere, BlueprintReadWrite, Category = "Groom")
 	UGroomAsset* GroomAsset;
 
-	/** Niagara component that will be attached to the system*/
-	//UPROPERTY()
-	class UNiagaraComponent* NiagaraComponent;
+	/** Niagara components that will be attached to the system*/
+	UPROPERTY(Transient)
+	TArray<class UNiagaraComponent*> NiagaraComponents;
 
 	/** 
 	 * When activated, the groom will be attached and skinned onto the skeletal mesh, if the groom component is a child of a skeletal/skinned component.
@@ -44,10 +44,6 @@ public:
 	UPROPERTY(EditAnywhere, BlueprintReadOnly, Category = "Groom")
 	class UGroomBindingAsset* BindingAsset;
 
-	/** Create default Niagara component for simulation */
-	UPROPERTY(EditAnywhere, BlueprintReadOnly, Category = "Groom")
-	bool	bCreateNiagaraComponent;
-
 	/** Boolean to check when the simulation should be reset */
 	bool bResetSimulation;
 
@@ -57,17 +53,17 @@ public:
 	/** Previous bone matrix to compare the difference and decide to reset or not the simulation */
 	FMatrix	PrevBoneMatrix;
 
-	/** Update Niagara component */
-	void UpdateNiagaraComponent();
+	/** Update Niagara components */
+	void UpdateHairSimulation();
+
+	/** Release Niagara components */
+	void ReleaseHairSimulation();
 
 	/** Update Group Description */
 	void UpdateHairGroupsDesc();
 
-	/** Enable simulated groups */
-	void EnableSimulatedGroups();
-
-	/** Disable Simulated groups */
-	void DisableSimulatedGroups();
+	/** Update simulated groups */
+	void UpdateSimulatedGroups();
 
 	//~ Begin UActorComponent Interface.
 	virtual void OnRegister() override;
diff --git a/Engine/Plugins/Experimental/HairStrands/Source/HairStrandsNiagara/Private/NiagaraDataInterfaceHairStrands.cpp b/Engine/Plugins/Experimental/HairStrands/Source/HairStrandsNiagara/Private/NiagaraDataInterfaceHairStrands.cpp
--- a/Engine/Plugins/Experimental/HairStrands/Source/HairStrandsNiagara/Private/NiagaraDataInterfaceHairStrands.cpp
+++ b/Engine/Plugins/Experimental/HairStrands/Source/HairStrandsNiagara/Private/NiagaraDataInterfaceHairStrands.cpp
@@ -69,8 +69,12 @@ static const FName ComputeNodePositionName(TEXT("ComputeNodePosition"));
 static const FName ComputeNodeOrientationName(TEXT("ComputeNodeOrientation"));
 static const FName ComputeNodeMassName(TEXT("ComputeNodeMass"));
 static const FName ComputeNodeInertiaName(TEXT("ComputeNodeInertia"));
+
+//------------------------------------------------------------------------------------------------------------
+
 static const FName ComputeEdgeLengthName(TEXT("ComputeEdgeLength"));
 static const FName ComputeEdgeRotationName(TEXT("ComputeEdgeRotation"));
+static const FName ComputeEdgeDirectionName(TEXT("ComputeEdgeDirection"));
 
 //------------------------------------------------------------------------------------------------------------
 
@@ -89,11 +93,9 @@ static const FName UpdateAngularVelocityName(TEXT("UpdateAngularVelocity"));
 
 static const FName AttachNodePositionName(TEXT("AttachNodePosition"));
 static const FName AttachNodeOrientationName(TEXT("AttachNodeOrientation"));
-static const FName AttachNodeStateName(TEXT("AttachNodeState"));
-
-//------------------------------------------------------------------------------------------------------------
 
-static const FName EvalSkinnedPositionName(TEXT("EvalSkinnedPosition"));
+static const FName AttachNodeStateName(TEXT("AttachNodeState"));
+static const FName UpdateNodeStateName(TEXT("UpdateNodeState"));
 
 //------------------------------------------------------------------------------------------------------------
 
@@ -139,7 +141,6 @@ static const FName ProjectSoftCollisionConstraintName(TEXT("ProjectSoftCollision
 
 //------------------------------------------------------------------------------------------------------------
 
-static const FName ComputeRestDirectionName(TEXT("ComputeRestDirection"));
 static const FName UpdateMaterialFrameName(TEXT("UpdateMaterialFrame"));
 static const FName ComputeMaterialFrameName(TEXT("ComputeMaterialFrame"));
 
@@ -168,7 +169,7 @@ const FString UNiagaraDataInterfaceHairStrands::CurvesOffsetsBufferName(TEXT("Cu
 const FString UNiagaraDataInterfaceHairStrands::RestPositionBufferName(TEXT("RestPositionBuffer_"));
 
 const FString UNiagaraDataInterfaceHairStrands::ResetSimulationName(TEXT("ResetSimulation_"));
-const FString UNiagaraDataInterfaceHairStrands::HasRootAttachedName(TEXT("HasRootAttached_"));
+const FString UNiagaraDataInterfaceHairStrands::InterpolationModeName(TEXT("InterpolationMode_"));
 const FString UNiagaraDataInterfaceHairStrands::RootBarycentricCoordinatesName(TEXT("RootBarycentricCoordinatesBuffer_"));
 
 const FString UNiagaraDataInterfaceHairStrands::RestRootOffsetName(TEXT("RestRootOffset_"));
@@ -191,6 +192,8 @@ const FString UNiagaraDataInterfaceHairStrands::DeformedPositionOffsetName(TEXT(
 const FString UNiagaraDataInterfaceHairStrands::BoundingBoxBufferAName(TEXT("BoundingBoxBufferA_"));
 const FString UNiagaraDataInterfaceHairStrands::BoundingBoxBufferBName(TEXT("BoundingBoxBufferB_"));
 
+const FString UNiagaraDataInterfaceHairStrands::ParamsScaleBufferName(TEXT("ParamsScaleBuffer_"));
+
 //------------------------------------------------------------------------------------------------------------
 
 struct FNDIHairStrandsParametersName
@@ -209,7 +212,7 @@ struct FNDIHairStrandsParametersName
 		CurvesOffsetsBufferName = UNiagaraDataInterfaceHairStrands::CurvesOffsetsBufferName + Suffix;
 		RestPositionBufferName = UNiagaraDataInterfaceHairStrands::RestPositionBufferName + Suffix;
 
-		HasRootAttachedName = UNiagaraDataInterfaceHairStrands::HasRootAttachedName + Suffix;
+		InterpolationModeName = UNiagaraDataInterfaceHairStrands::InterpolationModeName + Suffix;
 		ResetSimulationName = UNiagaraDataInterfaceHairStrands::ResetSimulationName + Suffix;
 		RootBarycentricCoordinatesName = UNiagaraDataInterfaceHairStrands::RootBarycentricCoordinatesName + Suffix;
 
@@ -232,6 +235,8 @@ struct FNDIHairStrandsParametersName
 
 		BoundingBoxBufferAName = UNiagaraDataInterfaceHairStrands::BoundingBoxBufferAName + Suffix;
 		BoundingBoxBufferBName = UNiagaraDataInterfaceHairStrands::BoundingBoxBufferBName + Suffix;
+
+		ParamsScaleBufferName = UNiagaraDataInterfaceHairStrands::ParamsScaleBufferName + Suffix;
 	}
 
 	FString NumStrandsName;
@@ -247,7 +252,7 @@ struct FNDIHairStrandsParametersName
 	FString RestPositionBufferName;
 
 	FString ResetSimulationName;
-	FString HasRootAttachedName;
+	FString InterpolationModeName;
 	FString RootBarycentricCoordinatesName;
 
 	FString RestRootOffsetName;
@@ -269,6 +274,8 @@ struct FNDIHairStrandsParametersName
 
 	FString BoundingBoxBufferAName;
 	FString BoundingBoxBufferBName;
+
+	FString ParamsScaleBufferName;
 };
 
 //------------------------------------------------------------------------------------------------------------
@@ -400,17 +407,18 @@ inline void ClearBuffer(FRHICommandList& RHICmdList, FNDIHairStrandsBuffer* Hair
 //------------------------------------------------------------------------------------------------------------
 
 void FNDIHairStrandsBuffer::Initialize(const FHairStrandsDatas*  HairStrandsDatas, const FHairStrandsRestResource* HairStrandsRestResource, 
-	const FHairStrandsDeformedResource*  HairStrandsDeformedResource, const FHairStrandsRootResource* HairStrandsRootResource)
+	const FHairStrandsDeformedResource*  HairStrandsDeformedResource, const FHairStrandsRootResource* HairStrandsRootResource, const FNDIHairStrandsData* NDIStrandsDatas)
 {
 	SourceDatas = HairStrandsDatas;
 	SourceRestResources = HairStrandsRestResource;
 	SourceDeformedResources = HairStrandsDeformedResource;
 	SourceRootResources = HairStrandsRootResource;
+	NDIDatas = NDIStrandsDatas;
 }
 
 void FNDIHairStrandsBuffer::InitRHI()
 {
-	if (SourceDatas != nullptr && SourceRestResources != nullptr)
+	if (SourceDatas != nullptr && SourceRestResources != nullptr && NDIDatas != nullptr)
 	{
 		{
 			const uint32 OffsetCount = SourceDatas->GetNumCurves() + 1;
@@ -446,6 +454,16 @@ void FNDIHairStrandsBuffer::InitRHI()
 			FMemory::Memcpy(BoundBufferData, ZeroData.GetData(), BoundBytes);
 			RHIUnlockVertexBuffer(BoundingBoxBufferB.Buffer);
 		}
+		{
+			const uint32 ScaleCount = 32 * NumScales;
+			const uint32 ScaleBytes = sizeof(float) * ScaleCount;
+
+			ParamsScaleBuffer.Initialize(sizeof(float), ScaleCount, EPixelFormat::PF_R32_FLOAT, BUF_Static);
+			void* ScaleBufferData = RHILockVertexBuffer(ParamsScaleBuffer.Buffer, 0, ScaleBytes, RLM_WriteOnly);
+
+			FMemory::Memcpy(ScaleBufferData, NDIDatas->ParamsScale.GetData(), ScaleBytes);
+			RHIUnlockVertexBuffer(ParamsScaleBuffer.Buffer);
+		}
 		if (SourceDeformedResources == nullptr)
 		{
 			const uint32 PositionsCount = SourceDatas->GetNumPoints();
@@ -460,6 +478,7 @@ void FNDIHairStrandsBuffer::ReleaseRHI()
 	DeformedPositionBuffer.Release();
 	BoundingBoxBufferA.Release();
 	BoundingBoxBufferB.Release();
+	ParamsScaleBuffer.Release();
 }
 
 //------------------------------------------------------------------------------------------------------------
@@ -478,6 +497,71 @@ void FNDIHairStrandsData::Release()
 	}
 }
 
+void FNDIHairStrandsData::Update(UNiagaraDataInterfaceHairStrands* Interface, FNiagaraSystemInstance* SystemInstance, const FHairStrandsDatas* StrandsDatas, 
+	UGroomAsset* GroomAsset, const int32 GroupIndex)
+{
+	if (Interface != nullptr && SystemInstance != nullptr)
+	{
+		WorldTransform = Interface->IsComponentValid() ? Interface->SourceComponent->GetComponentToWorld() :
+			SystemInstance->GetComponent()->GetComponentToWorld();
+
+		if (StrandsDatas != nullptr && GroomAsset != nullptr && GroupIndex >= 0 && GroupIndex < GroomAsset->HairGroupsPhysics.Num() &&
+			GroomAsset->HairGroupsPhysics[GroupIndex].SolverSettings.EnableSimulation)
+		{
+			FHairGroupsPhysics& HairPhysics = GroomAsset->HairGroupsPhysics[GroupIndex];
+			StrandsSize = static_cast<uint8>(HairPhysics.StrandsParameters.StrandsSize);
+
+			SubSteps = HairPhysics.SolverSettings.SubSteps;
+			IterationCount = HairPhysics.SolverSettings.IterationCount;
+			KinematicsTarget = HairPhysics.SolverSettings.KinematicsTarget;
+
+			GravityVector = HairPhysics.ExternalForces.GravityVector;
+			AirDrag = HairPhysics.ExternalForces.AirDrag;
+			AirVelocity = HairPhysics.ExternalForces.AirVelocity;
+
+			SolveBend = HairPhysics.MaterialConstraints.BendConstraint.SolveBend;
+			ProjectBend = HairPhysics.MaterialConstraints.BendConstraint.ProjectBend;
+			BendDamping = HairPhysics.MaterialConstraints.BendConstraint.BendDamping;
+			BendStiffness = HairPhysics.MaterialConstraints.BendConstraint.BendStiffness;
+
+			SolveStretch = HairPhysics.MaterialConstraints.StretchConstraint.SolveStretch;
+			ProjectStretch = HairPhysics.MaterialConstraints.StretchConstraint.ProjectStretch;
+			StretchDamping = HairPhysics.MaterialConstraints.StretchConstraint.StretchDamping;
+			StretchStiffness = HairPhysics.MaterialConstraints.StretchConstraint.StretchStiffness;
+
+			SolveCollision = HairPhysics.MaterialConstraints.CollisionConstraint.SolveCollision;
+			ProjectCollision = HairPhysics.MaterialConstraints.CollisionConstraint.ProjectCollision;
+			StaticFriction = HairPhysics.MaterialConstraints.CollisionConstraint.StaticFriction;
+			KineticFriction = HairPhysics.MaterialConstraints.CollisionConstraint.KineticFriction;
+			StrandsViscosity = HairPhysics.MaterialConstraints.CollisionConstraint.StrandsViscosity;
+			CollisionRadius = HairPhysics.MaterialConstraints.CollisionConstraint.CollisionRadius;
+
+			StrandsDensity = HairPhysics.StrandsParameters.StrandsDensity;
+			StrandsSmoothing = HairPhysics.StrandsParameters.StrandsSmoothing;
+			StrandsThickness = HairPhysics.StrandsParameters.StrandsThickness;
+
+			for (int32 i = 0; i < StrandsSize; ++i)
+			{
+				const float VertexCoord = static_cast<float>(i) / (StrandsSize-1.0);
+				ParamsScale[32 * BendOffset + i] = HairPhysics.MaterialConstraints.BendConstraint.BendScale.GetRichCurve()->Eval(VertexCoord);
+				ParamsScale[32 * StretchOffset + i] = HairPhysics.MaterialConstraints.StretchConstraint.StretchScale.GetRichCurve()->Eval(VertexCoord);
+				ParamsScale[32 * RadiusOffset + i] = HairPhysics.MaterialConstraints.CollisionConstraint.RadiusScale.GetRichCurve()->Eval(VertexCoord);
+				ParamsScale[32 * ThicknessOffset + i] = HairPhysics.StrandsParameters.ThicknessScale.GetRichCurve()->Eval(VertexCoord);
+			}
+
+			const FBox& StrandsBox = StrandsDatas->BoundingBox;
+
+			NumStrands = StrandsDatas->GetNumCurves();
+			BoxCenter = StrandsBox.GetCenter();
+			BoxExtent = StrandsBox.GetExtent();
+		}
+		else
+		{
+			ResetDatas();
+		}
+	}
+}
+
 bool FNDIHairStrandsData::Init(UNiagaraDataInterfaceHairStrands* Interface, FNiagaraSystemInstance* SystemInstance)
 {
 	HairStrandsBuffer = nullptr;
@@ -489,69 +573,17 @@ bool FNDIHairStrandsData::Init(UNiagaraDataInterfaceHairStrands* Interface, FNia
 		FHairStrandsDeformedResource* StrandsDeformedResource = nullptr;
 		FHairStrandsRootResource* StrandsRootResource = nullptr;
 		UGroomAsset* GroomAsset = nullptr;
+		int32 GroupIndex = 0;
 
 		{
-			Interface->ExtractDatasAndResources(SystemInstance, StrandsDatas, StrandsRestResource, StrandsDeformedResource, StrandsRootResource, GroomAsset);
+			Interface->ExtractDatasAndResources(SystemInstance, StrandsDatas, StrandsRestResource, StrandsDeformedResource, StrandsRootResource, GroomAsset, GroupIndex);
+			Update(Interface, SystemInstance, StrandsDatas, GroomAsset, GroupIndex);
 
 			HairStrandsBuffer = new FNDIHairStrandsBuffer();
-			HairStrandsBuffer->Initialize(StrandsDatas, StrandsRestResource, StrandsDeformedResource, StrandsRootResource);
-			BeginInitResource(HairStrandsBuffer);
-
-			WorldTransform = Interface->IsComponentValid() ? Interface->SourceComponent->GetComponentToWorld() :
-				SystemInstance->GetComponent()->GetComponentToWorld();
-
-			if (StrandsDatas != nullptr && GroomAsset != nullptr && GroomAsset->EnableSimulation)
-			{
-				StrandsSize = static_cast<uint8>(GroomAsset->StrandsSize);
-
-				SubSteps = GroomAsset->SubSteps;
-				IterationCount = GroomAsset->IterationCount;
-
-				GravityVector = GroomAsset->GravityVector;
-				AirDrag = GroomAsset->AirDrag;
-				AirVelocity = GroomAsset->AirVelocity;
+			HairStrandsBuffer->Initialize(StrandsDatas, StrandsRestResource, StrandsDeformedResource, StrandsRootResource, this);
 
-				SolveBend = GroomAsset->SolveBend;
-				ProjectBend = GroomAsset->ProjectBend;
-				BendDamping = GroomAsset->BendDamping;
-				BendStiffness = GroomAsset->BendStiffness;
-
-				SolveStretch = GroomAsset->SolveStretch;
-				ProjectStretch = GroomAsset->ProjectStretch;
-				StretchDamping = GroomAsset->StretchDamping;
-				StretchStiffness = GroomAsset->StretchStiffness;
-
-				SolveCollision = GroomAsset->SolveCollision;
-				ProjectCollision = GroomAsset->ProjectCollision;
-				StaticFriction = GroomAsset->StaticFriction;
-				KineticFriction = GroomAsset->KineticFriction;
-				StrandsViscosity = GroomAsset->StrandsViscosity;
-				CollisionRadius = GroomAsset->CollisionRadius;
-
-				StrandsDensity = GroomAsset->StrandsDensity;
-				StrandsSmoothing = GroomAsset->StrandsSmoothing;
-				StrandsThickness = GroomAsset->StrandsThickness;
-
-				for (int32 i = 0; i < 4; ++i)
-				{
-					const float VertexCoord = static_cast<float>(i) / (3);
-
-					BendScale[i] = GroomAsset->BendScale.GetRichCurve()->Eval(VertexCoord);
-					StretchScale[i] = GroomAsset->StretchScale.GetRichCurve()->Eval(VertexCoord);
-					RadiusScale[i] = GroomAsset->RadiusScale.GetRichCurve()->Eval(VertexCoord);
-					ThicknessScale[i] = GroomAsset->ThicknessScale.GetRichCurve()->Eval(VertexCoord);
-				}
-
-				const FBox& StrandsBox = StrandsDatas->BoundingBox;
-
-				NumStrands = StrandsDatas->GetNumCurves();
-				BoxCenter = StrandsBox.GetCenter();
-				BoxExtent = StrandsBox.GetExtent();
-			}
-			else
-			{
-				ResetDatas();
-			}
+			BeginInitResource(HairStrandsBuffer);
+		
 			TickCount = 0;
 			ForceReset = true;
 		}
@@ -582,7 +614,7 @@ struct FNDIHairStrandsParametersCS : public FNiagaraDataInterfaceParametersCS
 		RestPositionBuffer.Bind(ParameterMap, *ParamNames.RestPositionBufferName);
 
 		ResetSimulation.Bind(ParameterMap, *ParamNames.ResetSimulationName);
-		HasRootAttached.Bind(ParameterMap,*ParamNames.HasRootAttachedName);
+		InterpolationMode.Bind(ParameterMap,*ParamNames.InterpolationModeName);
 		RestRootOffset.Bind(ParameterMap, *ParamNames.RestRootOffsetName);
 		DeformedRootOffset.Bind(ParameterMap, *ParamNames.DeformedRootOffsetName);
 
@@ -606,6 +638,8 @@ struct FNDIHairStrandsParametersCS : public FNiagaraDataInterfaceParametersCS
 		BoundingBoxBufferA.Bind(ParameterMap, *ParamNames.BoundingBoxBufferAName);
 		BoundingBoxBufferB.Bind(ParameterMap, *ParamNames.BoundingBoxBufferBName);
 
+		ParamsScaleBuffer.Bind(ParameterMap, *ParamNames.ParamsScaleBufferName);
+
 		if (!DeformedPositionBuffer.IsBound())
 		{
 			UE_LOG(LogHairStrands, Warning, TEXT("Binding failed for FNDIHairStrandsParametersCS %s. Was it optimized out?"), *ParamNames.DeformedPositionBufferName)
@@ -618,12 +652,14 @@ struct FNDIHairStrandsParametersCS : public FNiagaraDataInterfaceParametersCS
 		{
 			UE_LOG(LogHairStrands, Warning, TEXT("Binding failed for FNDIHairStrandsParametersCS %s. Was it optimized out?"), *ParamNames.RestPositionBufferName)
 		}
-
+		if (!ParamsScaleBuffer.IsBound())
+		{
+			UE_LOG(LogHairStrands, Warning, TEXT("Binding failed for FNDIHairStrandsParametersCS %s. Was it optimized out?"), *ParamNames.ParamsScaleBufferName)
+		}
 		if (!RootBarycentricCoordinatesBuffer.IsBound())
 		{
 			UE_LOG(LogHairStrands, Warning, TEXT("Binding failed for FNDIHairStrandsParametersCS %s. Was it optimized out?"), *ParamNames.RootBarycentricCoordinatesName)
 		}
-
 		if (!RestTrianglePositionABuffer.IsBound())
 		{
 			UE_LOG(LogHairStrands, Warning, TEXT("Binding failed for FNDIHairStrandsParametersCS %s. Was it optimized out?"), *ParamNames.RestTrianglePositionAName)
@@ -636,7 +672,6 @@ struct FNDIHairStrandsParametersCS : public FNiagaraDataInterfaceParametersCS
 		{
 			UE_LOG(LogHairStrands, Warning, TEXT("Binding failed for FNDIHairStrandsParametersCS %s. Was it optimized out?"), *ParamNames.RestTrianglePositionCName)
 		}
-
 		if (!DeformedTrianglePositionABuffer.IsBound())
 		{
 			UE_LOG(LogHairStrands, Warning, TEXT("Binding failed for FNDIHairStrandsParametersCS %s. Was it optimized out?"), *ParamNames.DeformedTrianglePositionAName)
@@ -671,8 +706,8 @@ struct FNDIHairStrandsParametersCS : public FNiagaraDataInterfaceParametersCS
 			InterfaceProxy->SystemInstancesToProxyData.Find(Context.SystemInstance);
 
 		const bool IsHairValid = ProxyData != nullptr && ProxyData->HairStrandsBuffer != nullptr && ProxyData->HairStrandsBuffer->IsInitialized();
-		const bool IsRootValid = IsHairValid && ProxyData->HairStrandsBuffer->SourceRootResources != nullptr && ProxyData->HairStrandsBuffer->SourceRootResources->IsInitialized();
-		const bool IsRestValid = IsHairValid && ProxyData->HairStrandsBuffer->SourceRestResources != nullptr && ProxyData->HairStrandsBuffer->SourceRestResources->IsInitialized();
+		const bool IsRootValid = IsHairValid && ProxyData->HairStrandsBuffer->SourceRootResources != nullptr;//&& ProxyData->HairStrandsBuffer->SourceRootResources->IsInitialized();
+		const bool IsRestValid = IsHairValid && ProxyData->HairStrandsBuffer->SourceRestResources != nullptr;//&& ProxyData->HairStrandsBuffer->SourceRestResources->IsInitialized();
 		const bool IsDeformedValid = IsHairValid && ProxyData->HairStrandsBuffer->SourceDeformedResources != nullptr && ProxyData->HairStrandsBuffer->SourceDeformedResources->IsInitialized();
 
 		if (IsHairValid && IsRestValid)
@@ -682,23 +717,23 @@ struct FNDIHairStrandsParametersCS : public FNiagaraDataInterfaceParametersCS
 			FUnorderedAccessViewRHIRef PointPositionsUAV = IsDeformedValid ?
 				HairStrandsBuffer->SourceDeformedResources->DeformedPositionBuffer[HairStrandsBuffer->SourceDeformedResources->CurrentIndex].UAV : HairStrandsBuffer->DeformedPositionBuffer.UAV;
 
-			const int32 HasRootAttachedValue = (IsRootValid && HairStrandsBuffer->SourceRootResources->MeshProjectionLODs.Num() > 0 && 
+			int32 InterpolationModeValue = (IsRootValid && HairStrandsBuffer->SourceRootResources->MeshProjectionLODs.Num() > 0 && 
 				HairStrandsBuffer->SourceRootResources->MeshProjectionLODs[0].Status == FHairStrandsProjectionHairData::LODData::EStatus::Completed);
 
-			const FHairStrandsRootResource::FMeshProjectionLOD* MeshProjection = (HasRootAttachedValue == 1) ? &(HairStrandsBuffer->SourceRootResources->MeshProjectionLODs[0]) : nullptr;
+			const FHairStrandsRootResource::FMeshProjectionLOD* MeshProjection = (InterpolationModeValue == 1) ? &(HairStrandsBuffer->SourceRootResources->MeshProjectionLODs[0]) : nullptr;
 
-			FRHIShaderResourceView* RestTrianglePositionASRV = (HasRootAttachedValue == 1 && MeshProjection != nullptr) ?
+			FRHIShaderResourceView* RestTrianglePositionASRV = (InterpolationModeValue == 1 && MeshProjection != nullptr) ?
 				MeshProjection->RestRootTrianglePosition0Buffer.SRV.GetReference() : FNiagaraRenderer::GetDummyFloatBuffer();
-			FRHIShaderResourceView* RestTrianglePositionBSRV = (HasRootAttachedValue == 1 && MeshProjection != nullptr) ?
+			FRHIShaderResourceView* RestTrianglePositionBSRV = (InterpolationModeValue == 1 && MeshProjection != nullptr) ?
 				MeshProjection->RestRootTrianglePosition1Buffer.SRV.GetReference() : FNiagaraRenderer::GetDummyFloatBuffer();
-			FRHIShaderResourceView* RestTrianglePositionCSRV = (HasRootAttachedValue == 1 && MeshProjection != nullptr) ?
+			FRHIShaderResourceView* RestTrianglePositionCSRV = (InterpolationModeValue == 1 && MeshProjection != nullptr) ?
 				MeshProjection->RestRootTrianglePosition2Buffer.SRV.GetReference() : FNiagaraRenderer::GetDummyFloatBuffer();
 
-			FRHIShaderResourceView* DeformedTrianglePositionASRV = (HasRootAttachedValue == 1 && MeshProjection != nullptr) ?
+			FRHIShaderResourceView* DeformedTrianglePositionASRV = (InterpolationModeValue == 1 && MeshProjection != nullptr) ?
 				MeshProjection->DeformedRootTrianglePosition0Buffer.SRV.GetReference() : FNiagaraRenderer::GetDummyFloatBuffer();
-			FRHIShaderResourceView* DeformedTrianglePositionBSRV = (HasRootAttachedValue == 1 && MeshProjection != nullptr) ?
+			FRHIShaderResourceView* DeformedTrianglePositionBSRV = (InterpolationModeValue == 1 && MeshProjection != nullptr) ?
 				MeshProjection->DeformedRootTrianglePosition1Buffer.SRV.GetReference() : FNiagaraRenderer::GetDummyFloatBuffer();
-			FRHIShaderResourceView* DeformedTrianglePositionCSRV = (HasRootAttachedValue == 1 && MeshProjection != nullptr) ?
+			FRHIShaderResourceView* DeformedTrianglePositionCSRV = (InterpolationModeValue == 1 && MeshProjection != nullptr) ?
 				MeshProjection->DeformedRootTrianglePosition2Buffer.SRV.GetReference() : FNiagaraRenderer::GetDummyFloatBuffer();
 
 			const int32 SampleCountValue = (MeshProjection != nullptr) ? MeshProjection->SampleCount : 0;
@@ -707,18 +742,30 @@ struct FNDIHairStrandsParametersCS : public FNiagaraDataInterfaceParametersCS
 			FShaderResourceViewRHIRef MeshSampleWeightsBufferSRV = (MeshProjection != nullptr && MeshProjection->SampleCount > 0) ?
 				MeshProjection->MeshSampleWeightsBuffer.SRV.GetReference() : FNiagaraRenderer::GetDummyFloatBuffer();
 
-			FRHIShaderResourceView* RootBarycentricCoordinatesSRV = (HasRootAttachedValue == 1 && MeshProjection != nullptr) ?
+			FRHIShaderResourceView* RootBarycentricCoordinatesSRV = (InterpolationModeValue == 1 && MeshProjection != nullptr) ?
 				MeshProjection->RootTriangleBarycentricBuffer.SRV.GetReference() : FNiagaraRenderer::GetDummyFloatBuffer();
 
 			int32 bNeedSimReset = (ProxyData->TickCount <= MaxDelay ? 1 : 0);
 
-			//UE_LOG(LogHairStrands, Log, TEXT("Shader Reset : %d %d %d"), bNeedSimReset, IsRootValid, HasRootAttachedValue);
+			static const auto GlobalInterpolationCVar = IConsoleManager::Get().FindConsoleVariable(TEXT("r.HairStrands.Interpolation.Global"));
+			const int32 GlobalInterpolation = (GlobalInterpolationCVar != nullptr) ? GlobalInterpolationCVar->GetInt() : 0;
+
+			if ((InterpolationModeValue == 1) && (GlobalInterpolation == 1))
+			{
+				InterpolationModeValue = 2;
+			}
+
+			FUnorderedAccessViewRHIRef MeshSampleWeightsBufferUAV = (MeshProjection != nullptr && MeshProjection->SampleCount > 0) ?
+				MeshProjection->MeshSampleWeightsBuffer.UAV : Context.Batcher->GetEmptyRWBufferFromPool(RHICmdList, PF_R32_FLOAT);
+
+			//UE_LOG(LogHairStrands, Log, TEXT("Shader Reset : %d %d %d"), bNeedSimReset, IsRootValid, InterpolationModeValue);
 			FVector RestRootOffsetValue = FVector::ZeroVector;
 			FVector DeformedRootOffsetValue = FVector::ZeroVector;
 			FVector DeformedPositionOffsetValue =  IsDeformedValid ? HairStrandsBuffer->SourceDeformedResources->PositionOffset :
 																	 HairStrandsBuffer->SourceRestResources->PositionOffset;
 			FVector RestPositionOffsetValue = ProxyData->HairStrandsBuffer->SourceRestResources->PositionOffset;
 
+			RHICmdList.TransitionResource(EResourceTransitionAccess::ERWBarrier, EResourceTransitionPipeline::EComputeToCompute, PointPositionsUAV);
 			SetUAVParameter(RHICmdList, ComputeShaderRHI, DeformedPositionBuffer, PointPositionsUAV);
 			RHICmdList.TransitionResource(EResourceTransitionAccess::ERWBarrier, EResourceTransitionPipeline::EComputeToCompute, HairStrandsBuffer->BoundingBoxBufferA.UAV);
 			SetUAVParameter(RHICmdList, ComputeShaderRHI, BoundingBoxBufferA, HairStrandsBuffer->BoundingBoxBufferA.UAV);
@@ -736,7 +783,7 @@ struct FNDIHairStrandsParametersCS : public FNiagaraDataInterfaceParametersCS
 			SetShaderValue(RHICmdList, ComputeShaderRHI, BoxExtent, ProxyData->BoxExtent);
 
 			SetShaderValue(RHICmdList, ComputeShaderRHI, ResetSimulation, bNeedSimReset);
-			SetShaderValue(RHICmdList, ComputeShaderRHI, HasRootAttached, HasRootAttachedValue);
+			SetShaderValue(RHICmdList, ComputeShaderRHI, InterpolationMode, InterpolationModeValue);
 			SetShaderValue(RHICmdList, ComputeShaderRHI, RestRootOffset, RestRootOffsetValue);
 			SetShaderValue(RHICmdList, ComputeShaderRHI, DeformedRootOffset, DeformedRootOffsetValue);
 
@@ -753,9 +800,13 @@ struct FNDIHairStrandsParametersCS : public FNiagaraDataInterfaceParametersCS
 
 			SetShaderValue(RHICmdList, ComputeShaderRHI, SampleCount, SampleCountValue);
 			SetSRVParameter(RHICmdList, ComputeShaderRHI, RestSamplePositionsBuffer, RestSamplePositionsBufferSRV);
+
+			RHICmdList.TransitionResource(EResourceTransitionAccess::EReadable, EResourceTransitionPipeline::EComputeToCompute, MeshSampleWeightsBufferUAV);
 			SetSRVParameter(RHICmdList, ComputeShaderRHI, MeshSampleWeightsBuffer, MeshSampleWeightsBufferSRV);
 
 			SetSRVParameter(RHICmdList, ComputeShaderRHI, RootBarycentricCoordinatesBuffer, RootBarycentricCoordinatesSRV);
+
+			SetSRVParameter(RHICmdList, ComputeShaderRHI, ParamsScaleBuffer, HairStrandsBuffer->ParamsScaleBuffer.SRV);
 		}
 		else
 		{
@@ -774,7 +825,7 @@ struct FNDIHairStrandsParametersCS : public FNiagaraDataInterfaceParametersCS
 			SetShaderValue(RHICmdList, ComputeShaderRHI, BoxExtent, FVector(0,0,0));
 
 			SetShaderValue(RHICmdList, ComputeShaderRHI, ResetSimulation, 0);
-			SetShaderValue(RHICmdList, ComputeShaderRHI, HasRootAttached, 0);
+			SetShaderValue(RHICmdList, ComputeShaderRHI, InterpolationMode, 0);
 			SetShaderValue(RHICmdList, ComputeShaderRHI, RestRootOffset, FVector4(0, 0, 0, 0));
 			SetShaderValue(RHICmdList, ComputeShaderRHI, DeformedRootOffset, FVector4(0, 0, 0, 0));
 
@@ -794,6 +845,7 @@ struct FNDIHairStrandsParametersCS : public FNiagaraDataInterfaceParametersCS
 			SetSRVParameter(RHICmdList, ComputeShaderRHI, MeshSampleWeightsBuffer, FNiagaraRenderer::GetDummyFloatBuffer());
 
 			SetSRVParameter(RHICmdList, ComputeShaderRHI, RootBarycentricCoordinatesBuffer, FNiagaraRenderer::GetDummyFloatBuffer());
+			SetSRVParameter(RHICmdList, ComputeShaderRHI, ParamsScaleBuffer, FNiagaraRenderer::GetDummyFloatBuffer());
 		}
 	}
 
@@ -819,7 +871,7 @@ private:
 	LAYOUT_FIELD(FShaderResourceParameter, RestPositionBuffer);
 
 	LAYOUT_FIELD(FShaderParameter, ResetSimulation);
-	LAYOUT_FIELD(FShaderParameter, HasRootAttached);
+	LAYOUT_FIELD(FShaderParameter, InterpolationMode);
 	LAYOUT_FIELD(FShaderParameter, RestRootOffset);
 	LAYOUT_FIELD(FShaderParameter, DeformedRootOffset);
 
@@ -840,6 +892,7 @@ private:
 	LAYOUT_FIELD(FShaderParameter, RestPositionOffset);
 	LAYOUT_FIELD(FShaderParameter, DeformedPositionOffset);
 
+	LAYOUT_FIELD(FShaderResourceParameter, ParamsScaleBuffer);
 	LAYOUT_FIELD(FShaderResourceParameter, BoundingBoxBufferA);
 	LAYOUT_FIELD(FShaderResourceParameter, BoundingBoxBufferB);
 };
@@ -920,7 +973,6 @@ void UNiagaraDataInterfaceHairStrands::ExtractSourceComponent(FNiagaraSystemInst
 			if (UGroomComponent* ParentComp = Cast<UGroomComponent>(SimComp->GetAttachParent()))
 			{
 				SourceComponent = ParentComp;
-				SourceComponent = ParentComp;
 			}
 			else if (UGroomComponent* OuterComp = SimComp->GetTypedOuter<UGroomComponent>())
 			{
@@ -943,7 +995,8 @@ void UNiagaraDataInterfaceHairStrands::ExtractSourceComponent(FNiagaraSystemInst
 }
 
 void UNiagaraDataInterfaceHairStrands::ExtractDatasAndResources(FNiagaraSystemInstance* SystemInstance,FHairStrandsDatas*& OutStrandsDatas, 
-	FHairStrandsRestResource*& OutStrandsRestResource, FHairStrandsDeformedResource*& OutStrandsDeformedResource, FHairStrandsRootResource*& OutStrandsRootResource, UGroomAsset*& OutGroomAsset)
+	FHairStrandsRestResource*& OutStrandsRestResource, FHairStrandsDeformedResource*& OutStrandsDeformedResource, 
+	FHairStrandsRootResource*& OutStrandsRootResource, UGroomAsset*& OutGroomAsset, int32& OutGroupIndex)
 {
 	ExtractSourceComponent(SystemInstance);
 
@@ -951,23 +1004,36 @@ void UNiagaraDataInterfaceHairStrands::ExtractDatasAndResources(FNiagaraSystemIn
 	OutStrandsRestResource = nullptr;
 	OutStrandsDeformedResource = nullptr;
 	OutStrandsRootResource = nullptr;
-
-	const int32 GroupIndex = (DefaultSource != nullptr && (DefaultSource->SimulatedGroup < 
-		DefaultSource->GetNumHairGroups())) ? DefaultSource->SimulatedGroup : 0;
+	OutGroupIndex = -1;
 
 	if (IsComponentValid())
 	{
-		OutStrandsDatas = SourceComponent->GetGuideStrandsDatas(GroupIndex);
-		OutStrandsRestResource = SourceComponent->GetGuideStrandsRestResource(GroupIndex);
-		OutStrandsDeformedResource = SourceComponent->GetGuideStrandsDeformedResource(GroupIndex);
-		OutStrandsRootResource = SourceComponent->GetGuideStrandsRootResource(GroupIndex);
-		OutGroomAsset = SourceComponent->GroomAsset;
+		for (int32 NiagaraIndex = 0, NiagaraCount = SourceComponent->NiagaraComponents.Num(); NiagaraIndex < NiagaraCount; ++NiagaraIndex)
+		{
+			if (SourceComponent->NiagaraComponents[NiagaraIndex] == SystemInstance->GetComponent())
+			{
+				OutGroupIndex = NiagaraIndex;
+				break;
+			}
+		}
+		if (OutGroupIndex >= 0 && OutGroupIndex < SourceComponent->NiagaraComponents.Num())
+		{
+			OutStrandsDatas = SourceComponent->GetGuideStrandsDatas(OutGroupIndex);
+			OutStrandsRestResource = SourceComponent->GetGuideStrandsRestResource(OutGroupIndex);
+			OutStrandsDeformedResource = SourceComponent->GetGuideStrandsDeformedResource(OutGroupIndex);
+			OutStrandsRootResource = SourceComponent->GetGuideStrandsRootResource(OutGroupIndex);
+			OutGroomAsset = SourceComponent->GroomAsset;
+		}
 	}
-	else if (DefaultSource != nullptr && (GroupIndex < DefaultSource->GetNumHairGroups()))
+	else if (DefaultSource != nullptr)
 	{
-		OutStrandsDatas = &DefaultSource->HairGroupsData[GroupIndex].HairRenderData;
-		OutStrandsRestResource = DefaultSource->HairGroupsData[GroupIndex].HairStrandsRestResource;
-		OutGroomAsset = DefaultSource;
+		OutGroupIndex = 0;
+		if (OutGroupIndex < DefaultSource->GetNumHairGroups())
+		{
+			OutStrandsDatas = &DefaultSource->HairGroupsData[OutGroupIndex].HairRenderData;
+			OutStrandsRestResource = DefaultSource->HairGroupsData[OutGroupIndex].HairStrandsRestResource;
+			OutGroomAsset = DefaultSource;
+		}
 	}
 }
 
@@ -1004,11 +1070,12 @@ bool UNiagaraDataInterfaceHairStrands::PerInstanceTick(void* PerInstanceData, FN
 	FHairStrandsDeformedResource* StrandsDeformedResource = nullptr;
 	FHairStrandsRootResource* StrandsRootResource = nullptr;
 	UGroomAsset* GroomAsset = nullptr;
+	int32 GroupIndex = 0;
 
 	InstanceData->TickCount = FMath::Min(MaxDelay+1,InstanceData->TickCount+1);
 
-	ExtractDatasAndResources(SystemInstance, StrandsDatas, StrandsRestResource, StrandsDeformedResource, StrandsRootResource, GroomAsset);
-	InstanceData->HairStrandsBuffer->Initialize(StrandsDatas, StrandsRestResource, StrandsDeformedResource, StrandsRootResource);
+	ExtractDatasAndResources(SystemInstance, StrandsDatas, StrandsRestResource, StrandsDeformedResource, StrandsRootResource, GroomAsset, GroupIndex);
+	InstanceData->HairStrandsBuffer->Initialize(StrandsDatas, StrandsRestResource, StrandsDeformedResource, StrandsRootResource, InstanceData);
 
 	if (SourceComponent != nullptr)
 	{
@@ -1018,24 +1085,8 @@ bool UNiagaraDataInterfaceHairStrands::PerInstanceTick(void* PerInstanceData, FN
 		}
 		InstanceData->ForceReset = SourceComponent->bResetSimulation;
 	}
-
-	bool RequireReset = false;
-	if (StrandsDatas)
-	{
-		if (StrandsDeformedResource != nullptr)
-		{
-			InstanceData->BoxCenter = StrandsDeformedResource->PositionOffset;
-		}
-		if (IsComponentValid())
-		{
-			InstanceData->WorldTransform = SourceComponent->GetComponentToWorld();
-		}
-		else
-		{
-			InstanceData->WorldTransform = SystemInstance->GetComponent()->GetComponentToWorld();
-		}
-	}
-	return RequireReset;
+	InstanceData->Update(this, SystemInstance, StrandsDatas, GroomAsset, GroupIndex);
+	return false;
 }
 
 bool UNiagaraDataInterfaceHairStrands::CopyToInternal(UNiagaraDataInterface* Destination) const
@@ -1193,7 +1244,7 @@ void UNiagaraDataInterfaceHairStrands::GetFunctions(TArray<FNiagaraFunctionSigna
 		Sig.bMemberFunction = true;
 		Sig.bRequiresContext = false;
 		Sig.Inputs.Add(FNiagaraVariable(FNiagaraTypeDefinition(GetClass()), TEXT("Hair Strands")));
-		Sig.Outputs.Add(FNiagaraVariable(FNiagaraTypeDefinition::GetVec4Def(), TEXT("Bend Scale")));
+		Sig.Outputs.Add(FNiagaraVariable(FNiagaraTypeDefinition::GetFloatDef(), TEXT("Bend Scale")));
 
 		OutFunctions.Add(Sig);
 	}
@@ -1243,7 +1294,7 @@ void UNiagaraDataInterfaceHairStrands::GetFunctions(TArray<FNiagaraFunctionSigna
 		Sig.bMemberFunction = true;
 		Sig.bRequiresContext = false;
 		Sig.Inputs.Add(FNiagaraVariable(FNiagaraTypeDefinition(GetClass()), TEXT("Hair Strands")));
-		Sig.Outputs.Add(FNiagaraVariable(FNiagaraTypeDefinition::GetVec4Def(), TEXT("Stretch Scale")));
+		Sig.Outputs.Add(FNiagaraVariable(FNiagaraTypeDefinition::GetFloatDef(), TEXT("Stretch Scale")));
 
 		OutFunctions.Add(Sig);
 	}
@@ -1313,7 +1364,7 @@ void UNiagaraDataInterfaceHairStrands::GetFunctions(TArray<FNiagaraFunctionSigna
 		Sig.bMemberFunction = true;
 		Sig.bRequiresContext = false;
 		Sig.Inputs.Add(FNiagaraVariable(FNiagaraTypeDefinition(GetClass()), TEXT("Hair Strands")));
-		Sig.Outputs.Add(FNiagaraVariable(FNiagaraTypeDefinition::GetVec4Def(), TEXT("Radius Scale")));
+		Sig.Outputs.Add(FNiagaraVariable(FNiagaraTypeDefinition::GetFloatDef(), TEXT("Radius Scale")));
 
 		OutFunctions.Add(Sig);
 	}
@@ -1353,7 +1404,7 @@ void UNiagaraDataInterfaceHairStrands::GetFunctions(TArray<FNiagaraFunctionSigna
 		Sig.bMemberFunction = true;
 		Sig.bRequiresContext = false;
 		Sig.Inputs.Add(FNiagaraVariable(FNiagaraTypeDefinition(GetClass()), TEXT("Hair Strands")));
-		Sig.Outputs.Add(FNiagaraVariable(FNiagaraTypeDefinition::GetVec4Def(), TEXT("Thickness Scale")));
+		Sig.Outputs.Add(FNiagaraVariable(FNiagaraTypeDefinition::GetFloatDef(), TEXT("Thickness Scale")));
 
 		OutFunctions.Add(Sig);
 	}
@@ -1394,7 +1445,6 @@ void UNiagaraDataInterfaceHairStrands::GetFunctions(TArray<FNiagaraFunctionSigna
 		Sig.bMemberFunction = true;
 		Sig.bRequiresContext = false;
 		Sig.Inputs.Add(FNiagaraVariable(FNiagaraTypeDefinition(GetClass()), TEXT("Hair Strands")));
-		Sig.Inputs.Add(FNiagaraVariable(FNiagaraTypeDefinition::GetIntDef(), TEXT("Node Index")));
 		Sig.Inputs.Add(FNiagaraVariable(FNiagaraTypeDefinition::GetFloatDef(), TEXT("Smoothing Filter")));
 		Sig.Outputs.Add(FNiagaraVariable(FNiagaraTypeDefinition::GetVec3Def(), TEXT("Node Position")));
 
@@ -1406,7 +1456,6 @@ void UNiagaraDataInterfaceHairStrands::GetFunctions(TArray<FNiagaraFunctionSigna
 		Sig.bMemberFunction = true;
 		Sig.bRequiresContext = false;
 		Sig.Inputs.Add(FNiagaraVariable(FNiagaraTypeDefinition(GetClass()), TEXT("Hair Strands")));
-		Sig.Inputs.Add(FNiagaraVariable(FNiagaraTypeDefinition::GetIntDef(), TEXT("Node Index")));
 		Sig.Inputs.Add(FNiagaraVariable(FNiagaraTypeDefinition::GetVec3Def(), TEXT("Node Position")));
 		Sig.Outputs.Add(FNiagaraVariable(FNiagaraTypeDefinition::GetQuatDef(), TEXT("Node Orientation")));
 
@@ -1418,10 +1467,8 @@ void UNiagaraDataInterfaceHairStrands::GetFunctions(TArray<FNiagaraFunctionSigna
 		Sig.bMemberFunction = true;
 		Sig.bRequiresContext = false;
 		Sig.Inputs.Add(FNiagaraVariable(FNiagaraTypeDefinition(GetClass()), TEXT("Hair Strands")));
-		Sig.Inputs.Add(FNiagaraVariable(FNiagaraTypeDefinition::GetIntDef(), TEXT("Node Index")));
 		Sig.Inputs.Add(FNiagaraVariable(FNiagaraTypeDefinition::GetFloatDef(), TEXT("Strands Density")));
-		Sig.Inputs.Add(FNiagaraVariable(FNiagaraTypeDefinition::GetFloatDef(), TEXT("Root Thickness")));
-		Sig.Inputs.Add(FNiagaraVariable(FNiagaraTypeDefinition::GetFloatDef(), TEXT("Tip Thickness")));
+		Sig.Inputs.Add(FNiagaraVariable(FNiagaraTypeDefinition::GetFloatDef(), TEXT("Node Thickness")));
 		Sig.Outputs.Add(FNiagaraVariable(FNiagaraTypeDefinition::GetFloatDef(), TEXT("Node Mass")));
 
 		OutFunctions.Add(Sig);
@@ -1432,10 +1479,8 @@ void UNiagaraDataInterfaceHairStrands::GetFunctions(TArray<FNiagaraFunctionSigna
 		Sig.bMemberFunction = true;
 		Sig.bRequiresContext = false;
 		Sig.Inputs.Add(FNiagaraVariable(FNiagaraTypeDefinition(GetClass()), TEXT("Hair Strands")));
-		Sig.Inputs.Add(FNiagaraVariable(FNiagaraTypeDefinition::GetIntDef(), TEXT("Node Index")));
 		Sig.Inputs.Add(FNiagaraVariable(FNiagaraTypeDefinition::GetFloatDef(), TEXT("Strands Density")));
-		Sig.Inputs.Add(FNiagaraVariable(FNiagaraTypeDefinition::GetFloatDef(), TEXT("Root Thickness")));
-		Sig.Inputs.Add(FNiagaraVariable(FNiagaraTypeDefinition::GetFloatDef(), TEXT("Tip Thickness")));
+		Sig.Inputs.Add(FNiagaraVariable(FNiagaraTypeDefinition::GetFloatDef(), TEXT("Node Thickness")));
 		Sig.Outputs.Add(FNiagaraVariable(FNiagaraTypeDefinition::GetVec3Def(), TEXT("Node Inertia")));
 
 		OutFunctions.Add(Sig);
@@ -1446,7 +1491,6 @@ void UNiagaraDataInterfaceHairStrands::GetFunctions(TArray<FNiagaraFunctionSigna
 		Sig.bMemberFunction = true;
 		Sig.bRequiresContext = false;
 		Sig.Inputs.Add(FNiagaraVariable(FNiagaraTypeDefinition(GetClass()), TEXT("Hair Strands")));
-		Sig.Inputs.Add(FNiagaraVariable(FNiagaraTypeDefinition::GetIntDef(), TEXT("Node Index")));
 		Sig.Inputs.Add(FNiagaraVariable(FNiagaraTypeDefinition::GetVec3Def(), TEXT("Node Position")));
 		Sig.Inputs.Add(FNiagaraVariable(FNiagaraTypeDefinition::GetIntDef(), TEXT("Node Offset")));
 		Sig.Outputs.Add(FNiagaraVariable(FNiagaraTypeDefinition::GetFloatDef(), TEXT("Edge Length")));
@@ -1459,7 +1503,6 @@ void UNiagaraDataInterfaceHairStrands::GetFunctions(TArray<FNiagaraFunctionSigna
 		Sig.bMemberFunction = true;
 		Sig.bRequiresContext = false;
 		Sig.Inputs.Add(FNiagaraVariable(FNiagaraTypeDefinition(GetClass()), TEXT("Hair Strands")));
-		Sig.Inputs.Add(FNiagaraVariable(FNiagaraTypeDefinition::GetIntDef(), TEXT("Node Index")));
 		Sig.Inputs.Add(FNiagaraVariable(FNiagaraTypeDefinition::GetQuatDef(), TEXT("Node Orientation")));
 		Sig.Outputs.Add(FNiagaraVariable(FNiagaraTypeDefinition::GetQuatDef(), TEXT("Edge Rotation")));
 
@@ -1493,7 +1536,6 @@ void UNiagaraDataInterfaceHairStrands::GetFunctions(TArray<FNiagaraFunctionSigna
 		Sig.bMemberFunction = true;
 		Sig.bRequiresContext = false;
 		Sig.Inputs.Add(FNiagaraVariable(FNiagaraTypeDefinition(GetClass()), TEXT("Hair Strands")));
-		Sig.Inputs.Add(FNiagaraVariable(FNiagaraTypeDefinition::GetIntDef(), TEXT("Node Index")));
 		Sig.Inputs.Add(FNiagaraVariable(FNiagaraTypeDefinition::GetVec3Def(), TEXT("Rest Position")));
 		Sig.Inputs.Add(FNiagaraVariable(FNiagaraTypeDefinition::GetQuatDef(), TEXT("Rest Orientation")));
 		Sig.Outputs.Add(FNiagaraVariable(FNiagaraTypeDefinition::GetVec3Def(), TEXT("Local Position")));
@@ -1591,7 +1633,6 @@ void UNiagaraDataInterfaceHairStrands::GetFunctions(TArray<FNiagaraFunctionSigna
 		Sig.bMemberFunction = true;
 		Sig.bRequiresContext = false;
 		Sig.Inputs.Add(FNiagaraVariable(FNiagaraTypeDefinition(GetClass()), TEXT("Hair Strands")));
-		Sig.Inputs.Add(FNiagaraVariable(FNiagaraTypeDefinition::GetIntDef(), TEXT("Node Index")));
 		Sig.Inputs.Add(FNiagaraVariable(FNiagaraTypeDefinition::GetVec3Def(), TEXT("Local Position")));
 		Sig.Inputs.Add(FNiagaraVariable(FNiagaraTypeDefinition::GetQuatDef(), TEXT("Local Orientation")));
 		Sig.Outputs.Add(FNiagaraVariable(FNiagaraTypeDefinition::GetVec3Def(), TEXT("Node Position")));
@@ -1601,12 +1642,15 @@ void UNiagaraDataInterfaceHairStrands::GetFunctions(TArray<FNiagaraFunctionSigna
 	}
 	{
 		FNiagaraFunctionSignature Sig;
-		Sig.Name = EvalSkinnedPositionName;
+		Sig.Name = UpdateNodeStateName;
 		Sig.bMemberFunction = true;
 		Sig.bRequiresContext = false;
 		Sig.Inputs.Add(FNiagaraVariable(FNiagaraTypeDefinition(GetClass()), TEXT("Hair Strands")));
 		Sig.Inputs.Add(FNiagaraVariable(FNiagaraTypeDefinition::GetVec3Def(), TEXT("Rest Position")));
-		Sig.Outputs.Add(FNiagaraVariable(FNiagaraTypeDefinition::GetVec3Def(), TEXT("Skinned Position")));
+		Sig.Inputs.Add(FNiagaraVariable(FNiagaraTypeDefinition::GetVec3Def(), TEXT("Node Position")));
+		Sig.Inputs.Add(FNiagaraVariable(FNiagaraTypeDefinition::GetQuatDef(), TEXT("Node Orientation")));
+		Sig.Outputs.Add(FNiagaraVariable(FNiagaraTypeDefinition::GetVec3Def(), TEXT("Node Position")));
+		Sig.Outputs.Add(FNiagaraVariable(FNiagaraTypeDefinition::GetQuatDef(), TEXT("Node Orientation")));
 
 		OutFunctions.Add(Sig);
 	}
@@ -1616,7 +1660,6 @@ void UNiagaraDataInterfaceHairStrands::GetFunctions(TArray<FNiagaraFunctionSigna
 		Sig.bMemberFunction = true;
 		Sig.bRequiresContext = false;
 		Sig.Inputs.Add(FNiagaraVariable(FNiagaraTypeDefinition(GetClass()), TEXT("Hair Strands")));
-		Sig.Inputs.Add(FNiagaraVariable(FNiagaraTypeDefinition::GetIntDef(), TEXT("Node Index")));
 		Sig.Inputs.Add(FNiagaraVariable(FNiagaraTypeDefinition::GetVec3Def(), TEXT("Node Displace")));
 		Sig.Outputs.Add(FNiagaraVariable(FNiagaraTypeDefinition::GetBoolDef(), TEXT("Report Status")));
 
@@ -1628,7 +1671,6 @@ void UNiagaraDataInterfaceHairStrands::GetFunctions(TArray<FNiagaraFunctionSigna
 		Sig.bMemberFunction = true;
 		Sig.bRequiresContext = false;
 		Sig.Inputs.Add(FNiagaraVariable(FNiagaraTypeDefinition(GetClass()), TEXT("Hair Strands")));
-		Sig.Inputs.Add(FNiagaraVariable(FNiagaraTypeDefinition::GetIntDef(), TEXT("Node Index")));
 		Sig.Outputs.Add(FNiagaraVariable(FNiagaraTypeDefinition::GetBoolDef(), TEXT("Report Status")));
 
 		OutFunctions.Add(Sig);
@@ -1957,7 +1999,7 @@ void UNiagaraDataInterfaceHairStrands::GetFunctions(TArray<FNiagaraFunctionSigna
 	}
 	{
 		FNiagaraFunctionSignature Sig;
-		Sig.Name = ComputeRestDirectionName;
+		Sig.Name = ComputeEdgeDirectionName;
 		Sig.bMemberFunction = true;
 		Sig.bRequiresContext = false;
 		Sig.Inputs.Add(FNiagaraVariable(FNiagaraTypeDefinition(GetClass()), TEXT("Hair Strands")));
@@ -2021,7 +2063,7 @@ void UNiagaraDataInterfaceHairStrands::GetFunctions(TArray<FNiagaraFunctionSigna
 		Sig.bMemberFunction = true;
 		Sig.bRequiresContext = false;
 		Sig.Inputs.Add(FNiagaraVariable(FNiagaraTypeDefinition(GetClass()), TEXT("Hair Strands")));
-		Sig.Outputs.Add(FNiagaraVariable(FNiagaraTypeDefinition::GetBoolDef(), TEXT("Globa lInterpolation")));
+		Sig.Outputs.Add(FNiagaraVariable(FNiagaraTypeDefinition::GetBoolDef(), TEXT("Global lInterpolation")));
 
 		OutFunctions.Add(Sig);
 	}
@@ -2088,6 +2130,7 @@ DEFINE_NDI_DIRECT_FUNC_BINDER(UNiagaraDataInterfaceHairStrands, ComputeLocalStat
 DEFINE_NDI_DIRECT_FUNC_BINDER(UNiagaraDataInterfaceHairStrands, AttachNodePosition);
 DEFINE_NDI_DIRECT_FUNC_BINDER(UNiagaraDataInterfaceHairStrands, AttachNodeOrientation);
 DEFINE_NDI_DIRECT_FUNC_BINDER(UNiagaraDataInterfaceHairStrands, AttachNodeState);
+DEFINE_NDI_DIRECT_FUNC_BINDER(UNiagaraDataInterfaceHairStrands, UpdateNodeState);
 DEFINE_NDI_DIRECT_FUNC_BINDER(UNiagaraDataInterfaceHairStrands, UpdatePointPosition);
 DEFINE_NDI_DIRECT_FUNC_BINDER(UNiagaraDataInterfaceHairStrands, ResetPointPosition);
 DEFINE_NDI_DIRECT_FUNC_BINDER(UNiagaraDataInterfaceHairStrands, EvalSkinnedPosition);
@@ -2124,7 +2167,7 @@ DEFINE_NDI_DIRECT_FUNC_BINDER(UNiagaraDataInterfaceHairStrands, SetupSoftCollisi
 DEFINE_NDI_DIRECT_FUNC_BINDER(UNiagaraDataInterfaceHairStrands, ProjectSoftCollisionConstraint);
 DEFINE_NDI_DIRECT_FUNC_BINDER(UNiagaraDataInterfaceHairStrands, SolveSoftCollisionConstraint);
 
-DEFINE_NDI_DIRECT_FUNC_BINDER(UNiagaraDataInterfaceHairStrands, ComputeRestDirection);
+DEFINE_NDI_DIRECT_FUNC_BINDER(UNiagaraDataInterfaceHairStrands, ComputeEdgeDirection);
 DEFINE_NDI_DIRECT_FUNC_BINDER(UNiagaraDataInterfaceHairStrands, UpdateMaterialFrame);
 DEFINE_NDI_DIRECT_FUNC_BINDER(UNiagaraDataInterfaceHairStrands, ComputeMaterialFrame);
 DEFINE_NDI_DIRECT_FUNC_BINDER(UNiagaraDataInterfaceHairStrands, ComputeAirDragForce);
@@ -2191,7 +2234,7 @@ void UNiagaraDataInterfaceHairStrands::GetVMExternalFunction(const FVMExternalFu
 	}
 	else if (BindingInfo.Name == GetBendScaleName)
 	{
-		check(BindingInfo.GetNumInputs() == 1 && BindingInfo.GetNumOutputs() == 4);
+		check(BindingInfo.GetNumInputs() == 1 && BindingInfo.GetNumOutputs() == 1);
 		NDI_FUNC_BINDER(UNiagaraDataInterfaceHairStrands, GetBendScale)::Bind(this, OutFunc);
 	}
 	else if (BindingInfo.Name == GetSolveStretchName)
@@ -2216,7 +2259,7 @@ void UNiagaraDataInterfaceHairStrands::GetVMExternalFunction(const FVMExternalFu
 	}
 	else if (BindingInfo.Name == GetStretchScaleName)
 	{
-		check(BindingInfo.GetNumInputs() == 1 && BindingInfo.GetNumOutputs() == 4);
+		check(BindingInfo.GetNumInputs() == 1 && BindingInfo.GetNumOutputs() == 1);
 		NDI_FUNC_BINDER(UNiagaraDataInterfaceHairStrands, GetStretchScale)::Bind(this, OutFunc);
 	}
 	else if (BindingInfo.Name == GetSolveCollisionName)
@@ -2251,7 +2294,7 @@ void UNiagaraDataInterfaceHairStrands::GetVMExternalFunction(const FVMExternalFu
 	}
 	else if (BindingInfo.Name == GetRadiusScaleName)
 	{
-		check(BindingInfo.GetNumInputs() == 1 && BindingInfo.GetNumOutputs() == 4);
+		check(BindingInfo.GetNumInputs() == 1 && BindingInfo.GetNumOutputs() == 1);
 		NDI_FUNC_BINDER(UNiagaraDataInterfaceHairStrands, GetRadiusScale)::Bind(this, OutFunc);
 	}
 	else if (BindingInfo.Name == GetStrandsDensityName)
@@ -2271,7 +2314,7 @@ void UNiagaraDataInterfaceHairStrands::GetVMExternalFunction(const FVMExternalFu
 	}
 	else if (BindingInfo.Name == GetThicknessScaleName)
 	{
-		check(BindingInfo.GetNumInputs() == 1 && BindingInfo.GetNumOutputs() == 4);
+		check(BindingInfo.GetNumInputs() == 1 && BindingInfo.GetNumOutputs() == 1);
 		NDI_FUNC_BINDER(UNiagaraDataInterfaceHairStrands, GetThicknessScale)::Bind(this, OutFunc);
 	}
 	else if (BindingInfo.Name == GetWorldTransformName)
@@ -2291,32 +2334,32 @@ void UNiagaraDataInterfaceHairStrands::GetVMExternalFunction(const FVMExternalFu
 	}
 	else if (BindingInfo.Name == ComputeNodePositionName)
 	{
-		check(BindingInfo.GetNumInputs() == 3 && BindingInfo.GetNumOutputs() == 3);
+		check(BindingInfo.GetNumInputs() == 2 && BindingInfo.GetNumOutputs() == 3);
 		NDI_FUNC_BINDER(UNiagaraDataInterfaceHairStrands, ComputeNodePosition)::Bind(this, OutFunc);
 	}
 	else if (BindingInfo.Name == ComputeNodeOrientationName)
 	{
-		check(BindingInfo.GetNumInputs() == 5 && BindingInfo.GetNumOutputs() == 4);
+		check(BindingInfo.GetNumInputs() == 4 && BindingInfo.GetNumOutputs() == 4);
 		NDI_FUNC_BINDER(UNiagaraDataInterfaceHairStrands, ComputeNodeOrientation)::Bind(this, OutFunc);
 	}
 	else if (BindingInfo.Name == ComputeNodeMassName)
 	{
-		check(BindingInfo.GetNumInputs() == 5 && BindingInfo.GetNumOutputs() == 1);
+		check(BindingInfo.GetNumInputs() == 3 && BindingInfo.GetNumOutputs() == 1);
 		NDI_FUNC_BINDER(UNiagaraDataInterfaceHairStrands, ComputeNodeMass)::Bind(this, OutFunc);
 	}
 	else if (BindingInfo.Name == ComputeNodeInertiaName)
 	{
-		check(BindingInfo.GetNumInputs() == 5 && BindingInfo.GetNumOutputs() == 3);
+		check(BindingInfo.GetNumInputs() == 3 && BindingInfo.GetNumOutputs() == 3);
 		NDI_FUNC_BINDER(UNiagaraDataInterfaceHairStrands, ComputeNodeInertia)::Bind(this, OutFunc);
 	}
 	else if (BindingInfo.Name == ComputeEdgeLengthName)
 	{
-		check(BindingInfo.GetNumInputs() == 6 && BindingInfo.GetNumOutputs() == 1);
+		check(BindingInfo.GetNumInputs() == 5 && BindingInfo.GetNumOutputs() == 1);
 		NDI_FUNC_BINDER(UNiagaraDataInterfaceHairStrands, ComputeEdgeLength)::Bind(this, OutFunc);
 	}
 	else if (BindingInfo.Name == ComputeEdgeRotationName)
 	{
-		check(BindingInfo.GetNumInputs() == 6 && BindingInfo.GetNumOutputs() == 4);
+		check(BindingInfo.GetNumInputs() == 5 && BindingInfo.GetNumOutputs() == 4);
 		NDI_FUNC_BINDER(UNiagaraDataInterfaceHairStrands, ComputeEdgeRotation)::Bind(this, OutFunc);
 	}
 	else if (BindingInfo.Name == ComputeRestPositionName)
@@ -2331,7 +2374,7 @@ void UNiagaraDataInterfaceHairStrands::GetVMExternalFunction(const FVMExternalFu
 	}
 	else if (BindingInfo.Name == ComputeLocalStateName)
 	{
-		check(BindingInfo.GetNumInputs() == 9 && BindingInfo.GetNumOutputs() == 7);
+		check(BindingInfo.GetNumInputs() == 8 && BindingInfo.GetNumOutputs() == 7);
 		NDI_FUNC_BINDER(UNiagaraDataInterfaceHairStrands, ComputeLocalState)::Bind(this, OutFunc);
 	}
 	else if (BindingInfo.Name == AttachNodePositionName)
@@ -2346,22 +2389,22 @@ void UNiagaraDataInterfaceHairStrands::GetVMExternalFunction(const FVMExternalFu
 	}
 	else if (BindingInfo.Name == AttachNodeStateName)
 	{
-		check(BindingInfo.GetNumInputs() == 9 && BindingInfo.GetNumOutputs() == 7);
+		check(BindingInfo.GetNumInputs() == 8 && BindingInfo.GetNumOutputs() == 7);
 		NDI_FUNC_BINDER(UNiagaraDataInterfaceHairStrands, AttachNodeState)::Bind(this, OutFunc);
 	}
-	else if (BindingInfo.Name == EvalSkinnedPositionName)
+	else if (BindingInfo.Name == UpdateNodeStateName)
 	{
-		check(BindingInfo.GetNumInputs() == 4 && BindingInfo.GetNumOutputs() == 3);
-		NDI_FUNC_BINDER(UNiagaraDataInterfaceHairStrands, EvalSkinnedPosition)::Bind(this, OutFunc);
+		check(BindingInfo.GetNumInputs() == 11 && BindingInfo.GetNumOutputs() == 7);
+		NDI_FUNC_BINDER(UNiagaraDataInterfaceHairStrands, UpdateNodeState)::Bind(this, OutFunc);
 	}
 	else if (BindingInfo.Name == UpdatePointPositionName)
 	{
-		check(BindingInfo.GetNumInputs() == 5 && BindingInfo.GetNumOutputs() == 1);
+		check(BindingInfo.GetNumInputs() == 4 && BindingInfo.GetNumOutputs() == 1);
 		NDI_FUNC_BINDER(UNiagaraDataInterfaceHairStrands, UpdatePointPosition)::Bind(this, OutFunc);
 	}
 	else if (BindingInfo.Name == ResetPointPositionName)
 	{
-		check(BindingInfo.GetNumInputs() == 2 && BindingInfo.GetNumOutputs() == 1);
+		check(BindingInfo.GetNumInputs() == 1 && BindingInfo.GetNumOutputs() == 1);
 		NDI_FUNC_BINDER(UNiagaraDataInterfaceHairStrands, ResetPointPosition)::Bind(this, OutFunc);
 	}
 	else if (BindingInfo.Name == AdvectNodePositionName)
@@ -2484,10 +2527,10 @@ void UNiagaraDataInterfaceHairStrands::GetVMExternalFunction(const FVMExternalFu
 		check(BindingInfo.GetNumInputs() == 4 && BindingInfo.GetNumOutputs() == 5);
 		NDI_FUNC_BINDER(UNiagaraDataInterfaceHairStrands, SetupSoftCollisionConstraint)::Bind(this, OutFunc);
 	}
-	else if (BindingInfo.Name == ComputeRestDirectionName)
+	else if (BindingInfo.Name == ComputeEdgeDirectionName)
 	{
 		check(BindingInfo.GetNumInputs() == 8 && BindingInfo.GetNumOutputs() == 3);
-		NDI_FUNC_BINDER(UNiagaraDataInterfaceHairStrands, ComputeRestDirection)::Bind(this, OutFunc);
+		NDI_FUNC_BINDER(UNiagaraDataInterfaceHairStrands, ComputeEdgeDirection)::Bind(this, OutFunc);
 	}
 	else if (BindingInfo.Name == UpdateMaterialFrameName)
 	{
@@ -2691,21 +2734,7 @@ void UNiagaraDataInterfaceHairStrands::GetBendStiffness(FVectorVMContext& Contex
 }
 
 void UNiagaraDataInterfaceHairStrands::GetBendScale(FVectorVMContext& Context)
-{
-	VectorVM::FUserPtrHandler<FNDIHairStrandsData> InstData(Context);
-	VectorVM::FExternalFuncRegisterHandler<float> OutBendScaleX(Context);
-	VectorVM::FExternalFuncRegisterHandler<float> OutBendScaleY(Context);
-	VectorVM::FExternalFuncRegisterHandler<float> OutBendScaleZ(Context);
-	VectorVM::FExternalFuncRegisterHandler<float> OutBendScaleW(Context);
-
-	for (int32 InstanceIdx = 0; InstanceIdx < Context.NumInstances; ++InstanceIdx)
-	{
-		*OutBendScaleX.GetDestAndAdvance() = InstData->BendScale[0];
-		*OutBendScaleY.GetDestAndAdvance() = InstData->BendScale[1];
-		*OutBendScaleZ.GetDestAndAdvance() = InstData->BendScale[2];
-		*OutBendScaleW.GetDestAndAdvance() = InstData->BendScale[3];
-	}
-}
+{}
 
 void UNiagaraDataInterfaceHairStrands::GetSolveStretch(FVectorVMContext& Context)
 {
@@ -2753,19 +2782,6 @@ void UNiagaraDataInterfaceHairStrands::GetStretchStiffness(FVectorVMContext& Con
 
 void UNiagaraDataInterfaceHairStrands::GetStretchScale(FVectorVMContext& Context)
 {
-	VectorVM::FUserPtrHandler<FNDIHairStrandsData> InstData(Context);
-	VectorVM::FExternalFuncRegisterHandler<float> OutStretchScaleX(Context);
-	VectorVM::FExternalFuncRegisterHandler<float> OutStretchScaleY(Context);
-	VectorVM::FExternalFuncRegisterHandler<float> OutStretchScaleZ(Context);
-	VectorVM::FExternalFuncRegisterHandler<float> OutStretchScaleW(Context);
-
-	for (int32 InstanceIdx = 0; InstanceIdx < Context.NumInstances; ++InstanceIdx)
-	{
-		*OutStretchScaleX.GetDestAndAdvance() = InstData->StretchScale[0];
-		*OutStretchScaleY.GetDestAndAdvance() = InstData->StretchScale[1];
-		*OutStretchScaleZ.GetDestAndAdvance() = InstData->StretchScale[2];
-		*OutStretchScaleW.GetDestAndAdvance() = InstData->StretchScale[3];
-	}
 }
 
 
@@ -2837,19 +2853,6 @@ void UNiagaraDataInterfaceHairStrands::GetCollisionRadius(FVectorVMContext& Cont
 
 void UNiagaraDataInterfaceHairStrands::GetRadiusScale(FVectorVMContext& Context)
 {
-	VectorVM::FUserPtrHandler<FNDIHairStrandsData> InstData(Context);
-	VectorVM::FExternalFuncRegisterHandler<float> OutRadiusScaleX(Context);
-	VectorVM::FExternalFuncRegisterHandler<float> OutRadiusScaleY(Context);
-	VectorVM::FExternalFuncRegisterHandler<float> OutRadiusScaleZ(Context);
-	VectorVM::FExternalFuncRegisterHandler<float> OutRadiusScaleW(Context);
-
-	for (int32 InstanceIdx = 0; InstanceIdx < Context.NumInstances; ++InstanceIdx)
-	{
-		*OutRadiusScaleX.GetDestAndAdvance() = InstData->RadiusScale[0];
-		*OutRadiusScaleY.GetDestAndAdvance() = InstData->RadiusScale[1];
-		*OutRadiusScaleZ.GetDestAndAdvance() = InstData->RadiusScale[2];
-		*OutRadiusScaleW.GetDestAndAdvance() = InstData->RadiusScale[3];
-	}
 }
 
 void UNiagaraDataInterfaceHairStrands::GetStrandsSmoothing(FVectorVMContext& Context)
@@ -2887,19 +2890,6 @@ void UNiagaraDataInterfaceHairStrands::GetStrandsThickness(FVectorVMContext& Con
 
 void UNiagaraDataInterfaceHairStrands::GetThicknessScale(FVectorVMContext& Context)
 {
-	VectorVM::FUserPtrHandler<FNDIHairStrandsData> InstData(Context);
-	VectorVM::FExternalFuncRegisterHandler<float> OutThicknessScaleX(Context);
-	VectorVM::FExternalFuncRegisterHandler<float> OutThicknessScaleY(Context);
-	VectorVM::FExternalFuncRegisterHandler<float> OutThicknessScaleZ(Context);
-	VectorVM::FExternalFuncRegisterHandler<float> OutThicknessScaleW(Context);
-
-	for (int32 InstanceIdx = 0; InstanceIdx < Context.NumInstances; ++InstanceIdx)
-	{
-		*OutThicknessScaleX.GetDestAndAdvance() = InstData->ThicknessScale[0];
-		*OutThicknessScaleY.GetDestAndAdvance() = InstData->ThicknessScale[1];
-		*OutThicknessScaleZ.GetDestAndAdvance() = InstData->ThicknessScale[2];
-		*OutThicknessScaleW.GetDestAndAdvance() = InstData->ThicknessScale[3];
-	}
 }
 
 void UNiagaraDataInterfaceHairStrands::GetWorldTransform(FVectorVMContext& Context)
@@ -3032,6 +3022,11 @@ void UNiagaraDataInterfaceHairStrands::AttachNodeState(FVectorVMContext& Context
 	// @todo : implement function for cpu 
 }
 
+void UNiagaraDataInterfaceHairStrands::UpdateNodeState(FVectorVMContext& Context)
+{
+	// @todo : implement function for cpu 
+}
+
 void UNiagaraDataInterfaceHairStrands::AdvectNodePosition(FVectorVMContext& Context)
 {
 	// @todo : implement function for cpu 
@@ -3113,7 +3108,7 @@ void UNiagaraDataInterfaceHairStrands::ProjectBendRodMaterial(FVectorVMContext&
 	// @todo : implement function for cpu 
 }
 
-void UNiagaraDataInterfaceHairStrands::ComputeRestDirection(FVectorVMContext& Context)
+void UNiagaraDataInterfaceHairStrands::ComputeEdgeDirection(FVectorVMContext& Context)
 {
 	// @todo : implement function for cpu 
 }
@@ -3164,29 +3159,23 @@ void UNiagaraDataInterfaceHairStrands::NeedSimulationReset(FVectorVMContext& Con
 
 void UNiagaraDataInterfaceHairStrands::HasGlobalInterpolation(FVectorVMContext& Context)
 {
-	static const auto GlobalInterpolationCVar = IConsoleManager::Get().FindConsoleVariable(TEXT("r.HairStrands.Interpolation.Global"));
-	const int32 GlobalInterpolation = GlobalInterpolationCVar->GetInt();
-
 	VectorVM::FUserPtrHandler<FNDIHairStrandsData> InstData(Context);
 	VectorVM::FExternalFuncRegisterHandler<int32> OutGlobalInterpolation(Context);
 
 	for (int32 InstanceIdx = 0; InstanceIdx < Context.NumInstances; ++InstanceIdx)
 	{
-		*OutGlobalInterpolation.GetDestAndAdvance() = GlobalInterpolation;
+		*OutGlobalInterpolation.GetDestAndAdvance() = false;
 	}
 }
 
 void UNiagaraDataInterfaceHairStrands::HasKinematicsTarget(FVectorVMContext& Context)
 {
-	static const auto KinematicsTargetCVar = IConsoleManager::Get().FindConsoleVariable(TEXT("r.HairStrands.Interpolation.Target"));
-	const int32 KinematicsTarget = KinematicsTargetCVar->GetInt();
-
 	VectorVM::FUserPtrHandler<FNDIHairStrandsData> InstData(Context);
 	VectorVM::FExternalFuncRegisterHandler<int32> OutKinematicsTarget(Context);
 
 	for (int32 InstanceIdx = 0; InstanceIdx < Context.NumInstances; ++InstanceIdx)
 	{
-		*OutKinematicsTarget.GetDestAndAdvance() = KinematicsTarget;
+		*OutKinematicsTarget.GetDestAndAdvance() = InstData->KinematicsTarget;
 	}
 }
 
@@ -3253,6 +3242,50 @@ bool UNiagaraDataInterfaceHairStrands::GetFunctionHLSL(const FNiagaraDataInterfa
 		OutHLSL += FString::Format(FormatSample, ArgsSample);
 		return true;
 	}
+	if (FunctionInfo.DefinitionName == GetStretchScaleName)
+	{
+		static const TCHAR* FormatSample = TEXT(R"(
+		void {InstanceFunctionName}(out float OutStretchScale)
+		{
+			{HairStrandsContextName} OutStretchScale = DIContext.ParamsScaleBuffer[GGroupThreadId.x % DIContext.StrandSize];
+		}
+		)");
+		OutHLSL += FString::Format(FormatSample, ArgsSample);
+		return true;
+	}
+	if (FunctionInfo.DefinitionName == GetBendScaleName)
+	{
+		static const TCHAR* FormatSample = TEXT(R"(
+		void {InstanceFunctionName}(out float OutBendScale)
+		{
+			{HairStrandsContextName} OutBendScale = DIContext.ParamsScaleBuffer[32 + GGroupThreadId.x % DIContext.StrandSize];
+		}
+		)");
+		OutHLSL += FString::Format(FormatSample, ArgsSample);
+		return true;
+	}
+	if (FunctionInfo.DefinitionName == GetRadiusScaleName)
+	{
+		static const TCHAR* FormatSample = TEXT(R"(
+		void {InstanceFunctionName}(out float OutRadiusScale)
+		{
+			{HairStrandsContextName} OutRadiusScale = DIContext.ParamsScaleBuffer[64 + GGroupThreadId.x % DIContext.StrandSize];
+		}
+		)");
+		OutHLSL += FString::Format(FormatSample, ArgsSample);
+		return true;
+	}
+	if (FunctionInfo.DefinitionName == GetThicknessScaleName)
+	{
+		static const TCHAR* FormatSample = TEXT(R"(
+		void {InstanceFunctionName}(out float OutThicknessScale)
+		{
+			{HairStrandsContextName} OutThicknessScale = DIContext.ParamsScaleBuffer[96 + GGroupThreadId.x % DIContext.StrandSize];
+		}
+		)");
+		OutHLSL += FString::Format(FormatSample, ArgsSample);
+		return true;
+	}
 	else if (FunctionInfo.DefinitionName == GetPointPositionName)
 	{
 		static const TCHAR *FormatSample = TEXT(R"(
@@ -3267,10 +3300,10 @@ bool UNiagaraDataInterfaceHairStrands::GetFunctionHLSL(const FNiagaraDataInterfa
 	else if (FunctionInfo.DefinitionName == ComputeNodePositionName)
 	{
 		static const TCHAR *FormatSample = TEXT(R"(
-			void {InstanceFunctionName} (in int NodeIndex, in float SmoothingFilter, out float3 OutNodePosition)
+			void {InstanceFunctionName} (in float SmoothingFilter, out float3 OutNodePosition)
 			{
-				{HairStrandsContextName} DIHairStrands_ComputeNodePosition(DIContext,NodeIndex,OutNodePosition);
-				DIHairStrands_SmoothNodePosition(DIContext,NodeIndex,SmoothingFilter,OutNodePosition);
+				{HairStrandsContextName} DIHairStrands_ComputeNodePosition(DIContext,OutNodePosition);
+				DIHairStrands_SmoothNodePosition(DIContext,SmoothingFilter,OutNodePosition);
 			}
 			)");
 		OutHLSL += FString::Format(FormatSample, ArgsSample);
@@ -3279,9 +3312,9 @@ bool UNiagaraDataInterfaceHairStrands::GetFunctionHLSL(const FNiagaraDataInterfa
 	else if (FunctionInfo.DefinitionName == ComputeNodeOrientationName)
 	{
 		static const TCHAR *FormatSample = TEXT(R"(
-			void {InstanceFunctionName} (in int NodeIndex, in float3 NodePosition, out float4 OutNodeOrientation)
+			void {InstanceFunctionName} (in float3 NodePosition, out float4 OutNodeOrientation)
 			{
-				{HairStrandsContextName} DIHairStrands_ComputeNodeOrientation(DIContext,NodeIndex,NodePosition,OutNodeOrientation);
+				{HairStrandsContextName} DIHairStrands_ComputeNodeOrientation(DIContext,NodePosition,OutNodeOrientation);
 			}
 			)");
 		OutHLSL += FString::Format(FormatSample, ArgsSample);
@@ -3290,9 +3323,9 @@ bool UNiagaraDataInterfaceHairStrands::GetFunctionHLSL(const FNiagaraDataInterfa
 	else if (FunctionInfo.DefinitionName == ComputeNodeMassName)
 	{
 		static const TCHAR *FormatSample = TEXT(R"(
-			void {InstanceFunctionName} (in int NodeIndex, in float StrandsDensity, in float RootThickness, in float TipThickness, out float OutNodeMass)
+			void {InstanceFunctionName} (in float StrandsDensity, in float NodeThickness, out float OutNodeMass)
 			{
-				{HairStrandsContextName} DIHairStrands_ComputeNodeMass(DIContext,NodeIndex,StrandsDensity,RootThickness,TipThickness,OutNodeMass);
+				{HairStrandsContextName} DIHairStrands_ComputeNodeMass(DIContext,StrandsDensity,NodeThickness,OutNodeMass);
 			}
 			)");
 		OutHLSL += FString::Format(FormatSample, ArgsSample);
@@ -3301,9 +3334,9 @@ bool UNiagaraDataInterfaceHairStrands::GetFunctionHLSL(const FNiagaraDataInterfa
 	else if (FunctionInfo.DefinitionName == ComputeNodeInertiaName)
 	{
 		static const TCHAR *FormatSample = TEXT(R"(
-			void {InstanceFunctionName} (in int NodeIndex, in float StrandsDensity, in float RootThickness, in float TipThickness, out float3 OutNodeInertia)
+			void {InstanceFunctionName} (in float StrandsDensity, in float NodeThickness, out float3 OutNodeInertia)
 			{
-				{HairStrandsContextName} DIHairStrands_ComputeNodeInertia(DIContext,NodeIndex,StrandsDensity,RootThickness,TipThickness,OutNodeInertia);
+				{HairStrandsContextName} DIHairStrands_ComputeNodeInertia(DIContext,StrandsDensity,NodeThickness,OutNodeInertia);
 			}
 			)");
 		OutHLSL += FString::Format(FormatSample, ArgsSample);
@@ -3312,16 +3345,16 @@ bool UNiagaraDataInterfaceHairStrands::GetFunctionHLSL(const FNiagaraDataInterfa
 	else if (FunctionInfo.DefinitionName == ComputeEdgeLengthName)
 	{
 		static const TCHAR *FormatSample = TEXT(R"(
-			void {InstanceFunctionName} (in int NodeIndex, in float3 NodePosition, in int NodeOffset, out float OutEdgeLength)
+			void {InstanceFunctionName} (in float3 NodePosition, in int NodeOffset, out float OutEdgeLength)
 			{
 				{HairStrandsContextName} 
 				if(NodeOffset == 2)
 				{
-					DIHairStrands_ComputeEdgeVolume(DIContext,NodeIndex,NodePosition,OutEdgeLength);
+					DIHairStrands_ComputeEdgeVolume(DIContext,NodePosition,OutEdgeLength);
 				}
 				else
 				{
-					DIHairStrands_ComputeEdgeLength(DIContext,NodeIndex,NodePosition,NodeOffset,OutEdgeLength);
+					DIHairStrands_ComputeEdgeLength(DIContext,NodePosition,NodeOffset,OutEdgeLength);
 				}
 			}
 			)");
@@ -3331,9 +3364,9 @@ bool UNiagaraDataInterfaceHairStrands::GetFunctionHLSL(const FNiagaraDataInterfa
 	else if (FunctionInfo.DefinitionName == ComputeEdgeRotationName)
 	{
 		static const TCHAR *FormatSample = TEXT(R"(
-			void {InstanceFunctionName} (in int NodeIndex, in float4 NodeOrientation, out float4 OutEdgeRotation)
+			void {InstanceFunctionName} (in float4 NodeOrientation, out float4 OutEdgeRotation)
 			{
-				{HairStrandsContextName} DIHairStrands_ComputeEdgeRotation(DIContext,NodeIndex,NodeOrientation,OutEdgeRotation);
+				{HairStrandsContextName} DIHairStrands_ComputeEdgeRotation(DIContext,NodeOrientation,OutEdgeRotation);
 			}
 			)");
 		OutHLSL += FString::Format(FormatSample, ArgsSample);
@@ -3364,9 +3397,9 @@ bool UNiagaraDataInterfaceHairStrands::GetFunctionHLSL(const FNiagaraDataInterfa
 	else if (FunctionInfo.DefinitionName == ComputeLocalStateName)
 	{
 		static const TCHAR *FormatSample = TEXT(R"(
-			void {InstanceFunctionName} (in int NodeIndex, in float3 RestPosition, in float4 RestOrientation, out float3 LocalPosition, out float4 LocalOrientation)
+			void {InstanceFunctionName} (in float3 RestPosition, in float4 RestOrientation, out float3 LocalPosition, out float4 LocalOrientation)
 			{
-				{HairStrandsContextName} DIHairStrands_ComputeLocalState(DIContext,NodeIndex,RestPosition,RestOrientation,LocalPosition,LocalOrientation);
+				{HairStrandsContextName} DIHairStrands_ComputeLocalState(DIContext,RestPosition,RestOrientation,LocalPosition,LocalOrientation);
 			}
 			)");
 		OutHLSL += FString::Format(FormatSample, ArgsSample);
@@ -3383,17 +3416,6 @@ bool UNiagaraDataInterfaceHairStrands::GetFunctionHLSL(const FNiagaraDataInterfa
 		OutHLSL += FString::Format(FormatSample, ArgsSample);
 		return true;
 	}
-	else if (FunctionInfo.DefinitionName == EvalSkinnedPositionName)
-	{
-		static const TCHAR* FormatSample = TEXT(R"(
-			void {InstanceFunctionName} (in float3 RestPosition, out float3 SkinnedPosition)
-			{
-				{HairStrandsContextName} DIHairStrands_EvalSkinnedPosition(DIContext,RestPosition,SkinnedPosition);
-			}
-			)");
-		OutHLSL += FString::Format(FormatSample, ArgsSample);
-		return true;
-	}
 	else if (FunctionInfo.DefinitionName == AttachNodeOrientationName)
 	{
 		static const TCHAR *FormatSample = TEXT(R"(
@@ -3408,20 +3430,31 @@ bool UNiagaraDataInterfaceHairStrands::GetFunctionHLSL(const FNiagaraDataInterfa
 	else if (FunctionInfo.DefinitionName == AttachNodeStateName)
 	{
 		static const TCHAR *FormatSample = TEXT(R"(
-				void {InstanceFunctionName} (in int NodeIndex, in float3 LocalPosition, in float4 LocalOrientation, out float3 NodePosition, out float4 NodeOrientation)
+				void {InstanceFunctionName} ( in float3 LocalPosition, in float4 LocalOrientation, out float3 NodePosition, out float4 NodeOrientation)
 				{
-					{HairStrandsContextName} DIHairStrands_AttachNodeState(DIContext,NodeIndex,LocalPosition,LocalOrientation,NodePosition,NodeOrientation);
+					{HairStrandsContextName} DIHairStrands_AttachNodeState(DIContext,LocalPosition,LocalOrientation,NodePosition,NodeOrientation);
 				}
 				)");
 		OutHLSL += FString::Format(FormatSample, ArgsSample);
 		return true;
 	}
+	else if (FunctionInfo.DefinitionName == UpdateNodeStateName)
+	{
+		static const TCHAR* FormatSample = TEXT(R"(
+					void {InstanceFunctionName} ( in float3 RestPosition, in float3 NodePosition, in float4 NodeOrientation, out float3 OutNodePosition, out float4 OutNodeOrientation)
+					{
+						{HairStrandsContextName} DIHairStrands_UpdateNodeState(DIContext,RestPosition,NodePosition,NodeOrientation,OutNodePosition,OutNodeOrientation);
+					}
+					)");
+		OutHLSL += FString::Format(FormatSample, ArgsSample);
+		return true;
+	}
 	else if (FunctionInfo.DefinitionName == UpdatePointPositionName)
 	{
 		static const TCHAR *FormatSample = TEXT(R"(
-			void {InstanceFunctionName} (in int NodeIndex, in float3 NodeDisplace, out bool OutReportStatus)
+			void {InstanceFunctionName} (in float3 NodeDisplace, out bool OutReportStatus)
 			{
-				{HairStrandsContextName} DIHairStrands_UpdatePointPosition(DIContext,NodeIndex,NodeDisplace,OutReportStatus);
+				{HairStrandsContextName} DIHairStrands_UpdatePointPosition(DIContext,NodeDisplace,OutReportStatus);
 			}
 			)");
 		OutHLSL += FString::Format(FormatSample, ArgsSample);
@@ -3430,9 +3463,9 @@ bool UNiagaraDataInterfaceHairStrands::GetFunctionHLSL(const FNiagaraDataInterfa
 	else if (FunctionInfo.DefinitionName == ResetPointPositionName)
 	{
 		static const TCHAR *FormatSample = TEXT(R"(
-			void {InstanceFunctionName} (in int NodeIndex, out bool OutReportStatus)
+			void {InstanceFunctionName} (out bool OutReportStatus)
 			{
-				{HairStrandsContextName} DIHairStrands_ResetPointPosition(DIContext,NodeIndex,OutReportStatus);
+				{HairStrandsContextName} DIHairStrands_ResetPointPosition(DIContext,OutReportStatus);
 			}
 			)");
 		OutHLSL += FString::Format(FormatSample, ArgsSample);
@@ -3763,12 +3796,12 @@ in float RestLength, in float DeltaTime, in int NodeOffset, in float MaterialDam
 		OutHLSL += FString::Format(FormatSample, ArgsSample);
 		return true;
 	}
-	else if (FunctionInfo.DefinitionName == ComputeRestDirectionName)
+	else if (FunctionInfo.DefinitionName == ComputeEdgeDirectionName)
 	{
 		static const TCHAR *FormatSample = TEXT(R"(
 							void {InstanceFunctionName} (in float3 NodePosition, in float4 NodeOrientation, out float3 OutRestDirection)
 							{
-								{HairStrandsContextName} ComputeRestDirection(DIContext.StrandSize,NodePosition,NodeOrientation,OutRestDirection);
+								{HairStrandsContextName} DIHairStrands_ComputeEdgeDirection(DIContext,NodePosition,NodeOrientation,OutRestDirection);
 							}
 							)");
 		OutHLSL += FString::Format(FormatSample, ArgsSample);
@@ -3821,6 +3854,17 @@ in float RestLength, in float DeltaTime, in int NodeOffset, in float MaterialDam
 		OutHLSL += FString::Format(FormatSample, ArgsSample);
 		return true;
 	}
+	else if (FunctionInfo.DefinitionName == HasGlobalInterpolationName)
+	{
+		static const TCHAR* FormatSample = TEXT(R"(
+				void {InstanceFunctionName} ( out bool GlobalInterpolation)
+				{
+					{HairStrandsContextName} GlobalInterpolation  = (DIContext.InterpolationMode == 2);
+				}
+				)");
+		OutHLSL += FString::Format(FormatSample, ArgsSample);
+		return true;
+	}
 
 	OutHLSL += TEXT("\n");
 	return false;
diff --git a/Engine/Plugins/Experimental/HairStrands/Source/HairStrandsNiagara/Private/NiagaraDataInterfacePhysicsAsset.cpp b/Engine/Plugins/Experimental/HairStrands/Source/HairStrandsNiagara/Private/NiagaraDataInterfacePhysicsAsset.cpp
--- a/Engine/Plugins/Experimental/HairStrands/Source/HairStrandsNiagara/Private/NiagaraDataInterfacePhysicsAsset.cpp
+++ b/Engine/Plugins/Experimental/HairStrands/Source/HairStrandsNiagara/Private/NiagaraDataInterfacePhysicsAsset.cpp
@@ -24,7 +24,11 @@ static const FName GetNumCapsulesName(TEXT("GetNumCapsules"));
 
 //------------------------------------------------------------------------------------------------------------
 
+static const FName GetClosestElementName(TEXT("GetClosestElement"));
+static const FName GetElementPointName(TEXT("GetElementPoint"));
+static const FName GetElementDistanceName(TEXT("GetElementDistance"));
 static const FName GetClosestPointName(TEXT("GetClosestPoint"));
+static const FName GetClosestDistanceName(TEXT("GetClosestDistance"));
 static const FName GetTexturePointName(TEXT("GetTexturePoint"));
 static const FName GetProjectionPointName(TEXT("GetProjectionPoint"));
 
@@ -751,12 +755,64 @@ void UNiagaraDataInterfacePhysicsAsset::GetFunctions(TArray<FNiagaraFunctionSign
 		Sig.bMemberFunction = true;
 		Sig.bRequiresContext = false;
 		Sig.Inputs.Add(FNiagaraVariable(FNiagaraTypeDefinition(GetClass()), TEXT("Physics Asset")));
-		Sig.Inputs.Add(FNiagaraVariable(FNiagaraTypeDefinition::GetVec3Def(), TEXT("World Position")));
+		Sig.Inputs.Add(FNiagaraVariable(FNiagaraTypeDefinition::GetVec3Def(), TEXT("Node Position")));
 		Sig.Inputs.Add(FNiagaraVariable(FNiagaraTypeDefinition::GetFloatDef(), TEXT("Delta Time")));
 		Sig.Inputs.Add(FNiagaraVariable(FNiagaraTypeDefinition::GetFloatDef(), TEXT("Time Fraction")));
 		Sig.Outputs.Add(FNiagaraVariable(FNiagaraTypeDefinition::GetVec3Def(), TEXT("Closest Position")));
 		Sig.Outputs.Add(FNiagaraVariable(FNiagaraTypeDefinition::GetVec3Def(), TEXT("Closest Normal")));
 		Sig.Outputs.Add(FNiagaraVariable(FNiagaraTypeDefinition::GetVec3Def(), TEXT("Closest Velocity")));
+
+		OutFunctions.Add(Sig);
+	}
+	{
+		FNiagaraFunctionSignature Sig;
+		Sig.Name = GetClosestElementName;
+		Sig.bMemberFunction = true;
+		Sig.bRequiresContext = false;
+		Sig.Inputs.Add(FNiagaraVariable(FNiagaraTypeDefinition(GetClass()), TEXT("Physics Asset")));
+		Sig.Inputs.Add(FNiagaraVariable(FNiagaraTypeDefinition::GetVec3Def(), TEXT("Node Position")));
+		Sig.Inputs.Add(FNiagaraVariable(FNiagaraTypeDefinition::GetFloatDef(), TEXT("Time Fraction")));
+		Sig.Outputs.Add(FNiagaraVariable(FNiagaraTypeDefinition::GetIntDef(), TEXT("Closest Element")));
+
+		OutFunctions.Add(Sig);
+	}
+	{
+		FNiagaraFunctionSignature Sig;
+		Sig.Name = GetElementPointName;
+		Sig.bMemberFunction = true;
+		Sig.bRequiresContext = false;
+		Sig.Inputs.Add(FNiagaraVariable(FNiagaraTypeDefinition(GetClass()), TEXT("Physics Asset")));
+		Sig.Inputs.Add(FNiagaraVariable(FNiagaraTypeDefinition::GetVec3Def(), TEXT("Node Position")));
+		Sig.Inputs.Add(FNiagaraVariable(FNiagaraTypeDefinition::GetFloatDef(), TEXT("Delta Time")));
+		Sig.Inputs.Add(FNiagaraVariable(FNiagaraTypeDefinition::GetFloatDef(), TEXT("Time Fraction")));
+		Sig.Inputs.Add(FNiagaraVariable(FNiagaraTypeDefinition::GetIntDef(), TEXT("Element Index")));
+		Sig.Outputs.Add(FNiagaraVariable(FNiagaraTypeDefinition::GetVec3Def(), TEXT("Closest Position")));
+		Sig.Outputs.Add(FNiagaraVariable(FNiagaraTypeDefinition::GetVec3Def(), TEXT("Closest Normal")));
+		Sig.Outputs.Add(FNiagaraVariable(FNiagaraTypeDefinition::GetVec3Def(), TEXT("Closest Velocity")));
+
+		OutFunctions.Add(Sig);
+	}
+	{
+		FNiagaraFunctionSignature Sig;
+		Sig.Name = GetElementDistanceName;
+		Sig.bMemberFunction = true;
+		Sig.bRequiresContext = false;
+		Sig.Inputs.Add(FNiagaraVariable(FNiagaraTypeDefinition(GetClass()), TEXT("Physics Asset")));
+		Sig.Inputs.Add(FNiagaraVariable(FNiagaraTypeDefinition::GetVec3Def(), TEXT("Node Position")));
+		Sig.Inputs.Add(FNiagaraVariable(FNiagaraTypeDefinition::GetFloatDef(), TEXT("Time Fraction")));
+		Sig.Inputs.Add(FNiagaraVariable(FNiagaraTypeDefinition::GetIntDef(), TEXT("Element Index")));
+		Sig.Outputs.Add(FNiagaraVariable(FNiagaraTypeDefinition::GetFloatDef(), TEXT("Closest Distance")));
+
+		OutFunctions.Add(Sig);
+	}
+	{
+		FNiagaraFunctionSignature Sig;
+		Sig.Name = GetClosestDistanceName;
+		Sig.bMemberFunction = true;
+		Sig.bRequiresContext = false;
+		Sig.Inputs.Add(FNiagaraVariable(FNiagaraTypeDefinition(GetClass()), TEXT("Physics Asset")));
+		Sig.Inputs.Add(FNiagaraVariable(FNiagaraTypeDefinition::GetVec3Def(), TEXT("Node Position")));
+		Sig.Inputs.Add(FNiagaraVariable(FNiagaraTypeDefinition::GetFloatDef(), TEXT("Time Fraction")));
 		Sig.Outputs.Add(FNiagaraVariable(FNiagaraTypeDefinition::GetFloatDef(), TEXT("Closest Distance")));
 
 		OutFunctions.Add(Sig);
@@ -767,7 +823,7 @@ void UNiagaraDataInterfacePhysicsAsset::GetFunctions(TArray<FNiagaraFunctionSign
 		Sig.bMemberFunction = true;
 		Sig.bRequiresContext = false;
 		Sig.Inputs.Add(FNiagaraVariable(FNiagaraTypeDefinition(GetClass()), TEXT("Physics Asset")));
-		Sig.Inputs.Add(FNiagaraVariable(FNiagaraTypeDefinition::GetVec3Def(), TEXT("World Position")));
+		Sig.Inputs.Add(FNiagaraVariable(FNiagaraTypeDefinition::GetVec3Def(), TEXT("Node Position")));
 		Sig.Outputs.Add(FNiagaraVariable(FNiagaraTypeDefinition::GetIntDef(), TEXT("Element Index")));
 		Sig.Outputs.Add(FNiagaraVariable(FNiagaraTypeDefinition::GetVec3Def(), TEXT("Texture Position")));
 
@@ -779,7 +835,7 @@ void UNiagaraDataInterfacePhysicsAsset::GetFunctions(TArray<FNiagaraFunctionSign
 		Sig.bMemberFunction = true;
 		Sig.bRequiresContext = false;
 		Sig.Inputs.Add(FNiagaraVariable(FNiagaraTypeDefinition(GetClass()), TEXT("Physics Asset")));
-		Sig.Inputs.Add(FNiagaraVariable(FNiagaraTypeDefinition::GetVec3Def(), TEXT("World Position")));
+		Sig.Inputs.Add(FNiagaraVariable(FNiagaraTypeDefinition::GetVec3Def(), TEXT("Node Position")));
 		Sig.Inputs.Add(FNiagaraVariable(FNiagaraTypeDefinition::GetFloatDef(), TEXT("Delta Time")));
 		Sig.Inputs.Add(FNiagaraVariable(FNiagaraTypeDefinition::GetIntDef(), TEXT("Element Index")));
 		Sig.Inputs.Add(FNiagaraVariable(FNiagaraTypeDefinition::GetFloatDef(), TEXT("Texture Value")));
@@ -797,6 +853,10 @@ DEFINE_NDI_DIRECT_FUNC_BINDER(UNiagaraDataInterfacePhysicsAsset, GetNumBoxes);
 DEFINE_NDI_DIRECT_FUNC_BINDER(UNiagaraDataInterfacePhysicsAsset, GetNumSpheres);
 DEFINE_NDI_DIRECT_FUNC_BINDER(UNiagaraDataInterfacePhysicsAsset, GetNumCapsules);
 DEFINE_NDI_DIRECT_FUNC_BINDER(UNiagaraDataInterfacePhysicsAsset, GetClosestPoint);
+DEFINE_NDI_DIRECT_FUNC_BINDER(UNiagaraDataInterfacePhysicsAsset, GetClosestElement);
+DEFINE_NDI_DIRECT_FUNC_BINDER(UNiagaraDataInterfacePhysicsAsset, GetElementPoint);
+DEFINE_NDI_DIRECT_FUNC_BINDER(UNiagaraDataInterfacePhysicsAsset, GetElementDistance);
+DEFINE_NDI_DIRECT_FUNC_BINDER(UNiagaraDataInterfacePhysicsAsset, GetClosestDistance);
 DEFINE_NDI_DIRECT_FUNC_BINDER(UNiagaraDataInterfacePhysicsAsset, GetTexturePoint);
 DEFINE_NDI_DIRECT_FUNC_BINDER(UNiagaraDataInterfacePhysicsAsset, GetProjectionPoint);
 
@@ -819,9 +879,29 @@ void UNiagaraDataInterfacePhysicsAsset::GetVMExternalFunction(const FVMExternalF
 	}
 	else if (BindingInfo.Name == GetClosestPointName)
 	{
-		check(BindingInfo.GetNumInputs() == 6 && BindingInfo.GetNumOutputs() == 10);
+		check(BindingInfo.GetNumInputs() == 6 && BindingInfo.GetNumOutputs() == 9);
 		NDI_FUNC_BINDER(UNiagaraDataInterfacePhysicsAsset, GetClosestPoint)::Bind(this, OutFunc);
 	}
+	else if (BindingInfo.Name == GetClosestElementName)
+	{
+		check(BindingInfo.GetNumInputs() == 5 && BindingInfo.GetNumOutputs() == 1);
+		NDI_FUNC_BINDER(UNiagaraDataInterfacePhysicsAsset, GetClosestElement)::Bind(this, OutFunc);
+	}
+	else if (BindingInfo.Name == GetElementPointName)
+	{
+		check(BindingInfo.GetNumInputs() == 7 && BindingInfo.GetNumOutputs() == 9);
+		NDI_FUNC_BINDER(UNiagaraDataInterfacePhysicsAsset, GetElementPoint)::Bind(this, OutFunc);
+	}
+	else if (BindingInfo.Name == GetElementDistanceName)
+	{
+		check(BindingInfo.GetNumInputs() == 6 && BindingInfo.GetNumOutputs() == 1);
+		NDI_FUNC_BINDER(UNiagaraDataInterfacePhysicsAsset, GetElementDistance)::Bind(this, OutFunc);
+	}
+	else if (BindingInfo.Name == GetClosestDistanceName)
+	{
+		check(BindingInfo.GetNumInputs() == 5 && BindingInfo.GetNumOutputs() == 1);
+		NDI_FUNC_BINDER(UNiagaraDataInterfacePhysicsAsset, GetClosestDistance)::Bind(this, OutFunc);
+	}
 	else if (BindingInfo.Name == GetTexturePointName)
 	{
 		check(BindingInfo.GetNumInputs() == 4 && BindingInfo.GetNumOutputs() == 4);
@@ -850,6 +930,22 @@ void UNiagaraDataInterfacePhysicsAsset::GetClosestPoint(FVectorVMContext& Contex
 {
 }
 
+void UNiagaraDataInterfacePhysicsAsset::GetClosestElement(FVectorVMContext& Context)
+{
+}
+
+void UNiagaraDataInterfacePhysicsAsset::GetElementPoint(FVectorVMContext& Context)
+{
+}
+
+void UNiagaraDataInterfacePhysicsAsset::GetElementDistance(FVectorVMContext& Context)
+{
+}
+
+void UNiagaraDataInterfacePhysicsAsset::GetClosestDistance(FVectorVMContext& Context)
+{
+}
+
 void UNiagaraDataInterfacePhysicsAsset::GetTexturePoint(FVectorVMContext& Context)
 {
 }
@@ -912,11 +1008,59 @@ bool UNiagaraDataInterfacePhysicsAsset::GetFunctionHLSL(const FNiagaraDataInterf
 	else if (FunctionInfo.DefinitionName == GetClosestPointName)
 	{
 		static const TCHAR *FormatSample = TEXT(R"(
-		void {InstanceFunctionName}(in float3 WorldPosition, in float DeltaTime, in float TimeFraction, out float3 OutClosestPosition, 
-							out float3 OutClosestNormal, out float3 OutClosestVelocity, out float OutClosestDistance)
+		void {InstanceFunctionName}(in float3 NodePosition, in float DeltaTime, in float TimeFraction, out float3 OutClosestPosition, 
+							out float3 OutClosestNormal, out float3 OutClosestVelocity)
 		{
-			{PhysicsAssetContextName} DIPhysicsAsset_GetClosestPoint(DIContext,WorldPosition,DeltaTime,TimeFraction,
-				OutClosestPosition,OutClosestNormal,OutClosestVelocity,OutClosestDistance);
+			{PhysicsAssetContextName} DIPhysicsAsset_GetClosestPoint(DIContext,NodePosition,DeltaTime,TimeFraction,
+				OutClosestPosition,OutClosestNormal,OutClosestVelocity);
+		}
+		)");
+		OutHLSL += FString::Format(FormatSample, ArgsSample);
+		return true;
+	}
+	else if (FunctionInfo.DefinitionName == GetClosestElementName)
+	{
+		static const TCHAR* FormatSample = TEXT(R"(
+		void {InstanceFunctionName}(in float3 NodePosition, in float TimeFraction, out int OutClosestElement)
+		{
+			{PhysicsAssetContextName} DIPhysicsAsset_GetClosestElement(DIContext,NodePosition,TimeFraction,
+				OutClosestElement);
+		}
+		)");
+		OutHLSL += FString::Format(FormatSample, ArgsSample);
+		return true;
+	}
+	else if (FunctionInfo.DefinitionName == GetElementPointName)
+	{
+		static const TCHAR* FormatSample = TEXT(R"(
+		void {InstanceFunctionName}(in float3 NodePosition, in float DeltaTime, in float TimeFraction, in int ElementIndex, out float3 OutClosestPosition, 
+							out float3 OutClosestNormal, out float3 OutClosestVelocity)
+		{
+			{PhysicsAssetContextName} DIPhysicsAsset_GetElementPoint(DIContext,NodePosition,DeltaTime,TimeFraction,ElementIndex,
+				OutClosestPosition,OutClosestNormal,OutClosestVelocity);
+		}
+		)");
+		OutHLSL += FString::Format(FormatSample, ArgsSample);
+		return true;
+	}
+	else if (FunctionInfo.DefinitionName == GetElementDistanceName)
+	{
+		static const TCHAR* FormatSample = TEXT(R"(
+		void {InstanceFunctionName}(in float3 NodePosition, in float TimeFraction, in int ElementIndex, out float OutClosestDistance)
+		{
+			{PhysicsAssetContextName} DIPhysicsAsset_GetElementDistance(DIContext,NodePosition,TimeFraction,ElementIndex,
+				OutClosestDistance);
+		}
+		)");
+		OutHLSL += FString::Format(FormatSample, ArgsSample);
+		return true;
+	}
+	else if (FunctionInfo.DefinitionName == GetClosestDistanceName)
+	{
+		static const TCHAR* FormatSample = TEXT(R"(
+		void {InstanceFunctionName}(in float3 NodePosition, in float TimeFraction, out float OutClosestDistance)
+		{
+			{PhysicsAssetContextName} DIPhysicsAsset_GetClosestDistance(DIContext,NodePosition,TimeFraction,OutClosestDistance);
 		}
 		)");
 		OutHLSL += FString::Format(FormatSample, ArgsSample);
@@ -925,9 +1069,9 @@ bool UNiagaraDataInterfacePhysicsAsset::GetFunctionHLSL(const FNiagaraDataInterf
 	else if (FunctionInfo.DefinitionName == GetTexturePointName)
 	{
 		static const TCHAR *FormatSample = TEXT(R"(
-		void {InstanceFunctionName}(in float3 WorldPosition, out int OutElementIndex, out float3 OutTexturePosition)
+		void {InstanceFunctionName}(in float3 NodePosition, out int OutElementIndex, out float3 OutTexturePosition)
 		{
-			{PhysicsAssetContextName} DIPhysicsAsset_GetTexturePoint(DIContext,WorldPosition,OutElementIndex,OutTexturePosition);
+			{PhysicsAssetContextName} DIPhysicsAsset_GetTexturePoint(DIContext,NodePosition,OutElementIndex,OutTexturePosition);
 		}
 		)");
 		OutHLSL += FString::Format(FormatSample, ArgsSample);
@@ -936,10 +1080,10 @@ bool UNiagaraDataInterfacePhysicsAsset::GetFunctionHLSL(const FNiagaraDataInterf
 	else if (FunctionInfo.DefinitionName == GetProjectionPointName)
 	{
 		static const TCHAR *FormatSample = TEXT(R"(
-		void {InstanceFunctionName}(in float3 WorldPosition, in float DeltaTime, in int ElementIndex, in float TextureValue, in float3 TextureGradient, out float3 OutClosestPosition, 
+		void {InstanceFunctionName}(in float3 NodePosition, in float DeltaTime, in int ElementIndex, in float TextureValue, in float3 TextureGradient, out float3 OutClosestPosition, 
 							out float3 OutClosestNormal, out float3 OutClosestVelocity, out float OutClosestDistance)
 		{
-			{PhysicsAssetContextName} DIPhysicsAsset_GetProjectionPoint(DIContext,WorldPosition,DeltaTime,ElementIndex,TextureValue,TextureGradient,
+			{PhysicsAssetContextName} DIPhysicsAsset_GetProjectionPoint(DIContext,NodePosition,DeltaTime,ElementIndex,TextureValue,TextureGradient,
 				OutClosestPosition,OutClosestNormal,OutClosestVelocity,OutClosestDistance);
 		}
 		)");
diff --git a/Engine/Plugins/Experimental/HairStrands/Source/HairStrandsNiagara/Public/NiagaraDataInterfaceHairStrands.h b/Engine/Plugins/Experimental/HairStrands/Source/HairStrandsNiagara/Public/NiagaraDataInterfaceHairStrands.h
--- a/Engine/Plugins/Experimental/HairStrands/Source/HairStrandsNiagara/Public/NiagaraDataInterfaceHairStrands.h
+++ b/Engine/Plugins/Experimental/HairStrands/Source/HairStrandsNiagara/Public/NiagaraDataInterfaceHairStrands.h
@@ -10,14 +10,21 @@
 #include "GroomActor.h"
 #include "NiagaraDataInterfaceHairStrands.generated.h"
 
-static const int32 MaxDelay = 2;
+static const int32 MaxDelay = 8;
+static const int32 NumScales = 4;
+static const int32 StretchOffset = 0;
+static const int32 BendOffset = 1;
+static const int32 RadiusOffset = 2;
+static const int32 ThicknessOffset = 3;
+
+struct FNDIHairStrandsData;
 
 /** Render buffers that will be used in hlsl functions */
 struct FNDIHairStrandsBuffer : public FRenderResource
 {
 	/** Set the asset that will be used to affect the buffer */
 	void Initialize(const FHairStrandsDatas*  HairStrandsDatas, const FHairStrandsRestResource*  HairStrandsRestResource, 
-		const FHairStrandsDeformedResource*  HairStrandsDeformedResource, const FHairStrandsRootResource* HairStrandsRootResource );
+		const FHairStrandsDeformedResource*  HairStrandsDeformedResource, const FHairStrandsRootResource* HairStrandsRootResource, const FNDIHairStrandsData* NDIStrandsDatas);
 
 	/** Init the buffer */
 	virtual void InitRHI() override;
@@ -40,6 +47,9 @@ struct FNDIHairStrandsBuffer : public FRenderResource
 	/** Bounding Box Buffer B*/
 	FRWBuffer BoundingBoxBufferB;
 
+	/** Params scale buffer */
+	FRWBuffer ParamsScaleBuffer;
+
 	/** The strand asset datas from which to sample */
 	const FHairStrandsDatas* SourceDatas;
 
@@ -51,6 +61,9 @@ struct FNDIHairStrandsBuffer : public FRenderResource
 
 	/** The strand root resource to write into */
 	const FHairStrandsRootResource* SourceRootResources;
+
+	/** Niagara datas interface datas */
+	const FNDIHairStrandsData* NDIDatas;
 };
 
 /** Data stored per strand base instance*/
@@ -66,12 +79,17 @@ struct FNDIHairStrandsData
 	/** Release the buffers */
 	void Release();
 
+	/** Update the buffers */
+	void Update(UNiagaraDataInterfaceHairStrands* Interface, FNiagaraSystemInstance* SystemInstance, const FHairStrandsDatas* HairStrandsDatas, UGroomAsset* GroomAsset, const int32 GroupIndex);
+
 	inline void ResetDatas()
 	{
 		WorldTransform.SetIdentity();
 		BoxCenter = FVector(0, 0, 0);
 		BoxExtent = FVector(0, 0, 0);
 
+		KinematicsTarget = false;
+
 		TickCount = 0;
 		ForceReset = true;
 
@@ -105,6 +123,11 @@ struct FNDIHairStrandsData
 		StrandsDensity = 1.0;
 		StrandsSmoothing = 0.1;
 		StrandsThickness = 0.01;
+
+		for (int32 i = 0; i < 32 * NumScales; ++i)
+		{
+			ParamsScale[i] = 1.0;
+		}
 	}
 
 	inline void CopyDatas(const FNDIHairStrandsData* OtherDatas)
@@ -117,6 +140,8 @@ struct FNDIHairStrandsData
 			BoxCenter = OtherDatas->BoxCenter;
 			BoxExtent = OtherDatas->BoxExtent;
 
+			KinematicsTarget = OtherDatas->KinematicsTarget;
+
 			TickCount = OtherDatas->TickCount;
 			ForceReset = OtherDatas->ForceReset;
 
@@ -134,13 +159,11 @@ struct FNDIHairStrandsData
 			ProjectBend = OtherDatas->ProjectBend;
 			BendDamping = OtherDatas->BendDamping;
 			BendStiffness = OtherDatas->BendStiffness;
-			BendScale = OtherDatas->BendScale;
 
 			SolveStretch = OtherDatas->SolveStretch;
 			ProjectStretch = OtherDatas->ProjectStretch;
 			StretchDamping = OtherDatas->StretchDamping;
 			StretchStiffness = OtherDatas->StretchStiffness;
-			StretchScale = OtherDatas->StretchScale;
 
 			SolveCollision = OtherDatas->SolveCollision;
 			ProjectCollision = OtherDatas->ProjectCollision;
@@ -148,18 +171,21 @@ struct FNDIHairStrandsData
 			KineticFriction = OtherDatas->KineticFriction;
 			StrandsViscosity = OtherDatas->StrandsViscosity;
 			CollisionRadius = OtherDatas->CollisionRadius;
-			RadiusScale = OtherDatas->RadiusScale;
 
 			StrandsDensity = OtherDatas->StrandsDensity;
 			StrandsSmoothing = OtherDatas->StrandsSmoothing;
 			StrandsThickness = OtherDatas->StrandsThickness;
-			ThicknessScale = OtherDatas->ThicknessScale;
+
+			ParamsScale = OtherDatas->ParamsScale;
 		}
 	}
 
 	/** Cached World transform. */
 	FTransform WorldTransform;
 
+	/** Kinematics target */
+	bool KinematicsTarget;
+
 	/** Number of strands*/
 	int32 NumStrands;
 
@@ -208,9 +234,6 @@ struct FNDIHairStrandsData
 	/** Stiffness for the bend constraint in GPa */
 	float BendStiffness;
 
-	/** Stiffness scale along the strand */
-	TStaticArray<float,4> BendScale;
-
 	/** Enable the solve of the stretch constraint during the xpbd loop */
 	bool SolveStretch;
 
@@ -223,9 +246,6 @@ struct FNDIHairStrandsData
 	/** Stiffness for the stretch constraint in GPa */
 	float StretchStiffness;
 
-	/** Stiffness scale along the strand */
-	TStaticArray<float, 4> StretchScale;
-
 	/** Enable the solve of the collision constraint during the xpbd loop  */
 	bool SolveCollision;
 
@@ -244,9 +264,6 @@ struct FNDIHairStrandsData
 	/** Radius scale along the strand */
 	float CollisionRadius;
 
-	/** Radius scale along the strand */
-	TStaticArray<float, 4> RadiusScale;
-
 	/** Density of the strands in g/cm3 */
 	float StrandsDensity;
 
@@ -256,8 +273,8 @@ struct FNDIHairStrandsData
 	/** Strands thickness in cm that will be used for mass and inertia computation */
 	float StrandsThickness;
 
-	/** Thickness scale along the curve */
-	TStaticArray<float, 4> ThicknessScale;
+	/** Scales along the strand */
+	TStaticArray<float, 32 * NumScales> ParamsScale;
 };
 
 /** Data Interface for the strand base */
@@ -308,7 +325,8 @@ public:
 
 	/** Extract datas and resources */
 	void ExtractDatasAndResources(FNiagaraSystemInstance* SystemInstance, FHairStrandsDatas*& OutStrandsDatas,
-		FHairStrandsRestResource*& OutStrandsRestResource, FHairStrandsDeformedResource*& OutStrandsDeformedResource, FHairStrandsRootResource*& OutStrandsRootResource, UGroomAsset*& OutGroomAsset);
+		FHairStrandsRestResource*& OutStrandsRestResource, FHairStrandsDeformedResource*& OutStrandsDeformedResource, 
+		FHairStrandsRootResource*& OutStrandsRootResource, UGroomAsset*& OutGroomAsset, int32& OutGroupIndex);
 
 	/** Get the number of strands */
 	void GetNumStrands(FVectorVMContext& Context);
@@ -486,7 +504,7 @@ public:
 	void SetupSoftCollisionConstraint(FVectorVMContext& Context);
 
 	/** Compute the rest direction*/
-	void ComputeRestDirection(FVectorVMContext& Context);
+	void ComputeEdgeDirection(FVectorVMContext& Context);
 
 	/** Update the strands material frame */
 	void UpdateMaterialFrame(FVectorVMContext& Context);
@@ -503,6 +521,9 @@ public:
 	/** Attach the node position and orientation to the transform or to the skin cache */
 	void AttachNodeState(FVectorVMContext& Context);
 
+	/** Update the node position and orientation based on rbf transfer */
+	void UpdateNodeState(FVectorVMContext& Context);
+
 	/** Check if we need or not a simulation reset*/
 	void NeedSimulationReset(FVectorVMContext& Context);
 
@@ -552,7 +573,7 @@ public:
 	static const FString BoxExtentName;
 
 	/** Param to check if the roots have been attached to the skin */
-	static const FString HasRootAttachedName;
+	static const FString InterpolationModeName;
 
 	/** boolean to check if we need to rest the simulation*/
 	static const FString ResetSimulationName;
@@ -599,6 +620,9 @@ public:
 	/** Rbf Sample rest positions */
 	static const FString RestSamplePositionsName;
 
+	/** Params scale buffer */
+	static const FString ParamsScaleBufferName;
+
 protected:
 	/** Copy one niagara DI to this */
 	virtual bool CopyToInternal(UNiagaraDataInterface* Destination) const override;
diff --git a/Engine/Plugins/Experimental/HairStrands/Source/HairStrandsNiagara/Public/NiagaraDataInterfacePhysicsAsset.h b/Engine/Plugins/Experimental/HairStrands/Source/HairStrandsNiagara/Public/NiagaraDataInterfacePhysicsAsset.h
--- a/Engine/Plugins/Experimental/HairStrands/Source/HairStrandsNiagara/Public/NiagaraDataInterfacePhysicsAsset.h
+++ b/Engine/Plugins/Experimental/HairStrands/Source/HairStrandsNiagara/Public/NiagaraDataInterfacePhysicsAsset.h
@@ -166,9 +166,21 @@ public:
 	/** Get the number of capsules */
 	void GetNumCapsules(FVectorVMContext& Context);
 
+	/** Get the element point */
+	void GetElementPoint(FVectorVMContext& Context);
+
+	/** Get the element distance */
+	void GetElementDistance(FVectorVMContext& Context);
+
+	/** Get the closest element */
+	void GetClosestElement(FVectorVMContext& Context);
+
 	/** Get the closest point */
 	void GetClosestPoint(FVectorVMContext& Context);
 
+	/** Get the closest distance */
+	void GetClosestDistance(FVectorVMContext& Context);
+
 	/** Get the closest texture point */
 	void GetTexturePoint(FVectorVMContext& Context);
 
diff --git a/Engine/Shaders/Private/HairStrands/HairStrandsGuideDeform.usf b/Engine/Shaders/Private/HairStrands/HairStrandsGuideDeform.usf
--- a/Engine/Shaders/Private/HairStrands/HairStrandsGuideDeform.usf
+++ b/Engine/Shaders/Private/HairStrands/HairStrandsGuideDeform.usf
@@ -119,6 +119,10 @@ float3 DisplacePosition(uint VertexIndex, float3 Pos, float Seed)
 		const float FunctionValue = sqrt( dot(PositionDelta,PositionDelta)+1);
 		ControlPoint += FunctionValue * MeshSampleWeightsBuffer[i].xyz;
 	}
+	ControlPoint += MeshSampleWeightsBuffer[SampleCount].xyz;
+	ControlPoint += MeshSampleWeightsBuffer[SampleCount+1].xyz * RestControlPoint.x;
+	ControlPoint += MeshSampleWeightsBuffer[SampleCount+2].xyz * RestControlPoint.y;
+	ControlPoint += MeshSampleWeightsBuffer[SampleCount+3].xyz * RestControlPoint.z;
 	return ControlPoint - SimDeformedOffset;
 #endif
 
diff --git a/Engine/Shaders/Private/HairStrands/HairStrandsMeshInterpolate.usf b/Engine/Shaders/Private/HairStrands/HairStrandsMeshInterpolate.usf
--- a/Engine/Shaders/Private/HairStrands/HairStrandsMeshInterpolate.usf
+++ b/Engine/Shaders/Private/HairStrands/HairStrandsMeshInterpolate.usf
@@ -35,6 +35,10 @@ void MainCS(uint3 DispatchThreadId : SV_DispatchThreadID)
 			const float FunctionValue = sqrt( dot(PositionDelta,PositionDelta)+1);
 			ControlPoint += FunctionValue * MeshSampleWeightsBuffer[i].xyz;
 		}
+		ControlPoint += MeshSampleWeightsBuffer[MaxSampleCount].xyz;
+		ControlPoint += MeshSampleWeightsBuffer[MaxSampleCount+1].xyz * RestControlPoint.x;
+		ControlPoint += MeshSampleWeightsBuffer[MaxSampleCount+2].xyz * RestControlPoint.y;
+		ControlPoint += MeshSampleWeightsBuffer[MaxSampleCount+3].xyz * RestControlPoint.z;
 		OutDeformedRootTrianglePosition0[RootIndex+MaxRootCount] = float4(ControlPoint,1);
 	}
 	{
@@ -48,6 +52,10 @@ void MainCS(uint3 DispatchThreadId : SV_DispatchThreadID)
 			const float FunctionValue = sqrt( dot(PositionDelta,PositionDelta)+1);
 			ControlPoint += FunctionValue * MeshSampleWeightsBuffer[i].xyz;
 		}
+		ControlPoint += MeshSampleWeightsBuffer[MaxSampleCount].xyz;
+		ControlPoint += MeshSampleWeightsBuffer[MaxSampleCount+1].xyz * RestControlPoint.x;
+		ControlPoint += MeshSampleWeightsBuffer[MaxSampleCount+2].xyz * RestControlPoint.y;
+		ControlPoint += MeshSampleWeightsBuffer[MaxSampleCount+3].xyz * RestControlPoint.z;
 		OutDeformedRootTrianglePosition1[RootIndex+MaxRootCount] = float4(ControlPoint,1);
 	}
 	{
@@ -61,6 +69,10 @@ void MainCS(uint3 DispatchThreadId : SV_DispatchThreadID)
 			const float FunctionValue = sqrt( dot(PositionDelta,PositionDelta)+1);
 			ControlPoint += FunctionValue * MeshSampleWeightsBuffer[i].xyz;
 		}
+		ControlPoint += MeshSampleWeightsBuffer[MaxSampleCount].xyz;
+		ControlPoint += MeshSampleWeightsBuffer[MaxSampleCount+1].xyz * RestControlPoint.x;
+		ControlPoint += MeshSampleWeightsBuffer[MaxSampleCount+2].xyz * RestControlPoint.y;
+		ControlPoint += MeshSampleWeightsBuffer[MaxSampleCount+3].xyz * RestControlPoint.z;
 		OutDeformedRootTrianglePosition2[RootIndex+MaxRootCount] = float4(ControlPoint,1);
 	}
 }
\ No newline at end of file
diff --git a/Engine/Shaders/Private/HairStrands/HairStrandsSamplesUpdate.usf b/Engine/Shaders/Private/HairStrands/HairStrandsSamplesUpdate.usf
--- a/Engine/Shaders/Private/HairStrands/HairStrandsSamplesUpdate.usf
+++ b/Engine/Shaders/Private/HairStrands/HairStrandsSamplesUpdate.usf
@@ -15,10 +15,10 @@ RWBuffer<float4> OutSampleDeformationsBuffer;
 void MainCS(uint3 DispatchThreadId : SV_DispatchThreadID)
 { 
 	const uint SampleIndex = DispatchThreadId.x;
-	if (SampleIndex >= MaxSampleCount) 
+	if (SampleIndex >= (MaxSampleCount+4)) 
 		return;
 	
-	uint WeightsOffset = SampleIndex * MaxSampleCount;
+	uint WeightsOffset = SampleIndex * (MaxSampleCount+4);
 	float3 SampleDeformation = float3(0,0,0);
 	for(uint i = 0; i < MaxSampleCount; ++i, ++WeightsOffset)
 	{
diff --git a/Engine/Source/Runtime/Renderer/Private/HairStrands/HairStrandsMeshProjection.cpp b/Engine/Source/Runtime/Renderer/Private/HairStrands/HairStrandsMeshProjection.cpp
--- a/Engine/Source/Runtime/Renderer/Private/HairStrands/HairStrandsMeshProjection.cpp
+++ b/Engine/Source/Runtime/Renderer/Private/HairStrands/HairStrandsMeshProjection.cpp
@@ -888,7 +888,7 @@ static void AddHairStrandUpdateMeshSamplesPass(
 		Parameters->SampleDeformedPositionsBuffer = LODData.DeformedSamplePositionsBuffer->SRV;
 		Parameters->OutSampleDeformationsBuffer = LODData.MeshSampleWeightsBuffer->UAV;
 
-		const FIntVector DispatchGroupCount = FComputeShaderUtils::GetGroupCount(LODData.SampleCount, 128);
+		const FIntVector DispatchGroupCount = FComputeShaderUtils::GetGroupCount(LODData.SampleCount+4, 128);
 		check(DispatchGroupCount.X < 65536);
 		TShaderMapRef<FHairUpdateMeshSamplesCS> ComputeShader(ShaderMap);
 		FComputeShaderUtils::AddPass(
