Fix CIS

diff --git a/Engine/Plugins/Experimental/HairStrands/Source/HairStrandsNiagara/Private/NiagaraDataInterfaceHairStrands.cpp b/Engine/Plugins/Experimental/HairStrands/Source/HairStrandsNiagara/Private/NiagaraDataInterfaceHairStrands.cpp
--- a/Engine/Plugins/Experimental/HairStrands/Source/HairStrandsNiagara/Private/NiagaraDataInterfaceHairStrands.cpp
+++ b/Engine/Plugins/Experimental/HairStrands/Source/HairStrandsNiagara/Private/NiagaraDataInterfaceHairStrands.cpp
@@ -755,8 +755,9 @@ struct FNDIHairStrandsParametersCS : public FNiagaraDataInterfaceParametersCS
 				InterpolationModeValue = 2;
 			}
 
+			FUnorderedAccessViewRHIRef DummyFloatUAV(Context.Batcher->GetEmptyRWBufferFromPool(RHICmdList, PF_R32_FLOAT), false);
 			FUnorderedAccessViewRHIRef MeshSampleWeightsBufferUAV = (MeshProjection != nullptr && MeshProjection->SampleCount > 0) ?
-				MeshProjection->MeshSampleWeightsBuffer.UAV : Context.Batcher->GetEmptyRWBufferFromPool(RHICmdList, PF_R32_FLOAT);
+				MeshProjection->MeshSampleWeightsBuffer.UAV : DummyFloatUAV;
 
 			//UE_LOG(LogHairStrands, Log, TEXT("Shader Reset : %d %d %d"), bNeedSimReset, IsRootValid, InterpolationModeValue);
 			FVector RestRootOffsetValue = FVector::ZeroVector;
