Fix hair being rendered into the base pass when using hair strands debug material.

diff --git a/Engine/Plugins/Experimental/HairStrands/Source/HairStrandsCore/Private/GroomComponent.cpp b/Engine/Plugins/Experimental/HairStrands/Source/HairStrandsCore/Private/GroomComponent.cpp
--- a/Engine/Plugins/Experimental/HairStrands/Source/HairStrandsCore/Private/GroomComponent.cpp
+++ b/Engine/Plugins/Experimental/HairStrands/Source/HairStrandsCore/Private/GroomComponent.cpp
@@ -409,7 +409,6 @@ public:
 		if (DebugMode != EHairStrandsDebugMode::None)
 		{
 			FPrimitiveViewRelevance Result;
-			Result.bDrawRelevance = IsShown(View);
 			Result.bShadowRelevance = IsShadowCast(View);
 			Result.bDynamicRelevance = true;
 			MaterialRelevance.SetPrimitiveViewRelevance(Result);
@@ -428,7 +427,7 @@ public:
 		// Selection only
 		#if WITH_EDITOR
 		{
-			const bool bIsSelected = (IsSelected() || IsHovered()) && bIsViewModeValid;
+			const bool bIsSelected = (IsSelected() || IsHovered());
 			Result.bEditorStaticSelectionRelevance = bIsSelected;
 			Result.bDrawRelevance = bIsSelected;
 		}
