Fix CIS issue.

diff --git a/Engine/Plugins/Experimental/HairStrands/Source/HairStrandsCore/Private/GroomComponent.cpp b/Engine/Plugins/Experimental/HairStrands/Source/HairStrandsCore/Private/GroomComponent.cpp
--- a/Engine/Plugins/Experimental/HairStrands/Source/HairStrandsCore/Private/GroomComponent.cpp
+++ b/Engine/Plugins/Experimental/HairStrands/Source/HairStrandsCore/Private/GroomComponent.cpp
@@ -1172,7 +1172,7 @@ void UGroomComponent::InitResources(bool bIsBindingReloading)
 
 		if (bRunMeshProjection)
 		{				
-			FSkeletalMeshRenderData* TargetRenderData = SkeletalMeshComponent->GetSkeletalMeshRenderData();
+			FSkeletalMeshRenderData* TargetRenderData = SkeletalMeshComponent ? SkeletalMeshComponent->GetSkeletalMeshRenderData() : nullptr;
 			FHairStrandsProjectionMeshData TargetMeshData = ExtractMeshData(TargetRenderData);
 
 			// Create mapping between the source & target using their UV
