Ensure BoneIndex is in not out of bound

diff --git a/Engine/Plugins/Experimental/HairStrands/Source/HairStrandsCore/Private/GroomComponent.cpp b/Engine/Plugins/Experimental/HairStrands/Source/HairStrandsCore/Private/GroomComponent.cpp
--- a/Engine/Plugins/Experimental/HairStrands/Source/HairStrandsCore/Private/GroomComponent.cpp
+++ b/Engine/Plugins/Experimental/HairStrands/Source/HairStrandsCore/Private/GroomComponent.cpp
@@ -1384,12 +1384,12 @@ void UGroomComponent::OnComponentDestroyed(bool bDestroyingHierarchy)
 	ReleaseResources();
 
 #if WITH_EDITOR
-	if (bIsGroomAssetCallbackRegistered)
+	if (bIsGroomAssetCallbackRegistered && GroomAsset)
 	{
 		GroomAsset->GetOnGroomAssetChanged().RemoveAll(this);
 		bIsGroomAssetCallbackRegistered = false;
 	}
-	if (bIsGroomBindingAssetCallbackRegistered)
+	if (bIsGroomBindingAssetCallbackRegistered && BindingAsset)
 	{
 		BindingAsset->GetOnGroomBindingAssetChanged().RemoveAll(this);
 		bIsGroomBindingAssetCallbackRegistered = false;
diff --git a/Engine/Plugins/Experimental/HairStrands/Source/HairStrandsNiagara/Private/NiagaraDataInterfacePhysicsAsset.cpp b/Engine/Plugins/Experimental/HairStrands/Source/HairStrandsNiagara/Private/NiagaraDataInterfacePhysicsAsset.cpp
--- a/Engine/Plugins/Experimental/HairStrands/Source/HairStrandsNiagara/Private/NiagaraDataInterfacePhysicsAsset.cpp
+++ b/Engine/Plugins/Experimental/HairStrands/Source/HairStrandsNiagara/Private/NiagaraDataInterfacePhysicsAsset.cpp
@@ -134,83 +134,86 @@ void CreateInternalArrays(const TWeakObjectPtr<UPhysicsAsset> PhysicsAsset, cons
 				TArray<FTransform> BoneTransforms = (SkeletalMesh != nullptr) ? SkeletalMesh->GetComponentSpaceTransforms() : RestTransforms;
 				const bool bHasMasterPoseComponent = (SkeletalMesh != nullptr) && SkeletalMesh->MasterPoseComponent.IsValid();
 
-				uint32 NumBoxes = 0;
-				uint32 NumSpheres = 0;
-				uint32 NumCapsules = 0;
-				for (const UBodySetup* BodySetup : PhysicsAsset->SkeletalBodySetups)
+				if (BoneTransforms.Num() == RestTransforms.Num())
 				{
-					const FName BoneName = BodySetup->BoneName;
-					const int32 BoneIndex = RefSkeleton->FindBoneIndex(BoneName);
-					if (BoneIndex != INDEX_NONE)
+					uint32 NumBoxes = 0;
+					uint32 NumSpheres = 0;
+					uint32 NumCapsules = 0;
+					for (const UBodySetup* BodySetup : PhysicsAsset->SkeletalBodySetups)
 					{
-						NumBoxes += BodySetup->AggGeom.BoxElems.Num();
-						NumSpheres += BodySetup->AggGeom.SphereElems.Num();
-						NumCapsules += BodySetup->AggGeom.SphylElems.Num();
+						const FName BoneName = BodySetup->BoneName;
+						const int32 BoneIndex = RefSkeleton->FindBoneIndex(BoneName);
+						if (BoneIndex != INDEX_NONE && BoneIndex < BoneTransforms.Num())
+						{
+							NumBoxes += BodySetup->AggGeom.BoxElems.Num();
+							NumSpheres += BodySetup->AggGeom.SphereElems.Num();
+							NumCapsules += BodySetup->AggGeom.SphylElems.Num();
+						}
 					}
-				}
 
-				OutAssetArrays->ElementOffsets.BoxOffset = 0;
-				OutAssetArrays->ElementOffsets.SphereOffset = OutAssetArrays->ElementOffsets.BoxOffset + NumBoxes;
-				OutAssetArrays->ElementOffsets.CapsuleOffset = OutAssetArrays->ElementOffsets.SphereOffset + NumSpheres;
-				OutAssetArrays->ElementOffsets.NumElements = OutAssetArrays->ElementOffsets.CapsuleOffset + NumCapsules;
+					OutAssetArrays->ElementOffsets.BoxOffset = 0;
+					OutAssetArrays->ElementOffsets.SphereOffset = OutAssetArrays->ElementOffsets.BoxOffset + NumBoxes;
+					OutAssetArrays->ElementOffsets.CapsuleOffset = OutAssetArrays->ElementOffsets.SphereOffset + NumSpheres;
+					OutAssetArrays->ElementOffsets.NumElements = OutAssetArrays->ElementOffsets.CapsuleOffset + NumCapsules;
 
-				const uint32 NumTransforms = OutAssetArrays->ElementOffsets.NumElements * 3;
-				const uint32 NumExtents = OutAssetArrays->ElementOffsets.NumElements;
+					const uint32 NumTransforms = OutAssetArrays->ElementOffsets.NumElements * 3;
+					const uint32 NumExtents = OutAssetArrays->ElementOffsets.NumElements;
 
-				OutAssetArrays->CurrentTransform.SetNum(NumTransforms);
-				OutAssetArrays->InverseTransform.SetNum(NumTransforms);
-				OutAssetArrays->RestInverse.SetNum(NumTransforms);
-				OutAssetArrays->RestTransform.SetNum(NumTransforms);
-				OutAssetArrays->PreviousTransform.SetNum(NumTransforms);
-				OutAssetArrays->PreviousInverse.SetNum(NumTransforms);
-				OutAssetArrays->ElementExtent.SetNum(NumExtents);
+					OutAssetArrays->CurrentTransform.SetNum(NumTransforms);
+					OutAssetArrays->InverseTransform.SetNum(NumTransforms);
+					OutAssetArrays->RestInverse.SetNum(NumTransforms);
+					OutAssetArrays->RestTransform.SetNum(NumTransforms);
+					OutAssetArrays->PreviousTransform.SetNum(NumTransforms);
+					OutAssetArrays->PreviousInverse.SetNum(NumTransforms);
+					OutAssetArrays->ElementExtent.SetNum(NumExtents);
 
-				uint32 ElementCount = 0;
-				for (const UBodySetup* BodySetup : PhysicsAsset->SkeletalBodySetups)
-				{
-					const FName BoneName = BodySetup->BoneName;
-					const int32 BoneIndex = RefSkeleton->FindBoneIndex(BoneName);
-					if (BoneIndex != INDEX_NONE)
+					uint32 ElementCount = 0;
+					for (const UBodySetup* BodySetup : PhysicsAsset->SkeletalBodySetups)
 					{
-						const FTransform RestTransform = RestTransforms[BoneIndex];
-						const FTransform BoneTransform = bHasMasterPoseComponent ? SkeletalMesh->GetBoneTransform(BoneIndex) : BoneTransforms[BoneIndex] * WorldTransform;
-
-						for (const FKBoxElem& BoxElem : BodySetup->AggGeom.BoxElems)
+						const FName BoneName = BodySetup->BoneName;
+						const int32 BoneIndex = RefSkeleton->FindBoneIndex(BoneName);
+						if (BoneIndex != INDEX_NONE && BoneIndex < BoneTransforms.Num())
 						{
-							const FTransform RestElement = FTransform(BoxElem.Rotation, BoxElem.Center) * RestTransform;
-							FillCurrentTransforms(RestElement, ElementCount, OutAssetArrays->RestTransform, OutAssetArrays->RestInverse);
-							--ElementCount;
-
-							const FTransform ElementTransform = FTransform(BoxElem.Rotation, BoxElem.Center) * BoneTransform;
-							OutAssetArrays->ElementExtent[ElementCount] = FVector4(BoxElem.X, BoxElem.Y, BoxElem.Z, 0);
-							FillCurrentTransforms(ElementTransform, ElementCount, OutAssetArrays->CurrentTransform, OutAssetArrays->InverseTransform);
-						}
-
-						for (const FKSphereElem& SphereElem : BodySetup->AggGeom.SphereElems)
-						{
-							const FTransform RestElement = FTransform(SphereElem.Center) * RestTransform;
-							FillCurrentTransforms(RestElement, ElementCount, OutAssetArrays->RestTransform, OutAssetArrays->RestInverse);
-							--ElementCount;
-
-							const FTransform ElementTransform = FTransform(SphereElem.Center) * BoneTransform;
-							OutAssetArrays->ElementExtent[ElementCount] = FVector4(SphereElem.Radius, 0, 0, 0);
-							FillCurrentTransforms(ElementTransform, ElementCount, OutAssetArrays->CurrentTransform, OutAssetArrays->InverseTransform);
-						}
-
-						for (const FKSphylElem& CapsuleElem : BodySetup->AggGeom.SphylElems)
-						{
-							const FTransform RestElement = FTransform(CapsuleElem.Rotation, CapsuleElem.Center) * RestTransform;
-							FillCurrentTransforms(RestElement, ElementCount, OutAssetArrays->RestTransform, OutAssetArrays->RestInverse);
-							--ElementCount;
-
-							const FTransform ElementTransform = FTransform(CapsuleElem.Rotation, CapsuleElem.Center) * BoneTransform;
-							OutAssetArrays->ElementExtent[ElementCount] = FVector4(CapsuleElem.Radius, CapsuleElem.Length, 0, 0);
-							FillCurrentTransforms(ElementTransform, ElementCount, OutAssetArrays->CurrentTransform, OutAssetArrays->InverseTransform);
+							const FTransform RestTransform = RestTransforms[BoneIndex];
+							const FTransform BoneTransform = bHasMasterPoseComponent ? SkeletalMesh->GetBoneTransform(BoneIndex) : BoneTransforms[BoneIndex] * WorldTransform;
+
+							for (const FKBoxElem& BoxElem : BodySetup->AggGeom.BoxElems)
+							{
+								const FTransform RestElement = FTransform(BoxElem.Rotation, BoxElem.Center) * RestTransform;
+								FillCurrentTransforms(RestElement, ElementCount, OutAssetArrays->RestTransform, OutAssetArrays->RestInverse);
+								--ElementCount;
+
+								const FTransform ElementTransform = FTransform(BoxElem.Rotation, BoxElem.Center) * BoneTransform;
+								OutAssetArrays->ElementExtent[ElementCount] = FVector4(BoxElem.X, BoxElem.Y, BoxElem.Z, 0);
+								FillCurrentTransforms(ElementTransform, ElementCount, OutAssetArrays->CurrentTransform, OutAssetArrays->InverseTransform);
+							}
+
+							for (const FKSphereElem& SphereElem : BodySetup->AggGeom.SphereElems)
+							{
+								const FTransform RestElement = FTransform(SphereElem.Center) * RestTransform;
+								FillCurrentTransforms(RestElement, ElementCount, OutAssetArrays->RestTransform, OutAssetArrays->RestInverse);
+								--ElementCount;
+
+								const FTransform ElementTransform = FTransform(SphereElem.Center) * BoneTransform;
+								OutAssetArrays->ElementExtent[ElementCount] = FVector4(SphereElem.Radius, 0, 0, 0);
+								FillCurrentTransforms(ElementTransform, ElementCount, OutAssetArrays->CurrentTransform, OutAssetArrays->InverseTransform);
+							}
+
+							for (const FKSphylElem& CapsuleElem : BodySetup->AggGeom.SphylElems)
+							{
+								const FTransform RestElement = FTransform(CapsuleElem.Rotation, CapsuleElem.Center) * RestTransform;
+								FillCurrentTransforms(RestElement, ElementCount, OutAssetArrays->RestTransform, OutAssetArrays->RestInverse);
+								--ElementCount;
+
+								const FTransform ElementTransform = FTransform(CapsuleElem.Rotation, CapsuleElem.Center) * BoneTransform;
+								OutAssetArrays->ElementExtent[ElementCount] = FVector4(CapsuleElem.Radius, CapsuleElem.Length, 0, 0);
+								FillCurrentTransforms(ElementTransform, ElementCount, OutAssetArrays->CurrentTransform, OutAssetArrays->InverseTransform);
+							}
 						}
 					}
+					OutAssetArrays->PreviousTransform = OutAssetArrays->CurrentTransform;
+					OutAssetArrays->PreviousInverse = OutAssetArrays->InverseTransform;
 				}
-				OutAssetArrays->PreviousTransform = OutAssetArrays->CurrentTransform;
-				OutAssetArrays->PreviousInverse = OutAssetArrays->InverseTransform;
 			}
 		}
 	}
@@ -243,7 +246,7 @@ void UpdateInternalArrays(const TWeakObjectPtr<UPhysicsAsset> PhysicsAsset, cons
 			{
 				const FName BoneName = BodySetup->BoneName;
 				const int32 BoneIndex = RefSkeleton->FindBoneIndex(BoneName);
-				if (BoneIndex != INDEX_NONE)
+				if (BoneIndex != INDEX_NONE && BoneIndex < BoneTransforms.Num() )
 				{
 					const FTransform BoneTransform = bHasMasterPoseComponent ? SkeletalMesh->GetBoneTransform(BoneIndex) : BoneTransforms[BoneIndex] * WorldTransform;
 
@@ -325,25 +328,6 @@ void FNDIPhysicsAssetBuffer::InitRHI()
 		UE_LOG(LogPhysicsAsset, Warning, TEXT("Num Capsules = %d | Num Spheres = %d | Num Boxes = %d"), AssetArrays->ElementOffsets.NumElements - AssetArrays->ElementOffsets.CapsuleOffset,
 			AssetArrays->ElementOffsets.CapsuleOffset - AssetArrays->ElementOffsets.SphereOffset, AssetArrays->ElementOffsets.SphereOffset - AssetArrays->ElementOffsets.BoxOffset);
 	}
-
-	// /*const FPrimitiveComponentId LocalComponentId = SkeletalMesh->ComponentId;
-	//EWorldType::Type WorldType = SkeletalMesh->GetWorld() ? EWorldType::Type(SkeletalMesh->GetWorld()->WorldType) : EWorldType::None;
-	//WorldType = WorldType == EWorldType::Inactive ? EWorldType::Editor : WorldType;*/
-
-	//// Always setup the callback, even if the SkeletalMeshData are not ready yet
-	//if (SkeletalMesh.IsValid() && SkeletalMesh.Get() != nullptr)
-	//{
-	//	/*FSkeletalMeshObjectCallbackData CallbackData;
-	//	CallbackData.Run = MeshObjectCallback;
-	//	CallbackData.UserData = (uint64(LocalComponentId.PrimIDValue) & 0xFFFFFFFF) | (uint64(WorldType) << 32);
-	//	SkeletalMesh->MeshObjectCallbackData = CallbackData;*/
-
-	//	if (SkeletalMesh->MeshObject != nullptr)
-	//	{
-	//		UE_LOG(LogPhysicsAsset, Log, TEXT("Skin Cache there : %d %d %d"), SkeletalMesh->MeshObject, SkeletalMesh->MeshObject->GetCachedGeometry().LODIndex, 
-	//				SkeletalMesh->MeshObject->GetCachedGeometry().Sections.Num());
-	//	}
-	//}
 }
 
 void FNDIPhysicsAssetBuffer::ReleaseRHI()
