Fix hair binding asset to support multiple hair instance.

diff --git a/Engine/Plugins/Experimental/HairStrands/Source/HairStrandsCore/Private/GroomAsset.cpp b/Engine/Plugins/Experimental/HairStrands/Source/HairStrandsCore/Private/GroomAsset.cpp
--- a/Engine/Plugins/Experimental/HairStrands/Source/HairStrandsCore/Private/GroomAsset.cpp
+++ b/Engine/Plugins/Experimental/HairStrands/Source/HairStrandsCore/Private/GroomAsset.cpp
@@ -362,24 +362,24 @@ void FHairStrandsClusterCullingResource::ReleaseRHI()
 
 /////////////////////////////////////////////////////////////////////////////////////////
 
-FHairStrandsRootResource::FHairStrandsRootResource()
+FHairStrandsRestRootResource::FHairStrandsRestRootResource()
 {
 
 }
 
-FHairStrandsRootResource::FHairStrandsRootResource(const FHairStrandsRootData& InRootData): 
+FHairStrandsRestRootResource::FHairStrandsRestRootResource(const FHairStrandsRootData& InRootData):
 RootData(InRootData)
 {
 	PopulateFromRootData();
 }
 
-FHairStrandsRootResource::FHairStrandsRootResource(const FHairStrandsDatas* HairStrandsDatas, uint32 LODCount, const TArray<uint32>& NumSamples):
+FHairStrandsRestRootResource::FHairStrandsRestRootResource(const FHairStrandsDatas* HairStrandsDatas, uint32 LODCount, const TArray<uint32>& NumSamples):
 	RootData(HairStrandsDatas, LODCount, NumSamples)
 {
 	PopulateFromRootData();
 }
 
-void FHairStrandsRootResource::PopulateFromRootData()
+void FHairStrandsRestRootResource::PopulateFromRootData()
 {
 	uint32 LODIndex = 0;
 	for (FHairStrandsRootData::FMeshProjectionLOD& MeshProjectionLOD : RootData.MeshProjectionLODs)
@@ -387,12 +387,12 @@ void FHairStrandsRootResource::PopulateFromRootData()
 		FMeshProjectionLOD& LOD = MeshProjectionLODs.AddDefaulted_GetRef();
 
 		LOD.LODIndex = MeshProjectionLOD.LODIndex;
-		LOD.Status = FHairStrandsProjectionHairData::LODData::EStatus::Invalid;
+		LOD.Status = FHairStrandsProjectionHairData::EStatus::Invalid;
 		LOD.SampleCount = MeshProjectionLOD.SampleCount;
 	}
 }
 
-void FHairStrandsRootResource::InitRHI()
+void FHairStrandsRestRootResource::InitRHI()
 {
 	if (RootData.VertexToCurveIndexBuffer.Num() > 0)
 	{
@@ -409,7 +409,7 @@ void FHairStrandsRootResource::InitRHI()
 			const bool bHasValidCPUData = CPUData.RootTriangleBarycentricBuffer.Num() > 0;
 			if (bHasValidCPUData)
 			{
-				GPUData.Status = FHairStrandsProjectionHairData::LODData::EStatus::Completed;
+				GPUData.Status = FHairStrandsProjectionHairData::EStatus::Completed;
 
 				check(CPUData.RootTriangleBarycentricBuffer.Num() > 0);
 				CreateBuffer<FHairStrandsCurveTriangleBarycentricFormat>(CPUData.RootTriangleBarycentricBuffer, GPUData.RootTriangleBarycentricBuffer);
@@ -426,7 +426,7 @@ void FHairStrandsRootResource::InitRHI()
 			}
 			else
 			{
-				GPUData.Status = FHairStrandsProjectionHairData::LODData::EStatus::Initialized;
+				GPUData.Status = FHairStrandsProjectionHairData::EStatus::Initialized;
 
 				CreateBuffer<FHairStrandsCurveTriangleBarycentricFormat>(RootData.RootCount, GPUData.RootTriangleBarycentricBuffer);
 				CreateBuffer<FHairStrandsCurveTriangleIndexFormat>(RootData.RootCount, GPUData.RootTriangleIndexBuffer);
@@ -455,19 +455,11 @@ void FHairStrandsRootResource::InitRHI()
 				CreateBuffer<FHairStrandsIndexFormat>(CPUData.SampleCount, GPUData.MeshSampleIndicesBuffer);
 				CreateBuffer<FHairStrandsMeshTrianglePositionFormat>(CPUData.SampleCount, GPUData.RestSamplePositionsBuffer);
 			}
-			CreateBuffer<FHairStrandsMeshTrianglePositionFormat>(CPUData.SampleCount, GPUData.DeformedSamplePositionsBuffer);
-			CreateBuffer<FHairStrandsMeshTrianglePositionFormat>(CPUData.SampleCount+4, GPUData.MeshSampleWeightsBuffer);
-
-
-			// Strand hair roots translation and rotation in triangle-deformed position relative to the bound triangle 
-			CreateBuffer<FHairStrandsMeshTrianglePositionFormat>(RootData.RootCount, GPUData.DeformedRootTrianglePosition0Buffer);
-			CreateBuffer<FHairStrandsMeshTrianglePositionFormat>(RootData.RootCount, GPUData.DeformedRootTrianglePosition1Buffer);
-			CreateBuffer<FHairStrandsMeshTrianglePositionFormat>(RootData.RootCount, GPUData.DeformedRootTrianglePosition2Buffer);
 		}
 	}
 }
 
-void FHairStrandsRootResource::ReleaseRHI()
+void FHairStrandsRestRootResource::ReleaseRHI()
 {
 	RootPositionBuffer.Release();
 	RootNormalBuffer.Release();
@@ -475,19 +467,67 @@ void FHairStrandsRootResource::ReleaseRHI()
 	
 	for (FMeshProjectionLOD& GPUData : MeshProjectionLODs)
 	{
-		GPUData.Status = FHairStrandsProjectionHairData::LODData::EStatus::Invalid;
+		GPUData.Status = FHairStrandsProjectionHairData::EStatus::Invalid;
 		GPUData.RootTriangleIndexBuffer.Release();
 		GPUData.RootTriangleBarycentricBuffer.Release();
 		GPUData.RestRootTrianglePosition0Buffer.Release();
 		GPUData.RestRootTrianglePosition1Buffer.Release();
 		GPUData.RestRootTrianglePosition2Buffer.Release();
-		GPUData.DeformedRootTrianglePosition0Buffer.Release();
-		GPUData.DeformedRootTrianglePosition1Buffer.Release();
-		GPUData.DeformedRootTrianglePosition2Buffer.Release();
 		GPUData.SampleCount = 0;
 		GPUData.MeshInterpolationWeightsBuffer.Release();
 		GPUData.MeshSampleIndicesBuffer.Release();
 		GPUData.RestSamplePositionsBuffer.Release();
+	}
+	MeshProjectionLODs.Empty();
+}
+
+/////////////////////////////////////////////////////////////////////////////////////////
+
+FHairStrandsDeformedRootResource::FHairStrandsDeformedRootResource()
+{
+
+}
+
+FHairStrandsDeformedRootResource::FHairStrandsDeformedRootResource(const FHairStrandsRestRootResource* InRestResources)
+{
+	check(InRestResources);
+	uint32 LODIndex = 0;
+	RootCount = InRestResources->RootData.RootCount;
+	for (const FHairStrandsRestRootResource::FMeshProjectionLOD& InLOD : InRestResources->MeshProjectionLODs)
+	{
+		FMeshProjectionLOD& LOD = MeshProjectionLODs.AddDefaulted_GetRef();
+
+		LOD.Status = FHairStrandsProjectionHairData::EStatus::Invalid;
+		LOD.LODIndex = InLOD.LODIndex;
+		LOD.SampleCount = InLOD.SampleCount;
+	}
+}
+
+void FHairStrandsDeformedRootResource::InitRHI()
+{
+	if (RootCount > 0)
+	{
+		for (FMeshProjectionLOD& LOD : MeshProjectionLODs)
+		{		
+			LOD.Status = FHairStrandsProjectionHairData::EStatus::Initialized;
+			CreateBuffer<FHairStrandsMeshTrianglePositionFormat>(LOD.SampleCount, LOD.DeformedSamplePositionsBuffer);
+			CreateBuffer<FHairStrandsMeshTrianglePositionFormat>(LOD.SampleCount + 4, LOD.MeshSampleWeightsBuffer);
+
+			CreateBuffer<FHairStrandsMeshTrianglePositionFormat>(RootCount, LOD.DeformedRootTrianglePosition0Buffer);
+			CreateBuffer<FHairStrandsMeshTrianglePositionFormat>(RootCount, LOD.DeformedRootTrianglePosition1Buffer);
+			CreateBuffer<FHairStrandsMeshTrianglePositionFormat>(RootCount, LOD.DeformedRootTrianglePosition2Buffer);
+		}
+	}
+}
+
+void FHairStrandsDeformedRootResource::ReleaseRHI()
+{
+	for (FMeshProjectionLOD& GPUData : MeshProjectionLODs)
+	{
+		GPUData.Status = FHairStrandsProjectionHairData::EStatus::Invalid;
+		GPUData.DeformedRootTrianglePosition0Buffer.Release();
+		GPUData.DeformedRootTrianglePosition1Buffer.Release();
+		GPUData.DeformedRootTrianglePosition2Buffer.Release();
 		GPUData.DeformedSamplePositionsBuffer.Release();
 		GPUData.MeshSampleWeightsBuffer.Release();
 	}
@@ -1321,8 +1361,8 @@ void UGroomBindingAsset::InitResource()
 	for (FHairGroupData& Data : HairGroupDatas)
 	{
 		FHairGroupResource& Resource = HairGroupResources.AddDefaulted_GetRef();
-		Resource.SimRootResources = new FHairStrandsRootResource(Data.SimRootData);
-		Resource.RenRootResources = new FHairStrandsRootResource(Data.RenRootData);
+		Resource.SimRootResources = new FHairStrandsRestRootResource(Data.SimRootData);
+		Resource.RenRootResources = new FHairStrandsRestRootResource(Data.RenRootData);
 
 		BeginInitResource(Resource.SimRootResources);
 		BeginInitResource(Resource.RenRootResources);
@@ -1345,8 +1385,8 @@ void UGroomBindingAsset::ReleaseResource()
 	{
 		for (FHairGroupResource& Resource : HairGroupResources)
 		{
-			FHairStrandsRootResource* InSimRootResources = Resource.SimRootResources;
-			FHairStrandsRootResource* InRenRootResources = Resource.RenRootResources;
+			FHairStrandsRestRootResource* InSimRootResources = Resource.SimRootResources;
+			FHairStrandsRestRootResource* InRenRootResources = Resource.RenRootResources;
 			ENQUEUE_RENDER_COMMAND(ReleaseHairStrandsResourceCommand)(
 				[InSimRootResources, InRenRootResources](FRHICommandList& RHICmdList)
 			{
@@ -1365,8 +1405,8 @@ void UGroomBindingAsset::ReleaseResource()
 	FHairGroupResource ResourceToDelete;
 	while (HairGroupResourcesToDelete.Dequeue(ResourceToDelete))
 	{
-		FHairStrandsRootResource* InSimRootResources = ResourceToDelete.SimRootResources;
-		FHairStrandsRootResource* InRenRootResources = ResourceToDelete.RenRootResources;
+		FHairStrandsRestRootResource* InSimRootResources = ResourceToDelete.SimRootResources;
+		FHairStrandsRestRootResource* InRenRootResources = ResourceToDelete.RenRootResources;
 		ENQUEUE_RENDER_COMMAND(ReleaseHairStrandsResourceCommand)(
 			[InSimRootResources, InRenRootResources](FRHICommandList& RHICmdList)
 		{
diff --git a/Engine/Plugins/Experimental/HairStrands/Source/HairStrandsCore/Private/GroomComponent.cpp b/Engine/Plugins/Experimental/HairStrands/Source/HairStrandsCore/Private/GroomComponent.cpp
--- a/Engine/Plugins/Experimental/HairStrands/Source/HairStrandsCore/Private/GroomComponent.cpp
+++ b/Engine/Plugins/Experimental/HairStrands/Source/HairStrandsCore/Private/GroomComponent.cpp
@@ -853,14 +853,24 @@ FHairStrandsDeformedResource* UGroomComponent::GetGuideStrandsDeformedResource(u
 	return HairGroupResources->HairGroups[GroupIndex].SimDeformedResources;
 }
 
-FHairStrandsRootResource* UGroomComponent::GetGuideStrandsRootResource(uint32 GroupIndex)
+FHairStrandsRestRootResource* UGroomComponent::GetGuideStrandsRestRootResource(uint32 GroupIndex)
 {
 	if (!HairGroupResources || GroupIndex >= uint32(HairGroupResources->HairGroups.Num()))
 	{
 		return nullptr;
 	}
 
-	return HairGroupResources->HairGroups[GroupIndex].SimRootResources;
+	return HairGroupResources->HairGroups[GroupIndex].SimRestRootResources;
+}
+
+FHairStrandsDeformedRootResource* UGroomComponent::GetGuideStrandsDeformedRootResource(uint32 GroupIndex)
+{
+	if (!HairGroupResources || GroupIndex >= uint32(HairGroupResources->HairGroups.Num()))
+	{
+		return nullptr;
+	}
+
+	return HairGroupResources->HairGroups[GroupIndex].SimDeformedRootResources;
 }
 
 template<typename T> void SafeDelete(T*& Data) 
@@ -1005,9 +1015,11 @@ void UGroomComponent::InitResources(bool bIsBindingReloading)
 				check(GroupIt < BindingAsset->HairGroupResources.Num());
 				check(SkeletalMeshComponent->GetNumLODs() == BindingAsset->HairGroupResources[GroupIt].RenRootResources->RootData.MeshProjectionLODs.Num());
 				check(SkeletalMeshComponent->GetNumLODs() == BindingAsset->HairGroupResources[GroupIt].SimRootResources->RootData.MeshProjectionLODs.Num());
+
 				Res.bOwnRootResourceAllocation = false;
-				Res.RenRootResources = BindingAsset->HairGroupResources[GroupIt].RenRootResources;
-				Res.SimRootResources = BindingAsset->HairGroupResources[GroupIt].SimRootResources;
+				Res.RenRestRootResources = BindingAsset->HairGroupResources[GroupIt].RenRootResources;
+				Res.SimRestRootResources = BindingAsset->HairGroupResources[GroupIt].SimRootResources;
+
 				DebugHairGroup.bHasBinding = true;
 			}
 			else
@@ -1017,12 +1029,18 @@ void UGroomComponent::InitResources(bool bIsBindingReloading)
 				{
 					Res.bOwnRootResourceAllocation = true;
 					TArray<uint32> NumSamples; NumSamples.Init(0, LODCount);
-					Res.RenRootResources = new FHairStrandsRootResource(&GroupData.HairRenderData, LODCount, NumSamples);
-					Res.SimRootResources = new FHairStrandsRootResource(&GroupData.HairSimulationData, LODCount, NumSamples);
-					BeginInitResource(Res.RenRootResources);
-					BeginInitResource(Res.SimRootResources);
+					Res.RenRestRootResources = new FHairStrandsRestRootResource(&GroupData.HairRenderData, LODCount, NumSamples);
+					Res.SimRestRootResources = new FHairStrandsRestRootResource(&GroupData.HairSimulationData, LODCount, NumSamples);
+					BeginInitResource(Res.RenRestRootResources);
+					BeginInitResource(Res.SimRestRootResources);
 				}
 			}
+
+			Res.RenDeformedRootResources = new FHairStrandsDeformedRootResource(Res.RenRestRootResources);
+			Res.SimDeformedRootResources = new FHairStrandsDeformedRootResource(Res.SimRestRootResources);
+
+			BeginInitResource(Res.RenDeformedRootResources);
+			BeginInitResource(Res.SimDeformedRootResources);
 		}
 		
 		Res.RenderRestResources = GroupData.HairStrandsRestResource;
@@ -1152,8 +1170,8 @@ void UGroomComponent::InitResources(bool bIsBindingReloading)
 			OutputGroup.RenderGroupAABBBuffer			= &Res.HairGroupPublicDatas->GetGroupAABBBuffer();
 			OutputGroup.ClusterInfoBuffer				= &Res.ClusterCullingResources->ClusterInfoBuffer;
 
-			RenProjectionDatas.HairGroups.Add(ToProjectionHairData(Res.RenRootResources));
-			SimProjectionDatas.HairGroups.Add(ToProjectionHairData(Res.SimRootResources));
+			RenProjectionDatas.HairGroups.Add(ToProjectionHairData(Res.RenRestRootResources, Res.RenDeformedRootResources));
+			SimProjectionDatas.HairGroups.Add(ToProjectionHairData(Res.SimRestRootResources, Res.SimDeformedRootResources));
 
 			FHairStrandsPrimitiveResources::FHairGroup& Group = PrimitiveResources.Groups.AddDefaulted_GetRef();
 			Group.ClusterAABBBuffer = &Res.HairGroupPublicDatas->GetClusterAABBBuffer();
@@ -1242,9 +1260,11 @@ void UGroomComponent::DeleteHairGroupResources(FHairGroupResources*& InHairGroup
 			// Release the root resources only if they have been created internally (vs. being created by external asset)
 			if (Res.bOwnRootResourceAllocation)
 			{
-				SafeRelease(Res.RenRootResources);
-				SafeRelease(Res.SimRootResources);
+				SafeRelease(Res.RenRestRootResources);
+				SafeRelease(Res.SimRestRootResources);
 			}
+			SafeRelease(Res.RenDeformedRootResources);
+			SafeRelease(Res.SimDeformedRootResources);
 
 			SafeRelease(Res.RenderDeformedResources);
 			SafeRelease(Res.HairGroupPublicDatas);
diff --git a/Engine/Plugins/Experimental/HairStrands/Source/HairStrandsCore/Private/HairStrandsRendering.cpp b/Engine/Plugins/Experimental/HairStrands/Source/HairStrandsCore/Private/HairStrandsRendering.cpp
--- a/Engine/Plugins/Experimental/HairStrands/Source/HairStrandsCore/Private/HairStrandsRendering.cpp
+++ b/Engine/Plugins/Experimental/HairStrands/Source/HairStrandsCore/Private/HairStrandsRendering.cpp
@@ -170,7 +170,8 @@ static void AddDeformSimHairStrandsPass(
 	FRHIShaderResourceView* SimRootIndexBuffer,
 	FRHIUnorderedAccessView* OutSimDeformedPositionBuffer,
 	FVector& SimRestOffset,
-	FVector& SimDeformedOffset)
+	FVector& SimDeformedOffset,
+	FBufferTransitionQueue& OutTransitionQueue)
 {
 	static uint32 IterationCount = 0;
 	++IterationCount;
@@ -209,30 +210,38 @@ static void AddDeformSimHairStrandsPass(
 			Parameters->SimVertexToRootIndexBuffer = InSimHairData.VertexToCurveIndexBuffer->SRV;
 		}
 
-		const bool bSupportDynamicMesh = InSimHairData.RootCount > 0 && LODIndex >= 0 && LODIndex < uint32(InSimHairData.LODDatas.Num()) && InSimHairData.LODDatas[LODIndex].bIsValid && bIsVertexToCurveBuffersValid;
+		const bool bSupportDynamicMesh = 
+			InSimHairData.RootCount > 0 && 
+			LODIndex >= 0 && 
+			LODIndex < uint32(InSimHairData.RestLODDatas.Num()) && 
+			LODIndex < uint32(InSimHairData.DeformedLODDatas.Num()) &&
+			InSimHairData.RestLODDatas[LODIndex].IsValid() &&
+			InSimHairData.DeformedLODDatas[LODIndex].IsValid() &&
+			bIsVertexToCurveBuffersValid;
 		
 		bool bSupportGlobalInterpolation = false;
 		if (bSupportDynamicMesh)
 		{
-			bSupportGlobalInterpolation = (GStrandHairInterpolationGlobal == 1) && (InSimHairData.LODDatas[LODIndex].SampleCount > 0);
-			if (!bSupportGlobalInterpolation) {
+			bSupportGlobalInterpolation = (GStrandHairInterpolationGlobal == 1) && (InSimHairData.RestLODDatas[LODIndex].SampleCount > 0);
+			if (!bSupportGlobalInterpolation) 
+			{
 				InternalDeformationType = 4;
-				Parameters->SimRestPosition0Buffer = InSimHairData.LODDatas[LODIndex].RestRootTrianglePosition0Buffer->SRV;
-				Parameters->SimRestPosition1Buffer = InSimHairData.LODDatas[LODIndex].RestRootTrianglePosition1Buffer->SRV;
-				Parameters->SimRestPosition2Buffer = InSimHairData.LODDatas[LODIndex].RestRootTrianglePosition2Buffer->SRV;
+				Parameters->SimRestPosition0Buffer = InSimHairData.RestLODDatas[LODIndex].RestRootTrianglePosition0Buffer->SRV;
+				Parameters->SimRestPosition1Buffer = InSimHairData.RestLODDatas[LODIndex].RestRootTrianglePosition1Buffer->SRV;
+				Parameters->SimRestPosition2Buffer = InSimHairData.RestLODDatas[LODIndex].RestRootTrianglePosition2Buffer->SRV;
 
-				Parameters->SimDeformedPosition0Buffer = InSimHairData.LODDatas[LODIndex].DeformedRootTrianglePosition0Buffer->SRV;
-				Parameters->SimDeformedPosition1Buffer = InSimHairData.LODDatas[LODIndex].DeformedRootTrianglePosition1Buffer->SRV;
-				Parameters->SimDeformedPosition2Buffer = InSimHairData.LODDatas[LODIndex].DeformedRootTrianglePosition2Buffer->SRV;
+				Parameters->SimDeformedPosition0Buffer = InSimHairData.DeformedLODDatas[LODIndex].DeformedRootTrianglePosition0Buffer->SRV;
+				Parameters->SimDeformedPosition1Buffer = InSimHairData.DeformedLODDatas[LODIndex].DeformedRootTrianglePosition1Buffer->SRV;
+				Parameters->SimDeformedPosition2Buffer = InSimHairData.DeformedLODDatas[LODIndex].DeformedRootTrianglePosition2Buffer->SRV;
 
-				Parameters->SimRootBarycentricBuffer = InSimHairData.LODDatas[LODIndex].RootTriangleBarycentricBuffer->SRV;
+				Parameters->SimRootBarycentricBuffer = InSimHairData.RestLODDatas[LODIndex].RootTriangleBarycentricBuffer->SRV;
 			}
 			else
 			{
 				InternalDeformationType = 5;
-				Parameters->MeshSampleWeightsBuffer = InSimHairData.LODDatas[LODIndex].MeshSampleWeightsBuffer->SRV;
-				Parameters->RestSamplePositionsBuffer = InSimHairData.LODDatas[LODIndex].RestSamplePositionsBuffer->SRV;
-				Parameters->SampleCount = InSimHairData.LODDatas[LODIndex].SampleCount;
+				Parameters->MeshSampleWeightsBuffer = InSimHairData.DeformedLODDatas[LODIndex].MeshSampleWeightsBuffer->SRV;
+				Parameters->RestSamplePositionsBuffer = InSimHairData.RestLODDatas[LODIndex].RestSamplePositionsBuffer->SRV;
+				Parameters->SampleCount = InSimHairData.RestLODDatas[LODIndex].SampleCount;
 			}
 		}
 	}
@@ -250,6 +259,8 @@ static void AddDeformSimHairStrandsPass(
 		ComputeShader,
 		Parameters,
 		FIntVector(DispatchCountX, DispatchCountY, 1));
+
+	OutTransitionQueue.Add(OutSimDeformedPositionBuffer);
 }
 
 ///////////////////////////////////////////////////////////////////////////////////////////////////
@@ -300,7 +311,6 @@ class FHairInterpolationCS : public FGlobalShader
 		SHADER_PARAMETER_SRV(Buffer, RenderRestPosePositionBuffer)
 		SHADER_PARAMETER_UAV(RWBuffer, OutRenderDeformedPositionBuffer)
 
-		SHADER_PARAMETER_UAV(RWBuffer, OutClusterBuffer)
 		SHADER_PARAMETER_SRV(Buffer, VertexToClusterIdBuffer)
 
 		SHADER_PARAMETER_SRV(Buffer, SimRestPosePositionBuffer)
@@ -367,7 +377,8 @@ static void AddHairStrandsInterpolationPass(
 	const FUnorderedAccessViewRHIRef& OutRenderPositionBuffer,
 	const FUnorderedAccessViewRHIRef& OutRenderAttributeBuffer,
 	const FShaderResourceViewRHIRef& VertexToClusterIdBuffer,
-	const FShaderResourceViewRHIRef& SimRootPointIndexBuffer )
+	const FShaderResourceViewRHIRef& SimRootPointIndexBuffer,
+	FBufferTransitionQueue& OutTransitionQueue)
 {
 	const uint32 GroupSize = ComputeGroupSize();
 	const FIntVector DispatchCount = ComputeDispatchCount(VertexCount, GroupSize);
@@ -449,34 +460,44 @@ static void AddHairStrandsInterpolationPass(
 		{
 			Parameters->SimAttributeBuffer = SimAttributeBuffer;
 			Parameters->OutRenderAttributeBuffer = OutRenderAttributeBuffer;
+			OutTransitionQueue.Add(OutRenderAttributeBuffer);
 		}
 	}
 
-	const bool bSupportDynamicMesh = InRenHairData.RootCount > 0 && LODIndex >= 0 && LODIndex < InRenHairData.LODDatas.Num() && InRenHairData.LODDatas[LODIndex].bIsValid && bIsVertexToCurveBuffersValid;
+	const bool bSupportDynamicMesh = 
+		InRenHairData.RootCount > 0 && 
+		LODIndex >= 0 && 
+		LODIndex < InRenHairData.RestLODDatas.Num() && 
+		LODIndex < InRenHairData.DeformedLODDatas.Num() &&
+		InRenHairData.RestLODDatas[LODIndex].IsValid() &&
+		InRenHairData.DeformedLODDatas[LODIndex].IsValid() &&
+		bIsVertexToCurveBuffersValid;
+
 	bool bSupportGlobalInterpolation = false;
 	if (bSupportDynamicMesh)
 	{
-		bSupportGlobalInterpolation = (GStrandHairInterpolationGlobal == 1) && (InSimHairData.LODDatas[LODIndex].SampleCount > 0);
+		bSupportGlobalInterpolation = (GStrandHairInterpolationGlobal == 1) && (InSimHairData.RestLODDatas[LODIndex].SampleCount > 0);
 		{
-			Parameters->RestPosition0Buffer = InRenHairData.LODDatas[LODIndex].RestRootTrianglePosition0Buffer->SRV;
-			Parameters->RestPosition1Buffer = InRenHairData.LODDatas[LODIndex].RestRootTrianglePosition1Buffer->SRV;
-			Parameters->RestPosition2Buffer = InRenHairData.LODDatas[LODIndex].RestRootTrianglePosition2Buffer->SRV;
-
-			Parameters->DeformedPosition0Buffer = InRenHairData.LODDatas[LODIndex].DeformedRootTrianglePosition0Buffer->SRV;
-			Parameters->DeformedPosition1Buffer = InRenHairData.LODDatas[LODIndex].DeformedRootTrianglePosition1Buffer->SRV;
-			Parameters->DeformedPosition2Buffer = InRenHairData.LODDatas[LODIndex].DeformedRootTrianglePosition2Buffer->SRV;
+			Parameters->RestPosition0Buffer = InRenHairData.RestLODDatas[LODIndex].RestRootTrianglePosition0Buffer->SRV;
+			Parameters->RestPosition1Buffer = InRenHairData.RestLODDatas[LODIndex].RestRootTrianglePosition1Buffer->SRV;
+			Parameters->RestPosition2Buffer = InRenHairData.RestLODDatas[LODIndex].RestRootTrianglePosition2Buffer->SRV;
 
-			Parameters->RootBarycentricBuffer = InRenHairData.LODDatas[LODIndex].RootTriangleBarycentricBuffer->SRV;
+			Parameters->RootBarycentricBuffer = InRenHairData.RestLODDatas[LODIndex].RootTriangleBarycentricBuffer->SRV;
 
-			Parameters->SimRestPosition0Buffer = InSimHairData.LODDatas[LODIndex].RestRootTrianglePosition0Buffer->SRV;
-			Parameters->SimRestPosition1Buffer = InSimHairData.LODDatas[LODIndex].RestRootTrianglePosition1Buffer->SRV;
-			Parameters->SimRestPosition2Buffer = InSimHairData.LODDatas[LODIndex].RestRootTrianglePosition2Buffer->SRV;
+			Parameters->SimRestPosition0Buffer = InSimHairData.RestLODDatas[LODIndex].RestRootTrianglePosition0Buffer->SRV;
+			Parameters->SimRestPosition1Buffer = InSimHairData.RestLODDatas[LODIndex].RestRootTrianglePosition1Buffer->SRV;
+			Parameters->SimRestPosition2Buffer = InSimHairData.RestLODDatas[LODIndex].RestRootTrianglePosition2Buffer->SRV;
 
-			Parameters->SimDeformedPosition0Buffer = InSimHairData.LODDatas[LODIndex].DeformedRootTrianglePosition0Buffer->SRV;
-			Parameters->SimDeformedPosition1Buffer = InSimHairData.LODDatas[LODIndex].DeformedRootTrianglePosition1Buffer->SRV;
-			Parameters->SimDeformedPosition2Buffer = InSimHairData.LODDatas[LODIndex].DeformedRootTrianglePosition2Buffer->SRV;
+			Parameters->SimRootBarycentricBuffer = InSimHairData.RestLODDatas[LODIndex].RootTriangleBarycentricBuffer->SRV;
+		}
+		{
+			Parameters->DeformedPosition0Buffer = InRenHairData.DeformedLODDatas[LODIndex].DeformedRootTrianglePosition0Buffer->SRV;
+			Parameters->DeformedPosition1Buffer = InRenHairData.DeformedLODDatas[LODIndex].DeformedRootTrianglePosition1Buffer->SRV;
+			Parameters->DeformedPosition2Buffer = InRenHairData.DeformedLODDatas[LODIndex].DeformedRootTrianglePosition2Buffer->SRV;
 
-			Parameters->SimRootBarycentricBuffer = InSimHairData.LODDatas[LODIndex].RootTriangleBarycentricBuffer->SRV;
+			Parameters->SimDeformedPosition0Buffer = InSimHairData.DeformedLODDatas[LODIndex].DeformedRootTrianglePosition0Buffer->SRV;
+			Parameters->SimDeformedPosition1Buffer = InSimHairData.DeformedLODDatas[LODIndex].DeformedRootTrianglePosition1Buffer->SRV;
+			Parameters->SimDeformedPosition2Buffer = InSimHairData.DeformedLODDatas[LODIndex].DeformedRootTrianglePosition2Buffer->SRV;
 		}
 	}
 
@@ -502,6 +523,8 @@ static void AddHairStrandsInterpolationPass(
 		ComputeShader,
 		Parameters,
 		DispatchCount);
+
+	OutTransitionQueue.Add(OutRenderPositionBuffer);	
 }
 
 ///////////////////////////////////////////////////////////////////////////////////////////////////
@@ -538,7 +561,8 @@ static void AddHairClusterAABBPass(
 	const FHairStrandsProjectionHairData::HairGroup& InRenHairData,
 	const FVector& OutHairWorldOffset,
 	FHairStrandClusterData::FHairGroup& ClusterData,
-	const FShaderResourceViewRHIRef& RenderPositionBuffer)
+	const FShaderResourceViewRHIRef& RenderPositionBuffer,
+	FBufferTransitionQueue& OutTransitionQueue)
 {
 	const uint32 GroupSize = ComputeGroupSize();
 	const FIntVector DispatchCount = ComputeDispatchGroupCount2D(ClusterData.ClusterCount);
@@ -564,6 +588,9 @@ static void AddHairClusterAABBPass(
 		ComputeShader,
 		Parameters,
 		DispatchCount);
+
+	OutTransitionQueue.Add(Parameters->OutClusterAABBBuffer);
+	OutTransitionQueue.Add(Parameters->OutGroupAABBBuffer);
 }
 
 ///////////////////////////////////////////////////////////////////////////////////////////////////
@@ -593,7 +620,8 @@ static void AddHairTangentPass(
 	FRDGBuilder& GraphBuilder,
 	uint32 VertexCount,
 	const FShaderResourceViewRHIRef& PositionBuffer,
-	const FUnorderedAccessViewRHIRef& OutTangentBuffer)
+	const FUnorderedAccessViewRHIRef& OutTangentBuffer,
+	FBufferTransitionQueue& OutTransitionQueue)
 {
 	const uint32 GroupSize = ComputeGroupSize();
 	const FIntVector DispatchCount = ComputeDispatchCount(VertexCount, GroupSize);
@@ -616,6 +644,8 @@ static void AddHairTangentPass(
 		ComputeShader,
 		Parameters,
 		DispatchCount);
+
+	OutTransitionQueue.Add(OutTangentBuffer);
 }
 
 ///////////////////////////////////////////////////////////////////////////////////////////////////
@@ -649,7 +679,8 @@ static void AddGenerateRaytracingGeometryPass(
 	float HairRadius,
 	const FVector& HairWorldOffset,
 	const FShaderResourceViewRHIRef& PositionBuffer,
-	const FUnorderedAccessViewRHIRef& OutPositionBuffer)
+	const FUnorderedAccessViewRHIRef& OutPositionBuffer,
+	FBufferTransitionQueue& OutTransitionQueue)
 {
 	const uint32 GroupSize = ComputeGroupSize();
 	const FIntVector DispatchCount = ComputeDispatchCount(VertexCount, GroupSize);
@@ -674,6 +705,8 @@ static void AddGenerateRaytracingGeometryPass(
 		ComputeShader,
 		Parameters,
 		DispatchCount);
+
+	OutTransitionQueue.Add(OutPositionBuffer);
 }
 
 ///////////////////////////////////////////////////////////////////////////////////////////////////
@@ -704,7 +737,8 @@ static void AddClearClusterAABBPass(
 	FRDGBuilder& GraphBuilder,
 	uint32 ClusterCount,
 	FRHIUnorderedAccessView* OutClusterAABBuffer,
-	FRHIUnorderedAccessView* OutGroupAABBuffer)
+	FRHIUnorderedAccessView* OutGroupAABBuffer,
+	FBufferTransitionQueue& OutTransitionQueue)
 {
 	check(OutClusterAABBuffer);
 
@@ -724,6 +758,9 @@ static void AddClearClusterAABBPass(
 		ComputeShader,
 		Parameters,
 		DispatchCount);
+
+	OutTransitionQueue.Add(OutClusterAABBuffer);
+	OutTransitionQueue.Add(OutGroupAABBuffer);
 }
 
 ///////////////////////////////////////////////////////////////////////////////////////////////////
@@ -811,8 +848,6 @@ void ComputeHairStrandsInterpolation(
 		}
 	}
 
-	FRDGBuilder GraphBuilder(RHICmdList);
-
 	DECLARE_GPU_STAT(HairStrandsInterpolationCluster);
 	SCOPED_DRAW_EVENT(RHICmdList, HairStrandsInterpolationCluster);
 	SCOPED_GPU_STAT(RHICmdList, HairStrandsInterpolationCluster);
@@ -828,6 +863,8 @@ void ComputeHairStrandsInterpolation(
 
 	if (DeformationType != EDeformationType::RestStrands && DeformationType != EDeformationType::Simulation)
 	{
+		FBufferTransitionQueue TransitionQueue;
+		FRDGBuilder GraphBuilder(RHICmdList);
 		for (uint32 GroupIndex = 0; GroupIndex < GroupCount; ++GroupIndex)
 		{
 			FHairStrandsInterpolationInput::FHairGroup& Input = InInput->HairGroups[GroupIndex];
@@ -845,12 +882,17 @@ void ComputeHairStrandsInterpolation(
 				Input.SimRestPosePositionBuffer->SRV,
 				Input.SimRootPointIndexBuffer ? Input.SimRootPointIndexBuffer->SRV : nullptr,
 				Output.SimDeformedPositionBuffer[CurrIndex]->UAV, Input.InSimHairPositionOffset,
-				Input.OutHairPositionOffset);
+				Input.OutHairPositionOffset,
+				TransitionQueue);
 		}
+		GraphBuilder.Execute();
+		TransitBufferToReadable(RHICmdList, TransitionQueue);
 	}
 
 	if (DebugMode == EHairStrandsDebugMode::SimHairStrands)
 	{
+		FBufferTransitionQueue TransitionQueue;
+		FRDGBuilder GraphBuilder(RHICmdList);
 		for (uint32 GroupIndex = 0; GroupIndex < GroupCount; ++GroupIndex)
 		{
 			FHairStrandsInterpolationInput::FHairGroup& Input = InInput->HairGroups[GroupIndex];
@@ -863,7 +905,8 @@ void ComputeHairStrandsInterpolation(
 				GraphBuilder,
 				Input.SimVertexCount,
 				Output.SimDeformedPositionBuffer[CurrIndex]->SRV,
-				Output.SimTangentBuffer->UAV);
+				Output.SimTangentBuffer->UAV,
+				TransitionQueue);
 
 			const bool bHasSimulationEnabled = Input.bIsSimulationEnable && GHairStrandsInterpolateSimulation && DeformationType != EDeformationType::RestStrands;
 
@@ -881,217 +924,255 @@ void ComputeHairStrandsInterpolation(
 			Output.VFInput.bUseStableRasterization = Input.GroupDesc.bUseStableRasterization;
 			Output.VFInput.bScatterSceneLighting = Input.GroupDesc.bScatterSceneLighting;
 		}
+		GraphBuilder.Execute();
+		TransitBufferToReadable(RHICmdList, TransitionQueue);
 	}
 	else
 	{
-		for (uint32 GroupIndex = 0; GroupIndex < GroupCount; ++GroupIndex)
 		{
-			FHairStrandsInterpolationInput::FHairGroup& Input = InInput->HairGroups[GroupIndex];
-			FHairStrandsInterpolationOutput::HairGroup& Output = InOutput->HairGroups[GroupIndex];
-
-			const uint32 CurrIndex = *Output.CurrentIndex;
-			const uint32 PrevIndex = (CurrIndex + 1) % 2;
-
-			// If the deformation is driven by the physics simulation, then the output is always the 0 index
-			const uint32 SimIndex = CurrIndex;// GHairDeformationType == 0 ? 0 : CurrIndex;
-
-
-			check(ClusterData);
-
-			const uint32 BufferSizeInBytes = Input.RenderVertexCount * FHairStrandsAttributeFormat::SizeInByte;
-			if (bDebugModePatchedAttributeBuffer && Output.RenderPatchedAttributeBuffer.NumBytes != BufferSizeInBytes)
+			FBufferTransitionQueue TransitionQueue;
+			FRDGBuilder GraphBuilder(RHICmdList);
+			for (uint32 GroupIndex = 0; GroupIndex < GroupCount; ++GroupIndex)
 			{
-				Output.RenderPatchedAttributeBuffer.Release();
-				Output.RenderPatchedAttributeBuffer.Initialize(FHairStrandsAttributeFormat::SizeInByte, Input.RenderVertexCount, FHairStrandsAttributeFormat::Format, BUF_Static);
-			}
-
-			AddClearClusterAABBPass(
-				GraphBuilder,
-				Input.ClusterCount,
-				Output.RenderClusterAABBBuffer->UAV,
-				Output.RenderGroupAABBBuffer->UAV);
-		}
+				FHairStrandsInterpolationInput::FHairGroup& Input = InInput->HairGroups[GroupIndex];
+				FHairStrandsInterpolationOutput::HairGroup& Output = InOutput->HairGroups[GroupIndex];
 
-		for (uint32 GroupIndex = 0; GroupIndex < GroupCount; ++GroupIndex)
-		{
-			FHairStrandsInterpolationInput::FHairGroup& Input = InInput->HairGroups[GroupIndex];
-			FHairStrandsInterpolationOutput::HairGroup& Output = InOutput->HairGroups[GroupIndex];
+				const uint32 CurrIndex = *Output.CurrentIndex;
+				const uint32 PrevIndex = (CurrIndex + 1) % 2;
 
-			const uint32 CurrIndex = *Output.CurrentIndex;
-			const uint32 PrevIndex = (CurrIndex + 1) % 2;
+				// If the deformation is driven by the physics simulation, then the output is always the 0 index
+				const uint32 SimIndex = CurrIndex;// GHairDeformationType == 0 ? 0 : CurrIndex;
 
-			// If the deformation is driven by the physics simulation, then the output is always the 0 index
-			const uint32 SimIndex = CurrIndex;// GHairDeformationType == 0 ? 0 : CurrIndex;
 
-			FHairScaleAndClipDesc ScaleAndClipDesc;
-			ScaleAndClipDesc.InHairLength = Input.GroupDesc.HairLength;
-			ScaleAndClipDesc.InHairRadius = Input.GroupDesc.HairWidth * 0.5f;
-			ScaleAndClipDesc.OutHairRadius = (GStrandHairWidth > 0 ? GStrandHairWidth : Input.GroupDesc.HairWidth) * 0.5f;
-			ScaleAndClipDesc.MaxOutHairRadius = ScaleAndClipDesc.OutHairRadius * FMath::Max(1.f, FMath::Max(Input.GroupDesc.HairRootScale, Input.GroupDesc.HairTipScale));
-			ScaleAndClipDesc.HairRadiusRootScale = Input.GroupDesc.HairRootScale;
-			ScaleAndClipDesc.HairRadiusTipScale = Input.GroupDesc.HairTipScale;
-			ScaleAndClipDesc.HairLengthClip = Input.GroupDesc.HairClipLength;
+				check(ClusterData);
 
-			const bool bHasSimulationEnabled = Input.bIsSimulationEnable && GHairStrandsInterpolateSimulation && DeformationType != EDeformationType::RestStrands;
-			check(GroupIndex < uint32(InRenHairDatas.HairGroups.Num()));
-			check(GroupIndex < uint32(InSimHairDatas.HairGroups.Num()));
-			AddHairStrandsInterpolationPass(
-				GraphBuilder,
-				ShaderDrawData,
-				InRenHairDatas.HairGroups[GroupIndex],
-				InSimHairDatas.HairGroups[GroupIndex],
-				Input.InRenderHairPositionOffset,
-				Input.InSimHairPositionOffset,
-				Input.OutHairPositionOffset,
-				ScaleAndClipDesc,
-				LODIndex,
-				bHasSimulationEnabled,
-				bDebugModePatchedAttributeBuffer,
-				Input.RenderVertexCount,
-				Input.RenderRestPosePositionBuffer->SRV,
-				Input.RenderAttributeBuffer->SRV,
-				Input.Interpolation0Buffer->SRV,
-				Input.Interpolation1Buffer->SRV,
-				Input.SimRestPosePositionBuffer->SRV,
-				Output.SimDeformedPositionBuffer[SimIndex]->SRV,
-				Input.SimAttributeBuffer->SRV,
-				Output.RenderDeformedPositionBuffer[CurrIndex]->UAV,
-				Output.RenderPatchedAttributeBuffer.UAV,
-				Input.VertexToClusterIdBuffer->SRV,
-				Input.SimRootPointIndexBuffer->SRV);
+				const uint32 BufferSizeInBytes = Input.RenderVertexCount * FHairStrandsAttributeFormat::SizeInByte;
+				if (bDebugModePatchedAttributeBuffer && Output.RenderPatchedAttributeBuffer.NumBytes != BufferSizeInBytes)
+				{
+					Output.RenderPatchedAttributeBuffer.Release();
+					Output.RenderPatchedAttributeBuffer.Initialize(FHairStrandsAttributeFormat::SizeInByte, Input.RenderVertexCount, FHairStrandsAttributeFormat::Format, BUF_Static);
+				}
 
+				AddClearClusterAABBPass(
+					GraphBuilder,
+					Input.ClusterCount,
+					Output.RenderClusterAABBBuffer->UAV,
+					Output.RenderGroupAABBBuffer->UAV,
+					TransitionQueue);
+			}
+			GraphBuilder.Execute();
+			TransitBufferToReadable(RHICmdList, TransitionQueue);
 		}
 
-
-		for (uint32 GroupIndex = 0; GroupIndex < GroupCount; ++GroupIndex)
 		{
-			FHairStrandsInterpolationInput::FHairGroup& Input = InInput->HairGroups[GroupIndex];
-			FHairStrandsInterpolationOutput::HairGroup& Output = InOutput->HairGroups[GroupIndex];
-
-			const uint32 CurrIndex = *Output.CurrentIndex;
-			const uint32 PrevIndex = (CurrIndex + 1) % 2;
-
-			// If the deformation is driven by the physics simulation, then the output is always the 0 index
-			const uint32 SimIndex = CurrIndex;// GHairDeformationType == 0 ? 0 : CurrIndex;
-
-			// Initialize group cluster data for culling by the renderer
-			FHairStrandClusterData::FHairGroup& HairGroupCluster = ClusterData->HairGroups.Emplace_GetRef();
-			HairGroupCluster.ClusterCount = Input.ClusterCount;
-			HairGroupCluster.VertexCount = Input.ClusterVertexCount;
-			HairGroupCluster.GroupAABBBuffer = Output.RenderGroupAABBBuffer;
-			HairGroupCluster.ClusterAABBBuffer = Output.RenderClusterAABBBuffer;
-			HairGroupCluster.ClusterInfoBuffer = Output.ClusterInfoBuffer;
-			HairGroupCluster.VertexToClusterIdBuffer = Input.VertexToClusterIdBuffer;
-			HairGroupCluster.ClusterVertexIdBuffer = Input.ClusterVertexIdBuffer;
-			HairGroupCluster.ClusterIndexRadiusScaleInfoBuffer = Input.ClusterIndexRadiusScaleInfoBuffer;
-			HairGroupCluster.HairGroupPublicPtr = Output.HairGroupPublicData;
-			HairGroupCluster.LodBias = Input.GroupDesc.LodBias;
-			HairGroupCluster.LodAverageVertexPerPixel = Input.GroupDesc.LodAverageVertexPerPixel;
-
-			// Note: This code needs to exactly match the values FHairScaleAndClipDesc set int the previous loop.
-			const float OutHairRadius = (GStrandHairWidth > 0 ? GStrandHairWidth : Input.GroupDesc.HairWidth) * 0.5f;
-			const float MaxOutHairRadius = OutHairRadius * FMath::Max(1.f, FMath::Max(Input.GroupDesc.HairRootScale, Input.GroupDesc.HairTipScale));
-
-			Output.VFInput.HairRadius = MaxOutHairRadius;
-			Output.VFInput.HairLength = Input.GroupDesc.HairLength;
-			Output.VFInput.HairDensity = Input.GroupDesc.HairShadowDensity;
-			Output.VFInput.HairPositionBuffer = Output.RenderDeformedPositionBuffer[CurrIndex]->SRV;
-			Output.VFInput.HairPreviousPositionBuffer = Output.RenderDeformedPositionBuffer[PrevIndex]->SRV;
-			Output.VFInput.bUseStableRasterization = Input.GroupDesc.bUseStableRasterization;
-			Output.VFInput.bScatterSceneLighting = Input.GroupDesc.bScatterSceneLighting;
+			FBufferTransitionQueue TransitionQueue;
+			FRDGBuilder GraphBuilder(RHICmdList);
+			for (uint32 GroupIndex = 0; GroupIndex < GroupCount; ++GroupIndex)
+			{
+				FHairStrandsInterpolationInput::FHairGroup& Input = InInput->HairGroups[GroupIndex];
+				FHairStrandsInterpolationOutput::HairGroup& Output = InOutput->HairGroups[GroupIndex];
+
+				const uint32 CurrIndex = *Output.CurrentIndex;
+				const uint32 PrevIndex = (CurrIndex + 1) % 2;
+
+				// If the deformation is driven by the physics simulation, then the output is always the 0 index
+				const uint32 SimIndex = CurrIndex;// GHairDeformationType == 0 ? 0 : CurrIndex;
+
+				FHairScaleAndClipDesc ScaleAndClipDesc;
+				ScaleAndClipDesc.InHairLength = Input.GroupDesc.HairLength;
+				ScaleAndClipDesc.InHairRadius = Input.GroupDesc.HairWidth * 0.5f;
+				ScaleAndClipDesc.OutHairRadius = (GStrandHairWidth > 0 ? GStrandHairWidth : Input.GroupDesc.HairWidth) * 0.5f;
+				ScaleAndClipDesc.MaxOutHairRadius = ScaleAndClipDesc.OutHairRadius * FMath::Max(1.f, FMath::Max(Input.GroupDesc.HairRootScale, Input.GroupDesc.HairTipScale));
+				ScaleAndClipDesc.HairRadiusRootScale = Input.GroupDesc.HairRootScale;
+				ScaleAndClipDesc.HairRadiusTipScale = Input.GroupDesc.HairTipScale;
+				ScaleAndClipDesc.HairLengthClip = Input.GroupDesc.HairClipLength;
+
+				const bool bHasSimulationEnabled = Input.bIsSimulationEnable && GHairStrandsInterpolateSimulation && DeformationType != EDeformationType::RestStrands;
+				check(GroupIndex < uint32(InRenHairDatas.HairGroups.Num()));
+				check(GroupIndex < uint32(InSimHairDatas.HairGroups.Num()));
+				AddHairStrandsInterpolationPass(
+					GraphBuilder,
+					ShaderDrawData,
+					InRenHairDatas.HairGroups[GroupIndex],
+					InSimHairDatas.HairGroups[GroupIndex],
+					Input.InRenderHairPositionOffset,
+					Input.InSimHairPositionOffset,
+					Input.OutHairPositionOffset,
+					ScaleAndClipDesc,
+					LODIndex,
+					bHasSimulationEnabled,
+					bDebugModePatchedAttributeBuffer,
+					Input.RenderVertexCount,
+					Input.RenderRestPosePositionBuffer->SRV,
+					Input.RenderAttributeBuffer->SRV,
+					Input.Interpolation0Buffer->SRV,
+					Input.Interpolation1Buffer->SRV,
+					Input.SimRestPosePositionBuffer->SRV,
+					Output.SimDeformedPositionBuffer[SimIndex]->SRV,
+					Input.SimAttributeBuffer->SRV,
+					Output.RenderDeformedPositionBuffer[CurrIndex]->UAV,
+					Output.RenderPatchedAttributeBuffer.UAV,
+					Input.VertexToClusterIdBuffer->SRV,
+					Input.SimRootPointIndexBuffer->SRV,
+					TransitionQueue);
 
-			AddHairClusterAABBPass(
-				GraphBuilder,
-				InRenHairDatas.HairGroups[GroupIndex],
-				Input.OutHairPositionOffset,
-				HairGroupCluster,
-				Output.RenderDeformedPositionBuffer[CurrIndex]->SRV);
+			}
+			GraphBuilder.Execute();
+			TransitBufferToReadable(RHICmdList, TransitionQueue);
 		}
 
-		for (uint32 GroupIndex = 0; GroupIndex < GroupCount; ++GroupIndex)
-		{
-			FHairStrandsInterpolationInput::FHairGroup& Input = InInput->HairGroups[GroupIndex];
-			FHairStrandsInterpolationOutput::HairGroup& Output = InOutput->HairGroups[GroupIndex];
-
-			const uint32 CurrIndex = *Output.CurrentIndex;
-			const uint32 PrevIndex = (CurrIndex + 1) % 2;
-
-			AddHairTangentPass(
-				GraphBuilder,
-				Input.RenderVertexCount,
-				Output.VFInput.HairPositionBuffer,
-				Output.RenderTangentBuffer->UAV);
+		{		
+			FBufferTransitionQueue TransitionQueue;
+			FRDGBuilder GraphBuilder(RHICmdList);
+			for (uint32 GroupIndex = 0; GroupIndex < GroupCount; ++GroupIndex)
+			{
+				FHairStrandsInterpolationInput::FHairGroup& Input = InInput->HairGroups[GroupIndex];
+				FHairStrandsInterpolationOutput::HairGroup& Output = InOutput->HairGroups[GroupIndex];
+
+				const uint32 CurrIndex = *Output.CurrentIndex;
+				const uint32 PrevIndex = (CurrIndex + 1) % 2;
+
+				// If the deformation is driven by the physics simulation, then the output is always the 0 index
+				const uint32 SimIndex = CurrIndex;// GHairDeformationType == 0 ? 0 : CurrIndex;
+
+				// Initialize group cluster data for culling by the renderer
+				FHairStrandClusterData::FHairGroup& HairGroupCluster = ClusterData->HairGroups.Emplace_GetRef();
+				HairGroupCluster.ClusterCount = Input.ClusterCount;
+				HairGroupCluster.VertexCount = Input.ClusterVertexCount;
+				HairGroupCluster.GroupAABBBuffer = Output.RenderGroupAABBBuffer;
+				HairGroupCluster.ClusterAABBBuffer = Output.RenderClusterAABBBuffer;
+				HairGroupCluster.ClusterInfoBuffer = Output.ClusterInfoBuffer;
+				HairGroupCluster.VertexToClusterIdBuffer = Input.VertexToClusterIdBuffer;
+				HairGroupCluster.ClusterVertexIdBuffer = Input.ClusterVertexIdBuffer;
+				HairGroupCluster.ClusterIndexRadiusScaleInfoBuffer = Input.ClusterIndexRadiusScaleInfoBuffer;
+				HairGroupCluster.HairGroupPublicPtr = Output.HairGroupPublicData;
+				HairGroupCluster.LodBias = Input.GroupDesc.LodBias;
+				HairGroupCluster.LodAverageVertexPerPixel = Input.GroupDesc.LodAverageVertexPerPixel;
+
+				// Note: This code needs to exactly match the values FHairScaleAndClipDesc set int the previous loop.
+				const float OutHairRadius = (GStrandHairWidth > 0 ? GStrandHairWidth : Input.GroupDesc.HairWidth) * 0.5f;
+				const float MaxOutHairRadius = OutHairRadius * FMath::Max(1.f, FMath::Max(Input.GroupDesc.HairRootScale, Input.GroupDesc.HairTipScale));
+
+				Output.VFInput.HairRadius = MaxOutHairRadius;
+				Output.VFInput.HairLength = Input.GroupDesc.HairLength;
+				Output.VFInput.HairDensity = Input.GroupDesc.HairShadowDensity;
+				Output.VFInput.HairPositionBuffer = Output.RenderDeformedPositionBuffer[CurrIndex]->SRV;
+				Output.VFInput.HairPreviousPositionBuffer = Output.RenderDeformedPositionBuffer[PrevIndex]->SRV;
+				Output.VFInput.bUseStableRasterization = Input.GroupDesc.bUseStableRasterization;
+				Output.VFInput.bScatterSceneLighting = Input.GroupDesc.bScatterSceneLighting;
+
+				AddHairClusterAABBPass(
+					GraphBuilder,
+					InRenHairDatas.HairGroups[GroupIndex],
+					Input.OutHairPositionOffset,
+					HairGroupCluster,
+					Output.RenderDeformedPositionBuffer[CurrIndex]->SRV,
+					TransitionQueue);
+			}
+			GraphBuilder.Execute();
+			TransitBufferToReadable(RHICmdList, TransitionQueue);
 		}
 
-		for (uint32 GroupIndex = 0; GroupIndex < GroupCount; ++GroupIndex)
 		{
-			FHairStrandsInterpolationInput::FHairGroup& Input = InInput->HairGroups[GroupIndex];
-			FHairStrandsInterpolationOutput::HairGroup& Output = InOutput->HairGroups[GroupIndex];
+			FBufferTransitionQueue TransitionQueue;
+			FRDGBuilder GraphBuilder(RHICmdList);
+			for (uint32 GroupIndex = 0; GroupIndex < GroupCount; ++GroupIndex)
+			{
+				FHairStrandsInterpolationInput::FHairGroup& Input = InInput->HairGroups[GroupIndex];
+				FHairStrandsInterpolationOutput::HairGroup& Output = InOutput->HairGroups[GroupIndex];
 
-			const uint32 CurrIndex = *Output.CurrentIndex;
-			const uint32 PrevIndex = (CurrIndex + 1) % 2;
+				const uint32 CurrIndex = *Output.CurrentIndex;
+				const uint32 PrevIndex = (CurrIndex + 1) % 2;
 
-#if RHI_RAYTRACING
-			if (IsHairRayTracingEnabled())
-			{
-				const float HairRadiusScaleRT = (GHairRaytracingRadiusScale > 0 ? GHairRaytracingRadiusScale : Input.GroupDesc.HairRaytracingRadiusScale);
-				AddGenerateRaytracingGeometryPass(
+				AddHairTangentPass(
 					GraphBuilder,
 					Input.RenderVertexCount,
-					Output.VFInput.HairRadius * HairRadiusScaleRT,
-					Input.OutHairPositionOffset,
 					Output.VFInput.HairPositionBuffer,
-					Input.RaytracingPositionBuffer->UAV);
+					Output.RenderTangentBuffer->UAV,
+					TransitionQueue);
 			}
-#endif
-
-			Output.VFInput.HairTangentBuffer = Output.RenderTangentBuffer->SRV;
-			Output.VFInput.HairAttributeBuffer = bDebugModePatchedAttributeBuffer ? Output.RenderPatchedAttributeBuffer.SRV : Input.RenderAttributeBuffer->SRV;
-			Output.VFInput.HairMaterialBuffer = Output.RenderMaterialBuffer->SRV;
-			Output.VFInput.HairPositionOffset = Input.OutHairPositionOffset;
-			Output.VFInput.HairPreviousPositionOffset = Input.OutHairPreviousPositionOffset;
-			Output.VFInput.VertexCount = Input.RenderVertexCount;
+			GraphBuilder.Execute();
+			TransitBufferToReadable(RHICmdList, TransitionQueue);
 		}
 
-		for (uint32 GroupIndex = 0; GroupIndex < GroupCount; ++GroupIndex)
 		{
-			FHairStrandsInterpolationInput::FHairGroup& Input = InInput->HairGroups[GroupIndex];
-			FHairStrandsInterpolationOutput::HairGroup& Output = InOutput->HairGroups[GroupIndex];
-
-			const uint32 CurrIndex = *Output.CurrentIndex;
-			const uint32 PrevIndex = (CurrIndex + 1) % 2;
-
-			// TODO: find a more robust way to handle parameters passing to compute raster.
-			// At the moment there is a loose compling which will break if the vertex factor change.
-			Output.HairGroupPublicData->VFInput.HairPositionBuffer = Output.VFInput.HairPositionBuffer;
-			Output.HairGroupPublicData->VFInput.HairPositionOffset = Output.VFInput.HairPositionOffset;
-			Output.HairGroupPublicData->VFInput.VertexCount = Output.VFInput.VertexCount;
-			Output.HairGroupPublicData->VFInput.HairRadius = Output.VFInput.HairRadius;
-			Output.HairGroupPublicData->VFInput.HairLength = Output.VFInput.HairLength;
-			Output.HairGroupPublicData->VFInput.bUseStableRasterization = Output.VFInput.bUseStableRasterization;
-			Output.HairGroupPublicData->VFInput.bScatterSceneLighting = Output.VFInput.bScatterSceneLighting;
-			Output.HairGroupPublicData->VFInput.HairDensity = Output.VFInput.HairDensity;
-			Output.HairGroupPublicData->VFInput.LocalToWorldTransform = LocalToWorld;
+			FBufferTransitionQueue TransitionQueue;
+			FRDGBuilder GraphBuilder(RHICmdList);
+			for (uint32 GroupIndex = 0; GroupIndex < GroupCount; ++GroupIndex)
+			{
+				FHairStrandsInterpolationInput::FHairGroup& Input = InInput->HairGroups[GroupIndex];
+				FHairStrandsInterpolationOutput::HairGroup& Output = InOutput->HairGroups[GroupIndex];
+
+				const uint32 CurrIndex = *Output.CurrentIndex;
+				const uint32 PrevIndex = (CurrIndex + 1) % 2;
+
+				#if RHI_RAYTRACING
+				if (IsHairRayTracingEnabled())
+				{
+					const float HairRadiusScaleRT = (GHairRaytracingRadiusScale > 0 ? GHairRaytracingRadiusScale : Input.GroupDesc.HairRaytracingRadiusScale);
+					AddGenerateRaytracingGeometryPass(
+						GraphBuilder,
+						Input.RenderVertexCount,
+						Output.VFInput.HairRadius * HairRadiusScaleRT,
+						Input.OutHairPositionOffset,
+						Output.VFInput.HairPositionBuffer,
+						Input.RaytracingPositionBuffer->UAV,
+						TransitionQueue);
+				}
+				#endif
+
+				Output.VFInput.HairTangentBuffer = Output.RenderTangentBuffer->SRV;
+				Output.VFInput.HairAttributeBuffer = bDebugModePatchedAttributeBuffer ? Output.RenderPatchedAttributeBuffer.SRV : Input.RenderAttributeBuffer->SRV;
+				Output.VFInput.HairMaterialBuffer = Output.RenderMaterialBuffer->SRV;
+				Output.VFInput.HairPositionOffset = Input.OutHairPositionOffset;
+				Output.VFInput.HairPreviousPositionOffset = Input.OutHairPreviousPositionOffset;
+				Output.VFInput.VertexCount = Input.RenderVertexCount;
+			}
+			GraphBuilder.Execute();
+			TransitBufferToReadable(RHICmdList, TransitionQueue);
+		}
 
-#if RHI_RAYTRACING
-			if (IsHairRayTracingEnabled())
+		{
+			FRDGBuilder GraphBuilder(RHICmdList);
+			for (uint32 GroupIndex = 0; GroupIndex < GroupCount; ++GroupIndex)
 			{
-				FRHIUnorderedAccessView* UAV = Input.RaytracingPositionBuffer->UAV;
-				RHICmdList.TransitionResources(EResourceTransitionAccess::EReadable, EResourceTransitionPipeline::EComputeToGfx, &UAV, 1);
-
-				const bool bNeedFullBuild = !Input.bIsRTGeometryInitialized;
-				if (bNeedFullBuild)
-					BuildHairAccelerationStructure(RHICmdList, Input.RaytracingVertexCount, Input.RaytracingPositionBuffer->Buffer, Input.RaytracingGeometry);
-				else
-					UpdateHairAccelerationStructure(RHICmdList, Input.RaytracingGeometry);
-				Input.bIsRTGeometryInitialized = true;
+				FHairStrandsInterpolationInput::FHairGroup& Input = InInput->HairGroups[GroupIndex];
+				FHairStrandsInterpolationOutput::HairGroup& Output = InOutput->HairGroups[GroupIndex];
+
+				const uint32 CurrIndex = *Output.CurrentIndex;
+				const uint32 PrevIndex = (CurrIndex + 1) % 2;
+
+				// TODO: find a more robust way to handle parameters passing to compute raster.
+				// At the moment there is a loose compling which will break if the vertex factor change.
+				Output.HairGroupPublicData->VFInput.HairPositionBuffer = Output.VFInput.HairPositionBuffer;
+				Output.HairGroupPublicData->VFInput.HairPositionOffset = Output.VFInput.HairPositionOffset;
+				Output.HairGroupPublicData->VFInput.VertexCount = Output.VFInput.VertexCount;
+				Output.HairGroupPublicData->VFInput.HairRadius = Output.VFInput.HairRadius;
+				Output.HairGroupPublicData->VFInput.HairLength = Output.VFInput.HairLength;
+				Output.HairGroupPublicData->VFInput.bUseStableRasterization = Output.VFInput.bUseStableRasterization;
+				Output.HairGroupPublicData->VFInput.bScatterSceneLighting = Output.VFInput.bScatterSceneLighting;
+				Output.HairGroupPublicData->VFInput.HairDensity = Output.VFInput.HairDensity;
+				Output.HairGroupPublicData->VFInput.LocalToWorldTransform = LocalToWorld;
+
+	#if RHI_RAYTRACING
+				if (IsHairRayTracingEnabled())
+				{
+					FRHIUnorderedAccessView* UAV = Input.RaytracingPositionBuffer->UAV;
+					RHICmdList.TransitionResources(EResourceTransitionAccess::EReadable, EResourceTransitionPipeline::EComputeToGfx, &UAV, 1);
+
+					const bool bNeedFullBuild = !Input.bIsRTGeometryInitialized;
+					if (bNeedFullBuild)
+						BuildHairAccelerationStructure(RHICmdList, Input.RaytracingVertexCount, Input.RaytracingPositionBuffer->Buffer, Input.RaytracingGeometry);
+					else
+						UpdateHairAccelerationStructure(RHICmdList, Input.RaytracingGeometry);
+					Input.bIsRTGeometryInitialized = true;
+				}
+	#endif
 			}
-#endif
+			GraphBuilder.Execute();
 		}
 	}
 
-	GraphBuilder.Execute();
-
 	// update the current index
 	for (uint32 GroupIndex = 0; GroupIndex < GroupCount; ++GroupIndex)
 	{
@@ -1129,6 +1210,8 @@ void ResetHairStrandsInterpolation(
 
 			const uint32 CurrIndex = *Output.CurrentIndex;
 			const uint32 PrevIndex = (CurrIndex + 1) % 2;
+
+			FBufferTransitionQueue TransitionQueue;
 			FRDGBuilder GraphBuilder(RHICmdList);
 			AddDeformSimHairStrandsPass(
 				GraphBuilder,
@@ -1139,8 +1222,10 @@ void ResetHairStrandsInterpolation(
 				Input.SimRestPosePositionBuffer->SRV,
 				Input.SimRootPointIndexBuffer ? Input.SimRootPointIndexBuffer->SRV : nullptr,
 				Output.SimDeformedPositionBuffer[CurrIndex]->UAV, Input.InSimHairPositionOffset,
-				Input.OutHairPositionOffset);
+				Input.OutHairPositionOffset, 
+				TransitionQueue);
 			GraphBuilder.Execute();
+			TransitBufferToReadable(RHICmdList, TransitionQueue);
 		}
 	}
 }
@@ -1173,7 +1258,7 @@ void WritebackBuffer(TArray<WriteBackType>& InData, FRWBuffer& OutBuffer)
 
 static void ReadbackGroupData(
 	FHairStrandsRootData& OutCPUData,
-	FHairStrandsRootResource* InGPUData)
+	FHairStrandsRestRootResource* InGPUData)
 {
 	if (!InGPUData)
 	{
@@ -1186,7 +1271,7 @@ static void ReadbackGroupData(
 	for (uint32 LODIt = 0; LODIt < LODCount; ++LODIt)
 	{
 		FHairStrandsRootData::FMeshProjectionLOD& CPULOD = OutCPUData.MeshProjectionLODs[LODIt];
-		FHairStrandsRootResource::FMeshProjectionLOD& GPULOD = InGPUData->MeshProjectionLODs[LODIt];
+		FHairStrandsRestRootResource::FMeshProjectionLOD& GPULOD = InGPUData->MeshProjectionLODs[LODIt];
 		check(CPULOD.LODIndex == GPULOD.LODIndex);
 
 		ReadbackBuffer(CPULOD.RootTriangleIndexBuffer, GPULOD.RootTriangleIndexBuffer);
@@ -1203,8 +1288,7 @@ static void ReadbackGroupData(
 	}
 }
 
-void WritebackGroupData(FHairStrandsRootData& InCPUData,
-	FHairStrandsRootResource* OutGPUData)
+void WritebackGroupData(FHairStrandsRootData& InCPUData, FHairStrandsRestRootResource* OutGPUData)
 {
 	if (!OutGPUData)
 	{
@@ -1215,7 +1299,7 @@ void WritebackGroupData(FHairStrandsRootData& InCPUData,
 	for (uint32 LODIt = 0; LODIt < LODCount; ++LODIt)
 	{
 		FHairStrandsRootData::FMeshProjectionLOD& CPULOD = InCPUData.MeshProjectionLODs[LODIt];
-		FHairStrandsRootResource::FMeshProjectionLOD& GPULOD = OutGPUData->MeshProjectionLODs[LODIt];
+		FHairStrandsRestRootResource::FMeshProjectionLOD& GPULOD = OutGPUData->MeshProjectionLODs[LODIt];
 		check(CPULOD.LODIndex == GPULOD.LODIndex);
 
 		if (CPULOD.SampleCount > 0)
@@ -1466,8 +1550,8 @@ static void InternalProcessGroomBindingTask(FRHICommandListImmediate& RHICmdList
 	for (UGroomBindingAsset::FHairGroupData& GroupData : OutHairGroupDatas)
 	{
 		UGroomBindingAsset::FHairGroupResource& Resource = OutHairGroupResources.AddDefaulted_GetRef();
-		Resource.SimRootResources = new FHairStrandsRootResource(GroupData.SimRootData);
-		Resource.RenRootResources = new FHairStrandsRootResource(GroupData.RenRootData);
+		Resource.SimRootResources = new FHairStrandsRestRootResource(GroupData.SimRootData);
+		Resource.RenRootResources = new FHairStrandsRestRootResource(GroupData.RenRootData);
 
 		Resource.SimRootResources->InitRHI();
 		Resource.RenRootResources->InitRHI();
@@ -1489,8 +1573,8 @@ static void InternalProcessGroomBindingTask(FRHICommandListImmediate& RHICmdList
 	const uint32 GroupCount = OutHairGroupResources.Num();
 	for (const UGroomBindingAsset::FHairGroupResource& GroupResources : OutHairGroupResources)
 	{
-		RenProjectionDatas.HairGroups.Add(ToProjectionHairData(GroupResources.RenRootResources));
-		SimProjectionDatas.HairGroups.Add(ToProjectionHairData(GroupResources.SimRootResources));
+		RenProjectionDatas.HairGroups.Add(ToProjectionHairData(GroupResources.RenRootResources, nullptr));
+		SimProjectionDatas.HairGroups.Add(ToProjectionHairData(GroupResources.SimRootResources, nullptr));
 	}
 
 	FSkeletalMeshRenderData* TargetRenderData = TargetSkeletalMesh->GetResourceForRendering();
@@ -1684,44 +1768,64 @@ void AddGroomBindingTask(UGroomBindingAsset* BindingAsset)
 	EnqueueGroomBindingQuery((void*)BindingAsset, InternalProcessGroomBindingTask);
 }
 
-FHairStrandsProjectionHairData::HairGroup ToProjectionHairData(FHairStrandsRootResource* In)
+FHairStrandsProjectionHairData::HairGroup ToProjectionHairData(
+	FHairStrandsRestRootResource* InRest, 
+	FHairStrandsDeformedRootResource* InDeformed)
 {
 	check(IsInRenderingThread());
 
 	FHairStrandsProjectionHairData::HairGroup Out = {};
-	if (!In)
+	if (!InRest)
 		return Out;
 
-	Out.RootCount = In->RootData.RootCount;
-	Out.RootPositionBuffer = In->RootPositionBuffer.SRV;
-	Out.RootNormalBuffer = In->RootNormalBuffer.SRV;
-	Out.VertexToCurveIndexBuffer = &In->VertexToCurveIndexBuffer;
+	Out.RootCount = InRest->RootData.RootCount;
+	Out.RootPositionBuffer = InRest->RootPositionBuffer.SRV;
+	Out.RootNormalBuffer = InRest->RootNormalBuffer.SRV;
+	Out.VertexToCurveIndexBuffer = &InRest->VertexToCurveIndexBuffer;
 
-	for (FHairStrandsRootResource::FMeshProjectionLOD& MeshLODData : In->MeshProjectionLODs)
+	if (InRest && InDeformed)
 	{
-		FHairStrandsProjectionHairData::LODData& LODData = Out.LODDatas.AddDefaulted_GetRef();
-		LODData.LODIndex = MeshLODData.LODIndex;
-		LODData.RootTriangleIndexBuffer = &MeshLODData.RootTriangleIndexBuffer;
-		LODData.RootTriangleBarycentricBuffer = &MeshLODData.RootTriangleBarycentricBuffer;
-
-		LODData.RestRootTrianglePosition0Buffer = &MeshLODData.RestRootTrianglePosition0Buffer;
-		LODData.RestRootTrianglePosition1Buffer = &MeshLODData.RestRootTrianglePosition1Buffer;
-		LODData.RestRootTrianglePosition2Buffer = &MeshLODData.RestRootTrianglePosition2Buffer;
-
-		LODData.DeformedRootTrianglePosition0Buffer = &MeshLODData.DeformedRootTrianglePosition0Buffer;
-		LODData.DeformedRootTrianglePosition1Buffer = &MeshLODData.DeformedRootTrianglePosition1Buffer;
-		LODData.DeformedRootTrianglePosition2Buffer = &MeshLODData.DeformedRootTrianglePosition2Buffer;
-
-		LODData.Status = &MeshLODData.Status;
-		LODData.bIsValid = MeshLODData.Status == FHairStrandsProjectionHairData::LODData::EStatus::Completed || MeshLODData.Status == FHairStrandsProjectionHairData::LODData::EStatus::Initialized;
-
-		LODData.SampleCount = MeshLODData.SampleCount;
-		LODData.MeshInterpolationWeightsBuffer = &MeshLODData.MeshInterpolationWeightsBuffer;
-		LODData.MeshSampleIndicesBuffer = &MeshLODData.MeshSampleIndicesBuffer;
-		LODData.RestSamplePositionsBuffer = &MeshLODData.RestSamplePositionsBuffer;
-		LODData.DeformedSamplePositionsBuffer = &MeshLODData.DeformedSamplePositionsBuffer;
-		LODData.MeshSampleWeightsBuffer = &MeshLODData.MeshSampleWeightsBuffer;
+		check(InRest->MeshProjectionLODs.Num() == InRest->MeshProjectionLODs.Num());
 	}
+	const uint32 LODCount = InRest->MeshProjectionLODs.Num();
+	for (uint32 LODIt=0; LODIt<LODCount; ++LODIt)
+	{
+		{
+			FHairStrandsRestRootResource::FMeshProjectionLOD& Rest = InRest->MeshProjectionLODs[LODIt];
+			FHairStrandsProjectionHairData::RestLODData& OutRest = Out.RestLODDatas.AddDefaulted_GetRef();
+
+			OutRest.Status = &Rest.Status;
+			OutRest.LODIndex = Rest.LODIndex;
+
+			OutRest.RootTriangleIndexBuffer = &Rest.RootTriangleIndexBuffer;
+			OutRest.RootTriangleBarycentricBuffer = &Rest.RootTriangleBarycentricBuffer;
+
+			OutRest.RestRootTrianglePosition0Buffer = &Rest.RestRootTrianglePosition0Buffer;
+			OutRest.RestRootTrianglePosition1Buffer = &Rest.RestRootTrianglePosition1Buffer;
+			OutRest.RestRootTrianglePosition2Buffer = &Rest.RestRootTrianglePosition2Buffer;
+
+
+			OutRest.SampleCount	= Rest.SampleCount;
+			OutRest.MeshInterpolationWeightsBuffer	= &Rest.MeshInterpolationWeightsBuffer;
+			OutRest.MeshSampleIndicesBuffer			= &Rest.MeshSampleIndicesBuffer;
+			OutRest.RestSamplePositionsBuffer		= &Rest.RestSamplePositionsBuffer;
+		}
+
+		if (InDeformed)
+		{
+			FHairStrandsDeformedRootResource::FMeshProjectionLOD& Deformed = InDeformed->MeshProjectionLODs[LODIt];
+			FHairStrandsProjectionHairData::DeformedLODData& OutDeformed = Out.DeformedLODDatas.AddDefaulted_GetRef();
+
+			OutDeformed.Status = &Deformed.Status;
+			OutDeformed.LODIndex = Deformed.LODIndex;
 
+			OutDeformed.DeformedRootTrianglePosition0Buffer = &Deformed.DeformedRootTrianglePosition0Buffer;
+			OutDeformed.DeformedRootTrianglePosition1Buffer = &Deformed.DeformedRootTrianglePosition1Buffer;
+			OutDeformed.DeformedRootTrianglePosition2Buffer = &Deformed.DeformedRootTrianglePosition2Buffer;
+
+			OutDeformed.DeformedSamplePositionsBuffer	= &Deformed.DeformedSamplePositionsBuffer;
+			OutDeformed.MeshSampleWeightsBuffer			= &Deformed.MeshSampleWeightsBuffer;
+		}
+	}
 	return Out;
-}
\ No newline at end of file
+}
diff --git a/Engine/Plugins/Experimental/HairStrands/Source/HairStrandsCore/Public/GroomAsset.h b/Engine/Plugins/Experimental/HairStrands/Source/HairStrandsCore/Public/GroomAsset.h
--- a/Engine/Plugins/Experimental/HairStrands/Source/HairStrandsCore/Public/GroomAsset.h
+++ b/Engine/Plugins/Experimental/HairStrands/Source/HairStrandsCore/Public/GroomAsset.h
@@ -99,12 +99,12 @@ struct FHairStrandsRootData
 };
 
 /* Render buffers for root deformation for dynamic meshes */
-struct FHairStrandsRootResource : public FRenderResource
+struct FHairStrandsRestRootResource : public FRenderResource
 {
 	/** Build the hair strands resource */
-	FHairStrandsRootResource();
-	FHairStrandsRootResource(const FHairStrandsRootData& RootData);
-	FHairStrandsRootResource(const FHairStrandsDatas* HairStrandsDatas, uint32 LODCount, const TArray<uint32>& NumSamples);
+	FHairStrandsRestRootResource();
+	FHairStrandsRestRootResource(const FHairStrandsRootData& RootData);
+	FHairStrandsRestRootResource(const FHairStrandsDatas* HairStrandsDatas, uint32 LODCount, const TArray<uint32>& NumSamples);
 
 	/* Init the buffer */
 	virtual void InitRHI() override;
@@ -113,7 +113,7 @@ struct FHairStrandsRootResource : public FRenderResource
 	virtual void ReleaseRHI() override;
 
 	/* Get the resource name */
-	virtual FString GetFriendlyName() const override { return TEXT("FHairStrandsRootResource"); }
+	virtual FString GetFriendlyName() const override { return TEXT("FHairStrandsRestRootResource"); }
 
 	/* Populate GPU LOD data from RootData (this function doesn't initialize resources) */
 	void PopulateFromRootData();
@@ -124,7 +124,8 @@ struct FHairStrandsRootResource : public FRenderResource
 
 	struct FMeshProjectionLOD
 	{
-		FHairStrandsProjectionHairData::LODData::EStatus Status = FHairStrandsProjectionHairData::LODData::EStatus::Invalid;
+		const bool IsValid() const { return Status == FHairStrandsProjectionHairData::EStatus::Completed; }
+		FHairStrandsProjectionHairData::EStatus Status = FHairStrandsProjectionHairData::EStatus::Invalid;
 		int32 LODIndex = -1;
 
 		/* Triangle on which a root is attached */
@@ -138,6 +139,43 @@ struct FHairStrandsRootResource : public FRenderResource
 		FRWBuffer RestRootTrianglePosition1Buffer;
 		FRWBuffer RestRootTrianglePosition2Buffer;
 
+		/* Strand hair mesh interpolation matrix and sample indices */
+		uint32 SampleCount = 0;
+		FRWBuffer MeshInterpolationWeightsBuffer;
+		FRWBuffer MeshSampleIndicesBuffer;
+		FRWBuffer RestSamplePositionsBuffer;
+	};
+
+	/* Store the hair projection information for each mesh LOD */
+	TArray<FMeshProjectionLOD> MeshProjectionLODs;
+
+	/* Store CPU data for root info & root binding */
+	FHairStrandsRootData RootData;
+};
+
+
+/* Render buffers for root deformation for dynamic meshes */
+struct FHairStrandsDeformedRootResource : public FRenderResource
+{
+	/** Build the hair strands resource */
+	FHairStrandsDeformedRootResource();
+	FHairStrandsDeformedRootResource(const FHairStrandsRestRootResource* InRestResources);
+
+	/* Init the buffer */
+	virtual void InitRHI() override;
+
+	/* Release the buffer */
+	virtual void ReleaseRHI() override;
+
+	/* Get the resource name */
+	virtual FString GetFriendlyName() const override { return TEXT("FHairStrandsDeformedRootResource"); }
+
+	struct FMeshProjectionLOD
+	{
+		const bool IsValid() const { return Status == FHairStrandsProjectionHairData::EStatus::Completed; }
+		FHairStrandsProjectionHairData::EStatus Status = FHairStrandsProjectionHairData::EStatus::Invalid;
+		int32 LODIndex = -1;
+
 		/* Strand hair roots translation and rotation in triangle-deformed position relative to the bound triangle. Positions are relative the deformed root center*/
 		FRWBuffer DeformedRootTrianglePosition0Buffer;
 		FRWBuffer DeformedRootTrianglePosition1Buffer;
@@ -145,18 +183,13 @@ struct FHairStrandsRootResource : public FRenderResource
 
 		/* Strand hair mesh interpolation matrix and sample indices */
 		uint32 SampleCount = 0;
-		FRWBuffer MeshInterpolationWeightsBuffer;
-		FRWBuffer MeshSampleIndicesBuffer;
-		FRWBuffer RestSamplePositionsBuffer;
 		FRWBuffer DeformedSamplePositionsBuffer;
 		FRWBuffer MeshSampleWeightsBuffer;
 	};
 
 	/* Store the hair projection information for each mesh LOD */
+	uint32 RootCount = 0;
 	TArray<FMeshProjectionLOD> MeshProjectionLODs;
-
-	/* Store CPU data for root info & root binding */
-	FHairStrandsRootData RootData;
 };
 
 /* Render buffers that will be used for rendering */
@@ -736,8 +769,8 @@ public:
 	/** GPU and CPU binding data for both simulation and rendering. */
 	struct FHairGroupResource
 	{
-		FHairStrandsRootResource* SimRootResources = nullptr;
-		FHairStrandsRootResource* RenRootResources = nullptr;
+		FHairStrandsRestRootResource* SimRootResources = nullptr;
+		FHairStrandsRestRootResource* RenRootResources = nullptr;
 	};
 	typedef TArray<FHairGroupResource> FHairGroupResources;
 	FHairGroupResources HairGroupResources;
diff --git a/Engine/Plugins/Experimental/HairStrands/Source/HairStrandsCore/Public/GroomComponent.h b/Engine/Plugins/Experimental/HairStrands/Source/HairStrandsCore/Public/GroomComponent.h
--- a/Engine/Plugins/Experimental/HairStrands/Source/HairStrandsCore/Public/GroomComponent.h
+++ b/Engine/Plugins/Experimental/HairStrands/Source/HairStrandsCore/Public/GroomComponent.h
@@ -98,7 +98,8 @@ public:
 	FHairStrandsDeformedResource* GetGuideStrandsDeformedResource(uint32 GroupIndex);
 
 	/** Return the guide hairs root resources*/
-	FHairStrandsRootResource* GetGuideStrandsRootResource(uint32 GroupIndex);
+	FHairStrandsRestRootResource* GetGuideStrandsRestRootResource(uint32 GroupIndex);
+	FHairStrandsDeformedRootResource* GetGuideStrandsDeformedRootResource(uint32 GroupIndex);
 
 #if WITH_EDITOR
 	virtual void CheckForErrors() override;
@@ -133,8 +134,12 @@ private:
 
 		// Projection resources
 		bool bOwnRootResourceAllocation = true;
-		struct FHairStrandsRootResource* RenRootResources = nullptr;
-		struct FHairStrandsRootResource* SimRootResources = nullptr;
+		struct FHairStrandsRestRootResource* RenRestRootResources = nullptr;
+		struct FHairStrandsRestRootResource* SimRestRootResources = nullptr;
+
+		struct FHairStrandsDeformedRootResource* RenDeformedRootResources = nullptr;
+		struct FHairStrandsDeformedRootResource* SimDeformedRootResources = nullptr;
+
 	#if RHI_RAYTRACING
 		FHairStrandsRaytracingResource* RaytracingResources = nullptr;
 	#endif
diff --git a/Engine/Plugins/Experimental/HairStrands/Source/HairStrandsNiagara/Private/NiagaraDataInterfaceHairStrands.cpp b/Engine/Plugins/Experimental/HairStrands/Source/HairStrandsNiagara/Private/NiagaraDataInterfaceHairStrands.cpp
--- a/Engine/Plugins/Experimental/HairStrands/Source/HairStrandsNiagara/Private/NiagaraDataInterfaceHairStrands.cpp
+++ b/Engine/Plugins/Experimental/HairStrands/Source/HairStrandsNiagara/Private/NiagaraDataInterfaceHairStrands.cpp
@@ -406,13 +406,19 @@ inline void ClearBuffer(FRHICommandList& RHICmdList, FNDIHairStrandsBuffer* Hair
 
 //------------------------------------------------------------------------------------------------------------
 
-void FNDIHairStrandsBuffer::Initialize(const FHairStrandsDatas*  HairStrandsDatas, const FHairStrandsRestResource* HairStrandsRestResource, 
-	const FHairStrandsDeformedResource*  HairStrandsDeformedResource, const FHairStrandsRootResource* HairStrandsRootResource, const FNDIHairStrandsData* NDIStrandsDatas)
+void FNDIHairStrandsBuffer::Initialize(
+	const FHairStrandsDatas*  HairStrandsDatas, 
+	const FHairStrandsRestResource* HairStrandsRestResource, 
+	const FHairStrandsDeformedResource*  HairStrandsDeformedResource, 
+	const FHairStrandsRestRootResource* HairStrandsRestRootResource,
+	const FHairStrandsDeformedRootResource* HairStrandsDeformedRootResource,
+	const FNDIHairStrandsData* NDIStrandsDatas)
 {
 	SourceDatas = HairStrandsDatas;
 	SourceRestResources = HairStrandsRestResource;
 	SourceDeformedResources = HairStrandsDeformedResource;
-	SourceRootResources = HairStrandsRootResource;
+	SourceRestRootResources = HairStrandsRestRootResource;
+	SourceDeformedRootResources = HairStrandsDeformedRootResource;
 	NDIDatas = NDIStrandsDatas;
 }
 
@@ -571,16 +577,17 @@ bool FNDIHairStrandsData::Init(UNiagaraDataInterfaceHairStrands* Interface, FNia
 		FHairStrandsDatas* StrandsDatas = nullptr;
 		FHairStrandsRestResource* StrandsRestResource = nullptr;
 		FHairStrandsDeformedResource* StrandsDeformedResource = nullptr;
-		FHairStrandsRootResource* StrandsRootResource = nullptr;
+		FHairStrandsRestRootResource* StrandsRestRootResource = nullptr;
+		FHairStrandsDeformedRootResource* StrandsDeformedRootResource = nullptr;
 		UGroomAsset* GroomAsset = nullptr;
 		int32 GroupIndex = 0;
 
 		{
-			Interface->ExtractDatasAndResources(SystemInstance, StrandsDatas, StrandsRestResource, StrandsDeformedResource, StrandsRootResource, GroomAsset, GroupIndex);
+			Interface->ExtractDatasAndResources(SystemInstance, StrandsDatas, StrandsRestResource, StrandsDeformedResource, StrandsRestRootResource, StrandsDeformedRootResource, GroomAsset, GroupIndex);
 			Update(Interface, SystemInstance, StrandsDatas, GroomAsset, GroupIndex);
 
 			HairStrandsBuffer = new FNDIHairStrandsBuffer();
-			HairStrandsBuffer->Initialize(StrandsDatas, StrandsRestResource, StrandsDeformedResource, StrandsRootResource, this);
+			HairStrandsBuffer->Initialize(StrandsDatas, StrandsRestResource, StrandsDeformedResource, StrandsRestRootResource, StrandsDeformedRootResource, this);
 
 			BeginInitResource(HairStrandsBuffer);
 		
@@ -706,7 +713,7 @@ struct FNDIHairStrandsParametersCS : public FNiagaraDataInterfaceParametersCS
 			InterfaceProxy->SystemInstancesToProxyData.Find(Context.SystemInstance);
 
 		const bool IsHairValid = ProxyData != nullptr && ProxyData->HairStrandsBuffer != nullptr && ProxyData->HairStrandsBuffer->IsInitialized();
-		const bool IsRootValid = IsHairValid && ProxyData->HairStrandsBuffer->SourceRootResources != nullptr;//&& ProxyData->HairStrandsBuffer->SourceRootResources->IsInitialized();
+		const bool IsRootValid = IsHairValid && ProxyData->HairStrandsBuffer->SourceDeformedRootResources != nullptr;//&& ProxyData->HairStrandsBuffer->SourceRootResources->IsInitialized();
 		const bool IsRestValid = IsHairValid && ProxyData->HairStrandsBuffer->SourceRestResources != nullptr;//&& ProxyData->HairStrandsBuffer->SourceRestResources->IsInitialized();
 		const bool IsDeformedValid = IsHairValid && ProxyData->HairStrandsBuffer->SourceDeformedResources != nullptr && ProxyData->HairStrandsBuffer->SourceDeformedResources->IsInitialized();
 
@@ -717,33 +724,35 @@ struct FNDIHairStrandsParametersCS : public FNiagaraDataInterfaceParametersCS
 			FUnorderedAccessViewRHIRef PointPositionsUAV = IsDeformedValid ?
 				HairStrandsBuffer->SourceDeformedResources->DeformedPositionBuffer[HairStrandsBuffer->SourceDeformedResources->CurrentIndex].UAV : HairStrandsBuffer->DeformedPositionBuffer.UAV;
 
-			int32 InterpolationModeValue = (IsRootValid && HairStrandsBuffer->SourceRootResources->MeshProjectionLODs.Num() > 0 && 
-				HairStrandsBuffer->SourceRootResources->MeshProjectionLODs[0].Status == FHairStrandsProjectionHairData::LODData::EStatus::Completed);
+			int32 InterpolationModeValue = IsRootValid && HairStrandsBuffer->SourceDeformedRootResources->MeshProjectionLODs.Num() > 0 && 
+				HairStrandsBuffer->SourceDeformedRootResources->MeshProjectionLODs[0].IsValid();
 
-			const FHairStrandsRootResource::FMeshProjectionLOD* MeshProjection = (InterpolationModeValue == 1) ? &(HairStrandsBuffer->SourceRootResources->MeshProjectionLODs[0]) : nullptr;
+			const FHairStrandsRestRootResource::FMeshProjectionLOD* RestMeshProjection = (InterpolationModeValue == 1) ? &(HairStrandsBuffer->SourceRestRootResources->MeshProjectionLODs[0]) : nullptr;
+			const FHairStrandsDeformedRootResource::FMeshProjectionLOD* DeformedMeshProjection = (InterpolationModeValue == 1) ? &(HairStrandsBuffer->SourceDeformedRootResources->MeshProjectionLODs[0]) : nullptr;
 
-			FRHIShaderResourceView* RestTrianglePositionASRV = (InterpolationModeValue == 1 && MeshProjection != nullptr) ?
-				MeshProjection->RestRootTrianglePosition0Buffer.SRV.GetReference() : FNiagaraRenderer::GetDummyFloatBuffer();
-			FRHIShaderResourceView* RestTrianglePositionBSRV = (InterpolationModeValue == 1 && MeshProjection != nullptr) ?
-				MeshProjection->RestRootTrianglePosition1Buffer.SRV.GetReference() : FNiagaraRenderer::GetDummyFloatBuffer();
-			FRHIShaderResourceView* RestTrianglePositionCSRV = (InterpolationModeValue == 1 && MeshProjection != nullptr) ?
-				MeshProjection->RestRootTrianglePosition2Buffer.SRV.GetReference() : FNiagaraRenderer::GetDummyFloatBuffer();
+			FRHIShaderResourceView* RestTrianglePositionASRV = (InterpolationModeValue == 1 && RestMeshProjection != nullptr) ?
+				RestMeshProjection->RestRootTrianglePosition0Buffer.SRV.GetReference() : FNiagaraRenderer::GetDummyFloatBuffer();
+			FRHIShaderResourceView* RestTrianglePositionBSRV = (InterpolationModeValue == 1 && RestMeshProjection != nullptr) ?
+				RestMeshProjection->RestRootTrianglePosition1Buffer.SRV.GetReference() : FNiagaraRenderer::GetDummyFloatBuffer();
+			FRHIShaderResourceView* RestTrianglePositionCSRV = (InterpolationModeValue == 1 && RestMeshProjection != nullptr) ?
+				RestMeshProjection->RestRootTrianglePosition2Buffer.SRV.GetReference() : FNiagaraRenderer::GetDummyFloatBuffer();
 
-			FRHIShaderResourceView* DeformedTrianglePositionASRV = (InterpolationModeValue == 1 && MeshProjection != nullptr) ?
-				MeshProjection->DeformedRootTrianglePosition0Buffer.SRV.GetReference() : FNiagaraRenderer::GetDummyFloatBuffer();
-			FRHIShaderResourceView* DeformedTrianglePositionBSRV = (InterpolationModeValue == 1 && MeshProjection != nullptr) ?
-				MeshProjection->DeformedRootTrianglePosition1Buffer.SRV.GetReference() : FNiagaraRenderer::GetDummyFloatBuffer();
-			FRHIShaderResourceView* DeformedTrianglePositionCSRV = (InterpolationModeValue == 1 && MeshProjection != nullptr) ?
-				MeshProjection->DeformedRootTrianglePosition2Buffer.SRV.GetReference() : FNiagaraRenderer::GetDummyFloatBuffer();
+			FRHIShaderResourceView* DeformedTrianglePositionASRV = (InterpolationModeValue == 1 && DeformedMeshProjection != nullptr) ?
+				DeformedMeshProjection->DeformedRootTrianglePosition0Buffer.SRV.GetReference() : FNiagaraRenderer::GetDummyFloatBuffer();
+			FRHIShaderResourceView* DeformedTrianglePositionBSRV = (InterpolationModeValue == 1 && DeformedMeshProjection != nullptr) ?
+				DeformedMeshProjection->DeformedRootTrianglePosition1Buffer.SRV.GetReference() : FNiagaraRenderer::GetDummyFloatBuffer();
+			FRHIShaderResourceView* DeformedTrianglePositionCSRV = (InterpolationModeValue == 1 && DeformedMeshProjection != nullptr) ?
+				DeformedMeshProjection->DeformedRootTrianglePosition2Buffer.SRV.GetReference() : FNiagaraRenderer::GetDummyFloatBuffer();
 
-			const int32 SampleCountValue = (MeshProjection != nullptr) ? MeshProjection->SampleCount : 0;
-			FShaderResourceViewRHIRef RestSamplePositionsBufferSRV = (MeshProjection != nullptr && MeshProjection->SampleCount > 0) ?
-				MeshProjection->RestSamplePositionsBuffer.SRV.GetReference() : FNiagaraRenderer::GetDummyFloatBuffer();
-			FShaderResourceViewRHIRef MeshSampleWeightsBufferSRV = (MeshProjection != nullptr && MeshProjection->SampleCount > 0) ?
-				MeshProjection->MeshSampleWeightsBuffer.SRV.GetReference() : FNiagaraRenderer::GetDummyFloatBuffer();
+			const bool bHasSamples = (RestMeshProjection != nullptr && RestMeshProjection->SampleCount > 0);
+			const int32 SampleCountValue = bHasSamples ? RestMeshProjection->SampleCount : 0;
+			FShaderResourceViewRHIRef RestSamplePositionsBufferSRV = bHasSamples ?
+				RestMeshProjection->RestSamplePositionsBuffer.SRV.GetReference() : FNiagaraRenderer::GetDummyFloatBuffer();
+			FShaderResourceViewRHIRef MeshSampleWeightsBufferSRV = bHasSamples ?
+				DeformedMeshProjection->MeshSampleWeightsBuffer.SRV.GetReference() : FNiagaraRenderer::GetDummyFloatBuffer();
 
-			FRHIShaderResourceView* RootBarycentricCoordinatesSRV = (InterpolationModeValue == 1 && MeshProjection != nullptr) ?
-				MeshProjection->RootTriangleBarycentricBuffer.SRV.GetReference() : FNiagaraRenderer::GetDummyFloatBuffer();
+			FRHIShaderResourceView* RootBarycentricCoordinatesSRV = (InterpolationModeValue == 1 && RestMeshProjection != nullptr) ?
+				RestMeshProjection->RootTriangleBarycentricBuffer.SRV.GetReference() : FNiagaraRenderer::GetDummyFloatBuffer();
 
 			int32 bNeedSimReset = (ProxyData->TickCount <= MaxDelay ? 1 : 0);
 
@@ -756,8 +765,8 @@ struct FNDIHairStrandsParametersCS : public FNiagaraDataInterfaceParametersCS
 			}
 
 			FUnorderedAccessViewRHIRef DummyFloatUAV(Context.Batcher->GetEmptyRWBufferFromPool(RHICmdList, PF_R32_FLOAT), false);
-			FUnorderedAccessViewRHIRef MeshSampleWeightsBufferUAV = (MeshProjection != nullptr && MeshProjection->SampleCount > 0) ?
-				MeshProjection->MeshSampleWeightsBuffer.UAV : DummyFloatUAV;
+			FUnorderedAccessViewRHIRef MeshSampleWeightsBufferUAV = (RestMeshProjection != nullptr && RestMeshProjection->SampleCount > 0) ?
+				DeformedMeshProjection->MeshSampleWeightsBuffer.UAV : DummyFloatUAV;
 
 			//UE_LOG(LogHairStrands, Log, TEXT("Shader Reset : %d %d %d"), bNeedSimReset, IsRootValid, InterpolationModeValue);
 			FVector RestRootOffsetValue = FVector::ZeroVector;
@@ -995,16 +1004,23 @@ void UNiagaraDataInterfaceHairStrands::ExtractSourceComponent(FNiagaraSystemInst
 	}
 }
 
-void UNiagaraDataInterfaceHairStrands::ExtractDatasAndResources(FNiagaraSystemInstance* SystemInstance,FHairStrandsDatas*& OutStrandsDatas, 
-	FHairStrandsRestResource*& OutStrandsRestResource, FHairStrandsDeformedResource*& OutStrandsDeformedResource, 
-	FHairStrandsRootResource*& OutStrandsRootResource, UGroomAsset*& OutGroomAsset, int32& OutGroupIndex)
+void UNiagaraDataInterfaceHairStrands::ExtractDatasAndResources(
+	FNiagaraSystemInstance* SystemInstance,
+	FHairStrandsDatas*& OutStrandsDatas, 
+	FHairStrandsRestResource*& OutStrandsRestResource, 
+	FHairStrandsDeformedResource*& OutStrandsDeformedResource, 
+	FHairStrandsRestRootResource*& OutStrandsRestRootResource, 
+	FHairStrandsDeformedRootResource*& OutStrandsDeformedRootResource,
+	UGroomAsset*& OutGroomAsset,
+	int32& OutGroupIndex)
 {
 	ExtractSourceComponent(SystemInstance);
 
 	OutStrandsDatas = nullptr;
 	OutStrandsRestResource = nullptr;
 	OutStrandsDeformedResource = nullptr;
-	OutStrandsRootResource = nullptr;
+	OutStrandsRestRootResource = nullptr;
+	OutStrandsDeformedRootResource = nullptr;
 	OutGroupIndex = -1;
 
 	if (IsComponentValid())
@@ -1022,7 +1038,8 @@ void UNiagaraDataInterfaceHairStrands::ExtractDatasAndResources(FNiagaraSystemIn
 			OutStrandsDatas = SourceComponent->GetGuideStrandsDatas(OutGroupIndex);
 			OutStrandsRestResource = SourceComponent->GetGuideStrandsRestResource(OutGroupIndex);
 			OutStrandsDeformedResource = SourceComponent->GetGuideStrandsDeformedResource(OutGroupIndex);
-			OutStrandsRootResource = SourceComponent->GetGuideStrandsRootResource(OutGroupIndex);
+			OutStrandsRestRootResource = SourceComponent->GetGuideStrandsRestRootResource(OutGroupIndex);
+			OutStrandsDeformedRootResource = SourceComponent->GetGuideStrandsDeformedRootResource(OutGroupIndex);
 			OutGroomAsset = SourceComponent->GroomAsset;
 		}
 	}
@@ -1069,14 +1086,15 @@ bool UNiagaraDataInterfaceHairStrands::PerInstanceTick(void* PerInstanceData, FN
 	FHairStrandsDatas* StrandsDatas = nullptr;
 	FHairStrandsRestResource* StrandsRestResource = nullptr;
 	FHairStrandsDeformedResource* StrandsDeformedResource = nullptr;
-	FHairStrandsRootResource* StrandsRootResource = nullptr;
+	FHairStrandsRestRootResource* StrandsRestRootResource = nullptr;
+	FHairStrandsDeformedRootResource* StrandsDeformedRootResource = nullptr;
 	UGroomAsset* GroomAsset = nullptr;
 	int32 GroupIndex = 0;
 
 	InstanceData->TickCount = FMath::Min(MaxDelay+1,InstanceData->TickCount+1);
 
-	ExtractDatasAndResources(SystemInstance, StrandsDatas, StrandsRestResource, StrandsDeformedResource, StrandsRootResource, GroomAsset, GroupIndex);
-	InstanceData->HairStrandsBuffer->Initialize(StrandsDatas, StrandsRestResource, StrandsDeformedResource, StrandsRootResource, InstanceData);
+	ExtractDatasAndResources(SystemInstance, StrandsDatas, StrandsRestResource, StrandsDeformedResource, StrandsRestRootResource, StrandsDeformedRootResource, GroomAsset, GroupIndex);
+	InstanceData->HairStrandsBuffer->Initialize(StrandsDatas, StrandsRestResource, StrandsDeformedResource, StrandsRestRootResource, StrandsDeformedRootResource, InstanceData);
 
 	if (SourceComponent != nullptr)
 	{
diff --git a/Engine/Plugins/Experimental/HairStrands/Source/HairStrandsNiagara/Public/NiagaraDataInterfaceHairStrands.h b/Engine/Plugins/Experimental/HairStrands/Source/HairStrandsNiagara/Public/NiagaraDataInterfaceHairStrands.h
--- a/Engine/Plugins/Experimental/HairStrands/Source/HairStrandsNiagara/Public/NiagaraDataInterfaceHairStrands.h
+++ b/Engine/Plugins/Experimental/HairStrands/Source/HairStrandsNiagara/Public/NiagaraDataInterfaceHairStrands.h
@@ -23,8 +23,12 @@ struct FNDIHairStrandsData;
 struct FNDIHairStrandsBuffer : public FRenderResource
 {
 	/** Set the asset that will be used to affect the buffer */
-	void Initialize(const FHairStrandsDatas*  HairStrandsDatas, const FHairStrandsRestResource*  HairStrandsRestResource, 
-		const FHairStrandsDeformedResource*  HairStrandsDeformedResource, const FHairStrandsRootResource* HairStrandsRootResource, const FNDIHairStrandsData* NDIStrandsDatas);
+	void Initialize(const FHairStrandsDatas*  HairStrandsDatas, 
+		const FHairStrandsRestResource*  HairStrandsRestResource, 
+		const FHairStrandsDeformedResource*  HairStrandsDeformedResource, 
+		const FHairStrandsRestRootResource* HairStrandsRestRootResource, 
+		const FHairStrandsDeformedRootResource* HairStrandsDeformedRootResource,
+		const FNDIHairStrandsData* NDIStrandsDatas);
 
 	/** Init the buffer */
 	virtual void InitRHI() override;
@@ -60,7 +64,10 @@ struct FNDIHairStrandsBuffer : public FRenderResource
 	const FHairStrandsDeformedResource* SourceDeformedResources;
 
 	/** The strand root resource to write into */
-	const FHairStrandsRootResource* SourceRootResources;
+	const FHairStrandsRestRootResource* SourceRestRootResources;
+	
+	/** The strand root resource to write into */
+	const FHairStrandsDeformedRootResource* SourceDeformedRootResources;
 
 	/** Niagara datas interface datas */
 	const FNDIHairStrandsData* NDIDatas;
@@ -324,9 +331,15 @@ public:
 	bool IsComponentValid() const;
 
 	/** Extract datas and resources */
-	void ExtractDatasAndResources(FNiagaraSystemInstance* SystemInstance, FHairStrandsDatas*& OutStrandsDatas,
-		FHairStrandsRestResource*& OutStrandsRestResource, FHairStrandsDeformedResource*& OutStrandsDeformedResource, 
-		FHairStrandsRootResource*& OutStrandsRootResource, UGroomAsset*& OutGroomAsset, int32& OutGroupIndex);
+	void ExtractDatasAndResources(
+		FNiagaraSystemInstance* SystemInstance, 
+		FHairStrandsDatas*& OutStrandsDatas,
+		FHairStrandsRestResource*& OutStrandsRestResource, 
+		FHairStrandsDeformedResource*& OutStrandsDeformedResource, 
+		FHairStrandsRestRootResource*& OutStrandsRestRootResource, 
+		FHairStrandsDeformedRootResource*& OutStrandsDeformedRootResource,
+		UGroomAsset*& OutGroomAsset,
+		int32& OutGroupIndex);
 
 	/** Get the number of strands */
 	void GetNumStrands(FVectorVMContext& Context);
diff --git a/Engine/Source/Runtime/Renderer/Private/HairStrands/HairStrandsDebug.cpp b/Engine/Source/Runtime/Renderer/Private/HairStrands/HairStrandsDebug.cpp
--- a/Engine/Source/Runtime/Renderer/Private/HairStrands/HairStrandsDebug.cpp
+++ b/Engine/Source/Runtime/Renderer/Private/HairStrands/HairStrandsDebug.cpp
@@ -1026,18 +1026,18 @@ static void AddDebugProjectionHairPass(
 	const EPrimitiveType PrimitiveType = GeometryType == EDebugProjectionHairType::HairFrame ? PT_LineList : PT_TriangleList;
 	const uint32 PrimitiveCount = HairData.RootCount;
 
-	if (PrimitiveCount == 0 || LODIndex < 0 || LODIndex >= HairData.LODDatas.Num())
+	if (PrimitiveCount == 0 || LODIndex < 0 || LODIndex >= HairData.RestLODDatas.Num() || LODIndex >= HairData.DeformedLODDatas.Num())
 		return;
 
-	if (EDebugProjectionHairType::HairFrame == GeometryType && (!HairData.RootPositionBuffer || !HairData.RootNormalBuffer || !HairData.LODDatas[LODIndex].RootTriangleBarycentricBuffer))
+	if (EDebugProjectionHairType::HairFrame == GeometryType && (!HairData.RootPositionBuffer || !HairData.RootNormalBuffer || !HairData.RestLODDatas[LODIndex].RootTriangleBarycentricBuffer))
 		return;
 
-	if (!HairData.LODDatas[LODIndex].RestRootTrianglePosition0Buffer ||
-		!HairData.LODDatas[LODIndex].RestRootTrianglePosition1Buffer ||
-		!HairData.LODDatas[LODIndex].RestRootTrianglePosition2Buffer ||
-		!HairData.LODDatas[LODIndex].DeformedRootTrianglePosition0Buffer ||
-		!HairData.LODDatas[LODIndex].DeformedRootTrianglePosition1Buffer ||
-		!HairData.LODDatas[LODIndex].DeformedRootTrianglePosition2Buffer)
+	if (!HairData.RestLODDatas[LODIndex].RestRootTrianglePosition0Buffer ||
+		!HairData.RestLODDatas[LODIndex].RestRootTrianglePosition1Buffer ||
+		!HairData.RestLODDatas[LODIndex].RestRootTrianglePosition2Buffer ||
+		!HairData.DeformedLODDatas[LODIndex].DeformedRootTrianglePosition0Buffer ||
+		!HairData.DeformedLODDatas[LODIndex].DeformedRootTrianglePosition1Buffer ||
+		!HairData.DeformedLODDatas[LODIndex].DeformedRootTrianglePosition2Buffer)
 		return;
 
 	const FIntRect Viewport = View->ViewRect;
@@ -1053,16 +1053,16 @@ static void AddDebugProjectionHairPass(
 	{
 		Parameters->RootPositionBuffer = HairData.RootPositionBuffer;
 		Parameters->RootNormalBuffer = HairData.RootNormalBuffer;
-		Parameters->RootBarycentricBuffer = HairData.LODDatas[LODIndex].RootTriangleBarycentricBuffer->SRV;
+		Parameters->RootBarycentricBuffer = HairData.RestLODDatas[LODIndex].RootTriangleBarycentricBuffer->SRV;
 	}
 
-	Parameters->RestPosition0Buffer = HairData.LODDatas[LODIndex].RestRootTrianglePosition0Buffer->SRV;
-	Parameters->RestPosition1Buffer = HairData.LODDatas[LODIndex].RestRootTrianglePosition1Buffer->SRV;
-	Parameters->RestPosition2Buffer = HairData.LODDatas[LODIndex].RestRootTrianglePosition2Buffer->SRV;
+	Parameters->RestPosition0Buffer = HairData.RestLODDatas[LODIndex].RestRootTrianglePosition0Buffer->SRV;
+	Parameters->RestPosition1Buffer = HairData.RestLODDatas[LODIndex].RestRootTrianglePosition1Buffer->SRV;
+	Parameters->RestPosition2Buffer = HairData.RestLODDatas[LODIndex].RestRootTrianglePosition2Buffer->SRV;
 	
-	Parameters->DeformedPosition0Buffer = HairData.LODDatas[LODIndex].DeformedRootTrianglePosition0Buffer->SRV;
-	Parameters->DeformedPosition1Buffer = HairData.LODDatas[LODIndex].DeformedRootTrianglePosition1Buffer->SRV;
-	Parameters->DeformedPosition2Buffer = HairData.LODDatas[LODIndex].DeformedRootTrianglePosition2Buffer->SRV;
+	Parameters->DeformedPosition0Buffer = HairData.DeformedLODDatas[LODIndex].DeformedRootTrianglePosition0Buffer->SRV;
+	Parameters->DeformedPosition1Buffer = HairData.DeformedLODDatas[LODIndex].DeformedRootTrianglePosition1Buffer->SRV;
+	Parameters->DeformedPosition2Buffer = HairData.DeformedLODDatas[LODIndex].DeformedRootTrianglePosition2Buffer->SRV;
 
 	Parameters->ViewUniformBuffer = View->ViewUniformBuffer;
 	Parameters->RenderTargets[0] = FRenderTargetBinding(ColorTarget, ERenderTargetLoadAction::ELoad, 0);
diff --git a/Engine/Source/Runtime/Renderer/Private/HairStrands/HairStrandsInterface.cpp b/Engine/Source/Runtime/Renderer/Private/HairStrands/HairStrandsInterface.cpp
--- a/Engine/Source/Runtime/Renderer/Private/HairStrands/HairStrandsInterface.cpp
+++ b/Engine/Source/Runtime/Renderer/Private/HairStrands/HairStrandsInterface.cpp
@@ -301,7 +301,7 @@ void RunProjection(
 		FRDGBuilder GraphBuilder(RHICmdList);
 		for (FHairStrandsProjectionHairData::HairGroup& HairGroup : ProjectionHairData.HairGroups)
 		{
-			for (FHairStrandsProjectionHairData::LODData& LODData : HairGroup.LODDatas)
+			for (FHairStrandsProjectionHairData::RestLODData& LODData : HairGroup.RestLODDatas)
 			{
 				const uint32 LODIndex = LODData.LODIndex;
 				HairGroup.LocalToWorld = LocalToWorld;
@@ -362,29 +362,38 @@ void RunHairStrandsInterpolation(
 
 		FRDGBuilder GraphBuilder(RHICmdList);
 
-		for (FHairStrandsProjectionHairData::HairGroup& ProjectionHairData : E.RenProjectionHairDatas.HairGroups)
+		if (0 <= E.FrameLODIndex)
 		{
-			if (EHairStrandsInterpolationType::RenderStrands == Type && 0 <= E.FrameLODIndex && E.FrameLODIndex < ProjectionHairData.LODDatas.Num() && ProjectionHairData.LODDatas[E.FrameLODIndex].bIsValid)
+			if (EHairStrandsInterpolationType::RenderStrands == Type)
 			{
-				UpdateHairStrandsMeshTriangles(GraphBuilder, ShaderMap, E.FrameLODIndex, HairStrandsTriangleType::DeformedPose, MeshDataLOD, ProjectionHairData, TransitionQueue);
+				for (FHairStrandsProjectionHairData::HairGroup& ProjectionHairData : E.RenProjectionHairDatas.HairGroups)
+				{
+					if (E.FrameLODIndex < ProjectionHairData.DeformedLODDatas.Num() && ProjectionHairData.DeformedLODDatas[E.FrameLODIndex].IsValid())
+					{
+						UpdateHairStrandsMeshTriangles(GraphBuilder, ShaderMap, E.FrameLODIndex, HairStrandsTriangleType::DeformedPose, MeshDataLOD, ProjectionHairData, TransitionQueue);
+					}
+				}
 			}
-		}
 
-		for (FHairStrandsProjectionHairData::HairGroup& ProjectionHairData : E.SimProjectionHairDatas.HairGroups)
-		{
-			if (EHairStrandsInterpolationType::SimulationStrands == Type && 0 <= E.FrameLODIndex && E.FrameLODIndex < ProjectionHairData.LODDatas.Num() && ProjectionHairData.LODDatas[E.FrameLODIndex].bIsValid)
+			if (EHairStrandsInterpolationType::SimulationStrands == Type)
 			{
-				UpdateHairStrandsMeshTriangles(GraphBuilder, ShaderMap, E.FrameLODIndex, HairStrandsTriangleType::DeformedPose, MeshDataLOD, ProjectionHairData, TransitionQueue);
-			}
-		}
+				for (FHairStrandsProjectionHairData::HairGroup& ProjectionHairData : E.SimProjectionHairDatas.HairGroups)
+				{
+					if (E.FrameLODIndex < ProjectionHairData.DeformedLODDatas.Num() && ProjectionHairData.DeformedLODDatas[E.FrameLODIndex].IsValid())
+					{
+						UpdateHairStrandsMeshTriangles(GraphBuilder, ShaderMap, E.FrameLODIndex, HairStrandsTriangleType::DeformedPose, MeshDataLOD, ProjectionHairData, TransitionQueue);
+					}
+				}
 
-		for (FHairStrandsProjectionHairData::HairGroup& ProjectionHairData : E.SimProjectionHairDatas.HairGroups)
-		{
-			if (EHairStrandsInterpolationType::SimulationStrands == Type && 0 <= E.FrameLODIndex && E.FrameLODIndex < ProjectionHairData.LODDatas.Num() && ProjectionHairData.LODDatas[E.FrameLODIndex].bIsValid)
-			{
-				InitHairStrandsMeshSamples(GraphBuilder, ShaderMap, E.FrameLODIndex, HairStrandsTriangleType::DeformedPose, MeshDataLOD, ProjectionHairData, TransitionQueue);
-				UpdateHairStrandsMeshSamples(GraphBuilder, ShaderMap, E.FrameLODIndex, MeshDataLOD, ProjectionHairData, TransitionQueue);
-				//InterpolateHairStrandsMeshTriangles(GraphBuilder, ShaderMap, E.FrameLODIndex, MeshDataLOD, ProjectionHairData);
+				for (FHairStrandsProjectionHairData::HairGroup& ProjectionHairData : E.SimProjectionHairDatas.HairGroups)
+				{
+					if (E.FrameLODIndex < ProjectionHairData.DeformedLODDatas.Num() && ProjectionHairData.DeformedLODDatas[E.FrameLODIndex].IsValid())
+					{
+						InitHairStrandsMeshSamples(GraphBuilder, ShaderMap, E.FrameLODIndex, HairStrandsTriangleType::DeformedPose, MeshDataLOD, ProjectionHairData, TransitionQueue);
+						UpdateHairStrandsMeshSamples(GraphBuilder, ShaderMap, E.FrameLODIndex, MeshDataLOD, ProjectionHairData, TransitionQueue);
+						//InterpolateHairStrandsMeshTriangles(GraphBuilder, ShaderMap, E.FrameLODIndex, MeshDataLOD, ProjectionHairData);
+					}
+				}
 			}
 		}
 
@@ -537,8 +546,8 @@ FHairStrandsDebugInfos GetHairStandsDebugInfos()
 			if (GroupIt < uint32(E.RenProjectionHairDatas.HairGroups.Num()))
 			{
 				FHairStrandsProjectionHairData::HairGroup& ProjectionHair = E.RenProjectionHairDatas.HairGroups[GroupIt];
-				GroupInfo.LODCount = ProjectionHair.LODDatas.Num();
-				GroupInfo.bHasSkinInterpolation = ProjectionHair.LODDatas.Num() > 0;
+				GroupInfo.LODCount = ProjectionHair.DeformedLODDatas.Num();
+				GroupInfo.bHasSkinInterpolation = ProjectionHair.DeformedLODDatas.Num() > 0;
 			}
 			else
 			{
diff --git a/Engine/Source/Runtime/Renderer/Private/HairStrands/HairStrandsMeshProjection.cpp b/Engine/Source/Runtime/Renderer/Private/HairStrands/HairStrandsMeshProjection.cpp
--- a/Engine/Source/Runtime/Renderer/Private/HairStrands/HairStrandsMeshProjection.cpp
+++ b/Engine/Source/Runtime/Renderer/Private/HairStrands/HairStrandsMeshProjection.cpp
@@ -248,9 +248,9 @@ static void AddHairStrandMeshProjectionPass(
 {
 	if (!RootData.RootPositionBuffer ||
 		!RootData.RootNormalBuffer ||
-		LODIndex < 0 || LODIndex >= RootData.LODDatas.Num() ||
-		!RootData.LODDatas[LODIndex].RootTriangleIndexBuffer ||
-		!RootData.LODDatas[LODIndex].RootTriangleBarycentricBuffer ||
+		LODIndex < 0 || LODIndex >= RootData.RestLODDatas.Num() ||
+		!RootData.RestLODDatas[LODIndex].RootTriangleIndexBuffer ||
+		!RootData.RestLODDatas[LODIndex].RootTriangleBarycentricBuffer ||
 		!MeshSectionData.IndexBuffer ||
 		!MeshSectionData.PositionBuffer ||
 		MeshSectionData.TotalIndexCount == 0 ||
@@ -299,8 +299,8 @@ static void AddHairStrandMeshProjectionPass(
 		Parameters->MeshPrimitiveCount_Iteration	= (MeshPassIt < MeshPassCount-1) ? MeshPassNumPrimitive : (MeshSectionData.NumPrimitives - MeshPassNumPrimitive * MeshPassIt);
 
 		// The projection is always done onto the source/rest mesh
-		Parameters->OutRootTriangleIndex		= RootData.LODDatas[LODIndex].RootTriangleIndexBuffer->UAV;
-		Parameters->OutRootTriangleBarycentrics = RootData.LODDatas[LODIndex].RootTriangleBarycentricBuffer->UAV;
+		Parameters->OutRootTriangleIndex		= RootData.RestLODDatas[LODIndex].RootTriangleIndexBuffer->UAV;
+		Parameters->OutRootTriangleBarycentrics = RootData.RestLODDatas[LODIndex].RootTriangleBarycentricBuffer->UAV;
 		Parameters->OutRootTriangleDistance		= DistanceUAV;
 
 		const FIntVector DispatchGroupCount = FComputeShaderUtils::GetGroupCount(RootData.RootCount, 128);
@@ -313,8 +313,8 @@ static void AddHairStrandMeshProjectionPass(
 			Parameters,
 			DispatchGroupCount);
 
-		OutTransitionQueue.Add(RootData.LODDatas[LODIndex].RootTriangleIndexBuffer->UAV);
-		OutTransitionQueue.Add(RootData.LODDatas[LODIndex].RootTriangleBarycentricBuffer->UAV);
+		OutTransitionQueue.Add(RootData.RestLODDatas[LODIndex].RootTriangleIndexBuffer->UAV);
+		OutTransitionQueue.Add(RootData.RestLODDatas[LODIndex].RootTriangleBarycentricBuffer->UAV);
 	}
 }
 
@@ -326,7 +326,7 @@ void ProjectHairStrandsOntoMesh(
 	FHairStrandsProjectionHairData::HairGroup& ProjectionHairData, 
 	FBufferTransitionQueue& TransitionQueue)
 {
-	if (LODIndex < 0 || LODIndex >= ProjectionHairData.LODDatas.Num())
+	if (LODIndex < 0 || LODIndex >= ProjectionHairData.RestLODDatas.Num())
 		return;
 
 	FRDGBufferRef RootDistanceBuffer = GraphBuilder.CreateBuffer(FRDGBufferDesc::CreateStructuredDesc(sizeof(float), ProjectionHairData.RootCount), TEXT("HairStrandsTriangleDistance"));
@@ -334,9 +334,10 @@ void ProjectHairStrandsOntoMesh(
 	bool ClearDistance = true;
 	for (const FHairStrandsProjectionMeshData::Section& MeshSection : ProjectionMeshData.LODs[LODIndex].Sections)
 	{
-		check(ProjectionHairData.LODDatas[LODIndex].LODIndex == LODIndex);
+		check(ProjectionHairData.RestLODDatas[LODIndex].LODIndex == LODIndex);
 		AddHairStrandMeshProjectionPass(GraphBuilder, ShaderMap, ClearDistance, LODIndex, MeshSection, ProjectionHairData, RootDistanceBuffer, TransitionQueue);
-		ProjectionHairData.LODDatas[LODIndex].bIsValid = true;
+		check(ProjectionHairData.RestLODDatas[LODIndex].Status);
+		*ProjectionHairData.RestLODDatas[LODIndex].Status = FHairStrandsProjectionHairData::EStatus::Completed;
 		ClearDistance = false;
 	}
 }
@@ -449,12 +450,23 @@ static void AddHairStrandUpdateMeshTrianglesPass(
 	FHairStrandsProjectionHairData::HairGroup& RootData, 
 	FBufferTransitionQueue& OutTransitionQueue)
 {
-	if (RootData.RootCount == 0 || LODIndex < 0 || LODIndex >= RootData.LODDatas.Num())
+	if (RootData.RootCount == 0 || LODIndex < 0)
 	{
 		return;
 	}
-	FHairStrandsProjectionHairData::LODData& LODData = RootData.LODDatas[LODIndex];
-	check(LODData.LODIndex == LODIndex);
+
+	if (Type == HairStrandsTriangleType::RestPose && LODIndex >= RootData.RestLODDatas.Num())
+	{
+		return;
+	}
+
+	if (Type == HairStrandsTriangleType::DeformedPose && (LODIndex >= RootData.DeformedLODDatas.Num() || LODIndex >= RootData.RestLODDatas.Num()))
+	{
+		return;
+	}
+
+	FHairStrandsProjectionHairData::RestLODData& RestLODData = RootData.RestLODDatas[LODIndex];
+	check(RestLODData.LODIndex == LODIndex);
 
 	const int32 SectionCount = MeshData.Sections.Num();
 	FHairUpdateMeshTriangleCS::FParameters* Parameters = GraphBuilder.AllocParameters<FHairUpdateMeshTriangleCS::FParameters>();
@@ -581,27 +593,31 @@ static void AddHairStrandUpdateMeshTrianglesPass(
 		}
 	}
 
-	Parameters->RootTriangleIndex = LODData.RootTriangleIndexBuffer->SRV;
+	Parameters->RootTriangleIndex = RestLODData.RootTriangleIndexBuffer->SRV;
 	if (Type == HairStrandsTriangleType::RestPose)
 	{
-		Parameters->OutRootTrianglePosition0 = LODData.RestRootTrianglePosition0Buffer->UAV;
-		Parameters->OutRootTrianglePosition1 = LODData.RestRootTrianglePosition1Buffer->UAV;
-		Parameters->OutRootTrianglePosition2 = LODData.RestRootTrianglePosition2Buffer->UAV;
+		Parameters->OutRootTrianglePosition0 = RestLODData.RestRootTrianglePosition0Buffer->UAV;
+		Parameters->OutRootTrianglePosition1 = RestLODData.RestRootTrianglePosition1Buffer->UAV;
+		Parameters->OutRootTrianglePosition2 = RestLODData.RestRootTrianglePosition2Buffer->UAV;
 
-		OutTransitionQueue.Add(LODData.RestRootTrianglePosition0Buffer->UAV);
-		OutTransitionQueue.Add(LODData.RestRootTrianglePosition1Buffer->UAV);
-		OutTransitionQueue.Add(LODData.RestRootTrianglePosition2Buffer->UAV);
+		OutTransitionQueue.Add(RestLODData.RestRootTrianglePosition0Buffer->UAV);
+		OutTransitionQueue.Add(RestLODData.RestRootTrianglePosition1Buffer->UAV);
+		OutTransitionQueue.Add(RestLODData.RestRootTrianglePosition2Buffer->UAV);
 	}
 	else if (Type == HairStrandsTriangleType::DeformedPose)
 	{
-		Parameters->OutRootTrianglePosition0 = LODData.DeformedRootTrianglePosition0Buffer->UAV;
-		Parameters->OutRootTrianglePosition1 = LODData.DeformedRootTrianglePosition1Buffer->UAV;
-		Parameters->OutRootTrianglePosition2 = LODData.DeformedRootTrianglePosition2Buffer->UAV;
-		if (LODData.Status) (*LODData.Status) = FHairStrandsProjectionHairData::LODData::EStatus::Completed;
-
-		OutTransitionQueue.Add(LODData.DeformedRootTrianglePosition0Buffer->UAV);
-		OutTransitionQueue.Add(LODData.DeformedRootTrianglePosition1Buffer->UAV);
-		OutTransitionQueue.Add(LODData.DeformedRootTrianglePosition2Buffer->UAV);
+		FHairStrandsProjectionHairData::DeformedLODData& DeformedLODData = RootData.DeformedLODDatas[LODIndex];
+		check(DeformedLODData.LODIndex == LODIndex);
+
+		Parameters->OutRootTrianglePosition0 = DeformedLODData.DeformedRootTrianglePosition0Buffer->UAV;
+		Parameters->OutRootTrianglePosition1 = DeformedLODData.DeformedRootTrianglePosition1Buffer->UAV;
+		Parameters->OutRootTrianglePosition2 = DeformedLODData.DeformedRootTrianglePosition2Buffer->UAV;
+		check(DeformedLODData.Status);
+		(*DeformedLODData.Status) = FHairStrandsProjectionHairData::EStatus::Completed;
+
+		OutTransitionQueue.Add(DeformedLODData.DeformedRootTrianglePosition0Buffer->UAV);
+		OutTransitionQueue.Add(DeformedLODData.DeformedRootTrianglePosition1Buffer->UAV);
+		OutTransitionQueue.Add(DeformedLODData.DeformedRootTrianglePosition2Buffer->UAV);
 	}
 	else
 	{
@@ -688,27 +704,29 @@ static void AddHairStrandInterpolateMeshTrianglesPass(
 	FHairStrandsProjectionHairData::HairGroup& RootData,
 	FBufferTransitionQueue& OutTransitionQueue)
 {
-	if (RootData.RootCount == 0 || LODIndex < 0 || LODIndex >= RootData.LODDatas.Num())
+	if (RootData.RootCount == 0 || LODIndex < 0 || LODIndex >= RootData.RestLODDatas.Num() || LODIndex >= RootData.DeformedLODDatas.Num())
 	{
 		return;
 	}
-	FHairStrandsProjectionHairData::LODData& LODData = RootData.LODDatas[LODIndex];
-	check(LODData.LODIndex == LODIndex);
+	FHairStrandsProjectionHairData::RestLODData& RestLODData = RootData.RestLODDatas[LODIndex];
+	FHairStrandsProjectionHairData::DeformedLODData& DeformedLODData = RootData.DeformedLODDatas[LODIndex];
+	check(RestLODData.LODIndex == LODIndex);
+	check(DeformedLODData.LODIndex == LODIndex);
 
 	FHairInterpolateMeshTriangleCS::FParameters* Parameters = GraphBuilder.AllocParameters<FHairInterpolateMeshTriangleCS::FParameters>();
 	Parameters->MaxRootCount = RootData.RootCount;
-	Parameters->MaxSampleCount = LODData.SampleCount;
+	Parameters->MaxSampleCount = RestLODData.SampleCount;
 
-	Parameters->RestRootTrianglePosition0 = LODData.RestRootTrianglePosition0Buffer->SRV;
-	Parameters->RestRootTrianglePosition1 = LODData.RestRootTrianglePosition1Buffer->SRV;
-	Parameters->RestRootTrianglePosition2 = LODData.RestRootTrianglePosition2Buffer->SRV;
+	Parameters->RestRootTrianglePosition0 = RestLODData.RestRootTrianglePosition0Buffer->SRV;
+	Parameters->RestRootTrianglePosition1 = RestLODData.RestRootTrianglePosition1Buffer->SRV;
+	Parameters->RestRootTrianglePosition2 = RestLODData.RestRootTrianglePosition2Buffer->SRV;
 
-	Parameters->OutDeformedRootTrianglePosition0 = LODData.DeformedRootTrianglePosition0Buffer->UAV;
-	Parameters->OutDeformedRootTrianglePosition1 = LODData.DeformedRootTrianglePosition1Buffer->UAV;
-	Parameters->OutDeformedRootTrianglePosition2 = LODData.DeformedRootTrianglePosition2Buffer->UAV;
+	Parameters->OutDeformedRootTrianglePosition0 = DeformedLODData.DeformedRootTrianglePosition0Buffer->UAV;
+	Parameters->OutDeformedRootTrianglePosition1 = DeformedLODData.DeformedRootTrianglePosition1Buffer->UAV;
+	Parameters->OutDeformedRootTrianglePosition2 = DeformedLODData.DeformedRootTrianglePosition2Buffer->UAV;
 
-	Parameters->MeshSampleWeightsBuffer = LODData.MeshSampleWeightsBuffer->SRV;
-	Parameters->RestSamplePositionsBuffer = LODData.RestSamplePositionsBuffer->SRV;
+	Parameters->MeshSampleWeightsBuffer = DeformedLODData.MeshSampleWeightsBuffer->SRV;
+	Parameters->RestSamplePositionsBuffer = RestLODData.RestSamplePositionsBuffer->SRV;
 
 	const FIntVector DispatchGroupCount = FComputeShaderUtils::GetGroupCount(RootData.RootCount, 128);
 	check(DispatchGroupCount.X < 65536);
@@ -720,9 +738,9 @@ static void AddHairStrandInterpolateMeshTrianglesPass(
 		Parameters,
 		DispatchGroupCount);
 
-	OutTransitionQueue.Add(LODData.DeformedRootTrianglePosition0Buffer->UAV);
-	OutTransitionQueue.Add(LODData.DeformedRootTrianglePosition1Buffer->UAV);
-	OutTransitionQueue.Add(LODData.DeformedRootTrianglePosition2Buffer->UAV);
+	OutTransitionQueue.Add(DeformedLODData.DeformedRootTrianglePosition0Buffer->UAV);
+	OutTransitionQueue.Add(DeformedLODData.DeformedRootTrianglePosition1Buffer->UAV);
+	OutTransitionQueue.Add(DeformedLODData.DeformedRootTrianglePosition2Buffer->UAV);
 }
 
 void InterpolateHairStrandsMeshTriangles(
@@ -775,39 +793,53 @@ static void AddHairStrandInitMeshSamplesPass(
 	FHairStrandsProjectionHairData::HairGroup& RootData,
 	FBufferTransitionQueue& OutTransitionQueue)
 {
-	if (LODIndex < 0 || LODIndex >= RootData.LODDatas.Num())
+	if (LODIndex < 0)
 	{
 		return;
 	}
-	FHairStrandsProjectionHairData::LODData& LODData = RootData.LODDatas[LODIndex];
-	check(LODData.LODIndex == LODIndex);
+
+	if (Type == HairStrandsTriangleType::RestPose && LODIndex >= RootData.RestLODDatas.Num())
+	{
+		return;
+	}
+
+	if (Type == HairStrandsTriangleType::DeformedPose && (LODIndex >= RootData.RestLODDatas.Num() || LODIndex >= RootData.DeformedLODDatas.Num()))
+	{
+		return;
+	}
+
+	FHairStrandsProjectionHairData::RestLODData& RestLODData = RootData.RestLODDatas[LODIndex];
+	check(RestLODData.LODIndex == LODIndex);
 
 	const uint32 SectionCount = MeshData.Sections.Num();
-	if (SectionCount > 0 && LODData.SampleCount > 0)
+	if (SectionCount > 0 && RestLODData.SampleCount > 0)
 	{
 		FHairInitMeshSamplesCS::FParameters* Parameters = GraphBuilder.AllocParameters<FHairInitMeshSamplesCS::FParameters>();
 
 		Parameters->MaxVertexCount = MeshData.Sections[0].TotalVertexCount;
 		Parameters->VertexPositionsBuffer = MeshData.Sections[0].PositionBuffer;
 
-		Parameters->MaxSampleCount = LODData.SampleCount;
-		Parameters->SampleIndicesBuffer = LODData.MeshSampleIndicesBuffer->SRV;
+		Parameters->MaxSampleCount = RestLODData.SampleCount;
+		Parameters->SampleIndicesBuffer = RestLODData.MeshSampleIndicesBuffer->SRV;
 		if (Type == HairStrandsTriangleType::RestPose)
 		{
-			Parameters->OutSamplePositionsBuffer = LODData.RestSamplePositionsBuffer->UAV;
-			OutTransitionQueue.Add(LODData.RestSamplePositionsBuffer->UAV);
+			Parameters->OutSamplePositionsBuffer = RestLODData.RestSamplePositionsBuffer->UAV;
+			OutTransitionQueue.Add(RestLODData.RestSamplePositionsBuffer->UAV);
 		}
 		else if (Type == HairStrandsTriangleType::DeformedPose)
 		{
-			Parameters->OutSamplePositionsBuffer = LODData.DeformedSamplePositionsBuffer->UAV;
-			OutTransitionQueue.Add(LODData.DeformedSamplePositionsBuffer->UAV);
+			FHairStrandsProjectionHairData::DeformedLODData& DeformedLODData = RootData.DeformedLODDatas[LODIndex];
+			check(DeformedLODData.LODIndex == LODIndex);
+
+			Parameters->OutSamplePositionsBuffer = DeformedLODData.DeformedSamplePositionsBuffer->UAV;
+			OutTransitionQueue.Add(DeformedLODData.DeformedSamplePositionsBuffer->UAV);
 		}
 		else
 		{
 			return;
 		}
 
-		const FIntVector DispatchGroupCount = FComputeShaderUtils::GetGroupCount(LODData.SampleCount, 128);
+		const FIntVector DispatchGroupCount = FComputeShaderUtils::GetGroupCount(RestLODData.SampleCount, 128);
 		check(DispatchGroupCount.X < 65536);
 		TShaderMapRef<FHairInitMeshSamplesCS> ComputeShader(ShaderMap);
 		FComputeShaderUtils::AddPass(
@@ -869,26 +901,29 @@ static void AddHairStrandUpdateMeshSamplesPass(
 	FHairStrandsProjectionHairData::HairGroup& RootData,
 	FBufferTransitionQueue& OutTransitionQueue)
 {
-	if (LODIndex < 0 || LODIndex >= RootData.LODDatas.Num())
+	if (LODIndex < 0 || LODIndex >= RootData.RestLODDatas.Num() || LODIndex >= RootData.DeformedLODDatas.Num())
 	{
 		return;
 	}
-	FHairStrandsProjectionHairData::LODData& LODData = RootData.LODDatas[LODIndex];
-	check(LODData.LODIndex == LODIndex);
+
+	FHairStrandsProjectionHairData::RestLODData& RestLODData = RootData.RestLODDatas[LODIndex];
+	FHairStrandsProjectionHairData::DeformedLODData& DeformedLODData = RootData.DeformedLODDatas[LODIndex];
+	check(RestLODData.LODIndex == LODIndex);
+	check(DeformedLODData.LODIndex == LODIndex);
 
 	const uint32 SectionCount = MeshData.Sections.Num();
-	if (SectionCount > 0 && LODData.SampleCount > 0)
+	if (SectionCount > 0 && RestLODData.SampleCount > 0)
 	{
 		FHairUpdateMeshSamplesCS::FParameters* Parameters = GraphBuilder.AllocParameters<FHairUpdateMeshSamplesCS::FParameters>();
 
-		Parameters->MaxSampleCount = LODData.SampleCount;
-		Parameters->SampleIndicesBuffer = LODData.MeshSampleIndicesBuffer->SRV;
-		Parameters->InterpolationWeightsBuffer = LODData.MeshInterpolationWeightsBuffer->SRV;
-		Parameters->SampleRestPositionsBuffer = LODData.RestSamplePositionsBuffer->SRV;
-		Parameters->SampleDeformedPositionsBuffer = LODData.DeformedSamplePositionsBuffer->SRV;
-		Parameters->OutSampleDeformationsBuffer = LODData.MeshSampleWeightsBuffer->UAV;
+		Parameters->MaxSampleCount = RestLODData.SampleCount;
+		Parameters->SampleIndicesBuffer = RestLODData.MeshSampleIndicesBuffer->SRV;
+		Parameters->InterpolationWeightsBuffer = RestLODData.MeshInterpolationWeightsBuffer->SRV;
+		Parameters->SampleRestPositionsBuffer = RestLODData.RestSamplePositionsBuffer->SRV;
+		Parameters->SampleDeformedPositionsBuffer = DeformedLODData.DeformedSamplePositionsBuffer->SRV;
+		Parameters->OutSampleDeformationsBuffer = DeformedLODData.MeshSampleWeightsBuffer->UAV;
 
-		const FIntVector DispatchGroupCount = FComputeShaderUtils::GetGroupCount(LODData.SampleCount+4, 128);
+		const FIntVector DispatchGroupCount = FComputeShaderUtils::GetGroupCount(RestLODData.SampleCount+4, 128);
 		check(DispatchGroupCount.X < 65536);
 		TShaderMapRef<FHairUpdateMeshSamplesCS> ComputeShader(ShaderMap);
 		FComputeShaderUtils::AddPass(
@@ -898,7 +933,7 @@ static void AddHairStrandUpdateMeshSamplesPass(
 			Parameters,
 			DispatchGroupCount);
 
-		OutTransitionQueue.Add(LODData.MeshSampleWeightsBuffer->UAV);
+		OutTransitionQueue.Add(DeformedLODData.MeshSampleWeightsBuffer->UAV);
 	}
 }
 
@@ -981,10 +1016,10 @@ static void AddFollicleMaskPass(
 	const FHairStrandsProjectionHairData::HairGroup& HairData,
 	FRDGTextureRef OutTexture)
 {
-	if (LODIndex >= uint32(HairData.LODDatas.Num()) || HairData.RootCount == 0)
+	if (LODIndex >= uint32(HairData.RestLODDatas.Num()) || HairData.RootCount == 0)
 		return;
 
-	const FHairStrandsProjectionHairData::LODData& LODData = HairData.LODDatas[LODIndex];
+	const FHairStrandsProjectionHairData::RestLODData& LODData = HairData.RestLODDatas[LODIndex];
 	if (!LODData.RootTriangleBarycentricBuffer ||
 		!LODData.RestRootTrianglePosition0Buffer ||
 		!LODData.RestRootTrianglePosition1Buffer ||
diff --git a/Engine/Source/Runtime/Renderer/Public/HairStrandsInterface.h b/Engine/Source/Runtime/Renderer/Public/HairStrandsInterface.h
--- a/Engine/Source/Runtime/Renderer/Public/HairStrandsInterface.h
+++ b/Engine/Source/Runtime/Renderer/Public/HairStrandsInterface.h
@@ -212,12 +212,13 @@ struct FHairStrandsProjectionMeshData
 
 struct FHairStrandsProjectionHairData
 {
-	struct LODData
+	enum class EStatus { Invalid, Initialized, Completed };
+
+	struct RestLODData
 	{
-		enum class EStatus { Invalid, Initialized, Completed };
-		bool bIsValid = false;
 		int32 LODIndex = -1;
 		EStatus* Status = nullptr;
+		inline bool IsValid() const { return Status && (*Status == EStatus::Initialized || *Status == EStatus::Completed); }
 
 		// The index buffers stores the mesh section & the triangle index into a single uint32 
 		// (3 highest bits store the section (up to 8 sections)
@@ -231,17 +232,25 @@ struct FHairStrandsProjectionHairData
 		FRWBuffer* RestRootTrianglePosition1Buffer = nullptr;
 		FRWBuffer* RestRootTrianglePosition2Buffer = nullptr;
 
+		// Samples to be used for RBF mesh interpolation
+		uint32 SampleCount = 0;
+		FRWBuffer* MeshInterpolationWeightsBuffer = nullptr;
+		FRWBuffer* MeshSampleIndicesBuffer = nullptr;
+		FRWBuffer* RestSamplePositionsBuffer = nullptr;
+	};
+
+	struct DeformedLODData
+	{
+		int32 LODIndex = -1;
+		EStatus* Status = nullptr;
+		inline bool IsValid() const { return Status && (*Status == EStatus::Initialized || *Status == EStatus::Completed); }
+
 		// Deformed root triangles' positions are relative to root center (for preserving precision)
 		FRWBuffer* DeformedRootTrianglePosition0Buffer = nullptr;
 		FRWBuffer* DeformedRootTrianglePosition1Buffer = nullptr;
 		FRWBuffer* DeformedRootTrianglePosition2Buffer = nullptr;
 
-		// Samples to be used for rbf mesh interpolatrion
-		uint32 SampleCount = 0;
-		FRWBuffer* MeshInterpolationWeightsBuffer = nullptr;
-		FRWBuffer* MeshSampleIndicesBuffer = nullptr;
-
-		FRWBuffer* RestSamplePositionsBuffer = nullptr;
+		// Samples to be used for RBF mesh interpolation
 		FRWBuffer* DeformedSamplePositionsBuffer = nullptr;
 		FRWBuffer* MeshSampleWeightsBuffer = nullptr;
 	};
@@ -252,7 +261,8 @@ struct FHairStrandsProjectionHairData
 		FRHIShaderResourceView* RootNormalBuffer = nullptr;
 		FRWBuffer* VertexToCurveIndexBuffer = nullptr;
 
-		TArray<LODData> LODDatas;
+		TArray<RestLODData> RestLODDatas;
+		TArray<DeformedLODData> DeformedLODDatas;
 
 		uint32 RootCount = 0;
 		FTransform LocalToWorld = FTransform::Identity;
@@ -402,7 +412,7 @@ RENDERER_API void RunProjection(
 	FHairStrandsProjectionHairData& SimProjectionHairData);
 
 RENDERER_API FHairStrandsProjectionMeshData ExtractMeshData(class FSkeletalMeshRenderData* RenderData);
-FHairStrandsProjectionHairData::HairGroup ToProjectionHairData(struct FHairStrandsRootResource* In);
+FHairStrandsProjectionHairData::HairGroup ToProjectionHairData(struct FHairStrandsRestRootResource* InRest, struct FHairStrandsDeformedRootResource* InDeformed);
 
 typedef void (*TBindingProcess)(FRHICommandListImmediate& RHICmdList, void* Asset);
 RENDERER_API void EnqueueGroomBindingQuery(void* Asset, TBindingProcess BindingProcess);
