Remove unecessary transition

diff --git a/Engine/Plugins/Experimental/HairStrands/Source/HairStrandsNiagara/Private/NiagaraDataInterfaceHairStrands.cpp b/Engine/Plugins/Experimental/HairStrands/Source/HairStrandsNiagara/Private/NiagaraDataInterfaceHairStrands.cpp
--- a/Engine/Plugins/Experimental/HairStrands/Source/HairStrandsNiagara/Private/NiagaraDataInterfaceHairStrands.cpp
+++ b/Engine/Plugins/Experimental/HairStrands/Source/HairStrandsNiagara/Private/NiagaraDataInterfaceHairStrands.cpp
@@ -764,9 +764,9 @@ struct FNDIHairStrandsParametersCS : public FNiagaraDataInterfaceParametersCS
 				InterpolationModeValue = 2;
 			}
 
-			FUnorderedAccessViewRHIRef DummyFloatUAV(Context.Batcher->GetEmptyRWBufferFromPool(RHICmdList, PF_R32_FLOAT), false);
-			FUnorderedAccessViewRHIRef MeshSampleWeightsBufferUAV = (RestMeshProjection != nullptr && RestMeshProjection->SampleCount > 0) ?
-				DeformedMeshProjection->MeshSampleWeightsBuffer.UAV : DummyFloatUAV;
+			//FUnorderedAccessViewRHIRef DummyFloatUAV(Context.Batcher->GetEmptyRWBufferFromPool(RHICmdList, PF_R32_FLOAT), false);
+			//FUnorderedAccessViewRHIRef MeshSampleWeightsBufferUAV = (RestMeshProjection != nullptr && RestMeshProjection->SampleCount > 0) ?
+			//	DeformedMeshProjection->MeshSampleWeightsBuffer.UAV : Context.Batcher->GetEmptyRWBufferFromPool(RHICmdList, PF_R32_FLOAT);
 
 			//UE_LOG(LogHairStrands, Log, TEXT("Shader Reset : %d %d %d"), bNeedSimReset, IsRootValid, InterpolationModeValue);
 			FVector RestRootOffsetValue = FVector::ZeroVector;
@@ -811,7 +811,7 @@ struct FNDIHairStrandsParametersCS : public FNiagaraDataInterfaceParametersCS
 			SetShaderValue(RHICmdList, ComputeShaderRHI, SampleCount, SampleCountValue);
 			SetSRVParameter(RHICmdList, ComputeShaderRHI, RestSamplePositionsBuffer, RestSamplePositionsBufferSRV);
 
-			RHICmdList.TransitionResource(EResourceTransitionAccess::EReadable, EResourceTransitionPipeline::EComputeToCompute, MeshSampleWeightsBufferUAV);
+			//RHICmdList.TransitionResource(EResourceTransitionAccess::EReadable, EResourceTransitionPipeline::EComputeToCompute, MeshSampleWeightsBufferUAV);
 			SetSRVParameter(RHICmdList, ComputeShaderRHI, MeshSampleWeightsBuffer, MeshSampleWeightsBufferSRV);
 
 			SetSRVParameter(RHICmdList, ComputeShaderRHI, RootBarycentricCoordinatesBuffer, RootBarycentricCoordinatesSRV);
