Fix CIS issue.

diff --git a/Engine/Plugins/Experimental/HairStrands/Source/HairStrandsNiagara/Private/NiagaraDataInterfaceHairStrands.cpp b/Engine/Plugins/Experimental/HairStrands/Source/HairStrandsNiagara/Private/NiagaraDataInterfaceHairStrands.cpp
--- a/Engine/Plugins/Experimental/HairStrands/Source/HairStrandsNiagara/Private/NiagaraDataInterfaceHairStrands.cpp
+++ b/Engine/Plugins/Experimental/HairStrands/Source/HairStrandsNiagara/Private/NiagaraDataInterfaceHairStrands.cpp
@@ -748,7 +748,7 @@ struct FNDIHairStrandsParametersCS : public FNiagaraDataInterfaceParametersCS
 			const int32 SampleCountValue = bHasSamples ? RestMeshProjection->SampleCount : 0;
 			FShaderResourceViewRHIRef RestSamplePositionsBufferSRV = bHasSamples ?
 				RestMeshProjection->RestSamplePositionsBuffer.SRV.GetReference() : FNiagaraRenderer::GetDummyFloatBuffer();
-			FShaderResourceViewRHIRef MeshSampleWeightsBufferSRV = bHasSamples ?
+			FShaderResourceViewRHIRef MeshSampleWeightsBufferSRV = bHasSamples && DeformedMeshProjection != nullptr ?
 				DeformedMeshProjection->MeshSampleWeightsBuffer.SRV.GetReference() : FNiagaraRenderer::GetDummyFloatBuffer();
 
 			FRHIShaderResourceView* RootBarycentricCoordinatesSRV = (InterpolationModeValue == 1 && RestMeshProjection != nullptr) ?
@@ -765,7 +765,7 @@ struct FNDIHairStrandsParametersCS : public FNiagaraDataInterfaceParametersCS
 			}
 
 			//FUnorderedAccessViewRHIRef DummyFloatUAV(Context.Batcher->GetEmptyRWBufferFromPool(RHICmdList, PF_R32_FLOAT), false);
-			//FUnorderedAccessViewRHIRef MeshSampleWeightsBufferUAV = (RestMeshProjection != nullptr && RestMeshProjection->SampleCount > 0) ?
+			//FUnorderedAccessViewRHIRef MeshSampleWeightsBufferUAV = (DeformedMeshProjection != nullptr && RestMeshProjection->SampleCount > 0) ?
 			//	DeformedMeshProjection->MeshSampleWeightsBuffer.UAV : Context.Batcher->GetEmptyRWBufferFromPool(RHICmdList, PF_R32_FLOAT);
 
 			//UE_LOG(LogHairStrands, Log, TEXT("Shader Reset : %d %d %d"), bNeedSimReset, IsRootValid, InterpolationModeValue);
