Custom solver support

diff --git a/Engine/Plugins/Experimental/HairStrands/Source/HairStrandsCore/Private/GroomComponent.cpp b/Engine/Plugins/Experimental/HairStrands/Source/HairStrandsCore/Private/GroomComponent.cpp
--- a/Engine/Plugins/Experimental/HairStrands/Source/HairStrandsCore/Private/GroomComponent.cpp
+++ b/Engine/Plugins/Experimental/HairStrands/Source/HairStrandsCore/Private/GroomComponent.cpp
@@ -545,7 +545,6 @@ void UGroomComponent::UpdateHairSimulation()
 {
 	static UNiagaraSystem* CosseratRodsSystem = LoadObject<UNiagaraSystem>(nullptr, TEXT("/HairStrands/Emitters/GroomRodsSystem.GroomRodsSystem"));
 	static UNiagaraSystem* AngularSpringsSystem = LoadObject<UNiagaraSystem>(nullptr, TEXT("/HairStrands/Emitters/GroomSpringsSystem.GroomSpringsSystem"));
-
 	const int32 NumGroups = GroomAsset ? GroomAsset->HairGroupsPhysics.Num() : 0;
 	const int32 NumComponents = FMath::Max(NumGroups, NiagaraComponents.Num());
 
@@ -582,13 +581,16 @@ void UGroomComponent::UpdateHairSimulation()
 			if (GroomAsset->HairGroupsPhysics[i].SolverSettings.NiagaraSolver == EGroomNiagaraSolvers::AngularSprings)
 			{
 				NiagaraComponent->SetAsset(AngularSpringsSystem);
-				NiagaraComponent->Activate(true);
 			}
 			else if (GroomAsset->HairGroupsPhysics[i].SolverSettings.NiagaraSolver == EGroomNiagaraSolvers::CosseratRods)
 			{
 				NiagaraComponent->SetAsset(CosseratRodsSystem);
-				NiagaraComponent->Activate(true);
 			}
+			else 
+			{
+				NiagaraComponent->SetAsset(GroomAsset->HairGroupsPhysics[i].SolverSettings.CustomSystem.LoadSynchronous());
+			}
+			NiagaraComponent->Activate(true);
 			if (NiagaraComponent->GetSystemInstance())
 			{
 				NiagaraComponent->GetSystemInstance()->Reset(FNiagaraSystemInstance::EResetMode::ReInit);
diff --git a/Engine/Plugins/Experimental/HairStrands/Source/HairStrandsCore/Public/GroomAsset.h b/Engine/Plugins/Experimental/HairStrands/Source/HairStrandsCore/Public/GroomAsset.h
--- a/Engine/Plugins/Experimental/HairStrands/Source/HairStrandsCore/Public/GroomAsset.h
+++ b/Engine/Plugins/Experimental/HairStrands/Source/HairStrandsCore/Public/GroomAsset.h
@@ -24,7 +24,8 @@ enum class EGroomNiagaraSolvers : uint8
 {
 	None = 0 UMETA(Hidden),
 	CosseratRods = 0x02 UMETA(DisplatName = "GroomRods"),
-	AngularSprings = 0x04 UMETA(DisplatName = "GroomSprings")
+	AngularSprings = 0x04 UMETA(DisplatName = "GroomSprings"),
+	CustomSolver = 0x08 UMETA(DisplatName = "CustomSolver")
 };
 
 /** Size of each strands*/
@@ -366,6 +367,10 @@ struct HAIRSTRANDSCORE_API FHairSolverSettings
 	UPROPERTY(EditAnywhere, Category = "SolverSettings", meta = (ToolTip = "Niagara solver to be used for simulation"))
 	EGroomNiagaraSolvers NiagaraSolver;
 
+	/** Custom system to be used for simulation */
+	UPROPERTY(EditAnywhere, Category = "SolverSettings", meta = (EditCondition = "NiagaraSolver == EGroomNiagaraSolvers::CustomSolver", ToolTip = "Custom niagara system to be used if custom solver is picked"))
+	TSoftObjectPtr<UNiagaraSystem> CustomSystem;
+
 	/** Number of substeps to be used */
 	UPROPERTY(EditAnywhere, Category = "SolverSettings", meta = (ToolTip = "Number of sub steps to be done per frame. The actual solver calls are done at 24 fps"))
 	int32 SubSteps;
@@ -419,7 +424,7 @@ struct HAIRSTRANDSCORE_API FHairBendConstraint
 	float BendStiffness;
 
 	/** Stiffness scale along the strand */
-	UPROPERTY(EditAnywhere, Category = "BendConstraint", meta = (ViewMinInput = "0.0", ViewMaxInput = "1.0", ViewMinOutput = "0.0", ViewMaxOutput = "1.0", TimeLineLength = "1.0", XAxisName = "Strand Coordinate (0,1)", YAxisName = "Bend Scale", ToolTip = "This curve determines how much the bend stiffness will be scaled along each strand. \n The X axis range is [0,1], 0 mapping the root and 1 the tip"))
+	UPROPERTY(EditAnywhere, Category = "BendConstraint", meta = (DisplayName = "Stiffness Scale", ViewMinInput = "0.0", ViewMaxInput = "1.0", ViewMinOutput = "0.0", ViewMaxOutput = "1.0", TimeLineLength = "1.0", XAxisName = "Strand Coordinate (0,1)", YAxisName = "Bend Scale", ToolTip = "This curve determines how much the bend stiffness will be scaled along each strand. \n The X axis range is [0,1], 0 mapping the root and 1 the tip"))
 	FRuntimeFloatCurve BendScale;
 };
 
@@ -448,7 +453,7 @@ struct HAIRSTRANDSCORE_API FHairStretchConstraint
 	float StretchStiffness;
 
 	/** Stretch scale along the strand  */
-	UPROPERTY(EditAnywhere, Category = "StretchConstraint", meta = (ViewMinInput = "0.0", ViewMaxInput = "1.0", ViewMinOutput = "0.0", ViewMaxOutput = "1.0", TimeLineLength = "1.0", XAxisName = "Strand Coordinate (0,1)", YAxisName = "Stretch Scale", ToolTip = "This curve determines how much the stretch stiffness will be scaled along each strand. \n The X axis range is [0,1], 0 mapping the root and 1 the tip"))
+	UPROPERTY(EditAnywhere, Category = "StretchConstraint", meta = (DisplayName = "Stiffness Scale", ViewMinInput = "0.0", ViewMaxInput = "1.0", ViewMinOutput = "0.0", ViewMaxOutput = "1.0", TimeLineLength = "1.0", XAxisName = "Strand Coordinate (0,1)", YAxisName = "Stretch Scale", ToolTip = "This curve determines how much the stretch stiffness will be scaled along each strand. \n The X axis range is [0,1], 0 mapping the root and 1 the tip"))
 	FRuntimeFloatCurve StretchScale;
 };
 
