Fix groom properties not being forward when a groom component placed within a blueprint.

diff --git a/Engine/Plugins/Experimental/HairStrands/Source/HairStrandsCore/Private/GroomComponent.cpp b/Engine/Plugins/Experimental/HairStrands/Source/HairStrandsCore/Private/GroomComponent.cpp
--- a/Engine/Plugins/Experimental/HairStrands/Source/HairStrandsCore/Private/GroomComponent.cpp
+++ b/Engine/Plugins/Experimental/HairStrands/Source/HairStrandsCore/Private/GroomComponent.cpp
@@ -1381,6 +1381,10 @@ void UGroomComponent::OnRegister()
 	{
 		InitResources();
 	}
+	else
+	{
+		UpdateHairGroupsDescAndInvalidateRenderState();
+	}
 
 	// Insure the ticking of the Groom component always happens after the skeletalMeshComponent.
 	USkeletalMeshComponent* SkeletalMeshComponent = GetAttachParent() ? Cast<USkeletalMeshComponent>(GetAttachParent()) : nullptr;
@@ -1616,16 +1620,17 @@ void UGroomComponent::PostEditChangeProperty(FPropertyChangedEvent& PropertyChan
 		GroomAsset = nullptr;
 	}
 
+	bool bIsEventProcess = false;
+
 	const bool bRecreateResources = bAssetChanged || bBindingAssetChanged || PropertyThatChanged == nullptr || bBindToSkeletalChanged || bSourceSkeletalMeshChanged;
 	if (bRecreateResources)
 	{
 		// Release the resources before Super::PostEditChangeProperty so that they get
 		// re-initialized in OnRegister
 		ReleaseResources();
+		bIsEventProcess = true;
 	}
 
-	Super::PostEditChangeProperty(PropertyChangedEvent);
-
 	const bool bSupportSkinProjection = GetDefault<URendererSettings>()->bSupportSkinCacheShaders;
 	if (!bSupportSkinProjection)
 	{
@@ -1667,6 +1672,12 @@ void UGroomComponent::PostEditChangeProperty(FPropertyChangedEvent& PropertyChan
 		PropertyName == GET_MEMBER_NAME_CHECKED(FHairGroupDesc, bScatterSceneLighting))
 	{	
 		UpdateHairGroupsDescAndInvalidateRenderState();
+		bIsEventProcess = true;
+	}
+
+	if (!bIsEventProcess)
+	{
+		Super::PostEditChangeProperty(PropertyChangedEvent);
 	}
 
 #if WITH_EDITOR
