Fix hair interpolation

diff --git a/Engine/Plugins/Experimental/HairStrands/Source/HairStrandsCore/Private/HairStrandsRendering.cpp b/Engine/Plugins/Experimental/HairStrands/Source/HairStrandsCore/Private/HairStrandsRendering.cpp
--- a/Engine/Plugins/Experimental/HairStrands/Source/HairStrandsCore/Private/HairStrandsRendering.cpp
+++ b/Engine/Plugins/Experimental/HairStrands/Source/HairStrandsCore/Private/HairStrandsRendering.cpp
@@ -506,12 +506,14 @@ static void AddHairStrandsInterpolationPass(
 		ShaderDrawDebug::SetParameters(GraphBuilder, *ShaderDrawData, Parameters->ShaderDrawParameters);
 	}
 
+	const bool bHasLocalDeformation = (bHasSimulationEnable || bSupportGlobalInterpolation);
+
 	FHairInterpolationCS::FPermutationDomain PermutationVector;
 	PermutationVector.Set<FHairInterpolationCS::FGroupSize>(GetGroupSizePermutation(GroupSize));
 	PermutationVector.Set<FHairInterpolationCS::FDebug>(Parameters->HairDebugMode > 0 ? 1 : 0);
-	PermutationVector.Set<FHairInterpolationCS::FDynamicGeometry>((bSupportGlobalInterpolation && (GStrandHairInterpolationGlobal == 1)) ? 2 : 
-								bSupportDynamicMesh ? 1 : 0);
-	PermutationVector.Set<FHairInterpolationCS::FSimulation>((bHasSimulationEnable || bSupportGlobalInterpolation) ? 1 : 0);
+	PermutationVector.Set<FHairInterpolationCS::FDynamicGeometry>((bSupportDynamicMesh && bHasLocalDeformation) ? 2 :
+							(bSupportDynamicMesh && !bHasLocalDeformation) ? 1 : 0);
+	PermutationVector.Set<FHairInterpolationCS::FSimulation>(bHasLocalDeformation ? 1 : 0);
 	PermutationVector.Set<FHairInterpolationCS::FScaleAndClip>(bNeedScaleOrClip ? 1 : 0);
 
 	FGlobalShaderMap* ShaderMap = GetGlobalShaderMap(ERHIFeatureLevel::SM5);
