Removed unused legacy shader parameter construction code

diff --git a/Engine/Plugins/Experimental/HairStrands/Source/HairStrandsCore/Private/HairStrandsVertexFactory.cpp b/Engine/Plugins/Experimental/HairStrands/Source/HairStrandsCore/Private/HairStrandsVertexFactory.cpp
--- a/Engine/Plugins/Experimental/HairStrands/Source/HairStrandsCore/Private/HairStrandsVertexFactory.cpp
+++ b/Engine/Plugins/Experimental/HairStrands/Source/HairStrandsCore/Private/HairStrandsVertexFactory.cpp
@@ -220,28 +220,6 @@ IMPLEMENT_VERTEX_FACTORY_PARAMETER_TYPE(FHairStrandsVertexFactory, SF_Compute,
 IMPLEMENT_VERTEX_FACTORY_PARAMETER_TYPE(FHairStrandsVertexFactory, SF_RayHitGroup,	FHairStrandsVertexFactoryShaderParameters);
 #endif
 
-FVertexFactoryShaderParameters* FHairStrandsVertexFactory::ConstructShaderParameters(EShaderFrequency ShaderFrequency)
-{
-	if (ShaderFrequency == SF_Vertex)
-	{
-		return new FHairStrandsVertexFactoryShaderParameters();
-	}
-
-	if (ShaderFrequency == SF_Pixel)
-	{
-		return new FHairStrandsVertexFactoryShaderParameters();
-	}
-
-#if RHI_RAYTRACING
-	if (ShaderFrequency == SF_RayHitGroup)
-	{
-		return new FHairStrandsVertexFactoryShaderParameters();
-	}
-#endif // RHI_RAYTRACING
-
-	return NULL;
-}
-
 void FHairStrandsVertexFactory::ReleaseRHI()
 {
 	FVertexFactory::ReleaseRHI();
diff --git a/Engine/Plugins/Experimental/HairStrands/Source/HairStrandsCore/Public/HairStrandsVertexFactory.h b/Engine/Plugins/Experimental/HairStrands/Source/HairStrandsCore/Public/HairStrandsVertexFactory.h
--- a/Engine/Plugins/Experimental/HairStrands/Source/HairStrandsCore/Public/HairStrandsVertexFactory.h
+++ b/Engine/Plugins/Experimental/HairStrands/Source/HairStrandsCore/Public/HairStrandsVertexFactory.h
@@ -60,8 +60,6 @@ public:
 
 	static bool SupportsTessellationShaders() { return false; }
 
-	static FVertexFactoryShaderParameters* ConstructShaderParameters(EShaderFrequency ShaderFrequency);
-
 	inline FRHIShaderResourceView* GetPositionSRV(uint32 GroupIndex) const			{ check(Data.InterpolationOutput); return Data.InterpolationOutput->HairGroups[GroupIndex].VFInput.HairPositionBuffer; };
 	inline FRHIShaderResourceView* GetPreviousPositionSRV(uint32 GroupIndex) const	{ check(Data.InterpolationOutput); return Data.InterpolationOutput->HairGroups[GroupIndex].VFInput.HairPreviousPositionBuffer; }
 	inline FRHIShaderResourceView* GetAttributeSRV(uint32 GroupIndex) const			{ check(Data.InterpolationOutput); return Data.InterpolationOutput->HairGroups[GroupIndex].VFInput.HairAttributeBuffer; }
