Fix teleport issue. Expose Grid size to groom asset. Clean niagara assets. Update of the physics params. Fix projection and rods material.

diff --git a/Engine/Plugins/Experimental/HairStrands/Shaders/Private/NiagaraAngularSpringMaterial.ush b/Engine/Plugins/Experimental/HairStrands/Shaders/Private/NiagaraAngularSpringMaterial.ush
--- a/Engine/Plugins/Experimental/HairStrands/Shaders/Private/NiagaraAngularSpringMaterial.ush
+++ b/Engine/Plugins/Experimental/HairStrands/Shaders/Private/NiagaraAngularSpringMaterial.ush
@@ -149,7 +149,7 @@ void ProjectAngularSpringMaterial(in bool EnableConstraint, in int StrandsSize,
 		{
 			for(int i = 2; i < StrandsSize; ++i)
 			{
-				ComputeMaterialFrame(StrandsSize);
+				//ComputeMaterialFrame(StrandsSize);
 				if( LocalIndex == i )
 				{
 					UpdateAngularSpringMultiplier(RestLength,RestDirection,DeltaTime,true,0.0,MaterialCompliance,MaterialWeight,MaterialMultiplier);
diff --git a/Engine/Plugins/Experimental/HairStrands/Shaders/Private/NiagaraCosseratRodMaterial.ush b/Engine/Plugins/Experimental/HairStrands/Shaders/Private/NiagaraCosseratRodMaterial.ush
--- a/Engine/Plugins/Experimental/HairStrands/Shaders/Private/NiagaraCosseratRodMaterial.ush
+++ b/Engine/Plugins/Experimental/HairStrands/Shaders/Private/NiagaraCosseratRodMaterial.ush
@@ -65,7 +65,7 @@ void SolveAlignRodMaterial(in bool EnableConstraint, in int StrandsSize, in floa
 			in float MaterialCompliance, in float MaterialWeight, in float3 MaterialMultiplier, out float3 OutMaterialMultiplier)
 {
 	OutMaterialMultiplier = MaterialMultiplier;
-	if(EnableConstraint)
+	if(DeltaTime != 0.0 && EnableConstraint)
 	{
 		const int LocalIndex = (GGroupThreadId.x % StrandsSize);
 		if( LocalIndex > 1)
@@ -87,7 +87,7 @@ void SolveAlignRodMaterial(in bool EnableConstraint, in int StrandsSize, in floa
 	}
 } 
 
-void ProjectAlignRodMaterial(in bool EnableConstraint, in int StrandsSize, in float YoungModulus, in float RodThickness, in float RestLength, in float DeltaTime, out float3 OutNodePosition)
+void ProjectAlignRodMaterial(in bool EnableConstraint, in int StrandsSize, in float YoungModulus, in float RodThickness, in float RestLength, in float DeltaTime, out float4 OutNodeOrientation)
 {
 	const int LocalIndex = (GGroupThreadId.x % StrandsSize);
 	if(DeltaTime != 0.0 && EnableConstraint)
@@ -111,7 +111,7 @@ void ProjectAlignRodMaterial(in bool EnableConstraint, in int StrandsSize, in fl
 		}
 	}
 	GroupMemoryBarrier();
-	OutNodePosition = SharedNodePosition[GGroupThreadId.x];
+	OutNodeOrientation = SharedNodeOrientation[GGroupThreadId.x];
 }
 
 /* -----------------------------------------------------------------
@@ -175,7 +175,7 @@ void SolveStretchRodMaterial(in bool EnableConstraint, in int StrandsSize, in fl
 			in float MaterialCompliance, in float MaterialWeight, in float3 MaterialMultiplier, out float3 OutMaterialMultiplier)
 {
 	OutMaterialMultiplier = MaterialMultiplier;
-	if(EnableConstraint)
+	if(DeltaTime != 0.0 && EnableConstraint)
 	{
 		const int LocalIndex = (GGroupThreadId.x % StrandsSize);
 		if( LocalIndex > 1)
@@ -295,7 +295,8 @@ void SolveBendRodMaterial(in bool EnableConstraint, in int StrandsSize, in float
 {
 	OutMaterialMultiplier = MaterialMultiplier;
 
-	if(EnableConstraint)
+	UpdateMaterialFrame(StrandsSize);
+	if(DeltaTime != 0.0 && EnableConstraint)
 	{
 		const int LocalIndex = (GGroupThreadId.x % StrandsSize);
 		if( LocalIndex > 0)
@@ -317,9 +318,11 @@ void SolveBendRodMaterial(in bool EnableConstraint, in int StrandsSize, in float
 	}
 } 
 
-void ProjectBendRodMaterial(in bool EnableConstraint, in int StrandsSize, in float YoungModulus, in float RodThickness, in float RestLength, in float4 RestDarboux, in float DeltaTime, out float3 OutNodePosition)
+void ProjectBendRodMaterial(in bool EnableConstraint, in int StrandsSize, in float YoungModulus, in float RodThickness, in float RestLength, in float4 RestDarboux, in float DeltaTime, out float4 OutNodeOrientation)
 {
 	const int LocalIndex = (GGroupThreadId.x % StrandsSize);
+
+	UpdateMaterialFrame(StrandsSize);
 	if(DeltaTime != 0.0 && EnableConstraint)
 	{
 		float MaterialCompliance = 0.0;
@@ -341,7 +344,7 @@ void ProjectBendRodMaterial(in bool EnableConstraint, in int StrandsSize, in flo
 		}
 	}
 	GroupMemoryBarrier();
-	OutNodePosition = SharedNodePosition[GGroupThreadId.x];
+	OutNodeOrientation = SharedNodeOrientation[GGroupThreadId.x];
 }
 
 /* -----------------------------------------------------------------
diff --git a/Engine/Plugins/Experimental/HairStrands/Source/HairStrandsCore/Private/GroomAsset.cpp b/Engine/Plugins/Experimental/HairStrands/Source/HairStrandsCore/Private/GroomAsset.cpp
--- a/Engine/Plugins/Experimental/HairStrands/Source/HairStrandsCore/Private/GroomAsset.cpp
+++ b/Engine/Plugins/Experimental/HairStrands/Source/HairStrandsCore/Private/GroomAsset.cpp
@@ -813,6 +813,7 @@ FHairCollisionConstraint::FHairCollisionConstraint()
 	StaticFriction = 0.1;
 	StrandsViscosity = 1.0;
 	CollisionRadius = 0.1;
+	GridDimension = FIntVector(30,30,30);
 
 	RadiusScale.GetRichCurve()->SetKeyInterpMode(RadiusScale.GetRichCurve()->AddKey(0.f, 1.0f), ERichCurveInterpMode::RCIM_Cubic);
 	RadiusScale.GetRichCurve()->SetKeyInterpMode(RadiusScale.GetRichCurve()->AddKey(1.f, 0.1f), ERichCurveInterpMode::RCIM_Cubic);
diff --git a/Engine/Plugins/Experimental/HairStrands/Source/HairStrandsCore/Public/GroomAsset.h b/Engine/Plugins/Experimental/HairStrands/Source/HairStrandsCore/Public/GroomAsset.h
--- a/Engine/Plugins/Experimental/HairStrands/Source/HairStrandsCore/Public/GroomAsset.h
+++ b/Engine/Plugins/Experimental/HairStrands/Source/HairStrandsCore/Public/GroomAsset.h
@@ -484,6 +484,10 @@ struct HAIRSTRANDSCORE_API FHairCollisionConstraint
 	UPROPERTY(EditAnywhere, Category = "CollisionConstraint", meta = (ToolTip = "Viscosity parameter between 0 and 1 that will be used for self collision"))
 	float StrandsViscosity;
 
+	/** Dimension of the grid used to compute the viscosity force */
+	UPROPERTY(EditAnywhere, Category = "CollisionConstraint", meta = (ToolTip = "Dimension of the grid used to compute the viscosity force"))
+	FIntVector GridDimension;
+
 	/** Radius that will be used for the collision detection against the physics asset */
 	UPROPERTY(EditAnywhere, Category = "CollisionConstraint", meta = (ToolTip = "Radius that will be used for the collision detection against the physics asset"))
 	float CollisionRadius;
diff --git a/Engine/Plugins/Experimental/HairStrands/Source/HairStrandsNiagara/Private/NiagaraDataInterfaceHairStrands.cpp b/Engine/Plugins/Experimental/HairStrands/Source/HairStrandsNiagara/Private/NiagaraDataInterfaceHairStrands.cpp
--- a/Engine/Plugins/Experimental/HairStrands/Source/HairStrandsNiagara/Private/NiagaraDataInterfaceHairStrands.cpp
+++ b/Engine/Plugins/Experimental/HairStrands/Source/HairStrandsNiagara/Private/NiagaraDataInterfaceHairStrands.cpp
@@ -55,6 +55,7 @@ static const FName GetProjectCollisionName("GetProjectCollision");
 static const FName GetStaticFrictionName("GetStaticFriction");
 static const FName GetKineticFrictionName("GetKineticFriction");
 static const FName GetStrandsViscosityName("GetStrandsViscosity");
+static const FName GetGridDimensionName("GetGridDimension");
 static const FName GetCollisionRadiusName("GetCollisionRadius");
 static const FName GetRadiusScaleName("GetRadiusScale");
 
@@ -506,10 +507,10 @@ void FNDIHairStrandsData::Release()
 void FNDIHairStrandsData::Update(UNiagaraDataInterfaceHairStrands* Interface, FNiagaraSystemInstance* SystemInstance, const FHairStrandsDatas* StrandsDatas, 
 	UGroomAsset* GroomAsset, const int32 GroupIndex)
 {
-	if (Interface != nullptr && SystemInstance != nullptr)
+	if (Interface != nullptr)
 	{
 		WorldTransform = Interface->IsComponentValid() ? Interface->SourceComponent->GetComponentToWorld() :
-			SystemInstance->GetComponent()->GetComponentToWorld();
+			SystemInstance ? SystemInstance->GetComponent()->GetComponentToWorld() : FTransform::Identity;
 
 		if (StrandsDatas != nullptr && GroomAsset != nullptr && GroupIndex >= 0 && GroupIndex < GroomAsset->HairGroupsPhysics.Num() &&
 			GroomAsset->HairGroupsPhysics[GroupIndex].SolverSettings.EnableSimulation)
@@ -540,6 +541,7 @@ void FNDIHairStrandsData::Update(UNiagaraDataInterfaceHairStrands* Interface, FN
 			StaticFriction = HairPhysics.MaterialConstraints.CollisionConstraint.StaticFriction;
 			KineticFriction = HairPhysics.MaterialConstraints.CollisionConstraint.KineticFriction;
 			StrandsViscosity = HairPhysics.MaterialConstraints.CollisionConstraint.StrandsViscosity;
+			GridDimension = HairPhysics.MaterialConstraints.CollisionConstraint.GridDimension;
 			CollisionRadius = HairPhysics.MaterialConstraints.CollisionConstraint.CollisionRadius;
 
 			StrandsDensity = HairPhysics.StrandsParameters.StrandsDensity;
@@ -572,7 +574,7 @@ bool FNDIHairStrandsData::Init(UNiagaraDataInterfaceHairStrands* Interface, FNia
 {
 	HairStrandsBuffer = nullptr;
 
-	if (Interface != nullptr && SystemInstance != nullptr)
+	if (Interface != nullptr)
 	{
 		FHairStrandsDatas* StrandsDatas = nullptr;
 		FHairStrandsRestResource* StrandsRestResource = nullptr;
@@ -976,7 +978,7 @@ void UNiagaraDataInterfaceHairStrands::ExtractSourceComponent(FNiagaraSystemInst
 			SourceComponent = SourceActor->FindComponentByClass<UGroomComponent>();
 		}
 	}
-	else
+	else if(SystemInstance)
 	{
 		if (UNiagaraComponent* SimComp = SystemInstance->GetComponent())
 		{
@@ -1023,7 +1025,7 @@ void UNiagaraDataInterfaceHairStrands::ExtractDatasAndResources(
 	OutStrandsDeformedRootResource = nullptr;
 	OutGroupIndex = -1;
 
-	if (IsComponentValid())
+	if (IsComponentValid() && SystemInstance)
 	{
 		for (int32 NiagaraIndex = 0, NiagaraCount = SourceComponent->NiagaraComponents.Num(); NiagaraIndex < NiagaraCount; ++NiagaraIndex)
 		{
@@ -1367,6 +1369,16 @@ void UNiagaraDataInterfaceHairStrands::GetFunctions(TArray<FNiagaraFunctionSigna
 
 		OutFunctions.Add(Sig);
 	}
+	{
+		FNiagaraFunctionSignature Sig;
+		Sig.Name = GetGridDimensionName;
+		Sig.bMemberFunction = true;
+		Sig.bRequiresContext = false;
+		Sig.Inputs.Add(FNiagaraVariable(FNiagaraTypeDefinition(GetClass()), TEXT("Hair Strands")));
+		Sig.Outputs.Add(FNiagaraVariable(FNiagaraTypeDefinition::GetVec3Def(), TEXT("Grid Dimension")));
+
+		OutFunctions.Add(Sig);
+	}
 	{
 		FNiagaraFunctionSignature Sig;
 		Sig.Name = GetCollisionRadiusName;
@@ -1899,7 +1911,7 @@ void UNiagaraDataInterfaceHairStrands::GetFunctions(TArray<FNiagaraFunctionSigna
 		Sig.Inputs.Add(FNiagaraVariable(FNiagaraTypeDefinition(GetClass()), TEXT("Hair Strands")));
 		Sig.Inputs.Add(FNiagaraVariable(FNiagaraTypeDefinition::GetBoolDef(), TEXT("Enable Constraint")));
 		Sig.Inputs.Add(FNiagaraVariable(FNiagaraTypeDefinition::GetFloatDef(), TEXT("Rest Length")));
-		Sig.Inputs.Add(FNiagaraVariable(FNiagaraTypeDefinition::GetQuatDef(), TEXT("Rest Darboux")));
+		Sig.Inputs.Add(FNiagaraVariable(FNiagaraTypeDefinition::GetQuatDef(), TEXT("Rest Rotation")));
 		Sig.Inputs.Add(FNiagaraVariable(FNiagaraTypeDefinition::GetFloatDef(), TEXT("Delta Time")));
 		Sig.Inputs.Add(FNiagaraVariable(FNiagaraTypeDefinition::GetFloatDef(), TEXT("Material Damping")));
 		Sig.Inputs.Add(FNiagaraVariable(FNiagaraTypeDefinition::GetFloatDef(), TEXT("Material Compliance")));
@@ -1918,9 +1930,9 @@ void UNiagaraDataInterfaceHairStrands::GetFunctions(TArray<FNiagaraFunctionSigna
 		Sig.Inputs.Add(FNiagaraVariable(FNiagaraTypeDefinition::GetFloatDef(), TEXT("Bend Stiffness")));
 		Sig.Inputs.Add(FNiagaraVariable(FNiagaraTypeDefinition::GetFloatDef(), TEXT("Node Thickness")));
 		Sig.Inputs.Add(FNiagaraVariable(FNiagaraTypeDefinition::GetFloatDef(), TEXT("Rest Length")));
-		Sig.Inputs.Add(FNiagaraVariable(FNiagaraTypeDefinition::GetQuatDef(), TEXT("Rest Darboux")));
+		Sig.Inputs.Add(FNiagaraVariable(FNiagaraTypeDefinition::GetQuatDef(), TEXT("Rest Rotation")));
 		Sig.Inputs.Add(FNiagaraVariable(FNiagaraTypeDefinition::GetFloatDef(), TEXT("Delta Time")));
-		Sig.Outputs.Add(FNiagaraVariable(FNiagaraTypeDefinition::GetVec3Def(), TEXT("Node Position")));
+		Sig.Outputs.Add(FNiagaraVariable(FNiagaraTypeDefinition::GetQuatDef(), TEXT("Node Orientation")));
 
 		OutFunctions.Add(Sig);
 	}
@@ -2034,7 +2046,7 @@ void UNiagaraDataInterfaceHairStrands::GetFunctions(TArray<FNiagaraFunctionSigna
 		Sig.bMemberFunction = true;
 		Sig.bRequiresContext = false;
 		Sig.Inputs.Add(FNiagaraVariable(FNiagaraTypeDefinition(GetClass()), TEXT("Hair Strands")));
-		Sig.Outputs.Add(FNiagaraVariable(FNiagaraTypeDefinition::GetBoolDef(), TEXT("Update Status")));
+		Sig.Outputs.Add(FNiagaraVariable(FNiagaraTypeDefinition::GetQuatDef(), TEXT("Node Orientation")));
 
 		OutFunctions.Add(Sig);
 	}
@@ -2044,7 +2056,7 @@ void UNiagaraDataInterfaceHairStrands::GetFunctions(TArray<FNiagaraFunctionSigna
 		Sig.bMemberFunction = true;
 		Sig.bRequiresContext = false;
 		Sig.Inputs.Add(FNiagaraVariable(FNiagaraTypeDefinition(GetClass()), TEXT("Hair Strands")));
-		Sig.Outputs.Add(FNiagaraVariable(FNiagaraTypeDefinition::GetBoolDef(), TEXT("Compute Status")));
+		Sig.Outputs.Add(FNiagaraVariable(FNiagaraTypeDefinition::GetQuatDef(), TEXT("Node Orientation")));
 
 		OutFunctions.Add(Sig);
 	}
@@ -2126,6 +2138,7 @@ DEFINE_NDI_DIRECT_FUNC_BINDER(UNiagaraDataInterfaceHairStrands, GetProjectCollis
 DEFINE_NDI_DIRECT_FUNC_BINDER(UNiagaraDataInterfaceHairStrands, GetStaticFriction);
 DEFINE_NDI_DIRECT_FUNC_BINDER(UNiagaraDataInterfaceHairStrands, GetKineticFriction);
 DEFINE_NDI_DIRECT_FUNC_BINDER(UNiagaraDataInterfaceHairStrands, GetStrandsViscosity);
+DEFINE_NDI_DIRECT_FUNC_BINDER(UNiagaraDataInterfaceHairStrands, GetGridDimension);
 DEFINE_NDI_DIRECT_FUNC_BINDER(UNiagaraDataInterfaceHairStrands, GetCollisionRadius);
 DEFINE_NDI_DIRECT_FUNC_BINDER(UNiagaraDataInterfaceHairStrands, GetRadiusScale);
 
@@ -2306,6 +2319,11 @@ void UNiagaraDataInterfaceHairStrands::GetVMExternalFunction(const FVMExternalFu
 		check(BindingInfo.GetNumInputs() == 1 && BindingInfo.GetNumOutputs() == 1);
 		NDI_FUNC_BINDER(UNiagaraDataInterfaceHairStrands, GetStrandsViscosity)::Bind(this, OutFunc);
 	}
+	else if (BindingInfo.Name == GetGridDimensionName)
+	{
+		check(BindingInfo.GetNumInputs() == 1 && BindingInfo.GetNumOutputs() == 3);
+		NDI_FUNC_BINDER(UNiagaraDataInterfaceHairStrands, GetGridDimension)::Bind(this, OutFunc);
+	}
 	else if (BindingInfo.Name == GetCollisionRadiusName)
 	{
 		check(BindingInfo.GetNumInputs() == 1 && BindingInfo.GetNumOutputs() == 1);
@@ -2518,7 +2536,7 @@ void UNiagaraDataInterfaceHairStrands::GetVMExternalFunction(const FVMExternalFu
 	}
 	else if (BindingInfo.Name == ProjectBendRodMaterialName)
 	{
-		check(BindingInfo.GetNumInputs() == 10 && BindingInfo.GetNumOutputs() == 3);
+		check(BindingInfo.GetNumInputs() == 10 && BindingInfo.GetNumOutputs() == 4);
 		NDI_FUNC_BINDER(UNiagaraDataInterfaceHairStrands, ProjectBendRodMaterial)::Bind(this, OutFunc);
 	}
 	else if (BindingInfo.Name == SolveHardCollisionConstraintName)
@@ -2553,12 +2571,12 @@ void UNiagaraDataInterfaceHairStrands::GetVMExternalFunction(const FVMExternalFu
 	}
 	else if (BindingInfo.Name == UpdateMaterialFrameName)
 	{
-		check(BindingInfo.GetNumInputs() == 1 && BindingInfo.GetNumOutputs() == 1);
+		check(BindingInfo.GetNumInputs() == 1 && BindingInfo.GetNumOutputs() == 4);
 		NDI_FUNC_BINDER(UNiagaraDataInterfaceHairStrands, UpdateMaterialFrame)::Bind(this, OutFunc);
 	}
 	else if (BindingInfo.Name == ComputeMaterialFrameName)
 	{
-		check(BindingInfo.GetNumInputs() == 1 && BindingInfo.GetNumOutputs() == 1);
+		check(BindingInfo.GetNumInputs() == 1 && BindingInfo.GetNumOutputs() == 4);
 		NDI_FUNC_BINDER(UNiagaraDataInterfaceHairStrands, ComputeMaterialFrame)::Bind(this, OutFunc);
 	}
 	else if (BindingInfo.Name == ComputeAirDragForceName)
@@ -2859,6 +2877,20 @@ void UNiagaraDataInterfaceHairStrands::GetStrandsViscosity(FVectorVMContext& Con
 	}
 }
 
+void UNiagaraDataInterfaceHairStrands::GetGridDimension(FVectorVMContext& Context)
+{
+	VectorVM::FUserPtrHandler<FNDIHairStrandsData> InstData(Context);
+	VectorVM::FExternalFuncRegisterHandler<float> OutGridDimensionX(Context);
+	VectorVM::FExternalFuncRegisterHandler<float> OutGridDimensionY(Context);
+	VectorVM::FExternalFuncRegisterHandler<float> OutGridDimensionZ(Context);
+
+	for (int32 InstanceIdx = 0; InstanceIdx < Context.NumInstances; ++InstanceIdx)
+	{
+		*OutGridDimensionX.GetDestAndAdvance() = InstData->GridDimension.X;
+		*OutGridDimensionY.GetDestAndAdvance() = InstData->GridDimension.Y;
+		*OutGridDimensionZ.GetDestAndAdvance() = InstData->GridDimension.Z;
+	}
+}
 void UNiagaraDataInterfaceHairStrands::GetCollisionRadius(FVectorVMContext& Context)
 {
 	VectorVM::FUserPtrHandler<FNDIHairStrandsData> InstData(Context);
@@ -3729,10 +3761,10 @@ in float RestLength, in float DeltaTime, in int NodeOffset, in float MaterialDam
 	else if (FunctionInfo.DefinitionName == SolveBendRodMaterialName)
 	{
 		static const TCHAR *FormatSample = TEXT(R"(
-						void {InstanceFunctionName} (in bool EnableConstraint, in float RestLength, in float4 RestDarboux, in float DeltaTime, in float MaterialDamping, 
+						void {InstanceFunctionName} (in bool EnableConstraint, in float RestLength, in float4 RestRotation, in float DeltaTime, in float MaterialDamping, 
 				in float MaterialCompliance, in float MaterialWeight, in float3 MaterialMultiplier, out float3 OutMaterialMultiplier)
 						{
-							{HairStrandsContextName} SolveBendRodMaterial(EnableConstraint,DIContext.StrandSize,RestLength,RestDarboux,DeltaTime,MaterialDamping,MaterialCompliance,MaterialWeight,MaterialMultiplier,OutMaterialMultiplier);
+							{HairStrandsContextName} SolveBendRodMaterial(EnableConstraint,DIContext.StrandSize,RestLength,RestRotation,DeltaTime,MaterialDamping,MaterialCompliance,MaterialWeight,MaterialMultiplier,OutMaterialMultiplier);
 						}
 						)");
 		OutHLSL += FString::Format(FormatSample, ArgsSample);
@@ -3741,9 +3773,9 @@ in float RestLength, in float DeltaTime, in int NodeOffset, in float MaterialDam
 	else if (FunctionInfo.DefinitionName == ProjectBendRodMaterialName)
 	{
 		static const TCHAR *FormatSample = TEXT(R"(
-						void {InstanceFunctionName} (in bool EnableConstraint, in float YoungModulus, in float RodThickness, in float RestLength, in float4 RestDarboux, in float DeltaTime, out float3 OutNodePosition)
+						void {InstanceFunctionName} (in bool EnableConstraint, in float YoungModulus, in float RodThickness, in float RestLength, in float4 RestRotation, in float DeltaTime, out float4 OutNodeOrientation)
 						{
-							{HairStrandsContextName} ProjectBendRodMaterial(EnableConstraint,DIContext.StrandSize,YoungModulus,RodThickness,RestLength,RestDarboux,DeltaTime,OutNodePosition);
+							{HairStrandsContextName} ProjectBendRodMaterial(EnableConstraint,DIContext.StrandSize,YoungModulus,RodThickness,RestLength,RestRotation,DeltaTime,OutNodeOrientation);
 						}
 						)");
 		OutHLSL += FString::Format(FormatSample, ArgsSample);
@@ -3829,10 +3861,10 @@ in float RestLength, in float DeltaTime, in int NodeOffset, in float MaterialDam
 	else if (FunctionInfo.DefinitionName == UpdateMaterialFrameName)
 	{
 		static const TCHAR *FormatSample = TEXT(R"(
-					void {InstanceFunctionName} ( out bool OutUpdateStatus)
+					void {InstanceFunctionName} (out float4 OutNodeOrientation)
 					{
 						{HairStrandsContextName} UpdateMaterialFrame(DIContext.StrandSize);
-						OutUpdateStatus = true;
+						OutNodeOrientation = SharedNodeOrientation[GGroupThreadId.x];
 					}
 					)");
 		OutHLSL += FString::Format(FormatSample, ArgsSample);
@@ -3841,10 +3873,10 @@ in float RestLength, in float DeltaTime, in int NodeOffset, in float MaterialDam
 	else if (FunctionInfo.DefinitionName == ComputeMaterialFrameName)
 	{
 		static const TCHAR *FormatSample = TEXT(R"(
-						void {InstanceFunctionName} ( out bool OutUpdateStatus)
+						void {InstanceFunctionName} ( out float4 OutNodeOrientation)
 						{
 							{HairStrandsContextName} ComputeMaterialFrame(DIContext.StrandSize);
-							OutUpdateStatus = true;
+							OutNodeOrientation = SharedNodeOrientation[GGroupThreadId.x];
 						}
 						)");
 		OutHLSL += FString::Format(FormatSample, ArgsSample);
diff --git a/Engine/Plugins/Experimental/HairStrands/Source/HairStrandsNiagara/Private/NiagaraDataInterfacePhysicsAsset.cpp b/Engine/Plugins/Experimental/HairStrands/Source/HairStrandsNiagara/Private/NiagaraDataInterfacePhysicsAsset.cpp
--- a/Engine/Plugins/Experimental/HairStrands/Source/HairStrandsNiagara/Private/NiagaraDataInterfacePhysicsAsset.cpp
+++ b/Engine/Plugins/Experimental/HairStrands/Source/HairStrandsNiagara/Private/NiagaraDataInterfacePhysicsAsset.cpp
@@ -419,18 +419,18 @@ public:
 		{
 			UE_LOG(LogPhysicsAsset, Warning, TEXT("Binding failed for FNDIPhysicsAssetParametersCS %s. Was it optimized out?"), *ParamNames.InverseTransformBufferName)
 		}
-		if (!RestTransformBuffer.IsBound())
+		if (!ElementExtentBuffer.IsBound())
+		{
+			UE_LOG(LogPhysicsAsset, Warning, TEXT("Binding failed for FNDIPhysicsAssetParametersCS %s. Was it optimized out?"), *ParamNames.ElementExtentBufferName)
+		}
+		/*if (!RestTransformBuffer.IsBound())
 		{
 			UE_LOG(LogPhysicsAsset, Warning, TEXT("Binding failed for FNDIPhysicsAssetParametersCS %s. Was it optimized out?"), *ParamNames.RestTransformBufferName)
 		}
 		if (!RestInverseBuffer.IsBound())
 		{
 			UE_LOG(LogPhysicsAsset, Warning, TEXT("Binding failed for FNDIPhysicsAssetParametersCS %s. Was it optimized out?"), *ParamNames.RestInverseBufferName)
-		}
-		if (!ElementExtentBuffer.IsBound())
-		{
-			UE_LOG(LogPhysicsAsset, Warning, TEXT("Binding failed for FNDIPhysicsAssetParametersCS %s. Was it optimized out?"), *ParamNames.ElementExtentBufferName)
-		}
+		}*/
 	}
 
 	void Set(FRHICommandList& RHICmdList, const FNiagaraDataInterfaceSetArgs& Context) const
diff --git a/Engine/Plugins/Experimental/HairStrands/Source/HairStrandsNiagara/Private/NiagaraDataInterfacePressureGrid.cpp b/Engine/Plugins/Experimental/HairStrands/Source/HairStrandsNiagara/Private/NiagaraDataInterfacePressureGrid.cpp
--- a/Engine/Plugins/Experimental/HairStrands/Source/HairStrandsNiagara/Private/NiagaraDataInterfacePressureGrid.cpp
+++ b/Engine/Plugins/Experimental/HairStrands/Source/HairStrandsNiagara/Private/NiagaraDataInterfacePressureGrid.cpp
@@ -30,6 +30,7 @@ static const FName BuildGridTopologyName(TEXT("BuildGridTopology"));
 static const FName UpdateGridTransformName(TEXT("UpdateGridTransform"));
 static const FName AddGridVelocityName(TEXT("AddGridVelocity"));
 static const FName GetGridVelocityName(TEXT("GetGridVelocity"));
+static const FName SetGridDimensionName(TEXT("SetGridDimension"));
 
 //------------------------------------------------------------------------------------------------------------
 
@@ -123,6 +124,24 @@ void FNDIPressureGridData::Release()
 	}
 }
 
+void FNDIPressureGridData::Resize()
+{
+	if (NeedResize)
+	{
+		if (CurrentGridBuffer)
+		{
+			CurrentGridBuffer->Initialize(GridSize);
+			BeginInitResource(CurrentGridBuffer);
+		}
+		if (DestinationGridBuffer)
+		{
+			DestinationGridBuffer->Initialize(GridSize);
+			BeginInitResource(DestinationGridBuffer);
+		}
+		NeedResize = false;
+	}
+}
+
 bool FNDIPressureGridData::Init(UNiagaraDataInterfacePressureGrid* Interface, FNiagaraSystemInstance* SystemInstance)
 {
 	CurrentGridBuffer = nullptr;
@@ -130,6 +149,7 @@ bool FNDIPressureGridData::Init(UNiagaraDataInterfacePressureGrid* Interface, FN
 
 	GridOrigin = FVector4(0, 0, 0, 0);
 	GridSize = FIntVector(1, 1, 1);
+	NeedResize = true;
 	WorldTransform = WorldInverse = FMatrix::Identity;
 
 	if (Interface != nullptr)
@@ -137,12 +157,9 @@ bool FNDIPressureGridData::Init(UNiagaraDataInterfacePressureGrid* Interface, FN
 		GridSize = Interface->GridSize;
 
 		CurrentGridBuffer = new FNDIPressureGridBuffer();
-		CurrentGridBuffer->Initialize(Interface->GridSize);
-		BeginInitResource(CurrentGridBuffer);
-
 		DestinationGridBuffer = new FNDIPressureGridBuffer();
-		DestinationGridBuffer->Initialize(Interface->GridSize);
-		BeginInitResource(DestinationGridBuffer);
+
+		Resize();
 	}
 
 	return true;
@@ -281,8 +298,6 @@ bool UNiagaraDataInterfacePressureGrid::InitPerInstanceData(void* PerInstanceDat
 		ThisProxy->OutputSimulationStages_DEPRECATED = RT_OutputShaderStages;
 		ThisProxy->IterationSimulationStages_DEPRECATED = RT_IterationShaderStages;
 		ThisProxy->SetElementCount(ElementCount);
-
-		//ThisProxy->InitializePerInstanceData(InstanceID);
 	}
 	);
 
@@ -301,7 +316,6 @@ void UNiagaraDataInterfacePressureGrid::DestroyPerInstanceData(void* PerInstance
 		[ThisProxy, InstanceID = SystemInstance->GetId(), Batcher = SystemInstance->GetBatcher()](FRHICommandListImmediate& CmdList)
 	{
 		ThisProxy->SystemInstancesToProxyData.Remove(InstanceID);
-		//ThisProxy->DestroyPerInstanceData(Batcher, InstanceID);
 	}
 	);
 }
@@ -314,8 +328,22 @@ bool UNiagaraDataInterfacePressureGrid::PerInstanceTick(void* PerInstanceData, F
 	if (InstanceData)
 	{
 		InstanceData->WorldTransform = SystemInstance->GetComponent()->GetComponentToWorld().ToMatrixWithScale();
-	}
 
+		if (InstanceData->NeedResize)
+		{
+			const int32 ElementCount = (InstanceData->GridSize.X + 1) * (InstanceData->GridSize.Y + 1) * (InstanceData->GridSize.Z + 1);
+
+			FNDIPressureGridProxy* ThisProxy = GetProxyAs<FNDIPressureGridProxy>();
+			ENQUEUE_RENDER_COMMAND(FNiagaraDIPushInitialInstanceDataToRT) (
+				[ThisProxy, InstanceID = SystemInstance->GetId(), ElementCount](FRHICommandListImmediate& CmdList)
+			{
+				ThisProxy->SetElementCount(ElementCount);
+			}
+			);
+
+			InstanceData->Resize();
+		}
+	}
 	return RequireReset;
 }
 
@@ -501,6 +529,17 @@ void UNiagaraDataInterfacePressureGrid::GetFunctions(TArray<FNiagaraFunctionSign
 		Sig.Inputs.Add(FNiagaraVariable(FNiagaraTypeDefinition::GetIntDef(), TEXT("Grid Cell")));
 		Sig.Outputs.Add(FNiagaraVariable(FNiagaraTypeDefinition::GetVec3Def(), TEXT("Grid Velocity")));
 
+		OutFunctions.Add(Sig);
+	}
+	{
+		FNiagaraFunctionSignature Sig;
+		Sig.Name = SetGridDimensionName;
+		Sig.bMemberFunction = true;
+		Sig.bRequiresContext = false;
+		Sig.Inputs.Add(FNiagaraVariable(FNiagaraTypeDefinition(GetClass()), TEXT("Pressure Grid")));
+		Sig.Inputs.Add(FNiagaraVariable(FNiagaraTypeDefinition::GetVec3Def(), TEXT("Grid Dimension")));
+		Sig.Outputs.Add(FNiagaraVariable(FNiagaraTypeDefinition::GetBoolDef(), TEXT("Function Status")));
+
 		OutFunctions.Add(Sig);
 	}
 }
@@ -516,6 +555,7 @@ DEFINE_NDI_DIRECT_FUNC_BINDER(UNiagaraDataInterfacePressureGrid, BuildGridTopolo
 DEFINE_NDI_DIRECT_FUNC_BINDER(UNiagaraDataInterfacePressureGrid, UpdateGridTransform);
 DEFINE_NDI_DIRECT_FUNC_BINDER(UNiagaraDataInterfacePressureGrid, AddGridVelocity);
 DEFINE_NDI_DIRECT_FUNC_BINDER(UNiagaraDataInterfacePressureGrid, GetGridVelocity);
+DEFINE_NDI_DIRECT_FUNC_BINDER(UNiagaraDataInterfacePressureGrid, SetGridDimension);
 
 void UNiagaraDataInterfacePressureGrid::GetVMExternalFunction(const FVMExternalFunctionBindingInfo& BindingInfo, void* InstanceData, FVMExternalFunction &OutFunc)
 {
@@ -574,6 +614,11 @@ void UNiagaraDataInterfacePressureGrid::GetVMExternalFunction(const FVMExternalF
 		check(BindingInfo.GetNumInputs() == 2 && BindingInfo.GetNumOutputs() == 3);
 		NDI_FUNC_BINDER(UNiagaraDataInterfacePressureGrid, GetGridVelocity)::Bind(this, OutFunc);
 	}
+	else if (BindingInfo.Name == SetGridDimensionName)
+	{
+		check(BindingInfo.GetNumInputs() == 4 && BindingInfo.GetNumOutputs() == 1);
+		NDI_FUNC_BINDER(UNiagaraDataInterfacePressureGrid, SetGridDimension)::Bind(this, OutFunc);
+	}
 }
 
 void UNiagaraDataInterfacePressureGrid::BuildVelocityField(FVectorVMContext& Context)
@@ -621,6 +666,29 @@ void UNiagaraDataInterfacePressureGrid::GetGridVelocity(FVectorVMContext& Contex
 	// @todo : implement function for cpu 
 }
 
+void UNiagaraDataInterfacePressureGrid::SetGridDimension(FVectorVMContext& Context)
+{
+	VectorVM::FExternalFuncInputHandler<float> GridDimensionX(Context);
+	VectorVM::FExternalFuncInputHandler<float> GridDimensionY(Context);
+	VectorVM::FExternalFuncInputHandler<float> GridDimensionZ(Context);
+
+	VectorVM::FUserPtrHandler<FNDIPressureGridData> InstData(Context);
+	VectorVM::FExternalFuncRegisterHandler<bool> OutFunctionStatus(Context);
+
+	for (int32 InstanceIdx = 0; InstanceIdx < Context.NumInstances; ++InstanceIdx)
+	{
+		FIntVector GridDimension;
+		GridDimension.X = *GridDimensionX.GetDestAndAdvance();
+		GridDimension.Y = *GridDimensionY.GetDestAndAdvance();
+		GridDimension.Z = *GridDimensionZ.GetDestAndAdvance();
+
+		InstData->GridSize = GridDimension;
+		InstData->NeedResize = true;
+
+		*OutFunctionStatus.GetDestAndAdvance() = true;
+	}
+}
+
 void UNiagaraDataInterfacePressureGrid::BuildGridTopology(FVectorVMContext& Context)
 {
 	VectorVM::FExternalFuncInputHandler<float> CenterXParam(Context);
@@ -931,7 +999,7 @@ static void AddClearPressureGridPass(
 		FIntVector(DispatchCount, 1, 1));
 }
 
-inline void ClearBuffer(FNDIPressureGridBuffer* CurrentGridBuffer, FNDIPressureGridBuffer* DestinationGridBuffer, const FIntVector& GridSize, const bool CopyPressure)
+inline void ClearBuffer(FRHICommandList& RHICmdList, FNDIPressureGridBuffer* CurrentGridBuffer, FNDIPressureGridBuffer* DestinationGridBuffer, const FIntVector& GridSize, const bool CopyPressure)
 {
 	FRHIUnorderedAccessView* DestinationGridBufferUAV = DestinationGridBuffer->GridDataBuffer.UAV;
 	FRHIShaderResourceView* CurrentGridBufferSRV = CurrentGridBuffer->GridDataBuffer.SRV;
@@ -942,23 +1010,25 @@ inline void ClearBuffer(FNDIPressureGridBuffer* CurrentGridBuffer, FNDIPressureG
 		const FIntVector LocalGridSize = GridSize;
 		const bool LocalCopyPressure = CopyPressure;
 
-		ENQUEUE_RENDER_COMMAND(ClearPressureGrid)(
-			[DestinationGridBufferUAV, CurrentGridBufferSRV, CurrentGridBufferUAV, LocalGridSize, LocalCopyPressure]
-		(FRHICommandListImmediate& RHICmdListImm)
-		{
-			RHICmdListImm.TransitionResource(EResourceTransitionAccess::EWritable, EResourceTransitionPipeline::EComputeToCompute, DestinationGridBufferUAV);
-			RHICmdListImm.TransitionResource(EResourceTransitionAccess::EReadable, EResourceTransitionPipeline::EComputeToCompute, CurrentGridBufferUAV);
+		RHICmdList.ClearUAVUint(DestinationGridBufferUAV, FUintVector4(0, 0, 0, 0));
 
-			RHICmdListImm.ClearUAVUint(DestinationGridBufferUAV, FUintVector4(0,0,0,0));
+		//ENQUEUE_RENDER_COMMAND(ClearPressureGrid)(
+		//	[DestinationGridBufferUAV, CurrentGridBufferSRV, CurrentGridBufferUAV, LocalGridSize, LocalCopyPressure]
+		//(FRHICommandListImmediate& RHICmdListImm)
+		//{
+		//	RHICmdListImm.TransitionResource(EResourceTransitionAccess::EWritable, EResourceTransitionPipeline::EComputeToCompute, DestinationGridBufferUAV);
+		//	RHICmdListImm.TransitionResource(EResourceTransitionAccess::EReadable, EResourceTransitionPipeline::EComputeToCompute, CurrentGridBufferUAV);
 
-		/*	FRDGBuilder GraphBuilder(RHICmdListImm);
+		//	RHICmdListImm.ClearUAVUint(DestinationGridBufferUAV, FUintVector4(0,0,0,0));
 
-			AddClearPressureGridPass(
-				GraphBuilder,
-				CurrentGridBufferSRV, DestinationGridBufferUAV, LocalGridSize, LocalCopyPressure);
+		//	/*FRDGBuilder GraphBuilder(RHICmdListImm);
 
-			GraphBuilder.Execute();*/
-		});
+		//	AddClearPressureGridPass(
+		//		GraphBuilder,
+		//		CurrentGridBufferSRV, DestinationGridBufferUAV, LocalGridSize, LocalCopyPressure);
+
+		//	GraphBuilder.Execute();*/
+		//});
 	}
 }
 
@@ -1008,17 +1078,20 @@ void FNDIPressureGridProxy::PreStage(FRHICommandList& RHICmdList, const FNiagara
 	{
 		if (!Context.IsIterationStage)
 		{
-			ClearBuffer(ProxyData->CurrentGridBuffer, ProxyData->DestinationGridBuffer, ProxyData->GridSize, true);
+			ClearBuffer(RHICmdList, ProxyData->CurrentGridBuffer, ProxyData->DestinationGridBuffer, ProxyData->GridSize, true);
 		}
 		else
 		{
-			ENQUEUE_RENDER_COMMAND(CopyPressureGrid)(
+			FRHICopyTextureInfo CopyInfo;
+			RHICmdList.CopyTexture(ProxyData->CurrentGridBuffer->GridDataBuffer.Buffer,
+				ProxyData->DestinationGridBuffer->GridDataBuffer.Buffer, CopyInfo);
+			/*ENQUEUE_RENDER_COMMAND(CopyPressureGrid)(
 				[ProxyData](FRHICommandListImmediate& RHICmdList)
 			{
 				FRHICopyTextureInfo CopyInfo;
 				RHICmdList.CopyTexture(ProxyData->CurrentGridBuffer->GridDataBuffer.Buffer,
 					ProxyData->DestinationGridBuffer->GridDataBuffer.Buffer, CopyInfo);
-			});
+			});*/
 		}
 	}
 }
@@ -1032,13 +1105,16 @@ void FNDIPressureGridProxy::PostStage(FRHICommandList& RHICmdList, const FNiagar
 	{
 		if (ProxyData != nullptr)
 		{
-			ENQUEUE_RENDER_COMMAND(CopyPressureGrid)(
+			FRHICopyTextureInfo CopyInfo;
+			RHICmdList.CopyTexture(ProxyData->DestinationGridBuffer->GridDataBuffer.Buffer,
+				ProxyData->CurrentGridBuffer->GridDataBuffer.Buffer, CopyInfo);
+			/*ENQUEUE_RENDER_COMMAND(CopyPressureGrid)(
 				[ProxyData](FRHICommandListImmediate& RHICmdList)
 			{
 				FRHICopyTextureInfo CopyInfo;
 				RHICmdList.CopyTexture(ProxyData->DestinationGridBuffer->GridDataBuffer.Buffer,
 					ProxyData->CurrentGridBuffer->GridDataBuffer.Buffer, CopyInfo);
-			});
+			});*/
 		}
 	}
 }
@@ -1049,8 +1125,8 @@ void FNDIPressureGridProxy::ResetData(FRHICommandList& RHICmdList, const FNiagar
 
 	if (ProxyData != nullptr && ProxyData->DestinationGridBuffer != nullptr && ProxyData->CurrentGridBuffer != nullptr)
 	{
-		ClearBuffer(ProxyData->CurrentGridBuffer, ProxyData->DestinationGridBuffer, ProxyData->GridSize, false);
-		ClearBuffer(ProxyData->DestinationGridBuffer, ProxyData->CurrentGridBuffer, ProxyData->GridSize, false);
+		ClearBuffer(RHICmdList, ProxyData->CurrentGridBuffer, ProxyData->DestinationGridBuffer, ProxyData->GridSize, false);
+		ClearBuffer(RHICmdList, ProxyData->DestinationGridBuffer, ProxyData->CurrentGridBuffer, ProxyData->GridSize, false);
 	}
 }
 
diff --git a/Engine/Plugins/Experimental/HairStrands/Source/HairStrandsNiagara/Public/NiagaraDataInterfaceHairStrands.h b/Engine/Plugins/Experimental/HairStrands/Source/HairStrandsNiagara/Public/NiagaraDataInterfaceHairStrands.h
--- a/Engine/Plugins/Experimental/HairStrands/Source/HairStrandsNiagara/Public/NiagaraDataInterfaceHairStrands.h
+++ b/Engine/Plugins/Experimental/HairStrands/Source/HairStrandsNiagara/Public/NiagaraDataInterfaceHairStrands.h
@@ -10,7 +10,7 @@
 #include "GroomActor.h"
 #include "NiagaraDataInterfaceHairStrands.generated.h"
 
-static const int32 MaxDelay = 8;
+static const int32 MaxDelay = 2;
 static const int32 NumScales = 4;
 static const int32 StretchOffset = 0;
 static const int32 BendOffset = 1;
@@ -125,6 +125,7 @@ struct FNDIHairStrandsData
 		KineticFriction = 0.1;
 		StaticFriction = 0.1;
 		StrandsViscosity = 1.0;
+		GridDimension = FIntVector(30,30,30);
 		CollisionRadius = 1.0;
 
 		StrandsDensity = 1.0;
@@ -177,6 +178,7 @@ struct FNDIHairStrandsData
 			StaticFriction = OtherDatas->StaticFriction;
 			KineticFriction = OtherDatas->KineticFriction;
 			StrandsViscosity = OtherDatas->StrandsViscosity;
+			GridDimension = OtherDatas->GridDimension;
 			CollisionRadius = OtherDatas->CollisionRadius;
 
 			StrandsDensity = OtherDatas->StrandsDensity;
@@ -268,6 +270,9 @@ struct FNDIHairStrandsData
 	/** Radius that will be used for the collision detection against the physics asset */
 	float StrandsViscosity;
 
+	/** Grid Dimension used to compute the viscosity forces */
+	FIntVector GridDimension;
+
 	/** Radius scale along the strand */
 	float CollisionRadius;
 
@@ -387,6 +392,8 @@ public:
 
 	void GetStrandsViscosity(FVectorVMContext& Context);
 
+	void GetGridDimension(FVectorVMContext& Context);
+
 	void GetCollisionRadius(FVectorVMContext& Context);
 
 	void GetRadiusScale(FVectorVMContext& Context);
diff --git a/Engine/Plugins/Experimental/HairStrands/Source/HairStrandsNiagara/Public/NiagaraDataInterfacePressureGrid.h b/Engine/Plugins/Experimental/HairStrands/Source/HairStrandsNiagara/Public/NiagaraDataInterfacePressureGrid.h
--- a/Engine/Plugins/Experimental/HairStrands/Source/HairStrandsNiagara/Public/NiagaraDataInterfacePressureGrid.h
+++ b/Engine/Plugins/Experimental/HairStrands/Source/HairStrandsNiagara/Public/NiagaraDataInterfacePressureGrid.h
@@ -43,12 +43,18 @@ struct FNDIPressureGridData
 	/** Release the buffers */
 	void Release();
 
+	/** Resize the buffers */
+	void Resize();
+
 	/** Grid Origin */
 	FVector4 GridOrigin;
 
 	/** Grid Size */
 	FIntVector GridSize;
 
+	/** Need a resize */
+	bool NeedResize;
+
 	/** World Transform */
 	FMatrix WorldTransform;
 
@@ -144,6 +150,9 @@ public:
 	/** Add velocity to the grid */
 	void AddGridVelocity(FVectorVMContext& Context);
 
+	/** Set the grid dimension */
+	void SetGridDimension(FVectorVMContext& Context);
+
 	/** Name of the grid current buffer */
 	static const FString GridCurrentBufferName;
 
