Fix crash in GroomComponent.

diff --git a/Engine/Plugins/Experimental/HairStrands/Source/HairStrandsCore/Private/GroomComponent.cpp b/Engine/Plugins/Experimental/HairStrands/Source/HairStrandsCore/Private/GroomComponent.cpp
--- a/Engine/Plugins/Experimental/HairStrands/Source/HairStrandsCore/Private/GroomComponent.cpp
+++ b/Engine/Plugins/Experimental/HairStrands/Source/HairStrandsCore/Private/GroomComponent.cpp
@@ -1675,7 +1675,8 @@ void UGroomComponent::PostEditChangeProperty(FPropertyChangedEvent& PropertyChan
 		bIsEventProcess = true;
 	}
 
-	if (!bIsEventProcess)
+	// Always call parent PostEditChangeProperty as parent expect such a call (would crash in certain case otherwise)
+	//if (!bIsEventProcess)
 	{
 		Super::PostEditChangeProperty(PropertyChangedEvent);
 	}
