Fix crash when packaging a project with hair strands.

diff --git a/Engine/Plugins/Experimental/HairStrands/Source/HairStrandsCore/Private/GroomComponent.cpp b/Engine/Plugins/Experimental/HairStrands/Source/HairStrandsCore/Private/GroomComponent.cpp
--- a/Engine/Plugins/Experimental/HairStrands/Source/HairStrandsCore/Private/GroomComponent.cpp
+++ b/Engine/Plugins/Experimental/HairStrands/Source/HairStrandsCore/Private/GroomComponent.cpp
@@ -1099,6 +1099,12 @@ void UGroomComponent::InitResources(bool bIsBindingReloading)
 	Interpolation.Function = ComputeHairStrandsInterpolation;
 	Interpolation.ResetFunction = ResetHairStrandsInterpolation;
 
+	// Does not run projection code when running with null RHI as this is not needed, and will crash as the skeletal GPU resources are not created
+	if (GUsingNullRHI)
+	{
+		return;
+	}
+
 	UGroomAsset* LocalGroomAsset = GroomAsset;
 	USkeletalMesh* InSourceSkeletalMesh = SourceSkeletalMesh;
 	const bool bRunMeshProjection = bHasValidSketalMesh && (!BindingAsset || !bIsBindingCompatible);
