Fix hair component visibility when hidden in game option is ticked.

diff --git a/Engine/Plugins/Experimental/HairStrands/Source/HairStrandsCore/Private/GroomComponent.cpp b/Engine/Plugins/Experimental/HairStrands/Source/HairStrandsCore/Private/GroomComponent.cpp
--- a/Engine/Plugins/Experimental/HairStrands/Source/HairStrandsCore/Private/GroomComponent.cpp
+++ b/Engine/Plugins/Experimental/HairStrands/Source/HairStrandsCore/Private/GroomComponent.cpp
@@ -343,7 +343,7 @@ public:
 				continue;
 			}
 
-			if (VisibilityMap & (1 << ViewIndex))
+			if (IsShown(View) && (VisibilityMap & (1 << ViewIndex)))
 			{
 				for (uint32 GroupIt = 0; GroupIt < GroupCount; ++GroupIt)
 				{
@@ -416,7 +416,7 @@ public:
 		}
 
 		FPrimitiveViewRelevance Result;
-		Result.bHairStrands = bIsViewModeValid;
+		Result.bHairStrands = bIsViewModeValid && IsShown(View);
 
 		// Special pass for hair strands geometry (not part of the base pass, and shadowing is handlded in a custom fashion)
 		Result.bDrawRelevance = false;		
