Default/Debug hair materials loaded in constructor

diff --git a/Engine/Plugins/Experimental/HairStrands/Source/HairStrandsCore/Private/GroomComponent.cpp b/Engine/Plugins/Experimental/HairStrands/Source/HairStrandsCore/Private/GroomComponent.cpp
--- a/Engine/Plugins/Experimental/HairStrands/Source/HairStrandsCore/Private/GroomComponent.cpp
+++ b/Engine/Plugins/Experimental/HairStrands/Source/HairStrandsCore/Private/GroomComponent.cpp
@@ -144,33 +144,6 @@ static EHairMaterialCompatibility IsHairMaterialCompatible(UMaterialInterface* M
 	return EHairMaterialCompatibility::Valid;
 }
 
-static UMaterialInterface* GetHairDefaultMaterial()
-{
-#if WITH_EDITOR
-	static UMaterialInterface* HairDefaultMaterial = nullptr;
-	if (!HairDefaultMaterial)
-	{
-		HairDefaultMaterial = LoadObject<UMaterialInterface>(NULL, TEXT("/HairStrands/Materials/HairDefaultMaterial.HairDefaultMaterial"));
-	}
-	return HairDefaultMaterial;
-#else
-	return nullptr;
-#endif
-}
-
-static UMaterialInterface* GetHairDebugMaterial()
-{
-#if WITH_EDITOR
-	static UMaterialInterface* HairDebugMaterial = nullptr;
-	if (!HairDebugMaterial)
-	{
-		HairDebugMaterial = LoadObject<UMaterialInterface>(NULL, TEXT("/HairStrands/Materials/HairDebugMaterial.HairDebugMaterial"));
-	}
-	return HairDebugMaterial;
-#else
-	return nullptr;
-#endif
-}
 /////////////////////////////////////////////////////////////////////////////////////////
 /////////////////////////////////////////////////////////////////////////////////////////
 //  FStrandHairSceneProxy
@@ -196,6 +169,7 @@ public:
 		check(Component->GroomAsset->GetNumHairGroups() > 0);
 		check(Component->InterpolationOutput);
 		ComponentId = Component->ComponentId.PrimIDValue;
+		HairDebugMaterial = Component->HairDebugMaterial;
 
 		FHairStrandsVertexFactory::FDataType VFData;
 		VFData.InterpolationOutput = Component->InterpolationOutput;
@@ -328,7 +302,6 @@ public:
 				HairMaxRadius = FMath::Max(HairMaxRadius, HairGroup.VFInput.HairRadius);
 			}
 
-			UMaterialInterface* HairDebugMaterial = GetHairDebugMaterial();
 			const float HairClipLength = GetHairClipLength();
 			auto DebugMaterial = new FHairDebugModeMaterialRenderProxy(HairDebugMaterial ? HairDebugMaterial->GetRenderProxy() : nullptr, DebugModeScalar, 0, HairMaxRadius, HairClipLength);
 			Collector.RegisterOneFrameMaterialProxy(DebugMaterial);
@@ -444,6 +417,7 @@ private:
 	uint32 ComponentId = 0;
 	FHairStrandsVertexFactory VertexFactory;
 	FMaterialRelevance MaterialRelevance;
+	UMaterialInterface* HairDebugMaterial = nullptr;
 	struct HairGroup
 	{
 		UMaterialInterface* Material = nullptr;
@@ -479,6 +453,12 @@ UGroomComponent::UGroomComponent(const FObjectInitializer& ObjectInitializer)
 	NiagaraComponents.Empty();
 
 	SetCollisionProfileName(UCollisionProfile::PhysicsActor_ProfileName);
+
+	static ConstructorHelpers::FObjectFinder<UMaterialInterface> HairDebugMaterialRef(TEXT("/HairStrands/Materials/HairDebugMaterial.HairDebugMaterial"));
+	static ConstructorHelpers::FObjectFinder<UMaterialInterface> HairDefaultMaterialRef(TEXT("/HairStrands/Materials/HairDefaultMaterial.HairDefaultMaterial"));
+
+	HairDebugMaterial = HairDebugMaterialRef.Object;
+	HairDefaultMaterial = HairDefaultMaterialRef.Object;
 }
 
 void UGroomComponent::UpdateHairGroupsDesc()
@@ -545,6 +525,7 @@ void UGroomComponent::UpdateHairSimulation()
 {
 	static UNiagaraSystem* CosseratRodsSystem = LoadObject<UNiagaraSystem>(nullptr, TEXT("/HairStrands/Emitters/GroomRodsSystem.GroomRodsSystem"));
 	static UNiagaraSystem* AngularSpringsSystem = LoadObject<UNiagaraSystem>(nullptr, TEXT("/HairStrands/Emitters/GroomSpringsSystem.GroomSpringsSystem"));
+
 	const int32 NumGroups = GroomAsset ? GroomAsset->HairGroupsPhysics.Num() : 0;
 	const int32 NumComponents = FMath::Max(NumGroups, NiagaraComponents.Num());
 
@@ -819,7 +800,7 @@ UMaterialInterface* UGroomComponent::GetMaterial(int32 ElementIndex) const
 
 	if (bUseHairDefaultMaterial)
 	{
-		OverrideMaterial = GetHairDefaultMaterial();
+		OverrideMaterial = HairDefaultMaterial;
 	}
 
 	return OverrideMaterial;
@@ -1380,6 +1361,8 @@ void UGroomComponent::InvalidateAndRecreate()
 
 void UGroomComponent::OnRegister()
 {
+	
+
 	UpdateHairSimulation();
 	Super::OnRegister();
 
@@ -1551,7 +1534,6 @@ void UGroomComponent::GetUsedMaterials(TArray<UMaterialInterface*>& OutMaterials
 	UMeshComponent::GetUsedMaterials(OutMaterials, bGetDebugMaterials);
 	if (bGetDebugMaterials)
 	{
-		UMaterialInterface* HairDebugMaterial = GetHairDebugMaterial();
 		OutMaterials.Add(HairDebugMaterial);
 	}
 
@@ -1568,7 +1550,6 @@ void UGroomComponent::GetUsedMaterials(TArray<UMaterialInterface*>& OutMaterials
 
 	if (bRegisterDefaultMaterial)
 	{
-		UMaterialInterface* HairDefaultMaterial = GetHairDefaultMaterial();
 		OutMaterials.Add(HairDefaultMaterial);		
 	}
 #endif
diff --git a/Engine/Plugins/Experimental/HairStrands/Source/HairStrandsCore/Public/GroomComponent.h b/Engine/Plugins/Experimental/HairStrands/Source/HairStrandsCore/Public/GroomComponent.h
--- a/Engine/Plugins/Experimental/HairStrands/Source/HairStrandsCore/Public/GroomComponent.h
+++ b/Engine/Plugins/Experimental/HairStrands/Source/HairStrandsCore/Public/GroomComponent.h
@@ -44,6 +44,12 @@ public:
 	UPROPERTY(EditAnywhere, BlueprintReadOnly, Category = "Groom")
 	class UGroomBindingAsset* BindingAsset;
 
+	UPROPERTY()
+	UMaterialInterface* HairDebugMaterial;
+
+	UPROPERTY()
+	UMaterialInterface* HairDefaultMaterial;
+
 	/** Boolean to check when the simulation should be reset */
 	bool bResetSimulation;
 
