Fix CIS issue.

diff --git a/Engine/Plugins/Experimental/HairStrands/Source/HairStrandsCore/Private/GroomComponent.cpp b/Engine/Plugins/Experimental/HairStrands/Source/HairStrandsCore/Private/GroomComponent.cpp
--- a/Engine/Plugins/Experimental/HairStrands/Source/HairStrandsCore/Private/GroomComponent.cpp
+++ b/Engine/Plugins/Experimental/HairStrands/Source/HairStrandsCore/Private/GroomComponent.cpp
@@ -19,6 +19,7 @@
 #include "Misc/UObjectToken.h"
 #include "Misc/MapErrors.h"
 #include "NiagaraComponent.h"
+#include "UObject/ConstructorHelpers.h"
 
 static float GHairClipLength = -1;
 static FAutoConsoleVariableRef CVarHairClipLength(TEXT("r.HairStrands.DebugClipLength"), GHairClipLength, TEXT("Clip hair strands which have a lenth larger than this value. (default is -1, no effect)"));
