Fix bad default hair interpolation
Remove cvar for 4.25 release

diff --git a/Engine/Plugins/Experimental/HairStrands/Source/HairStrandsCore/Private/GroomBuilder.cpp b/Engine/Plugins/Experimental/HairStrands/Source/HairStrandsCore/Private/GroomBuilder.cpp
--- a/Engine/Plugins/Experimental/HairStrands/Source/HairStrandsCore/Private/GroomBuilder.cpp
+++ b/Engine/Plugins/Experimental/HairStrands/Source/HairStrandsCore/Private/GroomBuilder.cpp
@@ -1489,7 +1489,11 @@ bool FGroomBuilder::BuildGroom(const FHairDescription& HairDescription, const FG
 
 	GroomAsset->HairGroupsInfo.Reset();
 	GroomAsset->HairGroupsData.Reset();
-	GroomAsset->HairGroupsPhysics.Reset();
+
+	if (GroomAsset->HairGroupsPhysics.Num() != ProcessedHairDescription.HairGroups.Num())
+	{
+		GroomAsset->HairGroupsPhysics.Init(FHairGroupsPhysics(), ProcessedHairDescription.HairGroups.Num());
+	}
 
 	for (TPair<int32, FProcessedHairDescription::FHairGroup> HairGroupIt : ProcessedHairDescription.HairGroups)
 	{
@@ -1499,7 +1503,6 @@ bool FGroomBuilder::BuildGroom(const FHairDescription& HairDescription, const FG
 
 		GroomAsset->HairGroupsInfo.Add(MoveTemp(GroupInfo));
 		GroomAsset->HairGroupsData.Add(MoveTemp(GroupData));
-		GroomAsset->HairGroupsPhysics.Add(FHairGroupsPhysics());
 	}
 
 	// If there's usable closest guides and guide weights attributes, fill them into the asset
diff --git a/Engine/Plugins/Experimental/HairStrands/Source/HairStrandsCore/Private/GroomComponent.cpp b/Engine/Plugins/Experimental/HairStrands/Source/HairStrandsCore/Private/GroomComponent.cpp
--- a/Engine/Plugins/Experimental/HairStrands/Source/HairStrandsCore/Private/GroomComponent.cpp
+++ b/Engine/Plugins/Experimental/HairStrands/Source/HairStrandsCore/Private/GroomComponent.cpp
@@ -901,15 +901,18 @@ void UGroomComponent::UpdateSimulatedGroups()
 
 		FHairStrandsInterpolationInput* LocalInterpolationInput = InterpolationInput;
 		UGroomAsset* LocalGroomAsset = GroomAsset;
+		UGroomBindingAsset* LocalBindingAsset = BindingAsset;
 		ENQUEUE_RENDER_COMMAND(FHairStrandsTick_UEnableSimulatedGroups)(
-			[LocalInterpolationInput, LocalGroomAsset, Id, WorldType](FRHICommandListImmediate& RHICmdList)
+			[LocalInterpolationInput, LocalGroomAsset, LocalBindingAsset, Id, WorldType](FRHICommandListImmediate& RHICmdList)
 		{
 			int32 GroupIt = 0;
 			for (FHairStrandsInterpolationInput::FHairGroup& HairGroup : LocalInterpolationInput->HairGroups)
 			{
 				const bool bIsSimulationEnable = (LocalGroomAsset && GroupIt < LocalGroomAsset->HairGroupsPhysics.Num()) ? 
 					LocalGroomAsset->HairGroupsPhysics[GroupIt].SolverSettings.EnableSimulation : false;
+				const bool bHasGlobalInterpolation = (LocalBindingAsset && LocalGroomAsset->EnableGlobalInterpolation);
 				HairGroup.bIsSimulationEnable = bIsSimulationEnable;
+				HairGroup.bHasGlobalInterpolation = bHasGlobalInterpolation;
 				UpdateHairStrandsDebugInfo(Id, WorldType, GroupIt, bIsSimulationEnable);
 				++GroupIt;
 			}
@@ -1065,11 +1068,14 @@ void UGroomComponent::InitResources(bool bIsBindingReloading)
 		const bool bIsSimulationEnable = (GroomAsset && GroupIt < GroomAsset->HairGroupsPhysics.Num()) ?
 			GroomAsset->HairGroupsPhysics[GroupIt].SolverSettings.EnableSimulation : false;
 
+		const bool bHasGlobalInterpolation = (BindingAsset && GroomAsset->EnableGlobalInterpolation);
+
 		// For skinned groom, these value will be updated during TickComponent() call
 		// Deformed sim & render are expressed within the referential (unlike rest pose)
 		InterpolationInputGroup.OutHairPositionOffset = RenderRestHairPositionOffset;
 		InterpolationInputGroup.OutHairPreviousPositionOffset = RenderRestHairPositionOffset;
 		InterpolationInputGroup.bIsSimulationEnable = bIsSimulationEnable;
+		InterpolationInputGroup.bHasGlobalInterpolation = bHasGlobalInterpolation;
 		DebugHairGroup.bHasSimulation = InterpolationInputGroup.bIsSimulationEnable;
 
 		GroupIt++;
diff --git a/Engine/Plugins/Experimental/HairStrands/Source/HairStrandsCore/Private/HairStrandsRendering.cpp b/Engine/Plugins/Experimental/HairStrands/Source/HairStrandsCore/Private/HairStrandsRendering.cpp
--- a/Engine/Plugins/Experimental/HairStrands/Source/HairStrandsCore/Private/HairStrandsRendering.cpp
+++ b/Engine/Plugins/Experimental/HairStrands/Source/HairStrandsCore/Private/HairStrandsRendering.cpp
@@ -52,12 +52,6 @@ static FAutoConsoleVariableRef CVarStrandHairWidth(TEXT("r.HairStrands.StrandWid
 static int32 GStrandHairInterpolationDebug = 0;
 static FAutoConsoleVariableRef CVarStrandHairInterpolationDebug(TEXT("r.HairStrands.Interpolation.Debug"), GStrandHairInterpolationDebug, TEXT("Enable debug rendering for hair interpolation"));
 
-static int32 GStrandHairInterpolationGlobal = 0;
-static FAutoConsoleVariableRef CVarStrandHairInterpolationGlobal(TEXT("r.HairStrands.Interpolation.Global"), GStrandHairInterpolationGlobal, TEXT("Enable global radial basis function based hair interpolation"));
-
-static int32 GStrandHairInterpolationTarget = 0;
-static FAutoConsoleVariableRef CVarStrandHairInterpolationTarget(TEXT("r.HairStrands.Interpolation.Target"), GStrandHairInterpolationTarget, TEXT("Set the global interpolation as the kinematics target if simulation is enabled"));
-
 static FIntVector ComputeDispatchCount(uint32 ItemCount, uint32 GroupSize)
 {
 	const uint32 BatchCount = FMath::DivideAndRoundUp(ItemCount, GroupSize);
@@ -171,7 +165,8 @@ static void AddDeformSimHairStrandsPass(
 	FRHIUnorderedAccessView* OutSimDeformedPositionBuffer,
 	FVector& SimRestOffset,
 	FVector& SimDeformedOffset,
-	FBufferTransitionQueue& OutTransitionQueue)
+	FBufferTransitionQueue& OutTransitionQueue,
+	const bool bHasGlobalInterpolation)
 {
 	static uint32 IterationCount = 0;
 	++IterationCount;
@@ -222,7 +217,7 @@ static void AddDeformSimHairStrandsPass(
 		bool bSupportGlobalInterpolation = false;
 		if (bSupportDynamicMesh)
 		{
-			bSupportGlobalInterpolation = (GStrandHairInterpolationGlobal == 1) && (InSimHairData.RestLODDatas[LODIndex].SampleCount > 0);
+			bSupportGlobalInterpolation = bHasGlobalInterpolation && (InSimHairData.RestLODDatas[LODIndex].SampleCount > 0);
 			if (!bSupportGlobalInterpolation) 
 			{
 				InternalDeformationType = 4;
@@ -378,7 +373,8 @@ static void AddHairStrandsInterpolationPass(
 	const FUnorderedAccessViewRHIRef& OutRenderAttributeBuffer,
 	const FShaderResourceViewRHIRef& VertexToClusterIdBuffer,
 	const FShaderResourceViewRHIRef& SimRootPointIndexBuffer,
-	FBufferTransitionQueue& OutTransitionQueue)
+	FBufferTransitionQueue& OutTransitionQueue,
+	const bool bHasGlobalInterpolation)
 {
 	const uint32 GroupSize = ComputeGroupSize();
 	const FIntVector DispatchCount = ComputeDispatchCount(VertexCount, GroupSize);
@@ -476,7 +472,7 @@ static void AddHairStrandsInterpolationPass(
 	bool bSupportGlobalInterpolation = false;
 	if (bSupportDynamicMesh)
 	{
-		bSupportGlobalInterpolation = (GStrandHairInterpolationGlobal == 1) && (InSimHairData.RestLODDatas[LODIndex].SampleCount > 0);
+		bSupportGlobalInterpolation = bHasGlobalInterpolation && (InSimHairData.RestLODDatas[LODIndex].SampleCount > 0);
 		{
 			Parameters->RestPosition0Buffer = InRenHairData.RestLODDatas[LODIndex].RestRootTrianglePosition0Buffer->SRV;
 			Parameters->RestPosition1Buffer = InRenHairData.RestLODDatas[LODIndex].RestRootTrianglePosition1Buffer->SRV;
@@ -511,8 +507,8 @@ static void AddHairStrandsInterpolationPass(
 	FHairInterpolationCS::FPermutationDomain PermutationVector;
 	PermutationVector.Set<FHairInterpolationCS::FGroupSize>(GetGroupSizePermutation(GroupSize));
 	PermutationVector.Set<FHairInterpolationCS::FDebug>(Parameters->HairDebugMode > 0 ? 1 : 0);
-	PermutationVector.Set<FHairInterpolationCS::FDynamicGeometry>((bSupportDynamicMesh && bHasLocalDeformation) ? 2 :
-							(bSupportDynamicMesh && !bHasLocalDeformation) ? 1 : 0);
+	PermutationVector.Set<FHairInterpolationCS::FDynamicGeometry>((bSupportDynamicMesh && bSupportGlobalInterpolation) ? 2 :
+							(bSupportDynamicMesh && !bSupportGlobalInterpolation) ? 1 : 0);
 	PermutationVector.Set<FHairInterpolationCS::FSimulation>(bHasLocalDeformation ? 1 : 0);
 	PermutationVector.Set<FHairInterpolationCS::FScaleAndClip>(bNeedScaleOrClip ? 1 : 0);
 
@@ -885,7 +881,8 @@ void ComputeHairStrandsInterpolation(
 				Input.SimRootPointIndexBuffer ? Input.SimRootPointIndexBuffer->SRV : nullptr,
 				Output.SimDeformedPositionBuffer[CurrIndex]->UAV, Input.InSimHairPositionOffset,
 				Input.OutHairPositionOffset,
-				TransitionQueue);
+				TransitionQueue,
+				Input.bHasGlobalInterpolation);
 		}
 		GraphBuilder.Execute();
 		TransitBufferToReadable(RHICmdList, TransitionQueue);
@@ -1016,7 +1013,8 @@ void ComputeHairStrandsInterpolation(
 					Output.RenderPatchedAttributeBuffer.UAV,
 					Input.VertexToClusterIdBuffer->SRV,
 					Input.SimRootPointIndexBuffer->SRV,
-					TransitionQueue);
+					TransitionQueue,
+					Input.bHasGlobalInterpolation);
 
 			}
 			GraphBuilder.Execute();
@@ -1225,7 +1223,8 @@ void ResetHairStrandsInterpolation(
 				Input.SimRootPointIndexBuffer ? Input.SimRootPointIndexBuffer->SRV : nullptr,
 				Output.SimDeformedPositionBuffer[CurrIndex]->UAV, Input.InSimHairPositionOffset,
 				Input.OutHairPositionOffset, 
-				TransitionQueue);
+				TransitionQueue,
+				Input.bHasGlobalInterpolation);
 			GraphBuilder.Execute();
 			TransitBufferToReadable(RHICmdList, TransitionQueue);
 		}
diff --git a/Engine/Plugins/Experimental/HairStrands/Source/HairStrandsCore/Public/GroomAsset.h b/Engine/Plugins/Experimental/HairStrands/Source/HairStrandsCore/Public/GroomAsset.h
--- a/Engine/Plugins/Experimental/HairStrands/Source/HairStrandsCore/Public/GroomAsset.h
+++ b/Engine/Plugins/Experimental/HairStrands/Source/HairStrandsCore/Public/GroomAsset.h
@@ -359,10 +359,6 @@ struct HAIRSTRANDSCORE_API FHairSolverSettings
 	UPROPERTY(EditAnywhere, Category = "SolverSettings", meta = (ToolTip = "Enable the simulation on that group"))
 	bool EnableSimulation;
 
-	/** Allow the simulation to follow a kinematics target*/
-	UPROPERTY(EditAnywhere, Category = "SolverSettings", meta = (ToolTip = "Allow the simulation to follow a kinematics target"))
-	bool KinematicsTarget;
-
 	/** Niagara solver to be used for simulation */
 	UPROPERTY(EditAnywhere, Category = "SolverSettings", meta = (ToolTip = "Niagara solver to be used for simulation"))
 	EGroomNiagaraSolvers NiagaraSolver;
@@ -523,27 +519,27 @@ struct HAIRSTRANDSCORE_API FHairStrandsParameters
 {
 	GENERATED_BODY()
 
-		FHairStrandsParameters();
+	FHairStrandsParameters();
 
 	/** Number of particles per guide that will be used for simulation*/
 	UPROPERTY(EditAnywhere, Category = "StrandsParameters", meta = (ToolTip = "Number of particles per guide that will be used for simulation"))
-		EGroomStrandsSize StrandsSize;
+	EGroomStrandsSize StrandsSize;
 
 	/** Density of the strands in g/cm3 */
 	UPROPERTY(EditAnywhere, Category = "StrandsParameters", meta = (ToolTip = "Density of the strands in g/cm3"))
-		float StrandsDensity;
+	float StrandsDensity;
 
 	/** Smoothing between 0 and 1 of the incoming guides curves for better stability */
 	UPROPERTY(EditAnywhere, Category = "StrandsParameters", meta = (ToolTip = "Smoothing between 0 and 1 of the incoming guides curves for better stability"))
-		float StrandsSmoothing;
+	float StrandsSmoothing;
 
 	/** Strands thickness in cm that will be used for mass and inertia computation */
 	UPROPERTY(EditAnywhere, Category = "StrandsParameters", meta = (ToolTip = "Strands thickness in cm that will be used for mass and inertia computation"))
-		float StrandsThickness;
+	float StrandsThickness;
 
 	/** Thickness scale along the curve */
 	UPROPERTY(EditAnywhere, Category = "StrandsParameters", meta = (ViewMinInput = "0.0", ViewMaxInput = "1.0", ViewMinOutput = "0.0", ViewMaxOutput = "1.0", TimeLineLength = "1.0", XAxisName = "Strand Coordinate (0,1)", YAxisName = "Strands Thickness", ToolTip = "This curve determines how much the strands thickness will be scaled along each strand. \n The X axis range is [0,1], 0 mapping the root and 1 the tip"))
-		FRuntimeFloatCurve ThicknessScale;
+	FRuntimeFloatCurve ThicknessScale;
 };
 
 USTRUCT(BlueprintType)
@@ -635,6 +631,10 @@ public:
 	UPROPERTY(EditAnywhere, EditFixedSize, BlueprintReadWrite, Category = "HairPhysics", meta = (DisplayName = "Group"))
 	TArray<FHairGroupsPhysics> HairGroupsPhysics;
 
+	/** Enable radial basis function interpolation to be used instead of the local skin rigid transform */
+	UPROPERTY(EditAnywhere, EditFixedSize, BlueprintReadWrite, Category = "HairInterpolation", meta = (ToolTip = "Enable radial basis function interpolation to be used instead of the local skin rigid transform (WIP)"))
+	bool EnableGlobalInterpolation = false;
+
 	TArray<FHairGroupData> HairGroupsData;
 
 	//~ Begin UObject Interface.
@@ -768,8 +768,8 @@ public:
 	UPROPERTY(VisibleAnywhere, BlueprintReadWrite, Category = "BuildSettings", meta = (ClampMin = "0.01", UIMin = "0.01", UIMax = "1.0"))
 	USkeletalMesh* TargetSkeletalMesh;
 
-	/** Number of points used for the rbf interpolation */
-	UPROPERTY(VisibleAnywhere, BlueprintReadWrite, Category = "BuildSettings", meta = (ClampMin = "0.01", UIMin = "0.01", UIMax = "1.0"))
+	/** Number of points to be used for radial basis function interpolation  */
+	UPROPERTY(VisibleAnywhere, BlueprintReadWrite, Category = "HairInterpolation", meta = (ClampMax = "100", ClampMin = "0", ToolTip = "Number of points to be used for radial basis function interpolation (WIP)"))
 	int32 NumInterpolationPoints = 100;
 
 	UPROPERTY(EditAnywhere, EditFixedSize, BlueprintReadWrite, Category = "HairGroups", meta = (DisplayName = "Group"))
diff --git a/Engine/Plugins/Experimental/HairStrands/Source/HairStrandsCore/Public/GroomCreateBindingOptions.h b/Engine/Plugins/Experimental/HairStrands/Source/HairStrandsCore/Public/GroomCreateBindingOptions.h
--- a/Engine/Plugins/Experimental/HairStrands/Source/HairStrandsCore/Public/GroomCreateBindingOptions.h
+++ b/Engine/Plugins/Experimental/HairStrands/Source/HairStrandsCore/Public/GroomCreateBindingOptions.h
@@ -25,6 +25,6 @@ public:
 	USkeletalMesh* TargetSkeletalMesh;
 
 	/** Number of points used for the rbf interpolation */
-	UPROPERTY(EditAnywhere, BlueprintReadWrite, meta = (ShowOnlyInnerProperties), Category = BuildSettings)
+	UPROPERTY(EditAnywhere, BlueprintReadWrite, meta = (ShowOnlyInnerProperties, ClampMin = "0", ClampMax = "100"), Category = HairInterpolation)
 	int32 NumInterpolationPoints;
 };
diff --git a/Engine/Plugins/Experimental/HairStrands/Source/HairStrandsCore/Public/HairStrandsRendering.h b/Engine/Plugins/Experimental/HairStrands/Source/HairStrandsCore/Public/HairStrandsRendering.h
--- a/Engine/Plugins/Experimental/HairStrands/Source/HairStrandsCore/Public/HairStrandsRendering.h
+++ b/Engine/Plugins/Experimental/HairStrands/Source/HairStrandsCore/Public/HairStrandsRendering.h
@@ -41,6 +41,7 @@ struct FHairStrandsInterpolationInput
 		FReadBuffer*ClusterIndexRadiusScaleInfoBuffer = nullptr;
 
 		bool bIsSimulationEnable = false;
+		bool bHasGlobalInterpolation = false;
 		
 		FHairGroupDesc GroupDesc;
 
diff --git a/Engine/Plugins/Experimental/HairStrands/Source/HairStrandsNiagara/Private/NiagaraDataInterfaceHairStrands.cpp b/Engine/Plugins/Experimental/HairStrands/Source/HairStrandsNiagara/Private/NiagaraDataInterfaceHairStrands.cpp
--- a/Engine/Plugins/Experimental/HairStrands/Source/HairStrandsNiagara/Private/NiagaraDataInterfaceHairStrands.cpp
+++ b/Engine/Plugins/Experimental/HairStrands/Source/HairStrandsNiagara/Private/NiagaraDataInterfaceHairStrands.cpp
@@ -153,7 +153,6 @@ static const FName ComputeAirDragForceName(TEXT("ComputeAirDragForce"));
 
 static const FName NeedSimulationResetName(TEXT("NeedSimulationReset"));
 static const FName HasGlobalInterpolationName(TEXT("HasGlobalInterpolation"));
-static const FName HasKinematicsTargetName(TEXT("HasKinematicsTarget"));
 
 //------------------------------------------------------------------------------------------------------------
 
@@ -367,7 +366,7 @@ void FCopyBoundingBoxCS::UnbindBuffers(FRHICommandList& RHICmdList)
 }
 
 static void AddCopyBoundingBoxPass(
-	FRHICommandListImmediate& RHICmdList,
+	FRHICommandList& RHICmdList,
 	FRHIUnorderedAccessView* BoundingBoxBufferA,
 	FRHIUnorderedAccessView* OutBoundingBoxBufferB)
 {
@@ -396,12 +395,7 @@ inline void ClearBuffer(FRHICommandList& RHICmdList, FNDIHairStrandsBuffer* Hair
 
 	if (BoundingBoxBufferAUAV != nullptr && BoundingBoxBufferASRV != nullptr && BoundingBoxBufferBUAV != nullptr)
 	{
-		ENQUEUE_RENDER_COMMAND(FNDIHairStrandsCopyBoundingBox)(
-			[BoundingBoxBufferAUAV, BoundingBoxBufferBUAV, BoundingBoxBufferASRV, HairStrandsBuffer]
-		(FRHICommandListImmediate& RHICmdListImm)
-		{
-			AddCopyBoundingBoxPass(RHICmdListImm, HairStrandsBuffer->BoundingBoxBufferA.UAV, HairStrandsBuffer->BoundingBoxBufferB.UAV);
-		});
+		AddCopyBoundingBoxPass(RHICmdList, HairStrandsBuffer->BoundingBoxBufferA.UAV, HairStrandsBuffer->BoundingBoxBufferB.UAV);
 	}
 }
 
@@ -512,6 +506,9 @@ void FNDIHairStrandsData::Update(UNiagaraDataInterfaceHairStrands* Interface, FN
 		WorldTransform = Interface->IsComponentValid() ? Interface->SourceComponent->GetComponentToWorld() :
 			SystemInstance ? SystemInstance->GetComponent()->GetComponentToWorld() : FTransform::Identity;
 
+		GlobalInterpolation = (Interface->IsComponentValid() && Interface->SourceComponent->BindingAsset && Interface->SourceComponent->GroomAsset) ?
+			Interface->SourceComponent->GroomAsset->EnableGlobalInterpolation : false;
+
 		if (StrandsDatas != nullptr && GroomAsset != nullptr && GroupIndex >= 0 && GroupIndex < GroomAsset->HairGroupsPhysics.Num() &&
 			GroomAsset->HairGroupsPhysics[GroupIndex].SolverSettings.EnableSimulation)
 		{
@@ -520,7 +517,6 @@ void FNDIHairStrandsData::Update(UNiagaraDataInterfaceHairStrands* Interface, FN
 
 			SubSteps = HairPhysics.SolverSettings.SubSteps;
 			IterationCount = HairPhysics.SolverSettings.IterationCount;
-			KinematicsTarget = HairPhysics.SolverSettings.KinematicsTarget;
 
 			GravityVector = HairPhysics.ExternalForces.GravityVector;
 			AirDrag = HairPhysics.ExternalForces.AirDrag;
@@ -758,10 +754,7 @@ struct FNDIHairStrandsParametersCS : public FNiagaraDataInterfaceParametersCS
 
 			int32 bNeedSimReset = (ProxyData->TickCount <= MaxDelay ? 1 : 0);
 
-			static const auto GlobalInterpolationCVar = IConsoleManager::Get().FindConsoleVariable(TEXT("r.HairStrands.Interpolation.Global"));
-			const int32 GlobalInterpolation = (GlobalInterpolationCVar != nullptr) ? GlobalInterpolationCVar->GetInt() : 0;
-
-			if ((InterpolationModeValue == 1) && (GlobalInterpolation == 1))
+			if ((InterpolationModeValue == 1) && ProxyData->GlobalInterpolation)
 			{
 				InterpolationModeValue = 2;
 			}
@@ -2094,21 +2087,10 @@ void UNiagaraDataInterfaceHairStrands::GetFunctions(TArray<FNiagaraFunctionSigna
 		Sig.bMemberFunction = true;
 		Sig.bRequiresContext = false;
 		Sig.Inputs.Add(FNiagaraVariable(FNiagaraTypeDefinition(GetClass()), TEXT("Hair Strands")));
-		Sig.Outputs.Add(FNiagaraVariable(FNiagaraTypeDefinition::GetBoolDef(), TEXT("Global lInterpolation")));
-
-		OutFunctions.Add(Sig);
-	}
-	{
-		FNiagaraFunctionSignature Sig;
-		Sig.Name = HasKinematicsTargetName;
-		Sig.bMemberFunction = true;
-		Sig.bRequiresContext = false;
-		Sig.Inputs.Add(FNiagaraVariable(FNiagaraTypeDefinition(GetClass()), TEXT("Hair Strands")));
-		Sig.Outputs.Add(FNiagaraVariable(FNiagaraTypeDefinition::GetBoolDef(), TEXT("Kinematics Target")));
+		Sig.Outputs.Add(FNiagaraVariable(FNiagaraTypeDefinition::GetBoolDef(), TEXT("Global Interpolation")));
 
 		OutFunctions.Add(Sig);
 	}
-
 }
 
 DEFINE_NDI_DIRECT_FUNC_BINDER(UNiagaraDataInterfaceHairStrands, GetNumStrands);
@@ -2205,7 +2187,6 @@ DEFINE_NDI_DIRECT_FUNC_BINDER(UNiagaraDataInterfaceHairStrands, ComputeMaterialF
 DEFINE_NDI_DIRECT_FUNC_BINDER(UNiagaraDataInterfaceHairStrands, ComputeAirDragForce);
 DEFINE_NDI_DIRECT_FUNC_BINDER(UNiagaraDataInterfaceHairStrands, NeedSimulationReset);
 DEFINE_NDI_DIRECT_FUNC_BINDER(UNiagaraDataInterfaceHairStrands, HasGlobalInterpolation);
-DEFINE_NDI_DIRECT_FUNC_BINDER(UNiagaraDataInterfaceHairStrands, HasKinematicsTarget);
 
 void UNiagaraDataInterfaceHairStrands::GetVMExternalFunction(const FVMExternalFunctionBindingInfo& BindingInfo, void* InstanceData, FVMExternalFunction &OutFunc)
 {
@@ -2594,11 +2575,6 @@ void UNiagaraDataInterfaceHairStrands::GetVMExternalFunction(const FVMExternalFu
 		check(BindingInfo.GetNumInputs() == 1 && BindingInfo.GetNumOutputs() == 1);
 		NDI_FUNC_BINDER(UNiagaraDataInterfaceHairStrands, HasGlobalInterpolation)::Bind(this, OutFunc);
 	}
-	else if (BindingInfo.Name == HasKinematicsTargetName)
-	{
-		check(BindingInfo.GetNumInputs() == 1 && BindingInfo.GetNumOutputs() == 1);
-		NDI_FUNC_BINDER(UNiagaraDataInterfaceHairStrands, HasKinematicsTarget)::Bind(this, OutFunc);
-	}
 }
 
 void WriteTransform(const FMatrix& ToWrite, FVectorVMContext& Context)
@@ -3215,18 +3191,7 @@ void UNiagaraDataInterfaceHairStrands::HasGlobalInterpolation(FVectorVMContext&
 
 	for (int32 InstanceIdx = 0; InstanceIdx < Context.NumInstances; ++InstanceIdx)
 	{
-		*OutGlobalInterpolation.GetDestAndAdvance() = false;
-	}
-}
-
-void UNiagaraDataInterfaceHairStrands::HasKinematicsTarget(FVectorVMContext& Context)
-{
-	VectorVM::FUserPtrHandler<FNDIHairStrandsData> InstData(Context);
-	VectorVM::FExternalFuncRegisterHandler<int32> OutKinematicsTarget(Context);
-
-	for (int32 InstanceIdx = 0; InstanceIdx < Context.NumInstances; ++InstanceIdx)
-	{
-		*OutKinematicsTarget.GetDestAndAdvance() = InstData->KinematicsTarget;
+		*OutGlobalInterpolation.GetDestAndAdvance() = InstData->GlobalInterpolation;
 	}
 }
 
diff --git a/Engine/Plugins/Experimental/HairStrands/Source/HairStrandsNiagara/Public/NiagaraDataInterfaceHairStrands.h b/Engine/Plugins/Experimental/HairStrands/Source/HairStrandsNiagara/Public/NiagaraDataInterfaceHairStrands.h
--- a/Engine/Plugins/Experimental/HairStrands/Source/HairStrandsNiagara/Public/NiagaraDataInterfaceHairStrands.h
+++ b/Engine/Plugins/Experimental/HairStrands/Source/HairStrandsNiagara/Public/NiagaraDataInterfaceHairStrands.h
@@ -95,7 +95,7 @@ struct FNDIHairStrandsData
 		BoxCenter = FVector(0, 0, 0);
 		BoxExtent = FVector(0, 0, 0);
 
-		KinematicsTarget = false;
+		GlobalInterpolation = false;
 
 		TickCount = 0;
 		ForceReset = true;
@@ -148,7 +148,7 @@ struct FNDIHairStrandsData
 			BoxCenter = OtherDatas->BoxCenter;
 			BoxExtent = OtherDatas->BoxExtent;
 
-			KinematicsTarget = OtherDatas->KinematicsTarget;
+			GlobalInterpolation = OtherDatas->GlobalInterpolation;
 
 			TickCount = OtherDatas->TickCount;
 			ForceReset = OtherDatas->ForceReset;
@@ -192,8 +192,8 @@ struct FNDIHairStrandsData
 	/** Cached World transform. */
 	FTransform WorldTransform;
 
-	/** Kinematics target */
-	bool KinematicsTarget;
+	/** Global Interpolation */
+	bool GlobalInterpolation;
 
 	/** Number of strands*/
 	int32 NumStrands;
@@ -547,12 +547,9 @@ public:
 	/** Check if we need or not a simulation reset*/
 	void NeedSimulationReset(FVectorVMContext& Context);
 
-	/** Check if we need or not a simulation reset*/
+	/** Check if we have a global interpolation */
 	void HasGlobalInterpolation(FVectorVMContext& Context);
 
-	/** Check if we need or not a simulation reset*/
-	void HasKinematicsTarget(FVectorVMContext& Context);
-
 	/** Eval the skinned position given a rest position*/
 	void EvalSkinnedPosition(FVectorVMContext& Context);
 
