Remove ref to kinematicstarget

diff --git a/Engine/Plugins/Experimental/HairStrands/Source/HairStrandsCore/Private/GroomAsset.cpp b/Engine/Plugins/Experimental/HairStrands/Source/HairStrandsCore/Private/GroomAsset.cpp
--- a/Engine/Plugins/Experimental/HairStrands/Source/HairStrandsCore/Private/GroomAsset.cpp
+++ b/Engine/Plugins/Experimental/HairStrands/Source/HairStrandsCore/Private/GroomAsset.cpp
@@ -771,7 +771,6 @@ FHairSolverSettings::FHairSolverSettings()
 {
 	EnableSimulation = false;
 	NiagaraSolver = EGroomNiagaraSolvers::AngularSprings;
-	KinematicsTarget = false;
 	SubSteps = 5;
 	IterationCount = 5;
 }
