Fix CIS

diff --git a/Engine/Plugins/Experimental/HairStrands/Source/HairStrandsCore/Private/GroomComponent.cpp b/Engine/Plugins/Experimental/HairStrands/Source/HairStrandsCore/Private/GroomComponent.cpp
--- a/Engine/Plugins/Experimental/HairStrands/Source/HairStrandsCore/Private/GroomComponent.cpp
+++ b/Engine/Plugins/Experimental/HairStrands/Source/HairStrandsCore/Private/GroomComponent.cpp
@@ -910,7 +910,7 @@ void UGroomComponent::UpdateSimulatedGroups()
 			{
 				const bool bIsSimulationEnable = (LocalGroomAsset && GroupIt < LocalGroomAsset->HairGroupsPhysics.Num()) ? 
 					LocalGroomAsset->HairGroupsPhysics[GroupIt].SolverSettings.EnableSimulation : false;
-				const bool bHasGlobalInterpolation = (LocalBindingAsset && LocalGroomAsset->EnableGlobalInterpolation);
+				const bool bHasGlobalInterpolation = (LocalBindingAsset && LocalGroomAsset && LocalGroomAsset->EnableGlobalInterpolation);
 				HairGroup.HairInterpolationType =
 					(LocalGroomAsset && LocalGroomAsset->HairInterpolationType == EGroomInterpolationType::RigidTransform) ? 0 :
 					(LocalGroomAsset && LocalGroomAsset->HairInterpolationType == EGroomInterpolationType::OffsetTransform) ? 1 :
@@ -1072,7 +1072,7 @@ void UGroomComponent::InitResources(bool bIsBindingReloading)
 		const bool bIsSimulationEnable = (GroomAsset && GroupIt < GroomAsset->HairGroupsPhysics.Num()) ?
 			GroomAsset->HairGroupsPhysics[GroupIt].SolverSettings.EnableSimulation : false;
 
-		const bool bHasGlobalInterpolation = (BindingAsset && GroomAsset->EnableGlobalInterpolation);
+		const bool bHasGlobalInterpolation = (BindingAsset && GroomAsset && GroomAsset->EnableGlobalInterpolation);
 
 		InterpolationInputGroup.HairInterpolationType =
 			(GroomAsset && GroomAsset->HairInterpolationType == EGroomInterpolationType::RigidTransform) ? 0 :
