Fix missing release resources.

diff --git a/Engine/Plugins/Experimental/HairStrands/Source/HairStrandsCore/Private/HairStrandsVertexFactory.cpp b/Engine/Plugins/Experimental/HairStrands/Source/HairStrandsCore/Private/HairStrandsVertexFactory.cpp
--- a/Engine/Plugins/Experimental/HairStrands/Source/HairStrandsCore/Private/HairStrandsVertexFactory.cpp
+++ b/Engine/Plugins/Experimental/HairStrands/Source/HairStrandsCore/Private/HairStrandsVertexFactory.cpp
@@ -37,6 +37,12 @@ public:
 		SRVFloat = RHICreateShaderResourceView(VertexBufferRHI, sizeof(uint32), PF_R32_FLOAT);
 	}
 
+	virtual void ReleaseRHI() override
+	{
+		VertexBufferRHI.SafeRelease();
+		SRVUint.SafeRelease();
+		SRVFloat.SafeRelease();
+	}
 };
 TGlobalResource<FDummyCulledDispatchVertexIdsBuffer> GDummyCulledDispatchVertexIdsBuffer;
 
