Flagged GroomAssets as using UE5MainStreamObjectVersion (for compatibility with mesh description attribute versioning), and bumped the DDC key.

diff --git a/Engine/Plugins/Experimental/HairStrands/Source/HairStrandsCore/Private/GroomAsset.cpp b/Engine/Plugins/Experimental/HairStrands/Source/HairStrandsCore/Private/GroomAsset.cpp
--- a/Engine/Plugins/Experimental/HairStrands/Source/HairStrandsCore/Private/GroomAsset.cpp
+++ b/Engine/Plugins/Experimental/HairStrands/Source/HairStrandsCore/Private/GroomAsset.cpp
@@ -18,6 +18,7 @@
 #include "UObject/AnimObjectVersion.h"
 #include "UObject/PhysicsObjectVersion.h"
 #include "UObject/ReleaseObjectVersion.h"
+#include "UObject/UE5MainStreamObjectVersion.h"
 #include "NiagaraSystem.h"
 #include "GroomComponent.h"
 #include "Math/Box.h"
@@ -68,7 +69,8 @@ void UGroomAsset::Serialize(FArchive& Ar)
 	Super::Serialize(Ar);
 
 	Ar.UsingCustomVersion(FPhysicsObjectVersion::GUID);
-	Ar.UsingCustomVersion(FReleaseObjectVersion::GUID); // Needed to support MeshDescription AttributesSet serialization
+	Ar.UsingCustomVersion(FReleaseObjectVersion::GUID);			// Needed to support MeshDescription AttributesSet serialization
+	Ar.UsingCustomVersion(FUE5MainStreamObjectVersion::GUID);	// Needed to support MeshDescription AttributesSet serialization
 
 	if (Ar.CustomVer(FPhysicsObjectVersion::GUID) >= FPhysicsObjectVersion::GroomWithDescription)
 	{
@@ -647,7 +649,7 @@ void UGroomAsset::SetNumGroup(uint32 InGroupCount, bool bResetGroupData)
 // differences, etc.) replace the version GUID below with a new one.
 // In case of merge conflicts with DDC versions, you MUST generate a new GUID
 // and set this new GUID as the version.
-#define GROOM_DERIVED_DATA_VERSION TEXT("7AA52FFE75B54EC7A36CC989E5854A29")
+#define GROOM_DERIVED_DATA_VERSION TEXT("07D55AD513234B029DA8E702E2D7663F")
 
 #if WITH_EDITORONLY_DATA
 
