Fixing up legacy tools context creation in groom editor mode.

#rnx

diff --git a/Engine/Plugins/Experimental/HairStrands/Source/HairStrandsEditor/Private/GroomEditorMode.cpp b/Engine/Plugins/Experimental/HairStrands/Source/HairStrandsEditor/Private/GroomEditorMode.cpp
--- a/Engine/Plugins/Experimental/HairStrands/Source/HairStrandsEditor/Private/GroomEditorMode.cpp
+++ b/Engine/Plugins/Experimental/HairStrands/Source/HairStrandsEditor/Private/GroomEditorMode.cpp
@@ -171,20 +171,7 @@ void FGroomEditorMode::Enter()
 {
 	FEdMode::Enter();
 
-	// initialize FEdMode ToolsContext adapter
-
-	//ToolsContext = NewObject<UEdModeInteractiveToolsContext>(GetTransientPackage(), TEXT("ToolsContext"), RF_Transient);
-	ToolsContext = NewObject<UEdModeInteractiveToolsContext>();
-	ToolsContext->InitializeContextFromEdMode(this);
-
-	ToolsContext->OnToolNotificationMessage.AddLambda([this](const FText& Message)
-	{
-		this->OnToolNotificationMessage.Broadcast(Message);
-	});
-	ToolsContext->OnToolWarningMessage.AddLambda([this](const FText& Message)
-	{
-		this->OnToolWarningMessage.Broadcast(Message);
-	});
+	ToolsContext = Owner->GetInteractiveToolsContext();
 
 	if (!Toolkit.IsValid() && UsesToolkits())
 	{
@@ -231,6 +218,7 @@ void FGroomEditorMode::Enter()
 		CommandList->MapAction( UICommand,
 			FExecuteAction::CreateLambda([this, ToolIdentifier]() { ToolsContext->StartTool(ToolIdentifier); }),
 			FCanExecuteAction::CreateLambda([this, ToolIdentifier]() { return ToolsContext->CanStartTool(ToolIdentifier); }));
+		RegisteredTools.Emplace(UICommand, ToolIdentifier);
 	};
 
 	// register tool set
@@ -290,11 +278,14 @@ void FGroomEditorMode::Enter()
 
 void FGroomEditorMode::Exit()
 {
-	OnToolNotificationMessage.Clear();
-	OnToolWarningMessage.Clear();
-	
-	ToolsContext->ShutdownContext();
-	ToolsContext = nullptr;
+	// Shutdown and clean up our additions in the ToolsContext
+	const TSharedRef<FUICommandList>& CommandList = Toolkit->GetToolkitCommands();
+	for (auto& RegisteredTool : RegisteredTools)
+	{
+		CommandList->UnmapAction(RegisteredTool.Key);
+		ToolsContext->ToolManager->UnregisterToolType(RegisteredTool.Value);
+	}
+	RegisteredTools.SetNum(0);
 
 	if (Toolkit.IsValid())
 	{
@@ -329,6 +320,7 @@ void FGroomEditorMode::Exit()
 	}
 #endif // TOOLED_ENABLE_VIEWPORT_INTERACTION
 
+	ToolsContext = nullptr;
 
 	// Call base Exit method to ensure proper cleanup
 	FEdMode::Exit();
diff --git a/Engine/Plugins/Experimental/HairStrands/Source/HairStrandsEditor/Public/GroomEditorMode.h b/Engine/Plugins/Experimental/HairStrands/Source/HairStrandsEditor/Public/GroomEditorMode.h
--- a/Engine/Plugins/Experimental/HairStrands/Source/HairStrandsEditor/Public/GroomEditorMode.h
+++ b/Engine/Plugins/Experimental/HairStrands/Source/HairStrandsEditor/Public/GroomEditorMode.h
@@ -123,11 +123,6 @@ public:
 	// End of FEdMode interface
 	//////////////////
 
-public:
-	DECLARE_MULTICAST_DELEGATE_OneParam(FOnModelingModeToolNotification, const FText&);
-	FOnModelingModeToolNotification OnToolNotificationMessage;
-	FOnModelingModeToolNotification OnToolWarningMessage;
-
 public:
 
 	virtual UEdModeInteractiveToolsContext* GetToolsContext() const
@@ -141,6 +136,7 @@ public:
 	}
 
 protected:
+	TArray<TPair<TSharedPtr<FUICommandInfo>, FString>> RegisteredTools;
 
 	UEdModeInteractiveToolsContext* ToolsContext;
 
