Silenced PVS-Studio static analysis warnings.

diff --git a/Engine/Plugins/Experimental/HairStrands/Source/HairStrandsCore/Private/GroomComponent.cpp b/Engine/Plugins/Experimental/HairStrands/Source/HairStrandsCore/Private/GroomComponent.cpp
--- a/Engine/Plugins/Experimental/HairStrands/Source/HairStrandsCore/Private/GroomComponent.cpp
+++ b/Engine/Plugins/Experimental/HairStrands/Source/HairStrandsCore/Private/GroomComponent.cpp
@@ -1580,6 +1580,7 @@ void UGroomComponent::InitResources(bool bIsBindingReloading)
 
 			if (SkeletalMeshComponent)
 			{
+				check(BindingAsset);
 				check(GroupIt < BindingAsset->HairGroupResources.Num());
 				check(SkeletalMeshComponent->SkeletalMesh ? SkeletalMeshComponent->SkeletalMesh->GetLODInfoArray().Num() == BindingAsset->HairGroupResources[GroupIt].SimRootResources->RootData.MeshProjectionLODs.Num() : false);
 
@@ -1643,6 +1644,7 @@ void UGroomComponent::InitResources(bool bIsBindingReloading)
 
 			if (SkeletalMeshComponent)
 			{
+				check(BindingAsset);
 				check(GroupIt < BindingAsset->HairGroupResources.Num());
 				check(SkeletalMeshComponent->SkeletalMesh ? SkeletalMeshComponent->SkeletalMesh->GetLODInfoArray().Num() == BindingAsset->HairGroupResources[GroupIt].RenRootResources->RootData.MeshProjectionLODs.Num() : false);
 
@@ -1720,6 +1722,7 @@ void UGroomComponent::InitResources(bool bIsBindingReloading)
 
 					if (SkeletalMeshComponent)
 					{
+						check(BindingAsset);
 						check(SkeletalMeshComponent->GetNumLODs() == BindingAsset->HairGroupResources[GroupIt].RenRootResources->RootData.MeshProjectionLODs.Num());
 						
 						InstanceLOD.Guides.bOwnRootResourceAllocation = false;
