Fix inverted logic reported by static analysis.

diff --git a/Engine/Plugins/Experimental/HairStrands/Source/HairStrandsCore/Private/GroomTextureBuilder.cpp b/Engine/Plugins/Experimental/HairStrands/Source/HairStrandsCore/Private/GroomTextureBuilder.cpp
--- a/Engine/Plugins/Experimental/HairStrands/Source/HairStrandsCore/Private/GroomTextureBuilder.cpp
+++ b/Engine/Plugins/Experimental/HairStrands/Source/HairStrandsCore/Private/GroomTextureBuilder.cpp
@@ -216,7 +216,7 @@ static void InternalGenerateFollicleTexture_GPU(
 		FollicleMaskTexture,
 		[FollicleMaskTexture, MipCount, OutTexture](FRHICommandListImmediate& RHICmdList)
 	{
-		if (OutTexture == nullptr || OutTexture->Resource || OutTexture->Resource->GetTexture2DRHI())
+		if (OutTexture && OutTexture->Resource && OutTexture->Resource->GetTexture2DRHI())
 		{
 			FRHICopyTextureInfo CopyInfo;
 			CopyInfo.NumMips = MipCount;
