FDetailsViewArgs 8 parameter defaulted constructor deprecated and all usages fixed.

diff --git a/Engine/Plugins/2D/Paper2D/Source/Paper2DEditor/Private/ExtractSprites/SPaperExtractSpritesDialog.cpp b/Engine/Plugins/2D/Paper2D/Source/Paper2DEditor/Private/ExtractSprites/SPaperExtractSpritesDialog.cpp
--- a/Engine/Plugins/2D/Paper2D/Source/Paper2DEditor/Private/ExtractSprites/SPaperExtractSpritesDialog.cpp
+++ b/Engine/Plugins/2D/Paper2D/Source/Paper2DEditor/Private/ExtractSprites/SPaperExtractSpritesDialog.cpp
@@ -171,7 +171,10 @@ void SPaperExtractSpritesDialog::Construct(const FArguments& InArgs, UTexture2D*
 	PreviewExtractedSprites();
 
 	FPropertyEditorModule& EditModule = FModuleManager::Get().GetModuleChecked<FPropertyEditorModule>("PropertyEditor");
-	FDetailsViewArgs DetailsViewArgs(/*bUpdateFromSelection=*/ false, /*bLockable=*/ false, /*bAllowSearch=*/ false, /*InNameAreaSettings=*/ FDetailsViewArgs::HideNameArea, /*bHideSelectionTip=*/ true);
+	FDetailsViewArgs DetailsViewArgs;
+	DetailsViewArgs.bAllowSearch = false;
+	DetailsViewArgs.NameAreaSettings = FDetailsViewArgs::HideNameArea;
+	DetailsViewArgs.bHideSelectionTip = true;
 	MainPropertyView = EditModule.CreateDetailView(DetailsViewArgs);
 	MainPropertyView->SetObject(ExtractSpriteSettings);
 	MainPropertyView->OnFinishedChangingProperties().AddSP(this, &SPaperExtractSpritesDialog::OnFinishedChangingProperties);
diff --git a/Engine/Plugins/AI/EnvironmentQueryEditor/Source/EnvironmentQueryEditor/Private/EnvironmentQueryEditor.cpp b/Engine/Plugins/AI/EnvironmentQueryEditor/Source/EnvironmentQueryEditor/Private/EnvironmentQueryEditor.cpp
--- a/Engine/Plugins/AI/EnvironmentQueryEditor/Source/EnvironmentQueryEditor/Private/EnvironmentQueryEditor.cpp
+++ b/Engine/Plugins/AI/EnvironmentQueryEditor/Source/EnvironmentQueryEditor/Private/EnvironmentQueryEditor.cpp
@@ -321,9 +321,10 @@ void FEnvironmentQueryEditor::OnSelectedNodesChanged(const TSet<class UObject*>&
 void FEnvironmentQueryEditor::CreateInternalWidgets()
 {
 	FPropertyEditorModule& PropertyEditorModule = FModuleManager::GetModuleChecked<FPropertyEditorModule>( "PropertyEditor" );
-	const FDetailsViewArgs DetailsViewArgs( false, false, true, FDetailsViewArgs::ObjectsUseNameArea, false );
+	FDetailsViewArgs DetailsViewArgs;
+	DetailsViewArgs.NameAreaSettings = FDetailsViewArgs::ObjectsUseNameArea;
 	DetailsView = PropertyEditorModule.CreateDetailView( DetailsViewArgs );
-	DetailsView->SetObject( NULL );
+	DetailsView->SetObject( nullptr );
 	DetailsView->OnFinishedChangingProperties().AddSP(this, &FEnvironmentQueryEditor::OnFinishedChangingProperties);
 }
 
diff --git a/Engine/Plugins/Compositing/Composure/Source/ComposureEditor/Private/Sequencer/ComposureExportTrackEditor.cpp b/Engine/Plugins/Compositing/Composure/Source/ComposureEditor/Private/Sequencer/ComposureExportTrackEditor.cpp
--- a/Engine/Plugins/Compositing/Composure/Source/ComposureEditor/Private/Sequencer/ComposureExportTrackEditor.cpp
+++ b/Engine/Plugins/Compositing/Composure/Source/ComposureEditor/Private/Sequencer/ComposureExportTrackEditor.cpp
@@ -45,7 +45,10 @@ void FComposureExportTrackEditor::BuildTrackContextMenu(FMenuBuilder& MenuBuilde
 		FPropertyEditorModule& PropertyEditor = FModuleManager::Get().LoadModuleChecked<FPropertyEditorModule>("PropertyEditor");
 
 		// Create a details view for the track
-		FDetailsViewArgs DetailsViewArgs(false,false,false,FDetailsViewArgs::HideNameArea,true);
+		FDetailsViewArgs DetailsViewArgs;
+		DetailsViewArgs.bAllowSearch = false;
+		DetailsViewArgs.NameAreaSettings = FDetailsViewArgs::HideNameArea;
+		DetailsViewArgs.bHideSelectionTip = true;
 		DetailsViewArgs.DefaultsOnlyVisibility = EEditDefaultsOnlyNodeVisibility::Automatic;
 		DetailsViewArgs.bShowOptions = false;
 		DetailsViewArgs.ColumnWidth = 0.55f;
diff --git a/Engine/Plugins/Developer/TraceSourceFiltering/Source/SourceFilteringEditor/Private/SUserTraceFilteringWidget.cpp b/Engine/Plugins/Developer/TraceSourceFiltering/Source/SourceFilteringEditor/Private/SUserTraceFilteringWidget.cpp
--- a/Engine/Plugins/Developer/TraceSourceFiltering/Source/SourceFilteringEditor/Private/SUserTraceFilteringWidget.cpp
+++ b/Engine/Plugins/Developer/TraceSourceFiltering/Source/SourceFilteringEditor/Private/SUserTraceFilteringWidget.cpp
@@ -136,15 +136,10 @@ void SUserTraceFilteringWidget::ConstructInstanceDetailsView()
 	// Create a property view
 	FPropertyEditorModule& EditModule = FModuleManager::Get().GetModuleChecked<FPropertyEditorModule>("PropertyEditor");
 
-	FDetailsViewArgs DetailsViewArgs(
-		/*bUpdateFromSelection=*/ false,
-		/*bLockable=*/ false,
-		/*bAllowSearch=*/ false,
-		FDetailsViewArgs::HideNameArea,
-		/*bHideSelectionTip=*/ true,
-		/*InNotifyHook=*/ nullptr,
-		/*InSearchInitialKeyFocus=*/ false,
-		/*InViewIdentifier=*/ NAME_None);
+	FDetailsViewArgs DetailsViewArgs;
+	DetailsViewArgs.bAllowSearch = false;
+	DetailsViewArgs.NameAreaSettings = FDetailsViewArgs::HideNameArea;
+	DetailsViewArgs.bHideSelectionTip = true;
 	DetailsViewArgs.DefaultsOnlyVisibility = EEditDefaultsOnlyNodeVisibility::Automatic;
 	DetailsViewArgs.bShowOptions = false;	
 
diff --git a/Engine/Plugins/Editor/FacialAnimation/Source/FacialAnimationEditor/Private/SFacialAnimationBulkImporter.cpp b/Engine/Plugins/Editor/FacialAnimation/Source/FacialAnimationEditor/Private/SFacialAnimationBulkImporter.cpp
--- a/Engine/Plugins/Editor/FacialAnimation/Source/FacialAnimationEditor/Private/SFacialAnimationBulkImporter.cpp
+++ b/Engine/Plugins/Editor/FacialAnimation/Source/FacialAnimationEditor/Private/SFacialAnimationBulkImporter.cpp
@@ -22,7 +22,8 @@
 
 void SFacialAnimationBulkImporter::Construct(const FArguments& InArgs)
 {
-	FDetailsViewArgs DetailsViewArgs(false, false, true, FDetailsViewArgs::HideNameArea);
+	FDetailsViewArgs DetailsViewArgs;
+	DetailsViewArgs.NameAreaSettings = FDetailsViewArgs::HideNameArea;
 
 	FPropertyEditorModule& PropertyEditorModule = FModuleManager::Get().GetModuleChecked<FPropertyEditorModule>("PropertyEditor");
 	TSharedRef<IDetailsView> DetailsView = PropertyEditorModule.CreateDetailView(DetailsViewArgs);
diff --git a/Engine/Plugins/Editor/PluginBrowser/Source/PluginBrowser/Private/SPluginTile.cpp b/Engine/Plugins/Editor/PluginBrowser/Source/PluginBrowser/Private/SPluginTile.cpp
--- a/Engine/Plugins/Editor/PluginBrowser/Source/PluginBrowser/Private/SPluginTile.cpp
+++ b/Engine/Plugins/Editor/PluginBrowser/Source/PluginBrowser/Private/SPluginTile.cpp
@@ -566,7 +566,12 @@ void SPluginTile::OnEditPlugin()
 
 	// Create a property view
 	FPropertyEditorModule& EditModule = FModuleManager::Get().GetModuleChecked<FPropertyEditorModule>("PropertyEditor");
-	TSharedRef<IDetailsView> PropertyView = EditModule.CreateDetailView(FDetailsViewArgs(false, false, false, FDetailsViewArgs::ActorsUseNameArea, true));
+	
+	FDetailsViewArgs DetailsViewArgs;
+	DetailsViewArgs.bAllowSearch = false;
+	DetailsViewArgs.bHideSelectionTip = true;
+
+	TSharedRef<IDetailsView> PropertyView = EditModule.CreateDetailView(DetailsViewArgs);
 	PropertyView->SetObject(MetadataObject, true);
 
 	// Create the window
diff --git a/Engine/Plugins/Experimental/ChaosEditor/Source/FractureEditor/Private/FractureEditorModeToolkit.cpp b/Engine/Plugins/Experimental/ChaosEditor/Source/FractureEditor/Private/FractureEditorModeToolkit.cpp
--- a/Engine/Plugins/Experimental/ChaosEditor/Source/FractureEditor/Private/FractureEditorModeToolkit.cpp
+++ b/Engine/Plugins/Experimental/ChaosEditor/Source/FractureEditor/Private/FractureEditorModeToolkit.cpp
@@ -175,15 +175,10 @@ void FFractureEditorModeToolkit::Init(const TSharedPtr<IToolkitHost>& InitToolki
 
 	FCoreUObjectDelegates::OnObjectPropertyChanged.AddRaw(this, &FFractureEditorModeToolkit::OnObjectPostEditChange);
 
-	FDetailsViewArgs DetailsViewArgs(
-		false,  //bUpdateFromSelection=
-		false, //bLockable=
-		false, //bAllowSearch=
-		FDetailsViewArgs::HideNameArea,
-		true, //bHideSelectionTip=
-		nullptr, //InNotifyHook=
-		false, //InSearchInitialKeyFocus=
-		NAME_None); //InViewIdentifier=
+	FDetailsViewArgs DetailsViewArgs;
+	DetailsViewArgs.bAllowSearch = false;
+	DetailsViewArgs.NameAreaSettings = FDetailsViewArgs::HideNameArea;
+	DetailsViewArgs.bHideSelectionTip = true;
 	DetailsViewArgs.DefaultsOnlyVisibility = EEditDefaultsOnlyNodeVisibility::Automatic;
 	DetailsViewArgs.bShowOptions = false;
 	DetailsViewArgs.bAllowMultipleTopLevelObjects = true;
diff --git a/Engine/Plugins/Experimental/HairStrands/Source/HairStrandsEditor/Private/GroomCustomAssetEditorToolkit.cpp b/Engine/Plugins/Experimental/HairStrands/Source/HairStrandsEditor/Private/GroomCustomAssetEditorToolkit.cpp
--- a/Engine/Plugins/Experimental/HairStrands/Source/HairStrandsEditor/Private/GroomCustomAssetEditorToolkit.cpp
+++ b/Engine/Plugins/Experimental/HairStrands/Source/HairStrandsEditor/Private/GroomCustomAssetEditorToolkit.cpp
@@ -293,10 +293,6 @@ bool FGroomCustomAssetEditorToolkit::OnRequestClose()
 
 void FGroomCustomAssetEditorToolkit::InitCustomAssetEditor(const EToolkitMode::Type Mode, const TSharedPtr<class IToolkitHost>& InitToolkitHost, UGroomAsset* InCustomAsset)
 {
-	const bool bIsUpdatable		= false;
-	const bool bAllowFavorites	= true;
-	const bool bIsLockable		= false;
-
 	PreviewGroomComponent		 = nullptr;
 	PreviewStaticMeshComponent	 = nullptr;
 	PreviewSkeletalMeshComponent = nullptr;
@@ -306,15 +302,18 @@ void FGroomCustomAssetEditorToolkit::InitCustomAssetEditor(const EToolkitMode::T
 	
 	FPropertyEditorModule& PropertyEditorModule = FModuleManager::GetModuleChecked<FPropertyEditorModule>("PropertyEditor");
 
-	DetailView_LODProperties			= PropertyEditorModule.CreateDetailView(FDetailsViewArgs(bIsUpdatable, bIsLockable, true, FDetailsViewArgs::ObjectsUseNameArea, false));
-	DetailView_InterpolationProperties	= PropertyEditorModule.CreateDetailView(FDetailsViewArgs(bIsUpdatable, bIsLockable, true, FDetailsViewArgs::ObjectsUseNameArea, false));
-	DetailView_RenderingProperties		= PropertyEditorModule.CreateDetailView(FDetailsViewArgs(bIsUpdatable, bIsLockable, true, FDetailsViewArgs::ObjectsUseNameArea, false));
-	DetailView_PhysicsProperties		= PropertyEditorModule.CreateDetailView(FDetailsViewArgs(bIsUpdatable, bIsLockable, true, FDetailsViewArgs::ObjectsUseNameArea, false));
-	DetailView_CardsProperties			= PropertyEditorModule.CreateDetailView(FDetailsViewArgs(bIsUpdatable, bIsLockable, true, FDetailsViewArgs::ObjectsUseNameArea, false));
-	DetailView_MeshesProperties			= PropertyEditorModule.CreateDetailView(FDetailsViewArgs(bIsUpdatable, bIsLockable, true, FDetailsViewArgs::ObjectsUseNameArea, false));
-	DetailView_MaterialProperties		= PropertyEditorModule.CreateDetailView(FDetailsViewArgs(bIsUpdatable, bIsLockable, true, FDetailsViewArgs::ObjectsUseNameArea, false));
+	FDetailsViewArgs DetailsViewArgs;
+	DetailsViewArgs.NameAreaSettings = FDetailsViewArgs::ObjectsUseNameArea;
+
+	DetailView_LODProperties			= PropertyEditorModule.CreateDetailView(DetailsViewArgs);
+	DetailView_InterpolationProperties	= PropertyEditorModule.CreateDetailView(DetailsViewArgs);
+	DetailView_RenderingProperties		= PropertyEditorModule.CreateDetailView(DetailsViewArgs);
+	DetailView_PhysicsProperties		= PropertyEditorModule.CreateDetailView(DetailsViewArgs);
+	DetailView_CardsProperties			= PropertyEditorModule.CreateDetailView(DetailsViewArgs);
+	DetailView_MeshesProperties			= PropertyEditorModule.CreateDetailView(DetailsViewArgs);
+	DetailView_MaterialProperties		= PropertyEditorModule.CreateDetailView(DetailsViewArgs);
 #if GROOMEDITOR_ENABLE_COMPONENT_PANEL
-	DetailView_PreviewGroomComponent	= PropertyEditorModule.CreateDetailView(FDetailsViewArgs(bIsUpdatable, bIsLockable, true, FDetailsViewArgs::ObjectsUseNameArea, false));
+	DetailView_PreviewGroomComponent	= PropertyEditorModule.CreateDetailView(DetailsViewArgs);
 #endif
 
 	// Customization
diff --git a/Engine/Plugins/Experimental/ModelingToolsEditorMode/Source/ModelingToolsEditorMode/Private/ModelingToolsEditorModeToolkit.cpp b/Engine/Plugins/Experimental/ModelingToolsEditorMode/Source/ModelingToolsEditorMode/Private/ModelingToolsEditorModeToolkit.cpp
--- a/Engine/Plugins/Experimental/ModelingToolsEditorMode/Source/ModelingToolsEditorMode/Private/ModelingToolsEditorModeToolkit.cpp
+++ b/Engine/Plugins/Experimental/ModelingToolsEditorMode/Source/ModelingToolsEditorMode/Private/ModelingToolsEditorModeToolkit.cpp
@@ -47,15 +47,10 @@ void FModelingToolsEditorModeToolkit::Init(const TSharedPtr<IToolkitHost>& InitT
 {
 	FPropertyEditorModule& PropertyEditorModule = FModuleManager::GetModuleChecked<FPropertyEditorModule>("PropertyEditor");
 
-	FDetailsViewArgs DetailsViewArgs(
-		/*bUpdateFromSelection=*/ false,
-		/*bLockable=*/ false,
-		/*bAllowSearch=*/ false,
-		FDetailsViewArgs::HideNameArea,
-		/*bHideSelectionTip=*/ true,
-		/*InNotifyHook=*/ nullptr,
-		/*InSearchInitialKeyFocus=*/ false,
-		/*InViewIdentifier=*/ NAME_None);
+	FDetailsViewArgs DetailsViewArgs;
+	DetailsViewArgs.bAllowSearch = false;
+	DetailsViewArgs.NameAreaSettings = FDetailsViewArgs::HideNameArea;
+	DetailsViewArgs.bHideSelectionTip = true;
 	DetailsViewArgs.DefaultsOnlyVisibility = EEditDefaultsOnlyNodeVisibility::Automatic;
 	DetailsViewArgs.bShowOptions = false;
 	DetailsViewArgs.bAllowMultipleTopLevelObjects = true;
diff --git a/Engine/Plugins/FX/Niagara/Source/NiagaraEditor/Private/Toolkits/NiagaraParameterCollectionToolkit.cpp b/Engine/Plugins/FX/Niagara/Source/NiagaraEditor/Private/Toolkits/NiagaraParameterCollectionToolkit.cpp
--- a/Engine/Plugins/FX/Niagara/Source/NiagaraEditor/Private/Toolkits/NiagaraParameterCollectionToolkit.cpp
+++ b/Engine/Plugins/FX/Niagara/Source/NiagaraEditor/Private/Toolkits/NiagaraParameterCollectionToolkit.cpp
@@ -175,7 +175,10 @@ TSharedRef<SDockTab> FNiagaraParameterCollectionToolkit::SpawnTab_Main(const FSp
 	DetailArgs.bAllowSearch = false;
 
 	FPropertyEditorModule& PropertyEditorModule = FModuleManager::GetModuleChecked<FPropertyEditorModule>("PropertyEditor");
-	FDetailsViewArgs DetailsViewArgs(false, false, true, FDetailsViewArgs::HideNameArea, true, ParameterCollectionViewModel.Get());
+	FDetailsViewArgs DetailsViewArgs;
+	DetailsViewArgs.NameAreaSettings = FDetailsViewArgs::HideNameArea;
+	DetailsViewArgs.bHideSelectionTip = true;
+	DetailsViewArgs.NotifyHook = ParameterCollectionViewModel.Get();
 	TSharedRef<IDetailsView> DetailsView = PropertyEditorModule.CreateDetailView(DetailsViewArgs);
 
 	if (Instance->IsDefaultInstance())
diff --git a/Engine/Plugins/FX/Niagara/Source/NiagaraEditor/Private/Toolkits/NiagaraScriptToolkit.cpp b/Engine/Plugins/FX/Niagara/Source/NiagaraEditor/Private/Toolkits/NiagaraScriptToolkit.cpp
--- a/Engine/Plugins/FX/Niagara/Source/NiagaraEditor/Private/Toolkits/NiagaraScriptToolkit.cpp
+++ b/Engine/Plugins/FX/Niagara/Source/NiagaraEditor/Private/Toolkits/NiagaraScriptToolkit.cpp
@@ -355,7 +355,9 @@ TSharedRef<SDockTab> FNiagaraScriptToolkit::SpawnTabScriptDetails(const FSpawnTa
 	TWeakPtr<FNiagaraScriptViewModel> ScriptViewModelWeakPtr = ScriptViewModel;
 
 	FPropertyEditorModule& PropertyEditorModule = FModuleManager::GetModuleChecked<FPropertyEditorModule>("PropertyEditor");
-	FDetailsViewArgs DetailsViewArgs(false, false, true, FDetailsViewArgs::HideNameArea, true);
+	FDetailsViewArgs DetailsViewArgs;
+	DetailsViewArgs.NameAreaSettings = FDetailsViewArgs::HideNameArea;
+	DetailsViewArgs.bHideSelectionTip = true;
 	TSharedRef<IDetailsView> DetailsView = PropertyEditorModule.CreateDetailView(DetailsViewArgs);
 
 	DetailsView->OnFinishedChangingProperties().AddRaw(this, &FNiagaraScriptToolkit::OnEditedScriptPropertyFinishedChanging);
diff --git a/Engine/Plugins/FX/Niagara/Source/NiagaraEditor/Private/Widgets/SNiagaraParameterCollection.cpp b/Engine/Plugins/FX/Niagara/Source/NiagaraEditor/Private/Widgets/SNiagaraParameterCollection.cpp
--- a/Engine/Plugins/FX/Niagara/Source/NiagaraEditor/Private/Widgets/SNiagaraParameterCollection.cpp
+++ b/Engine/Plugins/FX/Niagara/Source/NiagaraEditor/Private/Widgets/SNiagaraParameterCollection.cpp
@@ -498,8 +498,13 @@ TSharedRef<ITableRow> SNiagaraParameterCollection::OnGenerateRowForParameter(TSh
 
 		FPropertyEditorModule& PropertyEditorModule = FModuleManager::GetModuleChecked<FPropertyEditorModule>("PropertyEditor");
 
+		FDetailsViewArgs DetailsViewArgs;
+		DetailsViewArgs.bAllowSearch = false;
+		DetailsViewArgs.NameAreaSettings = FDetailsViewArgs::HideNameArea;
+		DetailsViewArgs.bHideSelectionTip = true;
+
 		TSharedRef<IStructureDetailsView> StructureDetailsView = PropertyEditorModule.CreateStructureDetailView(
-			FDetailsViewArgs(false, false, false, FDetailsViewArgs::HideNameArea, true),
+			DetailsViewArgs,
 			FStructureDetailsViewArgs(),
 			nullptr);
 
@@ -525,8 +530,13 @@ TSharedRef<ITableRow> SNiagaraParameterCollection::OnGenerateRowForParameter(TSh
 
 		FPropertyEditorModule& PropertyEditorModule = FModuleManager::GetModuleChecked<FPropertyEditorModule>("PropertyEditor");
 
-		TSharedRef<IDetailsView> DetailsView = PropertyEditorModule.CreateDetailView(
-			FDetailsViewArgs(false, false, false, FDetailsViewArgs::HideNameArea, true, this));
+		FDetailsViewArgs DetailsViewArgs;
+		DetailsViewArgs.bAllowSearch = false;
+		DetailsViewArgs.NameAreaSettings = FDetailsViewArgs::HideNameArea;
+		DetailsViewArgs.bHideSelectionTip = true;
+		DetailsViewArgs.NotifyHook = this;
+
+		TSharedRef<IDetailsView> DetailsView = PropertyEditorModule.CreateDetailView(DetailsViewArgs);
 		DetailsView->SetObject(Item->GetDefaultValueObject());
 		DetailsWidget = DetailsView;
 	}
diff --git a/Engine/Plugins/FX/Niagara/Source/NiagaraEditor/Private/Widgets/SNiagaraSelectedObjectsDetails.cpp b/Engine/Plugins/FX/Niagara/Source/NiagaraEditor/Private/Widgets/SNiagaraSelectedObjectsDetails.cpp
--- a/Engine/Plugins/FX/Niagara/Source/NiagaraEditor/Private/Widgets/SNiagaraSelectedObjectsDetails.cpp
+++ b/Engine/Plugins/FX/Niagara/Source/NiagaraEditor/Private/Widgets/SNiagaraSelectedObjectsDetails.cpp
@@ -15,7 +15,9 @@ void SNiagaraSelectedObjectsDetails::Construct(const FArguments& InArgs, TShared
 	SelectedObjectsArray[0]->OnSelectedObjectsChanged().AddSP(this, &SNiagaraSelectedObjectsDetails::SelectedObjectsChanged);
 
 	FPropertyEditorModule& PropertyEditorModule = FModuleManager::GetModuleChecked<FPropertyEditorModule>("PropertyEditor");
-	FDetailsViewArgs DetailsViewArgs(false, false, true, FDetailsViewArgs::HideNameArea, true);
+	FDetailsViewArgs DetailsViewArgs;
+	DetailsViewArgs.NameAreaSettings = FDetailsViewArgs::HideNameArea;
+	DetailsViewArgs.bHideSelectionTip = true;
 	DetailsView = PropertyEditorModule.CreateDetailView(DetailsViewArgs);
 	DetailsView->SetObjects(SelectedObjectsArray[0]->GetSelectedObjects().Array());
 	DetailsView->OnFinishedChangingProperties().AddRaw(this, &SNiagaraSelectedObjectsDetails::OnDetailsPanelFinishedChangingProperties);
@@ -34,7 +36,9 @@ void SNiagaraSelectedObjectsDetails::Construct(const FArguments& InArgs, TShared
 	SelectedObjectsArray[1]->OnSelectedObjectsChanged().AddSP(this, &SNiagaraSelectedObjectsDetails::SelectedObjectsChangedSecond);
 
 	FPropertyEditorModule& PropertyEditorModule = FModuleManager::GetModuleChecked<FPropertyEditorModule>("PropertyEditor");
-	FDetailsViewArgs DetailsViewArgs(false, false, true, FDetailsViewArgs::HideNameArea, true);
+	FDetailsViewArgs DetailsViewArgs;
+	DetailsViewArgs.NameAreaSettings = FDetailsViewArgs::HideNameArea;
+	DetailsViewArgs.bHideSelectionTip = true;
 	DetailsView = PropertyEditorModule.CreateDetailView(DetailsViewArgs);
 	DetailsView->SetObjects(SelectedObjectsArray[0]->GetSelectedObjects().Array());
 	DetailsView->OnFinishedChangingProperties().AddRaw(this, &SNiagaraSelectedObjectsDetails::OnDetailsPanelFinishedChangingProperties);
diff --git a/Engine/Plugins/FX/Niagara/Source/NiagaraEditorWidgets/Private/Stack/SNiagaraStackFunctionInputValue.cpp b/Engine/Plugins/FX/Niagara/Source/NiagaraEditorWidgets/Private/Stack/SNiagaraStackFunctionInputValue.cpp
--- a/Engine/Plugins/FX/Niagara/Source/NiagaraEditorWidgets/Private/Stack/SNiagaraStackFunctionInputValue.cpp
+++ b/Engine/Plugins/FX/Niagara/Source/NiagaraEditorWidgets/Private/Stack/SNiagaraStackFunctionInputValue.cpp
@@ -294,8 +294,13 @@ TSharedRef<SWidget> SNiagaraStackFunctionInputValue::ConstructLocalValueStructWi
 		{
 			FPropertyEditorModule& PropertyEditorModule = FModuleManager::GetModuleChecked<FPropertyEditorModule>("PropertyEditor");
 
+			FDetailsViewArgs DetailsViewArgs;
+			DetailsViewArgs.bAllowSearch = false;
+			DetailsViewArgs.NameAreaSettings = FDetailsViewArgs::HideNameArea;
+			DetailsViewArgs.bHideSelectionTip = true;
+
 			TSharedRef<IStructureDetailsView> StructureDetailsView = PropertyEditorModule.CreateStructureDetailView(
-				FDetailsViewArgs(false, false, false, FDetailsViewArgs::HideNameArea, true),
+				FDetailsViewArgs(DetailsViewArgs),
 				FStructureDetailsViewArgs(),
 				nullptr);
 
diff --git a/Engine/Plugins/Media/MediaFrameworkUtilities/Source/MediaFrameworkUtilitiesEditor/Private/AssetEditor/MediaBundleEditorToolkit.cpp b/Engine/Plugins/Media/MediaFrameworkUtilities/Source/MediaFrameworkUtilitiesEditor/Private/AssetEditor/MediaBundleEditorToolkit.cpp
--- a/Engine/Plugins/Media/MediaFrameworkUtilities/Source/MediaFrameworkUtilitiesEditor/Private/AssetEditor/MediaBundleEditorToolkit.cpp
+++ b/Engine/Plugins/Media/MediaFrameworkUtilities/Source/MediaFrameworkUtilitiesEditor/Private/AssetEditor/MediaBundleEditorToolkit.cpp
@@ -37,10 +37,8 @@ void FMediaBundleEditorToolkit::InitMediaBundleEditor(const EToolkitMode::Type M
 {
 	GEditor->GetEditorSubsystem<UImportSubsystem>()->OnAssetPostImport.AddRaw(this, &FMediaBundleEditorToolkit::HandleAssetPostImport);
 
-	const bool bIsUpdatable = false;
-	const bool bAllowFavorites = true;
-	const bool bIsLockable = false;
-	const FDetailsViewArgs DetailsViewArgs(bIsUpdatable, bIsLockable, true, FDetailsViewArgs::ObjectsUseNameArea, false);
+	FDetailsViewArgs DetailsViewArgs;
+	DetailsViewArgs.NameAreaSettings = FDetailsViewArgs::ObjectsUseNameArea;
 	DetailsView = FModuleManager::GetModuleChecked<FPropertyEditorModule>("PropertyEditor").CreateDetailView(DetailsViewArgs);
 
 	const TSharedRef<FTabManager::FLayout> StandaloneDefaultLayout = FTabManager::NewLayout(MediaBundleEditorToolkit::Layout)
diff --git a/Engine/Plugins/Media/TimecodeSynchronizer/Source/TimecodeSynchronizerEditor/Private/AssetEditor/TimecodeSynchronizerEditorToolkit.cpp b/Engine/Plugins/Media/TimecodeSynchronizer/Source/TimecodeSynchronizerEditor/Private/AssetEditor/TimecodeSynchronizerEditorToolkit.cpp
--- a/Engine/Plugins/Media/TimecodeSynchronizer/Source/TimecodeSynchronizerEditor/Private/AssetEditor/TimecodeSynchronizerEditorToolkit.cpp
+++ b/Engine/Plugins/Media/TimecodeSynchronizer/Source/TimecodeSynchronizerEditor/Private/AssetEditor/TimecodeSynchronizerEditorToolkit.cpp
@@ -42,10 +42,8 @@ void FTimecodeSynchronizerEditorToolkit::InitTimecodeSynchronizerEditor(const ET
 {
 	GEditor->GetEditorSubsystem<UImportSubsystem>()->OnAssetPostImport.AddRaw(this, &FTimecodeSynchronizerEditorToolkit::HandleAssetPostImport);
 
-	const bool bIsUpdatable = false;
-	const bool bAllowFavorites = true;
-	const bool bIsLockable = false;
-	const FDetailsViewArgs DetailsViewArgs(bIsUpdatable, bIsLockable, true, FDetailsViewArgs::ObjectsUseNameArea, false);
+	FDetailsViewArgs DetailsViewArgs;
+	DetailsViewArgs.NameAreaSettings = FDetailsViewArgs::ObjectsUseNameArea;
 	DetailsView = FModuleManager::GetModuleChecked<FPropertyEditorModule>("PropertyEditor").CreateDetailView(DetailsViewArgs);
 	const TSharedRef<FTabManager::FLayout> StandaloneDefaultLayout = FTabManager::NewLayout(TimecodeSynchronizerEditorToolkit::Layout)
 		->AddArea
diff --git a/Engine/Plugins/VirtualProduction/Takes/Source/TakeRecorder/Private/Widgets/SLevelSequenceTakeEditor.cpp b/Engine/Plugins/VirtualProduction/Takes/Source/TakeRecorder/Private/Widgets/SLevelSequenceTakeEditor.cpp
--- a/Engine/Plugins/VirtualProduction/Takes/Source/TakeRecorder/Private/Widgets/SLevelSequenceTakeEditor.cpp
+++ b/Engine/Plugins/VirtualProduction/Takes/Source/TakeRecorder/Private/Widgets/SLevelSequenceTakeEditor.cpp
@@ -775,7 +775,10 @@ class FRecorderSourceObjectCustomization : public IDetailCustomization
 void SLevelSequenceTakeEditor::AddDetails(const TPair<const UClass*, TArray<UObject*> >& Pair, TArray<FObjectKey>& PreviousClasses)
 {
 	FPropertyEditorModule& PropertyEditorModule = FModuleManager::Get().LoadModuleChecked<FPropertyEditorModule>("PropertyEditor");
-	FDetailsViewArgs DetailsViewArgs(false, false, false, FDetailsViewArgs::HideNameArea, true);
+	FDetailsViewArgs DetailsViewArgs;
+	DetailsViewArgs.bAllowSearch = false;
+	DetailsViewArgs.NameAreaSettings = FDetailsViewArgs::HideNameArea;
+	DetailsViewArgs.bHideSelectionTip = true;
 	DetailsViewArgs.bShowScrollBar = false;
 
 	PreviousClasses.Remove(Pair.Key);
diff --git a/Engine/Source/Editor/AdvancedPreviewScene/Private/SAdvancedPreviewDetailsTab.cpp b/Engine/Source/Editor/AdvancedPreviewScene/Private/SAdvancedPreviewDetailsTab.cpp
--- a/Engine/Source/Editor/AdvancedPreviewScene/Private/SAdvancedPreviewDetailsTab.cpp
+++ b/Engine/Source/Editor/AdvancedPreviewScene/Private/SAdvancedPreviewDetailsTab.cpp
@@ -241,15 +241,9 @@ void SAdvancedPreviewDetailsTab::CreateSettingsView()
 	// Create a property view
 	FPropertyEditorModule& EditModule = FModuleManager::Get().GetModuleChecked<FPropertyEditorModule>("PropertyEditor");
 
-	FDetailsViewArgs DetailsViewArgs(
-		/*bUpdateFromSelection=*/ false,
-		/*bLockable=*/ false,
-		/*bAllowSearch=*/ true,
-		FDetailsViewArgs::HideNameArea,
-		/*bHideSelectionTip=*/ true,
-		/*InNotifyHook=*/ nullptr,
-		/*InSearchInitialKeyFocus=*/ false,
-		/*InViewIdentifier=*/ NAME_None);
+	FDetailsViewArgs DetailsViewArgs;
+	DetailsViewArgs.NameAreaSettings = FDetailsViewArgs::HideNameArea;
+	DetailsViewArgs.bHideSelectionTip = true;
 	DetailsViewArgs.DefaultsOnlyVisibility = EEditDefaultsOnlyNodeVisibility::Automatic;
 	DetailsViewArgs.bShowOptions = false;
 	DetailsViewArgs.bAllowMultipleTopLevelObjects = true;
diff --git a/Engine/Source/Editor/AnimationModifiers/Private/SAnimationModifierContentBrowserWindow.cpp b/Engine/Source/Editor/AnimationModifiers/Private/SAnimationModifierContentBrowserWindow.cpp
--- a/Engine/Source/Editor/AnimationModifiers/Private/SAnimationModifierContentBrowserWindow.cpp
+++ b/Engine/Source/Editor/AnimationModifiers/Private/SAnimationModifierContentBrowserWindow.cpp
@@ -166,15 +166,10 @@ void SAnimationModifierContentBrowserWindow::CreateInstanceDetailsView()
 	// Create a property view
 	FPropertyEditorModule& EditModule = FModuleManager::Get().GetModuleChecked<FPropertyEditorModule>("PropertyEditor");
 
-	FDetailsViewArgs DetailsViewArgs(
-		/*bUpdateFromSelection=*/ false,
-		/*bLockable=*/ false,
-		/*bAllowSearch=*/ false,
-		FDetailsViewArgs::HideNameArea,
-		/*bHideSelectionTip=*/ true,
-		/*InNotifyHook=*/ nullptr,
-		/*InSearchInitialKeyFocus=*/ false,
-		/*InViewIdentifier=*/ NAME_None);
+	FDetailsViewArgs DetailsViewArgs;
+	DetailsViewArgs.bAllowSearch = false;
+	DetailsViewArgs.NameAreaSettings = FDetailsViewArgs::HideNameArea;
+	DetailsViewArgs.bHideSelectionTip = true;
 	DetailsViewArgs.DefaultsOnlyVisibility = EEditDefaultsOnlyNodeVisibility::Automatic;
 	DetailsViewArgs.bShowOptions = false;
 
diff --git a/Engine/Source/Editor/AnimationModifiers/Private/SAnimationModifiersTab.cpp b/Engine/Source/Editor/AnimationModifiers/Private/SAnimationModifiersTab.cpp
--- a/Engine/Source/Editor/AnimationModifiers/Private/SAnimationModifiersTab.cpp
+++ b/Engine/Source/Editor/AnimationModifiers/Private/SAnimationModifiersTab.cpp
@@ -242,15 +242,10 @@ void SAnimationModifiersTab::CreateInstanceDetailsView()
 	// Create a property view
 	FPropertyEditorModule& EditModule = FModuleManager::Get().GetModuleChecked<FPropertyEditorModule>("PropertyEditor");
 
-	FDetailsViewArgs DetailsViewArgs(
-		/*bUpdateFromSelection=*/ false,
-		/*bLockable=*/ false,
-		/*bAllowSearch=*/ false,
-		FDetailsViewArgs::HideNameArea,
-		/*bHideSelectionTip=*/ true,
-		/*InNotifyHook=*/ nullptr,
-		/*InSearchInitialKeyFocus=*/ false,
-		/*InViewIdentifier=*/ NAME_None);
+	FDetailsViewArgs DetailsViewArgs;
+	DetailsViewArgs.bAllowSearch = false;
+	DetailsViewArgs.NameAreaSettings = FDetailsViewArgs::HideNameArea;
+	DetailsViewArgs.bHideSelectionTip = true;
 	DetailsViewArgs.DefaultsOnlyVisibility = EEditDefaultsOnlyNodeVisibility::Automatic;
 	DetailsViewArgs.bShowOptions = false;	
 
diff --git a/Engine/Source/Editor/AudioEditor/Private/SoundClassEditor.cpp b/Engine/Source/Editor/AudioEditor/Private/SoundClassEditor.cpp
--- a/Engine/Source/Editor/AudioEditor/Private/SoundClassEditor.cpp
+++ b/Engine/Source/Editor/AudioEditor/Private/SoundClassEditor.cpp
@@ -275,7 +275,8 @@ void FSoundClassEditor::CreateInternalWidgets()
 	GraphEditor = CreateGraphEditorWidget();
 
 	FPropertyEditorModule& PropertyEditorModule = FModuleManager::GetModuleChecked<FPropertyEditorModule>( "PropertyEditor" );
-	const FDetailsViewArgs DetailsViewArgs( false, false, true, FDetailsViewArgs::HideNameArea, false );
+	FDetailsViewArgs DetailsViewArgs;
+	DetailsViewArgs.NameAreaSettings = FDetailsViewArgs::HideNameArea;
 	DetailsView = PropertyEditorModule.CreateDetailView( DetailsViewArgs );
 	DetailsView->SetObject( SoundClass );
 }
diff --git a/Engine/Source/Editor/AudioEditor/Private/SoundSubmixEditor.cpp b/Engine/Source/Editor/AudioEditor/Private/SoundSubmixEditor.cpp
--- a/Engine/Source/Editor/AudioEditor/Private/SoundSubmixEditor.cpp
+++ b/Engine/Source/Editor/AudioEditor/Private/SoundSubmixEditor.cpp
@@ -273,7 +273,8 @@ void FSoundSubmixEditor::CreateInternalWidgets(USoundSubmixBase* InSoundSubmix)
 	GraphEditor = CreateGraphEditorWidget(InSoundSubmix);
 
 	FPropertyEditorModule& PropertyEditorModule = FModuleManager::GetModuleChecked<FPropertyEditorModule>("PropertyEditor");
-	const FDetailsViewArgs DetailsViewArgs(false, false, true, FDetailsViewArgs::HideNameArea, false);
+	FDetailsViewArgs DetailsViewArgs;
+	DetailsViewArgs.NameAreaSettings = FDetailsViewArgs::HideNameArea;
 	DetailsView = PropertyEditorModule.CreateDetailView(DetailsViewArgs);
 	DetailsView->SetObject(InSoundSubmix);
 }
diff --git a/Engine/Source/Editor/BehaviorTreeEditor/Private/BehaviorTreeEditor.cpp b/Engine/Source/Editor/BehaviorTreeEditor/Private/BehaviorTreeEditor.cpp
--- a/Engine/Source/Editor/BehaviorTreeEditor/Private/BehaviorTreeEditor.cpp
+++ b/Engine/Source/Editor/BehaviorTreeEditor/Private/BehaviorTreeEditor.cpp
@@ -692,14 +692,10 @@ TSharedRef<SWidget> FBehaviorTreeEditor::SpawnBlackboardEditor()
 
 TSharedRef<SWidget> FBehaviorTreeEditor::SpawnBlackboardDetails()
 {
-	const bool bIsUpdatable = false;
-	const bool bAllowFavorites = true;
-	const bool bIsLockable = false;
-	const bool bAllowSearch = true;
-	const bool bObjectsUseNameArea = false;
-	const bool bHideSelectionTip = true;
 	FPropertyEditorModule& PropertyEditorModule = FModuleManager::GetModuleChecked<FPropertyEditorModule>( "PropertyEditor" );
-	FDetailsViewArgs DetailsViewArgs( bIsUpdatable, bIsLockable, bAllowSearch, FDetailsViewArgs::HideNameArea, bHideSelectionTip );
+	FDetailsViewArgs DetailsViewArgs;
+	DetailsViewArgs.NameAreaSettings = FDetailsViewArgs::HideNameArea;
+	DetailsViewArgs.bHideSelectionTip = true;
 	DetailsViewArgs.NotifyHook = this;
 	BlackboardDetailsView = PropertyEditorModule.CreateDetailView( DetailsViewArgs );
 
@@ -892,7 +888,8 @@ void FBehaviorTreeEditor::GetAbortModePreview(const class UBTDecorator* Decorato
 void FBehaviorTreeEditor::CreateInternalWidgets()
 {
 	FPropertyEditorModule& PropertyEditorModule = FModuleManager::GetModuleChecked<FPropertyEditorModule>( "PropertyEditor" );
-	FDetailsViewArgs DetailsViewArgs( false, false, true, FDetailsViewArgs::HideNameArea, false );
+	FDetailsViewArgs DetailsViewArgs; 
+	DetailsViewArgs.NameAreaSettings = FDetailsViewArgs::HideNameArea;
 	DetailsViewArgs.NotifyHook = this;
 	DetailsViewArgs.DefaultsOnlyVisibility = EEditDefaultsOnlyNodeVisibility::Hide;
 	DetailsView = PropertyEditorModule.CreateDetailView( DetailsViewArgs );
diff --git a/Engine/Source/Editor/BehaviorTreeEditor/Private/SBehaviorTreeDiff.cpp b/Engine/Source/Editor/BehaviorTreeEditor/Private/SBehaviorTreeDiff.cpp
--- a/Engine/Source/Editor/BehaviorTreeEditor/Private/SBehaviorTreeDiff.cpp
+++ b/Engine/Source/Editor/BehaviorTreeEditor/Private/SBehaviorTreeDiff.cpp
@@ -442,7 +442,8 @@ void SBehaviorTreeDiff::FBehaviorTreeDiffPanel::GeneratePanel(UEdGraph* Graph, U
 		];
 
 	FPropertyEditorModule& PropertyEditorModule = FModuleManager::GetModuleChecked<FPropertyEditorModule>( "PropertyEditor" );
-	FDetailsViewArgs DetailsViewArgs( false, false, true, FDetailsViewArgs::ObjectsUseNameArea, false );
+	FDetailsViewArgs DetailsViewArgs;
+	DetailsViewArgs.NameAreaSettings = FDetailsViewArgs::ObjectsUseNameArea;
 	DetailsViewArgs.DefaultsOnlyVisibility = EEditDefaultsOnlyNodeVisibility::Hide;
 	DetailsView = PropertyEditorModule.CreateDetailView( DetailsViewArgs );
 	DetailsView->SetObject( NULL );
diff --git a/Engine/Source/Editor/ClothPainter/Private/SClothPaintTab.cpp b/Engine/Source/Editor/ClothPainter/Private/SClothPaintTab.cpp
--- a/Engine/Source/Editor/ClothPainter/Private/SClothPaintTab.cpp
+++ b/Engine/Source/Editor/ClothPainter/Private/SClothPaintTab.cpp
@@ -51,15 +51,10 @@ void SClothPaintTab::Construct(const FArguments& InArgs)
 	// Detail view for UClothingAssetCommon
 	FPropertyEditorModule& EditModule = FModuleManager::Get().GetModuleChecked<FPropertyEditorModule>("PropertyEditor");
 
-	FDetailsViewArgs DetailsViewArgs(
-		/*bUpdateFromSelection=*/ false,
-		/*bLockable=*/ false,
-		/*bAllowSearch=*/ false,
-		FDetailsViewArgs::HideNameArea,
-		/*bHideSelectionTip=*/ true,
-		/*InNotifyHook=*/ nullptr,
-		/*InSearchInitialKeyFocus=*/ false,
-		/*InViewIdentifier=*/ NAME_None);
+	FDetailsViewArgs DetailsViewArgs;
+	DetailsViewArgs.bAllowSearch = false;
+	DetailsViewArgs.NameAreaSettings = FDetailsViewArgs::HideNameArea;
+	DetailsViewArgs.bHideSelectionTip = true;
 	DetailsViewArgs.DefaultsOnlyVisibility = EEditDefaultsOnlyNodeVisibility::Automatic;
 	DetailsViewArgs.bShowOptions = false;
 	DetailsViewArgs.bAllowMultipleTopLevelObjects = true;
diff --git a/Engine/Source/Editor/ClothPainter/Private/SClothPaintWidget.cpp b/Engine/Source/Editor/ClothPainter/Private/SClothPaintWidget.cpp
--- a/Engine/Source/Editor/ClothPainter/Private/SClothPaintWidget.cpp
+++ b/Engine/Source/Editor/ClothPainter/Private/SClothPaintWidget.cpp
@@ -73,15 +73,10 @@ void SClothPaintWidget::CreateDetailsView(FClothPainter* InPainter)
 {
 	FPropertyEditorModule& EditModule = FModuleManager::Get().GetModuleChecked<FPropertyEditorModule>("PropertyEditor");
 
-	FDetailsViewArgs DetailsViewArgs(
-		/*bUpdateFromSelection=*/ false,
-		/*bLockable=*/ false,
-		/*bAllowSearch=*/ false,
-		FDetailsViewArgs::HideNameArea,
-		/*bHideSelectionTip=*/ true,
-		/*InNotifyHook=*/ nullptr,
-		/*InSearchInitialKeyFocus=*/ false,
-		/*InViewIdentifier=*/ NAME_None);
+	FDetailsViewArgs DetailsViewArgs;
+	DetailsViewArgs.bAllowSearch = false;
+	DetailsViewArgs.NameAreaSettings = FDetailsViewArgs::HideNameArea;
+	DetailsViewArgs.bHideSelectionTip = true;
 	DetailsViewArgs.DefaultsOnlyVisibility = EEditDefaultsOnlyNodeVisibility::Automatic;
 	DetailsViewArgs.bShowOptions = false;
 	DetailsViewArgs.bAllowMultipleTopLevelObjects = true;
diff --git a/Engine/Source/Editor/CurveAssetEditor/Private/CurveAssetEditor.cpp b/Engine/Source/Editor/CurveAssetEditor/Private/CurveAssetEditor.cpp
--- a/Engine/Source/Editor/CurveAssetEditor/Private/CurveAssetEditor.cpp
+++ b/Engine/Source/Editor/CurveAssetEditor/Private/CurveAssetEditor.cpp
@@ -177,7 +177,10 @@ void FCurveAssetEditor::InitCurveAssetEditor( const EToolkitMode::Type Mode, con
 			);
 
 		FPropertyEditorModule& PropertyEditorModule = FModuleManager::GetModuleChecked<FPropertyEditorModule>("PropertyEditor");
-		const FDetailsViewArgs DetailsViewArgs(false, false, false, FDetailsViewArgs::HideNameArea);
+		FDetailsViewArgs DetailsViewArgs;
+		DetailsViewArgs.bAllowSearch = false;
+		DetailsViewArgs.NameAreaSettings = FDetailsViewArgs::HideNameArea;
+
 		ColorCurveDetailsView = PropertyEditorModule.CreateDetailView(DetailsViewArgs);
 	}
 	const bool bCreateDefaultStandaloneMenu = true;
diff --git a/Engine/Source/Editor/CurveTableEditor/Private/CompositeCurveTableEditor.cpp b/Engine/Source/Editor/CurveTableEditor/Private/CompositeCurveTableEditor.cpp
--- a/Engine/Source/Editor/CurveTableEditor/Private/CompositeCurveTableEditor.cpp
+++ b/Engine/Source/Editor/CurveTableEditor/Private/CompositeCurveTableEditor.cpp
@@ -36,7 +36,9 @@ void FCompositeCurveTableEditor::UnregisterTabSpawners(const TSharedRef<class FT
 void FCompositeCurveTableEditor::CreateAndRegisterPropertiesTab(const TSharedRef<class FTabManager>& InTabManager)
 {
 	FPropertyEditorModule& PropertyEditorModule = FModuleManager::GetModuleChecked<FPropertyEditorModule>("PropertyEditor");
-	const FDetailsViewArgs DetailsViewArgs(/*bIsUpdatable*/false, /*bIsLockable*/false, true, FDetailsViewArgs::ObjectsUseNameArea, false);
+	FDetailsViewArgs DetailsViewArgs;
+	DetailsViewArgs.NameAreaSettings = FDetailsViewArgs::ObjectsUseNameArea;
+
 	DetailsView = PropertyEditorModule.CreateDetailView(DetailsViewArgs);
 
 	InTabManager->RegisterTabSpawner(PropertiesTabId, FOnSpawnTab::CreateSP(this, &FCompositeCurveTableEditor::SpawnTab_Properties))
diff --git a/Engine/Source/Editor/DataTableEditor/Private/CompositeDataTableEditor.cpp b/Engine/Source/Editor/DataTableEditor/Private/CompositeDataTableEditor.cpp
--- a/Engine/Source/Editor/DataTableEditor/Private/CompositeDataTableEditor.cpp
+++ b/Engine/Source/Editor/DataTableEditor/Private/CompositeDataTableEditor.cpp
@@ -67,7 +67,8 @@ void FCompositeDataTableEditor::CreateAndRegisterRowEditorTab(const TSharedRef<c
 void FCompositeDataTableEditor::CreateAndRegisterPropertiesTab(const TSharedRef<class FTabManager>& InTabManager)
 {
 	FPropertyEditorModule& PropertyEditorModule = FModuleManager::GetModuleChecked<FPropertyEditorModule>("PropertyEditor");
-	const FDetailsViewArgs DetailsViewArgs(/*bIsUpdatable*/false, /*bIsLockable*/false, true, FDetailsViewArgs::ObjectsUseNameArea, false);
+	FDetailsViewArgs DetailsViewArgs;
+	DetailsViewArgs.NameAreaSettings = FDetailsViewArgs::ObjectsUseNameArea;
 	DetailsView = PropertyEditorModule.CreateDetailView(DetailsViewArgs);
 
 	InTabManager->RegisterTabSpawner(PropertiesTabId, FOnSpawnTab::CreateSP(this, &FCompositeDataTableEditor::SpawnTab_Properties))
diff --git a/Engine/Source/Editor/DataTableEditor/Private/DataTableEditor.cpp b/Engine/Source/Editor/DataTableEditor/Private/DataTableEditor.cpp
--- a/Engine/Source/Editor/DataTableEditor/Private/DataTableEditor.cpp
+++ b/Engine/Source/Editor/DataTableEditor/Private/DataTableEditor.cpp
@@ -110,7 +110,9 @@ void FDataTableEditor::CreateAndRegisterDataTableTab(const TSharedRef<class FTab
 void FDataTableEditor::CreateAndRegisterDataTableDetailsTab(const TSharedRef<class FTabManager>& InTabManager)
 {
 	FPropertyEditorModule & EditModule = FModuleManager::Get().GetModuleChecked<FPropertyEditorModule>("PropertyEditor");
-	FDetailsViewArgs DetailsViewArgs(/*bUpdateFromSelection=*/ false, /*bLockable=*/ false, /*bAllowSearch=*/ true, /*InNameAreaSettings=*/ FDetailsViewArgs::HideNameArea, /*bHideSelectionTip=*/ true);
+	FDetailsViewArgs DetailsViewArgs;
+	DetailsViewArgs.NameAreaSettings = FDetailsViewArgs::HideNameArea;
+	DetailsViewArgs.bHideSelectionTip = true;
 	PropertyView = EditModule.CreateDetailView(DetailsViewArgs);
 
 	InTabManager->RegisterTabSpawner(DataTableDetailsTabId, FOnSpawnTab::CreateSP(this, &FDataTableEditor::SpawnTab_DataTableDetails))
diff --git a/Engine/Source/Editor/EditorStyle/Private/EditorStyleSettingsCustomization.cpp b/Engine/Source/Editor/EditorStyle/Private/EditorStyleSettingsCustomization.cpp
--- a/Engine/Source/Editor/EditorStyle/Private/EditorStyleSettingsCustomization.cpp
+++ b/Engine/Source/Editor/EditorStyle/Private/EditorStyleSettingsCustomization.cpp
@@ -38,16 +38,11 @@ public:
 
 		FPropertyEditorModule& PropertyEditorModule = FModuleManager::Get().GetModuleChecked<FPropertyEditorModule>("PropertyEditor");
 
-		FDetailsViewArgs DetailsViewArgs(
-			/*bUpdateFromSelection=*/ false,
-			/*bLockable=*/ false,
-			/*bAllowSearch=*/ false,
-			FDetailsViewArgs::HideNameArea,
-			/*bHideSelectionTip=*/ true,
-			/*InNotifyHook=*/ nullptr,
-			/*InSearchInitialKeyFocus=*/ false,
-			/*InViewIdentifier=*/ NAME_None);
+		FDetailsViewArgs DetailsViewArgs;
+		DetailsViewArgs.bAllowSearch = false;
 		DetailsViewArgs.bShowOptions = false;
+		DetailsViewArgs.bHideSelectionTip = true;
+		DetailsViewArgs.NameAreaSettings = FDetailsViewArgs::HideNameArea;
 
 		TSharedRef<IDetailsView> DetailsView = PropertyEditorModule.CreateDetailView(DetailsViewArgs);
 
diff --git a/Engine/Source/Editor/FoliageEdit/Private/SFoliagePalette.cpp b/Engine/Source/Editor/FoliageEdit/Private/SFoliagePalette.cpp
--- a/Engine/Source/Editor/FoliageEdit/Private/SFoliagePalette.cpp
+++ b/Engine/Source/Editor/FoliageEdit/Private/SFoliagePalette.cpp
@@ -185,7 +185,10 @@ void SFoliagePalette::Construct(const FArguments& InArgs)
 		FoliageTypeTextFilter::FItemToStringArray::CreateSP(this, &SFoliagePalette::GetPaletteItemFilterString)));
 
 	FPropertyEditorModule& PropertyModule = FModuleManager::LoadModuleChecked<FPropertyEditorModule>("PropertyEditor");
-	FDetailsViewArgs Args(false, false, false, FDetailsViewArgs::HideNameArea, true);
+	FDetailsViewArgs Args;
+	Args.NameAreaSettings = FDetailsViewArgs::HideNameArea;
+	Args.bAllowSearch = false; 
+	Args.bHideSelectionTip = true;
 	Args.bShowActorLabel = false;
 	DetailsWidget = PropertyModule.CreateDetailView(Args);
 	DetailsWidget->SetVisibility(FoliageEditMode->UISettings.GetShowPaletteItemDetails() ? EVisibility::SelfHitTestInvisible : EVisibility::Collapsed);
diff --git a/Engine/Source/Editor/HierarchicalLODOutliner/Private/HLODOutliner.cpp b/Engine/Source/Editor/HierarchicalLODOutliner/Private/HLODOutliner.cpp
--- a/Engine/Source/Editor/HierarchicalLODOutliner/Private/HLODOutliner.cpp
+++ b/Engine/Source/Editor/HierarchicalLODOutliner/Private/HLODOutliner.cpp
@@ -556,16 +556,11 @@ namespace HLODOutliner
 		// Create a property view
 		FPropertyEditorModule& EditModule = FModuleManager::Get().GetModuleChecked<FPropertyEditorModule>("PropertyEditor");
 
-		FNotifyHook* NotifyHook = this;
-		FDetailsViewArgs DetailsViewArgs(
-			/*bUpdateFromSelection=*/ false,
-			/*bLockable=*/ false,
-			/*bAllowSearch=*/ false,
-			FDetailsViewArgs::HideNameArea,
-			/*bHideSelectionTip=*/ true,
-			/*InNotifyHook=*/ NotifyHook,
-			/*InSearchInitialKeyFocus=*/ false,
-			/*InViewIdentifier=*/ NAME_None);
+		FDetailsViewArgs DetailsViewArgs;
+		DetailsViewArgs.bAllowSearch = false;
+		DetailsViewArgs.NameAreaSettings = FDetailsViewArgs::HideNameArea;
+		DetailsViewArgs.bHideSelectionTip = true;
+		DetailsViewArgs.NotifyHook = this;
 		DetailsViewArgs.DefaultsOnlyVisibility = EEditDefaultsOnlyNodeVisibility::Automatic;
 		DetailsViewArgs.bShowOptions = false;
 
diff --git a/Engine/Source/Editor/Kismet/Private/SKismetInspector.cpp b/Engine/Source/Editor/Kismet/Private/SKismetInspector.cpp
--- a/Engine/Source/Editor/Kismet/Private/SKismetInspector.cpp
+++ b/Engine/Source/Editor/Kismet/Private/SKismetInspector.cpp
@@ -343,8 +343,11 @@ void SKismetInspector::Construct(const FArguments& InArgs)
 		NotifyHook = Kismet2.Get();
 	}
 
-	FDetailsViewArgs::ENameAreaSettings NameAreaSettings = InArgs._HideNameArea ? FDetailsViewArgs::HideNameArea : FDetailsViewArgs::ObjectsUseNameArea;
-	FDetailsViewArgs DetailsViewArgs( /*bUpdateFromSelection=*/ false, /*bLockable=*/ false, /*bAllowSearch=*/ true, NameAreaSettings, /*bHideSelectionTip=*/ true, /*InNotifyHook=*/ NotifyHook, /*InSearchInitialKeyFocus=*/ false, /*InViewIdentifier=*/ InArgs._ViewIdentifier );
+	FDetailsViewArgs DetailsViewArgs;
+	DetailsViewArgs.NameAreaSettings = InArgs._HideNameArea ? FDetailsViewArgs::HideNameArea : FDetailsViewArgs::ObjectsUseNameArea;
+	DetailsViewArgs.bHideSelectionTip = true;
+	DetailsViewArgs.NotifyHook = NotifyHook;
+	DetailsViewArgs.ViewIdentifier = InArgs._ViewIdentifier;
 
 	PropertyView = EditModule.CreateDetailView( DetailsViewArgs );
 		
diff --git a/Engine/Source/Editor/Kismet/Private/UserDefinedEnumEditor.cpp b/Engine/Source/Editor/Kismet/Private/UserDefinedEnumEditor.cpp
--- a/Engine/Source/Editor/Kismet/Private/UserDefinedEnumEditor.cpp
+++ b/Engine/Source/Editor/Kismet/Private/UserDefinedEnumEditor.cpp
@@ -298,7 +298,10 @@ TSharedRef<SDockTab> FUserDefinedEnumEditor::SpawnEnumeratorsTab(const FSpawnTab
 	// Create a property view
 	FPropertyEditorModule& EditModule = FModuleManager::Get().GetModuleChecked<FPropertyEditorModule>("PropertyEditor");
 
-	FDetailsViewArgs DetailsViewArgs( /*bUpdateFromSelection=*/ false, /*bLockable=*/ false, /*bAllowSearch=*/ false, FDetailsViewArgs::HideNameArea, /*bHideSelectionTip=*/ true );
+	FDetailsViewArgs DetailsViewArgs;
+	DetailsViewArgs.NameAreaSettings = FDetailsViewArgs::HideNameArea;
+	DetailsViewArgs.bAllowSearch = false;
+	DetailsViewArgs.bHideSelectionTip = true;
 	DetailsViewArgs.bShowOptions = false;
 
 	PropertyView = EditModule.CreateDetailView( DetailsViewArgs );
diff --git a/Engine/Source/Editor/Kismet/Private/UserDefinedStructureEditor.cpp b/Engine/Source/Editor/Kismet/Private/UserDefinedStructureEditor.cpp
--- a/Engine/Source/Editor/Kismet/Private/UserDefinedStructureEditor.cpp
+++ b/Engine/Source/Editor/Kismet/Private/UserDefinedStructureEditor.cpp
@@ -389,7 +389,10 @@ TSharedRef<SDockTab> FUserDefinedStructureEditor::SpawnStructureTab(const FSpawn
 	{
 		// Create a property view
 		FPropertyEditorModule& EditModule = FModuleManager::Get().GetModuleChecked<FPropertyEditorModule>("PropertyEditor");
-		FDetailsViewArgs DetailsViewArgs( /*bUpdateFromSelection=*/ false, /*bLockable=*/ false, /*bAllowSearch=*/ false, FDetailsViewArgs::HideNameArea, /*bHideSelectionTip=*/ true);
+		FDetailsViewArgs DetailsViewArgs;
+		DetailsViewArgs.bAllowSearch = false;
+		DetailsViewArgs.NameAreaSettings = FDetailsViewArgs::HideNameArea;
+		DetailsViewArgs.bHideSelectionTip = true;
 		DetailsViewArgs.bShowOptions = false;
 		PropertyView = EditModule.CreateDetailView(DetailsViewArgs);
 		FOnGetDetailCustomizationInstance LayoutStructDetails = FOnGetDetailCustomizationInstance::CreateStatic(&FUserDefinedStructureDetails::MakeInstance);
diff --git a/Engine/Source/Editor/KismetWidgets/Private/SSingleObjectDetailsPanel.cpp b/Engine/Source/Editor/KismetWidgets/Private/SSingleObjectDetailsPanel.cpp
--- a/Engine/Source/Editor/KismetWidgets/Private/SSingleObjectDetailsPanel.cpp
+++ b/Engine/Source/Editor/KismetWidgets/Private/SSingleObjectDetailsPanel.cpp
@@ -16,7 +16,10 @@ void SSingleObjectDetailsPanel::Construct(const FArguments& InArgs, bool bAutoma
 	// Create a property view
 	FPropertyEditorModule& EditModule = FModuleManager::Get().GetModuleChecked<FPropertyEditorModule>("PropertyEditor");
 
-	FDetailsViewArgs DetailsViewArgs(/*bUpdateFromSelection=*/ false, /*bLockable=*/ false, bAllowSearch, FDetailsViewArgs::HideNameArea, /*bHideSelectionTip=*/ true);
+	FDetailsViewArgs DetailsViewArgs;
+	DetailsViewArgs.bAllowSearch = bAllowSearch;
+	DetailsViewArgs.NameAreaSettings = FDetailsViewArgs::HideNameArea;
+	DetailsViewArgs.bHideSelectionTip = true;
 	DetailsViewArgs.HostCommandList = InArgs._HostCommandList;
 	DetailsViewArgs.HostTabManager = InArgs._HostTabManager;
 
diff --git a/Engine/Source/Editor/LandscapeEditor/Private/SLandscapeEditor.cpp b/Engine/Source/Editor/LandscapeEditor/Private/SLandscapeEditor.cpp
--- a/Engine/Source/Editor/LandscapeEditor/Private/SLandscapeEditor.cpp
+++ b/Engine/Source/Editor/LandscapeEditor/Private/SLandscapeEditor.cpp
@@ -443,7 +443,9 @@ void SLandscapeEditor::Construct(const FArguments& InArgs, TSharedRef<FLandscape
 	ParentToolkit = InParentToolkit;
 
 	FPropertyEditorModule& PropertyEditorModule = FModuleManager::LoadModuleChecked<FPropertyEditorModule>("PropertyEditor");
-	FDetailsViewArgs DetailsViewArgs(false, false, false,FDetailsViewArgs::HideNameArea);
+	FDetailsViewArgs DetailsViewArgs;
+	DetailsViewArgs.NameAreaSettings = FDetailsViewArgs::HideNameArea;
+	DetailsViewArgs.bAllowSearch = false;
 
 	DetailsPanel = PropertyEditorModule.CreateDetailView(DetailsViewArgs);
 	DetailsPanel->SetIsPropertyVisibleDelegate(FIsPropertyVisible::CreateSP(this, &SLandscapeEditor::GetIsPropertyVisible));
diff --git a/Engine/Source/Editor/LevelEditor/Private/SLevelEditor.cpp b/Engine/Source/Editor/LevelEditor/Private/SLevelEditor.cpp
--- a/Engine/Source/Editor/LevelEditor/Private/SLevelEditor.cpp
+++ b/Engine/Source/Editor/LevelEditor/Private/SLevelEditor.cpp
@@ -908,7 +908,9 @@ TSharedRef<SDockTab> SLevelEditor::SpawnLevelEditorTab( const FSpawnTabArgs& Arg
 	else if (TabIdentifier == LevelEditorTabIds::WorldSettings)
 	{
 		FPropertyEditorModule& PropPlugin = FModuleManager::LoadModuleChecked<FPropertyEditorModule>("PropertyEditor");
-		FDetailsViewArgs DetailsViewArgs( false, false, true, FDetailsViewArgs::HideNameArea, false, GUnrealEd );
+		FDetailsViewArgs DetailsViewArgs;
+		DetailsViewArgs.NameAreaSettings = FDetailsViewArgs::HideNameArea;
+		DetailsViewArgs.NotifyHook = GUnrealEd;
 		DetailsViewArgs.bShowActorLabel = false;
 
 		WorldSettingsView = PropPlugin.CreateDetailView( DetailsViewArgs );
diff --git a/Engine/Source/Editor/LocalizationDashboard/Private/LocalizationDashboard.cpp b/Engine/Source/Editor/LocalizationDashboard/Private/LocalizationDashboard.cpp
--- a/Engine/Source/Editor/LocalizationDashboard/Private/LocalizationDashboard.cpp
+++ b/Engine/Source/Editor/LocalizationDashboard/Private/LocalizationDashboard.cpp
@@ -75,7 +75,9 @@ void SLocalizationDashboard::Construct(const FArguments& InArgs, const TSharedPt
 
 		// Settings Details View
 		{
-			FDetailsViewArgs DetailsViewArgs(false, false, false, FDetailsViewArgs::ENameAreaSettings::HideNameArea, false, nullptr, false, NAME_None);
+			FDetailsViewArgs DetailsViewArgs;
+			DetailsViewArgs.bAllowSearch = false; 
+			DetailsViewArgs.NameAreaSettings = FDetailsViewArgs::HideNameArea;
 			TSharedRef<IDetailsView> DetailsView = PropertyModule.CreateDetailView(DetailsViewArgs);
 			DetailsView->SetIsPropertyEditingEnabledDelegate(FIsPropertyEditingEnabled::CreateSP(this, &SLocalizationDashboard::CanMakeEdits));
 			DetailsView->SetObject(GetMutableDefault<ULocalizationDashboardSettings>(), true);
@@ -89,7 +91,9 @@ void SLocalizationDashboard::Construct(const FArguments& InArgs, const TSharedPt
 
 		// Game Targets Details View
 		{
-			FDetailsViewArgs DetailsViewArgs(false, false, false, FDetailsViewArgs::ENameAreaSettings::HideNameArea, false, nullptr, false, NAME_None);
+			FDetailsViewArgs DetailsViewArgs;
+			DetailsViewArgs.bAllowSearch = false; 
+			DetailsViewArgs.NameAreaSettings = FDetailsViewArgs::HideNameArea;
 			TSharedRef<IDetailsView> DetailsView = PropertyModule.CreateDetailView(DetailsViewArgs);
 			DetailsView->SetIsPropertyEditingEnabledDelegate(FIsPropertyEditingEnabled::CreateSP(this, &SLocalizationDashboard::CanMakeEdits));
 			DetailsView->SetObject(ULocalizationSettings::GetGameTargetSet(), true);
@@ -103,7 +107,9 @@ void SLocalizationDashboard::Construct(const FArguments& InArgs, const TSharedPt
 
 		// Engine Targets Details View
 		{
-			FDetailsViewArgs DetailsViewArgs(false, false, false, FDetailsViewArgs::ENameAreaSettings::HideNameArea, false, nullptr, false, NAME_None);
+			FDetailsViewArgs DetailsViewArgs;
+			DetailsViewArgs.bAllowSearch = false; 
+			DetailsViewArgs.NameAreaSettings = FDetailsViewArgs::HideNameArea;
 			TSharedRef<IDetailsView> DetailsView = PropertyModule.CreateDetailView(DetailsViewArgs);
 			DetailsView->SetIsPropertyEditingEnabledDelegate(FIsPropertyEditingEnabled::CreateSP(this, &SLocalizationDashboard::CanMakeEdits));
 			DetailsView->SetObject(ULocalizationSettings::GetEngineTargetSet(), true);
diff --git a/Engine/Source/Editor/LocalizationDashboard/Private/SLocalizationTargetEditor.cpp b/Engine/Source/Editor/LocalizationDashboard/Private/SLocalizationTargetEditor.cpp
--- a/Engine/Source/Editor/LocalizationDashboard/Private/SLocalizationTargetEditor.cpp
+++ b/Engine/Source/Editor/LocalizationDashboard/Private/SLocalizationTargetEditor.cpp
@@ -13,7 +13,10 @@ void SLocalizationTargetEditor::Construct(const FArguments& InArgs, ULocalizatio
 	LocalizationTarget = InLocalizationTarget;
 
 	FPropertyEditorModule& PropertyModule = FModuleManager::LoadModuleChecked<FPropertyEditorModule>("PropertyEditor");
-	FDetailsViewArgs DetailsViewArgs(false, false, false, FDetailsViewArgs::ENameAreaSettings::HideNameArea, false, nullptr, false, NAME_None);
+	FDetailsViewArgs DetailsViewArgs;
+	DetailsViewArgs.bAllowSearch = false;
+	DetailsViewArgs.NameAreaSettings = FDetailsViewArgs::HideNameArea;
+
 	TSharedRef<IDetailsView> DetailsView = PropertyModule.CreateDetailView(DetailsViewArgs);
 	DetailsView->SetObject(InLocalizationTarget, true);
 	DetailsView->SetIsPropertyEditingEnabledDelegate(IsPropertyEditingEnabled);
diff --git a/Engine/Source/Editor/MaterialEditor/Private/MaterialEditor.cpp b/Engine/Source/Editor/MaterialEditor/Private/MaterialEditor.cpp
--- a/Engine/Source/Editor/MaterialEditor/Private/MaterialEditor.cpp
+++ b/Engine/Source/Editor/MaterialEditor/Private/MaterialEditor.cpp
@@ -1365,7 +1365,10 @@ void FMaterialEditor::CreateInternalWidgets()
 
 	FPropertyEditorModule& PropertyEditorModule = FModuleManager::GetModuleChecked<FPropertyEditorModule>( "PropertyEditor" );
 
-	const FDetailsViewArgs DetailsViewArgs( false, false, true, FDetailsViewArgs::HideNameArea, true, this );
+	FDetailsViewArgs DetailsViewArgs;
+	DetailsViewArgs.NameAreaSettings = FDetailsViewArgs::HideNameArea;
+	DetailsViewArgs.bHideSelectionTip = true;
+	DetailsViewArgs.NotifyHook = this;
 	MaterialDetailsView = PropertyEditorModule.CreateDetailView( DetailsViewArgs );
 
 	FOnGetDetailCustomizationInstance LayoutExpressionParameterDetails = FOnGetDetailCustomizationInstance::CreateStatic(
diff --git a/Engine/Source/Editor/MaterialEditor/Private/MaterialInstanceEditor.cpp b/Engine/Source/Editor/MaterialEditor/Private/MaterialInstanceEditor.cpp
--- a/Engine/Source/Editor/MaterialEditor/Private/MaterialInstanceEditor.cpp
+++ b/Engine/Source/Editor/MaterialEditor/Private/MaterialInstanceEditor.cpp
@@ -757,7 +757,10 @@ void FMaterialInstanceEditor::CreateInternalWidgets()
 	PreviewUIViewport = SNew(SMaterialEditorUIPreviewViewport, GetMaterialInterface());
 
 	FPropertyEditorModule& PropertyEditorModule = FModuleManager::GetModuleChecked<FPropertyEditorModule>( "PropertyEditor" );
-	FDetailsViewArgs DetailsViewArgs( false, false, true, FDetailsViewArgs::HideNameArea, true, this );
+	FDetailsViewArgs DetailsViewArgs;
+	DetailsViewArgs.NameAreaSettings = FDetailsViewArgs::HideNameArea;
+	DetailsViewArgs.bHideSelectionTip = true;
+	DetailsViewArgs.NotifyHook = this;
 	DetailsViewArgs.bShowModifiedPropertiesOption = false;
 	DetailsViewArgs.bShowCustomFilterOption = true;
 	MaterialInstanceDetails = PropertyEditorModule.CreateDetailView( DetailsViewArgs );
diff --git a/Engine/Source/Editor/MovieSceneTools/Private/Sections/ThumbnailSection.cpp b/Engine/Source/Editor/MovieSceneTools/Private/Sections/ThumbnailSection.cpp
--- a/Engine/Source/Editor/MovieSceneTools/Private/Sections/ThumbnailSection.cpp
+++ b/Engine/Source/Editor/MovieSceneTools/Private/Sections/ThumbnailSection.cpp
@@ -157,7 +157,10 @@ void FThumbnailSection::BuildSectionContextMenu(FMenuBuilder& MenuBuilder, const
 
 					FPropertyEditorModule& PropertyModule = FModuleManager::LoadModuleChecked<FPropertyEditorModule>("PropertyEditor");
 
-					FDetailsViewArgs Args(false, false, false, FDetailsViewArgs::HideNameArea);
+					FDetailsViewArgs Args;
+					Args.bAllowSearch = false;
+					Args.NameAreaSettings = FDetailsViewArgs::HideNameArea;
+
 					TSharedRef<IDetailsView> DetailView = PropertyModule.CreateDetailView(Args);
 					DetailView->SetObject(GetMutableDefault<UMovieSceneUserThumbnailSettings>());
 					InMenuBuilder.AddWidget(DetailView, FText(), true);
diff --git a/Engine/Source/Editor/MovieSceneTools/Private/TrackEditors/EventTrackEditor.cpp b/Engine/Source/Editor/MovieSceneTools/Private/TrackEditors/EventTrackEditor.cpp
--- a/Engine/Source/Editor/MovieSceneTools/Private/TrackEditors/EventTrackEditor.cpp
+++ b/Engine/Source/Editor/MovieSceneTools/Private/TrackEditors/EventTrackEditor.cpp
@@ -198,7 +198,10 @@ void FEventTrackEditor::BuildTrackContextMenu(FMenuBuilder& MenuBuilder, UMovieS
 		FPropertyEditorModule& PropertyEditor = FModuleManager::Get().LoadModuleChecked<FPropertyEditorModule>("PropertyEditor");
 
 		// Create a details view for the track
-		FDetailsViewArgs DetailsViewArgs(false,false,false,FDetailsViewArgs::HideNameArea,true);
+		FDetailsViewArgs DetailsViewArgs;
+		DetailsViewArgs.bAllowSearch = false;
+		DetailsViewArgs.NameAreaSettings = FDetailsViewArgs::HideNameArea;
+		DetailsViewArgs.bHideSelectionTip = true;
 		DetailsViewArgs.DefaultsOnlyVisibility = EEditDefaultsOnlyNodeVisibility::Automatic;
 		DetailsViewArgs.bShowOptions = false;
 		DetailsViewArgs.ColumnWidth = 0.55f;
diff --git a/Engine/Source/Editor/Persona/Private/SPersonaDetails.cpp b/Engine/Source/Editor/Persona/Private/SPersonaDetails.cpp
--- a/Engine/Source/Editor/Persona/Private/SPersonaDetails.cpp
+++ b/Engine/Source/Editor/Persona/Private/SPersonaDetails.cpp
@@ -6,7 +6,9 @@
 
 void SPersonaDetails::Construct(const FArguments& InArgs)
 {
-	FDetailsViewArgs DetailsViewArgs(false, false, true, FDetailsViewArgs::HideNameArea, true);
+	FDetailsViewArgs DetailsViewArgs;
+	DetailsViewArgs.NameAreaSettings = FDetailsViewArgs::HideNameArea;
+	DetailsViewArgs.bHideSelectionTip = true;
 
 	FPropertyEditorModule& PropertyEditorModule = FModuleManager::LoadModuleChecked<FPropertyEditorModule>("PropertyEditor");
 	DetailsView = PropertyEditorModule.CreateDetailView(DetailsViewArgs);
diff --git a/Engine/Source/Editor/PhysicsAssetEditor/Private/PhysicsAssetEditorProfilesSummoner.cpp b/Engine/Source/Editor/PhysicsAssetEditor/Private/PhysicsAssetEditorProfilesSummoner.cpp
--- a/Engine/Source/Editor/PhysicsAssetEditor/Private/PhysicsAssetEditorProfilesSummoner.cpp
+++ b/Engine/Source/Editor/PhysicsAssetEditor/Private/PhysicsAssetEditorProfilesSummoner.cpp
@@ -30,7 +30,9 @@ TSharedPtr<SToolTip> FPhysicsAssetEditorProfilesSummoner::CreateTabToolTipWidget
 
 TSharedRef<SWidget> FPhysicsAssetEditorProfilesSummoner::CreateTabBody(const FWorkflowTabSpawnInfo& Info) const
 {
-	FDetailsViewArgs DetailsViewArgs(false, false, true, FDetailsViewArgs::HideNameArea, true);
+	FDetailsViewArgs DetailsViewArgs;
+	DetailsViewArgs.NameAreaSettings = FDetailsViewArgs::HideNameArea;
+	DetailsViewArgs.bHideSelectionTip = true;
 	DetailsViewArgs.bAllowSearch = false;
 
 	FPropertyEditorModule& PropertyEditorModule = FModuleManager::LoadModuleChecked<FPropertyEditorModule>("PropertyEditor");
diff --git a/Engine/Source/Editor/PhysicsAssetEditor/Private/PhysicsAssetEditorSharedData.cpp b/Engine/Source/Editor/PhysicsAssetEditor/Private/PhysicsAssetEditorSharedData.cpp
--- a/Engine/Source/Editor/PhysicsAssetEditor/Private/PhysicsAssetEditorSharedData.cpp
+++ b/Engine/Source/Editor/PhysicsAssetEditor/Private/PhysicsAssetEditorSharedData.cpp
@@ -2012,7 +2012,9 @@ void FPhysicsAssetEditorSharedData::OpenNewBodyDlg(EAppReturnType::Type* NewBody
 
 TSharedRef<SWidget> FPhysicsAssetEditorSharedData::CreateGenerateBodiesWidget(const FSimpleDelegate& InOnCreate, const FSimpleDelegate& InOnCancel, const TAttribute<bool>& InIsEnabled, const TAttribute<FText>& InCreateButtonText, bool bForNewAsset)
 {
-	FDetailsViewArgs DetailsViewArgs(false, false, true, FDetailsViewArgs::HideNameArea, true);
+	FDetailsViewArgs DetailsViewArgs;
+	DetailsViewArgs.NameAreaSettings = FDetailsViewArgs::HideNameArea;
+	DetailsViewArgs.bHideSelectionTip = true;
 	DetailsViewArgs.bAllowSearch = false;
 
 	FPropertyEditorModule& PropertyEditorModule = FModuleManager::LoadModuleChecked<FPropertyEditorModule>("PropertyEditor");
diff --git a/Engine/Source/Editor/PropertyEditor/Private/PropertyEditorToolkit.cpp b/Engine/Source/Editor/PropertyEditor/Private/PropertyEditorToolkit.cpp
--- a/Engine/Source/Editor/PropertyEditor/Private/PropertyEditorToolkit.cpp
+++ b/Engine/Source/Editor/PropertyEditor/Private/PropertyEditorToolkit.cpp
@@ -312,7 +312,8 @@ void FPropertyEditorToolkit::CreatePropertyTable()
 void FPropertyEditorToolkit::CreateDetailsPanel()
 {
 	FPropertyEditorModule& PropertyEditorModule = FModuleManager::GetModuleChecked<FPropertyEditorModule>("PropertyEditor");
-	const FDetailsViewArgs DetailsViewArgs(false, false, true, FDetailsViewArgs::HideNameArea);
+	FDetailsViewArgs DetailsViewArgs;
+	DetailsViewArgs.NameAreaSettings = FDetailsViewArgs::HideNameArea;
 	DetailsView = PropertyEditorModule.CreateDetailView(DetailsViewArgs);
 }
 
diff --git a/Engine/Source/Editor/PropertyEditor/Public/IDetailsView.h b/Engine/Source/Editor/PropertyEditor/Public/IDetailsView.h
--- a/Engine/Source/Editor/PropertyEditor/Public/IDetailsView.h
+++ b/Engine/Source/Editor/PropertyEditor/Public/IDetailsView.h
@@ -102,26 +102,19 @@ struct FDetailsViewArgs
 	float ColumnWidth;
 
 public:
-	/** Default constructor */
-	FDetailsViewArgs( const bool InUpdateFromSelection = false
-					, const bool InLockable = false
-					, const bool InAllowSearch = true
-					, const ENameAreaSettings InNameAreaSettings = ActorsUseNameArea
-					, const bool InHideSelectionTip = false
-					, FNotifyHook* InNotifyHook = NULL
-					, const bool InSearchInitialKeyFocus = false
-					, FName InViewIdentifier = NAME_None )
+
+	FDetailsViewArgs()
 		: DefaultsOnlyVisibility(EEditDefaultsOnlyNodeVisibility::Show)
-		, ViewIdentifier( InViewIdentifier )
-		, NotifyHook( InNotifyHook )
-		, NameAreaSettings( InNameAreaSettings )
-		, bUpdatesFromSelection( InUpdateFromSelection )
-		, bLockable(InLockable)
-		, bAllowSearch( InAllowSearch )
-		, bHideSelectionTip( InHideSelectionTip )
-		, bSearchInitialKeyFocus( InSearchInitialKeyFocus )
+		, ViewIdentifier(NAME_None)
+		, NotifyHook(nullptr)
+		, NameAreaSettings(ActorsUseNameArea)
+		, bUpdatesFromSelection(false)
+		, bLockable(false)
+		, bAllowSearch(true)
+		, bHideSelectionTip(false)
+		, bSearchInitialKeyFocus(false)
 		, bShowPropertyMatrixButton(true)
-		, bShowOptions( true )
+		, bShowOptions(true)
 		, bShowModifiedPropertiesOption(true)
 		, bShowActorLabel(true)
 		, bShowDifferingPropertiesOption(false)
@@ -137,6 +130,28 @@ public:
 		, ColumnWidth(.65f)
 	{
 	}
+
+	/** Default constructor */
+	UE_DEPRECATED(5.0, "This constructor is deprecated, please create an empty FDetailsViewArgs and explicitly set any values that you wish to change.")
+	FDetailsViewArgs( const bool InUpdateFromSelection
+					, const bool InLockable = false
+					, const bool InAllowSearch = true
+					, const ENameAreaSettings InNameAreaSettings = ActorsUseNameArea
+					, const bool InHideSelectionTip = false
+					, FNotifyHook* InNotifyHook = NULL
+					, const bool InSearchInitialKeyFocus = false
+					, FName InViewIdentifier = NAME_None )
+		: FDetailsViewArgs()
+	{
+		bUpdatesFromSelection = InUpdateFromSelection;
+		bLockable = InLockable;
+		bAllowSearch = InAllowSearch;
+		NameAreaSettings = InNameAreaSettings;
+		bHideSelectionTip = InHideSelectionTip;
+		NotifyHook = InNotifyHook;
+		bSearchInitialKeyFocus = InSearchInitialKeyFocus;
+		ViewIdentifier = InViewIdentifier;
+	}
 };
 
 
diff --git a/Engine/Source/Editor/SequenceRecorder/Private/SClassPropertyRecorderSettings.cpp b/Engine/Source/Editor/SequenceRecorder/Private/SClassPropertyRecorderSettings.cpp
--- a/Engine/Source/Editor/SequenceRecorder/Private/SClassPropertyRecorderSettings.cpp
+++ b/Engine/Source/Editor/SequenceRecorder/Private/SClassPropertyRecorderSettings.cpp
@@ -72,7 +72,9 @@ FReply SClassPropertyRecorderSettings::HandleChoosePropertiesButtonClicked()
 {
 	if (ClassHandle.IsValid() && PropertiesHandle.IsValid())
 	{
-		FDetailsViewArgs DetailsViewArgs(false, false, true, FDetailsViewArgs::HideNameArea, true);
+		FDetailsViewArgs DetailsViewArgs;
+		DetailsViewArgs.NameAreaSettings = FDetailsViewArgs::HideNameArea;
+		DetailsViewArgs.bHideSelectionTip = true;
 		FPropertyEditorModule& PropertyEditorModule = FModuleManager::GetModuleChecked<FPropertyEditorModule>("PropertyEditor");
 		TSharedRef<IDetailsView> DetailsView = PropertyEditorModule.CreateDetailView(DetailsViewArgs);
 
diff --git a/Engine/Source/Editor/Sequencer/Private/SequencerContextMenus.cpp b/Engine/Source/Editor/Sequencer/Private/SequencerContextMenus.cpp
--- a/Engine/Source/Editor/Sequencer/Private/SequencerContextMenus.cpp
+++ b/Engine/Source/Editor/Sequencer/Private/SequencerContextMenus.cpp
@@ -2155,15 +2155,10 @@ void FEasingContextMenu::EasingOptionsMenu(FMenuBuilder& MenuBuilder)
 {
 	FPropertyEditorModule& EditModule = FModuleManager::Get().GetModuleChecked<FPropertyEditorModule>("PropertyEditor");
 
-	FDetailsViewArgs DetailsViewArgs(
-		/*bUpdateFromSelection=*/ false,
-		/*bLockable=*/ false,
-		/*bAllowSearch=*/ false,
-		FDetailsViewArgs::HideNameArea,
-		/*bHideSelectionTip=*/ true,
-		/*InNotifyHook=*/ nullptr,
-		/*InSearchInitialKeyFocus=*/ false,
-		/*InViewIdentifier=*/ NAME_None);
+	FDetailsViewArgs DetailsViewArgs;
+	DetailsViewArgs.bAllowSearch = false;
+	DetailsViewArgs.NameAreaSettings = FDetailsViewArgs::HideNameArea;
+	DetailsViewArgs.bHideSelectionTip = true;
 	DetailsViewArgs.bShowOptions = false;
 	DetailsViewArgs.bShowScrollBar = false;
 
diff --git a/Engine/Source/Editor/SkeletonEditor/Private/SkeletonTreeSocketItem.cpp b/Engine/Source/Editor/SkeletonEditor/Private/SkeletonTreeSocketItem.cpp
--- a/Engine/Source/Editor/SkeletonEditor/Private/SkeletonTreeSocketItem.cpp
+++ b/Engine/Source/Editor/SkeletonEditor/Private/SkeletonTreeSocketItem.cpp
@@ -79,7 +79,10 @@ TSharedRef< SWidget > FSkeletonTreeSocketItem::GenerateInlineEditWidget(const TA
 	if (GetDefault<UPersonaOptions>()->bUseInlineSocketEditor)
 	{
 		FPropertyEditorModule& EditModule = FModuleManager::Get().GetModuleChecked<FPropertyEditorModule>("PropertyEditor");
-		FDetailsViewArgs DetailsViewArgs(/*bUpdateFromSelection=*/ false, /*bLockable=*/ false, /*bAllowSearch=*/ false, /*InNameAreaSettings=*/ FDetailsViewArgs::HideNameArea, /*bHideSelectionTip=*/ true);
+		FDetailsViewArgs DetailsViewArgs;
+		DetailsViewArgs.bAllowSearch = false;
+		DetailsViewArgs.NameAreaSettings = FDetailsViewArgs::HideNameArea;
+		DetailsViewArgs.bHideSelectionTip = true;
 		DetailsViewArgs.bAllowFavoriteSystem = false;
 		DetailsViewArgs.bShowScrollBar = false;
 		TSharedRef<IDetailsView> DetailsView = EditModule.CreateDetailView(DetailsViewArgs);
diff --git a/Engine/Source/Editor/UMGEditor/Private/Details/SWidgetDetailsView.cpp b/Engine/Source/Editor/UMGEditor/Private/Details/SWidgetDetailsView.cpp
--- a/Engine/Source/Editor/UMGEditor/Private/Details/SWidgetDetailsView.cpp
+++ b/Engine/Source/Editor/UMGEditor/Private/Details/SWidgetDetailsView.cpp
@@ -41,16 +41,10 @@ void SWidgetDetailsView::Construct(const FArguments& InArgs, TSharedPtr<FWidgetB
 	// Create a property view
 	FPropertyEditorModule& EditModule = FModuleManager::Get().GetModuleChecked<FPropertyEditorModule>("PropertyEditor");
 
-	FNotifyHook* NotifyHook = this;
-	FDetailsViewArgs DetailsViewArgs(
-		/*bUpdateFromSelection=*/ false,
-		/*bLockable=*/ false,
-		/*bAllowSearch=*/ true,
-		FDetailsViewArgs::HideNameArea,
-		/*bHideSelectionTip=*/ true,
-		/*InNotifyHook=*/ NotifyHook,
-		/*InSearchInitialKeyFocus=*/ false,
-		/*InViewIdentifier=*/ NAME_None);
+	FDetailsViewArgs DetailsViewArgs;
+	DetailsViewArgs.NameAreaSettings = FDetailsViewArgs::HideNameArea;
+	DetailsViewArgs.bHideSelectionTip = true;
+	DetailsViewArgs.NotifyHook = this;
 	DetailsViewArgs.DefaultsOnlyVisibility = EEditDefaultsOnlyNodeVisibility::Automatic;
 
 	PropertyView = EditModule.CreateDetailView(DetailsViewArgs);
diff --git a/Engine/Source/Editor/UnrealEd/Private/SCSVImportOptions.cpp b/Engine/Source/Editor/UnrealEd/Private/SCSVImportOptions.cpp
--- a/Engine/Source/Editor/UnrealEd/Private/SCSVImportOptions.cpp
+++ b/Engine/Source/Editor/UnrealEd/Private/SCSVImportOptions.cpp
@@ -31,7 +31,10 @@ void SCSVImportOptions::Construct(const FArguments& InArgs)
 
 	// Create properties view
 	FPropertyEditorModule & EditModule = FModuleManager::Get().GetModuleChecked<FPropertyEditorModule>("PropertyEditor");
-	FDetailsViewArgs DetailsViewArgs(/*bUpdateFromSelection=*/ false, /*bLockable=*/ false, /*bAllowSearch=*/ false, /*InNameAreaSettings=*/ FDetailsViewArgs::HideNameArea, /*bHideSelectionTip=*/ true);
+	FDetailsViewArgs DetailsViewArgs;
+	DetailsViewArgs.bAllowSearch = false;
+	DetailsViewArgs.NameAreaSettings = FDetailsViewArgs::HideNameArea;
+	DetailsViewArgs.bHideSelectionTip = true;
 	PropertyView = EditModule.CreateDetailView(DetailsViewArgs);
 
 	PropertyView->SetIsPropertyVisibleDelegate(FIsPropertyVisible::CreateLambda([](const FPropertyAndParent& InPropertyAndParent)
diff --git a/Engine/Source/Editor/UnrealEd/Private/Toolkits/BaseToolkit.cpp b/Engine/Source/Editor/UnrealEd/Private/Toolkits/BaseToolkit.cpp
--- a/Engine/Source/Editor/UnrealEd/Private/Toolkits/BaseToolkit.cpp
+++ b/Engine/Source/Editor/UnrealEd/Private/Toolkits/BaseToolkit.cpp
@@ -144,36 +144,29 @@ void FModeToolkit::Init(const TSharedPtr< class IToolkitHost >& InitToolkitHost,
 
 	FPropertyEditorModule& PropertyEditorModule = FModuleManager::GetModuleChecked<FPropertyEditorModule>("PropertyEditor");
 
-	FDetailsViewArgs ModeDetailsViewArgs(
-		/*bUpdateFromSelection=*/ false,
-		/*bLockable=*/ false,
-		/*bAllowSearch=*/ false,
-		FDetailsViewArgs::HideNameArea,
-		/*bHideSelectionTip=*/ true,
-		/*InNotifyHook=*/ nullptr,
-		/*InSearchInitialKeyFocus=*/ false,
-		/*InViewIdentifier=*/ NAME_None);
-	ModeDetailsViewArgs.DefaultsOnlyVisibility = EEditDefaultsOnlyNodeVisibility::Automatic;
-	ModeDetailsViewArgs.bShowOptions = false;
-	ModeDetailsViewArgs.bAllowMultipleTopLevelObjects = true;
-
-	ModeDetailsView = PropertyEditorModule.CreateDetailView(ModeDetailsViewArgs);
-
-	FDetailsViewArgs DetailsViewArgs(
-		/*bUpdateFromSelection=*/ false,
-		/*bLockable=*/ false,
-		/*bAllowSearch=*/ false,
-		FDetailsViewArgs::HideNameArea,
-		/*bHideSelectionTip=*/ true,
-		/*InNotifyHook=*/ nullptr,
-		/*InSearchInitialKeyFocus=*/ false,
-		/*InViewIdentifier=*/ NAME_None);
-	DetailsViewArgs.DefaultsOnlyVisibility = EEditDefaultsOnlyNodeVisibility::Automatic;
-	DetailsViewArgs.bShowOptions = false;
-	DetailsViewArgs.bAllowMultipleTopLevelObjects = true;
-
-	DetailsView = PropertyEditorModule.CreateDetailView(DetailsViewArgs);
+	{
+		FDetailsViewArgs ModeDetailsViewArgs;
+		ModeDetailsViewArgs.bAllowSearch = false;
+		ModeDetailsViewArgs.NameAreaSettings = FDetailsViewArgs::HideNameArea;
+		ModeDetailsViewArgs.bHideSelectionTip = true;
+		ModeDetailsViewArgs.DefaultsOnlyVisibility = EEditDefaultsOnlyNodeVisibility::Automatic;
+		ModeDetailsViewArgs.bShowOptions = false;
+		ModeDetailsViewArgs.bAllowMultipleTopLevelObjects = true;
+
+		ModeDetailsView = PropertyEditorModule.CreateDetailView(ModeDetailsViewArgs);
+	}
 
+	{
+		FDetailsViewArgs DetailsViewArgs;
+		DetailsViewArgs.bAllowSearch = false;
+		DetailsViewArgs.NameAreaSettings = FDetailsViewArgs::HideNameArea;
+		DetailsViewArgs.bHideSelectionTip = true;
+		DetailsViewArgs.DefaultsOnlyVisibility = EEditDefaultsOnlyNodeVisibility::Automatic;
+		DetailsViewArgs.bShowOptions = false;
+		DetailsViewArgs.bAllowMultipleTopLevelObjects = true;
+
+		DetailsView = PropertyEditorModule.CreateDetailView(DetailsViewArgs);
+	}
 
 	FToolkitManager::Get().RegisterNewToolkit(SharedThis(this));
 }
diff --git a/Engine/Source/Editor/UnrealEd/Private/Toolkits/SimpleAssetEditor.cpp b/Engine/Source/Editor/UnrealEd/Private/Toolkits/SimpleAssetEditor.cpp
--- a/Engine/Source/Editor/UnrealEd/Private/Toolkits/SimpleAssetEditor.cpp
+++ b/Engine/Source/Editor/UnrealEd/Private/Toolkits/SimpleAssetEditor.cpp
@@ -48,16 +48,13 @@ FSimpleAssetEditor::~FSimpleAssetEditor()
 
 void FSimpleAssetEditor::InitEditor( const EToolkitMode::Type Mode, const TSharedPtr< class IToolkitHost >& InitToolkitHost, const TArray<UObject*>& ObjectsToEdit, FGetDetailsViewObjects GetDetailsViewObjects )
 {
-	const bool bIsUpdatable = false;
-	const bool bAllowFavorites = true;
-	const bool bIsLockable = false;
-
 	EditingObjects = ObjectsToEdit;
 	GEditor->GetEditorSubsystem<UImportSubsystem>()->OnAssetPostImport.AddSP(this, &FSimpleAssetEditor::HandleAssetPostImport);
 	GEditor->OnObjectsReplaced().AddSP(this, &FSimpleAssetEditor::OnObjectsReplaced);
 
 	FPropertyEditorModule& PropertyEditorModule = FModuleManager::GetModuleChecked<FPropertyEditorModule>( "PropertyEditor" );
-	const FDetailsViewArgs DetailsViewArgs( bIsUpdatable, bIsLockable, true, FDetailsViewArgs::ObjectsUseNameArea, false );
+	FDetailsViewArgs DetailsViewArgs;
+	DetailsViewArgs.NameAreaSettings = FDetailsViewArgs::ObjectsUseNameArea;
 	DetailsView = PropertyEditorModule.CreateDetailView( DetailsViewArgs );
 	const TSharedRef<FTabManager::FLayout> StandaloneDefaultLayout = FTabManager::NewLayout( "Standalone_SimpleAssetEditor_Layout_v3" )
 	->AddArea
diff --git a/Engine/Source/Editor/UnrealEd/Private/Tools/BaseAssetToolkit.cpp b/Engine/Source/Editor/UnrealEd/Private/Tools/BaseAssetToolkit.cpp
--- a/Engine/Source/Editor/UnrealEd/Private/Tools/BaseAssetToolkit.cpp
+++ b/Engine/Source/Editor/UnrealEd/Private/Tools/BaseAssetToolkit.cpp
@@ -149,7 +149,9 @@ void FBaseAssetToolkit::CreateWidgets()
 	LayoutExtender = MakeShared<FLayoutExtender>();
 
 	FPropertyEditorModule& PropertyEditorModule = FModuleManager::GetModuleChecked<FPropertyEditorModule>("PropertyEditor");
-	const FDetailsViewArgs DetailsViewArgs(false, false, true, FDetailsViewArgs::ObjectsUseNameArea, true);
+	FDetailsViewArgs DetailsViewArgs;
+	DetailsViewArgs.NameAreaSettings = FDetailsViewArgs::ObjectsUseNameArea;
+	DetailsViewArgs.bHideSelectionTip = true;
 	DetailsView = PropertyEditorModule.CreateDetailView(DetailsViewArgs);
 }
 
diff --git a/Engine/Source/Editor/WorldBrowser/Private/SWorldDetails.cpp b/Engine/Source/Editor/WorldBrowser/Private/SWorldDetails.cpp
--- a/Engine/Source/Editor/WorldBrowser/Private/SWorldDetails.cpp
+++ b/Engine/Source/Editor/WorldBrowser/Private/SWorldDetails.cpp
@@ -66,7 +66,10 @@ void SWorldDetails::OnBrowseWorld(UWorld* InWorld)
 		WorldModel->CollectionChanged.AddSP(this, &SWorldDetails::OnCollectionChanged);
 	
 		FPropertyEditorModule& PropertyModule = FModuleManager::LoadModuleChecked<FPropertyEditorModule>("PropertyEditor");
-		FDetailsViewArgs Args(false, false, false, FDetailsViewArgs::HideNameArea, true);
+		FDetailsViewArgs Args;
+		Args.bAllowSearch = false;
+		Args.NameAreaSettings = FDetailsViewArgs::HideNameArea;
+		Args.bHideSelectionTip = true;
 		Args.bShowActorLabel = false;
 	
 		DetailsView = PropertyModule.CreateDetailView(Args);
