Fixed more deprecation warnings.

diff --git a/Engine/Plugins/Experimental/VirtualHeightfieldMesh/Source/VirtualHeightfieldMesh/Private/VirtualHeightfieldMeshSceneProxy.cpp b/Engine/Plugins/Experimental/VirtualHeightfieldMesh/Source/VirtualHeightfieldMesh/Private/VirtualHeightfieldMeshSceneProxy.cpp
--- a/Engine/Plugins/Experimental/VirtualHeightfieldMesh/Source/VirtualHeightfieldMesh/Private/VirtualHeightfieldMeshSceneProxy.cpp
+++ b/Engine/Plugins/Experimental/VirtualHeightfieldMesh/Source/VirtualHeightfieldMesh/Private/VirtualHeightfieldMeshSceneProxy.cpp
@@ -1070,7 +1070,7 @@ namespace VirtualHeightfieldMesh
 		GraphBuilder.Execute();
 
 		// The SRV objects referenced by final rendering are managed outside of RDG.
-		InBuffers.InstanceBufferSRV = RHICreateShaderResourceView(InBuffers.InstanceBuffer->GetStructuredBufferRHI());
+		InBuffers.InstanceBufferSRV = RHICreateShaderResourceView(InBuffers.InstanceBuffer->GetRHI());
 	}
 
 	/** Initialize the volatie resources used in the render graph. */
diff --git a/Engine/Plugins/Runtime/HairStrands/Source/HairStrandsCore/Private/GroomComponent.cpp b/Engine/Plugins/Runtime/HairStrands/Source/HairStrandsCore/Private/GroomComponent.cpp
--- a/Engine/Plugins/Runtime/HairStrands/Source/HairStrandsCore/Private/GroomComponent.cpp
+++ b/Engine/Plugins/Runtime/HairStrands/Source/HairStrandsCore/Private/GroomComponent.cpp
@@ -831,7 +831,7 @@ public:
 		if (bUseCulling)
 		{
 			BatchElement.NumPrimitives = 0;
-			BatchElement.IndirectArgsBuffer = bUseCulling ? Instance->HairGroupPublicData->GetDrawIndirectBuffer().Buffer->GetVertexBufferRHI() : nullptr;
+			BatchElement.IndirectArgsBuffer = bUseCulling ? Instance->HairGroupPublicData->GetDrawIndirectBuffer().Buffer->GetRHI() : nullptr;
 			BatchElement.IndirectArgsOffset = 0;
 		}
 		else
diff --git a/Engine/Plugins/Runtime/HairStrands/Source/HairStrandsCore/Private/HairStrandsRendering.cpp b/Engine/Plugins/Runtime/HairStrands/Source/HairStrandsCore/Private/HairStrandsRendering.cpp
--- a/Engine/Plugins/Runtime/HairStrands/Source/HairStrandsCore/Private/HairStrandsRendering.cpp
+++ b/Engine/Plugins/Runtime/HairStrands/Source/HairStrandsCore/Private/HairStrandsRendering.cpp
@@ -1140,7 +1140,7 @@ static void BuildHairAccelerationStructure_Cards(FRHICommandList& RHICmdList,
 	Initializer.bFastBuild = true;
 	Initializer.bAllowUpdate = true;
 
-	FVertexBufferRHIRef PositionBuffer(DeformedResource->GetBuffer(FHairCardsDeformedResource::Current).Buffer->GetVertexBufferRHI()); // This will likely flicker result in half speed motion as everyother frame will use the wrong buffer
+	FVertexBufferRHIRef PositionBuffer(DeformedResource->GetBuffer(FHairCardsDeformedResource::Current).Buffer->GetRHI()); // This will likely flicker result in half speed motion as everyother frame will use the wrong buffer
 
 	FRayTracingGeometrySegment Segment;
 	Segment.VertexBuffer = PositionBuffer;
@@ -1169,7 +1169,7 @@ static void BuildHairAccelerationStructure_Meshes(FRHICommandList& RHICmdList,
 	Initializer.bFastBuild = true;
 	Initializer.bAllowUpdate = true;
 
-	FVertexBufferRHIRef PositionBuffer(DeformedResource->GetBuffer(FHairMeshesDeformedResource::Current).Buffer->GetVertexBufferRHI()); // This will likely flicker result in half speed motion as everyother frame will use the wrong buffer
+	FVertexBufferRHIRef PositionBuffer(DeformedResource->GetBuffer(FHairMeshesDeformedResource::Current).Buffer->GetRHI()); // This will likely flicker result in half speed motion as everyother frame will use the wrong buffer
 
 	FRayTracingGeometrySegment Segment;
 	Segment.VertexBuffer = PositionBuffer;
@@ -1495,7 +1495,7 @@ void ComputeHairStrandsInterpolation(
 				const bool bNeedFullBuild = !Instance->Strands.RenRaytracingResource->bIsRTGeometryInitialized;
 				if (bNeedFullBuild)
 				{
-					FVertexBufferRHIRef PositionBuffer(Instance->Strands.RenRaytracingResource->PositionBuffer.Buffer->GetVertexBufferRHI());
+					FVertexBufferRHIRef PositionBuffer(Instance->Strands.RenRaytracingResource->PositionBuffer.Buffer->GetRHI());
 					BuildHairAccelerationStructure_Strands(RHICmdList, Instance->Strands.RenRaytracingResource->VertexCount, PositionBuffer, &Instance->Strands.RenRaytracingResource->RayTracingGeometry);
 				}
 				else
