Fixed versioning in HairDescription serialization. Previously, resaving a groom asset of an older version would lead to the bulkdata versioning not being updated, leading to an inconsistency between the package version and the HairDescription version. Also ensured that custom versioning is correctly set up when serializing to HairDescriptionBulkData.

diff --git a/Engine/Plugins/Runtime/HairStrands/Source/HairStrandsCore/Private/GroomAsset.cpp b/Engine/Plugins/Runtime/HairStrands/Source/HairStrandsCore/Private/GroomAsset.cpp
--- a/Engine/Plugins/Runtime/HairStrands/Source/HairStrandsCore/Private/GroomAsset.cpp
+++ b/Engine/Plugins/Runtime/HairStrands/Source/HairStrandsCore/Private/GroomAsset.cpp
@@ -1169,7 +1169,7 @@ void UGroomAsset::SetHairWidth(float Width)
 // differences, etc.) replace the version GUID below with a new one.
 // In case of merge conflicts with DDC versions, you MUST generate a new GUID
 // and set this new GUID as the version.
-#define GROOM_DERIVED_DATA_VERSION TEXT("C571C104EA0E447CB099D4BEC0C871C1")
+#define GROOM_DERIVED_DATA_VERSION TEXT("04296541C35C4CE0999D08B6170BB67A")
 
 #if WITH_EDITORONLY_DATA
 
diff --git a/Engine/Plugins/Runtime/HairStrands/Source/HairStrandsCore/Private/HairDescription.cpp b/Engine/Plugins/Runtime/HairStrands/Source/HairStrandsCore/Private/HairDescription.cpp
--- a/Engine/Plugins/Runtime/HairStrands/Source/HairStrandsCore/Private/HairDescription.cpp
+++ b/Engine/Plugins/Runtime/HairStrands/Source/HairStrandsCore/Private/HairDescription.cpp
@@ -61,8 +61,13 @@ bool FHairDescription::IsValid() const
 	return (NumStrands > 0) && (NumVertices > 0);
 }
 
+
 void FHairDescription::Serialize(FArchive& Ar)
 {
+	Ar.UsingCustomVersion(FEditorObjectVersion::GUID);
+	Ar.UsingCustomVersion(FReleaseObjectVersion::GUID);
+	Ar.UsingCustomVersion(FUE5MainStreamObjectVersion::GUID);
+
 	Ar << NumVertices;
 	Ar << NumStrands;
 
@@ -75,11 +80,35 @@ void FHairDescription::Serialize(FArchive& Ar)
 
 void FHairDescriptionBulkData::Serialize(FArchive& Ar, UObject* Owner)
 {
-	if (Ar.IsLoading())
+	Ar.UsingCustomVersion(FEditorObjectVersion::GUID);
+	Ar.UsingCustomVersion(FReleaseObjectVersion::GUID);
+	Ar.UsingCustomVersion(FUE5MainStreamObjectVersion::GUID);
+
+	if (Ar.IsTransacting())
 	{
-		// If loading, take the package custom version so it can be applied to the bulk data archive
-		// when unpacking HairDescription from it
-		CustomVersions = Ar.GetCustomVersions();
+		// If transacting, keep these members alive the other side of an undo, otherwise their values will get lost
+		CustomVersions.Serialize(Ar);
+		Ar << bBulkDataUpdated;
+	}
+	else
+	{
+		if (Ar.IsLoading())
+		{
+			// If loading, take the package custom version so it can be applied to the bulk data archive
+			// when unpacking HairDescription from it
+			CustomVersions = Ar.GetCustomVersions();
+		}
+		else if (Ar.IsSaving())
+		{
+			// If the bulk data hasn't been updated since this was loaded, there's a possibility that it has old versioning.
+			// Explicitly load and resave the FHairDescription so that its version is in sync with the FHairDescriptionBulkData.
+			if (!bBulkDataUpdated)
+			{
+				FHairDescription HairDescription;
+				LoadHairDescription(HairDescription);
+				SaveHairDescription(HairDescription);
+			}
+		}
 	}
 
 	BulkData.Serialize(Ar, Owner);
@@ -102,6 +131,10 @@ void FHairDescriptionBulkData::SaveHairDescription(FHairDescription& HairDescrip
 
 	// Use bulk data hash instead of guid to identify content to improve DDC cache hit
 	ComputeGuidFromHash();
+
+	// Mark the HairDescriptionBulkData as having been updated.
+	// This means we know that its version is up-to-date.
+	bBulkDataUpdated = true;
 }
 
 
diff --git a/Engine/Plugins/Runtime/HairStrands/Source/HairStrandsCore/Public/HairDescription.h b/Engine/Plugins/Runtime/HairStrands/Source/HairStrandsCore/Public/HairDescription.h
--- a/Engine/Plugins/Runtime/HairStrands/Source/HairStrandsCore/Public/HairDescription.h
+++ b/Engine/Plugins/Runtime/HairStrands/Source/HairStrandsCore/Public/HairDescription.h
@@ -146,6 +146,9 @@ private:
 
 	/** Custom version to propagate to archive when serializing the bulk data */
 	FCustomVersionContainer CustomVersions;
+
+	/** Whether the bulk data has been written via SaveHairDescription */
+	bool bBulkDataUpdated = false;
 };
 
 template <typename AttributeType> using TStrandAttributesRef = TMeshAttributesRef<FStrandID, AttributeType>;
