Fix CIS static analysis issue.

diff --git a/Engine/Plugins/Runtime/HairStrands/Source/HairStrandsCore/Private/GroomComponent.cpp b/Engine/Plugins/Runtime/HairStrands/Source/HairStrandsCore/Private/GroomComponent.cpp
--- a/Engine/Plugins/Runtime/HairStrands/Source/HairStrandsCore/Private/GroomComponent.cpp
+++ b/Engine/Plugins/Runtime/HairStrands/Source/HairStrandsCore/Private/GroomComponent.cpp
@@ -1722,7 +1722,7 @@ static USkeletalMeshComponent* ValidateBindingAsset(
 					GroupIt < BindingAsset->HairGroupResources.Num() &&
 					CardsLODIndex < uint32(BindingAsset->HairGroupResources[GroupIt].CardsRootResources.Num()) &&
 					BindingAsset->HairGroupResources[GroupIt].CardsRootResources[CardsLODIndex] != nullptr &&
-					(SkeletalMeshComponent->SkeletalMesh ? SkeletalMeshComponent->SkeletalMesh->GetLODInfoArray().Num() == BindingAsset->HairGroupResources[GroupIt].CardsRootResources[CardsLODIndex]->RootData.MeshProjectionLODs.Num() : false);
+					((SkeletalMeshComponent && SkeletalMeshComponent->SkeletalMesh) ? SkeletalMeshComponent->SkeletalMesh->GetLODInfoArray().Num() == BindingAsset->HairGroupResources[GroupIt].CardsRootResources[CardsLODIndex]->RootData.MeshProjectionLODs.Num() : false);
 
 				if (!bIsCardsBindingCompatible)
 				{
