Fix hair strands lagging one frame behind.

diff --git a/Engine/Plugins/Runtime/HairStrands/Source/HairStrandsCore/Private/GroomManager.cpp b/Engine/Plugins/Runtime/HairStrands/Source/HairStrandsCore/Private/GroomManager.cpp
--- a/Engine/Plugins/Runtime/HairStrands/Source/HairStrandsCore/Private/GroomManager.cpp
+++ b/Engine/Plugins/Runtime/HairStrands/Source/HairStrandsCore/Private/GroomManager.cpp
@@ -564,7 +564,7 @@ void RunHairStrandsDebug(
 // HairStrands Bookmark API
 
 void ProcessHairStrandsBookmark(
-	FRDGBuilder& GraphBuilder,
+	FRDGBuilder* GraphBuilder,
 	EHairStrandsBookmark Bookmark,
 	FHairStrandsBookmarkParameters& Parameters)
 {
@@ -580,7 +580,8 @@ void ProcessHairStrandsBookmark(
 
 		if (bHasHairStardsnProcess)
 		{
-			RunHairStrandsProcess(GraphBuilder, Parameters.ShaderMap, Parameters.DebugShaderData);
+			check(GraphBuilder);
+			RunHairStrandsProcess(*GraphBuilder, Parameters.ShaderMap, Parameters.DebugShaderData);
 		}
 	}
 	else if (Bookmark == EHairStrandsBookmark::ProcessLODSelection)
@@ -606,8 +607,9 @@ void ProcessHairStrandsBookmark(
 	}
 	else if (Bookmark == EHairStrandsBookmark::ProcessGuideInterpolation)
 	{
+		check(GraphBuilder);
 		RunHairStrandsInterpolation(
-			GraphBuilder,
+			*GraphBuilder,
 			Parameters.WorldType,
 			Parameters.SkinCache,
 			Parameters.DebugShaderData,
@@ -623,8 +625,9 @@ void ProcessHairStrandsBookmark(
 	}
 	else if (Bookmark == EHairStrandsBookmark::ProcessStrandsInterpolation)
 	{
+		check(GraphBuilder);
 		RunHairStrandsInterpolation(
-			GraphBuilder,
+			*GraphBuilder,
 			Parameters.WorldType,
 			Parameters.SkinCache,
 			Parameters.DebugShaderData,
@@ -634,8 +637,9 @@ void ProcessHairStrandsBookmark(
 	}
 	else if (Bookmark == EHairStrandsBookmark::ProcessDebug)
 	{
+		check(GraphBuilder);
 		RunHairStrandsDebug(
-			GraphBuilder,
+			*GraphBuilder,
 			Parameters.ShaderMap,
 			Parameters.WorldType,
 			*Parameters.View,
diff --git a/Engine/Plugins/Runtime/HairStrands/Source/HairStrandsCore/Private/HairStrandsCore.cpp b/Engine/Plugins/Runtime/HairStrands/Source/HairStrandsCore/Private/HairStrandsCore.cpp
--- a/Engine/Plugins/Runtime/HairStrands/Source/HairStrandsCore/Private/HairStrandsCore.cpp
+++ b/Engine/Plugins/Runtime/HairStrands/Source/HairStrandsCore/Private/HairStrandsCore.cpp
@@ -11,7 +11,7 @@
 IMPLEMENT_MODULE(FHairStrandsCore, HairStrandsCore);
 
 void ProcessHairStrandsBookmark(
-	FRDGBuilder& GraphBuilder,
+	FRDGBuilder* GraphBuilder,
 	EHairStrandsBookmark Bookmark,
 	FHairStrandsBookmarkParameters& Parameters);
 
diff --git a/Engine/Source/Runtime/Renderer/Private/HairStrands/HairStrandsInterface.cpp b/Engine/Source/Runtime/Renderer/Private/HairStrands/HairStrandsInterface.cpp
--- a/Engine/Source/Runtime/Renderer/Private/HairStrands/HairStrandsInterface.cpp
+++ b/Engine/Source/Runtime/Renderer/Private/HairStrands/HairStrandsInterface.cpp
@@ -323,7 +323,15 @@ void RunHairStrandsBookmark(FRDGBuilder& GraphBuilder, EHairStrandsBookmark Book
 {
 	if (GHairStrandsBookmarkFunction)
 	{
-		GHairStrandsBookmarkFunction(GraphBuilder, Bookmark, Parameters);
+		GHairStrandsBookmarkFunction(&GraphBuilder, Bookmark, Parameters);
+	}
+}
+
+void RunHairStrandsBookmark(EHairStrandsBookmark Bookmark, FHairStrandsBookmarkParameters& Parameters)
+{
+	if (GHairStrandsBookmarkFunction)
+	{
+		GHairStrandsBookmarkFunction(nullptr, Bookmark, Parameters);
 	}
 }
 
diff --git a/Engine/Source/Runtime/Renderer/Private/HairStrands/HairStrandsRendering.h b/Engine/Source/Runtime/Renderer/Private/HairStrands/HairStrandsRendering.h
--- a/Engine/Source/Runtime/Renderer/Private/HairStrands/HairStrandsRendering.h
+++ b/Engine/Source/Runtime/Renderer/Private/HairStrands/HairStrandsRendering.h
@@ -48,5 +48,9 @@ void RunHairStrandsBookmark(
 	EHairStrandsBookmark Bookmark, 
 	FHairStrandsBookmarkParameters& Parameters);
 
+void RunHairStrandsBookmark(
+	EHairStrandsBookmark Bookmark,
+	FHairStrandsBookmarkParameters& Parameters);
+
 FHairStrandsBookmarkParameters CreateHairStrandsBookmarkParameters(FViewInfo& View);
 FHairStrandsBookmarkParameters CreateHairStrandsBookmarkParameters(TArray<FViewInfo>& View);
\ No newline at end of file
diff --git a/Engine/Source/Runtime/Renderer/Private/SceneRendering.cpp b/Engine/Source/Runtime/Renderer/Private/SceneRendering.cpp
--- a/Engine/Source/Runtime/Renderer/Private/SceneRendering.cpp
+++ b/Engine/Source/Runtime/Renderer/Private/SceneRendering.cpp
@@ -3677,20 +3677,19 @@ static void RenderViewFamily_RenderThread(FRHICommandListImmediate& RHICmdList,
 				// Render the scene.
 				SceneRenderer->Render(GraphBuilder);
 			}
+			GraphBuilder.Execute();
+
+			CSV_SCOPED_TIMING_STAT_EXCLUSIVE(PostRenderCleanUp);
 			
 			if (SceneRenderer->Views.Num() > 0 && !ViewFamily.EngineShowFlags.HitProxies)
 			{
-				FHairStrandsBookmarkParameters& Parameters = *GraphBuilder.AllocObject<FHairStrandsBookmarkParameters>(CreateHairStrandsBookmarkParameters(SceneRenderer->Views));
+				FHairStrandsBookmarkParameters Parameters = CreateHairStrandsBookmarkParameters(SceneRenderer->Views);
 				if (Parameters.bHasElements)
 				{
-					RunHairStrandsBookmark(GraphBuilder, EHairStrandsBookmark::ProcessEndOfFrame, Parameters);
+					RunHairStrandsBookmark(EHairStrandsBookmark::ProcessEndOfFrame, Parameters);
 				}
 			}
 
-			GraphBuilder.Execute();
-
-			CSV_SCOPED_TIMING_STAT_EXCLUSIVE(PostRenderCleanUp);
-
 			// Only reset per-frame scene state once all views have processed their frame, including those in planar reflections
 			for (int32 CacheType = 0; CacheType < UE_ARRAY_COUNT(SceneRenderer->Scene->DistanceFieldSceneData.PrimitiveModifiedBounds); CacheType++)
 			{
diff --git a/Engine/Source/Runtime/Renderer/Private/SceneVisibility.cpp b/Engine/Source/Runtime/Renderer/Private/SceneVisibility.cpp
--- a/Engine/Source/Runtime/Renderer/Private/SceneVisibility.cpp
+++ b/Engine/Source/Runtime/Renderer/Private/SceneVisibility.cpp
@@ -3116,7 +3116,7 @@ void FSceneRenderer::PreVisibilityFrameSetup(FRDGBuilder& GraphBuilder)
 		FHairStrandsBookmarkParameters Parameters = CreateHairStrandsBookmarkParameters(Views);
 		if (Parameters.bHasElements)
 		{
-			RunHairStrandsBookmark(GraphBuilder, EHairStrandsBookmark::ProcessLODSelection, Parameters);
+			RunHairStrandsBookmark(EHairStrandsBookmark::ProcessLODSelection, Parameters);
 		}
 	}
 
diff --git a/Engine/Source/Runtime/Renderer/Public/HairStrandsInterface.h b/Engine/Source/Runtime/Renderer/Public/HairStrandsInterface.h
--- a/Engine/Source/Runtime/Renderer/Public/HairStrandsInterface.h
+++ b/Engine/Source/Runtime/Renderer/Public/HairStrandsInterface.h
@@ -385,5 +385,5 @@ struct FHairStrandsBookmarkParameters
 };
 
 typedef void (*THairStrandsParameterFunction)(FHairStrandsBookmarkParameters& Parameters);
-typedef void (*THairStrandsBookmarkFunction)(FRDGBuilder& GraphBuilder, EHairStrandsBookmark Bookmark, FHairStrandsBookmarkParameters& Parameters);
+typedef void (*THairStrandsBookmarkFunction)(FRDGBuilder* GraphBuilder, EHairStrandsBookmark Bookmark, FHairStrandsBookmarkParameters& Parameters);
 RENDERER_API void RegisterBookmarkFunction(THairStrandsBookmarkFunction Bookmark, THairStrandsParameterFunction Parameters);
