Fix precision issue with hair strands position, when the skeletal mesh's bounding box, on which the groom is attached to, is too large.

diff --git a/Engine/Plugins/Runtime/HairStrands/Source/HairStrandsCore/Private/GroomComponent.cpp b/Engine/Plugins/Runtime/HairStrands/Source/HairStrandsCore/Private/GroomComponent.cpp
--- a/Engine/Plugins/Runtime/HairStrands/Source/HairStrandsCore/Private/GroomComponent.cpp
+++ b/Engine/Plugins/Runtime/HairStrands/Source/HairStrandsCore/Private/GroomComponent.cpp
@@ -1819,7 +1819,7 @@ void UGroomComponent::InitResources(bool bIsBindingReloading)
 			}
 			HairGroupInstance->Guides.RestResource = GroupData.Guides.RestResource;
 
-			HairGroupInstance->Guides.DeformedResource = new FHairStrandsDeformedResource(GroupData.Guides.Data.RenderData, true, true);
+			HairGroupInstance->Guides.DeformedResource = new FHairStrandsDeformedResource(GroupData.Guides.Data.RenderData, true, true, HairGroupInstance->Guides.RestResource->PositionOffset);
 			BeginInitResource(HairGroupInstance->Guides.DeformedResource);
 
 			HairGroupInstance->Guides.DeformedResource->GetPositionOffset(FHairStrandsDeformedResource::EFrameType::Current)  = HairGroupInstance->Guides.RestResource->PositionOffset;
@@ -1888,7 +1888,7 @@ void UGroomComponent::InitResources(bool bIsBindingReloading)
 			}
 		
 			HairGroupInstance->Strands.RestResource = GroupData.Strands.RestResource;
-			HairGroupInstance->Strands.DeformedResource = new FHairStrandsDeformedResource(GroupData.Strands.Data.RenderData, false, bDynamicResources);
+			HairGroupInstance->Strands.DeformedResource = new FHairStrandsDeformedResource(GroupData.Strands.Data.RenderData, false, bDynamicResources, HairGroupInstance->Strands.RestResource->PositionOffset);
 			BeginInitResource(HairGroupInstance->Strands.DeformedResource);
 			HairGroupInstance->Strands.ClusterCullingResource = GroupData.Strands.ClusterCullingResource;
 
@@ -1966,7 +1966,7 @@ void UGroomComponent::InitResources(bool bIsBindingReloading)
 					}
 
 					InstanceLOD.Guides.RestResource = LOD.Guides.RestResource;
-					InstanceLOD.Guides.DeformedResource = new FHairStrandsDeformedResource(LOD.Guides.Data.RenderData, false, true);
+					InstanceLOD.Guides.DeformedResource = new FHairStrandsDeformedResource(LOD.Guides.Data.RenderData, false, true, InstanceLOD.Guides.RestResource->PositionOffset);
 					BeginInitResource(InstanceLOD.Guides.DeformedResource);
 
 					// Initialize deformed position relative position offset to the rest pose offset
diff --git a/Engine/Plugins/Runtime/HairStrands/Source/HairStrandsCore/Private/GroomManager.cpp b/Engine/Plugins/Runtime/HairStrands/Source/HairStrandsCore/Private/GroomManager.cpp
--- a/Engine/Plugins/Runtime/HairStrands/Source/HairStrandsCore/Private/GroomManager.cpp
+++ b/Engine/Plugins/Runtime/HairStrands/Source/HairStrandsCore/Private/GroomManager.cpp
@@ -162,6 +162,22 @@ void RunHairStrandsInterpolation(
 							MeshDataLOD, 
 							Instance->Strands.RestRootResource, 
 							Instance->Strands.DeformedRootResource);
+
+						AddHairStrandUpdatePositionOffsetPass(
+							GraphBuilder,
+							ShaderMap,
+							MeshLODIndex,
+							Instance->Strands.DeformedRootResource,
+							Instance->Strands.DeformedResource);
+					}
+					else if (Instance->Strands.HasValidData())
+					{
+						AddHairStrandUpdatePositionOffsetPass(
+							GraphBuilder,
+							ShaderMap,
+							MeshLODIndex,
+							nullptr,
+							Instance->Strands.DeformedResource);
 					}
 				}
 				else if (InstanceGeometryType == EHairGeometryType::Cards)
@@ -180,6 +196,22 @@ void RunHairStrandsInterpolation(
 								MeshDataLOD,
 								CardsInstance.Guides.RestRootResource,
 								CardsInstance.Guides.DeformedRootResource);
+
+							AddHairStrandUpdatePositionOffsetPass(
+								GraphBuilder,
+								ShaderMap,
+								MeshLODIndex,
+								CardsInstance.Guides.DeformedRootResource,
+								CardsInstance.Guides.DeformedResource);
+						}
+						else if (CardsInstance.Guides.IsValid())
+						{
+							AddHairStrandUpdatePositionOffsetPass(
+								GraphBuilder,
+								ShaderMap,
+								MeshLODIndex,
+								nullptr,
+								CardsInstance.Guides.DeformedResource);
 						}
 					}
 				}
@@ -221,6 +253,22 @@ void RunHairStrandsInterpolation(
 							Instance->Guides.RestRootResource,
 							Instance->Guides.DeformedRootResource);
 					}
+
+					AddHairStrandUpdatePositionOffsetPass(
+						GraphBuilder,
+						ShaderMap,
+						Instance->Debug.MeshLODIndex,
+						Instance->Guides.DeformedRootResource,
+						Instance->Guides.DeformedResource);
+				}
+				else if (Instance->Guides.IsValid())
+				{
+					AddHairStrandUpdatePositionOffsetPass(
+						GraphBuilder,
+						ShaderMap,
+						Instance->Debug.MeshLODIndex,
+						nullptr,
+						Instance->Guides.DeformedResource);
 				}
 			}
 		}
diff --git a/Engine/Plugins/Runtime/HairStrands/Source/HairStrandsCore/Private/GroomResources.cpp b/Engine/Plugins/Runtime/HairStrands/Source/HairStrandsCore/Private/GroomResources.cpp
--- a/Engine/Plugins/Runtime/HairStrands/Source/HairStrandsCore/Private/GroomResources.cpp
+++ b/Engine/Plugins/Runtime/HairStrands/Source/HairStrandsCore/Private/GroomResources.cpp
@@ -512,8 +512,8 @@ void FHairStrandsRestResource::ReleaseRHI()
 
 /////////////////////////////////////////////////////////////////////////////////////////
 
-FHairStrandsDeformedResource::FHairStrandsDeformedResource(const FHairStrandsDatas::FRenderData& HairStrandRenderData, bool bInInitializedData, bool bInDynamic) :
-	RenderData(HairStrandRenderData), bInitializedData(bInInitializedData), bDynamic(bInDynamic)
+FHairStrandsDeformedResource::FHairStrandsDeformedResource(const FHairStrandsDatas::FRenderData& HairStrandRenderData, bool bInInitializedData, bool bInDynamic, const FVector& InDefaultOffset) :
+	RenderData(HairStrandRenderData), bInitializedData(bInInitializedData), bDynamic(bInDynamic), DefaultOffset(InDefaultOffset)
 {}
 
 void FHairStrandsDeformedResource::InitRHI()
@@ -542,6 +542,11 @@ void FHairStrandsDeformedResource::InitRHI()
 	}
 	InternalCreateVertexBufferRDG<FHairStrandsTangentFormat>(GraphBuilder, VertexCount * FHairStrandsTangentFormat::ComponentCount, TangentBuffer, TEXT("HairStrandsDeformed_TangentBuffer"));
 
+	TArray<FVector4> DefaultOffsets;
+	DefaultOffsets.Add(DefaultOffset);
+	InternalCreateVertexBufferRDG<FHairStrandsPositionOffsetFormat>(GraphBuilder, DefaultOffsets, DeformedOffsetBuffer[0], TEXT("HairStrandsDeformed_DeformedOffsetBuffer0"));
+	InternalCreateVertexBufferRDG<FHairStrandsPositionOffsetFormat>(GraphBuilder, DefaultOffsets, DeformedOffsetBuffer[1], TEXT("HairStrandsDeformed_DeformedOffsetBuffer1"));
+
 	GraphBuilder.Execute();
 }
 
@@ -553,6 +558,9 @@ void FHairStrandsDeformedResource::ReleaseRHI()
 		DeformedPositionBuffer[1].Release();
 	}
 	TangentBuffer.Release();
+
+	DeformedOffsetBuffer[0].Release();
+	DeformedOffsetBuffer[1].Release();
 }
 
 bool FHairStrandsDeformedResource::NeedsToUpdateTangent()
diff --git a/Engine/Plugins/Runtime/HairStrands/Source/HairStrandsCore/Private/HairStrandsMeshProjection.cpp b/Engine/Plugins/Runtime/HairStrands/Source/HairStrandsCore/Private/HairStrandsMeshProjection.cpp
--- a/Engine/Plugins/Runtime/HairStrands/Source/HairStrandsCore/Private/HairStrandsMeshProjection.cpp
+++ b/Engine/Plugins/Runtime/HairStrands/Source/HairStrandsCore/Private/HairStrandsMeshProjection.cpp
@@ -13,6 +13,9 @@
 static int32 GHairProjectionMaxTrianglePerProjectionIteration = 8;
 static FAutoConsoleVariableRef CVarHairProjectionMaxTrianglePerProjectionIteration(TEXT("r.HairStrands.Projection.MaxTrianglePerIteration"), GHairProjectionMaxTrianglePerProjectionIteration, TEXT("Change the number of triangles which are iterated over during one projection iteration step. In kilo triangle (e.g., 8 == 8000 triangles). Default is 8."));
 
+static int32 GHairStrandsUseGPUPositionOffset = 1;
+static FAutoConsoleVariableRef CVarHairStrandsUseGPUPositionOffset(TEXT("r.HairStrands.UseGPUPositionOffset"), GHairStrandsUseGPUPositionOffset, TEXT("Use GPU position offset to improve hair strands position precision."));
+
 ///////////////////////////////////////////////////////////////////////////////////////////////////
 
 #define MAX_HAIRSTRANDS_SECTION_COUNT 255
@@ -578,6 +581,7 @@ public:
 	{
 		FGlobalShader::ModifyCompilationEnvironment(Parameters, OutEnvironment);
 		OutEnvironment.SetDefine(TEXT("MAX_SECTION_COUNT"), SectionArrayCount);
+		OutEnvironment.SetDefine(TEXT("SHADER_MESH_UPDATE"), SectionArrayCount);
 	}
 };
 
@@ -1748,4 +1752,71 @@ void AddComputeMipsPass(
 			FComputeShaderUtils::Dispatch(RHICmdList, ComputeShader, *Parameters, GroupCount);
 		});
 	}
+}
+
+///////////////////////////////////////////////////////////////////////////////////////////////////
+
+class FHairUpdatePositionOffsetCS : public FGlobalShader
+{
+public:
+private:
+	DECLARE_GLOBAL_SHADER(FHairUpdatePositionOffsetCS);
+	SHADER_USE_PARAMETER_STRUCT(FHairUpdatePositionOffsetCS, FGlobalShader);
+
+	BEGIN_SHADER_PARAMETER_STRUCT(FParameters, )
+		SHADER_PARAMETER(FVector, CPUPositionOffset)
+		SHADER_PARAMETER(uint32, bUseCPUOffset)
+		SHADER_PARAMETER_RDG_BUFFER_SRV(StructuredBuffer, RootTrianglePosition0Buffer)
+		SHADER_PARAMETER_RDG_BUFFER_UAV(RWBuffer, OutOffsetBuffer)
+	END_SHADER_PARAMETER_STRUCT()
+
+public:
+	static bool ShouldCompilePermutation(const FGlobalShaderPermutationParameters& Parameters)
+	{
+		return IsHairStrandsSupported(EHairStrandsShaderType::All, Parameters.Platform);
+	}
+	static void ModifyCompilationEnvironment(const FGlobalShaderPermutationParameters& Parameters, FShaderCompilerEnvironment& OutEnvironment)
+	{
+		FGlobalShader::ModifyCompilationEnvironment(Parameters, OutEnvironment);
+		OutEnvironment.SetDefine(TEXT("SHADER_OFFSET_UPDATE"), 1);
+	}
+};
+
+IMPLEMENT_GLOBAL_SHADER(FHairUpdatePositionOffsetCS, "/Engine/Private/HairStrands/HairStrandsMeshUpdate.usf", "MainCS", SF_Compute);
+
+void AddHairStrandUpdatePositionOffsetPass(
+	FRDGBuilder& GraphBuilder,
+	FGlobalShaderMap* ShaderMap,
+	const int32 LODIndex,
+	FHairStrandsDeformedRootResource* DeformedRootResources,
+	FHairStrandsDeformedResource* DeformedResources)
+{
+	if (LODIndex < 0 || DeformedResources == nullptr)
+	{
+		return;
+	}
+
+	FRDGImportedBuffer OutPositionOffsetBuffer    = Register(GraphBuilder, DeformedResources->GetPositionOffsetBuffer(FHairStrandsDeformedResource::Current), ERDGImportedBufferFlags::CreateUAV);
+	FRDGImportedBuffer RootTrianglePositionBuffer;
+	if (DeformedRootResources)
+	{
+		RootTrianglePositionBuffer = Register(GraphBuilder, DeformedRootResources->LODs[LODIndex].DeformedRootTrianglePosition0Buffer, ERDGImportedBufferFlags::CreateSRV);
+	}
+
+	const uint32 OffsetIndex = DeformedResources->GetIndex(FHairStrandsDeformedResource::Current);
+	FHairUpdatePositionOffsetCS::FParameters* PassParameters = GraphBuilder.AllocParameters<FHairUpdatePositionOffsetCS::FParameters>();
+	PassParameters->CPUPositionOffset = DeformedResources->PositionOffset[OffsetIndex];
+	PassParameters->bUseCPUOffset = (DeformedRootResources == nullptr || GHairStrandsUseGPUPositionOffset>0) ? 0 : 1;
+	PassParameters->RootTrianglePosition0Buffer = RootTrianglePositionBuffer.SRV;
+	PassParameters->OutOffsetBuffer = OutPositionOffsetBuffer.UAV;
+
+	TShaderMapRef<FHairUpdatePositionOffsetCS> ComputeShader(ShaderMap);
+	FComputeShaderUtils::AddPass(
+		GraphBuilder,
+		RDG_EVENT_NAME("HairStrandsUpdatePositionOffset"),
+		ComputeShader,
+		PassParameters,
+		FIntVector(1, 1, 1));
+
+	GraphBuilder.SetBufferAccessFinal(OutPositionOffsetBuffer.Buffer, ERHIAccess::SRVMask);
 }
\ No newline at end of file
diff --git a/Engine/Plugins/Runtime/HairStrands/Source/HairStrandsCore/Private/HairStrandsRendering.cpp b/Engine/Plugins/Runtime/HairStrands/Source/HairStrandsCore/Private/HairStrandsRendering.cpp
--- a/Engine/Plugins/Runtime/HairStrands/Source/HairStrandsCore/Private/HairStrandsRendering.cpp
+++ b/Engine/Plugins/Runtime/HairStrands/Source/HairStrandsCore/Private/HairStrandsRendering.cpp
@@ -190,7 +190,6 @@ class FDeformGuideCS : public FGlobalShader
 		SHADER_PARAMETER(uint32, VertexCount)
 		SHADER_PARAMETER(uint32, IterationCount)
 		SHADER_PARAMETER(FVector, SimRestOffset)
-		SHADER_PARAMETER(FVector, SimDeformedOffset)
 		SHADER_PARAMETER(uint32, DispatchCountX)
 
 		SHADER_PARAMETER_RDG_BUFFER_SRV(Buffer<float4>, SimRestPosition0Buffer)
@@ -204,6 +203,7 @@ class FDeformGuideCS : public FGlobalShader
 		SHADER_PARAMETER_RDG_BUFFER_SRV(Buffer<uint>, SimRootBarycentricBuffer)
 		SHADER_PARAMETER_RDG_BUFFER_SRV(Buffer<uint>, SimVertexToRootIndexBuffer)
 
+		SHADER_PARAMETER_RDG_BUFFER_SRV(Buffer, SimDeformedOffsetBuffer)
 		SHADER_PARAMETER_RDG_BUFFER_SRV(Buffer, SimRestPosePositionBuffer)
 		SHADER_PARAMETER_RDG_BUFFER_UAV(RWBuffer, OutSimDeformedPositionBuffer)
 
@@ -229,7 +229,7 @@ static void AddDeformSimHairStrandsPass(
 	FRDGBufferSRVRef SimRestPosePositionBuffer,
 	FRDGImportedBuffer OutSimDeformedPositionBuffer,
 	const FVector& SimRestOffset,
-	const FVector& SimDeformedOffset,
+	FRDGBufferSRVRef SimDeformedOffsetBuffer,
 	const bool bHasGlobalInterpolation)
 {
 	static uint32 IterationCount = 0;
@@ -256,7 +256,7 @@ static void AddDeformSimHairStrandsPass(
 	Parameters->OutSimDeformedPositionBuffer = OutSimDeformedPositionBuffer.UAV;
 	Parameters->VertexCount = VertexCount;
 	Parameters->IterationCount = IterationCount % 10000;
-	Parameters->SimDeformedOffset = SimDeformedOffset;
+	Parameters->SimDeformedOffsetBuffer = SimDeformedOffsetBuffer;
 	Parameters->SimRestOffset = SimRestOffset;
 	Parameters->DispatchCountX = DispatchCountX;
 
@@ -395,8 +395,6 @@ class FHairInterpolationCS : public FGlobalShader
 		SHADER_PARAMETER(uint32, HairDebugMode)
 		SHADER_PARAMETER(FVector, InRenderHairPositionOffset)
 		SHADER_PARAMETER(FVector, InSimHairPositionOffset)
-		SHADER_PARAMETER(FVector, OutRenderHairPositionOffset)
-		SHADER_PARAMETER(FVector, OutSimHairPositionOffset)
 		SHADER_PARAMETER(FIntPoint, HairStrandsCullIndex)
 
 		SHADER_PARAMETER(float, InHairLength)
@@ -449,6 +447,9 @@ class FHairInterpolationCS : public FGlobalShader
 
 		SHADER_PARAMETER_RDG_BUFFER_SRV(Buffer<uint>, SimRootPointIndexBuffer)
 
+		SHADER_PARAMETER_RDG_BUFFER_SRV(Buffer<float4>, OutSimHairPositionOffsetBuffer)
+		SHADER_PARAMETER_RDG_BUFFER_SRV(Buffer<float4>, OutRenHairPositionOffsetBuffer)
+
 		SHADER_PARAMETER_RDG_BUFFER_SRV(Buffer<uint>, HairStrandsVF_CullingIndirectBuffer)
 		SHADER_PARAMETER_RDG_BUFFER_SRV(Buffer<uint>, HairStrandsVF_CullingIndexBuffer)
 		SHADER_PARAMETER_RDG_BUFFER_SRV(Buffer<float>, HairStrandsVF_CullingRadiusScaleBuffer)
@@ -481,8 +482,8 @@ static void AddHairStrandsInterpolationPass(
 	FRDGHairStrandsCullingData& CullingData,
 	const FVector& InRenderHairWorldOffset,
 	const FVector& InSimHairWorldOffset,
-	const FVector& OutRenderHairWorldOffset,
-	const FVector& OutSimHairWorldOffset,
+	const FRDGBufferSRVRef& OutRenHairPositionOffsetBuffer,
+	const FRDGBufferSRVRef& OutSimHairPositionOffsetBuffer,
 	const FHairStrandsRestRootResource* RenRestRootResources,
 	const FHairStrandsRestRootResource* SimRestRootResources,
 	const FHairStrandsDeformedRootResource* RenDeformedRootResources,
@@ -515,8 +516,10 @@ static void AddHairStrandsInterpolationPass(
 	Parameters->VertexCount = VertexCount;
 	Parameters->InRenderHairPositionOffset = InRenderHairWorldOffset;
 	Parameters->InSimHairPositionOffset =  InSimHairWorldOffset;
-	Parameters->OutRenderHairPositionOffset = OutRenderHairWorldOffset;
-	Parameters->OutSimHairPositionOffset = OutSimHairWorldOffset;
+
+	Parameters->OutSimHairPositionOffsetBuffer = OutSimHairPositionOffsetBuffer;
+	Parameters->OutRenHairPositionOffsetBuffer = OutRenHairPositionOffsetBuffer;
+
 	Parameters->DispatchCountX = DispatchCount.X;
 	Parameters->SimRootPointIndexBuffer = SimRootPointIndexBuffer;
 	
@@ -699,9 +702,10 @@ class FHairClusterAABBCS : public FGlobalShader
 	BEGIN_SHADER_PARAMETER_STRUCT(FParameters, )
 		SHADER_PARAMETER(uint32, DispatchCountX)
 		SHADER_PARAMETER(uint32, ClusterCount)
-		SHADER_PARAMETER(FVector, OutHairPositionOffset)
 		SHADER_PARAMETER(FMatrix, LocalToWorldMatrix)
+
 		SHADER_PARAMETER_RDG_BUFFER_SRV(Buffer, RenderDeformedPositionBuffer)
+		SHADER_PARAMETER_RDG_BUFFER_SRV(Buffer, RenderDeformedOffsetBuffer)
 		SHADER_PARAMETER_RDG_BUFFER_SRV(Buffer, ClusterVertexIdBuffer)
 		SHADER_PARAMETER_RDG_BUFFER_SRV(Buffer, ClusterIdBuffer)
 		SHADER_PARAMETER_RDG_BUFFER_SRV(Buffer, ClusterIndexOffsetBuffer)
@@ -725,7 +729,7 @@ static void AddHairClusterAABBPass(
 	FRDGBuilder& GraphBuilder,
 	FGlobalShaderMap* ShaderMap,
 	const FTransform& InRenLocalToWorld,
-	const FVector& OutHairWorldOffset,
+	FRDGBufferSRVRef RenderDeformedOffsetBuffer,
 	FHairStrandClusterData::FHairGroup& ClusterData,
 
 	FRDGHairStrandsCullingData& ClusterAABBData,
@@ -742,8 +746,8 @@ static void AddHairClusterAABBPass(
 	Parameters->DispatchCountX = DispatchCount.X;
 	Parameters->ClusterCount = ClusterData.ClusterCount;
 	Parameters->LocalToWorldMatrix = InRenLocalToWorld.ToMatrixWithScale();
-	Parameters->OutHairPositionOffset = OutHairWorldOffset;
 	Parameters->RenderDeformedPositionBuffer = RenderPositionBuffer.SRV;
+	Parameters->RenderDeformedOffsetBuffer = RenderDeformedOffsetBuffer;
 	Parameters->ClusterVertexIdBuffer = RegisterAsSRV(GraphBuilder, *ClusterData.ClusterVertexIdBuffer);
 	Parameters->ClusterIdBuffer = GraphBuilder.CreateSRV(ClusterIdBuffer, PF_R32_UINT);
 	Parameters->ClusterIndexOffsetBuffer = GraphBuilder.CreateSRV(ClusterIndexOffsetBuffer, PF_R32_UINT);
@@ -985,8 +989,8 @@ class FHairRaytracingGeometryCS : public FGlobalShader
 	BEGIN_SHADER_PARAMETER_STRUCT(FParameters, )
 		SHADER_PARAMETER(uint32, VertexCount)
 		SHADER_PARAMETER(uint32, DispatchCountX)
-		SHADER_PARAMETER(FVector, StrandHairWorldOffset)
 		SHADER_PARAMETER(float, StrandHairRadius)
+		SHADER_PARAMETER_RDG_BUFFER_SRV(Buffer, PositionOffsetBuffer)
 		SHADER_PARAMETER_RDG_BUFFER_SRV(Buffer, PositionBuffer)
 		SHADER_PARAMETER_RDG_BUFFER_UAV(RWBuffer, OutputPositionBuffer)
 	END_SHADER_PARAMETER_STRUCT()
@@ -1002,7 +1006,7 @@ static void AddGenerateRaytracingGeometryPass(
 	FGlobalShaderMap* ShaderMap,
 	uint32 VertexCount,
 	float HairRadius,
-	const FVector& HairWorldOffset,
+	const FRDGBufferSRVRef& HairWorldOffsetBuffer,
 	const FRDGBufferSRVRef& PositionBuffer,
 	const FRDGBufferUAVRef& OutPositionBuffer)
 {
@@ -1012,7 +1016,7 @@ static void AddGenerateRaytracingGeometryPass(
 	FHairRaytracingGeometryCS::FParameters* Parameters = GraphBuilder.AllocParameters<FHairRaytracingGeometryCS::FParameters>();
 	Parameters->VertexCount = VertexCount;
 	Parameters->DispatchCountX = DispatchCount.X;
-	Parameters->StrandHairWorldOffset = HairWorldOffset;
+	Parameters->PositionOffsetBuffer = HairWorldOffsetBuffer;
 	Parameters->StrandHairRadius = HairRadius;
 	Parameters->PositionBuffer = PositionBuffer;
 	Parameters->OutputPositionBuffer = OutPositionBuffer;
@@ -1267,6 +1271,10 @@ FHairGroupPublicData::FVertexFactoryInput ComputeHairStrandsVertexInputData(FHai
 
 		OutVFInput.Strands.PositionOffset = Instance->Guides.DeformedResource->GetPositionOffset(FHairStrandsDeformedResource::EFrameType::Current);
 		OutVFInput.Strands.PrevPositionOffset = Instance->Guides.DeformedResource->GetPositionOffset(FHairStrandsDeformedResource::EFrameType::Previous);
+
+		OutVFInput.Strands.PositionOffsetBuffer = Instance->Guides.DeformedResource->GetPositionOffsetBuffer(FHairStrandsDeformedResource::EFrameType::Current).SRV;
+		OutVFInput.Strands.PrevPositionOffsetBuffer = Instance->Guides.DeformedResource->GetPositionOffsetBuffer(FHairStrandsDeformedResource::EFrameType::Previous).SRV;
+
 		OutVFInput.Strands.VertexCount = Instance->Guides.RestResource->GetVertexCount();
 		OutVFInput.Strands.HairRadius = (GStrandHairWidth > 0 ? GStrandHairWidth : Instance->Strands.Modifier.HairWidth) * 0.5f;
 		OutVFInput.Strands.HairLength = Instance->Strands.Modifier.HairLength;
@@ -1282,6 +1290,9 @@ FHairGroupPublicData::FVertexFactoryInput ComputeHairStrandsVertexInputData(FHai
 		OutVFInput.Strands.AttributeBuffer = bDebugModePatchedAttributeBuffer ? Instance->Strands.DebugAttributeBuffer.SRV : Instance->Strands.RestResource->AttributeBuffer.SRV;
 		OutVFInput.Strands.MaterialBuffer = Instance->Strands.RestResource->MaterialBuffer.SRV;
 
+		OutVFInput.Strands.PositionOffsetBuffer = Instance->Strands.DeformedResource->GetPositionOffsetBuffer(FHairStrandsDeformedResource::EFrameType::Current).SRV;
+		OutVFInput.Strands.PrevPositionOffsetBuffer = Instance->Strands.DeformedResource->GetPositionOffsetBuffer(FHairStrandsDeformedResource::EFrameType::Previous).SRV;
+
 		OutVFInput.Strands.PositionOffset = Instance->Strands.DeformedResource->GetPositionOffset(FHairStrandsDeformedResource::EFrameType::Current);
 		OutVFInput.Strands.PrevPositionOffset = Instance->Strands.DeformedResource->GetPositionOffset(FHairStrandsDeformedResource::EFrameType::Previous);
 		OutVFInput.Strands.VertexCount = Instance->Strands.RestResource->GetVertexCount();
@@ -1361,7 +1372,7 @@ void ComputeHairStrandsInterpolation(
 			RegisterAsSRV(GraphBuilder, Instance->Guides.RestResource->RestPositionBuffer),
 			Register(GraphBuilder, Instance->Guides.DeformedResource->GetBuffer(FHairStrandsDeformedResource::Current), ERDGImportedBufferFlags::CreateUAV),
 			Instance->Guides.RestResource->PositionOffset,
-			Instance->Guides.DeformedResource->GetPositionOffset(FHairStrandsDeformedResource::Current),
+			RegisterAsSRV(GraphBuilder, Instance->Guides.DeformedResource->GetPositionOffsetBuffer(FHairStrandsDeformedResource::Current)),
 			Instance->Guides.bHasGlobalInterpolation);
 	}
 
@@ -1417,8 +1428,8 @@ void ComputeHairStrandsInterpolation(
 					CullingData,
 					Instance->Strands.RestResource->PositionOffset,
 					Instance->Guides.RestResource->PositionOffset,
-					Instance->Strands.DeformedResource->GetPositionOffset(FHairStrandsDeformedResource::Current),
-					Instance->Guides.DeformedResource->GetPositionOffset(FHairStrandsDeformedResource::Current),
+					RegisterAsSRV(GraphBuilder, Instance->Strands.DeformedResource->GetPositionOffsetBuffer(FHairStrandsDeformedResource::Current)),
+					RegisterAsSRV(GraphBuilder, Instance->Guides.DeformedResource->GetPositionOffsetBuffer(FHairStrandsDeformedResource::Current)),
 					Instance->Strands.RestRootResource,
 					Instance->Guides.RestRootResource,
 					Instance->Strands.DeformedRootResource,
@@ -1452,7 +1463,7 @@ void ComputeHairStrandsInterpolation(
 						GraphBuilder,
 						ShaderMap,
 						Instance->LocalToWorld,
-						Instance->Strands.DeformedResource->GetPositionOffset(FHairStrandsDeformedResource::Current),
+						RegisterAsSRV(GraphBuilder, Instance->Strands.DeformedResource->GetPositionOffsetBuffer(FHairStrandsDeformedResource::EFrameType::Current)),
 						HairGroupCluster,
 						CullingData,
 						Strands_DeformedPosition);
@@ -1482,7 +1493,7 @@ void ComputeHairStrandsInterpolation(
 				ShaderMap,
 				Instance->Strands.RestResource->GetVertexCount(),
 				ScaleAndClipDesc.MaxOutHairRadius * HairRadiusScaleRT,
-				Instance->Strands.DeformedResource->GetPositionOffset(FHairStrandsDeformedResource::Current),
+				RegisterAsSRV(GraphBuilder, Instance->Strands.DeformedResource->GetPositionOffsetBuffer(FHairStrandsDeformedResource::EFrameType::Current)),
 				Strands_DeformedPosition.SRV,
 				Raytracing_PositionBuffer.UAV);
 
@@ -1543,8 +1554,8 @@ void ComputeHairStrandsInterpolation(
 					CullingData,
 					LOD.Guides.RestResource->PositionOffset,
 					Instance->Guides.RestResource->PositionOffset,
-					LOD.Guides.DeformedResource->GetPositionOffset(FHairStrandsDeformedResource::Current),
-					Instance->Guides.DeformedResource->GetPositionOffset(FHairStrandsDeformedResource::Current),
+					RegisterAsSRV(GraphBuilder, LOD.Guides.DeformedResource->GetPositionOffsetBuffer(FHairStrandsDeformedResource::Current)),
+					RegisterAsSRV(GraphBuilder, Instance->Guides.DeformedResource->GetPositionOffsetBuffer(FHairStrandsDeformedResource::Current)),
 					LOD.Guides.RestRootResource,
 					Instance->Guides.RestRootResource,
 					LOD.Guides.DeformedRootResource,
@@ -1683,6 +1694,6 @@ void ResetHairStrandsInterpolation(
 		RegisterAsSRV(GraphBuilder, Instance->Guides.RestResource->RestPositionBuffer),
 		DeformedPositionBuffer,
 		Instance->Guides.RestResource->PositionOffset,
-		Instance->Guides.DeformedResource->GetPositionOffset(FHairStrandsDeformedResource::Current), 
+		RegisterAsSRV(GraphBuilder, Instance->Guides.DeformedResource->GetPositionOffsetBuffer(FHairStrandsDeformedResource::Current)),
 		Instance->Guides.bHasGlobalInterpolation);
 }
diff --git a/Engine/Plugins/Runtime/HairStrands/Source/HairStrandsCore/Private/HairStrandsVertexFactory.cpp b/Engine/Plugins/Runtime/HairStrands/Source/HairStrandsCore/Private/HairStrandsVertexFactory.cpp
--- a/Engine/Plugins/Runtime/HairStrands/Source/HairStrandsCore/Private/HairStrandsVertexFactory.cpp
+++ b/Engine/Plugins/Runtime/HairStrands/Source/HairStrandsCore/Private/HairStrandsVertexFactory.cpp
@@ -54,8 +54,8 @@ inline FRHIShaderResourceView* GetAttributeSRV(const FHairGroupPublicData::FVert
 inline FRHIShaderResourceView* GetMaterialSRV(const FHairGroupPublicData::FVertexFactoryInput& VFInput)			{ return VFInput.Strands.MaterialBuffer; }
 inline FRHIShaderResourceView* GetTangentSRV(const FHairGroupPublicData::FVertexFactoryInput& VFInput)			{ return VFInput.Strands.TangentBuffer; }
 
-inline const FVector& GetPositionOffset(const FHairGroupPublicData::FVertexFactoryInput& VFInput)				{ return VFInput.Strands.PositionOffset; }
-inline const FVector& GetPreviousPositionOffset(const FHairGroupPublicData::FVertexFactoryInput& VFInput)		{ return VFInput.Strands.PrevPositionOffset; }
+inline FRHIShaderResourceView* GetPositionOffsetSRV(const FHairGroupPublicData::FVertexFactoryInput& VFInput) { return VFInput.Strands.PositionOffsetBuffer; };
+inline FRHIShaderResourceView* GetPreviousPositionOffsetSRV(const FHairGroupPublicData::FVertexFactoryInput& VFInput) { return VFInput.Strands.PrevPositionOffsetBuffer; }
 
 inline bool  UseStableRasterization(const FHairGroupPublicData::FVertexFactoryInput& VFInput)					{ return VFInput.Strands.bUseStableRasterization; };
 inline bool  UseScatterSceneLighting(const FHairGroupPublicData::FVertexFactoryInput& VFInput)					{ return VFInput.Strands.bScatterSceneLighting; };
@@ -76,13 +76,13 @@ public:
 	LAYOUT_FIELD(FShaderParameter, Length);
 	LAYOUT_FIELD(FShaderParameter, RadiusAtDepth1_Primary);	// unused
 	LAYOUT_FIELD(FShaderParameter, RadiusAtDepth1_Velocity);	// unused
-	LAYOUT_FIELD(FShaderParameter, PositionOffset);
-	LAYOUT_FIELD(FShaderParameter, PreviousPositionOffset);
 	LAYOUT_FIELD(FShaderParameter, Density);
 	LAYOUT_FIELD(FShaderParameter, Culling);
 	LAYOUT_FIELD(FShaderParameter, StableRasterization);
 	LAYOUT_FIELD(FShaderParameter, ScatterSceneLighing);
 
+	LAYOUT_FIELD(FShaderResourceParameter, PositionOffsetBuffer);
+	LAYOUT_FIELD(FShaderResourceParameter, PreviousPositionOffsetBuffer);
 	LAYOUT_FIELD(FShaderResourceParameter, PositionBuffer);
 	LAYOUT_FIELD(FShaderResourceParameter, PreviousPositionBuffer);
 	LAYOUT_FIELD(FShaderResourceParameter, AttributeBuffer);
@@ -96,13 +96,15 @@ public:
 	{
 		Radius.Bind(ParameterMap, TEXT("HairStrandsVF_Radius"));
 		Length.Bind(ParameterMap, TEXT("HairStrandsVF_Length"));
-		PositionOffset.Bind(ParameterMap, TEXT("HairStrandsVF_PositionOffset"));
-		PreviousPositionOffset.Bind(ParameterMap, TEXT("HairStrandsVF_PreviousPositionOffset"));
+		Density.Bind(ParameterMap, TEXT("HairStrandsVF_Density"));
 		Density.Bind(ParameterMap, TEXT("HairStrandsVF_Density"));	
 		Culling.Bind(ParameterMap, TEXT("HairStrandsVF_CullingEnable"));
 		StableRasterization.Bind(ParameterMap, TEXT("HairStrandsVF_bUseStableRasterization"));
 		ScatterSceneLighing.Bind(ParameterMap, TEXT("HairStrandsVF_bScatterSceneLighing"));
 
+		PositionOffsetBuffer.Bind(ParameterMap, TEXT("HairStrandsVF_PositionOffsetBuffer"));
+		PreviousPositionOffsetBuffer.Bind(ParameterMap, TEXT("HairStrandsVF_PreviousPositionOffsetBuffer"));
+
 		PositionBuffer.Bind(ParameterMap, TEXT("HairStrandsVF_PositionBuffer"));
 		PreviousPositionBuffer.Bind(ParameterMap, TEXT("HairStrandsVF_PreviousPositionBuffer"));
 		AttributeBuffer.Bind(ParameterMap, TEXT("HairStrandsVF_AttributeBuffer"));
@@ -139,8 +141,8 @@ public:
 		VFS_BindParam(ShaderBindings, TangentBuffer, GetTangentSRV(VFInput));
 		VFS_BindParam(ShaderBindings, Radius, GetMaxStrandRadius(VFInput));
 		VFS_BindParam(ShaderBindings, Length, GetMaxStrandLength(VFInput));
-		VFS_BindParam(ShaderBindings, PositionOffset, GetPositionOffset(VFInput));
-		VFS_BindParam(ShaderBindings, PreviousPositionOffset, GetPreviousPositionOffset(VFInput));
+		VFS_BindParam(ShaderBindings, PositionOffsetBuffer, GetPositionOffsetSRV(VFInput));
+		VFS_BindParam(ShaderBindings, PreviousPositionOffsetBuffer, GetPreviousPositionOffsetSRV(VFInput));
 		VFS_BindParam(ShaderBindings, Density, GetHairDensity(VFInput));
 		VFS_BindParam(ShaderBindings, StableRasterization, UseStableRasterization(VFInput) ? 1u : 0u);
 		VFS_BindParam(ShaderBindings, ScatterSceneLighing, UseScatterSceneLighting(VFInput) ? 1u : 0u);
diff --git a/Engine/Plugins/Runtime/HairStrands/Source/HairStrandsCore/Public/GroomResources.h b/Engine/Plugins/Runtime/HairStrands/Source/HairStrandsCore/Public/GroomResources.h
--- a/Engine/Plugins/Runtime/HairStrands/Source/HairStrandsCore/Public/GroomResources.h
+++ b/Engine/Plugins/Runtime/HairStrands/Source/HairStrandsCore/Public/GroomResources.h
@@ -204,7 +204,7 @@ struct FHairStrandsRestResource : public FRenderResource
 struct FHairStrandsDeformedResource : public FRenderResource
 {
 	/** Build the hair strands resource */
-	FHairStrandsDeformedResource(const FHairStrandsDatas::FRenderData& HairStrandRenderData, bool bInitializeData, bool bDynamic);
+	FHairStrandsDeformedResource(const FHairStrandsDatas::FRenderData& HairStrandRenderData, bool bInitializeData, bool bDynamic, const FVector& InDefaultOffset);
 
 	/* Init the buffer */
 	virtual void InitRHI() override;
@@ -228,6 +228,9 @@ struct FHairStrandsDeformedResource : public FRenderResource
 	/* Strand hair deformed position buffer (previous and current) */
 	FRDGExternalBuffer DeformedPositionBuffer[2];
 
+	/* Strand hair deformed position buffer (previous and current) */
+	FRDGExternalBuffer DeformedOffsetBuffer[2];
+
 	/* Strand hair tangent buffer */
 	FRDGExternalBuffer TangentBuffer;
 
@@ -246,6 +249,7 @@ struct FHairStrandsDeformedResource : public FRenderResource
 	/* Whether the underlying resource is dynamic not (single or double buffer allocated) */
 	const bool bDynamic = true; 
 	bool bInitializedTangent = true;
+	FVector DefaultOffset = FVector::ZeroVector;
 
 	enum EFrameType
 	{
@@ -257,6 +261,7 @@ struct FHairStrandsDeformedResource : public FRenderResource
 	inline uint32 GetIndex(EFrameType T) const					{ return (!bDynamic || T == EFrameType::Current) ? CurrentIndex : 1u - CurrentIndex; }
 	inline FRDGExternalBuffer& GetBuffer(EFrameType T)			{ return DeformedPositionBuffer[GetIndex(T)];  }
 	inline FVector& GetPositionOffset(EFrameType T)				{ return PositionOffset[GetIndex(T)]; }
+	inline FRDGExternalBuffer& GetPositionOffsetBuffer(EFrameType T) { return DeformedOffsetBuffer[GetIndex(T)]; }
 	inline const FVector& GetPositionOffset(EFrameType T) const { return PositionOffset[GetIndex(T)]; }
 	inline void SwapBuffer()									{ if (bDynamic) { CurrentIndex = 1u - CurrentIndex; } }
 	bool NeedsToUpdateTangent();
diff --git a/Engine/Plugins/Runtime/HairStrands/Source/HairStrandsCore/Public/HairStrandsDatas.h b/Engine/Plugins/Runtime/HairStrands/Source/HairStrandsCore/Public/HairStrandsDatas.h
--- a/Engine/Plugins/Runtime/HairStrands/Source/HairStrandsCore/Public/HairStrandsDatas.h
+++ b/Engine/Plugins/Runtime/HairStrands/Source/HairStrandsCore/Public/HairStrandsDatas.h
@@ -83,6 +83,15 @@ struct FHairStrandsPositionFormat
 	static const EPixelFormat Format = PF_R16G16B16A16_UINT;
 };
 
+struct FHairStrandsPositionOffsetFormat
+{
+	typedef FVector4 Type;
+	static const uint32 ComponentCount = 1;
+	static const uint32 SizeInByte = sizeof(Type);
+	static const EVertexElementType VertexElementType = VET_Float4;
+	static const EPixelFormat Format = PF_A32B32G32R32F;
+};
+
 struct FHairStrandsAttributeFormat
 {
 	typedef FPackedHairAttributeVertex Type;
diff --git a/Engine/Plugins/Runtime/HairStrands/Source/HairStrandsCore/Public/HairStrandsMeshProjection.h b/Engine/Plugins/Runtime/HairStrands/Source/HairStrandsCore/Public/HairStrandsMeshProjection.h
--- a/Engine/Plugins/Runtime/HairStrands/Source/HairStrandsCore/Public/HairStrandsMeshProjection.h
+++ b/Engine/Plugins/Runtime/HairStrands/Source/HairStrandsCore/Public/HairStrandsMeshProjection.h
@@ -16,6 +16,7 @@ struct FHairMeshesRestResource;
 struct FHairMeshesDeformedResource;
 struct FHairCardsRestResource;
 struct FHairCardsDeformedResource;
+struct FHairStrandsDeformedResource;
 
 struct FHairStrandsProjectionMeshData
 {
@@ -147,4 +148,11 @@ void AddHairCardsRBFInterpolationPass(
 	FHairStrandsRestRootResource* RestRootResources,
 	FHairStrandsDeformedRootResource* DeformedRootResources);
 
+void AddHairStrandUpdatePositionOffsetPass(
+	FRDGBuilder& GraphBuilder,
+	FGlobalShaderMap* ShaderMap,
+	const int32 LODIndex,
+	FHairStrandsDeformedRootResource* DeformedRootResources,
+	FHairStrandsDeformedResource* DeformedResources);
+
 FHairStrandsProjectionMeshData ExtractMeshData(FSkeletalMeshRenderData* RenderData);
\ No newline at end of file
diff --git a/Engine/Shaders/Private/HairStrands/HairStrandsGuideDeform.usf b/Engine/Shaders/Private/HairStrands/HairStrandsGuideDeform.usf
--- a/Engine/Shaders/Private/HairStrands/HairStrandsGuideDeform.usf
+++ b/Engine/Shaders/Private/HairStrands/HairStrandsGuideDeform.usf
@@ -5,11 +5,11 @@
 #include "HairStrandsMeshRBF.ush"
 
 float3 SimRestOffset;
-float3 SimDeformedOffset;
 uint VertexCount;
 uint IterationCount;
 uint DispatchCountX; 
 
+Buffer<float4> SimDeformedOffsetBuffer;
 Buffer<uint4> SimRestPosePositionBuffer;
 RWBuffer<uint4> OutSimDeformedPositionBuffer;
 
@@ -67,6 +67,8 @@ Buffer<float4>  MeshSampleWeightsBuffer;
 
 float3 DisplacePosition(uint VertexIndex, float3 Pos, float Seed)
 {
+	const float3 OutSimHairPositionOffset = SimDeformedOffsetBuffer[0];
+
 // By pass
 #if PERMUTATION_DEFORMATION == DEFORMATION_BYPASS
 	return Pos;
@@ -86,7 +88,7 @@ float3 DisplacePosition(uint VertexIndex, float3 Pos, float Seed)
 #endif
 
 #if PERMUTATION_DEFORMATION == DEFORMATION_OFFSET_GUIDE
-	return Pos + (SimRestOffset - SimDeformedOffset);
+	return Pos + (SimRestOffset - OutSimHairPositionOffset);
 #endif
 
 #if PERMUTATION_DEFORMATION == DEFORMATION_OFFSET_GUIDE_DYNAMIC
@@ -96,13 +98,13 @@ float3 DisplacePosition(uint VertexIndex, float3 Pos, float Seed)
 	const FHairMeshTriangle RestTriangle = GetTriangleTransformation(RootIndex, SimRestPosition0Buffer, SimRestPosition1Buffer, SimRestPosition2Buffer);
 	const FHairMeshTriangle DeformedTriangle = GetTriangleTransformation(RootIndex, SimDeformedPosition0Buffer, SimDeformedPosition1Buffer, SimDeformedPosition2Buffer);
 
-	return DeformGuideByTriangle(VertexIndex, RestTriangle, DeformedTriangle, RootBarycentric) - SimDeformedOffset;
+	return DeformGuideByTriangle(VertexIndex, RestTriangle, DeformedTriangle, RootBarycentric) - OutSimHairPositionOffset;
 #endif
 
 #if PERMUTATION_DEFORMATION == DEFORMATION_OFFSET_GUIDE_GLOBAL
 	const float3 RestControlPoint = Pos + SimRestOffset;
 	const float3 ControlPoint = ApplyRBF(RestControlPoint, SampleCount, RestSamplePositionsBuffer, MeshSampleWeightsBuffer);
-	return ControlPoint - SimDeformedOffset;
+	return ControlPoint - OutSimHairPositionOffset;
 #endif
 }
 
diff --git a/Engine/Shaders/Private/HairStrands/HairStrandsInterpolation.usf b/Engine/Shaders/Private/HairStrands/HairStrandsInterpolation.usf
--- a/Engine/Shaders/Private/HairStrands/HairStrandsInterpolation.usf
+++ b/Engine/Shaders/Private/HairStrands/HairStrandsInterpolation.usf
@@ -33,8 +33,6 @@ uint DispatchCountX;
 uint HairDebugMode;
 float3 InRenderHairPositionOffset;
 float3 InSimHairPositionOffset;
-float3 OutRenderHairPositionOffset;
-float3 OutSimHairPositionOffset;
 int2 HairStrandsCullIndex;
 float4x4 LocalToWorldMatrix;
 
@@ -101,6 +99,9 @@ Buffer<float4>	 SimAttributeBuffer;
 RWBuffer<float4> OutRenderAttributeBuffer; 
 #endif
 
+Buffer<float4> OutSimHairPositionOffsetBuffer;
+Buffer<float4> OutRenHairPositionOffsetBuffer;
+
 #if PERMUTATION_GROUP_SIZE == 64
 	#define GROUP_SIZE 64
 #elif PERMUTATION_GROUP_SIZE == 32
@@ -109,8 +110,8 @@ RWBuffer<float4> OutRenderAttributeBuffer;
 	#error Unknown group size	
 #endif
 
-float3 ComputeStaticGeometryOffset(uint GuideIndex, float GuideVertexWeight)
-{	
+float3 ComputeStaticGeometryOffset(uint GuideIndex, float GuideVertexWeight, float3 OutSimHairPositionOffset)
+{
 	const float3 RestGuidePoint = f16tof32(SimRestPosePositionBuffer[GuideIndex].xyz) + InSimHairPositionOffset;
 
 	const float3 DeformedGuidePoint = f16tof32(DeformedSimPositionBuffer[GuideIndex].xyz) + OutSimHairPositionOffset;
@@ -133,7 +134,8 @@ float3 ComputeDynamicGeometryOffset(
 	float GuideVertexWeight, 
 	FHairMeshTriangle RestTriangle, 
 	FHairMeshTriangle DeformedTriangle,
-	float3 RootBarycentric)
+	float3 RootBarycentric,
+	float3 OutSimHairPositionOffset)
 {
 	const float3 RestGuidePoint = f16tof32(SimRestPosePositionBuffer[GuideIndex].xyz) + InSimHairPositionOffset;
 	const float3 LocalRestGuidePoint = ToTriangle(RestGuidePoint, RootBarycentric, RestTriangle);
@@ -199,6 +201,9 @@ bool ScaleAndClip(uint ControlPointId, inout uint InPackedAlpha)
 [numthreads(GROUP_SIZE, 1, 1)]
 void MainCS(uint2 DispatchThreadId : SV_DispatchThreadID)
 {
+	const float3 SimHairPositionOffset = OutSimHairPositionOffsetBuffer[0].xyz;
+	const float3 OutRenHairPositionOffset = OutRenHairPositionOffsetBuffer[0].xyz;
+
 	uint VertexIndex = DispatchThreadId.x + DispatchThreadId.y * DispatchCountX * GROUP_SIZE;
 #if PERMUTATION_CULLING
 	if (HairStrandsVF_bIsCullingEnable)
@@ -266,8 +271,8 @@ void MainCS(uint2 DispatchThreadId : SV_DispatchThreadID)
 				const uint GuideIndex0 = GuideVertexIndices[K];
 				const uint GuideIndex1 = GuideIndex0 + 1;
 
-				const float3 Offset0 = ComputeStaticGeometryOffset(GuideIndex0, GuideVertexWeights[K]);
-				const float3 Offset1 = ComputeStaticGeometryOffset(GuideIndex1, GuideVertexWeights[K]);
+				const float3 Offset0 = ComputeStaticGeometryOffset(GuideIndex0, GuideVertexWeights[K], SimHairPositionOffset);
+				const float3 Offset1 = ComputeStaticGeometryOffset(GuideIndex1, GuideVertexWeights[K], SimHairPositionOffset);
 
 				CurrOffset += lerp(Offset0, Offset1, GuideVertexLerps[K]);
 			}
@@ -276,7 +281,7 @@ void MainCS(uint2 DispatchThreadId : SV_DispatchThreadID)
 			const bool bIsValid = ScaleAndClip(VertexIndex, PackedRestPosition.w);
 			WorldVertexPosition = bIsValid ? ControlPoint + CurrOffset : float3(POSITIVE_INFINITY, POSITIVE_INFINITY, POSITIVE_INFINITY);
 
-			OutRenderDeformedPositionBuffer[VertexIndex] = float4(f32tof16(WorldVertexPosition - OutRenderHairPositionOffset), PackedRestPosition.w);
+			OutRenderDeformedPositionBuffer[VertexIndex] = float4(f32tof16(WorldVertexPosition - OutRenHairPositionOffset), PackedRestPosition.w);
 		}
 		#endif	
 
@@ -308,8 +313,8 @@ void MainCS(uint2 DispatchThreadId : SV_DispatchThreadID)
 				const FHairMeshTriangle RestTriangle		= GetTriangleTransformation(RootIndex, SimRestPosition0Buffer, SimRestPosition1Buffer, SimRestPosition2Buffer);
 				const FHairMeshTriangle DeformedTriangle	= GetTriangleTransformation(RootIndex, SimDeformedPosition0Buffer, SimDeformedPosition1Buffer, SimDeformedPosition2Buffer);
 
-				const float3 Offset0 = ComputeDynamicGeometryOffset(GuideIndex0, GuideVertexWeights[K], RestTriangle, DeformedTriangle, RootBarycentric);
-				const float3 Offset1 = ComputeDynamicGeometryOffset(GuideIndex1, GuideVertexWeights[K], RestTriangle, DeformedTriangle, RootBarycentric);
+				const float3 Offset0 = ComputeDynamicGeometryOffset(GuideIndex0, GuideVertexWeights[K], RestTriangle, DeformedTriangle, RootBarycentric, SimHairPositionOffset);
+				const float3 Offset1 = ComputeDynamicGeometryOffset(GuideIndex1, GuideVertexWeights[K], RestTriangle, DeformedTriangle, RootBarycentric, SimHairPositionOffset);
 
 				CurrOffset += VectorToWorld(lerp(Offset0, Offset1, GuideVertexLerps[K]),DeformedTriangle);
 			}
@@ -335,7 +340,7 @@ void MainCS(uint2 DispatchThreadId : SV_DispatchThreadID)
 				if (bDebugEnabled && VertexIndex==0)
 				{	
 					AddCross(InRenderHairPositionOffset, 1.2, RestColor);
-					AddCross(OutRenderHairPositionOffset, 1, DeformedColor);
+					AddCross(OutRenHairPositionOffset, 1, DeformedColor);
 				}
 			
 				if (bDebugEnabled && VertexIndex >= 0 && VertexIndex < 64)
@@ -357,7 +362,7 @@ void MainCS(uint2 DispatchThreadId : SV_DispatchThreadID)
 			}
 			#endif
 
-			OutRenderDeformedPositionBuffer[VertexIndex] = float4(f32tof16(WorldVertexPosition - OutRenderHairPositionOffset), PackedRestPosition.w);
+			OutRenderDeformedPositionBuffer[VertexIndex] = float4(f32tof16(WorldVertexPosition - OutRenHairPositionOffset), PackedRestPosition.w);
 		}
 		#endif	
 
@@ -400,7 +405,7 @@ void MainCS(uint2 DispatchThreadId : SV_DispatchThreadID)
 			const bool bIsValid = ScaleAndClip(VertexIndex, PackedRestPosition.w);
 			WorldVertexPosition = bIsValid ? ControlPoint + CurrOffset : float3(POSITIVE_INFINITY, POSITIVE_INFINITY, POSITIVE_INFINITY);
 
-			OutRenderDeformedPositionBuffer[VertexIndex] = float4(f32tof16(WorldVertexPosition - OutRenderHairPositionOffset), PackedRestPosition.w);
+			OutRenderDeformedPositionBuffer[VertexIndex] = float4(f32tof16(WorldVertexPosition - OutRenHairPositionOffset), PackedRestPosition.w);
 		}
 		#endif	
 
@@ -438,8 +443,8 @@ void MainCS(uint2 DispatchThreadId : SV_DispatchThreadID)
 				
 				const float3 LocalPoint = VectorToTriangle(ControlPoint,RestTriangleG);
 
-				const float3 Offset0 = ComputeRotateGeometryOffset(GuideIndex0, 1.0, RestTriangleG, DeformedTriangleG);
-				const float3 Offset1 = ComputeRotateGeometryOffset(GuideIndex1, 1.0, RestTriangleG, DeformedTriangleG);
+				const float3 Offset0 = ComputeRotateGeometryOffset(GuideIndex0, 1.0, RestTriangleG, DeformedTriangleG/*, SimHairPositionOffset*/);
+				const float3 Offset1 = ComputeRotateGeometryOffset(GuideIndex1, 1.0, RestTriangleG, DeformedTriangleG/*, SimHairPositionOffset*/);
 
 				CurrOffset += GuideVertexWeights[K] * VectorToWorld(lerp(Offset0, Offset1, GuideVertexLerps[K])+LocalPoint,DeformedTriangleG);
 			}
@@ -450,7 +455,7 @@ void MainCS(uint2 DispatchThreadId : SV_DispatchThreadID)
 			const bool bIsValid = ScaleAndClip(VertexIndex, PackedRestPosition.w);
 			WorldVertexPosition = bIsValid ? ControlPoint + CurrOffset : float3(POSITIVE_INFINITY, POSITIVE_INFINITY, POSITIVE_INFINITY);
 
-			OutRenderDeformedPositionBuffer[VertexIndex] = float4(f32tof16(WorldVertexPosition - OutRenderHairPositionOffset), PackedRestPosition.w);
+			OutRenderDeformedPositionBuffer[VertexIndex] = float4(f32tof16(WorldVertexPosition - OutRenHairPositionOffset), PackedRestPosition.w);
 		}
 		#endif
 
@@ -487,8 +492,8 @@ void MainCS(uint2 DispatchThreadId : SV_DispatchThreadID)
 				
 				const float3 LocalPoint = ToTriangle(ControlPoint,RootBarycentricG,RestTriangleG);
 
-				const float3 Offset0 = ComputeDynamicGeometryOffset(GuideIndex0, 1.0, RestTriangleG, DeformedTriangleG, RootBarycentricG);
-				const float3 Offset1 = ComputeDynamicGeometryOffset(GuideIndex1, 1.0, RestTriangleG, DeformedTriangleG, RootBarycentricG);
+				const float3 Offset0 = ComputeDynamicGeometryOffset(GuideIndex0, 1.0, RestTriangleG, DeformedTriangleG, RootBarycentricG, SimHairPositionOffset);
+				const float3 Offset1 = ComputeDynamicGeometryOffset(GuideIndex1, 1.0, RestTriangleG, DeformedTriangleG, RootBarycentricG, SimHairPositionOffset);
 
 				CurrOffset += GuideVertexWeights[K] * ToWorld(lerp(Offset0, Offset1, GuideVertexLerps[K])+LocalPoint,RootBarycentricG,DeformedTriangleG);
 			}
@@ -497,7 +502,7 @@ void MainCS(uint2 DispatchThreadId : SV_DispatchThreadID)
 			const bool bIsValid = ScaleAndClip(VertexIndex, PackedRestPosition.w);
 			WorldVertexPosition = bIsValid ? CurrOffset : float3(POSITIVE_INFINITY, POSITIVE_INFINITY, POSITIVE_INFINITY);
 
-			OutRenderDeformedPositionBuffer[VertexIndex] = float4(f32tof16(WorldVertexPosition - OutRenderHairPositionOffset), PackedRestPosition.w);
+			OutRenderDeformedPositionBuffer[VertexIndex] = float4(f32tof16(WorldVertexPosition - OutRenHairPositionOffset), PackedRestPosition.w);
 		}
 		#endif
 
@@ -574,7 +579,7 @@ void MainCS(uint2 DispatchThreadId : SV_DispatchThreadID)
 
 uint DispatchCountX;
 uint ClusterCount;
-float3 OutHairPositionOffset;
+//float3 OutHairPositionOffset;
 float4x4 LocalToWorldMatrix;
 
 Buffer<uint4> RenderDeformedPositionBuffer;
@@ -583,6 +588,8 @@ Buffer<uint> ClusterIdBuffer;
 Buffer<uint> ClusterIndexOffsetBuffer;
 Buffer<uint> ClusterIndexCountBuffer;
 
+Buffer<float4> RenderDeformedOffsetBuffer;
+
 RWBuffer<int> OutClusterAABBBuffer; 	// Cluster data packed as {uint3 AABBMin, uint3 AABBMax}
 RWBuffer<int> OutGroupAABBBuffer; 		// Group data packed as {uint3 AABBMin, uint3 AABBMax}
 
@@ -616,6 +623,8 @@ void ClusterAABBEvaluationCS(uint3 DispatchThreadId : SV_DispatchThreadID, uint3
 
 		if (ThreadReadOffset < IndexCount)
 		{
+			const float3 OutHairPositionOffset = RenderDeformedOffsetBuffer[0];
+
 			const uint VertexIndex = ClusterVertexIdBuffer[IndexOffset + ThreadReadOffset];
 			const FHairControlPoint CP = UnpackHairControlPoint(RenderDeformedPositionBuffer[VertexIndex], OutHairPositionOffset, 1, 1);
 
diff --git a/Engine/Shaders/Private/HairStrands/HairStrandsMeshUpdate.usf b/Engine/Shaders/Private/HairStrands/HairStrandsMeshUpdate.usf
--- a/Engine/Shaders/Private/HairStrands/HairStrandsMeshUpdate.usf
+++ b/Engine/Shaders/Private/HairStrands/HairStrandsMeshUpdate.usf
@@ -3,6 +3,8 @@
 #include "../Common.ush"
 #include "HairStrandsMeshProjectionCommon.ush"  
 
+#if SHADER_MESH_UPDATE
+
 #ifndef MAX_SECTION_COUNT
 #error MAX_SECTION_COUNT needs to be defined
 #endif
@@ -200,4 +202,29 @@ void MainCS(uint3 DispatchThreadId : SV_DispatchThreadID)
 	OutRootTrianglePosition1[RootIndex] = (float4(Tri.P1, 1));
 	OutRootTrianglePosition2[RootIndex] = (float4(Tri.P2, 1));
 #endif
-}
\ No newline at end of file
+}
+
+#endif // SHADER_MESH_UPDATE
+
+#if SHADER_OFFSET_UPDATE
+
+float3 CPUPositionOffset;
+uint bUseCPUOffset;
+Buffer<float4> RootTrianglePosition0Buffer;
+RWBuffer<float4> OutOffsetBuffer;
+
+[numthreads(1, 1, 1)]
+void MainCS(uint3 DispatchThreadId : SV_DispatchThreadID)
+{
+	if (DispatchThreadId.x > 0)
+		return;
+
+	float3 RefPosition = CPUPositionOffset;
+	if (bUseCPUOffset == 0)
+	{
+		RefPosition = RootTrianglePosition0Buffer[0];
+	}
+	OutOffsetBuffer[0] = float4(RefPosition, 1);
+}
+
+#endif // SHADER_OFFSET_UPDATE
\ No newline at end of file
diff --git a/Engine/Shaders/Private/HairStrands/HairStrandsRaytracingGeometry.usf b/Engine/Shaders/Private/HairStrands/HairStrandsRaytracingGeometry.usf
--- a/Engine/Shaders/Private/HairStrands/HairStrandsRaytracingGeometry.usf
+++ b/Engine/Shaders/Private/HairStrands/HairStrandsRaytracingGeometry.usf
@@ -8,6 +8,7 @@ uint DispatchCountX;
 float3 StrandHairWorldOffset;
 float StrandHairRadius;
 
+Buffer<float4>	 PositionOffsetBuffer;
 Buffer<uint4> 	 PositionBuffer;
 RWBuffer<float4> OutputPositionBuffer;
 
@@ -36,7 +37,9 @@ FControlPoint GetControlPoint(uint Index)
 		const uint4 PackedData = PositionBuffer[Index];
 		const uint PackedAlpha = asuint(PackedData.a);
 
-		Out.Position 	= f16tof32(PackedData.xyz) + StrandHairWorldOffset;
+		const float3 OutPositionOffset = PositionOffsetBuffer[0].xyz;
+
+		Out.Position 	= f16tof32(PackedData.xyz) + OutPositionOffset;
 		Out.PointType 	= PackedAlpha & 0x3;
 		Out.WorldRadius	= StrandHairRadius * float((PackedAlpha & 0x00FC) >> 2) / 63.0f;
 	}
diff --git a/Engine/Shaders/Private/HairStrands/HairStrandsVertexFactory.ush b/Engine/Shaders/Private/HairStrands/HairStrandsVertexFactory.ush
--- a/Engine/Shaders/Private/HairStrands/HairStrandsVertexFactory.ush
+++ b/Engine/Shaders/Private/HairStrands/HairStrandsVertexFactory.ush
@@ -37,13 +37,15 @@
 // This is set by compilation enviromenent of the vertex factory
 ///////////////////////////////////////////////////////////////////////////////
 
-float3 HairStrandsVF_PositionOffset;			// Vertex position (3x16bits) are relative to this offset for improving precision while keeping vertex footprint low
-float3 HairStrandsVF_PreviousPositionOffset;	// Vertex position (3x16bits) are relative to this offset for improving precision while keeping vertex footprint low
 float  HairStrandsVF_Density;
 uint   HairStrandsVF_CullingEnable;
 uint   HairStrandsVF_bUseStableRasterization;
 uint   HairStrandsVF_bScatterSceneLighing;
 
+// Vertex position (3x16bits) are relative to this offset for improving precision while keeping vertex footprint low
+Buffer<float4> HairStrandsVF_PositionOffsetBuffer;
+Buffer<float4> HairStrandsVF_PreviousPositionOffsetBuffer;
+
 Buffer<uint4>  HairStrandsVF_PreviousPositionBuffer;
 Buffer<float4> HairStrandsVF_TangentBuffer;
 
@@ -188,10 +190,12 @@ uint GetHairPrimitiveId(FVertexFactoryInput Input)
 
 FHairControlPoint GetVertexPosition(const FVertexInfo VertexInfo, bool bInvalidJointVertex=true)
 {
+	const float3 PositionOffset = HairStrandsVF_PositionOffsetBuffer[0].xyz;
+
 	const uint4 PackedData = HairStrandsVF_PositionBuffer[VertexInfo.VertexIndex];
 	FHairControlPoint Out = UnpackHairControlPoint(
 		PackedData, 
-		HairStrandsVF_PositionOffset, 
+		PositionOffset,
 		HairStrandsVF_Radius * VertexInfo.LodRadiusScale, 
 		HairStrandsVF_Length);
 	
@@ -221,11 +225,13 @@ float GetWorldStrandLength(FVertexFactoryInput Input)
 
 float3 GetVertexPreviousPosition(FVertexFactoryInput Input, bool bInvalidJointVertex = true)
 {
+	const float3 PreviousPositionOffset = HairStrandsVF_PreviousPositionOffsetBuffer[0].xyz;
+
 	FVertexInfo VertexInfo = GetVertexInfo(Input);
 	const uint4 PackedData = HairStrandsVF_PreviousPositionBuffer[VertexInfo.VertexIndex];
 	FHairControlPoint Out = UnpackHairControlPoint(
 		PackedData,
-		HairStrandsVF_PreviousPositionOffset,
+		PreviousPositionOffset,
 		HairStrandsVF_Radius * VertexInfo.LodRadiusScale,
 		HairStrandsVF_Length);
 
diff --git a/Engine/Shaders/Private/HairStrands/HairStrandsVisibilityRasterCompute.usf b/Engine/Shaders/Private/HairStrands/HairStrandsVisibilityRasterCompute.usf
--- a/Engine/Shaders/Private/HairStrands/HairStrandsVisibilityRasterCompute.usf
+++ b/Engine/Shaders/Private/HairStrands/HairStrandsVisibilityRasterCompute.usf
@@ -81,6 +81,7 @@ float3			HairStrandsVF_PositionOffset;
 uint			HairStrandsVF_VertexCount;
 float4x4		HairStrandsVF_LocalToWorldPrimitiveTransform;
 Buffer<uint4>	HairStrandsVF_PositionBuffer;
+Buffer<float4>	HairStrandsVF_PositionOffsetBuffer;
 
 #if PERMUTATION_CULLING == 1
 Buffer<uint>	HairStrandsVF_CullingIndirectBuffer;
@@ -274,9 +275,11 @@ void MainCS(uint3 DispatchThreadID : SV_DispatchThreadID)
 
 	ResolvedView = ResolveView();
 
+	const float3 PositionOffset = HairStrandsVF_PositionOffsetBuffer[0].xyz;
+
 	const FHairControlPoint CP0 = UnpackHairControlPoint(
 		HairStrandsVF_PositionBuffer[VertexIndex0],
-		HairStrandsVF_PositionOffset,
+		PositionOffset,
 		HairStrandsVF_Radius * LodRadiusScale0,
 		HairStrandsVF_Length);
 
@@ -285,7 +288,7 @@ void MainCS(uint3 DispatchThreadID : SV_DispatchThreadID)
 
 	const FHairControlPoint CP1 = UnpackHairControlPoint(
 		HairStrandsVF_PositionBuffer[VertexIndex1],
-		HairStrandsVF_PositionOffset,
+		PositionOffset,
 		HairStrandsVF_Radius * LodRadiusScale1,
 		HairStrandsVF_Length);
 
diff --git a/Engine/Shaders/Private/HairStrands/HairStrandsVoxelRasterCompute.usf b/Engine/Shaders/Private/HairStrands/HairStrandsVoxelRasterCompute.usf
--- a/Engine/Shaders/Private/HairStrands/HairStrandsVoxelRasterCompute.usf
+++ b/Engine/Shaders/Private/HairStrands/HairStrandsVoxelRasterCompute.usf
@@ -19,6 +19,7 @@ float3			HairStrandsVF_PositionOffset;
 uint			HairStrandsVF_VertexCount;
 float4x4		HairStrandsVF_LocalToWorldPrimitiveTransform;
 Buffer<uint4>	HairStrandsVF_PositionBuffer;
+Buffer<float4>	HairStrandsVF_PositionOffsetBuffer;
 
 #if PERMUTATION_CULLING == 1
 Buffer<uint>	HairStrandsVF_CullingIndirectBuffer;
@@ -96,9 +97,11 @@ void MainCS(uint3 DispatchThreadID : SV_DispatchThreadID)
 	}
 #endif
 
+	const float4 PositionOffset = HairStrandsVF_PositionOffsetBuffer[0];
+
 	const FHairControlPoint CP0 = UnpackHairControlPoint(
 		HairStrandsVF_PositionBuffer[VertexIndex0],
-		HairStrandsVF_PositionOffset,
+		PositionOffset,
 		HairStrandsVF_Radius * LodRadiusScale0,
 		HairStrandsVF_Length);
 
@@ -107,7 +110,7 @@ void MainCS(uint3 DispatchThreadID : SV_DispatchThreadID)
 
 	const FHairControlPoint CP1 = UnpackHairControlPoint(
 		HairStrandsVF_PositionBuffer[VertexIndex1],
-		HairStrandsVF_PositionOffset,
+		PositionOffset,
 		HairStrandsVF_Radius * LodRadiusScale1,
 		HairStrandsVF_Length);
 
diff --git a/Engine/Source/Runtime/Renderer/Private/HairStrands/HairStrandsVisibility.cpp b/Engine/Source/Runtime/Renderer/Private/HairStrands/HairStrandsVisibility.cpp
--- a/Engine/Source/Runtime/Renderer/Private/HairStrands/HairStrandsVisibility.cpp
+++ b/Engine/Source/Runtime/Renderer/Private/HairStrands/HairStrandsVisibility.cpp
@@ -2942,10 +2942,10 @@ class FVisiblityRasterComputeCS : public FGlobalShader
 		SHADER_PARAMETER(float, HairStrandsVF_Radius)
 		SHADER_PARAMETER(float, HairStrandsVF_Length)
 		SHADER_PARAMETER(uint32, HairStrandsVF_bUseStableRasterization)
-		SHADER_PARAMETER(FVector, HairStrandsVF_PositionOffset)
 		SHADER_PARAMETER(uint32, HairStrandsVF_VertexCount)
 		SHADER_PARAMETER(FMatrix, HairStrandsVF_LocalToWorldPrimitiveTransform)
 		SHADER_PARAMETER_SRV(Buffer, HairStrandsVF_PositionBuffer)
+		SHADER_PARAMETER_SRV(Buffer, HairStrandsVF_PositionOffsetBuffer)
 		SHADER_PARAMETER_RDG_BUFFER_SRV(Buffer, HairStrandsVF_CullingIndirectBuffer)
 		SHADER_PARAMETER_RDG_BUFFER_SRV(Buffer, HairStrandsVF_CullingIndexBuffer)
 		SHADER_PARAMETER_RDG_BUFFER_SRV(Buffer, HairStrandsVF_CullingRadiusScaleBuffer)
@@ -3124,7 +3124,7 @@ static FRasterComputeOutput AddVisibilityComputeRasterPass(
 
 			const FHairGroupPublicData::FVertexFactoryInput& VFInput = HairGroupPublicData->VFInput;
 			PassParameters->HairStrandsVF_PositionBuffer	= VFInput.Strands.PositionBuffer;
-			PassParameters->HairStrandsVF_PositionOffset	= VFInput.Strands.PositionOffset;
+			PassParameters->HairStrandsVF_PositionOffsetBuffer = VFInput.Strands.PositionOffsetBuffer;
 			PassParameters->HairStrandsVF_VertexCount		= VFInput.Strands.VertexCount;
 			PassParameters->HairStrandsVF_Radius			= VFInput.Strands.HairRadius;
 			PassParameters->HairStrandsVF_Length			= VFInput.Strands.HairLength;
diff --git a/Engine/Source/Runtime/Renderer/Private/HairStrands/HairStrandsVoxelization.cpp b/Engine/Source/Runtime/Renderer/Private/HairStrands/HairStrandsVoxelization.cpp
--- a/Engine/Source/Runtime/Renderer/Private/HairStrands/HairStrandsVoxelization.cpp
+++ b/Engine/Source/Runtime/Renderer/Private/HairStrands/HairStrandsVoxelization.cpp
@@ -993,6 +993,7 @@ class FVoxelRasterComputeCS : public FGlobalShader
 		SHADER_PARAMETER(uint32,  HairStrandsVF_VertexCount)
 		SHADER_PARAMETER(FMatrix, HairStrandsVF_LocalToWorldPrimitiveTransform)
 		SHADER_PARAMETER_SRV(Buffer, HairStrandsVF_PositionBuffer)
+		SHADER_PARAMETER_SRV(Buffer, HairStrandsVF_PositionOffsetBuffer)
 		SHADER_PARAMETER_RDG_BUFFER_SRV(Buffer, HairStrandsVF_CullingIndirectBuffer)
 		SHADER_PARAMETER_RDG_BUFFER_SRV(Buffer, HairStrandsVF_CullingIndexBuffer)
 		SHADER_PARAMETER_RDG_BUFFER_SRV(Buffer, HairStrandsVF_CullingRadiusScaleBuffer)
@@ -1060,6 +1061,7 @@ static void AddVirtualVoxelizationComputeRasterPass(
 
 				PassParameters->HairStrandsVF_PositionBuffer = VFInput.Strands.PositionBuffer;
 				PassParameters->HairStrandsVF_PositionOffset = VFInput.Strands.PositionOffset;
+				PassParameters->HairStrandsVF_PositionOffsetBuffer = VFInput.Strands.PositionOffsetBuffer;
 				PassParameters->HairStrandsVF_VertexCount = VFInput.Strands.VertexCount;
 				PassParameters->HairStrandsVF_Radius = VFInput.Strands.HairRadius;
 				PassParameters->HairStrandsVF_Length = VFInput.Strands.HairLength;
diff --git a/Engine/Source/Runtime/Renderer/Public/HairStrandsInterface.h b/Engine/Source/Runtime/Renderer/Public/HairStrandsInterface.h
--- a/Engine/Source/Runtime/Renderer/Public/HairStrandsInterface.h
+++ b/Engine/Source/Runtime/Renderer/Public/HairStrandsInterface.h
@@ -196,6 +196,10 @@ public:
 
 			FVector PositionOffset = FVector::ZeroVector;
 			FVector PrevPositionOffset = FVector::ZeroVector;
+
+			FShaderResourceViewRHIRef PositionOffsetBuffer = nullptr;
+			FShaderResourceViewRHIRef PrevPositionOffsetBuffer = nullptr;
+
 			uint32 VertexCount = 0;
 			float HairRadius = 0;
 			float HairLength = 0;
