Expose variable to force a groom component to use cards/meshes geometry instead of strands.

diff --git a/Engine/Plugins/Runtime/HairStrands/Source/HairStrandsCore/Public/GroomComponent.h b/Engine/Plugins/Runtime/HairStrands/Source/HairStrandsCore/Public/GroomComponent.h
--- a/Engine/Plugins/Runtime/HairStrands/Source/HairStrandsCore/Public/GroomComponent.h
+++ b/Engine/Plugins/Runtime/HairStrands/Source/HairStrandsCore/Public/GroomComponent.h
@@ -165,6 +165,10 @@ public:
 	UPROPERTY(EditAnywhere, BlueprintReadWrite, Category = "Groom")
 	TArray<FHairGroupDesc> GroomGroupsDesc;
 
+	/** Force the groom to use cards/meshes geometry instead of strands. */
+	UPROPERTY(EditAnywhere, BlueprintReadWrite, Category = "Groom")
+	bool bUseCards = false;
+
 	/** Hair group instance access */
 	uint32 GetGroupCount() const { return HairGroupInstances.Num();  }
 	FHairGroupInstance* GetGroupInstance(uint32 Index) { return Index < uint32(HairGroupInstances.Num()) ? HairGroupInstances[Index] : nullptr; }
@@ -205,7 +209,7 @@ private:
 	bool bIsGroomBindingAssetCallbackRegistered;
 	int32 PredictedLODIndex = -1;
 	bool bValidationEnable = true;
-	bool bUseCards = false;
+	
 
 	EWorldType::Type GetWorldType() const; 
 	void InitResources(bool bIsBindingReloading=false);
