Fix hair cards/meshes velocity output.

diff --git a/Engine/Plugins/Runtime/HairStrands/Source/HairStrandsCore/Private/GroomComponent.cpp b/Engine/Plugins/Runtime/HairStrands/Source/HairStrandsCore/Private/GroomComponent.cpp
--- a/Engine/Plugins/Runtime/HairStrands/Source/HairStrandsCore/Private/GroomComponent.cpp
+++ b/Engine/Plugins/Runtime/HairStrands/Source/HairStrandsCore/Private/GroomComponent.cpp
@@ -830,10 +830,9 @@ public:
 		bool bHasPrecomputedVolumetricLightmap;
 		FMatrix PreviousLocalToWorld;
 		int32 SingleCaptureIndex;
-		bool bOutputVelocity = false;
-		bool bDrawVelocity = false; // Velocity vector is done in a custom fashion
+		bool bOutputVelocity = GeometryType == EHairGeometryType::Cards || GeometryType == EHairGeometryType::Meshes;
+		bool bDrawVelocity = bOutputVelocity; // Velocity vector is done in a custom fashion
 		GetScene().GetPrimitiveUniformShaderParameters_RenderThread(GetPrimitiveSceneInfo(), bHasPrecomputedVolumetricLightmap, PreviousLocalToWorld, SingleCaptureIndex, bOutputVelocity);
-		bOutputVelocity = false;
 		FDynamicPrimitiveUniformBuffer& DynamicPrimitiveUniformBuffer = Collector.AllocateOneFrameResource<FDynamicPrimitiveUniformBuffer>();
 		DynamicPrimitiveUniformBuffer.Set(GetLocalToWorld(), PreviousLocalToWorld, GetBounds(), GetLocalBounds(), true, false, bDrawVelocity, bOutputVelocity);
 		BatchElement.PrimitiveUniformBufferResource = &DynamicPrimitiveUniformBuffer.UniformBuffer;
