Add support for selecting attachment type per LOD for groom asset.

Attachement can be rigid (attached to a bone) or skinned (attached to a surface).

diff --git a/Engine/Plugins/Runtime/HairStrands/Source/HairStrandsCore/Private/GroomComponent.cpp b/Engine/Plugins/Runtime/HairStrands/Source/HairStrandsCore/Private/GroomComponent.cpp
--- a/Engine/Plugins/Runtime/HairStrands/Source/HairStrandsCore/Private/GroomComponent.cpp
+++ b/Engine/Plugins/Runtime/HairStrands/Source/HairStrandsCore/Private/GroomComponent.cpp
@@ -574,6 +574,10 @@ public:
 		for (FHairGroupInstance* Instance : HairGroupInstances)
 		{
 			Instance->LocalToWorld = HairLocalToWorld;
+			if (Instance->BindingType == EHairBindingType::Skinning)
+			{
+				Instance->LocalToWorld = Instance->Debug.SkeletalLocalToWorld;
+			}
 		}
 	}
 
@@ -833,8 +837,16 @@ public:
 		bool bOutputVelocity = GeometryType == EHairGeometryType::Cards || GeometryType == EHairGeometryType::Meshes;
 		bool bDrawVelocity = bOutputVelocity; // Velocity vector is done in a custom fashion
 		GetScene().GetPrimitiveUniformShaderParameters_RenderThread(GetPrimitiveSceneInfo(), bHasPrecomputedVolumetricLightmap, PreviousLocalToWorld, SingleCaptureIndex, bOutputVelocity);
+
+		FMatrix OverrideLocalToWorld = GetLocalToWorld();
+		if (Instance->BindingType == EHairBindingType::Skinning)
+		{
+			OverrideLocalToWorld = Instance->Debug.SkeletalLocalToWorld.ToMatrixWithScale();
+			PreviousLocalToWorld = Instance->Debug.SkeletalPreviousLocalToWorld.ToMatrixWithScale();
+		}
+
 		FDynamicPrimitiveUniformBuffer& DynamicPrimitiveUniformBuffer = Collector.AllocateOneFrameResource<FDynamicPrimitiveUniformBuffer>();
-		DynamicPrimitiveUniformBuffer.Set(GetLocalToWorld(), PreviousLocalToWorld, GetBounds(), GetLocalBounds(), true, false, bDrawVelocity, bOutputVelocity);
+		DynamicPrimitiveUniformBuffer.Set(OverrideLocalToWorld, PreviousLocalToWorld, GetBounds(), GetLocalBounds(), true, false, bDrawVelocity, bOutputVelocity);
 		BatchElement.PrimitiveUniformBufferResource = &DynamicPrimitiveUniformBuffer.UniformBuffer;
 
 		BatchElement.FirstIndex = 0;
@@ -1764,6 +1776,18 @@ static USkeletalMeshComponent* ValidateBindingAsset(
 
 void CreateHairStrandsDebugAttributeBuffer(FRDGExternalBuffer* DebugAttributeBuffer, uint32 VertexCount);
 
+
+static EHairBindingType ToHairBindingType(EGroomBindingType In)
+{
+	switch (In)
+	{
+	case EGroomBindingType::Rigid:		return EHairBindingType::Rigid;
+	case EGroomBindingType::Skinning:	return EHairBindingType::Skinning;
+	case EGroomBindingType::NoneBinding:return EHairBindingType::NoneBinding;
+	}
+	return EHairBindingType::NoneBinding;
+}
+
 void UGroomComponent::InitResources(bool bIsBindingReloading)
 {
 	LLM_SCOPE(ELLMTag::Meshes) // This should be a Groom LLM tag, but there is no LLM tag bit left
@@ -1781,18 +1805,32 @@ void UGroomComponent::InitResources(bool bIsBindingReloading)
 
 	InitializedResources = GroomAsset;
 
+	// Initialize LOD screen size & visibility
+	bool bHasNeedBindingData = false;
+	for (int32 GroupIt = 0, GroupCount = GroomAsset->HairGroupsData.Num(); GroupIt < GroupCount; ++GroupIt)
+	{
+		const FHairGroupsLOD& GroupLOD = GroomAsset->HairGroupsLOD[GroupIt];
+		for (const FHairLODSettings& LODSettings : GroupLOD.LODs)
+		{
+			bHasNeedBindingData = bHasNeedBindingData || LODSettings.BindingType == EGroomBindingType::Skinning;
+		}
+	}
+
 	const FPrimitiveComponentId LocalComponentId = ComponentId;
 	const EWorldType::Type WorldType = GetWorldType();
 	const bool bIsStrandsEnabled = IsHairStrandsEnabled(EHairStrandsShaderType::Strands);
 
 	// Insure that the binding asset is compatible, otherwise no binding will be used
-	USkeletalMeshComponent* SkeletalMeshComponent = ValidateBindingAsset(GroomAsset, BindingAsset, GetAttachParent() ? Cast<USkeletalMeshComponent>(GetAttachParent()) : nullptr, bIsBindingReloading, bValidationEnable, this);
+	USkeletalMeshComponent* SkeletalMeshComponent = GetAttachParent() ? Cast<USkeletalMeshComponent>(GetAttachParent()) : nullptr;
+	SkeletalMeshComponent = ValidateBindingAsset(GroomAsset, BindingAsset, SkeletalMeshComponent, bIsBindingReloading, bValidationEnable, this);
+	UGroomBindingAsset* LocalBindingAsset = bHasNeedBindingData ? BindingAsset : nullptr;
 
 	// Insure the ticking of the Groom component always happens after the skeletalMeshComponent.
 	if (SkeletalMeshComponent)
 	{
 		RegisteredSkeletalMeshComponent = SkeletalMeshComponent;
 		AddTickPrerequisiteComponent(SkeletalMeshComponent);
+		bUseAttachParentBound = true;
 	}
 
 	FTransform HairLocalToWorld = GetComponentTransform();
@@ -1824,6 +1862,25 @@ void UGroomComponent::InitResources(bool bIsBindingReloading)
 
 		const bool bDynamicResources = IsHairStrandsSimulationEnable() || IsHairStrandsBindingEnable();
 
+		// Initialize LOD screen size & visibility
+		{
+			HairGroupInstance->HairGroupPublicData = new FHairGroupPublicData(GroupIt);
+			TArray<float> CPULODScreenSize;
+			TArray<bool> LODVisibility;
+			TArray<EHairGeometryType> LODGeometryTypes;
+			const FHairGroupsLOD& GroupLOD = GroomAsset->HairGroupsLOD[GroupIt];
+			for (const FHairLODSettings& LODSettings : GroupLOD.LODs)
+			{
+				CPULODScreenSize.Add(LODSettings.ScreenSize);
+				LODVisibility.Add(LODSettings.bVisible);
+				LODGeometryTypes.Add(ConvertToHairGeometryType(LODSettings.GeometryType));
+				HairGroupInstance->HairGroupPublicData->BindingTypes.Add(ToHairBindingType(LODSettings.BindingType));
+			}
+			HairGroupInstance->HairGroupPublicData->SetLODScreenSizes(CPULODScreenSize);
+			HairGroupInstance->HairGroupPublicData->SetLODVisibilities(LODVisibility);
+			HairGroupInstance->HairGroupPublicData->SetLODGeometryTypes(LODGeometryTypes);
+		}
+
 		// Sim data 
 		// Simulation guides are used for two purposes:
 		// * Physics simulation
@@ -1833,14 +1890,14 @@ void UGroomComponent::InitResources(bool bIsBindingReloading)
 		{
 			HairGroupInstance->Guides.Data = &GroupData.Guides.Data;
 
-			if (SkeletalMeshComponent && BindingAsset)
+			if (SkeletalMeshComponent && LocalBindingAsset)
 			{
-				check(BindingAsset);
-				check(GroupIt < BindingAsset->HairGroupResources.Num());
-				check(SkeletalMeshComponent->SkeletalMesh ? SkeletalMeshComponent->SkeletalMesh->GetLODInfoArray().Num() == BindingAsset->HairGroupResources[GroupIt].SimRootResources->RootData.MeshProjectionLODs.Num() : false);
+				check(LocalBindingAsset);
+				check(GroupIt < LocalBindingAsset->HairGroupResources.Num());
+				check(SkeletalMeshComponent->SkeletalMesh ? SkeletalMeshComponent->SkeletalMesh->GetLODInfoArray().Num() == LocalBindingAsset->HairGroupResources[GroupIt].SimRootResources->RootData.MeshProjectionLODs.Num() : false);
 
 				HairGroupInstance->Guides.bOwnRootResourceAllocation = false;
-				HairGroupInstance->Guides.RestRootResource = BindingAsset->HairGroupResources[GroupIt].SimRootResources;
+				HairGroupInstance->Guides.RestRootResource = LocalBindingAsset->HairGroupResources[GroupIt].SimRootResources;
 
 				HairGroupInstance->Guides.DeformedRootResource = new FHairStrandsDeformedRootResource(HairGroupInstance->Guides.RestRootResource);
 				BeginInitResource(HairGroupInstance->Guides.DeformedRootResource);
@@ -1858,26 +1915,9 @@ void UGroomComponent::InitResources(bool bIsBindingReloading)
 				(GroomAsset->HairGroupsPhysics[GroupIt].SolverSettings.EnableSimulation) :
 				false;
 
-			HairGroupInstance->Guides.bHasGlobalInterpolation = BindingAsset && GroomAsset && GroomAsset->EnableGlobalInterpolation;
+			HairGroupInstance->Guides.bHasGlobalInterpolation = LocalBindingAsset && GroomAsset && GroomAsset->EnableGlobalInterpolation;
 		}
 
-		// Initialize LOD screen size & visibility
-		{
-			HairGroupInstance->HairGroupPublicData = new FHairGroupPublicData(GroupIt);
-			TArray<float> CPULODScreenSize;
-			TArray<bool> LODVisibility;
-			TArray<EHairGeometryType> LODGeometryTypes;
-			const FHairGroupsLOD& GroupLOD = GroomAsset->HairGroupsLOD[GroupIt];
-			for (const FHairLODSettings& LODSettings : GroupLOD.LODs)
-			{
-				CPULODScreenSize.Add(LODSettings.ScreenSize);
-				LODVisibility.Add(LODSettings.bVisible);
-				LODGeometryTypes.Add(ConvertToHairGeometryType(LODSettings.GeometryType));
-			}
-			HairGroupInstance->HairGroupPublicData->SetLODScreenSizes(CPULODScreenSize);
-			HairGroupInstance->HairGroupPublicData->SetLODVisibilities(LODVisibility);
-			HairGroupInstance->HairGroupPublicData->SetLODGeometryTypes(LODGeometryTypes);
-		}
 
 		// LODBias is in the Modifier which is needed for LOD selection regardless if the strands are there or not
 		HairGroupInstance->Strands.Modifier = GetGroomGroupsDesc(GroomAsset, this, GroupIt);
@@ -1902,14 +1942,14 @@ void UGroomComponent::InitResources(bool bIsBindingReloading)
 			}
 			#endif
 
-			if (SkeletalMeshComponent && BindingAsset && IsHairStrandsBindingEnable())
+			if (SkeletalMeshComponent && LocalBindingAsset && IsHairStrandsBindingEnable())
 			{
-				check(BindingAsset);
-				check(GroupIt < BindingAsset->HairGroupResources.Num());
-				check(SkeletalMeshComponent->SkeletalMesh ? SkeletalMeshComponent->SkeletalMesh->GetLODInfoArray().Num() == BindingAsset->HairGroupResources[GroupIt].RenRootResources->RootData.MeshProjectionLODs.Num() : false);
+				check(LocalBindingAsset);
+				check(GroupIt < LocalBindingAsset->HairGroupResources.Num());
+				check(SkeletalMeshComponent->SkeletalMesh ? SkeletalMeshComponent->SkeletalMesh->GetLODInfoArray().Num() == LocalBindingAsset->HairGroupResources[GroupIt].RenRootResources->RootData.MeshProjectionLODs.Num() : false);
 
 				HairGroupInstance->Strands.bOwnRootResourceAllocation = false;
-				HairGroupInstance->Strands.RestRootResource = BindingAsset->HairGroupResources[GroupIt].RenRootResources;
+				HairGroupInstance->Strands.RestRootResource = LocalBindingAsset->HairGroupResources[GroupIt].RenRootResources;
 
 				HairGroupInstance->Strands.DeformedRootResource = new FHairStrandsDeformedRootResource(HairGroupInstance->Strands.RestRootResource);
 				BeginInitResource(HairGroupInstance->Strands.DeformedRootResource);
@@ -1989,13 +2029,13 @@ void UGroomComponent::InitResources(bool bIsBindingReloading)
 					InstanceLOD.Guides.InterpolationData = &LOD.Guides.InterpolationData;
 					InstanceLOD.Guides.InterpolationResource = LOD.Guides.InterpolationResource;
 
-					if (SkeletalMeshComponent && BindingAsset)
+					if (SkeletalMeshComponent && LocalBindingAsset)
 					{
-						check(GroupIt < BindingAsset->HairGroupResources.Num());
-						check(SkeletalMeshComponent->SkeletalMesh ? SkeletalMeshComponent->SkeletalMesh->GetLODInfoArray().Num() == BindingAsset->HairGroupResources[GroupIt].CardsRootResources[CardsLODIndex]->RootData.MeshProjectionLODs.Num() : false);
+						check(GroupIt < LocalBindingAsset->HairGroupResources.Num());
+						check(SkeletalMeshComponent->SkeletalMesh ? SkeletalMeshComponent->SkeletalMesh->GetLODInfoArray().Num() == LocalBindingAsset->HairGroupResources[GroupIt].CardsRootResources[CardsLODIndex]->RootData.MeshProjectionLODs.Num() : false);
 
 						InstanceLOD.Guides.bOwnRootResourceAllocation = false;
-						InstanceLOD.Guides.RestRootResource = BindingAsset->HairGroupResources[GroupIt].CardsRootResources[CardsLODIndex];
+						InstanceLOD.Guides.RestRootResource = LocalBindingAsset->HairGroupResources[GroupIt].CardsRootResources[CardsLODIndex];
 
 						InstanceLOD.Guides.DeformedRootResource = new FHairStrandsDeformedRootResource(InstanceLOD.Guides.RestRootResource);
 						BeginInitResource(InstanceLOD.Guides.DeformedRootResource);
@@ -2307,8 +2347,7 @@ void UGroomComponent::TickComponent(float DeltaTime, enum ELevelTick TickType, F
 
 	// When a groom binding and similation are disabled, and the groom component is parented with a skeletal mesh, we can optionnaly 
 	// attach the groom to a particular socket/bone
-	const bool bStaticAttachement = !IsHairStrandsBindingEnable() && !IsHairStrandsSimulationEnable();
-	if (RegisteredSkeletalMeshComponent && bStaticAttachement && !AttachmentName.IsEmpty())
+	if (RegisteredSkeletalMeshComponent && !AttachmentName.IsEmpty())
 	{
 		const FName BoneName(AttachmentName);
 		if (GetAttachSocketName() != BoneName)
@@ -2377,6 +2416,7 @@ void UGroomComponent::TickComponent(float DeltaTime, enum ELevelTick TickType, F
 		for (FHairGroupInstance* Instance : LocalHairGroupInstances)
 		{
 			Instance->WorldType = WorldType;
+			Instance->Debug.SkeletalPreviousLocalToWorld = Instance->Debug.SkeletalLocalToWorld;
 			Instance->Debug.SkeletalLocalToWorld = SkinLocalToWorld;
 		}
 	});
diff --git a/Engine/Plugins/Runtime/HairStrands/Source/HairStrandsCore/Private/GroomDebug.cpp b/Engine/Plugins/Runtime/HairStrands/Source/HairStrandsCore/Private/GroomDebug.cpp
--- a/Engine/Plugins/Runtime/HairStrands/Source/HairStrandsCore/Private/GroomDebug.cpp
+++ b/Engine/Plugins/Runtime/HairStrands/Source/HairStrandsCore/Private/GroomDebug.cpp
@@ -91,6 +91,8 @@ namespace GroomDebug
 	};
 }
 
+bool IsHairStrandsSkinCacheEnable();
+
 static void GetGroomInterpolationData(
 	FRDGBuilder& GraphBuilder,
 	const TArray<FHairGroupInstance*> Instances,
@@ -105,15 +107,18 @@ static void GetGroomInterpolationData(
 			continue;
 
 		FCachedGeometry CachedGeometry;
-		if (SkinCache)
-		{
-			CachedGeometry = SkinCache->GetCachedGeometry(Instance->Debug.SkeletalComponent->ComponentId.PrimIDValue);
-		}
-		else if (Instance->Debug.SkeletalComponent)
+		if (Instance->Debug.SkeletalComponent)
 		{
-			const ERHIFeatureLevel::Type FeatureLevel = GMaxRHIFeatureLevel;
-			FGlobalShaderMap* ShaderMap = GetGlobalShaderMap(FeatureLevel);
-			BuildCacheGeometry(GraphBuilder, ShaderMap, Instance->Debug.SkeletalComponent, CachedGeometry);
+			if (SkinCache)
+			{
+				CachedGeometry = SkinCache->GetCachedGeometry(Instance->Debug.SkeletalComponent->ComponentId.PrimIDValue);
+			}
+			
+			if (IsHairStrandsSkinCacheEnable() && CachedGeometry.Sections.Num() == 0)
+			{
+				FGlobalShaderMap* ShaderMap = GetGlobalShaderMap(GMaxRHIFeatureLevel);
+				BuildCacheGeometry(GraphBuilder, ShaderMap, Instance->Debug.SkeletalComponent, CachedGeometry);
+			}
 		}
 		if (CachedGeometry.Sections.Num() == 0)
 			continue;
@@ -776,6 +781,29 @@ BEGIN_SHADER_PARAMETER_STRUCT(FHairDebugCanvasParameter, )
 	RENDER_TARGET_BINDING_SLOTS()
 END_SHADER_PARAMETER_STRUCT()
 
+static const TCHAR* ToString(EHairGeometryType In)
+{
+	switch (In)
+	{
+	case EHairGeometryType::NoneGeometry:	return TEXT("None");
+	case EHairGeometryType::Strands:		return TEXT("Strands");
+	case EHairGeometryType::Cards:			return TEXT("Cards");
+	case EHairGeometryType::Meshes:			return TEXT("Meshes");
+	}
+	return TEXT("None");
+}
+
+static const TCHAR* ToString(EHairBindingType In)
+{
+	switch (In)
+	{
+	case EHairBindingType::NoneBinding: return TEXT("None");
+	case EHairBindingType::Rigid:		return TEXT("Rigid");
+	case EHairBindingType::Skinning:	return TEXT("Skinning");
+	}
+	return TEXT("None");
+}
+
 void RunHairStrandsDebug(
 	FRDGBuilder& GraphBuilder,
 	FGlobalShaderMap* ShaderMap,
@@ -817,6 +845,8 @@ void RunHairStrandsDebug(
 			const FLinearColor DebugColor(1, 1, 0);
 			const FLinearColor DebugGroupColor(0.5f, 0, 0);
 			FString Line;
+			// Active groom
+			// Name | Group x / x | LOD x / x | GeometryType | BindingType | Sim | RBF | VertexCount
 
 			Line = FString::Printf(TEXT("----------------------------------------------------------------"));
 			Canvas.DrawShadowedString(X, Y += YStep, *Line, GetStatsFont(), DebugColor);
@@ -830,33 +860,22 @@ void RunHairStrandsDebug(
 				const bool bHasSkinInterpolation = Instance->Strands.RestRootResource != nullptr;
 				const bool bHasBindingAsset = bHasSkinInterpolation && !Instance->Strands.bOwnRootResourceAllocation;
 
-				Line = FString::Printf(TEXT(" * Id:%d | WorldType:%s | Group:%d/%d | Asset : %s | Skeletal : %s "),
-					Instance->Debug.ComponentId,
-					ToString(Instance->WorldType),
+				if (!bIsActive)
+					continue;
+
+				Line = FString::Printf(TEXT(" * Group:%d/%d | LOD:%1.1f/%d | GeometryType:%s | BindingType:%s | Sim:%d | RBF:%d | VertexCount:%d | Name: %s"),
 					Instance->Debug.GroupIndex,
 					Instance->Debug.GroupCount,
-					*Instance->Debug.GroomAssetName,
-					*Instance->Debug.SkeletalComponentName);
-				Canvas.DrawShadowedString(X, Y += YStep, *Line, GetStatsFont(), bIsActive ? DebugColor : InactiveColor);
-				if (Instance->Strands.IsValid())
-				{
-					Line = FString::Printf(TEXT("        |> CurveCount : %d | VertexCount : %d | MaxRadius : %f | MaxLength : %f | Skinned: %s | Binding: %s | Simulation: %s| LOD count : %d"),
-						Instance->Strands.Data->GetNumCurves(),
-						Instance->Strands.Data->GetNumPoints(),
-						Instance->HairGroupPublicData->VFInput.Strands.HairRadius,
-						Instance->HairGroupPublicData->VFInput.Strands.HairLength,
-						bHasSkinInterpolation ? TEXT("True") : TEXT("False"),
-						bHasBindingAsset ? TEXT("True") : TEXT("False"),
-						Instance->Guides.bIsSimulationEnable ? TEXT("True") : TEXT("False"),
-						Instance->Strands.ClusterCullingResource->Data.ClusterLODInfos.Num());
-				}
-				else
-				{
-					Line = FString::Printf(TEXT("        |> HasStrands : %s | HasCards : %s | HasMeshes : %s"),
-						Instance->Strands.IsValid() ? TEXT("True") : TEXT("False"),
-						Instance->Cards.IsValid() ? TEXT("True") : TEXT("False"),
-						Instance->Meshes.IsValid() ? TEXT("True") : TEXT("False"));
-				}
+
+					Instance->HairGroupPublicData->LODIndex,
+					Instance->HairGroupPublicData->GetLODScreenSizes().Num(),
+
+					ToString(Instance->GeometryType),
+					ToString(Instance->BindingType),
+					0, // TODO: Instance->bHasSimulation,
+					0, // TODO: Instance->bHasRBF,
+					Instance->HairGroupPublicData->VertexCount,
+					*Instance->Debug.GroomAssetName);
 				Canvas.DrawShadowedString(X, Y += YStep, *Line, GetStatsFont(), bIsActive ? DebugGroupColor : InactiveColor);
 			}
 
@@ -902,7 +921,7 @@ void RunHairStrandsDebug(
 				TArray<int32> HairLODIndices;
 				for (FHairGroupInstance* Instance : Instances)
 				{
-					if (!Instance->HairGroupPublicData || Instance->Guides.RestRootResource == nullptr || Instance->Guides.DeformedRootResource == nullptr)
+					if (Instance->WorldType != WorldType || !Instance->HairGroupPublicData || Instance->Guides.RestRootResource == nullptr || Instance->Guides.DeformedRootResource == nullptr)
 						continue;
 
 					const int32 MeshLODIndex = Instance->Debug.MeshLODIndex;
diff --git a/Engine/Plugins/Runtime/HairStrands/Source/HairStrandsCore/Private/GroomManager.cpp b/Engine/Plugins/Runtime/HairStrands/Source/HairStrandsCore/Private/GroomManager.cpp
--- a/Engine/Plugins/Runtime/HairStrands/Source/HairStrandsCore/Private/GroomManager.cpp
+++ b/Engine/Plugins/Runtime/HairStrands/Source/HairStrandsCore/Private/GroomManager.cpp
@@ -25,6 +25,14 @@ static FAutoConsoleVariableRef CVarHairStrands_UseCards(TEXT("r.HairStrands.UseC
 static int32 GHairStrands_SwapBufferEndOfFrame = 1;
 static FAutoConsoleVariableRef CVarGHairStrands_SwapBufferEndOfFrame(TEXT("r.HairStrands.SwapEndOfFrame"), GHairStrands_SwapBufferEndOfFrame, TEXT("Swap rendering buffer at the end of frame. This is an experimental toggle. Default:1"));
 
+static int32 GHairStrands_ManualSkinCache = 1;
+static FAutoConsoleVariableRef CVarGHairStrands_ManualSkinCache(TEXT("r.HairStrands.ManualSkinCache"), GHairStrands_ManualSkinCache, TEXT("If skin cache is not enabled, and grooms use skinning method, this enable a simple skin cache mechanisme for groom. Default:enable"));
+
+bool IsHairStrandsSkinCacheEnable()
+{
+	return GHairStrands_ManualSkinCache > 0;
+}
+
 DEFINE_LOG_CATEGORY_STATIC(LogGroomManager, Log, All);
 
 ////////////////////////////////////////////////////////////////////////////////////////////////////////////////////////
@@ -124,7 +132,8 @@ void RunHairStrandsInterpolation(
 			{
 				CachedGeometry = SkinCache->GetCachedGeometry(Instance->Debug.SkeletalComponent->ComponentId.PrimIDValue);
 			}
-			else
+
+			if (IsHairStrandsSkinCacheEnable() && CachedGeometry.Sections.Num() == 0)
 			{
 				//#hair_todo: Need to have a (frame) cache to insure that we don't recompute the same projection several time
 				// Actual populate the cache with only the needed part basd on the groom projection data. At the moment it recompute everything ...
@@ -546,6 +555,7 @@ static void RunHairLODSelection(EWorldType::Type WorldType, const TArray<const F
 		Instance->HairGroupPublicData->VFInput.GeometryType = GeometryType;
 		Instance->HairGroupPublicData->VFInput.bHasLODSwitch = (FMath::FloorToInt(PrevLODIndex) != FMath::FloorToInt(LODIndex));
 		Instance->GeometryType = GeometryType;
+		Instance->BindingType = Instance->HairGroupPublicData->GetBindingType(IntLODIndex);
 	}
 }
 
diff --git a/Engine/Plugins/Runtime/HairStrands/Source/HairStrandsCore/Private/HairStrandsMeshProjection.cpp b/Engine/Plugins/Runtime/HairStrands/Source/HairStrandsCore/Private/HairStrandsMeshProjection.cpp
--- a/Engine/Plugins/Runtime/HairStrands/Source/HairStrandsCore/Private/HairStrandsMeshProjection.cpp
+++ b/Engine/Plugins/Runtime/HairStrands/Source/HairStrandsCore/Private/HairStrandsMeshProjection.cpp
@@ -1763,6 +1763,9 @@ private:
 	DECLARE_GLOBAL_SHADER(FHairUpdatePositionOffsetCS);
 	SHADER_USE_PARAMETER_STRUCT(FHairUpdatePositionOffsetCS, FGlobalShader);
 
+	class FSourceType : SHADER_PERMUTATION_INT("PERMUTATION_SOURCE_TYPE", 2);
+	using FPermutationDomain = TShaderPermutationDomain<FSourceType>;
+
 	BEGIN_SHADER_PARAMETER_STRUCT(FParameters, )
 		SHADER_PARAMETER(FVector, CPUPositionOffset)
 		SHADER_PARAMETER(uint32, bUseCPUOffset)
@@ -1806,11 +1809,13 @@ void AddHairStrandUpdatePositionOffsetPass(
 	const uint32 OffsetIndex = DeformedResources->GetIndex(FHairStrandsDeformedResource::Current);
 	FHairUpdatePositionOffsetCS::FParameters* PassParameters = GraphBuilder.AllocParameters<FHairUpdatePositionOffsetCS::FParameters>();
 	PassParameters->CPUPositionOffset = DeformedResources->PositionOffset[OffsetIndex];
-	PassParameters->bUseCPUOffset = (DeformedRootResources == nullptr || GHairStrandsUseGPUPositionOffset>0) ? 0 : 1;
+	PassParameters->bUseCPUOffset = (DeformedRootResources == nullptr || GHairStrandsUseGPUPositionOffset>0) ? 1 : 0;
 	PassParameters->RootTrianglePosition0Buffer = RootTrianglePositionBuffer.SRV;
 	PassParameters->OutOffsetBuffer = OutPositionOffsetBuffer.UAV;
 
-	TShaderMapRef<FHairUpdatePositionOffsetCS> ComputeShader(ShaderMap);
+	FHairUpdatePositionOffsetCS::FPermutationDomain PermutationVector;
+	PermutationVector.Set<FHairUpdatePositionOffsetCS::FSourceType>(PassParameters->bUseCPUOffset ? 0 : 1);
+	TShaderMapRef<FHairUpdatePositionOffsetCS> ComputeShader(ShaderMap, PermutationVector);
 	FComputeShaderUtils::AddPass(
 		GraphBuilder,
 		RDG_EVENT_NAME("HairStrandsUpdatePositionOffset"),
diff --git a/Engine/Plugins/Runtime/HairStrands/Source/HairStrandsCore/Private/HairStrandsRendering.cpp b/Engine/Plugins/Runtime/HairStrands/Source/HairStrandsCore/Private/HairStrandsRendering.cpp
--- a/Engine/Plugins/Runtime/HairStrands/Source/HairStrandsCore/Private/HairStrandsRendering.cpp
+++ b/Engine/Plugins/Runtime/HairStrands/Source/HairStrandsCore/Private/HairStrandsRendering.cpp
@@ -1537,6 +1537,8 @@ void ComputeHairStrandsInterpolation(
 				ScaleAndClipDesc.HairRadiusTipScale = 1;
 				ScaleAndClipDesc.HairLengthClip = 1;
 
+				const bool bHasSkinning = Instance->BindingType == EHairBindingType::Skinning;
+
 				FRDGHairStrandsCullingData CullingData;
 				AddHairStrandsInterpolationPass(
 					GraphBuilder,
@@ -1554,10 +1556,10 @@ void ComputeHairStrandsInterpolation(
 					Instance->Guides.RestResource->PositionOffset,
 					RegisterAsSRV(GraphBuilder, LOD.Guides.DeformedResource->GetPositionOffsetBuffer(FHairStrandsDeformedResource::Current)),
 					RegisterAsSRV(GraphBuilder, Instance->Guides.DeformedResource->GetPositionOffsetBuffer(FHairStrandsDeformedResource::Current)),
-					LOD.Guides.RestRootResource,
-					Instance->Guides.RestRootResource,
-					LOD.Guides.DeformedRootResource,
-					Instance->Guides.DeformedRootResource,
+					bHasSkinning ? LOD.Guides.RestRootResource : nullptr ,
+					bHasSkinning ? Instance->Guides.RestRootResource : nullptr,
+					bHasSkinning ? LOD.Guides.DeformedRootResource : nullptr,
+					bHasSkinning ? Instance->Guides.DeformedRootResource : nullptr,
 					RegisterAsSRV(GraphBuilder, LOD.Guides.RestResource->RestPositionBuffer),
 					RegisterAsSRV(GraphBuilder, LOD.Guides.RestResource->AttributeBuffer),
 					RegisterAsSRV(GraphBuilder, LOD.Guides.InterpolationResource->Interpolation0Buffer),
@@ -1569,7 +1571,7 @@ void ComputeHairStrandsInterpolation(
 					nullptr,
 					nullptr,
 					nullptr,
-					RegisterAsSRV(GraphBuilder, LOD.Guides.InterpolationResource->SimRootPointIndexBuffer));
+					RegisterAsSRV(GraphBuilder, LOD.Guides.InterpolationResource->SimRootPointIndexBuffer)); // <- this should be optional
 
 				AddHairCardsDeformationPass(
 					GraphBuilder,
@@ -1662,7 +1664,10 @@ void ComputeHairStrandsInterpolation(
 	}
 
 	Instance->HairGroupPublicData->VFInput.GeometryType = InstanceGeometryType;
-	Instance->HairGroupPublicData->VFInput.LocalToWorldTransform = Instance->LocalToWorld;
+	if (Instance->BindingType == EHairBindingType::Rigid)
+		Instance->HairGroupPublicData->VFInput.LocalToWorldTransform = Instance->LocalToWorld;
+	else
+		Instance->HairGroupPublicData->VFInput.LocalToWorldTransform = Instance->Debug.SkeletalLocalToWorld;
 	Instance->HairGroupPublicData->bSupportVoxelization = Instance->Strands.Modifier.bSupportVoxelization && Instance->bCastShadow;
 }
 
diff --git a/Engine/Plugins/Runtime/HairStrands/Source/HairStrandsCore/Public/GroomAssetInterpolation.h b/Engine/Plugins/Runtime/HairStrands/Source/HairStrandsCore/Public/GroomAssetInterpolation.h
--- a/Engine/Plugins/Runtime/HairStrands/Source/HairStrandsCore/Public/GroomAssetInterpolation.h
+++ b/Engine/Plugins/Runtime/HairStrands/Source/HairStrandsCore/Public/GroomAssetInterpolation.h
@@ -40,6 +40,15 @@ enum class EGroomGeometryType : uint8
 	Meshes
 };
 
+UENUM(BlueprintType)
+enum class EGroomBindingType : uint8
+{
+	NoneBinding	UMETA(Hidden),
+	Rigid		UMETA(DisplayName = "Rigid",    ToolTip = "When attached to a skeltal mesh, the hair follow the provided attachement name"),
+	Skinning	UMETA(DisplayName = "Skinning", ToolTip = "When attached to a skeltal mesh, the hair follow the skin surface"),
+};
+
+
 USTRUCT(BlueprintType)
 struct HAIRSTRANDSCORE_API FHairLODSettings
 {
@@ -73,6 +82,9 @@ struct HAIRSTRANDSCORE_API FHairLODSettings
 	UPROPERTY(EditAnywhere, Category = "DecimationSettings", meta = (ToolTip = "If enable this LOD version will use cards representation"))
 	EGroomGeometryType GeometryType = EGroomGeometryType::Strands;
 
+	UPROPERTY(EditAnywhere, Category = "DecimationSettings", meta = (ToolTip = "If enable this LOD version will use the provided attachment points"))
+	EGroomBindingType BindingType = EGroomBindingType::NoneBinding;
+
 	bool operator==(const FHairLODSettings& A) const;
 };
 
diff --git a/Engine/Plugins/Runtime/HairStrands/Source/HairStrandsCore/Public/GroomInstance.h b/Engine/Plugins/Runtime/HairStrands/Source/HairStrandsCore/Public/GroomInstance.h
--- a/Engine/Plugins/Runtime/HairStrands/Source/HairStrandsCore/Public/GroomInstance.h
+++ b/Engine/Plugins/Runtime/HairStrands/Source/HairStrandsCore/Public/GroomInstance.h
@@ -171,6 +171,7 @@ struct HAIRSTRANDSCORE_API FHairGroupInstance
 		USkeletalMeshComponent*	SkeletalComponent = nullptr;
 		FString					SkeletalComponentName;
 		FTransform				SkeletalLocalToWorld = FTransform::Identity;
+		FTransform				SkeletalPreviousLocalToWorld = FTransform::Identity;
 		bool					bDrawCardsGuides = false;
 
 		// Transfer
@@ -186,6 +187,7 @@ struct HAIRSTRANDSCORE_API FHairGroupInstance
 	EWorldType::Type		WorldType;
 	FHairGroupPublicData*	HairGroupPublicData = nullptr;
 	EHairGeometryType		GeometryType = EHairGeometryType::NoneGeometry;
+	EHairBindingType		BindingType = EHairBindingType::NoneBinding;
 	const FBoxSphereBounds*	ProxyBounds = nullptr;
 	const FBoxSphereBounds* ProxyLocalBounds = nullptr;
 	bool					bUseCPULODSelection = true;
diff --git a/Engine/Shaders/Private/HairStrands/HairStrandsMeshUpdate.usf b/Engine/Shaders/Private/HairStrands/HairStrandsMeshUpdate.usf
--- a/Engine/Shaders/Private/HairStrands/HairStrandsMeshUpdate.usf
+++ b/Engine/Shaders/Private/HairStrands/HairStrandsMeshUpdate.usf
@@ -210,7 +210,9 @@ void MainCS(uint3 DispatchThreadId : SV_DispatchThreadID)
 
 float3 CPUPositionOffset;
 uint bUseCPUOffset;
+#if PERMUTATION_SOURCE_TYPE == 1
 Buffer<float4> RootTrianglePosition0Buffer;
+#endif
 RWBuffer<float4> OutOffsetBuffer;
 
 [numthreads(1, 1, 1)]
@@ -220,10 +222,13 @@ void MainCS(uint3 DispatchThreadId : SV_DispatchThreadID)
 		return;
 
 	float3 RefPosition = CPUPositionOffset;
+	#if PERMUTATION_SOURCE_TYPE == 1
 	if (bUseCPUOffset == 0)
 	{
 		RefPosition = RootTrianglePosition0Buffer[0];
 	}
+	#endif
+
 	OutOffsetBuffer[0] = float4(RefPosition, 1);
 }
 
diff --git a/Engine/Source/Runtime/Renderer/Public/HairStrandsInterface.h b/Engine/Source/Runtime/Renderer/Public/HairStrandsInterface.h
--- a/Engine/Source/Runtime/Renderer/Public/HairStrandsInterface.h
+++ b/Engine/Source/Runtime/Renderer/Public/HairStrandsInterface.h
@@ -128,6 +128,13 @@ enum EHairGeometryType
 	NoneGeometry
 };
 
+enum EHairBindingType
+{
+	NoneBinding,
+	Rigid,
+	Skinning
+};
+
 ////////////////////////////////////////////////////////////////////////////////////////////////////
 // Public group data 
 class RENDERER_API FHairGroupPublicData : public FRenderResource
@@ -170,6 +177,12 @@ public:
 	void SetLODVisibilities(const TArray<bool>& InLODVisibility) { LODVisibilities = InLODVisibility; }
 	const TArray<bool>& GetLODVisibilities() const { return LODVisibilities; }
 
+	EHairBindingType GetBindingType(int32 InLODIndex) 
+	{ 
+		if (InLODIndex < 0 && InLODIndex >= BindingTypes.Num()) return EHairBindingType::NoneBinding;
+		return BindingTypes[InLODIndex];
+	}	
+
 	void SetLODScreenSizes(const TArray<float>& ScreenSizes) { LODScreenSizes = ScreenSizes; }
 	const TArray<float>& GetLODScreenSizes() const { return LODScreenSizes;  }
 
@@ -220,6 +233,7 @@ public:
 
 		bool bHasLODSwitch = false;
 		EHairGeometryType GeometryType = EHairGeometryType::NoneGeometry;
+		EHairBindingType BindingType = EHairBindingType::NoneBinding;
 		FTransform LocalToWorldTransform;
 	};
 	FVertexFactoryInput VFInput;
@@ -252,6 +266,8 @@ public:
 	TArray<float> LODScreenSizes;
 	TArray<EHairGeometryType> LODGeometryTypes;
 
+	TArray<EHairBindingType> BindingTypes;
+
 	// Data change every frame by the groom proxy based on views data
 	float LODIndex = 0;			// Current LOD used for all views
 	float LODBias = 0;			// Current LOD bias
