Set default binding value to Skinning.

diff --git a/Engine/Plugins/Runtime/HairStrands/Source/HairStrandsCore/Public/GroomAssetInterpolation.h b/Engine/Plugins/Runtime/HairStrands/Source/HairStrandsCore/Public/GroomAssetInterpolation.h
--- a/Engine/Plugins/Runtime/HairStrands/Source/HairStrandsCore/Public/GroomAssetInterpolation.h
+++ b/Engine/Plugins/Runtime/HairStrands/Source/HairStrandsCore/Public/GroomAssetInterpolation.h
@@ -83,7 +83,7 @@ struct HAIRSTRANDSCORE_API FHairLODSettings
 	EGroomGeometryType GeometryType = EGroomGeometryType::Strands;
 
 	UPROPERTY(EditAnywhere, Category = "DecimationSettings", meta = (ToolTip = "If enable this LOD version will use the provided attachment points"))
-	EGroomBindingType BindingType = EGroomBindingType::NoneBinding;
+	EGroomBindingType BindingType = EGroomBindingType::Skinning;
 
 	bool operator==(const FHairLODSettings& A) const;
 };
