Fix use of the PreviewMesh in the NiagaraDataInterfacePhysicsAsset

diff --git a/Engine/Plugins/Runtime/HairStrands/Source/HairStrandsCore/Private/Niagara/NiagaraDataInterfacePhysicsAsset.cpp b/Engine/Plugins/Runtime/HairStrands/Source/HairStrandsCore/Private/Niagara/NiagaraDataInterfacePhysicsAsset.cpp
--- a/Engine/Plugins/Runtime/HairStrands/Source/HairStrandsCore/Private/Niagara/NiagaraDataInterfacePhysicsAsset.cpp
+++ b/Engine/Plugins/Runtime/HairStrands/Source/HairStrandsCore/Private/Niagara/NiagaraDataInterfacePhysicsAsset.cpp
@@ -142,7 +142,7 @@ void CreateInternalArrays(const TArray<TWeakObjectPtr<UPhysicsAsset>>& PhysicsAs
 			TWeakObjectPtr<UPhysicsAsset> PhysicsAsset = PhysicsAssets[ComponentIndex];
 			if (PhysicsAsset.IsValid() && PhysicsAsset.Get() != nullptr)
 			{
-				USkeletalMesh* SkelMesh = PhysicsAsset->GetPreviewMesh();
+				USkeletalMesh* SkelMesh = (SkeletalMeshs[ComponentIndex].Get() && SkeletalMeshs[ComponentIndex]->SkeletalMesh) ? SkeletalMeshs[ComponentIndex]->SkeletalMesh : PhysicsAsset->GetPreviewMesh();
 				if (!SkelMesh)
 				{
 					continue;
@@ -208,7 +208,7 @@ void CreateInternalArrays(const TArray<TWeakObjectPtr<UPhysicsAsset>>& PhysicsAs
 				TWeakObjectPtr<UPhysicsAsset> PhysicsAsset = PhysicsAssets[ComponentIndex];
 				if (PhysicsAsset.IsValid() && PhysicsAsset.Get() != nullptr)
 				{
-					USkeletalMesh* SkelMesh = PhysicsAsset->GetPreviewMesh();
+					USkeletalMesh* SkelMesh = (SkeletalMeshs[ComponentIndex].Get() && SkeletalMeshs[ComponentIndex]->SkeletalMesh) ? SkeletalMeshs[ComponentIndex]->SkeletalMesh : PhysicsAsset->GetPreviewMesh();
 					if (!SkelMesh)
 					{
 						continue;
@@ -305,7 +305,7 @@ void UpdateInternalArrays(const TArray<TWeakObjectPtr<UPhysicsAsset>>& PhysicsAs
 			TWeakObjectPtr<UPhysicsAsset> PhysicsAsset = PhysicsAssets[ComponentIndex];
 			if (PhysicsAsset.IsValid() && PhysicsAsset.Get() != nullptr)
 			{
-				USkeletalMesh* SkelMesh = PhysicsAsset->GetPreviewMesh();
+				USkeletalMesh* SkelMesh = (SkeletalMeshs[ComponentIndex].Get() && SkeletalMeshs[ComponentIndex]->SkeletalMesh) ? SkeletalMeshs[ComponentIndex]->SkeletalMesh : PhysicsAsset->GetPreviewMesh();
 				if (!SkelMesh)
 				{
 					continue;
@@ -425,10 +425,11 @@ void FNDIPhysicsAssetData::Init(UNiagaraDataInterfacePhysicsAsset* Interface, FN
 			TWeakObjectPtr<UPhysicsAsset> PhysicsAsset = Interface->PhysicsAssets[ComponentIndex];
 			if (PhysicsAsset.IsValid() && PhysicsAsset.Get() != nullptr)
 			{
-				const bool HasPreviewMesh = PhysicsAsset != nullptr && PhysicsAsset->GetPreviewMesh() != nullptr;
-				if (HasPreviewMesh)
+				USkeletalMesh* SkelMesh = (Interface->SourceComponents[ComponentIndex].IsValid() && Interface->SourceComponents[ComponentIndex].Get() && Interface->SourceComponents[ComponentIndex]->SkeletalMesh) ?
+					Interface->SourceComponents[ComponentIndex]->SkeletalMesh : PhysicsAsset->GetPreviewMesh();
+				if (SkelMesh)
 				{
-					BoundingBox += PhysicsAsset->GetPreviewMesh()->GetImportedBounds().GetBox();
+					BoundingBox += SkelMesh->GetImportedBounds().GetBox();
 				}
 			}
 		}
