Fix groom translation when not attached to a skeletal component.

diff --git a/Engine/Plugins/Runtime/HairStrands/Source/HairStrandsCore/Private/GroomComponent.cpp b/Engine/Plugins/Runtime/HairStrands/Source/HairStrandsCore/Private/GroomComponent.cpp
--- a/Engine/Plugins/Runtime/HairStrands/Source/HairStrandsCore/Private/GroomComponent.cpp
+++ b/Engine/Plugins/Runtime/HairStrands/Source/HairStrandsCore/Private/GroomComponent.cpp
@@ -1805,14 +1805,14 @@ void UGroomComponent::InitResources(bool bIsBindingReloading)
 
 	InitializedResources = GroomAsset;
 
-	// Initialize LOD screen size & visibility
+	// 1. Check if we need any kind of binding data 
 	bool bHasNeedBindingData = false;
 	for (int32 GroupIt = 0, GroupCount = GroomAsset->HairGroupsData.Num(); GroupIt < GroupCount; ++GroupIt)
 	{
 		const FHairGroupsLOD& GroupLOD = GroomAsset->HairGroupsLOD[GroupIt];
 		for (const FHairLODSettings& LODSettings : GroupLOD.LODs)
 		{
-			bHasNeedBindingData = bHasNeedBindingData || LODSettings.BindingType == EGroomBindingType::Skinning;
+			bHasNeedBindingData  = bHasNeedBindingData || LODSettings.BindingType == EGroomBindingType::Skinning;
 		}
 	}
 
@@ -1820,9 +1820,12 @@ void UGroomComponent::InitResources(bool bIsBindingReloading)
 	const EWorldType::Type WorldType = GetWorldType();
 	const bool bIsStrandsEnabled = IsHairStrandsEnabled(EHairStrandsShaderType::Strands);
 
-	// Insure that the binding asset is compatible, otherwise no binding will be used
+	// 2. Insure that the binding asset is compatible, otherwise no binding
 	USkeletalMeshComponent* SkeletalMeshComponent = GetAttachParent() ? Cast<USkeletalMeshComponent>(GetAttachParent()) : nullptr;
-	SkeletalMeshComponent = ValidateBindingAsset(GroomAsset, BindingAsset, SkeletalMeshComponent, bIsBindingReloading, bValidationEnable, this);
+	if (bHasNeedBindingData && SkeletalMeshComponent)
+	{
+		SkeletalMeshComponent = ValidateBindingAsset(GroomAsset, BindingAsset, SkeletalMeshComponent, bIsBindingReloading, bValidationEnable, this);
+	}
 	UGroomBindingAsset* LocalBindingAsset = bHasNeedBindingData ? BindingAsset : nullptr;
 
 	// Insure the ticking of the Groom component always happens after the skeletalMeshComponent.
@@ -1851,6 +1854,7 @@ void UGroomComponent::InitResources(bool bIsBindingReloading)
 			HairGroupInstance->Debug.SkeletalComponentName = SkeletalMeshComponent->GetPathName();
 		}
 		HairGroupInstance->GeometryType = EHairGeometryType::NoneGeometry;
+		HairGroupInstance->BindingType = EHairBindingType::NoneBinding;
 
 		FHairGroupData& GroupData = GroomAsset->HairGroupsData[GroupIt];
 		const uint32 SkeletalLODCount = SkeletalMeshComponent ? SkeletalMeshComponent->GetNumLODs() : 0;
@@ -1871,10 +1875,20 @@ void UGroomComponent::InitResources(bool bIsBindingReloading)
 			const FHairGroupsLOD& GroupLOD = GroomAsset->HairGroupsLOD[GroupIt];
 			for (const FHairLODSettings& LODSettings : GroupLOD.LODs)
 			{
+				EHairBindingType BindingType = ToHairBindingType(LODSettings.BindingType);
+				if (BindingType == EHairBindingType::Skinning && !LocalBindingAsset)
+				{
+					BindingType = EHairBindingType::NoneBinding;
+				}
+				else if (BindingType == EHairBindingType::Rigid && !SkeletalMeshComponent)
+				{
+					BindingType = EHairBindingType::NoneBinding;
+				}
+
 				CPULODScreenSize.Add(LODSettings.ScreenSize);
 				LODVisibility.Add(LODSettings.bVisible);
 				LODGeometryTypes.Add(ConvertToHairGeometryType(LODSettings.GeometryType));
-				HairGroupInstance->HairGroupPublicData->BindingTypes.Add(ToHairBindingType(LODSettings.BindingType));
+				HairGroupInstance->HairGroupPublicData->BindingTypes.Add(BindingType);
 			}
 			HairGroupInstance->HairGroupPublicData->SetLODScreenSizes(CPULODScreenSize);
 			HairGroupInstance->HairGroupPublicData->SetLODVisibilities(LODVisibility);
diff --git a/Engine/Plugins/Runtime/HairStrands/Source/HairStrandsCore/Private/HairStrandsRendering.cpp b/Engine/Plugins/Runtime/HairStrands/Source/HairStrandsCore/Private/HairStrandsRendering.cpp
--- a/Engine/Plugins/Runtime/HairStrands/Source/HairStrandsCore/Private/HairStrandsRendering.cpp
+++ b/Engine/Plugins/Runtime/HairStrands/Source/HairStrandsCore/Private/HairStrandsRendering.cpp
@@ -1412,6 +1412,8 @@ void ComputeHairStrandsInterpolation(
 		Instance->HairGroupPublicData->VFInput = ComputeHairStrandsVertexInputData(Instance);
 		{
 			{
+				const bool bHasSkinning = Instance->BindingType == EHairBindingType::Skinning;
+
 				AddHairStrandsInterpolationPass(
 					GraphBuilder,
 					ShaderMap,
@@ -1428,10 +1430,10 @@ void ComputeHairStrandsInterpolation(
 					Instance->Guides.RestResource->PositionOffset,
 					RegisterAsSRV(GraphBuilder, Instance->Strands.DeformedResource->GetPositionOffsetBuffer(FHairStrandsDeformedResource::Current)),
 					RegisterAsSRV(GraphBuilder, Instance->Guides.DeformedResource->GetPositionOffsetBuffer(FHairStrandsDeformedResource::Current)),
-					Instance->Strands.RestRootResource,
-					Instance->Guides.RestRootResource,
-					Instance->Strands.DeformedRootResource,
-					Instance->Guides.DeformedRootResource,
+					bHasSkinning ? Instance->Strands.RestRootResource : nullptr,
+					bHasSkinning ? Instance->Guides.RestRootResource : nullptr,
+					bHasSkinning ? Instance->Strands.DeformedRootResource : nullptr,
+					bHasSkinning ? Instance->Guides.DeformedRootResource : nullptr,
 					RegisterAsSRV(GraphBuilder, Instance->Strands.RestResource->RestPositionBuffer),
 					RegisterAsSRV(GraphBuilder, Instance->Strands.RestResource->AttributeBuffer),
 					RegisterAsSRV(GraphBuilder, Instance->Strands.InterpolationResource->Interpolation0Buffer),
