Avoid running hair root triangle update when skinning attachment is disabled.

diff --git a/Engine/Plugins/Runtime/HairStrands/Source/HairStrandsCore/Private/GroomDebug.cpp b/Engine/Plugins/Runtime/HairStrands/Source/HairStrandsCore/Private/GroomDebug.cpp
--- a/Engine/Plugins/Runtime/HairStrands/Source/HairStrandsCore/Private/GroomDebug.cpp
+++ b/Engine/Plugins/Runtime/HairStrands/Source/HairStrandsCore/Private/GroomDebug.cpp
@@ -921,7 +921,7 @@ void RunHairStrandsDebug(
 				TArray<int32> HairLODIndices;
 				for (FHairGroupInstance* Instance : Instances)
 				{
-					if (Instance->WorldType != WorldType || !Instance->HairGroupPublicData || Instance->Guides.RestRootResource == nullptr || Instance->Guides.DeformedRootResource == nullptr)
+					if (Instance->WorldType != WorldType || !Instance->HairGroupPublicData || Instance->Guides.RestRootResource == nullptr || Instance->Guides.DeformedRootResource == nullptr || Instance->BindingType != EHairBindingType::Skinning)
 						continue;
 
 					const int32 MeshLODIndex = Instance->Debug.MeshLODIndex;
diff --git a/Engine/Plugins/Runtime/HairStrands/Source/HairStrandsCore/Private/GroomManager.cpp b/Engine/Plugins/Runtime/HairStrands/Source/HairStrandsCore/Private/GroomManager.cpp
--- a/Engine/Plugins/Runtime/HairStrands/Source/HairStrandsCore/Private/GroomManager.cpp
+++ b/Engine/Plugins/Runtime/HairStrands/Source/HairStrandsCore/Private/GroomManager.cpp
@@ -159,9 +159,10 @@ void RunHairStrandsInterpolation(
 			if (EHairStrandsInterpolationType::RenderStrands == Type)
 			{
 				const EHairGeometryType InstanceGeometryType = Instance->GeometryType;
+				const EHairBindingType InstanceBindingType = Instance->BindingType;
 				if (InstanceGeometryType == EHairGeometryType::Strands)
 				{
-					if (Instance->Strands.HasValidRootData() && Instance->Strands.DeformedRootResource->IsValid(MeshLODIndex))
+					if (Instance->Strands.HasValidRootData() && Instance->Strands.DeformedRootResource->IsValid(MeshLODIndex) && InstanceBindingType == EHairBindingType::Skinning)
 					{
 						AddHairStrandUpdateMeshTrianglesPass(
 							GraphBuilder, 
@@ -195,7 +196,7 @@ void RunHairStrandsInterpolation(
 					if (Instance->Cards.IsValid(HairLODIndex))
 					{
 						FHairGroupInstance::FCards::FLOD& CardsInstance = Instance->Cards.LODs[HairLODIndex];
-						if (CardsInstance.Guides.IsValid() && Instance->Guides.HasValidRootData() && CardsInstance.Guides.DeformedRootResource->IsValid(MeshLODIndex) && NeedsUpdateCardsMeshTriangles())
+						if (CardsInstance.Guides.IsValid() && Instance->Guides.HasValidRootData() && CardsInstance.Guides.DeformedRootResource->IsValid(MeshLODIndex) && NeedsUpdateCardsMeshTriangles() && InstanceBindingType == EHairBindingType::Skinning)
 						{
 							AddHairStrandUpdateMeshTrianglesPass(
 								GraphBuilder,
@@ -231,7 +232,7 @@ void RunHairStrandsInterpolation(
 			}
 			else if (EHairStrandsInterpolationType::SimulationStrands == Type)
 			{
-				if (Instance->Guides.IsValid() && Instance->Guides.HasValidRootData() && Instance->Guides.DeformedRootResource->IsValid(MeshLODIndex))
+				if (Instance->Guides.IsValid() && Instance->Guides.HasValidRootData() && Instance->Guides.DeformedRootResource->IsValid(MeshLODIndex) )
 				{
 					AddHairStrandUpdateMeshTrianglesPass(
 						GraphBuilder,
diff --git a/Engine/Plugins/Runtime/HairStrands/Source/HairStrandsCore/Private/HairStrandsRendering.cpp b/Engine/Plugins/Runtime/HairStrands/Source/HairStrandsCore/Private/HairStrandsRendering.cpp
--- a/Engine/Plugins/Runtime/HairStrands/Source/HairStrandsCore/Private/HairStrandsRendering.cpp
+++ b/Engine/Plugins/Runtime/HairStrands/Source/HairStrandsCore/Private/HairStrandsRendering.cpp
@@ -850,6 +850,7 @@ static void AddHairCardsDeformationPass(
 
 	const bool bIsVertexToCurveBuffersValid = RestRootResources && RestRootResources->VertexToCurveIndexBuffer.Buffer != nullptr;
 	const bool bSupportDynamicMesh =
+		Instance->BindingType == EHairBindingType::Skinning &&
 		bIsVertexToCurveBuffersValid &&
 		RestRootResources &&
 		RestRootResources->RootData.RootCount > 0 &&
