Prototype Wrapped Object Pointers (ObjectHandle/ObjectPtr)

Call-site changes affecting compilation of QAGame for non-scalar wrapped object pointer upgrades in Engine + Shootergame

diff --git a/Engine/Plugins/Experimental/ControlRig/Source/ControlRig/Private/ControlRigComponent.cpp b/Engine/Plugins/Experimental/ControlRig/Source/ControlRig/Private/ControlRigComponent.cpp
--- a/Engine/Plugins/Experimental/ControlRig/Source/ControlRig/Private/ControlRigComponent.cpp
+++ b/Engine/Plugins/Experimental/ControlRig/Source/ControlRig/Private/ControlRigComponent.cpp
@@ -1190,7 +1190,7 @@ void UControlRigComponent::ValidateMappingData()
 			MappedElement.ElementIndex = INDEX_NONE;
 			MappedElement.SubIndex = INDEX_NONE;
 
-			AActor* MappedOwner = MappedElement.ComponentReference.OtherActor == nullptr ? GetOwner() : MappedElement.ComponentReference.OtherActor;
+			AActor* MappedOwner = MappedElement.ComponentReference.OtherActor == nullptr ? GetOwner() : ToRawPtr(MappedElement.ComponentReference.OtherActor);
 			MappedElement.SceneComponent = Cast<USceneComponent>(MappedElement.ComponentReference.GetComponent(MappedOwner));
 
 			if (MappedElement.SceneComponent == nullptr ||
diff --git a/Engine/Plugins/Runtime/HairStrands/Source/HairStrandsCore/Private/Niagara/NiagaraDataInterfacePhysicsAsset.cpp b/Engine/Plugins/Runtime/HairStrands/Source/HairStrandsCore/Private/Niagara/NiagaraDataInterfacePhysicsAsset.cpp
--- a/Engine/Plugins/Runtime/HairStrands/Source/HairStrandsCore/Private/Niagara/NiagaraDataInterfacePhysicsAsset.cpp
+++ b/Engine/Plugins/Runtime/HairStrands/Source/HairStrandsCore/Private/Niagara/NiagaraDataInterfacePhysicsAsset.cpp
@@ -132,7 +132,7 @@ void GetNumPrimitives(const TArray<TWeakObjectPtr<UPhysicsAsset>>& PhysicsAssets
 		TWeakObjectPtr<UPhysicsAsset> PhysicsAsset = PhysicsAssets[ComponentIndex];
 		if (PhysicsAsset.IsValid() && PhysicsAsset.Get() != nullptr)
 		{
-			USkeletalMesh* SkelMesh = (SkeletalMeshs[ComponentIndex].Get() && SkeletalMeshs[ComponentIndex]->SkeletalMesh) ? SkeletalMeshs[ComponentIndex]->SkeletalMesh : PhysicsAsset->GetPreviewMesh();
+			USkeletalMesh* SkelMesh = (SkeletalMeshs[ComponentIndex].Get() && SkeletalMeshs[ComponentIndex]->SkeletalMesh) ? ToRawPtr(SkeletalMeshs[ComponentIndex]->SkeletalMesh) : PhysicsAsset->GetPreviewMesh();
 			if (!SkelMesh)
 			{
 				continue;
@@ -214,7 +214,7 @@ void CreateInternalArrays(const TArray<TWeakObjectPtr<UPhysicsAsset>>& PhysicsAs
 				TWeakObjectPtr<UPhysicsAsset> PhysicsAsset = PhysicsAssets[ComponentIndex];
 				if (PhysicsAsset.IsValid() && PhysicsAsset.Get() != nullptr)
 				{
-					USkeletalMesh* SkelMesh = (SkeletalMeshs[ComponentIndex].Get() && SkeletalMeshs[ComponentIndex]->SkeletalMesh) ? SkeletalMeshs[ComponentIndex]->SkeletalMesh : PhysicsAsset->GetPreviewMesh();
+					USkeletalMesh* SkelMesh = (SkeletalMeshs[ComponentIndex].Get() && SkeletalMeshs[ComponentIndex]->SkeletalMesh) ? ToRawPtr(SkeletalMeshs[ComponentIndex]->SkeletalMesh) : PhysicsAsset->GetPreviewMesh();
 					if (!SkelMesh)
 					{
 						continue;
@@ -324,7 +324,7 @@ void UpdateInternalArrays(const TArray<TWeakObjectPtr<UPhysicsAsset>>& PhysicsAs
 			TWeakObjectPtr<UPhysicsAsset> PhysicsAsset = PhysicsAssets[ComponentIndex];
 			if (PhysicsAsset.IsValid() && PhysicsAsset.Get() != nullptr)
 			{
-				USkeletalMesh* SkelMesh = (SkeletalMeshs[ComponentIndex].Get() && SkeletalMeshs[ComponentIndex]->SkeletalMesh) ? SkeletalMeshs[ComponentIndex]->SkeletalMesh : PhysicsAsset->GetPreviewMesh();
+				USkeletalMesh* SkelMesh = (SkeletalMeshs[ComponentIndex].Get() && SkeletalMeshs[ComponentIndex]->SkeletalMesh) ? ToRawPtr(SkeletalMeshs[ComponentIndex]->SkeletalMesh) : PhysicsAsset->GetPreviewMesh();
 				if (!SkelMesh)
 				{
 					continue;
@@ -445,7 +445,7 @@ void FNDIPhysicsAssetData::Init(UNiagaraDataInterfacePhysicsAsset* Interface, FN
 			if (PhysicsAsset.IsValid() && PhysicsAsset.Get() != nullptr)
 			{
 				USkeletalMesh* SkelMesh = (Interface->SourceComponents[ComponentIndex].IsValid() && Interface->SourceComponents[ComponentIndex].Get() && Interface->SourceComponents[ComponentIndex]->SkeletalMesh) ?
-					Interface->SourceComponents[ComponentIndex]->SkeletalMesh : PhysicsAsset->GetPreviewMesh();
+					ToRawPtr(Interface->SourceComponents[ComponentIndex]->SkeletalMesh) : PhysicsAsset->GetPreviewMesh();
 				if (SkelMesh)
 				{
 					BoundingBox += SkelMesh->GetImportedBounds().GetBox();
