Fix invalid dirty on attachment change.

diff --git a/Engine/Plugins/Runtime/HairStrands/Source/HairStrandsCore/Private/GroomComponent.cpp b/Engine/Plugins/Runtime/HairStrands/Source/HairStrandsCore/Private/GroomComponent.cpp
--- a/Engine/Plugins/Runtime/HairStrands/Source/HairStrandsCore/Private/GroomComponent.cpp
+++ b/Engine/Plugins/Runtime/HairStrands/Source/HairStrandsCore/Private/GroomComponent.cpp
@@ -2374,6 +2374,14 @@ void UGroomComponent::OnAttachmentChanged()
 	}
 }
 
+// Override the DetachFromComponent so that we do not mark the actor as modified/dirty when attaching to a specific bone
+void UGroomComponent::DetachFromComponent(const FDetachmentTransformRules& In)
+{
+	FDetachmentTransformRules DetachmentRules = In;
+	DetachmentRules.bCallModify = false;
+	Super::DetachFromComponent(DetachmentRules);
+}
+
 void UGroomComponent::TickComponent(float DeltaTime, enum ELevelTick TickType, FActorComponentTickFunction *ThisTickFunction)
 {
 	Super::TickComponent(DeltaTime, TickType, ThisTickFunction);	
diff --git a/Engine/Plugins/Runtime/HairStrands/Source/HairStrandsCore/Public/GroomComponent.h b/Engine/Plugins/Runtime/HairStrands/Source/HairStrandsCore/Public/GroomComponent.h
--- a/Engine/Plugins/Runtime/HairStrands/Source/HairStrandsCore/Public/GroomComponent.h
+++ b/Engine/Plugins/Runtime/HairStrands/Source/HairStrandsCore/Public/GroomComponent.h
@@ -93,6 +93,7 @@ public:
 	virtual void OnUnregister() override;
 	virtual void OnComponentDestroyed(bool bDestroyingHierarchy) override;
 	virtual void OnAttachmentChanged() override;
+	virtual void DetachFromComponent(const FDetachmentTransformRules& DetachmentRules) override;
 	virtual void TickComponent(float DeltaTime, enum ELevelTick TickType, FActorComponentTickFunction *ThisTickFunction) override;
 	virtual void SendRenderTransform_Concurrent() override;
 	//~ End UActorComponent Interface.
