Disable manual skin cache path for hair, as output buffer does not respect the section layout, causing issue when the triangle data are fetch later one. Hair requires skin cache in any case, as the manual skin cache path does not send the skin data correctly and can result in animation stuttering.

diff --git a/Engine/Plugins/Runtime/HairStrands/Source/HairStrandsCore/Private/GroomManager.cpp b/Engine/Plugins/Runtime/HairStrands/Source/HairStrandsCore/Private/GroomManager.cpp
--- a/Engine/Plugins/Runtime/HairStrands/Source/HairStrandsCore/Private/GroomManager.cpp
+++ b/Engine/Plugins/Runtime/HairStrands/Source/HairStrandsCore/Private/GroomManager.cpp
@@ -25,8 +25,8 @@ static FAutoConsoleVariableRef CVarHairStrands_UseCards(TEXT("r.HairStrands.UseC
 static int32 GHairStrands_SwapBufferEndOfFrame = 1;
 static FAutoConsoleVariableRef CVarGHairStrands_SwapBufferEndOfFrame(TEXT("r.HairStrands.SwapEndOfFrame"), GHairStrands_SwapBufferEndOfFrame, TEXT("Swap rendering buffer at the end of frame. This is an experimental toggle. Default:1"));
 
-static int32 GHairStrands_ManualSkinCache = 1;
-static FAutoConsoleVariableRef CVarGHairStrands_ManualSkinCache(TEXT("r.HairStrands.ManualSkinCache"), GHairStrands_ManualSkinCache, TEXT("If skin cache is not enabled, and grooms use skinning method, this enable a simple skin cache mechanisme for groom. Default:enable"));
+static int32 GHairStrands_ManualSkinCache = 0;
+static FAutoConsoleVariableRef CVarGHairStrands_ManualSkinCache(TEXT("r.HairStrands.ManualSkinCache"), GHairStrands_ManualSkinCache, TEXT("If skin cache is not enabled, and grooms use skinning method, this enable a simple skin cache mechanisme for groom. Default:disable"));
 
 bool IsHairStrandsSkinCacheEnable()
 {
