Add CVar for enabling/disabling hair DDC log activity. Disable by default.

diff --git a/Engine/Plugins/Runtime/HairStrands/Source/HairStrandsCore/Private/GroomAsset.cpp b/Engine/Plugins/Runtime/HairStrands/Source/HairStrandsCore/Private/GroomAsset.cpp
--- a/Engine/Plugins/Runtime/HairStrands/Source/HairStrandsCore/Private/GroomAsset.cpp
+++ b/Engine/Plugins/Runtime/HairStrands/Source/HairStrandsCore/Private/GroomAsset.cpp
@@ -47,6 +47,15 @@ static FAutoConsoleVariableRef CVarHairStrandsLoadAsset(TEXT("r.HairStrands.Load
 static int32 GEnableGroomAsyncLoad = 0;
 static FAutoConsoleVariableRef CVarGroomAsyncLoad(TEXT("r.HairStrands.AsyncLoad"), GEnableGroomAsyncLoad, TEXT("Allow groom asset to be loaded asynchronously in the editor"));
 
+// Editor async groom load can be useful in a workflow that consists mostly of loading grooms from a hot DDC
+static int32 GHairStrandsDDCLogEnable = 0;
+static FAutoConsoleVariableRef CVarHairStrandsDDCLogEnable(TEXT("r.HairStrands.DDCLog"), GHairStrandsDDCLogEnable, TEXT("Enable DDC logging for groom assets and groom binding assets"));
+
+bool IsHairStrandsDDCLogEnable()
+{
+	return GHairStrandsDDCLogEnable > 0;
+}
+
 /////////////////////////////////////////////////////////////////////////////////////////
 
 enum class EHairAtlasTextureType
@@ -1524,7 +1533,10 @@ bool UGroomAsset::CacheStrandsData(uint32 GroupIndex, FProcessedHairDescription&
 	TArray<uint8> DerivedData;
 	if (GetDerivedDataCacheRef().GetSynchronous(*DerivedDataKey, DerivedData, GetPathName()))
 	{
-		UE_LOG(LogHairStrands, Log, TEXT("[Groom/DDC] Strands - Found (Groom:%s Group6:%d)."), *GetName(), GroupIndex);
+		if (IsHairStrandsDDCLogEnable())
+		{
+			UE_LOG(LogHairStrands, Log, TEXT("[Groom/DDC] Strands - Found (Groom:%s Group6:%d)."), *GetName(), GroupIndex);
+		}
 		FMemoryReader Ar(DerivedData, /*bIsPersistent=*/ true);
 
 		int64 UncompressedSize = 0;
@@ -1546,7 +1558,10 @@ bool UGroomAsset::CacheStrandsData(uint32 GroupIndex, FProcessedHairDescription&
 			return false;
 		}
 
-		UE_LOG(LogHairStrands, Log, TEXT("[Groom/DDC] Strands - Not found (Groom:%s Group:%d)."), *GetName(), GroupIndex);
+		if (IsHairStrandsDDCLogEnable())
+		{
+			UE_LOG(LogHairStrands, Log, TEXT("[Groom/DDC] Strands - Not found (Groom:%s Group:%d)."), *GetName(), GroupIndex);
+		}
 		// Load the HairDescription from the bulk data if needed
 		if (!HairDescription)
 		{
@@ -1618,13 +1633,19 @@ bool UGroomAsset::CacheCardsGeometry(uint32 GroupIndex, const FString& StrandsKe
 	TArray<uint8> DerivedData;
 	if (GetDerivedDataCacheRef().GetSynchronous(*DerivedDataKey, DerivedData, GetPathName()))
 	{
-		UE_LOG(LogHairStrands, Log, TEXT("[Groom/DDC] Cards - Found (Groom:%s Group:%d)."), *GetName(), GroupIndex);
+		if (IsHairStrandsDDCLogEnable())
+		{
+			UE_LOG(LogHairStrands, Log, TEXT("[Groom/DDC] Cards - Found (Groom:%s Group:%d)."), *GetName(), GroupIndex);
+		}
 		FMemoryReader Ar(DerivedData, true);
 		Ar << HairGroupData.Cards.LODs;
 	}
 	else
 	{
-		UE_LOG(LogHairStrands, Log, TEXT("[Groom/DDC] Cards - Not found (Groom:%s Group:%d)."), *GetName(), GroupIndex);
+		if (IsHairStrandsDDCLogEnable())
+		{
+			UE_LOG(LogHairStrands, Log, TEXT("[Groom/DDC] Cards - Not found (Groom:%s Group:%d)."), *GetName(), GroupIndex);
+		}
 
 		if (!BuildCardsGeometry(GroupIndex))
 		{
@@ -1689,14 +1710,20 @@ bool UGroomAsset::CacheMeshesGeometry(uint32 GroupIndex)
 	TArray<uint8> DerivedData;
 	if (GetDerivedDataCacheRef().GetSynchronous(*DerivedDataKey, DerivedData, GetPathName()))
 	{
-		UE_LOG(LogHairStrands, Log, TEXT("[Groom/DDC] Meshes - Found (Groom:%s Group:%d)."), *GetName(), GroupIndex);
+		if (IsHairStrandsDDCLogEnable())
+		{
+			UE_LOG(LogHairStrands, Log, TEXT("[Groom/DDC] Meshes - Found (Groom:%s Group:%d)."), *GetName(), GroupIndex);
+		}
 
 		FMemoryReader Ar(DerivedData, true);
 		Ar << HairGroupData.Meshes.LODs;
 	}
 	else
 	{
-		UE_LOG(LogHairStrands, Log, TEXT("[Groom/DDC] Meshes - Not found (Groom:%s Group:%d)."), *GetName(), GroupIndex);
+		if (IsHairStrandsDDCLogEnable())
+		{
+			UE_LOG(LogHairStrands, Log, TEXT("[Groom/DDC] Meshes - Not found (Groom:%s Group:%d)."), *GetName(), GroupIndex);
+		}
 
 		if (!BuildMeshesGeometry(GroupIndex))
 		{
diff --git a/Engine/Plugins/Runtime/HairStrands/Source/HairStrandsCore/Private/GroomBindingAsset.cpp b/Engine/Plugins/Runtime/HairStrands/Source/HairStrandsCore/Private/GroomBindingAsset.cpp
--- a/Engine/Plugins/Runtime/HairStrands/Source/HairStrandsCore/Private/GroomBindingAsset.cpp
+++ b/Engine/Plugins/Runtime/HairStrands/Source/HairStrandsCore/Private/GroomBindingAsset.cpp
@@ -14,6 +14,8 @@
 
 /////////////////////////////////////////////////////////////////////////////////////////
 
+bool IsHairStrandsDDCLogEnable();
+
 FArchive& operator<<(FArchive& Ar, FHairStrandsRootData& RootData)
 {
 	RootData.Serialize(Ar);
@@ -613,7 +615,10 @@ void UGroomBindingAsset::CacheDerivedDatas()
 		TArray<uint8> DerivedData;
 		if (GetDerivedDataCacheRef().GetSynchronous(*DerivedDataKey, DerivedData, GetPathName()))
 		{
-			UE_LOG(LogHairStrands, Log, TEXT("[GroomBinding/DDC] Found (GroomiBinding:%s)."), *GetName());
+			if (IsHairStrandsDDCLogEnable())
+			{
+				UE_LOG(LogHairStrands, Log, TEXT("[GroomBinding/DDC] Found (GroomiBinding:%s)."), *GetName());
+			}
 
 			FMemoryReader Ar(DerivedData, /*bIsPersistent=*/ true);
 
@@ -630,7 +635,10 @@ void UGroomBindingAsset::CacheDerivedDatas()
 		}
 		else
 		{
-			UE_LOG(LogHairStrands, Log, TEXT("[GroomBinding/DDC] Not found (GroomiBinding:%s)."), *GetName());
+			if (IsHairStrandsDDCLogEnable())
+			{
+				UE_LOG(LogHairStrands, Log, TEXT("[GroomBinding/DDC] Not found (GroomiBinding:%s)."), *GetName());
+			}
 
 			// Build groom binding data
 			bIsValid = FGroomBindingBuilder::BuildBinding(this, false, false);
