Rename all the hair resources (textures/buffers) to be prefixed with Hair. to help filtering them out with the VIS command.

diff --git a/Engine/Plugins/Runtime/HairStrands/Source/HairStrandsCore/Private/GroomDebug.cpp b/Engine/Plugins/Runtime/HairStrands/Source/HairStrandsCore/Private/GroomDebug.cpp
--- a/Engine/Plugins/Runtime/HairStrands/Source/HairStrandsCore/Private/GroomDebug.cpp
+++ b/Engine/Plugins/Runtime/HairStrands/Source/HairStrandsCore/Private/GroomDebug.cpp
@@ -862,7 +862,7 @@ void RunHairStrandsDebug(
 			FRDGTextureRef DepthTexture;
 			{
 				const FRDGTextureDesc Desc = FRDGTextureDesc::Create2D(SceneColorTexture->Desc.Extent, PF_DepthStencil, FClearValueBinding::DepthFar, TexCreate_DepthStencilTargetable | TexCreate_ShaderResource);
-				DepthTexture = GraphBuilder.CreateTexture(Desc, TEXT("HairInterpolationDepthTexture"));
+				DepthTexture = GraphBuilder.CreateTexture(Desc, TEXT("Hair.InterpolationDepthTexture"));
 			}
 
 			if (GHairDebugMeshProjection_SkinCacheMesh > 0)
diff --git a/Engine/Plugins/Runtime/HairStrands/Source/HairStrandsCore/Private/GroomGeometryCache.cpp b/Engine/Plugins/Runtime/HairStrands/Source/HairStrandsCore/Private/GroomGeometryCache.cpp
--- a/Engine/Plugins/Runtime/HairStrands/Source/HairStrandsCore/Private/GroomGeometryCache.cpp
+++ b/Engine/Plugins/Runtime/HairStrands/Source/HairStrandsCore/Private/GroomGeometryCache.cpp
@@ -81,9 +81,9 @@ static void BuildBoneMatrices(USkeletalMeshComponent* SkeletalMeshComponent, con
 		TArray<FVector4> BoneMatrices;
 		BuildBoneMatrices(SkeletalMeshComponent, LODData, LODIndex, MatrixOffsets, BoneMatrices);
 
-		FRDGBufferRef DeformedPositionsBuffer	= GraphBuilder.CreateBuffer(FRDGBufferDesc::CreateBufferDesc(sizeof(float), LODData.StaticVertexBuffers.PositionVertexBuffer.GetNumVertices() * 3), TEXT("HairStrandsSkinnedDeformedPositions"));
-		FRDGBufferRef BoneMatricesBuffer		= CreateStructuredBuffer(GraphBuilder, TEXT("HairStrandsSkinnedBoneMatrices"), sizeof(float) * 4, BoneMatrices.Num(), BoneMatrices.GetData(), sizeof(float) * 4 * BoneMatrices.Num());
-		FRDGBufferRef MatrixOffsetsBuffer		= CreateStructuredBuffer(GraphBuilder, TEXT("HairStrandsSkinnedMatrixOffsets"), sizeof(uint32), MatrixOffsets.Num(), MatrixOffsets.GetData(), sizeof(uint32) * MatrixOffsets.Num());
+		FRDGBufferRef DeformedPositionsBuffer	= GraphBuilder.CreateBuffer(FRDGBufferDesc::CreateBufferDesc(sizeof(float), LODData.StaticVertexBuffers.PositionVertexBuffer.GetNumVertices() * 3), TEXT("Hair.SkinnedDeformedPositions"));
+		FRDGBufferRef BoneMatricesBuffer		= CreateStructuredBuffer(GraphBuilder, TEXT("Hair.SkinnedBoneMatrices"), sizeof(float) * 4, BoneMatrices.Num(), BoneMatrices.GetData(), sizeof(float) * 4 * BoneMatrices.Num());
+		FRDGBufferRef MatrixOffsetsBuffer		= CreateStructuredBuffer(GraphBuilder, TEXT("Hair.SkinnedMatrixOffsets"), sizeof(uint32), MatrixOffsets.Num(), MatrixOffsets.GetData(), sizeof(uint32) * MatrixOffsets.Num());
 
 		AddSkinUpdatePass(GraphBuilder, ShaderMap, SkeletalMeshComponent->GetSkinWeightBuffer(LODIndex), LODData, BoneMatricesBuffer, MatrixOffsetsBuffer, DeformedPositionsBuffer);
 
diff --git a/Engine/Plugins/Runtime/HairStrands/Source/HairStrandsCore/Private/GroomResources.cpp b/Engine/Plugins/Runtime/HairStrands/Source/HairStrandsCore/Private/GroomResources.cpp
--- a/Engine/Plugins/Runtime/HairStrands/Source/HairStrandsCore/Private/GroomResources.cpp
+++ b/Engine/Plugins/Runtime/HairStrands/Source/HairStrandsCore/Private/GroomResources.cpp
@@ -269,7 +269,7 @@ void CreateHairStrandsDebugAttributeBuffer(FRDGExternalBuffer* DebugAttributeBuf
 		if (GUsingNullRHI) { return; }
 		FMemMark Mark(FMemStack::Get());
 		FRDGBuilder GraphBuilder(RHICmdList);
-		InternalCreateVertexBufferRDG<FHairStrandsAttributeFormat>(GraphBuilder, VertexCount, *LocalAttributeBuffer, TEXT("HairStrands_DebugAttributeBuffer"));
+		InternalCreateVertexBufferRDG<FHairStrandsAttributeFormat>(GraphBuilder, VertexCount, *LocalAttributeBuffer, TEXT("Hair.Strands_DebugAttributeBuffer"));
 		GraphBuilder.Execute();
 	});
 }
@@ -302,9 +302,9 @@ FHairCardsRestResource::FHairCardsRestResource(const FHairCardsDatas::FRenderDat
 void FHairCardsRestResource::InitRHI()
 {
 	// These resources are kept as regular (i.e., non-RDG resources) as they need to be bound at the input assembly stage by the Vertex declaraction which requires FVertexBuffer type
-	CreateBuffer<FHairCardsPositionFormat>(RenderData.Positions, RestPositionBuffer, TEXT("HairCardsRest_PositionBuffer"));
-	CreateBuffer<FHairCardsNormalFormat>(RenderData.Normals, NormalsBuffer, TEXT("HairCardsRest_NormalBuffer"));
-	CreateBuffer<FHairCardsUVFormat>(RenderData.UVs, UVsBuffer, TEXT("HairCardsRest_UVBuffer"));
+	CreateBuffer<FHairCardsPositionFormat>(RenderData.Positions, RestPositionBuffer, TEXT("Hair.CardsRest_PositionBuffer"));
+	CreateBuffer<FHairCardsNormalFormat>(RenderData.Normals, NormalsBuffer, TEXT("Hair.CardsRest_NormalBuffer"));
+	CreateBuffer<FHairCardsUVFormat>(RenderData.UVs, UVsBuffer, TEXT("Hair.CardsRest_UVBuffer"));
 
 	FSamplerStateRHIRef DefaultSampler = TStaticSamplerState<SF_Bilinear, AM_Clamp, AM_Clamp, AM_Clamp>::GetRHI();
 	DepthSampler = DefaultSampler;
@@ -358,19 +358,19 @@ void FHairCardsProceduralResource::InitRHI()
 	FRHICommandListImmediate& RHICmdList = FRHICommandListExecutor::GetImmediateCommandList();
 	FRDGBuilder GraphBuilder(RHICmdList);
 
-	InternalCreateVertexBufferRDG<FHairCardsAtlasRectFormat>(GraphBuilder, RenderData.CardsRect, AtlasRectBuffer, TEXT("HairCardsProcedural_AtlasRectBuffer"));
-	InternalCreateVertexBufferRDG<FHairCardsDimensionFormat>(GraphBuilder, RenderData.CardsLengths, LengthBuffer, TEXT("HairCardsProcedural_LengthBuffer"));
+	InternalCreateVertexBufferRDG<FHairCardsAtlasRectFormat>(GraphBuilder, RenderData.CardsRect, AtlasRectBuffer, TEXT("Hair.CardsProcedural_AtlasRectBuffer"));
+	InternalCreateVertexBufferRDG<FHairCardsDimensionFormat>(GraphBuilder, RenderData.CardsLengths, LengthBuffer, TEXT("Hair.CardsProcedural_LengthBuffer"));
 
-	InternalCreateVertexBufferRDG<FHairCardsOffsetAndCount>(GraphBuilder, RenderData.CardItToCluster, CardItToClusterBuffer, TEXT("HairCardsProcedural_CardItToClusterBuffer"));
-	InternalCreateVertexBufferRDG<FHairCardsOffsetAndCount>(GraphBuilder, RenderData.ClusterIdToVertices, ClusterIdToVerticesBuffer, TEXT("HairCardsProcedural_ClusterIdToVerticesBuffer"));
-	InternalCreateVertexBufferRDG<FHairCardsBoundsFormat>(GraphBuilder, RenderData.ClusterBounds, ClusterBoundBuffer, TEXT("HairCardsProcedural_ClusterBoundBuffer"));
+	InternalCreateVertexBufferRDG<FHairCardsOffsetAndCount>(GraphBuilder, RenderData.CardItToCluster, CardItToClusterBuffer, TEXT("Hair.CardsProcedural_CardItToClusterBuffer"));
+	InternalCreateVertexBufferRDG<FHairCardsOffsetAndCount>(GraphBuilder, RenderData.ClusterIdToVertices, ClusterIdToVerticesBuffer, TEXT("Hair.CardsProcedural_ClusterIdToVerticesBuffer"));
+	InternalCreateVertexBufferRDG<FHairCardsBoundsFormat>(GraphBuilder, RenderData.ClusterBounds, ClusterBoundBuffer, TEXT("Hair.CardsProcedural_ClusterBoundBuffer"));
 
-	InternalCreateVertexBufferRDG<FHairCardsVoxelDensityFormat>(GraphBuilder, RenderData.VoxelDensity, CardVoxel.DensityBuffer, TEXT("HairCardsProcedural_VoxelDensityBuffer"));
-	InternalCreateVertexBufferRDG<FHairCardsVoxelTangentFormat>(GraphBuilder, RenderData.VoxelTangent, CardVoxel.TangentBuffer, TEXT("HairCardsProcedural_VoxelTangentBuffer"));
-	InternalCreateVertexBufferRDG<FHairCardsVoxelTangentFormat>(GraphBuilder, RenderData.VoxelNormal, CardVoxel.NormalBuffer, TEXT("HairCardsProcedural_VoxelNormalBuffer"));
+	InternalCreateVertexBufferRDG<FHairCardsVoxelDensityFormat>(GraphBuilder, RenderData.VoxelDensity, CardVoxel.DensityBuffer, TEXT("Hair.CardsProcedural_VoxelDensityBuffer"));
+	InternalCreateVertexBufferRDG<FHairCardsVoxelTangentFormat>(GraphBuilder, RenderData.VoxelTangent, CardVoxel.TangentBuffer, TEXT("Hair.CardsProcedural_VoxelTangentBuffer"));
+	InternalCreateVertexBufferRDG<FHairCardsVoxelTangentFormat>(GraphBuilder, RenderData.VoxelNormal, CardVoxel.NormalBuffer, TEXT("Hair.CardsProcedural_VoxelNormalBuffer"));
 
-	InternalCreateVertexBufferRDG<FHairCardsStrandsPositionFormat>(GraphBuilder, RenderData.CardsStrandsPositions, CardsStrandsPositions, TEXT("HairCardsProcedural_CardsStrandsPositions"));
-	InternalCreateVertexBufferRDG<FHairCardsStrandsAttributeFormat>(GraphBuilder, RenderData.CardsStrandsAttributes, CardsStrandsAttributes, TEXT("HairCardsProcedural_CardsStrandsAttributes"));
+	InternalCreateVertexBufferRDG<FHairCardsStrandsPositionFormat>(GraphBuilder, RenderData.CardsStrandsPositions, CardsStrandsPositions, TEXT("Hair.CardsProcedural_CardsStrandsPositions"));
+	InternalCreateVertexBufferRDG<FHairCardsStrandsAttributeFormat>(GraphBuilder, RenderData.CardsStrandsAttributes, CardsStrandsAttributes, TEXT("Hair.CardsProcedural_CardsStrandsAttributes"));
 
 	GraphBuilder.Execute();
 }
@@ -407,18 +407,18 @@ void FHairCardsDeformedResource::InitRHI()
 	const uint32 VertexCount = RenderData.Positions.Num();
 	if (bInitializedData)
 	{
-		InternalCreateVertexBufferRDG<FHairCardsPositionFormat>(GraphBuilder, RenderData.Positions, DeformedPositionBuffer[0], TEXT("HairCardsDeformed_Position0"));
+		InternalCreateVertexBufferRDG<FHairCardsPositionFormat>(GraphBuilder, RenderData.Positions, DeformedPositionBuffer[0], TEXT("Hair.CardsDeformed_Position0"));
 		if (bDynamic)
 		{
-			InternalCreateVertexBufferRDG<FHairCardsPositionFormat>(GraphBuilder, RenderData.Positions, DeformedPositionBuffer[1], TEXT("HairCardsDeformed_Position1"));
+			InternalCreateVertexBufferRDG<FHairCardsPositionFormat>(GraphBuilder, RenderData.Positions, DeformedPositionBuffer[1], TEXT("Hair.CardsDeformed_Position1"));
 		}
 	}
 	else
 	{
-		InternalCreateVertexBufferRDG<FHairCardsPositionFormat>(GraphBuilder, VertexCount, DeformedPositionBuffer[0], TEXT("HairCardsDeformed_Position0"));
+		InternalCreateVertexBufferRDG<FHairCardsPositionFormat>(GraphBuilder, VertexCount, DeformedPositionBuffer[0], TEXT("Hair.CardsDeformed_Position0"));
 		if (bDynamic)
 		{
-			InternalCreateVertexBufferRDG<FHairCardsPositionFormat>(GraphBuilder, VertexCount, DeformedPositionBuffer[1], TEXT("HairCardsDeformed_Position1"));
+			InternalCreateVertexBufferRDG<FHairCardsPositionFormat>(GraphBuilder, VertexCount, DeformedPositionBuffer[1], TEXT("Hair.CardsDeformed_Position1"));
 		}
 	}
 	GraphBuilder.Execute();
@@ -451,9 +451,9 @@ FHairMeshesRestResource::FHairMeshesRestResource(const FHairMeshesDatas::FRender
 void FHairMeshesRestResource::InitRHI()
 {
 	// These resources are kept as regular (i.e., non-RDG resources) as they need to be bound at the input assembly stage by the Vertex declaraction which requires FVertexBuffer type
-	CreateBuffer<FHairCardsPositionFormat>(RenderData.Positions, RestPositionBuffer, TEXT("HairMeshesRest_Positions"));
-	CreateBuffer<FHairCardsNormalFormat>(RenderData.Normals, NormalsBuffer, TEXT("HairMeshesRest_Normals"));
-	CreateBuffer<FHairCardsUVFormat>(RenderData.UVs, UVsBuffer, TEXT("HairMeshesRest_UVs"));
+	CreateBuffer<FHairCardsPositionFormat>(RenderData.Positions, RestPositionBuffer, TEXT("Hair.MeshesRest_Positions"));
+	CreateBuffer<FHairCardsNormalFormat>(RenderData.Normals, NormalsBuffer, TEXT("Hair.MeshesRest_Normals"));
+	CreateBuffer<FHairCardsUVFormat>(RenderData.UVs, UVsBuffer, TEXT("Hair.MeshesRest_UVs"));
 }
 
 void FHairMeshesRestResource::ReleaseRHI()
@@ -496,18 +496,18 @@ void FHairMeshesDeformedResource::InitRHI()
 	const uint32 VertexCount = RenderData.Positions.Num();
 	if (bInitializedData)
 	{
-		InternalCreateVertexBufferRDG<FHairCardsPositionFormat>(GraphBuilder, RenderData.Positions, DeformedPositionBuffer[0], TEXT("HairMeshesDeformed_Positions0"));
+		InternalCreateVertexBufferRDG<FHairCardsPositionFormat>(GraphBuilder, RenderData.Positions, DeformedPositionBuffer[0], TEXT("Hair.MeshesDeformed_Positions0"));
 		if (bDynamic)
 		{
-			InternalCreateVertexBufferRDG<FHairCardsPositionFormat>(GraphBuilder, RenderData.Positions, DeformedPositionBuffer[1], TEXT("HairMeshesDeformed_Positions1"));
+			InternalCreateVertexBufferRDG<FHairCardsPositionFormat>(GraphBuilder, RenderData.Positions, DeformedPositionBuffer[1], TEXT("Hair.MeshesDeformed_Positions1"));
 		}
 	}
 	else
 	{
-		InternalCreateVertexBufferRDG<FHairCardsPositionFormat>(GraphBuilder, VertexCount, DeformedPositionBuffer[0], TEXT("HairMeshesDeformed_Positions0"));
+		InternalCreateVertexBufferRDG<FHairCardsPositionFormat>(GraphBuilder, VertexCount, DeformedPositionBuffer[0], TEXT("Hair.MeshesDeformed_Positions0"));
 		if (bDynamic)
 		{
-			InternalCreateVertexBufferRDG<FHairCardsPositionFormat>(GraphBuilder, VertexCount, DeformedPositionBuffer[1], TEXT("HairMeshesDeformed_Positions1"));
+			InternalCreateVertexBufferRDG<FHairCardsPositionFormat>(GraphBuilder, VertexCount, DeformedPositionBuffer[1], TEXT("Hair.MeshesDeformed_Positions1"));
 		}
 	}
 
@@ -536,9 +536,9 @@ void FHairStrandsRestResource::InitRHI()
 	FRHICommandListImmediate& RHICmdList = FRHICommandListExecutor::GetImmediateCommandList();
 	FRDGBuilder GraphBuilder(RHICmdList);
 
-	InternalCreateVertexBufferRDG<FHairStrandsPositionFormat>(GraphBuilder, RenderData.Positions, RestPositionBuffer, TEXT("HairStrandsRest_RestPositionBuffer"));
-	InternalCreateVertexBufferRDG<FHairStrandsAttributeFormat>(GraphBuilder, RenderData.Attributes, AttributeBuffer, TEXT("HairStrandsRest_AttributeBuffer"));
-	InternalCreateVertexBufferRDG<FHairStrandsMaterialFormat>(GraphBuilder, RenderData.Materials, MaterialBuffer, TEXT("HairStrandsRest_MaterialBuffer"));
+	InternalCreateVertexBufferRDG<FHairStrandsPositionFormat>(GraphBuilder, RenderData.Positions, RestPositionBuffer, TEXT("Hair.StrandsRest_RestPositionBuffer"));
+	InternalCreateVertexBufferRDG<FHairStrandsAttributeFormat>(GraphBuilder, RenderData.Attributes, AttributeBuffer, TEXT("Hair.StrandsRest_AttributeBuffer"));
+	InternalCreateVertexBufferRDG<FHairStrandsMaterialFormat>(GraphBuilder, RenderData.Materials, MaterialBuffer, TEXT("Hair.StrandsRest_MaterialBuffer"));
 
 	GraphBuilder.Execute();
 }
@@ -566,26 +566,26 @@ void FHairStrandsDeformedResource::InitRHI()
 	const uint32 VertexCount = RenderData.Positions.Num();
 	if (bInitializedData)
 	{
-		InternalCreateVertexBufferRDG<FHairStrandsPositionFormat>(GraphBuilder, RenderData.Positions, DeformedPositionBuffer[0], TEXT("HairStrandsDeformed_DeformedPositionBuffer0")); // , bDynamic ? ERHIAccess::UAVCompute : ERHIAccess::SRVMask
+		InternalCreateVertexBufferRDG<FHairStrandsPositionFormat>(GraphBuilder, RenderData.Positions, DeformedPositionBuffer[0], TEXT("Hair.StrandsDeformed_DeformedPositionBuffer0")); // , bDynamic ? ERHIAccess::UAVCompute : ERHIAccess::SRVMask
 		if (bDynamic)
 		{
-			InternalCreateVertexBufferRDG<FHairStrandsPositionFormat>(GraphBuilder, RenderData.Positions, DeformedPositionBuffer[1], TEXT("HairStrandsDeformed_DeformedPositionBuffer1")); // , ERHIAccess::UAVCompute
+			InternalCreateVertexBufferRDG<FHairStrandsPositionFormat>(GraphBuilder, RenderData.Positions, DeformedPositionBuffer[1], TEXT("Hair.StrandsDeformed_DeformedPositionBuffer1")); // , ERHIAccess::UAVCompute
 		}
 	}
 	else
 	{
-		InternalCreateVertexBufferRDG<FHairStrandsPositionFormat>(GraphBuilder, VertexCount, DeformedPositionBuffer[0], TEXT("HairStrandsDeformed_DeformedPositionBuffer0"));
+		InternalCreateVertexBufferRDG<FHairStrandsPositionFormat>(GraphBuilder, VertexCount, DeformedPositionBuffer[0], TEXT("Hair.StrandsDeformed_DeformedPositionBuffer0"));
 		if (bDynamic)
 		{
-			InternalCreateVertexBufferRDG<FHairStrandsPositionFormat>(GraphBuilder, VertexCount, DeformedPositionBuffer[1], TEXT("HairStrandsDeformed_DeformedPositionBuffer1"));
+			InternalCreateVertexBufferRDG<FHairStrandsPositionFormat>(GraphBuilder, VertexCount, DeformedPositionBuffer[1], TEXT("Hair.StrandsDeformed_DeformedPositionBuffer1"));
 		}
 	}
-	InternalCreateVertexBufferRDG<FHairStrandsTangentFormat>(GraphBuilder, VertexCount * FHairStrandsTangentFormat::ComponentCount, TangentBuffer, TEXT("HairStrandsDeformed_TangentBuffer"));
+	InternalCreateVertexBufferRDG<FHairStrandsTangentFormat>(GraphBuilder, VertexCount * FHairStrandsTangentFormat::ComponentCount, TangentBuffer, TEXT("Hair.StrandsDeformed_TangentBuffer"));
 
 	TArray<FVector4> DefaultOffsets;
 	DefaultOffsets.Add(DefaultOffset);
-	InternalCreateVertexBufferRDG<FHairStrandsPositionOffsetFormat>(GraphBuilder, DefaultOffsets, DeformedOffsetBuffer[0], TEXT("HairStrandsDeformed_DeformedOffsetBuffer0"));
-	InternalCreateVertexBufferRDG<FHairStrandsPositionOffsetFormat>(GraphBuilder, DefaultOffsets, DeformedOffsetBuffer[1], TEXT("HairStrandsDeformed_DeformedOffsetBuffer1"));
+	InternalCreateVertexBufferRDG<FHairStrandsPositionOffsetFormat>(GraphBuilder, DefaultOffsets, DeformedOffsetBuffer[0], TEXT("Hair.StrandsDeformed_DeformedOffsetBuffer0"));
+	InternalCreateVertexBufferRDG<FHairStrandsPositionOffsetFormat>(GraphBuilder, DefaultOffsets, DeformedOffsetBuffer[1], TEXT("Hair.StrandsDeformed_DeformedOffsetBuffer1"));
 
 	GraphBuilder.Execute();
 }
@@ -659,11 +659,11 @@ void FHairStrandsClusterCullingResource::InitRHI()
 	FRHICommandListImmediate& RHICmdList = FRHICommandListExecutor::GetImmediateCommandList();
 	FRDGBuilder GraphBuilder(RHICmdList);
 
-	InternalCreateStructuredBufferRDG(GraphBuilder, Data.PackedClusterInfos, sizeof(FHairStrandsClusterCullingData::FHairClusterInfo::Packed), ClusterInfoBuffer, TEXT("HairStrandsClusterCulling_ClusterInfoBuffer"));
-	InternalCreateStructuredBufferRDG(GraphBuilder, Data.ClusterLODInfos, sizeof(FHairStrandsClusterCullingData::FHairClusterLODInfo), ClusterLODInfoBuffer, TEXT("HairStrandsClusterCulling_ClusterLODInfoBuffer"));
+	InternalCreateStructuredBufferRDG(GraphBuilder, Data.PackedClusterInfos, sizeof(FHairStrandsClusterCullingData::FHairClusterInfo::Packed), ClusterInfoBuffer, TEXT("Hair.StrandsClusterCulling_ClusterInfoBuffer"));
+	InternalCreateStructuredBufferRDG(GraphBuilder, Data.ClusterLODInfos, sizeof(FHairStrandsClusterCullingData::FHairClusterLODInfo), ClusterLODInfoBuffer, TEXT("Hair.StrandsClusterCulling_ClusterLODInfoBuffer"));
 
-	InternalCreateVertexBufferRDG(GraphBuilder, Data.ClusterVertexIds, PF_R32_UINT, ClusterVertexIdBuffer, TEXT("HairStrandsClusterCulling_ClusterVertexIds"));
-	InternalCreateVertexBufferRDG(GraphBuilder, Data.VertexToClusterIds, PF_R32_UINT, VertexToClusterIdBuffer, TEXT("HairStrandsClusterCulling_VertexToClusterIds"));
+	InternalCreateVertexBufferRDG(GraphBuilder, Data.ClusterVertexIds, PF_R32_UINT, ClusterVertexIdBuffer, TEXT("Hair.StrandsClusterCulling_ClusterVertexIds"));
+	InternalCreateVertexBufferRDG(GraphBuilder, Data.VertexToClusterIds, PF_R32_UINT, VertexToClusterIdBuffer, TEXT("Hair.StrandsClusterCulling_VertexToClusterIds"));
 
 	GraphBuilder.Execute();
 }
@@ -712,9 +712,9 @@ void FHairStrandsRestRootResource::InitRHI()
 		FRHICommandListImmediate& RHICmdList = FRHICommandListExecutor::GetImmediateCommandList();
 		FRDGBuilder GraphBuilder(RHICmdList);
 
-		InternalCreateVertexBufferRDG<FHairStrandsIndexFormat>(GraphBuilder, RootData.VertexToCurveIndexBuffer, VertexToCurveIndexBuffer, TEXT("HairStrandsRestRoot_VertexToCurveIndexBuffer"));
-		InternalCreateVertexBufferRDG<FHairStrandsRootPositionFormat>(GraphBuilder, RootData.RootPositionBuffer, RootPositionBuffer, TEXT("HairStrandsRestRoot_RootPositionBuffer"));
-		InternalCreateVertexBufferRDG<FHairStrandsRootNormalFormat>(GraphBuilder, RootData.RootNormalBuffer, RootNormalBuffer, TEXT("HairStrandsRestRoot_RootNormalBuffer"));
+		InternalCreateVertexBufferRDG<FHairStrandsIndexFormat>(GraphBuilder, RootData.VertexToCurveIndexBuffer, VertexToCurveIndexBuffer, TEXT("Hair.StrandsRestRoot_VertexToCurveIndexBuffer"));
+		InternalCreateVertexBufferRDG<FHairStrandsRootPositionFormat>(GraphBuilder, RootData.RootPositionBuffer, RootPositionBuffer, TEXT("Hair.StrandsRestRoot_RootPositionBuffer"));
+		InternalCreateVertexBufferRDG<FHairStrandsRootNormalFormat>(GraphBuilder, RootData.RootNormalBuffer, RootNormalBuffer, TEXT("Hair.StrandsRestRoot_RootNormalBuffer"));
 		
 		check(LODs.Num() == RootData.MeshProjectionLODs.Num());
 		for (uint32 LODIt=0, LODCount = LODs.Num(); LODIt<LODCount; ++LODIt)
@@ -728,29 +728,29 @@ void FHairStrandsRestRootResource::InitRHI()
 				GPUData.Status = FLOD::EStatus::Completed;
 
 				check(CPUData.RootTriangleBarycentricBuffer.Num() > 0);
-				InternalCreateVertexBufferRDG<FHairStrandsCurveTriangleBarycentricFormat>(GraphBuilder, CPUData.RootTriangleBarycentricBuffer, GPUData.RootTriangleBarycentricBuffer, TEXT("HairStrandsRestRoot_RootTriangleBarycentricBuffer"));
+				InternalCreateVertexBufferRDG<FHairStrandsCurveTriangleBarycentricFormat>(GraphBuilder, CPUData.RootTriangleBarycentricBuffer, GPUData.RootTriangleBarycentricBuffer, TEXT("Hair.StrandsRestRoot_RootTriangleBarycentricBuffer"));
 
 				check(CPUData.RootTriangleIndexBuffer.Num() > 0);
-				InternalCreateVertexBufferRDG<FHairStrandsCurveTriangleIndexFormat>(GraphBuilder, CPUData.RootTriangleIndexBuffer, GPUData.RootTriangleIndexBuffer, TEXT("HairStrandsRestRoot_RootTriangleIndexBuffer"));
+				InternalCreateVertexBufferRDG<FHairStrandsCurveTriangleIndexFormat>(GraphBuilder, CPUData.RootTriangleIndexBuffer, GPUData.RootTriangleIndexBuffer, TEXT("Hair.StrandsRestRoot_RootTriangleIndexBuffer"));
 
 				check(CPUData.RestRootTrianglePosition0Buffer.Num() > 0);
 				check(CPUData.RestRootTrianglePosition1Buffer.Num() > 0);
 				check(CPUData.RestRootTrianglePosition2Buffer.Num() > 0);
-				InternalCreateVertexBufferRDG<FHairStrandsMeshTrianglePositionFormat>(GraphBuilder, CPUData.RestRootTrianglePosition0Buffer, GPUData.RestRootTrianglePosition0Buffer, TEXT("HairStrandsRestRoot_RestRootTrianglePosition0Buffer"));
-				InternalCreateVertexBufferRDG<FHairStrandsMeshTrianglePositionFormat>(GraphBuilder, CPUData.RestRootTrianglePosition1Buffer, GPUData.RestRootTrianglePosition1Buffer, TEXT("HairStrandsRestRoot_RestRootTrianglePosition1Buffer"));
-				InternalCreateVertexBufferRDG<FHairStrandsMeshTrianglePositionFormat>(GraphBuilder, CPUData.RestRootTrianglePosition2Buffer, GPUData.RestRootTrianglePosition2Buffer, TEXT("HairStrandsRestRoot_RestRootTrianglePosition2Buffer"));
+				InternalCreateVertexBufferRDG<FHairStrandsMeshTrianglePositionFormat>(GraphBuilder, CPUData.RestRootTrianglePosition0Buffer, GPUData.RestRootTrianglePosition0Buffer, TEXT("Hair.StrandsRestRoot_RestRootTrianglePosition0Buffer"));
+				InternalCreateVertexBufferRDG<FHairStrandsMeshTrianglePositionFormat>(GraphBuilder, CPUData.RestRootTrianglePosition1Buffer, GPUData.RestRootTrianglePosition1Buffer, TEXT("Hair.StrandsRestRoot_RestRootTrianglePosition1Buffer"));
+				InternalCreateVertexBufferRDG<FHairStrandsMeshTrianglePositionFormat>(GraphBuilder, CPUData.RestRootTrianglePosition2Buffer, GPUData.RestRootTrianglePosition2Buffer, TEXT("Hair.StrandsRestRoot_RestRootTrianglePosition2Buffer"));
 			}
 			else
 			{
 				GPUData.Status = FLOD::EStatus::Initialized;
 
-				InternalCreateVertexBufferRDG<FHairStrandsCurveTriangleBarycentricFormat>(GraphBuilder, RootData.RootCount, GPUData.RootTriangleBarycentricBuffer, TEXT("HairStrandsRestRoot_RootTriangleBarycentricBuffer"));
-				InternalCreateVertexBufferRDG<FHairStrandsCurveTriangleIndexFormat>(GraphBuilder, RootData.RootCount, GPUData.RootTriangleIndexBuffer, TEXT("HairStrandsRestRoot_RootTriangleIndexBuffer"));
+				InternalCreateVertexBufferRDG<FHairStrandsCurveTriangleBarycentricFormat>(GraphBuilder, RootData.RootCount, GPUData.RootTriangleBarycentricBuffer, TEXT("Hair.StrandsRestRoot_RootTriangleBarycentricBuffer"));
+				InternalCreateVertexBufferRDG<FHairStrandsCurveTriangleIndexFormat>(GraphBuilder, RootData.RootCount, GPUData.RootTriangleIndexBuffer, TEXT("Hair.StrandsRestRoot_RootTriangleIndexBuffer"));
 
 				// Create buffers. Initialization will be done by render passes
-				InternalCreateVertexBufferRDG<FHairStrandsMeshTrianglePositionFormat>(GraphBuilder, RootData.RootCount, GPUData.RestRootTrianglePosition0Buffer, TEXT("HairStrandsRestRoot_RestRootTrianglePosition0Buffer"));
-				InternalCreateVertexBufferRDG<FHairStrandsMeshTrianglePositionFormat>(GraphBuilder, RootData.RootCount, GPUData.RestRootTrianglePosition1Buffer, TEXT("HairStrandsRestRoot_RestRootTrianglePosition1Buffer"));
-				InternalCreateVertexBufferRDG<FHairStrandsMeshTrianglePositionFormat>(GraphBuilder, RootData.RootCount, GPUData.RestRootTrianglePosition2Buffer, TEXT("HairStrandsRestRoot_RestRootTrianglePosition2Buffer"));
+				InternalCreateVertexBufferRDG<FHairStrandsMeshTrianglePositionFormat>(GraphBuilder, RootData.RootCount, GPUData.RestRootTrianglePosition0Buffer, TEXT("Hair.StrandsRestRoot_RestRootTrianglePosition0Buffer"));
+				InternalCreateVertexBufferRDG<FHairStrandsMeshTrianglePositionFormat>(GraphBuilder, RootData.RootCount, GPUData.RestRootTrianglePosition1Buffer, TEXT("Hair.StrandsRestRoot_RestRootTrianglePosition1Buffer"));
+				InternalCreateVertexBufferRDG<FHairStrandsMeshTrianglePositionFormat>(GraphBuilder, RootData.RootCount, GPUData.RestRootTrianglePosition2Buffer, TEXT("Hair.StrandsRestRoot_RestRootTrianglePosition2Buffer"));
 			}
 
 			GPUData.SampleCount = CPUData.SampleCount;
@@ -761,15 +761,15 @@ void FHairStrandsRestRootResource::InitRHI()
 				check(CPUData.MeshSampleIndicesBuffer.Num() == CPUData.SampleCount);
 				check(CPUData.RestSamplePositionsBuffer.Num() == CPUData.SampleCount);
 
-				InternalCreateVertexBufferRDG<FHairStrandsWeightFormat>(GraphBuilder, CPUData.MeshInterpolationWeightsBuffer, GPUData.MeshInterpolationWeightsBuffer, TEXT("HairStrandsRestRoot_MeshInterpolationWeightsBuffer"));
-				InternalCreateVertexBufferRDG<FHairStrandsIndexFormat>(GraphBuilder, CPUData.MeshSampleIndicesBuffer, GPUData.MeshSampleIndicesBuffer, TEXT("HairStrandsRestRoot_MeshSampleIndicesBuffer"));
-				InternalCreateVertexBufferRDG<FHairStrandsMeshTrianglePositionFormat>(GraphBuilder, CPUData.RestSamplePositionsBuffer, GPUData.RestSamplePositionsBuffer, TEXT("HairStrandsRestRoot_RestSamplePositionsBuffer"));
+				InternalCreateVertexBufferRDG<FHairStrandsWeightFormat>(GraphBuilder, CPUData.MeshInterpolationWeightsBuffer, GPUData.MeshInterpolationWeightsBuffer, TEXT("Hair.StrandsRestRoot_MeshInterpolationWeightsBuffer"));
+				InternalCreateVertexBufferRDG<FHairStrandsIndexFormat>(GraphBuilder, CPUData.MeshSampleIndicesBuffer, GPUData.MeshSampleIndicesBuffer, TEXT("Hair.StrandsRestRoot_MeshSampleIndicesBuffer"));
+				InternalCreateVertexBufferRDG<FHairStrandsMeshTrianglePositionFormat>(GraphBuilder, CPUData.RestSamplePositionsBuffer, GPUData.RestSamplePositionsBuffer, TEXT("Hair.StrandsRestRoot_RestSamplePositionsBuffer"));
 			}
 			else
 			{
-				InternalCreateVertexBufferRDG<FHairStrandsWeightFormat>(GraphBuilder, (CPUData.SampleCount+4) * (CPUData.SampleCount+4), GPUData.MeshInterpolationWeightsBuffer, TEXT("HairStrandsRestRoot_MeshInterpolationWeightsBuffer"));
-				InternalCreateVertexBufferRDG<FHairStrandsIndexFormat>(GraphBuilder, CPUData.SampleCount, GPUData.MeshSampleIndicesBuffer, TEXT("HairStrandsRestRoot_MeshSampleIndicesBuffer"));
-				InternalCreateVertexBufferRDG<FHairStrandsMeshTrianglePositionFormat>(GraphBuilder, CPUData.SampleCount, GPUData.RestSamplePositionsBuffer, TEXT("HairStrandsRestRoot_RestSamplePositionsBuffer"));
+				InternalCreateVertexBufferRDG<FHairStrandsWeightFormat>(GraphBuilder, (CPUData.SampleCount+4) * (CPUData.SampleCount+4), GPUData.MeshInterpolationWeightsBuffer, TEXT("Hair.StrandsRestRoot_MeshInterpolationWeightsBuffer"));
+				InternalCreateVertexBufferRDG<FHairStrandsIndexFormat>(GraphBuilder, CPUData.SampleCount, GPUData.MeshSampleIndicesBuffer, TEXT("Hair.StrandsRestRoot_MeshSampleIndicesBuffer"));
+				InternalCreateVertexBufferRDG<FHairStrandsMeshTrianglePositionFormat>(GraphBuilder, CPUData.SampleCount, GPUData.RestSamplePositionsBuffer, TEXT("Hair.StrandsRestRoot_RestSamplePositionsBuffer"));
 			}
 		}
 		GraphBuilder.Execute();
@@ -836,12 +836,12 @@ void FHairStrandsDeformedRootResource::InitRHI()
 		for (FLOD& LOD : LODs)
 		{		
 			LOD.Status = FLOD::EStatus::Initialized;
-			InternalCreateVertexBufferRDG<FHairStrandsMeshTrianglePositionFormat>(GraphBuilder, LOD.SampleCount, LOD.DeformedSamplePositionsBuffer, TEXT("HairStrandsRootDeformed_DeformedSamplePositionsBuffer"));
-			InternalCreateVertexBufferRDG<FHairStrandsMeshTrianglePositionFormat>(GraphBuilder, LOD.SampleCount + 4, LOD.MeshSampleWeightsBuffer, TEXT("HairStrandsRootDeformed_MeshSampleWeightsBuffer"));
+			InternalCreateVertexBufferRDG<FHairStrandsMeshTrianglePositionFormat>(GraphBuilder, LOD.SampleCount, LOD.DeformedSamplePositionsBuffer, TEXT("Hair.StrandsRootDeformed_DeformedSamplePositionsBuffer"));
+			InternalCreateVertexBufferRDG<FHairStrandsMeshTrianglePositionFormat>(GraphBuilder, LOD.SampleCount + 4, LOD.MeshSampleWeightsBuffer, TEXT("Hair.StrandsRootDeformed_MeshSampleWeightsBuffer"));
 
-			InternalCreateVertexBufferRDG<FHairStrandsMeshTrianglePositionFormat>(GraphBuilder, RootCount, LOD.DeformedRootTrianglePosition0Buffer, TEXT("HairStrandsRootDeformed_DeformedRootTrianglePosition0Buffer"));
-			InternalCreateVertexBufferRDG<FHairStrandsMeshTrianglePositionFormat>(GraphBuilder, RootCount, LOD.DeformedRootTrianglePosition1Buffer, TEXT("HairStrandsRootDeformed_DeformedRootTrianglePosition1Buffer"));
-			InternalCreateVertexBufferRDG<FHairStrandsMeshTrianglePositionFormat>(GraphBuilder, RootCount, LOD.DeformedRootTrianglePosition2Buffer, TEXT("HairStrandsRootDeformed_DeformedRootTrianglePosition2Buffer"));
+			InternalCreateVertexBufferRDG<FHairStrandsMeshTrianglePositionFormat>(GraphBuilder, RootCount, LOD.DeformedRootTrianglePosition0Buffer, TEXT("Hair.StrandsRootDeformed_DeformedRootTrianglePosition0Buffer"));
+			InternalCreateVertexBufferRDG<FHairStrandsMeshTrianglePositionFormat>(GraphBuilder, RootCount, LOD.DeformedRootTrianglePosition1Buffer, TEXT("Hair.StrandsRootDeformed_DeformedRootTrianglePosition1Buffer"));
+			InternalCreateVertexBufferRDG<FHairStrandsMeshTrianglePositionFormat>(GraphBuilder, RootCount, LOD.DeformedRootTrianglePosition2Buffer, TEXT("Hair.StrandsRootDeformed_DeformedRootTrianglePosition2Buffer"));
 		}
 	}
 	GraphBuilder.Execute();
@@ -1017,9 +1017,9 @@ void FHairStrandsInterpolationResource::InitRHI()
 	FRHICommandListImmediate& RHICmdList = FRHICommandListExecutor::GetImmediateCommandList();
 	FRDGBuilder GraphBuilder(RHICmdList);
 
-	InternalCreateVertexBufferRDG<FHairStrandsInterpolation0Format>(GraphBuilder, RenderData.Interpolation0, Interpolation0Buffer, TEXT("HairStrandsInterpolation_Interpolation0Buffer"));
-	InternalCreateVertexBufferRDG<FHairStrandsInterpolation1Format>(GraphBuilder, RenderData.Interpolation1, Interpolation1Buffer, TEXT("HairStrandsInterpolation_Interpolation1Buffer"));
-	InternalCreateVertexBufferRDG<FHairStrandsRootIndexFormat>(GraphBuilder, SimRootPointIndex, SimRootPointIndexBuffer, TEXT("HairStrandsInterpolation_SimRootPointIndex"));
+	InternalCreateVertexBufferRDG<FHairStrandsInterpolation0Format>(GraphBuilder, RenderData.Interpolation0, Interpolation0Buffer, TEXT("Hair.StrandsInterpolation_Interpolation0Buffer"));
+	InternalCreateVertexBufferRDG<FHairStrandsInterpolation1Format>(GraphBuilder, RenderData.Interpolation1, Interpolation1Buffer, TEXT("Hair.StrandsInterpolation_Interpolation1Buffer"));
+	InternalCreateVertexBufferRDG<FHairStrandsRootIndexFormat>(GraphBuilder, SimRootPointIndex, SimRootPointIndexBuffer, TEXT("Hair.StrandsInterpolation_SimRootPointIndex"));
 	GraphBuilder.Execute();
 	//SimRootPointIndex.SetNum(0);
 }
@@ -1068,7 +1068,7 @@ void FHairCardsInterpolationResource::InitRHI()
 	FRHICommandListImmediate& RHICmdList = FRHICommandListExecutor::GetImmediateCommandList();
 	FRDGBuilder GraphBuilder(RHICmdList);
 
-	InternalCreateVertexBufferRDG<FHairCardsInterpolationFormat>(GraphBuilder, RenderData.Interpolation, InterpolationBuffer, TEXT("HairCardsInterpolation_InterpolationBuffer"));
+	InternalCreateVertexBufferRDG<FHairCardsInterpolationFormat>(GraphBuilder, RenderData.Interpolation, InterpolationBuffer, TEXT("Hair.CardsInterpolation_InterpolationBuffer"));
 	GraphBuilder.Execute();
 }
 
@@ -1101,7 +1101,7 @@ void FHairStrandsRaytracingResource::InitRHI()
 	FRHICommandListImmediate& RHICmdList = FRHICommandListExecutor::GetImmediateCommandList();
 	FRDGBuilder GraphBuilder(RHICmdList);
 
-	InternalCreateVertexBufferRDG<FHairStrandsRaytracingFormat>(GraphBuilder, VertexCount, PositionBuffer, TEXT("HairStrandsRaytracing_PositionBuffer"));
+	InternalCreateVertexBufferRDG<FHairStrandsRaytracingFormat>(GraphBuilder, VertexCount, PositionBuffer, TEXT("Hair.StrandsRaytracing_PositionBuffer"));
 	GraphBuilder.Execute();
 }
 
diff --git a/Engine/Plugins/Runtime/HairStrands/Source/HairStrandsCore/Private/GroomTextureBuilder.cpp b/Engine/Plugins/Runtime/HairStrands/Source/HairStrandsCore/Private/GroomTextureBuilder.cpp
--- a/Engine/Plugins/Runtime/HairStrands/Source/HairStrandsCore/Private/GroomTextureBuilder.cpp
+++ b/Engine/Plugins/Runtime/HairStrands/Source/HairStrandsCore/Private/GroomTextureBuilder.cpp
@@ -825,35 +825,35 @@ static void InternalBuildStrandsTextures_GPU(
 	Desc.Format = PF_G16;
 	Desc.ClearValue = FClearValueBinding::White;
 	FRDGTextureRef DepthTexture[2];
-	DepthTexture[0] = GraphBuilder.CreateTexture(Desc, TEXT("DepthTexture"));
-	DepthTexture[1] = GraphBuilder.CreateTexture(Desc, TEXT("DepthTexture"));
+	DepthTexture[0] = GraphBuilder.CreateTexture(Desc, TEXT("Hair.DepthTexture"));
+	DepthTexture[1] = GraphBuilder.CreateTexture(Desc, TEXT("Hair.DepthTexture"));
 
 	Desc.Format = PF_B8G8R8A8;
 	Desc.ClearValue = FClearValueBinding::Transparent;
 	FRDGTextureRef CoverageTexture[2];
-	CoverageTexture[0] = GraphBuilder.CreateTexture(Desc, TEXT("CoverageTexture"));
-	CoverageTexture[1] = GraphBuilder.CreateTexture(Desc, TEXT("CoverageTexture"));
+	CoverageTexture[0] = GraphBuilder.CreateTexture(Desc, TEXT("Hair.CoverageTexture"));
+	CoverageTexture[1] = GraphBuilder.CreateTexture(Desc, TEXT("Hair.CoverageTexture"));
 
 	Desc.Format = PF_B8G8R8A8;
 	Desc.ClearValue = FClearValueBinding::Transparent;
 	FRDGTextureRef TangentTexture[2];
-	TangentTexture[0] = GraphBuilder.CreateTexture(Desc, TEXT("TangentTexture"));
-	TangentTexture[1] = GraphBuilder.CreateTexture(Desc, TEXT("TangentTexture"));
+	TangentTexture[0] = GraphBuilder.CreateTexture(Desc, TEXT("Hair.TangentTexture"));
+	TangentTexture[1] = GraphBuilder.CreateTexture(Desc, TEXT("Hair.TangentTexture"));
 
 	Desc.Format = PF_B8G8R8A8;
 	Desc.ClearValue = FClearValueBinding::Transparent;
 	FRDGTextureRef AttributeTexture[2];
-	AttributeTexture[0] = GraphBuilder.CreateTexture(Desc, TEXT("AttributeTexture"));
-	AttributeTexture[1] = GraphBuilder.CreateTexture(Desc, TEXT("AttributeTexture"));
+	AttributeTexture[0] = GraphBuilder.CreateTexture(Desc, TEXT("Hair.AttributeTexture"));
+	AttributeTexture[1] = GraphBuilder.CreateTexture(Desc, TEXT("Hair.AttributeTexture"));
 
 	Desc.Format = PF_R32_UINT;
 	Desc.ClearValue = FClearValueBinding::Black;
-	FRDGTextureRef TriangleMaskTexture = GraphBuilder.CreateTexture(Desc, TEXT("TriangleMaskTexture"));
+	FRDGTextureRef TriangleMaskTexture = GraphBuilder.CreateTexture(Desc, TEXT("Hair.TriangleMaskTexture"));
 	
 	Desc.ClearValue = FClearValueBinding(1);
 	Desc.Format = PF_DepthStencil;
 	Desc.Flags = TexCreate_DepthStencilTargetable;
-	FRDGTextureRef DepthTestTexture = GraphBuilder.CreateTexture(Desc, TEXT("DepthTestTexture"));
+	FRDGTextureRef DepthTestTexture = GraphBuilder.CreateTexture(Desc, TEXT("Hair.DepthTestTexture"));
 
 	bool bClear = true;
 	const uint32 GroupCount = InInfo.GroomAsset->GetNumHairGroups();
diff --git a/Engine/Plugins/Runtime/HairStrands/Source/HairStrandsCore/Private/HairCardsBuilder.cpp b/Engine/Plugins/Runtime/HairStrands/Source/HairStrandsCore/Private/HairCardsBuilder.cpp
--- a/Engine/Plugins/Runtime/HairStrands/Source/HairStrandsCore/Private/HairCardsBuilder.cpp
+++ b/Engine/Plugins/Runtime/HairStrands/Source/HairStrandsCore/Private/HairCardsBuilder.cpp
@@ -276,7 +276,7 @@ static void AddHairCardAtlasTexturePass(
 	FRDGTextureRef OutCoverageTexture,
 	FRDGTextureRef OutAttributeTexture)
 {
-	FRDGTextureRef DepthTestTexture = GraphBuilder.CreateTexture(FRDGTextureDesc::Create2D(OutAttributeTexture->Desc.Extent, PF_DepthStencil, FClearValueBinding::DepthOne, TexCreate_DepthStencilTargetable | TexCreate_ShaderResource, 1), TEXT("CardsDepthTest"));
+	FRDGTextureRef DepthTestTexture = GraphBuilder.CreateTexture(FRDGTextureDesc::Create2D(OutAttributeTexture->Desc.Extent, PF_DepthStencil, FClearValueBinding::DepthOne, TexCreate_DepthStencilTargetable | TexCreate_ShaderResource, 1), TEXT("Hair.CardsDepthTest"));
 
 	bool bClear = true;
 	for (const FHairCardsProceduralAtlas::Rect& Rect : InAtlas.Rects)
@@ -3846,10 +3846,10 @@ void RunHairCardsAtlasQueries(
 			continue;
 
 		// Allocate resources for generating the atlas texture (for editor mode only)
-		FRDGTextureRef DepthTexture		= GraphBuilder.CreateTexture(FRDGTextureDesc::Create2D(Q.ProceduralData->Atlas.Resolution, PF_R8G8B8A8, FClearValueBinding::Black, TexCreate_RenderTargetable | TexCreate_ShaderResource, 1), TEXT("CardsDepth"));
-		FRDGTextureRef TangentTexture	= GraphBuilder.CreateTexture(FRDGTextureDesc::Create2D(Q.ProceduralData->Atlas.Resolution, PF_R8G8B8A8, FClearValueBinding::Black, TexCreate_RenderTargetable | TexCreate_ShaderResource, 1), TEXT("CardTangent"));
-		FRDGTextureRef CoverageTexture	= GraphBuilder.CreateTexture(FRDGTextureDesc::Create2D(Q.ProceduralData->Atlas.Resolution, PF_R8G8B8A8, FClearValueBinding::Black, TexCreate_RenderTargetable | TexCreate_ShaderResource, 1), TEXT("CardCoverage"));
-		FRDGTextureRef AttributeTexture = GraphBuilder.CreateTexture(FRDGTextureDesc::Create2D(Q.ProceduralData->Atlas.Resolution, PF_R8G8B8A8, FClearValueBinding::Black, TexCreate_RenderTargetable | TexCreate_ShaderResource, 1), TEXT("CardAttribute"));
+		FRDGTextureRef DepthTexture		= GraphBuilder.CreateTexture(FRDGTextureDesc::Create2D(Q.ProceduralData->Atlas.Resolution, PF_R8G8B8A8, FClearValueBinding::Black, TexCreate_RenderTargetable | TexCreate_ShaderResource, 1), TEXT("Hair.CardsDepth"));
+		FRDGTextureRef TangentTexture	= GraphBuilder.CreateTexture(FRDGTextureDesc::Create2D(Q.ProceduralData->Atlas.Resolution, PF_R8G8B8A8, FClearValueBinding::Black, TexCreate_RenderTargetable | TexCreate_ShaderResource, 1), TEXT("Hair.CardTangent"));
+		FRDGTextureRef CoverageTexture	= GraphBuilder.CreateTexture(FRDGTextureDesc::Create2D(Q.ProceduralData->Atlas.Resolution, PF_R8G8B8A8, FClearValueBinding::Black, TexCreate_RenderTargetable | TexCreate_ShaderResource, 1), TEXT("Hair.CardCoverage"));
+		FRDGTextureRef AttributeTexture = GraphBuilder.CreateTexture(FRDGTextureDesc::Create2D(Q.ProceduralData->Atlas.Resolution, PF_R8G8B8A8, FClearValueBinding::Black, TexCreate_RenderTargetable | TexCreate_ShaderResource, 1), TEXT("Hair.CardAttribute"));
 		AddHairCardAtlasTexturePass(
 			GraphBuilder,
 			Q.ProceduralData->Atlas,
diff --git a/Engine/Plugins/Runtime/HairStrands/Source/HairStrandsCore/Private/HairStrandsMeshProjection.cpp b/Engine/Plugins/Runtime/HairStrands/Source/HairStrandsCore/Private/HairStrandsMeshProjection.cpp
--- a/Engine/Plugins/Runtime/HairStrands/Source/HairStrandsCore/Private/HairStrandsMeshProjection.cpp
+++ b/Engine/Plugins/Runtime/HairStrands/Source/HairStrandsCore/Private/HairStrandsMeshProjection.cpp
@@ -111,7 +111,7 @@ static FRDGBufferRef AddMeshSectionId(
 		return nullptr;
 
 	// Initialized the section ID to a large number, as the shader will do an atomic min on the section ID.
-	FRDGBufferRef VertexSectionIdBuffer = GraphBuilder.CreateBuffer(FRDGBufferDesc::CreateBufferDesc(sizeof(uint32), MeshData.Sections[0].TotalVertexCount), TEXT("SectionId"));
+	FRDGBufferRef VertexSectionIdBuffer = GraphBuilder.CreateBuffer(FRDGBufferDesc::CreateBufferDesc(sizeof(uint32), MeshData.Sections[0].TotalVertexCount), TEXT("Hair.SectionId"));
 	FRDGBufferUAVRef VertexSectionIdBufferUAV = GraphBuilder.CreateUAV(VertexSectionIdBuffer, PF_R32_UINT);
 	AddClearUAVPass(GraphBuilder, VertexSectionIdBufferUAV, ~0u);
 	for (const FHairStrandsProjectionMeshData::Section& MeshSection : MeshData.Sections)
@@ -270,7 +270,7 @@ static void AddMeshTransferPass(
 		return;
 	}
 	
-	FRDGBufferRef PositionDistanceBuffer = GraphBuilder.CreateBuffer(FRDGBufferDesc::CreateBufferDesc(sizeof(uint32), TargetSectionData.TotalVertexCount), TEXT("DistanceBuffer"));
+	FRDGBufferRef PositionDistanceBuffer = GraphBuilder.CreateBuffer(FRDGBufferDesc::CreateBufferDesc(sizeof(uint32), TargetSectionData.TotalVertexCount), TEXT("Hair.DistanceBuffer"));
 	FRDGBufferUAVRef PositionDistanceBufferUAV = GraphBuilder.CreateUAV(PositionDistanceBuffer, PF_R32_UINT);
 
 	// For projecting hair onto a skeletal mesh, 1 thread is spawn for each hair which iterates over all triangles.
@@ -462,7 +462,7 @@ void ProjectHairStrandsOntoMesh(
 
 	const uint32 RootCount = RestResources->RootData.RootCount;
 	FHairStrandsRestRootResource::FLOD& RestLODDatas = RestResources->LODs[LODIndex];
-	FRDGBufferRef RootDistanceBuffer = GraphBuilder.CreateBuffer(FRDGBufferDesc::CreateStructuredDesc(sizeof(float), RootCount), TEXT("HairStrandsTriangleDistance"));
+	FRDGBufferRef RootDistanceBuffer = GraphBuilder.CreateBuffer(FRDGBufferDesc::CreateStructuredDesc(sizeof(float), RootCount), TEXT("Hair.TriangleDistance"));
 
 	bool ClearDistance = true;
 	for (const FHairStrandsProjectionMeshData::Section& MeshSection : ProjectionMeshData.LODs[LODIndex].Sections)
@@ -1658,7 +1658,7 @@ void GenerateFolliculeMask(
 	if (OutTexture == nullptr)
 	{
 		FRDGTextureDesc OutputDesc = FRDGTextureDesc::Create2D(Resolution, Format, FClearValueBinding(ClearColor), TexCreate_ShaderResource | TexCreate_RenderTargetable | TexCreate_UAV, MipCount);
-		OutTexture = GraphBuilder.CreateTexture(OutputDesc, TEXT("FollicleMask"));
+		OutTexture = GraphBuilder.CreateTexture(OutputDesc, TEXT("Hair.FollicleMask"));
 	}
 
 	AddFollicleMaskPass(GraphBuilder, ShaderMap, bClear, KernelSizeInPixels, Channel, LODIndex, RestResources, OutTexture);
@@ -1681,7 +1681,7 @@ void GenerateFolliculeMask(
 	if (OutTexture == nullptr)
 	{
 		FRDGTextureDesc OutputDesc = FRDGTextureDesc::Create2D(Resolution, Format, FClearValueBinding(ClearColor), TexCreate_ShaderResource | TexCreate_RenderTargetable | TexCreate_UAV, MipCount);
-		OutTexture = GraphBuilder.CreateTexture(OutputDesc, TEXT("FollicleMask"));
+		OutTexture = GraphBuilder.CreateTexture(OutputDesc, TEXT("Hair.FollicleMask"));
 	}
 
 	for (const FRDGBufferRef& RootUVBuffer : RootUVBuffers)
diff --git a/Engine/Source/Runtime/Renderer/Private/HairStrands/HairStrandsCluster.cpp b/Engine/Source/Runtime/Renderer/Private/HairStrands/HairStrandsCluster.cpp
--- a/Engine/Source/Runtime/Renderer/Private/HairStrands/HairStrandsCluster.cpp
+++ b/Engine/Source/Runtime/Renderer/Private/HairStrands/HairStrandsCluster.cpp
@@ -225,7 +225,7 @@ FHairStrandsMacroGroupViews CreateHairStrandsMacroGroups(
 				RDG_EVENT_SCOPE(GraphBuilder, "HairStrandsAABB");
 				RDG_GPU_STAT_SCOPE(GraphBuilder, HairStrandsAABB);
 
-				MacroGroups.MacroGroupResources.MacroGroupAABBsBuffer = GraphBuilder.CreateBuffer(FRDGBufferDesc::CreateBufferDesc(4, 6 * MacroGroupCount), TEXT("HairMacroGroupAABBBuffer"));
+				MacroGroups.MacroGroupResources.MacroGroupAABBsBuffer = GraphBuilder.CreateBuffer(FRDGBufferDesc::CreateBufferDesc(4, 6 * MacroGroupCount), TEXT("Hair.MacroGroupAABBBuffer"));
 				FRDGBufferUAVRef MacroGroupAABBBufferUAV = GraphBuilder.CreateUAV(MacroGroups.MacroGroupResources.MacroGroupAABBsBuffer, PF_R32_SINT);
 				for (FHairStrandsMacroGroupData& MacroGroup : MacroGroups.Datas)
 				{				
diff --git a/Engine/Source/Runtime/Renderer/Private/HairStrands/HairStrandsClusters.cpp b/Engine/Source/Runtime/Renderer/Private/HairStrands/HairStrandsClusters.cpp
--- a/Engine/Source/Runtime/Renderer/Private/HairStrands/HairStrandsClusters.cpp
+++ b/Engine/Source/Runtime/Renderer/Private/HairStrands/HairStrandsClusters.cpp
@@ -271,19 +271,19 @@ static void AddClusterCullingPass(
 	const FHairHZBParameters& HZBParameters,
 	FHairStrandClusterData::FHairGroup& ClusterData)
 {
-	FRDGBufferRef DispatchIndirectParametersClusterCount = GraphBuilder.CreateBuffer(FRDGBufferDesc::CreateIndirectDesc<FRHIDispatchIndirectParameters>(), TEXT("HairDispatchIndirectParametersClusterCount"));
-	FRDGBufferRef DispatchIndirectParametersClusterCount2D = GraphBuilder.CreateBuffer(FRDGBufferDesc::CreateIndirectDesc<FRHIDispatchIndirectParameters>(), TEXT("HairDispatchIndirectParametersClusterCount2D"));
-	FRDGBufferRef DispatchIndirectParametersClusterCountDiv512 = GraphBuilder.CreateBuffer(FRDGBufferDesc::CreateIndirectDesc<FRHIDispatchIndirectParameters>(), TEXT("HairDispatchIndirectParametersClusterCountDiv512"));
-	FRDGBufferRef DispatchIndirectParametersClusterCountDiv512Div512 = GraphBuilder.CreateBuffer(FRDGBufferDesc::CreateIndirectDesc<FRHIDispatchIndirectParameters>(), TEXT("HairDispatchIndirectParametersClusterCountDiv512Div512"));
+	FRDGBufferRef DispatchIndirectParametersClusterCount = GraphBuilder.CreateBuffer(FRDGBufferDesc::CreateIndirectDesc<FRHIDispatchIndirectParameters>(), TEXT("Hair.DispatchIndirectParametersClusterCount"));
+	FRDGBufferRef DispatchIndirectParametersClusterCount2D = GraphBuilder.CreateBuffer(FRDGBufferDesc::CreateIndirectDesc<FRHIDispatchIndirectParameters>(), TEXT("Hair.DispatchIndirectParametersClusterCount2D"));
+	FRDGBufferRef DispatchIndirectParametersClusterCountDiv512 = GraphBuilder.CreateBuffer(FRDGBufferDesc::CreateIndirectDesc<FRHIDispatchIndirectParameters>(), TEXT("Hair.DispatchIndirectParametersClusterCountDiv512"));
+	FRDGBufferRef DispatchIndirectParametersClusterCountDiv512Div512 = GraphBuilder.CreateBuffer(FRDGBufferDesc::CreateIndirectDesc<FRHIDispatchIndirectParameters>(), TEXT("Hair.DispatchIndirectParametersClusterCountDiv512Div512"));
 	
-	FRDGBufferRef GlobalClusterIdBuffer = GraphBuilder.CreateBuffer(FRDGBufferDesc::CreateBufferDesc(sizeof(uint32), ClusterData.ClusterCount), TEXT("HairGlobalClusterIdBuffer"));
-	FRDGBufferRef GlobalIndexStartBuffer = GraphBuilder.CreateBuffer(FRDGBufferDesc::CreateBufferDesc(sizeof(uint32), ClusterData.ClusterCount), TEXT("HairGlobalIndexStartBuffer"));
-	FRDGBufferRef GlobalIndexCountBuffer = GraphBuilder.CreateBuffer(FRDGBufferDesc::CreateBufferDesc(sizeof(uint32), ClusterData.ClusterCount), TEXT("HairGlobalIndexCountBuffer"));
-	FRDGBufferRef GlobalRadiusScaleBuffer = GraphBuilder.CreateBuffer(FRDGBufferDesc::CreateBufferDesc(sizeof(float), ClusterData.ClusterCount), TEXT("HairGlobalRadiusScaleBuffer"));
+	FRDGBufferRef GlobalClusterIdBuffer = GraphBuilder.CreateBuffer(FRDGBufferDesc::CreateBufferDesc(sizeof(uint32), ClusterData.ClusterCount), TEXT("Hair.GlobalClusterIdBuffer"));
+	FRDGBufferRef GlobalIndexStartBuffer = GraphBuilder.CreateBuffer(FRDGBufferDesc::CreateBufferDesc(sizeof(uint32), ClusterData.ClusterCount), TEXT("Hair.GlobalIndexStartBuffer"));
+	FRDGBufferRef GlobalIndexCountBuffer = GraphBuilder.CreateBuffer(FRDGBufferDesc::CreateBufferDesc(sizeof(uint32), ClusterData.ClusterCount), TEXT("Hair.GlobalIndexCountBuffer"));
+	FRDGBufferRef GlobalRadiusScaleBuffer = GraphBuilder.CreateBuffer(FRDGBufferDesc::CreateBufferDesc(sizeof(float), ClusterData.ClusterCount), TEXT("Hair.GlobalRadiusScaleBuffer"));
 
-	FRDGBufferRef PerBlocklTotalIndexCountBuffer = GraphBuilder.CreateBuffer(FRDGBufferDesc::CreateBufferDesc(sizeof(uint32), ClusterData.ClusterCount), TEXT("PerBlocklTotalIndexCountBuffer"));
-	FRDGBufferRef PerBlocklTotalIndexCountPreFixSumBuffer = GraphBuilder.CreateBuffer(FRDGBufferDesc::CreateBufferDesc(sizeof(uint32) * 2, ClusterData.ClusterCount), TEXT("PerBlocklTotalIndexCountPreFixSumBuffer"));
-	FRDGBufferRef PerBlocklIndexCountPreFixSumBuffer = GraphBuilder.CreateBuffer(FRDGBufferDesc::CreateBufferDesc(sizeof(uint32) * 2, ClusterData.ClusterCount), TEXT("PerBlocklIndexCountPreFixSumBuffer"));
+	FRDGBufferRef PerBlocklTotalIndexCountBuffer = GraphBuilder.CreateBuffer(FRDGBufferDesc::CreateBufferDesc(sizeof(uint32), ClusterData.ClusterCount), TEXT("Hair.PerBlocklTotalIndexCountBuffer"));
+	FRDGBufferRef PerBlocklTotalIndexCountPreFixSumBuffer = GraphBuilder.CreateBuffer(FRDGBufferDesc::CreateBufferDesc(sizeof(uint32) * 2, ClusterData.ClusterCount), TEXT("Hair.PerBlocklTotalIndexCountPreFixSumBuffer"));
+	FRDGBufferRef PerBlocklIndexCountPreFixSumBuffer = GraphBuilder.CreateBuffer(FRDGBufferDesc::CreateBufferDesc(sizeof(uint32) * 2, ClusterData.ClusterCount), TEXT("Hair.PerBlocklIndexCountPreFixSumBuffer"));
 
 	FRDGImportedBuffer DrawIndirectParametersBuffer = Register(GraphBuilder, ClusterData.HairGroupPublicPtr->GetDrawIndirectBuffer(), ERDGImportedBufferFlags::CreateUAV);
 	FRDGImportedBuffer DrawIndirectParametersRasterComputeBuffer = Register(GraphBuilder, ClusterData.HairGroupPublicPtr->GetDrawIndirectRasterComputeBuffer(), ERDGImportedBufferFlags::CreateUAV);
@@ -305,7 +305,7 @@ static void AddClusterCullingPass(
 		static IConsoleVariable* CVarClusterDebug = IConsoleManager::Get().FindConsoleVariable(TEXT("r.HairStrands.Cluster.Debug"));
 		bClusterDebugAABBBuffer = IsHairStrandsClusterDebugAABBEnable();
 		bClusterDebug = IsHairStrandsClusterDebugEnable();
-		ClusterDebugInfoBuffer = GraphBuilder.CreateBuffer(FRDGBufferDesc::CreateStructuredDesc(sizeof(FHairClusterDebugInfo), ClusterData.ClusterCount), TEXT("CulledCompactedIndexBuffer"));
+		ClusterDebugInfoBuffer = GraphBuilder.CreateBuffer(FRDGBufferDesc::CreateStructuredDesc(sizeof(FHairClusterDebugInfo), ClusterData.ClusterCount), TEXT("Hair.CulledCompactedIndexBuffer"));
 	}
 #endif
 
diff --git a/Engine/Source/Runtime/Renderer/Private/HairStrands/HairStrandsComposition.cpp b/Engine/Source/Runtime/Renderer/Private/HairStrands/HairStrandsComposition.cpp
--- a/Engine/Source/Runtime/Renderer/Private/HairStrands/HairStrandsComposition.cpp
+++ b/Engine/Source/Runtime/Renderer/Private/HairStrands/HairStrandsComposition.cpp
@@ -188,7 +188,7 @@ static FRDGTextureRef AddHairDOFDepthPass(
 	FRDGTextureRef OutDOFDepthTexture = nullptr;
 	{
 		FRDGTextureDesc Desc = FRDGTextureDesc::Create2D(OutputResolution, PF_R32_FLOAT, FClearValueBinding::Black, TexCreate_RenderTargetable | TexCreate_ShaderResource, 1);
-		OutDOFDepthTexture = GraphBuilder.CreateTexture(Desc, TEXT("HairDOFDepth"));
+		OutDOFDepthTexture = GraphBuilder.CreateTexture(Desc, TEXT("Hair.DOFDepth"));
 	}
 
 	FHairDOFDepthPS::FParameters* Parameters = GraphBuilder.AllocParameters<FHairDOFDepthPS::FParameters>();
@@ -285,7 +285,7 @@ static void AddHairVisibilityFastResolveMSAAPass(
 		Desc.Format = PF_R8G8B8A8;
 		Desc.Flags = TexCreate_RenderTargetable | TexCreate_ShaderResource;
 		Desc.ClearValue = FClearValueBinding(0);
-		DummyTexture = GraphBuilder.CreateTexture(Desc, TEXT("HairDummyTexture"));
+		DummyTexture = GraphBuilder.CreateTexture(Desc, TEXT("Hair.DummyTexture"));
 	}
 
 	FVector2D PixelVelocity(1.f / (Resolution.X * 2), 1.f / (Resolution.Y * 2));
@@ -384,7 +384,7 @@ static void AddHairVisibilityFastResolveMaskPass(
 	FRDGTextureRef DummyTexture;
 	{
 		FRDGTextureDesc Desc = FRDGTextureDesc::Create2D(Resolution, PF_R8G8B8A8, FClearValueBinding::Black, TexCreate_RenderTargetable);
-		DummyTexture = GraphBuilder.CreateTexture(Desc, TEXT("HairDummyTexture"));
+		DummyTexture = GraphBuilder.CreateTexture(Desc, TEXT("Hair.DummyTexture"));
 	}
 
 	FHairVisibilityFastResolveMaskPS::FParameters* Parameters = GraphBuilder.AllocParameters<FHairVisibilityFastResolveMaskPS::FParameters>();
diff --git a/Engine/Source/Runtime/Renderer/Private/HairStrands/HairStrandsDebug.cpp b/Engine/Source/Runtime/Renderer/Private/HairStrands/HairStrandsDebug.cpp
--- a/Engine/Source/Runtime/Renderer/Private/HairStrands/HairStrandsDebug.cpp
+++ b/Engine/Source/Runtime/Renderer/Private/HairStrands/HairStrandsDebug.cpp
@@ -121,10 +121,10 @@ FHairCullInfo GetHairStrandsCullInfo()
 FHairStrandsDebugData::Data FHairStrandsDebugData::CreateData(FRDGBuilder& GraphBuilder)
 {
 	FHairStrandsDebugData::Data Out;
-	Out.ShadingPointBuffer = GraphBuilder.CreateBuffer(FRDGBufferDesc::CreateStructuredDesc(sizeof(ShadingInfo), MaxShadingPointCount), TEXT("HairDebugShadingPoint"));
-	Out.ShadingPointCounter = GraphBuilder.CreateBuffer(FRDGBufferDesc::CreateBufferDesc(sizeof(uint32), 1), TEXT("HairDebugShadingPointCounter"));
-	Out.SampleBuffer = GraphBuilder.CreateBuffer(FRDGBufferDesc::CreateStructuredDesc(sizeof(Sample), MaxSampleCount), TEXT("HairDebugSample"));
-	Out.SampleCounter = GraphBuilder.CreateBuffer(FRDGBufferDesc::CreateBufferDesc(sizeof(uint32), 1), TEXT("HairDebugSampleCounter"));
+	Out.ShadingPointBuffer = GraphBuilder.CreateBuffer(FRDGBufferDesc::CreateStructuredDesc(sizeof(ShadingInfo), MaxShadingPointCount), TEXT("Hair.DebugShadingPoint"));
+	Out.ShadingPointCounter = GraphBuilder.CreateBuffer(FRDGBufferDesc::CreateBufferDesc(sizeof(uint32), 1), TEXT("Hair.DebugShadingPointCounter"));
+	Out.SampleBuffer = GraphBuilder.CreateBuffer(FRDGBufferDesc::CreateStructuredDesc(sizeof(Sample), MaxSampleCount), TEXT("Hair.DebugSample"));
+	Out.SampleCounter = GraphBuilder.CreateBuffer(FRDGBufferDesc::CreateBufferDesc(sizeof(uint32), 1), TEXT("Hair.DebugSampleCounter"));
 	AddClearUAVPass(GraphBuilder, GraphBuilder.CreateUAV(Out.ShadingPointCounter, PF_R32_UINT), 0u);
 	AddClearUAVPass(GraphBuilder, GraphBuilder.CreateUAV(Out.SampleCounter, PF_R32_UINT), 0u);
 	return Out;
diff --git a/Engine/Source/Runtime/Renderer/Private/HairStrands/HairStrandsDeepShadow.cpp b/Engine/Source/Runtime/Renderer/Private/HairStrands/HairStrandsDeepShadow.cpp
--- a/Engine/Source/Runtime/Renderer/Private/HairStrands/HairStrandsDeepShadow.cpp
+++ b/Engine/Source/Runtime/Renderer/Private/HairStrands/HairStrandsDeepShadow.cpp
@@ -286,8 +286,8 @@ void RenderHairStrandsDeepShadows(
 
 		// Create Atlas resources for DOM. It is shared for all lights, across all views
 		bool bClear = true;
-		FRDGTextureRef FrontDepthAtlasTexture = GraphBuilder.CreateTexture(FRDGTextureDesc::Create2D(AtlasResolution, PF_DepthStencil, FClearValueBinding::DepthFar, TexCreate_DepthStencilTargetable | TexCreate_ShaderResource), TEXT("ShadowDepth"));
-		FRDGTextureRef DeepShadowLayersAtlasTexture = GraphBuilder.CreateTexture(FRDGTextureDesc::Create2D(AtlasResolution, PF_FloatRGBA, FClearValueBinding::Transparent, TexCreate_RenderTargetable | TexCreate_ShaderResource), TEXT("DeepShadowLayers"));
+		FRDGTextureRef FrontDepthAtlasTexture = GraphBuilder.CreateTexture(FRDGTextureDesc::Create2D(AtlasResolution, PF_DepthStencil, FClearValueBinding::DepthFar, TexCreate_DepthStencilTargetable | TexCreate_ShaderResource), TEXT("Hair.ShadowDepth"));
+		FRDGTextureRef DeepShadowLayersAtlasTexture = GraphBuilder.CreateTexture(FRDGTextureDesc::Create2D(AtlasResolution, PF_FloatRGBA, FClearValueBinding::Transparent, TexCreate_RenderTargetable | TexCreate_ShaderResource), TEXT("Hair.DeepShadowLayers"));
 
 		// TODO add support for multiple view: need to deduplicate light which are visible accross several views
 		// Allocate atlas CPU slot
@@ -344,8 +344,8 @@ void RenderHairStrandsDeepShadows(
 		Resources.TotalAtlasSlotCount = TotalAtlasSlotIndex;
 		Resources.AtlasSlotResolution = AtlasSlotResolution;
 
-		FRDGBufferRef DeepShadowViewInfoBuffer = GraphBuilder.CreateBuffer(FRDGBufferDesc::CreateStructuredDesc(20 * sizeof(float), FMath::Max(1u, TotalAtlasSlotIndex)), TEXT("DeepShadowViewInfo"));
-		FRDGBufferRef DeepShadowWorldToLightBuffer = GraphBuilder.CreateBuffer(FRDGBufferDesc::CreateStructuredDesc(16 * sizeof(float), FMath::Max(1u, TotalAtlasSlotIndex)), TEXT("DeepShadowWorldToLightTransform"));
+		FRDGBufferRef DeepShadowViewInfoBuffer = GraphBuilder.CreateBuffer(FRDGBufferDesc::CreateStructuredDesc(20 * sizeof(float), FMath::Max(1u, TotalAtlasSlotIndex)), TEXT("Hair.DeepShadowViewInfo"));
+		FRDGBufferRef DeepShadowWorldToLightBuffer = GraphBuilder.CreateBuffer(FRDGBufferDesc::CreateStructuredDesc(16 * sizeof(float), FMath::Max(1u, TotalAtlasSlotIndex)), TEXT("Hair.DeepShadowWorldToLightTransform"));
 		FRDGBufferSRVRef DeepShadowViewInfoBufferSRV = GraphBuilder.CreateSRV(DeepShadowViewInfoBuffer);
 
 		Resources.bIsGPUDriven = GDeepShadowGPUDriven > 0;
diff --git a/Engine/Source/Runtime/Renderer/Private/HairStrands/HairStrandsInterface.cpp b/Engine/Source/Runtime/Renderer/Private/HairStrands/HairStrandsInterface.cpp
--- a/Engine/Source/Runtime/Renderer/Private/HairStrands/HairStrandsInterface.cpp
+++ b/Engine/Source/Runtime/Renderer/Private/HairStrands/HairStrandsInterface.cpp
@@ -252,14 +252,14 @@ void FHairGroupPublicData::InitRHI()
 	FMemMark Mark(FMemStack::Get());
 	FRHICommandListImmediate& RHICmdList = FRHICommandListExecutor::GetImmediateCommandList();
 	FRDGBuilder GraphBuilder(RHICmdList);
-	InternalCreateIndirectBufferRDG(GraphBuilder, DrawIndirectBuffer, TEXT("HairStrandsCluster_DrawIndirectBuffer"), FUintVector4(GroupControlTriangleStripVertexCount, 1, 0, 0));
-	InternalCreateIndirectBufferRDG(GraphBuilder, DrawIndirectRasterComputeBuffer, TEXT("HairStrandsCluster_DrawIndirectRasterComputeBuffer"), FUintVector4(0, 1, 0, 0));
+	InternalCreateIndirectBufferRDG(GraphBuilder, DrawIndirectBuffer, TEXT("Hair.Cluster_DrawIndirectBuffer"), FUintVector4(GroupControlTriangleStripVertexCount, 1, 0, 0));
+	InternalCreateIndirectBufferRDG(GraphBuilder, DrawIndirectRasterComputeBuffer, TEXT("Hair.Cluster_DrawIndirectRasterComputeBuffer"), FUintVector4(0, 1, 0, 0));
 
-	InternalCreateVertexBufferRDG(GraphBuilder, sizeof(int32), ClusterCount * 6, EPixelFormat::PF_R32_SINT, ClusterAABBBuffer, TEXT("HairStrandsCluster_ClusterAABBBuffer"));
-	InternalCreateVertexBufferRDG(GraphBuilder, sizeof(int32), 6, EPixelFormat::PF_R32_SINT, GroupAABBBuffer, TEXT("HairStrandsCluster_GroupAABBBuffer"));
+	InternalCreateVertexBufferRDG(GraphBuilder, sizeof(int32), ClusterCount * 6, EPixelFormat::PF_R32_SINT, ClusterAABBBuffer, TEXT("Hair.Cluster_ClusterAABBBuffer"));
+	InternalCreateVertexBufferRDG(GraphBuilder, sizeof(int32), 6, EPixelFormat::PF_R32_SINT, GroupAABBBuffer, TEXT("Hair.Cluster_GroupAABBBuffer"));
 
-	InternalCreateVertexBufferRDG(GraphBuilder, sizeof(int32), VertexCount, EPixelFormat::PF_R32_UINT, CulledVertexIdBuffer, TEXT("HairStrandsCluster_CulledVertexIdBuffer"));
-	InternalCreateVertexBufferRDG(GraphBuilder, sizeof(float), VertexCount, EPixelFormat::PF_R32_FLOAT, CulledVertexRadiusScaleBuffer, TEXT("HairStrandsCluster_CulledVertexRadiusScaleBuffer"), true);
+	InternalCreateVertexBufferRDG(GraphBuilder, sizeof(int32), VertexCount, EPixelFormat::PF_R32_UINT, CulledVertexIdBuffer, TEXT("Hair.Cluster_CulledVertexIdBuffer"));
+	InternalCreateVertexBufferRDG(GraphBuilder, sizeof(float), VertexCount, EPixelFormat::PF_R32_FLOAT, CulledVertexRadiusScaleBuffer, TEXT("Hair.Cluster_CulledVertexRadiusScaleBuffer"), true);
 
 	GraphBuilder.Execute();
 }
diff --git a/Engine/Source/Runtime/Renderer/Private/HairStrands/HairStrandsLUT.cpp b/Engine/Source/Runtime/Renderer/Private/HairStrands/HairStrandsLUT.cpp
--- a/Engine/Source/Runtime/Renderer/Private/HairStrands/HairStrandsLUT.cpp
+++ b/Engine/Source/Runtime/Renderer/Private/HairStrands/HairStrandsLUT.cpp
@@ -69,7 +69,7 @@ static FRDGTextureRef AddHairLUTPass(
 	OutputDesc.Format = PF_FloatRGBA;
 	OutputDesc.NumMips = 1;
 	OutputDesc.Flags = TexCreate_ShaderResource | TexCreate_UAV;
-	FRDGTextureRef HairLUTTexture = GraphBuilder.CreateTexture(OutputDesc, TEXT("HairLUT"));
+	FRDGTextureRef HairLUTTexture = GraphBuilder.CreateTexture(OutputDesc, TEXT("Hair.LUT"));
 
 	FHairLUTCS::FParameters* Parameters = GraphBuilder.AllocParameters<FHairLUTCS::FParameters>();
 	Parameters->OutputColor = GraphBuilder.CreateUAV(FRDGTextureUAVDesc(HairLUTTexture, 0));
@@ -283,7 +283,7 @@ static FRDGTextureRef AddHairCoverageLUTPass(FRDGBuilder& GraphBuilder, const FV
 
 	const uint32 ElementCount = FHairCountToCoverageData::HairCount * FHairCountToCoverageData::HairRadiusCount;
 	const uint32 SizeInBytes = ElementCount * sizeof(float);
-	FRDGBufferRef UploadBuffer = GraphBuilder.CreateBuffer(FRDGBufferDesc::CreateStructuredDesc(sizeof(float), ElementCount), TEXT("HairCoverageData"));
+	FRDGBufferRef UploadBuffer = GraphBuilder.CreateBuffer(FRDGBufferDesc::CreateStructuredDesc(sizeof(float), ElementCount), TEXT("Hair.CoverageData"));
 	FHairCoverageLUTCS::FUploadParameters* UploadParameters = GraphBuilder.AllocParameters<FHairCoverageLUTCS::FUploadParameters>();
 	UploadParameters->UploadBuffer = UploadBuffer;
 
@@ -307,7 +307,7 @@ static FRDGTextureRef AddHairCoverageLUTPass(FRDGBuilder& GraphBuilder, const FV
 	OutputDesc.Format = PF_R32_FLOAT;
 	OutputDesc.NumMips = 1;
 	OutputDesc.Flags = TexCreate_UAV | TexCreate_ShaderResource;
-	FRDGTextureRef HairLUTTexture = GraphBuilder.CreateTexture(OutputDesc, TEXT("HairCoverageLUT"));
+	FRDGTextureRef HairLUTTexture = GraphBuilder.CreateTexture(OutputDesc, TEXT("Hair.CoverageLUT"));
 
 	FHairCoverageLUTCS::FParameters* Parameters = GraphBuilder.AllocParameters<FHairCoverageLUTCS::FParameters>();
 	Parameters->OutputResolution = OutputResolution;
diff --git a/Engine/Source/Runtime/Renderer/Private/HairStrands/HairStrandsScatter.cpp b/Engine/Source/Runtime/Renderer/Private/HairStrands/HairStrandsScatter.cpp
--- a/Engine/Source/Runtime/Renderer/Private/HairStrands/HairStrandsScatter.cpp
+++ b/Engine/Source/Runtime/Renderer/Private/HairStrands/HairStrandsScatter.cpp
@@ -56,7 +56,7 @@ static FRDGTextureRef AddPreScatterComposePass(
 		Desc.Format = InSceneColorTexture->Desc.Format;
 		Desc.Flags = TexCreate_RenderTargetable | TexCreate_ShaderResource;
 		Desc.ClearValue = FClearValueBinding(0);
-		OutputTexture = GraphBuilder.CreateTexture(Desc, TEXT("HairComposedTexture"));
+		OutputTexture = GraphBuilder.CreateTexture(Desc, TEXT("Hair.ComposedTexture"));
 	}
 
 	FVector2D PixelVelocity(1.f / (Resolution.X * 2), 1.f / (Resolution.Y * 2));
@@ -167,7 +167,7 @@ static FRDGTextureRef AddScatterPass(
 		PixelRadiusAtDepth1 = DiameterAtDepth1 * DiameterToRadius;
 	}
 
-	FRDGTextureRef OutDiffusionOutput = GraphBuilder.CreateTexture(InDiffusionInput->Desc, TEXT("HairDiffusionOutput"));
+	FRDGTextureRef OutDiffusionOutput = GraphBuilder.CreateTexture(InDiffusionInput->Desc, TEXT("Hair.DiffusionOutput"));
 
 	FHairScatterPS::FParameters* Parameters = GraphBuilder.AllocParameters<FHairScatterPS::FParameters>();
 	Parameters->Enable = GHairStrandsScatter_Enable > 0 ? 1 : 0;
diff --git a/Engine/Source/Runtime/Renderer/Private/HairStrands/HairStrandsTransmittance.cpp b/Engine/Source/Runtime/Renderer/Private/HairStrands/HairStrandsTransmittance.cpp
--- a/Engine/Source/Runtime/Renderer/Private/HairStrands/HairStrandsTransmittance.cpp
+++ b/Engine/Source/Runtime/Renderer/Private/HairStrands/HairStrandsTransmittance.cpp
@@ -210,7 +210,7 @@ static FRDGBufferRef AddDeepShadowTransmittanceMaskPass(
 	FRDGBufferRef OutBuffer = GraphBuilder.CreateBuffer(FRDGBufferDesc::CreateStructuredDesc(
 		4 * sizeof(float),
 		Params.HairVisibilityNodeData->Desc.NumElements),
-		TEXT("HairTransmittanceNodeData"));
+		TEXT("Hair.TransmittanceNodeData"));
 
 	FDeepTransmittanceMaskCS::FParameters* Parameters = GraphBuilder.AllocParameters<FDeepTransmittanceMaskCS::FParameters>();
 	Parameters->ViewUniformBuffer = View.ViewUniformBuffer;
@@ -406,7 +406,7 @@ static void AddDeepShadowOpaqueMaskPass(
 	{
 		FRDGTextureDesc Desc = OutShadowMask->Desc;
 		Desc.Flags |= TexCreate_ShaderResource;
-		RayMarchMask = GraphBuilder.CreateTexture(Desc, TEXT("RayMarchMask"));
+		RayMarchMask = GraphBuilder.CreateTexture(Desc, TEXT("Hair.RayMarchMask"));
 		FRHICopyTextureInfo CopyInfo;
 		CopyInfo.Size = OutShadowMask->Desc.GetSize();
 		AddCopyTexturePass(GraphBuilder, OutShadowMask, RayMarchMask, CopyInfo);
diff --git a/Engine/Source/Runtime/Renderer/Private/HairStrands/HairStrandsVisibility.cpp b/Engine/Source/Runtime/Renderer/Private/HairStrands/HairStrandsVisibility.cpp
--- a/Engine/Source/Runtime/Renderer/Private/HairStrands/HairStrandsVisibility.cpp
+++ b/Engine/Source/Runtime/Renderer/Private/HairStrands/HairStrandsVisibility.cpp
@@ -332,7 +332,7 @@ static FRDGTextureRef AddClearLightSamplePass(
 {	
 	const uint32 SampleTextureResolution = FMath::CeilToInt(FMath::Sqrt(static_cast<float>(MaxNodeCount)));
 	FRDGTextureDesc Desc = FRDGTextureDesc::Create2D(FIntPoint(SampleTextureResolution, SampleTextureResolution), PF_FloatRGBA, FClearValueBinding::Black, TexCreate_UAV | TexCreate_ShaderResource | TexCreate_RenderTargetable);
-	FRDGTextureRef Output = GraphBuilder.CreateTexture(Desc, TEXT("HairLightSample"));
+	FRDGTextureRef Output = GraphBuilder.CreateTexture(Desc, TEXT("Hair.LightSample"));
 
 	FHairLightSampleClearPS::FParameters* ParametersPS = GraphBuilder.AllocParameters<FHairLightSampleClearPS::FParameters>();
 	ParametersPS->MaxViewportResolution = Desc.Extent;
@@ -688,7 +688,7 @@ static void AddHairMaterialGBufferPass(
 	if (OutVelocityTexture == nullptr)
 	{
 		FRDGTextureDesc VelocityDesc = FVelocityRendering::GetRenderTargetDesc(ViewInfo->GetShaderPlatform(), OutDepthTexture->Desc.Extent);
-		OutVelocityTexture = GraphBuilder.CreateTexture(VelocityDesc, TEXT("DummyVelocity"));
+		OutVelocityTexture = GraphBuilder.CreateTexture(VelocityDesc, TEXT("Hair.DummyVelocity"));
 		bIsVelocityDummy = true;
 	}
 	PassParameters->RenderTargets[0] = FRenderTargetBinding(OutBufferATexture, ERenderTargetLoadAction::ELoad, 0);
@@ -1032,7 +1032,7 @@ static FRDGBufferRef AddUpdateSampleCoveragePass(
 	const FRDGTextureRef NodeIndexAndOffset,
 	const FRDGBufferRef InNodeDataBuffer)
 {
-	FRDGBufferRef OutNodeDataBuffer = GraphBuilder.CreateBuffer(FRDGBufferDesc::CreateStructuredDesc(InNodeDataBuffer->Desc.BytesPerElement, InNodeDataBuffer->Desc.NumElements), TEXT("HairCompactNodeData"));
+	FRDGBufferRef OutNodeDataBuffer = GraphBuilder.CreateBuffer(FRDGBufferDesc::CreateStructuredDesc(InNodeDataBuffer->Desc.BytesPerElement, InNodeDataBuffer->Desc.NumElements), TEXT("Hair.CompactNodeData"));
 
 	FUpdateSampleCoverageCS::FParameters* Parameters = GraphBuilder.AllocParameters<FUpdateSampleCoverageCS::FParameters>();
 	Parameters->Resolution = NodeIndexAndOffset->Desc.Extent;
@@ -1087,8 +1087,8 @@ static FMaterialPassOutput AddHairMaterialPass(
 	const uint32 MaxNodeCount = CompactNodeVis->Desc.NumElements;
 
 	FMaterialPassOutput Output;
-	Output.NodeData		= GraphBuilder.CreateBuffer(FRDGBufferDesc::CreateStructuredDesc(sizeof(HairStrandsVisibilityInternal::NodeData), MaxNodeCount), TEXT("HairCompactNodeData"));
-	Output.NodeVelocity = GraphBuilder.CreateBuffer(FRDGBufferDesc::CreateBufferDesc(4, CompactNodeVis->Desc.NumElements), TEXT("HairCompactNodeVelocity"));
+	Output.NodeData		= GraphBuilder.CreateBuffer(FRDGBufferDesc::CreateStructuredDesc(sizeof(HairStrandsVisibilityInternal::NodeData), MaxNodeCount), TEXT("Hair.CompactNodeData"));
+	Output.NodeVelocity = GraphBuilder.CreateBuffer(FRDGBufferDesc::CreateBufferDesc(4, CompactNodeVis->Desc.NumElements), TEXT("Hair.CompactNodeVelocity"));
 
 	const uint32 ResolutionDim = FMath::CeilToInt(FMath::Sqrt(static_cast<float>(MaxNodeCount)));
 	const FIntPoint Resolution(ResolutionDim, ResolutionDim);
@@ -1128,14 +1128,14 @@ static FMaterialPassOutput AddHairMaterialPass(
 	if (bOutputEmissive)
 	{
 		FRDGTextureDesc OutputDesc = FRDGTextureDesc::Create2D(Resolution, PF_FloatRGBA, FClearValueBinding::Black, TexCreate_RenderTargetable);
-		Output.EmissiveTexture = GraphBuilder.CreateTexture(OutputDesc, TEXT("HairMaterialEmissiveOutput"));
+		Output.EmissiveTexture = GraphBuilder.CreateTexture(OutputDesc, TEXT("Hair.MaterialEmissiveOutput"));
 		PassParameters->RenderTargets[0] = FRenderTargetBinding(Output.EmissiveTexture, ERenderTargetLoadAction::EClear, 0);
 		
 	}
 	else if (bIsPlatformRequireRenderTarget)
 	{
 		FRDGTextureDesc OutputDesc = FRDGTextureDesc::Create2D(Resolution, PF_FloatRGBA, FClearValueBinding::Black, TexCreate_RenderTargetable);
-		FRDGTextureRef OutDummyTexture0 = GraphBuilder.CreateTexture(OutputDesc, TEXT("HairMaterialDummyOutput"));
+		FRDGTextureRef OutDummyTexture0 = GraphBuilder.CreateTexture(OutputDesc, TEXT("Hair.MaterialDummyOutput"));
 		PassParameters->RenderTargets[0] = FRenderTargetBinding(OutDummyTexture0, ERenderTargetLoadAction::EClear, 0);
 	}
 	
@@ -1242,7 +1242,7 @@ static void AddHairVelocityPass(
 	const FIntPoint Resolution = OutVelocityTexture->Desc.Extent;
 	{
 		FRDGTextureDesc Desc = FRDGTextureDesc::Create2D(Resolution, PF_R32_UINT, FClearValueBinding::None, TexCreate_UAV);
-		OutResolveMaskTexture = GraphBuilder.CreateTexture(Desc, TEXT("VelocityResolveMaskTexture"));
+		OutResolveMaskTexture = GraphBuilder.CreateTexture(Desc, TEXT("Hair.VelocityResolveMaskTexture"));
 	}
 
 	check(OutVelocityTexture->Desc.Format == PF_G16R16 || OutVelocityTexture->Desc.Format == PF_A16B16G16R16);
@@ -1318,7 +1318,7 @@ static FRDGTextureRef AddHairLightChannelMaskPass(
 	FRDGTextureRef NodeOffsetAndCount)
 {
 	FRDGTextureDesc Desc = FRDGTextureDesc::Create2D(Resolution, PF_R32_UINT, FClearValueBinding::None, TexCreate_UAV | TexCreate_ShaderResource);
-	FRDGTextureRef OutLightChannelMaskTexture = GraphBuilder.CreateTexture(Desc, TEXT("HairLightChannelMask"));
+	FRDGTextureRef OutLightChannelMaskTexture = GraphBuilder.CreateTexture(Desc, TEXT("Hair.LightChannelMask"));
 
 	FHairLightChannelMaskCS::FPermutationDomain PermutationVector;
 	PermutationVector.Set<FHairLightChannelMaskCS::FVendor>(GetVendor());
@@ -1374,9 +1374,9 @@ TRDGUniformBufferRef<FVisibilityPassUniformParameters> CreatePassDummyTextures(F
 	FVisibilityPassUniformParameters* UniformParameters = GraphBuilder.AllocParameters<FVisibilityPassUniformParameters>();
 
 	FRDGTextureDesc Desc = FRDGTextureDesc::Create2D(FIntPoint(1,1), PF_R32_UINT, FClearValueBinding::None, TexCreate_UAV | TexCreate_ShaderResource);
-	UniformParameters->PPLLCounter		= GraphBuilder.CreateUAV(GraphBuilder.CreateTexture(Desc, TEXT("HairVisibilityPPLLNodeIndex")));
-	UniformParameters->PPLLNodeIndex	= GraphBuilder.CreateUAV(GraphBuilder.CreateTexture(Desc, TEXT("HairVisibilityPPLLNodeIndex")));
-	UniformParameters->PPLLNodeData		= GraphBuilder.CreateUAV(GraphBuilder.CreateBuffer(FRDGBufferDesc::CreateStructuredDesc(sizeof(PPLLNodeData), 1), TEXT("DummyPPLLNodeData")));
+	UniformParameters->PPLLCounter		= GraphBuilder.CreateUAV(GraphBuilder.CreateTexture(Desc, TEXT("Hair.VisibilityPPLLNodeIndex")));
+	UniformParameters->PPLLNodeIndex	= GraphBuilder.CreateUAV(GraphBuilder.CreateTexture(Desc, TEXT("Hair.VisibilityPPLLNodeIndex")));
+	UniformParameters->PPLLNodeData		= GraphBuilder.CreateUAV(GraphBuilder.CreateBuffer(FRDGBufferDesc::CreateStructuredDesc(sizeof(PPLLNodeData), 1), TEXT("Hair.DummyPPLLNodeData")));
 
 	return GraphBuilder.CreateUniformBuffer(UniformParameters);
 }
@@ -1786,7 +1786,7 @@ static FRDGBufferRef AddCopyIndirectArgPass(
 {
 	check(CounterTexture);
 
-	FRDGBufferRef OutBuffer = GraphBuilder.CreateBuffer(FRDGBufferDesc::CreateIndirectDesc<FRHIDispatchIndirectParameters>(), TEXT("HairVisibilityIndirectArgBuffer"));
+	FRDGBufferRef OutBuffer = GraphBuilder.CreateBuffer(FRDGBufferDesc::CreateIndirectDesc<FRHIDispatchIndirectParameters>(), TEXT("Hair.VisibilityIndirectArgBuffer"));
 
 	FCopyIndirectBufferCS::FParameters* Parameters = GraphBuilder.AllocParameters<FCopyIndirectBufferCS::FParameters>();
 	Parameters->ThreadGroupSize = ThreadGroupSize;
@@ -1936,17 +1936,17 @@ static void AddHairVisibilityPrimitiveIdCompactionPass(
 
 	{
 		FRDGTextureDesc Desc = FRDGTextureDesc::Create2D(FIntPoint(1, 1), PF_R32_UINT, FClearValueBinding::None, TexCreate_UAV | TexCreate_ShaderResource);
-		OutCompactCounter = GraphBuilder.CreateTexture(Desc, TEXT("HairVisibilityCompactCounter"));
+		OutCompactCounter = GraphBuilder.CreateTexture(Desc, TEXT("Hair.VisibilityCompactCounter"));
 	}
 
 	{
 		FRDGTextureDesc Desc = FRDGTextureDesc::Create2D(Resolution, PF_R32_UINT, FClearValueBinding::None, TexCreate_UAV | TexCreate_ShaderResource);
-		OutCompactNodeIndex = GraphBuilder.CreateTexture(Desc, TEXT("HairVisibilityCompactNodeIndex"));
+		OutCompactNodeIndex = GraphBuilder.CreateTexture(Desc, TEXT("Hair.VisibilityCompactNodeIndex"));
 	}
 
 	{
 		FRDGTextureDesc Desc = FRDGTextureDesc::Create2D(Resolution, PF_R16G16B16A16_UINT, FClearValueBinding::None, TexCreate_UAV | TexCreate_ShaderResource);
-		OutCategorizationTexture = GraphBuilder.CreateTexture(Desc, TEXT("CategorizationTexture"));
+		OutCategorizationTexture = GraphBuilder.CreateTexture(Desc, TEXT("Hair.CategorizationTexture"));
 	}
 
 	const uint32 ClearValues[4] = { 0u,0u,0u,0u };
@@ -1958,11 +1958,11 @@ static void AddHairVisibilityPrimitiveIdCompactionPass(
 	const uint32 MSAASampleCount = GetHairVisibilityRenderMode() == HairVisibilityRenderMode_MSAA ? GetMaxSamplePerPixel() : 1;
 	const uint32 PPLLMaxRenderNodePerPixel = GetMaxSamplePerPixel();
 	const uint32 MaxRenderNodeCount = GetTotalSampleCountForAllocation(Resolution);
-	OutCompactNodeData = GraphBuilder.CreateBuffer(FRDGBufferDesc::CreateStructuredDesc(bUseVisibility ? sizeof(HairStrandsVisibilityInternal::NodeVis) : sizeof(HairStrandsVisibilityInternal::NodeData), MaxRenderNodeCount), TEXT("HairVisibilityPrimitiveIdCompactNodeData"));
+	OutCompactNodeData = GraphBuilder.CreateBuffer(FRDGBufferDesc::CreateStructuredDesc(bUseVisibility ? sizeof(HairStrandsVisibilityInternal::NodeVis) : sizeof(HairStrandsVisibilityInternal::NodeData), MaxRenderNodeCount), TEXT("Hair.VisibilityPrimitiveIdCompactNodeData"));
 
 	{
 		// Pixel coord of the node. Stored as 2*uint16, packed into a single uint32
-		OutCompactNodeCoord = GraphBuilder.CreateBuffer(FRDGBufferDesc::CreateStructuredDesc(sizeof(uint32), MaxRenderNodeCount), TEXT("HairVisibilityPrimitiveIdCompactNodeCoord"));
+		OutCompactNodeCoord = GraphBuilder.CreateBuffer(FRDGBufferDesc::CreateStructuredDesc(sizeof(uint32), MaxRenderNodeCount), TEXT("Hair.VisibilityPrimitiveIdCompactNodeCoord"));
 	}
 
 	const bool bWriteOutVelocity = OutVelocityTexture != nullptr;
@@ -2088,17 +2088,17 @@ static void AddHairVisibilityCompactionComputeRasterPass(
 
 	{
 		FRDGTextureDesc Desc = FRDGTextureDesc::Create2D(FIntPoint(1,1), PF_R32_UINT, FClearValueBinding::None, TexCreate_UAV);
-		OutCompactCounter = GraphBuilder.CreateTexture(Desc, TEXT("HairVisibilityCompactCounter"));
+		OutCompactCounter = GraphBuilder.CreateTexture(Desc, TEXT("Hair.VisibilityCompactCounter"));
 	}
 
 	{
 		FRDGTextureDesc Desc = FRDGTextureDesc::Create2D(Resolution, PF_R32_UINT, FClearValueBinding::None, TexCreate_UAV | TexCreate_ShaderResource);
-		OutCompactNodeIndex = GraphBuilder.CreateTexture(Desc, TEXT("HairVisibilityCompactNodeIndex"));
+		OutCompactNodeIndex = GraphBuilder.CreateTexture(Desc, TEXT("Hair.VisibilityCompactNodeIndex"));
 	}
 
 	{
 		FRDGTextureDesc Desc = FRDGTextureDesc::Create2D(Resolution, PF_R16G16B16A16_UINT, FClearValueBinding::None, TexCreate_UAV | TexCreate_ShaderResource);
-		OutCategorizationTexture = GraphBuilder.CreateTexture(Desc, TEXT("CategorizationTexture"));
+		OutCategorizationTexture = GraphBuilder.CreateTexture(Desc, TEXT("Hair.CategorizationTexture"));
 	}
 
 	const uint32 ClearValues[4] = { 0u,0u,0u,0u };
@@ -2110,8 +2110,8 @@ static void AddHairVisibilityCompactionComputeRasterPass(
 	const uint32 MSAASampleCount = GetHairVisibilityRenderMode() == HairVisibilityRenderMode_MSAA ? GetMaxSamplePerPixel() : 1;
 	const uint32 PPLLMaxRenderNodePerPixel = GetMaxSamplePerPixel();
 	const uint32 MaxRenderNodeCount = GetTotalSampleCountForAllocation(Resolution);
-	OutCompactNodeData = GraphBuilder.CreateBuffer(FRDGBufferDesc::CreateStructuredDesc(sizeof(HairStrandsVisibilityInternal::NodeVis), MaxRenderNodeCount), TEXT("HairVisibilityPrimitiveIdCompactNodeData"));
-	OutCompactNodeCoord = GraphBuilder.CreateBuffer(FRDGBufferDesc::CreateStructuredDesc(sizeof(uint32), MaxRenderNodeCount), TEXT("HairVisibilityPrimitiveIdCompactNodeCoord"));
+	OutCompactNodeData = GraphBuilder.CreateBuffer(FRDGBufferDesc::CreateStructuredDesc(sizeof(HairStrandsVisibilityInternal::NodeVis), MaxRenderNodeCount), TEXT("Hair.VisibilityPrimitiveIdCompactNodeData"));
+	OutCompactNodeCoord = GraphBuilder.CreateBuffer(FRDGBufferDesc::CreateStructuredDesc(sizeof(uint32), MaxRenderNodeCount), TEXT("Hair.VisibilityPrimitiveIdCompactNodeCoord"));
 
 	FRDGTextureRef DefaultTexture = GSystemTextures.GetBlackDummy(GraphBuilder);
 	FHairVisibilityCompactionComputeRasterCS::FParameters* PassParameters = GraphBuilder.AllocParameters<FHairVisibilityCompactionComputeRasterCS::FParameters>();
@@ -2193,15 +2193,15 @@ static void AddGenerateTilePass(
 	FRDGTextureRef TileCounter;
 	{
 		FRDGTextureDesc Desc = FRDGTextureDesc::Create2D(FIntPoint(1,1), PF_R32_UINT, FClearValueBinding::None, TexCreate_UAV | TexCreate_ShaderResource);
-		TileCounter = GraphBuilder.CreateTexture(Desc, TEXT("HairTileCounter"));
+		TileCounter = GraphBuilder.CreateTexture(Desc, TEXT("Hair.TileCounter"));
 	}
 
 	{
 		FRDGTextureDesc Desc = FRDGTextureDesc::Create2D(TileResolution, PF_R32_UINT, FClearValueBinding::None, TexCreate_UAV | TexCreate_ShaderResource);
-		OutTileIndexTexture = GraphBuilder.CreateTexture(Desc, TEXT("HairTileIndexTexture"));
+		OutTileIndexTexture = GraphBuilder.CreateTexture(Desc, TEXT("Hair.TileIndexTexture"));
 	}
 
-	OutTileBuffer = GraphBuilder.CreateBuffer(FRDGBufferDesc::CreateBufferDesc(sizeof(uint32), TileCount), TEXT("HairTileBuffer"));
+	OutTileBuffer = GraphBuilder.CreateBuffer(FRDGBufferDesc::CreateBufferDesc(sizeof(uint32), TileCount), TEXT("Hair.TileBuffer"));
 
 	uint32 ClearValues[4] = { 0,0,0,0 };
 	AddClearUAVPass(GraphBuilder, GraphBuilder.CreateUAV(TileCounter), ClearValues);
@@ -2261,7 +2261,7 @@ static FRDGTextureRef AddHairVisibilityFillOpaqueDepth(
 		check(GetHairVisibilityRenderMode() == HairVisibilityRenderMode_MSAA);
 
 		FRDGTextureDesc Desc = FRDGTextureDesc::Create2D(Resolution, PF_DepthStencil, FClearValueBinding::DepthFar, TexCreate_DepthStencilTargetable | TexCreate_ShaderResource, 1, GetMaxSamplePerPixel());
-		OutVisibilityDepthTexture = GraphBuilder.CreateTexture(Desc, TEXT("HairVisibilityDepthTexture"));
+		OutVisibilityDepthTexture = GraphBuilder.CreateTexture(Desc, TEXT("Hair.VisibilityDepthTexture"));
 	}
 
 	FHairVisibilityFillOpaqueDepthPS::FParameters* Parameters = GraphBuilder.AllocParameters<FHairVisibilityFillOpaqueDepthPS::FParameters>();
@@ -2515,22 +2515,22 @@ static void AddHairVisibilityMSAAPass(
 	{
 		{
 			FRDGTextureDesc Desc = FRDGTextureDesc::Create2D(Resolution, PF_R32G32_UINT, FClearValueBinding::None, TexCreate_RenderTargetable | TexCreate_ShaderResource, 1, MSAASampleCount);
-			OutVisibilityIdTexture = GraphBuilder.CreateTexture(Desc, TEXT("HairVisibilityIDTexture"));
+			OutVisibilityIdTexture = GraphBuilder.CreateTexture(Desc, TEXT("Hair.VisibilityIDTexture"));
 		}
 
 		{
 			FRDGTextureDesc Desc = FRDGTextureDesc::Create2D(Resolution, PF_R8G8B8A8, FClearValueBinding(FLinearColor(0, 0, 0, 0)), TexCreate_RenderTargetable | TexCreate_ShaderResource, 1, MSAASampleCount);
-			OutVisibilityMaterialTexture = GraphBuilder.CreateTexture(Desc, TEXT("HairVisibilityMaterialTexture"));
+			OutVisibilityMaterialTexture = GraphBuilder.CreateTexture(Desc, TEXT("Hair.VisibilityMaterialTexture"));
 		}
 
 		{
 			FRDGTextureDesc Desc = FRDGTextureDesc::Create2D(Resolution, PF_R8G8B8A8, FClearValueBinding(FLinearColor(0, 0, 0, 0)), TexCreate_RenderTargetable | TexCreate_ShaderResource, 1, MSAASampleCount);
-			OutVisibilityAttributeTexture = GraphBuilder.CreateTexture(Desc, TEXT("HairVisibilityAttributeTexture"));
+			OutVisibilityAttributeTexture = GraphBuilder.CreateTexture(Desc, TEXT("Hair.VisibilityAttributeTexture"));
 		}
 
 		{
 			FRDGTextureDesc Desc = FRDGTextureDesc::Create2D(Resolution, PF_G16R16, FClearValueBinding(FLinearColor(0, 0, 0, 0)), TexCreate_RenderTargetable | TexCreate_ShaderResource, 1, MSAASampleCount);
-			OutVisibilityVelocityTexture = GraphBuilder.CreateTexture(Desc, TEXT("HairVisibilityVelocityTexture"));
+			OutVisibilityVelocityTexture = GraphBuilder.CreateTexture(Desc, TEXT("Hair.VisibilityVelocityTexture"));
 		}
 		AddClearGraphicPass(GraphBuilder, RDG_EVENT_NAME("HairStrandsClearVisibilityMSAAIdTexture"), ViewInfo, 0xFFFFFFFF, OutVisibilityIdTexture);
 
@@ -2576,17 +2576,17 @@ static void AddHairVisibilityPPLLPass(
 {
 	{
 		FRDGTextureDesc Desc = FRDGTextureDesc::Create2D(FIntPoint(1,1), PF_R32_UINT, FClearValueBinding::None, TexCreate_UAV | TexCreate_ShaderResource);
-		OutVisibilityPPLLNodeCounter = GraphBuilder.CreateTexture(Desc, TEXT("HairVisibilityPPLLCounter"));
+		OutVisibilityPPLLNodeCounter = GraphBuilder.CreateTexture(Desc, TEXT("Hair.VisibilityPPLLCounter"));
 	}
 
 	{
 		FRDGTextureDesc Desc = FRDGTextureDesc::Create2D(Resolution, PF_R32_UINT, FClearValueBinding::None, TexCreate_UAV | TexCreate_ShaderResource);
-		OutVisibilityPPLLNodeIndex = GraphBuilder.CreateTexture(Desc, TEXT("HairVisibilityPPLLNodeIndex"));
+		OutVisibilityPPLLNodeIndex = GraphBuilder.CreateTexture(Desc, TEXT("Hair.VisibilityPPLLNodeIndex"));
 	}
 
 	const uint32 PPLLMaxTotalListElementCount = GetTotalSampleCountForAllocation(Resolution);
 	{
-		OutVisibilityPPLLNodeData = GraphBuilder.CreateBuffer(FRDGBufferDesc::CreateStructuredDesc(sizeof(PPLLNodeData), PPLLMaxTotalListElementCount), TEXT("HairVisibilityPPLLNodeData"));
+		OutVisibilityPPLLNodeData = GraphBuilder.CreateBuffer(FRDGBufferDesc::CreateStructuredDesc(sizeof(PPLLNodeData), PPLLMaxTotalListElementCount), TEXT("Hair.VisibilityPPLLNodeData"));
 	}
 	const uint32 ClearValue0[4] = { 0,0,0,0 };
 	const uint32 ClearValueInvalid[4] = { 0xFFFFFFFF, 0xFFFFFFFF, 0xFFFFFFFF, 0xFFFFFFFF };
@@ -2640,14 +2640,14 @@ static FHairPrimaryTransmittance AddHairViewTransmittancePass(
 	PassParameters->UniformBuffer = CreatePassDummyTextures(GraphBuilder);
 	FHairPrimaryTransmittance Out;
 
-	Out.TransmittanceTexture = GraphBuilder.CreateTexture(Desc, TEXT("HairViewTransmittanceTexture"));
+	Out.TransmittanceTexture = GraphBuilder.CreateTexture(Desc, TEXT("Hair.ViewTransmittanceTexture"));
 	PassParameters->RenderTargets[0] = FRenderTargetBinding(Out.TransmittanceTexture, ERenderTargetLoadAction::EClear, 0);
 
 	if (RenderMode == HairVisibilityRenderMode_TransmittanceAndHairCount)
 	{
 		Desc.Format = PF_G32R32F;
 		Desc.ClearValue = FClearValueBinding(FLinearColor(0.0f, 0.0f, 0.0f, 0.0f));
-		Out.HairCountTexture = GraphBuilder.CreateTexture(Desc, TEXT("HairViewHairCountTexture"));
+		Out.HairCountTexture = GraphBuilder.CreateTexture(Desc, TEXT("Hair.ViewHairCountTexture"));
 		PassParameters->RenderTargets[1] = FRenderTargetBinding(Out.HairCountTexture, ERenderTargetLoadAction::EClear, 0);
 	}
 
@@ -2965,7 +2965,7 @@ static FRDGTextureRef AddHairHairCountToTransmittancePass(
 	const bool bUseOneChannel = HairCountTexture->Desc.Format == PF_R32_UINT;
 
 	FRDGTextureDesc Desc = FRDGTextureDesc::Create2D(OutputResolution, PF_R32_FLOAT, FClearValueBinding(FLinearColor(0.0f, 0.0f, 0.0f, 0.0f)), TexCreate_UAV | TexCreate_ShaderResource | TexCreate_RenderTargetable);
-	FRDGTextureRef OutputTexture = GraphBuilder.CreateTexture(Desc, TEXT("HairVisibilityTexture"));
+	FRDGTextureRef OutputTexture = GraphBuilder.CreateTexture(Desc, TEXT("Hair.VisibilityTexture"));
 	FRDGTextureRef HairCoverageLUT = HairLUT.Textures[HairLUTType_Coverage];
 
 	FHairCountToCoverageCS::FParameters* PassParameters = GraphBuilder.AllocParameters<FHairCountToCoverageCS::FParameters>();
@@ -3103,26 +3103,26 @@ static FRasterComputeOutput AddVisibilityComputeRasterPass(
 	FRDGTextureUAVRef VisibilityTexture3UAV = nullptr;
 
 	uint32 ClearValues[4] = { 0,0,0,0 };
-	Out.HairCountTexture = GraphBuilder.CreateTexture(DescCount, TEXT("HairViewTransmittanceTexture"));
+	Out.HairCountTexture = GraphBuilder.CreateTexture(DescCount, TEXT("Hair.ViewTransmittanceTexture"));
 	FRDGTextureUAVRef HairCountTextureUAV = GraphBuilder.CreateUAV(Out.HairCountTexture);
 	AddClearUAVPass(GraphBuilder, HairCountTextureUAV, ClearValues);
 
-	Out.VisibilityTexture0 = GraphBuilder.CreateTexture(DescVis, TEXT("HairVisibilityTexture0"));
+	Out.VisibilityTexture0 = GraphBuilder.CreateTexture(DescVis, TEXT("Hair.VisibilityTexture0"));
 	VisibilityTexture0UAV = GraphBuilder.CreateUAV(Out.VisibilityTexture0);
 	AddClearUAVPass(GraphBuilder, VisibilityTexture0UAV, ClearValues);
 	if (SamplePerPixelCount > 1)
 	{
-		Out.VisibilityTexture1 = GraphBuilder.CreateTexture(DescVis, TEXT("HairVisibilityTexture1"));
+		Out.VisibilityTexture1 = GraphBuilder.CreateTexture(DescVis, TEXT("Hair.VisibilityTexture1"));
 		VisibilityTexture1UAV = GraphBuilder.CreateUAV(Out.VisibilityTexture1);
 		AddClearUAVPass(GraphBuilder, VisibilityTexture1UAV, ClearValues);
 		if (SamplePerPixelCount > 2)
 		{
-			Out.VisibilityTexture2 = GraphBuilder.CreateTexture(DescVis, TEXT("HairVisibilityTexture2"));
+			Out.VisibilityTexture2 = GraphBuilder.CreateTexture(DescVis, TEXT("Hair.VisibilityTexture2"));
 			VisibilityTexture2UAV = GraphBuilder.CreateUAV(Out.VisibilityTexture2);
 			AddClearUAVPass(GraphBuilder, VisibilityTexture2UAV, ClearValues);
 			if (SamplePerPixelCount > 3)
 			{
-				Out.VisibilityTexture3 = GraphBuilder.CreateTexture(DescVis, TEXT("HairVisibilityTexture3"));
+				Out.VisibilityTexture3 = GraphBuilder.CreateTexture(DescVis, TEXT("Hair.VisibilityTexture3"));
 				VisibilityTexture3UAV = GraphBuilder.CreateUAV(Out.VisibilityTexture3);
 				AddClearUAVPass(GraphBuilder, VisibilityTexture3UAV, ClearValues);
 			}
@@ -3272,7 +3272,7 @@ FHairStrandsVisibilityViews RenderHairStrandsVisibilityBuffer(
 			const EHairVisibilityRenderMode RenderMode = GetHairVisibilityRenderMode();
 			check(RenderMode == HairVisibilityRenderMode_MSAA || RenderMode == HairVisibilityRenderMode_PPLL || RenderMode == HairVisibilityRenderMode_ComputeRaster);
 
-			FRDGTextureRef HairOnlyDepthTexture = GraphBuilder.CreateTexture(SceneDepthTexture->Desc, TEXT("HairStrandsHairOnlyDepthTexture"));
+			FRDGTextureRef HairOnlyDepthTexture = GraphBuilder.CreateTexture(SceneDepthTexture->Desc, TEXT("Hair.HairOnlyDepthTexture"));
 			FRDGTextureRef CategorizationTexture = nullptr;
 			FRDGTextureRef CompactNodeIndex = nullptr;
 			FRDGBufferRef  CompactNodeData = nullptr;
diff --git a/Engine/Source/Runtime/Renderer/Private/HairStrands/HairStrandsVoxelization.cpp b/Engine/Source/Runtime/Renderer/Private/HairStrands/HairStrandsVoxelization.cpp
--- a/Engine/Source/Runtime/Renderer/Private/HairStrands/HairStrandsVoxelization.cpp
+++ b/Engine/Source/Runtime/Renderer/Private/HairStrands/HairStrandsVoxelization.cpp
@@ -545,15 +545,15 @@ static void AddAllocateVoxelPagesPass(
 	}
 	check(OutTotalPageIndexCount > 0);
 	
-	FRDGBufferRef PageIndexBuffer = GraphBuilder.CreateBuffer(FRDGBufferDesc::CreateBufferDesc(sizeof(uint32), OutTotalPageIndexCount), TEXT("PageIndexBuffer"));
-	FRDGBufferRef PageIndexOccupancyBuffer = GraphBuilder.CreateBuffer(FRDGBufferDesc::CreateBufferDesc(sizeof(uint32)*2, OutTotalPageIndexCount), TEXT("PageIndexOccupancyBuffer"));
-	FRDGBufferRef PageIndexCoordBuffer = GraphBuilder.CreateBuffer(FRDGBufferDesc::CreateBufferDesc(sizeof(uint32), OutTotalPageIndexCount), TEXT("PageIndexCoordBuffer"));
-	FRDGBufferRef PageIndexGlobalCounter = GraphBuilder.CreateBuffer(FRDGBufferDesc::CreateBufferDesc(sizeof(uint32), 2), TEXT("PageIndexGlobalCounter"));
-	FRDGBufferRef NodeDescBuffer = GraphBuilder.CreateBuffer(FRDGBufferDesc::CreateStructuredDesc(sizeof(FPackedVirtualVoxelNodeDesc), MacroGroupCount), TEXT("VirtualVoxelNodeDescBuffer"));
-	FRDGBufferRef IndirectArgsBuffer = GraphBuilder.CreateBuffer(FRDGBufferDesc::CreateIndirectDesc<FRHIDispatchIndirectParameters>(MacroGroupCount), TEXT("VirtualVoxelIndirectArgsBuffer"));
+	FRDGBufferRef PageIndexBuffer = GraphBuilder.CreateBuffer(FRDGBufferDesc::CreateBufferDesc(sizeof(uint32), OutTotalPageIndexCount), TEXT("Hair.PageIndexBuffer"));
+	FRDGBufferRef PageIndexOccupancyBuffer = GraphBuilder.CreateBuffer(FRDGBufferDesc::CreateBufferDesc(sizeof(uint32)*2, OutTotalPageIndexCount), TEXT("Hair.PageIndexOccupancyBuffer"));
+	FRDGBufferRef PageIndexCoordBuffer = GraphBuilder.CreateBuffer(FRDGBufferDesc::CreateBufferDesc(sizeof(uint32), OutTotalPageIndexCount), TEXT("Hair.PageIndexCoordBuffer"));
+	FRDGBufferRef PageIndexGlobalCounter = GraphBuilder.CreateBuffer(FRDGBufferDesc::CreateBufferDesc(sizeof(uint32), 2), TEXT("Hair.PageIndexGlobalCounter"));
+	FRDGBufferRef NodeDescBuffer = GraphBuilder.CreateBuffer(FRDGBufferDesc::CreateStructuredDesc(sizeof(FPackedVirtualVoxelNodeDesc), MacroGroupCount), TEXT("Hair.VirtualVoxelNodeDescBuffer"));
+	FRDGBufferRef IndirectArgsBuffer = GraphBuilder.CreateBuffer(FRDGBufferDesc::CreateIndirectDesc<FRHIDispatchIndirectParameters>(MacroGroupCount), TEXT("Hair.VirtualVoxelIndirectArgsBuffer"));
 
 	const uint32 TotalPageCount = PageCountResolution.X * PageCountResolution.Y * PageCountResolution.Z;
-	FRDGBufferRef PageToPageIndexBuffer = GraphBuilder.CreateBuffer(FRDGBufferDesc::CreateBufferDesc(sizeof(uint32), TotalPageCount), TEXT("PageToPageIndexBuffer"));
+	FRDGBufferRef PageToPageIndexBuffer = GraphBuilder.CreateBuffer(FRDGBufferDesc::CreateBufferDesc(sizeof(uint32), TotalPageCount), TEXT("Hair.PageToPageIndexBuffer"));
 
 	FRDGBufferUAVRef PageIndexBufferUAV = GraphBuilder.CreateUAV(PageIndexBuffer, PF_R32_UINT);
 	FRDGBufferUAVRef PageIndexOccupancyBufferUAV = GraphBuilder.CreateUAV(PageIndexOccupancyBuffer, PF_R32G32_UINT);
@@ -561,9 +561,9 @@ static void AddAllocateVoxelPagesPass(
 	
 	// Stored FVoxelizationViewInfo structs
 	// See HairStrandsVoxelPageCommonStruct.ush for more details
-	FRDGBufferRef VoxelizationViewInfoBuffer = GraphBuilder.CreateBuffer(FRDGBufferDesc::CreateStructuredDesc(24 * sizeof(float), MacroGroupCount), TEXT("VoxelizationViewInfo"));
-	FRDGBufferRef PageIndexResolutionBuffer = GraphBuilder.CreateBuffer(FRDGBufferDesc::CreateBufferDesc(MacroGroupCount * 4 * sizeof(uint32), OutTotalPageIndexCount), TEXT("PageIndexResolutionBuffer"));
-	FRDGBufferRef PageIndexAllocationIndirectBufferArgs = GraphBuilder.CreateBuffer(FRDGBufferDesc::CreateIndirectDesc<FRHIDispatchIndirectParameters>(MacroGroupCount), TEXT("PageIndexAllocationIndirectBufferArgs"));
+	FRDGBufferRef VoxelizationViewInfoBuffer = GraphBuilder.CreateBuffer(FRDGBufferDesc::CreateStructuredDesc(24 * sizeof(float), MacroGroupCount), TEXT("Hair.VoxelizationViewInfo"));
+	FRDGBufferRef PageIndexResolutionBuffer = GraphBuilder.CreateBuffer(FRDGBufferDesc::CreateBufferDesc(MacroGroupCount * 4 * sizeof(uint32), OutTotalPageIndexCount), TEXT("Hair.PageIndexResolutionBuffer"));
+	FRDGBufferRef PageIndexAllocationIndirectBufferArgs = GraphBuilder.CreateBuffer(FRDGBufferDesc::CreateIndirectDesc<FRHIDispatchIndirectParameters>(MacroGroupCount), TEXT("Hair.PageIndexAllocationIndirectBufferArgs"));
 
 	AddClearUAVPass(GraphBuilder, PageIndexBufferUAV, 0u);
 	AddClearUAVPass(GraphBuilder, PageIndexOccupancyBufferUAV, 0u);
@@ -629,8 +629,8 @@ static void AddAllocateVoxelPagesPass(
 		{			
 			const uint32 AverageWorkItemPerCluster = 8; // Arbitrary/Guess number
 			const uint32 MaxAllocationCount = TotalClusterCount * AverageWorkItemPerCluster;
-			FRDGBufferRef ScatterCounter = GraphBuilder.CreateBuffer(FRDGBufferDesc::CreateBufferDesc(sizeof(uint32), 1), TEXT("PageScatterCounter"));
-			FRDGBufferRef ScatterBuffer = GraphBuilder.CreateBuffer(FRDGBufferDesc::CreateBufferDesc(2 * sizeof(uint32), MaxAllocationCount), TEXT("PageScatterBuffer"));
+			FRDGBufferRef ScatterCounter = GraphBuilder.CreateBuffer(FRDGBufferDesc::CreateBufferDesc(sizeof(uint32), 1), TEXT("Hair.PageScatterCounter"));
+			FRDGBufferRef ScatterBuffer = GraphBuilder.CreateBuffer(FRDGBufferDesc::CreateBufferDesc(2 * sizeof(uint32), MaxAllocationCount), TEXT("Hair.PageScatterBuffer"));
 
 			FRDGBufferUAVRef ScatterCounterUAV = GraphBuilder.CreateUAV(ScatterCounter, PF_R32_UINT);
 			FRDGBufferUAVRef ScatterBufferUAV  = GraphBuilder.CreateUAV(ScatterBuffer, PF_R32G32_UINT);
@@ -670,7 +670,7 @@ static void AddAllocateVoxelPagesPass(
 			FRDGBufferSRVRef ScatterCounterSRV = GraphBuilder.CreateSRV(ScatterCounter, PF_R32_UINT);
 
 			// Build indirect buffer args
-			FRDGBufferRef ScatterIndirectArgsBuffer = GraphBuilder.CreateBuffer(FRDGBufferDesc::CreateIndirectDesc<FRHIDispatchIndirectParameters>(1), TEXT("PageScatterIndirectArgs"));
+			FRDGBufferRef ScatterIndirectArgsBuffer = GraphBuilder.CreateBuffer(FRDGBufferDesc::CreateIndirectDesc<FRHIDispatchIndirectParameters>(1), TEXT("Hair.PageScatterIndirectArgs"));
 			{
 				check(MacroGroup.MacroGroupId < MacroGroupCount);
 
@@ -902,7 +902,7 @@ FVirtualVoxelResources AllocateVirtualVoxelResources(
 			FClearValueBinding::Black, 
 			TexCreate_UAV | TexCreate_ShaderResource,
 			MipCount);
-		Out.PageTexture = GraphBuilder.CreateTexture(Desc, TEXT("VoxelPageTexture"));
+		Out.PageTexture = GraphBuilder.CreateTexture(Desc, TEXT("Hair.VoxelPageTexture"));
 	}
 
 	Out.Parameters.Common.PageIndexBuffer			= GraphBuilder.CreateSRV(Out.PageIndexBuffer, PF_R32_UINT);
@@ -931,7 +931,7 @@ static FRDGBufferRef IndirectVoxelPageClear(
 	SCOPED_DRAW_EVENT(GraphBuilder.RHICmdList, HairStrandsIndVoxelPageClear);
 	SCOPED_GPU_STAT(GraphBuilder.RHICmdList, HairStrandsIndVoxelPageClear);
 
-	FRDGBufferRef ClearIndArgsBuffer = GraphBuilder.CreateBuffer(FRDGBufferDesc::CreateIndirectDesc<FRHIDispatchIndirectParameters>(1), TEXT("VirtualVoxelClearIndirectArgsBuffer"));
+	FRDGBufferRef ClearIndArgsBuffer = GraphBuilder.CreateBuffer(FRDGBufferDesc::CreateIndirectDesc<FRHIDispatchIndirectParameters>(1), TEXT("Hair.VirtualVoxelClearIndirectArgsBuffer"));
 
 	// Generate the indirect buffer required to clear all voxel allocated linearly in the page volume texture, using the global counter.
 	{
@@ -1205,7 +1205,7 @@ static void AddVirtualVoxelizationRasterPass(
 
 	// For debug purpose
 	#if 0
-	FRDGTextureRef DebugOutputTexture = GraphBuilder.CreateTexture(FRDGTextureDesc::Create2D(RasterResolution, PF_R32_UINT, FClearValueBinding::Black, TexCreate_RenderTargetable), TEXT("DummyTexture"));
+	FRDGTextureRef DebugOutputTexture = GraphBuilder.CreateTexture(FRDGTextureDesc::Create2D(RasterResolution, PF_R32_UINT, FClearValueBinding::Black, TexCreate_RenderTargetable), TEXT("Hair.DummyTexture"));
 	PassParameters->RenderTargets[0] = FRenderTargetBinding(DebugOutputTexture, ERenderTargetLoadAction::EClear);
 	#endif
 
@@ -1321,7 +1321,7 @@ static void AddVirtualVoxelGenerateMipPass(
 	{
 		const uint32 TargetMipIndex = MipIt + 1;
 		const uint32 DispatchGroupSize = 64;
-		FRDGBufferRef MipIndirectArgs = GraphBuilder.CreateBuffer(FRDGBufferDesc::CreateIndirectDesc<FRHIDispatchIndirectParameters>(), TEXT("VirtualVoxelMipIndirectArgsBuffer"));
+		FRDGBufferRef MipIndirectArgs = GraphBuilder.CreateBuffer(FRDGBufferDesc::CreateIndirectDesc<FRHIDispatchIndirectParameters>(), TEXT("Hair.VirtualVoxelMipIndirectArgsBuffer"));
 		MipIndirectArgsBuffers.Add(MipIndirectArgs);
 
 		FVirtualVoxelIndirectArgMipCS::FParameters* Parameters = GraphBuilder.AllocParameters<FVirtualVoxelIndirectArgMipCS::FParameters>();
