Remove unneeded check on initialized RT data (make same as cards and strands)

diff --git a/Engine/Plugins/Runtime/HairStrands/Source/HairStrandsCore/Private/HairStrandsRendering.cpp b/Engine/Plugins/Runtime/HairStrands/Source/HairStrandsCore/Private/HairStrandsRendering.cpp
--- a/Engine/Plugins/Runtime/HairStrands/Source/HairStrandsCore/Private/HairStrandsRendering.cpp
+++ b/Engine/Plugins/Runtime/HairStrands/Source/HairStrandsCore/Private/HairStrandsRendering.cpp
@@ -1659,7 +1659,7 @@ void ComputeHairStrandsInterpolation(
 
 		#if RHI_RAYTRACING
 		FHairGroupInstance::FMeshes::FLOD& LOD = Instance->Meshes.LODs[HairLODIndex];
-		if (LOD.RaytracingResource && !LOD.RaytracingResource->bIsRTGeometryInitialized)
+		if (LOD.RaytracingResource)
 		{
 			GraphBuilder.AddPass(
 				RDG_EVENT_NAME("HairMeshesUpdateBLAS"),
