Add per LOD configuration for hair simulation and hair RBF deformation.

Ran preflight on this CL.

diff --git a/Engine/Plugins/Runtime/HairStrands/Source/HairStrandsCore/Private/GroomAsset.cpp b/Engine/Plugins/Runtime/HairStrands/Source/HairStrandsCore/Private/GroomAsset.cpp
--- a/Engine/Plugins/Runtime/HairStrands/Source/HairStrandsCore/Private/GroomAsset.cpp
+++ b/Engine/Plugins/Runtime/HairStrands/Source/HairStrandsCore/Private/GroomAsset.cpp
@@ -56,6 +56,10 @@ bool IsHairStrandsDDCLogEnable()
 	return GHairStrandsDDCLogEnable > 0;
 }
 
+static int32 GHairMaxSimulatedLOD = -1;
+static FAutoConsoleVariableRef CVarHairMaxSimulatedLOD(TEXT("r.HairStrands.MaxSimulatedLOD"), GHairMaxSimulatedLOD, TEXT("Maximum hair LOD to be simulated"));
+static bool IsHairLODSimulationEnabled(const int32 LODIndex) { return (LODIndex >= 0 && (GHairMaxSimulatedLOD < 0 || (GHairMaxSimulatedLOD >= 0 && LODIndex <= GHairMaxSimulatedLOD))); }
+
 /////////////////////////////////////////////////////////////////////////////////////////
 
 enum class EHairAtlasTextureType
@@ -2337,6 +2341,59 @@ void UGroomAsset::InitMeshesResources()
 	}
 }
 
+bool UGroomAsset::IsSimulationEnable(int32 GroupIndex, int32 LODIndex)
+{
+	if (GroupIndex < 0 || GroupIndex >= HairGroupsLOD.Num() || !IsHairStrandsSimulationEnable())
+	{
+		return false;
+	}
+
+	if (LODIndex >= HairGroupsLOD[GroupIndex].LODs.Num())
+	{
+		return false;
+	}
+
+	check(HairGroupsPhysics.Num() == HairGroupsLOD.Num());
+
+	// If the LOD index is not forced, then its value could be -1. In this case we return the 'global' asset value
+	if (LODIndex < 0)
+	{
+		return HairGroupsPhysics[GroupIndex].SolverSettings.EnableSimulation;
+	}
+
+	if (!IsHairLODSimulationEnabled(LODIndex))
+	{
+		return false;
+	}
+
+	return
+		HairGroupsLOD[GroupIndex].LODs[LODIndex].Simulation == EGroomOverrideType::Enable ||
+		(HairGroupsLOD[GroupIndex].LODs[LODIndex].Simulation == EGroomOverrideType::Auto && HairGroupsPhysics[GroupIndex].SolverSettings.EnableSimulation);
+}
+
+bool UGroomAsset::IsGlobalInterpolationEnable(int32 GroupIndex, int32 LODIndex)
+{
+	if (GroupIndex < 0 || GroupIndex >= HairGroupsLOD.Num())
+	{
+		return false;
+	}
+
+	if (LODIndex < 0 || LODIndex >= HairGroupsLOD[GroupIndex].LODs.Num())
+	{
+		return false;
+	}
+
+	// If the LOD index is not forced, then its value could be -1. In this case we return the 'global' asset value
+	if (LODIndex < 0)
+	{
+		return EnableGlobalInterpolation;
+	}
+
+	return
+		HairGroupsLOD[GroupIndex].LODs[LODIndex].GlobalInterpolation == EGroomOverrideType::Enable ||
+		(HairGroupsLOD[GroupIndex].LODs[LODIndex].GlobalInterpolation == EGroomOverrideType::Auto && EnableGlobalInterpolation);
+}
+
 bool FProcessedHairDescription::IsValid() const
 {
 	for (TPair<int32, FProcessedHairDescription::FHairGroup> HairGroupIt : HairGroups)
diff --git a/Engine/Plugins/Runtime/HairStrands/Source/HairStrandsCore/Private/GroomComponent.cpp b/Engine/Plugins/Runtime/HairStrands/Source/HairStrandsCore/Private/GroomComponent.cpp
--- a/Engine/Plugins/Runtime/HairStrands/Source/HairStrandsCore/Private/GroomComponent.cpp
+++ b/Engine/Plugins/Runtime/HairStrands/Source/HairStrandsCore/Private/GroomComponent.cpp
@@ -32,10 +32,6 @@ static float GHairClipLength = -1;
 static FAutoConsoleVariableRef CVarHairClipLength(TEXT("r.HairStrands.DebugClipLength"), GHairClipLength, TEXT("Clip hair strands which have a lenth larger than this value. (default is -1, no effect)"));
 float GetHairClipLength() { return GHairClipLength > 0 ? GHairClipLength : 100000;  }
 
-static int32 GHairMaxSimulatedLOD = -1;
-static FAutoConsoleVariableRef CVarHairMaxSimulatedLOD(TEXT("r.HairStrands.MaxSimulatedLOD"), GHairMaxSimulatedLOD, TEXT("Maximum hair LOD to be simulated"));
-bool IsHairLODSimulationEnabled(const int32 LODIndex) { return (LODIndex >= 0 && (GHairMaxSimulatedLOD < 0 || (GHairMaxSimulatedLOD >= 0 && LODIndex <= GHairMaxSimulatedLOD))); }
-
 static int32 GHairEnableAdaptiveSubsteps = 0;  
 static FAutoConsoleVariableRef CVarHairEnableAdaptiveSubsteps(TEXT("r.HairStrands.EnableAdaptiveSubsteps"), GHairEnableAdaptiveSubsteps, TEXT("Enable adaptive solver substeps"));
 bool IsHairAdaptiveSubstepsEnabled() { return (GHairEnableAdaptiveSubsteps == 1); }
@@ -44,6 +40,30 @@ bool IsHairAdaptiveSubstepsEnabled() { return (GHairEnableAdaptiveSubsteps == 1)
 
 ////////////////////////////////////////////////////////////////////////////////////////////////////////////////////////
 
+static EHairBindingType ToHairBindingType(EGroomBindingType In)
+{
+	switch (In)
+	{
+	case EGroomBindingType::Rigid		: return EHairBindingType::Rigid;
+	case EGroomBindingType::Skinning	: return EHairBindingType::Skinning;
+	case EGroomBindingType::NoneBinding	: return EHairBindingType::NoneBinding;
+	}
+	return EHairBindingType::NoneBinding;
+}
+
+static EHairInterpolationType ToHairInterpolationType(EGroomInterpolationType In)
+{
+	switch (In)
+	{
+	case EGroomInterpolationType::None				: return EHairInterpolationType::NoneSkinning; 
+	case EGroomInterpolationType::RigidTransform	: return EHairInterpolationType::RigidSkinning;
+	case EGroomInterpolationType::OffsetTransform	: return EHairInterpolationType::OffsetSkinning;
+	case EGroomInterpolationType::SmoothTransform	: return EHairInterpolationType::SmoothSkinning;
+	}
+	return EHairInterpolationType::NoneSkinning;
+}
+////////////////////////////////////////////////////////////////////////////////////////////////////////////////////////
+
 static FHairGroupDesc GetGroomGroupsDesc(const UGroomAsset* Asset, UGroomComponent* Component, uint32 GroupIndex)
 {
 	if (!Asset)
@@ -1023,7 +1043,7 @@ void UGroomComponent::ReleaseHairSimulation()
 	NiagaraComponents.Empty();
 }
 
-void EnableHairSimulation(UGroomComponent* GroomComponent, const bool bEnableSimulation)
+static void InternalUpdateHairSimulation(UGroomComponent* GroomComponent)
 {
 	if (!GroomComponent)
 	{
@@ -1038,11 +1058,12 @@ void EnableHairSimulation(UGroomComponent* GroomComponent, const bool bEnableSim
 
 	bool NeedSpringsSolver = false;
 	bool NeedRodsSolver = false;
-	if (GroomAsset)
+	if (GroomAsset) 
 	{
 		for (int32 i = 0; i < NumGroups; ++i)
 		{
-			ValidComponents[i] = GroomAsset->HairGroupsPhysics[i].SolverSettings.EnableSimulation && IsHairStrandsSimulationEnable() && bEnableSimulation;
+			const int32 LODIndex = GroomComponent->GroomGroupsDesc[i].LODForcedIndex;
+			ValidComponents[i] = GroomAsset->IsSimulationEnable(i, LODIndex);
 			if (ValidComponents[i] && (GroomAsset->HairGroupsPhysics[i].SolverSettings.NiagaraSolver == EGroomNiagaraSolvers::AngularSprings))
 			{
 				NeedSpringsSolver = true;
@@ -1123,9 +1144,9 @@ void EnableHairSimulation(UGroomComponent* GroomComponent, const bool bEnableSim
 	GroomComponent->UpdateSimulatedGroups();
 }
 
-void UGroomComponent::UpdateHairSimulation()  
+void UGroomComponent::UpdateHairSimulation()
 {
-	EnableHairSimulation(this,true);
+	InternalUpdateHairSimulation(this);
 }
 
 void UGroomComponent::SetGroomAsset(UGroomAsset* Asset)
@@ -1256,25 +1277,17 @@ void UGroomComponent::SetForcedLOD(int32 LODIndex)
 		LODIndex = -1;
 	}
 
-	const bool bValidLODA = IsHairLODSimulationEnabled(GetForcedLOD());
-	const bool bValidLODB = IsHairLODSimulationEnabled(LODIndex);
-
+	bool bHasLODChanged = false;
 	for (FHairGroupDesc& HairDesc : GroomGroupsDesc)
 	{
+		bHasLODChanged = bHasLODChanged || HairDesc.LODForcedIndex != LODIndex;
 		HairDesc.LODForcedIndex = LODIndex;
 	}
 	UpdateHairGroupsDesc();
 	
-	if (bValidLODA != bValidLODB)
+	if (bHasLODChanged)
 	{
-		if (bValidLODB)
-		{
-			EnableHairSimulation(this, true);
-		}
-		else 
-		{
-			EnableHairSimulation(this, false);
-		}
+		InternalUpdateHairSimulation(this);
 	}
 
 	// Do not invalidate completly the proxy, but just update LOD index on the rendering thread
@@ -1674,30 +1687,27 @@ void UGroomComponent::UpdateSimulatedGroups()
 		
 		const bool bIsStrandsEnabled = IsHairStrandsEnabled(EHairStrandsShaderType::Strands);
 
+		// For now use the force LOD to drive enabling/disabling simulation & RBF
+		const int32 LODIndex = GetForcedLOD();
+
 		TArray<FHairGroupInstance*> LocalInstances = HairGroupInstances;
 		UGroomAsset* LocalGroomAsset = GroomAsset;
 		UGroomBindingAsset* LocalBindingAsset = BindingAsset;
 		ENQUEUE_RENDER_COMMAND(FHairStrandsTick_UEnableSimulatedGroups)(
-			[LocalInstances, LocalGroomAsset, LocalBindingAsset, Id, WorldType, bIsStrandsEnabled](FRHICommandListImmediate& RHICmdList)
-		{
+			[LocalInstances, LocalGroomAsset, LocalBindingAsset, Id, WorldType, bIsStrandsEnabled, LODIndex](FRHICommandListImmediate& RHICmdList)
+			{
 			int32 GroupIt = 0;
 			for (FHairGroupInstance* Instance : LocalInstances)
 			{
-				const bool bIsSimulationEnable = (LocalGroomAsset && GroupIt < LocalGroomAsset->HairGroupsPhysics.Num()) ? 
-					(LocalGroomAsset->HairGroupsPhysics[GroupIt].SolverSettings.EnableSimulation && IsHairStrandsSimulationEnable()): false;
-				const bool bHasGlobalInterpolation = LocalBindingAsset && LocalGroomAsset && LocalGroomAsset->EnableGlobalInterpolation;
-				Instance->Strands.HairInterpolationType = 0;
-				if (bIsStrandsEnabled)
+				Instance->Strands.HairInterpolationType = EHairInterpolationType::NoneSkinning;
+				if (bIsStrandsEnabled && LocalGroomAsset)
 				{
-					Instance->Strands.HairInterpolationType =
-						(LocalGroomAsset && LocalGroomAsset->HairInterpolationType == EGroomInterpolationType::RigidTransform) ? 0 :
-						(LocalGroomAsset && LocalGroomAsset->HairInterpolationType == EGroomInterpolationType::OffsetTransform) ? 1 :
-						(LocalGroomAsset && LocalGroomAsset->HairInterpolationType == EGroomInterpolationType::SmoothTransform) ? 2 : 0;
+					Instance->Strands.HairInterpolationType = ToHairInterpolationType(LocalGroomAsset->HairInterpolationType);
 				}
 				if (Instance->Guides.IsValid())
 				{
-					Instance->Guides.bIsSimulationEnable = bIsSimulationEnable;
-					Instance->Guides.bHasGlobalInterpolation = bHasGlobalInterpolation;
+					Instance->Guides.bIsSimulationEnable	 = LocalGroomAsset && LocalGroomAsset->IsSimulationEnable(GroupIt, LODIndex);
+					Instance->Guides.bHasGlobalInterpolation = LocalGroomAsset && LocalGroomAsset->IsGlobalInterpolationEnable(GroupIt, LODIndex);
 				}
 				++GroupIt;
 			}
@@ -1799,18 +1809,6 @@ static USkeletalMeshComponent* ValidateBindingAsset(
 
 void CreateHairStrandsDebugAttributeBuffer(FRDGExternalBuffer* DebugAttributeBuffer, uint32 VertexCount);
 
-
-static EHairBindingType ToHairBindingType(EGroomBindingType In)
-{
-	switch (In)
-	{
-	case EGroomBindingType::Rigid:		return EHairBindingType::Rigid;
-	case EGroomBindingType::Skinning:	return EHairBindingType::Skinning;
-	case EGroomBindingType::NoneBinding:return EHairBindingType::NoneBinding;
-	}
-	return EHairBindingType::NoneBinding;
-}
-
 void UGroomComponent::InitResources(bool bIsBindingReloading)
 {
 	LLM_SCOPE(ELLMTag::Meshes) // This should be a Groom LLM tag, but there is no LLM tag bit left
@@ -1828,8 +1826,10 @@ void UGroomComponent::InitResources(bool bIsBindingReloading)
 
 	InitializedResources = GroomAsset;
 
-	// 1. Check if we need any kind of binding data 
+	// 1. Check if we need any kind of binding data, simulation data, or RBF data
 	bool bHasNeedBindingData = false;
+	bool bHasNeedGlobalDeformation = false;
+	bool bHasNeedSimulation = false;
 	for (int32 GroupIt = 0, GroupCount = GroomAsset->HairGroupsData.Num(); GroupIt < GroupCount; ++GroupIt)
 	{
 		const FHairGroupsLOD& GroupLOD = GroomAsset->HairGroupsLOD[GroupIt];
@@ -1837,6 +1837,12 @@ void UGroomComponent::InitResources(bool bIsBindingReloading)
 		{
 			bHasNeedBindingData  = bHasNeedBindingData || LODSettings.BindingType == EGroomBindingType::Skinning;
 		}
+
+		for (uint32 LODIt = 0, LODCount = GroomAsset->GetLODCount(); LODIt < LODCount; ++LODIt)
+		{
+			bHasNeedSimulation = bHasNeedSimulation || GroomAsset->IsSimulationEnable(GroupIt, LODIt);
+			bHasNeedGlobalDeformation = bHasNeedGlobalDeformation || GroomAsset->IsGlobalInterpolationEnable(GroupIt, LODIt);
+		}
 	}
 
 	const FPrimitiveComponentId LocalComponentId = ComponentId;
@@ -1859,9 +1865,6 @@ void UGroomComponent::InitResources(bool bIsBindingReloading)
 		bUseAttachParentBound = true;
 	}
 
-	FTransform HairLocalToWorld = GetComponentTransform();
-	FTransform SkinLocalToWorld = SkeletalMeshComponent ? SkeletalMeshComponent->GetComponentTransform() : FTransform::Identity;
-	
 	for (int32 GroupIt = 0, GroupCount = GroomAsset->HairGroupsData.Num(); GroupIt < GroupCount; ++GroupIt)
 	{
 		FHairGroupInstance* HairGroupInstance = new FHairGroupInstance();
@@ -1882,10 +1885,7 @@ void UGroomComponent::InitResources(bool bIsBindingReloading)
 		FHairGroupData& GroupData = GroomAsset->HairGroupsData[GroupIt];
 		const uint32 SkeletalLODCount = SkeletalMeshComponent ? SkeletalMeshComponent->GetNumLODs() : 0;
 
-		const uint32 HairInterpolationType =
-			GroomAsset->HairInterpolationType == EGroomInterpolationType::RigidTransform ? 0 :
-			GroomAsset->HairInterpolationType == EGroomInterpolationType::OffsetTransform ? 1 :
-			GroomAsset->HairInterpolationType == EGroomInterpolationType::SmoothTransform ? 2 : 0;
+		const EHairInterpolationType HairInterpolationType = ToHairInterpolationType(GroomAsset->HairInterpolationType);
 
 		const bool bDynamicResources = IsHairStrandsSimulationEnable() || IsHairStrandsBindingEnable();
 
@@ -1896,8 +1896,9 @@ void UGroomComponent::InitResources(bool bIsBindingReloading)
 			TArray<bool> LODVisibility;
 			TArray<EHairGeometryType> LODGeometryTypes;
 			const FHairGroupsLOD& GroupLOD = GroomAsset->HairGroupsLOD[GroupIt];
-			for (const FHairLODSettings& LODSettings : GroupLOD.LODs)
+			for (uint32 LODIt = 0, LODCount = GroupLOD.LODs.Num(); LODIt < LODCount; ++LODIt)
 			{
+				const FHairLODSettings& LODSettings = GroupLOD.LODs[LODIt];
 				EHairBindingType BindingType = ToHairBindingType(LODSettings.BindingType);
 				if (BindingType == EHairBindingType::Skinning && !LocalBindingAsset)
 				{
@@ -1908,10 +1909,15 @@ void UGroomComponent::InitResources(bool bIsBindingReloading)
 					BindingType = EHairBindingType::NoneBinding;
 				}
 
+				const bool LODSimulation = GroomAsset->IsSimulationEnable(GroupIt, LODIt);
+				const bool LODGlobalInterpolation = GroomAsset->IsGlobalInterpolationEnable(GroupIt, LODIt);
+
 				CPULODScreenSize.Add(LODSettings.ScreenSize);
 				LODVisibility.Add(LODSettings.bVisible);
 				LODGeometryTypes.Add(ConvertToHairGeometryType(LODSettings.GeometryType));
 				HairGroupInstance->HairGroupPublicData->BindingTypes.Add(BindingType);
+				HairGroupInstance->HairGroupPublicData->LODSimulations.Add(LODSimulation);
+				HairGroupInstance->HairGroupPublicData->LODGlobalInterpolations.Add(LODGlobalInterpolation);
 			}
 			HairGroupInstance->HairGroupPublicData->SetLODScreenSizes(CPULODScreenSize);
 			HairGroupInstance->HairGroupPublicData->SetLODVisibilities(LODVisibility);
@@ -1923,7 +1929,7 @@ void UGroomComponent::InitResources(bool bIsBindingReloading)
 		// * Physics simulation
 		// * RBF deformation.
 		// Therefore, even if simulation is disabled, we need to run partially the update if the binding system is enabled (skin deformation + RBF correction)
-		if (GroupData.Guides.IsValid() && (bDynamicResources || bIsStrandsEnabled))
+		if (GroupData.Guides.IsValid() && (bDynamicResources || bIsStrandsEnabled) && (bHasNeedSimulation || bHasNeedGlobalDeformation))
 		{
 			HairGroupInstance->Guides.Data = &GroupData.Guides.Data;
 
@@ -1947,15 +1953,12 @@ void UGroomComponent::InitResources(bool bIsBindingReloading)
 			HairGroupInstance->Guides.DeformedResource->GetPositionOffset(FHairStrandsDeformedResource::EFrameType::Current)  = HairGroupInstance->Guides.RestResource->PositionOffset;
 			HairGroupInstance->Guides.DeformedResource->GetPositionOffset(FHairStrandsDeformedResource::EFrameType::Previous) = HairGroupInstance->Guides.RestResource->PositionOffset;
 
-			HairGroupInstance->Guides.bIsSimulationEnable = 
-				(IsHairStrandsSimulationEnable() && GroomAsset && GroupIt < GroomAsset->HairGroupsPhysics.Num()) ?
-				(GroomAsset->HairGroupsPhysics[GroupIt].SolverSettings.EnableSimulation) :
-				false;
-
-			HairGroupInstance->Guides.bHasGlobalInterpolation = LocalBindingAsset && GroomAsset && GroomAsset->EnableGlobalInterpolation;
+			// Initialize the simulation and the global deformation to its default behavior by setting it with LODIndex = -1
+			const int32 LODIndex = -1;
+			HairGroupInstance->Guides.bIsSimulationEnable = GroomAsset && GroomAsset->IsGlobalInterpolationEnable(GroupIt,LODIndex);
+			HairGroupInstance->Guides.bHasGlobalInterpolation = LocalBindingAsset && GroomAsset && GroomAsset->IsGlobalInterpolationEnable(GroupIt,LODIndex);
 		}
 
-
 		// LODBias is in the Modifier which is needed for LOD selection regardless if the strands are there or not
 		HairGroupInstance->Strands.Modifier = GetGroomGroupsDesc(GroomAsset, this, GroupIt);
 
diff --git a/Engine/Plugins/Runtime/HairStrands/Source/HairStrandsCore/Private/GroomDebug.cpp b/Engine/Plugins/Runtime/HairStrands/Source/HairStrandsCore/Private/GroomDebug.cpp
--- a/Engine/Plugins/Runtime/HairStrands/Source/HairStrandsCore/Private/GroomDebug.cpp
+++ b/Engine/Plugins/Runtime/HairStrands/Source/HairStrandsCore/Private/GroomDebug.cpp
@@ -832,7 +832,7 @@ void RunHairStrandsDebug(
 				if (!bIsActive)
 					continue;
 
-				Line = FString::Printf(TEXT(" * Group:%d/%d | LOD:%1.1f/%d | GeometryType:%s | BindingType:%s | Sim:%d | RBF:%d | VertexCount:%d | Name: %s"),
+				Line = FString::Printf(TEXT(" * Group:%d/%d | LOD:%1.2f/%d | GeometryType:%s | BindingType:%s | Sim:%d | RBF:%d | VertexCount:%d | Name: %s"),
 					Instance->Debug.GroupIndex,
 					Instance->Debug.GroupCount,
 
@@ -841,8 +841,8 @@ void RunHairStrandsDebug(
 
 					ToString(Instance->GeometryType),
 					ToString(Instance->BindingType),
-					0, // TODO: Instance->bHasSimulation,
-					0, // TODO: Instance->bHasRBF,
+					Instance->Guides.bIsSimulationEnable,
+					Instance->Guides.bHasGlobalInterpolation,
 					Instance->HairGroupPublicData->VertexCount,
 					*Instance->Debug.GroomAssetName);
 				Canvas.DrawShadowedString(X, Y += YStep, *Line, GetStatsFont(), bIsActive ? DebugGroupColor : InactiveColor);
diff --git a/Engine/Plugins/Runtime/HairStrands/Source/HairStrandsCore/Private/GroomManager.cpp b/Engine/Plugins/Runtime/HairStrands/Source/HairStrandsCore/Private/GroomManager.cpp
--- a/Engine/Plugins/Runtime/HairStrands/Source/HairStrandsCore/Private/GroomManager.cpp
+++ b/Engine/Plugins/Runtime/HairStrands/Source/HairStrandsCore/Private/GroomManager.cpp
@@ -102,8 +102,6 @@ void UnregisterHairStrands(uint32 ComponentId)
 
 ///////////////////////////////////////////////////////////////////////////////////////////////////
 
-bool HasHairInstanceSimulationEnable(FHairGroupInstance* Instance, int32 MeshLODIndex);
-
 static bool IsInstanceVisible(
 	const FSceneView* View,
 	FHairGroupInstance* Instance)
@@ -187,10 +185,11 @@ static void RunInternalHairStrandsInterpolation(
 		Instance->Debug.MeshLODIndex = MeshLODIndex;
 		if (0 <= MeshLODIndex)
 		{
+			const EHairGeometryType InstanceGeometryType = Instance->GeometryType;
+			const EHairBindingType InstanceBindingType = Instance->BindingType;
+
 			if (EHairStrandsInterpolationType::RenderStrands == Type)
 			{
-				const EHairGeometryType InstanceGeometryType = Instance->GeometryType;
-				const EHairBindingType InstanceBindingType = Instance->BindingType;
 				if (InstanceGeometryType == EHairGeometryType::Strands)
 				{
 					if (Instance->Strands.HasValidRootData() && Instance->Strands.DeformedRootResource->IsValid(MeshLODIndex) && InstanceBindingType == EHairBindingType::Skinning)
@@ -227,7 +226,7 @@ static void RunInternalHairStrandsInterpolation(
 					if (Instance->Cards.IsValid(HairLODIndex))
 					{
 						FHairGroupInstance::FCards::FLOD& CardsInstance = Instance->Cards.LODs[HairLODIndex];
-						if (CardsInstance.Guides.IsValid() && Instance->Guides.HasValidRootData() && CardsInstance.Guides.DeformedRootResource->IsValid(MeshLODIndex) && NeedsUpdateCardsMeshTriangles() && InstanceBindingType == EHairBindingType::Skinning)
+						if (CardsInstance.Guides.IsValid() && CardsInstance.Guides.HasValidRootData() && CardsInstance.Guides.DeformedRootResource->IsValid(MeshLODIndex) && NeedsUpdateCardsMeshTriangles() && InstanceBindingType == EHairBindingType::Skinning)
 						{
 							AddHairStrandUpdateMeshTrianglesPass(
 								GraphBuilder,
@@ -263,7 +262,9 @@ static void RunInternalHairStrandsInterpolation(
 			}
 			else if (EHairStrandsInterpolationType::SimulationStrands == Type)
 			{
-				if (Instance->Guides.IsValid() && Instance->Guides.HasValidRootData() && Instance->Guides.DeformedRootResource->IsValid(MeshLODIndex) )
+				if (Instance->Guides.IsValid() && Instance->Guides.HasValidRootData() && Instance->Guides.DeformedRootResource->IsValid(MeshLODIndex) && 
+					InstanceBindingType == EHairBindingType::Skinning &&
+					(Instance->Guides.bIsSimulationEnable || Instance->Guides.bHasGlobalInterpolation))
 				{
 					AddHairStrandUpdateMeshTrianglesPass(
 						GraphBuilder,
@@ -324,7 +325,7 @@ static void RunInternalHairStrandsInterpolation(
 				continue;
 
 			// Frustum culling guide deformation if the instance does not have any simulation
-			if (Instance->GeometryType == EHairGeometryType::Strands && !HasHairInstanceSimulationEnable(Instance, Instance->Debug.MeshLODIndex) && !IsInstanceVisible(View, Instance))
+			if (Instance->GeometryType == EHairGeometryType::Strands && !Instance->Guides.bIsSimulationEnable  && !IsInstanceVisible(View, Instance))
 			{
 				continue;
 			}
@@ -653,6 +654,9 @@ static void RunHairLODSelection(EWorldType::Type WorldType, const TArray<const F
 		Instance->HairGroupPublicData->VFInput.bHasLODSwitch = (FMath::FloorToInt(PrevLODIndex) != FMath::FloorToInt(LODIndex));
 		Instance->GeometryType = GeometryType;
 		Instance->BindingType = Instance->HairGroupPublicData->GetBindingType(IntLODIndex);
+		Instance->Guides.bIsSimulationEnable = Instance->HairGroupPublicData->IsSimulationEnable(IntLODIndex);
+		Instance->Guides.bHasGlobalInterpolation = Instance->HairGroupPublicData->IsGlobalInterpolationEnable(IntLODIndex);
+
 	}
 }
 
diff --git a/Engine/Plugins/Runtime/HairStrands/Source/HairStrandsCore/Private/HairStrandsRendering.cpp b/Engine/Plugins/Runtime/HairStrands/Source/HairStrandsCore/Private/HairStrandsRendering.cpp
--- a/Engine/Plugins/Runtime/HairStrands/Source/HairStrandsCore/Private/HairStrandsRendering.cpp
+++ b/Engine/Plugins/Runtime/HairStrands/Source/HairStrandsCore/Private/HairStrandsRendering.cpp
@@ -58,8 +58,6 @@ EHairCardsSimulationType GetHairCardsSimulationType()
 		(GHairCardsInterpolationType >= 1 ? EHairCardsSimulationType::Guide : EHairCardsSimulationType::None);
 }
 
-bool IsHairLODSimulationEnabled(const int32 LODIndex);
-
 bool NeedsUpdateCardsMeshTriangles()
 {
 	return GetHairCardsSimulationType() == EHairCardsSimulationType::Guide;
@@ -327,6 +325,36 @@ static void AddDeformSimHairStrandsPass(
 
 ///////////////////////////////////////////////////////////////////////////////////////////////////
 
+struct FCopyBufferParameter
+{
+	BEGIN_SHADER_PARAMETER_STRUCT(FParameters, )
+		RDG_BUFFER_ACCESS(DstBuffer, ERHIAccess::CopyDest)
+	END_SHADER_PARAMETER_STRUCT()
+};
+
+static void AddCopyBufferPass(
+	FRDGBuilder& GraphBuilder,
+	FRDGBufferRef DstBuffer,
+	FBufferRHIRef SrcBuffer)
+{
+
+	FCopyBufferParameter::FParameters* Parameters = GraphBuilder.AllocParameters<FCopyBufferParameter::FParameters>();
+	Parameters->DstBuffer = DstBuffer;
+
+	const uint32 TotalByteCount = DstBuffer->Desc.BytesPerElement * DstBuffer->Desc.NumElements;
+
+	GraphBuilder.AddPass(
+		RDG_EVENT_NAME("HairCardsUpdateBLAS"),
+		Parameters,
+		ERDGPassFlags::Copy,
+		[Parameters, SrcBuffer, TotalByteCount](FRHICommandList& RHICmdList)
+		{
+			RHICmdList.CopyBufferRegion(Parameters->DstBuffer->GetRHI(), 0, SrcBuffer, 0, TotalByteCount);
+		});
+}
+
+///////////////////////////////////////////////////////////////////////////////////////////////////
+
 struct FRDGHairStrandsCullingData
 {
 	bool bCullingResultAvailable = false;
@@ -470,6 +498,14 @@ public:
 
 IMPLEMENT_GLOBAL_SHADER(FHairInterpolationCS, "/Engine/Private/HairStrands/HairStrandsInterpolation.usf", "MainCS", SF_Compute);
 
+// AddHairStrandsInterpolationPass takes in input
+// * Guide resources (optional)
+// * Guide root triangles (optional)
+// * Strands resources
+// * Strands root triangles (optional)
+// 
+// This functions is in charge of deforming the rendering strands (OutRenderPositionBuffer) based on the guides
+// The guides are either the result of a simulation deformation, or a global interpolation deformation (aka RBF deformation), or both
 static void AddHairStrandsInterpolationPass(
 	FRDGBuilder& GraphBuilder,
 	FGlobalShaderMap* ShaderMap,
@@ -478,8 +514,8 @@ static void AddHairStrandsInterpolationPass(
 	const uint32 VertexCount,
 	const FHairScaleAndClipDesc ScaleAndClipDesc,
 	const int32 MeshLODIndex,
-	const bool bPatchedAttributeBuffer, 	
-	const uint32 HairInterpolationType,
+	const bool bPatchedAttributeBuffer,
+	const EHairInterpolationType HairInterpolationType,
 	const EHairGeometryType InstanceGeometryType,
 	FRDGHairStrandsCullingData& CullingData,
 	const FVector& InRenderHairWorldOffset,
@@ -535,16 +571,20 @@ static void AddHairStrandsInterpolationPass(
 	Parameters->HairLengthClipEnable = ScaleAndClipDesc.IsEnable() ? 1 : 0;
 	Parameters->AttributeBuffer = RenderAttributeBuffer;	
 	
-	const bool bIsVertexToCurveBuffersValid =
+	const bool bIsVertexToCurveBuffersValid_Sim =
 		SimRestRootResources &&
-		SimRestRootResources->VertexToCurveIndexBuffer.Buffer != nullptr &&
+		SimRestRootResources->VertexToCurveIndexBuffer.Buffer != nullptr;
+	if (bIsVertexToCurveBuffersValid_Sim)
+	{
+		Parameters->SimVertexToRootIndexBuffer = RegisterAsSRV(GraphBuilder, SimRestRootResources->VertexToCurveIndexBuffer);
+	}
+
+	const bool bIsVertexToCurveBuffersValid_Ren =
 		RenRestRootResources &&
 		RenRestRootResources->VertexToCurveIndexBuffer.Buffer != nullptr;
-
-	if (bIsVertexToCurveBuffersValid)
+	if (bIsVertexToCurveBuffersValid_Ren)
 	{
 		Parameters->RenVertexToRootIndexBuffer = RegisterAsSRV(GraphBuilder, RenRestRootResources->VertexToCurveIndexBuffer);
-		Parameters->SimVertexToRootIndexBuffer = RegisterAsSRV(GraphBuilder, SimRestRootResources->VertexToCurveIndexBuffer);
 	}
 
 	Parameters->VertexToClusterIdBuffer = VertexToClusterIdBuffer;
@@ -555,7 +595,7 @@ static void AddHairStrandsInterpolationPass(
 	Parameters->HairDebugMode = 0;
 	{
 		const FHairCullInfo Info = GetHairStrandsCullInfo();
-		const bool bCullingEnable = Info.CullMode != EHairCullMode::None && bIsVertexToCurveBuffersValid;
+		const bool bCullingEnable = Info.CullMode != EHairCullMode::None && bIsVertexToCurveBuffersValid_Ren;
 
 		if (bPatchedAttributeBuffer && OutRenderAttributeBuffer)
 		{
@@ -589,7 +629,7 @@ static void AddHairStrandsInterpolationPass(
 	}
 
 	const bool bSupportDynamicMesh = 
-		bIsVertexToCurveBuffersValid &&
+		bIsVertexToCurveBuffersValid_Ren &&
 		RenRestRootResources &&
 		RenRestRootResources->RootData.RootCount > 0 && 
 		MeshLODIndex >= 0 && 
@@ -598,22 +638,14 @@ static void AddHairStrandsInterpolationPass(
 		RenRestRootResources->LODs[MeshLODIndex].IsValid() &&
 		RenDeformedRootResources->LODs[MeshLODIndex].IsValid();
 
+	// Guides
 	bool bSupportGlobalInterpolation = false;
-	if (bSupportDynamicMesh)
+	if (bSupportDynamicMesh && (Instance->Guides.bIsSimulationEnable || Instance->Guides.bHasGlobalInterpolation))
 	{
 		const FHairStrandsRestRootResource::FLOD& Sim_RestLODDatas = SimRestRootResources->LODs[MeshLODIndex];
-		const FHairStrandsRestRootResource::FLOD& Ren_RestLODDatas = RenRestRootResources->LODs[MeshLODIndex];
-		const FHairStrandsDeformedRootResource::FLOD& Sim_DeformedLODDatas = SimDeformedRootResources->LODs[MeshLODIndex];
-		const FHairStrandsDeformedRootResource::FLOD& Ren_DeformedLODDatas = RenDeformedRootResources->LODs[MeshLODIndex];
-
+		const FHairStrandsDeformedRootResource::FLOD& Sim_DeformedLODDatas = SimDeformedRootResources->LODs[MeshLODIndex];\
 		bSupportGlobalInterpolation = Instance->Guides.bHasGlobalInterpolation && (Sim_RestLODDatas.SampleCount > 0);
 		{
-			Parameters->RestPosition0Buffer = RegisterAsSRV(GraphBuilder, Ren_RestLODDatas.RestRootTrianglePosition0Buffer);
-			Parameters->RestPosition1Buffer = RegisterAsSRV(GraphBuilder, Ren_RestLODDatas.RestRootTrianglePosition1Buffer);
-			Parameters->RestPosition2Buffer = RegisterAsSRV(GraphBuilder, Ren_RestLODDatas.RestRootTrianglePosition2Buffer);
-
-			Parameters->RootBarycentricBuffer = RegisterAsSRV(GraphBuilder, Ren_RestLODDatas.RootTriangleBarycentricBuffer);
-
 			Parameters->SimRestPosition0Buffer = RegisterAsSRV(GraphBuilder, Sim_RestLODDatas.RestRootTrianglePosition0Buffer);
 			Parameters->SimRestPosition1Buffer = RegisterAsSRV(GraphBuilder, Sim_RestLODDatas.RestRootTrianglePosition1Buffer);
 			Parameters->SimRestPosition2Buffer = RegisterAsSRV(GraphBuilder, Sim_RestLODDatas.RestRootTrianglePosition2Buffer);
@@ -621,31 +653,61 @@ static void AddHairStrandsInterpolationPass(
 			Parameters->SimRootBarycentricBuffer = RegisterAsSRV(GraphBuilder, Sim_RestLODDatas.RootTriangleBarycentricBuffer);
 		}
 		{
-			Parameters->DeformedPosition0Buffer = RegisterAsSRV(GraphBuilder, Ren_DeformedLODDatas.DeformedRootTrianglePosition0Buffer);
-			Parameters->DeformedPosition1Buffer = RegisterAsSRV(GraphBuilder, Ren_DeformedLODDatas.DeformedRootTrianglePosition1Buffer);
-			Parameters->DeformedPosition2Buffer = RegisterAsSRV(GraphBuilder, Ren_DeformedLODDatas.DeformedRootTrianglePosition2Buffer);
-
 			Parameters->SimDeformedPosition0Buffer = RegisterAsSRV(GraphBuilder, Sim_DeformedLODDatas.DeformedRootTrianglePosition0Buffer);
 			Parameters->SimDeformedPosition1Buffer = RegisterAsSRV(GraphBuilder, Sim_DeformedLODDatas.DeformedRootTrianglePosition1Buffer);
 			Parameters->SimDeformedPosition2Buffer = RegisterAsSRV(GraphBuilder, Sim_DeformedLODDatas.DeformedRootTrianglePosition2Buffer);
 		}
 	}
 
+	// Strands
+	if (bSupportDynamicMesh)
+	{
+		const FHairStrandsRestRootResource::FLOD& Ren_RestLODDatas = RenRestRootResources->LODs[MeshLODIndex];
+		const FHairStrandsDeformedRootResource::FLOD& Ren_DeformedLODDatas = RenDeformedRootResources->LODs[MeshLODIndex];
+		{
+			Parameters->RestPosition0Buffer = RegisterAsSRV(GraphBuilder, Ren_RestLODDatas.RestRootTrianglePosition0Buffer);
+			Parameters->RestPosition1Buffer = RegisterAsSRV(GraphBuilder, Ren_RestLODDatas.RestRootTrianglePosition1Buffer);
+			Parameters->RestPosition2Buffer = RegisterAsSRV(GraphBuilder, Ren_RestLODDatas.RestRootTrianglePosition2Buffer);
+			Parameters->RootBarycentricBuffer = RegisterAsSRV(GraphBuilder, Ren_RestLODDatas.RootTriangleBarycentricBuffer);
+		}
+		{
+			Parameters->DeformedPosition0Buffer = RegisterAsSRV(GraphBuilder, Ren_DeformedLODDatas.DeformedRootTrianglePosition0Buffer);
+			Parameters->DeformedPosition1Buffer = RegisterAsSRV(GraphBuilder, Ren_DeformedLODDatas.DeformedRootTrianglePosition1Buffer);
+			Parameters->DeformedPosition2Buffer = RegisterAsSRV(GraphBuilder, Ren_DeformedLODDatas.DeformedRootTrianglePosition2Buffer);
+		}
+	}
+
 	if (ShaderDrawDebug::IsShaderDrawDebugEnabled() && ShaderDrawData)
 	{
 		ShaderDrawDebug::SetParameters(GraphBuilder, *ShaderDrawData, Parameters->ShaderDrawParameters);
 	}
 
-	const bool bUseSingleGuide = GHairStrandsUseSingleGuideInterpolation > 0;
-	const bool bHasLocalDeformation = (Instance->Guides.bIsSimulationEnable && ((MeshLODIndex < 0) || ((MeshLODIndex >= 0) && IsHairLODSimulationEnabled(MeshLODIndex)))) || bSupportGlobalInterpolation;
-	const bool bCullingEnable = InstanceGeometryType == EHairGeometryType::Strands && CullingData.bCullingResultAvailable;
+	const bool bUseSingleGuide		= GHairStrandsUseSingleGuideInterpolation > 0;
+	const bool bHasLocalDeformation = Instance->Guides.bIsSimulationEnable || bSupportGlobalInterpolation;
+	const bool bCullingEnable		= InstanceGeometryType == EHairGeometryType::Strands && CullingData.bCullingResultAvailable;
 	Parameters->HairStrandsVF_bIsCullingEnable = bCullingEnable ? 1 : 0;
 
+	// Select dynamic geometry permutation, based on the Simulation/RBF/InterpolationType
+	int32 DynamicGeometryType = 0;
+	{
+		switch (HairInterpolationType)
+		{
+		case EHairInterpolationType::NoneSkinning			: DynamicGeometryType = 0; break; // INTERPOLATION_RIGID
+		case EHairInterpolationType::RigidSkinning			: DynamicGeometryType = 1; break; // INTERPOLATION_SKINNING_OFFSET
+		case EHairInterpolationType::OffsetSkinning			: DynamicGeometryType = 2; break; // INTERPOLATION_SKINNING_TRANSLATION
+		case EHairInterpolationType::SmoothSkinning			: DynamicGeometryType = 3; break; // INTERPOLATION_SKINNING_ROTATION
+		//case EHairInterpolationType::SmoothOffsetSkinning	: DynamicGeometryType = 4; break; // INTERPOLATION_SKINNING_TRANSLATION_AND_ROTATION - Experimental, not used
+		}		
+
+		if ( bSupportDynamicMesh && !bHasLocalDeformation) DynamicGeometryType = 1; // INTERPOLATION_SKINNING_OFFSET
+		if (!bSupportDynamicMesh && !bHasLocalDeformation) DynamicGeometryType = 0; // INTERPOLATION_RIGID
+		if (!bSupportDynamicMesh &&  bHasLocalDeformation) DynamicGeometryType = 0; // INTERPOLATION_RIGID
+	}
+
 	FHairInterpolationCS::FPermutationDomain PermutationVector;
 	PermutationVector.Set<FHairInterpolationCS::FGroupSize>(GroupSize);
 	PermutationVector.Set<FHairInterpolationCS::FDebug>(Parameters->HairDebugMode > 0 ? 1 : 0);
-	PermutationVector.Set<FHairInterpolationCS::FDynamicGeometry>((bSupportDynamicMesh && bHasLocalDeformation) ? HairInterpolationType+1 :
-							(bSupportDynamicMesh && !bHasLocalDeformation) ? 1 : 0);
+	PermutationVector.Set<FHairInterpolationCS::FDynamicGeometry>(DynamicGeometryType);
 	PermutationVector.Set<FHairInterpolationCS::FSimulation>(bHasLocalDeformation ? 1 : 0);
 	PermutationVector.Set<FHairInterpolationCS::FSingleGuide>(bUseSingleGuide ? 1 : 0);
 	PermutationVector.Set<FHairInterpolationCS::FCulling>(bCullingEnable ? 1 : 0);
@@ -1430,6 +1492,7 @@ void ComputeHairStrandsInterpolation(
 		Instance->HairGroupPublicData->VFInput = ComputeHairStrandsVertexInputData(Instance);
 		{
 			{
+				const bool bValidGuide = Instance->Guides.bIsSimulationEnable || Instance->Guides.bHasGlobalInterpolation;
 				const bool bHasSkinning = Instance->BindingType == EHairBindingType::Skinning;
 
 				AddHairStrandsInterpolationPass(
@@ -1445,20 +1508,20 @@ void ComputeHairStrandsInterpolation(
 					InstanceGeometryType,
 					CullingData,
 					Instance->Strands.RestResource->PositionOffset,
-					Instance->Guides.RestResource->PositionOffset,
+					bValidGuide ? Instance->Guides.RestResource->PositionOffset : FVector::ZeroVector,
 					RegisterAsSRV(GraphBuilder, Instance->Strands.DeformedResource->GetPositionOffsetBuffer(FHairStrandsDeformedResource::Current)),
-					RegisterAsSRV(GraphBuilder, Instance->Guides.DeformedResource->GetPositionOffsetBuffer(FHairStrandsDeformedResource::Current)),
+					bValidGuide ? RegisterAsSRV(GraphBuilder, Instance->Guides.DeformedResource->GetPositionOffsetBuffer(FHairStrandsDeformedResource::Current)) : nullptr,
 					bHasSkinning ? Instance->Strands.RestRootResource : nullptr,
-					bHasSkinning ? Instance->Guides.RestRootResource : nullptr,
+					bHasSkinning && bValidGuide ? Instance->Guides.RestRootResource : nullptr,
 					bHasSkinning ? Instance->Strands.DeformedRootResource : nullptr,
-					bHasSkinning ? Instance->Guides.DeformedRootResource : nullptr,
+					bHasSkinning && bValidGuide  ? Instance->Guides.DeformedRootResource : nullptr,
 					RegisterAsSRV(GraphBuilder, Instance->Strands.RestResource->RestPositionBuffer),
 					RegisterAsSRV(GraphBuilder, Instance->Strands.RestResource->AttributeBuffer),
 					RegisterAsSRV(GraphBuilder, Instance->Strands.InterpolationResource->Interpolation0Buffer),
 					RegisterAsSRV(GraphBuilder, Instance->Strands.InterpolationResource->Interpolation1Buffer),
-					RegisterAsSRV(GraphBuilder, Instance->Guides.RestResource->RestPositionBuffer),
-					RegisterAsSRV(GraphBuilder, Instance->Guides.DeformedResource->GetBuffer(FHairStrandsDeformedResource::Current)),
-					RegisterAsSRV(GraphBuilder, Instance->Guides.RestResource->AttributeBuffer),
+					bValidGuide ? RegisterAsSRV(GraphBuilder, Instance->Guides.RestResource->RestPositionBuffer) : nullptr,
+					bValidGuide ? RegisterAsSRV(GraphBuilder, Instance->Guides.DeformedResource->GetBuffer(FHairStrandsDeformedResource::Current)) : nullptr,
+					bValidGuide ? RegisterAsSRV(GraphBuilder, Instance->Guides.RestResource->AttributeBuffer) : nullptr,
 					Strands_DeformedPosition,
 					&Strands_DeformedPrevPosition, 
 					Strands_DebugAttribute,
@@ -1542,8 +1605,15 @@ void ComputeHairStrandsInterpolation(
 		const bool bIsCardsValid = Instance->Cards.IsValid(HairLODIndex);
 		if (bIsCardsValid)
 		{
+			const bool bValidGuide = Instance->Guides.bIsSimulationEnable || Instance->Guides.bHasGlobalInterpolation;
+			const bool bHasSkinning = Instance->BindingType == EHairBindingType::Skinning;
+
 			FHairGroupInstance::FCards::FLOD& LOD = Instance->Cards.LODs[HairLODIndex];
-			const EHairCardsSimulationType CardsSimulationType = GetHairCardsSimulationType();
+			const EHairCardsSimulationType CardsSimulationType = 
+				bHasSkinning || bValidGuide ?
+				GetHairCardsSimulationType() :
+				EHairCardsSimulationType::None;
+
 			if (LOD.Guides.Data && CardsSimulationType == EHairCardsSimulationType::Guide)
 			{
 				FRDGImportedBuffer Guides_DeformedPositionBuffer = Register(GraphBuilder, LOD.Guides.DeformedResource->GetBuffer(FHairStrandsDeformedResource::Current), ERDGImportedBufferFlags::CreateViews);
@@ -1558,8 +1628,6 @@ void ComputeHairStrandsInterpolation(
 				ScaleAndClipDesc.HairRadiusTipScale = 1;
 				ScaleAndClipDesc.HairLengthClip = 1;
 
-				const bool bHasSkinning = Instance->BindingType == EHairBindingType::Skinning;
-
 				FRDGHairStrandsCullingData CullingData;
 				AddHairStrandsInterpolationPass(
 					GraphBuilder,
@@ -1574,20 +1642,20 @@ void ComputeHairStrandsInterpolation(
 					InstanceGeometryType,
 					CullingData,
 					LOD.Guides.RestResource->PositionOffset,
-					Instance->Guides.RestResource->PositionOffset,
+					bValidGuide ? Instance->Guides.RestResource->PositionOffset : FVector::ZeroVector,
 					RegisterAsSRV(GraphBuilder, LOD.Guides.DeformedResource->GetPositionOffsetBuffer(FHairStrandsDeformedResource::Current)),
-					RegisterAsSRV(GraphBuilder, Instance->Guides.DeformedResource->GetPositionOffsetBuffer(FHairStrandsDeformedResource::Current)),
+					bValidGuide ? RegisterAsSRV(GraphBuilder, Instance->Guides.DeformedResource->GetPositionOffsetBuffer(FHairStrandsDeformedResource::Current)) : nullptr,
 					bHasSkinning ? LOD.Guides.RestRootResource : nullptr ,
-					bHasSkinning ? Instance->Guides.RestRootResource : nullptr,
+					bHasSkinning && bValidGuide ? Instance->Guides.RestRootResource : nullptr,
 					bHasSkinning ? LOD.Guides.DeformedRootResource : nullptr,
-					bHasSkinning ? Instance->Guides.DeformedRootResource : nullptr,
+					bHasSkinning && bValidGuide ? Instance->Guides.DeformedRootResource : nullptr,
 					RegisterAsSRV(GraphBuilder, LOD.Guides.RestResource->RestPositionBuffer),
 					RegisterAsSRV(GraphBuilder, LOD.Guides.RestResource->AttributeBuffer),
 					RegisterAsSRV(GraphBuilder, LOD.Guides.InterpolationResource->Interpolation0Buffer),
 					RegisterAsSRV(GraphBuilder, LOD.Guides.InterpolationResource->Interpolation1Buffer),
-					RegisterAsSRV(GraphBuilder, Instance->Guides.RestResource->RestPositionBuffer),
-					RegisterAsSRV(GraphBuilder, Instance->Guides.DeformedResource->GetBuffer(FHairStrandsDeformedResource::Current)),
-					RegisterAsSRV(GraphBuilder, Instance->Guides.RestResource->AttributeBuffer),
+					bValidGuide ? RegisterAsSRV(GraphBuilder, Instance->Guides.RestResource->RestPositionBuffer) : nullptr,
+					bValidGuide ? RegisterAsSRV(GraphBuilder, Instance->Guides.DeformedResource->GetBuffer(FHairStrandsDeformedResource::Current)) : nullptr,
+					bValidGuide ? RegisterAsSRV(GraphBuilder, Instance->Guides.RestResource->AttributeBuffer) : nullptr,
 					Guides_DeformedPositionBuffer,
 					nullptr,
 					nullptr,
@@ -1612,6 +1680,12 @@ void ComputeHairStrandsInterpolation(
 					Instance->Guides.RestRootResource,
 					Instance->Guides.DeformedRootResource);
 			}
+			else
+			{
+				// TODO remove this copy. For being able to do that we need to build a VF permutation which take the rest position as input
+				FRDGBufferRef DeformedResource = Register(GraphBuilder, LOD.DeformedResource->GetBuffer(FHairCardsDeformedResource::Current), ERDGImportedBufferFlags::None).Buffer;
+				AddCopyBufferPass(GraphBuilder, DeformedResource, LOD.RestResource->RestPositionBuffer.VertexBufferRHI);
+			}
 
 			#if RHI_RAYTRACING
 			if (LOD.RaytracingResource)
@@ -1692,22 +1766,16 @@ void ComputeHairStrandsInterpolation(
 	Instance->HairGroupPublicData->bSupportVoxelization = Instance->Strands.Modifier.bSupportVoxelization && Instance->bCastShadow;
 }
 
-bool HasHairInstanceSimulationEnable(FHairGroupInstance* Instance, int32 MeshLODIndex)
-{
-	const bool bHasNoSimulation = !Instance || (Instance && Instance->Guides.bIsSimulationEnable && IsHairLODSimulationEnabled(MeshLODIndex)) || !IsHairStrandsBindingEnable();
-	return !bHasNoSimulation;
-}
-
 void ResetHairStrandsInterpolation(
 	FRDGBuilder& GraphBuilder,
 	FGlobalShaderMap* ShaderMap,
 	FHairGroupInstance* Instance,
 	int32 MeshLODIndex)
 {
-	if (!HasHairInstanceSimulationEnable(Instance, MeshLODIndex))
-	{
-		return;
-	}
+	if (!Instance || 
+		(Instance && Instance->Guides.bIsSimulationEnable) ||
+		(!Instance->Guides.bHasGlobalInterpolation && !Instance->Guides.bIsSimulationEnable) ||
+		!IsHairStrandsBindingEnable()) return;
 
 	DECLARE_GPU_STAT(HairStrandsGuideDeform);
 	RDG_EVENT_SCOPE(GraphBuilder, "HairStrandsGuideDeform");
diff --git a/Engine/Plugins/Runtime/HairStrands/Source/HairStrandsCore/Private/Niagara/NiagaraDataInterfaceHairStrands.cpp b/Engine/Plugins/Runtime/HairStrands/Source/HairStrandsCore/Private/Niagara/NiagaraDataInterfaceHairStrands.cpp
--- a/Engine/Plugins/Runtime/HairStrands/Source/HairStrandsCore/Private/Niagara/NiagaraDataInterfaceHairStrands.cpp
+++ b/Engine/Plugins/Runtime/HairStrands/Source/HairStrandsCore/Private/Niagara/NiagaraDataInterfaceHairStrands.cpp
@@ -406,8 +406,8 @@ void FNDIHairStrandsData::Release()
 	}
 }
 
-void FNDIHairStrandsData::Update(UNiagaraDataInterfaceHairStrands* Interface, FNiagaraSystemInstance* SystemInstance, const FHairStrandsDatas* StrandsDatas, 
-	UGroomAsset* GroomAsset, const int32 GroupIndex)
+void FNDIHairStrandsData::Update(UNiagaraDataInterfaceHairStrands* Interface, FNiagaraSystemInstance* SystemInstance, const FHairStrandsDatas* StrandsDatas,
+	UGroomAsset* GroomAsset, const int32 GroupIndex, const int32 LODIndex)
 {
 	if (Interface != nullptr)
 	{
@@ -416,11 +416,12 @@ void FNDIHairStrandsData::Update(UNiagaraDataInterfaceHairStrands* Interface, FN
 
 		GlobalInterpolation = (Interface->IsComponentValid() && Interface->SourceComponent->BindingAsset && Interface->SourceComponent->GroomAsset) ?
 			Interface->SourceComponent->GroomAsset->EnableGlobalInterpolation : false;
+		HairGroupInstance = Interface->SourceComponent->GetGroupInstance(GroupIndex);
 
 		TickingGroup = Interface->IsComponentValid() ? ComputeTickingGroup(Interface->SourceComponent) : NiagaraFirstTickGroup;
 
 		if (StrandsDatas != nullptr && GroomAsset != nullptr && GroupIndex >= 0 && GroupIndex < GroomAsset->HairGroupsPhysics.Num() &&
-			GroomAsset->HairGroupsPhysics[GroupIndex].SolverSettings.EnableSimulation)
+			GroomAsset->IsSimulationEnable(GroupIndex, LODIndex))
 		{
 			FHairGroupsPhysics& HairPhysics = GroomAsset->HairGroupsPhysics[GroupIndex];
 			StrandsSize = static_cast<uint8>(HairPhysics.StrandsParameters.StrandsSize);
@@ -487,10 +488,11 @@ bool FNDIHairStrandsData::Init(UNiagaraDataInterfaceHairStrands* Interface, FNia
 		FHairStrandsDeformedRootResource* StrandsDeformedRootResource = nullptr;
 		UGroomAsset* GroomAsset = nullptr;
 		int32 GroupIndex = 0;
+		int32 LODIndex = 0;
 
 		{
-			Interface->ExtractDatasAndResources(SystemInstance, StrandsDatas, StrandsRestResource, StrandsDeformedResource, StrandsRestRootResource, StrandsDeformedRootResource, GroomAsset, GroupIndex);
-			Update(Interface, SystemInstance, StrandsDatas, GroomAsset, GroupIndex);
+			Interface->ExtractDatasAndResources(SystemInstance, StrandsDatas, StrandsRestResource, StrandsDeformedResource, StrandsRestRootResource, StrandsDeformedRootResource, GroomAsset, GroupIndex, LODIndex);
+			Update(Interface, SystemInstance, StrandsDatas, GroomAsset, GroupIndex, LODIndex);
 
 			HairStrandsBuffer = new FNDIHairStrandsBuffer();
 			HairStrandsBuffer->Initialize(StrandsDatas, StrandsRestResource, StrandsDeformedResource, StrandsRestRootResource, StrandsDeformedRootResource, ParamsScale);
@@ -507,6 +509,13 @@ bool FNDIHairStrandsData::Init(UNiagaraDataInterfaceHairStrands* Interface, FNia
 
 //------------------------------------------------------------------------------------------------------------
 
+enum class EHairSimulationInterpolationMode : uint8
+{
+	Rigid = 0,
+	Skinned = 1,
+	RBF = 2
+};
+
 struct FNDIHairStrandsParametersCS : public FNiagaraDataInterfaceParametersCS
 {
 	DECLARE_TYPE_LAYOUT(FNDIHairStrandsParametersCS, NonVirtual);
@@ -552,7 +561,7 @@ struct FNDIHairStrandsParametersCS : public FNiagaraDataInterfaceParametersCS
 
 		
 	}
-
+	
 	void Set(FRHICommandList& RHICmdList, const FNiagaraDataInterfaceSetArgs& Context) const
 	{
 		check(IsInRenderingThread()); 
@@ -564,10 +573,14 @@ struct FNDIHairStrandsParametersCS : public FNiagaraDataInterfaceParametersCS
 		FNDIHairStrandsData* ProxyData =
 			InterfaceProxy->SystemInstancesToProxyData.Find(Context.SystemInstanceID);
 
+
 		const bool IsHairValid = ProxyData != nullptr && ProxyData->HairStrandsBuffer != nullptr && ProxyData->HairStrandsBuffer->IsInitialized();
-		const bool IsRootValid = IsHairValid && ProxyData->HairStrandsBuffer->SourceDeformedRootResources != nullptr;//&& ProxyData->HairStrandsBuffer->SourceRootResources->IsInitialized();
+		const bool bHasSkinningBinding = IsHairValid && ProxyData->HairGroupInstance && ProxyData->HairGroupInstance->BindingType == EHairBindingType::Skinning;
+		const bool IsRootValid = IsHairValid && ProxyData->HairStrandsBuffer->SourceDeformedRootResources != nullptr && bHasSkinningBinding;//&& ProxyData->HairStrandsBuffer->SourceRootResources->IsInitialized();
 		const bool IsRestValid = IsHairValid && ProxyData->HairStrandsBuffer->SourceRestResources != nullptr;//&& ProxyData->HairStrandsBuffer->SourceRestResources->IsInitialized();
 		const bool IsDeformedValid = IsHairValid && ProxyData->HairStrandsBuffer->SourceDeformedResources != nullptr && ProxyData->HairStrandsBuffer->SourceDeformedResources->IsInitialized();
+		
+		const bool bHasLODSwitched = ProxyData->HairGroupInstance && ProxyData->HairGroupInstance->HairGroupPublicData && ProxyData->HairGroupInstance->HairGroupPublicData->VFInput.bHasLODSwitch;
 
 		if (IsHairValid && IsRestValid)
 		{
@@ -577,23 +590,27 @@ struct FNDIHairStrandsParametersCS : public FNiagaraDataInterfaceParametersCS
 				HairStrandsBuffer->SourceDeformedResources->DeformedPositionBuffer[HairStrandsBuffer->SourceDeformedResources->CurrentIndex].UAV : HairStrandsBuffer->DeformedPositionBuffer.UAV;
 
 			const int32 MeshLODIndex = IsRootValid ? HairStrandsBuffer->SourceDeformedRootResources->MeshLODIndex : -1;
-			int32 InterpolationModeValue = IsRootValid && HairStrandsBuffer->SourceDeformedRootResources->IsValid(MeshLODIndex) ? 1 : 0;
+			EHairSimulationInterpolationMode InterpolationModeValue = EHairSimulationInterpolationMode::Rigid;
+			if (IsRootValid && HairStrandsBuffer->SourceDeformedRootResources->IsValid(MeshLODIndex))
+			{
+				InterpolationModeValue = EHairSimulationInterpolationMode::Skinned;
+			}
 
-			const FHairStrandsRestRootResource::FLOD* RestMeshProjection = (InterpolationModeValue == 1) ? &(HairStrandsBuffer->SourceRestRootResources->LODs[MeshLODIndex]) : nullptr;
-			const FHairStrandsDeformedRootResource::FLOD* DeformedMeshProjection = (InterpolationModeValue == 1) ? &(HairStrandsBuffer->SourceDeformedRootResources->LODs[MeshLODIndex]) : nullptr;
+			const FHairStrandsRestRootResource::FLOD* RestMeshProjection = (InterpolationModeValue == EHairSimulationInterpolationMode::Skinned) ? &(HairStrandsBuffer->SourceRestRootResources->LODs[MeshLODIndex]) : nullptr;
+			const FHairStrandsDeformedRootResource::FLOD* DeformedMeshProjection = (InterpolationModeValue == EHairSimulationInterpolationMode::Skinned) ? &(HairStrandsBuffer->SourceDeformedRootResources->LODs[MeshLODIndex]) : nullptr;
 
-			FRHIShaderResourceView* RestTrianglePositionASRV = (InterpolationModeValue == 1 && RestMeshProjection != nullptr) ?
+			FRHIShaderResourceView* RestTrianglePositionASRV = (InterpolationModeValue == EHairSimulationInterpolationMode::Skinned && RestMeshProjection != nullptr) ?
 				RestMeshProjection->RestRootTrianglePosition0Buffer.SRV.GetReference() : FNiagaraRenderer::GetDummyFloatBuffer();
-			FRHIShaderResourceView* RestTrianglePositionBSRV = (InterpolationModeValue == 1 && RestMeshProjection != nullptr) ?
+			FRHIShaderResourceView* RestTrianglePositionBSRV = (InterpolationModeValue == EHairSimulationInterpolationMode::Skinned && RestMeshProjection != nullptr) ?
 				RestMeshProjection->RestRootTrianglePosition1Buffer.SRV.GetReference() : FNiagaraRenderer::GetDummyFloatBuffer();
-			FRHIShaderResourceView* RestTrianglePositionCSRV = (InterpolationModeValue == 1 && RestMeshProjection != nullptr) ?
+			FRHIShaderResourceView* RestTrianglePositionCSRV = (InterpolationModeValue == EHairSimulationInterpolationMode::Skinned && RestMeshProjection != nullptr) ?
 				RestMeshProjection->RestRootTrianglePosition2Buffer.SRV.GetReference() : FNiagaraRenderer::GetDummyFloatBuffer();
 
-			FRHIShaderResourceView* DeformedTrianglePositionASRV = (InterpolationModeValue == 1 && DeformedMeshProjection != nullptr) ?
+			FRHIShaderResourceView* DeformedTrianglePositionASRV = (InterpolationModeValue == EHairSimulationInterpolationMode::Skinned && DeformedMeshProjection != nullptr) ?
 				DeformedMeshProjection->DeformedRootTrianglePosition0Buffer.SRV.GetReference() : FNiagaraRenderer::GetDummyFloatBuffer();
-			FRHIShaderResourceView* DeformedTrianglePositionBSRV = (InterpolationModeValue == 1 && DeformedMeshProjection != nullptr) ?
+			FRHIShaderResourceView* DeformedTrianglePositionBSRV = (InterpolationModeValue == EHairSimulationInterpolationMode::Skinned && DeformedMeshProjection != nullptr) ?
 				DeformedMeshProjection->DeformedRootTrianglePosition1Buffer.SRV.GetReference() : FNiagaraRenderer::GetDummyFloatBuffer();
-			FRHIShaderResourceView* DeformedTrianglePositionCSRV = (InterpolationModeValue == 1 && DeformedMeshProjection != nullptr) ?
+			FRHIShaderResourceView* DeformedTrianglePositionCSRV = (InterpolationModeValue == EHairSimulationInterpolationMode::Skinned && DeformedMeshProjection != nullptr) ?
 				DeformedMeshProjection->DeformedRootTrianglePosition2Buffer.SRV.GetReference() : FNiagaraRenderer::GetDummyFloatBuffer();
 
 			const bool bHasSamples = (RestMeshProjection != nullptr && RestMeshProjection->SampleCount > 0);
@@ -603,14 +620,14 @@ struct FNDIHairStrandsParametersCS : public FNiagaraDataInterfaceParametersCS
 			FShaderResourceViewRHIRef MeshSampleWeightsBufferSRV = bHasSamples && DeformedMeshProjection != nullptr ?
 				DeformedMeshProjection->MeshSampleWeightsBuffer.SRV.GetReference() : FNiagaraRenderer::GetDummyFloatBuffer();
 
-			FRHIShaderResourceView* RootBarycentricCoordinatesSRV = (InterpolationModeValue == 1 && RestMeshProjection != nullptr) ?
+			FRHIShaderResourceView* RootBarycentricCoordinatesSRV = (InterpolationModeValue == EHairSimulationInterpolationMode::Skinned && RestMeshProjection != nullptr) ?
 				RestMeshProjection->RootTriangleBarycentricBuffer.SRV.GetReference() : FNiagaraRenderer::GetDummyFloatBuffer();
 
-			int32 bNeedSimReset = (ProxyData->TickCount <= MaxDelay ? 1 : 0);
+			int32 bNeedSimReset = bHasLODSwitched || (ProxyData->TickCount <= MaxDelay ? 1 : 0);
 
-			if ((InterpolationModeValue == 1) && ProxyData->GlobalInterpolation)
+			if (InterpolationModeValue == EHairSimulationInterpolationMode::Skinned && ProxyData->GlobalInterpolation)
 			{
-				InterpolationModeValue = 2;
+				InterpolationModeValue = EHairSimulationInterpolationMode::RBF;
 			}
 
 			//FUnorderedAccessViewRHIRef DummyFloatUAV(Context.Batcher->GetEmptyRWBufferFromPool(RHICmdList, PF_R32_FLOAT), false);
@@ -645,7 +662,7 @@ struct FNDIHairStrandsParametersCS : public FNiagaraDataInterfaceParametersCS
 			SetShaderValue(RHICmdList, ComputeShaderRHI, StrandSize, ProxyData->StrandsSize);
 
 			SetShaderValue(RHICmdList, ComputeShaderRHI, ResetSimulation, bNeedSimReset);
-			SetShaderValue(RHICmdList, ComputeShaderRHI, InterpolationMode, InterpolationModeValue);
+			SetShaderValue(RHICmdList, ComputeShaderRHI, InterpolationMode, int32(InterpolationModeValue));
 			SetShaderValue(RHICmdList, ComputeShaderRHI, RestRootOffset, RestRootOffsetValue);
 			SetShaderValue(RHICmdList, ComputeShaderRHI, DeformedRootOffset, DeformedRootOffsetValue);
 
@@ -872,7 +889,8 @@ void UNiagaraDataInterfaceHairStrands::ExtractDatasAndResources(
 	FHairStrandsRestRootResource*& OutStrandsRestRootResource, 
 	FHairStrandsDeformedRootResource*& OutStrandsDeformedRootResource,
 	UGroomAsset*& OutGroomAsset,
-	int32& OutGroupIndex)
+	int32& OutGroupIndex,
+	int32& OutLODIndex)
 {
 	ExtractSourceComponent(SystemInstance);
 
@@ -882,6 +900,7 @@ void UNiagaraDataInterfaceHairStrands::ExtractDatasAndResources(
 	OutStrandsRestRootResource = nullptr;
 	OutStrandsDeformedRootResource = nullptr;
 	OutGroupIndex = -1;
+	OutLODIndex = -1;
 
 	if (IsComponentValid() && SystemInstance)
 	{
@@ -904,11 +923,13 @@ void UNiagaraDataInterfaceHairStrands::ExtractDatasAndResources(
 			OutStrandsRestRootResource = SourceComponent->GetGuideStrandsRestRootResource(OutGroupIndex);
 			OutStrandsDeformedRootResource = SourceComponent->GetGuideStrandsDeformedRootResource(OutGroupIndex);
 			OutGroomAsset = SourceComponent->GroomAsset;
+			OutLODIndex = SourceComponent->GetForcedLOD();
 		}
 	}
 	else if (DefaultSource != nullptr)
 	{
 		OutGroupIndex = 0;
+		OutLODIndex = 0;
 		if (OutGroupIndex < DefaultSource->GetNumHairGroups())
 		{
 			OutStrandsDatas = &DefaultSource->HairGroupsData[OutGroupIndex].Strands.Data;
@@ -964,10 +985,11 @@ bool UNiagaraDataInterfaceHairStrands::PerInstanceTick(void* PerInstanceData, FN
 	FHairStrandsDeformedRootResource* StrandsDeformedRootResource = nullptr;
 	UGroomAsset* GroomAsset = nullptr;
 	int32 GroupIndex = 0;
+	int32 LODIndex = 0;
 
 	InstanceData->TickCount = FMath::Min(MaxDelay+1,InstanceData->TickCount+1);
 
-	ExtractDatasAndResources(SystemInstance, StrandsDatas, StrandsRestResource, StrandsDeformedResource, StrandsRestRootResource, StrandsDeformedRootResource, GroomAsset, GroupIndex);
+	ExtractDatasAndResources(SystemInstance, StrandsDatas, StrandsRestResource, StrandsDeformedResource, StrandsRestRootResource, StrandsDeformedRootResource, GroomAsset, GroupIndex, LODIndex);
 	InstanceData->HairStrandsBuffer->Update(StrandsDatas, StrandsRestResource, StrandsDeformedResource, StrandsRestRootResource, StrandsDeformedRootResource);
 
 	if (SourceComponent != nullptr)
@@ -978,7 +1000,7 @@ bool UNiagaraDataInterfaceHairStrands::PerInstanceTick(void* PerInstanceData, FN
 		}
 		InstanceData->ForceReset = SourceComponent->bResetSimulation;
 	}
-	InstanceData->Update(this, SystemInstance, StrandsDatas, GroomAsset, GroupIndex);
+	InstanceData->Update(this, SystemInstance, StrandsDatas, GroomAsset, GroupIndex, LODIndex);
 	return false;
 }
 
diff --git a/Engine/Plugins/Runtime/HairStrands/Source/HairStrandsCore/Public/GroomAsset.h b/Engine/Plugins/Runtime/HairStrands/Source/HairStrandsCore/Public/GroomAsset.h
--- a/Engine/Plugins/Runtime/HairStrands/Source/HairStrandsCore/Public/GroomAsset.h
+++ b/Engine/Plugins/Runtime/HairStrands/Source/HairStrandsCore/Public/GroomAsset.h
@@ -403,6 +403,9 @@ public:
 	virtual const TArray<UAssetUserData*>* GetAssetUserDataArray() const override;
 	//~ End IInterface_AssetUserData Interface
 
+	bool IsSimulationEnable(int32 GroupIndex, int32 LODIndex);
+	bool IsGlobalInterpolationEnable(int32 GroupIndex, int32 LODIndex);
+
 #if WITH_EDITORONLY_DATA
 	void SetDebugMode(EHairStrandsDebugMode InMode) { DebugMode = InMode; OnGroomAssetChanged.Broadcast(); }
 	EHairStrandsDebugMode GetDebugMode() const { return DebugMode; }
diff --git a/Engine/Plugins/Runtime/HairStrands/Source/HairStrandsCore/Public/GroomAssetInterpolation.h b/Engine/Plugins/Runtime/HairStrands/Source/HairStrandsCore/Public/GroomAssetInterpolation.h
--- a/Engine/Plugins/Runtime/HairStrands/Source/HairStrandsCore/Public/GroomAssetInterpolation.h
+++ b/Engine/Plugins/Runtime/HairStrands/Source/HairStrandsCore/Public/GroomAssetInterpolation.h
@@ -48,6 +48,13 @@ enum class EGroomBindingType : uint8
 	Skinning	UMETA(DisplayName = "Skinning", ToolTip = "When attached to a skeltal mesh, the hair follow the skin surface"),
 };
 
+UENUM(BlueprintType)
+enum class EGroomOverrideType : uint8
+{
+	Auto	UMETA(DisplayName = "Auto", ToolTip = "Use the asset value"),
+	Enable	UMETA(DisplayName = "Enable", ToolTip = "Override the asset value, and force enabled"),
+	Disable UMETA(DisplayName = "Disable", ToolTip = "Override the asset value, and force disabled")
+};
 
 USTRUCT(BlueprintType)
 struct HAIRSTRANDSCORE_API FHairLODSettings
@@ -84,6 +91,12 @@ struct HAIRSTRANDSCORE_API FHairLODSettings
 
 	UPROPERTY(EditAnywhere, Category = "DecimationSettings", meta = (ToolTip = "If enable this LOD version will use the provided attachment points"))
 	EGroomBindingType BindingType = EGroomBindingType::Skinning;
+	
+	UPROPERTY(EditAnywhere, Category = "DecimationSettings", meta = (ToolTip = "Groom simulation"))
+	EGroomOverrideType Simulation = EGroomOverrideType::Auto;
+
+	UPROPERTY(EditAnywhere, Category = "DecimationSettings", meta = (DisplayName = "RBF Interpolation", ToolTip = "Global interpolation"))
+	EGroomOverrideType GlobalInterpolation = EGroomOverrideType::Auto;
 
 	bool operator==(const FHairLODSettings& A) const;
 };
diff --git a/Engine/Plugins/Runtime/HairStrands/Source/HairStrandsCore/Public/GroomInstance.h b/Engine/Plugins/Runtime/HairStrands/Source/HairStrandsCore/Public/GroomInstance.h
--- a/Engine/Plugins/Runtime/HairStrands/Source/HairStrandsCore/Public/GroomInstance.h
+++ b/Engine/Plugins/Runtime/HairStrands/Source/HairStrandsCore/Public/GroomInstance.h
@@ -68,7 +68,7 @@ struct HAIRSTRANDSCORE_API FHairGroupInstance
 		FHairStrandsInterpolationDatas* InterpolationData = nullptr;
 		FHairStrandsInterpolationResource* InterpolationResource = nullptr;
 
-		uint32 HairInterpolationType = 0;
+		EHairInterpolationType HairInterpolationType = EHairInterpolationType::NoneSkinning;
 	};
 
 	//////////////////////////////////////////////////////////////////////////////////////////
diff --git a/Engine/Plugins/Runtime/HairStrands/Source/HairStrandsCore/Public/Niagara/NiagaraDataInterfaceHairStrands.h b/Engine/Plugins/Runtime/HairStrands/Source/HairStrandsCore/Public/Niagara/NiagaraDataInterfaceHairStrands.h
--- a/Engine/Plugins/Runtime/HairStrands/Source/HairStrandsCore/Public/Niagara/NiagaraDataInterfaceHairStrands.h
+++ b/Engine/Plugins/Runtime/HairStrands/Source/HairStrandsCore/Public/Niagara/NiagaraDataInterfaceHairStrands.h
@@ -97,12 +97,13 @@ struct FNDIHairStrandsData
 	void Release();
 
 	/** Update the buffers */
-	void Update(UNiagaraDataInterfaceHairStrands* Interface, FNiagaraSystemInstance* SystemInstance, const FHairStrandsDatas* HairStrandsDatas, UGroomAsset* GroomAsset, const int32 GroupIndex);
+	void Update(UNiagaraDataInterfaceHairStrands* Interface, FNiagaraSystemInstance* SystemInstance, const FHairStrandsDatas* HairStrandsDatas, UGroomAsset* GroomAsset, const int32 GroupIndex, const int32 LODIndex);
 
 	inline void ResetDatas()
 	{
 		WorldTransform.SetIdentity();
 		GlobalInterpolation = false;
+		HairGroupInstance = nullptr;
 
 		TickCount = 0;
 		ForceReset = true;
@@ -156,6 +157,8 @@ struct FNDIHairStrandsData
 			WorldTransform = OtherDatas->WorldTransform;
 
 			GlobalInterpolation = OtherDatas->GlobalInterpolation;
+			BindingType = OtherDatas->BindingType;
+			HairGroupInstance = OtherDatas->HairGroupInstance;
 
 			TickCount = OtherDatas->TickCount;
 			ForceReset = OtherDatas->ForceReset;
@@ -218,6 +221,12 @@ struct FNDIHairStrandsData
 
 	/** Strands Gpu buffer */
 	FNDIHairStrandsBuffer* HairStrandsBuffer;
+
+	/** Hair group instance */
+	FHairGroupInstance* HairGroupInstance;
+
+	/** Binding type between the groom asset and the attached skeletal mesh */
+	EHairBindingType BindingType;
 	
 	/** Number of substeps to be used */
 	int32 SubSteps;
@@ -353,7 +362,8 @@ public:
 		FHairStrandsRestRootResource*& OutStrandsRestRootResource, 
 		FHairStrandsDeformedRootResource*& OutStrandsDeformedRootResource,
 		UGroomAsset*& OutGroomAsset,
-		int32& OutGroupIndex);
+		int32& OutGroupIndex,
+		int32& OutLODIndex);
 
 	/** Get the number of strands */
 	void GetNumStrands(FVectorVMContext& Context);
diff --git a/Engine/Source/Runtime/Renderer/Public/HairStrandsInterface.h b/Engine/Source/Runtime/Renderer/Public/HairStrandsInterface.h
--- a/Engine/Source/Runtime/Renderer/Public/HairStrandsInterface.h
+++ b/Engine/Source/Runtime/Renderer/Public/HairStrandsInterface.h
@@ -135,6 +135,14 @@ enum EHairBindingType
 	Skinning
 };
 
+enum EHairInterpolationType
+{
+	NoneSkinning,				// Use no skinning data (i.e. with no binding)
+	RigidSkinning,				// Use skinning data & apply a rigid triangle deformation
+	OffsetSkinning,				// Use skinning data & apply a offset deformation
+	SmoothSkinning,				// Use skinning data & apply a smooth (rotation) offset
+};
+
 ////////////////////////////////////////////////////////////////////////////////////////////////////
 // Public group data 
 class RENDERER_API FHairGroupPublicData : public FRenderResource
@@ -177,12 +185,24 @@ public:
 	void SetLODVisibilities(const TArray<bool>& InLODVisibility) { LODVisibilities = InLODVisibility; }
 	const TArray<bool>& GetLODVisibilities() const { return LODVisibilities; }
 
-	EHairBindingType GetBindingType(int32 InLODIndex) 
+	EHairBindingType GetBindingType(int32 InLODIndex) const
 	{ 
 		if (InLODIndex < 0 && InLODIndex >= BindingTypes.Num()) return EHairBindingType::NoneBinding;
 		return BindingTypes[InLODIndex];
 	}	
 
+	bool IsSimulationEnable(int32 InLODIndex) const
+	{
+		if (InLODIndex < 0 && InLODIndex >= LODSimulations.Num()) return false;
+		return LODSimulations[InLODIndex];
+	}
+
+	bool IsGlobalInterpolationEnable(int32 InLODIndex) const 
+	{
+		if (InLODIndex < 0 && InLODIndex >= LODGlobalInterpolations.Num()) return false;
+		return LODGlobalInterpolations[InLODIndex];
+	}
+
 	void SetLODScreenSizes(const TArray<float>& ScreenSizes) { LODScreenSizes = ScreenSizes; }
 	const TArray<float>& GetLODScreenSizes() const { return LODScreenSizes;  }
 
@@ -263,7 +283,9 @@ public:
 	   CPU LOD selection is enabled otherwise the GPU selection is used. CPU LOD selection use the CPU 
 	   bounding box, which might not be as accurate as the GPU ones*/
 	TArray<bool> LODVisibilities;
-	TArray<float> LODScreenSizes;
+	TArray<float>LODScreenSizes;
+	TArray<bool> LODSimulations;
+	TArray<bool> LODGlobalInterpolations;
 	TArray<EHairGeometryType> LODGeometryTypes;
 
 	TArray<EHairBindingType> BindingTypes;
