Fix CIS issue.

diff --git a/Engine/Plugins/Runtime/HairStrands/Source/HairStrandsCore/Private/Niagara/NiagaraDataInterfaceHairStrands.cpp b/Engine/Plugins/Runtime/HairStrands/Source/HairStrandsCore/Private/Niagara/NiagaraDataInterfaceHairStrands.cpp
--- a/Engine/Plugins/Runtime/HairStrands/Source/HairStrandsCore/Private/Niagara/NiagaraDataInterfaceHairStrands.cpp
+++ b/Engine/Plugins/Runtime/HairStrands/Source/HairStrandsCore/Private/Niagara/NiagaraDataInterfaceHairStrands.cpp
@@ -580,10 +580,12 @@ struct FNDIHairStrandsParametersCS : public FNiagaraDataInterfaceParametersCS
 		const bool IsRestValid = IsHairValid && ProxyData->HairStrandsBuffer->SourceRestResources != nullptr;//&& ProxyData->HairStrandsBuffer->SourceRestResources->IsInitialized();
 		const bool IsDeformedValid = IsHairValid && ProxyData->HairStrandsBuffer->SourceDeformedResources != nullptr && ProxyData->HairStrandsBuffer->SourceDeformedResources->IsInitialized();
 		
-		const bool bHasLODSwitched = ProxyData->HairGroupInstance && ProxyData->HairGroupInstance->HairGroupPublicData && ProxyData->HairGroupInstance->HairGroupPublicData->VFInput.bHasLODSwitch;
+		const bool bHasLODSwitched = IsHairValid && ProxyData->HairGroupInstance && ProxyData->HairGroupInstance->HairGroupPublicData && ProxyData->HairGroupInstance->HairGroupPublicData->VFInput.bHasLODSwitch;
 
 		if (IsHairValid && IsRestValid)
 		{
+			check(ProxyData);
+
 			FNDIHairStrandsBuffer* HairStrandsBuffer = ProxyData->HairStrandsBuffer;
 
 			FUnorderedAccessViewRHIRef PointPositionsUAV = IsDeformedValid ?
