Fix groom crash when initialization hair simulation.

diff --git a/Engine/Plugins/Runtime/HairStrands/Source/HairStrandsCore/Private/GroomComponent.cpp b/Engine/Plugins/Runtime/HairStrands/Source/HairStrandsCore/Private/GroomComponent.cpp
--- a/Engine/Plugins/Runtime/HairStrands/Source/HairStrandsCore/Private/GroomComponent.cpp
+++ b/Engine/Plugins/Runtime/HairStrands/Source/HairStrandsCore/Private/GroomComponent.cpp
@@ -1045,7 +1045,7 @@ void UGroomComponent::ReleaseHairSimulation()
 
 static void InternalUpdateHairSimulation(UGroomComponent* GroomComponent)
 {
-	if (!GroomComponent)
+	if (!GroomComponent || GroomComponent->GroomGroupsDesc.Num() == 0)
 	{
 		return;
 	}
