Remove overlap event support for Groom component as this is not needed and expensive on CPU.

diff --git a/Engine/Plugins/Runtime/HairStrands/Source/HairStrandsCore/Private/GroomComponent.cpp b/Engine/Plugins/Runtime/HairStrands/Source/HairStrandsCore/Private/GroomComponent.cpp
--- a/Engine/Plugins/Runtime/HairStrands/Source/HairStrandsCore/Private/GroomComponent.cpp
+++ b/Engine/Plugins/Runtime/HairStrands/Source/HairStrandsCore/Private/GroomComponent.cpp
@@ -982,7 +982,9 @@ UGroomComponent::UGroomComponent(const FObjectInitializer& ObjectInitializer)
 	SourceSkeletalMesh = nullptr; 
 	NiagaraComponents.Empty();
 	PhysicsAsset = nullptr;
-
+	
+	// Overlap events are expensive and not needed (at least at the moment) as we don't need to collide against other component.
+	SetGenerateOverlapEvents(false);
 	SetCollisionProfileName(UCollisionProfile::PhysicsActor_ProfileName);
 
 	static ConstructorHelpers::FObjectFinder<UMaterialInterface> Strands_DebugMaterialRef(TEXT("/HairStrands/Materials/HairDebugMaterial.HairDebugMaterial"));
