Fix crash when binding a groom to an (temporarly) invalid skeletal mesh.

diff --git a/Engine/Plugins/Runtime/HairStrands/Source/HairStrandsCore/Private/GroomComponent.cpp b/Engine/Plugins/Runtime/HairStrands/Source/HairStrandsCore/Private/GroomComponent.cpp
--- a/Engine/Plugins/Runtime/HairStrands/Source/HairStrandsCore/Private/GroomComponent.cpp
+++ b/Engine/Plugins/Runtime/HairStrands/Source/HairStrandsCore/Private/GroomComponent.cpp
@@ -1856,6 +1856,10 @@ void UGroomComponent::InitResources(bool bIsBindingReloading)
 	if (bHasNeedBindingData && SkeletalMeshComponent)
 	{
 		SkeletalMeshComponent = ValidateBindingAsset(GroomAsset, BindingAsset, SkeletalMeshComponent, bIsBindingReloading, bValidationEnable, this);
+		if (SkeletalMeshComponent->SkeletalMesh == nullptr)
+		{
+			SkeletalMeshComponent = nullptr;
+		}
 	}
 	UGroomBindingAsset* LocalBindingAsset = bHasNeedBindingData ? BindingAsset : nullptr;
 
