Remove update of hair RT geometry (strands/cards/meshes) when the binding type is set to rigid.

diff --git a/Engine/Plugins/Runtime/HairStrands/Source/HairStrandsCore/Private/GroomComponent.cpp b/Engine/Plugins/Runtime/HairStrands/Source/HairStrandsCore/Private/GroomComponent.cpp
--- a/Engine/Plugins/Runtime/HairStrands/Source/HairStrandsCore/Private/GroomComponent.cpp
+++ b/Engine/Plugins/Runtime/HairStrands/Source/HairStrandsCore/Private/GroomComponent.cpp
@@ -605,6 +605,22 @@ public:
 		}
 	}
 
+	void GetOverrideLocalToTransform(const FHairGroupInstance* Instance, FMatrix& OutLocalToWorld) const
+	{
+		if (Instance->BindingType == EHairBindingType::Skinning)
+		{
+			OutLocalToWorld = Instance->Debug.SkeletalLocalToWorld.ToMatrixWithScale();
+		}
+	}
+
+	void GetOverridePreviousLocalToTransform(const FHairGroupInstance* Instance, FMatrix& OutPreviousLocalToWorld) const
+	{
+		if (Instance->BindingType == EHairBindingType::Skinning)
+		{
+			OutPreviousLocalToWorld = Instance->Debug.SkeletalPreviousLocalToWorld.ToMatrixWithScale();
+		}
+	}
+
 #if RHI_RAYTRACING
 	virtual bool IsRayTracingRelevant() const { return true; }
 	virtual bool IsRayTracingStaticRelevant() const { return false; }
@@ -622,6 +638,10 @@ public:
 		{
 			const FHairGroup& GroupData = HairGroups[GroupIt];
 
+			FHairGroupInstance* Instance = HairGroupInstances[GroupIt];
+			FMatrix OverrideLocalToWorld = GetLocalToWorld();
+			GetOverrideLocalToTransform(Instance, OverrideLocalToWorld);
+
 			if (GroupData.PublicData->VFInput.GeometryType == EHairGeometryType::Strands)
 			{
 				if (GroupData.RayTracingGeometry_Strands && GroupData.RayTracingGeometry_Strands->RayTracingGeometryRHI.IsValid())
@@ -630,7 +650,7 @@ public:
 					{
 						check(Segment.VertexBuffer.IsValid());
 					}
-					AddOpaqueRaytracingInstance(GetLocalToWorld(), GroupData.RayTracingGeometry_Strands, RaytracingInstanceMask_ThinShadow, OutRayTracingInstances);
+					AddOpaqueRaytracingInstance(OverrideLocalToWorld, GroupData.RayTracingGeometry_Strands, RaytracingInstanceMask_ThinShadow, OutRayTracingInstances);
 				}
 			}
 			else if (GroupData.PublicData->VFInput.GeometryType == EHairGeometryType::Cards)
@@ -644,14 +664,13 @@ public:
 					}
 				// Either use shadow only (no material evaluation) or full material evaluation during tracing
 				#if 0
-					AddOpaqueRaytracingInstance(GetLocalToWorld(), GroupData.RayTracingGeometries_Cards[LODIndex], RaytracingInstanceMask_ThinShadow, OutRayTracingInstances);
+					AddOpaqueRaytracingInstance(OverrideLocalToWorld, GroupData.RayTracingGeometries_Cards[LODIndex], RaytracingInstanceMask_ThinShadow, OutRayTracingInstances);
 				#else
-					FHairGroupInstance* Instance = HairGroupInstances[GroupIt];
 
 					FMeshBatch* MeshBatch = CreateMeshBatch(Context.ReferenceView, Context.ReferenceViewFamily, Context.RayTracingMeshResourceCollector, GroupData, Instance, GroupIt, /*LODIndex,*/ nullptr);
 					TArray<FMeshBatch> MeshBatches;
 					MeshBatches.Add(*MeshBatch);
-					AddOpaqueRaytracingInstance(GetLocalToWorld(), GroupData.RayTracingGeometries_Cards[LODIndex], RaytracingInstanceMask_Opaque, MeshBatches, OutRayTracingInstances);
+					AddOpaqueRaytracingInstance(OverrideLocalToWorld, GroupData.RayTracingGeometries_Cards[LODIndex], RaytracingInstanceMask_Opaque, MeshBatches, OutRayTracingInstances);
 				#endif
 				}
 			}
@@ -666,14 +685,13 @@ public:
 					}
 				// Either use shadow only (no material evaluation) or full material evaluation during tracing
 				#if 0
-					AddOpaqueRaytracingInstance(GetLocalToWorld(), GroupData.RayTracingGeometries_Meshes[LODIndex], RaytracingInstanceMask_ThinShadow, OutRayTracingInstances);
+					AddOpaqueRaytracingInstance(OverrideLocalToWorld, GroupData.RayTracingGeometries_Meshes[LODIndex], RaytracingInstanceMask_ThinShadow, OutRayTracingInstances);
 				#else
-					FHairGroupInstance* Instance = HairGroupInstances[GroupIt];
 
 					FMeshBatch* MeshBatch = CreateMeshBatch(Context.ReferenceView, Context.ReferenceViewFamily, Context.RayTracingMeshResourceCollector, GroupData, Instance, GroupIt, /*LODIndex,*/ nullptr);
 					TArray<FMeshBatch> MeshBatches;
 					MeshBatches.Add(*MeshBatch);
-					AddOpaqueRaytracingInstance(GetLocalToWorld(), GroupData.RayTracingGeometries_Meshes[LODIndex], RaytracingInstanceMask_Opaque, MeshBatches, OutRayTracingInstances);
+					AddOpaqueRaytracingInstance(OverrideLocalToWorld, GroupData.RayTracingGeometries_Meshes[LODIndex], RaytracingInstanceMask_Opaque, MeshBatches, OutRayTracingInstances);
 				#endif
 				}
 			}
@@ -867,11 +885,8 @@ public:
 		GetScene().GetPrimitiveUniformShaderParameters_RenderThread(GetPrimitiveSceneInfo(), bHasPrecomputedVolumetricLightmap, PreviousLocalToWorld, SingleCaptureIndex, bOutputVelocity);
 
 		FMatrix OverrideLocalToWorld = GetLocalToWorld();
-		if (Instance->BindingType == EHairBindingType::Skinning)
-		{
-			OverrideLocalToWorld = Instance->Debug.SkeletalLocalToWorld.ToMatrixWithScale();
-			PreviousLocalToWorld = Instance->Debug.SkeletalPreviousLocalToWorld.ToMatrixWithScale();
-		}
+		GetOverrideLocalToTransform(Instance, OverrideLocalToWorld);
+		GetOverridePreviousLocalToTransform(Instance, PreviousLocalToWorld);
 
 		FDynamicPrimitiveUniformBuffer& DynamicPrimitiveUniformBuffer = Collector.AllocateOneFrameResource<FDynamicPrimitiveUniformBuffer>();
 		DynamicPrimitiveUniformBuffer.Set(OverrideLocalToWorld, PreviousLocalToWorld, GetBounds(), GetLocalBounds(), true, false, bDrawVelocity, bOutputVelocity);
diff --git a/Engine/Plugins/Runtime/HairStrands/Source/HairStrandsCore/Private/HairStrandsRendering.cpp b/Engine/Plugins/Runtime/HairStrands/Source/HairStrandsCore/Private/HairStrandsRendering.cpp
--- a/Engine/Plugins/Runtime/HairStrands/Source/HairStrandsCore/Private/HairStrandsRendering.cpp
+++ b/Engine/Plugins/Runtime/HairStrands/Source/HairStrandsCore/Private/HairStrandsRendering.cpp
@@ -344,7 +344,7 @@ static void AddCopyBufferPass(
 	const uint32 TotalByteCount = DstBuffer->Desc.BytesPerElement * DstBuffer->Desc.NumElements;
 
 	GraphBuilder.AddPass(
-		RDG_EVENT_NAME("HairCardsUpdateBLAS"),
+		RDG_EVENT_NAME("HairCopyPosition"),
 		Parameters,
 		ERDGPassFlags::Copy,
 		[Parameters, SrcBuffer, TotalByteCount](FRHICommandList& RHICmdList)
@@ -1567,35 +1567,38 @@ void ComputeHairStrandsInterpolation(
 		if (Instance->Strands.RenRaytracingResource)
 		{
 			const float HairRadiusScaleRT = (GHairRaytracingRadiusScale > 0 ? GHairRaytracingRadiusScale : Instance->Strands.Modifier.HairRaytracingRadiusScale);
-
-			FRDGImportedBuffer Raytracing_PositionBuffer = Register(GraphBuilder, Instance->Strands.RenRaytracingResource->PositionBuffer, ERDGImportedBufferFlags::CreateViews);
-			AddGenerateRaytracingGeometryPass(
-				GraphBuilder,
-				ShaderMap,
-				Instance->Strands.RestResource->GetVertexCount(),
-				ScaleAndClipDesc.MaxOutHairRadius * HairRadiusScaleRT,
-				RegisterAsSRV(GraphBuilder, Instance->Strands.DeformedResource->GetPositionOffsetBuffer(FHairStrandsDeformedResource::EFrameType::Current)),
-				CullingData,
-				Strands_DeformedPosition.SRV,
-				Raytracing_PositionBuffer.UAV);
-
-			GraphBuilder.AddPass(
-			RDG_EVENT_NAME("HairStrandsUpdateBLAS"),
-			ERDGPassFlags::NeverCull,
-			[Instance](FRHICommandList& RHICmdList)
+			const bool bNeedUpdate = !Instance->Strands.RenRaytracingResource->bIsRTGeometryInitialized || Instance->BindingType == EHairBindingType::Skinning;
+			if (bNeedUpdate)
 			{
-				const bool bNeedFullBuild = !Instance->Strands.RenRaytracingResource->bIsRTGeometryInitialized;
-				if (bNeedFullBuild)
-				{
-					FBufferRHIRef PositionBuffer(Instance->Strands.RenRaytracingResource->PositionBuffer.Buffer->GetRHI());
-					BuildHairAccelerationStructure_Strands(RHICmdList, Instance->Strands.RenRaytracingResource->VertexCount, PositionBuffer, &Instance->Strands.RenRaytracingResource->RayTracingGeometry);
-				}
-				else
+				FRDGImportedBuffer Raytracing_PositionBuffer = Register(GraphBuilder, Instance->Strands.RenRaytracingResource->PositionBuffer, ERDGImportedBufferFlags::CreateViews);
+				AddGenerateRaytracingGeometryPass(
+					GraphBuilder,
+					ShaderMap,
+					Instance->Strands.RestResource->GetVertexCount(),
+					ScaleAndClipDesc.MaxOutHairRadius * HairRadiusScaleRT,
+					RegisterAsSRV(GraphBuilder, Instance->Strands.DeformedResource->GetPositionOffsetBuffer(FHairStrandsDeformedResource::EFrameType::Current)),
+					CullingData,
+					Strands_DeformedPosition.SRV,
+					Raytracing_PositionBuffer.UAV);
+
+				GraphBuilder.AddPass(
+					RDG_EVENT_NAME("HairStrandsUpdateBLAS"),
+					ERDGPassFlags::NeverCull,
+				[Instance](FRHICommandList& RHICmdList)
 				{
-					UpdateHairAccelerationStructure(RHICmdList, &Instance->Strands.RenRaytracingResource->RayTracingGeometry);
-				}
-				Instance->Strands.RenRaytracingResource->bIsRTGeometryInitialized = true;
-			});
+					const bool bNeedFullBuild = !Instance->Strands.RenRaytracingResource->bIsRTGeometryInitialized;
+					if (bNeedFullBuild)
+					{
+						FBufferRHIRef PositionBuffer(Instance->Strands.RenRaytracingResource->PositionBuffer.Buffer->GetRHI());
+						BuildHairAccelerationStructure_Strands(RHICmdList, Instance->Strands.RenRaytracingResource->VertexCount, PositionBuffer, &Instance->Strands.RenRaytracingResource->RayTracingGeometry);
+					}
+					else
+					{
+						UpdateHairAccelerationStructure(RHICmdList, &Instance->Strands.RenRaytracingResource->RayTracingGeometry);
+					}
+					Instance->Strands.RenRaytracingResource->bIsRTGeometryInitialized = true;
+				});
+			}
 		}
 		#endif
 	}
@@ -1690,17 +1693,20 @@ void ComputeHairStrandsInterpolation(
 			#if RHI_RAYTRACING
 			if (LOD.RaytracingResource)
 			{
-				GraphBuilder.AddPass(
-					RDG_EVENT_NAME("HairCardsUpdateBLAS"),
-					ERDGPassFlags::NeverCull,
-					[Instance, HairLODIndex](FRHICommandList& RHICmdList)
+				const bool bNeedUpdate = !LOD.RaytracingResource->bIsRTGeometryInitialized || Instance->BindingType == EHairBindingType::Skinning;
+				if (bNeedUpdate)
+				{
+					GraphBuilder.AddPass(
+						RDG_EVENT_NAME("HairCardsUpdateBLAS"),
+						ERDGPassFlags::NeverCull,
+						[Instance, HairLODIndex](FRHICommandList& RHICmdList)
 					{
 						FHairGroupInstance::FCards::FLOD& LocalLOD = Instance->Cards.LODs[HairLODIndex];
 
 						const bool bNeedFullBuild = !LocalLOD.RaytracingResource->bIsRTGeometryInitialized;
 						if (bNeedFullBuild)
 						{
-							BuildHairAccelerationStructure_Cards(RHICmdList, LocalLOD.RestResource, LocalLOD.DeformedResource,  &LocalLOD.RaytracingResource->RayTracingGeometry);
+							BuildHairAccelerationStructure_Cards(RHICmdList, LocalLOD.RestResource, LocalLOD.DeformedResource, &LocalLOD.RaytracingResource->RayTracingGeometry);
 						}
 						else
 						{
@@ -1708,6 +1714,7 @@ void ComputeHairStrandsInterpolation(
 						}
 						LocalLOD.RaytracingResource->bIsRTGeometryInitialized = true;
 					});
+				}
 			}
 			#endif
 		}
@@ -1735,25 +1742,27 @@ void ComputeHairStrandsInterpolation(
 		FHairGroupInstance::FMeshes::FLOD& LOD = Instance->Meshes.LODs[HairLODIndex];
 		if (LOD.RaytracingResource)
 		{
-			GraphBuilder.AddPass(
-				RDG_EVENT_NAME("HairMeshesUpdateBLAS"),
-				ERDGPassFlags::NeverCull,
-				[Instance, HairLODIndex](FRHICommandList& RHICmdList)
+			const bool bNeedUpdate = !LOD.RaytracingResource->bIsRTGeometryInitialized || Instance->BindingType == EHairBindingType::Skinning;
+			if (bNeedUpdate)
 			{
-
-				FHairGroupInstance::FMeshes::FLOD& LocalLOD = Instance->Meshes.LODs[HairLODIndex];
-				
-				const bool bNeedFullBuild = !LocalLOD.RaytracingResource->bIsRTGeometryInitialized;
-				if (bNeedFullBuild)
-				{
-					BuildHairAccelerationStructure_Meshes(RHICmdList, LocalLOD.RestResource, LocalLOD.DeformedResource,  &LocalLOD.RaytracingResource->RayTracingGeometry);
-				}
-				else
+				GraphBuilder.AddPass(
+					RDG_EVENT_NAME("HairMeshesUpdateBLAS"),
+					ERDGPassFlags::NeverCull,
+					[Instance, HairLODIndex](FRHICommandList& RHICmdList)
 				{
-					UpdateHairAccelerationStructure(RHICmdList, &LocalLOD.RaytracingResource->RayTracingGeometry);
-				}
-				LocalLOD.RaytracingResource->bIsRTGeometryInitialized = true;
-			});
+					FHairGroupInstance::FMeshes::FLOD& LocalLOD = Instance->Meshes.LODs[HairLODIndex];				
+					const bool bNeedFullBuild = !LocalLOD.RaytracingResource->bIsRTGeometryInitialized;
+					if (bNeedFullBuild)
+					{
+						BuildHairAccelerationStructure_Meshes(RHICmdList, LocalLOD.RestResource, LocalLOD.DeformedResource,  &LocalLOD.RaytracingResource->RayTracingGeometry);
+					}
+					else
+					{
+						UpdateHairAccelerationStructure(RHICmdList, &LocalLOD.RaytracingResource->RayTracingGeometry);
+					}
+					LocalLOD.RaytracingResource->bIsRTGeometryInitialized = true;
+				});
+			}
 		}
 		#endif
 	}
