Remove own allocation for root binding data as this is no longer used.

diff --git a/Engine/Plugins/Runtime/HairStrands/Source/HairStrandsCore/Private/GroomComponent.cpp b/Engine/Plugins/Runtime/HairStrands/Source/HairStrandsCore/Private/GroomComponent.cpp
--- a/Engine/Plugins/Runtime/HairStrands/Source/HairStrandsCore/Private/GroomComponent.cpp
+++ b/Engine/Plugins/Runtime/HairStrands/Source/HairStrandsCore/Private/GroomComponent.cpp
@@ -1960,7 +1960,6 @@ void UGroomComponent::InitResources(bool bIsBindingReloading)
 				check(GroupIt < LocalBindingAsset->HairGroupResources.Num());
 				check(SkeletalMeshComponent->SkeletalMesh ? SkeletalMeshComponent->SkeletalMesh->GetLODInfoArray().Num() == LocalBindingAsset->HairGroupResources[GroupIt].SimRootResources->RootData.MeshProjectionLODs.Num() : false);
 
-				HairGroupInstance->Guides.bOwnRootResourceAllocation = false;
 				HairGroupInstance->Guides.RestRootResource = LocalBindingAsset->HairGroupResources[GroupIt].SimRootResources;
 
 				HairGroupInstance->Guides.DeformedRootResource = new FHairStrandsDeformedRootResource(HairGroupInstance->Guides.RestRootResource);
@@ -2009,7 +2008,6 @@ void UGroomComponent::InitResources(bool bIsBindingReloading)
 				check(GroupIt < LocalBindingAsset->HairGroupResources.Num());
 				check(SkeletalMeshComponent->SkeletalMesh ? SkeletalMeshComponent->SkeletalMesh->GetLODInfoArray().Num() == LocalBindingAsset->HairGroupResources[GroupIt].RenRootResources->RootData.MeshProjectionLODs.Num() : false);
 
-				HairGroupInstance->Strands.bOwnRootResourceAllocation = false;
 				HairGroupInstance->Strands.RestRootResource = LocalBindingAsset->HairGroupResources[GroupIt].RenRootResources;
 
 				HairGroupInstance->Strands.DeformedRootResource = new FHairStrandsDeformedRootResource(HairGroupInstance->Strands.RestRootResource);
@@ -2095,7 +2093,6 @@ void UGroomComponent::InitResources(bool bIsBindingReloading)
 						check(GroupIt < LocalBindingAsset->HairGroupResources.Num());
 						check(SkeletalMeshComponent->SkeletalMesh ? SkeletalMeshComponent->SkeletalMesh->GetLODInfoArray().Num() == LocalBindingAsset->HairGroupResources[GroupIt].CardsRootResources[CardsLODIndex]->RootData.MeshProjectionLODs.Num() : false);
 
-						InstanceLOD.Guides.bOwnRootResourceAllocation = false;
 						InstanceLOD.Guides.RestRootResource = LocalBindingAsset->HairGroupResources[GroupIt].CardsRootResources[CardsLODIndex];
 
 						InstanceLOD.Guides.DeformedRootResource = new FHairStrandsDeformedRootResource(InstanceLOD.Guides.RestRootResource);
@@ -2198,10 +2195,6 @@ void UGroomComponent::ReleaseResources()
 			// Guides
 			if (LocalInstance->Guides.IsValid())
 			{
-				if (LocalInstance->Guides.bOwnRootResourceAllocation)
-				{
-					InternalResourceRelease(LocalInstance->Guides.RestRootResource);
-				}
 				InternalResourceRelease(LocalInstance->Guides.DeformedRootResource);
 				InternalResourceRelease(LocalInstance->Guides.DeformedResource);
 			}
@@ -2209,10 +2202,6 @@ void UGroomComponent::ReleaseResources()
 			// Strands
 			if (LocalInstance->Strands.IsValid())
 			{
-				if (LocalInstance->Strands.bOwnRootResourceAllocation)
-				{
-					InternalResourceRelease(LocalInstance->Strands.RestRootResource);
-				}
 				InternalResourceRelease(LocalInstance->Strands.DeformedRootResource);
 				InternalResourceRelease(LocalInstance->Strands.DeformedResource);
 
@@ -2232,10 +2221,6 @@ void UGroomComponent::ReleaseResources()
 				{
 					if (LocalInstance->Cards.IsValid(CardLODIt))
 					{
-						if (LocalInstance->Cards.LODs[CardLODIt].Guides.bOwnRootResourceAllocation)
-						{
-							InternalResourceRelease(LocalInstance->Cards.LODs[CardLODIt].Guides.RestRootResource);
-						}
 						InternalResourceRelease(LocalInstance->Cards.LODs[CardLODIt].Guides.DeformedRootResource);
 						InternalResourceRelease(LocalInstance->Cards.LODs[CardLODIt].Guides.DeformedResource);
 						InternalResourceRelease(LocalInstance->Cards.LODs[CardLODIt].DeformedResource);
diff --git a/Engine/Plugins/Runtime/HairStrands/Source/HairStrandsCore/Private/GroomDebug.cpp b/Engine/Plugins/Runtime/HairStrands/Source/HairStrandsCore/Private/GroomDebug.cpp
--- a/Engine/Plugins/Runtime/HairStrands/Source/HairStrandsCore/Private/GroomDebug.cpp
+++ b/Engine/Plugins/Runtime/HairStrands/Source/HairStrandsCore/Private/GroomDebug.cpp
@@ -826,9 +826,6 @@ void RunHairStrandsDebug(
 			for (FHairGroupInstance* Instance : Instances)
 			{
 				const bool bIsActive = Instance->WorldType == WorldType;
-				const bool bHasSkinInterpolation = Instance->Strands.RestRootResource != nullptr;
-				const bool bHasBindingAsset = bHasSkinInterpolation && !Instance->Strands.bOwnRootResourceAllocation;
-
 				if (!bIsActive)
 					continue;
 
diff --git a/Engine/Plugins/Runtime/HairStrands/Source/HairStrandsCore/Public/GroomInstance.h b/Engine/Plugins/Runtime/HairStrands/Source/HairStrandsCore/Public/GroomInstance.h
--- a/Engine/Plugins/Runtime/HairStrands/Source/HairStrandsCore/Public/GroomInstance.h
+++ b/Engine/Plugins/Runtime/HairStrands/Source/HairStrandsCore/Public/GroomInstance.h
@@ -54,7 +54,6 @@ struct HAIRSTRANDSCORE_API FHairGroupInstance
 
 		// Resources - Rest root data, for deforming strands attached to a skinned mesh surface
 		// Resources - Deformed root data, for deforming strands attached to a skinned mesh surface
-		bool bOwnRootResourceAllocation = true;
 		FHairStrandsRestRootResource* RestRootResource = nullptr;
 		FHairStrandsDeformedRootResource* DeformedRootResource = nullptr;
 
