Fix crash groom crash when non/invalid binding asset are provided.

diff --git a/Engine/Plugins/Runtime/HairStrands/Source/HairStrandsCore/Private/GroomComponent.cpp b/Engine/Plugins/Runtime/HairStrands/Source/HairStrandsCore/Private/GroomComponent.cpp
--- a/Engine/Plugins/Runtime/HairStrands/Source/HairStrandsCore/Private/GroomComponent.cpp
+++ b/Engine/Plugins/Runtime/HairStrands/Source/HairStrandsCore/Private/GroomComponent.cpp
@@ -1871,7 +1871,7 @@ void UGroomComponent::InitResources(bool bIsBindingReloading)
 	if (bHasNeedBindingData && SkeletalMeshComponent)
 	{
 		SkeletalMeshComponent = ValidateBindingAsset(GroomAsset, BindingAsset, SkeletalMeshComponent, bIsBindingReloading, bValidationEnable, this);
-		if (SkeletalMeshComponent->SkeletalMesh == nullptr)
+		if (SkeletalMeshComponent && SkeletalMeshComponent->SkeletalMesh == nullptr)
 		{
 			SkeletalMeshComponent = nullptr;
 		}
