Fix crash with hair strands when visualizing guides/influence/cluster and groom does not have simulation or RBF data.

 * Insure all hair debug mode goes through the same function to get consistency
 * Add special path for groom editor to display guides data even groom does not have simulation or RBF

diff --git a/Engine/Plugins/Runtime/HairStrands/Source/HairStrandsCore/Private/GroomComponent.cpp b/Engine/Plugins/Runtime/HairStrands/Source/HairStrandsCore/Private/GroomComponent.cpp
--- a/Engine/Plugins/Runtime/HairStrands/Source/HairStrandsCore/Private/GroomComponent.cpp
+++ b/Engine/Plugins/Runtime/HairStrands/Source/HairStrandsCore/Private/GroomComponent.cpp
@@ -38,6 +38,8 @@ bool IsHairAdaptiveSubstepsEnabled() { return (GHairEnableAdaptiveSubsteps == 1)
 
 #define LOCTEXT_NAMESPACE "GroomComponent"
 
+EHairStrandsDebugMode GetHairStrandsGeometryDebugMode(FHairGroupInstance* Instance);
+
 ////////////////////////////////////////////////////////////////////////////////////////////////////////////////////////
 
 static EHairBindingType ToHairBindingType(EGroomBindingType In)
@@ -721,7 +723,7 @@ public:
 
 
 		FMaterialRenderProxy* Strands_MaterialProxy = nullptr;
-		const EHairStrandsDebugMode DebugMode = Instances[0]->Debug.DebugMode != EHairStrandsDebugMode::NoneDebug ? Instances[0]->Debug.DebugMode : GetHairStrandsDebugStrandsMode();
+		const EHairStrandsDebugMode DebugMode = GetHairStrandsGeometryDebugMode(Instances[0]);
 		if (DebugMode != EHairStrandsDebugMode::NoneDebug)
 		{
 			float DebugModeScalar = 0;
@@ -1950,7 +1952,7 @@ void UGroomComponent::InitResources(bool bIsBindingReloading)
 		// * Physics simulation
 		// * RBF deformation.
 		// Therefore, even if simulation is disabled, we need to run partially the update if the binding system is enabled (skin deformation + RBF correction)
-		if (GroupData.Guides.IsValid() && (bDynamicResources || bIsStrandsEnabled) && (bHasNeedSimulation || bHasNeedGlobalDeformation))
+		if (GroupData.Guides.IsValid() && (bDynamicResources || bIsStrandsEnabled) && (bHasNeedSimulation || bHasNeedGlobalDeformation || bPreviewMode))
 		{
 			HairGroupInstance->Guides.Data = &GroupData.Guides.Data;
 
diff --git a/Engine/Plugins/Runtime/HairStrands/Source/HairStrandsCore/Private/HairStrandsRendering.cpp b/Engine/Plugins/Runtime/HairStrands/Source/HairStrandsCore/Private/HairStrandsRendering.cpp
--- a/Engine/Plugins/Runtime/HairStrands/Source/HairStrandsCore/Private/HairStrandsRendering.cpp
+++ b/Engine/Plugins/Runtime/HairStrands/Source/HairStrandsCore/Private/HairStrandsRendering.cpp
@@ -1299,7 +1299,15 @@ void RegisterClusterData(FHairGroupInstance* Instance, FHairStrandClusterData* I
 
 EHairStrandsDebugMode GetHairStrandsGeometryDebugMode(FHairGroupInstance* Instance)
 {
-	return Instance->Debug.DebugMode != EHairStrandsDebugMode::NoneDebug ? Instance->Debug.DebugMode : GetHairStrandsDebugStrandsMode();
+	EHairStrandsDebugMode DebugMode = Instance->Debug.DebugMode != EHairStrandsDebugMode::NoneDebug ? Instance->Debug.DebugMode : GetHairStrandsDebugStrandsMode();
+	if ((DebugMode == EHairStrandsDebugMode::SimHairStrands ||
+		 DebugMode == EHairStrandsDebugMode::RenderHairStrands ||
+		 DebugMode == EHairStrandsDebugMode::RenderVisCluster)
+		&& Instance->Guides.Data == nullptr)
+	{
+		DebugMode = EHairStrandsDebugMode::NoneDebug;
+	}
+	return DebugMode;
 }
 
 FHairScaleAndClipDesc ComputeHairScaleAndClipDesc(FHairGroupInstance* Instance)
@@ -1492,7 +1500,8 @@ void ComputeHairStrandsInterpolation(
 		Instance->HairGroupPublicData->VFInput = ComputeHairStrandsVertexInputData(Instance);
 		{
 			{
-				const bool bValidGuide = Instance->Guides.bIsSimulationEnable || Instance->Guides.bHasGlobalInterpolation;
+				// "WITH_EDITOR && bDebugModePatchedAttributeBuffer" is a special path when visualizing groom within the groom editor, so that we can diplay guides even when there is no simulation or global interpolation enabled
+				const bool bValidGuide = Instance->Guides.bIsSimulationEnable || Instance->Guides.bHasGlobalInterpolation || (WITH_EDITOR && bDebugModePatchedAttributeBuffer);
 				const bool bHasSkinning = Instance->BindingType == EHairBindingType::Skinning;
 
 				AddHairStrandsInterpolationPass(
diff --git a/Engine/Plugins/Runtime/HairStrands/Source/HairStrandsCore/Public/GroomComponent.h b/Engine/Plugins/Runtime/HairStrands/Source/HairStrandsCore/Public/GroomComponent.h
--- a/Engine/Plugins/Runtime/HairStrands/Source/HairStrandsCore/Public/GroomComponent.h
+++ b/Engine/Plugins/Runtime/HairStrands/Source/HairStrandsCore/Public/GroomComponent.h
@@ -187,6 +187,11 @@ public:
 	virtual int32 GetNumMaterials() const override;
 	//~ End UPrimitiveComponent Interface
 
+#if WITH_EDITORONLY_DATA
+	// Set the component in preview mode, forcing the loading of certain data
+	void SetPreviewMode(bool bValue) { bPreviewMode = bValue; };
+#endif
+
 private:
 	TArray<FHairGroupInstance*> HairGroupInstances;
 
@@ -211,7 +216,7 @@ private:
 	bool bIsGroomBindingAssetCallbackRegistered;
 	int32 PredictedLODIndex = -1;
 	bool bValidationEnable = true;
-	
+	bool bPreviewMode = false;
 
 	EWorldType::Type GetWorldType() const; 
 	void InitResources(bool bIsBindingReloading=false);
diff --git a/Engine/Plugins/Runtime/HairStrands/Source/HairStrandsEditor/Private/GroomCustomAssetEditorToolkit.cpp b/Engine/Plugins/Runtime/HairStrands/Source/HairStrandsEditor/Private/GroomCustomAssetEditorToolkit.cpp
--- a/Engine/Plugins/Runtime/HairStrands/Source/HairStrandsEditor/Private/GroomCustomAssetEditorToolkit.cpp
+++ b/Engine/Plugins/Runtime/HairStrands/Source/HairStrandsEditor/Private/GroomCustomAssetEditorToolkit.cpp
@@ -238,6 +238,7 @@ void FGroomCustomAssetEditorToolkit::InitPreviewComponents()
 	PreviewGroomComponent = NewObject<UGroomComponent>(GetTransientPackage(), NAME_None, RF_Transient);
 	PreviewGroomComponent->CastShadow = 1;
 	PreviewGroomComponent->bCastDynamicShadow = 1;
+	PreviewGroomComponent->SetPreviewMode(true);
 	PreviewGroomComponent->SetGroomAsset(GroomAsset.Get());
 	PreviewGroomComponent->Activate(true);
 	
