Add debug mode for visualizing hair groups within the groom editor.

diff --git a/Engine/Plugins/Runtime/HairStrands/Source/HairStrandsCore/Private/GroomComponent.cpp b/Engine/Plugins/Runtime/HairStrands/Source/HairStrandsCore/Private/GroomComponent.cpp
--- a/Engine/Plugins/Runtime/HairStrands/Source/HairStrandsCore/Private/GroomComponent.cpp
+++ b/Engine/Plugins/Runtime/HairStrands/Source/HairStrandsCore/Private/GroomComponent.cpp
@@ -38,8 +38,36 @@ bool IsHairAdaptiveSubstepsEnabled() { return (GHairEnableAdaptiveSubsteps == 1)
 
 #define LOCTEXT_NAMESPACE "GroomComponent"
 
+////////////////////////////////////////////////////////////////////////////////////////////////////////////////////////
+
 EHairStrandsDebugMode GetHairStrandsGeometryDebugMode(FHairGroupInstance* Instance);
 
+const FLinearColor GetHairGroupDebugColor(int32 GroupIt)
+{
+	static TArray<FLinearColor> IndexedColor;
+	if (IndexedColor.Num() == 0)
+	{
+		IndexedColor.Add(FLinearColor(FColor::Cyan));
+		IndexedColor.Add(FLinearColor(FColor::Magenta));
+		IndexedColor.Add(FLinearColor(FColor::Orange));
+		IndexedColor.Add(FLinearColor(FColor::Silver));
+		IndexedColor.Add(FLinearColor(FColor::Emerald));
+		IndexedColor.Add(FLinearColor(FColor::Red));
+		IndexedColor.Add(FLinearColor(FColor::Green));
+		IndexedColor.Add(FLinearColor(FColor::Blue));
+		IndexedColor.Add(FLinearColor(FColor::Yellow));
+		IndexedColor.Add(FLinearColor(FColor::Purple));
+		IndexedColor.Add(FLinearColor(FColor::Turquoise));
+	
+	}
+	if (GroupIt >= IndexedColor.Num())
+	{
+		IndexedColor.Add(FLinearColor::MakeRandomColor());
+	}
+	check(GroupIt < IndexedColor.Num());
+	return IndexedColor[GroupIt];	
+}
+
 ////////////////////////////////////////////////////////////////////////////////////////////////////////////////////////
 
 static EHairBindingType ToHairBindingType(EGroomBindingType In)
@@ -106,23 +134,27 @@ public:
 	const float HairMinRadius;
 	const float HairMaxRadius;
 	const float HairClipLength;
+	const FVector HairGroupColor;
 
 	FName DebugModeParamName;
 	FName MinHairRadiusParamName;
 	FName MaxHairRadiusParamName;
 	FName HairClipLengthParamName;
+	FName HairGroupHairColorParamName;
 
 	/** Initialization constructor. */
-	FHairDebugModeMaterialRenderProxy(const FMaterialRenderProxy* InParent, float InMode, float InMinRadius, float InMaxRadius, float InHairClipLength) :
+	FHairDebugModeMaterialRenderProxy(const FMaterialRenderProxy* InParent, float InMode, float InMinRadius, float InMaxRadius, float InHairClipLength, FVector InGroupColor) :
 		Parent(InParent),
 		DebugMode(InMode),
 		HairMinRadius(InMinRadius),
 		HairMaxRadius(InMaxRadius),
 		HairClipLength(InHairClipLength),
+		HairGroupColor(InGroupColor),
 		DebugModeParamName(NAME_FloatProperty),
 		MinHairRadiusParamName(NAME_ByteProperty),
 		MaxHairRadiusParamName(NAME_IntProperty),
-		HairClipLengthParamName(NAME_BoolProperty)
+		HairClipLengthParamName(NAME_BoolProperty),
+		HairGroupHairColorParamName(NAME_VectorProperty)
 	{}
 
 	virtual const FMaterial* GetMaterialNoFallback(ERHIFeatureLevel::Type InFeatureLevel) const override
@@ -137,6 +169,11 @@ public:
 
 	virtual bool GetVectorValue(const FHashedMaterialParameterInfo& ParameterInfo, FLinearColor* OutValue, const FMaterialRenderContext& Context) const override
 	{
+		if (ParameterInfo.Name == HairGroupHairColorParamName)
+		{
+			*OutValue = HairGroupColor;
+			return true;
+		}
 		return Parent->GetVectorValue(ParameterInfo, OutValue, Context);
 	}
 
@@ -162,10 +199,7 @@ public:
 			*OutValue = HairClipLength;
 			return true;
 		}
-		else
-		{
-			return Parent->GetScalarValue(ParameterInfo, OutValue, Context);
-		}
+		return Parent->GetScalarValue(ParameterInfo, OutValue, Context);
 	}
 
 	virtual bool GetTextureValue(const FHashedMaterialParameterInfo& ParameterInfo, const URuntimeVirtualTexture** OutValue, const FMaterialRenderContext& Context) const override
@@ -721,41 +755,6 @@ public:
 
 		QUICK_SCOPE_CYCLE_COUNTER(STAT_HairStrandsSceneProxy_GetDynamicMeshElements);
 
-
-		FMaterialRenderProxy* Strands_MaterialProxy = nullptr;
-		const EHairStrandsDebugMode DebugMode = GetHairStrandsGeometryDebugMode(Instances[0]);
-		if (DebugMode != EHairStrandsDebugMode::NoneDebug)
-		{
-			float DebugModeScalar = 0;
-			switch(DebugMode)
-			{
-			case EHairStrandsDebugMode::NoneDebug					: DebugModeScalar =99.f; break;
-			case EHairStrandsDebugMode::SimHairStrands				: DebugModeScalar = 0.f; break;
-			case EHairStrandsDebugMode::RenderHairStrands			: DebugModeScalar = 0.f; break;
-			case EHairStrandsDebugMode::RenderHairRootUV			: DebugModeScalar = 1.f; break;
-			case EHairStrandsDebugMode::RenderHairUV				: DebugModeScalar = 2.f; break;
-			case EHairStrandsDebugMode::RenderHairSeed				: DebugModeScalar = 3.f; break;
-			case EHairStrandsDebugMode::RenderHairDimension			: DebugModeScalar = 4.f; break;
-			case EHairStrandsDebugMode::RenderHairRadiusVariation	: DebugModeScalar = 5.f; break;
-			case EHairStrandsDebugMode::RenderHairRootUDIM			: DebugModeScalar = 6.f; break;
-			case EHairStrandsDebugMode::RenderHairBaseColor			: DebugModeScalar = 7.f; break;
-			case EHairStrandsDebugMode::RenderHairRoughness			: DebugModeScalar = 8.f; break;
-			case EHairStrandsDebugMode::RenderVisCluster			: DebugModeScalar = 0.f; break;
-			};
-
-			// TODO: fix this as the radius is incorrect. This code run before the interpolation code, which is where HairRadius is updated.
-			float HairMaxRadius = 0;
-			for (FHairGroupInstance* Instance : Instances)
-			{
-				HairMaxRadius = FMath::Max(HairMaxRadius, Instance->Strands.Modifier.HairWidth * 0.5f);
-			}
-
-			const float HairClipLength = GetHairClipLength();
-			auto DebugMaterial = new FHairDebugModeMaterialRenderProxy(Strands_DebugMaterial ? Strands_DebugMaterial->GetRenderProxy() : nullptr, DebugModeScalar, 0, HairMaxRadius, HairClipLength);
-			Collector.RegisterOneFrameMaterialProxy(DebugMaterial);
-			Strands_MaterialProxy = DebugMaterial;
-		}
-
 		// Need information back from the rendering thread to knwo which representation to use (strands/cards/mesh)
 		for (int32 ViewIndex = 0; ViewIndex < Views.Num(); ViewIndex++)
 		{
@@ -769,6 +768,42 @@ public:
 			{
 				for (uint32 GroupIt = 0; GroupIt < GroupCount; ++GroupIt)
 				{
+					FMaterialRenderProxy* Strands_MaterialProxy = nullptr;
+					const EHairStrandsDebugMode DebugMode = GetHairStrandsGeometryDebugMode(Instances[GroupIt]);
+					if (DebugMode != EHairStrandsDebugMode::NoneDebug)
+					{
+						float DebugModeScalar = 0;
+						switch(DebugMode)
+						{
+						case EHairStrandsDebugMode::NoneDebug					: DebugModeScalar =99.f; break;
+						case EHairStrandsDebugMode::SimHairStrands				: DebugModeScalar = 0.f; break;
+						case EHairStrandsDebugMode::RenderHairStrands			: DebugModeScalar = 0.f; break;
+						case EHairStrandsDebugMode::RenderHairRootUV			: DebugModeScalar = 1.f; break;
+						case EHairStrandsDebugMode::RenderHairUV				: DebugModeScalar = 2.f; break;
+						case EHairStrandsDebugMode::RenderHairSeed				: DebugModeScalar = 3.f; break;
+						case EHairStrandsDebugMode::RenderHairDimension			: DebugModeScalar = 4.f; break;
+						case EHairStrandsDebugMode::RenderHairRadiusVariation	: DebugModeScalar = 5.f; break;
+						case EHairStrandsDebugMode::RenderHairRootUDIM			: DebugModeScalar = 6.f; break;
+						case EHairStrandsDebugMode::RenderHairBaseColor			: DebugModeScalar = 7.f; break;
+						case EHairStrandsDebugMode::RenderHairRoughness			: DebugModeScalar = 8.f; break;
+						case EHairStrandsDebugMode::RenderVisCluster			: DebugModeScalar = 0.f; break;
+						case EHairStrandsDebugMode::RenderHairGroup				: DebugModeScalar = 9.f; break;
+						};
+
+						// TODO: fix this as the radius is incorrect. This code run before the interpolation code, which is where HairRadius is updated.
+						float HairMaxRadius = 0;
+						for (FHairGroupInstance* Instance : Instances)
+						{
+							HairMaxRadius = FMath::Max(HairMaxRadius, Instance->Strands.Modifier.HairWidth * 0.5f);
+						}
+
+						const FVector HairGroupColor = FVector(GetHairGroupDebugColor(GroupIt));
+						const float HairClipLength = GetHairClipLength();
+						auto DebugMaterial = new FHairDebugModeMaterialRenderProxy(Strands_DebugMaterial ? Strands_DebugMaterial->GetRenderProxy() : nullptr, DebugModeScalar, 0, HairMaxRadius, HairClipLength, HairGroupColor);
+						Collector.RegisterOneFrameMaterialProxy(DebugMaterial);
+						Strands_MaterialProxy = DebugMaterial;
+					}
+
 					FMeshBatch* MeshBatch = CreateMeshBatch(View, ViewFamily, Collector, HairGroups[GroupIt], Instances[GroupIt], GroupIt, Strands_MaterialProxy);
 					if (MeshBatch == nullptr)
 					{
diff --git a/Engine/Plugins/Runtime/HairStrands/Source/HairStrandsCore/Public/GroomComponent.h b/Engine/Plugins/Runtime/HairStrands/Source/HairStrandsCore/Public/GroomComponent.h
--- a/Engine/Plugins/Runtime/HairStrands/Source/HairStrandsCore/Public/GroomComponent.h
+++ b/Engine/Plugins/Runtime/HairStrands/Source/HairStrandsCore/Public/GroomComponent.h
@@ -238,4 +238,9 @@ public:
 private:
 	TArray<UGroomComponent*> GroomComponents;
 };
-#endif
\ No newline at end of file
+
+/** Return the debug color of an hair group */
+HAIRSTRANDSCORE_API const FLinearColor GetHairGroupDebugColor(int32 GroupIt);
+
+#endif
+
diff --git a/Engine/Plugins/Runtime/HairStrands/Source/HairStrandsEditor/Private/GroomAssetDetails.cpp b/Engine/Plugins/Runtime/HairStrands/Source/HairStrandsEditor/Private/GroomAssetDetails.cpp
--- a/Engine/Plugins/Runtime/HairStrands/Source/HairStrandsEditor/Private/GroomAssetDetails.cpp
+++ b/Engine/Plugins/Runtime/HairStrands/Source/HairStrandsEditor/Private/GroomAssetDetails.cpp
@@ -899,7 +899,7 @@ TSharedRef<SWidget> FGroomRenderingDetails::MakeGroupNameCustomization(int32 Gro
 		return SNew(STextBlock)
 			.Font(IDetailLayoutBuilder::GetDetailFont())
 			.ColorAndOpacity(HairGroupColor)
-			.Text(FText::Format(LOCTEXT("Cards", "Cards {0}"), FText::AsNumber(GroupIndex)));
+			.Text(FText::Format(LOCTEXT("Cards", "Cards {0} "), FText::AsNumber(GroupIndex)));
 	}
 	break;
 	case EMaterialPanelType::Meshes:
@@ -907,7 +907,7 @@ TSharedRef<SWidget> FGroomRenderingDetails::MakeGroupNameCustomization(int32 Gro
 		return SNew(STextBlock)
 			.Font(IDetailLayoutBuilder::GetDetailFont())
 			.ColorAndOpacity(HairGroupColor)
-			.Text(FText::Format(LOCTEXT("Meshes", "Meshes {0}"), FText::AsNumber(GroupIndex)));
+			.Text(FText::Format(LOCTEXT("Meshes", "Meshes {0} "), FText::AsNumber(GroupIndex)));
 	}
 	break;
 	default:
@@ -915,7 +915,7 @@ TSharedRef<SWidget> FGroomRenderingDetails::MakeGroupNameCustomization(int32 Gro
 		return SNew(STextBlock)
 			.Font(IDetailLayoutBuilder::GetDetailFont())
 			.ColorAndOpacity(HairGroupColor)
-			.Text(FText::Format(LOCTEXT("Group", "Group ID {0}"), FText::AsNumber(GroupIndex)));
+			.Text(FText::Format(LOCTEXT("Group", "Group ID {0} "), FText::AsNumber(GroupIndex)));
 	}
 	break;
 	}
@@ -927,6 +927,7 @@ TSharedRef<SWidget> FGroomRenderingDetails::MakeGroupNameCustomization(int32 Gro
 void FGroomRenderingDetails::OnGenerateElementForHairGroup(TSharedRef<IPropertyHandle> StructProperty, int32 GroupIndex, IDetailChildrenBuilder& ChildrenBuilder, IDetailLayoutBuilder* DetailLayout)
 {
 	const FSlateFontInfo DetailFontInfo = IDetailLayoutBuilder::GetDetailFont();
+	const FLinearColor GroupColor = GetHairGroupDebugColor(GroupIndex);
 
 	FProperty* Property = StructProperty->GetProperty();
 
@@ -954,12 +955,23 @@ void FGroomRenderingDetails::OnGenerateElementForHairGroup(TSharedRef<IPropertyH
 		]
 		+ SHorizontalBox::Slot()
 		.AutoWidth()
+		.VAlign(VAlign_Center)
 		.HAlign(HAlign_Right)
 		[
 			MakeGroupNameCustomization(GroupIndex)
 		]
 		+ SHorizontalBox::Slot()
 		.AutoWidth()
+		.HAlign(HAlign_Right)
+		.VAlign(VAlign_Center)
+		[
+			SNew(SColorBlock)
+			.Color(GroupColor)
+			.Size(FVector2D(8,8))
+			.CornerRadius(FVector4(4.0f, 4.0f, 4.0f, 4.0f))
+		]
+		+ SHorizontalBox::Slot()
+		.AutoWidth()
 		[
 			MakeGroupNameButtonCustomization(GroupIndex, Property)
 		]
diff --git a/Engine/Plugins/Runtime/HairStrands/Source/HairStrandsEditor/Private/GroomEditorCommands.cpp b/Engine/Plugins/Runtime/HairStrands/Source/HairStrandsEditor/Private/GroomEditorCommands.cpp
--- a/Engine/Plugins/Runtime/HairStrands/Source/HairStrandsEditor/Private/GroomEditorCommands.cpp
+++ b/Engine/Plugins/Runtime/HairStrands/Source/HairStrandsEditor/Private/GroomEditorCommands.cpp
@@ -38,6 +38,7 @@ void FGroomEditorCommands::RegisterCommands()
 	UI_COMMAND(ViewMode_BaseColor,		"BaseColor",		"View vertex base color",		EUserInterfaceActionType::Button, FInputChord());
 	UI_COMMAND(ViewMode_Roughness,		"Roughness",		"View vertex roughness",		EUserInterfaceActionType::Button, FInputChord());
 	UI_COMMAND(ViewMode_VisCluster,		"VisCluster",		"View clusters",				EUserInterfaceActionType::Button, FInputChord());
+	UI_COMMAND(ViewMode_Group,			"Hair Groups",		"View hair groups",				EUserInterfaceActionType::Button, FInputChord());
 
 	UI_COMMAND(ViewMode_CardsGuides,	"CardsGuide",		"Cards Guides",					EUserInterfaceActionType::Button, FInputChord());
 }
diff --git a/Engine/Plugins/Runtime/HairStrands/Source/HairStrandsEditor/Private/GroomEditorViewportToolBar.cpp b/Engine/Plugins/Runtime/HairStrands/Source/HairStrandsEditor/Private/GroomEditorViewportToolBar.cpp
--- a/Engine/Plugins/Runtime/HairStrands/Source/HairStrandsEditor/Private/GroomEditorViewportToolBar.cpp
+++ b/Engine/Plugins/Runtime/HairStrands/Source/HairStrandsEditor/Private/GroomEditorViewportToolBar.cpp
@@ -64,6 +64,10 @@ TSharedRef<SWidget> SGroomEditorViewportToolbar::GenerateShowMenu() const
 
 		ViewMenuBuilder.AddMenuSeparator();
 
+		ViewMenuBuilder.AddMenuEntry(Commands.ViewMode_Group);
+
+		ViewMenuBuilder.AddMenuSeparator();
+
 		ViewMenuBuilder.AddMenuEntry(Commands.ViewMode_CardsGuides);
 	}
 
diff --git a/Engine/Plugins/Runtime/HairStrands/Source/HairStrandsEditor/Private/Widgets/SGroomEditorViewport.cpp b/Engine/Plugins/Runtime/HairStrands/Source/HairStrandsEditor/Private/Widgets/SGroomEditorViewport.cpp
--- a/Engine/Plugins/Runtime/HairStrands/Source/HairStrandsEditor/Private/Widgets/SGroomEditorViewport.cpp
+++ b/Engine/Plugins/Runtime/HairStrands/Source/HairStrandsEditor/Private/Widgets/SGroomEditorViewport.cpp
@@ -200,6 +200,7 @@ void SGroomEditorViewport::BindCommands()
 	CommandList->MapAction(FGroomEditorCommands::Get().ViewMode_BaseColor,			FExecuteAction::CreateSP(this, &SGroomEditorViewport::OnViewMode, EHairStrandsDebugMode::RenderHairBaseColor),		FCanExecuteAction::CreateSP(this, &SGroomEditorViewport::CanViewMode));
 	CommandList->MapAction(FGroomEditorCommands::Get().ViewMode_Roughness,			FExecuteAction::CreateSP(this, &SGroomEditorViewport::OnViewMode, EHairStrandsDebugMode::RenderHairRoughness),		FCanExecuteAction::CreateSP(this, &SGroomEditorViewport::CanViewMode));
 	CommandList->MapAction(FGroomEditorCommands::Get().ViewMode_VisCluster,			FExecuteAction::CreateSP(this, &SGroomEditorViewport::OnViewMode, EHairStrandsDebugMode::RenderVisCluster),			FCanExecuteAction::CreateSP(this, &SGroomEditorViewport::CanViewMode));
+	CommandList->MapAction(FGroomEditorCommands::Get().ViewMode_Group,				FExecuteAction::CreateSP(this, &SGroomEditorViewport::OnViewMode, EHairStrandsDebugMode::RenderHairGroup),			FCanExecuteAction::CreateSP(this, &SGroomEditorViewport::CanViewMode));
 	CommandList->MapAction(FGroomEditorCommands::Get().ViewMode_CardsGuides,		FExecuteAction::CreateSP(this, &SGroomEditorViewport::OnCardsGuides), FCanExecuteAction::CreateSP(this, &SGroomEditorViewport::CanCardsGuides));
 }
 
diff --git a/Engine/Plugins/Runtime/HairStrands/Source/HairStrandsEditor/Public/GroomEditorCommands.h b/Engine/Plugins/Runtime/HairStrands/Source/HairStrandsEditor/Public/GroomEditorCommands.h
--- a/Engine/Plugins/Runtime/HairStrands/Source/HairStrandsEditor/Public/GroomEditorCommands.h
+++ b/Engine/Plugins/Runtime/HairStrands/Source/HairStrandsEditor/Public/GroomEditorCommands.h
@@ -32,6 +32,7 @@ public:
 	TSharedPtr<FUICommandInfo> ViewMode_BaseColor;
 	TSharedPtr<FUICommandInfo> ViewMode_Roughness;
 	TSharedPtr<FUICommandInfo> ViewMode_VisCluster;
+	TSharedPtr<FUICommandInfo> ViewMode_Group;
 
 	TSharedPtr<FUICommandInfo> ViewMode_CardsGuides;
 
diff --git a/Engine/Source/Runtime/Renderer/Public/HairStrandsInterface.h b/Engine/Source/Runtime/Renderer/Public/HairStrandsInterface.h
--- a/Engine/Source/Runtime/Renderer/Public/HairStrandsInterface.h
+++ b/Engine/Source/Runtime/Renderer/Public/HairStrandsInterface.h
@@ -66,6 +66,7 @@ enum class EHairStrandsDebugMode : uint8
 	RenderHairBaseColor,
 	RenderHairRoughness,
 	RenderVisCluster,
+	RenderHairGroup,
 	Count
 };
 
