Remove unused hair strands guide deformation permutation.

diff --git a/Engine/Plugins/Runtime/HairStrands/Source/HairStrandsCore/Private/HairStrandsRendering.cpp b/Engine/Plugins/Runtime/HairStrands/Source/HairStrandsCore/Private/HairStrandsRendering.cpp
--- a/Engine/Plugins/Runtime/HairStrands/Source/HairStrandsCore/Private/HairStrandsRendering.cpp
+++ b/Engine/Plugins/Runtime/HairStrands/Source/HairStrandsCore/Private/HairStrandsRendering.cpp
@@ -104,8 +104,6 @@ enum class EDeformationType : uint8
 	Simulation,		// Use the output of the hair simulation
 	RestStrands,	// Use the rest strands position (no weighted interpolation)
 	RestGuide,		// Use the rest guide as input of the interpolation (no deformation), only weighted interpolation
-	Wave,			// Apply a wave pattern to deform the guides
-	NormalDirection,// Apply a stretch pattern aligned with the guide root's normal
 	OffsetGuide		// Offset the guides
 };
 
@@ -116,8 +114,6 @@ static EDeformationType GetDeformationType()
 	case 0: return EDeformationType::Simulation;
 	case 1: return EDeformationType::RestStrands;
 	case 2: return EDeformationType::RestGuide;
-	case 3: return EDeformationType::Wave;
-	case 4: return EDeformationType::NormalDirection;
 	}
 
 	return EDeformationType::Simulation;
@@ -183,12 +179,11 @@ class FDeformGuideCS : public FGlobalShader
 	SHADER_USE_PARAMETER_STRUCT(FDeformGuideCS, FGlobalShader);
 
 	class FGroupSize : SHADER_PERMUTATION_INT("PERMUTATION_GROUP_SIZE", 2);
-	class FDeformationType : SHADER_PERMUTATION_INT("PERMUTATION_DEFORMATION", 6);
+	class FDeformationType : SHADER_PERMUTATION_INT("PERMUTATION_DEFORMATION", 4);
 	using FPermutationDomain = TShaderPermutationDomain<FGroupSize, FDeformationType>;
 
 	BEGIN_SHADER_PARAMETER_STRUCT(FParameters, )
 		SHADER_PARAMETER(uint32, VertexCount)
-		SHADER_PARAMETER(uint32, IterationCount)
 		SHADER_PARAMETER(FVector, SimRestOffset)
 		SHADER_PARAMETER(uint32, DispatchCountX)
 
@@ -232,19 +227,23 @@ static void AddDeformSimHairStrandsPass(
 	FRDGBufferSRVRef SimDeformedOffsetBuffer,
 	const bool bHasGlobalInterpolation)
 {
-	static uint32 IterationCount = 0;
-	++IterationCount;
-
-	int32 InternalDeformationType = -1;
+	enum EInternalDeformationType
+	{
+		InternalDeformationType_ByPass = 0,
+		InternalDeformationType_Offset = 1,
+		InternalDeformationType_Skinned = 2,
+		InternalDeformationType_RBF = 3,
+		InternalDeformationTypeCount
+	};
+
+	EInternalDeformationType InternalDeformationType = InternalDeformationTypeCount;
 	switch (DeformationType)
 	{
-	case EDeformationType::RestGuide: InternalDeformationType = 0; break;
-	case EDeformationType::Wave: InternalDeformationType = 1; break;
-	case EDeformationType::NormalDirection: InternalDeformationType = 2; break;
-	case EDeformationType::OffsetGuide : InternalDeformationType = 3; break;
+	case EDeformationType::RestGuide   : InternalDeformationType = InternalDeformationType_ByPass; break;
+	case EDeformationType::OffsetGuide : InternalDeformationType = InternalDeformationType_Offset; break;
 	}
 
-	if (InternalDeformationType < 0) return;
+	if (InternalDeformationType == InternalDeformationTypeCount) return;
 
 	const uint32 GroupSize = ComputeGroupSize();
 	const uint32 DispatchCount = FMath::DivideAndRoundUp(VertexCount, GroupSize);
@@ -255,7 +254,6 @@ static void AddDeformSimHairStrandsPass(
 	Parameters->SimRestPosePositionBuffer = SimRestPosePositionBuffer;
 	Parameters->OutSimDeformedPositionBuffer = OutSimDeformedPositionBuffer.UAV;
 	Parameters->VertexCount = VertexCount;
-	Parameters->IterationCount = IterationCount % 10000;
 	Parameters->SimDeformedOffsetBuffer = SimDeformedOffsetBuffer;
 	Parameters->SimRestOffset = SimRestOffset;
 	Parameters->DispatchCountX = DispatchCountX;
@@ -287,7 +285,7 @@ static void AddDeformSimHairStrandsPass(
 			bSupportGlobalInterpolation = bHasGlobalInterpolation && (RestLODDatas.SampleCount > 0);
 			if (!bSupportGlobalInterpolation) 
 			{
-				InternalDeformationType = 4;
+				InternalDeformationType = InternalDeformationType_Skinned;
 				Parameters->SimRestPosition0Buffer = RegisterAsSRV(GraphBuilder, RestLODDatas.RestRootTrianglePosition0Buffer);
 				Parameters->SimRestPosition1Buffer = RegisterAsSRV(GraphBuilder, RestLODDatas.RestRootTrianglePosition1Buffer);
 				Parameters->SimRestPosition2Buffer = RegisterAsSRV(GraphBuilder, RestLODDatas.RestRootTrianglePosition2Buffer);
@@ -300,7 +298,7 @@ static void AddDeformSimHairStrandsPass(
 			}
 			else
 			{
-				InternalDeformationType = 5;
+				InternalDeformationType = InternalDeformationType_RBF;
 				Parameters->MeshSampleWeightsBuffer = RegisterAsSRV(GraphBuilder, DeformedLODDatas.MeshSampleWeightsBuffer);
 				Parameters->RestSamplePositionsBuffer = RegisterAsSRV(GraphBuilder, RestLODDatas.RestSamplePositionsBuffer);
 				Parameters->SampleCount = RestLODDatas.SampleCount;
diff --git a/Engine/Shaders/Private/HairStrands/HairStrandsGuideDeform.usf b/Engine/Shaders/Private/HairStrands/HairStrandsGuideDeform.usf
--- a/Engine/Shaders/Private/HairStrands/HairStrandsGuideDeform.usf
+++ b/Engine/Shaders/Private/HairStrands/HairStrandsGuideDeform.usf
@@ -6,7 +6,6 @@
 
 float3 SimRestOffset;
 uint VertexCount;
-uint IterationCount;
 uint DispatchCountX; 
 
 Buffer<float4> SimDeformedOffsetBuffer;
@@ -24,11 +23,9 @@ RWBuffer<uint4> OutSimDeformedPositionBuffer;
 
 
 #define DEFORMATION_BYPASS 0
-#define DEFORMATION_WAVE 1
-#define DEFORMATION_NORMAL_DIRECTION 2
-#define DEFORMATION_OFFSET_GUIDE 3
-#define DEFORMATION_OFFSET_GUIDE_DYNAMIC 4
-#define DEFORMATION_OFFSET_GUIDE_GLOBAL 5
+#define DEFORMATION_OFFSET_GUIDE 1
+#define DEFORMATION_OFFSET_GUIDE_DYNAMIC 2
+#define DEFORMATION_OFFSET_GUIDE_GLOBAL 3
 
 #if PERMUTATION_DEFORMATION == DEFORMATION_OFFSET_GUIDE_DYNAMIC
 
@@ -65,28 +62,15 @@ Buffer<float4>  RestSamplePositionsBuffer;
 Buffer<float4>  MeshSampleWeightsBuffer;
 #endif
 
-float3 DisplacePosition(uint VertexIndex, float3 Pos, float Seed)
+float3 DisplacePosition(uint VertexIndex, float3 Pos)
 {
-	const float3 OutSimHairPositionOffset = SimDeformedOffsetBuffer[0];
+	const float3 OutSimHairPositionOffset = SimDeformedOffsetBuffer[0].xyz;
 
 // By pass
 #if PERMUTATION_DEFORMATION == DEFORMATION_BYPASS
 	return Pos;
 #endif
 
-// Wave
-#if PERMUTATION_DEFORMATION == DEFORMATION_WAVE
-	const float Scale = 100.f;
-	Pos *= Scale;
-	//	return float3(0, cos(Pos.y + Seed / 16.f), cos(Pos.z + Seed / 32.f)) / Scale;
-	return Pos + float3(0, cos(Pos.y + Seed / 16.f), cos(Pos.z + Seed / 32.f));
-#endif
-
-// Align with some normal
-#if PERMUTATION_DEFORMATION == DEFORMATION_NORMAL_DIRECTION
-	return Pos;
-#endif
-
 #if PERMUTATION_DEFORMATION == DEFORMATION_OFFSET_GUIDE
 	return Pos + (SimRestOffset - OutSimHairPositionOffset);
 #endif
@@ -116,7 +100,7 @@ void MainCS(uint2 DispatchThreadId : SV_DispatchThreadID)
 	{
 		const uint4  Data 				= SimRestPosePositionBuffer[VertexIndex];
 		const float3 ControlPoint		= f16tof32(Data.xyz);
-		const float3 DisplacedPosition	= DisplacePosition(VertexIndex, ControlPoint, IterationCount);
+		const float3 DisplacedPosition	= DisplacePosition(VertexIndex, ControlPoint);
 		OutSimDeformedPositionBuffer[VertexIndex] = uint4(f32tof16(DisplacedPosition), Data.w);
 	}
 }
