Fix invalid local to world transform for hair groom voxelization when a groom is not bound to a skeletal mesh.

diff --git a/Engine/Plugins/Runtime/HairStrands/Source/HairStrandsCore/Private/HairStrandsRendering.cpp b/Engine/Plugins/Runtime/HairStrands/Source/HairStrandsCore/Private/HairStrandsRendering.cpp
--- a/Engine/Plugins/Runtime/HairStrands/Source/HairStrandsCore/Private/HairStrandsRendering.cpp
+++ b/Engine/Plugins/Runtime/HairStrands/Source/HairStrandsCore/Private/HairStrandsRendering.cpp
@@ -1775,10 +1775,10 @@ void ComputeHairStrandsInterpolation(
 	}
 
 	Instance->HairGroupPublicData->VFInput.GeometryType = InstanceGeometryType;
-	if (Instance->BindingType == EHairBindingType::Rigid)
-		Instance->HairGroupPublicData->VFInput.LocalToWorldTransform = Instance->LocalToWorld;
-	else
+	if (Instance->BindingType == EHairBindingType::Skinning)
 		Instance->HairGroupPublicData->VFInput.LocalToWorldTransform = Instance->Debug.SkeletalLocalToWorld;
+	else
+		Instance->HairGroupPublicData->VFInput.LocalToWorldTransform = Instance->LocalToWorld;
 	Instance->HairGroupPublicData->bSupportVoxelization = Instance->Strands.Modifier.bSupportVoxelization && Instance->bCastShadow;
 }
 
