Remove deformation position allocation/update for groom LOD having no deformation/simulation/RBF deformation.

diff --git a/Engine/Plugins/Runtime/HairStrands/Source/HairStrandsCore/Private/GroomAsset.cpp b/Engine/Plugins/Runtime/HairStrands/Source/HairStrandsCore/Private/GroomAsset.cpp
--- a/Engine/Plugins/Runtime/HairStrands/Source/HairStrandsCore/Private/GroomAsset.cpp
+++ b/Engine/Plugins/Runtime/HairStrands/Source/HairStrandsCore/Private/GroomAsset.cpp
@@ -2341,7 +2341,37 @@ void UGroomAsset::InitMeshesResources()
 	}
 }
 
-bool UGroomAsset::IsSimulationEnable(int32 GroupIndex, int32 LODIndex)
+EGroomGeometryType UGroomAsset::GetGeometryType(int32 GroupIndex, int32 LODIndex) const
+{
+	if (GroupIndex < 0 || GroupIndex >= HairGroupsLOD.Num())
+	{
+		return EGroomGeometryType::Strands;
+	}
+
+	if (LODIndex < 0 || LODIndex >= HairGroupsLOD[GroupIndex].LODs.Num())
+	{
+		return EGroomGeometryType::Strands;
+	}
+
+	return HairGroupsLOD[GroupIndex].LODs[LODIndex].GeometryType;
+}
+
+EGroomBindingType UGroomAsset::GetBindingType(int32 GroupIndex, int32 LODIndex) const
+{
+	if (GroupIndex < 0 || GroupIndex >= HairGroupsLOD.Num() || !IsHairStrandsBindingEnable())
+	{
+		return EGroomBindingType::Rigid; // Fallback to rigid by default
+	}
+
+	if (LODIndex < 0 || LODIndex >= HairGroupsLOD[GroupIndex].LODs.Num())
+	{
+		return EGroomBindingType::Rigid; // Fallback to rigid by default
+	}
+
+	return HairGroupsLOD[GroupIndex].LODs[LODIndex].BindingType;
+}
+
+bool UGroomAsset::IsSimulationEnable(int32 GroupIndex, int32 LODIndex) const 
 {
 	if (GroupIndex < 0 || GroupIndex >= HairGroupsLOD.Num() || !IsHairStrandsSimulationEnable())
 	{
@@ -2371,7 +2401,7 @@ bool UGroomAsset::IsSimulationEnable(int32 GroupIndex, int32 LODIndex)
 		(HairGroupsLOD[GroupIndex].LODs[LODIndex].Simulation == EGroomOverrideType::Auto && HairGroupsPhysics[GroupIndex].SolverSettings.EnableSimulation);
 }
 
-bool UGroomAsset::IsGlobalInterpolationEnable(int32 GroupIndex, int32 LODIndex)
+bool UGroomAsset::IsGlobalInterpolationEnable(int32 GroupIndex, int32 LODIndex) const
 {
 	if (GroupIndex < 0 || GroupIndex >= HairGroupsLOD.Num())
 	{
diff --git a/Engine/Plugins/Runtime/HairStrands/Source/HairStrandsCore/Private/GroomComponent.cpp b/Engine/Plugins/Runtime/HairStrands/Source/HairStrandsCore/Private/GroomComponent.cpp
--- a/Engine/Plugins/Runtime/HairStrands/Source/HairStrandsCore/Private/GroomComponent.cpp
+++ b/Engine/Plugins/Runtime/HairStrands/Source/HairStrandsCore/Private/GroomComponent.cpp
@@ -417,7 +417,7 @@ struct FHairCardsVertexFactories
 	EShaderPlatform ShaderPlatform = EShaderPlatform::SP_NumPlatforms;
 };
 
-static EHairGeometryType ConvertToHairGeometryType(EGroomGeometryType Type)
+static EHairGeometryType ToHairGeometryType(EGroomGeometryType Type)
 {
 	switch (Type)
 	{
@@ -851,7 +851,8 @@ public:
 			{
 				return nullptr;
 			}
-			const bool bIsCurrent0 = Instance->Meshes.LODs[IntLODIndex].DeformedResource->CurrentIndex == 0;
+			const bool bIsDynamic  = Instance->Meshes.LODs[IntLODIndex].DeformedResource != nullptr;
+			const bool bIsCurrent0 = bIsDynamic ? Instance->Meshes.LODs[IntLODIndex].DeformedResource->CurrentIndex == 0 : true;
 			VertexFactory = (FVertexFactory*)CardsAndMeshesVertexFactories.GetVertexFactory(GroupIndex, IntLODIndex, bIsCurrent0, GeometryType);
 			check(VertexFactory);
 			HairVertexCount = Instance->Meshes.LODs[IntLODIndex].RestResource->PrimitiveCount * 3;
@@ -867,7 +868,9 @@ public:
 			{
 				return nullptr;
 			}
-			const bool bIsCurrent0 = Instance->Cards.LODs[IntLODIndex].DeformedResource->CurrentIndex == 0;
+
+			const bool bIsDynamic  = Instance->Cards.LODs[IntLODIndex].DeformedResource != nullptr;
+			const bool bIsCurrent0 = bIsDynamic ? Instance->Cards.LODs[IntLODIndex].DeformedResource->CurrentIndex == 0 : true;
 			VertexFactory = (FVertexFactory*)CardsAndMeshesVertexFactories.GetVertexFactory(GroupIndex, IntLODIndex, bIsCurrent0, GeometryType);
 			check(VertexFactory);
 			HairVertexCount = Instance->Cards.LODs[IntLODIndex].RestResource->PrimitiveCount * 3;
@@ -1881,27 +1884,34 @@ void UGroomComponent::InitResources(bool bIsBindingReloading)
 	InitializedResources = GroomAsset;
 
 	// 1. Check if we need any kind of binding data, simulation data, or RBF data
-	bool bHasNeedBindingData = false;
+	//
+	//					  Requires    	| Requires     | Requires
+	//         Features | Skeletal mesh	| Binding data | Guides
+	// -----------------|---------------|--------------|-------------
+	//    Rigid binding | No			| No		   | No
+	// Skinning binding | Yes			| Yes		   | No
+	//       Simulation | No			| No		   | Yes
+	//              RBF | Yes		    | Yes		   | Yes
+	//
+	bool bHasNeedSkinningBinding = false;			   
 	bool bHasNeedGlobalDeformation = false;
 	bool bHasNeedSimulation = false;
 	for (int32 GroupIt = 0, GroupCount = GroomAsset->HairGroupsData.Num(); GroupIt < GroupCount; ++GroupIt)
 	{
-		const FHairGroupsLOD& GroupLOD = GroomAsset->HairGroupsLOD[GroupIt];
-		for (const FHairLODSettings& LODSettings : GroupLOD.LODs)
-		{
-			bHasNeedBindingData  = bHasNeedBindingData || LODSettings.BindingType == EGroomBindingType::Skinning;
-		}
-
 		for (uint32 LODIt = 0, LODCount = GroomAsset->GetLODCount(); LODIt < LODCount; ++LODIt)
 		{
-			bHasNeedSimulation = bHasNeedSimulation || GroomAsset->IsSimulationEnable(GroupIt, LODIt);
-			bHasNeedGlobalDeformation = bHasNeedGlobalDeformation || GroomAsset->IsGlobalInterpolationEnable(GroupIt, LODIt);
+			const EGroomGeometryType GeometryType = GroomAsset->GetGeometryType(GroupIt, LODIt);
+			const EGroomBindingType BindingType = GroomAsset->GetBindingType(GroupIt, LODIt);
+
+			// Force global interpolation to be enable for meshes with skinning binding as we use RBF defomation for 'sticking' meshes onto skel. mesh surface
+			const bool bForceGlobalDeformation = GeometryType == EGroomGeometryType::Meshes && BindingType == EGroomBindingType::Skinning;
+			
+			bHasNeedSkinningBinding   = bHasNeedSkinningBinding || BindingType == EGroomBindingType::Skinning;
+			bHasNeedSimulation		  = bHasNeedSimulation || GroomAsset->IsSimulationEnable(GroupIt, LODIt);
+			bHasNeedGlobalDeformation = bHasNeedGlobalDeformation || GroomAsset->IsGlobalInterpolationEnable(GroupIt, LODIt) || bForceGlobalDeformation;
 		}
 	}
-
-	const FPrimitiveComponentId LocalComponentId = ComponentId;
-	const EWorldType::Type WorldType = GetWorldType();
-	const bool bIsStrandsEnabled = IsHairStrandsEnabled(EHairStrandsShaderType::Strands);
+	const bool bHasNeedBindingData = bHasNeedSkinningBinding || bHasNeedGlobalDeformation;
 
 	// 2. Insure that the binding asset is compatible, otherwise no binding
 	USkeletalMeshComponent* SkeletalMeshComponent = GetAttachParent() ? Cast<USkeletalMeshComponent>(GetAttachParent()) : nullptr;
@@ -1913,16 +1923,26 @@ void UGroomComponent::InitResources(bool bIsBindingReloading)
 			SkeletalMeshComponent = nullptr;
 		}
 	}
-	UGroomBindingAsset* LocalBindingAsset = bHasNeedBindingData ? BindingAsset : nullptr;
 
 	// Insure the ticking of the Groom component always happens after the skeletalMeshComponent.
+	UGroomBindingAsset* LocalBindingAsset = nullptr;
 	if (SkeletalMeshComponent)
 	{
 		RegisteredSkeletalMeshComponent = SkeletalMeshComponent;
 		AddTickPrerequisiteComponent(SkeletalMeshComponent);
 		bUseAttachParentBound = true;
+		LocalBindingAsset = bHasNeedBindingData ? BindingAsset : nullptr;
+	}
+
+	// Update requirement based on binding asset availability
+	if (LocalBindingAsset == nullptr)
+	{
+		bHasNeedSkinningBinding = false;
+		bHasNeedGlobalDeformation = false;
 	}
 
+	const bool bIsStrandsEnabled = IsHairStrandsEnabled(EHairStrandsShaderType::Strands);
+	const EWorldType::Type WorldType = GetWorldType();
 	for (int32 GroupIt = 0, GroupCount = GroomAsset->HairGroupsData.Num(); GroupIt < GroupCount; ++GroupIt)
 	{
 		FHairGroupInstance* HairGroupInstance = new FHairGroupInstance();
@@ -1932,7 +1952,7 @@ void UGroomComponent::InitResources(bool bIsBindingReloading)
 		HairGroupInstance->Debug.GroupIndex = GroupIt;
 		HairGroupInstance->Debug.GroupCount = GroupCount;
 		HairGroupInstance->Debug.GroomAssetName = GroomAsset->GetName();
-		HairGroupInstance->Debug.SkeletalComponent = IsHairStrandsBindingEnable() ? SkeletalMeshComponent : nullptr;
+		HairGroupInstance->Debug.SkeletalComponent = SkeletalMeshComponent;
 		if (SkeletalMeshComponent)
 		{
 			HairGroupInstance->Debug.SkeletalComponentName = SkeletalMeshComponent->GetPathName();
@@ -1941,12 +1961,9 @@ void UGroomComponent::InitResources(bool bIsBindingReloading)
 		HairGroupInstance->BindingType = EHairBindingType::NoneBinding;
 
 		FHairGroupData& GroupData = GroomAsset->HairGroupsData[GroupIt];
-		const uint32 SkeletalLODCount = SkeletalMeshComponent ? SkeletalMeshComponent->GetNumLODs() : 0;
 
 		const EHairInterpolationType HairInterpolationType = ToHairInterpolationType(GroomAsset->HairInterpolationType);
 
-		const bool bDynamicResources = IsHairStrandsSimulationEnable() || IsHairStrandsBindingEnable();
-
 		// Initialize LOD screen size & visibility
 		{
 			HairGroupInstance->HairGroupPublicData = new FHairGroupPublicData(GroupIt);
@@ -1957,7 +1974,7 @@ void UGroomComponent::InitResources(bool bIsBindingReloading)
 			for (uint32 LODIt = 0, LODCount = GroupLOD.LODs.Num(); LODIt < LODCount; ++LODIt)
 			{
 				const FHairLODSettings& LODSettings = GroupLOD.LODs[LODIt];
-				EHairBindingType BindingType = ToHairBindingType(LODSettings.BindingType);
+				EHairBindingType BindingType = ToHairBindingType(GroomAsset->GetBindingType(GroupIt, LODIt));
 				if (BindingType == EHairBindingType::Skinning && !LocalBindingAsset)
 				{
 					BindingType = EHairBindingType::NoneBinding;
@@ -1967,12 +1984,16 @@ void UGroomComponent::InitResources(bool bIsBindingReloading)
 					BindingType = EHairBindingType::NoneBinding;
 				}
 
+				// Force global interpolation to be enable for meshes with skinning binding as we use RBF defomation for 'sticking' meshes onto skel. mesh surface
+				const EHairGeometryType GeometryType = ToHairGeometryType(GroomAsset->GetGeometryType(GroupIt, LODIt));
+				const bool bForceGlobalInterpolation = GeometryType == EHairGeometryType::Meshes && BindingType == EHairBindingType::Skinning;
+
 				const bool LODSimulation = GroomAsset->IsSimulationEnable(GroupIt, LODIt);
-				const bool LODGlobalInterpolation = GroomAsset->IsGlobalInterpolationEnable(GroupIt, LODIt);
+				const bool LODGlobalInterpolation = LocalBindingAsset && (GroomAsset->IsGlobalInterpolationEnable(GroupIt, LODIt) || bForceGlobalInterpolation);
 
 				CPULODScreenSize.Add(LODSettings.ScreenSize);
 				LODVisibility.Add(LODSettings.bVisible);
-				LODGeometryTypes.Add(ConvertToHairGeometryType(LODSettings.GeometryType));
+				LODGeometryTypes.Add(GeometryType);
 				HairGroupInstance->HairGroupPublicData->BindingTypes.Add(BindingType);
 				HairGroupInstance->HairGroupPublicData->LODSimulations.Add(LODSimulation);
 				HairGroupInstance->HairGroupPublicData->LODGlobalInterpolations.Add(LODGlobalInterpolation);
@@ -1982,17 +2003,21 @@ void UGroomComponent::InitResources(bool bIsBindingReloading)
 			HairGroupInstance->HairGroupPublicData->SetLODGeometryTypes(LODGeometryTypes);
 		}
 
+		// Not needed anymore
+		//const bool bDynamicResources = IsHairStrandsSimulationEnable() || IsHairStrandsBindingEnable();
+
 		// Sim data 
 		// Simulation guides are used for two purposes:
 		// * Physics simulation
 		// * RBF deformation.
 		// Therefore, even if simulation is disabled, we need to run partially the update if the binding system is enabled (skin deformation + RBF correction)
-		if (GroupData.Guides.IsValid() && (bDynamicResources || bIsStrandsEnabled) && (bHasNeedSimulation || bHasNeedGlobalDeformation || bPreviewMode))
+		if (GroupData.Guides.IsValid() && (bHasNeedSimulation || bHasNeedGlobalDeformation || bPreviewMode))
 		{
 			HairGroupInstance->Guides.Data = &GroupData.Guides.Data;
 
-			if (SkeletalMeshComponent && LocalBindingAsset)
+			if (LocalBindingAsset)
 			{
+				check(SkeletalMeshComponent);
 				check(LocalBindingAsset);
 				check(GroupIt < LocalBindingAsset->HairGroupResources.Num());
 				check(SkeletalMeshComponent->SkeletalMesh ? SkeletalMeshComponent->SkeletalMesh->GetLODInfoArray().Num() == LocalBindingAsset->HairGroupResources[GroupIt].SimRootResources->RootData.MeshProjectionLODs.Num() : false);
@@ -2004,7 +2029,8 @@ void UGroomComponent::InitResources(bool bIsBindingReloading)
 			}
 			HairGroupInstance->Guides.RestResource = GroupData.Guides.RestResource;
 
-			HairGroupInstance->Guides.DeformedResource = new FHairStrandsDeformedResource(GroupData.Guides.Data.RenderData, true, true, HairGroupInstance->Guides.RestResource->PositionOffset);
+			// If guides are allocated, deformed resources are always needs since they are either used with simulation, or RBF deformation. Both are dynamics, and require deformed positions
+			HairGroupInstance->Guides.DeformedResource = new FHairStrandsDeformedResource(GroupData.Guides.Data.RenderData, true, HairGroupInstance->Guides.RestResource->PositionOffset);
 			BeginInitResource(HairGroupInstance->Guides.DeformedResource);
 
 			HairGroupInstance->Guides.DeformedResource->GetPositionOffset(FHairStrandsDeformedResource::EFrameType::Current)  = HairGroupInstance->Guides.RestResource->PositionOffset;
@@ -2039,22 +2065,53 @@ void UGroomComponent::InitResources(bool bIsBindingReloading)
 			}
 			#endif
 
-			if (SkeletalMeshComponent && LocalBindingAsset && IsHairStrandsBindingEnable())
+			// Check if strands needs to load/allocate root resources
+			bool bNeedRootResources = false;
+			bool bNeedDynamicResources = false;
+			if (LocalBindingAsset)
+			{
+				const FHairGroupsLOD& GroupLOD = GroomAsset->HairGroupsLOD[GroupIt];
+				for (uint32 LODIt = 0, LODCount = GroupLOD.LODs.Num(); LODIt < LODCount; ++LODIt)
+				{
+					if (HairGroupInstance->HairGroupPublicData->GetGeometryType(LODIt) == EHairGeometryType::Strands)
+					{
+						if (HairGroupInstance->HairGroupPublicData->GetBindingType(LODIt) == EHairBindingType::Skinning ||
+							HairGroupInstance->HairGroupPublicData->IsGlobalInterpolationEnable(LODIt))
+						{
+							bNeedRootResources = true;
+						}
+
+						if (HairGroupInstance->HairGroupPublicData->GetBindingType(LODIt) == EHairBindingType::Skinning ||
+							HairGroupInstance->HairGroupPublicData->IsSimulationEnable(LODIt))
+						{
+							bNeedDynamicResources = true;
+						}
+					}
+				}
+			}
+
+			if (bNeedRootResources)
 			{
 				check(LocalBindingAsset);
+				check(SkeletalMeshComponent);
 				check(GroupIt < LocalBindingAsset->HairGroupResources.Num());
 				check(SkeletalMeshComponent->SkeletalMesh ? SkeletalMeshComponent->SkeletalMesh->GetLODInfoArray().Num() == LocalBindingAsset->HairGroupResources[GroupIt].RenRootResources->RootData.MeshProjectionLODs.Num() : false);
 
 				HairGroupInstance->Strands.RestRootResource = LocalBindingAsset->HairGroupResources[GroupIt].RenRootResources;
-
 				HairGroupInstance->Strands.DeformedRootResource = new FHairStrandsDeformedRootResource(HairGroupInstance->Strands.RestRootResource);
 				BeginInitResource(HairGroupInstance->Strands.DeformedRootResource);
 			}
-		
+
 			HairGroupInstance->Strands.RestResource = GroupData.Strands.RestResource;
-			HairGroupInstance->Strands.DeformedResource = new FHairStrandsDeformedResource(GroupData.Strands.Data.RenderData, false, bDynamicResources, HairGroupInstance->Strands.RestResource->PositionOffset);
-			BeginInitResource(HairGroupInstance->Strands.DeformedResource);
-			HairGroupInstance->Strands.ClusterCullingResource = GroupData.Strands.ClusterCullingResource;
+			if (bNeedDynamicResources)
+			{
+				HairGroupInstance->Strands.DeformedResource = new FHairStrandsDeformedResource(GroupData.Strands.Data.RenderData, false, HairGroupInstance->Strands.RestResource->PositionOffset);
+				BeginInitResource(HairGroupInstance->Strands.DeformedResource);
+
+				// Initialize deformed position relative position offset to the rest pose offset
+				HairGroupInstance->Strands.DeformedResource->GetPositionOffset(FHairStrandsDeformedResource::EFrameType::Current)  = HairGroupInstance->Strands.RestResource->PositionOffset;
+				HairGroupInstance->Strands.DeformedResource->GetPositionOffset(FHairStrandsDeformedResource::EFrameType::Previous) = HairGroupInstance->Strands.RestResource->PositionOffset;
+			} 
 
 			// Create an debug buffer for storing cluster visalization data. This is only used for debug purpose, hence only enable in editor build.
 			// Special case for debug mode were the attribute buffer is patch with some custom data to show hair properties (strands belonging to the same cluster, ...)
@@ -2063,6 +2120,7 @@ void UGroomComponent::InitResources(bool bIsBindingReloading)
 			#endif
 
 			// An empty groom doesn't have a ClusterCullingResource
+			HairGroupInstance->Strands.ClusterCullingResource = GroupData.Strands.ClusterCullingResource;
 			if (HairGroupInstance->Strands.ClusterCullingResource)
 			{
 				// This codes assumes strands LOD are contigus and the highest (i.e., 0...x). Change this code to something more robust
@@ -2086,9 +2144,6 @@ void UGroomComponent::InitResources(bool bIsBindingReloading)
 				BeginInitResource(HairGroupInstance->HairGroupPublicData);
 			}
 
-			// Initialize deformed position relative position offset to the rest pose offset
-			HairGroupInstance->Strands.DeformedResource->GetPositionOffset(FHairStrandsDeformedResource::EFrameType::Current)  = HairGroupInstance->Strands.RestResource->PositionOffset;
-			HairGroupInstance->Strands.DeformedResource->GetPositionOffset(FHairStrandsDeformedResource::EFrameType::Previous) = HairGroupInstance->Strands.RestResource->PositionOffset;
 
 			HairGroupInstance->Strands.HairInterpolationType = HairInterpolationType;
 		}
@@ -2100,12 +2155,22 @@ void UGroomComponent::InitResources(bool bIsBindingReloading)
 			FHairGroupInstance::FCards::FLOD& InstanceLOD = HairGroupInstance->Cards.LODs.AddDefaulted_GetRef();
 			if (LOD.IsValid())
 			{
+				const EHairBindingType BindingType	= HairGroupInstance->HairGroupPublicData->GetBindingType(CardsLODIndex);
+				const bool bHasSimulation			= HairGroupInstance->HairGroupPublicData->IsSimulationEnable(CardsLODIndex);
+				const bool bHasGlobalDeformation	= HairGroupInstance->HairGroupPublicData->IsGlobalInterpolationEnable(CardsLODIndex);
+				const bool bNeedDeformedPositions	= bHasSimulation || bHasGlobalDeformation || BindingType == EHairBindingType::Skinning;
+				const bool bNeedRootData			= bHasGlobalDeformation || BindingType == EHairBindingType::Skinning;
+
 				InstanceLOD.Data = &LOD.Data;
 				InstanceLOD.RestResource = LOD.RestResource;
 				InstanceLOD.InterpolationData = &LOD.InterpolationData;
 				InstanceLOD.InterpolationResource = LOD.InterpolationResource;
-				InstanceLOD.DeformedResource = new FHairCardsDeformedResource(LOD.Data.RenderData, false, bDynamicResources);
-				BeginInitResource(InstanceLOD.DeformedResource);
+
+				if (bNeedDeformedPositions)
+				{
+					InstanceLOD.DeformedResource = new FHairCardsDeformedResource(LOD.Data.RenderData, false);
+					BeginInitResource(InstanceLOD.DeformedResource);
+				}
 
 				#if RHI_RAYTRACING
 				if (IsRayTracingEnabled())
@@ -2119,30 +2184,32 @@ void UGroomComponent::InitResources(bool bIsBindingReloading)
 				InstanceLOD.Material = nullptr;
 
 				// Strands data/resources
-				if (bDynamicResources)
+				if (bNeedDeformedPositions)
 				{
 					InstanceLOD.Guides.Data = &LOD.Guides.Data;
 					InstanceLOD.Guides.InterpolationData = &LOD.Guides.InterpolationData;
 					InstanceLOD.Guides.InterpolationResource = LOD.Guides.InterpolationResource;
 
-					if (SkeletalMeshComponent && LocalBindingAsset)
+					if (bNeedRootData)
 					{
+						check(LocalBindingAsset);
 						check(GroupIt < LocalBindingAsset->HairGroupResources.Num());
 						check(SkeletalMeshComponent->SkeletalMesh ? SkeletalMeshComponent->SkeletalMesh->GetLODInfoArray().Num() == LocalBindingAsset->HairGroupResources[GroupIt].CardsRootResources[CardsLODIndex]->RootData.MeshProjectionLODs.Num() : false);
 
 						InstanceLOD.Guides.RestRootResource = LocalBindingAsset->HairGroupResources[GroupIt].CardsRootResources[CardsLODIndex];
-
 						InstanceLOD.Guides.DeformedRootResource = new FHairStrandsDeformedRootResource(InstanceLOD.Guides.RestRootResource);
 						BeginInitResource(InstanceLOD.Guides.DeformedRootResource);
 					}
 
 					InstanceLOD.Guides.RestResource = LOD.Guides.RestResource;
-					InstanceLOD.Guides.DeformedResource = new FHairStrandsDeformedResource(LOD.Guides.Data.RenderData, false, true, InstanceLOD.Guides.RestResource->PositionOffset);
-					BeginInitResource(InstanceLOD.Guides.DeformedResource);
+					{
+						InstanceLOD.Guides.DeformedResource = new FHairStrandsDeformedResource(LOD.Guides.Data.RenderData, false, InstanceLOD.Guides.RestResource->PositionOffset);
+						BeginInitResource(InstanceLOD.Guides.DeformedResource);
 
-					// Initialize deformed position relative position offset to the rest pose offset
-					InstanceLOD.Guides.DeformedResource->GetPositionOffset(FHairStrandsDeformedResource::EFrameType::Current) = InstanceLOD.Guides.RestResource->PositionOffset;
-					InstanceLOD.Guides.DeformedResource->GetPositionOffset(FHairStrandsDeformedResource::EFrameType::Previous) = InstanceLOD.Guides.RestResource->PositionOffset;
+						// Initialize deformed position relative position offset to the rest pose offset
+						InstanceLOD.Guides.DeformedResource->GetPositionOffset(FHairStrandsDeformedResource::EFrameType::Current) = InstanceLOD.Guides.RestResource->PositionOffset;
+						InstanceLOD.Guides.DeformedResource->GetPositionOffset(FHairStrandsDeformedResource::EFrameType::Previous) = InstanceLOD.Guides.RestResource->PositionOffset;
+					}
 
 					InstanceLOD.Guides.HairInterpolationType = HairInterpolationType;
 				}
@@ -2152,14 +2219,21 @@ void UGroomComponent::InitResources(bool bIsBindingReloading)
 		// Meshes resources
 		for (FHairGroupData::FMeshes::FLOD& LOD : GroupData.Meshes.LODs)
 		{
-			const int32 MeshesIt = HairGroupInstance->Meshes.LODs.Num();
+			const int32 MeshLODIndex = HairGroupInstance->Meshes.LODs.Num();
 			FHairGroupInstance::FMeshes::FLOD& InstanceLOD = HairGroupInstance->Meshes.LODs.AddDefaulted_GetRef();
 			if (LOD.IsValid())
 			{
+				const EHairBindingType BindingType = HairGroupInstance->HairGroupPublicData->GetBindingType(MeshLODIndex);
+				const bool bHasGlobalDeformation   = HairGroupInstance->HairGroupPublicData->IsGlobalInterpolationEnable(MeshLODIndex);
+				const bool bNeedDeformedPositions  = bHasGlobalDeformation || BindingType == EHairBindingType::Skinning;
+
 				InstanceLOD.Data = &LOD.Data;
 				InstanceLOD.RestResource = LOD.RestResource;
-				InstanceLOD.DeformedResource = new FHairMeshesDeformedResource(LOD.Data.RenderData, true, bDynamicResources);
-				BeginInitResource(InstanceLOD.DeformedResource);
+				if (bNeedDeformedPositions)
+				{
+					InstanceLOD.DeformedResource = new FHairMeshesDeformedResource(LOD.Data.RenderData, true);
+					BeginInitResource(InstanceLOD.DeformedResource);
+				}
 
 				#if RHI_RAYTRACING
 				if (IsRayTracingEnabled())
diff --git a/Engine/Plugins/Runtime/HairStrands/Source/HairStrandsCore/Private/GroomManager.cpp b/Engine/Plugins/Runtime/HairStrands/Source/HairStrandsCore/Private/GroomManager.cpp
--- a/Engine/Plugins/Runtime/HairStrands/Source/HairStrandsCore/Private/GroomManager.cpp
+++ b/Engine/Plugins/Runtime/HairStrands/Source/HairStrandsCore/Private/GroomManager.cpp
@@ -188,12 +188,21 @@ static void RunInternalHairStrandsInterpolation(
 			const EHairGeometryType InstanceGeometryType = Instance->GeometryType;
 			const EHairBindingType InstanceBindingType = Instance->BindingType;
 
+			const uint32 HairLODIndex = Instance->HairGroupPublicData->LODIndex;
+			const EHairBindingType BindingType = Instance->HairGroupPublicData->GetBindingType(HairLODIndex);
+			const bool bSimulationEnable = Instance->HairGroupPublicData->IsSimulationEnable(HairLODIndex);
+			const bool bGlobalDeformationEnable = Instance->HairGroupPublicData->IsGlobalInterpolationEnable(HairLODIndex);
+			check(InstanceBindingType == BindingType);
+
 			if (EHairStrandsInterpolationType::RenderStrands == Type)
 			{
 				if (InstanceGeometryType == EHairGeometryType::Strands)
 				{
-					if (Instance->Strands.HasValidRootData() && Instance->Strands.DeformedRootResource->IsValid(MeshLODIndex) && InstanceBindingType == EHairBindingType::Skinning)
+					if (BindingType == EHairBindingType::Skinning || bGlobalDeformationEnable)
 					{
+						check(Instance->Strands.HasValidRootData());
+						check(Instance->Strands.DeformedRootResource->IsValid(MeshLODIndex));
+
 						AddHairStrandUpdateMeshTrianglesPass(
 							GraphBuilder, 
 							ShaderMap, 
@@ -210,8 +219,10 @@ static void RunInternalHairStrandsInterpolation(
 							Instance->Strands.DeformedRootResource,
 							Instance->Strands.DeformedResource);
 					}
-					else if (Instance->Strands.HasValidData())
+					else if (bSimulationEnable)
 					{
+						check(Instance->Strands.HasValidData());
+
 						AddHairStrandUpdatePositionOffsetPass(
 							GraphBuilder,
 							ShaderMap,
@@ -222,12 +233,15 @@ static void RunInternalHairStrandsInterpolation(
 				}
 				else if (InstanceGeometryType == EHairGeometryType::Cards)
 				{
-					const uint32 HairLODIndex = Instance->HairGroupPublicData->LODIndex;
 					if (Instance->Cards.IsValid(HairLODIndex))
 					{
 						FHairGroupInstance::FCards::FLOD& CardsInstance = Instance->Cards.LODs[HairLODIndex];
-						if (CardsInstance.Guides.IsValid() && CardsInstance.Guides.HasValidRootData() && CardsInstance.Guides.DeformedRootResource->IsValid(MeshLODIndex) && NeedsUpdateCardsMeshTriangles() && InstanceBindingType == EHairBindingType::Skinning)
+						if (BindingType == EHairBindingType::Skinning || bGlobalDeformationEnable)
 						{
+							check(CardsInstance.Guides.IsValid());
+							check(CardsInstance.Guides.HasValidRootData());
+							check(CardsInstance.Guides.DeformedRootResource->IsValid(MeshLODIndex));
+
 							AddHairStrandUpdateMeshTrianglesPass(
 								GraphBuilder,
 								ShaderMap,
@@ -244,8 +258,9 @@ static void RunInternalHairStrandsInterpolation(
 								CardsInstance.Guides.DeformedRootResource,
 								CardsInstance.Guides.DeformedResource);
 						}
-						else if (CardsInstance.Guides.IsValid())
+						else if (bSimulationEnable)
 						{
+							check(CardsInstance.Guides.IsValid());
 							AddHairStrandUpdatePositionOffsetPass(
 								GraphBuilder,
 								ShaderMap,
@@ -262,23 +277,18 @@ static void RunInternalHairStrandsInterpolation(
 			}
 			else if (EHairStrandsInterpolationType::SimulationStrands == Type)
 			{
-				if (Instance->Guides.IsValid() && Instance->Guides.HasValidRootData() && Instance->Guides.DeformedRootResource->IsValid(MeshLODIndex) && 
-					InstanceBindingType == EHairBindingType::Skinning &&
-					(Instance->Guides.bIsSimulationEnable || Instance->Guides.bHasGlobalInterpolation))
+				// Guide update need to run only if simulation is enabled, or if RBF is enabled (since RFB are transfer through guides)
+				if (bGlobalDeformationEnable || bSimulationEnable)
 				{
-					AddHairStrandUpdateMeshTrianglesPass(
-						GraphBuilder,
-						ShaderMap,
-						Instance->Debug.MeshLODIndex,
-						HairStrandsTriangleType::DeformedPose,
-						MeshDataLOD,
-						Instance->Guides.RestRootResource,
-						Instance->Guides.DeformedRootResource);
-
-					if((Instance->GeometryType == EHairGeometryType::Strands && Instance->Guides.bHasGlobalInterpolation) ||
-						Instance->GeometryType == EHairGeometryType::Meshes || Instance->GeometryType == EHairGeometryType::Cards)
+					check(Instance->Guides.IsValid());
+
+					if (BindingType == EHairBindingType::Skinning)
 					{
-						AddHairStrandInitMeshSamplesPass(
+						check(Instance->Guides.IsValid());
+						check(Instance->Guides.HasValidRootData());
+						check(Instance->Guides.DeformedRootResource->IsValid(MeshLODIndex));
+
+						AddHairStrandUpdateMeshTrianglesPass(
 							GraphBuilder,
 							ShaderMap,
 							Instance->Debug.MeshLODIndex,
@@ -287,30 +297,44 @@ static void RunInternalHairStrandsInterpolation(
 							Instance->Guides.RestRootResource,
 							Instance->Guides.DeformedRootResource);
 
-						AddHairStrandUpdateMeshSamplesPass(
+						if (bGlobalDeformationEnable)
+						{
+							AddHairStrandInitMeshSamplesPass(
+								GraphBuilder,
+								ShaderMap,
+								Instance->Debug.MeshLODIndex,
+								HairStrandsTriangleType::DeformedPose,
+								MeshDataLOD,
+								Instance->Guides.RestRootResource,
+								Instance->Guides.DeformedRootResource);
+
+							AddHairStrandUpdateMeshSamplesPass(
+								GraphBuilder,
+								ShaderMap,
+								Instance->Debug.MeshLODIndex,
+								MeshDataLOD,
+								Instance->Guides.RestRootResource,
+								Instance->Guides.DeformedRootResource);
+						}
+
+						AddHairStrandUpdatePositionOffsetPass(
 							GraphBuilder,
 							ShaderMap,
 							Instance->Debug.MeshLODIndex,
-							MeshDataLOD,
-							Instance->Guides.RestRootResource,
-							Instance->Guides.DeformedRootResource);
+							Instance->Guides.DeformedRootResource,
+							Instance->Guides.DeformedResource);
 					}
+					else if (bSimulationEnable)
+					{
+						check(Instance->Guides.IsValid());
 
-					AddHairStrandUpdatePositionOffsetPass(
-						GraphBuilder,
-						ShaderMap,
-						Instance->Debug.MeshLODIndex,
-						Instance->Guides.DeformedRootResource,
-						Instance->Guides.DeformedResource);
-				}
-				else if (Instance->Guides.IsValid())
-				{
-					AddHairStrandUpdatePositionOffsetPass(
-						GraphBuilder,
-						ShaderMap,
-						Instance->Debug.MeshLODIndex,
-						nullptr,
-						Instance->Guides.DeformedResource);
+						AddHairStrandUpdatePositionOffsetPass(
+							GraphBuilder,
+							ShaderMap,
+							Instance->Debug.MeshLODIndex,
+							nullptr,
+							Instance->Guides.DeformedResource);
+					}
 				}
 			}
 		}
@@ -506,7 +530,10 @@ static void RunHairBufferSwap(EWorldType::Type WorldType, const TArray<const FSc
 			{
 				if (Instance->Cards.IsValid(LODIt))
 				{
-					Instance->Cards.LODs[LODIt].DeformedResource->SwapBuffer();
+					if (Instance->Cards.LODs[LODIt].DeformedResource)
+					{
+						Instance->Cards.LODs[LODIt].DeformedResource->SwapBuffer();
+					}
 					if (Instance->Cards.LODs[LODIt].Guides.IsValid())
 					{
 						Instance->Cards.LODs[LODIt].Guides.DeformedResource->SwapBuffer();
@@ -518,7 +545,10 @@ static void RunHairBufferSwap(EWorldType::Type WorldType, const TArray<const FSc
 			{
 				if (Instance->Meshes.IsValid(LODIt))
 				{
-					Instance->Meshes.LODs[LODIt].DeformedResource->SwapBuffer();
+					if (Instance->Meshes.LODs[LODIt].DeformedResource)
+					{
+						Instance->Meshes.LODs[LODIt].DeformedResource->SwapBuffer();
+					}
 				}
 			}
 
diff --git a/Engine/Plugins/Runtime/HairStrands/Source/HairStrandsCore/Private/GroomResources.cpp b/Engine/Plugins/Runtime/HairStrands/Source/HairStrandsCore/Private/GroomResources.cpp
--- a/Engine/Plugins/Runtime/HairStrands/Source/HairStrandsCore/Private/GroomResources.cpp
+++ b/Engine/Plugins/Runtime/HairStrands/Source/HairStrandsCore/Private/GroomResources.cpp
@@ -393,8 +393,8 @@ void FHairCardsProceduralResource::ReleaseRHI()
 
 /////////////////////////////////////////////////////////////////////////////////////////
 
-FHairCardsDeformedResource::FHairCardsDeformedResource(const FHairCardsDatas::FRenderData& HairCardsRenderData, bool bInInitializedData, bool bInDynamic) :
-	RenderData(HairCardsRenderData), bInitializedData(!bInDynamic || bInInitializedData), bDynamic(bInDynamic)
+FHairCardsDeformedResource::FHairCardsDeformedResource(const FHairCardsDatas::FRenderData& HairCardsRenderData, bool bInInitializedData) :
+	RenderData(HairCardsRenderData), bInitializedData(bInInitializedData)
 {}
 
 void FHairCardsDeformedResource::InitRHI()
@@ -408,18 +408,12 @@ void FHairCardsDeformedResource::InitRHI()
 	if (bInitializedData)
 	{
 		InternalCreateVertexBufferRDG<FHairCardsPositionFormat>(GraphBuilder, RenderData.Positions, DeformedPositionBuffer[0], TEXT("Hair.CardsDeformed_Position0"));
-		if (bDynamic)
-		{
-			InternalCreateVertexBufferRDG<FHairCardsPositionFormat>(GraphBuilder, RenderData.Positions, DeformedPositionBuffer[1], TEXT("Hair.CardsDeformed_Position1"));
-		}
+		InternalCreateVertexBufferRDG<FHairCardsPositionFormat>(GraphBuilder, RenderData.Positions, DeformedPositionBuffer[1], TEXT("Hair.CardsDeformed_Position1"));
 	}
 	else
 	{
 		InternalCreateVertexBufferRDG<FHairCardsPositionFormat>(GraphBuilder, VertexCount, DeformedPositionBuffer[0], TEXT("Hair.CardsDeformed_Position0"));
-		if (bDynamic)
-		{
-			InternalCreateVertexBufferRDG<FHairCardsPositionFormat>(GraphBuilder, VertexCount, DeformedPositionBuffer[1], TEXT("Hair.CardsDeformed_Position1"));
-		}
+		InternalCreateVertexBufferRDG<FHairCardsPositionFormat>(GraphBuilder, VertexCount, DeformedPositionBuffer[1], TEXT("Hair.CardsDeformed_Position1"));
 	}
 	GraphBuilder.Execute();
 }
@@ -427,10 +421,7 @@ void FHairCardsDeformedResource::InitRHI()
 void FHairCardsDeformedResource::ReleaseRHI()
 {
 	DeformedPositionBuffer[0].Release();
-	if (bDynamic)
-	{
-		DeformedPositionBuffer[1].Release();
-	}
+	DeformedPositionBuffer[1].Release();
 }
 
 /////////////////////////////////////////////////////////////////////////////////////////
@@ -482,8 +473,8 @@ void FHairMeshesRestResource::ReleaseResource()
 
 /////////////////////////////////////////////////////////////////////////////////////////
 
-FHairMeshesDeformedResource::FHairMeshesDeformedResource(const FHairMeshesDatas::FRenderData& HairMeshesData, bool bInInitializedData, bool bInDynamic) :
-	RenderData(HairMeshesData), bInitializedData(!bInDynamic || bInInitializedData), bDynamic(bInDynamic)
+FHairMeshesDeformedResource::FHairMeshesDeformedResource(const FHairMeshesDatas::FRenderData& HairMeshesData, bool bInInitializedData) :
+	RenderData(HairMeshesData), bInitializedData(bInInitializedData)
 {}
 
 void FHairMeshesDeformedResource::InitRHI()
@@ -497,18 +488,12 @@ void FHairMeshesDeformedResource::InitRHI()
 	if (bInitializedData)
 	{
 		InternalCreateVertexBufferRDG<FHairCardsPositionFormat>(GraphBuilder, RenderData.Positions, DeformedPositionBuffer[0], TEXT("Hair.MeshesDeformed_Positions0"));
-		if (bDynamic)
-		{
-			InternalCreateVertexBufferRDG<FHairCardsPositionFormat>(GraphBuilder, RenderData.Positions, DeformedPositionBuffer[1], TEXT("Hair.MeshesDeformed_Positions1"));
-		}
+		InternalCreateVertexBufferRDG<FHairCardsPositionFormat>(GraphBuilder, RenderData.Positions, DeformedPositionBuffer[1], TEXT("Hair.MeshesDeformed_Positions1"));
 	}
 	else
 	{
 		InternalCreateVertexBufferRDG<FHairCardsPositionFormat>(GraphBuilder, VertexCount, DeformedPositionBuffer[0], TEXT("Hair.MeshesDeformed_Positions0"));
-		if (bDynamic)
-		{
-			InternalCreateVertexBufferRDG<FHairCardsPositionFormat>(GraphBuilder, VertexCount, DeformedPositionBuffer[1], TEXT("Hair.MeshesDeformed_Positions1"));
-		}
+		InternalCreateVertexBufferRDG<FHairCardsPositionFormat>(GraphBuilder, VertexCount, DeformedPositionBuffer[1], TEXT("Hair.MeshesDeformed_Positions1"));
 	}
 
 	GraphBuilder.Execute();
@@ -517,10 +502,7 @@ void FHairMeshesDeformedResource::InitRHI()
 void FHairMeshesDeformedResource::ReleaseRHI()
 {
 	DeformedPositionBuffer[0].Release();
-	if (bDynamic)
-	{
-		DeformedPositionBuffer[1].Release();
-	}
+	DeformedPositionBuffer[1].Release();
 }
 
 /////////////////////////////////////////////////////////////////////////////////////////
@@ -536,24 +518,58 @@ void FHairStrandsRestResource::InitRHI()
 	FRHICommandListImmediate& RHICmdList = FRHICommandListExecutor::GetImmediateCommandList();
 	FRDGBuilder GraphBuilder(RHICmdList);
 
-	InternalCreateVertexBufferRDG<FHairStrandsPositionFormat>(GraphBuilder, RenderData.Positions, RestPositionBuffer, TEXT("Hair.StrandsRest_RestPositionBuffer"));
+	InternalCreateVertexBufferRDG<FHairStrandsPositionFormat>(GraphBuilder, RenderData.Positions, RestPositionBuffer, TEXT("Hair.StrandsRest_PositionBuffer"));
 	InternalCreateVertexBufferRDG<FHairStrandsAttributeFormat>(GraphBuilder, RenderData.Attributes, AttributeBuffer, TEXT("Hair.StrandsRest_AttributeBuffer"));
 	InternalCreateVertexBufferRDG<FHairStrandsMaterialFormat>(GraphBuilder, RenderData.Materials, MaterialBuffer, TEXT("Hair.StrandsRest_MaterialBuffer"));
 
+	TArray<FVector4> RestOffset;
+	RestOffset.Add(PositionOffset);
+	InternalCreateVertexBufferRDG<FHairStrandsPositionOffsetFormat>(GraphBuilder, RestOffset, PositionOffsetBuffer, TEXT("Hair.StrandsRest_PositionOffsetBuffer"));
+	
 	GraphBuilder.Execute();
 }
 
+void AddHairTangentPass(
+	FRDGBuilder& GraphBuilder,
+	FGlobalShaderMap* ShaderMap,
+	uint32 VertexCount,
+	FHairGroupPublicData* HairGroupPublicData,
+	FRDGBufferSRVRef PositionBuffer,
+	FRDGImportedBuffer OutTangentBuffer);
+
+FRDGExternalBuffer FHairStrandsRestResource::GetTangentBuffer(FRDGBuilder& GraphBuilder, FGlobalShaderMap* ShaderMap)
+{
+	// Lazy allocation and update
+	if (TangentBuffer.Buffer == nullptr)
+	{
+		const uint32 VertexCount = RenderData.Positions.Num();
+		InternalCreateVertexBufferRDG<FHairStrandsTangentFormat>(GraphBuilder, VertexCount * FHairStrandsTangentFormat::ComponentCount, TangentBuffer, TEXT("Hair.StrandsRest_TangentBuffer"));
+
+		AddHairTangentPass(
+			GraphBuilder,
+			ShaderMap,
+			VertexCount,
+			nullptr,
+			RegisterAsSRV(GraphBuilder, RestPositionBuffer),
+			Register(GraphBuilder, TangentBuffer, ERDGImportedBufferFlags::CreateUAV));
+	}
+
+	return TangentBuffer;
+}
+
 void FHairStrandsRestResource::ReleaseRHI()
 {
 	RestPositionBuffer.Release();
+	PositionOffsetBuffer.Release();
 	AttributeBuffer.Release();
 	MaterialBuffer.Release();
+	TangentBuffer.Release();
 }
 
 /////////////////////////////////////////////////////////////////////////////////////////
 
-FHairStrandsDeformedResource::FHairStrandsDeformedResource(const FHairStrandsDatas::FRenderData& HairStrandRenderData, bool bInInitializedData, bool bInDynamic, const FVector& InDefaultOffset) :
-	RenderData(HairStrandRenderData), bInitializedData(bInInitializedData), bDynamic(bInDynamic), DefaultOffset(InDefaultOffset)
+FHairStrandsDeformedResource::FHairStrandsDeformedResource(const FHairStrandsDatas::FRenderData& HairStrandRenderData, bool bInInitializedData, const FVector& InDefaultOffset) :
+	RenderData(HairStrandRenderData), bInitializedData(bInInitializedData), DefaultOffset(InDefaultOffset)
 {}
 
 void FHairStrandsDeformedResource::InitRHI()
@@ -567,18 +583,12 @@ void FHairStrandsDeformedResource::InitRHI()
 	if (bInitializedData)
 	{
 		InternalCreateVertexBufferRDG<FHairStrandsPositionFormat>(GraphBuilder, RenderData.Positions, DeformedPositionBuffer[0], TEXT("Hair.StrandsDeformed_DeformedPositionBuffer0")); // , bDynamic ? ERHIAccess::UAVCompute : ERHIAccess::SRVMask
-		if (bDynamic)
-		{
-			InternalCreateVertexBufferRDG<FHairStrandsPositionFormat>(GraphBuilder, RenderData.Positions, DeformedPositionBuffer[1], TEXT("Hair.StrandsDeformed_DeformedPositionBuffer1")); // , ERHIAccess::UAVCompute
-		}
+		InternalCreateVertexBufferRDG<FHairStrandsPositionFormat>(GraphBuilder, RenderData.Positions, DeformedPositionBuffer[1], TEXT("Hair.StrandsDeformed_DeformedPositionBuffer1")); // , ERHIAccess::UAVCompute
 	}
 	else
 	{
 		InternalCreateVertexBufferRDG<FHairStrandsPositionFormat>(GraphBuilder, VertexCount, DeformedPositionBuffer[0], TEXT("Hair.StrandsDeformed_DeformedPositionBuffer0"));
-		if (bDynamic)
-		{
-			InternalCreateVertexBufferRDG<FHairStrandsPositionFormat>(GraphBuilder, VertexCount, DeformedPositionBuffer[1], TEXT("Hair.StrandsDeformed_DeformedPositionBuffer1"));
-		}
+		InternalCreateVertexBufferRDG<FHairStrandsPositionFormat>(GraphBuilder, VertexCount, DeformedPositionBuffer[1], TEXT("Hair.StrandsDeformed_DeformedPositionBuffer1"));
 	}
 	InternalCreateVertexBufferRDG<FHairStrandsTangentFormat>(GraphBuilder, VertexCount * FHairStrandsTangentFormat::ComponentCount, TangentBuffer, TEXT("Hair.StrandsDeformed_TangentBuffer"));
 
@@ -593,30 +603,13 @@ void FHairStrandsDeformedResource::InitRHI()
 void FHairStrandsDeformedResource::ReleaseRHI()
 {
 	DeformedPositionBuffer[0].Release();
-	if (bDynamic)
-	{
-		DeformedPositionBuffer[1].Release();
-	}
+	DeformedPositionBuffer[1].Release();
 	TangentBuffer.Release();
 
 	DeformedOffsetBuffer[0].Release();
 	DeformedOffsetBuffer[1].Release();
 }
 
-bool FHairStrandsDeformedResource::NeedsToUpdateTangent()
-{ 
-	if (bDynamic)
-	{
-		return true;
-	}
-	else
-	{ 
-		const bool bInit = bInitializedTangent;
-		bInitializedTangent = false;
-		return bInit;
-	} 
-}
-
 /////////////////////////////////////////////////////////////////////////////////////////
 
 FArchive& operator<<(FArchive& Ar, FHairStrandsClusterCullingData::FHairClusterInfo& Info);
diff --git a/Engine/Plugins/Runtime/HairStrands/Source/HairStrandsCore/Private/HairCardsVertexFactory.cpp b/Engine/Plugins/Runtime/HairStrands/Source/HairStrandsCore/Private/HairCardsVertexFactory.cpp
--- a/Engine/Plugins/Runtime/HairStrands/Source/HairStrandsCore/Private/HairCardsVertexFactory.cpp
+++ b/Engine/Plugins/Runtime/HairStrands/Source/HairStrandsCore/Private/HairCardsVertexFactory.cpp
@@ -37,14 +37,23 @@ FHairCardsUniformBuffer CreateHairCardsVFUniformBuffer(
 	{
 		const FHairGroupInstance::FCards::FLOD& LOD = Instance->Cards.LODs[LODIndex];
 
+		// Cards atlas UV are inverted so fetching needs to be inverted on the y-axis
+		UniformParameters.bInvertUV = LOD.RestResource->bInvertUV;
+
 		// When the geometry is not-dynamic (no binding to skeletal mesh, no simulation), only a single vertex buffer is allocated. 
 		// In this case we force the buffer index to 0
-		const bool bIsDynamic = LOD.DeformedResource->bDynamic;
+		const bool bIsDynamic = LOD.DeformedResource != nullptr;
+		if (bIsDynamic)
+		{
+			UniformParameters.PositionBuffer = LOD.DeformedResource->DeformedPositionBuffer[bIsDynamic ? Current : 0].SRV.GetReference();
+			UniformParameters.PreviousPositionBuffer = LOD.DeformedResource->DeformedPositionBuffer[bIsDynamic ? 1u-Current : 0].SRV.GetReference();
+		}
+		else
+		{
+			UniformParameters.PositionBuffer = LOD.RestResource->RestPositionBuffer.ShaderResourceViewRHI;
+			UniformParameters.PreviousPositionBuffer = LOD.RestResource->RestPositionBuffer.ShaderResourceViewRHI;
+		}
 
-		// Cards atlas UV are inverted so fetching needs to be inverted on the y-axis
-		UniformParameters.bInvertUV = LOD.RestResource->bInvertUV;
-		UniformParameters.PositionBuffer = LOD.DeformedResource->DeformedPositionBuffer[bIsDynamic ? Current : 0].SRV.GetReference();
-		UniformParameters.PreviousPositionBuffer = LOD.DeformedResource->DeformedPositionBuffer[bIsDynamic ? 1u-Current : 0].SRV.GetReference();
 		UniformParameters.NormalsBuffer = LOD.RestResource->NormalsBuffer.ShaderResourceViewRHI.GetReference();
 		UniformParameters.UVsBuffer = LOD.RestResource->UVsBuffer.ShaderResourceViewRHI.GetReference();
 
@@ -65,12 +74,20 @@ FHairCardsUniformBuffer CreateHairCardsVFUniformBuffer(
 
 		// When the geometry is not-dynamic (no binding to skeletal mesh, no simulation), only a single vertex buffer is allocated. 
 		// In this case we force the buffer index to 0
-		const bool bIsDynamic = LOD.DeformedResource->bDynamic;
+		const bool bIsDynamic = LOD.DeformedResource != nullptr;
+		if (bIsDynamic)
+		{
+			UniformParameters.PositionBuffer = LOD.DeformedResource->DeformedPositionBuffer[bIsDynamic ? Current : 0].SRV.GetReference();
+			UniformParameters.PreviousPositionBuffer = LOD.DeformedResource->DeformedPositionBuffer[bIsDynamic ? 1u - Current : 0].SRV.GetReference();
+		}
+		else
+		{
+			UniformParameters.PositionBuffer = LOD.RestResource->RestPositionBuffer.ShaderResourceViewRHI;
+			UniformParameters.PreviousPositionBuffer = LOD.RestResource->RestPositionBuffer.ShaderResourceViewRHI;
+		}
 
 		// Meshes UV are not inverted so no need to invert the y-axis
 		UniformParameters.bInvertUV = 0;
-		UniformParameters.PositionBuffer = LOD.DeformedResource->DeformedPositionBuffer[bIsDynamic ? Current : 0].SRV.GetReference();
-		UniformParameters.PreviousPositionBuffer = LOD.DeformedResource->DeformedPositionBuffer[bIsDynamic ? 1u-Current : 0].SRV.GetReference();
 		UniformParameters.NormalsBuffer = LOD.RestResource->NormalsBuffer.ShaderResourceViewRHI.GetReference();
 		UniformParameters.UVsBuffer = LOD.RestResource->UVsBuffer.ShaderResourceViewRHI.GetReference();
 
@@ -153,14 +170,16 @@ public:
 		if (InstanceGeometryType == EHairGeometryType::Cards)
 		{
 			const FHairGroupInstance::FCards::FLOD& LOD = Instance->Cards.LODs[LODIndex];
-			const uint32 UniformIndex = LOD.DeformedResource->GetIndex(FHairCardsDeformedResource::Current);
+			const bool bIsDynamic = LOD.DeformedResource != nullptr;
+			const uint32 UniformIndex = bIsDynamic ? LOD.DeformedResource->GetIndex(FHairCardsDeformedResource::Current) : 0;
 			check(LOD.UniformBuffer[UniformIndex]);
 			VertexFactoryUniformBuffer = LOD.UniformBuffer[UniformIndex];
 		}
 		else if (InstanceGeometryType == EHairGeometryType::Meshes)
 		{
 			const FHairGroupInstance::FMeshes::FLOD& LOD = Instance->Meshes.LODs[LODIndex];
-			const uint32 UniformIndex = LOD.DeformedResource->GetIndex(FHairMeshesDeformedResource::Current);
+			const bool bIsDynamic = LOD.DeformedResource != nullptr;
+			const uint32 UniformIndex = bIsDynamic ? LOD.DeformedResource->GetIndex(FHairMeshesDeformedResource::Current) : 0;
 			check(LOD.UniformBuffer[UniformIndex]);
 			VertexFactoryUniformBuffer = LOD.UniformBuffer[UniformIndex];
 		}
diff --git a/Engine/Plugins/Runtime/HairStrands/Source/HairStrandsCore/Private/HairStrandsRendering.cpp b/Engine/Plugins/Runtime/HairStrands/Source/HairStrandsCore/Private/HairStrandsRendering.cpp
--- a/Engine/Plugins/Runtime/HairStrands/Source/HairStrandsCore/Private/HairStrandsRendering.cpp
+++ b/Engine/Plugins/Runtime/HairStrands/Source/HairStrandsCore/Private/HairStrandsRendering.cpp
@@ -23,9 +23,6 @@
 #include "GroomManager.h"
 #include "GroomInstance.h"
 
-static int32 GHairDeformationType = 0;
-static FAutoConsoleVariableRef CVarHairDeformationType(TEXT("r.HairStrands.DeformationType"), GHairDeformationType, TEXT("Type of procedural deformation applied on hair strands (0:use simulation's output, 1:use rest strands, 2: use rest guides, 3:wave pattern, 4:follow root normal)"));
-
 static float GHairRaytracingRadiusScale = 0;
 static FAutoConsoleVariableRef CVarHairRaytracingRadiusScale(TEXT("r.HairStrands.RaytracingRadiusScale"), GHairRaytracingRadiusScale, TEXT("Override the per instance scale factor for raytracing hair strands geometry (0: disabled, >0:enabled)"));
 
@@ -102,23 +99,10 @@ inline uint32 GetGroupSizePermutation(uint32 GroupSize)
 enum class EDeformationType : uint8
 {
 	Simulation,		// Use the output of the hair simulation
-	RestStrands,	// Use the rest strands position (no weighted interpolation)
 	RestGuide,		// Use the rest guide as input of the interpolation (no deformation), only weighted interpolation
 	OffsetGuide		// Offset the guides
 };
 
-static EDeformationType GetDeformationType()
-{
-	switch (GHairDeformationType)
-	{
-	case 0: return EDeformationType::Simulation;
-	case 1: return EDeformationType::RestStrands;
-	case 2: return EDeformationType::RestGuide;
-	}
-
-	return EDeformationType::Simulation;
-}
-
 ////////////////////////////////////////////////////////////////////////////////////////////////////////////////////////
 class FTransferVelocityPassCS : public FGlobalShader
 {
@@ -323,36 +307,6 @@ static void AddDeformSimHairStrandsPass(
 
 ///////////////////////////////////////////////////////////////////////////////////////////////////
 
-struct FCopyBufferParameter
-{
-	BEGIN_SHADER_PARAMETER_STRUCT(FParameters, )
-		RDG_BUFFER_ACCESS(DstBuffer, ERHIAccess::CopyDest)
-	END_SHADER_PARAMETER_STRUCT()
-};
-
-static void AddCopyBufferPass(
-	FRDGBuilder& GraphBuilder,
-	FRDGBufferRef DstBuffer,
-	FBufferRHIRef SrcBuffer)
-{
-
-	FCopyBufferParameter::FParameters* Parameters = GraphBuilder.AllocParameters<FCopyBufferParameter::FParameters>();
-	Parameters->DstBuffer = DstBuffer;
-
-	const uint32 TotalByteCount = DstBuffer->Desc.BytesPerElement * DstBuffer->Desc.NumElements;
-
-	GraphBuilder.AddPass(
-		RDG_EVENT_NAME("HairCopyPosition"),
-		Parameters,
-		ERDGPassFlags::Copy,
-		[Parameters, SrcBuffer, TotalByteCount](FRHICommandList& RHICmdList)
-		{
-			RHICmdList.CopyBufferRegion(Parameters->DstBuffer->GetRHI(), 0, SrcBuffer, 0, TotalByteCount);
-		});
-}
-
-///////////////////////////////////////////////////////////////////////////////////////////////////
-
 struct FRDGHairStrandsCullingData
 {
 	bool bCullingResultAvailable = false;
@@ -508,14 +462,14 @@ static void AddHairStrandsInterpolationPass(
 	FRDGBuilder& GraphBuilder,
 	FGlobalShaderMap* ShaderMap,
 	const FShaderDrawDebugData* ShaderDrawData,
-	FHairGroupInstance* Instance,
+	const FHairGroupInstance* Instance,
 	const uint32 VertexCount,
 	const FHairScaleAndClipDesc ScaleAndClipDesc,
 	const int32 MeshLODIndex,
 	const bool bPatchedAttributeBuffer,
 	const EHairInterpolationType HairInterpolationType,
 	const EHairGeometryType InstanceGeometryType,
-	FRDGHairStrandsCullingData& CullingData,
+	const FRDGHairStrandsCullingData& CullingData,
 	const FVector& InRenderHairWorldOffset,
 	const FVector& InSimHairWorldOffset,
 	const FRDGBufferSRVRef& OutRenHairPositionOffsetBuffer,
@@ -796,7 +750,7 @@ static void AddHairClusterAABBPass(
 
 	FRDGHairStrandsCullingData& ClusterAABBData,
 
-	FRDGImportedBuffer& RenderPositionBuffer)
+	FRDGBufferSRVRef RenderPositionBufferSRV)
 {
 	const uint32 GroupSize = ComputeGroupSize();
 	const FIntVector DispatchCount = ComputeDispatchGroupCount2D(ClusterData.ClusterCount);
@@ -808,7 +762,7 @@ static void AddHairClusterAABBPass(
 	Parameters->DispatchCountX = DispatchCount.X;
 	Parameters->ClusterCount = ClusterData.ClusterCount;
 	Parameters->LocalToWorldMatrix = InRenLocalToWorld.ToMatrixWithScale();
-	Parameters->RenderDeformedPositionBuffer = RenderPositionBuffer.SRV;
+	Parameters->RenderDeformedPositionBuffer = RenderPositionBufferSRV;
 	Parameters->RenderDeformedOffsetBuffer = RenderDeformedOffsetBuffer;
 	Parameters->ClusterVertexIdBuffer = RegisterAsSRV(GraphBuilder, *ClusterData.ClusterVertexIdBuffer);
 	Parameters->ClusterIdBuffer = GraphBuilder.CreateSRV(ClusterIdBuffer, PF_R32_UINT);
@@ -987,7 +941,7 @@ public:
 
 IMPLEMENT_GLOBAL_SHADER(FHairTangentCS, "/Engine/Private/HairStrands/HairStrandsTangent.usf", "MainCS", SF_Compute);
 
-static void AddHairTangentPass(
+void AddHairTangentPass(
 	FRDGBuilder& GraphBuilder,
 	FGlobalShaderMap* ShaderMap,
 	uint32 VertexCount,
@@ -997,7 +951,7 @@ static void AddHairTangentPass(
 {
 	const uint32 GroupSize = ComputeGroupSize();
 	const FIntVector DispatchCount = ComputeDispatchCount(VertexCount, GroupSize);
-	const bool bCullingEnable = HairGroupPublicData->GetCullingResultAvailable();
+	const bool bCullingEnable = HairGroupPublicData && HairGroupPublicData->GetCullingResultAvailable();
 
 	FHairTangentCS::FParameters* Parameters = GraphBuilder.AllocParameters<FHairTangentCS::FParameters>();
 	Parameters->PositionBuffer = PositionBuffer;
@@ -1297,7 +1251,10 @@ void RegisterClusterData(FHairGroupInstance* Instance, FHairStrandClusterData* I
 
 EHairStrandsDebugMode GetHairStrandsGeometryDebugMode(FHairGroupInstance* Instance)
 {
+	// 1. Use the per-instance debug mode by default, and fallback on the global debug otherwise
 	EHairStrandsDebugMode DebugMode = Instance->Debug.DebugMode != EHairStrandsDebugMode::NoneDebug ? Instance->Debug.DebugMode : GetHairStrandsDebugStrandsMode();
+
+	// 2. If groom does not have guides (e.g., because it does not have simulation/RBF data), then we fallback to default view mode
 	if ((DebugMode == EHairStrandsDebugMode::SimHairStrands ||
 		 DebugMode == EHairStrandsDebugMode::RenderHairStrands ||
 		 DebugMode == EHairStrandsDebugMode::RenderVisCluster)
@@ -1314,11 +1271,22 @@ FHairScaleAndClipDesc ComputeHairScaleAndClipDesc(FHairGroupInstance* Instance)
 	Out.bEnable = true;
 	Out.InHairLength = Instance->Strands.Data->StrandsCurves.MaxLength;
 	Out.InHairRadius = Instance->Strands.Modifier.HairWidth * 0.5f;
-	Out.OutHairRadius = (GStrandHairWidth > 0 ? GStrandHairWidth : Instance->Strands.Modifier.HairWidth) * 0.5f;
-	Out.MaxOutHairRadius = Out.OutHairRadius * FMath::Max(1.f, FMath::Max(Instance->Strands.Modifier.HairRootScale, Instance->Strands.Modifier.HairTipScale));
-	Out.HairRadiusRootScale = Instance->Strands.Modifier.HairRootScale;
-	Out.HairRadiusTipScale = Instance->Strands.Modifier.HairTipScale;
-	Out.HairLengthClip = FMath::Clamp(Instance->Strands.Modifier.HairClipLength / Instance->Strands.Data->StrandsCurves.MaxLength, 0.f, 1.f);
+
+	Out.OutHairRadius = Out.InHairRadius;
+	Out.MaxOutHairRadius = Out.OutHairRadius;
+	Out.HairRadiusRootScale = 1.f;
+	Out.HairRadiusTipScale = 1.f;
+	Out.HairLengthClip = 1.f;
+
+	const bool bSupportDeformation = Instance->Strands.DeformedResource != nullptr;
+	if (bSupportDeformation)
+	{
+		Out.OutHairRadius = (GStrandHairWidth > 0 ? GStrandHairWidth : Instance->Strands.Modifier.HairWidth) * 0.5f;
+		Out.MaxOutHairRadius = Out.OutHairRadius * FMath::Max(1.f, FMath::Max(Instance->Strands.Modifier.HairRootScale, Instance->Strands.Modifier.HairTipScale));
+		Out.HairRadiusRootScale = Instance->Strands.Modifier.HairRootScale;
+		Out.HairRadiusTipScale = Instance->Strands.Modifier.HairTipScale;
+		Out.HairLengthClip = FMath::Clamp(Instance->Strands.Modifier.HairClipLength / Instance->Strands.Data->StrandsCurves.MaxLength, 0.f, 1.f);
+	}
 	
 	return Out;
 }
@@ -1328,6 +1296,7 @@ bool NeedsPatchAttributeBuffer(EHairStrandsDebugMode DebugMode)
 	return DebugMode == EHairStrandsDebugMode::RenderHairStrands || DebugMode == EHairStrandsDebugMode::RenderVisCluster;
 }
 
+// Compute/Update the hair strands description which will be used for rendering (VF) / voxelization & co.
 FHairGroupPublicData::FVertexFactoryInput ComputeHairStrandsVertexInputData(FHairGroupInstance* Instance)
 {
 	FHairGroupPublicData::FVertexFactoryInput OutVFInput;
@@ -1339,8 +1308,12 @@ FHairGroupPublicData::FVertexFactoryInput ComputeHairStrandsVertexInputData(FHai
 	const EHairStrandsDebugMode DebugMode = GetHairStrandsGeometryDebugMode(Instance);
 	const bool bDebugModePatchedAttributeBuffer = NeedsPatchAttributeBuffer(DebugMode);
 
+	const bool bSupportDeformation = Instance->Strands.DeformedResource != nullptr;
+
+	// 1. Guides
 	if (DebugMode == EHairStrandsDebugMode::SimHairStrands)
 	{
+
 		OutVFInput.Strands.PositionBuffer = Instance->Guides.DeformedResource->GetBuffer(FHairStrandsDeformedResource::EFrameType::Current).SRV;
 		OutVFInput.Strands.PrevPositionBuffer = Instance->Guides.DeformedResource->GetBuffer(FHairStrandsDeformedResource::EFrameType::Previous).SRV;
 		OutVFInput.Strands.TangentBuffer = Instance->Guides.DeformedResource->TangentBuffer.SRV;
@@ -1360,8 +1333,10 @@ FHairGroupPublicData::FVertexFactoryInput ComputeHairStrandsVertexInputData(FHai
 		OutVFInput.Strands.bUseStableRasterization = Instance->Strands.Modifier.bUseStableRasterization;
 		OutVFInput.Strands.bScatterSceneLighting = Instance->Strands.Modifier.bScatterSceneLighting;
 	}
-	else
+	// 2. Render Strands with deformation
+	else if (bSupportDeformation)
 	{
+
 		OutVFInput.Strands.PositionBuffer = Instance->Strands.DeformedResource->GetBuffer(FHairStrandsDeformedResource::EFrameType::Current).SRV;
 		OutVFInput.Strands.PrevPositionBuffer = Instance->Strands.DeformedResource->GetBuffer(FHairStrandsDeformedResource::EFrameType::Previous).SRV;
 		OutVFInput.Strands.TangentBuffer = Instance->Strands.DeformedResource->TangentBuffer.SRV;
@@ -1380,6 +1355,28 @@ FHairGroupPublicData::FVertexFactoryInput ComputeHairStrandsVertexInputData(FHai
 		OutVFInput.Strands.bScatterSceneLighting = Instance->Strands.Modifier.bScatterSceneLighting;
 		OutVFInput.Strands.bUseStableRasterization = Instance->Strands.Modifier.bUseStableRasterization;
 	}
+	// 3. Render Strands in rest position: used when there are no skinning (rigid binding or no binding), no simulation, and no RBF
+	else
+	{
+
+		OutVFInput.Strands.PositionBuffer = Instance->Strands.RestResource->RestPositionBuffer.SRV;
+		OutVFInput.Strands.PrevPositionBuffer = Instance->Strands.RestResource->RestPositionBuffer.SRV;
+		OutVFInput.Strands.TangentBuffer = Instance->Strands.RestResource->TangentBuffer.SRV;
+		OutVFInput.Strands.AttributeBuffer = Instance->Strands.RestResource->AttributeBuffer.SRV;
+		OutVFInput.Strands.MaterialBuffer = Instance->Strands.RestResource->MaterialBuffer.SRV;
+
+		OutVFInput.Strands.PositionOffsetBuffer = Instance->Strands.RestResource->PositionOffsetBuffer.SRV;
+		OutVFInput.Strands.PrevPositionOffsetBuffer = Instance->Strands.RestResource->PositionOffsetBuffer.SRV;
+
+		OutVFInput.Strands.PositionOffset = Instance->Strands.RestResource->PositionOffset;
+		OutVFInput.Strands.PrevPositionOffset = Instance->Strands.RestResource->PositionOffset;
+		OutVFInput.Strands.VertexCount = Instance->Strands.RestResource->GetVertexCount();
+		OutVFInput.Strands.HairRadius = ScaleAndClipDesc.MaxOutHairRadius;
+		OutVFInput.Strands.HairLength = Instance->Strands.Modifier.HairLength;
+		OutVFInput.Strands.HairDensity = Instance->Strands.Modifier.HairShadowDensity;
+		OutVFInput.Strands.bScatterSceneLighting = Instance->Strands.Modifier.bScatterSceneLighting;
+		OutVFInput.Strands.bUseStableRasterization = Instance->Strands.Modifier.bUseStableRasterization;
+	}
 
 	return OutVFInput;
 }
@@ -1399,65 +1396,36 @@ void ComputeHairStrandsInterpolation(
 	FHairGroupInstance* Instance,
 	int32 MeshLODIndex,
 	FHairStrandClusterData* InClusterData)
-{
-	// Note: We are breaking this code up into several, larger for loops. In the previous version, the typical code path was:
-	// for each group:
-	//     AddClearClusterAABBPass()
-	//     AddHairStrandsInterpolationPass()
-	//     AddHairClusterAABBPass()
-	//     AddHairTangentPass()
-	//     AddGenerateRaytracingGeometryPass()
-	// 
-	// The problem is that it creates bubbles in the GPU, since each pass was dependent on the previous one. So it has been
-	// modified to be of the form:
-	//
-	// for each group:
-	//     AddClearClusterAABBPass()
-	// for each group:
-	//     AddHairStrandsInterpolationPass()
-	// ...
-
-	if (!Instance) return;
+{	
+	if (!Instance)
+	{
+		return;
+	}
 
 	// Reset
 	Instance->HairGroupPublicData->VFInput.Strands	= FHairGroupPublicData::FVertexFactoryInput::FStrands();
 	Instance->HairGroupPublicData->VFInput.Cards	= FHairGroupPublicData::FVertexFactoryInput::FCards();
 	Instance->HairGroupPublicData->VFInput.Meshes	= FHairGroupPublicData::FVertexFactoryInput::FMeshes();
-	const EHairGeometryType InstanceGeometryType = Instance->GeometryType;
-
-	DECLARE_GPU_STAT(HairStrandsInterpolation);
-	RDG_EVENT_SCOPE(GraphBuilder, "HairStrandsInterpolation");
-	RDG_GPU_STAT_SCOPE(GraphBuilder, HairStrandsInterpolation);
+	const EHairGeometryType InstanceGeometryType	= Instance->GeometryType;
 
-	// Debug mode:
-	// * None	: Display hair normally
-	// * Sim	: Show sim strands
-	// * Render : Show rendering strands with sim color influence
-	const EDeformationType DeformationType = GetDeformationType();
-	const EHairStrandsDebugMode DebugMode = GetHairStrandsGeometryDebugMode(Instance);
-	const bool bDebugModePatchedAttributeBuffer = NeedsPatchAttributeBuffer(DebugMode);
 
-	if (DeformationType != EDeformationType::RestStrands && DeformationType != EDeformationType::Simulation)
+	if (InstanceGeometryType == EHairGeometryType::Strands)
 	{
-		AddDeformSimHairStrandsPass(
-			GraphBuilder,
-			ShaderMap,
-			DeformationType,
-			MeshLODIndex,
-			Instance->Guides.RestResource->GetVertexCount(),
-			Instance->Guides.RestRootResource,
-			Instance->Guides.DeformedRootResource,
-			RegisterAsSRV(GraphBuilder, Instance->Guides.RestResource->RestPositionBuffer),
-			Register(GraphBuilder, Instance->Guides.DeformedResource->GetBuffer(FHairStrandsDeformedResource::Current), ERDGImportedBufferFlags::CreateUAV),
-			Instance->Guides.RestResource->PositionOffset,
-			RegisterAsSRV(GraphBuilder, Instance->Guides.DeformedResource->GetPositionOffsetBuffer(FHairStrandsDeformedResource::Current)),
-			Instance->Guides.bHasGlobalInterpolation);
-	}
+		DECLARE_GPU_STAT(HairStrandsInterpolation);
+		RDG_EVENT_SCOPE(GraphBuilder, "HairStrandsInterpolation");
+		RDG_GPU_STAT_SCOPE(GraphBuilder, HairStrandsInterpolation);
+
+		// Debug mode:
+		// * None	: Display hair normally
+		// * Sim	: Show sim strands
+		// * Render : Show rendering strands with sim color influence
+		const EHairStrandsDebugMode DebugMode = GetHairStrandsGeometryDebugMode(Instance);
+		const bool bDebugModePatchedAttributeBuffer = NeedsPatchAttributeBuffer(DebugMode);
+		const FHairScaleAndClipDesc ScaleAndClipDesc = ComputeHairScaleAndClipDesc(Instance);
 
-	if (DebugMode == EHairStrandsDebugMode::SimHairStrands && InstanceGeometryType == EHairGeometryType::Strands)
-	{
-		if (Instance->Guides.DeformedResource->NeedsToUpdateTangent())
+		if (DebugMode == EHairStrandsDebugMode::SimHairStrands)
 		{
+			check(Instance->Guides.DeformedResource);
 			AddHairTangentPass(
 				GraphBuilder,
 				ShaderMap,
@@ -1465,48 +1433,49 @@ void ComputeHairStrandsInterpolation(
 				Instance->HairGroupPublicData,
 				RegisterAsSRV(GraphBuilder, Instance->Guides.DeformedResource->GetBuffer(FHairStrandsDeformedResource::Current)),
 				Register(GraphBuilder, Instance->Guides.DeformedResource->TangentBuffer, ERDGImportedBufferFlags::CreateUAV));
-		}
-
-		Instance->HairGroupPublicData->VFInput = ComputeHairStrandsVertexInputData(Instance);
-	}
-	else if (InstanceGeometryType == EHairGeometryType::Strands)
-	{
-
-		FRDGHairStrandsCullingData CullingData = ImportCullingData(GraphBuilder, Instance->HairGroupPublicData);
 
+			Instance->HairGroupPublicData->VFInput = ComputeHairStrandsVertexInputData(Instance);
+		}
+		else
 		{
+			const bool bNeedDeformation = Instance->Strands.DeformedResource != nullptr;
+
+			// 1. Clear cluster AABBs (use for voxel allocation & optionally for culling)
+			FRDGHairStrandsCullingData CullingData = ImportCullingData(GraphBuilder, Instance->HairGroupPublicData);
 			AddClearClusterAABBPass(
 				GraphBuilder,
 				ShaderMap,
 				CullingData.ClusterCount,
 				CullingData.ClusterAABBBuffer,
 				CullingData.GroupAABBBuffer);
-		}
 
-		FRDGImportedBuffer Strands_DeformedPosition		= Register(GraphBuilder, Instance->Strands.DeformedResource->GetBuffer(FHairStrandsDeformedResource::Current), ERDGImportedBufferFlags::CreateViews);
-		FRDGImportedBuffer Strands_DeformedPrevPosition = Register(GraphBuilder, Instance->Strands.DeformedResource->GetBuffer(FHairStrandsDeformedResource::Previous), ERDGImportedBufferFlags::CreateViews);
-		FRDGImportedBuffer Strands_DeformedTangent		= Register(GraphBuilder, Instance->Strands.DeformedResource->TangentBuffer, ERDGImportedBufferFlags::CreateViews);		
-		#if WITH_EDITOR
-		FRDGImportedBuffer Strands_DebugAttributeReg	= Register(GraphBuilder, Instance->Strands.DebugAttributeBuffer, ERDGImportedBufferFlags::CreateUAV);
-		FRDGImportedBuffer*Strands_DebugAttribute		= &Strands_DebugAttributeReg;
-		#else
-		FRDGImportedBuffer*Strands_DebugAttribute		= nullptr;
-		#endif
-
-		// Note: This code needs to exactly match the values FHairScaleAndClipDesc set int the previous loop.
-		const FHairScaleAndClipDesc ScaleAndClipDesc = ComputeHairScaleAndClipDesc(Instance);
-		Instance->HairGroupPublicData->VFInput = ComputeHairStrandsVertexInputData(Instance);
-		{
+			// 2. Deform hair if needed (e.g.: skinning, simulation, RBF) and recompute tangent
+			FRDGBufferSRVRef Strands_PositionSRV = nullptr;
+			FRDGBufferSRVRef Strands_PositionOffsetSRV = nullptr;
+			if (bNeedDeformation)
 			{
+				FRDGImportedBuffer Strands_DeformedPosition = Register(GraphBuilder, Instance->Strands.DeformedResource->GetBuffer(FHairStrandsDeformedResource::Current), ERDGImportedBufferFlags::CreateViews);
+				FRDGImportedBuffer Strands_DeformedPrevPosition = Register(GraphBuilder, Instance->Strands.DeformedResource->GetBuffer(FHairStrandsDeformedResource::Previous), ERDGImportedBufferFlags::CreateViews);
+				FRDGImportedBuffer Strands_DeformedTangent = Register(GraphBuilder, Instance->Strands.DeformedResource->TangentBuffer, ERDGImportedBufferFlags::CreateViews);
+			#if WITH_EDITOR
+				FRDGImportedBuffer Strands_DebugAttributeReg = Register(GraphBuilder, Instance->Strands.DebugAttributeBuffer, ERDGImportedBufferFlags::CreateUAV);
+				FRDGImportedBuffer* Strands_DebugAttribute = &Strands_DebugAttributeReg;
+			#else
+				FRDGImportedBuffer* Strands_DebugAttribute = nullptr;
+			#endif
+				Strands_PositionSRV = Strands_DeformedPosition.SRV;
+				Strands_PositionOffsetSRV = RegisterAsSRV(GraphBuilder, Instance->Strands.DeformedResource->GetPositionOffsetBuffer(FHairStrandsDeformedResource::EFrameType::Current));
+
 				// "WITH_EDITOR && bDebugModePatchedAttributeBuffer" is a special path when visualizing groom within the groom editor, so that we can diplay guides even when there is no simulation or global interpolation enabled
 				const bool bValidGuide = Instance->Guides.bIsSimulationEnable || Instance->Guides.bHasGlobalInterpolation || (WITH_EDITOR && bDebugModePatchedAttributeBuffer);
 				const bool bHasSkinning = Instance->BindingType == EHairBindingType::Skinning;
 
+				// 2.1 Compute deformation position based on simulation/skinning/RBF
 				AddHairStrandsInterpolationPass(
 					GraphBuilder,
 					ShaderMap,
 					ShaderDrawData,
-					Instance, 
+					Instance,
 					Instance->Strands.RestResource->GetVertexCount(),
 					ScaleAndClipDesc,
 					MeshLODIndex,
@@ -1516,12 +1485,12 @@ void ComputeHairStrandsInterpolation(
 					CullingData,
 					Instance->Strands.RestResource->PositionOffset,
 					bValidGuide ? Instance->Guides.RestResource->PositionOffset : FVector::ZeroVector,
-					RegisterAsSRV(GraphBuilder, Instance->Strands.DeformedResource->GetPositionOffsetBuffer(FHairStrandsDeformedResource::Current)),
+					Strands_PositionOffsetSRV,
 					bValidGuide ? RegisterAsSRV(GraphBuilder, Instance->Guides.DeformedResource->GetPositionOffsetBuffer(FHairStrandsDeformedResource::Current)) : nullptr,
 					bHasSkinning ? Instance->Strands.RestRootResource : nullptr,
 					bHasSkinning && bValidGuide ? Instance->Guides.RestRootResource : nullptr,
 					bHasSkinning ? Instance->Strands.DeformedRootResource : nullptr,
-					bHasSkinning && bValidGuide  ? Instance->Guides.DeformedRootResource : nullptr,
+					bHasSkinning && bValidGuide ? Instance->Guides.DeformedRootResource : nullptr,
 					RegisterAsSRV(GraphBuilder, Instance->Strands.RestResource->RestPositionBuffer),
 					RegisterAsSRV(GraphBuilder, Instance->Strands.RestResource->AttributeBuffer),
 					RegisterAsSRV(GraphBuilder, Instance->Strands.InterpolationResource->Interpolation0Buffer),
@@ -1530,16 +1499,34 @@ void ComputeHairStrandsInterpolation(
 					bValidGuide ? RegisterAsSRV(GraphBuilder, Instance->Guides.DeformedResource->GetBuffer(FHairStrandsDeformedResource::Current)) : nullptr,
 					bValidGuide ? RegisterAsSRV(GraphBuilder, Instance->Guides.RestResource->AttributeBuffer) : nullptr,
 					Strands_DeformedPosition,
-					&Strands_DeformedPrevPosition, 
+					&Strands_DeformedPrevPosition,
 					Strands_DebugAttribute,
 					RegisterAsSRV(GraphBuilder, Instance->Strands.ClusterCullingResource->VertexToClusterIdBuffer),
 					RegisterAsSRV(GraphBuilder, Instance->Strands.InterpolationResource->SimRootPointIndexBuffer));
 
+				// 2.2 Update tangent data based on the deformed positions
+				AddHairTangentPass(
+					GraphBuilder,
+					ShaderMap,
+					Instance->Strands.RestResource->GetVertexCount(),
+					Instance->HairGroupPublicData,
+					Strands_DeformedPosition.SRV,
+					Strands_DeformedTangent);
+			}
+			else
+			{
+				Strands_PositionSRV = RegisterAsSRV(GraphBuilder, Instance->Strands.RestResource->RestPositionBuffer);
+				Strands_PositionOffsetSRV = RegisterAsSRV(GraphBuilder, Instance->Strands.RestResource->PositionOffsetBuffer);
+				// Generated the static tangent if they haven;t been generated yet
+				Instance->Strands.RestResource->GetTangentBuffer(GraphBuilder, ShaderMap);
+
 			}
-		}
 
-		{		
-			FHairStrandClusterData::FHairGroup& HairGroupCluster =  InClusterData->HairGroups[Instance->HairGroupPublicData->ClusterDataIndex];
+			// 2.1 Update the VF input with the update resources
+			Instance->HairGroupPublicData->VFInput = ComputeHairStrandsVertexInputData(Instance);
+
+			// 3. Compute cluster AABBs (used for LODing and voxelization)
+			FHairStrandClusterData::FHairGroup& HairGroupCluster = InClusterData->HairGroups[Instance->HairGroupPublicData->ClusterDataIndex];
 			if (HairGroupCluster.bVisible)
 			{
 				// Optim: If an instance is using CPU selection and does not voxelize it's data, then there is no need for having valid AABB
@@ -1551,150 +1538,150 @@ void ComputeHairStrandsInterpolation(
 						GraphBuilder,
 						ShaderMap,
 						Instance->LocalToWorld,
-						RegisterAsSRV(GraphBuilder, Instance->Strands.DeformedResource->GetPositionOffsetBuffer(FHairStrandsDeformedResource::EFrameType::Current)),
+						Strands_PositionOffsetSRV,
 						HairGroupCluster,
 						CullingData,
-						Strands_DeformedPosition);
+						Strands_PositionSRV);
 				}
 			}
-		}
 
-		if (Instance->Strands.DeformedResource->NeedsToUpdateTangent())
-		{
-			AddHairTangentPass(
-				GraphBuilder,
-				ShaderMap,
-				Instance->Strands.RestResource->GetVertexCount(),
-				Instance->HairGroupPublicData,
-				Strands_DeformedPosition.SRV,
-				Strands_DeformedTangent);
-		}
-
-		#if RHI_RAYTRACING
-		if (Instance->Strands.RenRaytracingResource)
-		{
-			const float HairRadiusScaleRT = (GHairRaytracingRadiusScale > 0 ? GHairRaytracingRadiusScale : Instance->Strands.Modifier.HairRaytracingRadiusScale);
-			const bool bNeedUpdate = !Instance->Strands.RenRaytracingResource->bIsRTGeometryInitialized || Instance->BindingType == EHairBindingType::Skinning;
-			if (bNeedUpdate)
+			// 4. Update raytracing geometry
+			#if RHI_RAYTRACING
+			if (Instance->Strands.RenRaytracingResource)
 			{
-				FRDGImportedBuffer Raytracing_PositionBuffer = Register(GraphBuilder, Instance->Strands.RenRaytracingResource->PositionBuffer, ERDGImportedBufferFlags::CreateViews);
-				AddGenerateRaytracingGeometryPass(
-					GraphBuilder,
-					ShaderMap,
-					Instance->Strands.RestResource->GetVertexCount(),
-					ScaleAndClipDesc.MaxOutHairRadius * HairRadiusScaleRT,
-					RegisterAsSRV(GraphBuilder, Instance->Strands.DeformedResource->GetPositionOffsetBuffer(FHairStrandsDeformedResource::EFrameType::Current)),
-					CullingData,
-					Strands_DeformedPosition.SRV,
-					Raytracing_PositionBuffer.UAV);
-
-				GraphBuilder.AddPass(
-					RDG_EVENT_NAME("HairStrandsUpdateBLAS"),
-					ERDGPassFlags::NeverCull,
-				[Instance](FRHICommandList& RHICmdList)
+				const float HairRadiusScaleRT = (GHairRaytracingRadiusScale > 0 ? GHairRaytracingRadiusScale : Instance->Strands.Modifier.HairRaytracingRadiusScale);
+				const bool bNeedUpdate = !Instance->Strands.RenRaytracingResource->bIsRTGeometryInitialized || Instance->BindingType == EHairBindingType::Skinning;
+				if (bNeedUpdate)
 				{
-					const bool bNeedFullBuild = !Instance->Strands.RenRaytracingResource->bIsRTGeometryInitialized;
-					if (bNeedFullBuild)
-					{
-						FBufferRHIRef PositionBuffer(Instance->Strands.RenRaytracingResource->PositionBuffer.Buffer->GetRHI());
-						BuildHairAccelerationStructure_Strands(RHICmdList, Instance->Strands.RenRaytracingResource->VertexCount, PositionBuffer, &Instance->Strands.RenRaytracingResource->RayTracingGeometry);
-					}
-					else
+					FRDGImportedBuffer Raytracing_PositionBuffer = Register(GraphBuilder, Instance->Strands.RenRaytracingResource->PositionBuffer, ERDGImportedBufferFlags::CreateViews);
+					AddGenerateRaytracingGeometryPass(
+						GraphBuilder,
+						ShaderMap,
+						Instance->Strands.RestResource->GetVertexCount(),
+						ScaleAndClipDesc.MaxOutHairRadius * HairRadiusScaleRT,
+						Strands_PositionOffsetSRV,
+						CullingData,
+						Strands_PositionSRV,
+						Raytracing_PositionBuffer.UAV);
+
+					GraphBuilder.AddPass(
+						RDG_EVENT_NAME("HairStrandsUpdateBLAS"),
+						ERDGPassFlags::NeverCull,
+					[Instance](FRHICommandList& RHICmdList)
 					{
-						UpdateHairAccelerationStructure(RHICmdList, &Instance->Strands.RenRaytracingResource->RayTracingGeometry);
-					}
-					Instance->Strands.RenRaytracingResource->bIsRTGeometryInitialized = true;
-				});
+						const bool bNeedFullBuild = !Instance->Strands.RenRaytracingResource->bIsRTGeometryInitialized;
+						if (bNeedFullBuild)
+						{
+							FBufferRHIRef PositionBuffer(Instance->Strands.RenRaytracingResource->PositionBuffer.Buffer->GetRHI());
+							BuildHairAccelerationStructure_Strands(RHICmdList, Instance->Strands.RenRaytracingResource->VertexCount, PositionBuffer, &Instance->Strands.RenRaytracingResource->RayTracingGeometry);
+						}
+						else
+						{
+							UpdateHairAccelerationStructure(RHICmdList, &Instance->Strands.RenRaytracingResource->RayTracingGeometry);
+						}
+						Instance->Strands.RenRaytracingResource->bIsRTGeometryInitialized = true;
+					});
+				}
 			}
+			#endif
 		}
-		#endif
+
+		// Sanity check
+		check(Instance->HairGroupPublicData->VFInput.Strands.PositionBuffer);
 	}
 	else if (InstanceGeometryType == EHairGeometryType::Cards)
 	{	
+		DECLARE_GPU_STAT(HairCardsInterpolation);
+		RDG_EVENT_SCOPE(GraphBuilder, "HairCardsInterpolation");
+		RDG_GPU_STAT_SCOPE(GraphBuilder, HairCardsInterpolation);
+
 		const uint32 HairLODIndex = Instance->HairGroupPublicData->GetIntLODIndex();
 		const bool bIsCardsValid = Instance->Cards.IsValid(HairLODIndex);
 		if (bIsCardsValid)
 		{
 			const bool bValidGuide = Instance->Guides.bIsSimulationEnable || Instance->Guides.bHasGlobalInterpolation;
 			const bool bHasSkinning = Instance->BindingType == EHairBindingType::Skinning;
+			const bool bNeedDeformation = bValidGuide || bHasSkinning;
 
 			FHairGroupInstance::FCards::FLOD& LOD = Instance->Cards.LODs[HairLODIndex];
-			const EHairCardsSimulationType CardsSimulationType = 
-				bHasSkinning || bValidGuide ?
-				GetHairCardsSimulationType() :
-				EHairCardsSimulationType::None;
 
-			if (LOD.Guides.Data && CardsSimulationType == EHairCardsSimulationType::Guide)
+			if (bNeedDeformation)
 			{
-				FRDGImportedBuffer Guides_DeformedPositionBuffer = Register(GraphBuilder, LOD.Guides.DeformedResource->GetBuffer(FHairStrandsDeformedResource::Current), ERDGImportedBufferFlags::CreateViews);
-
-				FHairScaleAndClipDesc ScaleAndClipDesc;
-				ScaleAndClipDesc.bEnable = false;
-				ScaleAndClipDesc.InHairLength = LOD.Guides.Data->StrandsCurves.MaxLength;
-				ScaleAndClipDesc.InHairRadius = LOD.Guides.Data->StrandsCurves.MaxRadius;
-				ScaleAndClipDesc.OutHairRadius = (GStrandHairWidth > 0 ? GStrandHairWidth * 0.5f : ScaleAndClipDesc.InHairRadius);
-				ScaleAndClipDesc.MaxOutHairRadius = ScaleAndClipDesc.OutHairRadius;
-				ScaleAndClipDesc.HairRadiusRootScale = 1;
-				ScaleAndClipDesc.HairRadiusTipScale = 1;
-				ScaleAndClipDesc.HairLengthClip = 1;
-
-				FRDGHairStrandsCullingData CullingData;
-				AddHairStrandsInterpolationPass(
-					GraphBuilder,
-					ShaderMap,
-					ShaderDrawData,
-					Instance,
-					LOD.Guides.RestResource->GetVertexCount(),
-					ScaleAndClipDesc,
-					MeshLODIndex,
-					false,
-					LOD.Guides.HairInterpolationType,
-					InstanceGeometryType,
-					CullingData,
-					LOD.Guides.RestResource->PositionOffset,
-					bValidGuide ? Instance->Guides.RestResource->PositionOffset : FVector::ZeroVector,
-					RegisterAsSRV(GraphBuilder, LOD.Guides.DeformedResource->GetPositionOffsetBuffer(FHairStrandsDeformedResource::Current)),
-					bValidGuide ? RegisterAsSRV(GraphBuilder, Instance->Guides.DeformedResource->GetPositionOffsetBuffer(FHairStrandsDeformedResource::Current)) : nullptr,
-					bHasSkinning ? LOD.Guides.RestRootResource : nullptr ,
-					bHasSkinning && bValidGuide ? Instance->Guides.RestRootResource : nullptr,
-					bHasSkinning ? LOD.Guides.DeformedRootResource : nullptr,
-					bHasSkinning && bValidGuide ? Instance->Guides.DeformedRootResource : nullptr,
-					RegisterAsSRV(GraphBuilder, LOD.Guides.RestResource->RestPositionBuffer),
-					RegisterAsSRV(GraphBuilder, LOD.Guides.RestResource->AttributeBuffer),
-					RegisterAsSRV(GraphBuilder, LOD.Guides.InterpolationResource->Interpolation0Buffer),
-					RegisterAsSRV(GraphBuilder, LOD.Guides.InterpolationResource->Interpolation1Buffer),
-					bValidGuide ? RegisterAsSRV(GraphBuilder, Instance->Guides.RestResource->RestPositionBuffer) : nullptr,
-					bValidGuide ? RegisterAsSRV(GraphBuilder, Instance->Guides.DeformedResource->GetBuffer(FHairStrandsDeformedResource::Current)) : nullptr,
-					bValidGuide ? RegisterAsSRV(GraphBuilder, Instance->Guides.RestResource->AttributeBuffer) : nullptr,
-					Guides_DeformedPositionBuffer,
-					nullptr,
-					nullptr,
-					nullptr,
-					RegisterAsSRV(GraphBuilder, LOD.Guides.InterpolationResource->SimRootPointIndexBuffer)); // <- this should be optional
+				check(LOD.Guides.Data);
 
-				AddHairCardsDeformationPass(
-					GraphBuilder,
-					ShaderMap,
-					ShaderDrawData,
-					Instance,
-					MeshLODIndex);
-			}
-			else if (LOD.Guides.Data && CardsSimulationType == EHairCardsSimulationType::RBF)
-			{
-				AddHairCardsRBFInterpolationPass(
-					GraphBuilder,
-					ShaderMap,
-					MeshLODIndex,
-					LOD.RestResource,
-					LOD.DeformedResource,
-					Instance->Guides.RestRootResource,
-					Instance->Guides.DeformedRootResource);
-			}
-			else
-			{
-				// TODO remove this copy. For being able to do that we need to build a VF permutation which take the rest position as input
-				FRDGBufferRef DeformedResource = Register(GraphBuilder, LOD.DeformedResource->GetBuffer(FHairCardsDeformedResource::Current), ERDGImportedBufferFlags::None).Buffer;
-				AddCopyBufferPass(GraphBuilder, DeformedResource, LOD.RestResource->RestPositionBuffer.VertexBufferRHI);
+				const EHairCardsSimulationType CardsSimulationType = 
+					bHasSkinning || bValidGuide ?
+					GetHairCardsSimulationType() :
+					EHairCardsSimulationType::None;
+
+				// 1. Cards are deformed based on guides motion (simulation or RBF applied on guides)
+				if (CardsSimulationType == EHairCardsSimulationType::Guide)
+				{
+					FRDGImportedBuffer Guides_DeformedPositionBuffer = Register(GraphBuilder, LOD.Guides.DeformedResource->GetBuffer(FHairStrandsDeformedResource::Current), ERDGImportedBufferFlags::CreateViews);
+
+					FHairScaleAndClipDesc ScaleAndClipDesc;
+					ScaleAndClipDesc.bEnable = false;
+					ScaleAndClipDesc.InHairLength = LOD.Guides.Data->StrandsCurves.MaxLength;
+					ScaleAndClipDesc.InHairRadius = LOD.Guides.Data->StrandsCurves.MaxRadius;
+					ScaleAndClipDesc.OutHairRadius = (GStrandHairWidth > 0 ? GStrandHairWidth * 0.5f : ScaleAndClipDesc.InHairRadius);
+					ScaleAndClipDesc.MaxOutHairRadius = ScaleAndClipDesc.OutHairRadius;
+					ScaleAndClipDesc.HairRadiusRootScale = 1;
+					ScaleAndClipDesc.HairRadiusTipScale = 1;
+					ScaleAndClipDesc.HairLengthClip = 1;
+
+					FRDGHairStrandsCullingData CullingData;
+					AddHairStrandsInterpolationPass(
+						GraphBuilder,
+						ShaderMap,
+						ShaderDrawData,
+						Instance,
+						LOD.Guides.RestResource->GetVertexCount(),
+						ScaleAndClipDesc,
+						MeshLODIndex,
+						false,
+						LOD.Guides.HairInterpolationType,
+						InstanceGeometryType,
+						CullingData,
+						LOD.Guides.RestResource->PositionOffset,
+						bValidGuide ? Instance->Guides.RestResource->PositionOffset : FVector::ZeroVector,
+						RegisterAsSRV(GraphBuilder, LOD.Guides.DeformedResource->GetPositionOffsetBuffer(FHairStrandsDeformedResource::Current)),
+						bValidGuide ? RegisterAsSRV(GraphBuilder, Instance->Guides.DeformedResource->GetPositionOffsetBuffer(FHairStrandsDeformedResource::Current)) : nullptr,
+						bHasSkinning ? LOD.Guides.RestRootResource : nullptr ,
+						bHasSkinning && bValidGuide ? Instance->Guides.RestRootResource : nullptr,
+						bHasSkinning ? LOD.Guides.DeformedRootResource : nullptr,
+						bHasSkinning && bValidGuide ? Instance->Guides.DeformedRootResource : nullptr,
+						RegisterAsSRV(GraphBuilder, LOD.Guides.RestResource->RestPositionBuffer),
+						RegisterAsSRV(GraphBuilder, LOD.Guides.RestResource->AttributeBuffer),
+						RegisterAsSRV(GraphBuilder, LOD.Guides.InterpolationResource->Interpolation0Buffer),
+						RegisterAsSRV(GraphBuilder, LOD.Guides.InterpolationResource->Interpolation1Buffer),
+						bValidGuide ? RegisterAsSRV(GraphBuilder, Instance->Guides.RestResource->RestPositionBuffer) : nullptr,
+						bValidGuide ? RegisterAsSRV(GraphBuilder, Instance->Guides.DeformedResource->GetBuffer(FHairStrandsDeformedResource::Current)) : nullptr,
+						bValidGuide ? RegisterAsSRV(GraphBuilder, Instance->Guides.RestResource->AttributeBuffer) : nullptr,
+						Guides_DeformedPositionBuffer,
+						nullptr,
+						nullptr,
+						nullptr,
+						RegisterAsSRV(GraphBuilder, LOD.Guides.InterpolationResource->SimRootPointIndexBuffer)); // <- this should be optional
+
+					AddHairCardsDeformationPass(
+						GraphBuilder,
+						ShaderMap,
+						ShaderDrawData,
+						Instance,
+						MeshLODIndex);
+				}
+				// 2. Cards are deformed only based on skel. mesh RBF data)
+				else if (CardsSimulationType == EHairCardsSimulationType::RBF)
+				{
+					AddHairCardsRBFInterpolationPass(
+						GraphBuilder,
+						ShaderMap,
+						MeshLODIndex,
+						LOD.RestResource,
+						LOD.DeformedResource,
+						Instance->Guides.RestRootResource,
+						Instance->Guides.DeformedRootResource);
+				}
 			}
 
 			#if RHI_RAYTRACING
@@ -1728,12 +1715,24 @@ void ComputeHairStrandsInterpolation(
 	}
 	else if (InstanceGeometryType == EHairGeometryType::Meshes)
 	{
+		DECLARE_GPU_STAT(HairMeshesInterpolation);
+		RDG_EVENT_SCOPE(GraphBuilder, "HairMeshesInterpolation");
+		RDG_GPU_STAT_SCOPE(GraphBuilder, HairMeshesInterpolation);
+
 		const uint32 HairLODIndex = Instance->HairGroupPublicData->GetIntLODIndex();
 		if (Instance->Meshes.IsValid(HairLODIndex))
 		{
-			FHairGroupInstance::FMeshes::FLOD& MeshesInstance = Instance->Meshes.LODs[HairLODIndex];
-			if (Instance->Guides.IsValid() && Instance->Guides.HasValidRootData() && Instance->Guides.DeformedRootResource->IsValid(MeshLODIndex))
+			const bool bNeedDeformation = Instance->Meshes.LODs[HairLODIndex].DeformedResource != nullptr;
+			if (bNeedDeformation)
 			{
+				FHairGroupInstance::FMeshes::FLOD& MeshesInstance = Instance->Meshes.LODs[HairLODIndex];
+
+				check(Instance->BindingType == EHairBindingType::Skinning)
+				check(Instance->Guides.IsValid());
+				check(Instance->Guides.HasValidRootData());
+				check(Instance->Guides.DeformedRootResource);
+				check(Instance->Guides.DeformedRootResource->IsValid(MeshLODIndex));
+
 				AddHairMeshesRBFInterpolationPass(
 					GraphBuilder,
 					ShaderMap,
@@ -1743,35 +1742,35 @@ void ComputeHairStrandsInterpolation(
 					Instance->Guides.RestRootResource,
 					Instance->Guides.DeformedRootResource);
 			}
-		}
 
-		#if RHI_RAYTRACING
-		FHairGroupInstance::FMeshes::FLOD& LOD = Instance->Meshes.LODs[HairLODIndex];
-		if (LOD.RaytracingResource)
-		{
-			const bool bNeedUpdate = !LOD.RaytracingResource->bIsRTGeometryInitialized || Instance->BindingType == EHairBindingType::Skinning;
-			if (bNeedUpdate)
+			#if RHI_RAYTRACING
+			FHairGroupInstance::FMeshes::FLOD& LOD = Instance->Meshes.LODs[HairLODIndex];
+			if (LOD.RaytracingResource)
 			{
-				GraphBuilder.AddPass(
-					RDG_EVENT_NAME("HairMeshesUpdateBLAS"),
-					ERDGPassFlags::NeverCull,
-					[Instance, HairLODIndex](FRHICommandList& RHICmdList)
+				const bool bNeedUpdate = !LOD.RaytracingResource->bIsRTGeometryInitialized || bNeedDeformation;
+				if (bNeedUpdate)
 				{
-					FHairGroupInstance::FMeshes::FLOD& LocalLOD = Instance->Meshes.LODs[HairLODIndex];				
-					const bool bNeedFullBuild = !LocalLOD.RaytracingResource->bIsRTGeometryInitialized;
-					if (bNeedFullBuild)
-					{
-						BuildHairAccelerationStructure_Meshes(RHICmdList, LocalLOD.RestResource, LocalLOD.DeformedResource,  &LocalLOD.RaytracingResource->RayTracingGeometry);
-					}
-					else
+					GraphBuilder.AddPass(
+						RDG_EVENT_NAME("HairMeshesUpdateBLAS"),
+						ERDGPassFlags::NeverCull,
+						[Instance, HairLODIndex](FRHICommandList& RHICmdList)
 					{
-						UpdateHairAccelerationStructure(RHICmdList, &LocalLOD.RaytracingResource->RayTracingGeometry);
-					}
-					LocalLOD.RaytracingResource->bIsRTGeometryInitialized = true;
-				});
+						FHairGroupInstance::FMeshes::FLOD& LocalLOD = Instance->Meshes.LODs[HairLODIndex];				
+						const bool bNeedFullBuild = !LocalLOD.RaytracingResource->bIsRTGeometryInitialized;
+						if (bNeedFullBuild)
+						{
+							BuildHairAccelerationStructure_Meshes(RHICmdList, LocalLOD.RestResource, LocalLOD.DeformedResource,  &LocalLOD.RaytracingResource->RayTracingGeometry);
+						}
+						else
+						{
+							UpdateHairAccelerationStructure(RHICmdList, &LocalLOD.RaytracingResource->RayTracingGeometry);
+						}
+						LocalLOD.RaytracingResource->bIsRTGeometryInitialized = true;
+					});
+				}
 			}
+			#endif
 		}
-		#endif
 	}
 
 	Instance->HairGroupPublicData->VFInput.GeometryType = InstanceGeometryType;
diff --git a/Engine/Plugins/Runtime/HairStrands/Source/HairStrandsCore/Public/GroomAsset.h b/Engine/Plugins/Runtime/HairStrands/Source/HairStrandsCore/Public/GroomAsset.h
--- a/Engine/Plugins/Runtime/HairStrands/Source/HairStrandsCore/Public/GroomAsset.h
+++ b/Engine/Plugins/Runtime/HairStrands/Source/HairStrandsCore/Public/GroomAsset.h
@@ -403,8 +403,10 @@ public:
 	virtual const TArray<UAssetUserData*>* GetAssetUserDataArray() const override;
 	//~ End IInterface_AssetUserData Interface
 
-	bool IsSimulationEnable(int32 GroupIndex, int32 LODIndex);
-	bool IsGlobalInterpolationEnable(int32 GroupIndex, int32 LODIndex);
+	EGroomGeometryType GetGeometryType(int32 GroupIndex, int32 LODIndex) const;
+	EGroomBindingType GetBindingType(int32 GroupIndex, int32 LODIndex) const;
+	bool IsSimulationEnable(int32 GroupIndex, int32 LODIndex) const;
+	bool IsGlobalInterpolationEnable(int32 GroupIndex, int32 LODIndex) const;
 
 #if WITH_EDITORONLY_DATA
 	void SetDebugMode(EHairStrandsDebugMode InMode) { DebugMode = InMode; OnGroomAssetChanged.Broadcast(); }
diff --git a/Engine/Plugins/Runtime/HairStrands/Source/HairStrandsCore/Public/GroomInstance.h b/Engine/Plugins/Runtime/HairStrands/Source/HairStrandsCore/Public/GroomInstance.h
--- a/Engine/Plugins/Runtime/HairStrands/Source/HairStrandsCore/Public/GroomInstance.h
+++ b/Engine/Plugins/Runtime/HairStrands/Source/HairStrandsCore/Public/GroomInstance.h
@@ -42,7 +42,7 @@ struct HAIRSTRANDSCORE_API FHairGroupInstance
 	struct FStrandsBase
 	{
 		bool HasValidData() const { return Data != nullptr && Data->GetNumPoints() > 0; }
-		bool IsValid() const { return DeformedResource != nullptr; }
+		bool IsValid() const { return RestResource != nullptr; }
 
 		// Data - Render & sim (rest) data
 		FHairStrandsDatas* Data = nullptr;
@@ -101,10 +101,10 @@ struct HAIRSTRANDSCORE_API FHairGroupInstance
 	struct FCards
 	{
 		bool IsValid() const { for (const FLOD& LOD: LODs) { if (LOD.IsValid()) { return true; } } return false; }
-		bool IsValid(int32 LODIndex) const { return LODIndex >= 0 && LODIndex < LODs.Num() && LODs[LODIndex].DeformedResource != nullptr; }
+		bool IsValid(int32 LODIndex) const { return LODIndex >= 0 && LODIndex < LODs.Num() && LODs[LODIndex].RestResource != nullptr; }
 		struct FLOD
 		{
-			bool IsValid() const { return DeformedResource != nullptr; }
+			bool IsValid() const { return RestResource != nullptr; }
 
 			// Data
 			FHairCardsDatas* Data = nullptr;
@@ -133,10 +133,10 @@ struct HAIRSTRANDSCORE_API FHairGroupInstance
 	struct FMeshes
 	{
 		bool IsValid() const { for (const FLOD& LOD : LODs) { if (LOD.IsValid()) { return true; } } return false; }
-		bool IsValid(int32 LODIndex) const { return LODIndex >= 0 && LODIndex < LODs.Num() && LODs[LODIndex].DeformedResource != nullptr; }
+		bool IsValid(int32 LODIndex) const { return LODIndex >= 0 && LODIndex < LODs.Num() && LODs[LODIndex].RestResource != nullptr; }
 		struct FLOD
 		{
-			bool IsValid() const { return DeformedResource != nullptr; }
+			bool IsValid() const { return RestResource != nullptr; }
 
 			// Data
 			FHairMeshesDatas* Data = nullptr;
diff --git a/Engine/Plugins/Runtime/HairStrands/Source/HairStrandsCore/Public/GroomResources.h b/Engine/Plugins/Runtime/HairStrands/Source/HairStrandsCore/Public/GroomResources.h
--- a/Engine/Plugins/Runtime/HairStrands/Source/HairStrandsCore/Public/GroomResources.h
+++ b/Engine/Plugins/Runtime/HairStrands/Source/HairStrandsCore/Public/GroomResources.h
@@ -173,19 +173,29 @@ struct FHairStrandsRestResource : public FRenderResource
 	/* Get the resource name */
 	virtual FString GetFriendlyName() const override { return TEXT("FHairStrandsResource"); }
 
+	FRDGExternalBuffer GetTangentBuffer(class FRDGBuilder& GraphBuilder, class FGlobalShaderMap* ShaderMap);
+
 	/* Return the memory size for GPU resources */
 	uint32 GetResourcesSize() const
 	{
 		uint32 Total = 0;
 		Total += GetBufferTotalNumBytes(RestPositionBuffer);
+		Total += GetBufferTotalNumBytes(PositionOffsetBuffer);
 		Total += GetBufferTotalNumBytes(AttributeBuffer);
 		Total += GetBufferTotalNumBytes(MaterialBuffer);
+		Total += GetBufferTotalNumBytes(TangentBuffer);
 		return Total;
 	}
 
 	/* Strand hair rest position buffer */
 	FRDGExternalBuffer  RestPositionBuffer;
 
+	/* Strand hair rest offset position buffer */
+	FRDGExternalBuffer  PositionOffsetBuffer;
+
+	/* Strand hair tangent buffer (non-allocated unless used for static geometry) */
+	FRDGExternalBuffer TangentBuffer;
+
 	/* Strand hair attribute buffer */
 	FRDGExternalBuffer AttributeBuffer;
 
@@ -204,7 +214,7 @@ struct FHairStrandsRestResource : public FRenderResource
 struct FHairStrandsDeformedResource : public FRenderResource
 {
 	/** Build the hair strands resource */
-	FHairStrandsDeformedResource(const FHairStrandsDatas::FRenderData& HairStrandRenderData, bool bInitializeData, bool bDynamic, const FVector& InDefaultOffset);
+	FHairStrandsDeformedResource(const FHairStrandsDatas::FRenderData& HairStrandRenderData, bool bInitializeData, const FVector& InDefaultOffset);
 
 	/* Init the buffer */
 	virtual void InitRHI() override;
@@ -247,8 +257,6 @@ struct FHairStrandsDeformedResource : public FRenderResource
 	uint32 CurrentIndex = 0;
 
 	/* Whether the underlying resource is dynamic not (single or double buffer allocated) */
-	const bool bDynamic = true; 
-	bool bInitializedTangent = true;
 	FVector DefaultOffset = FVector::ZeroVector;
 
 	enum EFrameType
@@ -258,13 +266,13 @@ struct FHairStrandsDeformedResource : public FRenderResource
 	};
 
 	// Helper accessors
-	inline uint32 GetIndex(EFrameType T) const					{ return (!bDynamic || T == EFrameType::Current) ? CurrentIndex : 1u - CurrentIndex; }
+	inline uint32 GetIndex(EFrameType T) const					{ return T == EFrameType::Current ? CurrentIndex : 1u - CurrentIndex; }
 	inline FRDGExternalBuffer& GetBuffer(EFrameType T)			{ return DeformedPositionBuffer[GetIndex(T)];  }
 	inline FVector& GetPositionOffset(EFrameType T)				{ return PositionOffset[GetIndex(T)]; }
 	inline FRDGExternalBuffer& GetPositionOffsetBuffer(EFrameType T) { return DeformedOffsetBuffer[GetIndex(T)]; }
 	inline const FVector& GetPositionOffset(EFrameType T) const { return PositionOffset[GetIndex(T)]; }
-	inline void SwapBuffer()									{ if (bDynamic) { CurrentIndex = 1u - CurrentIndex; } }
-	bool NeedsToUpdateTangent();
+	inline void SwapBuffer()									{ CurrentIndex = 1u - CurrentIndex; }
+	//bool NeedsToUpdateTangent();
 };
 
 struct FHairStrandsClusterCullingResource : public FRenderResource
@@ -488,7 +496,7 @@ struct FHairCardsProceduralResource : public FRenderResource
 struct FHairCardsDeformedResource : public FRenderResource
 {
 	/** Build the hair strands resource */
-	FHairCardsDeformedResource(const FHairCardsDatas::FRenderData& HairStrandRenderData, bool bInitializeData, bool bDynamic);
+	FHairCardsDeformedResource(const FHairCardsDatas::FRenderData& HairStrandRenderData, bool bInitializeData);
 
 	/* Init the buffer */
 	virtual void InitRHI() override;
@@ -517,8 +525,6 @@ struct FHairCardsDeformedResource : public FRenderResource
 	/* Whether the GPU data should be initialized with the asset data or not */
 	const bool bInitializedData = false;
 
-	const bool bDynamic = true;
-
 	/* Whether the GPU data should be initialized with the asset data or not */
 	uint32 CurrentIndex = 0;
 
@@ -529,9 +535,9 @@ struct FHairCardsDeformedResource : public FRenderResource
 	};
 
 	// Helper accessors
-	inline uint32 GetIndex(EFrameType T)				{ return (!bDynamic || T == EFrameType::Current) ? CurrentIndex : 1u - CurrentIndex; }
+	inline uint32 GetIndex(EFrameType T)				{ return T == EFrameType::Current ? CurrentIndex : 1u - CurrentIndex; }
 	inline FRDGExternalBuffer& GetBuffer(EFrameType T)	{ return DeformedPositionBuffer[GetIndex(T)];  }
-	inline void SwapBuffer()							{ if (bDynamic) { CurrentIndex = 1u - CurrentIndex; } }
+	inline void SwapBuffer()							{ CurrentIndex = 1u - CurrentIndex; }
 };
 
 /** Hair cards points interpolation attributes */
@@ -653,7 +659,7 @@ struct FHairMeshesRestResource : public FRenderResource
 struct FHairMeshesDeformedResource : public FRenderResource
 {
 	/** Build the hair strands resource */
-	FHairMeshesDeformedResource(const FHairMeshesDatas::FRenderData& HairMeshesRenderData, bool bInInitializedData, bool bDynamic);
+	FHairMeshesDeformedResource(const FHairMeshesDatas::FRenderData& HairMeshesRenderData, bool bInInitializedData);
 
 	/* Init the buffer */
 	virtual void InitRHI() override;
@@ -681,7 +687,6 @@ struct FHairMeshesDeformedResource : public FRenderResource
 
 	/* Whether the GPU data should be initialized with the asset data or not */
 	const bool bInitializedData = false;
-	const bool bDynamic = true;
 
 	/* Whether the GPU data should be initialized with the asset data or not */
 	uint32 CurrentIndex = 0;
@@ -693,9 +698,9 @@ struct FHairMeshesDeformedResource : public FRenderResource
 	};
 
 	// Helper accessors
-	inline uint32 GetIndex(EFrameType T) { return (!bDynamic || T == EFrameType::Current) ? CurrentIndex : 1u - CurrentIndex; }
+	inline uint32 GetIndex(EFrameType T) { return T == EFrameType::Current ? CurrentIndex : 1u - CurrentIndex; }
 	inline FRDGExternalBuffer& GetBuffer(EFrameType T) { return DeformedPositionBuffer[GetIndex(T)]; }
-	inline void SwapBuffer() { if (bDynamic) { CurrentIndex = 1u - CurrentIndex; } }
+	inline void SwapBuffer() { CurrentIndex = 1u - CurrentIndex; }
 };
 
 ////////////////////////////////////////////////////////////////////////////////////////////////////////////////////////
diff --git a/Engine/Source/Runtime/Renderer/Public/HairStrandsInterface.h b/Engine/Source/Runtime/Renderer/Public/HairStrandsInterface.h
--- a/Engine/Source/Runtime/Renderer/Public/HairStrandsInterface.h
+++ b/Engine/Source/Runtime/Renderer/Public/HairStrandsInterface.h
@@ -186,6 +186,12 @@ public:
 	void SetLODVisibilities(const TArray<bool>& InLODVisibility) { LODVisibilities = InLODVisibility; }
 	const TArray<bool>& GetLODVisibilities() const { return LODVisibilities; }
 
+	EHairGeometryType GetGeometryType(int32 InLODIndex) const
+	{
+		if (InLODIndex < 0 && InLODIndex >= LODGeometryTypes.Num()) return EHairGeometryType::NoneGeometry;
+		return LODGeometryTypes[InLODIndex];
+	}
+
 	EHairBindingType GetBindingType(int32 InLODIndex) const
 	{ 
 		if (InLODIndex < 0 && InLODIndex >= BindingTypes.Num()) return EHairBindingType::NoneBinding;
