Rename groom rest resource position buffers for naming consistency.

diff --git a/Engine/Plugins/Runtime/HairStrands/Source/HairStrandsCore/Private/GroomDebug.cpp b/Engine/Plugins/Runtime/HairStrands/Source/HairStrandsCore/Private/GroomDebug.cpp
--- a/Engine/Plugins/Runtime/HairStrands/Source/HairStrandsCore/Private/GroomDebug.cpp
+++ b/Engine/Plugins/Runtime/HairStrands/Source/HairStrandsCore/Private/GroomDebug.cpp
@@ -727,13 +727,13 @@ static void AddDrawDebugCardsGuidesPass(
 	Parameters->SimVertexCount = Instance->Guides.RestResource->GetVertexCount();
 
 	Parameters->RenRestOffset = LOD.Guides.RestResource->PositionOffset;
-	Parameters->RenRestPosition = RegisterAsSRV(GraphBuilder, LOD.Guides.RestResource->RestPositionBuffer);
+	Parameters->RenRestPosition = RegisterAsSRV(GraphBuilder, LOD.Guides.RestResource->PositionBuffer);
 
 	Parameters->RenDeformedOffset = LOD.Guides.DeformedResource->GetPositionOffset(FHairStrandsDeformedResource::Current);
 	Parameters->RenDeformedPosition = RegisterAsSRV(GraphBuilder, LOD.Guides.DeformedResource->GetBuffer(FHairStrandsDeformedResource::Current));
 
 	Parameters->SimRestOffset = Instance->Guides.RestResource->PositionOffset;
-	Parameters->SimRestPosition = RegisterAsSRV(GraphBuilder, Instance->Guides.RestResource->RestPositionBuffer);
+	Parameters->SimRestPosition = RegisterAsSRV(GraphBuilder, Instance->Guides.RestResource->PositionBuffer);
 
 	Parameters->SimDeformedOffset = Instance->Guides.DeformedResource->GetPositionOffset(FHairStrandsDeformedResource::Current);
 	Parameters->SimDeformedPosition = RegisterAsSRV(GraphBuilder, Instance->Guides.DeformedResource->GetBuffer(FHairStrandsDeformedResource::Current));
diff --git a/Engine/Plugins/Runtime/HairStrands/Source/HairStrandsCore/Private/GroomResources.cpp b/Engine/Plugins/Runtime/HairStrands/Source/HairStrandsCore/Private/GroomResources.cpp
--- a/Engine/Plugins/Runtime/HairStrands/Source/HairStrandsCore/Private/GroomResources.cpp
+++ b/Engine/Plugins/Runtime/HairStrands/Source/HairStrandsCore/Private/GroomResources.cpp
@@ -508,7 +508,7 @@ void FHairMeshesDeformedResource::ReleaseRHI()
 /////////////////////////////////////////////////////////////////////////////////////////
 
 FHairStrandsRestResource::FHairStrandsRestResource(const FHairStrandsDatas::FRenderData& HairStrandRenderData, const FVector& InPositionOffset) :
-	RestPositionBuffer(), AttributeBuffer(), MaterialBuffer(), PositionOffset(InPositionOffset), RenderData(HairStrandRenderData)
+	PositionBuffer(), AttributeBuffer(), MaterialBuffer(), PositionOffset(InPositionOffset), RenderData(HairStrandRenderData)
 {}
 
 void FHairStrandsRestResource::InitRHI()
@@ -518,7 +518,7 @@ void FHairStrandsRestResource::InitRHI()
 	FRHICommandListImmediate& RHICmdList = FRHICommandListExecutor::GetImmediateCommandList();
 	FRDGBuilder GraphBuilder(RHICmdList);
 
-	InternalCreateVertexBufferRDG<FHairStrandsPositionFormat>(GraphBuilder, RenderData.Positions, RestPositionBuffer, TEXT("Hair.StrandsRest_PositionBuffer"));
+	InternalCreateVertexBufferRDG<FHairStrandsPositionFormat>(GraphBuilder, RenderData.Positions, PositionBuffer, TEXT("Hair.StrandsRest_PositionBuffer"));
 	InternalCreateVertexBufferRDG<FHairStrandsAttributeFormat>(GraphBuilder, RenderData.Attributes, AttributeBuffer, TEXT("Hair.StrandsRest_AttributeBuffer"));
 	InternalCreateVertexBufferRDG<FHairStrandsMaterialFormat>(GraphBuilder, RenderData.Materials, MaterialBuffer, TEXT("Hair.StrandsRest_MaterialBuffer"));
 
@@ -550,7 +550,7 @@ FRDGExternalBuffer FHairStrandsRestResource::GetTangentBuffer(FRDGBuilder& Graph
 			ShaderMap,
 			VertexCount,
 			nullptr,
-			RegisterAsSRV(GraphBuilder, RestPositionBuffer),
+			RegisterAsSRV(GraphBuilder, PositionBuffer),
 			Register(GraphBuilder, TangentBuffer, ERDGImportedBufferFlags::CreateUAV));
 	}
 
@@ -559,7 +559,7 @@ FRDGExternalBuffer FHairStrandsRestResource::GetTangentBuffer(FRDGBuilder& Graph
 
 void FHairStrandsRestResource::ReleaseRHI()
 {
-	RestPositionBuffer.Release();
+	PositionBuffer.Release();
 	PositionOffsetBuffer.Release();
 	AttributeBuffer.Release();
 	MaterialBuffer.Release();
diff --git a/Engine/Plugins/Runtime/HairStrands/Source/HairStrandsCore/Private/GroomTextureBuilder.cpp b/Engine/Plugins/Runtime/HairStrands/Source/HairStrandsCore/Private/GroomTextureBuilder.cpp
--- a/Engine/Plugins/Runtime/HairStrands/Source/HairStrandsCore/Private/GroomTextureBuilder.cpp
+++ b/Engine/Plugins/Runtime/HairStrands/Source/HairStrandsCore/Private/GroomTextureBuilder.cpp
@@ -956,7 +956,7 @@ static void InternalBuildStrandsTextures_GPU(
 				VoxelOffsetAndCount,
 				VoxelData,
 
-				GroupData.Strands.RestResource->RestPositionBuffer.SRV,
+				GroupData.Strands.RestResource->PositionBuffer.SRV,
 				GroupData.Strands.RestResource->AttributeBuffer.SRV,
 				GroupData.Strands.RestResource->PositionOffset,
 				RenderingData.GeometrySettings.HairWidth * 0.5f,
diff --git a/Engine/Plugins/Runtime/HairStrands/Source/HairStrandsCore/Private/HairStrandsRendering.cpp b/Engine/Plugins/Runtime/HairStrands/Source/HairStrandsCore/Private/HairStrandsRendering.cpp
--- a/Engine/Plugins/Runtime/HairStrands/Source/HairStrandsCore/Private/HairStrandsRendering.cpp
+++ b/Engine/Plugins/Runtime/HairStrands/Source/HairStrandsCore/Private/HairStrandsRendering.cpp
@@ -849,7 +849,7 @@ static void AddHairCardsDeformationPass(
 	FHairCardsDeformationCS::FParameters* Parameters = GraphBuilder.AllocParameters<FHairCardsDeformationCS::FParameters>();
 	Parameters->GuideVertexCount			= LOD.Guides.RestResource->GetVertexCount();
 	Parameters->GuideRestPositionOffset		= LOD.Guides.RestResource->PositionOffset;
-	Parameters->GuideRestPositionBuffer		= RegisterAsSRV(GraphBuilder, LOD.Guides.RestResource->RestPositionBuffer);
+	Parameters->GuideRestPositionBuffer		= RegisterAsSRV(GraphBuilder, LOD.Guides.RestResource->PositionBuffer);
 	Parameters->GuideDeformedPositionBuffer = RegisterAsSRV(GraphBuilder, LOD.Guides.DeformedResource->GetBuffer(FHairStrandsDeformedResource::Current));
 	Parameters->GuideDeformedPositionOffsetBuffer = RegisterAsSRV(GraphBuilder, LOD.Guides.DeformedResource->GetPositionOffsetBuffer(FHairStrandsDeformedResource::Current));
 	
@@ -1359,8 +1359,8 @@ FHairGroupPublicData::FVertexFactoryInput ComputeHairStrandsVertexInputData(FHai
 	else
 	{
 
-		OutVFInput.Strands.PositionBuffer = Instance->Strands.RestResource->RestPositionBuffer.SRV;
-		OutVFInput.Strands.PrevPositionBuffer = Instance->Strands.RestResource->RestPositionBuffer.SRV;
+		OutVFInput.Strands.PositionBuffer = Instance->Strands.RestResource->PositionBuffer.SRV;
+		OutVFInput.Strands.PrevPositionBuffer = Instance->Strands.RestResource->PositionBuffer.SRV;
 		OutVFInput.Strands.TangentBuffer = Instance->Strands.RestResource->TangentBuffer.SRV;
 		OutVFInput.Strands.AttributeBuffer = Instance->Strands.RestResource->AttributeBuffer.SRV;
 		OutVFInput.Strands.MaterialBuffer = Instance->Strands.RestResource->MaterialBuffer.SRV;
@@ -1491,11 +1491,11 @@ void ComputeHairStrandsInterpolation(
 					bHasSkinning && bValidGuide ? Instance->Guides.RestRootResource : nullptr,
 					bHasSkinning ? Instance->Strands.DeformedRootResource : nullptr,
 					bHasSkinning && bValidGuide ? Instance->Guides.DeformedRootResource : nullptr,
-					RegisterAsSRV(GraphBuilder, Instance->Strands.RestResource->RestPositionBuffer),
+					RegisterAsSRV(GraphBuilder, Instance->Strands.RestResource->PositionBuffer),
 					RegisterAsSRV(GraphBuilder, Instance->Strands.RestResource->AttributeBuffer),
 					RegisterAsSRV(GraphBuilder, Instance->Strands.InterpolationResource->Interpolation0Buffer),
 					RegisterAsSRV(GraphBuilder, Instance->Strands.InterpolationResource->Interpolation1Buffer),
-					bValidGuide ? RegisterAsSRV(GraphBuilder, Instance->Guides.RestResource->RestPositionBuffer) : nullptr,
+					bValidGuide ? RegisterAsSRV(GraphBuilder, Instance->Guides.RestResource->PositionBuffer) : nullptr,
 					bValidGuide ? RegisterAsSRV(GraphBuilder, Instance->Guides.DeformedResource->GetBuffer(FHairStrandsDeformedResource::Current)) : nullptr,
 					bValidGuide ? RegisterAsSRV(GraphBuilder, Instance->Guides.RestResource->AttributeBuffer) : nullptr,
 					Strands_DeformedPosition,
@@ -1515,7 +1515,7 @@ void ComputeHairStrandsInterpolation(
 			}
 			else
 			{
-				Strands_PositionSRV = RegisterAsSRV(GraphBuilder, Instance->Strands.RestResource->RestPositionBuffer);
+				Strands_PositionSRV = RegisterAsSRV(GraphBuilder, Instance->Strands.RestResource->PositionBuffer);
 				Strands_PositionOffsetSRV = RegisterAsSRV(GraphBuilder, Instance->Strands.RestResource->PositionOffsetBuffer);
 				// Generated the static tangent if they haven;t been generated yet
 				Instance->Strands.RestResource->GetTangentBuffer(GraphBuilder, ShaderMap);
@@ -1650,11 +1650,11 @@ void ComputeHairStrandsInterpolation(
 						bHasSkinning && bValidGuide ? Instance->Guides.RestRootResource : nullptr,
 						bHasSkinning ? LOD.Guides.DeformedRootResource : nullptr,
 						bHasSkinning && bValidGuide ? Instance->Guides.DeformedRootResource : nullptr,
-						RegisterAsSRV(GraphBuilder, LOD.Guides.RestResource->RestPositionBuffer),
+						RegisterAsSRV(GraphBuilder, LOD.Guides.RestResource->PositionBuffer),
 						RegisterAsSRV(GraphBuilder, LOD.Guides.RestResource->AttributeBuffer),
 						RegisterAsSRV(GraphBuilder, LOD.Guides.InterpolationResource->Interpolation0Buffer),
 						RegisterAsSRV(GraphBuilder, LOD.Guides.InterpolationResource->Interpolation1Buffer),
-						bValidGuide ? RegisterAsSRV(GraphBuilder, Instance->Guides.RestResource->RestPositionBuffer) : nullptr,
+						bValidGuide ? RegisterAsSRV(GraphBuilder, Instance->Guides.RestResource->PositionBuffer) : nullptr,
 						bValidGuide ? RegisterAsSRV(GraphBuilder, Instance->Guides.DeformedResource->GetBuffer(FHairStrandsDeformedResource::Current)) : nullptr,
 						bValidGuide ? RegisterAsSRV(GraphBuilder, Instance->Guides.RestResource->AttributeBuffer) : nullptr,
 						Guides_DeformedPositionBuffer,
@@ -1807,7 +1807,7 @@ void ResetHairStrandsInterpolation(
 		Instance->Guides.RestResource->GetVertexCount(),
 		Instance->Guides.RestRootResource,
 		Instance->Guides.DeformedRootResource,
-		RegisterAsSRV(GraphBuilder, Instance->Guides.RestResource->RestPositionBuffer),
+		RegisterAsSRV(GraphBuilder, Instance->Guides.RestResource->PositionBuffer),
 		DeformedPositionBuffer,
 		Instance->Guides.RestResource->PositionOffset,
 		RegisterAsSRV(GraphBuilder, Instance->Guides.DeformedResource->GetPositionOffsetBuffer(FHairStrandsDeformedResource::Current)),
diff --git a/Engine/Plugins/Runtime/HairStrands/Source/HairStrandsCore/Private/Niagara/NiagaraDataInterfaceHairStrands.cpp b/Engine/Plugins/Runtime/HairStrands/Source/HairStrandsCore/Private/Niagara/NiagaraDataInterfaceHairStrands.cpp
--- a/Engine/Plugins/Runtime/HairStrands/Source/HairStrandsCore/Private/Niagara/NiagaraDataInterfaceHairStrands.cpp
+++ b/Engine/Plugins/Runtime/HairStrands/Source/HairStrandsCore/Private/Niagara/NiagaraDataInterfaceHairStrands.cpp
@@ -654,7 +654,7 @@ struct FNDIHairStrandsParametersCS : public FNiagaraDataInterfaceParametersCS
 			SetUAVParameter(RHICmdList, ComputeShaderRHI, BoundingBoxBuffer, HairStrandsBuffer->BoundingBoxBuffer.UAV);
 
 			SetSRVParameter(RHICmdList, ComputeShaderRHI, CurvesOffsetsBuffer, HairStrandsBuffer->CurvesOffsetsBuffer.SRV);
-			SetSRVParameter(RHICmdList, ComputeShaderRHI, RestPositionBuffer, HairStrandsBuffer->SourceRestResources->RestPositionBuffer.SRV);
+			SetSRVParameter(RHICmdList, ComputeShaderRHI, RestPositionBuffer, HairStrandsBuffer->SourceRestResources->PositionBuffer.SRV);
 
 			SetShaderValue(RHICmdList, ComputeShaderRHI, BoundingBoxOffsets, HairStrandsBuffer->BoundingBoxOffsets);
 			SetShaderValue(RHICmdList, ComputeShaderRHI, WorldTransform, ProxyData->WorldTransform.ToMatrixWithScale());
diff --git a/Engine/Plugins/Runtime/HairStrands/Source/HairStrandsCore/Public/GroomResources.h b/Engine/Plugins/Runtime/HairStrands/Source/HairStrandsCore/Public/GroomResources.h
--- a/Engine/Plugins/Runtime/HairStrands/Source/HairStrandsCore/Public/GroomResources.h
+++ b/Engine/Plugins/Runtime/HairStrands/Source/HairStrandsCore/Public/GroomResources.h
@@ -179,7 +179,7 @@ struct FHairStrandsRestResource : public FRenderResource
 	uint32 GetResourcesSize() const
 	{
 		uint32 Total = 0;
-		Total += GetBufferTotalNumBytes(RestPositionBuffer);
+		Total += GetBufferTotalNumBytes(PositionBuffer);
 		Total += GetBufferTotalNumBytes(PositionOffsetBuffer);
 		Total += GetBufferTotalNumBytes(AttributeBuffer);
 		Total += GetBufferTotalNumBytes(MaterialBuffer);
@@ -188,7 +188,7 @@ struct FHairStrandsRestResource : public FRenderResource
 	}
 
 	/* Strand hair rest position buffer */
-	FRDGExternalBuffer  RestPositionBuffer;
+	FRDGExternalBuffer  PositionBuffer;
 
 	/* Strand hair rest offset position buffer */
 	FRDGExternalBuffer  PositionOffsetBuffer;
