Fix crash with groom meshes requesting RBF data while not skinned attached.

diff --git a/Engine/Plugins/Runtime/HairStrands/Source/HairStrandsCore/Private/GroomComponent.cpp b/Engine/Plugins/Runtime/HairStrands/Source/HairStrandsCore/Private/GroomComponent.cpp
--- a/Engine/Plugins/Runtime/HairStrands/Source/HairStrandsCore/Private/GroomComponent.cpp
+++ b/Engine/Plugins/Runtime/HairStrands/Source/HairStrandsCore/Private/GroomComponent.cpp
@@ -1763,8 +1763,9 @@ void UGroomComponent::UpdateSimulatedGroups()
 				}
 				if (Instance->Guides.IsValid())
 				{
-					Instance->Guides.bIsSimulationEnable	 = LocalGroomAsset && LocalGroomAsset->IsSimulationEnable(GroupIt, LODIndex);
-					Instance->Guides.bHasGlobalInterpolation = LocalGroomAsset && LocalGroomAsset->IsGlobalInterpolationEnable(GroupIt, LODIndex);
+					check(Instance->HairGroupPublicData);
+					Instance->Guides.bIsSimulationEnable	 = Instance->HairGroupPublicData->IsGlobalInterpolationEnable(LODIndex);
+					Instance->Guides.bHasGlobalInterpolation = Instance->HairGroupPublicData->IsSimulationEnable(LODIndex);
 				}
 				++GroupIt;
 			}
@@ -1903,12 +1904,12 @@ void UGroomComponent::InitResources(bool bIsBindingReloading)
 			const EGroomGeometryType GeometryType = GroomAsset->GetGeometryType(GroupIt, LODIt);
 			const EGroomBindingType BindingType = GroomAsset->GetBindingType(GroupIt, LODIt);
 
-			// Force global interpolation to be enable for meshes with skinning binding as we use RBF defomation for 'sticking' meshes onto skel. mesh surface
-			const bool bForceGlobalDeformation = GeometryType == EGroomGeometryType::Meshes && BindingType == EGroomBindingType::Skinning;
-			
+			// Note on Global Deformation:
+			// * Global deformation require to have skinning binding
+			// * Force global interpolation to be enable for meshes with skinning binding as we use RBF defomation for 'sticking' meshes onto skel. mesh surface			
 			bHasNeedSkinningBinding   = bHasNeedSkinningBinding || BindingType == EGroomBindingType::Skinning;
 			bHasNeedSimulation		  = bHasNeedSimulation || GroomAsset->IsSimulationEnable(GroupIt, LODIt);
-			bHasNeedGlobalDeformation = bHasNeedGlobalDeformation || GroomAsset->IsGlobalInterpolationEnable(GroupIt, LODIt) || bForceGlobalDeformation;
+			bHasNeedGlobalDeformation = bHasNeedGlobalDeformation || (BindingType == EGroomBindingType::Skinning && (GroomAsset->IsGlobalInterpolationEnable(GroupIt, LODIt) || GeometryType == EGroomGeometryType::Meshes));
 		}
 	}
 	const bool bHasNeedBindingData = bHasNeedSkinningBinding || bHasNeedGlobalDeformation;
@@ -1984,12 +1985,11 @@ void UGroomComponent::InitResources(bool bIsBindingReloading)
 					BindingType = EHairBindingType::NoneBinding;
 				}
 
-				// Force global interpolation to be enable for meshes with skinning binding as we use RBF defomation for 'sticking' meshes onto skel. mesh surface
+				// * Force global interpolation to be enable for meshes with skinning binding as we use RBF defomation for 'sticking' meshes onto skel. mesh surface
+				// * Global deformation are allowed only with 'Skinning' binding type
 				const EHairGeometryType GeometryType = ToHairGeometryType(GroomAsset->GetGeometryType(GroupIt, LODIt));
-				const bool bForceGlobalInterpolation = GeometryType == EHairGeometryType::Meshes && BindingType == EHairBindingType::Skinning;
-
 				const bool LODSimulation = GroomAsset->IsSimulationEnable(GroupIt, LODIt);
-				const bool LODGlobalInterpolation = LocalBindingAsset && (GroomAsset->IsGlobalInterpolationEnable(GroupIt, LODIt) || bForceGlobalInterpolation);
+				const bool LODGlobalInterpolation = LocalBindingAsset && (BindingType == EHairBindingType::Skinning && (GroomAsset->IsGlobalInterpolationEnable(GroupIt, LODIt) || GeometryType == EHairGeometryType::Meshes));
 
 				CPULODScreenSize.Add(LODSettings.ScreenSize);
 				LODVisibility.Add(LODSettings.bVisible);
@@ -2161,6 +2161,12 @@ void UGroomComponent::InitResources(bool bIsBindingReloading)
 				const bool bNeedDeformedPositions	= bHasSimulation || bHasGlobalDeformation || BindingType == EHairBindingType::Skinning;
 				const bool bNeedRootData			= bHasGlobalDeformation || BindingType == EHairBindingType::Skinning;
 
+				// Sanity check
+				if (bHasGlobalDeformation)
+				{
+					check(BindingType == EHairBindingType::Skinning);
+				}
+
 				InstanceLOD.Data = &LOD.Data;
 				InstanceLOD.RestResource = LOD.RestResource;
 				InstanceLOD.InterpolationData = &LOD.InterpolationData;
@@ -2225,7 +2231,7 @@ void UGroomComponent::InitResources(bool bIsBindingReloading)
 			{
 				const EHairBindingType BindingType = HairGroupInstance->HairGroupPublicData->GetBindingType(MeshLODIndex);
 				const bool bHasGlobalDeformation   = HairGroupInstance->HairGroupPublicData->IsGlobalInterpolationEnable(MeshLODIndex);
-				const bool bNeedDeformedPositions  = bHasGlobalDeformation || BindingType == EHairBindingType::Skinning;
+				const bool bNeedDeformedPositions  = bHasGlobalDeformation && BindingType == EHairBindingType::Skinning;
 
 				InstanceLOD.Data = &LOD.Data;
 				InstanceLOD.RestResource = LOD.RestResource;
diff --git a/Engine/Source/Runtime/Renderer/Public/HairStrandsInterface.h b/Engine/Source/Runtime/Renderer/Public/HairStrandsInterface.h
--- a/Engine/Source/Runtime/Renderer/Public/HairStrandsInterface.h
+++ b/Engine/Source/Runtime/Renderer/Public/HairStrandsInterface.h
@@ -194,19 +194,19 @@ public:
 
 	EHairBindingType GetBindingType(int32 InLODIndex) const
 	{ 
-		if (InLODIndex < 0 && InLODIndex >= BindingTypes.Num()) return EHairBindingType::NoneBinding;
+		if (InLODIndex < 0 || InLODIndex >= BindingTypes.Num()) return EHairBindingType::NoneBinding;
 		return BindingTypes[InLODIndex];
 	}	
 
 	bool IsSimulationEnable(int32 InLODIndex) const
 	{
-		if (InLODIndex < 0 && InLODIndex >= LODSimulations.Num()) return false;
+		if (InLODIndex < 0 || InLODIndex >= LODSimulations.Num()) return false;
 		return LODSimulations[InLODIndex];
 	}
 
 	bool IsGlobalInterpolationEnable(int32 InLODIndex) const 
 	{
-		if (InLODIndex < 0 && InLODIndex >= LODGlobalInterpolations.Num()) return false;
+		if (InLODIndex < 0 || InLODIndex >= LODGlobalInterpolations.Num()) return false;
 		return LODGlobalInterpolations[InLODIndex];
 	}
 
