Fix invalid primitive ID during hair material pass.

The generation of PrimitiveID has changed making it's value non fetchable during the mesh processing. This is a consequence of the recent change of GPU culling being enabled. To solve the issue the hair material pass fetches now directly the value from the Primitive stream buffer.

diff --git a/Engine/Plugins/Runtime/HairStrands/Source/HairStrandsCore/Private/HairStrandsVertexFactory.cpp b/Engine/Plugins/Runtime/HairStrands/Source/HairStrandsCore/Private/HairStrandsVertexFactory.cpp
--- a/Engine/Plugins/Runtime/HairStrands/Source/HairStrandsCore/Private/HairStrandsVertexFactory.cpp
+++ b/Engine/Plugins/Runtime/HairStrands/Source/HairStrandsCore/Private/HairStrandsVertexFactory.cpp
@@ -13,6 +13,8 @@
 #include "HairStrandsInterface.h"
 #include "GroomInstance.h"
 
+#define VF_STRANDS_SUPPORT_GPU_SCENE 1
+
 /////////////////////////////////////////////////////////////////////////////////////////
 
 template<typename T> inline void VFS_BindParam(FMeshDrawSingleShaderBindings& ShaderBindings, const FShaderResourceParameter& Param, T* Value)	{ if (Param.IsBound() && Value) ShaderBindings.Add(Param, Value); }
@@ -174,14 +176,18 @@ bool FHairStrandsVertexFactory::ShouldCompilePermutation(const FVertexFactorySha
 
 void FHairStrandsVertexFactory::ModifyCompilationEnvironment(const FVertexFactoryShaderPermutationParameters& Parameters, FShaderCompilerEnvironment& OutEnvironment)
 {
-	const bool bUseGPUSceneAndPrimitiveIdStream = Parameters.VertexFactoryType->SupportsPrimitiveIdStream() && UseGPUScene(Parameters.Platform, GetMaxSupportedFeatureLevel(Parameters.Platform));
+	const bool bUseGPUSceneAndPrimitiveIdStream = 
+		VF_STRANDS_SUPPORT_GPU_SCENE
+		&& Parameters.VertexFactoryType->SupportsPrimitiveIdStream() 
+		&& UseGPUScene(Parameters.Platform, GetMaxSupportedFeatureLevel(Parameters.Platform));
 	OutEnvironment.SetDefine(TEXT("VF_SUPPORTS_PRIMITIVE_SCENE_DATA"), bUseGPUSceneAndPrimitiveIdStream);
 	OutEnvironment.SetDefine(TEXT("HAIR_STRAND_MESH_FACTORY"), TEXT("1"));
 }
 
 void FHairStrandsVertexFactory::ValidateCompiledResult(const FVertexFactoryType* Type, EShaderPlatform Platform, const FShaderParameterMap& ParameterMap, TArray<FString>& OutErrors)
 {
-	if (Type->SupportsPrimitiveIdStream() 
+	if (VF_STRANDS_SUPPORT_GPU_SCENE
+		&& Type->SupportsPrimitiveIdStream()
 		&& UseGPUScene(Platform, GetMaxSupportedFeatureLevel(Platform)) 
 		&& ParameterMap.ContainsParameterAllocation(FPrimitiveUniformShaderParameters::StaticStructMetadata.GetShaderVariableName()))
 	{
@@ -222,7 +228,7 @@ void FHairStrandsVertexFactory::InitRHI()
 
 	// VertexFactory needs to be able to support max possible shader platform and feature level
 	// in case if we switch feature level at runtime.
-	const bool bCanUseGPUScene = UseGPUScene(GMaxRHIShaderPlatform, GMaxRHIFeatureLevel);
+	const bool bCanUseGPUScene = VF_STRANDS_SUPPORT_GPU_SCENE && UseGPUScene(GMaxRHIShaderPlatform, GMaxRHIFeatureLevel);
 
 	FVertexDeclarationElementList Elements;
 	SetPrimitiveIdStreamIndex(EVertexInputStreamType::Default, -1);
@@ -234,10 +240,12 @@ void FHairStrandsVertexFactory::InitRHI()
 		bNeedsDeclaration = true;
 	}
 
-	check(Streams.Num() > 0);
-
-	InitDeclaration(Elements);
-	check(IsValidRef(GetDeclaration()));
+	if (bNeedsDeclaration)
+	{
+		check(Streams.Num() > 0);
+		InitDeclaration(Elements);
+		check(IsValidRef(GetDeclaration()));
+	}
 }
 
 IMPLEMENT_VERTEX_FACTORY_PARAMETER_TYPE(FHairStrandsVertexFactory, SF_Vertex,		FHairStrandsVertexFactoryShaderParameters);
@@ -252,4 +260,4 @@ void FHairStrandsVertexFactory::ReleaseRHI()
 	FVertexFactory::ReleaseRHI();
 }
 
-IMPLEMENT_VERTEX_FACTORY_TYPE_EX(FHairStrandsVertexFactory,"/Engine/Private/HairStrands/HairStrandsVertexFactory.ush",true,false,true,true,false,true,true,false);
+IMPLEMENT_VERTEX_FACTORY_TYPE_EX(FHairStrandsVertexFactory,"/Engine/Private/HairStrands/HairStrandsVertexFactory.ush",true,false,true,true,false,true, VF_STRANDS_SUPPORT_GPU_SCENE,false);
diff --git a/Engine/Shaders/Private/HairStrands/HairStrandsMaterialGBufferPS.usf b/Engine/Shaders/Private/HairStrands/HairStrandsMaterialGBufferPS.usf
--- a/Engine/Shaders/Private/HairStrands/HairStrandsMaterialGBufferPS.usf
+++ b/Engine/Shaders/Private/HairStrands/HairStrandsMaterialGBufferPS.usf
@@ -9,6 +9,7 @@
 void Main(
 	in float2 UV : TEXCOORD
 	, in float4 SvPosition: SV_Position
+	, in uint  InPrimitiveId : PRIMITIVE_ID
 	, out float4 OutGBufferA : SV_Target0
 	, out float4 OutGBufferB : SV_Target1
 	, out float4 OutGBufferC : SV_Target2
@@ -56,7 +57,7 @@ void Main(
 	const float2 SamplePixelCoord = PixelCoord.xy + float2(0.5f, 0.5f);
 
 	const float SampleCoverage8bit = 1; // TODO?
-	FEvaluateOutput Out = Evaluate(SamplePixelCoord, SampleDepth, VisSample.PrimitiveId, SampleCoverage8bit, MaterialPass_PrimitiveId, false);
+	FEvaluateOutput Out = Evaluate(SamplePixelCoord, SampleDepth, VisSample.PrimitiveId, SampleCoverage8bit, InPrimitiveId, false);
 	WriteGBuffer(
 		Out.NodeData.Tangent,
 		Out.NodeData.Specular,
diff --git a/Engine/Shaders/Private/HairStrands/HairStrandsMaterialGBufferVS.usf b/Engine/Shaders/Private/HairStrands/HairStrandsMaterialGBufferVS.usf
--- a/Engine/Shaders/Private/HairStrands/HairStrandsMaterialGBufferVS.usf
+++ b/Engine/Shaders/Private/HairStrands/HairStrandsMaterialGBufferVS.usf
@@ -10,8 +10,10 @@ float3 DecodeVelocityFromTexture(float4 In);
  
 void Main(
 	uint VertexId : SV_VertexID, 
+	FVertexFactoryInput Input,
 	out float2 OutUV : TEXCOORD,
-	out float4 OutPosition : SV_POSITION)
+	out float4 OutPosition : SV_POSITION,
+	nointerpolation out uint  OutPrimitiveId : PRIMITIVE_ID)
 {
 	const float2 UV = float2(VertexId & 1, VertexId >> 1) * 2;
 	const float2 Pos = float2(UV.x, 1 - UV.y) * 2 - float2(1, 1);// * 4 - float2(2,2);
diff --git a/Engine/Shaders/Private/HairStrands/HairStrandsMaterialPS.usf b/Engine/Shaders/Private/HairStrands/HairStrandsMaterialPS.usf
--- a/Engine/Shaders/Private/HairStrands/HairStrandsMaterialPS.usf
+++ b/Engine/Shaders/Private/HairStrands/HairStrandsMaterialPS.usf
@@ -6,6 +6,7 @@ void Main(
 	in float4 SvPosition: SV_Position
 	, in uint2 Resolution : DISPATCH_RESOLUTION
 	, in uint TotalNodeCount : DISPATCH_NODE_COUNT
+	, in uint InPrimitiveId : PRIMITIVE_ID
 	#if HAIR_MATERIAL_DEBUG_OR_EMISSIVE_OUTPUT
 	, out float4 OutColor0 : SV_Target0
 	#endif
@@ -31,7 +32,7 @@ void Main(
 
 		if (SampleMaterialId == MaterialPass_MaterialId)
 		{
-			FEvaluateOutput Out = Evaluate(SamplePixelCoord, SampleDepth, SamplePrimitiveId, SampleCoverage8bit, MaterialPass_PrimitiveId, MaterialPassParameters.bUpdateSampleCoverage > 0);
+			FEvaluateOutput Out = Evaluate(SamplePixelCoord, SampleDepth, SamplePrimitiveId, SampleCoverage8bit, InPrimitiveId, MaterialPassParameters.bUpdateSampleCoverage > 0);
 
 			// Note: Use the syntax _ over the syntax . as on certain platform the UAV buffers are 
 			// stored within a static const struct preventing doing a write operation
diff --git a/Engine/Shaders/Private/HairStrands/HairStrandsMaterialVS.usf b/Engine/Shaders/Private/HairStrands/HairStrandsMaterialVS.usf
--- a/Engine/Shaders/Private/HairStrands/HairStrandsMaterialVS.usf
+++ b/Engine/Shaders/Private/HairStrands/HairStrandsMaterialVS.usf
@@ -10,9 +10,11 @@ float3 DecodeVelocityFromTexture(float4 In);
  
 void Main(
 	uint VertexId : SV_VertexID, 
+	FVertexFactoryInput Input,
 	out float4 Position : SV_POSITION,
 	nointerpolation out uint2 OutResolution : DISPATCH_RESOLUTION,
-	nointerpolation out uint  OutTotalNodeCount : DISPATCH_NODE_COUNT)
+	nointerpolation out uint  OutTotalNodeCount : DISPATCH_NODE_COUNT,
+	nointerpolation out uint  OutPrimitiveId : PRIMITIVE_ID)
 {
 	OutTotalNodeCount = MaterialPassParameters.TotalNodeCounter.Load(uint3(0,0,0));
 	OutResolution.x = ceil(sqrt(OutTotalNodeCount));
@@ -24,6 +26,8 @@ void Main(
 	const float2 Pos = float2(-UV.x, UV.y) * 4 + float2(-1,+1) + float2(ClipCoord.x, -ClipCoord.y);
 	Position = float4(Pos, 0.5f, 1);
 
+	const FVertexFactoryIntermediates VFIntermediates = GetVertexFactoryIntermediates(Input);
+	OutPrimitiveId = VFIntermediates.PrimitiveId; 
 }
 
 #endif // VERTEXSHADER
