Fix crash when generating follicle texture or hair strands textures.

diff --git a/Engine/Plugins/Runtime/HairStrands/Source/HairStrandsCore/Private/GroomTextureBuilder.cpp b/Engine/Plugins/Runtime/HairStrands/Source/HairStrandsCore/Private/GroomTextureBuilder.cpp
--- a/Engine/Plugins/Runtime/HairStrands/Source/HairStrandsCore/Private/GroomTextureBuilder.cpp
+++ b/Engine/Plugins/Runtime/HairStrands/Source/HairStrandsCore/Private/GroomTextureBuilder.cpp
@@ -60,7 +60,8 @@ void FGroomTextureBuilder::AllocateFollicleTextureResources(UTexture2D* Out, con
 	Out->PlatformData->SizeY = Resolution.Y;
 	Out->PlatformData->PixelFormat = PF_B8G8R8A8;
 
-	Out->UpdateResource();
+	// No need to create the resources when using the CPU path.
+	//Out->UpdateResource();
 }
 
 #if WITH_EDITOR
@@ -472,7 +473,7 @@ static void InternalAllocateStrandsTexture(UTexture2D* Out, const FIntPoint& Res
 	Out->PlatformData->SizeY = Resolution.Y;
 	Out->PlatformData->PixelFormat = Format;
 
-	Out->UpdateResource();
+	//Out->UpdateResource();
 }
 
 static void InternalAllocateStrandsTexture_Depth(UTexture2D* Out, const FIntPoint& Resolution, uint32 MipCount)
@@ -757,8 +758,7 @@ static void AddReadbackPass(
 	[InputTexture, OutTexture, BytePerPixel](FRHICommandListImmediate& RHICmdList)
 	{
 		const FIntPoint Resolution = InputTexture->Desc.Extent;
-		check(OutTexture->GetSurfaceWidth() == Resolution.X);
-		check(OutTexture->GetSurfaceHeight() == Resolution.Y);
+		check(OutTexture);
 
 		FRHIResourceCreateInfo CreateInfo(TEXT("ReadbackPass_StagingTexture"));
 		FTexture2DRHIRef StagingTexture = RHICreateTexture2D(
@@ -781,6 +781,9 @@ static void AddReadbackPass(
 
 		// don't think we can mark package as a dirty in the package build
 	#if WITH_EDITORONLY_DATA
+		check(OutTexture->Source.GetSizeX() == Resolution.X);
+		check(OutTexture->Source.GetSizeY() == Resolution.Y);
+
 		void* InData = nullptr;
 		int32 Width = 0, Height = 0;
 		RHICmdList.MapStagingSurface(StagingTexture, InData, Width, Height);
