Fix crash when building raytracing geometry with groom bound to skel. mesh with rigid binding.

diff --git a/Engine/Plugins/Runtime/HairStrands/Source/HairStrandsCore/Private/HairStrandsRendering.cpp b/Engine/Plugins/Runtime/HairStrands/Source/HairStrandsCore/Private/HairStrandsRendering.cpp
--- a/Engine/Plugins/Runtime/HairStrands/Source/HairStrandsCore/Private/HairStrandsRendering.cpp
+++ b/Engine/Plugins/Runtime/HairStrands/Source/HairStrandsCore/Private/HairStrandsRendering.cpp
@@ -1177,7 +1177,11 @@ static void BuildHairAccelerationStructure_Cards(FRHICommandList& RHICmdList,
 	Initializer.bFastBuild = true;
 	Initializer.bAllowUpdate = true;
 
-	FBufferRHIRef PositionBuffer(DeformedResource->GetBuffer(FHairCardsDeformedResource::Current).Buffer->GetRHI()); // This will likely flicker result in half speed motion as everyother frame will use the wrong buffer
+	FBufferRHIRef PositionBuffer = RestResource->RestPositionBuffer.VertexBufferRHI;
+	if (DeformedResource)
+	{
+		PositionBuffer = DeformedResource->GetBuffer(FHairCardsDeformedResource::Current).Buffer->GetRHI(); // This will likely flicker result in half speed motion as everyother frame will use the wrong buffer
+	}
 
 	FRayTracingGeometrySegment Segment;
 	Segment.VertexBuffer = PositionBuffer;
@@ -1206,7 +1210,11 @@ static void BuildHairAccelerationStructure_Meshes(FRHICommandList& RHICmdList,
 	Initializer.bFastBuild = true;
 	Initializer.bAllowUpdate = true;
 
-	FBufferRHIRef PositionBuffer(DeformedResource->GetBuffer(FHairMeshesDeformedResource::Current).Buffer->GetRHI()); // This will likely flicker result in half speed motion as everyother frame will use the wrong buffer
+	FBufferRHIRef PositionBuffer = RestResource->RestPositionBuffer.VertexBufferRHI;
+	if (DeformedResource)
+	{	
+		PositionBuffer = DeformedResource->GetBuffer(FHairMeshesDeformedResource::Current).Buffer->GetRHI(); // This will likely flicker result in half speed motion as everyother frame will use the wrong buffer
+	}
 
 	FRayTracingGeometrySegment Segment;
 	Segment.VertexBuffer = PositionBuffer;
