Use groom gpu offset buffer for hair sim

diff --git a/Engine/Plugins/FX/Niagara/Source/NiagaraCore/Private/NiagaraCustomVersion.cpp b/Engine/Plugins/FX/Niagara/Source/NiagaraCore/Private/NiagaraCustomVersion.cpp
--- a/Engine/Plugins/FX/Niagara/Source/NiagaraCore/Private/NiagaraCustomVersion.cpp
+++ b/Engine/Plugins/FX/Niagara/Source/NiagaraCore/Private/NiagaraCustomVersion.cpp
@@ -8,4 +8,5 @@ const FGuid FNiagaraCustomVersion::GUID(0xFCF57AFA, 0x50764283, 0xB9A9E658, 0xFF
 // Register the custom version with core
 FCustomVersionRegistration GRegisterNiagaraCustomVersion(FNiagaraCustomVersion::GUID, FNiagaraCustomVersion::LatestVersion, TEXT("NiagaraVer"));
 
-const FGuid FNiagaraCustomVersion::LatestScriptCompileVersion(0x3B921B8D, 0x31F94404, 0xBE937958, 0x8C8332AE);
+const FGuid FNiagaraCustomVersion::LatestScriptCompileVersion(0x12D464DC, 0x39F84AAB, 0xAEC3B076, 0xD1FCCE02);
+
diff --git a/Engine/Plugins/Runtime/HairStrands/Shaders/Private/NiagaraDataInterfaceHairStrands.ush b/Engine/Plugins/Runtime/HairStrands/Shaders/Private/NiagaraDataInterfaceHairStrands.ush
--- a/Engine/Plugins/Runtime/HairStrands/Shaders/Private/NiagaraDataInterfaceHairStrands.ush
+++ b/Engine/Plugins/Runtime/HairStrands/Shaders/Private/NiagaraDataInterfaceHairStrands.ush
@@ -35,7 +35,7 @@ struct FDIHairStrandsContext
 	float3				RestRootOffset;
 	float3				DeformedRootOffset;
 	float3				RestPositionOffset;
-	float3				DeformedPositionOffset;
+	Buffer<float4>		DeformedPositionOffset;
 	bool				ResetSimulation;
 	int					SampleCount;
 	int4				BoundingBoxOffsets;
@@ -64,7 +64,7 @@ int					InterpolationMode_##NAME;\
 float3				RestRootOffset_##NAME;\
 float3				DeformedRootOffset_##NAME;\
 float3				RestPositionOffset_##NAME;\
-float3				DeformedPositionOffset_##NAME;\
+Buffer<float4>		DeformedPositionOffset_##NAME;\
 bool				ResetSimulation_##NAME;\
 RWBuffer<uint>	    BoundingBoxBuffer_##NAME;\
 int					SampleCount_##NAME;\
@@ -594,7 +594,7 @@ void DIHairStrands_ResetPointPosition(in FDIHairStrandsContext DIContext, out bo
 		const int PointEnd = DIContext.CurvesOffsetsBuffer[StrandIndex+1]; 
 		const int PointCount = PointEnd - PointBegin;
 
-		const uint4 PositionOffset = uint4( f32tof16(DIContext.RestPositionOffset - DIContext.DeformedPositionOffset), 0 );
+		const uint4 PositionOffset = uint4( f32tof16(DIContext.RestPositionOffset - DIContext.DeformedPositionOffset[0].xyz), 0 );
 
 		ReportStatus = true;
 		for (int PointIndex = 0; PointIndex < PointCount; ++PointIndex)
@@ -698,7 +698,7 @@ void DIHairStrands_UpdatePointPosition(in FDIHairStrandsContext DIContext, in fl
 	const int PointEnd = DIContext.CurvesOffsetsBuffer[StrandIndex + 1];
 	const int PointCount = PointEnd - PointBegin;
 
-	const float3 PositionOffset = DIContext.RestPositionOffset - DIContext.DeformedPositionOffset;
+	const float3 PositionOffset = DIContext.RestPositionOffset - DIContext.DeformedPositionOffset[0].xyz;
 
 	for (int PointIndex = 0; PointIndex < PointCount; ++PointIndex)
 	{
@@ -717,7 +717,7 @@ void DIHairStrands_UpdatePointPosition(in FDIHairStrandsContext DIContext, in fl
 			const float3 LocalPosition = DIHairStrands_TriangleLocalPosition(RestTrianglePosition, RestTriangleOrientation, 
 				f16tof32(DIContext.RestPositionBuffer[PointBegin + PointIndex].xyz) + DIContext.RestPositionOffset) + LocalDisplace;
 			
-			const float3 WorldPosition = DIHairStrands_TriangleWorldPosition(DeformedTrianglePosition, DeformedTriangleOrientation, LocalPosition) - DIContext.DeformedPositionOffset;
+			const float3 WorldPosition = DIHairStrands_TriangleWorldPosition(DeformedTrianglePosition, DeformedTriangleOrientation, LocalPosition) - DIContext.DeformedPositionOffset[0].xyz;
 			DIContext.DeformedPositionBuffer[PointBegin + PointIndex] = uint4(f32tof16(WorldPosition), DIContext.RestPositionBuffer[PointBegin + PointIndex].w);
 		}
 	}
@@ -726,7 +726,7 @@ void DIHairStrands_UpdatePointPosition(in FDIHairStrandsContext DIContext, in fl
 // Get the point position
 void DIHairStrands_GetPointPosition(in FDIHairStrandsContext DIContext, in int PointIndex, out float3 OutPointPosition)
 {
-	OutPointPosition = mul(float4( f16tof32(DIContext.DeformedPositionBuffer[PointIndex].xyz)+DIContext.DeformedPositionOffset, 1.0), DIContext.WorldTransform).xyz;
+	OutPointPosition = mul(float4(f16tof32(DIContext.DeformedPositionBuffer[PointIndex].xyz) + DIContext.DeformedPositionOffset[0].xyz, 1.0), DIContext.WorldTransform).xyz;
 }
 
 /* -----------------------------------------------------------------
diff --git a/Engine/Plugins/Runtime/HairStrands/Source/HairStrandsCore/Private/GroomComponent.cpp b/Engine/Plugins/Runtime/HairStrands/Source/HairStrandsCore/Private/GroomComponent.cpp
--- a/Engine/Plugins/Runtime/HairStrands/Source/HairStrandsCore/Private/GroomComponent.cpp
+++ b/Engine/Plugins/Runtime/HairStrands/Source/HairStrandsCore/Private/GroomComponent.cpp
@@ -1720,6 +1720,20 @@ FHairStrandsRestRootResource* UGroomComponent::GetGuideStrandsRestRootResource(u
 	return HairGroupInstances[GroupIndex]->Guides.RestRootResource;
 }
 
+const FTransform& UGroomComponent::GetGuideStrandsLocalToWorld(uint32 GroupIndex) const
+{
+	if (GroupIndex >= uint32(HairGroupInstances.Num()))
+	{
+		return FTransform::Identity;
+	}
+
+	if (!HairGroupInstances[GroupIndex]->Guides.IsValid())
+	{
+		return FTransform::Identity;
+	}
+	return HairGroupInstances[GroupIndex]->LocalToWorld;
+}
+
 FHairStrandsDeformedRootResource* UGroomComponent::GetGuideStrandsDeformedRootResource(uint32 GroupIndex)
 {
 	if (GroupIndex >= uint32(HairGroupInstances.Num()))
diff --git a/Engine/Plugins/Runtime/HairStrands/Source/HairStrandsCore/Private/Niagara/NiagaraDataInterfaceHairStrands.cpp b/Engine/Plugins/Runtime/HairStrands/Source/HairStrandsCore/Private/Niagara/NiagaraDataInterfaceHairStrands.cpp
--- a/Engine/Plugins/Runtime/HairStrands/Source/HairStrandsCore/Private/Niagara/NiagaraDataInterfaceHairStrands.cpp
+++ b/Engine/Plugins/Runtime/HairStrands/Source/HairStrandsCore/Private/Niagara/NiagaraDataInterfaceHairStrands.cpp
@@ -411,13 +411,12 @@ void FNDIHairStrandsData::Release()
 	}
 }
 
-void FNDIHairStrandsData::Update(UNiagaraDataInterfaceHairStrands* Interface, FNiagaraSystemInstance* SystemInstance, const FHairStrandsDatas* StrandsDatas,
-	UGroomAsset* GroomAsset, const int32 GroupIndex, const int32 LODIndex)
+void FNDIHairStrandsData::Update(UNiagaraDataInterfaceHairStrands* Interface, FNiagaraSystemInstance* SystemInstance, const FHairStrandsDatas* StrandsDatas, 
+	UGroomAsset* GroomAsset, const int32 GroupIndex, const int32 LODIndex, const FTransform& LocalToWorld)
 {
 	if (Interface != nullptr)
 	{
-		WorldTransform = Interface->IsComponentValid() ? Interface->SourceComponent->GetComponentToWorld() :
-			SystemInstance ? SystemInstance->GetWorldTransform() : FTransform::Identity;
+		WorldTransform = LocalToWorld;
 
 		GlobalInterpolation = (Interface->IsComponentValid() && Interface->SourceComponent->BindingAsset && Interface->SourceComponent->GroomAsset) ?
 			Interface->SourceComponent->GroomAsset->EnableGlobalInterpolation : false;
@@ -496,8 +495,9 @@ bool FNDIHairStrandsData::Init(UNiagaraDataInterfaceHairStrands* Interface, FNia
 		int32 LODIndex = 0;
 
 		{
-			Interface->ExtractDatasAndResources(SystemInstance, StrandsDatas, StrandsRestResource, StrandsDeformedResource, StrandsRestRootResource, StrandsDeformedRootResource, GroomAsset, GroupIndex, LODIndex);
-			Update(Interface, SystemInstance, StrandsDatas, GroomAsset, GroupIndex, LODIndex);
+			FTransform LocalToWorld = FTransform::Identity;
+			Interface->ExtractDatasAndResources(SystemInstance, StrandsDatas, StrandsRestResource, StrandsDeformedResource, StrandsRestRootResource, StrandsDeformedRootResource, GroomAsset, GroupIndex, LODIndex, LocalToWorld);
+			Update(Interface, SystemInstance, StrandsDatas, GroomAsset, GroupIndex, LODIndex, LocalToWorld);
 
 			HairStrandsBuffer = new FNDIHairStrandsBuffer();
 			HairStrandsBuffer->Initialize(StrandsDatas, StrandsRestResource, StrandsDeformedResource, StrandsRestRootResource, StrandsDeformedRootResource, ParamsScale);
@@ -566,7 +566,7 @@ struct FNDIHairStrandsParametersCS : public FNiagaraDataInterfaceParametersCS
 
 		
 	}
-	
+
 	void Set(FRHICommandList& RHICmdList, const FNiagaraDataInterfaceSetArgs& Context) const
 	{
 		check(IsInRenderingThread()); 
@@ -578,13 +578,12 @@ struct FNDIHairStrandsParametersCS : public FNiagaraDataInterfaceParametersCS
 		FNDIHairStrandsData* ProxyData =
 			InterfaceProxy->SystemInstancesToProxyData.Find(Context.SystemInstanceID);
 
-
 		const bool IsHairValid = ProxyData != nullptr && ProxyData->HairStrandsBuffer != nullptr && ProxyData->HairStrandsBuffer->IsInitialized();
 		const bool bHasSkinningBinding = IsHairValid && ProxyData->HairGroupInstance && ProxyData->HairGroupInstance->BindingType == EHairBindingType::Skinning;
 		const bool IsRootValid = IsHairValid && ProxyData->HairStrandsBuffer->SourceDeformedRootResources != nullptr && bHasSkinningBinding;//&& ProxyData->HairStrandsBuffer->SourceRootResources->IsInitialized();
 		const bool IsRestValid = IsHairValid && ProxyData->HairStrandsBuffer->SourceRestResources != nullptr;//&& ProxyData->HairStrandsBuffer->SourceRestResources->IsInitialized();
 		const bool IsDeformedValid = IsHairValid && ProxyData->HairStrandsBuffer->SourceDeformedResources != nullptr && ProxyData->HairStrandsBuffer->SourceDeformedResources->IsInitialized();
-		
+
 		const bool bHasLODSwitched = IsHairValid && ProxyData->HairGroupInstance && ProxyData->HairGroupInstance->HairGroupPublicData && ProxyData->HairGroupInstance->HairGroupPublicData->VFInput.bHasLODSwitch;
 
 		if (IsHairValid && IsRestValid)
@@ -596,6 +595,9 @@ struct FNDIHairStrandsParametersCS : public FNiagaraDataInterfaceParametersCS
 			FUnorderedAccessViewRHIRef PointPositionsUAV = IsDeformedValid ?
 				HairStrandsBuffer->SourceDeformedResources->DeformedPositionBuffer[HairStrandsBuffer->SourceDeformedResources->CurrentIndex].UAV : HairStrandsBuffer->DeformedPositionBuffer.UAV;
 
+			FRHIShaderResourceView* DeformedPositionOffsetSRV = IsDeformedValid ? 
+				HairStrandsBuffer->SourceDeformedResources->DeformedOffsetBuffer[HairStrandsBuffer->SourceDeformedResources->CurrentIndex].SRV.GetReference() : FNiagaraRenderer::GetDummyFloatBuffer();
+
 			const int32 MeshLODIndex = IsRootValid ? HairStrandsBuffer->SourceDeformedRootResources->MeshLODIndex : -1;
 			EHairSimulationInterpolationMode InterpolationModeValue = EHairSimulationInterpolationMode::Rigid;
 			if (IsRootValid && HairStrandsBuffer->SourceDeformedRootResources->IsValid(MeshLODIndex))
@@ -644,8 +646,7 @@ struct FNDIHairStrandsParametersCS : public FNiagaraDataInterfaceParametersCS
 			//UE_LOG(LogHairStrands, Log, TEXT("Shader Reset : %d %d %d"), bNeedSimReset, IsRootValid, InterpolationModeValue);
 			FVector RestRootOffsetValue = FVector::ZeroVector;
 			FVector DeformedRootOffsetValue = FVector::ZeroVector;
-			FVector DeformedPositionOffsetValue =  IsDeformedValid ? HairStrandsBuffer->SourceDeformedResources->GetPositionOffset(FHairStrandsDeformedResource::EFrameType::Current) :
-																	 HairStrandsBuffer->SourceRestResources->PositionOffset;
+			
 			FVector RestPositionOffsetValue = ProxyData->HairStrandsBuffer->SourceRestResources->PositionOffset;
 
 			FRHITransitionInfo Transitions[] = {
@@ -674,7 +675,7 @@ struct FNDIHairStrandsParametersCS : public FNiagaraDataInterfaceParametersCS
 			SetShaderValue(RHICmdList, ComputeShaderRHI, DeformedRootOffset, DeformedRootOffsetValue);
 
 			SetShaderValue(RHICmdList, ComputeShaderRHI, RestPositionOffset, RestPositionOffsetValue);
-			SetShaderValue(RHICmdList, ComputeShaderRHI, DeformedPositionOffset, DeformedPositionOffsetValue);
+			SetSRVParameter(RHICmdList, ComputeShaderRHI, DeformedPositionOffset, DeformedPositionOffsetSRV);
 
 			SetSRVParameter(RHICmdList, ComputeShaderRHI, RestTrianglePositionABuffer, RestTrianglePositionASRV);
 			SetSRVParameter(RHICmdList, ComputeShaderRHI, RestTrianglePositionBBuffer, RestTrianglePositionBSRV);
@@ -713,7 +714,7 @@ struct FNDIHairStrandsParametersCS : public FNiagaraDataInterfaceParametersCS
 			SetShaderValue(RHICmdList, ComputeShaderRHI, DeformedRootOffset, FVector4(0, 0, 0, 0));
 
 			SetShaderValue(RHICmdList, ComputeShaderRHI, RestPositionOffset, FVector(0,0,0));
-			SetShaderValue(RHICmdList, ComputeShaderRHI, DeformedPositionOffset, FVector(0,0,0));
+			SetSRVParameter(RHICmdList, ComputeShaderRHI, DeformedPositionOffset, FNiagaraRenderer::GetDummyFloatBuffer());
 
 			SetSRVParameter(RHICmdList, ComputeShaderRHI, RestTrianglePositionABuffer, FNiagaraRenderer::GetDummyFloatBuffer());
 			SetSRVParameter(RHICmdList, ComputeShaderRHI, RestTrianglePositionBBuffer, FNiagaraRenderer::GetDummyFloatBuffer());
@@ -771,7 +772,7 @@ private:
 	LAYOUT_FIELD(FShaderResourceParameter, MeshSampleWeightsBuffer);
 
 	LAYOUT_FIELD(FShaderParameter, RestPositionOffset);
-	LAYOUT_FIELD(FShaderParameter, DeformedPositionOffset);
+	LAYOUT_FIELD(FShaderResourceParameter, DeformedPositionOffset);
 
 	LAYOUT_FIELD(FShaderResourceParameter, ParamsScaleBuffer);
 	LAYOUT_FIELD(FShaderResourceParameter, BoundingBoxBuffer);
@@ -896,8 +897,9 @@ void UNiagaraDataInterfaceHairStrands::ExtractDatasAndResources(
 	FHairStrandsRestRootResource*& OutStrandsRestRootResource, 
 	FHairStrandsDeformedRootResource*& OutStrandsDeformedRootResource,
 	UGroomAsset*& OutGroomAsset,
-	int32& OutGroupIndex,
-	int32& OutLODIndex)
+	int32& OutGroupIndex, 
+	int32& OutLODIndex,
+	FTransform& OutLocalToWorld)
 {
 	ExtractSourceComponent(SystemInstance);
 
@@ -931,12 +933,14 @@ void UNiagaraDataInterfaceHairStrands::ExtractDatasAndResources(
 			OutStrandsDeformedRootResource = SourceComponent->GetGuideStrandsDeformedRootResource(OutGroupIndex);
 			OutGroomAsset = SourceComponent->GroomAsset;
 			OutLODIndex = SourceComponent->GetForcedLOD();
+			OutLocalToWorld = SourceComponent->GetGuideStrandsLocalToWorld(OutGroupIndex);
 		}
 	}
 	else if (DefaultSource != nullptr)
 	{
 		OutGroupIndex = 0;
 		OutLODIndex = 0;
+		OutLocalToWorld = FTransform::Identity;
 		if (OutGroupIndex < DefaultSource->GetNumHairGroups())
 		{
 			OutStrandsDatas = &DefaultSource->HairGroupsData[OutGroupIndex].Strands.Data;
@@ -994,9 +998,10 @@ bool UNiagaraDataInterfaceHairStrands::PerInstanceTick(void* PerInstanceData, FN
 	int32 GroupIndex = 0;
 	int32 LODIndex = 0;
 
-	InstanceData->TickCount = FMath::Min(GHairSimulationMaxDelay +1,InstanceData->TickCount+1);
+	InstanceData->TickCount = FMath::Min(GHairSimulationMaxDelay + 1, InstanceData->TickCount + 1);
 
-	ExtractDatasAndResources(SystemInstance, StrandsDatas, StrandsRestResource, StrandsDeformedResource, StrandsRestRootResource, StrandsDeformedRootResource, GroomAsset, GroupIndex, LODIndex);
+	FTransform LocalToWorld = FTransform::Identity;
+	ExtractDatasAndResources(SystemInstance, StrandsDatas, StrandsRestResource, StrandsDeformedResource, StrandsRestRootResource, StrandsDeformedRootResource, GroomAsset, GroupIndex, LODIndex, LocalToWorld);
 	InstanceData->HairStrandsBuffer->Update(StrandsDatas, StrandsRestResource, StrandsDeformedResource, StrandsRestRootResource, StrandsDeformedRootResource);
 
 	if (SourceComponent != nullptr)
@@ -1007,7 +1012,7 @@ bool UNiagaraDataInterfaceHairStrands::PerInstanceTick(void* PerInstanceData, FN
 		}
 		InstanceData->ForceReset = SourceComponent->bResetSimulation;
 	}
-	InstanceData->Update(this, SystemInstance, StrandsDatas, GroomAsset, GroupIndex, LODIndex);
+	InstanceData->Update(this, SystemInstance, StrandsDatas, GroomAsset, GroupIndex, LODIndex, LocalToWorld);
 	return false;
 }
 
diff --git a/Engine/Plugins/Runtime/HairStrands/Source/HairStrandsCore/Public/GroomComponent.h b/Engine/Plugins/Runtime/HairStrands/Source/HairStrandsCore/Public/GroomComponent.h
--- a/Engine/Plugins/Runtime/HairStrands/Source/HairStrandsCore/Public/GroomComponent.h
+++ b/Engine/Plugins/Runtime/HairStrands/Source/HairStrandsCore/Public/GroomComponent.h
@@ -123,6 +123,8 @@ public:
 	/** Return the guide hairs root resources*/
 	FHairStrandsRestRootResource* GetGuideStrandsRestRootResource(uint32 GroupIndex);
 	FHairStrandsDeformedRootResource* GetGuideStrandsDeformedRootResource(uint32 GroupIndex);
+	const FTransform& GetGuideStrandsLocalToWorld(uint32 GroupIndex) const;
+
 
 #if WITH_EDITOR
 	virtual void CheckForErrors() override;
diff --git a/Engine/Plugins/Runtime/HairStrands/Source/HairStrandsCore/Public/Niagara/NiagaraDataInterfaceHairStrands.h b/Engine/Plugins/Runtime/HairStrands/Source/HairStrandsCore/Public/Niagara/NiagaraDataInterfaceHairStrands.h
--- a/Engine/Plugins/Runtime/HairStrands/Source/HairStrandsCore/Public/Niagara/NiagaraDataInterfaceHairStrands.h
+++ b/Engine/Plugins/Runtime/HairStrands/Source/HairStrandsCore/Public/Niagara/NiagaraDataInterfaceHairStrands.h
@@ -97,7 +97,7 @@ struct FNDIHairStrandsData
 	void Release();
 
 	/** Update the buffers */
-	void Update(UNiagaraDataInterfaceHairStrands* Interface, FNiagaraSystemInstance* SystemInstance, const FHairStrandsDatas* HairStrandsDatas, UGroomAsset* GroomAsset, const int32 GroupIndex, const int32 LODIndex);
+	void Update(UNiagaraDataInterfaceHairStrands* Interface, FNiagaraSystemInstance* SystemInstance, const FHairStrandsDatas* HairStrandsDatas, UGroomAsset* GroomAsset, const int32 GroupIndex, const int32 LODIndex, const FTransform& LocalToWorld);
 
 	inline void ResetDatas()
 	{
@@ -364,7 +364,8 @@ public:
 		FHairStrandsDeformedRootResource*& OutStrandsDeformedRootResource,
 		UGroomAsset*& OutGroomAsset,
 		int32& OutGroupIndex,
-		int32& OutLODIndex);
+		int32& OutLODIndex, 
+		FTransform& OutLocalToWorld);
 
 	/** Get the number of strands */
 	void GetNumStrands(FVectorVMContext& Context);
