Fix hair guide debug mode not being correctly displayed.

diff --git a/Engine/Plugins/Runtime/HairStrands/Source/HairStrandsCore/Private/HairStrandsRendering.cpp b/Engine/Plugins/Runtime/HairStrands/Source/HairStrandsCore/Private/HairStrandsRendering.cpp
--- a/Engine/Plugins/Runtime/HairStrands/Source/HairStrandsCore/Private/HairStrandsRendering.cpp
+++ b/Engine/Plugins/Runtime/HairStrands/Source/HairStrandsCore/Private/HairStrandsRendering.cpp
@@ -1433,6 +1433,9 @@ void ComputeHairStrandsInterpolation(
 
 		if (DebugMode == EHairStrandsDebugMode::SimHairStrands)
 		{
+			// Disable culling when drawing only guides, as the culling output has been computed for the strands, not for the guides.
+			Instance->HairGroupPublicData->SetCullingResultAvailable(false);
+
 			check(Instance->Guides.DeformedResource);
 			AddHairTangentPass(
 				GraphBuilder,
