Fix hair strands simulation not be enabled/used when not attached to a skeletal mesh.

diff --git a/Engine/Plugins/Runtime/HairStrands/Source/HairStrandsCore/Private/GroomComponent.cpp b/Engine/Plugins/Runtime/HairStrands/Source/HairStrandsCore/Private/GroomComponent.cpp
--- a/Engine/Plugins/Runtime/HairStrands/Source/HairStrandsCore/Private/GroomComponent.cpp
+++ b/Engine/Plugins/Runtime/HairStrands/Source/HairStrandsCore/Private/GroomComponent.cpp
@@ -2087,20 +2087,20 @@ void UGroomComponent::InitResources(bool bIsBindingReloading)
 			// Check if strands needs to load/allocate root resources
 			bool bNeedRootResources = false;
 			bool bNeedDynamicResources = false;
-			if (LocalBindingAsset)
 			{
 				const FHairGroupsLOD& GroupLOD = GroomAsset->HairGroupsLOD[GroupIt];
 				for (uint32 LODIt = 0, LODCount = GroupLOD.LODs.Num(); LODIt < LODCount; ++LODIt)
 				{
 					if (HairGroupInstance->HairGroupPublicData->GetGeometryType(LODIt) == EHairGeometryType::Strands)
 					{
-						if (HairGroupInstance->HairGroupPublicData->GetBindingType(LODIt) == EHairBindingType::Skinning ||
-							HairGroupInstance->HairGroupPublicData->IsGlobalInterpolationEnable(LODIt))
+						if (LocalBindingAsset &&
+							(HairGroupInstance->HairGroupPublicData->GetBindingType(LODIt) == EHairBindingType::Skinning ||
+							 HairGroupInstance->HairGroupPublicData->IsGlobalInterpolationEnable(LODIt)))
 						{
 							bNeedRootResources = true;
 						}
 
-						if (HairGroupInstance->HairGroupPublicData->GetBindingType(LODIt) == EHairBindingType::Skinning ||
+						if ((LocalBindingAsset && HairGroupInstance->HairGroupPublicData->GetBindingType(LODIt) == EHairBindingType::Skinning) ||
 							HairGroupInstance->HairGroupPublicData->IsSimulationEnable(LODIt))
 						{
 							bNeedDynamicResources = true;
