Fix hair sim reset when skel mesh hierarchy is changing

diff --git a/Engine/Plugins/Runtime/HairStrands/Source/HairStrandsCore/Private/Niagara/NiagaraDataInterfaceHairStrands.cpp b/Engine/Plugins/Runtime/HairStrands/Source/HairStrandsCore/Private/Niagara/NiagaraDataInterfaceHairStrands.cpp
--- a/Engine/Plugins/Runtime/HairStrands/Source/HairStrandsCore/Private/Niagara/NiagaraDataInterfaceHairStrands.cpp
+++ b/Engine/Plugins/Runtime/HairStrands/Source/HairStrandsCore/Private/Niagara/NiagaraDataInterfaceHairStrands.cpp
@@ -819,6 +819,30 @@ void FNDIHairStrandsProxy::DestroyPerInstanceData(NiagaraEmitterInstanceBatcher*
 
 //------------------------------------------------------------------------------------------------------------
 
+FORCEINLINE bool RequiresSimulationReset(FNiagaraSystemInstance* SystemInstance, uint32& OldSkeletalMeshes)
+{
+	uint32 NewSkeletalMeshes = 0;
+	if (USceneComponent* AttachComponent = SystemInstance->GetAttachComponent())
+	{
+		if (AActor* RootActor = AttachComponent->GetAttachmentRootActor())
+		{
+			for (UActorComponent* ActorComp : RootActor->GetComponents())
+			{
+				USkeletalMeshComponent* SkelMeshComp = Cast<USkeletalMeshComponent>(ActorComp);
+				if (SkelMeshComp && SkelMeshComp->SkeletalMesh)
+				{
+					NewSkeletalMeshes += GetTypeHash(SkelMeshComp->SkeletalMesh->GetName());
+				}
+			}
+		}
+	}
+	bool bNeedReset = NewSkeletalMeshes != OldSkeletalMeshes;
+	OldSkeletalMeshes = NewSkeletalMeshes;
+	return bNeedReset;
+}
+
+//------------------------------------------------------------------------------------------------------------
+
 UNiagaraDataInterfaceHairStrands::UNiagaraDataInterfaceHairStrands(FObjectInitializer const& ObjectInitializer)
 	: Super(ObjectInitializer)
 	, DefaultSource(nullptr)
@@ -1006,11 +1030,12 @@ bool UNiagaraDataInterfaceHairStrands::PerInstanceTick(void* PerInstanceData, FN
 
 	if (SourceComponent != nullptr)
 	{
-		if (SourceComponent->bResetSimulation)
+		const bool bRequiresReset = SourceComponent->bResetSimulation || RequiresSimulationReset(SystemInstance, InstanceData->SkeletalMeshes);
+		if (bRequiresReset)
 		{
 			InstanceData->TickCount = 0;
 		}
-		InstanceData->ForceReset = SourceComponent->bResetSimulation;
+		InstanceData->ForceReset = bRequiresReset;
 	}
 	InstanceData->Update(this, SystemInstance, StrandsDatas, GroomAsset, GroupIndex, LODIndex, LocalToWorld);
 	return false;
diff --git a/Engine/Plugins/Runtime/HairStrands/Source/HairStrandsCore/Public/Niagara/NiagaraDataInterfaceHairStrands.h b/Engine/Plugins/Runtime/HairStrands/Source/HairStrandsCore/Public/Niagara/NiagaraDataInterfaceHairStrands.h
--- a/Engine/Plugins/Runtime/HairStrands/Source/HairStrandsCore/Public/Niagara/NiagaraDataInterfaceHairStrands.h
+++ b/Engine/Plugins/Runtime/HairStrands/Source/HairStrandsCore/Public/Niagara/NiagaraDataInterfaceHairStrands.h
@@ -146,6 +146,7 @@ struct FNDIHairStrandsData
 		{
 			ParamsScale[i] = 1.0;
 		}
+		SkeletalMeshes = 0;
 	}
 
 	inline void CopyDatas(const FNDIHairStrandsData* OtherDatas)
@@ -197,6 +198,8 @@ struct FNDIHairStrandsData
 
 			ParamsScale = OtherDatas->ParamsScale;
 
+			SkeletalMeshes = OtherDatas->SkeletalMeshes;
+
 			TickingGroup = OtherDatas->TickingGroup;
 		}
 	}
@@ -300,6 +303,9 @@ struct FNDIHairStrandsData
 	/** Scales along the strand */
 	TStaticArray<float, 32 * NumScales> ParamsScale;
 
+	/** List of all the skel meshes in the hierarchy*/
+	uint32 SkeletalMeshes;
+
 	/** The instance ticking group */
 	ETickingGroup TickingGroup;
 };
