Remove primitive Id stream support from HairStrands factory.

diff --git a/Engine/Plugins/Runtime/HairStrands/Source/HairStrandsCore/Private/HairStrandsVertexFactory.cpp b/Engine/Plugins/Runtime/HairStrands/Source/HairStrandsCore/Private/HairStrandsVertexFactory.cpp
--- a/Engine/Plugins/Runtime/HairStrands/Source/HairStrandsCore/Private/HairStrandsVertexFactory.cpp
+++ b/Engine/Plugins/Runtime/HairStrands/Source/HairStrandsCore/Private/HairStrandsVertexFactory.cpp
@@ -13,7 +13,7 @@
 #include "HairStrandsInterface.h"
 #include "GroomInstance.h"
 
-#define VF_STRANDS_SUPPORT_GPU_SCENE 1
+#define VF_STRANDS_SUPPORT_GPU_SCENE 0
 
 /////////////////////////////////////////////////////////////////////////////////////////
 
