Fix crash when groom is attached to a skel mesh, but the skel mesh has not beein initialized yet.

diff --git a/Engine/Plugins/Runtime/HairStrands/Source/HairStrandsCore/Private/GroomComponent.cpp b/Engine/Plugins/Runtime/HairStrands/Source/HairStrandsCore/Private/GroomComponent.cpp
--- a/Engine/Plugins/Runtime/HairStrands/Source/HairStrandsCore/Private/GroomComponent.cpp
+++ b/Engine/Plugins/Runtime/HairStrands/Source/HairStrandsCore/Private/GroomComponent.cpp
@@ -2547,7 +2547,7 @@ void UGroomComponent::TickComponent(float DeltaTime, enum ELevelTick TickType, F
 
 	// When a groom binding and similation are disabled, and the groom component is parented with a skeletal mesh, we can optionnaly 
 	// attach the groom to a particular socket/bone
-	if (RegisteredSkeletalMeshComponent && !AttachmentName.IsEmpty())
+	if (RegisteredSkeletalMeshComponent && RegisteredSkeletalMeshComponent->SkeletalMesh && !AttachmentName.IsEmpty())
 	{
 		const FName BoneName(AttachmentName);
 		if (GetAttachSocketName() != BoneName)
