Fix hair sim CIS

diff --git a/Engine/Plugins/Runtime/HairStrands/Source/HairStrandsCore/Private/Niagara/NiagaraDataInterfaceHairStrands.cpp b/Engine/Plugins/Runtime/HairStrands/Source/HairStrandsCore/Private/Niagara/NiagaraDataInterfaceHairStrands.cpp
--- a/Engine/Plugins/Runtime/HairStrands/Source/HairStrandsCore/Private/Niagara/NiagaraDataInterfaceHairStrands.cpp
+++ b/Engine/Plugins/Runtime/HairStrands/Source/HairStrandsCore/Private/Niagara/NiagaraDataInterfaceHairStrands.cpp
@@ -7,6 +7,7 @@
 #include "NiagaraSystemInstance.h"
 #include "NiagaraEmitterInstanceBatcher.h"
 
+#include "Components/SkeletalMeshComponent.h"
 #include "ShaderParameterUtils.h"
 #include "ClearQuad.h"
 #include "RenderGraphBuilder.h"
