Fix CIS issue.

diff --git a/Engine/Plugins/Runtime/HairStrands/Source/HairStrandsCore/Private/HairStrandsVertexFactory.cpp b/Engine/Plugins/Runtime/HairStrands/Source/HairStrandsCore/Private/HairStrandsVertexFactory.cpp
--- a/Engine/Plugins/Runtime/HairStrands/Source/HairStrandsCore/Private/HairStrandsVertexFactory.cpp
+++ b/Engine/Plugins/Runtime/HairStrands/Source/HairStrandsCore/Private/HairStrandsVertexFactory.cpp
@@ -186,13 +186,14 @@ void FHairStrandsVertexFactory::ModifyCompilationEnvironment(const FVertexFactor
 
 void FHairStrandsVertexFactory::ValidateCompiledResult(const FVertexFactoryType* Type, EShaderPlatform Platform, const FShaderParameterMap& ParameterMap, TArray<FString>& OutErrors)
 {
-	if (VF_STRANDS_SUPPORT_GPU_SCENE
-		&& Type->SupportsPrimitiveIdStream()
+#if VF_STRANDS_SUPPORT_GPU_SCENE
+	if (Type->SupportsPrimitiveIdStream()
 		&& UseGPUScene(Platform, GetMaxSupportedFeatureLevel(Platform)) 
 		&& ParameterMap.ContainsParameterAllocation(FPrimitiveUniformShaderParameters::StaticStructMetadata.GetShaderVariableName()))
 	{
 		OutErrors.AddUnique(*FString::Printf(TEXT("Shader attempted to bind the Primitive uniform buffer even though Vertex Factory %s computes a PrimitiveId per-instance.  This will break auto-instancing.  Shaders should use GetPrimitiveData(PrimitiveId).Member instead of Primitive.Member."), Type->GetName()));
 	}
+#endif
 }
 
 void FHairStrandsVertexFactory::SetData(const FDataType& InData)
@@ -226,12 +227,12 @@ void FHairStrandsVertexFactory::InitRHI()
 	// We create different streams based on feature level
 	check(HasValidFeatureLevel());
 
-	// VertexFactory needs to be able to support max possible shader platform and feature level
-	// in case if we switch feature level at runtime.
-	const bool bCanUseGPUScene = VF_STRANDS_SUPPORT_GPU_SCENE && UseGPUScene(GMaxRHIShaderPlatform, GMaxRHIFeatureLevel);
-
 	FVertexDeclarationElementList Elements;
 	SetPrimitiveIdStreamIndex(EVertexInputStreamType::Default, -1);
+#if VF_STRANDS_SUPPORT_GPU_SCENE
+	// VertexFactory needs to be able to support max possible shader platform and feature level
+	// in case if we switch feature level at runtime.
+	const bool bCanUseGPUScene = UseGPUScene(GMaxRHIShaderPlatform, GMaxRHIFeatureLevel);
 	if (GetType()->SupportsPrimitiveIdStream() && bCanUseGPUScene)
 	{
 		// When the VF is used for rendering in normal mesh passes, this vertex buffer and offset will be overridden
@@ -239,6 +240,7 @@ void FHairStrandsVertexFactory::InitRHI()
 		SetPrimitiveIdStreamIndex(EVertexInputStreamType::Default, Elements.Last().StreamIndex);
 		bNeedsDeclaration = true;
 	}
+#endif
 
 	if (bNeedsDeclaration)
 	{
