Groom Editor Mode transitioned to UEdMode.

diff --git a/Engine/Plugins/Runtime/HairStrands/Source/HairStrandsEditor/Private/GroomCustomAssetEditorToolkit.cpp b/Engine/Plugins/Runtime/HairStrands/Source/HairStrandsEditor/Private/GroomCustomAssetEditorToolkit.cpp
--- a/Engine/Plugins/Runtime/HairStrands/Source/HairStrandsEditor/Private/GroomCustomAssetEditorToolkit.cpp
+++ b/Engine/Plugins/Runtime/HairStrands/Source/HairStrandsEditor/Private/GroomCustomAssetEditorToolkit.cpp
@@ -15,7 +15,6 @@
 #include "Toolkits/AssetEditorToolkit.h"
 #include "HairStrandsEditor.h"
 #include "GroomEditorCommands.h"
-#include "GroomEditorMode.h"
 #include "GroomEditorStyle.h"
 #include "GroomAssetDetails.h"
 #include "GroomMaterialDetails.h"
@@ -108,12 +107,6 @@ void FGroomCustomAssetEditorToolkit::UnregisterTabSpawners(const TSharedRef<clas
 #endif
 }
 
-FEdMode* FGroomCustomAssetEditorToolkit::GetEditorMode() const 
-{
-	static FGroomEditorMode *Mode = new FGroomEditorMode;
-	return Mode;
-}
-
 static uint32 GOpenedGroomEditorCount = 0;
 
 void FGroomCustomAssetEditorToolkit::DocPropChanged(UObject *InObject, FPropertyChangedEvent &Property)
diff --git a/Engine/Plugins/Runtime/HairStrands/Source/HairStrandsEditor/Private/GroomEditorMode.cpp b/Engine/Plugins/Runtime/HairStrands/Source/HairStrandsEditor/Private/GroomEditorMode.cpp
--- a/Engine/Plugins/Runtime/HairStrands/Source/HairStrandsEditor/Private/GroomEditorMode.cpp
+++ b/Engine/Plugins/Runtime/HairStrands/Source/HairStrandsEditor/Private/GroomEditorMode.cpp
@@ -1,225 +1,35 @@
 // Copyright Epic Games, Inc. All Rights Reserved.
 
 #include "GroomEditorMode.h"
-#include "Toolkits/ToolkitManager.h"
+#include "GroomEditorCommands.h"
 #include "Framework/Commands/UICommandList.h"
-#include "Framework/Application/SlateApplication.h"
-#include "Framework/MultiBox/MultiBoxBuilder.h"
-#include "EditorViewportClient.h"
 #include "EdModeInteractiveToolsContext.h"
+#include "InteractiveToolManager.h"
 
-#include "EditorModeManager.h"
-#include "GroomEditorCommands.h"
-#include "Toolkits/BaseToolkit.h"
-#include "EdMode.h"
-
-#define LOCTEXT_NAMESPACE "FGroomEditorMode"
-
-const FEditorModeID FGroomEditorMode::EM_GroomEditorModeId = TEXT("EM_GroomEditorMode");
-
-FGroomEditorMode::FGroomEditorMode()
-{
-	ToolsContext = nullptr;
-
-	UICommandList = MakeShareable(new FUICommandList);
-}
-
-FGroomEditorMode::~FGroomEditorMode()
-{
-	if (ToolsContext != nullptr)
-	{
-		ToolsContext->ShutdownContext();
-		ToolsContext = nullptr;
-	}
-}
+#define LOCTEXT_NAMESPACE "GroomEditor"
 
-void FGroomEditorMode::ActorSelectionChangeNotify()
-{
-}
+const FEditorModeID UGroomEditorMode::EM_GroomEditorModeId = TEXT("EM_GroomEditorMode");
 
-bool FGroomEditorMode::ProcessEditDelete()
+UGroomEditorMode::UGroomEditorMode()
 {
-	return ToolsContext->ProcessEditDelete();
+	Info = FEditorModeInfo(
+		EM_GroomEditorModeId,
+		LOCTEXT("GroomEditorMode", "Groom Editor"),
+		FSlateIcon(),
+		false);
 }
 
-bool FGroomEditorMode::AllowWidgetMove()
+bool UGroomEditorMode::AllowWidgetMove()
 { 
 	return false; 
 }
 
-bool FGroomEditorMode::ShouldDrawWidget() const
-{ 
-	// allow standard xform gizmo if we don't have an active tool
-	if (ToolsContext != nullptr)
-	{
-		return ToolsContext->ToolManager->HasAnyActiveTool() == false;
-	}
-	return true; 
-}
-
-bool FGroomEditorMode::UsesTransformWidget() const
-{ 
-	return true; 
-}
-
-void FGroomEditorMode::Tick(FEditorViewportClient* ViewportClient, float DeltaTime)
-{
-	FEdMode::Tick(ViewportClient, DeltaTime);
-
-	if (ToolsContext != nullptr)
-	{
-		ToolsContext->Tick(ViewportClient, DeltaTime);
-	}
-}
-
-void FGroomEditorMode::Render(const FSceneView* View, FViewport* Viewport, FPrimitiveDrawInterface* PDI)
-{
-	FEdMode::Render(View, Viewport, PDI);
-
-	if (ToolsContext != nullptr)
-	{
-		ToolsContext->Render(View, Viewport, PDI);
-	}
-}
-
-bool FGroomEditorMode::InputKey(FEditorViewportClient* ViewportClient, FViewport* Viewport, FKey Key, EInputEvent Event)
-{
-	// try hotkeys
-	if (Event != IE_Released)
-	{
-		if (UICommandList->ProcessCommandBindings(Key, FSlateApplication::Get().GetModifierKeys(), false/*Event == IE_Repeat*/))
-		{
-			return true;
-		}
-	}
-
-	bool bHandled = ToolsContext->InputKey(ViewportClient, Viewport, Key, Event);
-	if (bHandled == false)
-	{
-		bHandled = FEdMode::InputKey(ViewportClient, Viewport, Key, Event);
-	}
-	return bHandled;	
-}
-
-bool FGroomEditorMode::InputAxis(FEditorViewportClient* InViewportClient, FViewport* Viewport, int32 ControllerId, FKey Key, float Delta, float DeltaTime)
+void UGroomEditorMode::Enter()
 {
-	// mouse axes: EKeys::MouseX, EKeys::MouseY, EKeys::MouseWheelAxis
-	return FEdMode::InputAxis(InViewportClient, Viewport, ControllerId, Key, Delta, DeltaTime);
-}
-
-bool FGroomEditorMode::StartTracking(FEditorViewportClient* InViewportClient, FViewport* InViewport) 
-{ 
-	bool bHandled = FEdMode::StartTracking(InViewportClient, InViewport);
-#ifdef ENABLE_DEBUG_PRINTING
-	UE_LOG(LogTemp, Warning, TEXT("START TRACKING - base handled was %d"), (int)bHandled);
-#endif
-
-	bHandled |= ToolsContext->StartTracking(InViewportClient, InViewport);
-
-	return bHandled;
-}
-
-bool FGroomEditorMode::CapturedMouseMove(FEditorViewportClient* InViewportClient, FViewport* InViewport, int32 InMouseX, int32 InMouseY)
-{
-	bool bHandled = ToolsContext->CapturedMouseMove(InViewportClient, InViewport, InMouseX, InMouseY);
-	return bHandled;
-}
-
-bool FGroomEditorMode::EndTracking(FEditorViewportClient* InViewportClient, FViewport* InViewport)
-{
-	bool bHandled = ToolsContext->EndTracking(InViewportClient, InViewport);
-	return bHandled;
-}
-
-bool FGroomEditorMode::ReceivedFocus(FEditorViewportClient* ViewportClient, FViewport* Viewport)
-{
-#ifdef ENABLE_DEBUG_PRINTING
-	UE_LOG(LogTemp, Warning, TEXT("RECEIVED FOCUS"));
-#endif
-
-	return false;
-}
-
-bool FGroomEditorMode::LostFocus(FEditorViewportClient* ViewportClient, FViewport* Viewport)
-{
-#ifdef ENABLE_DEBUG_PRINTING
-	UE_LOG(LogTemp, Warning, TEXT("LOST FOCUS"));
-#endif
-
-	return false;
-}
-
-bool FGroomEditorMode::MouseEnter(FEditorViewportClient* ViewportClient, FViewport* Viewport, int32 x, int32 y)
-{
-	bool bHandled = ToolsContext->MouseEnter(ViewportClient, Viewport, x, y);
-	return bHandled;
-}
-
-bool FGroomEditorMode::MouseMove(FEditorViewportClient* ViewportClient, FViewport* Viewport, int32 x, int32 y)
-{
-	bool bHandled = ToolsContext->MouseMove(ViewportClient, Viewport, x, y);
-	return bHandled;
-}
-
-bool FGroomEditorMode::MouseLeave(FEditorViewportClient* ViewportClient, FViewport* Viewport)
-{
-	bool bHandled = ToolsContext->MouseLeave(ViewportClient, Viewport);
-	return bHandled;
-}
-
-void FGroomEditorMode::Enter()
-{
-	FEdMode::Enter();
-
-	ToolsContext = Owner->GetInteractiveToolsContext();
-
-	if (!Toolkit.IsValid() && UsesToolkits())
-	{
-		//Toolkit = MakeShareable(new FGroomEditorModeToolkit);
-		//Toolkit->Init(Owner->GetToolkitHost());
-
-		//const FModelingToolsManagerCommands& ToolManagerCommands = FModelingToolsManagerCommands::Get();
-		//const TSharedRef<FUICommandList>& CommandList = Toolkit->GetToolkitCommands();
-
-		/*CommandList->MapAction(
-			ToolManagerCommands.AcceptActiveTool,
-			FExecuteAction::CreateLambda([this]() { ToolsContext->EndTool(EToolShutdownType::Accept); }),
-			FCanExecuteAction::CreateLambda([this]() { return ToolsContext->CanAcceptActiveTool(); }),
-			FGetActionCheckState(),
-			FIsActionButtonVisible::CreateLambda([this]() {return ToolsContext->ActiveToolHasAccept(); }),
-			EUIActionRepeatMode::RepeatDisabled
-		);
-
-		CommandList->MapAction(
-			ToolManagerCommands.CancelActiveTool,
-			FExecuteAction::CreateLambda([this]() { ToolsContext->EndTool(EToolShutdownType::Cancel); }),
-			FCanExecuteAction::CreateLambda([this]() { return ToolsContext->CanCancelActiveTool(); }),
-			FGetActionCheckState(),
-			FIsActionButtonVisible::CreateLambda([this]() {return ToolsContext->ActiveToolHasAccept(); }),
-			EUIActionRepeatMode::RepeatDisabled
-		);
-
-		CommandList->MapAction(
-			ToolManagerCommands.CompleteActiveTool,
-			FExecuteAction::CreateLambda([this]() { ToolsContext->EndTool(EToolShutdownType::Completed); }),
-			FCanExecuteAction::CreateLambda([this]() { return ToolsContext->CanCompleteActiveTool(); }),
-			FGetActionCheckState(),
-			FIsActionButtonVisible::CreateLambda([this]() {return ToolsContext->CanCompleteActiveTool(); }),
-			EUIActionRepeatMode::RepeatDisabled
-		);*/
-
-	}
+	Super::Enter();
 
 	const FGroomEditorCommands& ToolManagerCommands = FGroomEditorCommands::Get();
 	const TSharedRef<FUICommandList>& CommandList = Toolkit->GetToolkitCommands();
-	auto RegisterToolFunc = [this, &CommandList](TSharedPtr<FUICommandInfo> UICommand, FString ToolIdentifier, UInteractiveToolBuilder* Builder)
-	{
-		ToolsContext->ToolManager->RegisterToolType(ToolIdentifier, Builder);
-		CommandList->MapAction( UICommand,
-			FExecuteAction::CreateLambda([this, ToolIdentifier]() { ToolsContext->StartTool(ToolIdentifier); }),
-			FCanExecuteAction::CreateLambda([this, ToolIdentifier]() { return !ToolsContext->HasActiveTool() && ToolsContext->CanStartTool(ToolIdentifier); }));
-		RegisteredTools.Emplace(UICommand, ToolIdentifier);
-	};
 
 	// register tool set
 
@@ -228,19 +38,9 @@ void FGroomEditorMode::Enter()
 	////
 	//auto HairPlaceToolBuilder = NewObject<UHairPlaceToolBuilder>();
 	//HairPlaceToolBuilder->AssetAPI = ToolsContext->GetAssetAPI();
-	//RegisterToolFunc(ToolManagerCommands.BeginHairPlaceTool, TEXT("HairPlaceTool"), HairPlaceToolBuilder);
+	//RegisterTool(ToolManagerCommands.BeginHairPlaceTool, TEXT("HairPlaceTool"), HairPlaceToolBuilder);
 			
-	ToolsContext->ToolManager->SelectActiveToolType(EToolSide::Left, TEXT("HairPlaceTool"));
-
-	// listen for Tool start/end events to bind/unbind any hotkeys relevant to that Tool
-	ToolsContext->ToolManager->OnToolStarted.AddLambda([this](UInteractiveToolManager* Manager, UInteractiveTool* Tool)
-	{
-		//FModelingToolActionCommands::UpdateToolCommandBinding(Tool, UICommandList, false);
-	});
-	ToolsContext->ToolManager->OnToolEnded.AddLambda([this](UInteractiveToolManager* Manager, UInteractiveTool* Tool)
-	{
-		//FModelingToolActionCommands::UpdateToolCommandBinding(Tool, UICommandList, true);
-	});
+	//ToolsContext->ToolManager->SelectActiveToolType(EToolSide::Left, TEXT("HairPlaceTool"));
 
 #ifdef TOOLED_ENABLE_VIEWPORT_INTERACTION
 	///
@@ -276,23 +76,8 @@ void FGroomEditorMode::Enter()
 #endif  // TOOLED_ENABLE_VIEWPORT_INTERACTION
 }
 
-void FGroomEditorMode::Exit()
+void UGroomEditorMode::Exit()
 {
-	// Shutdown and clean up our additions in the ToolsContext
-	const TSharedRef<FUICommandList>& CommandList = Toolkit->GetToolkitCommands();
-	for (auto& RegisteredTool : RegisteredTools)
-	{
-		CommandList->UnmapAction(RegisteredTool.Key);
-		ToolsContext->ToolManager->UnregisterToolType(RegisteredTool.Value);
-	}
-	RegisteredTools.SetNum(0);
-
-	if (Toolkit.IsValid())
-	{
-		FToolkitManager::Get().CloseToolkit(Toolkit.ToSharedRef());
-		Toolkit.Reset();
-	}
-
 #ifdef TOOLED_ENABLE_VIEWPORT_INTERACTION
 	///
 	// Viewport Interaction
@@ -320,20 +105,7 @@ void FGroomEditorMode::Exit()
 	}
 #endif // TOOLED_ENABLE_VIEWPORT_INTERACTION
 
-	ToolsContext = nullptr;
-
-	// Call base Exit method to ensure proper cleanup
-	FEdMode::Exit();
-}
-
-bool FGroomEditorMode::UsesToolkits() const
-{
-	return true;
-}
-
-void FGroomEditorMode::AddReferencedObjects(FReferenceCollector& Collector)
-{
-	Collector.AddReferencedObject(ToolsContext);
+	Super::Exit();
 }
 
 #undef LOCTEXT_NAMESPACE
\ No newline at end of file
diff --git a/Engine/Plugins/Runtime/HairStrands/Source/HairStrandsEditor/Private/HairStrandsEditor.cpp b/Engine/Plugins/Runtime/HairStrands/Source/HairStrandsEditor/Private/HairStrandsEditor.cpp
--- a/Engine/Plugins/Runtime/HairStrands/Source/HairStrandsEditor/Private/HairStrandsEditor.cpp
+++ b/Engine/Plugins/Runtime/HairStrands/Source/HairStrandsEditor/Private/HairStrandsEditor.cpp
@@ -14,7 +14,6 @@
 #include "FbxHairTranslator.h"
 
 #include "GroomEditorCommands.h"
-#include "GroomEditorMode.h"
 #include "GroomAsset.h"
 #include "GroomComponentDetailsCustomization.h"
 
@@ -107,11 +106,6 @@ void FGroomEditor::StartupModule()
 	}
 
 	FGroomEditorCommands::Register();
-	FEditorModeRegistry::Get().RegisterMode<FGroomEditorMode>(
-		FGroomEditorMode::EM_GroomEditorModeId,
-		LOCTEXT("GroomEditorMode", "Groom Editor"),
-		FSlateIcon(),
-		false);
 
 	// Asset create/edition helper/wrapper for creating/edition asset withn the HairStrandsCore 
 	// project without any editor dependencies
@@ -124,8 +118,6 @@ void FGroomEditor::StartupModule()
 
 void FGroomEditor::ShutdownModule()
 {
-	FEditorModeRegistry::Get().UnregisterMode(FGroomEditorMode::EM_GroomEditorModeId);
-
 	// #ueent_todo: Unregister the translators
 	FAssetToolsModule* AssetToolsModule = FModuleManager::GetModulePtr<FAssetToolsModule>("AssetTools");
 
diff --git a/Engine/Plugins/Runtime/HairStrands/Source/HairStrandsEditor/Public/GroomCustomAssetEditorToolkit.h b/Engine/Plugins/Runtime/HairStrands/Source/HairStrandsEditor/Public/GroomCustomAssetEditorToolkit.h
--- a/Engine/Plugins/Runtime/HairStrands/Source/HairStrandsEditor/Public/GroomCustomAssetEditorToolkit.h
+++ b/Engine/Plugins/Runtime/HairStrands/Source/HairStrandsEditor/Public/GroomCustomAssetEditorToolkit.h
@@ -60,7 +60,6 @@ public:
 	virtual FString GetWorldCentricTabPrefix() const override;
 	virtual FLinearColor GetWorldCentricTabColorScale() const override;
 	virtual bool IsPrimaryEditor() const override { return true; }
-	FEdMode* GetEditorMode() const override;
 	virtual bool OnRequestClose() override;
 	/** End IToolkit interface */
 
diff --git a/Engine/Plugins/Runtime/HairStrands/Source/HairStrandsEditor/Public/GroomEditorMode.h b/Engine/Plugins/Runtime/HairStrands/Source/HairStrandsEditor/Public/GroomEditorMode.h
--- a/Engine/Plugins/Runtime/HairStrands/Source/HairStrandsEditor/Public/GroomEditorMode.h
+++ b/Engine/Plugins/Runtime/HairStrands/Source/HairStrandsEditor/Public/GroomEditorMode.h
@@ -2,44 +2,27 @@
 #pragma once
 
 #include "CoreMinimal.h"
-#include "EdMode.h"
+#include "Tools/LegacyEdModeWidgetHelpers.h"
 
-#include "InputState.h"
-#include "InteractiveToolManager.h"
-#include "EdModeInteractiveToolsContext.h"
+#include "GroomEditorMode.generated.h"
 
-class FEditorComponentSourceFactory;
-class FEditorToolAssetAPI;
-class FUICommandList;
-
-class FGroomEditorMode : public FEdMode
+UCLASS(Transient)
+class UGroomEditorMode : public UBaseLegacyWidgetEdMode
 {
+	GENERATED_BODY()
+
 public:
 	const static FEditorModeID EM_GroomEditorModeId;
-public:
 
-	FGroomEditorMode();
-	virtual ~FGroomEditorMode();
+	UGroomEditorMode();
 
 	////////////////
-	// FEdMode interface
+	// UBaseLegacyWidgetEdMode interface
 	////////////////
 
-	virtual void Tick(FEditorViewportClient* ViewportClient, float DeltaTime) override;
-	virtual void Render(const FSceneView* View, FViewport* Viewport, FPrimitiveDrawInterface* PDI) override;
-	virtual bool UsesToolkits() const override;
-
-	virtual void AddReferencedObjects(FReferenceCollector& Collector) override;
-
 	// these disable the standard gizmo, which is probably want we want in
 	// these tools as we can't hit-test the standard gizmo...
 	virtual bool AllowWidgetMove() override;
-	virtual bool ShouldDrawWidget() const override;
-	virtual bool UsesTransformWidget() const override;
-
-	virtual void ActorSelectionChangeNotify() override;
-
-	virtual bool ProcessEditDelete();
 
 	/*
 	 * focus events
@@ -51,98 +34,10 @@ public:
 	// called when we "end" this editor mode (ie switch to another tab)
 	virtual void Exit() override;
 
-	// called when viewport window is focused
-	virtual bool ReceivedFocus(FEditorViewportClient* ViewportClient, FViewport* Viewport) override;
-
-	// called when viewport window loses focus (ie when some other window is focused)
-	// *not* called when editor is backgrounded, but is called when editor is minimized
-	virtual bool LostFocus(FEditorViewportClient* ViewportClient, FViewport* Viewport) override;
-
-	/*
-	 * Mouse position events - These are called when no mouse button is down
-	 */
-
-	// called when mouse moves over viewport window
-	virtual bool MouseEnter(FEditorViewportClient* ViewportClient, FViewport* Viewport, int32 x, int32 y) override;
-
-	// called when mouse leaves viewport window
-	virtual bool MouseLeave(FEditorViewportClient* ViewportClient, FViewport* Viewport) override;
-
-	// called on any mouse-move event. *not* called during tracking/capturing, eg if any button is down
-	virtual bool MouseMove(FEditorViewportClient* ViewportClient, FViewport* Viewport, int32 x, int32 y) override;
-
-	/*
-	 * Input Button/Axis Events & Mouse Capture
-	 *
-	 *   event sequence for left mouse down/capture/up is:
-	 *     - InputKey( EKeys::LeftMouseButton, IE_PRESSED )
-	 *     - StartTracking()
-	 *     - CapturedMouseMove() and InputAxis() (repeated for each mouse-move)
-	 *     - Inputkey( EKeys::LeftMouseButton, IE_RELEASED )
-	 *     - EndTracking()
-	 *
-	 *   for doubleclick, we get one of the above sequence, then a second
-	 *   where instead of the first IE_PRESSED we get a IE_DoubleClick
-	 *
-	 *   for mouse wheel, we get following sequence
-	 *     - InputKey ( EKeys::MouseScrollUp / Down, IE_PRESSED )
-	 *     - InputAxis() for wheel move
-	 *     - InputKey ( EKeys::MouseScrollUp / Down, IE_RELEASED )
-	 *   it appears that there will only ever be *one* InputAxis() between the pressed/released sequenece
-	 *
-	 *   Note that this wheel event sequence can happen *during* a 
-	 *   middle-mouse tracking sequence on a wheel-button (nice!)
-	 */
-
-	// This is not just called for keyboard keys, it is also called for mouse down/up events!
-	// Eg for left-press we get Key = EKeys::LeftMouseButton and Event = IE_Pressed
-	// Return value indicates "handled'. If we return true for mouse press events then
-	// the StartTracking/CapturedMouseMove/EndTracking sequence is *not* called (but we
-	// also don't get MouseMove(), so the mouse-movement events appear to be lost?)
-	virtual bool InputKey(FEditorViewportClient* ViewportClient, FViewport* Viewport, FKey Key, EInputEvent Event) override;
-
-
-	// Called for 1D Axis movements - EKeys::MouseX, EKeys::MouseY, and EKeys::MouseWheelAxis
-	// Called even if we return true from InputKey which otherwise blocks tracking. 
-	// Return value indicates "handled" but has no effect on CapturedMouseMove events (ie we still get them)
-	virtual bool InputAxis(FEditorViewportClient* InViewportClient, FViewport* Viewport, int32 ControllerId, FKey Key, float Delta, float DeltaTime);
-
-	// called on mouse-down. return value indicates whether this was "handled" but
-	// does not mean we get exclusive capture events!
-	virtual bool StartTracking(FEditorViewportClient* InViewportClient, FViewport* InViewport) override;
-
-	// called during mouse-down mouse-move. Always called, return value is used to indicate whether
-	// we "handled" this mouse move (see FEditorModeTools::CapturedMouseMove) but does not pre-empt
-	// other editor modes from seeing this move event...
-	virtual bool CapturedMouseMove(FEditorViewportClient* InViewportClient, FViewport* InViewport, int32 InMouseX, int32 InMouseY) override;
-
-	// always called on mouse-up
-	virtual bool EndTracking(FEditorViewportClient* InViewportClient, FViewport* InViewport) override;
-
 	//////////////////
-	// End of FEdMode interface
+	// End of UBaseLegacyWidgetEdMode interface
 	//////////////////
 
-public:
-
-	virtual UEdModeInteractiveToolsContext* GetToolsContext() const
-	{
-		return ToolsContext;
-	}
-
-	virtual UInteractiveToolManager* GetToolManager() const
-	{ 
-		return ToolsContext->ToolManager;
-	}
-
-protected:
-	TArray<TPair<TSharedPtr<FUICommandInfo>, FString>> RegisteredTools;
-
-	UEdModeInteractiveToolsContext* ToolsContext;
-
-	/** Command list lives here so that the key bindings on the commands can be processed in the viewport. */
-	TSharedPtr<FUICommandList> UICommandList;
-
 public:
 	/** Cached pointer to the viewport world interaction object we're using to interact with mesh elements */
 	class UViewportWorldInteraction* ViewportWorldInteraction;
