Fix groom attachment with groom having only rigid binding as binding type.

diff --git a/Engine/Plugins/Runtime/HairStrands/Source/HairStrandsCore/Private/GroomComponent.cpp b/Engine/Plugins/Runtime/HairStrands/Source/HairStrandsCore/Private/GroomComponent.cpp
--- a/Engine/Plugins/Runtime/HairStrands/Source/HairStrandsCore/Private/GroomComponent.cpp
+++ b/Engine/Plugins/Runtime/HairStrands/Source/HairStrandsCore/Private/GroomComponent.cpp
@@ -1973,6 +1973,7 @@ void UGroomComponent::InitResources(bool bIsBindingReloading)
 	bool bHasNeedSkinningBinding = false;			   
 	bool bHasNeedGlobalDeformation = false;
 	bool bHasNeedSimulation = false;
+	bool bHasNeedSkeletalMesh = false;
 	for (int32 GroupIt = 0, GroupCount = GroomAsset->HairGroupsData.Num(); GroupIt < GroupCount; ++GroupIt)
 	{
 		for (uint32 LODIt = 0, LODCount = GroomAsset->GetLODCount(); LODIt < LODCount; ++LODIt)
@@ -1983,6 +1984,7 @@ void UGroomComponent::InitResources(bool bIsBindingReloading)
 			// Note on Global Deformation:
 			// * Global deformation require to have skinning binding
 			// * Force global interpolation to be enable for meshes with skinning binding as we use RBF defomation for 'sticking' meshes onto skel. mesh surface			
+			bHasNeedSkeletalMesh	  = bHasNeedSkeletalMesh || BindingType == EGroomBindingType::Rigid;
 			bHasNeedSkinningBinding   = bHasNeedSkinningBinding || BindingType == EGroomBindingType::Skinning;
 			bHasNeedSimulation		  = bHasNeedSimulation || GroomAsset->IsSimulationEnable(GroupIt, LODIt);
 			bHasNeedGlobalDeformation = bHasNeedGlobalDeformation || (BindingType == EGroomBindingType::Skinning && (GroomAsset->IsGlobalInterpolationEnable(GroupIt, LODIt) || GeometryType == EGroomGeometryType::Meshes));
@@ -2006,6 +2008,10 @@ void UGroomComponent::InitResources(bool bIsBindingReloading)
 			}
 		}
 	}
+	else if (bHasNeedSkeletalMesh && ParentMeshComponent)
+	{
+		ValidatedMeshComponent = Cast<USkeletalMeshComponent>(ParentMeshComponent)->SkeletalMesh ? ParentMeshComponent : nullptr;
+	}
 
 	// Insure the ticking of the Groom component always happens after the skeletalMeshComponent.
 	UGroomBindingAsset* LocalBindingAsset = nullptr;
