Remove UAV flags from hair static resources to save memory.

UAV flags enforce memory alignment, which create small lost of memory.

diff --git a/Engine/Plugins/Runtime/HairStrands/Source/HairStrandsCore/Private/GroomResources.cpp b/Engine/Plugins/Runtime/HairStrands/Source/HairStrandsCore/Private/GroomResources.cpp
--- a/Engine/Plugins/Runtime/HairStrands/Source/HairStrandsCore/Private/GroomResources.cpp
+++ b/Engine/Plugins/Runtime/HairStrands/Source/HairStrandsCore/Private/GroomResources.cpp
@@ -23,6 +23,12 @@
 #include "RenderTargetPool.h"
 #include "GroomBindingBuilder.h"
 
+enum class EHairResourceUsageType : uint8
+{
+	Static,
+	Dynamic
+};
+
 /////////////////////////////////////////////////////////////////////////////////////////
 // FRWBuffer utils 
 void UploadDataToBuffer(FReadBuffer& OutBuffer, uint32 DataSizeInBytes, const void* InCpuData)
@@ -90,7 +96,7 @@ void CreateBuffer(const TArray<typename FormatType::Type>& InData, FHairCardsVer
 /////////////////////////////////////////////////////////////////////////////////////////
 // RDG buffers utils 
 template<typename FormatType>
-void InternalCreateVertexBufferRDG(FRDGBuilder& GraphBuilder, const TArray<typename FormatType::Type>& InData, FRDGExternalBuffer& Out, const TCHAR* DebugName)
+void InternalCreateVertexBufferRDG(FRDGBuilder& GraphBuilder, const TArray<typename FormatType::Type>& InData, FRDGExternalBuffer& Out, const TCHAR* DebugName, EHairResourceUsageType UsageType)
 {
 	FRDGBufferRef Buffer = nullptr;
 
@@ -102,7 +108,11 @@ void InternalCreateVertexBufferRDG(FRDGBuilder& GraphBuilder, const TArray<typen
 		return;
 	}
 
-	const FRDGBufferDesc Desc = FRDGBufferDesc::CreateBufferDesc(FormatType::SizeInByte, InData.Num());
+	FRDGBufferDesc Desc = FRDGBufferDesc::CreateBufferDesc(FormatType::SizeInByte, InData.Num());
+	if (UsageType != EHairResourceUsageType::Dynamic)
+	{
+		Desc.Usage = Desc.Usage & (~BUF_UnorderedAccess);
+	}
 	Buffer = CreateVertexBuffer(
 		GraphBuilder,
 		DebugName,
@@ -115,7 +125,7 @@ void InternalCreateVertexBufferRDG(FRDGBuilder& GraphBuilder, const TArray<typen
 }
 
 template<typename DataType>
-void InternalCreateVertexBufferRDG(FRDGBuilder& GraphBuilder, const TArray<DataType>& InData, EPixelFormat Format, FRDGExternalBuffer& Out, const TCHAR* DebugName)
+void InternalCreateVertexBufferRDG(FRDGBuilder& GraphBuilder, const TArray<DataType>& InData, EPixelFormat Format, FRDGExternalBuffer& Out, const TCHAR* DebugName, EHairResourceUsageType UsageType)
 {
 	FRDGBufferRef Buffer = nullptr;
 
@@ -127,7 +137,11 @@ void InternalCreateVertexBufferRDG(FRDGBuilder& GraphBuilder, const TArray<DataT
 		return;
 	}
 
-	const FRDGBufferDesc Desc = FRDGBufferDesc::CreateBufferDesc(sizeof(DataType), InData.Num());
+	FRDGBufferDesc Desc = FRDGBufferDesc::CreateBufferDesc(sizeof(DataType), InData.Num());
+	if (UsageType != EHairResourceUsageType::Dynamic)
+	{
+		Desc.Usage = Desc.Usage & (~BUF_UnorderedAccess);
+	}
 	Buffer = CreateVertexBuffer(
 		GraphBuilder,
 		DebugName,
@@ -140,8 +154,11 @@ void InternalCreateVertexBufferRDG(FRDGBuilder& GraphBuilder, const TArray<DataT
 }
 
 template<typename FormatType>
-void InternalCreateVertexBufferRDG(FRDGBuilder& GraphBuilder, uint32 InVertexCount, FRDGExternalBuffer& Out, const TCHAR* DebugName)
+void InternalCreateVertexBufferRDG(FRDGBuilder& GraphBuilder, uint32 InVertexCount, FRDGExternalBuffer& Out, const TCHAR* DebugName, EHairResourceUsageType Usage)
 {
+	// Sanity check
+	check(Usage == EHairResourceUsageType::Dynamic);
+
 	const uint32 DataCount = InVertexCount;
 	const uint32 DataSizeInBytes = FormatType::SizeInByte * DataCount;
 	if (DataSizeInBytes == 0)
@@ -150,7 +167,7 @@ void InternalCreateVertexBufferRDG(FRDGBuilder& GraphBuilder, uint32 InVertexCou
 		return;
 	}
 
-	const FRDGBufferDesc Desc = FRDGBufferDesc::CreateBufferDesc(FormatType::SizeInByte, InVertexCount);
+	FRDGBufferDesc Desc = FRDGBufferDesc::CreateBufferDesc(FormatType::SizeInByte, InVertexCount);
 	FRDGBufferRef Buffer = GraphBuilder.CreateBuffer(Desc, DebugName, ERDGBufferFlags::MultiFrame);
 
 	auto IsFloatFormat = []()
@@ -190,10 +207,8 @@ void InternalCreateVertexBufferRDG(FRDGBuilder& GraphBuilder, uint32 InVertexCou
 }
 
 template<typename FormatType>
-void InternalCreateStructuredBufferRDG(FRDGBuilder& GraphBuilder, const TArray<FormatType>& InData, uint32 InSizeInByte, FRDGExternalBuffer& Out, const TCHAR* DebugName)
+void InternalCreateStructuredBufferRDG(FRDGBuilder& GraphBuilder, const TArray<FormatType>& InData, uint32 InSizeInByte, FRDGExternalBuffer& Out, const TCHAR* DebugName, EHairResourceUsageType UsageType)
 {
-	FRDGBufferRef Buffer = nullptr;
-
 	const uint32 DataCount = InData.Num();
 	const uint32 DataSizeInBytes = InSizeInByte * DataCount;
 	if (DataSizeInBytes == 0)
@@ -201,16 +216,18 @@ void InternalCreateStructuredBufferRDG(FRDGBuilder& GraphBuilder, const TArray<F
 		Out.Buffer = nullptr;
 		return;
 	}
-	// #hair_todo: Create this with a create+clear pass instead?
-	Buffer = CreateStructuredBuffer(
-		GraphBuilder,
-		DebugName,
-		InSizeInByte,
-		DataCount,
-		InData.GetData(),
-		DataSizeInBytes,
-		ERDGInitialDataFlags::NoCopy);
+	
+	FRDGBufferDesc Desc = FRDGBufferDesc::CreateStructuredDesc(InSizeInByte, DataCount);
+	if (UsageType != EHairResourceUsageType::Dynamic)
+	{
+		Desc.Usage = Desc.Usage & (~BUF_UnorderedAccess);
+	}
 
+	FRDGBufferRef Buffer = GraphBuilder.CreateBuffer(Desc, DebugName);
+	if (InData.GetData() && DataSizeInBytes)
+	{
+		AddBufferUploadPass(GraphBuilder, Buffer, InData.GetData(), DataSizeInBytes, ERDGInitialDataFlags::NoCopy);
+	}
 	ConvertToExternalBufferWithViews(GraphBuilder, Buffer, Out);
 }
 
@@ -268,7 +285,7 @@ void CreateHairStrandsDebugAttributeBuffer(FRDGExternalBuffer* DebugAttributeBuf
 		if (GUsingNullRHI) { return; }
 		FMemMark Mark(FMemStack::Get());
 		FRDGBuilder GraphBuilder(RHICmdList);
-		InternalCreateVertexBufferRDG<FHairStrandsAttributeFormat>(GraphBuilder, VertexCount, *LocalAttributeBuffer, TEXT("Hair.Strands_DebugAttributeBuffer"));
+		InternalCreateVertexBufferRDG<FHairStrandsAttributeFormat>(GraphBuilder, VertexCount, *LocalAttributeBuffer, TEXT("Hair.Strands_DebugAttributeBuffer"), EHairResourceUsageType::Dynamic);
 		GraphBuilder.Execute();
 	});
 }
@@ -357,19 +374,19 @@ void FHairCardsProceduralResource::InitRHI()
 	FRHICommandListImmediate& RHICmdList = FRHICommandListExecutor::GetImmediateCommandList();
 	FRDGBuilder GraphBuilder(RHICmdList);
 
-	InternalCreateVertexBufferRDG<FHairCardsAtlasRectFormat>(GraphBuilder, RenderData.CardsRect, AtlasRectBuffer, TEXT("Hair.CardsProcedural_AtlasRectBuffer"));
-	InternalCreateVertexBufferRDG<FHairCardsDimensionFormat>(GraphBuilder, RenderData.CardsLengths, LengthBuffer, TEXT("Hair.CardsProcedural_LengthBuffer"));
+	InternalCreateVertexBufferRDG<FHairCardsAtlasRectFormat>(GraphBuilder, RenderData.CardsRect, AtlasRectBuffer, TEXT("Hair.CardsProcedural_AtlasRectBuffer"), EHairResourceUsageType::Static);
+	InternalCreateVertexBufferRDG<FHairCardsDimensionFormat>(GraphBuilder, RenderData.CardsLengths, LengthBuffer, TEXT("Hair.CardsProcedural_LengthBuffer"), EHairResourceUsageType::Static);
 
-	InternalCreateVertexBufferRDG<FHairCardsOffsetAndCount>(GraphBuilder, RenderData.CardItToCluster, CardItToClusterBuffer, TEXT("Hair.CardsProcedural_CardItToClusterBuffer"));
-	InternalCreateVertexBufferRDG<FHairCardsOffsetAndCount>(GraphBuilder, RenderData.ClusterIdToVertices, ClusterIdToVerticesBuffer, TEXT("Hair.CardsProcedural_ClusterIdToVerticesBuffer"));
-	InternalCreateVertexBufferRDG<FHairCardsBoundsFormat>(GraphBuilder, RenderData.ClusterBounds, ClusterBoundBuffer, TEXT("Hair.CardsProcedural_ClusterBoundBuffer"));
+	InternalCreateVertexBufferRDG<FHairCardsOffsetAndCount>(GraphBuilder, RenderData.CardItToCluster, CardItToClusterBuffer, TEXT("Hair.CardsProcedural_CardItToClusterBuffer"), EHairResourceUsageType::Static);
+	InternalCreateVertexBufferRDG<FHairCardsOffsetAndCount>(GraphBuilder, RenderData.ClusterIdToVertices, ClusterIdToVerticesBuffer, TEXT("Hair.CardsProcedural_ClusterIdToVerticesBuffer"), EHairResourceUsageType::Static);
+	InternalCreateVertexBufferRDG<FHairCardsBoundsFormat>(GraphBuilder, RenderData.ClusterBounds, ClusterBoundBuffer, TEXT("Hair.CardsProcedural_ClusterBoundBuffer"), EHairResourceUsageType::Static);
 
-	InternalCreateVertexBufferRDG<FHairCardsVoxelDensityFormat>(GraphBuilder, RenderData.VoxelDensity, CardVoxel.DensityBuffer, TEXT("Hair.CardsProcedural_VoxelDensityBuffer"));
-	InternalCreateVertexBufferRDG<FHairCardsVoxelTangentFormat>(GraphBuilder, RenderData.VoxelTangent, CardVoxel.TangentBuffer, TEXT("Hair.CardsProcedural_VoxelTangentBuffer"));
-	InternalCreateVertexBufferRDG<FHairCardsVoxelTangentFormat>(GraphBuilder, RenderData.VoxelNormal, CardVoxel.NormalBuffer, TEXT("Hair.CardsProcedural_VoxelNormalBuffer"));
+	InternalCreateVertexBufferRDG<FHairCardsVoxelDensityFormat>(GraphBuilder, RenderData.VoxelDensity, CardVoxel.DensityBuffer, TEXT("Hair.CardsProcedural_VoxelDensityBuffer"), EHairResourceUsageType::Static);
+	InternalCreateVertexBufferRDG<FHairCardsVoxelTangentFormat>(GraphBuilder, RenderData.VoxelTangent, CardVoxel.TangentBuffer, TEXT("Hair.CardsProcedural_VoxelTangentBuffer"), EHairResourceUsageType::Static);
+	InternalCreateVertexBufferRDG<FHairCardsVoxelTangentFormat>(GraphBuilder, RenderData.VoxelNormal, CardVoxel.NormalBuffer, TEXT("Hair.CardsProcedural_VoxelNormalBuffer"), EHairResourceUsageType::Static);
 
-	InternalCreateVertexBufferRDG<FHairCardsStrandsPositionFormat>(GraphBuilder, RenderData.CardsStrandsPositions, CardsStrandsPositions, TEXT("Hair.CardsProcedural_CardsStrandsPositions"));
-	InternalCreateVertexBufferRDG<FHairCardsStrandsAttributeFormat>(GraphBuilder, RenderData.CardsStrandsAttributes, CardsStrandsAttributes, TEXT("Hair.CardsProcedural_CardsStrandsAttributes"));
+	InternalCreateVertexBufferRDG<FHairCardsStrandsPositionFormat>(GraphBuilder, RenderData.CardsStrandsPositions, CardsStrandsPositions, TEXT("Hair.CardsProcedural_CardsStrandsPositions"), EHairResourceUsageType::Static);
+	InternalCreateVertexBufferRDG<FHairCardsStrandsAttributeFormat>(GraphBuilder, RenderData.CardsStrandsAttributes, CardsStrandsAttributes, TEXT("Hair.CardsProcedural_CardsStrandsAttributes"), EHairResourceUsageType::Static);
 
 	GraphBuilder.Execute();
 }
@@ -406,13 +423,13 @@ void FHairCardsDeformedResource::InitRHI()
 	const uint32 VertexCount = RenderData.Positions.Num();
 	if (bInitializedData)
 	{
-		InternalCreateVertexBufferRDG<FHairCardsPositionFormat>(GraphBuilder, RenderData.Positions, DeformedPositionBuffer[0], TEXT("Hair.CardsDeformed_Position0"));
-		InternalCreateVertexBufferRDG<FHairCardsPositionFormat>(GraphBuilder, RenderData.Positions, DeformedPositionBuffer[1], TEXT("Hair.CardsDeformed_Position1"));
+		InternalCreateVertexBufferRDG<FHairCardsPositionFormat>(GraphBuilder, RenderData.Positions, DeformedPositionBuffer[0], TEXT("Hair.CardsDeformed_Position0"), EHairResourceUsageType::Dynamic);
+		InternalCreateVertexBufferRDG<FHairCardsPositionFormat>(GraphBuilder, RenderData.Positions, DeformedPositionBuffer[1], TEXT("Hair.CardsDeformed_Position1"), EHairResourceUsageType::Dynamic);
 	}
 	else
 	{
-		InternalCreateVertexBufferRDG<FHairCardsPositionFormat>(GraphBuilder, VertexCount, DeformedPositionBuffer[0], TEXT("Hair.CardsDeformed_Position0"));
-		InternalCreateVertexBufferRDG<FHairCardsPositionFormat>(GraphBuilder, VertexCount, DeformedPositionBuffer[1], TEXT("Hair.CardsDeformed_Position1"));
+		InternalCreateVertexBufferRDG<FHairCardsPositionFormat>(GraphBuilder, VertexCount, DeformedPositionBuffer[0], TEXT("Hair.CardsDeformed_Position0"), EHairResourceUsageType::Dynamic);
+		InternalCreateVertexBufferRDG<FHairCardsPositionFormat>(GraphBuilder, VertexCount, DeformedPositionBuffer[1], TEXT("Hair.CardsDeformed_Position1"), EHairResourceUsageType::Dynamic);
 	}
 	GraphBuilder.Execute();
 }
@@ -486,13 +503,13 @@ void FHairMeshesDeformedResource::InitRHI()
 	const uint32 VertexCount = RenderData.Positions.Num();
 	if (bInitializedData)
 	{
-		InternalCreateVertexBufferRDG<FHairCardsPositionFormat>(GraphBuilder, RenderData.Positions, DeformedPositionBuffer[0], TEXT("Hair.MeshesDeformed_Positions0"));
-		InternalCreateVertexBufferRDG<FHairCardsPositionFormat>(GraphBuilder, RenderData.Positions, DeformedPositionBuffer[1], TEXT("Hair.MeshesDeformed_Positions1"));
+		InternalCreateVertexBufferRDG<FHairCardsPositionFormat>(GraphBuilder, RenderData.Positions, DeformedPositionBuffer[0], TEXT("Hair.MeshesDeformed_Positions0"), EHairResourceUsageType::Dynamic);
+		InternalCreateVertexBufferRDG<FHairCardsPositionFormat>(GraphBuilder, RenderData.Positions, DeformedPositionBuffer[1], TEXT("Hair.MeshesDeformed_Positions1"), EHairResourceUsageType::Dynamic);
 	}
 	else
 	{
-		InternalCreateVertexBufferRDG<FHairCardsPositionFormat>(GraphBuilder, VertexCount, DeformedPositionBuffer[0], TEXT("Hair.MeshesDeformed_Positions0"));
-		InternalCreateVertexBufferRDG<FHairCardsPositionFormat>(GraphBuilder, VertexCount, DeformedPositionBuffer[1], TEXT("Hair.MeshesDeformed_Positions1"));
+		InternalCreateVertexBufferRDG<FHairCardsPositionFormat>(GraphBuilder, VertexCount, DeformedPositionBuffer[0], TEXT("Hair.MeshesDeformed_Positions0"), EHairResourceUsageType::Dynamic);
+		InternalCreateVertexBufferRDG<FHairCardsPositionFormat>(GraphBuilder, VertexCount, DeformedPositionBuffer[1], TEXT("Hair.MeshesDeformed_Positions1"), EHairResourceUsageType::Dynamic);
 	}
 
 	GraphBuilder.Execute();
@@ -517,13 +534,13 @@ void FHairStrandsRestResource::InitRHI()
 	FRHICommandListImmediate& RHICmdList = FRHICommandListExecutor::GetImmediateCommandList();
 	FRDGBuilder GraphBuilder(RHICmdList);
 
-	InternalCreateVertexBufferRDG<FHairStrandsPositionFormat>(GraphBuilder, RenderData.Positions, PositionBuffer, TEXT("Hair.StrandsRest_PositionBuffer"));
-	InternalCreateVertexBufferRDG<FHairStrandsAttributeFormat>(GraphBuilder, RenderData.Attributes, AttributeBuffer, TEXT("Hair.StrandsRest_AttributeBuffer"));
-	InternalCreateVertexBufferRDG<FHairStrandsMaterialFormat>(GraphBuilder, RenderData.Materials, MaterialBuffer, TEXT("Hair.StrandsRest_MaterialBuffer"));
+	InternalCreateVertexBufferRDG<FHairStrandsPositionFormat>(GraphBuilder, RenderData.Positions, PositionBuffer, TEXT("Hair.StrandsRest_PositionBuffer"), EHairResourceUsageType::Static);
+	InternalCreateVertexBufferRDG<FHairStrandsAttributeFormat>(GraphBuilder, RenderData.Attributes, AttributeBuffer, TEXT("Hair.StrandsRest_AttributeBuffer"), EHairResourceUsageType::Static);
+	InternalCreateVertexBufferRDG<FHairStrandsMaterialFormat>(GraphBuilder, RenderData.Materials, MaterialBuffer, TEXT("Hair.StrandsRest_MaterialBuffer"), EHairResourceUsageType::Static);
 
 	TArray<FVector4> RestOffset;
 	RestOffset.Add(PositionOffset);
-	InternalCreateVertexBufferRDG<FHairStrandsPositionOffsetFormat>(GraphBuilder, RestOffset, PositionOffsetBuffer, TEXT("Hair.StrandsRest_PositionOffsetBuffer"));
+	InternalCreateVertexBufferRDG<FHairStrandsPositionOffsetFormat>(GraphBuilder, RestOffset, PositionOffsetBuffer, TEXT("Hair.StrandsRest_PositionOffsetBuffer"), EHairResourceUsageType::Static);
 	
 	GraphBuilder.Execute();
 }
@@ -542,7 +559,7 @@ FRDGExternalBuffer FHairStrandsRestResource::GetTangentBuffer(FRDGBuilder& Graph
 	if (TangentBuffer.Buffer == nullptr)
 	{
 		const uint32 VertexCount = RenderData.Positions.Num();
-		InternalCreateVertexBufferRDG<FHairStrandsTangentFormat>(GraphBuilder, VertexCount * FHairStrandsTangentFormat::ComponentCount, TangentBuffer, TEXT("Hair.StrandsRest_TangentBuffer"));
+		InternalCreateVertexBufferRDG<FHairStrandsTangentFormat>(GraphBuilder, VertexCount * FHairStrandsTangentFormat::ComponentCount, TangentBuffer, TEXT("Hair.StrandsRest_TangentBuffer"), EHairResourceUsageType::Dynamic);
 
 		AddHairTangentPass(
 			GraphBuilder,
@@ -581,20 +598,20 @@ void FHairStrandsDeformedResource::InitRHI()
 	const uint32 VertexCount = RenderData.Positions.Num();
 	if (bInitializedData)
 	{
-		InternalCreateVertexBufferRDG<FHairStrandsPositionFormat>(GraphBuilder, RenderData.Positions, DeformedPositionBuffer[0], TEXT("Hair.StrandsDeformed_DeformedPositionBuffer0")); // , bDynamic ? ERHIAccess::UAVCompute : ERHIAccess::SRVMask
-		InternalCreateVertexBufferRDG<FHairStrandsPositionFormat>(GraphBuilder, RenderData.Positions, DeformedPositionBuffer[1], TEXT("Hair.StrandsDeformed_DeformedPositionBuffer1")); // , ERHIAccess::UAVCompute
+		InternalCreateVertexBufferRDG<FHairStrandsPositionFormat>(GraphBuilder, RenderData.Positions, DeformedPositionBuffer[0], TEXT("Hair.StrandsDeformed_DeformedPositionBuffer0"), EHairResourceUsageType::Dynamic); // , bDynamic ? ERHIAccess::UAVCompute : ERHIAccess::SRVMask
+		InternalCreateVertexBufferRDG<FHairStrandsPositionFormat>(GraphBuilder, RenderData.Positions, DeformedPositionBuffer[1], TEXT("Hair.StrandsDeformed_DeformedPositionBuffer1"), EHairResourceUsageType::Dynamic); // , ERHIAccess::UAVCompute
 	}
 	else
 	{
-		InternalCreateVertexBufferRDG<FHairStrandsPositionFormat>(GraphBuilder, VertexCount, DeformedPositionBuffer[0], TEXT("Hair.StrandsDeformed_DeformedPositionBuffer0"));
-		InternalCreateVertexBufferRDG<FHairStrandsPositionFormat>(GraphBuilder, VertexCount, DeformedPositionBuffer[1], TEXT("Hair.StrandsDeformed_DeformedPositionBuffer1"));
+		InternalCreateVertexBufferRDG<FHairStrandsPositionFormat>(GraphBuilder, VertexCount, DeformedPositionBuffer[0], TEXT("Hair.StrandsDeformed_DeformedPositionBuffer0"), EHairResourceUsageType::Dynamic);
+		InternalCreateVertexBufferRDG<FHairStrandsPositionFormat>(GraphBuilder, VertexCount, DeformedPositionBuffer[1], TEXT("Hair.StrandsDeformed_DeformedPositionBuffer1"), EHairResourceUsageType::Dynamic);
 	}
-	InternalCreateVertexBufferRDG<FHairStrandsTangentFormat>(GraphBuilder, VertexCount * FHairStrandsTangentFormat::ComponentCount, TangentBuffer, TEXT("Hair.StrandsDeformed_TangentBuffer"));
+	InternalCreateVertexBufferRDG<FHairStrandsTangentFormat>(GraphBuilder, VertexCount * FHairStrandsTangentFormat::ComponentCount, TangentBuffer, TEXT("Hair.StrandsDeformed_TangentBuffer"), EHairResourceUsageType::Dynamic);
 
 	TArray<FVector4> DefaultOffsets;
 	DefaultOffsets.Add(DefaultOffset);
-	InternalCreateVertexBufferRDG<FHairStrandsPositionOffsetFormat>(GraphBuilder, DefaultOffsets, DeformedOffsetBuffer[0], TEXT("Hair.StrandsDeformed_DeformedOffsetBuffer0"));
-	InternalCreateVertexBufferRDG<FHairStrandsPositionOffsetFormat>(GraphBuilder, DefaultOffsets, DeformedOffsetBuffer[1], TEXT("Hair.StrandsDeformed_DeformedOffsetBuffer1"));
+	InternalCreateVertexBufferRDG<FHairStrandsPositionOffsetFormat>(GraphBuilder, DefaultOffsets, DeformedOffsetBuffer[0], TEXT("Hair.StrandsDeformed_DeformedOffsetBuffer0"), EHairResourceUsageType::Dynamic);
+	InternalCreateVertexBufferRDG<FHairStrandsPositionOffsetFormat>(GraphBuilder, DefaultOffsets, DeformedOffsetBuffer[1], TEXT("Hair.StrandsDeformed_DeformedOffsetBuffer1"), EHairResourceUsageType::Dynamic);
 
 	GraphBuilder.Execute();
 }
@@ -651,11 +668,11 @@ void FHairStrandsClusterCullingResource::InitRHI()
 	FRHICommandListImmediate& RHICmdList = FRHICommandListExecutor::GetImmediateCommandList();
 	FRDGBuilder GraphBuilder(RHICmdList);
 
-	InternalCreateStructuredBufferRDG(GraphBuilder, Data.PackedClusterInfos, sizeof(FHairStrandsClusterCullingData::FHairClusterInfo::Packed), ClusterInfoBuffer, TEXT("Hair.StrandsClusterCulling_ClusterInfoBuffer"));
-	InternalCreateStructuredBufferRDG(GraphBuilder, Data.ClusterLODInfos, sizeof(FHairStrandsClusterCullingData::FHairClusterLODInfo), ClusterLODInfoBuffer, TEXT("Hair.StrandsClusterCulling_ClusterLODInfoBuffer"));
+	InternalCreateStructuredBufferRDG(GraphBuilder, Data.PackedClusterInfos, sizeof(FHairStrandsClusterCullingData::FHairClusterInfo::Packed), ClusterInfoBuffer, TEXT("Hair.StrandsClusterCulling_ClusterInfoBuffer"), EHairResourceUsageType::Static);
+	InternalCreateStructuredBufferRDG(GraphBuilder, Data.ClusterLODInfos, sizeof(FHairStrandsClusterCullingData::FHairClusterLODInfo), ClusterLODInfoBuffer, TEXT("Hair.StrandsClusterCulling_ClusterLODInfoBuffer"), EHairResourceUsageType::Static);
 
-	InternalCreateVertexBufferRDG(GraphBuilder, Data.ClusterVertexIds, PF_R32_UINT, ClusterVertexIdBuffer, TEXT("Hair.StrandsClusterCulling_ClusterVertexIds"));
-	InternalCreateVertexBufferRDG(GraphBuilder, Data.VertexToClusterIds, PF_R32_UINT, VertexToClusterIdBuffer, TEXT("Hair.StrandsClusterCulling_VertexToClusterIds"));
+	InternalCreateVertexBufferRDG(GraphBuilder, Data.ClusterVertexIds, PF_R32_UINT, ClusterVertexIdBuffer, TEXT("Hair.StrandsClusterCulling_ClusterVertexIds"), EHairResourceUsageType::Static);
+	InternalCreateVertexBufferRDG(GraphBuilder, Data.VertexToClusterIds, PF_R32_UINT, VertexToClusterIdBuffer, TEXT("Hair.StrandsClusterCulling_VertexToClusterIds"), EHairResourceUsageType::Static);
 
 	GraphBuilder.Execute();
 }
@@ -704,9 +721,9 @@ void FHairStrandsRestRootResource::InitRHI()
 		FRHICommandListImmediate& RHICmdList = FRHICommandListExecutor::GetImmediateCommandList();
 		FRDGBuilder GraphBuilder(RHICmdList);
 
-		InternalCreateVertexBufferRDG<FHairStrandsIndexFormat>(GraphBuilder, RootData.VertexToCurveIndexBuffer, VertexToCurveIndexBuffer, TEXT("Hair.StrandsRestRoot_VertexToCurveIndexBuffer"));
-		InternalCreateVertexBufferRDG<FHairStrandsRootPositionFormat>(GraphBuilder, RootData.RootPositionBuffer, RootPositionBuffer, TEXT("Hair.StrandsRestRoot_RootPositionBuffer"));
-		InternalCreateVertexBufferRDG<FHairStrandsRootNormalFormat>(GraphBuilder, RootData.RootNormalBuffer, RootNormalBuffer, TEXT("Hair.StrandsRestRoot_RootNormalBuffer"));
+		InternalCreateVertexBufferRDG<FHairStrandsIndexFormat>(GraphBuilder, RootData.VertexToCurveIndexBuffer, VertexToCurveIndexBuffer, TEXT("Hair.StrandsRestRoot_VertexToCurveIndexBuffer"), EHairResourceUsageType::Static);
+		InternalCreateVertexBufferRDG<FHairStrandsRootPositionFormat>(GraphBuilder, RootData.RootPositionBuffer, RootPositionBuffer, TEXT("Hair.StrandsRestRoot_RootPositionBuffer"), EHairResourceUsageType::Static);
+		InternalCreateVertexBufferRDG<FHairStrandsRootNormalFormat>(GraphBuilder, RootData.RootNormalBuffer, RootNormalBuffer, TEXT("Hair.StrandsRestRoot_RootNormalBuffer"), EHairResourceUsageType::Static);
 		
 		check(LODs.Num() == RootData.MeshProjectionLODs.Num());
 		for (uint32 LODIt=0, LODCount = LODs.Num(); LODIt<LODCount; ++LODIt)
@@ -720,48 +737,48 @@ void FHairStrandsRestRootResource::InitRHI()
 				GPUData.Status = FLOD::EStatus::Completed;
 
 				check(CPUData.RootTriangleBarycentricBuffer.Num() > 0);
-				InternalCreateVertexBufferRDG<FHairStrandsCurveTriangleBarycentricFormat>(GraphBuilder, CPUData.RootTriangleBarycentricBuffer, GPUData.RootTriangleBarycentricBuffer, TEXT("Hair.StrandsRestRoot_RootTriangleBarycentricBuffer"));
+				InternalCreateVertexBufferRDG<FHairStrandsCurveTriangleBarycentricFormat>(GraphBuilder, CPUData.RootTriangleBarycentricBuffer, GPUData.RootTriangleBarycentricBuffer, TEXT("Hair.StrandsRestRoot_RootTriangleBarycentricBuffer"), EHairResourceUsageType::Static);
 
 				check(CPUData.RootTriangleIndexBuffer.Num() > 0);
-				InternalCreateVertexBufferRDG<FHairStrandsCurveTriangleIndexFormat>(GraphBuilder, CPUData.RootTriangleIndexBuffer, GPUData.RootTriangleIndexBuffer, TEXT("Hair.StrandsRestRoot_RootTriangleIndexBuffer"));
+				InternalCreateVertexBufferRDG<FHairStrandsCurveTriangleIndexFormat>(GraphBuilder, CPUData.RootTriangleIndexBuffer, GPUData.RootTriangleIndexBuffer, TEXT("Hair.StrandsRestRoot_RootTriangleIndexBuffer"), EHairResourceUsageType::Static);
 
 				check(CPUData.RestRootTrianglePosition0Buffer.Num() > 0);
 				check(CPUData.RestRootTrianglePosition1Buffer.Num() > 0);
 				check(CPUData.RestRootTrianglePosition2Buffer.Num() > 0);
-				InternalCreateVertexBufferRDG<FHairStrandsMeshTrianglePositionFormat>(GraphBuilder, CPUData.RestRootTrianglePosition0Buffer, GPUData.RestRootTrianglePosition0Buffer, TEXT("Hair.StrandsRestRoot_RestRootTrianglePosition0Buffer"));
-				InternalCreateVertexBufferRDG<FHairStrandsMeshTrianglePositionFormat>(GraphBuilder, CPUData.RestRootTrianglePosition1Buffer, GPUData.RestRootTrianglePosition1Buffer, TEXT("Hair.StrandsRestRoot_RestRootTrianglePosition1Buffer"));
-				InternalCreateVertexBufferRDG<FHairStrandsMeshTrianglePositionFormat>(GraphBuilder, CPUData.RestRootTrianglePosition2Buffer, GPUData.RestRootTrianglePosition2Buffer, TEXT("Hair.StrandsRestRoot_RestRootTrianglePosition2Buffer"));
+				InternalCreateVertexBufferRDG<FHairStrandsMeshTrianglePositionFormat>(GraphBuilder, CPUData.RestRootTrianglePosition0Buffer, GPUData.RestRootTrianglePosition0Buffer, TEXT("Hair.StrandsRestRoot_RestRootTrianglePosition0Buffer"), EHairResourceUsageType::Static);
+				InternalCreateVertexBufferRDG<FHairStrandsMeshTrianglePositionFormat>(GraphBuilder, CPUData.RestRootTrianglePosition1Buffer, GPUData.RestRootTrianglePosition1Buffer, TEXT("Hair.StrandsRestRoot_RestRootTrianglePosition1Buffer"), EHairResourceUsageType::Static);
+				InternalCreateVertexBufferRDG<FHairStrandsMeshTrianglePositionFormat>(GraphBuilder, CPUData.RestRootTrianglePosition2Buffer, GPUData.RestRootTrianglePosition2Buffer, TEXT("Hair.StrandsRestRoot_RestRootTrianglePosition2Buffer"), EHairResourceUsageType::Static);
 			}
 			else
 			{
 				GPUData.Status = FLOD::EStatus::Initialized;
-
-				InternalCreateVertexBufferRDG<FHairStrandsCurveTriangleBarycentricFormat>(GraphBuilder, RootData.RootCount, GPUData.RootTriangleBarycentricBuffer, TEXT("Hair.StrandsRestRoot_RootTriangleBarycentricBuffer"));
-				InternalCreateVertexBufferRDG<FHairStrandsCurveTriangleIndexFormat>(GraphBuilder, RootData.RootCount, GPUData.RootTriangleIndexBuffer, TEXT("Hair.StrandsRestRoot_RootTriangleIndexBuffer"));
-
+			
+				InternalCreateVertexBufferRDG<FHairStrandsCurveTriangleBarycentricFormat>(GraphBuilder, RootData.RootCount, GPUData.RootTriangleBarycentricBuffer, TEXT("Hair.StrandsRestRoot_RootTriangleBarycentricBuffer"), EHairResourceUsageType::Dynamic);
+				InternalCreateVertexBufferRDG<FHairStrandsCurveTriangleIndexFormat>(GraphBuilder, RootData.RootCount, GPUData.RootTriangleIndexBuffer, TEXT("Hair.StrandsRestRoot_RootTriangleIndexBuffer"), EHairResourceUsageType::Dynamic);
+			
 				// Create buffers. Initialization will be done by render passes
-				InternalCreateVertexBufferRDG<FHairStrandsMeshTrianglePositionFormat>(GraphBuilder, RootData.RootCount, GPUData.RestRootTrianglePosition0Buffer, TEXT("Hair.StrandsRestRoot_RestRootTrianglePosition0Buffer"));
-				InternalCreateVertexBufferRDG<FHairStrandsMeshTrianglePositionFormat>(GraphBuilder, RootData.RootCount, GPUData.RestRootTrianglePosition1Buffer, TEXT("Hair.StrandsRestRoot_RestRootTrianglePosition1Buffer"));
-				InternalCreateVertexBufferRDG<FHairStrandsMeshTrianglePositionFormat>(GraphBuilder, RootData.RootCount, GPUData.RestRootTrianglePosition2Buffer, TEXT("Hair.StrandsRestRoot_RestRootTrianglePosition2Buffer"));
+				InternalCreateVertexBufferRDG<FHairStrandsMeshTrianglePositionFormat>(GraphBuilder, RootData.RootCount, GPUData.RestRootTrianglePosition0Buffer, TEXT("Hair.StrandsRestRoot_RestRootTrianglePosition0Buffer"), EHairResourceUsageType::Dynamic);
+				InternalCreateVertexBufferRDG<FHairStrandsMeshTrianglePositionFormat>(GraphBuilder, RootData.RootCount, GPUData.RestRootTrianglePosition1Buffer, TEXT("Hair.StrandsRestRoot_RestRootTrianglePosition1Buffer"), EHairResourceUsageType::Dynamic);
+				InternalCreateVertexBufferRDG<FHairStrandsMeshTrianglePositionFormat>(GraphBuilder, RootData.RootCount, GPUData.RestRootTrianglePosition2Buffer, TEXT("Hair.StrandsRestRoot_RestRootTrianglePosition2Buffer"), EHairResourceUsageType::Dynamic);
 			}
 
 			GPUData.SampleCount = CPUData.SampleCount;
 			const bool bHasValidCPUWeights = CPUData.MeshSampleIndicesBuffer.Num() > 0;
-			if(bHasValidCPUWeights)
+			if (bHasValidCPUWeights)
 			{
 				//check(CPUData.MeshInterpolationWeightsBuffer.Num() == (CPUData.SampleCount+4) * (CPUData.SampleCount+4));
 				check(CPUData.MeshSampleIndicesBuffer.Num() == CPUData.SampleCount);
 				check(CPUData.RestSamplePositionsBuffer.Num() == CPUData.SampleCount);
 
-				InternalCreateVertexBufferRDG<FHairStrandsWeightFormat>(GraphBuilder, CPUData.MeshInterpolationWeightsBuffer, GPUData.MeshInterpolationWeightsBuffer, TEXT("Hair.StrandsRestRoot_MeshInterpolationWeightsBuffer"));
-				InternalCreateVertexBufferRDG<FHairStrandsIndexFormat>(GraphBuilder, CPUData.MeshSampleIndicesBuffer, GPUData.MeshSampleIndicesBuffer, TEXT("Hair.StrandsRestRoot_MeshSampleIndicesBuffer"));
-				InternalCreateVertexBufferRDG<FHairStrandsMeshTrianglePositionFormat>(GraphBuilder, CPUData.RestSamplePositionsBuffer, GPUData.RestSamplePositionsBuffer, TEXT("Hair.StrandsRestRoot_RestSamplePositionsBuffer"));
+				InternalCreateVertexBufferRDG<FHairStrandsWeightFormat>(GraphBuilder, CPUData.MeshInterpolationWeightsBuffer, GPUData.MeshInterpolationWeightsBuffer, TEXT("Hair.StrandsRestRoot_MeshInterpolationWeightsBuffer"), EHairResourceUsageType::Static);
+				InternalCreateVertexBufferRDG<FHairStrandsIndexFormat>(GraphBuilder, CPUData.MeshSampleIndicesBuffer, GPUData.MeshSampleIndicesBuffer, TEXT("Hair.StrandsRestRoot_MeshSampleIndicesBuffer"), EHairResourceUsageType::Static);
+				InternalCreateVertexBufferRDG<FHairStrandsMeshTrianglePositionFormat>(GraphBuilder, CPUData.RestSamplePositionsBuffer, GPUData.RestSamplePositionsBuffer, TEXT("Hair.StrandsRestRoot_RestSamplePositionsBuffer"), EHairResourceUsageType::Static);
 			}
 			else
 			{
-				InternalCreateVertexBufferRDG<FHairStrandsWeightFormat>(GraphBuilder, (CPUData.SampleCount+4) * (CPUData.SampleCount+4), GPUData.MeshInterpolationWeightsBuffer, TEXT("Hair.StrandsRestRoot_MeshInterpolationWeightsBuffer"));
-				InternalCreateVertexBufferRDG<FHairStrandsIndexFormat>(GraphBuilder, CPUData.SampleCount, GPUData.MeshSampleIndicesBuffer, TEXT("Hair.StrandsRestRoot_MeshSampleIndicesBuffer"));
-				InternalCreateVertexBufferRDG<FHairStrandsMeshTrianglePositionFormat>(GraphBuilder, CPUData.SampleCount, GPUData.RestSamplePositionsBuffer, TEXT("Hair.StrandsRestRoot_RestSamplePositionsBuffer"));
+				InternalCreateVertexBufferRDG<FHairStrandsWeightFormat>(GraphBuilder, (CPUData.SampleCount+4) * (CPUData.SampleCount+4), GPUData.MeshInterpolationWeightsBuffer, TEXT("Hair.StrandsRestRoot_MeshInterpolationWeightsBuffer"), EHairResourceUsageType::Dynamic);
+				InternalCreateVertexBufferRDG<FHairStrandsIndexFormat>(GraphBuilder, CPUData.SampleCount, GPUData.MeshSampleIndicesBuffer, TEXT("Hair.StrandsRestRoot_MeshSampleIndicesBuffer"), EHairResourceUsageType::Dynamic);
+				InternalCreateVertexBufferRDG<FHairStrandsMeshTrianglePositionFormat>(GraphBuilder, CPUData.SampleCount, GPUData.RestSamplePositionsBuffer, TEXT("Hair.StrandsRestRoot_RestSamplePositionsBuffer"), EHairResourceUsageType::Dynamic);
 			}
 		}
 		GraphBuilder.Execute();
@@ -828,12 +845,12 @@ void FHairStrandsDeformedRootResource::InitRHI()
 		for (FLOD& LOD : LODs)
 		{		
 			LOD.Status = FLOD::EStatus::Initialized;
-			InternalCreateVertexBufferRDG<FHairStrandsMeshTrianglePositionFormat>(GraphBuilder, LOD.SampleCount, LOD.DeformedSamplePositionsBuffer, TEXT("Hair.StrandsRootDeformed_DeformedSamplePositionsBuffer"));
-			InternalCreateVertexBufferRDG<FHairStrandsMeshTrianglePositionFormat>(GraphBuilder, LOD.SampleCount + 4, LOD.MeshSampleWeightsBuffer, TEXT("Hair.StrandsRootDeformed_MeshSampleWeightsBuffer"));
+			InternalCreateVertexBufferRDG<FHairStrandsMeshTrianglePositionFormat>(GraphBuilder, LOD.SampleCount, LOD.DeformedSamplePositionsBuffer, TEXT("Hair.StrandsRootDeformed_DeformedSamplePositionsBuffer"), EHairResourceUsageType::Dynamic);
+			InternalCreateVertexBufferRDG<FHairStrandsMeshTrianglePositionFormat>(GraphBuilder, LOD.SampleCount + 4, LOD.MeshSampleWeightsBuffer, TEXT("Hair.StrandsRootDeformed_MeshSampleWeightsBuffer"), EHairResourceUsageType::Dynamic);
 
-			InternalCreateVertexBufferRDG<FHairStrandsMeshTrianglePositionFormat>(GraphBuilder, RootCount, LOD.DeformedRootTrianglePosition0Buffer, TEXT("Hair.StrandsRootDeformed_DeformedRootTrianglePosition0Buffer"));
-			InternalCreateVertexBufferRDG<FHairStrandsMeshTrianglePositionFormat>(GraphBuilder, RootCount, LOD.DeformedRootTrianglePosition1Buffer, TEXT("Hair.StrandsRootDeformed_DeformedRootTrianglePosition1Buffer"));
-			InternalCreateVertexBufferRDG<FHairStrandsMeshTrianglePositionFormat>(GraphBuilder, RootCount, LOD.DeformedRootTrianglePosition2Buffer, TEXT("Hair.StrandsRootDeformed_DeformedRootTrianglePosition2Buffer"));
+			InternalCreateVertexBufferRDG<FHairStrandsMeshTrianglePositionFormat>(GraphBuilder, RootCount, LOD.DeformedRootTrianglePosition0Buffer, TEXT("Hair.StrandsRootDeformed_DeformedRootTrianglePosition0Buffer"), EHairResourceUsageType::Dynamic);
+			InternalCreateVertexBufferRDG<FHairStrandsMeshTrianglePositionFormat>(GraphBuilder, RootCount, LOD.DeformedRootTrianglePosition1Buffer, TEXT("Hair.StrandsRootDeformed_DeformedRootTrianglePosition1Buffer"), EHairResourceUsageType::Dynamic);
+			InternalCreateVertexBufferRDG<FHairStrandsMeshTrianglePositionFormat>(GraphBuilder, RootCount, LOD.DeformedRootTrianglePosition2Buffer, TEXT("Hair.StrandsRootDeformed_DeformedRootTrianglePosition2Buffer"), EHairResourceUsageType::Dynamic);
 		}
 	}
 	GraphBuilder.Execute();
@@ -1009,9 +1026,9 @@ void FHairStrandsInterpolationResource::InitRHI()
 	FRHICommandListImmediate& RHICmdList = FRHICommandListExecutor::GetImmediateCommandList();
 	FRDGBuilder GraphBuilder(RHICmdList);
 
-	InternalCreateVertexBufferRDG<FHairStrandsInterpolation0Format>(GraphBuilder, RenderData.Interpolation0, Interpolation0Buffer, TEXT("Hair.StrandsInterpolation_Interpolation0Buffer"));
-	InternalCreateVertexBufferRDG<FHairStrandsInterpolation1Format>(GraphBuilder, RenderData.Interpolation1, Interpolation1Buffer, TEXT("Hair.StrandsInterpolation_Interpolation1Buffer"));
-	InternalCreateVertexBufferRDG<FHairStrandsRootIndexFormat>(GraphBuilder, SimRootPointIndex, SimRootPointIndexBuffer, TEXT("Hair.StrandsInterpolation_SimRootPointIndex"));
+	InternalCreateVertexBufferRDG<FHairStrandsInterpolation0Format>(GraphBuilder, RenderData.Interpolation0, Interpolation0Buffer, TEXT("Hair.StrandsInterpolation_Interpolation0Buffer"), EHairResourceUsageType::Static);
+	InternalCreateVertexBufferRDG<FHairStrandsInterpolation1Format>(GraphBuilder, RenderData.Interpolation1, Interpolation1Buffer, TEXT("Hair.StrandsInterpolation_Interpolation1Buffer"), EHairResourceUsageType::Static);
+	InternalCreateVertexBufferRDG<FHairStrandsRootIndexFormat>(GraphBuilder, SimRootPointIndex, SimRootPointIndexBuffer, TEXT("Hair.StrandsInterpolation_SimRootPointIndex"), EHairResourceUsageType::Static);
 	GraphBuilder.Execute();
 	//SimRootPointIndex.SetNum(0);
 }
@@ -1060,7 +1077,7 @@ void FHairCardsInterpolationResource::InitRHI()
 	FRHICommandListImmediate& RHICmdList = FRHICommandListExecutor::GetImmediateCommandList();
 	FRDGBuilder GraphBuilder(RHICmdList);
 
-	InternalCreateVertexBufferRDG<FHairCardsInterpolationFormat>(GraphBuilder, RenderData.Interpolation, InterpolationBuffer, TEXT("Hair.CardsInterpolation_InterpolationBuffer"));
+	InternalCreateVertexBufferRDG<FHairCardsInterpolationFormat>(GraphBuilder, RenderData.Interpolation, InterpolationBuffer, TEXT("Hair.CardsInterpolation_InterpolationBuffer"), EHairResourceUsageType::Static);
 	GraphBuilder.Execute();
 }
 
@@ -1093,7 +1110,7 @@ void FHairStrandsRaytracingResource::InitRHI()
 	FRHICommandListImmediate& RHICmdList = FRHICommandListExecutor::GetImmediateCommandList();
 	FRDGBuilder GraphBuilder(RHICmdList);
 
-	InternalCreateVertexBufferRDG<FHairStrandsRaytracingFormat>(GraphBuilder, VertexCount, PositionBuffer, TEXT("Hair.StrandsRaytracing_PositionBuffer"));
+	InternalCreateVertexBufferRDG<FHairStrandsRaytracingFormat>(GraphBuilder, VertexCount, PositionBuffer, TEXT("Hair.StrandsRaytracing_PositionBuffer"), EHairResourceUsageType::Dynamic);
 	GraphBuilder.Execute();
 }
 
diff --git a/Engine/Source/Runtime/Renderer/Private/HairStrands/HairStrandsInterface.cpp b/Engine/Source/Runtime/Renderer/Private/HairStrands/HairStrandsInterface.cpp
--- a/Engine/Source/Runtime/Renderer/Private/HairStrands/HairStrandsInterface.cpp
+++ b/Engine/Source/Runtime/Renderer/Private/HairStrands/HairStrandsInterface.cpp
@@ -208,13 +208,25 @@ void ConvertToExternalBufferWithViews(FRDGBuilder& GraphBuilder, FRDGBufferRef&
 	ConvertToExternalBuffer(GraphBuilder, InBuffer, OutBuffer.Buffer);
 	if (Format != PF_Unknown)
 	{
-		OutBuffer.SRV = OutBuffer.Buffer->GetOrCreateSRV(FRDGBufferSRVDesc(InBuffer, Format));
-		OutBuffer.UAV = OutBuffer.Buffer->GetOrCreateUAV(FRDGBufferUAVDesc(InBuffer, Format));
+		if (InBuffer->Desc.Usage & BUF_ShaderResource)
+		{
+			OutBuffer.SRV = OutBuffer.Buffer->GetOrCreateSRV(FRDGBufferSRVDesc(InBuffer, Format));
+		}
+		if (InBuffer->Desc.Usage & BUF_UnorderedAccess)
+		{
+			OutBuffer.UAV = OutBuffer.Buffer->GetOrCreateUAV(FRDGBufferUAVDesc(InBuffer, Format));
+		}
 	}
 	else
 	{
-		OutBuffer.SRV = OutBuffer.Buffer->GetOrCreateSRV(FRDGBufferSRVDesc(InBuffer));
-		OutBuffer.UAV = OutBuffer.Buffer->GetOrCreateUAV(FRDGBufferUAVDesc(InBuffer));
+		if (InBuffer->Desc.Usage & BUF_ShaderResource)
+		{
+			OutBuffer.SRV = OutBuffer.Buffer->GetOrCreateSRV(FRDGBufferSRVDesc(InBuffer));
+		}
+		if (InBuffer->Desc.Usage & BUF_UnorderedAccess)
+		{
+			OutBuffer.UAV = OutBuffer.Buffer->GetOrCreateUAV(FRDGBufferUAVDesc(InBuffer));
+		}
 	}
 	OutBuffer.Format = Format;
 }
