Change groom guides rest resources and strands interpolation data, to be lazily allocated to reduce memory overhead.

These resources are allocated only if used by a groom component with simulation/RBF deformation.

diff --git a/Engine/Plugins/Runtime/HairStrands/Source/HairStrandsCore/Private/GroomAsset.cpp b/Engine/Plugins/Runtime/HairStrands/Source/HairStrandsCore/Private/GroomAsset.cpp
--- a/Engine/Plugins/Runtime/HairStrands/Source/HairStrandsCore/Private/GroomAsset.cpp
+++ b/Engine/Plugins/Runtime/HairStrands/Source/HairStrandsCore/Private/GroomAsset.cpp
@@ -439,12 +439,18 @@ void UGroomAsset::UpdateResource()
 
 		check(GroupIndex < uint32(GetNumHairGroups()));
 		FHairGroupData& GroupData = HairGroupsData[GroupIndex];
-		InternalUpdateResource(GroupData.Guides.RestResource);
+		if (GroupData.Guides.RestResource)
+		{
+			InternalUpdateResource(GroupData.Guides.RestResource);
+		}
 
 		if (IsHairStrandsEnabled(EHairStrandsShaderType::Strands))
 		{
 			InternalUpdateResource(GroupData.Strands.RestResource);
-			InternalUpdateResource(GroupData.Strands.InterpolationResource);
+			if (GroupData.Strands.InterpolationResource)
+			{
+				InternalUpdateResource(GroupData.Strands.InterpolationResource);
+			}
 
 			if ((ChangeType & GroomChangeType_LOD) && GroupData.Strands.HasValidData())
 			{
@@ -2159,8 +2165,18 @@ void UGroomAsset::UpdateCachedSettings()
 
 void UGroomAsset::InitGuideResources()
 {
-	for (FHairGroupData& GroupData : HairGroupsData)
+	// Guide resources are lazy allocated, as these resources are only used when RBF/simulation is enabled by a groom component
+	// for (uint32 GroupIndex = 0, GroupCount = GetNumHairGroups(); GroupIndex < GroupCount; ++GroupIndex)
+	//{
+	//	AllocateGuidesResources(GroupIndex);
+	//}
+}
+
+void UGroomAsset::AllocateGuidesResources(uint32 GroupIndex)
+{
+	if (GroupIndex < uint32(GetNumHairGroups()))
 	{
+		FHairGroupData& GroupData = HairGroupsData[GroupIndex];
 		if (GroupData.Guides.HasValidData())
 		{
 			GroupData.Guides.RestResource = new FHairStrandsRestResource(GroupData.Guides.Data.RenderData, GroupData.Guides.Data.BoundingBox.GetCenter());
@@ -2169,6 +2185,17 @@ void UGroomAsset::InitGuideResources()
 	}
 }
 
+void UGroomAsset::AllocateInterpolationResources(uint32 GroupIndex)
+{
+	if (GroupIndex < uint32(GetNumHairGroups()))
+	{
+		FHairGroupData& GroupData = HairGroupsData[GroupIndex];
+		check(GroupData.Guides.IsValid());
+		GroupData.Strands.InterpolationResource = new FHairStrandsInterpolationResource(GroupData.Strands.InterpolationData.RenderData, GroupData.Guides.Data);
+		BeginInitResource(GroupData.Strands.InterpolationResource);
+	}
+}
+
 void UGroomAsset::InitStrandsResources()
 {
 	// Even though we shouldn't build the strands resources if the platforms does 
@@ -2205,9 +2232,8 @@ void UGroomAsset::InitStrandsResources()
 				BeginInitResource(GroupData.Strands.ClusterCullingResource);
 			}
 
-			check(GroupData.Guides.IsValid());
-			GroupData.Strands.InterpolationResource = new FHairStrandsInterpolationResource(GroupData.Strands.InterpolationData.RenderData, GroupData.Guides.Data);
-			BeginInitResource(GroupData.Strands.InterpolationResource);
+			// Interpolation are lazy allocated, as these resources are only used when RBF/simulation is enabled by a groom component
+			// AllocateInterpolationResources(GroupIndex);
 		}
 	}
 }
diff --git a/Engine/Plugins/Runtime/HairStrands/Source/HairStrandsCore/Private/GroomComponent.cpp b/Engine/Plugins/Runtime/HairStrands/Source/HairStrandsCore/Private/GroomComponent.cpp
--- a/Engine/Plugins/Runtime/HairStrands/Source/HairStrandsCore/Private/GroomComponent.cpp
+++ b/Engine/Plugins/Runtime/HairStrands/Source/HairStrandsCore/Private/GroomComponent.cpp
@@ -2100,7 +2100,8 @@ void UGroomComponent::InitResources(bool bIsBindingReloading)
 		// * Physics simulation
 		// * RBF deformation.
 		// Therefore, even if simulation is disabled, we need to run partially the update if the binding system is enabled (skin deformation + RBF correction)
-		if (GroupData.Guides.IsValid() && (bHasNeedSimulation || bHasNeedGlobalDeformation || bPreviewMode))
+		const bool bNeedGuides = GroupData.Guides.HasValidData() && (bHasNeedSimulation || bHasNeedGlobalDeformation || bPreviewMode);
+		if (bNeedGuides)
 		{
 			HairGroupInstance->Guides.Data = &GroupData.Guides.Data;
 
@@ -2119,6 +2120,14 @@ void UGroomComponent::InitResources(bool bIsBindingReloading)
 				HairGroupInstance->Guides.DeformedRootResource = new FHairStrandsDeformedRootResource(HairGroupInstance->Guides.RestRootResource);
 				BeginInitResource(HairGroupInstance->Guides.DeformedRootResource);
 			}
+
+			// (Lazy) Allocation of the guide rest resources
+			if (!GroupData.Guides.RestResource)
+			{
+				// Lazy allocation of the guide resources
+				GroomAsset->AllocateGuidesResources(GroupIt);
+				check(GroupData.Guides.RestResource);
+			}
 			HairGroupInstance->Guides.RestResource = GroupData.Guides.RestResource;
 
 			// If guides are allocated, deformed resources are always needs since they are either used with simulation, or RBF deformation. Both are dynamics, and require deformed positions
@@ -2144,7 +2153,18 @@ void UGroomComponent::InitResources(bool bIsBindingReloading)
 
 			HairGroupInstance->Strands.Data = &GroupData.Strands.Data;
 			HairGroupInstance->Strands.InterpolationData = &GroupData.Strands.InterpolationData;
-			HairGroupInstance->Strands.InterpolationResource = GroupData.Strands.InterpolationResource;
+
+			// (Lazy) Allocate interpolation resources, only if guides are required
+			if (bNeedGuides)
+			{
+				if (!GroupData.Strands.InterpolationResource)
+				{
+					// Lazy allocation of the guide resources
+					GroomAsset->AllocateInterpolationResources(GroupIt);
+					check(GroupData.Strands.InterpolationResource);
+				}
+				HairGroupInstance->Strands.InterpolationResource = GroupData.Strands.InterpolationResource;
+			}
 
 			// Material
 			HairGroupInstance->Strands.Material = nullptr;
diff --git a/Engine/Plugins/Runtime/HairStrands/Source/HairStrandsCore/Private/HairStrandsRendering.cpp b/Engine/Plugins/Runtime/HairStrands/Source/HairStrandsCore/Private/HairStrandsRendering.cpp
--- a/Engine/Plugins/Runtime/HairStrands/Source/HairStrandsCore/Private/HairStrandsRendering.cpp
+++ b/Engine/Plugins/Runtime/HairStrands/Source/HairStrandsCore/Private/HairStrandsRendering.cpp
@@ -1507,8 +1507,8 @@ void ComputeHairStrandsInterpolation(
 					bHasSkinning && bValidGuide ? Instance->Guides.DeformedRootResource : nullptr,
 					RegisterAsSRV(GraphBuilder, Instance->Strands.RestResource->PositionBuffer),
 					RegisterAsSRV(GraphBuilder, Instance->Strands.RestResource->AttributeBuffer),
-					RegisterAsSRV(GraphBuilder, Instance->Strands.InterpolationResource->Interpolation0Buffer),
-					RegisterAsSRV(GraphBuilder, Instance->Strands.InterpolationResource->Interpolation1Buffer),
+					bValidGuide ? RegisterAsSRV(GraphBuilder, Instance->Strands.InterpolationResource->Interpolation0Buffer) : nullptr,
+					bValidGuide ? RegisterAsSRV(GraphBuilder, Instance->Strands.InterpolationResource->Interpolation1Buffer) : nullptr,
 					bValidGuide ? RegisterAsSRV(GraphBuilder, Instance->Guides.RestResource->PositionBuffer) : nullptr,
 					bValidGuide ? RegisterAsSRV(GraphBuilder, Instance->Guides.DeformedResource->GetBuffer(FHairStrandsDeformedResource::Current)) : nullptr,
 					bValidGuide ? RegisterAsSRV(GraphBuilder, Instance->Guides.RestResource->AttributeBuffer) : nullptr,
@@ -1516,7 +1516,7 @@ void ComputeHairStrandsInterpolation(
 					&Strands_DeformedPrevPosition,
 					Strands_DebugAttribute,
 					RegisterAsSRV(GraphBuilder, Instance->Strands.ClusterCullingResource->VertexToClusterIdBuffer),
-					RegisterAsSRV(GraphBuilder, Instance->Strands.InterpolationResource->SimRootPointIndexBuffer));
+					bValidGuide ? RegisterAsSRV(GraphBuilder, Instance->Strands.InterpolationResource->SimRootPointIndexBuffer) : nullptr);
 
 				// 2.2 Update tangent data based on the deformed positions
 				AddHairTangentPass(
diff --git a/Engine/Plugins/Runtime/HairStrands/Source/HairStrandsCore/Public/GroomAsset.h b/Engine/Plugins/Runtime/HairStrands/Source/HairStrandsCore/Public/GroomAsset.h
--- a/Engine/Plugins/Runtime/HairStrands/Source/HairStrandsCore/Public/GroomAsset.h
+++ b/Engine/Plugins/Runtime/HairStrands/Source/HairStrandsCore/Public/GroomAsset.h
@@ -454,6 +454,10 @@ public:
 private:
 	void ApplyStripFlags(uint8 StripFlags, const class ITargetPlatform* CookTarget);
 
+	void AllocateGuidesResources(uint32 GroupIndex);
+	void AllocateInterpolationResources(uint32 GroupIndex);
+	friend class UGroomComponent;
+
 #if WITH_EDITORONLY_DATA
 	bool HasImportedStrandsData() const;
 
