Fix CIS issue.

Static analysis get confused as a variable, which is already early out if null, get retested.

diff --git a/Engine/Plugins/Runtime/HairStrands/Source/HairStrandsCore/Private/GroomComponent.cpp b/Engine/Plugins/Runtime/HairStrands/Source/HairStrandsCore/Private/GroomComponent.cpp
--- a/Engine/Plugins/Runtime/HairStrands/Source/HairStrandsCore/Private/GroomComponent.cpp
+++ b/Engine/Plugins/Runtime/HairStrands/Source/HairStrandsCore/Private/GroomComponent.cpp
@@ -2139,8 +2139,8 @@ void UGroomComponent::InitResources(bool bIsBindingReloading)
 
 			// Initialize the simulation and the global deformation to its default behavior by setting it with LODIndex = -1
 			const int32 LODIndex = -1;
-			HairGroupInstance->Guides.bIsSimulationEnable = GroomAsset && GroomAsset->IsGlobalInterpolationEnable(GroupIt,LODIndex);
-			HairGroupInstance->Guides.bHasGlobalInterpolation = LocalBindingAsset && GroomAsset && GroomAsset->IsGlobalInterpolationEnable(GroupIt,LODIndex);
+			HairGroupInstance->Guides.bIsSimulationEnable = GroomAsset->IsGlobalInterpolationEnable(GroupIt,LODIndex);
+			HairGroupInstance->Guides.bHasGlobalInterpolation = LocalBindingAsset && GroomAsset->IsGlobalInterpolationEnable(GroupIt,LODIndex);
 		}
 
 		// LODBias is in the Modifier which is needed for LOD selection regardless if the strands are there or not
