Fix invalid groom bounding box when a groom component does not have a groom asset.

The (default) bound computation was missing its local transform.

diff --git a/Engine/Plugins/Runtime/HairStrands/Source/HairStrandsCore/Private/GroomComponent.cpp b/Engine/Plugins/Runtime/HairStrands/Source/HairStrandsCore/Private/GroomComponent.cpp
--- a/Engine/Plugins/Runtime/HairStrands/Source/HairStrandsCore/Private/GroomComponent.cpp
+++ b/Engine/Plugins/Runtime/HairStrands/Source/HairStrandsCore/Private/GroomComponent.cpp
@@ -1494,7 +1494,8 @@ FBoxSphereBounds UGroomComponent::CalcBounds(const FTransform& InLocalToWorld) c
 	}
 	else
 	{
-		return FBoxSphereBounds(EForceInit::ForceInitToZero);
+		FBoxSphereBounds LocalBounds(EForceInit::ForceInitToZero);
+		return FBoxSphereBounds(LocalBounds.TransformBy(InLocalToWorld));
 	}
 }
 
