Mark Groom resources as persistent RDG resources.

diff --git a/Engine/Plugins/Runtime/HairStrands/Source/HairStrandsCore/Private/GroomResources.cpp b/Engine/Plugins/Runtime/HairStrands/Source/HairStrandsCore/Private/GroomResources.cpp
--- a/Engine/Plugins/Runtime/HairStrands/Source/HairStrandsCore/Private/GroomResources.cpp
+++ b/Engine/Plugins/Runtime/HairStrands/Source/HairStrandsCore/Private/GroomResources.cpp
@@ -95,6 +95,24 @@ void CreateBuffer(const TArray<typename FormatType::Type>& InData, FHairCardsVer
 
 /////////////////////////////////////////////////////////////////////////////////////////
 // RDG buffers utils 
+
+static FRDGBufferRef InternalCreateVertexBuffer(
+	FRDGBuilder& GraphBuilder,
+	const TCHAR* Name,
+	const FRDGBufferDesc& Desc,
+	const void* InitialData,
+	uint64 InitialDataSize,
+	ERDGInitialDataFlags InitialDataFlags)
+{
+	checkf(Desc.UnderlyingType == FRDGBufferDesc::EUnderlyingType::VertexBuffer, TEXT("CreateVertexBuffer called with an FRDGBufferDesc underlying type that is not 'VertexBuffer'. Buffer: %s"), Name);
+	FRDGBufferRef Buffer = GraphBuilder.CreateBuffer(Desc, Name, ERDGBufferFlags::MultiFrame);
+	if (InitialData && InitialDataSize)
+	{
+		AddBufferUploadPass(GraphBuilder, Buffer, InitialData, InitialDataSize, InitialDataFlags);
+	}
+	return Buffer;
+}
+
 template<typename FormatType>
 void InternalCreateVertexBufferRDG(FRDGBuilder& GraphBuilder, const TArray<typename FormatType::Type>& InData, FRDGExternalBuffer& Out, const TCHAR* DebugName, EHairResourceUsageType UsageType)
 {
@@ -113,7 +131,7 @@ void InternalCreateVertexBufferRDG(FRDGBuilder& GraphBuilder, const TArray<typen
 	{
 		Desc.Usage = Desc.Usage & (~BUF_UnorderedAccess);
 	}
-	Buffer = CreateVertexBuffer(
+	Buffer = InternalCreateVertexBuffer(
 		GraphBuilder,
 		DebugName,
 		Desc,
@@ -142,7 +160,7 @@ void InternalCreateVertexBufferRDG(FRDGBuilder& GraphBuilder, const TArray<DataT
 	{
 		Desc.Usage = Desc.Usage & (~BUF_UnorderedAccess);
 	}
-	Buffer = CreateVertexBuffer(
+	Buffer = InternalCreateVertexBuffer(
 		GraphBuilder,
 		DebugName,
 		Desc,
@@ -223,7 +241,7 @@ void InternalCreateStructuredBufferRDG(FRDGBuilder& GraphBuilder, const TArray<F
 		Desc.Usage = Desc.Usage & (~BUF_UnorderedAccess);
 	}
 
-	FRDGBufferRef Buffer = GraphBuilder.CreateBuffer(Desc, DebugName);
+	FRDGBufferRef Buffer = GraphBuilder.CreateBuffer(Desc, DebugName, ERDGBufferFlags::MultiFrame);
 	if (InData.GetData() && DataSizeInBytes)
 	{
 		AddBufferUploadPass(GraphBuilder, Buffer, InData.GetData(), DataSizeInBytes, ERDGInitialDataFlags::NoCopy);
