Second pass of hair resources transition issue.

THis CL adds an explicit transition pass prior to all mesh draw pass rendering, converting resources used by HairStrands VF into SRV state, making them untracked for the rest of the RDG builder.

diff --git a/Engine/Plugins/Runtime/HairStrands/Source/HairStrandsCore/Private/HairStrandsRendering.cpp b/Engine/Plugins/Runtime/HairStrands/Source/HairStrandsCore/Private/HairStrandsRendering.cpp
--- a/Engine/Plugins/Runtime/HairStrands/Source/HairStrandsCore/Private/HairStrandsRendering.cpp
+++ b/Engine/Plugins/Runtime/HairStrands/Source/HairStrandsCore/Private/HairStrandsRendering.cpp
@@ -1307,8 +1307,22 @@ bool NeedsPatchAttributeBuffer(EHairStrandsDebugMode DebugMode)
 	return DebugMode == EHairStrandsDebugMode::RenderHairStrands || DebugMode == EHairStrandsDebugMode::RenderVisCluster;
 }
 
+static void ConvertHairStrandsVFParameters(
+	FRDGBuilder* GraphBuilder,
+	FRDGImportedBuffer& Buffer,
+	FShaderResourceViewRHIRef& BufferRHISRV,
+	FRDGExternalBuffer& ExternalBuffer)
+{
+	if (GraphBuilder)
+	{
+		Buffer = Register(*GraphBuilder, ExternalBuffer, ERDGImportedBufferFlags::CreateSRV);
+	}
+	BufferRHISRV = ExternalBuffer.SRV;
+}
+#define CONVERT_HAIRSSTRANDS_VF_PARAMETERS(OutName, ExternalBuffer) ConvertHairStrandsVFParameters(GraphBuilder, OutName, OutName##RHISRV, ExternalBuffer);
+
 // Compute/Update the hair strands description which will be used for rendering (VF) / voxelization & co.
-FHairGroupPublicData::FVertexFactoryInput ComputeHairStrandsVertexInputData(FHairGroupInstance* Instance)
+static FHairGroupPublicData::FVertexFactoryInput InternalComputeHairStrandsVertexInputData(FRDGBuilder* GraphBuilder, FHairGroupInstance* Instance)
 {
 	FHairGroupPublicData::FVertexFactoryInput OutVFInput;
 	if (!Instance || Instance->GeometryType != EHairGeometryType::Strands)
@@ -1324,19 +1338,16 @@ FHairGroupPublicData::FVertexFactoryInput ComputeHairStrandsVertexInputData(FHai
 	// 1. Guides
 	if (DebugMode == EHairStrandsDebugMode::SimHairStrands)
 	{
-
-		OutVFInput.Strands.PositionBuffer = Instance->Guides.DeformedResource->GetBuffer(FHairStrandsDeformedResource::EFrameType::Current).SRV;
-		OutVFInput.Strands.PrevPositionBuffer = Instance->Guides.DeformedResource->GetBuffer(FHairStrandsDeformedResource::EFrameType::Previous).SRV;
-		OutVFInput.Strands.TangentBuffer = Instance->Guides.DeformedResource->TangentBuffer.SRV;
-		OutVFInput.Strands.AttributeBuffer = Instance->Guides.RestResource->AttributeBuffer.SRV;
-		OutVFInput.Strands.MaterialBuffer = Instance->Guides.RestResource->MaterialBuffer.SRV;
+		CONVERT_HAIRSSTRANDS_VF_PARAMETERS(OutVFInput.Strands.PositionBuffer,			Instance->Guides.DeformedResource->GetBuffer(FHairStrandsDeformedResource::EFrameType::Current));
+		CONVERT_HAIRSSTRANDS_VF_PARAMETERS(OutVFInput.Strands.PrevPositionBuffer,		Instance->Guides.DeformedResource->GetBuffer(FHairStrandsDeformedResource::EFrameType::Previous));
+		CONVERT_HAIRSSTRANDS_VF_PARAMETERS(OutVFInput.Strands.TangentBuffer,			Instance->Guides.DeformedResource->TangentBuffer);
+		CONVERT_HAIRSSTRANDS_VF_PARAMETERS(OutVFInput.Strands.AttributeBuffer,			Instance->Guides.RestResource->AttributeBuffer);
+		CONVERT_HAIRSSTRANDS_VF_PARAMETERS(OutVFInput.Strands.MaterialBuffer,			Instance->Guides.RestResource->MaterialBuffer);
+		CONVERT_HAIRSSTRANDS_VF_PARAMETERS(OutVFInput.Strands.PositionOffsetBuffer,		Instance->Guides.DeformedResource->GetPositionOffsetBuffer(FHairStrandsDeformedResource::EFrameType::Current));
+		CONVERT_HAIRSSTRANDS_VF_PARAMETERS(OutVFInput.Strands.PrevPositionOffsetBuffer,	Instance->Guides.DeformedResource->GetPositionOffsetBuffer(FHairStrandsDeformedResource::EFrameType::Previous));
 
 		OutVFInput.Strands.PositionOffset = Instance->Guides.DeformedResource->GetPositionOffset(FHairStrandsDeformedResource::EFrameType::Current);
 		OutVFInput.Strands.PrevPositionOffset = Instance->Guides.DeformedResource->GetPositionOffset(FHairStrandsDeformedResource::EFrameType::Previous);
-
-		OutVFInput.Strands.PositionOffsetBuffer = Instance->Guides.DeformedResource->GetPositionOffsetBuffer(FHairStrandsDeformedResource::EFrameType::Current).SRV;
-		OutVFInput.Strands.PrevPositionOffsetBuffer = Instance->Guides.DeformedResource->GetPositionOffsetBuffer(FHairStrandsDeformedResource::EFrameType::Previous).SRV;
-
 		OutVFInput.Strands.VertexCount = Instance->Guides.RestResource->GetVertexCount();
 		OutVFInput.Strands.HairRadius = (GStrandHairWidth > 0 ? GStrandHairWidth : Instance->Strands.Modifier.HairWidth) * 0.5f;
 		OutVFInput.Strands.HairLength = Instance->Strands.Modifier.HairLength;
@@ -1347,15 +1358,13 @@ FHairGroupPublicData::FVertexFactoryInput ComputeHairStrandsVertexInputData(FHai
 	// 2. Render Strands with deformation
 	else if (bSupportDeformation)
 	{
-
-		OutVFInput.Strands.PositionBuffer = Instance->Strands.DeformedResource->GetBuffer(FHairStrandsDeformedResource::EFrameType::Current).SRV;
-		OutVFInput.Strands.PrevPositionBuffer = Instance->Strands.DeformedResource->GetBuffer(FHairStrandsDeformedResource::EFrameType::Previous).SRV;
-		OutVFInput.Strands.TangentBuffer = Instance->Strands.DeformedResource->TangentBuffer.SRV;
-		OutVFInput.Strands.AttributeBuffer = bDebugModePatchedAttributeBuffer ? Instance->Strands.DebugAttributeBuffer.SRV : Instance->Strands.RestResource->AttributeBuffer.SRV;
-		OutVFInput.Strands.MaterialBuffer = Instance->Strands.RestResource->MaterialBuffer.SRV;
-
-		OutVFInput.Strands.PositionOffsetBuffer = Instance->Strands.DeformedResource->GetPositionOffsetBuffer(FHairStrandsDeformedResource::EFrameType::Current).SRV;
-		OutVFInput.Strands.PrevPositionOffsetBuffer = Instance->Strands.DeformedResource->GetPositionOffsetBuffer(FHairStrandsDeformedResource::EFrameType::Previous).SRV;
+		CONVERT_HAIRSSTRANDS_VF_PARAMETERS(OutVFInput.Strands.PositionBuffer,			Instance->Strands.DeformedResource->GetBuffer(FHairStrandsDeformedResource::EFrameType::Current));
+		CONVERT_HAIRSSTRANDS_VF_PARAMETERS(OutVFInput.Strands.PrevPositionBuffer,		Instance->Strands.DeformedResource->GetBuffer(FHairStrandsDeformedResource::EFrameType::Previous));
+		CONVERT_HAIRSSTRANDS_VF_PARAMETERS(OutVFInput.Strands.TangentBuffer,			Instance->Strands.DeformedResource->TangentBuffer);
+		CONVERT_HAIRSSTRANDS_VF_PARAMETERS(OutVFInput.Strands.AttributeBuffer,			bDebugModePatchedAttributeBuffer ? Instance->Strands.DebugAttributeBuffer : Instance->Strands.RestResource->AttributeBuffer);
+		CONVERT_HAIRSSTRANDS_VF_PARAMETERS(OutVFInput.Strands.MaterialBuffer,			Instance->Strands.RestResource->MaterialBuffer);
+		CONVERT_HAIRSSTRANDS_VF_PARAMETERS(OutVFInput.Strands.PositionOffsetBuffer,		Instance->Strands.DeformedResource->GetPositionOffsetBuffer(FHairStrandsDeformedResource::EFrameType::Current));
+		CONVERT_HAIRSSTRANDS_VF_PARAMETERS(OutVFInput.Strands.PrevPositionOffsetBuffer,	Instance->Strands.DeformedResource->GetPositionOffsetBuffer(FHairStrandsDeformedResource::EFrameType::Previous));
 
 		OutVFInput.Strands.PositionOffset = Instance->Strands.DeformedResource->GetPositionOffset(FHairStrandsDeformedResource::EFrameType::Current);
 		OutVFInput.Strands.PrevPositionOffset = Instance->Strands.DeformedResource->GetPositionOffset(FHairStrandsDeformedResource::EFrameType::Previous);
@@ -1369,15 +1378,13 @@ FHairGroupPublicData::FVertexFactoryInput ComputeHairStrandsVertexInputData(FHai
 	// 3. Render Strands in rest position: used when there are no skinning (rigid binding or no binding), no simulation, and no RBF
 	else
 	{
-
-		OutVFInput.Strands.PositionBuffer = Instance->Strands.RestResource->PositionBuffer.SRV;
-		OutVFInput.Strands.PrevPositionBuffer = Instance->Strands.RestResource->PositionBuffer.SRV;
-		OutVFInput.Strands.TangentBuffer = Instance->Strands.RestResource->TangentBuffer.SRV;
-		OutVFInput.Strands.AttributeBuffer = Instance->Strands.RestResource->AttributeBuffer.SRV;
-		OutVFInput.Strands.MaterialBuffer = Instance->Strands.RestResource->MaterialBuffer.SRV;
-
-		OutVFInput.Strands.PositionOffsetBuffer = Instance->Strands.RestResource->PositionOffsetBuffer.SRV;
-		OutVFInput.Strands.PrevPositionOffsetBuffer = Instance->Strands.RestResource->PositionOffsetBuffer.SRV;
+		CONVERT_HAIRSSTRANDS_VF_PARAMETERS(OutVFInput.Strands.PositionBuffer,			Instance->Strands.RestResource->PositionBuffer);
+		CONVERT_HAIRSSTRANDS_VF_PARAMETERS(OutVFInput.Strands.PrevPositionBuffer,		Instance->Strands.RestResource->PositionBuffer);
+		CONVERT_HAIRSSTRANDS_VF_PARAMETERS(OutVFInput.Strands.TangentBuffer,			Instance->Strands.RestResource->TangentBuffer);
+		CONVERT_HAIRSSTRANDS_VF_PARAMETERS(OutVFInput.Strands.AttributeBuffer,			Instance->Strands.RestResource->AttributeBuffer);
+		CONVERT_HAIRSSTRANDS_VF_PARAMETERS(OutVFInput.Strands.MaterialBuffer,			Instance->Strands.RestResource->MaterialBuffer);
+		CONVERT_HAIRSSTRANDS_VF_PARAMETERS(OutVFInput.Strands.PositionOffsetBuffer,		Instance->Strands.RestResource->PositionOffsetBuffer);
+		CONVERT_HAIRSSTRANDS_VF_PARAMETERS(OutVFInput.Strands.PrevPositionOffsetBuffer, Instance->Strands.RestResource->PositionOffsetBuffer);
 
 		OutVFInput.Strands.PositionOffset = Instance->Strands.RestResource->PositionOffset;
 		OutVFInput.Strands.PrevPositionOffset = Instance->Strands.RestResource->PositionOffset;
@@ -1392,6 +1399,11 @@ FHairGroupPublicData::FVertexFactoryInput ComputeHairStrandsVertexInputData(FHai
 	return OutVFInput;
 }
 
+FHairGroupPublicData::FVertexFactoryInput ComputeHairStrandsVertexInputData(FHairGroupInstance* Instance)
+{
+	return InternalComputeHairStrandsVertexInputData(nullptr, Instance);
+}
+
 void AddBufferTransitionToReadablePass(FRDGBuilder& GraphBuilder, FRHIUnorderedAccessView* UAV)
 {
 	AddPass(GraphBuilder, [UAV](FRHICommandList& RHICmdList)
@@ -1448,7 +1460,7 @@ void ComputeHairStrandsInterpolation(
 				RegisterAsSRV(GraphBuilder, Instance->Guides.DeformedResource->GetBuffer(FHairStrandsDeformedResource::Current)),
 				Register(GraphBuilder, Instance->Guides.DeformedResource->TangentBuffer, ERDGImportedBufferFlags::CreateUAV));
 
-			Instance->HairGroupPublicData->VFInput = ComputeHairStrandsVertexInputData(Instance);
+			Instance->HairGroupPublicData->VFInput = InternalComputeHairStrandsVertexInputData(&GraphBuilder, Instance);
 		}
 		else
 		{
@@ -1537,7 +1549,7 @@ void ComputeHairStrandsInterpolation(
 			}
 
 			// 2.1 Update the VF input with the update resources
-			Instance->HairGroupPublicData->VFInput = ComputeHairStrandsVertexInputData(Instance);
+			Instance->HairGroupPublicData->VFInput = InternalComputeHairStrandsVertexInputData(&GraphBuilder, Instance);
 
 			// 3. Compute cluster AABBs (used for LODing and voxelization)
 			FHairStrandClusterData::FHairGroup& HairGroupCluster = InClusterData->HairGroups[Instance->HairGroupPublicData->ClusterDataIndex];
@@ -1601,7 +1613,7 @@ void ComputeHairStrandsInterpolation(
 		}
 
 		// Sanity check
-		check(Instance->HairGroupPublicData->VFInput.Strands.PositionBuffer);
+		check(Instance->HairGroupPublicData->VFInput.Strands.PositionBuffer.Buffer);
 	}
 	else if (InstanceGeometryType == EHairGeometryType::Cards)
 	{	
diff --git a/Engine/Plugins/Runtime/HairStrands/Source/HairStrandsCore/Private/HairStrandsVertexFactory.cpp b/Engine/Plugins/Runtime/HairStrands/Source/HairStrandsCore/Private/HairStrandsVertexFactory.cpp
--- a/Engine/Plugins/Runtime/HairStrands/Source/HairStrandsCore/Private/HairStrandsVertexFactory.cpp
+++ b/Engine/Plugins/Runtime/HairStrands/Source/HairStrandsCore/Private/HairStrandsVertexFactory.cpp
@@ -49,14 +49,14 @@ public:
 TGlobalResource<FDummyCulledDispatchVertexIdsBuffer> GDummyCulledDispatchVertexIdsBuffer;
 
 /////////////////////////////////////////////////////////////////////////////////////////
-inline FRHIShaderResourceView* GetPositionSRV(const FHairGroupPublicData::FVertexFactoryInput& VFInput)			{ return VFInput.Strands.PositionBuffer; };
-inline FRHIShaderResourceView* GetPreviousPositionSRV(const FHairGroupPublicData::FVertexFactoryInput& VFInput)	{ return VFInput.Strands.PrevPositionBuffer; }
-inline FRHIShaderResourceView* GetAttributeSRV(const FHairGroupPublicData::FVertexFactoryInput& VFInput)		{ return VFInput.Strands.AttributeBuffer; }
-inline FRHIShaderResourceView* GetMaterialSRV(const FHairGroupPublicData::FVertexFactoryInput& VFInput)			{ return VFInput.Strands.MaterialBuffer; }
-inline FRHIShaderResourceView* GetTangentSRV(const FHairGroupPublicData::FVertexFactoryInput& VFInput)			{ return VFInput.Strands.TangentBuffer; }
-
-inline FRHIShaderResourceView* GetPositionOffsetSRV(const FHairGroupPublicData::FVertexFactoryInput& VFInput) { return VFInput.Strands.PositionOffsetBuffer; };
-inline FRHIShaderResourceView* GetPreviousPositionOffsetSRV(const FHairGroupPublicData::FVertexFactoryInput& VFInput) { return VFInput.Strands.PrevPositionOffsetBuffer; }
+inline FRHIShaderResourceView* GetPositionSRV(const FHairGroupPublicData::FVertexFactoryInput& VFInput)			{ return VFInput.Strands.PositionBufferRHISRV; }; 
+inline FRHIShaderResourceView* GetPreviousPositionSRV(const FHairGroupPublicData::FVertexFactoryInput& VFInput)	{ return VFInput.Strands.PrevPositionBufferRHISRV; }
+inline FRHIShaderResourceView* GetAttributeSRV(const FHairGroupPublicData::FVertexFactoryInput& VFInput)		{ return VFInput.Strands.AttributeBufferRHISRV; }
+inline FRHIShaderResourceView* GetMaterialSRV(const FHairGroupPublicData::FVertexFactoryInput& VFInput)			{ return VFInput.Strands.MaterialBufferRHISRV; }
+inline FRHIShaderResourceView* GetTangentSRV(const FHairGroupPublicData::FVertexFactoryInput& VFInput)			{ return VFInput.Strands.TangentBufferRHISRV; }
+
+inline FRHIShaderResourceView* GetPositionOffsetSRV(const FHairGroupPublicData::FVertexFactoryInput& VFInput)	{ return VFInput.Strands.PositionOffsetBufferRHISRV; };
+inline FRHIShaderResourceView* GetPreviousPositionOffsetSRV(const FHairGroupPublicData::FVertexFactoryInput& VFInput) { return VFInput.Strands.PrevPositionOffsetBufferRHISRV; }
 
 inline bool  UseStableRasterization(const FHairGroupPublicData::FVertexFactoryInput& VFInput)					{ return VFInput.Strands.bUseStableRasterization; };
 inline bool  UseScatterSceneLighting(const FHairGroupPublicData::FVertexFactoryInput& VFInput)					{ return VFInput.Strands.bScatterSceneLighting; };
diff --git a/Engine/Source/Runtime/Renderer/Private/HairStrands/HairStrandsVisibility.cpp b/Engine/Source/Runtime/Renderer/Private/HairStrands/HairStrandsVisibility.cpp
--- a/Engine/Source/Runtime/Renderer/Private/HairStrands/HairStrandsVisibility.cpp
+++ b/Engine/Source/Runtime/Renderer/Private/HairStrands/HairStrandsVisibility.cpp
@@ -2993,6 +2993,45 @@ static FRDGTextureRef AddHairHairCountToTransmittancePass(
 	return OutputTexture;
 }
 
+// Transit resources used during the MeshDraw passes
+static void AddMeshDrawTransitionPass(
+	FRDGBuilder& GraphBuilder,
+	const FViewInfo& ViewInfo,
+	const FHairStrandsMacroGroupDatas& MacroGroupDatas)
+{
+	for (const FHairStrandsMacroGroupData& MacroGroup : MacroGroupDatas)
+	{
+		for (const FHairStrandsMacroGroupData::PrimitiveInfo& PrimitiveInfo : MacroGroup.PrimitivesInfos)
+		{
+			FHairGroupPublicData* HairGroupPublicData = reinterpret_cast<FHairGroupPublicData*>(PrimitiveInfo.Mesh->Elements[0].VertexFactoryUserData);
+			check(HairGroupPublicData);
+
+			FHairGroupPublicData::FVertexFactoryInput& VFInput = HairGroupPublicData->VFInput;
+			ConvertToUntrackedBuffer(GraphBuilder, VFInput.Strands.PositionBuffer.Buffer,			ERHIAccess::SRVMask);
+			ConvertToUntrackedBuffer(GraphBuilder, VFInput.Strands.PrevPositionBuffer.Buffer,		ERHIAccess::SRVMask);
+			ConvertToUntrackedBuffer(GraphBuilder, VFInput.Strands.TangentBuffer.Buffer,			ERHIAccess::SRVMask);
+			ConvertToUntrackedBuffer(GraphBuilder, VFInput.Strands.AttributeBuffer.Buffer,			ERHIAccess::SRVMask);
+			ConvertToUntrackedBuffer(GraphBuilder, VFInput.Strands.MaterialBuffer.Buffer,			ERHIAccess::SRVMask);
+			ConvertToUntrackedBuffer(GraphBuilder, VFInput.Strands.PositionOffsetBuffer.Buffer,		ERHIAccess::SRVMask);
+			ConvertToUntrackedBuffer(GraphBuilder, VFInput.Strands.PrevPositionOffsetBuffer.Buffer,	ERHIAccess::SRVMask);
+
+			FRDGBufferRef CulledVertexIdBuffer = Register(GraphBuilder, HairGroupPublicData->CulledVertexIdBuffer, ERDGImportedBufferFlags::None).Buffer;
+			FRDGBufferRef CulledVertexRadiusScaleBuffer = Register(GraphBuilder, HairGroupPublicData->CulledVertexRadiusScaleBuffer, ERDGImportedBufferFlags::None).Buffer;
+			FRDGBufferRef DrawIndirectBuffer = Register(GraphBuilder, HairGroupPublicData->DrawIndirectBuffer, ERDGImportedBufferFlags::None).Buffer;
+			ConvertToUntrackedBuffer(GraphBuilder, CulledVertexIdBuffer, ERHIAccess::SRVMask);
+			ConvertToUntrackedBuffer(GraphBuilder, CulledVertexRadiusScaleBuffer, ERHIAccess::SRVMask);
+			ConvertToUntrackedBuffer(GraphBuilder, DrawIndirectBuffer, ERHIAccess::IndirectArgs);
+
+			VFInput.Strands.PositionBuffer				= FRDGImportedBuffer();
+			VFInput.Strands.PrevPositionBuffer			= FRDGImportedBuffer();
+			VFInput.Strands.TangentBuffer				= FRDGImportedBuffer();
+			VFInput.Strands.AttributeBuffer				= FRDGImportedBuffer();
+			VFInput.Strands.MaterialBuffer				= FRDGImportedBuffer();
+			VFInput.Strands.PositionOffsetBuffer		= FRDGImportedBuffer();
+			VFInput.Strands.PrevPositionOffsetBuffer	= FRDGImportedBuffer();
+		}
+	}
+}
 ///////////////////////////////////////////////////////////////////////////////////////////////////
 
 class FVisiblityRasterComputeCS : public FGlobalShader
@@ -3020,8 +3059,8 @@ class FVisiblityRasterComputeCS : public FGlobalShader
 		SHADER_PARAMETER(uint32, HairStrandsVF_bUseStableRasterization)
 		SHADER_PARAMETER(uint32, HairStrandsVF_VertexCount)
 		SHADER_PARAMETER(FMatrix, HairStrandsVF_LocalToWorldPrimitiveTransform)
-		SHADER_PARAMETER_SRV(Buffer, HairStrandsVF_PositionBuffer)
-		SHADER_PARAMETER_SRV(Buffer, HairStrandsVF_PositionOffsetBuffer)
+		SHADER_PARAMETER_RDG_BUFFER_SRV(Buffer, HairStrandsVF_PositionBuffer)
+		SHADER_PARAMETER_RDG_BUFFER_SRV(Buffer, HairStrandsVF_PositionOffsetBuffer)
 		SHADER_PARAMETER_RDG_BUFFER_SRV(Buffer, HairStrandsVF_CullingIndirectBuffer)
 		SHADER_PARAMETER_RDG_BUFFER_SRV(Buffer, HairStrandsVF_CullingIndexBuffer)
 		SHADER_PARAMETER_RDG_BUFFER_SRV(Buffer, HairStrandsVF_CullingRadiusScaleBuffer)
@@ -3199,8 +3238,8 @@ static FRasterComputeOutput AddVisibilityComputeRasterPass(
 			check(HairGroupPublicData);
 
 			const FHairGroupPublicData::FVertexFactoryInput& VFInput = HairGroupPublicData->VFInput;
-			PassParameters->HairStrandsVF_PositionBuffer	= VFInput.Strands.PositionBuffer;
-			PassParameters->HairStrandsVF_PositionOffsetBuffer = VFInput.Strands.PositionOffsetBuffer;
+			PassParameters->HairStrandsVF_PositionBuffer	= VFInput.Strands.PositionBuffer.SRV;
+			PassParameters->HairStrandsVF_PositionOffsetBuffer = VFInput.Strands.PositionOffsetBuffer.SRV;
 			PassParameters->HairStrandsVF_VertexCount		= VFInput.Strands.VertexCount;
 			PassParameters->HairStrandsVF_Radius			= VFInput.Strands.HairRadius;
 			PassParameters->HairStrandsVF_Length			= VFInput.Strands.HairLength;
@@ -3279,6 +3318,12 @@ void RenderHairStrandsVisibilityBuffer(
 			FRDGTextureRef CompactNodeIndex = nullptr;
 			FRDGBufferRef  CompactNodeData = nullptr;
 			FRDGTextureRef NodeCounter = nullptr;
+
+			if (RenderMode != HairVisibilityRenderMode_ComputeRaster)
+			{
+				AddMeshDrawTransitionPass(GraphBuilder, View, MacroGroupDatas);
+			}
+
 			if (RenderMode == HairVisibilityRenderMode_ComputeRaster)
 			{
 				FRasterComputeOutput RasterOutput = AddVisibilityComputeRasterPass(
diff --git a/Engine/Source/Runtime/Renderer/Private/HairStrands/HairStrandsVoxelization.cpp b/Engine/Source/Runtime/Renderer/Private/HairStrands/HairStrandsVoxelization.cpp
--- a/Engine/Source/Runtime/Renderer/Private/HairStrands/HairStrandsVoxelization.cpp
+++ b/Engine/Source/Runtime/Renderer/Private/HairStrands/HairStrandsVoxelization.cpp
@@ -993,8 +993,8 @@ class FVoxelRasterComputeCS : public FGlobalShader
 		SHADER_PARAMETER(FVector, HairStrandsVF_PositionOffset)
 		SHADER_PARAMETER(uint32,  HairStrandsVF_VertexCount)
 		SHADER_PARAMETER(FMatrix, HairStrandsVF_LocalToWorldPrimitiveTransform)
-		SHADER_PARAMETER_SRV(Buffer, HairStrandsVF_PositionBuffer)
-		SHADER_PARAMETER_SRV(Buffer, HairStrandsVF_PositionOffsetBuffer)
+		SHADER_PARAMETER_RDG_BUFFER_SRV(Buffer, HairStrandsVF_PositionBuffer)
+		SHADER_PARAMETER_RDG_BUFFER_SRV(Buffer, HairStrandsVF_PositionOffsetBuffer)
 		SHADER_PARAMETER_RDG_BUFFER_SRV(Buffer, HairStrandsVF_CullingIndirectBuffer)
 		SHADER_PARAMETER_RDG_BUFFER_SRV(Buffer, HairStrandsVF_CullingIndexBuffer)
 		SHADER_PARAMETER_RDG_BUFFER_SRV(Buffer, HairStrandsVF_CullingRadiusScaleBuffer)
@@ -1050,7 +1050,7 @@ static void AddVirtualVoxelizationComputeRasterPass(
 			if (HairGroupPublicData->DoesSupportVoxelization())
 			{
 				const FHairGroupPublicData::FVertexFactoryInput& VFInput = HairGroupPublicData->VFInput;
-				if (HairGroupPublicData->VFInput.Strands.PositionBuffer == nullptr)
+				if (HairGroupPublicData->VFInput.Strands.PositionBuffer.Buffer == nullptr)
 				{
 					continue;
 				}
@@ -1064,9 +1064,9 @@ static void AddVirtualVoxelizationComputeRasterPass(
 				PassParameters->OutPageTexture = PageTextureUAV;
 				PassParameters->FrameIdMod8 = FrameIdMode8;
 
-				PassParameters->HairStrandsVF_PositionBuffer = VFInput.Strands.PositionBuffer;
+				PassParameters->HairStrandsVF_PositionBuffer = VFInput.Strands.PositionBuffer.SRV;
 				PassParameters->HairStrandsVF_PositionOffset = VFInput.Strands.PositionOffset;
-				PassParameters->HairStrandsVF_PositionOffsetBuffer = VFInput.Strands.PositionOffsetBuffer;
+				PassParameters->HairStrandsVF_PositionOffsetBuffer = VFInput.Strands.PositionOffsetBuffer.SRV;
 				PassParameters->HairStrandsVF_VertexCount = VFInput.Strands.VertexCount;
 				PassParameters->HairStrandsVF_Radius = VFInput.Strands.HairRadius;
 				PassParameters->HairStrandsVF_Length = VFInput.Strands.HairLength;
diff --git a/Engine/Source/Runtime/Renderer/Public/HairStrandsInterface.h b/Engine/Source/Runtime/Renderer/Public/HairStrandsInterface.h
--- a/Engine/Source/Runtime/Renderer/Public/HairStrandsInterface.h
+++ b/Engine/Source/Runtime/Renderer/Public/HairStrandsInterface.h
@@ -31,9 +31,9 @@ ENUM_CLASS_FLAGS(ERDGImportedBufferFlags);
 
 struct RENDERER_API FRDGExternalBuffer
 {
-	TRefCountPtr<FRDGPooledBuffer> Buffer;
-	FShaderResourceViewRHIRef SRV;
-	FUnorderedAccessViewRHIRef UAV;
+	TRefCountPtr<FRDGPooledBuffer> Buffer = nullptr;
+	FShaderResourceViewRHIRef SRV = nullptr;
+	FUnorderedAccessViewRHIRef UAV = nullptr;
 	EPixelFormat Format = PF_Unknown;
 	void Release();
 };
@@ -214,18 +214,25 @@ public:
 	{
 		struct FStrands
 		{
-			FShaderResourceViewRHIRef PositionBuffer = nullptr;
-			FShaderResourceViewRHIRef PrevPositionBuffer = nullptr;
-			FShaderResourceViewRHIRef TangentBuffer = nullptr;
-			FShaderResourceViewRHIRef MaterialBuffer = nullptr;
-			FShaderResourceViewRHIRef AttributeBuffer = nullptr;
+			FRDGImportedBuffer PositionBuffer;
+			FRDGImportedBuffer PrevPositionBuffer;
+			FRDGImportedBuffer TangentBuffer;
+			FRDGImportedBuffer MaterialBuffer;
+			FRDGImportedBuffer AttributeBuffer;
+			FRDGImportedBuffer PositionOffsetBuffer;
+			FRDGImportedBuffer PrevPositionOffsetBuffer;
+
+			FShaderResourceViewRHIRef PositionBufferRHISRV				= nullptr;
+			FShaderResourceViewRHIRef PrevPositionBufferRHISRV			= nullptr;
+			FShaderResourceViewRHIRef TangentBufferRHISRV				= nullptr;
+			FShaderResourceViewRHIRef MaterialBufferRHISRV				= nullptr;
+			FShaderResourceViewRHIRef AttributeBufferRHISRV				= nullptr;
+			FShaderResourceViewRHIRef PositionOffsetBufferRHISRV		= nullptr;
+			FShaderResourceViewRHIRef PrevPositionOffsetBufferRHISRV	= nullptr;
 
 			FVector PositionOffset = FVector::ZeroVector;
 			FVector PrevPositionOffset = FVector::ZeroVector;
 
-			FShaderResourceViewRHIRef PositionOffsetBuffer = nullptr;
-			FShaderResourceViewRHIRef PrevPositionOffsetBuffer = nullptr;
-
 			uint32 VertexCount = 0;
 			float HairRadius = 0;
 			float HairLength = 0;
