Change hair strands resource naming to precise usage (guide, strands, cards deformer).

This will help to diagnose remaining validation errors happening with Niagara.

diff --git a/Engine/Plugins/Runtime/HairStrands/Source/HairStrandsCore/Private/GroomAsset.cpp b/Engine/Plugins/Runtime/HairStrands/Source/HairStrandsCore/Private/GroomAsset.cpp
--- a/Engine/Plugins/Runtime/HairStrands/Source/HairStrandsCore/Private/GroomAsset.cpp
+++ b/Engine/Plugins/Runtime/HairStrands/Source/HairStrandsCore/Private/GroomAsset.cpp
@@ -1976,7 +1976,7 @@ bool UGroomAsset::BuildCardsGeometry(uint32 GroupIndex)
 					GroupData.Guides.Data.RenderData = LOD.Guides.Data.RenderData;
 				}
 
-				LOD.Guides.RestResource = new FHairStrandsRestResource(LOD.Guides.Data.RenderData, LOD.Guides.Data.BoundingBox.GetCenter());
+				LOD.Guides.RestResource = new FHairStrandsRestResource(LOD.Guides.Data.RenderData, LOD.Guides.Data.BoundingBox.GetCenter(), EHairStrandsResourcesType::Cards);
 				BeginInitResource(LOD.Guides.RestResource);
 
 				LOD.Guides.InterpolationResource = new FHairStrandsInterpolationResource(LOD.Guides.InterpolationData.RenderData, GroupData.Guides.Data);
@@ -2190,7 +2190,7 @@ void UGroomAsset::AllocateGuidesResources(uint32 GroupIndex)
 		FHairGroupData& GroupData = HairGroupsData[GroupIndex];
 		if (GroupData.Guides.HasValidData())
 		{
-			GroupData.Guides.RestResource = new FHairStrandsRestResource(GroupData.Guides.Data.RenderData, GroupData.Guides.Data.BoundingBox.GetCenter());
+			GroupData.Guides.RestResource = new FHairStrandsRestResource(GroupData.Guides.Data.RenderData, GroupData.Guides.Data.BoundingBox.GetCenter(), EHairStrandsResourcesType::Guides);
 			BeginInitResource(GroupData.Guides.RestResource);
 		}
 	}
@@ -2227,7 +2227,7 @@ void UGroomAsset::InitStrandsResources()
 
 		if (GroupData.Strands.HasValidData())
 		{
-			GroupData.Strands.RestResource = new FHairStrandsRestResource(GroupData.Strands.Data.RenderData, GroupData.Strands.Data.BoundingBox.GetCenter());
+			GroupData.Strands.RestResource = new FHairStrandsRestResource(GroupData.Strands.Data.RenderData, GroupData.Strands.Data.BoundingBox.GetCenter(), EHairStrandsResourcesType::Strands);
 			BeginInitResource(GroupData.Strands.RestResource);
 
 			if (GroupData.Strands.ClusterCullingData.IsValid())
@@ -2329,7 +2329,7 @@ void UGroomAsset::InitCardsResources()
 				LOD.InterpolationResource = new FHairCardsInterpolationResource(LOD.InterpolationData.RenderData);
 				BeginInitResource(LOD.InterpolationResource);
 
-				LOD.Guides.RestResource = new FHairStrandsRestResource(LOD.Guides.Data.RenderData, LOD.Guides.Data.BoundingBox.GetCenter());
+				LOD.Guides.RestResource = new FHairStrandsRestResource(LOD.Guides.Data.RenderData, LOD.Guides.Data.BoundingBox.GetCenter(), EHairStrandsResourcesType::Cards);
 				BeginInitResource(LOD.Guides.RestResource);
 
 				LOD.Guides.InterpolationResource = new FHairStrandsInterpolationResource(LOD.Guides.InterpolationData.RenderData, GroupData.Guides.Data);
diff --git a/Engine/Plugins/Runtime/HairStrands/Source/HairStrandsCore/Private/GroomBindingAsset.cpp b/Engine/Plugins/Runtime/HairStrands/Source/HairStrandsCore/Private/GroomBindingAsset.cpp
--- a/Engine/Plugins/Runtime/HairStrands/Source/HairStrandsCore/Private/GroomBindingAsset.cpp
+++ b/Engine/Plugins/Runtime/HairStrands/Source/HairStrandsCore/Private/GroomBindingAsset.cpp
@@ -91,7 +91,7 @@ void UGroomBindingAsset::InitResource()
 		Resource.SimRootResources = nullptr;
 		if (Data.SimRootData.IsValid())
 		{
-			Resource.SimRootResources = new FHairStrandsRestRootResource(Data.SimRootData);
+			Resource.SimRootResources = new FHairStrandsRestRootResource(Data.SimRootData, EHairStrandsResourcesType::Guides);
 			BeginInitResource(Resource.SimRootResources);
 		}
 
@@ -99,7 +99,7 @@ void UGroomBindingAsset::InitResource()
 		Resource.RenRootResources = nullptr;
 		if (IsHairStrandsEnabled(EHairStrandsShaderType::Strands) && Data.RenRootData.IsValid())
 		{
-			Resource.RenRootResources = new FHairStrandsRestRootResource(Data.RenRootData);
+			Resource.RenRootResources = new FHairStrandsRestRootResource(Data.RenRootData, EHairStrandsResourcesType::Strands);
 			BeginInitResource(Resource.RenRootResources);
 		}
 
@@ -113,7 +113,7 @@ void UGroomBindingAsset::InitResource()
 				Resource.CardsRootResources[CardsLODIt] = nullptr;
 				if (Data.CardsRootData[CardsLODIt].IsValid())
 				{
-					Resource.CardsRootResources[CardsLODIt] = new FHairStrandsRestRootResource(Data.CardsRootData[CardsLODIt]);
+					Resource.CardsRootResources[CardsLODIt] = new FHairStrandsRestRootResource(Data.CardsRootData[CardsLODIt], EHairStrandsResourcesType::Cards);
 					BeginInitResource(Resource.CardsRootResources[CardsLODIt]);
 				}
 			}
diff --git a/Engine/Plugins/Runtime/HairStrands/Source/HairStrandsCore/Private/GroomComponent.cpp b/Engine/Plugins/Runtime/HairStrands/Source/HairStrandsCore/Private/GroomComponent.cpp
--- a/Engine/Plugins/Runtime/HairStrands/Source/HairStrandsCore/Private/GroomComponent.cpp
+++ b/Engine/Plugins/Runtime/HairStrands/Source/HairStrandsCore/Private/GroomComponent.cpp
@@ -2118,7 +2118,7 @@ void UGroomComponent::InitResources(bool bIsBindingReloading)
 
 				HairGroupInstance->Guides.RestRootResource = LocalBindingAsset->HairGroupResources[GroupIt].SimRootResources;
 
-				HairGroupInstance->Guides.DeformedRootResource = new FHairStrandsDeformedRootResource(HairGroupInstance->Guides.RestRootResource);
+				HairGroupInstance->Guides.DeformedRootResource = new FHairStrandsDeformedRootResource(HairGroupInstance->Guides.RestRootResource, EHairStrandsResourcesType::Guides);
 				BeginInitResource(HairGroupInstance->Guides.DeformedRootResource);
 			}
 
@@ -2132,7 +2132,7 @@ void UGroomComponent::InitResources(bool bIsBindingReloading)
 			HairGroupInstance->Guides.RestResource = GroupData.Guides.RestResource;
 
 			// If guides are allocated, deformed resources are always needs since they are either used with simulation, or RBF deformation. Both are dynamics, and require deformed positions
-			HairGroupInstance->Guides.DeformedResource = new FHairStrandsDeformedResource(GroupData.Guides.Data.RenderData, true, HairGroupInstance->Guides.RestResource->PositionOffset);
+			HairGroupInstance->Guides.DeformedResource = new FHairStrandsDeformedResource(GroupData.Guides.Data.RenderData, true, HairGroupInstance->Guides.RestResource->PositionOffset, EHairStrandsResourcesType::Guides);
 			BeginInitResource(HairGroupInstance->Guides.DeformedResource);
 
 			HairGroupInstance->Guides.DeformedResource->GetPositionOffset(FHairStrandsDeformedResource::EFrameType::Current)  = HairGroupInstance->Guides.RestResource->PositionOffset;
@@ -2214,14 +2214,14 @@ void UGroomComponent::InitResources(bool bIsBindingReloading)
 				}
 
 				HairGroupInstance->Strands.RestRootResource = LocalBindingAsset->HairGroupResources[GroupIt].RenRootResources;
-				HairGroupInstance->Strands.DeformedRootResource = new FHairStrandsDeformedRootResource(HairGroupInstance->Strands.RestRootResource);
+				HairGroupInstance->Strands.DeformedRootResource = new FHairStrandsDeformedRootResource(HairGroupInstance->Strands.RestRootResource, EHairStrandsResourcesType::Strands);
 				BeginInitResource(HairGroupInstance->Strands.DeformedRootResource);
 			}
 		
 			HairGroupInstance->Strands.RestResource = GroupData.Strands.RestResource;
 			if (bNeedDynamicResources)
 			{
-				HairGroupInstance->Strands.DeformedResource = new FHairStrandsDeformedResource(GroupData.Strands.Data.RenderData, false, HairGroupInstance->Strands.RestResource->PositionOffset);
+				HairGroupInstance->Strands.DeformedResource = new FHairStrandsDeformedResource(GroupData.Strands.Data.RenderData, false, HairGroupInstance->Strands.RestResource->PositionOffset, EHairStrandsResourcesType::Strands);
 			BeginInitResource(HairGroupInstance->Strands.DeformedResource);
 
 				// Initialize deformed position relative position offset to the rest pose offset
@@ -2322,13 +2322,13 @@ void UGroomComponent::InitResources(bool bIsBindingReloading)
 						}
 
 						InstanceLOD.Guides.RestRootResource = LocalBindingAsset->HairGroupResources[GroupIt].CardsRootResources[CardsLODIndex];
-						InstanceLOD.Guides.DeformedRootResource = new FHairStrandsDeformedRootResource(InstanceLOD.Guides.RestRootResource);
+						InstanceLOD.Guides.DeformedRootResource = new FHairStrandsDeformedRootResource(InstanceLOD.Guides.RestRootResource, EHairStrandsResourcesType::Cards);
 						BeginInitResource(InstanceLOD.Guides.DeformedRootResource);
 					}
 
 					InstanceLOD.Guides.RestResource = LOD.Guides.RestResource;
 					{
-						InstanceLOD.Guides.DeformedResource = new FHairStrandsDeformedResource(LOD.Guides.Data.RenderData, false, InstanceLOD.Guides.RestResource->PositionOffset);
+						InstanceLOD.Guides.DeformedResource = new FHairStrandsDeformedResource(LOD.Guides.Data.RenderData, false, InstanceLOD.Guides.RestResource->PositionOffset, EHairStrandsResourcesType::Cards);
 						BeginInitResource(InstanceLOD.Guides.DeformedResource);
 
 						// Initialize deformed position relative position offset to the rest pose offset
diff --git a/Engine/Plugins/Runtime/HairStrands/Source/HairStrandsCore/Private/GroomResources.cpp b/Engine/Plugins/Runtime/HairStrands/Source/HairStrandsCore/Private/GroomResources.cpp
--- a/Engine/Plugins/Runtime/HairStrands/Source/HairStrandsCore/Private/GroomResources.cpp
+++ b/Engine/Plugins/Runtime/HairStrands/Source/HairStrandsCore/Private/GroomResources.cpp
@@ -29,6 +29,8 @@ enum class EHairResourceUsageType : uint8
 	Dynamic
 };
 
+#define HAIRSTRANDS_RESOUCE_NAME(Type, Name) (Type == EHairStrandsResourcesType::Guides  ? TEXT("Name##(Guides)") : (Type == EHairStrandsResourcesType::Strands ? TEXT("Name##(Strands)") : TEXT("Name##(Cards)")))
+
 /////////////////////////////////////////////////////////////////////////////////////////
 // FRWBuffer utils 
 void UploadDataToBuffer(FReadBuffer& OutBuffer, uint32 DataSizeInBytes, const void* InCpuData)
@@ -541,8 +543,8 @@ void FHairMeshesDeformedResource::ReleaseRHI()
 
 /////////////////////////////////////////////////////////////////////////////////////////
 
-FHairStrandsRestResource::FHairStrandsRestResource(const FHairStrandsDatas::FRenderData& HairStrandRenderData, const FVector& InPositionOffset) :
-	PositionBuffer(), AttributeBuffer(), MaterialBuffer(), PositionOffset(InPositionOffset), RenderData(HairStrandRenderData)
+FHairStrandsRestResource::FHairStrandsRestResource(const FHairStrandsDatas::FRenderData& HairStrandRenderData, const FVector& InPositionOffset, EHairStrandsResourcesType InCurveType) :
+	PositionBuffer(), AttributeBuffer(), MaterialBuffer(), PositionOffset(InPositionOffset), RenderData(HairStrandRenderData), CurveType(InCurveType)
 {}
 
 void FHairStrandsRestResource::InitRHI()
@@ -552,13 +554,13 @@ void FHairStrandsRestResource::InitRHI()
 	FRHICommandListImmediate& RHICmdList = FRHICommandListExecutor::GetImmediateCommandList();
 	FRDGBuilder GraphBuilder(RHICmdList);
 
-	InternalCreateVertexBufferRDG<FHairStrandsPositionFormat>(GraphBuilder, RenderData.Positions, PositionBuffer, TEXT("Hair.StrandsRest_PositionBuffer"), EHairResourceUsageType::Static);
-	InternalCreateVertexBufferRDG<FHairStrandsAttributeFormat>(GraphBuilder, RenderData.Attributes, AttributeBuffer, TEXT("Hair.StrandsRest_AttributeBuffer"), EHairResourceUsageType::Static);
-	InternalCreateVertexBufferRDG<FHairStrandsMaterialFormat>(GraphBuilder, RenderData.Materials, MaterialBuffer, TEXT("Hair.StrandsRest_MaterialBuffer"), EHairResourceUsageType::Static);
+	InternalCreateVertexBufferRDG<FHairStrandsPositionFormat>(GraphBuilder, RenderData.Positions, PositionBuffer, HAIRSTRANDS_RESOUCE_NAME(CurveType, Hair.StrandsRest_PositionBuffer), EHairResourceUsageType::Static);
+	InternalCreateVertexBufferRDG<FHairStrandsAttributeFormat>(GraphBuilder, RenderData.Attributes, AttributeBuffer, HAIRSTRANDS_RESOUCE_NAME(CurveType, Hair.StrandsRest_AttributeBuffer), EHairResourceUsageType::Static);
+	InternalCreateVertexBufferRDG<FHairStrandsMaterialFormat>(GraphBuilder, RenderData.Materials, MaterialBuffer, HAIRSTRANDS_RESOUCE_NAME(CurveType, Hair.StrandsRest_MaterialBuffer), EHairResourceUsageType::Static);
 
 	TArray<FVector4> RestOffset;
 	RestOffset.Add(PositionOffset);
-	InternalCreateVertexBufferRDG<FHairStrandsPositionOffsetFormat>(GraphBuilder, RestOffset, PositionOffsetBuffer, TEXT("Hair.StrandsRest_PositionOffsetBuffer"), EHairResourceUsageType::Static);
+	InternalCreateVertexBufferRDG<FHairStrandsPositionOffsetFormat>(GraphBuilder, RestOffset, PositionOffsetBuffer, HAIRSTRANDS_RESOUCE_NAME(CurveType, Hair.StrandsRest_PositionOffsetBuffer), EHairResourceUsageType::Static);
 	
 	GraphBuilder.Execute();
 }
@@ -602,8 +604,8 @@ void FHairStrandsRestResource::ReleaseRHI()
 
 /////////////////////////////////////////////////////////////////////////////////////////
 
-FHairStrandsDeformedResource::FHairStrandsDeformedResource(const FHairStrandsDatas::FRenderData& HairStrandRenderData, bool bInInitializedData, const FVector& InDefaultOffset) :
-	RenderData(HairStrandRenderData), bInitializedData(bInInitializedData), DefaultOffset(InDefaultOffset)
+FHairStrandsDeformedResource::FHairStrandsDeformedResource(const FHairStrandsDatas::FRenderData& HairStrandRenderData, bool bInInitializedData, const FVector& InDefaultOffset, EHairStrandsResourcesType InCurveType) :
+	RenderData(HairStrandRenderData), bInitializedData(bInInitializedData), DefaultOffset(InDefaultOffset), CurveType(InCurveType)
 {}
 
 void FHairStrandsDeformedResource::InitRHI()
@@ -616,20 +618,20 @@ void FHairStrandsDeformedResource::InitRHI()
 	const uint32 VertexCount = RenderData.Positions.Num();
 	if (bInitializedData)
 	{
-		InternalCreateVertexBufferRDG<FHairStrandsPositionFormat>(GraphBuilder, RenderData.Positions, DeformedPositionBuffer[0], TEXT("Hair.StrandsDeformed_DeformedPositionBuffer0"), EHairResourceUsageType::Dynamic); // , bDynamic ? ERHIAccess::UAVCompute : ERHIAccess::SRVMask
-		InternalCreateVertexBufferRDG<FHairStrandsPositionFormat>(GraphBuilder, RenderData.Positions, DeformedPositionBuffer[1], TEXT("Hair.StrandsDeformed_DeformedPositionBuffer1"), EHairResourceUsageType::Dynamic); // , ERHIAccess::UAVCompute
+		InternalCreateVertexBufferRDG<FHairStrandsPositionFormat>(GraphBuilder, RenderData.Positions, DeformedPositionBuffer[0], HAIRSTRANDS_RESOUCE_NAME(CurveType, Hair.StrandsDeformed_DeformedPositionBuffer0), EHairResourceUsageType::Dynamic);
+		InternalCreateVertexBufferRDG<FHairStrandsPositionFormat>(GraphBuilder, RenderData.Positions, DeformedPositionBuffer[1], HAIRSTRANDS_RESOUCE_NAME(CurveType, Hair.StrandsDeformed_DeformedPositionBuffer1), EHairResourceUsageType::Dynamic);
 	}
 	else
 	{
-		InternalCreateVertexBufferRDG<FHairStrandsPositionFormat>(GraphBuilder, VertexCount, DeformedPositionBuffer[0], TEXT("Hair.StrandsDeformed_DeformedPositionBuffer0"), EHairResourceUsageType::Dynamic);
-		InternalCreateVertexBufferRDG<FHairStrandsPositionFormat>(GraphBuilder, VertexCount, DeformedPositionBuffer[1], TEXT("Hair.StrandsDeformed_DeformedPositionBuffer1"), EHairResourceUsageType::Dynamic);
+		InternalCreateVertexBufferRDG<FHairStrandsPositionFormat>(GraphBuilder, VertexCount, DeformedPositionBuffer[0], HAIRSTRANDS_RESOUCE_NAME(CurveType, Hair.StrandsDeformed_DeformedPositionBuffer0), EHairResourceUsageType::Dynamic);
+		InternalCreateVertexBufferRDG<FHairStrandsPositionFormat>(GraphBuilder, VertexCount, DeformedPositionBuffer[1], HAIRSTRANDS_RESOUCE_NAME(CurveType, Hair.StrandsDeformed_DeformedPositionBuffer1), EHairResourceUsageType::Dynamic);
 	}
-	InternalCreateVertexBufferRDG<FHairStrandsTangentFormat>(GraphBuilder, VertexCount * FHairStrandsTangentFormat::ComponentCount, TangentBuffer, TEXT("Hair.StrandsDeformed_TangentBuffer"), EHairResourceUsageType::Dynamic);
+	InternalCreateVertexBufferRDG<FHairStrandsTangentFormat>(GraphBuilder, VertexCount * FHairStrandsTangentFormat::ComponentCount, TangentBuffer, HAIRSTRANDS_RESOUCE_NAME(CurveType, Hair.StrandsDeformed_TangentBuffer), EHairResourceUsageType::Dynamic);
 
 	TArray<FVector4> DefaultOffsets;
 	DefaultOffsets.Add(DefaultOffset);
-	InternalCreateVertexBufferRDG<FHairStrandsPositionOffsetFormat>(GraphBuilder, DefaultOffsets, DeformedOffsetBuffer[0], TEXT("Hair.StrandsDeformed_DeformedOffsetBuffer0"), EHairResourceUsageType::Dynamic);
-	InternalCreateVertexBufferRDG<FHairStrandsPositionOffsetFormat>(GraphBuilder, DefaultOffsets, DeformedOffsetBuffer[1], TEXT("Hair.StrandsDeformed_DeformedOffsetBuffer1"), EHairResourceUsageType::Dynamic);
+	InternalCreateVertexBufferRDG<FHairStrandsPositionOffsetFormat>(GraphBuilder, DefaultOffsets, DeformedOffsetBuffer[0], HAIRSTRANDS_RESOUCE_NAME(CurveType, Hair.StrandsDeformed_DeformedOffsetBuffer0), EHairResourceUsageType::Dynamic);
+	InternalCreateVertexBufferRDG<FHairStrandsPositionOffsetFormat>(GraphBuilder, DefaultOffsets, DeformedOffsetBuffer[1], HAIRSTRANDS_RESOUCE_NAME(CurveType, Hair.StrandsDeformed_DeformedOffsetBuffer1), EHairResourceUsageType::Dynamic);
 
 	GraphBuilder.Execute();
 }
@@ -705,14 +707,14 @@ void FHairStrandsClusterCullingResource::ReleaseRHI()
 
 /////////////////////////////////////////////////////////////////////////////////////////
 
-FHairStrandsRestRootResource::FHairStrandsRestRootResource(const FHairStrandsRootData& InRootData):
-RootData(InRootData)
+FHairStrandsRestRootResource::FHairStrandsRestRootResource(const FHairStrandsRootData& InRootData, EHairStrandsResourcesType InCurveType) :
+RootData(InRootData), CurveType(InCurveType)
 {
 	PopulateFromRootData();
 }
 
-FHairStrandsRestRootResource::FHairStrandsRestRootResource(const FHairStrandsDatas* HairStrandsDatas, uint32 LODCount, const TArray<uint32>& NumSamples):
-	RootData(HairStrandsDatas, LODCount, NumSamples)
+FHairStrandsRestRootResource::FHairStrandsRestRootResource(const FHairStrandsDatas* HairStrandsDatas, uint32 LODCount, const TArray<uint32>& NumSamples, EHairStrandsResourcesType InCurveType):
+	RootData(HairStrandsDatas, LODCount, NumSamples), CurveType(InCurveType)
 {
 	PopulateFromRootData();
 }
@@ -739,9 +741,9 @@ void FHairStrandsRestRootResource::InitRHI()
 		FRHICommandListImmediate& RHICmdList = FRHICommandListExecutor::GetImmediateCommandList();
 		FRDGBuilder GraphBuilder(RHICmdList);
 
-		InternalCreateVertexBufferRDG<FHairStrandsIndexFormat>(GraphBuilder, RootData.VertexToCurveIndexBuffer, VertexToCurveIndexBuffer, TEXT("Hair.StrandsRestRoot_VertexToCurveIndexBuffer"), EHairResourceUsageType::Static);
-		InternalCreateVertexBufferRDG<FHairStrandsRootPositionFormat>(GraphBuilder, RootData.RootPositionBuffer, RootPositionBuffer, TEXT("Hair.StrandsRestRoot_RootPositionBuffer"), EHairResourceUsageType::Static);
-		InternalCreateVertexBufferRDG<FHairStrandsRootNormalFormat>(GraphBuilder, RootData.RootNormalBuffer, RootNormalBuffer, TEXT("Hair.StrandsRestRoot_RootNormalBuffer"), EHairResourceUsageType::Static);
+		InternalCreateVertexBufferRDG<FHairStrandsIndexFormat>(GraphBuilder, RootData.VertexToCurveIndexBuffer, VertexToCurveIndexBuffer, HAIRSTRANDS_RESOUCE_NAME(CurveType, Hair.StrandsRestRoot_VertexToCurveIndexBuffer), EHairResourceUsageType::Static);
+		InternalCreateVertexBufferRDG<FHairStrandsRootPositionFormat>(GraphBuilder, RootData.RootPositionBuffer, RootPositionBuffer, HAIRSTRANDS_RESOUCE_NAME(CurveType, Hair.StrandsRestRoot_RootPositionBuffer), EHairResourceUsageType::Static);
+		InternalCreateVertexBufferRDG<FHairStrandsRootNormalFormat>(GraphBuilder, RootData.RootNormalBuffer, RootNormalBuffer, HAIRSTRANDS_RESOUCE_NAME(CurveType, Hair.StrandsRestRoot_RootNormalBuffer), EHairResourceUsageType::Static);
 		
 		check(LODs.Num() == RootData.MeshProjectionLODs.Num());
 		for (uint32 LODIt=0, LODCount = LODs.Num(); LODIt<LODCount; ++LODIt)
@@ -755,29 +757,29 @@ void FHairStrandsRestRootResource::InitRHI()
 				GPUData.Status = FLOD::EStatus::Completed;
 
 				check(CPUData.RootTriangleBarycentricBuffer.Num() > 0);
-				InternalCreateVertexBufferRDG<FHairStrandsCurveTriangleBarycentricFormat>(GraphBuilder, CPUData.RootTriangleBarycentricBuffer, GPUData.RootTriangleBarycentricBuffer, TEXT("Hair.StrandsRestRoot_RootTriangleBarycentricBuffer"), EHairResourceUsageType::Static);
+				InternalCreateVertexBufferRDG<FHairStrandsCurveTriangleBarycentricFormat>(GraphBuilder, CPUData.RootTriangleBarycentricBuffer, GPUData.RootTriangleBarycentricBuffer, HAIRSTRANDS_RESOUCE_NAME(CurveType, Hair.StrandsRestRoot_RootTriangleBarycentricBuffer), EHairResourceUsageType::Static);
 
 				check(CPUData.RootTriangleIndexBuffer.Num() > 0);
-				InternalCreateVertexBufferRDG<FHairStrandsCurveTriangleIndexFormat>(GraphBuilder, CPUData.RootTriangleIndexBuffer, GPUData.RootTriangleIndexBuffer, TEXT("Hair.StrandsRestRoot_RootTriangleIndexBuffer"), EHairResourceUsageType::Static);
+				InternalCreateVertexBufferRDG<FHairStrandsCurveTriangleIndexFormat>(GraphBuilder, CPUData.RootTriangleIndexBuffer, GPUData.RootTriangleIndexBuffer, HAIRSTRANDS_RESOUCE_NAME(CurveType, Hair.StrandsRestRoot_RootTriangleIndexBuffer), EHairResourceUsageType::Static);
 
 				check(CPUData.RestRootTrianglePosition0Buffer.Num() > 0);
 				check(CPUData.RestRootTrianglePosition1Buffer.Num() > 0);
 				check(CPUData.RestRootTrianglePosition2Buffer.Num() > 0);
-				InternalCreateVertexBufferRDG<FHairStrandsMeshTrianglePositionFormat>(GraphBuilder, CPUData.RestRootTrianglePosition0Buffer, GPUData.RestRootTrianglePosition0Buffer, TEXT("Hair.StrandsRestRoot_RestRootTrianglePosition0Buffer"), EHairResourceUsageType::Static);
-				InternalCreateVertexBufferRDG<FHairStrandsMeshTrianglePositionFormat>(GraphBuilder, CPUData.RestRootTrianglePosition1Buffer, GPUData.RestRootTrianglePosition1Buffer, TEXT("Hair.StrandsRestRoot_RestRootTrianglePosition1Buffer"), EHairResourceUsageType::Static);
-				InternalCreateVertexBufferRDG<FHairStrandsMeshTrianglePositionFormat>(GraphBuilder, CPUData.RestRootTrianglePosition2Buffer, GPUData.RestRootTrianglePosition2Buffer, TEXT("Hair.StrandsRestRoot_RestRootTrianglePosition2Buffer"), EHairResourceUsageType::Static);
+				InternalCreateVertexBufferRDG<FHairStrandsMeshTrianglePositionFormat>(GraphBuilder, CPUData.RestRootTrianglePosition0Buffer, GPUData.RestRootTrianglePosition0Buffer, HAIRSTRANDS_RESOUCE_NAME(CurveType, Hair.StrandsRestRoot_RestRootTrianglePosition0Buffer), EHairResourceUsageType::Static);
+				InternalCreateVertexBufferRDG<FHairStrandsMeshTrianglePositionFormat>(GraphBuilder, CPUData.RestRootTrianglePosition1Buffer, GPUData.RestRootTrianglePosition1Buffer, HAIRSTRANDS_RESOUCE_NAME(CurveType, Hair.StrandsRestRoot_RestRootTrianglePosition1Buffer), EHairResourceUsageType::Static);
+				InternalCreateVertexBufferRDG<FHairStrandsMeshTrianglePositionFormat>(GraphBuilder, CPUData.RestRootTrianglePosition2Buffer, GPUData.RestRootTrianglePosition2Buffer, HAIRSTRANDS_RESOUCE_NAME(CurveType, Hair.StrandsRestRoot_RestRootTrianglePosition2Buffer), EHairResourceUsageType::Static);
 			}
 			else
 			{
 				GPUData.Status = FLOD::EStatus::Initialized;
 			
-				InternalCreateVertexBufferRDG<FHairStrandsCurveTriangleBarycentricFormat>(GraphBuilder, RootData.RootCount, GPUData.RootTriangleBarycentricBuffer, TEXT("Hair.StrandsRestRoot_RootTriangleBarycentricBuffer"), EHairResourceUsageType::Dynamic);
-				InternalCreateVertexBufferRDG<FHairStrandsCurveTriangleIndexFormat>(GraphBuilder, RootData.RootCount, GPUData.RootTriangleIndexBuffer, TEXT("Hair.StrandsRestRoot_RootTriangleIndexBuffer"), EHairResourceUsageType::Dynamic);
+				InternalCreateVertexBufferRDG<FHairStrandsCurveTriangleBarycentricFormat>(GraphBuilder, RootData.RootCount, GPUData.RootTriangleBarycentricBuffer, HAIRSTRANDS_RESOUCE_NAME(CurveType, Hair.StrandsRestRoot_RootTriangleBarycentricBuffer), EHairResourceUsageType::Dynamic);
+				InternalCreateVertexBufferRDG<FHairStrandsCurveTriangleIndexFormat>(GraphBuilder, RootData.RootCount, GPUData.RootTriangleIndexBuffer, HAIRSTRANDS_RESOUCE_NAME(CurveType, Hair.StrandsRestRoot_RootTriangleIndexBuffer), EHairResourceUsageType::Dynamic);
 			
 				// Create buffers. Initialization will be done by render passes
-				InternalCreateVertexBufferRDG<FHairStrandsMeshTrianglePositionFormat>(GraphBuilder, RootData.RootCount, GPUData.RestRootTrianglePosition0Buffer, TEXT("Hair.StrandsRestRoot_RestRootTrianglePosition0Buffer"), EHairResourceUsageType::Dynamic);
-				InternalCreateVertexBufferRDG<FHairStrandsMeshTrianglePositionFormat>(GraphBuilder, RootData.RootCount, GPUData.RestRootTrianglePosition1Buffer, TEXT("Hair.StrandsRestRoot_RestRootTrianglePosition1Buffer"), EHairResourceUsageType::Dynamic);
-				InternalCreateVertexBufferRDG<FHairStrandsMeshTrianglePositionFormat>(GraphBuilder, RootData.RootCount, GPUData.RestRootTrianglePosition2Buffer, TEXT("Hair.StrandsRestRoot_RestRootTrianglePosition2Buffer"), EHairResourceUsageType::Dynamic);
+				InternalCreateVertexBufferRDG<FHairStrandsMeshTrianglePositionFormat>(GraphBuilder, RootData.RootCount, GPUData.RestRootTrianglePosition0Buffer, HAIRSTRANDS_RESOUCE_NAME(CurveType, Hair.StrandsRestRoot_RestRootTrianglePosition0Buffer), EHairResourceUsageType::Dynamic);
+				InternalCreateVertexBufferRDG<FHairStrandsMeshTrianglePositionFormat>(GraphBuilder, RootData.RootCount, GPUData.RestRootTrianglePosition1Buffer, HAIRSTRANDS_RESOUCE_NAME(CurveType, Hair.StrandsRestRoot_RestRootTrianglePosition1Buffer), EHairResourceUsageType::Dynamic);
+				InternalCreateVertexBufferRDG<FHairStrandsMeshTrianglePositionFormat>(GraphBuilder, RootData.RootCount, GPUData.RestRootTrianglePosition2Buffer, HAIRSTRANDS_RESOUCE_NAME(CurveType, Hair.StrandsRestRoot_RestRootTrianglePosition2Buffer), EHairResourceUsageType::Dynamic);
 			}
 
 			GPUData.SampleCount = CPUData.SampleCount;
@@ -788,15 +790,15 @@ void FHairStrandsRestRootResource::InitRHI()
 				check(CPUData.MeshSampleIndicesBuffer.Num() == CPUData.SampleCount);
 				check(CPUData.RestSamplePositionsBuffer.Num() == CPUData.SampleCount);
 
-				InternalCreateVertexBufferRDG<FHairStrandsWeightFormat>(GraphBuilder, CPUData.MeshInterpolationWeightsBuffer, GPUData.MeshInterpolationWeightsBuffer, TEXT("Hair.StrandsRestRoot_MeshInterpolationWeightsBuffer"), EHairResourceUsageType::Static);
-				InternalCreateVertexBufferRDG<FHairStrandsIndexFormat>(GraphBuilder, CPUData.MeshSampleIndicesBuffer, GPUData.MeshSampleIndicesBuffer, TEXT("Hair.StrandsRestRoot_MeshSampleIndicesBuffer"), EHairResourceUsageType::Static);
-				InternalCreateVertexBufferRDG<FHairStrandsMeshTrianglePositionFormat>(GraphBuilder, CPUData.RestSamplePositionsBuffer, GPUData.RestSamplePositionsBuffer, TEXT("Hair.StrandsRestRoot_RestSamplePositionsBuffer"), EHairResourceUsageType::Static);
+				InternalCreateVertexBufferRDG<FHairStrandsWeightFormat>(GraphBuilder, CPUData.MeshInterpolationWeightsBuffer, GPUData.MeshInterpolationWeightsBuffer, HAIRSTRANDS_RESOUCE_NAME(CurveType, Hair.StrandsRestRoot_MeshInterpolationWeightsBuffer), EHairResourceUsageType::Static);
+				InternalCreateVertexBufferRDG<FHairStrandsIndexFormat>(GraphBuilder, CPUData.MeshSampleIndicesBuffer, GPUData.MeshSampleIndicesBuffer, HAIRSTRANDS_RESOUCE_NAME(CurveType, Hair.StrandsRestRoot_MeshSampleIndicesBuffer), EHairResourceUsageType::Static);
+				InternalCreateVertexBufferRDG<FHairStrandsMeshTrianglePositionFormat>(GraphBuilder, CPUData.RestSamplePositionsBuffer, GPUData.RestSamplePositionsBuffer, HAIRSTRANDS_RESOUCE_NAME(CurveType, Hair.StrandsRestRoot_RestSamplePositionsBuffer), EHairResourceUsageType::Static);
 			}
 			else
 			{
-				InternalCreateVertexBufferRDG<FHairStrandsWeightFormat>(GraphBuilder, (CPUData.SampleCount+4) * (CPUData.SampleCount+4), GPUData.MeshInterpolationWeightsBuffer, TEXT("Hair.StrandsRestRoot_MeshInterpolationWeightsBuffer"), EHairResourceUsageType::Dynamic);
-				InternalCreateVertexBufferRDG<FHairStrandsIndexFormat>(GraphBuilder, CPUData.SampleCount, GPUData.MeshSampleIndicesBuffer, TEXT("Hair.StrandsRestRoot_MeshSampleIndicesBuffer"), EHairResourceUsageType::Dynamic);
-				InternalCreateVertexBufferRDG<FHairStrandsMeshTrianglePositionFormat>(GraphBuilder, CPUData.SampleCount, GPUData.RestSamplePositionsBuffer, TEXT("Hair.StrandsRestRoot_RestSamplePositionsBuffer"), EHairResourceUsageType::Dynamic);
+				InternalCreateVertexBufferRDG<FHairStrandsWeightFormat>(GraphBuilder, (CPUData.SampleCount+4) * (CPUData.SampleCount+4), GPUData.MeshInterpolationWeightsBuffer, HAIRSTRANDS_RESOUCE_NAME(CurveType, Hair.StrandsRestRoot_MeshInterpolationWeightsBuffer), EHairResourceUsageType::Dynamic);
+				InternalCreateVertexBufferRDG<FHairStrandsIndexFormat>(GraphBuilder, CPUData.SampleCount, GPUData.MeshSampleIndicesBuffer, HAIRSTRANDS_RESOUCE_NAME(CurveType, Hair.StrandsRestRoot_MeshSampleIndicesBuffer), EHairResourceUsageType::Dynamic);
+				InternalCreateVertexBufferRDG<FHairStrandsMeshTrianglePositionFormat>(GraphBuilder, CPUData.SampleCount, GPUData.RestSamplePositionsBuffer, HAIRSTRANDS_RESOUCE_NAME(CurveType, Hair.StrandsRestRoot_RestSamplePositionsBuffer), EHairResourceUsageType::Dynamic);
 			}
 		}
 		GraphBuilder.Execute();
@@ -831,12 +833,13 @@ void FHairStrandsRestRootResource::ReleaseRHI()
 
 /////////////////////////////////////////////////////////////////////////////////////////
 
-FHairStrandsDeformedRootResource::FHairStrandsDeformedRootResource()
+FHairStrandsDeformedRootResource::FHairStrandsDeformedRootResource(EHairStrandsResourcesType InCurveType):CurveType(InCurveType)
 {
 
 }
 
-FHairStrandsDeformedRootResource::FHairStrandsDeformedRootResource(const FHairStrandsRestRootResource* InRestResources)
+FHairStrandsDeformedRootResource::FHairStrandsDeformedRootResource(const FHairStrandsRestRootResource* InRestResources, EHairStrandsResourcesType InCurveType):
+CurveType(InCurveType)
 {
 	check(InRestResources);
 	uint32 LODIndex = 0;
@@ -863,12 +866,12 @@ void FHairStrandsDeformedRootResource::InitRHI()
 		for (FLOD& LOD : LODs)
 		{		
 			LOD.Status = FLOD::EStatus::Initialized;
-			InternalCreateVertexBufferRDG<FHairStrandsMeshTrianglePositionFormat>(GraphBuilder, LOD.SampleCount, LOD.DeformedSamplePositionsBuffer, TEXT("Hair.StrandsRootDeformed_DeformedSamplePositionsBuffer"), EHairResourceUsageType::Dynamic);
-			InternalCreateVertexBufferRDG<FHairStrandsMeshTrianglePositionFormat>(GraphBuilder, LOD.SampleCount + 4, LOD.MeshSampleWeightsBuffer, TEXT("Hair.StrandsRootDeformed_MeshSampleWeightsBuffer"), EHairResourceUsageType::Dynamic);
+			InternalCreateVertexBufferRDG<FHairStrandsMeshTrianglePositionFormat>(GraphBuilder, LOD.SampleCount, LOD.DeformedSamplePositionsBuffer, HAIRSTRANDS_RESOUCE_NAME(CurveType, Hair.StrandsRootDeformed_DeformedSamplePositionsBuffer), EHairResourceUsageType::Dynamic);
+			InternalCreateVertexBufferRDG<FHairStrandsMeshTrianglePositionFormat>(GraphBuilder, LOD.SampleCount + 4, LOD.MeshSampleWeightsBuffer, HAIRSTRANDS_RESOUCE_NAME(CurveType, Hair.StrandsRootDeformed_MeshSampleWeightsBuffer), EHairResourceUsageType::Dynamic);
 
-			InternalCreateVertexBufferRDG<FHairStrandsMeshTrianglePositionFormat>(GraphBuilder, RootCount, LOD.DeformedRootTrianglePosition0Buffer, TEXT("Hair.StrandsRootDeformed_DeformedRootTrianglePosition0Buffer"), EHairResourceUsageType::Dynamic);
-			InternalCreateVertexBufferRDG<FHairStrandsMeshTrianglePositionFormat>(GraphBuilder, RootCount, LOD.DeformedRootTrianglePosition1Buffer, TEXT("Hair.StrandsRootDeformed_DeformedRootTrianglePosition1Buffer"), EHairResourceUsageType::Dynamic);
-			InternalCreateVertexBufferRDG<FHairStrandsMeshTrianglePositionFormat>(GraphBuilder, RootCount, LOD.DeformedRootTrianglePosition2Buffer, TEXT("Hair.StrandsRootDeformed_DeformedRootTrianglePosition2Buffer"), EHairResourceUsageType::Dynamic);
+			InternalCreateVertexBufferRDG<FHairStrandsMeshTrianglePositionFormat>(GraphBuilder, RootCount, LOD.DeformedRootTrianglePosition0Buffer, HAIRSTRANDS_RESOUCE_NAME(CurveType, Hair.StrandsRootDeformed_DeformedRootTrianglePosition0Buffer), EHairResourceUsageType::Dynamic);
+			InternalCreateVertexBufferRDG<FHairStrandsMeshTrianglePositionFormat>(GraphBuilder, RootCount, LOD.DeformedRootTrianglePosition1Buffer, HAIRSTRANDS_RESOUCE_NAME(CurveType, Hair.StrandsRootDeformed_DeformedRootTrianglePosition1Buffer), EHairResourceUsageType::Dynamic);
+			InternalCreateVertexBufferRDG<FHairStrandsMeshTrianglePositionFormat>(GraphBuilder, RootCount, LOD.DeformedRootTrianglePosition2Buffer, HAIRSTRANDS_RESOUCE_NAME(CurveType, Hair.StrandsRootDeformed_DeformedRootTrianglePosition2Buffer), EHairResourceUsageType::Dynamic);
 		}
 	}
 	GraphBuilder.Execute();
diff --git a/Engine/Plugins/Runtime/HairStrands/Source/HairStrandsCore/Public/GroomResources.h b/Engine/Plugins/Runtime/HairStrands/Source/HairStrandsCore/Public/GroomResources.h
--- a/Engine/Plugins/Runtime/HairStrands/Source/HairStrandsCore/Public/GroomResources.h
+++ b/Engine/Plugins/Runtime/HairStrands/Source/HairStrandsCore/Public/GroomResources.h
@@ -16,12 +16,19 @@ inline uint32 GetBufferTotalNumBytes(const FRDGExternalBuffer& In)
 	return In.Buffer ? In.Buffer->Desc.GetTotalNumBytes() : 0;
 }
 
+enum class EHairStrandsResourcesType : uint8
+{
+	Guides,		// Guides used for simulation
+	Strands,	// Rendering strands 
+	Cards		// Guides used for deforming the cards geometry (which is different from the simulation guides)
+};
+
 /* Render buffers for root deformation for dynamic meshes */
 struct FHairStrandsRestRootResource : public FRenderResource
 {
 	/** Build the hair strands resource */
-	FHairStrandsRestRootResource(const FHairStrandsRootData& RootData);
-	FHairStrandsRestRootResource(const FHairStrandsDatas* HairStrandsDatas, uint32 LODCount, const TArray<uint32>& NumSamples);
+	FHairStrandsRestRootResource(const FHairStrandsRootData& RootData, EHairStrandsResourcesType CurveType);
+	FHairStrandsRestRootResource(const FHairStrandsDatas* HairStrandsDatas, uint32 LODCount, const TArray<uint32>& NumSamples, EHairStrandsResourcesType CurveType);
 
 	/* Init the buffer */
 	virtual void InitRHI() override;
@@ -91,14 +98,17 @@ struct FHairStrandsRestRootResource : public FRenderResource
 
 	/* Store CPU data for root info & root binding */
 	FHairStrandsRootData RootData;
+
+	/* Type of curves */
+	const EHairStrandsResourcesType CurveType;
 };
 
 /* Render buffers for root deformation for dynamic meshes */
 struct FHairStrandsDeformedRootResource : public FRenderResource
 {
 	/** Build the hair strands resource */
-	FHairStrandsDeformedRootResource();
-	FHairStrandsDeformedRootResource(const FHairStrandsRestRootResource* InRestResources);
+	FHairStrandsDeformedRootResource(EHairStrandsResourcesType CurveType);
+	FHairStrandsDeformedRootResource(const FHairStrandsRestRootResource* InRestResources, EHairStrandsResourcesType CurveType);
 
 	/* Init the buffer */
 	virtual void InitRHI() override;
@@ -156,13 +166,16 @@ struct FHairStrandsDeformedRootResource : public FRenderResource
 
 	/* Last update MeshLODIndex */
 	int32 MeshLODIndex = -1;
+
+	/* Type of curves */
+	const EHairStrandsResourcesType CurveType;
 };
 
 /* Render buffers that will be used for rendering */
 struct FHairStrandsRestResource : public FRenderResource
 {
 	/** Build the hair strands resource */
-	FHairStrandsRestResource(const FHairStrandsDatas::FRenderData& HairStrandRenderData, const FVector& PositionOffset);
+	FHairStrandsRestResource(const FHairStrandsDatas::FRenderData& HairStrandRenderData, const FVector& PositionOffset, EHairStrandsResourcesType CurveType);
 
 	/* Init the buffer */
 	virtual void InitRHI() override;
@@ -208,13 +221,16 @@ struct FHairStrandsRestResource : public FRenderResource
 	/* Reference to the hair strands render data */
 	const FHairStrandsDatas::FRenderData& RenderData;
 
+	/* Type of curves */
+	const EHairStrandsResourcesType CurveType;
+
 	inline uint32 GetVertexCount() const { return RenderData.Positions.Num() / FHairStrandsPositionFormat::ComponentCount; }
 };
 
 struct FHairStrandsDeformedResource : public FRenderResource
 {
 	/** Build the hair strands resource */
-	FHairStrandsDeformedResource(const FHairStrandsDatas::FRenderData& HairStrandRenderData, bool bInitializeData, const FVector& InDefaultOffset);
+	FHairStrandsDeformedResource(const FHairStrandsDatas::FRenderData& HairStrandRenderData, bool bInitializeData, const FVector& InDefaultOffset, EHairStrandsResourcesType CurveType);
 
 	/* Init the buffer */
 	virtual void InitRHI() override;
@@ -259,6 +275,9 @@ struct FHairStrandsDeformedResource : public FRenderResource
 	/* Whether the underlying resource is dynamic not (single or double buffer allocated) */
 	FVector DefaultOffset = FVector::ZeroVector;
 
+	/* Type of curves */
+	const EHairStrandsResourcesType CurveType;
+
 	enum EFrameType
 	{
 		Previous,
