Add variable to suppress groom binding builder warning.

diff --git a/Engine/Plugins/Runtime/HairStrands/Source/HairStrandsCore/Private/GroomBindingBuilder.cpp b/Engine/Plugins/Runtime/HairStrands/Source/HairStrandsCore/Private/GroomBindingBuilder.cpp
--- a/Engine/Plugins/Runtime/HairStrands/Source/HairStrandsCore/Private/GroomBindingBuilder.cpp
+++ b/Engine/Plugins/Runtime/HairStrands/Source/HairStrandsCore/Private/GroomBindingBuilder.cpp
@@ -49,6 +49,9 @@ DEFINE_LOG_CATEGORY_STATIC(LogGroomBindingBuilder, Log, All);
 
 #define LOCTEXT_NAMESPACE "GroomBindingBuilder"
 
+static int32 GHairStrandsBindingBuilderWarningEnable = 1;
+static FAutoConsoleVariableRef CVarHairStrandsBindingBuilderWarningEnable(TEXT("r.HairStrands.Log.BindingBuilderWarning"), GHairStrandsBindingBuilderWarningEnable, TEXT("Enable/disable warning during groom binding builder"));
+
 ///////////////////////////////////////////////////////////////////////////////////////////////////
 
 FString FGroomBindingBuilder::GetVersion()
@@ -1650,7 +1653,7 @@ namespace GroomBinding_Transfer
 		const IMeshLODData& SourceMeshLODData = InSourceMeshData->GetMeshLODData(SourceLODIndex);
 		const bool bIsMatchingSectionValid = MatchingSection < SourceMeshLODData.GetNumSections();
 		const int32 SourceSectionId = bIsMatchingSectionValid ? MatchingSection : 0;
-		if (!bIsMatchingSectionValid)
+		if (!bIsMatchingSectionValid && GHairStrandsBindingBuilderWarningEnable > 0)
 		{
 			UE_LOG(LogHairStrands, Warning, TEXT("[Groom] Binding asset will not respect the requested 'Matching section' %d. The source skeletal mesh does not have such a section. Instead 'Matching Section' 0 will be used."), MatchingSection);
 		}
@@ -1678,7 +1681,10 @@ namespace GroomBinding_Transfer
 			const IMeshLODData& TargetMeshLODData = InTargetMeshData->GetMeshLODData(TargetLODIndex);
 			if (LocalTargetSectionId >= TargetMeshLODData.GetNumSections())
 			{
-				UE_LOG(LogHairStrands, Warning, TEXT("[Groom] Binding asset will not respect the requested 'Matching section' %d for LOD %d. The target skeletal mesh does not have such a section for this LOD. Instead section 0 will be used for this given LOD."), TargetSectionId, TargetLODIndex);
+				if (GHairStrandsBindingBuilderWarningEnable > 0)
+				{
+					UE_LOG(LogHairStrands, Warning, TEXT("[Groom] Binding asset will not respect the requested 'Matching section' %d for LOD %d. The target skeletal mesh does not have such a section for this LOD. Instead section 0 will be used for this given LOD."), TargetSectionId, TargetLODIndex);
+				}
 
 				LocalTargetSectionId = 0;
 				LocalSourceSectionId = 0;
@@ -2195,7 +2201,10 @@ static void InternalBuildBinding_GPU(FRDGBuilder& GraphBuilder, UGroomBindingAss
 		!BindingAsset->TargetSkeletalMesh ||
 		BindingAsset->Groom->GetNumHairGroups() == 0)
 	{
-		UE_LOG(LogHairStrands, Warning, TEXT("[Groom] Error - Binding asset can be created/rebuilt."));
+		if (GHairStrandsBindingBuilderWarningEnable > 0)
+		{
+			UE_LOG(LogHairStrands, Warning, TEXT("[Groom] Error - Binding asset can be created/rebuilt."));
+		}
 		return;
 	}
 
