Fix page fault crash in hair strands system due to invalid culling index.

diff --git a/Engine/Plugins/Runtime/HairStrands/Source/HairStrandsCore/Private/HairStrandsRendering.cpp b/Engine/Plugins/Runtime/HairStrands/Source/HairStrandsCore/Private/HairStrandsRendering.cpp
--- a/Engine/Plugins/Runtime/HairStrands/Source/HairStrandsCore/Private/HairStrandsRendering.cpp
+++ b/Engine/Plugins/Runtime/HairStrands/Source/HairStrandsCore/Private/HairStrandsRendering.cpp
@@ -726,6 +726,7 @@ class FHairClusterAABBCS : public FGlobalShader
 		SHADER_PARAMETER_RDG_BUFFER_SRV(Buffer, ClusterIdBuffer)
 		SHADER_PARAMETER_RDG_BUFFER_SRV(Buffer, ClusterIndexOffsetBuffer)
 		SHADER_PARAMETER_RDG_BUFFER_SRV(Buffer, ClusterIndexCountBuffer)
+		SHADER_PARAMETER_RDG_BUFFER_SRV(Buffer<uint>, HairStrandsVF_CullingIndirectBuffer)
 		SHADER_PARAMETER_RDG_BUFFER_UAV(RWBuffer, OutClusterAABBBuffer)
 		SHADER_PARAMETER_RDG_BUFFER_UAV(RWBuffer, OutGroupAABBBuffer)
 	END_SHADER_PARAMETER_STRUCT()
@@ -750,7 +751,8 @@ static void AddHairClusterAABBPass(
 
 	FRDGHairStrandsCullingData& ClusterAABBData,
 
-	FRDGBufferSRVRef RenderPositionBufferSRV)
+	FRDGBufferSRVRef RenderPositionBufferSRV,
+	FRDGBufferSRVRef& DrawIndirectRasterComputeBuffer)
 {
 	const uint32 GroupSize = ComputeGroupSize();
 	const FIntVector DispatchCount = ComputeDispatchGroupCount2D(ClusterData.ClusterCount);
@@ -768,9 +770,16 @@ static void AddHairClusterAABBPass(
 	Parameters->ClusterIdBuffer = GraphBuilder.CreateSRV(ClusterIdBuffer, PF_R32_UINT);
 	Parameters->ClusterIndexOffsetBuffer = GraphBuilder.CreateSRV(ClusterIndexOffsetBuffer, PF_R32_UINT);
 	Parameters->ClusterIndexCountBuffer = GraphBuilder.CreateSRV(ClusterIndexCountBuffer, PF_R32_UINT);
+	Parameters->HairStrandsVF_CullingIndirectBuffer = DrawIndirectRasterComputeBuffer; // Used for checking max vertex count
 	Parameters->OutClusterAABBBuffer = ClusterAABBData.ClusterAABBBuffer.UAV;
 	Parameters->OutGroupAABBBuffer = ClusterAABBData.GroupAABBBuffer.UAV;
 
+	// Sanity check
+	check(ClusterData.ClusterCount == ClusterIdBuffer->Desc.NumElements);
+	check(ClusterData.ClusterCount == ClusterIndexOffsetBuffer->Desc.NumElements);
+	check(ClusterData.ClusterCount == ClusterIndexCountBuffer->Desc.NumElements);
+	check(ClusterData.ClusterCount * 6 == ClusterAABBData.ClusterAABBBuffer.Buffer->Desc.NumElements);
+
 	FHairClusterAABBCS::FPermutationDomain PermutationVector;
 	PermutationVector.Set<FHairClusterAABBCS::FGroupSize>(GroupSize);
 	TShaderMapRef<FHairClusterAABBCS> ComputeShader(ShaderMap, PermutationVector);
@@ -1560,6 +1569,7 @@ void ComputeHairStrandsInterpolation(
 
 				if (bNeedAABB)
 				{
+					FRDGImportedBuffer Strands_CulledVertexCount = Register(GraphBuilder, Instance->HairGroupPublicData->GetDrawIndirectRasterComputeBuffer(), ERDGImportedBufferFlags::CreateSRV);
 					AddHairClusterAABBPass(
 						GraphBuilder,
 						ShaderMap,
@@ -1567,7 +1577,8 @@ void ComputeHairStrandsInterpolation(
 						Strands_PositionOffsetSRV,
 						HairGroupCluster,
 						CullingData,
-						Strands_PositionSRV);
+						Strands_PositionSRV,
+						Strands_CulledVertexCount.SRV);
 				}
 			}
 
diff --git a/Engine/Shaders/Private/HairStrands/HairStrandsInterpolation.usf b/Engine/Shaders/Private/HairStrands/HairStrandsInterpolation.usf
--- a/Engine/Shaders/Private/HairStrands/HairStrandsInterpolation.usf
+++ b/Engine/Shaders/Private/HairStrands/HairStrandsInterpolation.usf
@@ -600,6 +600,7 @@ Buffer<uint> ClusterIdBuffer;
 Buffer<uint> ClusterIndexOffsetBuffer;
 Buffer<uint> ClusterIndexCountBuffer;
 
+Buffer<uint> HairStrandsVF_CullingIndirectBuffer;
 Buffer<float4> RenderDeformedOffsetBuffer;
 
 RWBuffer<int> OutClusterAABBBuffer; 	// Cluster data packed as {uint3 AABBMin, uint3 AABBMax}
@@ -624,6 +625,8 @@ void ClusterAABBEvaluationCS(uint3 DispatchThreadId : SV_DispatchThreadID, uint3
 	const uint IndexOffset = ClusterIndexOffsetBuffer[ClusterFetchIndex];
 	const uint IndexCount  = ClusterIndexCountBuffer[ClusterFetchIndex];
 
+	const uint VertexCountAfterCulling = HairStrandsVF_CullingIndirectBuffer[3];
+
 	int3 ClusterAABBMin = 9999999.0f;
 	int3 ClusterAABBMax =-9999999.0f;
 
@@ -636,8 +639,7 @@ void ClusterAABBEvaluationCS(uint3 DispatchThreadId : SV_DispatchThreadID, uint3
 		if (ThreadReadOffset < IndexCount)
 		{
 			const float3 OutHairPositionOffset = RenderDeformedOffsetBuffer[0].xyz;
-
-			const uint VertexIndex = ClusterVertexIdBuffer[IndexOffset + ThreadReadOffset];
+			const uint VertexIndex = clamp(ClusterVertexIdBuffer[IndexOffset + ThreadReadOffset], 0, VertexCountAfterCulling-1);
 			const FHairControlPoint CP = UnpackHairControlPoint(RenderDeformedPositionBuffer[VertexIndex], OutHairPositionOffset, 1, 1);
 
 			if (all(IsFinite(CP.Position)))
diff --git a/Engine/Shaders/Private/HairStrands/HairStrandsTangent.usf b/Engine/Shaders/Private/HairStrands/HairStrandsTangent.usf
--- a/Engine/Shaders/Private/HairStrands/HairStrandsTangent.usf
+++ b/Engine/Shaders/Private/HairStrands/HairStrandsTangent.usf
@@ -43,8 +43,8 @@ void MainCS(uint2 DispatchThreadId : SV_DispatchThreadID)
 #if PERMUTATION_CULLING
 	if (HairStrandsVF_bIsCullingEnable)
 	{
-		const uint FetchIndex = DispatchThreadId.x + DispatchThreadId.y * GROUP_SIZE * GROUP_COUNT_X;
-		const uint VertexCountAfterCulling = HairStrandsVF_CullingIndirectBuffer[3];
+		const int FetchIndex = DispatchThreadId.x + DispatchThreadId.y * GROUP_SIZE * GROUP_COUNT_X;
+		const int VertexCountAfterCulling = clamp(HairStrandsVF_CullingIndirectBuffer[3], 0, MaxVertexCount);
 		bValid = false;
 		if (FetchIndex < VertexCountAfterCulling)
 		{
