Fix crash when decimation value are incorrect.

Sanitize the decimation factors to insure vertex/curves count are always correct.

diff --git a/Engine/Plugins/Runtime/HairStrands/Source/HairStrandsCore/Private/GroomBuilder.cpp b/Engine/Plugins/Runtime/HairStrands/Source/HairStrandsCore/Private/GroomBuilder.cpp
--- a/Engine/Plugins/Runtime/HairStrands/Source/HairStrandsCore/Private/GroomBuilder.cpp
+++ b/Engine/Plugins/Runtime/HairStrands/Source/HairStrandsCore/Private/GroomBuilder.cpp
@@ -1508,7 +1508,13 @@ bool FGroomBuilder::BuildGroom(FProcessedHairDescription& ProcessedHairDescripti
 	check(GroupIndex < GroupCount);
 	check(uint32(GroomAsset->GetNumHairGroups()) == GroupCount);
 
-	FHairGroupsInterpolation& InSettings = GroomAsset->HairGroupsInterpolation[GroupIndex];
+	const FHairGroupsInterpolation& InSettings = GroomAsset->HairGroupsInterpolation[GroupIndex];
+
+	// Sanitize decimation values. Do not update the 'InSettings' values directly as this would change 
+	// the groom asset and thus would change the DDC key
+	const float CurveDecimation		= FMath::Clamp(InSettings.DecimationSettings.CurveDecimation, 0.f, 1.f);
+	const float VertexDecimation	= FMath::Clamp(InSettings.DecimationSettings.VertexDecimation, 0.f, 1.f);
+	const float HairToGuideDensity	= FMath::Clamp(InSettings.InterpolationSettings.HairToGuideDensity, 0.f, 1.f);	
 
 	for (TPair<int32, FProcessedHairDescription::FHairGroup>& HairGroupIt : ProcessedHairDescription.HairGroups)
 	{
@@ -1536,11 +1542,11 @@ bool FGroomBuilder::BuildGroom(FProcessedHairDescription& ProcessedHairDescripti
 			HairStrandsBuilder::BuildInternalData(RenData, !ProcessedHairDescription.bHasUVData);
 
 			// Decimate
-			if (InSettings.DecimationSettings.CurveDecimation < 1 || InSettings.DecimationSettings.VertexDecimation < 1)
+			if (CurveDecimation < 1 || VertexDecimation < 1)
 			{
 				FHairStrandsDatas FullData = RenData;
 				RenData.Reset();
-				Decimate(FullData, InSettings.DecimationSettings.CurveDecimation, InSettings.DecimationSettings.VertexDecimation, RenData);
+				Decimate(FullData, CurveDecimation, VertexDecimation, RenData);
 			}
 		}
 
@@ -1562,7 +1568,7 @@ bool FGroomBuilder::BuildGroom(FProcessedHairDescription& ProcessedHairDescripti
 			else
 			{
 				SimData.Reset();
-				Decimate(GroupData.Strands.Data, InSettings.InterpolationSettings.HairToGuideDensity, 1, SimData);
+				Decimate(GroupData.Strands.Data, HairToGuideDensity, 1, SimData);
 			}
 		}
 	}
@@ -1659,7 +1665,7 @@ void FGroomBuilder::BuildData(FHairGroupData& GroupData, const FHairGroupsInterp
 
 inline uint32 DecimatePointCount(uint32 InCount, float InDecimationFactor)
 {
-	return FMath::Max(2.f, FMath::CeilToFloat(InCount * InDecimationFactor));
+	return FMath::Clamp(uint32(FMath::CeilToInt(InCount * InDecimationFactor)), 2u, InCount);
 }
 
 static void DecimateCurve(
@@ -1911,6 +1917,11 @@ static void DecimateCurve(
 		float PrevAngle = 0;
 		for (FHairLODSettings& S : Settings)
 		{
+			// Sanitize the decimation values
+			S.CurveDecimation  = FMath::Clamp(S.CurveDecimation,  0.f,  1.f);
+			S.VertexDecimation = FMath::Clamp(S.VertexDecimation, 0.f,  1.f);
+			S.AngularThreshold = FMath::Clamp(S.AngularThreshold, 0.f, 90.f);
+
 			if (S.VertexDecimation > PrevFactor)
 			{
 				S.VertexDecimation = PrevFactor;
@@ -1941,8 +1952,8 @@ static void DecimateCurve(
 
 	for (uint8 LODIt = 0; LODIt < LODCount; ++LODIt)
 	{
-		const int32 LODTargetVertexCount = FMath::Max(2.f, FMath::CeilToFloat(InCount * Settings[LODIt].VertexDecimation));
-		const float LODAngularThreshold = FMath::DegreesToRadians(Settings[LODIt].AngularThreshold);
+		const int32 LODTargetVertexCount = FMath::Clamp(uint32(FMath::CeilToInt(InCount * Settings[LODIt].VertexDecimation)), 2u, InCount);
+		const float LODAngularThreshold  = FMath::DegreesToRadians(Settings[LODIt].AngularThreshold);
 
 		// 'bCanDecimate' tracks if it is possible to reduce the remaining vertives even more while respecting the user angular constrain
 		bool bCanDecimate = true;
@@ -2254,7 +2265,8 @@ static void InternalBuildClusterData(
 		LODCurveCount.SetNum(LODCount);
 		for (uint8 LODIt = 0; LODIt < LODCount; ++LODIt)
 		{
-			LODCurveCount[LODIt] = FMath::Max(1u, uint32(FMath::CeilToInt(Cluster.ClusterCurves.Num() * InSettings.LODs[LODIt].CurveDecimation)));
+			const float CurveDecimation = FMath::Clamp(InSettings.LODs[LODIt].CurveDecimation, 0.0f, 1.0f);
+			LODCurveCount[LODIt] = FMath::Clamp(FMath::CeilToInt(Cluster.ClusterCurves.Num() * CurveDecimation), 1, Cluster.ClusterCurves.Num());
 		}
 
 		// 4.4 Decimate each curve for all LODs
