Fix follicle texture generation crash when providing invalid inputs.

diff --git a/Engine/Plugins/Runtime/HairStrands/Source/HairStrandsEditor/Private/GroomActions.cpp b/Engine/Plugins/Runtime/HairStrands/Source/HairStrandsEditor/Private/GroomActions.cpp
--- a/Engine/Plugins/Runtime/HairStrands/Source/HairStrandsEditor/Private/GroomActions.cpp
+++ b/Engine/Plugins/Runtime/HairStrands/Source/HairStrandsEditor/Private/GroomActions.cpp
@@ -357,12 +357,15 @@ void FGroomActions::ExecuteCreateFollicleTexture(TArray<TWeakObjectPtr<UGroomAss
 		TArray<FFollicleInfo> Infos;
 		for (FFollicleMaskOptions& Option : CurrentOptions->Grooms)
 		{
-			Option.Groom->ConditionalPostLoad();
+			if (Option.Groom)
+			{
+				Option.Groom->ConditionalPostLoad();
 
-			FFollicleInfo& Info = Infos.AddDefaulted_GetRef();
-			Info.GroomAsset			= Option.Groom;
-			Info.Channel			= FFollicleInfo::EChannel(uint8(Option.Channel));
-			Info.KernelSizeInPixels = FMath::Max(2, CurrentOptions->RootRadius);
+				FFollicleInfo& Info = Infos.AddDefaulted_GetRef();
+				Info.GroomAsset			= Option.Groom;
+				Info.Channel			= FFollicleInfo::EChannel(uint8(Option.Channel));
+				Info.KernelSizeInPixels = FMath::Max(2, CurrentOptions->RootRadius);
+			}
 		}
 
 		const uint32 Resolution = FMath::RoundUpToPowerOfTwo(CurrentOptions->Resolution);
