Fix debug mode changing in groom editor when rebuilding groom asset.

diff --git a/Engine/Plugins/Runtime/HairStrands/Source/HairStrandsCore/Private/GroomComponent.cpp b/Engine/Plugins/Runtime/HairStrands/Source/HairStrandsCore/Private/GroomComponent.cpp
--- a/Engine/Plugins/Runtime/HairStrands/Source/HairStrandsCore/Private/GroomComponent.cpp
+++ b/Engine/Plugins/Runtime/HairStrands/Source/HairStrandsCore/Private/GroomComponent.cpp
@@ -2146,6 +2146,9 @@ void UGroomComponent::InitResources(bool bIsBindingReloading)
 
 		// LODBias is in the Modifier which is needed for LOD selection regardless if the strands are there or not
 		HairGroupInstance->Strands.Modifier = GetGroomGroupsDesc(GroomAsset, this, GroupIt);
+		#if WITH_EDITORONLY_DATA
+		HairGroupInstance->Debug.DebugMode = GroomAsset ? GroomAsset->GetDebugMode() : EHairStrandsDebugMode::NoneDebug;
+		#endif// #if WITH_EDITORONLY_DATA
 
 		// Strands data/resources
 		if (bIsStrandsEnabled && GroupData.Strands.IsValid())
