Fix hair cards & hair meshes transistion error.

diff --git a/Engine/Plugins/Runtime/HairStrands/Source/HairStrandsCore/Private/HairStrandsRendering.cpp b/Engine/Plugins/Runtime/HairStrands/Source/HairStrandsCore/Private/HairStrandsRendering.cpp
--- a/Engine/Plugins/Runtime/HairStrands/Source/HairStrandsCore/Private/HairStrandsRendering.cpp
+++ b/Engine/Plugins/Runtime/HairStrands/Source/HairStrandsCore/Private/HairStrandsRendering.cpp
@@ -1719,6 +1719,8 @@ void ComputeHairStrandsInterpolation(
 						Instance->Guides.RestRootResource,
 						Instance->Guides.DeformedRootResource);
 				}
+
+				ConvertToUntrackedBuffer(GraphBuilder, Register(GraphBuilder, LOD.DeformedResource->GetBuffer(FHairCardsDeformedResource::Current), ERDGImportedBufferFlags::None).Buffer, ERHIAccess::SRVMask);
 			}
 
 			#if RHI_RAYTRACING
@@ -1778,6 +1780,8 @@ void ComputeHairStrandsInterpolation(
 					MeshesInstance.DeformedResource,
 					Instance->Guides.RestRootResource,
 					Instance->Guides.DeformedRootResource);
+
+				ConvertToUntrackedBuffer(GraphBuilder, Register(GraphBuilder, MeshesInstance.DeformedResource->GetBuffer(FHairMeshesDeformedResource::Current), ERDGImportedBufferFlags::None).Buffer, ERHIAccess::SRVMask);
 			}
 
 			#if RHI_RAYTRACING
