Clean groom asset/component overrides

diff --git a/Engine/Plugins/Runtime/HairStrands/Source/HairStrandsCore/Private/GroomAsset.cpp b/Engine/Plugins/Runtime/HairStrands/Source/HairStrandsCore/Private/GroomAsset.cpp
--- a/Engine/Plugins/Runtime/HairStrands/Source/HairStrandsCore/Private/GroomAsset.cpp
+++ b/Engine/Plugins/Runtime/HairStrands/Source/HairStrandsCore/Private/GroomAsset.cpp
@@ -535,6 +535,7 @@ void UGroomAsset::UpdateHairGroupsInfo()
 		Info.NumGuides = Data.Guides.Data.GetNumCurves();
 		Info.NumCurveVertices = Data.Strands.Data.GetNumPoints();
 		Info.NumGuideVertices = Data.Guides.Data.GetNumPoints();
+		Info.MaxCurveLength = Data.Strands.Data.StrandsCurves.MaxLength;
 		if (bForceReset)
 		{
 			Info.bIsVisible = true;
diff --git a/Engine/Plugins/Runtime/HairStrands/Source/HairStrandsCore/Private/GroomComponent.cpp b/Engine/Plugins/Runtime/HairStrands/Source/HairStrandsCore/Private/GroomComponent.cpp
--- a/Engine/Plugins/Runtime/HairStrands/Source/HairStrandsCore/Private/GroomComponent.cpp
+++ b/Engine/Plugins/Runtime/HairStrands/Source/HairStrandsCore/Private/GroomComponent.cpp
@@ -106,15 +106,13 @@ static FHairGroupDesc GetGroomGroupsDesc(const UGroomAsset* Asset, UGroomCompone
 	}
 
 	FHairGroupDesc O = Component->GroomGroupsDesc[GroupIndex];
-	O.HairCount  = Asset->HairGroupsData[GroupIndex].Strands.Data.GetNumCurves();
-	O.GuideCount = Asset->HairGroupsData[GroupIndex].Guides.Data.GetNumCurves();
 	O.HairLength = Asset->HairGroupsData[GroupIndex].Strands.Data.StrandsCurves.MaxLength;
 	O.LODBias 	 = Asset->EffectiveLODBias[GroupIndex] > 0 ? FMath::Max(O.LODBias, Asset->EffectiveLODBias[GroupIndex]) : O.LODBias;
 
 	if (!O.HairWidth_Override)					{ O.HairWidth					= Asset->HairGroupsRendering[GroupIndex].GeometrySettings.HairWidth;					}
 	if (!O.HairRootScale_Override)				{ O.HairRootScale				= Asset->HairGroupsRendering[GroupIndex].GeometrySettings.HairRootScale;				}
 	if (!O.HairTipScale_Override)				{ O.HairTipScale				= Asset->HairGroupsRendering[GroupIndex].GeometrySettings.HairTipScale;					}
-	if (!O.HairClipLength_Override)				{ O.HairClipLength				= Asset->HairGroupsRendering[GroupIndex].GeometrySettings.HairClipScale * O.HairLength;	}
+	if (!O.HairClipScale_Override)				{ O.HairClipScale				= Asset->HairGroupsRendering[GroupIndex].GeometrySettings.HairClipScale;				}
 	if (!O.bSupportVoxelization_Override)		{ O.bSupportVoxelization		= Asset->HairGroupsRendering[GroupIndex].ShadowSettings.bVoxelize;						}
 	if (!O.HairShadowDensity_Override)			{ O.HairShadowDensity			= Asset->HairGroupsRendering[GroupIndex].ShadowSettings.HairShadowDensity;				}
 	if (!O.HairRaytracingRadiusScale_Override)	{ O.HairRaytracingRadiusScale	= Asset->HairGroupsRendering[GroupIndex].ShadowSettings.HairRaytracingRadiusScale;		}
@@ -1081,11 +1079,6 @@ void UGroomComponent::UpdateHairGroupsDesc()
 	{
 		GroomGroupsDesc.Init(FHairGroupDesc(), GroupCount);
 	}
-
-	for (uint32 GroupIt = 0; GroupIt < GroupCount; ++GroupIt)
-	{
-		GroomGroupsDesc[GroupIt] = GetGroomGroupsDesc(GroomAsset, this, GroupIt);
-	}
 }
 
 void UGroomComponent::ReleaseHairSimulation()
@@ -1287,8 +1280,8 @@ void UGroomComponent::SetHairLengthScale(float Scale)
 	Scale = FMath::Clamp(Scale, 0.f, 1.f);
 	for (FHairGroupDesc& HairDesc : GroomGroupsDesc)
 	{
-		HairDesc.HairClipLength = Scale * HairDesc.HairLength;
-		HairDesc.HairClipLength_Override = true;
+		HairDesc.HairClipScale = Scale;
+		HairDesc.HairClipScale_Override = true;
 	}
 	UpdateHairGroupsDescAndInvalidateRenderState();
 }
@@ -2868,10 +2861,9 @@ void UGroomComponent::PostEditChangeProperty(FPropertyChangedEvent& PropertyChan
 #endif
 
 	if (PropertyName == GET_MEMBER_NAME_CHECKED(FHairGroupDesc, HairWidth) || 
-		PropertyName == GET_MEMBER_NAME_CHECKED(FHairGroupDesc, HairLength) ||
 		PropertyName == GET_MEMBER_NAME_CHECKED(FHairGroupDesc, HairRootScale) ||
 		PropertyName == GET_MEMBER_NAME_CHECKED(FHairGroupDesc, HairTipScale) ||
-		PropertyName == GET_MEMBER_NAME_CHECKED(FHairGroupDesc, HairClipLength) ||
+		PropertyName == GET_MEMBER_NAME_CHECKED(FHairGroupDesc, HairClipScale) ||
 		PropertyName == GET_MEMBER_NAME_CHECKED(FHairGroupDesc, HairShadowDensity) ||
 		PropertyName == GET_MEMBER_NAME_CHECKED(FHairGroupDesc, HairRaytracingRadiusScale) ||
 		PropertyName == GET_MEMBER_NAME_CHECKED(FHairGroupDesc, LODBias) ||
diff --git a/Engine/Plugins/Runtime/HairStrands/Source/HairStrandsCore/Public/GroomAsset.h b/Engine/Plugins/Runtime/HairStrands/Source/HairStrandsCore/Public/GroomAsset.h
--- a/Engine/Plugins/Runtime/HairStrands/Source/HairStrandsCore/Public/GroomAsset.h
+++ b/Engine/Plugins/Runtime/HairStrands/Source/HairStrandsCore/Public/GroomAsset.h
@@ -48,6 +48,9 @@ struct HAIRSTRANDSCORE_API FHairGroupInfo
 
 	UPROPERTY(VisibleAnywhere, Category = "Info", meta = (DisplayName = "Guide Vertex Count"))
 	int32 NumGuideVertices = 0;
+
+	UPROPERTY(VisibleAnywhere, Category = "Info", meta = (DisplayName = "Length of the longest hair strands"))
+	float MaxCurveLength = 0;
 };
 
 USTRUCT(BlueprintType)
diff --git a/Engine/Plugins/Runtime/HairStrands/Source/HairStrandsCore/Public/GroomDesc.h b/Engine/Plugins/Runtime/HairStrands/Source/HairStrandsCore/Public/GroomDesc.h
--- a/Engine/Plugins/Runtime/HairStrands/Source/HairStrandsCore/Public/GroomDesc.h
+++ b/Engine/Plugins/Runtime/HairStrands/Source/HairStrandsCore/Public/GroomDesc.h
@@ -14,58 +14,50 @@ struct FHairGroupDesc
 {
 	GENERATED_USTRUCT_BODY()
 
-	/** Number of hairs within this hair group. */
-	UPROPERTY(VisibleAnywhere, BlueprintReadOnly, Category = "Groom")
-	int32 HairCount = 0;
-
-	/** Number of simulation guides within this hair group. */
-	UPROPERTY(VisibleAnywhere, BlueprintReadOnly, Category = "Groom")
-	int32 GuideCount = 0;
-
 	/** Length of the longest hair strands */
-	UPROPERTY(VisibleAnywhere, BlueprintReadOnly, Category = "Groom")
+	UPROPERTY()
 	float HairLength = 0;
 
 	/** Hair width (in centimeters) */
-	UPROPERTY(EditAnywhere, BlueprintReadWrite, Category = "Groom", meta = (ClampMin = "0.0001", UIMin = "0.001", UIMax = "1.0", SliderExponent = 6))
-	float HairWidth = 0;
-	UPROPERTY()
+	UPROPERTY(EditAnywhere, BlueprintReadWrite, Category = "Groom", meta = (editcondition = "HairWidth_Override", ClampMin = "0.0001", UIMin = "0.001", UIMax = "1.0", SliderExponent = 6))
+	float HairWidth = 0.1;
+	UPROPERTY(EditAnywhere, BlueprintReadWrite, Category = Overrides, meta = (PinHiddenByDefault, InlineEditConditionToggle))
 	bool HairWidth_Override = false;
 
 	/** Scale the hair width at the root */
-	UPROPERTY(EditAnywhere, BlueprintReadWrite, Category = "Groom", AdvancedDisplay, meta = (ClampMin = "0.0001", UIMin = "0.001", UIMax = "2.0", SliderExponent = 6))
-	float HairRootScale = 0;
-	UPROPERTY()
+	UPROPERTY(EditAnywhere, BlueprintReadWrite, Category = "Groom", AdvancedDisplay, meta = (editcondition = "HairRootScale_Override", ClampMin = "0.0001", UIMin = "0.001", UIMax = "2.0", SliderExponent = 6))
+	float HairRootScale = 1.0f;
+	UPROPERTY(EditAnywhere, BlueprintReadWrite, Category = Overrides, meta = (PinHiddenByDefault, InlineEditConditionToggle))
 	bool HairRootScale_Override = false;
 
 	/** Scale the hair with at the tip */
-	UPROPERTY(EditAnywhere, BlueprintReadWrite, Category = "Groom", AdvancedDisplay, meta = (ClampMin = "0.0001", UIMin = "0.001", UIMax = "2.0", SliderExponent = 6))
-	float HairTipScale = 0;
-	UPROPERTY()
+	UPROPERTY(EditAnywhere, BlueprintReadWrite, Category = "Groom", AdvancedDisplay, meta = (editcondition = "HairTipScale_Override", ClampMin = "0.0001", UIMin = "0.001", UIMax = "2.0", SliderExponent = 6))
+	float HairTipScale = 1.0f;
+	UPROPERTY(EditAnywhere, BlueprintReadWrite, Category = Overrides, meta = (PinHiddenByDefault, InlineEditConditionToggle))
 	bool HairTipScale_Override = false;
 
-	/** Normalized hair clip length, i.e. at which length hair will be clipped. 1 means no clipping. 0 means hairs are fully clipped */
-	UPROPERTY(EditAnywhere, BlueprintReadWrite, Category = "Groom", AdvancedDisplay, meta = (ClampMin = "0.0", UIMin = "0.0", UIMax = "1.0", SliderExponent = 1))
-	float HairClipLength = 0;
-	UPROPERTY()
-	bool HairClipLength_Override = false;
+	/** Normalized hair clip scale, i.e. at which normalized length hair will be clipped. 1 means no clipping. 0 means hairs are fully clipped */
+	UPROPERTY(EditAnywhere, BlueprintReadWrite, Category = "Groom", AdvancedDisplay, meta = (editcondition = "HairClipScale_Override", ClampMin = "0.0", UIMin = "0.0", UIMax = "1.0", SliderExponent = 1))
+	float HairClipScale = 1.0f;
+	UPROPERTY(EditAnywhere, BlueprintReadWrite, Category = Overrides, meta = (PinHiddenByDefault, InlineEditConditionToggle))
+	bool HairClipScale_Override = false;
 
 	/** Override the hair shadow density factor (unit less). */
-	UPROPERTY(EditAnywhere, BlueprintReadWrite, Category = "Groom", AdvancedDisplay, meta = (ClampMin = "0.0001", UIMin = "0.001", UIMax = "10.0", SliderExponent = 6))
-	float HairShadowDensity = 0;
-	UPROPERTY()
+	UPROPERTY(EditAnywhere, BlueprintReadWrite, Category = "Groom", AdvancedDisplay, meta = (editcondition = "HairShadowDensity_Override", ClampMin = "0.0001", UIMin = "0.001", UIMax = "10.0", SliderExponent = 6))
+	float HairShadowDensity = 1.0f;
+	UPROPERTY(EditAnywhere, BlueprintReadWrite, Category = Overrides, meta = (PinHiddenByDefault, InlineEditConditionToggle))
 	bool HairShadowDensity_Override = false;
 
 	/** Scale the hair geometry radius for ray tracing effects (e.g. shadow) */
-	UPROPERTY(EditAnywhere, BlueprintReadWrite, Category = "Groom", AdvancedDisplay, meta = (ClampMin = "0.0001", UIMin = "0.001", UIMax = "10.0", SliderExponent = 6))
-	float HairRaytracingRadiusScale = 0;
-	UPROPERTY()
+	UPROPERTY(EditAnywhere, BlueprintReadWrite, Category = "Groom", AdvancedDisplay, meta = (editcondition = "HairRaytracingRadiusScale_Override", ClampMin = "0.0001", UIMin = "0.001", UIMax = "10.0", SliderExponent = 6))
+	float HairRaytracingRadiusScale = 1.0f;
+	UPROPERTY(EditAnywhere, BlueprintReadWrite, Category = Overrides, meta = (PinHiddenByDefault, InlineEditConditionToggle))
 	bool HairRaytracingRadiusScale_Override = false;
 
 	/** Enable hair strands geomtry for raytracing */
-	UPROPERTY(EditAnywhere, BlueprintReadWrite, Category = "Groom", AdvancedDisplay, meta = (ClampMin = "0.0001", UIMin = "0.001", UIMax = "10.0", SliderExponent = 6))
+	UPROPERTY(EditAnywhere, BlueprintReadWrite, Category = "Groom", AdvancedDisplay, meta = (editcondition = "bUseHairRaytracingGeometry_Override", ClampMin = "0.0001", UIMin = "0.001", UIMax = "10.0", SliderExponent = 6))
 	bool bUseHairRaytracingGeometry = false;
-	UPROPERTY()
+	UPROPERTY(EditAnywhere, BlueprintReadWrite, Category = Overrides, meta = (PinHiddenByDefault, InlineEditConditionToggle))
 	bool bUseHairRaytracingGeometry_Override = false;
 
 	/** Bias the selected LOD. A value >0 will progressively select lower detailed lods. Used when r.HairStrands.Cluster.Culling = 1. */
@@ -73,15 +65,15 @@ struct FHairGroupDesc
 	float LODBias = 0;
 
 	/** Insure the hair does not alias. When enable, group of hairs might appear thicker. Isolated hair should remain thin. */
-	UPROPERTY(EditAnywhere, BlueprintReadOnly, Category = "Groom")
+	UPROPERTY(EditAnywhere, BlueprintReadOnly, Category = "Groom", meta = (editcondition = "bUseStableRasterization_Override"))
 	bool bUseStableRasterization = false;
-	UPROPERTY()
+	UPROPERTY(EditAnywhere, BlueprintReadWrite, Category = Overrides, meta = (PinHiddenByDefault, InlineEditConditionToggle))
 	bool bUseStableRasterization_Override = false;
 
 	/** Light hair with the scene color. This is used for vellus/short hair to bring light from the surrounding surface, like skin. */
-	UPROPERTY(EditAnywhere, BlueprintReadOnly, Category = "Groom")
+	UPROPERTY(EditAnywhere, BlueprintReadOnly, Category = "Groom", meta = (editcondition = "bScatterSceneLighting_Override"))
 	bool bScatterSceneLighting = false;
-	UPROPERTY()
+	UPROPERTY(EditAnywhere, BlueprintReadWrite, Category = Overrides, meta = (PinHiddenByDefault, InlineEditConditionToggle))
 	bool bScatterSceneLighting_Override = false;
 
 	UPROPERTY()
diff --git a/Engine/Plugins/Runtime/HairStrands/Source/HairStrandsEditor/Private/GroomComponentDetailsCustomization.cpp b/Engine/Plugins/Runtime/HairStrands/Source/HairStrandsEditor/Private/GroomComponentDetailsCustomization.cpp
--- a/Engine/Plugins/Runtime/HairStrands/Source/HairStrandsEditor/Private/GroomComponentDetailsCustomization.cpp
+++ b/Engine/Plugins/Runtime/HairStrands/Source/HairStrandsEditor/Private/GroomComponentDetailsCustomization.cpp
@@ -86,77 +86,58 @@ void FGroomComponentDetailsCustomization::CustomizeDescGroupProperties(IDetailLa
 	}
 }
 
-void FGroomComponentDetailsCustomization::SetOverride(int32 GroupIndex, TSharedPtr<IPropertyHandle> ChildHandle, bool bValue)
+void FGroomComponentDetailsCustomization::OnResetToDefault(int32 GroupIndex, TSharedPtr<IPropertyHandle> ChildHandle)
 {
+	if (ChildHandle == nullptr || GroomComponentPtr == nullptr || GroupIndex < 0 || GroupIndex >= GroomComponentPtr->GroomGroupsDesc.Num())
+	{
+		return;
+	}
+	
 	FName PropertyName = ChildHandle->GetProperty()->GetFName();
 
-	// Reset the override flag so that the groom component will fallback onto the groom asset value
+	FHairGroupDesc Default;
 	if (PropertyName == GET_MEMBER_NAME_CHECKED(FHairGroupDesc, HairWidth))
 	{
-		GroomComponentPtr->GroomGroupsDesc[GroupIndex].HairWidth_Override = bValue;
+		GroomComponentPtr->GroomGroupsDesc[GroupIndex].HairWidth = Default.HairWidth;
 	}
 	else if (PropertyName == GET_MEMBER_NAME_CHECKED(FHairGroupDesc, HairRootScale))
 	{
-		GroomComponentPtr->GroomGroupsDesc[GroupIndex].HairRootScale_Override = bValue;
+		GroomComponentPtr->GroomGroupsDesc[GroupIndex].HairRootScale = Default.HairRootScale;
 	}
 	else if (PropertyName == GET_MEMBER_NAME_CHECKED(FHairGroupDesc, HairTipScale))
 	{
-		GroomComponentPtr->GroomGroupsDesc[GroupIndex].HairTipScale_Override = bValue;
+		GroomComponentPtr->GroomGroupsDesc[GroupIndex].HairTipScale = Default.HairTipScale;
 	}
-	else if (PropertyName == GET_MEMBER_NAME_CHECKED(FHairGroupDesc, HairClipLength))
+	else if (PropertyName == GET_MEMBER_NAME_CHECKED(FHairGroupDesc, HairClipScale))
 	{
-		GroomComponentPtr->GroomGroupsDesc[GroupIndex].HairClipLength_Override = bValue;
+		GroomComponentPtr->GroomGroupsDesc[GroupIndex].HairClipScale = Default.HairClipScale;
 	}
 	else if (PropertyName == GET_MEMBER_NAME_CHECKED(FHairGroupDesc, HairShadowDensity))
 	{
-		GroomComponentPtr->GroomGroupsDesc[GroupIndex].HairShadowDensity_Override = bValue;
+		GroomComponentPtr->GroomGroupsDesc[GroupIndex].HairShadowDensity = Default.HairShadowDensity;
 	}
 	else if (PropertyName == GET_MEMBER_NAME_CHECKED(FHairGroupDesc, HairRaytracingRadiusScale))
 	{
-		GroomComponentPtr->GroomGroupsDesc[GroupIndex].HairRaytracingRadiusScale_Override = bValue;
+		GroomComponentPtr->GroomGroupsDesc[GroupIndex].HairRaytracingRadiusScale = Default.HairRaytracingRadiusScale;
 	}
 	else if (PropertyName == GET_MEMBER_NAME_CHECKED(FHairGroupDesc, bUseHairRaytracingGeometry))
 	{
-		GroomComponentPtr->GroomGroupsDesc[GroupIndex].bUseHairRaytracingGeometry_Override = bValue;
+		GroomComponentPtr->GroomGroupsDesc[GroupIndex].bUseHairRaytracingGeometry = Default.bUseHairRaytracingGeometry;
 	}
 	else if (PropertyName == GET_MEMBER_NAME_CHECKED(FHairGroupDesc, bUseStableRasterization))
 	{
-		GroomComponentPtr->GroomGroupsDesc[GroupIndex].bUseStableRasterization_Override = bValue;
+		GroomComponentPtr->GroomGroupsDesc[GroupIndex].bUseStableRasterization = Default.bUseStableRasterization;
 	}
 	else if (PropertyName == GET_MEMBER_NAME_CHECKED(FHairGroupDesc, bScatterSceneLighting))
 	{
-		GroomComponentPtr->GroomGroupsDesc[GroupIndex].bScatterSceneLighting_Override = bValue;
+		GroomComponentPtr->GroomGroupsDesc[GroupIndex].bScatterSceneLighting = Default.bScatterSceneLighting;
 	}
-}
 
-void FGroomComponentDetailsCustomization::OnResetToDefault(int32 GroupIndex, TSharedPtr<IPropertyHandle> ChildHandle)
-{
-	if (ChildHandle == nullptr || GroomComponentPtr == nullptr || GroupIndex < 0 || GroupIndex >= GroomComponentPtr->GroomGroupsDesc.Num())
-	{
-		return;
-	}
-	
 	FScopedTransaction ScopedTransaction(NSLOCTEXT("UnrealEd", "PropertyWindowResetToDefault", "Reset to Default"));
-	//GroomComponentPtr->Modify();
-	SetOverride(GroupIndex, ChildHandle, false);
-	GroomComponentPtr->UpdateHairGroupsDescAndInvalidateRenderState();
-}
-
-void FGroomComponentDetailsCustomization::OnValueChanged(int32 GroupIndex, TSharedPtr<IPropertyHandle> ChildHandle)
-{
-	if (ChildHandle == nullptr || GroomComponentPtr == nullptr || GroupIndex < 0 || GroupIndex >= GroomComponentPtr->GroomGroupsDesc.Num())
-	{
-		return;
-	}
-
-	FScopedTransaction ScopedTransaction(NSLOCTEXT("UnrealEd", "PropertyWindowPreValueChanged", "PreValue Changed"));
-	//GroomComponentPtr->Modify();
-	SetOverride(GroupIndex, ChildHandle, true);
 	GroomComponentPtr->UpdateHairGroupsDescAndInvalidateRenderState();
 }
 
-
-static TSharedRef<SUniformGridPanel> MakeHairInfoGrid(const FSlateFontInfo& DetailFontInfo, FHairGroupDesc& GroupDesc)
+static TSharedRef<SUniformGridPanel> MakeHairInfoGrid(const FSlateFontInfo& DetailFontInfo, const FHairGroupInfo& Infos)
 {
 	TSharedRef<SUniformGridPanel> Grid = SNew(SUniformGridPanel).SlotPadding(2.0f);
 
@@ -191,21 +172,21 @@ static TSharedRef<SUniformGridPanel> MakeHairInfoGrid(const FSlateFontInfo& Deta
 	[
 		SNew(STextBlock)
 		.Font(DetailFontInfo)
-		.Text(FText::AsNumber(GroupDesc.HairCount))
+		.Text(FText::AsNumber(Infos.NumCurves))
 	];
 	Grid->AddSlot(1, 1) // x, y
 	.HAlign(HAlign_Right)
 	[
 		SNew(STextBlock)
 		.Font(DetailFontInfo)
-		.Text(FText::AsNumber(GroupDesc.GuideCount))
+		.Text(FText::AsNumber(Infos.NumGuides))
 	];
 	Grid->AddSlot(2, 1) // x, y
 	.HAlign(HAlign_Right)
 	[
 		SNew(STextBlock)
 		.Font(DetailFontInfo)
-		.Text(FText::AsNumber(GroupDesc.HairLength))
+		.Text(FText::AsNumber(Infos.MaxCurveLength))
 	];
 
 	return Grid;
@@ -249,11 +230,16 @@ void FGroomComponentDetailsCustomization::OnGenerateElementForHairGroup(TSharedR
 
 	if (GroomComponentPtr != nullptr && GroupIndex>=0 && GroupIndex < GroomComponentPtr->GroomGroupsDesc.Num())
 	{
+		FHairGroupInfo Infos;
+		if (GroomComponentPtr->GroomAsset && GroupIndex < GroomComponentPtr->GroomAsset->HairGroupsInfo.Num())
+		{
+			Infos = GroomComponentPtr->GroomAsset->HairGroupsInfo[GroupIndex];
+		}
 		ChildrenBuilder.AddCustomRow(LOCTEXT("HairInfo_Separator", "Separator"))
 		.ValueContent()
 		.HAlign(HAlign_Fill)
 		[
-			MakeHairInfoGrid(DetailFontInfo, GroomComponentPtr->GroomGroupsDesc[GroupIndex])
+			MakeHairInfoGrid(DetailFontInfo, Infos)
 		];
 	}
 
@@ -262,18 +248,7 @@ void FGroomComponentDetailsCustomization::OnGenerateElementForHairGroup(TSharedR
 	for (uint32 ChildIt = 0; ChildIt < ChildrenCount; ++ChildIt)
 	{
 		TSharedPtr<IPropertyHandle> ChildHandle = StructProperty->GetChildHandle(ChildIt);
-
-		// These properties are not display with regular widget but with a dedicated array
-		FName PropertyName = ChildHandle->GetProperty()->GetFName();
-		if (PropertyName == GET_MEMBER_NAME_CHECKED(FHairGroupDesc, HairCount) ||
-			PropertyName == GET_MEMBER_NAME_CHECKED(FHairGroupDesc, HairLength) ||
-			PropertyName == GET_MEMBER_NAME_CHECKED(FHairGroupDesc, GuideCount)) 
-		{
-			continue;
-		}
-
 		ChildHandle->SetOnPropertyResetToDefault(FSimpleDelegate::CreateSP(this, &FGroomComponentDetailsCustomization::OnResetToDefault, GroupIndex, ChildHandle));
-		ChildHandle->SetOnPropertyValuePreChange(FSimpleDelegate::CreateSP(this, &FGroomComponentDetailsCustomization::OnValueChanged, GroupIndex, ChildHandle));
 		ChildrenBuilder.AddProperty(ChildHandle.ToSharedRef());
 	}
 }
diff --git a/Engine/Plugins/Runtime/HairStrands/Source/HairStrandsEditor/Private/GroomComponentDetailsCustomization.h b/Engine/Plugins/Runtime/HairStrands/Source/HairStrandsEditor/Private/GroomComponentDetailsCustomization.h
--- a/Engine/Plugins/Runtime/HairStrands/Source/HairStrandsEditor/Private/GroomComponentDetailsCustomization.h
+++ b/Engine/Plugins/Runtime/HairStrands/Source/HairStrandsEditor/Private/GroomComponentDetailsCustomization.h
@@ -31,6 +31,4 @@ private:
 	void CustomizeDescGroupProperties(IDetailLayoutBuilder& DetailLayout, IDetailCategoryBuilder& StrandsGroupFilesCategory);
 	void OnGenerateElementForHairGroup(TSharedRef<IPropertyHandle> StructProperty, int32 GroupIndex, IDetailChildrenBuilder& ChildrenBuilder, IDetailLayoutBuilder* DetailLayout);
 	void OnResetToDefault(int32 GroupIndex, TSharedPtr<IPropertyHandle> ChildHandle);
-	void OnValueChanged(int32 GroupIndex, TSharedPtr<IPropertyHandle> ChildHandle);
-	void SetOverride(int32 GroupIndex, TSharedPtr<IPropertyHandle> ChildHandle, bool bValue);
 };
