Add missing files.

diff --git a/Engine/Plugins/Runtime/HairStrands/Source/HairStrandsCore/Private/HairStrandsRendering.cpp b/Engine/Plugins/Runtime/HairStrands/Source/HairStrandsCore/Private/HairStrandsRendering.cpp
--- a/Engine/Plugins/Runtime/HairStrands/Source/HairStrandsCore/Private/HairStrandsRendering.cpp
+++ b/Engine/Plugins/Runtime/HairStrands/Source/HairStrandsCore/Private/HairStrandsRendering.cpp
@@ -1305,7 +1305,7 @@ FHairScaleAndClipDesc ComputeHairScaleAndClipDesc(FHairGroupInstance* Instance)
 		Out.MaxOutHairRadius = Out.OutHairRadius * FMath::Max(1.f, FMath::Max(Instance->Strands.Modifier.HairRootScale, Instance->Strands.Modifier.HairTipScale));
 		Out.HairRadiusRootScale = Instance->Strands.Modifier.HairRootScale;
 		Out.HairRadiusTipScale = Instance->Strands.Modifier.HairTipScale;
-		Out.HairLengthClip = FMath::Clamp(Instance->Strands.Modifier.HairClipLength / Instance->Strands.Data->StrandsCurves.MaxLength, 0.f, 1.f);
+		Out.HairLengthClip = FMath::Clamp(Instance->Strands.Modifier.HairClipScale, 0.f, 1.f);
 	}
 	
 	return Out;
