Remove clip option from groom asset.

This feature is deprecated in in 4.27.

diff --git a/Engine/Plugins/Runtime/HairStrands/Source/HairStrandsCore/Private/GroomAssetRendering.cpp b/Engine/Plugins/Runtime/HairStrands/Source/HairStrandsCore/Private/GroomAssetRendering.cpp
--- a/Engine/Plugins/Runtime/HairStrands/Source/HairStrandsCore/Private/GroomAssetRendering.cpp
+++ b/Engine/Plugins/Runtime/HairStrands/Source/HairStrandsCore/Private/GroomAssetRendering.cpp
@@ -7,7 +7,6 @@ FHairGeometrySettings::FHairGeometrySettings()
 	HairWidth = 0.01;
 	HairRootScale = 1;
 	HairTipScale = 1;
-	HairClipScale = 1;
 }
 
 FHairShadowSettings::FHairShadowSettings()
@@ -36,8 +35,7 @@ bool FHairGeometrySettings::operator==(const FHairGeometrySettings& A) const
 	return
 		HairWidth == A.HairWidth &&
 		HairRootScale == A.HairRootScale &&
-		HairTipScale == A.HairTipScale &&
-		HairClipScale == A.HairClipScale;
+		HairTipScale == A.HairTipScale;
 }
 
 bool FHairShadowSettings::operator==(const FHairShadowSettings& A) const
diff --git a/Engine/Plugins/Runtime/HairStrands/Source/HairStrandsCore/Private/GroomComponent.cpp b/Engine/Plugins/Runtime/HairStrands/Source/HairStrandsCore/Private/GroomComponent.cpp
--- a/Engine/Plugins/Runtime/HairStrands/Source/HairStrandsCore/Private/GroomComponent.cpp
+++ b/Engine/Plugins/Runtime/HairStrands/Source/HairStrandsCore/Private/GroomComponent.cpp
@@ -29,10 +29,6 @@
 #include "GroomManager.h"
 #include "GroomInstance.h"
 
-static float GHairClipLength = -1;
-static FAutoConsoleVariableRef CVarHairClipLength(TEXT("r.HairStrands.DebugClipLength"), GHairClipLength, TEXT("Clip hair strands which have a lenth larger than this value. (default is -1, no effect)"));
-float GetHairClipLength() { return GHairClipLength > 0 ? GHairClipLength : 100000;  }
-
 static int32 GHairEnableAdaptiveSubsteps = 0;  
 static FAutoConsoleVariableRef CVarHairEnableAdaptiveSubsteps(TEXT("r.HairStrands.EnableAdaptiveSubsteps"), GHairEnableAdaptiveSubsteps, TEXT("Enable adaptive solver substeps"));
 bool IsHairAdaptiveSubstepsEnabled() { return (GHairEnableAdaptiveSubsteps == 1); }
@@ -112,7 +108,6 @@ static FHairGroupDesc GetGroomGroupsDesc(const UGroomAsset* Asset, UGroomCompone
 	if (!O.HairWidth_Override)					{ O.HairWidth					= Asset->HairGroupsRendering[GroupIndex].GeometrySettings.HairWidth;					}
 	if (!O.HairRootScale_Override)				{ O.HairRootScale				= Asset->HairGroupsRendering[GroupIndex].GeometrySettings.HairRootScale;				}
 	if (!O.HairTipScale_Override)				{ O.HairTipScale				= Asset->HairGroupsRendering[GroupIndex].GeometrySettings.HairTipScale;					}
-	if (!O.HairClipScale_Override)				{ O.HairClipScale				= Asset->HairGroupsRendering[GroupIndex].GeometrySettings.HairClipScale;				}
 	if (!O.bSupportVoxelization_Override)		{ O.bSupportVoxelization		= Asset->HairGroupsRendering[GroupIndex].ShadowSettings.bVoxelize;						}
 	if (!O.HairShadowDensity_Override)			{ O.HairShadowDensity			= Asset->HairGroupsRendering[GroupIndex].ShadowSettings.HairShadowDensity;				}
 	if (!O.HairRaytracingRadiusScale_Override)	{ O.HairRaytracingRadiusScale	= Asset->HairGroupsRendering[GroupIndex].ShadowSettings.HairRaytracingRadiusScale;		}
@@ -135,27 +130,23 @@ public:
 	const float DebugMode;
 	const float HairMinRadius;
 	const float HairMaxRadius;
-	const float HairClipLength;
 	const FVector HairGroupColor;
 
 	FName DebugModeParamName;
 	FName MinHairRadiusParamName;
 	FName MaxHairRadiusParamName;
-	FName HairClipLengthParamName;
 	FName HairGroupHairColorParamName;
 
 	/** Initialization constructor. */
-	FHairDebugModeMaterialRenderProxy(const FMaterialRenderProxy* InParent, float InMode, float InMinRadius, float InMaxRadius, float InHairClipLength, FVector InGroupColor) :
+	FHairDebugModeMaterialRenderProxy(const FMaterialRenderProxy* InParent, float InMode, float InMinRadius, float InMaxRadius, FVector InGroupColor) :
 		Parent(InParent),
 		DebugMode(InMode),
 		HairMinRadius(InMinRadius),
 		HairMaxRadius(InMaxRadius),
-		HairClipLength(InHairClipLength),
 		HairGroupColor(InGroupColor),
 		DebugModeParamName(NAME_FloatProperty),
 		MinHairRadiusParamName(NAME_ByteProperty),
 		MaxHairRadiusParamName(NAME_IntProperty),
-		HairClipLengthParamName(NAME_BoolProperty),
 		HairGroupHairColorParamName(NAME_VectorProperty)
 	{}
 
@@ -196,11 +187,6 @@ public:
 			*OutValue = HairMaxRadius;
 			return true;
 		}
-		else if (ParameterInfo.Name == HairClipLengthParamName)
-		{
-			*OutValue = HairClipLength;
-			return true;
-		}
 		return Parent->GetScalarValue(ParameterInfo, OutValue, Context);
 	}
 
@@ -770,41 +756,40 @@ public:
 			{
 				for (uint32 GroupIt = 0; GroupIt < GroupCount; ++GroupIt)
 				{
-		FMaterialRenderProxy* Strands_MaterialProxy = nullptr;
+					FMaterialRenderProxy* Strands_MaterialProxy = nullptr;
 					const EHairStrandsDebugMode DebugMode = GetHairStrandsGeometryDebugMode(Instances[GroupIt]);
-		if (DebugMode != EHairStrandsDebugMode::NoneDebug)
-		{
-			float DebugModeScalar = 0;
-			switch(DebugMode)
-			{
-			case EHairStrandsDebugMode::NoneDebug					: DebugModeScalar =99.f; break;
-			case EHairStrandsDebugMode::SimHairStrands				: DebugModeScalar = 0.f; break;
-			case EHairStrandsDebugMode::RenderHairStrands			: DebugModeScalar = 0.f; break;
-			case EHairStrandsDebugMode::RenderHairRootUV			: DebugModeScalar = 1.f; break;
-			case EHairStrandsDebugMode::RenderHairUV				: DebugModeScalar = 2.f; break;
-			case EHairStrandsDebugMode::RenderHairSeed				: DebugModeScalar = 3.f; break;
-			case EHairStrandsDebugMode::RenderHairDimension			: DebugModeScalar = 4.f; break;
-			case EHairStrandsDebugMode::RenderHairRadiusVariation	: DebugModeScalar = 5.f; break;
-			case EHairStrandsDebugMode::RenderHairRootUDIM			: DebugModeScalar = 6.f; break;
-			case EHairStrandsDebugMode::RenderHairBaseColor			: DebugModeScalar = 7.f; break;
-			case EHairStrandsDebugMode::RenderHairRoughness			: DebugModeScalar = 8.f; break;
-			case EHairStrandsDebugMode::RenderVisCluster			: DebugModeScalar = 0.f; break;
+					if (DebugMode != EHairStrandsDebugMode::NoneDebug)
+					{
+						float DebugModeScalar = 0;
+						switch(DebugMode)
+						{
+						case EHairStrandsDebugMode::NoneDebug					: DebugModeScalar =99.f; break;
+						case EHairStrandsDebugMode::SimHairStrands				: DebugModeScalar = 0.f; break;
+						case EHairStrandsDebugMode::RenderHairStrands			: DebugModeScalar = 0.f; break;
+						case EHairStrandsDebugMode::RenderHairRootUV			: DebugModeScalar = 1.f; break;
+						case EHairStrandsDebugMode::RenderHairUV				: DebugModeScalar = 2.f; break;
+						case EHairStrandsDebugMode::RenderHairSeed				: DebugModeScalar = 3.f; break;
+						case EHairStrandsDebugMode::RenderHairDimension			: DebugModeScalar = 4.f; break;
+						case EHairStrandsDebugMode::RenderHairRadiusVariation	: DebugModeScalar = 5.f; break;
+						case EHairStrandsDebugMode::RenderHairRootUDIM			: DebugModeScalar = 6.f; break;
+						case EHairStrandsDebugMode::RenderHairBaseColor			: DebugModeScalar = 7.f; break;
+						case EHairStrandsDebugMode::RenderHairRoughness			: DebugModeScalar = 8.f; break;
+						case EHairStrandsDebugMode::RenderVisCluster			: DebugModeScalar = 0.f; break;
 						case EHairStrandsDebugMode::RenderHairGroup				: DebugModeScalar = 9.f; break;
-			};
+						};
 
-			// TODO: fix this as the radius is incorrect. This code run before the interpolation code, which is where HairRadius is updated.
-			float HairMaxRadius = 0;
-			for (FHairGroupInstance* Instance : Instances)
-			{
-				HairMaxRadius = FMath::Max(HairMaxRadius, Instance->Strands.Modifier.HairWidth * 0.5f);
-			}
+						// TODO: fix this as the radius is incorrect. This code run before the interpolation code, which is where HairRadius is updated.
+						float HairMaxRadius = 0;
+						for (FHairGroupInstance* Instance : Instances)
+						{
+							HairMaxRadius = FMath::Max(HairMaxRadius, Instance->Strands.Modifier.HairWidth * 0.5f);
+						}
 
 						const FVector HairGroupColor = FVector(GetHairGroupDebugColor(GroupIt));
-			const float HairClipLength = GetHairClipLength();
-						auto DebugMaterial = new FHairDebugModeMaterialRenderProxy(Strands_DebugMaterial ? Strands_DebugMaterial->GetRenderProxy() : nullptr, DebugModeScalar, 0, HairMaxRadius, HairClipLength, HairGroupColor);
-			Collector.RegisterOneFrameMaterialProxy(DebugMaterial);
-			Strands_MaterialProxy = DebugMaterial;
-		}
+						auto DebugMaterial = new FHairDebugModeMaterialRenderProxy(Strands_DebugMaterial ? Strands_DebugMaterial->GetRenderProxy() : nullptr, DebugModeScalar, 0, HairMaxRadius, HairGroupColor);
+						Collector.RegisterOneFrameMaterialProxy(DebugMaterial);
+						Strands_MaterialProxy = DebugMaterial;
+					}
 
 					FMeshBatch* MeshBatch = CreateMeshBatch(View, ViewFamily, Collector, HairGroups[GroupIt], Instances[GroupIt], GroupIt, Strands_MaterialProxy);
 					if (MeshBatch == nullptr)
@@ -1275,17 +1260,6 @@ void UGroomComponent::SetStableRasterization(bool bEnable)
 	UpdateHairGroupsDescAndInvalidateRenderState();
 }
 
-void UGroomComponent::SetHairLengthScale(float Scale) 
-{ 
-	Scale = FMath::Clamp(Scale, 0.f, 1.f);
-	for (FHairGroupDesc& HairDesc : GroomGroupsDesc)
-	{
-		HairDesc.HairClipScale = Scale;
-		HairDesc.HairClipScale_Override = true;
-	}
-	UpdateHairGroupsDescAndInvalidateRenderState();
-}
-
 void UGroomComponent::SetHairRootScale(float Scale)
 {
 	Scale = FMath::Clamp(Scale, 0.f, 10.f);
@@ -2863,7 +2837,6 @@ void UGroomComponent::PostEditChangeProperty(FPropertyChangedEvent& PropertyChan
 	if (PropertyName == GET_MEMBER_NAME_CHECKED(FHairGroupDesc, HairWidth) || 
 		PropertyName == GET_MEMBER_NAME_CHECKED(FHairGroupDesc, HairRootScale) ||
 		PropertyName == GET_MEMBER_NAME_CHECKED(FHairGroupDesc, HairTipScale) ||
-		PropertyName == GET_MEMBER_NAME_CHECKED(FHairGroupDesc, HairClipScale) ||
 		PropertyName == GET_MEMBER_NAME_CHECKED(FHairGroupDesc, HairShadowDensity) ||
 		PropertyName == GET_MEMBER_NAME_CHECKED(FHairGroupDesc, HairRaytracingRadiusScale) ||
 		PropertyName == GET_MEMBER_NAME_CHECKED(FHairGroupDesc, LODBias) ||
diff --git a/Engine/Plugins/Runtime/HairStrands/Source/HairStrandsCore/Private/HairStrandsRendering.cpp b/Engine/Plugins/Runtime/HairStrands/Source/HairStrandsCore/Private/HairStrandsRendering.cpp
--- a/Engine/Plugins/Runtime/HairStrands/Source/HairStrandsCore/Private/HairStrandsRendering.cpp
+++ b/Engine/Plugins/Runtime/HairStrands/Source/HairStrandsCore/Private/HairStrandsRendering.cpp
@@ -335,15 +335,13 @@ FRDGHairStrandsCullingData ImportCullingData(FRDGBuilder& GraphBuilder, FHairGro
 	return Out;
 }
 
-struct FHairScaleAndClipDesc
+struct FHairScaleDesc
 {
-	float InHairLength = 0;
 	float InHairRadius = 0;
 	float OutHairRadius = 0;
 	float MaxOutHairRadius = 0;
 	float HairRadiusRootScale = 1;
 	float HairRadiusTipScale = 1;
-	float HairLengthClip = 1;
 	bool  bEnable = true;
 
 	bool IsEnable() const 
@@ -352,8 +350,7 @@ struct FHairScaleAndClipDesc
 			bEnable && (
 			InHairRadius != OutHairRadius ||
 			HairRadiusRootScale != 1 ||
-			HairRadiusTipScale != 1 ||
-			HairLengthClip < 1);
+			HairRadiusTipScale != 1);
 	}
 };
 
@@ -385,8 +382,7 @@ class FHairInterpolationCS : public FGlobalShader
 		SHADER_PARAMETER(float, MaxOutHairRadius)
 		SHADER_PARAMETER(float, HairRadiusRootScale)
 		SHADER_PARAMETER(float, HairRadiusTipScale)
-		SHADER_PARAMETER(float, HairLengthClip)
-		SHADER_PARAMETER(uint32, HairLengthClipEnable)
+		SHADER_PARAMETER(uint32, HairScalingEnable)
 		SHADER_PARAMETER(uint32,  HairStrandsVF_bIsCullingEnable)
 
 		SHADER_PARAMETER(FMatrix, LocalToWorldMatrix)
@@ -464,7 +460,7 @@ static void AddHairStrandsInterpolationPass(
 	const FShaderDrawDebugData* ShaderDrawData,
 	const FHairGroupInstance* Instance,
 	const uint32 VertexCount,
-	const FHairScaleAndClipDesc ScaleAndClipDesc,
+	const FHairScaleDesc ScaleDesc,
 	const int32 MeshLODIndex,
 	const bool bPatchedAttributeBuffer,
 	const EHairInterpolationType HairInterpolationType,
@@ -513,14 +509,12 @@ static void AddHairStrandsInterpolationPass(
 	Parameters->DispatchCountX = DispatchCount.X;
 	Parameters->SimRootPointIndexBuffer = SimRootPointIndexBuffer;
 	
-	Parameters->InHairLength = ScaleAndClipDesc.InHairLength;
-	Parameters->InHairRadius = ScaleAndClipDesc.InHairRadius;
-	Parameters->OutHairRadius = ScaleAndClipDesc.OutHairRadius;
-	Parameters->MaxOutHairRadius = ScaleAndClipDesc.MaxOutHairRadius;
-	Parameters->HairRadiusRootScale = ScaleAndClipDesc.HairRadiusRootScale;
-	Parameters->HairRadiusTipScale = ScaleAndClipDesc.HairRadiusTipScale;
-	Parameters->HairLengthClip = ScaleAndClipDesc.HairLengthClip * ScaleAndClipDesc.InHairLength;
-	Parameters->HairLengthClipEnable = ScaleAndClipDesc.IsEnable() ? 1 : 0;
+	Parameters->InHairRadius = ScaleDesc.InHairRadius;
+	Parameters->OutHairRadius = ScaleDesc.OutHairRadius;
+	Parameters->MaxOutHairRadius = ScaleDesc.MaxOutHairRadius;
+	Parameters->HairRadiusRootScale = ScaleDesc.HairRadiusRootScale;
+	Parameters->HairRadiusTipScale = ScaleDesc.HairRadiusTipScale;
+	Parameters->HairScalingEnable = ScaleDesc.IsEnable() ? 1u : 0u;
 	Parameters->AttributeBuffer = RenderAttributeBuffer;	
 	
 	const bool bIsVertexToCurveBuffersValid_Sim =
@@ -1285,18 +1279,16 @@ EHairStrandsDebugMode GetHairStrandsGeometryDebugMode(FHairGroupInstance* Instan
 	return DebugMode;
 }
 
-FHairScaleAndClipDesc ComputeHairScaleAndClipDesc(FHairGroupInstance* Instance)
+FHairScaleDesc ComputeHairScaleDesc(FHairGroupInstance* Instance)
 {
-	FHairScaleAndClipDesc Out;
+	FHairScaleDesc Out;
 	Out.bEnable = true;
-	Out.InHairLength = Instance->Strands.Data->StrandsCurves.MaxLength;
 	Out.InHairRadius = Instance->Strands.Modifier.HairWidth * 0.5f;
 
 	Out.OutHairRadius = Out.InHairRadius;
 	Out.MaxOutHairRadius = Out.OutHairRadius;
 	Out.HairRadiusRootScale = 1.f;
 	Out.HairRadiusTipScale = 1.f;
-	Out.HairLengthClip = 1.f;
 
 	const bool bSupportDeformation = Instance->Strands.DeformedResource != nullptr;
 	if (bSupportDeformation)
@@ -1305,7 +1297,6 @@ FHairScaleAndClipDesc ComputeHairScaleAndClipDesc(FHairGroupInstance* Instance)
 		Out.MaxOutHairRadius = Out.OutHairRadius * FMath::Max(1.f, FMath::Max(Instance->Strands.Modifier.HairRootScale, Instance->Strands.Modifier.HairTipScale));
 		Out.HairRadiusRootScale = Instance->Strands.Modifier.HairRootScale;
 		Out.HairRadiusTipScale = Instance->Strands.Modifier.HairTipScale;
-		Out.HairLengthClip = FMath::Clamp(Instance->Strands.Modifier.HairClipScale, 0.f, 1.f);
 	}
 	
 	return Out;
@@ -1337,7 +1328,7 @@ static FHairGroupPublicData::FVertexFactoryInput InternalComputeHairStrandsVerte
 	if (!Instance || Instance->GeometryType != EHairGeometryType::Strands)
 		return OutVFInput;
 
-	const FHairScaleAndClipDesc ScaleAndClipDesc = ComputeHairScaleAndClipDesc(Instance);
+	const FHairScaleDesc ScaleDesc = ComputeHairScaleDesc(Instance);
 
 	const EHairStrandsDebugMode DebugMode = GetHairStrandsGeometryDebugMode(Instance);
 	const bool bDebugModePatchedAttributeBuffer = NeedsPatchAttributeBuffer(DebugMode);
@@ -1378,7 +1369,7 @@ static FHairGroupPublicData::FVertexFactoryInput InternalComputeHairStrandsVerte
 		OutVFInput.Strands.PositionOffset = Instance->Strands.DeformedResource->GetPositionOffset(FHairStrandsDeformedResource::EFrameType::Current);
 		OutVFInput.Strands.PrevPositionOffset = Instance->Strands.DeformedResource->GetPositionOffset(FHairStrandsDeformedResource::EFrameType::Previous);
 		OutVFInput.Strands.VertexCount = Instance->Strands.RestResource->GetVertexCount();
-		OutVFInput.Strands.HairRadius = ScaleAndClipDesc.MaxOutHairRadius;
+		OutVFInput.Strands.HairRadius = ScaleDesc.MaxOutHairRadius;
 		OutVFInput.Strands.HairLength = Instance->Strands.Modifier.HairLength;
 		OutVFInput.Strands.HairDensity = Instance->Strands.Modifier.HairShadowDensity;
 		OutVFInput.Strands.bScatterSceneLighting = Instance->Strands.Modifier.bScatterSceneLighting;
@@ -1398,7 +1389,7 @@ static FHairGroupPublicData::FVertexFactoryInput InternalComputeHairStrandsVerte
 		OutVFInput.Strands.PositionOffset = Instance->Strands.RestResource->PositionOffset;
 		OutVFInput.Strands.PrevPositionOffset = Instance->Strands.RestResource->PositionOffset;
 		OutVFInput.Strands.VertexCount = Instance->Strands.RestResource->GetVertexCount();
-		OutVFInput.Strands.HairRadius = ScaleAndClipDesc.MaxOutHairRadius;
+		OutVFInput.Strands.HairRadius = ScaleDesc.MaxOutHairRadius;
 		OutVFInput.Strands.HairLength = Instance->Strands.Modifier.HairLength;
 		OutVFInput.Strands.HairDensity = Instance->Strands.Modifier.HairShadowDensity;
 		OutVFInput.Strands.bScatterSceneLighting = Instance->Strands.Modifier.bScatterSceneLighting;
@@ -1453,7 +1444,7 @@ void ComputeHairStrandsInterpolation(
 		// * Render : Show rendering strands with sim color influence
 		const EHairStrandsDebugMode DebugMode = GetHairStrandsGeometryDebugMode(Instance);
 		const bool bDebugModePatchedAttributeBuffer = NeedsPatchAttributeBuffer(DebugMode);
-		const FHairScaleAndClipDesc ScaleAndClipDesc = ComputeHairScaleAndClipDesc(Instance);
+		const FHairScaleDesc ScaleDesc = ComputeHairScaleDesc(Instance);
 
 		if (DebugMode == EHairStrandsDebugMode::SimHairStrands)
 		{
@@ -1512,7 +1503,7 @@ void ComputeHairStrandsInterpolation(
 					ShaderDrawData,
 					Instance,
 					Instance->Strands.RestResource->GetVertexCount(),
-					ScaleAndClipDesc,
+					ScaleDesc,
 					MeshLODIndex,
 					bDebugModePatchedAttributeBuffer,
 					Instance->Strands.HairInterpolationType,
@@ -1595,7 +1586,7 @@ void ComputeHairStrandsInterpolation(
 						GraphBuilder,
 						ShaderMap,
 						Instance->Strands.RestResource->GetVertexCount(),
-						ScaleAndClipDesc.MaxOutHairRadius * HairRadiusScaleRT,
+						ScaleDesc.MaxOutHairRadius * HairRadiusScaleRT,
 						Strands_PositionOffsetSRV,
 						CullingData,
 						Strands_PositionSRV,
@@ -1656,15 +1647,13 @@ void ComputeHairStrandsInterpolation(
 				{
 					FRDGImportedBuffer Guides_DeformedPositionBuffer = Register(GraphBuilder, LOD.Guides.DeformedResource->GetBuffer(FHairStrandsDeformedResource::Current), ERDGImportedBufferFlags::CreateViews);
 
-					FHairScaleAndClipDesc ScaleAndClipDesc;
-					ScaleAndClipDesc.bEnable = false;
-					ScaleAndClipDesc.InHairLength = LOD.Guides.Data->StrandsCurves.MaxLength;
-					ScaleAndClipDesc.InHairRadius = LOD.Guides.Data->StrandsCurves.MaxRadius;
-					ScaleAndClipDesc.OutHairRadius = (GStrandHairWidth > 0 ? GStrandHairWidth * 0.5f : ScaleAndClipDesc.InHairRadius);
-					ScaleAndClipDesc.MaxOutHairRadius = ScaleAndClipDesc.OutHairRadius;
-					ScaleAndClipDesc.HairRadiusRootScale = 1;
-					ScaleAndClipDesc.HairRadiusTipScale = 1;
-					ScaleAndClipDesc.HairLengthClip = 1;
+					FHairScaleDesc ScaleDesc;
+					ScaleDesc.bEnable = false;
+					ScaleDesc.InHairRadius = LOD.Guides.Data->StrandsCurves.MaxRadius;
+					ScaleDesc.OutHairRadius = (GStrandHairWidth > 0 ? GStrandHairWidth * 0.5f : ScaleDesc.InHairRadius);
+					ScaleDesc.MaxOutHairRadius = ScaleDesc.OutHairRadius;
+					ScaleDesc.HairRadiusRootScale = 1;
+					ScaleDesc.HairRadiusTipScale = 1;
 
 					FRDGHairStrandsCullingData CullingData;
 					AddHairStrandsInterpolationPass(
@@ -1673,7 +1662,7 @@ void ComputeHairStrandsInterpolation(
 						ShaderDrawData,
 						Instance,
 						LOD.Guides.RestResource->GetVertexCount(),
-						ScaleAndClipDesc,
+						ScaleDesc,
 						MeshLODIndex,
 						false,
 						LOD.Guides.HairInterpolationType,
diff --git a/Engine/Plugins/Runtime/HairStrands/Source/HairStrandsCore/Public/GroomAssetRendering.h b/Engine/Plugins/Runtime/HairStrands/Source/HairStrandsCore/Public/GroomAssetRendering.h
--- a/Engine/Plugins/Runtime/HairStrands/Source/HairStrandsCore/Public/GroomAssetRendering.h
+++ b/Engine/Plugins/Runtime/HairStrands/Source/HairStrandsCore/Public/GroomAssetRendering.h
@@ -30,10 +30,6 @@ struct HAIRSTRANDSCORE_API FHairGeometrySettings
 	UPROPERTY(EditAnywhere, BlueprintReadWrite, Category = "GeometrySettings", AdvancedDisplay, meta = (ClampMin = "0.0001", UIMin = "0.001", UIMax = "2.0", SliderExponent = 6))
 	float HairTipScale;
 
-	/** Normalized hair clip length, i.e. at which length hair will be clipped. 1 means no clipping. 0 means hairs are fully clipped */
-	UPROPERTY(EditAnywhere, BlueprintReadWrite, Category = "GeometrySettings", AdvancedDisplay, meta = (ClampMin = "0.0", UIMin = "0.0", UIMax = "1.0", SliderExponent = 1))
-	float HairClipScale;
-
 	bool operator==(const FHairGeometrySettings& A) const;
 };
 
diff --git a/Engine/Plugins/Runtime/HairStrands/Source/HairStrandsCore/Public/GroomComponent.h b/Engine/Plugins/Runtime/HairStrands/Source/HairStrandsCore/Public/GroomComponent.h
--- a/Engine/Plugins/Runtime/HairStrands/Source/HairStrandsCore/Public/GroomComponent.h
+++ b/Engine/Plugins/Runtime/HairStrands/Source/HairStrandsCore/Public/GroomComponent.h
@@ -146,7 +146,6 @@ public:
 
 	void SetStableRasterization(bool bEnable);
 	void SetGroomAsset(UGroomAsset* Asset, UGroomBindingAsset* InBinding);
-	void SetHairLengthScale(float Scale);
 	void SetHairRootScale(float Scale);
 	void SetHairWidth(float HairWidth);
 	void SetScatterSceneLighting(bool Enable);
diff --git a/Engine/Plugins/Runtime/HairStrands/Source/HairStrandsCore/Public/GroomDesc.h b/Engine/Plugins/Runtime/HairStrands/Source/HairStrandsCore/Public/GroomDesc.h
--- a/Engine/Plugins/Runtime/HairStrands/Source/HairStrandsCore/Public/GroomDesc.h
+++ b/Engine/Plugins/Runtime/HairStrands/Source/HairStrandsCore/Public/GroomDesc.h
@@ -36,12 +36,6 @@ struct FHairGroupDesc
 	UPROPERTY(EditAnywhere, BlueprintReadWrite, Category = Overrides, meta = (PinHiddenByDefault, InlineEditConditionToggle))
 	bool HairTipScale_Override = false;
 
-	/** Normalized hair clip scale, i.e. at which normalized length hair will be clipped. 1 means no clipping. 0 means hairs are fully clipped */
-	UPROPERTY(EditAnywhere, BlueprintReadWrite, Category = "Groom", AdvancedDisplay, meta = (editcondition = "HairClipScale_Override", ClampMin = "0.0", UIMin = "0.0", UIMax = "1.0", SliderExponent = 1))
-	float HairClipScale = 1.0f;
-	UPROPERTY(EditAnywhere, BlueprintReadWrite, Category = Overrides, meta = (PinHiddenByDefault, InlineEditConditionToggle))
-	bool HairClipScale_Override = false;
-
 	/** Override the hair shadow density factor (unit less). */
 	UPROPERTY(EditAnywhere, BlueprintReadWrite, Category = "Groom", AdvancedDisplay, meta = (editcondition = "HairShadowDensity_Override", ClampMin = "0.0001", UIMin = "0.001", UIMax = "10.0", SliderExponent = 6))
 	float HairShadowDensity = 1.0f;
diff --git a/Engine/Plugins/Runtime/HairStrands/Source/HairStrandsEditor/Private/GroomComponentDetailsCustomization.cpp b/Engine/Plugins/Runtime/HairStrands/Source/HairStrandsEditor/Private/GroomComponentDetailsCustomization.cpp
--- a/Engine/Plugins/Runtime/HairStrands/Source/HairStrandsEditor/Private/GroomComponentDetailsCustomization.cpp
+++ b/Engine/Plugins/Runtime/HairStrands/Source/HairStrandsEditor/Private/GroomComponentDetailsCustomization.cpp
@@ -108,10 +108,6 @@ void FGroomComponentDetailsCustomization::OnResetToDefault(int32 GroupIndex, TSh
 	{
 		GroomComponentPtr->GroomGroupsDesc[GroupIndex].HairTipScale = Default.HairTipScale;
 	}
-	else if (PropertyName == GET_MEMBER_NAME_CHECKED(FHairGroupDesc, HairClipScale))
-	{
-		GroomComponentPtr->GroomGroupsDesc[GroupIndex].HairClipScale = Default.HairClipScale;
-	}
 	else if (PropertyName == GET_MEMBER_NAME_CHECKED(FHairGroupDesc, HairShadowDensity))
 	{
 		GroomComponentPtr->GroomGroupsDesc[GroupIndex].HairShadowDensity = Default.HairShadowDensity;
diff --git a/Engine/Shaders/Private/HairStrands/HairStrandsInterpolation.usf b/Engine/Shaders/Private/HairStrands/HairStrandsInterpolation.usf
--- a/Engine/Shaders/Private/HairStrands/HairStrandsInterpolation.usf
+++ b/Engine/Shaders/Private/HairStrands/HairStrandsInterpolation.usf
@@ -48,8 +48,7 @@ float OutHairRadius;
 float MaxOutHairRadius;
 float HairRadiusRootScale;
 float HairRadiusTipScale;
-float HairLengthClip;
-uint  HairLengthClipEnable;
+uint  HairScalingEnable;
 
 Buffer<uint4>  RenderRestPosePositionBuffer;
 Buffer<uint4>  SimRestPosePositionBuffer;
@@ -179,10 +178,9 @@ float3 ComputeRotateGeometryOffset(uint GuideIndex, float GuideVertexWeight,
 }
 
 // Compute the new hair width and return if the current vertex is valid (true), or if it needs to be clipped (false)
-bool ScaleAndClip(uint ControlPointId, inout uint InPackedAlpha)
+void ScaleStrand(uint ControlPointId, inout uint InPackedAlpha)
 {
-	bool bIsValid = true;
-	if (HairLengthClipEnable)
+	if (HairScalingEnable)
 	{
 		// Instead of doing the actual interpolation we could rely on dithering as an optimization
 		const FHairAttributes Attributes0 = DecodeHairAttribute(AttributeBuffer[ControlPointId]);
@@ -191,17 +189,12 @@ bool ScaleAndClip(uint ControlPointId, inout uint InPackedAlpha)
 		// todo: Move all these computations into shared function which can be shared with the hair VF
 		const float InNormVertexRadius = ((InPackedAlpha & 0x00FC) >> 2) / 63.f;
 		//const float InVertexRadius = InNormVertexRadius * InHairRadius;
-		const float InVertexLength = ((InPackedAlpha & 0xFF00) >> 8) / 255.f * InHairLength;
 
 		const float OutVertexRadius = OutHairRadius * InNormVertexRadius * lerp(HairRadiusRootScale, HairRadiusTipScale, U);
 		const uint PackedVertexRadius = (uint(saturate(OutVertexRadius / MaxOutHairRadius) * 63.f) & 0x3F) << 2;
 
 		InPackedAlpha = (InPackedAlpha & 0xFF03) | PackedVertexRadius;
-
-		bIsValid = InVertexLength <= HairLengthClip;
 	}
-	
-	return bIsValid;
 }
 
 [numthreads(GROUP_SIZE, 1, 1)]
@@ -285,8 +278,8 @@ void MainCS(uint2 DispatchThreadId : SV_DispatchThreadID)
 			}
 			#endif
 
-			const bool bIsValid = ScaleAndClip(VertexIndex, PackedRestPosition.w);
-			WorldVertexPosition = bIsValid ? ControlPoint + CurrOffset : float3(POSITIVE_INFINITY, POSITIVE_INFINITY, POSITIVE_INFINITY);
+			ScaleStrand(VertexIndex, PackedRestPosition.w);
+			WorldVertexPosition = ControlPoint + CurrOffset;
 
 			OutPosition = uint4(f32tof16(WorldVertexPosition - OutRenHairPositionOffset), PackedRestPosition.w);
 			OutRenderDeformedPositionBuffer[VertexIndex] = OutPosition;
@@ -337,8 +330,8 @@ void MainCS(uint2 DispatchThreadId : SV_DispatchThreadID)
 			ControlPoint = ToTriangle(ControlPoint,RootBarycentric,RestTriangle);
 			ControlPoint = ToWorld(ControlPoint,RootBarycentric,DeformedTriangle) + CurrOffset;
 
-			const bool bIsValid = ScaleAndClip(VertexIndex, PackedRestPosition.w);
-			WorldVertexPosition = bIsValid ? ControlPoint : float3(POSITIVE_INFINITY, POSITIVE_INFINITY, POSITIVE_INFINITY);
+			ScaleStrand(VertexIndex, PackedRestPosition.w);
+			WorldVertexPosition = ControlPoint;
 
 			#if PERMUTATION_DEBUG == 1 && ALLOW_DEBUG_RENDERING
 			{
@@ -411,8 +404,8 @@ void MainCS(uint2 DispatchThreadId : SV_DispatchThreadID)
 
 			ControlPoint += RootBarycentric.x * DeformedTriangle.P0 + RootBarycentric.y * DeformedTriangle.P1 + RootBarycentric.z * DeformedTriangle.P2;
 
-			const bool bIsValid = ScaleAndClip(VertexIndex, PackedRestPosition.w);
-			WorldVertexPosition = bIsValid ? ControlPoint + CurrOffset : float3(POSITIVE_INFINITY, POSITIVE_INFINITY, POSITIVE_INFINITY);
+			ScaleStrand(VertexIndex, PackedRestPosition.w);
+			WorldVertexPosition = ControlPoint + CurrOffset;
 
 			OutPosition = uint4(f32tof16(WorldVertexPosition - OutRenHairPositionOffset), PackedRestPosition.w);
 			OutRenderDeformedPositionBuffer[VertexIndex] = OutPosition;
@@ -462,8 +455,8 @@ void MainCS(uint2 DispatchThreadId : SV_DispatchThreadID)
 
 			ControlPoint =  RootBarycentric.x * DeformedTriangle.P0 + RootBarycentric.y * DeformedTriangle.P1 + RootBarycentric.z * DeformedTriangle.P2;
 
-			const bool bIsValid = ScaleAndClip(VertexIndex, PackedRestPosition.w);
-			WorldVertexPosition = bIsValid ? ControlPoint + CurrOffset : float3(POSITIVE_INFINITY, POSITIVE_INFINITY, POSITIVE_INFINITY);
+			ScaleStrand(VertexIndex, PackedRestPosition.w);
+			WorldVertexPosition = ControlPoint + CurrOffset;
 
 			OutPosition = uint4(f32tof16(WorldVertexPosition - OutRenHairPositionOffset), PackedRestPosition.w);
 			OutRenderDeformedPositionBuffer[VertexIndex] = OutPosition;
@@ -510,8 +503,8 @@ void MainCS(uint2 DispatchThreadId : SV_DispatchThreadID)
 			}
 			#endif
 
-			const bool bIsValid = ScaleAndClip(VertexIndex, PackedRestPosition.w);
-			WorldVertexPosition = bIsValid ? CurrOffset : float3(POSITIVE_INFINITY, POSITIVE_INFINITY, POSITIVE_INFINITY);
+			ScaleStrand(VertexIndex, PackedRestPosition.w);
+			WorldVertexPosition = CurrOffset;
 
 			OutPosition = uint4(f32tof16(WorldVertexPosition - OutRenHairPositionOffset), PackedRestPosition.w);
 			OutRenderDeformedPositionBuffer[VertexIndex] = OutPosition;
