Fix reset to default override.

diff --git a/Engine/Plugins/Runtime/HairStrands/Source/HairStrandsEditor/Private/GroomAssetDetails.cpp b/Engine/Plugins/Runtime/HairStrands/Source/HairStrandsEditor/Private/GroomAssetDetails.cpp
--- a/Engine/Plugins/Runtime/HairStrands/Source/HairStrandsEditor/Private/GroomAssetDetails.cpp
+++ b/Engine/Plugins/Runtime/HairStrands/Source/HairStrandsEditor/Private/GroomAssetDetails.cpp
@@ -698,25 +698,258 @@ bool FGroomRenderingDetails::IsStrandsMaterialPickerEnabled(int32 GroupIndex) co
 	return true;
 }
 
-static void ExpandStruct(TSharedPtr<IPropertyHandle>& PropertyHandle, IDetailChildrenBuilder& ChildrenBuilder)
+
+template <typename T>
+bool AssignIfDifferent(T& Dest, const T& Src, bool bSetValue)
+{
+	const bool bHasChanged = Dest != Src;
+	if (bHasChanged && bSetValue)
+	{
+		Dest = Src;
+	}
+	return bHasChanged;
+}
+	
+#define HAIR_RESET0(GroomMemberName, StructTypeName, MemberName)										{ if (PropertyName == GET_MEMBER_NAME_CHECKED(StructTypeName, MemberName)) { bHasChanged = AssignIfDifferent(GroomAsset->GroomMemberName[GroupIndex].MemberName, Default.MemberName, bSetValue); } }
+#define HAIR_RESET1(GroomMemberName, StructTypeName, StructMemberName, MemberName)						{ if (PropertyName == GET_MEMBER_NAME_CHECKED(StructTypeName, MemberName)) { bHasChanged = AssignIfDifferent(GroomAsset->GroomMemberName[GroupIndex].StructMemberName.MemberName, Default.MemberName, bSetValue); } }
+#define HAIR_RESET2(GroomMemberName, StructTypeName, StructMemberName, SubStructMemberName, MemberName)	{ if (PropertyName == GET_MEMBER_NAME_CHECKED(StructTypeName, MemberName)) { bHasChanged = AssignIfDifferent(GroomAsset->GroomMemberName[GroupIndex].StructMemberName.SubStructMemberName.MemberName, Default.MemberName, bSetValue); } }
+
+bool FGroomRenderingDetails::CommonResetToDefault(TSharedPtr<IPropertyHandle> ChildHandle, int32 GroupIndex, int32 LODIndex, bool bSetValue)
 {
+	bool bHasChanged = false;
+	if (ChildHandle == nullptr || GroomAsset == nullptr || GroupIndex < 0 || GroupIndex >= GroomAsset->GetNumHairGroups())
+	{
+		return bHasChanged;
+	}
+
+	FName PropertyName = ChildHandle->GetProperty()->GetFName();
+
+	// Hair strands
+	{
+		{
+			FHairGeometrySettings Default;
+			HAIR_RESET1(HairGroupsRendering, FHairGeometrySettings, GeometrySettings, HairWidth);
+			HAIR_RESET1(HairGroupsRendering, FHairGeometrySettings, GeometrySettings, HairRootScale);
+			HAIR_RESET1(HairGroupsRendering, FHairGeometrySettings, GeometrySettings, HairTipScale);
+		}
+
+		{
+			FHairShadowSettings Default;
+			HAIR_RESET1(HairGroupsRendering, FHairShadowSettings, ShadowSettings, bVoxelize);
+			HAIR_RESET1(HairGroupsRendering, FHairShadowSettings, ShadowSettings, bUseHairRaytracingGeometry);
+			HAIR_RESET1(HairGroupsRendering, FHairShadowSettings, ShadowSettings, HairRaytracingRadiusScale);
+			HAIR_RESET1(HairGroupsRendering, FHairShadowSettings, ShadowSettings, HairRaytracingRadiusScale);
+		}
+
+		{
+			FHairAdvancedRenderingSettings Default;
+			HAIR_RESET1(HairGroupsRendering, FHairAdvancedRenderingSettings, AdvancedSettings, bScatterSceneLighting);
+			HAIR_RESET1(HairGroupsRendering, FHairAdvancedRenderingSettings, AdvancedSettings, bUseStableRasterization);
+		}
+	}
+
+	// Interpolation
+	{
+		{
+			FHairDecimationSettings Default;
+			HAIR_RESET1(HairGroupsInterpolation, FHairDecimationSettings, DecimationSettings, CurveDecimation);
+			HAIR_RESET1(HairGroupsInterpolation, FHairDecimationSettings, DecimationSettings, VertexDecimation);
+		}
+
+		{
+			FHairInterpolationSettings Default;
+			HAIR_RESET1(HairGroupsInterpolation, FHairInterpolationSettings, InterpolationSettings, bOverrideGuides);
+			HAIR_RESET1(HairGroupsInterpolation, FHairInterpolationSettings, InterpolationSettings, HairToGuideDensity);
+			HAIR_RESET1(HairGroupsInterpolation, FHairInterpolationSettings, InterpolationSettings, bRandomizeGuide);
+			HAIR_RESET1(HairGroupsInterpolation, FHairInterpolationSettings, InterpolationSettings, bUseUniqueGuide);
+		}
+	}
+
+	// LODs
+	if (LODIndex>=0)
+	{
+		{
+			FHairGroupsLOD Default;
+			HAIR_RESET0(HairGroupsLOD, FHairGroupsLOD, ClusterWorldSize);
+			HAIR_RESET0(HairGroupsLOD, FHairGroupsLOD, ClusterScreenSizeScale);
+		}
+
+		{
+			FHairLODSettings Default;
+			HAIR_RESET1(HairGroupsLOD, FHairLODSettings, LODs[LODIndex], CurveDecimation);
+			HAIR_RESET1(HairGroupsLOD, FHairLODSettings, LODs[LODIndex], VertexDecimation);
+			HAIR_RESET1(HairGroupsLOD, FHairLODSettings, LODs[LODIndex], AngularThreshold);
+			HAIR_RESET1(HairGroupsLOD, FHairLODSettings, LODs[LODIndex], ScreenSize);
+			HAIR_RESET1(HairGroupsLOD, FHairLODSettings, LODs[LODIndex], ThicknessScale);
+			HAIR_RESET1(HairGroupsLOD, FHairLODSettings, LODs[LODIndex], GeometryType);
+		}
+	}
+
+	// Cards
+	{
+		{
+			FHairGroupsCardsSourceDescription Default;
+			HAIR_RESET0(HairGroupsCards, FHairGroupsCardsSourceDescription, SourceType);
+			HAIR_RESET0(HairGroupsCards, FHairGroupsCardsSourceDescription, ProceduralMesh);
+			HAIR_RESET0(HairGroupsCards, FHairGroupsCardsSourceDescription, ImportedMesh);
+			HAIR_RESET0(HairGroupsCards, FHairGroupsCardsSourceDescription, GroupIndex);
+			HAIR_RESET0(HairGroupsCards, FHairGroupsCardsSourceDescription, LODIndex);
+		}
+
+		{
+			FHairGroupCardsTextures Default;
+			HAIR_RESET1(HairGroupsCards, FHairGroupCardsTextures, Textures, DepthTexture);
+			HAIR_RESET1(HairGroupsCards, FHairGroupCardsTextures, Textures, CoverageTexture);
+			HAIR_RESET1(HairGroupsCards, FHairGroupCardsTextures, Textures, TangentTexture);
+			HAIR_RESET1(HairGroupsCards, FHairGroupCardsTextures, Textures, AttributeTexture);
+			HAIR_RESET1(HairGroupsCards, FHairGroupCardsTextures, Textures, AuxilaryDataTexture);
+		}
+	}
+
+	// Meshes
+	{
+		{
+			FHairGroupsMeshesSourceDescription Default;
+			HAIR_RESET0(HairGroupsMeshes, FHairGroupsMeshesSourceDescription, ImportedMesh);
+			HAIR_RESET0(HairGroupsMeshes, FHairGroupsMeshesSourceDescription, GroupIndex);
+			HAIR_RESET0(HairGroupsMeshes, FHairGroupsMeshesSourceDescription, LODIndex);
+		}
+
+		{
+			FHairGroupCardsTextures Default;
+			HAIR_RESET1(HairGroupsMeshes, FHairGroupCardsTextures, Textures, DepthTexture);
+			HAIR_RESET1(HairGroupsMeshes, FHairGroupCardsTextures, Textures, CoverageTexture);
+			HAIR_RESET1(HairGroupsMeshes, FHairGroupCardsTextures, Textures, TangentTexture);
+			HAIR_RESET1(HairGroupsMeshes, FHairGroupCardsTextures, Textures, AttributeTexture);
+			HAIR_RESET1(HairGroupsMeshes, FHairGroupCardsTextures, Textures, AuxilaryDataTexture);
+		}
+	}
+
+	// Physics
+	{
+		{
+			FHairSolverSettings Default;
+			HAIR_RESET1(HairGroupsPhysics, FHairSolverSettings, SolverSettings, EnableSimulation);
+			HAIR_RESET1(HairGroupsPhysics, FHairSolverSettings, SolverSettings, NiagaraSolver);
+			HAIR_RESET1(HairGroupsPhysics, FHairSolverSettings, SolverSettings, CustomSystem);
+			HAIR_RESET1(HairGroupsPhysics, FHairSolverSettings, SolverSettings, SubSteps);
+			HAIR_RESET1(HairGroupsPhysics, FHairSolverSettings, SolverSettings, IterationCount);
+		}
+
+		{
+			FHairExternalForces Default;
+			HAIR_RESET1(HairGroupsPhysics, FHairExternalForces, ExternalForces, GravityVector);
+			HAIR_RESET1(HairGroupsPhysics, FHairExternalForces, ExternalForces, AirDrag);
+			HAIR_RESET1(HairGroupsPhysics, FHairExternalForces, ExternalForces, AirVelocity);
+		}
+
+		{
+			FHairBendConstraint Default;
+			HAIR_RESET2(HairGroupsPhysics, FHairBendConstraint, MaterialConstraints, BendConstraint, SolveBend);
+			HAIR_RESET2(HairGroupsPhysics, FHairBendConstraint, MaterialConstraints, BendConstraint, ProjectBend);
+			HAIR_RESET2(HairGroupsPhysics, FHairBendConstraint, MaterialConstraints, BendConstraint, BendDamping);
+			HAIR_RESET2(HairGroupsPhysics, FHairBendConstraint, MaterialConstraints, BendConstraint, BendStiffness);
+//			HAIR_RESET2(HairGroupsPhysics, FHairBendConstraint, MaterialConstraints, BendConstraint, BendScale);
+		}
+
+		{
+			FHairStretchConstraint Default;
+			HAIR_RESET2(HairGroupsPhysics, FHairStretchConstraint, MaterialConstraints, StretchConstraint, SolveStretch);
+			HAIR_RESET2(HairGroupsPhysics, FHairStretchConstraint, MaterialConstraints, StretchConstraint, ProjectStretch);
+			HAIR_RESET2(HairGroupsPhysics, FHairStretchConstraint, MaterialConstraints, StretchConstraint, StretchDamping);
+			HAIR_RESET2(HairGroupsPhysics, FHairStretchConstraint, MaterialConstraints, StretchConstraint, StretchStiffness);
+//			HAIR_RESET2(HairGroupsPhysics, FHairStretchConstraint, MaterialConstraints, StretchConstraint, StretchScale);
+		}
+
+		{
+			FHairCollisionConstraint Default;
+			HAIR_RESET2(HairGroupsPhysics, FHairCollisionConstraint, MaterialConstraints, CollisionConstraint, SolveCollision);
+			HAIR_RESET2(HairGroupsPhysics, FHairCollisionConstraint, MaterialConstraints, CollisionConstraint, ProjectCollision);
+			HAIR_RESET2(HairGroupsPhysics, FHairCollisionConstraint, MaterialConstraints, CollisionConstraint, StaticFriction);
+			HAIR_RESET2(HairGroupsPhysics, FHairCollisionConstraint, MaterialConstraints, CollisionConstraint, KineticFriction);
+			HAIR_RESET2(HairGroupsPhysics, FHairCollisionConstraint, MaterialConstraints, CollisionConstraint, StrandsViscosity);
+			HAIR_RESET2(HairGroupsPhysics, FHairCollisionConstraint, MaterialConstraints, CollisionConstraint, GridDimension);
+			HAIR_RESET2(HairGroupsPhysics, FHairCollisionConstraint, MaterialConstraints, CollisionConstraint, CollisionRadius);
+//			HAIR_RESET2(HairGroupsPhysics, FHairCollisionConstraint, MaterialConstraints, CollisionConstraint, RadiusScale);
+		}
+		
+
+		{
+			FHairStrandsParameters Default;
+			HAIR_RESET1(HairGroupsPhysics, FHairStrandsParameters, StrandsParameters, StrandsSize);
+			HAIR_RESET1(HairGroupsPhysics, FHairStrandsParameters, StrandsParameters, StrandsDensity);
+			HAIR_RESET1(HairGroupsPhysics, FHairStrandsParameters, StrandsParameters, StrandsSmoothing);
+			HAIR_RESET1(HairGroupsPhysics, FHairStrandsParameters, StrandsParameters, StrandsThickness);
+//			HAIR_RESET1(HairGroupsPhysics, FHairStrandsParameters, StrandsParameters, ThicknessScale);
+		}
+	}
+
+	if (bSetValue && bHasChanged)
+	{
+		FPropertyChangedEvent PropertyChangedEvent(ChildHandle->GetProperty(), EPropertyChangeType::ValueSet);
+		GroomAsset->PostEditChangeProperty(PropertyChangedEvent);
+	}
+	return bHasChanged;
+}
+
+bool FGroomRenderingDetails::ShouldResetToDefault(TSharedPtr<IPropertyHandle> ChildHandle, int32 GroupIndex, int32 LODIndex)
+{
+	return CommonResetToDefault(ChildHandle, GroupIndex, LODIndex, false);
+}
+
+void FGroomRenderingDetails::ResetToDefault(TSharedPtr<IPropertyHandle> ChildHandle, int32 GroupIndex, int32 LODIndex)
+{
+	CommonResetToDefault(ChildHandle, GroupIndex, LODIndex, true);
+}
+
+void FGroomRenderingDetails::AddPropertyWithCustomReset(TSharedPtr<IPropertyHandle>& PropertyHandle, IDetailChildrenBuilder& Builder, int32 GroupIndex, int32 LODIndex)
+{	
+	FIsResetToDefaultVisible IsResetVisible = FIsResetToDefaultVisible::CreateSP(this, &FGroomRenderingDetails::ShouldResetToDefault, GroupIndex, LODIndex);
+	FResetToDefaultHandler ResetHandler = FResetToDefaultHandler::CreateSP(this, &FGroomRenderingDetails::ResetToDefault, GroupIndex, LODIndex);
+	FResetToDefaultOverride ResetOverride = FResetToDefaultOverride::Create(IsResetVisible, ResetHandler);
+	Builder.AddProperty(PropertyHandle.ToSharedRef()).OverrideResetToDefault(ResetOverride);
+}
+
+void FGroomRenderingDetails::ExpandStruct(TSharedPtr<IPropertyHandle>& PropertyHandle, IDetailChildrenBuilder& ChildrenBuilder, int32 GroupIndex, int32 LODIndex, bool bOverrideReset)
+{
+
 	uint32 ChildrenCount = 0;
 	PropertyHandle->GetNumChildren(ChildrenCount);
 	for (uint32 ChildIt = 0; ChildIt < ChildrenCount; ++ChildIt)
 	{
 		TSharedPtr<IPropertyHandle> ChildHandle = PropertyHandle->GetChildHandle(ChildIt);
-		ChildrenBuilder.AddProperty(ChildHandle.ToSharedRef());
+		if (bOverrideReset)
+		{
+			FIsResetToDefaultVisible IsResetVisible = FIsResetToDefaultVisible::CreateSP(this, &FGroomRenderingDetails::ShouldResetToDefault, GroupIndex, LODIndex);
+			FResetToDefaultHandler ResetHandler     = FResetToDefaultHandler::CreateSP(this, &FGroomRenderingDetails::ResetToDefault, GroupIndex, LODIndex);
+			FResetToDefaultOverride ResetOverride   = FResetToDefaultOverride::Create(IsResetVisible, ResetHandler);
+			ChildrenBuilder.AddProperty(ChildHandle.ToSharedRef()).OverrideResetToDefault(ResetOverride);
+		}
+		else
+		{
+			ChildrenBuilder.AddProperty(ChildHandle.ToSharedRef());
+		}
 	}
 }
 
-static void ExpandStruct(TSharedRef<IPropertyHandle>& PropertyHandle, IDetailChildrenBuilder& ChildrenBuilder)
+void FGroomRenderingDetails::ExpandStruct(TSharedRef<IPropertyHandle>& PropertyHandle, IDetailChildrenBuilder& ChildrenBuilder, int32 GroupIndex, int32 LODIndex, bool bOverrideReset)
 {
 	uint32 ChildrenCount = 0;
 	PropertyHandle->GetNumChildren(ChildrenCount);
 	for (uint32 ChildIt = 0; ChildIt < ChildrenCount; ++ChildIt)
 	{
 		TSharedPtr<IPropertyHandle> ChildHandle = PropertyHandle->GetChildHandle(ChildIt);
-		ChildrenBuilder.AddProperty(ChildHandle.ToSharedRef());
+		if (bOverrideReset)
+		{
+			FIsResetToDefaultVisible IsResetVisible = FIsResetToDefaultVisible::CreateSP(this, &FGroomRenderingDetails::ShouldResetToDefault, GroupIndex, LODIndex);
+			FResetToDefaultHandler ResetHandler = FResetToDefaultHandler::CreateSP(this, &FGroomRenderingDetails::ResetToDefault, GroupIndex, LODIndex);
+			FResetToDefaultOverride ResetOverride = FResetToDefaultOverride::Create(IsResetVisible, ResetHandler);
+			ChildrenBuilder.AddProperty(ChildHandle.ToSharedRef()).OverrideResetToDefault(ResetOverride);
+		}
+		else
+		{
+			ChildrenBuilder.AddProperty(ChildHandle.ToSharedRef());
+		}
 	}
 }
 
@@ -788,7 +1021,7 @@ FReply FGroomRenderingDetails::OnSaveCards(int32 DescIndex, FProperty* Property)
 
 void FGroomRenderingDetails::AddLODSlot(TSharedRef<IPropertyHandle>& LODHandle, IDetailChildrenBuilder& ChildrenBuilder, int32 GroupIndex, int32 LODIndex)
 {	
-	ExpandStruct(LODHandle, ChildrenBuilder);
+	ExpandStruct(LODHandle, ChildrenBuilder, GroupIndex, LODIndex, true);
 }
 
 void FGroomRenderingDetails::OnGenerateElementForLODs(TSharedRef<IPropertyHandle> StructProperty, int32 LODIndex, IDetailChildrenBuilder& ChildrenBuilder, IDetailLayoutBuilder* DetailLayout, int32 GroupIndex)
@@ -839,7 +1072,7 @@ void FGroomRenderingDetails::OnGenerateElementForLODs(TSharedRef<IPropertyHandle
 
 	// Rename the array entry name by its group name and adds all its existing properties
 	StructProperty->SetPropertyDisplayName(LOCTEXT("LODProperties", "LOD Properties"));
-	ExpandStruct(StructProperty, ChildrenBuilder);
+	ExpandStruct(StructProperty, ChildrenBuilder, GroupIndex, LODIndex, true);
 }
 
 TSharedRef<SWidget> FGroomRenderingDetails::MakeGroupNameButtonCustomization(int32 GroupIndex, FProperty* Property)
@@ -1041,7 +1274,7 @@ void FGroomRenderingDetails::OnGenerateElementForHairGroup(TSharedRef<IPropertyH
 				PropertyName == GET_MEMBER_NAME_CHECKED(FHairGroupsRendering, ShadowSettings) ||
 				PropertyName == GET_MEMBER_NAME_CHECKED(FHairGroupsRendering, AdvancedSettings))
 			{
-				ExpandStruct(ChildHandle, ChildrenBuilder);
+				ExpandStruct(ChildHandle, ChildrenBuilder, GroupIndex, -1, true);
 			}
 			else
 			{
@@ -1054,7 +1287,7 @@ void FGroomRenderingDetails::OnGenerateElementForHairGroup(TSharedRef<IPropertyH
 			if (PropertyName == GET_MEMBER_NAME_CHECKED(FHairGroupsCardsSourceDescription, SourceType))
 			{
 				// Add the Source type selection and just below add a save button
-				ChildrenBuilder.AddProperty(ChildHandle.ToSharedRef());
+				AddPropertyWithCustomReset(ChildHandle, ChildrenBuilder, GroupIndex, -1);
 				if (GroomAsset != nullptr && GroupIndex >= 0 && GroupIndex < GroomAsset->HairGroupsCards.Num() && (PanelType == EMaterialPanelType::Cards))
 				{
 					FText ToolTipTextForGeneration(FText::FromString(TEXT("Generate procedural cards data (meshes and textures) based on current procedural settings. Cards generation needs to run prior to the (re)loading of the cards data.")));
@@ -1103,13 +1336,13 @@ void FGroomRenderingDetails::OnGenerateElementForHairGroup(TSharedRef<IPropertyH
 			}
 			else
 			{
-				ChildrenBuilder.AddProperty(ChildHandle.ToSharedRef());
+				AddPropertyWithCustomReset(ChildHandle, ChildrenBuilder, GroupIndex, -1);
 			}
 		}
 		break;
 		case EMaterialPanelType::Meshes:
 		{
-			ChildrenBuilder.AddProperty(ChildHandle.ToSharedRef());
+			AddPropertyWithCustomReset(ChildHandle, ChildrenBuilder, GroupIndex, -1);
 		}
 		break;
 		case EMaterialPanelType::Interpolation:
@@ -1117,11 +1350,11 @@ void FGroomRenderingDetails::OnGenerateElementForHairGroup(TSharedRef<IPropertyH
 			if (PropertyName == GET_MEMBER_NAME_CHECKED(FHairGroupsInterpolation, DecimationSettings) ||
 				PropertyName == GET_MEMBER_NAME_CHECKED(FHairGroupsInterpolation, InterpolationSettings))
 			{
-				ExpandStruct(ChildHandle, ChildrenBuilder);
+				ExpandStruct(ChildHandle, ChildrenBuilder, GroupIndex, -1, true);
 			}
 			else
 			{
-				ChildrenBuilder.AddProperty(ChildHandle.ToSharedRef());
+				AddPropertyWithCustomReset(ChildHandle, ChildrenBuilder, GroupIndex, -1);
 			}
 		}
 		break;
@@ -1137,13 +1370,13 @@ void FGroomRenderingDetails::OnGenerateElementForHairGroup(TSharedRef<IPropertyH
 			}
 			else
 			{
-				ChildrenBuilder.AddProperty(ChildHandle.ToSharedRef());
+				AddPropertyWithCustomReset(ChildHandle, ChildrenBuilder, GroupIndex, -1);
 			}
 		}
 		break;
 		case EMaterialPanelType::Physics:
 		{
-			ChildrenBuilder.AddProperty(ChildHandle.ToSharedRef());
+			AddPropertyWithCustomReset(ChildHandle, ChildrenBuilder, GroupIndex, -1);
 		}
 		break;
 		default:
diff --git a/Engine/Plugins/Runtime/HairStrands/Source/HairStrandsEditor/Private/GroomAssetDetails.h b/Engine/Plugins/Runtime/HairStrands/Source/HairStrandsEditor/Private/GroomAssetDetails.h
--- a/Engine/Plugins/Runtime/HairStrands/Source/HairStrandsEditor/Private/GroomAssetDetails.h
+++ b/Engine/Plugins/Runtime/HairStrands/Source/HairStrandsEditor/Private/GroomAssetDetails.h
@@ -82,6 +82,14 @@ private:
 	FReply OnRefreshCards(int32 GroupIndex, FProperty* Property);
 	FReply OnSaveCards(int32 GroupIndex, FProperty* Property);
 
+	void ExpandStruct(TSharedPtr<IPropertyHandle>& PropertyHandle, IDetailChildrenBuilder& ChildrenBuilder, int32 GroupIndex, int32 LODIndex, bool bOverrideReset);
+	void ExpandStruct(TSharedRef<IPropertyHandle>& PropertyHandle, IDetailChildrenBuilder& ChildrenBuilder, int32 GroupIndex, int32 LODIndex, bool bOverrideReset);
+	void AddPropertyWithCustomReset(TSharedPtr<IPropertyHandle>& PropertyHandle, IDetailChildrenBuilder& Builder, int32 GroupIndex, int32 LODIndex);
+
+	bool CommonResetToDefault(TSharedPtr<IPropertyHandle> ChildHandle, int32 GroupIndex, int32 LODIndex, bool bSetValue);
+	bool ShouldResetToDefault(TSharedPtr<IPropertyHandle> ChildHandle, int32 GroupIndex, int32 LODInex);
+	void ResetToDefault(TSharedPtr<IPropertyHandle> ChildHandle, int32 GroupIndex, int32 LODIndex);
+
 private:
 	IDetailLayoutBuilder* GroomDetailLayout = nullptr;
 	EMaterialPanelType PanelType = EMaterialPanelType::Strands;
