Fix groom property reset on component.

diff --git a/Engine/Plugins/Runtime/HairStrands/Source/HairStrandsEditor/Private/GroomComponentDetailsCustomization.cpp b/Engine/Plugins/Runtime/HairStrands/Source/HairStrandsEditor/Private/GroomComponentDetailsCustomization.cpp
--- a/Engine/Plugins/Runtime/HairStrands/Source/HairStrandsEditor/Private/GroomComponentDetailsCustomization.cpp
+++ b/Engine/Plugins/Runtime/HairStrands/Source/HairStrandsEditor/Private/GroomComponentDetailsCustomization.cpp
@@ -86,51 +86,64 @@ void FGroomComponentDetailsCustomization::CustomizeDescGroupProperties(IDetailLa
 	}
 }
 
-void FGroomComponentDetailsCustomization::OnResetToDefault(int32 GroupIndex, TSharedPtr<IPropertyHandle> ChildHandle)
+template <typename T>
+bool AssignIfDifferentHairComponent(T& Dest, const T& Src, bool bSetValue)
 {
+	const bool bHasChanged = Dest != Src;
+	if (bHasChanged && bSetValue)
+	{
+		Dest = Src;
+	}
+	return bHasChanged;
+}
+#define HAIR_RESET_COMPONENT(MemberName) { if (PropertyName == GET_MEMBER_NAME_CHECKED(FHairGroupDesc, MemberName)) { bHasChanged = AssignIfDifferentHairComponent(GroomComponentPtr->GroomGroupsDesc[GroupIndex].MemberName, Default.MemberName, bSetValue); } }
+
+bool FGroomComponentDetailsCustomization::CommonResetToDefault(TSharedPtr<IPropertyHandle> ChildHandle, int32 GroupIndex, bool bSetValue)
+{
+	bool bHasChanged = false;
 	if (ChildHandle == nullptr || GroomComponentPtr == nullptr || GroupIndex < 0 || GroupIndex >= GroomComponentPtr->GroomGroupsDesc.Num())
 	{
-		return;
+		return bHasChanged;
 	}
 	
 	FName PropertyName = ChildHandle->GetProperty()->GetFName();
 
 	FHairGroupDesc Default;
-	if (PropertyName == GET_MEMBER_NAME_CHECKED(FHairGroupDesc, HairWidth))
-	{
-		GroomComponentPtr->GroomGroupsDesc[GroupIndex].HairWidth = Default.HairWidth;
-	}
-	else if (PropertyName == GET_MEMBER_NAME_CHECKED(FHairGroupDesc, HairRootScale))
-	{
-		GroomComponentPtr->GroomGroupsDesc[GroupIndex].HairRootScale = Default.HairRootScale;
-	}
-	else if (PropertyName == GET_MEMBER_NAME_CHECKED(FHairGroupDesc, HairTipScale))
-	{
-		GroomComponentPtr->GroomGroupsDesc[GroupIndex].HairTipScale = Default.HairTipScale;
-	}
-	else if (PropertyName == GET_MEMBER_NAME_CHECKED(FHairGroupDesc, HairShadowDensity))
-	{
-		GroomComponentPtr->GroomGroupsDesc[GroupIndex].HairShadowDensity = Default.HairShadowDensity;
-	}
-	else if (PropertyName == GET_MEMBER_NAME_CHECKED(FHairGroupDesc, HairRaytracingRadiusScale))
-	{
-		GroomComponentPtr->GroomGroupsDesc[GroupIndex].HairRaytracingRadiusScale = Default.HairRaytracingRadiusScale;
-	}
-	else if (PropertyName == GET_MEMBER_NAME_CHECKED(FHairGroupDesc, bUseHairRaytracingGeometry))
-	{
-		GroomComponentPtr->GroomGroupsDesc[GroupIndex].bUseHairRaytracingGeometry = Default.bUseHairRaytracingGeometry;
-	}
-	else if (PropertyName == GET_MEMBER_NAME_CHECKED(FHairGroupDesc, bUseStableRasterization))
-	{
-		GroomComponentPtr->GroomGroupsDesc[GroupIndex].bUseStableRasterization = Default.bUseStableRasterization;
-	}
-	else if (PropertyName == GET_MEMBER_NAME_CHECKED(FHairGroupDesc, bScatterSceneLighting))
+	HAIR_RESET_COMPONENT(HairWidth);
+	HAIR_RESET_COMPONENT(HairRootScale);
+	HAIR_RESET_COMPONENT(HairTipScale);
+	HAIR_RESET_COMPONENT(HairShadowDensity);
+	HAIR_RESET_COMPONENT(HairRaytracingRadiusScale);
+	HAIR_RESET_COMPONENT(bUseHairRaytracingGeometry);
+	HAIR_RESET_COMPONENT(bUseStableRasterization);
+	HAIR_RESET_COMPONENT(bScatterSceneLighting);
+	HAIR_RESET_COMPONENT(LODBias);
+
+	if (bSetValue && bHasChanged)
 	{
-		GroomComponentPtr->GroomGroupsDesc[GroupIndex].bScatterSceneLighting = Default.bScatterSceneLighting;
+		FScopedTransaction ScopedTransaction(NSLOCTEXT("UnrealEd", "PropertyWindowResetToDefault", "Reset to Default"));
+		GroomComponentPtr->UpdateHairGroupsDescAndInvalidateRenderState();
 	}
 
-	FScopedTransaction ScopedTransaction(NSLOCTEXT("UnrealEd", "PropertyWindowResetToDefault", "Reset to Default"));
-	GroomComponentPtr->UpdateHairGroupsDescAndInvalidateRenderState();
+	return bHasChanged;
+}
+
+bool FGroomComponentDetailsCustomization::ShouldResetToDefault(TSharedPtr<IPropertyHandle> ChildHandle, int32 GroupIndex)
+{
+	return CommonResetToDefault(ChildHandle, GroupIndex, false);
+}
+
+void FGroomComponentDetailsCustomization::ResetToDefault(TSharedPtr<IPropertyHandle> ChildHandle, int32 GroupIndex)
+{
+	CommonResetToDefault(ChildHandle, GroupIndex, true);
+}
+
+void FGroomComponentDetailsCustomization::AddPropertyWithCustomReset(TSharedPtr<IPropertyHandle>& PropertyHandle, IDetailChildrenBuilder& Builder, int32 GroupIndex)
+{
+	FIsResetToDefaultVisible IsResetVisible = FIsResetToDefaultVisible::CreateSP(this, &FGroomComponentDetailsCustomization::ShouldResetToDefault, GroupIndex);
+	FResetToDefaultHandler ResetHandler = FResetToDefaultHandler::CreateSP(this, &FGroomComponentDetailsCustomization::ResetToDefault, GroupIndex);
+	FResetToDefaultOverride ResetOverride = FResetToDefaultOverride::Create(IsResetVisible, ResetHandler);
+	Builder.AddProperty(PropertyHandle.ToSharedRef()).OverrideResetToDefault(ResetOverride);
 }
 
 static TSharedRef<SUniformGridPanel> MakeHairInfoGrid(const FSlateFontInfo& DetailFontInfo, const FHairGroupInfo& Infos)
@@ -244,8 +257,7 @@ void FGroomComponentDetailsCustomization::OnGenerateElementForHairGroup(TSharedR
 	for (uint32 ChildIt = 0; ChildIt < ChildrenCount; ++ChildIt)
 	{
 		TSharedPtr<IPropertyHandle> ChildHandle = StructProperty->GetChildHandle(ChildIt);
-		ChildHandle->SetOnPropertyResetToDefault(FSimpleDelegate::CreateSP(this, &FGroomComponentDetailsCustomization::OnResetToDefault, GroupIndex, ChildHandle));
-		ChildrenBuilder.AddProperty(ChildHandle.ToSharedRef());
+		AddPropertyWithCustomReset(ChildHandle, ChildrenBuilder, GroupIndex);
 	}
 }
 
diff --git a/Engine/Plugins/Runtime/HairStrands/Source/HairStrandsEditor/Private/GroomComponentDetailsCustomization.h b/Engine/Plugins/Runtime/HairStrands/Source/HairStrandsEditor/Private/GroomComponentDetailsCustomization.h
--- a/Engine/Plugins/Runtime/HairStrands/Source/HairStrandsEditor/Private/GroomComponentDetailsCustomization.h
+++ b/Engine/Plugins/Runtime/HairStrands/Source/HairStrandsEditor/Private/GroomComponentDetailsCustomization.h
@@ -30,5 +30,9 @@ private:
 
 	void CustomizeDescGroupProperties(IDetailLayoutBuilder& DetailLayout, IDetailCategoryBuilder& StrandsGroupFilesCategory);
 	void OnGenerateElementForHairGroup(TSharedRef<IPropertyHandle> StructProperty, int32 GroupIndex, IDetailChildrenBuilder& ChildrenBuilder, IDetailLayoutBuilder* DetailLayout);
-	void OnResetToDefault(int32 GroupIndex, TSharedPtr<IPropertyHandle> ChildHandle);
+
+	bool CommonResetToDefault(TSharedPtr<IPropertyHandle> ChildHandle, int32 GroupIndex, bool bSetValue);
+	bool ShouldResetToDefault(TSharedPtr<IPropertyHandle> ChildHandle, int32 GroupIndex);
+	void ResetToDefault(TSharedPtr<IPropertyHandle> ChildHandle, int32 GroupIndex);
+	void AddPropertyWithCustomReset(TSharedPtr<IPropertyHandle>& PropertyHandle, IDetailChildrenBuilder& Builder, int32 GroupIndex);
 };
