Change hair radius scaling to be done into VF, rather than interpolation in order to support non-interpolated hair.

This allows to fix root/tip scaling when a groom does not need interpolation (i.e., no-skinning & no-simulation).

diff --git a/Engine/Plugins/Runtime/HairStrands/Source/HairStrandsCore/Private/GroomBuilder.cpp b/Engine/Plugins/Runtime/HairStrands/Source/HairStrandsCore/Private/GroomBuilder.cpp
--- a/Engine/Plugins/Runtime/HairStrands/Source/HairStrandsCore/Private/GroomBuilder.cpp
+++ b/Engine/Plugins/Runtime/HairStrands/Source/HairStrandsCore/Private/GroomBuilder.cpp
@@ -32,7 +32,7 @@ static FAutoConsoleVariableRef CVarHairClusterBuilder_MaxVoxelResolution(TEXT("r
 FString FGroomBuilder::GetVersion()
 {
 	// Important to update the version when groom building changes
-	return TEXT("1c");
+	return TEXT("1g");
 }
 
 namespace HairStrandsBuilder
@@ -246,13 +246,13 @@ namespace HairStrandsBuilder
 				CopyVectorToPosition(PointPosition - HairBoxCenter, PackedPosition);
 				PackedPosition.ControlPointType = (PointIndex == 0) ? 1u : (PointIndex == (PointCount - 1) ? 2u : 0u);
 				PackedPosition.NormalizedRadius = uint8(FMath::Clamp(NormalizedRadius * 63.f, 0.f, 63.f));
-				PackedPosition.NormalizedLength = uint8(FMath::Clamp(NormalizedLength *255.f, 0.f, 255.f));
+				PackedPosition.UCoord = uint8(FMath::Clamp(CoordU * 255.f, 0.f, 255.f));
 
 				const FVector2D RootUV = Curves.CurvesRootUV[CurveIndex];
 				FHairStrandsAttributeFormat::Type& PackedAttributes = OutPackedAttributes[PointIndex + IndexOffset];
 				PackedAttributes.Unused0 = 0;
 				PackedAttributes.Unused1 = 0;
-				PackedAttributes.UCoord = FMath::Clamp(CoordU * 255.f, 0.f, 255.f);
+				PackedAttributes.NormalizedLength = uint8(FMath::Clamp(NormalizedLength * 255.f, 0.f, 255.f));
 				PackedAttributes.Seed = CurveSeed;
 
 				// Root UV support UDIM texture coordinate but limit the spans of the UDIM to be in 256x256 instead of 9999x9999.
diff --git a/Engine/Plugins/Runtime/HairStrands/Source/HairStrandsCore/Private/HairStrandsDatas.cpp b/Engine/Plugins/Runtime/HairStrands/Source/HairStrandsCore/Private/HairStrandsDatas.cpp
--- a/Engine/Plugins/Runtime/HairStrands/Source/HairStrandsCore/Private/HairStrandsDatas.cpp
+++ b/Engine/Plugins/Runtime/HairStrands/Source/HairStrandsCore/Private/HairStrandsDatas.cpp
@@ -2,6 +2,7 @@
 
 #include "HairStrandsDatas.h"
 #include "UObject/ReleaseObjectVersion.h"
+#include "UObject/UE5ReleaseStreamObjectVersion.h"
 
 void FHairStrandsInterpolationDatas::SetNum(const uint32 NumCurves)
 {
@@ -66,7 +67,7 @@ FArchive& operator<<(FArchive& Ar, FPackedHairVertex& Vertex)
 	Ar << Vertex.Y;
 	Ar << Vertex.Z;
 
-	Ar << Vertex.NormalizedLength;
+	Ar << Vertex.UCoord;
 
 	if (Ar.IsLoading())
 	{
@@ -89,7 +90,7 @@ FArchive& operator<<(FArchive& Ar, FPackedHairAttributeVertex& Vertex)
 
 	Ar << Vertex.RootU;
 	Ar << Vertex.RootV;
-	Ar << Vertex.UCoord;
+	Ar << Vertex.NormalizedLength;
 	Ar << Vertex.Seed;
 
 	if (Ar.CustomVer(FReleaseObjectVersion::GUID) >= FReleaseObjectVersion::GroomAssetVersion1)
@@ -234,9 +235,25 @@ void FHairStrandsCurves::Serialize(FArchive& Ar)
 void FHairStrandsDatas::FRenderData::Serialize(FArchive& Ar)
 {
 	Ar.UsingCustomVersion(FReleaseObjectVersion::GUID);
+	Ar.UsingCustomVersion(FUE5ReleaseStreamObjectVersion::GUID);
 
 	Ar << Positions;
 	Ar << Attributes;
+
+	if (Ar.CustomVer(FUE5ReleaseStreamObjectVersion::GUID) < FUE5ReleaseStreamObjectVersion::HairStrandsVertexFormatChange && Ar.IsLoading())
+	{
+		check(Positions.Num() == Attributes.Num());
+		const uint32 VertexCount = Positions.Num();
+		for (uint32 VertexIt = 0; VertexIt < VertexCount; ++VertexIt)
+		{
+			// Prior to FUE5ReleaseStreamObjectVersion::HairStrandsVertexFormatChange, UCoord and NormalizedLength were swapped
+			const uint8 UCoord = Attributes[VertexIt].NormalizedLength;
+			const uint8 NormalizedLength = Positions[VertexIt].UCoord;
+			Attributes[VertexIt].NormalizedLength = NormalizedLength;
+			Positions[VertexIt].UCoord = UCoord;
+		}
+	}
+
 	if (Ar.CustomVer(FReleaseObjectVersion::GUID) >= FReleaseObjectVersion::GroomAssetVersion2)
 	{
 		Ar << Materials;
diff --git a/Engine/Plugins/Runtime/HairStrands/Source/HairStrandsCore/Private/HairStrandsRendering.cpp b/Engine/Plugins/Runtime/HairStrands/Source/HairStrandsCore/Private/HairStrandsRendering.cpp
--- a/Engine/Plugins/Runtime/HairStrands/Source/HairStrandsCore/Private/HairStrandsRendering.cpp
+++ b/Engine/Plugins/Runtime/HairStrands/Source/HairStrandsCore/Private/HairStrandsRendering.cpp
@@ -335,25 +335,6 @@ FRDGHairStrandsCullingData ImportCullingData(FRDGBuilder& GraphBuilder, FHairGro
 	return Out;
 }
 
-struct FHairScaleDesc
-{
-	float InHairRadius = 0;
-	float OutHairRadius = 0;
-	float MaxOutHairRadius = 0;
-	float HairRadiusRootScale = 1;
-	float HairRadiusTipScale = 1;
-	bool  bEnable = true;
-
-	bool IsEnable() const 
-	{
-		return 
-			bEnable && (
-			InHairRadius != OutHairRadius ||
-			HairRadiusRootScale != 1 ||
-			HairRadiusTipScale != 1);
-	}
-};
-
 class FHairInterpolationCS : public FGlobalShader
 {
 	DECLARE_GLOBAL_SHADER(FHairInterpolationCS);
@@ -375,16 +356,7 @@ class FHairInterpolationCS : public FGlobalShader
 		SHADER_PARAMETER(FVector, InRenderHairPositionOffset)
 		SHADER_PARAMETER(FVector, InSimHairPositionOffset)
 		SHADER_PARAMETER(FIntPoint, HairStrandsCullIndex)
-
-		SHADER_PARAMETER(float, InHairLength)
-		SHADER_PARAMETER(float, InHairRadius)
-		SHADER_PARAMETER(float, OutHairRadius)
-		SHADER_PARAMETER(float, MaxOutHairRadius)
-		SHADER_PARAMETER(float, HairRadiusRootScale)
-		SHADER_PARAMETER(float, HairRadiusTipScale)
-		SHADER_PARAMETER(uint32, HairScalingEnable)
 		SHADER_PARAMETER(uint32,  HairStrandsVF_bIsCullingEnable)
-
 		SHADER_PARAMETER(FMatrix, LocalToWorldMatrix)
 
 		SHADER_PARAMETER_RDG_BUFFER_SRV(Buffer, RenderRestPosePositionBuffer)
@@ -460,7 +432,6 @@ static void AddHairStrandsInterpolationPass(
 	const FShaderDrawDebugData* ShaderDrawData,
 	const FHairGroupInstance* Instance,
 	const uint32 VertexCount,
-	const FHairScaleDesc ScaleDesc,
 	const int32 MeshLODIndex,
 	const bool bPatchedAttributeBuffer,
 	const EHairInterpolationType HairInterpolationType,
@@ -508,14 +479,8 @@ static void AddHairStrandsInterpolationPass(
 
 	Parameters->DispatchCountX = DispatchCount.X;
 	Parameters->SimRootPointIndexBuffer = SimRootPointIndexBuffer;
-	
-	Parameters->InHairRadius = ScaleDesc.InHairRadius;
-	Parameters->OutHairRadius = ScaleDesc.OutHairRadius;
-	Parameters->MaxOutHairRadius = ScaleDesc.MaxOutHairRadius;
-	Parameters->HairRadiusRootScale = ScaleDesc.HairRadiusRootScale;
-	Parameters->HairRadiusTipScale = ScaleDesc.HairRadiusTipScale;
-	Parameters->HairScalingEnable = ScaleDesc.IsEnable() ? 1u : 0u;
-	Parameters->AttributeBuffer = RenderAttributeBuffer;	
+
+	Parameters->AttributeBuffer = RenderAttributeBuffer;
 	
 	const bool bIsVertexToCurveBuffersValid_Sim =
 		SimRestRootResources &&
@@ -1010,7 +975,9 @@ class FHairRaytracingGeometryCS : public FGlobalShader
 	BEGIN_SHADER_PARAMETER_STRUCT(FParameters, )
 		SHADER_PARAMETER(uint32, VertexCount)
 		SHADER_PARAMETER(uint32, DispatchCountX)
-		SHADER_PARAMETER(float, StrandHairRadius)
+		SHADER_PARAMETER(float, HairStrandsVF_HairRadius)
+		SHADER_PARAMETER(float, HairStrandsVF_HairRootScale)
+		SHADER_PARAMETER(float, HairStrandsVF_HairTipScale)
 
 		SHADER_PARAMETER(uint32, HairStrandsVF_bIsCullingEnable)
 		SHADER_PARAMETER_RDG_BUFFER_SRV(Buffer,	HairStrandsVF_CullingIndirectBuffer)
@@ -1034,6 +1001,8 @@ static void AddGenerateRaytracingGeometryPass(
 	FGlobalShaderMap* ShaderMap,
 	uint32 VertexCount,
 	float HairRadius,
+	float RootScale,
+	float TipScale,
 	const FRDGBufferSRVRef& HairWorldOffsetBuffer,
 	FRDGHairStrandsCullingData& CullingData,
 	const FRDGBufferSRVRef& PositionBuffer,
@@ -1046,7 +1015,9 @@ static void AddGenerateRaytracingGeometryPass(
 	Parameters->VertexCount = VertexCount;
 	Parameters->DispatchCountX = DispatchCount.X;
 	Parameters->PositionOffsetBuffer = HairWorldOffsetBuffer;
-	Parameters->StrandHairRadius = HairRadius;
+	Parameters->HairStrandsVF_HairRadius = HairRadius;
+	Parameters->HairStrandsVF_HairRootScale = RootScale;
+	Parameters->HairStrandsVF_HairTipScale = TipScale;
 	Parameters->PositionBuffer = PositionBuffer;
 	Parameters->OutputPositionBuffer = OutPositionBuffer;
 
@@ -1279,29 +1250,6 @@ EHairStrandsDebugMode GetHairStrandsGeometryDebugMode(FHairGroupInstance* Instan
 	return DebugMode;
 }
 
-FHairScaleDesc ComputeHairScaleDesc(FHairGroupInstance* Instance)
-{
-	FHairScaleDesc Out;
-	Out.bEnable = true;
-	Out.InHairRadius = Instance->Strands.Modifier.HairWidth * 0.5f;
-
-	Out.OutHairRadius = Out.InHairRadius;
-	Out.MaxOutHairRadius = Out.OutHairRadius;
-	Out.HairRadiusRootScale = 1.f;
-	Out.HairRadiusTipScale = 1.f;
-
-	const bool bSupportDeformation = Instance->Strands.DeformedResource != nullptr;
-	if (bSupportDeformation)
-	{
-		Out.OutHairRadius = (GStrandHairWidth > 0 ? GStrandHairWidth : Instance->Strands.Modifier.HairWidth) * 0.5f;
-		Out.MaxOutHairRadius = Out.OutHairRadius * FMath::Max(1.f, FMath::Max(Instance->Strands.Modifier.HairRootScale, Instance->Strands.Modifier.HairTipScale));
-		Out.HairRadiusRootScale = Instance->Strands.Modifier.HairRootScale;
-		Out.HairRadiusTipScale = Instance->Strands.Modifier.HairTipScale;
-	}
-	
-	return Out;
-}
-
 bool NeedsPatchAttributeBuffer(EHairStrandsDebugMode DebugMode)
 {
 	return DebugMode == EHairStrandsDebugMode::RenderHairStrands || DebugMode == EHairStrandsDebugMode::RenderVisCluster;
@@ -1328,8 +1276,6 @@ static FHairGroupPublicData::FVertexFactoryInput InternalComputeHairStrandsVerte
 	if (!Instance || Instance->GeometryType != EHairGeometryType::Strands)
 		return OutVFInput;
 
-	const FHairScaleDesc ScaleDesc = ComputeHairScaleDesc(Instance);
-
 	const EHairStrandsDebugMode DebugMode = GetHairStrandsGeometryDebugMode(Instance);
 	const bool bDebugModePatchedAttributeBuffer = NeedsPatchAttributeBuffer(DebugMode);
 
@@ -1348,12 +1294,7 @@ static FHairGroupPublicData::FVertexFactoryInput InternalComputeHairStrandsVerte
 
 		OutVFInput.Strands.PositionOffset = Instance->Guides.DeformedResource->GetPositionOffset(FHairStrandsDeformedResource::EFrameType::Current);
 		OutVFInput.Strands.PrevPositionOffset = Instance->Guides.DeformedResource->GetPositionOffset(FHairStrandsDeformedResource::EFrameType::Previous);
-		OutVFInput.Strands.VertexCount = Instance->Guides.RestResource->GetVertexCount();
-		OutVFInput.Strands.HairRadius = (GStrandHairWidth > 0 ? GStrandHairWidth : Instance->Strands.Modifier.HairWidth) * 0.5f;
-		OutVFInput.Strands.HairLength = Instance->Strands.Modifier.HairLength;
-		OutVFInput.Strands.HairDensity = Instance->Strands.Modifier.HairShadowDensity;
-		OutVFInput.Strands.bUseStableRasterization = Instance->Strands.Modifier.bUseStableRasterization;
-		OutVFInput.Strands.bScatterSceneLighting = Instance->Strands.Modifier.bScatterSceneLighting;
+		OutVFInput.Strands.VertexCount = Instance->Guides.RestResource->GetVertexCount();		
 	}
 	// 2. Render Strands with deformation
 	else if (bSupportDeformation)
@@ -1369,11 +1310,6 @@ static FHairGroupPublicData::FVertexFactoryInput InternalComputeHairStrandsVerte
 		OutVFInput.Strands.PositionOffset = Instance->Strands.DeformedResource->GetPositionOffset(FHairStrandsDeformedResource::EFrameType::Current);
 		OutVFInput.Strands.PrevPositionOffset = Instance->Strands.DeformedResource->GetPositionOffset(FHairStrandsDeformedResource::EFrameType::Previous);
 		OutVFInput.Strands.VertexCount = Instance->Strands.RestResource->GetVertexCount();
-		OutVFInput.Strands.HairRadius = ScaleDesc.MaxOutHairRadius;
-		OutVFInput.Strands.HairLength = Instance->Strands.Modifier.HairLength;
-		OutVFInput.Strands.HairDensity = Instance->Strands.Modifier.HairShadowDensity;
-		OutVFInput.Strands.bScatterSceneLighting = Instance->Strands.Modifier.bScatterSceneLighting;
-		OutVFInput.Strands.bUseStableRasterization = Instance->Strands.Modifier.bUseStableRasterization;
 	}
 	// 3. Render Strands in rest position: used when there are no skinning (rigid binding or no binding), no simulation, and no RBF
 	else
@@ -1389,13 +1325,16 @@ static FHairGroupPublicData::FVertexFactoryInput InternalComputeHairStrandsVerte
 		OutVFInput.Strands.PositionOffset = Instance->Strands.RestResource->PositionOffset;
 		OutVFInput.Strands.PrevPositionOffset = Instance->Strands.RestResource->PositionOffset;
 		OutVFInput.Strands.VertexCount = Instance->Strands.RestResource->GetVertexCount();
-		OutVFInput.Strands.HairRadius = ScaleDesc.MaxOutHairRadius;
-		OutVFInput.Strands.HairLength = Instance->Strands.Modifier.HairLength;
-		OutVFInput.Strands.HairDensity = Instance->Strands.Modifier.HairShadowDensity;
-		OutVFInput.Strands.bScatterSceneLighting = Instance->Strands.Modifier.bScatterSceneLighting;
-		OutVFInput.Strands.bUseStableRasterization = Instance->Strands.Modifier.bUseStableRasterization;
 	}
 
+	OutVFInput.Strands.HairRadius = (GStrandHairWidth > 0 ? GStrandHairWidth : Instance->Strands.Modifier.HairWidth) * 0.5f;
+	OutVFInput.Strands.HairRootScale = Instance->Strands.Modifier.HairRootScale;
+	OutVFInput.Strands.HairTipScale = Instance->Strands.Modifier.HairTipScale;
+	OutVFInput.Strands.HairLength = Instance->Strands.Modifier.HairLength;
+	OutVFInput.Strands.HairDensity = Instance->Strands.Modifier.HairShadowDensity;
+	OutVFInput.Strands.bScatterSceneLighting = Instance->Strands.Modifier.bScatterSceneLighting;
+	OutVFInput.Strands.bUseStableRasterization = Instance->Strands.Modifier.bUseStableRasterization;
+
 	return OutVFInput;
 }
 
@@ -1444,7 +1383,6 @@ void ComputeHairStrandsInterpolation(
 		// * Render : Show rendering strands with sim color influence
 		const EHairStrandsDebugMode DebugMode = GetHairStrandsGeometryDebugMode(Instance);
 		const bool bDebugModePatchedAttributeBuffer = NeedsPatchAttributeBuffer(DebugMode);
-		const FHairScaleDesc ScaleDesc = ComputeHairScaleDesc(Instance);
 
 		if (DebugMode == EHairStrandsDebugMode::SimHairStrands)
 		{
@@ -1503,7 +1441,6 @@ void ComputeHairStrandsInterpolation(
 					ShaderDrawData,
 					Instance,
 					Instance->Strands.RestResource->GetVertexCount(),
-					ScaleDesc,
 					MeshLODIndex,
 					bDebugModePatchedAttributeBuffer,
 					Instance->Strands.HairInterpolationType,
@@ -1586,7 +1523,9 @@ void ComputeHairStrandsInterpolation(
 						GraphBuilder,
 						ShaderMap,
 						Instance->Strands.RestResource->GetVertexCount(),
-						ScaleDesc.MaxOutHairRadius * HairRadiusScaleRT,
+						Instance->HairGroupPublicData->VFInput.Strands.HairRadius * HairRadiusScaleRT,
+						Instance->HairGroupPublicData->VFInput.Strands.HairRootScale,
+						Instance->HairGroupPublicData->VFInput.Strands.HairTipScale,
 						Strands_PositionOffsetSRV,
 						CullingData,
 						Strands_PositionSRV,
@@ -1647,14 +1586,6 @@ void ComputeHairStrandsInterpolation(
 				{
 					FRDGImportedBuffer Guides_DeformedPositionBuffer = Register(GraphBuilder, LOD.Guides.DeformedResource->GetBuffer(FHairStrandsDeformedResource::Current), ERDGImportedBufferFlags::CreateViews);
 
-					FHairScaleDesc ScaleDesc;
-					ScaleDesc.bEnable = false;
-					ScaleDesc.InHairRadius = LOD.Guides.Data->StrandsCurves.MaxRadius;
-					ScaleDesc.OutHairRadius = (GStrandHairWidth > 0 ? GStrandHairWidth * 0.5f : ScaleDesc.InHairRadius);
-					ScaleDesc.MaxOutHairRadius = ScaleDesc.OutHairRadius;
-					ScaleDesc.HairRadiusRootScale = 1;
-					ScaleDesc.HairRadiusTipScale = 1;
-
 					FRDGHairStrandsCullingData CullingData;
 					AddHairStrandsInterpolationPass(
 						GraphBuilder,
@@ -1662,7 +1593,6 @@ void ComputeHairStrandsInterpolation(
 						ShaderDrawData,
 						Instance,
 						LOD.Guides.RestResource->GetVertexCount(),
-						ScaleDesc,
 						MeshLODIndex,
 						false,
 						LOD.Guides.HairInterpolationType,
diff --git a/Engine/Plugins/Runtime/HairStrands/Source/HairStrandsCore/Private/HairStrandsVertexFactory.cpp b/Engine/Plugins/Runtime/HairStrands/Source/HairStrandsCore/Private/HairStrandsVertexFactory.cpp
--- a/Engine/Plugins/Runtime/HairStrands/Source/HairStrandsCore/Private/HairStrandsVertexFactory.cpp
+++ b/Engine/Plugins/Runtime/HairStrands/Source/HairStrandsCore/Private/HairStrandsVertexFactory.cpp
@@ -61,6 +61,8 @@ inline FRHIShaderResourceView* GetPreviousPositionOffsetSRV(const FHairGroupPubl
 inline bool  UseStableRasterization(const FHairGroupPublicData::FVertexFactoryInput& VFInput)					{ return VFInput.Strands.bUseStableRasterization; };
 inline bool  UseScatterSceneLighting(const FHairGroupPublicData::FVertexFactoryInput& VFInput)					{ return VFInput.Strands.bScatterSceneLighting; };
 inline float GetMaxStrandRadius(const FHairGroupPublicData::FVertexFactoryInput& VFInput)						{ return VFInput.Strands.HairRadius; };
+inline float GetStrandRootScale(const FHairGroupPublicData::FVertexFactoryInput& VFInput)						{ return VFInput.Strands.HairRootScale; };
+inline float GetStrandTipScale(const FHairGroupPublicData::FVertexFactoryInput& VFInput)						{ return VFInput.Strands.HairTipScale; };
 inline float GetMaxStrandLength(const FHairGroupPublicData::FVertexFactoryInput& VFInput)						{ return VFInput.Strands.HairLength; };
 inline float GetHairDensity(const FHairGroupPublicData::FVertexFactoryInput& VFInput)							{ return VFInput.Strands.HairDensity; };
 
@@ -74,9 +76,11 @@ class FHairStrandsVertexFactoryShaderParameters : public FVertexFactoryShaderPar
 public:
 
 	LAYOUT_FIELD(FShaderParameter, Radius);
+	LAYOUT_FIELD(FShaderParameter, RootScale);
+	LAYOUT_FIELD(FShaderParameter, TipScale);
 	LAYOUT_FIELD(FShaderParameter, Length);
 	LAYOUT_FIELD(FShaderParameter, RadiusAtDepth1_Primary);	// unused
-	LAYOUT_FIELD(FShaderParameter, RadiusAtDepth1_Velocity);	// unused
+	LAYOUT_FIELD(FShaderParameter, RadiusAtDepth1_Velocity);// unused
 	LAYOUT_FIELD(FShaderParameter, Density);
 	LAYOUT_FIELD(FShaderParameter, Culling);
 	LAYOUT_FIELD(FShaderParameter, StableRasterization);
@@ -96,6 +100,8 @@ public:
 	void Bind(const FShaderParameterMap& ParameterMap)
 	{
 		Radius.Bind(ParameterMap, TEXT("HairStrandsVF_Radius"));
+		RootScale.Bind(ParameterMap, TEXT("HairStrandsVF_RootScale"));
+		TipScale.Bind(ParameterMap, TEXT("HairStrandsVF_TipScale"));
 		Length.Bind(ParameterMap, TEXT("HairStrandsVF_Length"));
 		Density.Bind(ParameterMap, TEXT("HairStrandsVF_Density"));
 		Density.Bind(ParameterMap, TEXT("HairStrandsVF_Density"));	
@@ -141,6 +147,8 @@ public:
 		VFS_BindParam(ShaderBindings, MaterialBuffer, GetMaterialSRV(VFInput));
 		VFS_BindParam(ShaderBindings, TangentBuffer, GetTangentSRV(VFInput));
 		VFS_BindParam(ShaderBindings, Radius, GetMaxStrandRadius(VFInput));
+		VFS_BindParam(ShaderBindings, RootScale, GetStrandRootScale(VFInput));
+		VFS_BindParam(ShaderBindings, TipScale, GetStrandTipScale(VFInput));
 		VFS_BindParam(ShaderBindings, Length, GetMaxStrandLength(VFInput));
 		VFS_BindParam(ShaderBindings, PositionOffsetBuffer, GetPositionOffsetSRV(VFInput));
 		VFS_BindParam(ShaderBindings, PreviousPositionOffsetBuffer, GetPreviousPositionOffsetSRV(VFInput));
diff --git a/Engine/Plugins/Runtime/HairStrands/Source/HairStrandsCore/Public/HairStrandsDatas.h b/Engine/Plugins/Runtime/HairStrands/Source/HairStrandsCore/Public/HairStrandsDatas.h
--- a/Engine/Plugins/Runtime/HairStrands/Source/HairStrandsCore/Public/HairStrandsDatas.h
+++ b/Engine/Plugins/Runtime/HairStrands/Source/HairStrandsCore/Public/HairStrandsDatas.h
@@ -15,14 +15,14 @@ struct FPackedHairVertex
 	FFloat16 X, Y, Z;
 	uint8 ControlPointType : 2;
 	uint8 NormalizedRadius : 6;
-	uint8 NormalizedLength;
+	uint8 UCoord;
 };
 
 struct FPackedHairAttributeVertex
 {
 	uint8 RootU;
 	uint8 RootV;
-	uint8 UCoord;
+	uint8 NormalizedLength;
 	uint8 Seed;
 
 	uint8 IndexU;
diff --git a/Engine/Shaders/Private/HairStrands/HairCardsDebug.usf b/Engine/Shaders/Private/HairStrands/HairCardsDebug.usf
--- a/Engine/Shaders/Private/HairStrands/HairCardsDebug.usf
+++ b/Engine/Shaders/Private/HairStrands/HairCardsDebug.usf
@@ -157,8 +157,8 @@ void MainCS(uint3 DispatchThreadId : SV_DispatchThreadID)
 		Color	 = float4(1, 0, 1, 1);
 	}
 
-	const FHairControlPoint P0 = UnpackHairControlPoint(PackedP0, Offset, 1, 1);
-	const FHairControlPoint P1 = UnpackHairControlPoint(PackedP1, Offset, 1, 1);
+	const FHairControlPoint P0 = UnpackHairControlPoint(PackedP0, Offset, 1, 1, 1);
+	const FHairControlPoint P1 = UnpackHairControlPoint(PackedP1, Offset, 1, 1, 1);
 
 	if (P1.Type != HAIR_CONTROLPOINT_START)
 	{
diff --git a/Engine/Shaders/Private/HairStrands/HairCardsDeformation.usf b/Engine/Shaders/Private/HairStrands/HairCardsDeformation.usf
--- a/Engine/Shaders/Private/HairStrands/HairCardsDeformation.usf
+++ b/Engine/Shaders/Private/HairStrands/HairCardsDeformation.usf
@@ -64,7 +64,8 @@ void MainCS(uint2 DispatchThreadId : SV_DispatchThreadID)
 		return;
 
 	const float GuideRadius = 1;
-	const float GuideLength = 1;
+	const float GuideRootScale = 1;
+	const float GuideTipScale = 1;
 	const float3 GuideDeformedPositionOffset = GuideDeformedPositionOffsetBuffer[0].xyz;
 
 #if PERMUTATION_DYNAMIC_GEOMETRY == 0
@@ -80,25 +81,29 @@ void MainCS(uint2 DispatchThreadId : SV_DispatchThreadID)
 		GuideRestPositionBuffer[GuideIndex0],
 		GuideRestPositionOffset,
 		GuideRadius,
-		GuideLength);
+		GuideRootScale,
+		GuideTipScale);
 
 	const FHairControlPoint RestGuideP1 = UnpackHairControlPoint(
 		GuideRestPositionBuffer[GuideIndex1],
 		GuideRestPositionOffset,
 		GuideRadius,
-		GuideLength);
+		GuideRootScale,
+		GuideTipScale);
 
 	const FHairControlPoint DeformedGuideP0 = UnpackHairControlPoint(
 		GuideDeformedPositionBuffer[GuideIndex0],
 		GuideDeformedPositionOffset,
 		GuideRadius,
-		GuideLength);
+		GuideRootScale,
+		GuideTipScale);
 
 	const FHairControlPoint DeformedGuideP1 = UnpackHairControlPoint(
 		GuideDeformedPositionBuffer[GuideIndex1],
 		GuideDeformedPositionOffset,
 		GuideRadius,
-		GuideLength);
+		GuideRootScale,
+		GuideTipScale);
 
 	const float3 RestGuidePosition		= lerp(RestGuideP0.Position, RestGuideP1.Position, GuideLerp);
 	const float3 DeformedGuidePosition	= lerp(DeformedGuideP0.Position, DeformedGuideP1.Position, GuideLerp);
diff --git a/Engine/Shaders/Private/HairStrands/HairStrandsAttributeCommon.ush b/Engine/Shaders/Private/HairStrands/HairStrandsAttributeCommon.ush
--- a/Engine/Shaders/Private/HairStrands/HairStrandsAttributeCommon.ush
+++ b/Engine/Shaders/Private/HairStrands/HairStrandsAttributeCommon.ush
@@ -9,6 +9,8 @@
 ////////////////////////////////////////////////////////////////////////////////
 
 float  HairStrandsVF_Radius;
+float  HairStrandsVF_RootScale;
+float  HairStrandsVF_TipScale;
 float  HairStrandsVF_Length;
 
 Buffer<uint4>  HairStrandsVF_PositionBuffer;
@@ -31,13 +33,13 @@ FHairVertexAttributes GetHairStrandsAttributes(uint HairPrimitiveId, float2 Vert
 	const uint VertexId1 = HairPrimitiveId+1;
 
 	// Instead of doing the actual interpolation we could rely on dithering as an optimization
-	const FHairAttributes Attributes0 = DecodeHairAttribute(HairStrandsVF_AttributeBuffer[VertexId0]);
-	const FHairAttributes Attributes1 = DecodeHairAttribute(HairStrandsVF_AttributeBuffer[VertexId1]);
+	const FHairAttributes Attributes0 = DecodeHairAttribute(HairStrandsVF_AttributeBuffer[VertexId0], HairStrandsVF_Length);
+	const FHairAttributes Attributes1 = DecodeHairAttribute(HairStrandsVF_AttributeBuffer[VertexId1], HairStrandsVF_Length);
 
-	const uint4 PackedData = HairStrandsVF_PositionBuffer[VertexId0];
-	const uint PackedAlpha = asuint(PackedData.a);
-	const float WorldRadius = ((PackedAlpha & 0x00FC) >> 2) / 63.f * HairStrandsVF_Radius;
-	const float WorldLength = ((PackedAlpha & 0xFF00) >> 8) / 255.f * HairStrandsVF_Length;
+	const FHairControlPoint ControlPoint0 = UnpackHairControlPoint(HairStrandsVF_PositionBuffer[VertexId0], float3(0, 0 ,0), HairStrandsVF_Radius, HairStrandsVF_RootScale, HairStrandsVF_TipScale);
+	const FHairControlPoint ControlPoint1 = UnpackHairControlPoint(HairStrandsVF_PositionBuffer[VertexId1], float3(0, 0, 0), HairStrandsVF_Radius, HairStrandsVF_RootScale, HairStrandsVF_TipScale);
+	const float WorldLength = lerp(Attributes0.WorldLength, Attributes1.WorldLength, VertexUV.x);
+	const float WorldRadius = lerp(ControlPoint0.WorldRadius, ControlPoint1.WorldRadius, VertexUV.x);
 
 	FHairVertexAttributes Out;
 	// Constant along the strand
@@ -46,8 +48,8 @@ FHairVertexAttributes GetHairStrandsAttributes(uint HairPrimitiveId, float2 Vert
 	Out.HairRootUV		= Attributes0.RootUV;
 
 	// Varying along the strand
+	Out.HairUV.x		= lerp(ControlPoint0.UCoord, ControlPoint1.UCoord, VertexUV.x);
 	Out.HairUV.y		= VertexUV.y;
-	Out.HairUV.x		= lerp(Attributes0.UCoord, Attributes1.UCoord, VertexUV.x);
 
 	// Instead of doing the actual interpolation we could rely on dithering as an optimization
 	const float4 Material0 = HairStrandsVF_MaterialBuffer[VertexId0];
diff --git a/Engine/Shaders/Private/HairStrands/HairStrandsInterpolation.usf b/Engine/Shaders/Private/HairStrands/HairStrandsInterpolation.usf
--- a/Engine/Shaders/Private/HairStrands/HairStrandsInterpolation.usf
+++ b/Engine/Shaders/Private/HairStrands/HairStrandsInterpolation.usf
@@ -42,14 +42,6 @@ float3 InSimHairPositionOffset;
 int2 HairStrandsCullIndex;
 float4x4 LocalToWorldMatrix;
 
-float InHairLength;
-float InHairRadius;
-float OutHairRadius;
-float MaxOutHairRadius;
-float HairRadiusRootScale;
-float HairRadiusTipScale;
-uint  HairScalingEnable;
-
 Buffer<uint4>  RenderRestPosePositionBuffer;
 Buffer<uint4>  SimRestPosePositionBuffer;
 Buffer<uint4>  DeformedSimPositionBuffer;
@@ -177,26 +169,6 @@ float3 ComputeRotateGeometryOffset(uint GuideIndex, float GuideVertexWeight,
 	return (LocalDeformedGuidePoint - LocalRestGuidePoint) * GuideVertexWeight;
 }
 
-// Compute the new hair width and return if the current vertex is valid (true), or if it needs to be clipped (false)
-void ScaleStrand(uint ControlPointId, inout uint InPackedAlpha)
-{
-	if (HairScalingEnable)
-	{
-		// Instead of doing the actual interpolation we could rely on dithering as an optimization
-		const FHairAttributes Attributes0 = DecodeHairAttribute(AttributeBuffer[ControlPointId]);
-		const float U = Attributes0.UCoord;
-
-		// todo: Move all these computations into shared function which can be shared with the hair VF
-		const float InNormVertexRadius = ((InPackedAlpha & 0x00FC) >> 2) / 63.f;
-		//const float InVertexRadius = InNormVertexRadius * InHairRadius;
-
-		const float OutVertexRadius = OutHairRadius * InNormVertexRadius * lerp(HairRadiusRootScale, HairRadiusTipScale, U);
-		const uint PackedVertexRadius = (uint(saturate(OutVertexRadius / MaxOutHairRadius) * 63.f) & 0x3F) << 2;
-
-		InPackedAlpha = (InPackedAlpha & 0xFF03) | PackedVertexRadius;
-	}
-}
-
 [numthreads(GROUP_SIZE, 1, 1)]
 void MainCS(uint2 DispatchThreadId : SV_DispatchThreadID)
 {
@@ -257,7 +229,7 @@ void MainCS(uint2 DispatchThreadId : SV_DispatchThreadID)
 		#if PERMUTATION_DYNAMIC_GEOMETRY == INTERPOLATION_RIGID
 		{
 			float3 CurrOffset = 0;
-			uint4 PackedRestPosition = RenderRestPosePositionBuffer[VertexIndex];
+			const uint4 PackedRestPosition = RenderRestPosePositionBuffer[VertexIndex];
 			float3 ControlPoint = f16tof32(PackedRestPosition.xyz) + InRenderHairPositionOffset;
 
 			#if PERMUTATION_SIMULATION == 1
@@ -278,7 +250,6 @@ void MainCS(uint2 DispatchThreadId : SV_DispatchThreadID)
 			}
 			#endif
 
-			ScaleStrand(VertexIndex, PackedRestPosition.w);
 			WorldVertexPosition = ControlPoint + CurrOffset;
 
 			OutPosition = uint4(f32tof16(WorldVertexPosition - OutRenHairPositionOffset), PackedRestPosition.w);
@@ -290,7 +261,7 @@ void MainCS(uint2 DispatchThreadId : SV_DispatchThreadID)
 		// Apply dynamic mesh deformation (translation / rotation)
 		#if PERMUTATION_DYNAMIC_GEOMETRY == INTERPOLATION_SKINNING_OFFSET
 		{
-			uint4 PackedRestPosition = RenderRestPosePositionBuffer[VertexIndex];
+			const uint4 PackedRestPosition = RenderRestPosePositionBuffer[VertexIndex];
 			const float3 RestControlPoint = f16tof32(PackedRestPosition.xyz) + InRenderHairPositionOffset;
 			float3 ControlPoint = RestControlPoint;
 
@@ -330,7 +301,6 @@ void MainCS(uint2 DispatchThreadId : SV_DispatchThreadID)
 			ControlPoint = ToTriangle(ControlPoint,RootBarycentric,RestTriangle);
 			ControlPoint = ToWorld(ControlPoint,RootBarycentric,DeformedTriangle) + CurrOffset;
 
-			ScaleStrand(VertexIndex, PackedRestPosition.w);
 			WorldVertexPosition = ControlPoint;
 
 			#if PERMUTATION_DEBUG == 1 && ALLOW_DEBUG_RENDERING
@@ -371,7 +341,7 @@ void MainCS(uint2 DispatchThreadId : SV_DispatchThreadID)
 		// 3. Linear blend skinning (translation)
 		#if PERMUTATION_DYNAMIC_GEOMETRY == INTERPOLATION_SKINNING_TRANSLATION
 		{
-			uint4 PackedRestPosition = RenderRestPosePositionBuffer[VertexIndex];
+			const uint4 PackedRestPosition = RenderRestPosePositionBuffer[VertexIndex];
 			const float3 RestControlPoint = f16tof32(PackedRestPosition.xyz) + InRenderHairPositionOffset;
 			float3 ControlPoint = RestControlPoint;
 
@@ -404,7 +374,6 @@ void MainCS(uint2 DispatchThreadId : SV_DispatchThreadID)
 
 			ControlPoint += RootBarycentric.x * DeformedTriangle.P0 + RootBarycentric.y * DeformedTriangle.P1 + RootBarycentric.z * DeformedTriangle.P2;
 
-			ScaleStrand(VertexIndex, PackedRestPosition.w);
 			WorldVertexPosition = ControlPoint + CurrOffset;
 
 			OutPosition = uint4(f32tof16(WorldVertexPosition - OutRenHairPositionOffset), PackedRestPosition.w);
@@ -415,7 +384,7 @@ void MainCS(uint2 DispatchThreadId : SV_DispatchThreadID)
 		// 4. Linear blend skinning (rotation)
 		#if PERMUTATION_DYNAMIC_GEOMETRY == INTERPOLATION_SKINNING_ROTATION
 		{
-			uint4 PackedRestPosition = RenderRestPosePositionBuffer[VertexIndex];
+			const uint4 PackedRestPosition = RenderRestPosePositionBuffer[VertexIndex];
 			const float3 RestControlPoint = f16tof32(PackedRestPosition.xyz) + InRenderHairPositionOffset;
 			float3 ControlPoint = RestControlPoint;
 
@@ -455,7 +424,6 @@ void MainCS(uint2 DispatchThreadId : SV_DispatchThreadID)
 
 			ControlPoint =  RootBarycentric.x * DeformedTriangle.P0 + RootBarycentric.y * DeformedTriangle.P1 + RootBarycentric.z * DeformedTriangle.P2;
 
-			ScaleStrand(VertexIndex, PackedRestPosition.w);
 			WorldVertexPosition = ControlPoint + CurrOffset;
 
 			OutPosition = uint4(f32tof16(WorldVertexPosition - OutRenHairPositionOffset), PackedRestPosition.w);
@@ -466,7 +434,7 @@ void MainCS(uint2 DispatchThreadId : SV_DispatchThreadID)
 		// 5. Linear blend skinning (translation + rotation)
 		#if PERMUTATION_DYNAMIC_GEOMETRY == INTERPOLATION_SKINNING_TRANSLATION_AND_ROTATION
 		{
-			uint4 PackedRestPosition = RenderRestPosePositionBuffer[VertexIndex];
+			const uint4 PackedRestPosition = RenderRestPosePositionBuffer[VertexIndex];
 			const float3 RestControlPoint = f16tof32(PackedRestPosition.xyz) + InRenderHairPositionOffset;
 			float3 ControlPoint = RestControlPoint;
 
@@ -503,7 +471,6 @@ void MainCS(uint2 DispatchThreadId : SV_DispatchThreadID)
 			}
 			#endif
 
-			ScaleStrand(VertexIndex, PackedRestPosition.w);
 			WorldVertexPosition = CurrOffset;
 
 			OutPosition = uint4(f32tof16(WorldVertexPosition - OutRenHairPositionOffset), PackedRestPosition.w);
@@ -633,7 +600,7 @@ void ClusterAABBEvaluationCS(uint3 DispatchThreadId : SV_DispatchThreadID, uint3
 		{
 			const float3 OutHairPositionOffset = RenderDeformedOffsetBuffer[0].xyz;
 			const uint VertexIndex = clamp(ClusterVertexIdBuffer[IndexOffset + ThreadReadOffset], 0u, VertexCountAfterCulling-1);
-			const FHairControlPoint CP = UnpackHairControlPoint(RenderDeformedPositionBuffer[VertexIndex], OutHairPositionOffset, 1, 1);
+			const FHairControlPoint CP = UnpackHairControlPoint(RenderDeformedPositionBuffer[VertexIndex], OutHairPositionOffset, 1, 1, 1);
 
 			if (all(IsFinite(CP.Position)))
 			{
diff --git a/Engine/Shaders/Private/HairStrands/HairStrandsRaytracingGeometry.usf b/Engine/Shaders/Private/HairStrands/HairStrandsRaytracingGeometry.usf
--- a/Engine/Shaders/Private/HairStrands/HairStrandsRaytracingGeometry.usf
+++ b/Engine/Shaders/Private/HairStrands/HairStrandsRaytracingGeometry.usf
@@ -1,12 +1,14 @@
 // Copyright Epic Games, Inc. All Rights Reserved.
 
 #include "../Common.ush" 
+#include "HairStrandsVertexFactoryCommon.ush"
 
 uint VertexCount;
 uint DispatchCountX;
 
-float3 StrandHairWorldOffset;
-float StrandHairRadius;
+float HairStrandsVF_HairRadius;
+float HairStrandsVF_HairRootScale;
+float HairStrandsVF_HairTipScale;
 
 Buffer<float4>	 PositionOffsetBuffer;
 Buffer<uint4> 	 PositionBuffer;
@@ -28,27 +30,18 @@ Buffer<uint>	HairStrandsVF_CullingIndirectBufferArgs;
 	#error Unknown group size	
 #endif
 
-struct FControlPoint
+FHairControlPoint GetControlPoint(uint ControlPointId)
 {
-	float3 Position;
-	float PointType;
-	float WorldRadius;
-};
+	FHairControlPoint Out = (FHairControlPoint)0;
 
-FControlPoint GetControlPoint(uint Index)
-{
-	FControlPoint Out = (FControlPoint)0;
-
-	if (Index < VertexCount)
+	if (ControlPointId < VertexCount)
 	{
-		const uint4 PackedData = PositionBuffer[Index];
-		const uint PackedAlpha = asuint(PackedData.a);
-
-		const float3 OutPositionOffset = PositionOffsetBuffer[0].xyz;
-
-		Out.Position 	= f16tof32(PackedData.xyz) + OutPositionOffset;
-		Out.PointType 	= PackedAlpha & 0x3;
-		Out.WorldRadius	= StrandHairRadius * float((PackedAlpha & 0x00FC) >> 2) / 63.0f;
+		Out = UnpackHairControlPoint(
+			PositionBuffer[ControlPointId],
+			PositionOffsetBuffer[0].xyz,
+			HairStrandsVF_HairRadius,
+			HairStrandsVF_HairRootScale,
+			HairStrandsVF_HairTipScale);
 	}
 
 	return Out;
@@ -114,17 +107,17 @@ void MainCS(uint2 DispatchThreadId : SV_DispatchThreadID)
 	}
 #endif
 
-	FControlPoint P0 = GetControlPoint(VertexIndex0);
+	FHairControlPoint P0 = GetControlPoint(VertexIndex0);
 
 	// Point at the end of a strand is filled with invalid position. The positions/triangle are 
 	// ignored by DXR during BVH building
-	if (P0.PointType == 2)
+	if (P0.Type == 2)
 	{
 		SetInvalidSegment(VertexIndex0);
 		return;
 	}
 
-	FControlPoint P1 = GetControlPoint(VertexIndex1);
+	FHairControlPoint P1 = GetControlPoint(VertexIndex1);
 
 	const float3 Tangent = normalize(P1.Position - P0.Position);
 	float3 B = float3(0,0,-1);
diff --git a/Engine/Shaders/Private/HairStrands/HairStrandsTexturesGeneration.usf b/Engine/Shaders/Private/HairStrands/HairStrandsTexturesGeneration.usf
--- a/Engine/Shaders/Private/HairStrands/HairStrandsTexturesGeneration.usf
+++ b/Engine/Shaders/Private/HairStrands/HairStrandsTexturesGeneration.usf
@@ -134,6 +134,9 @@ bool Trace(
 {
 	bool bIsValid = false;
 
+	const float InVF_RootScale = 1.0f;
+	const float InVF_TipScale = 1.0f;
+
 	const float3 WP0 = InRayO;
 	const float3 WP1 = InRayO + InRayD * InMaxDistance;
 
@@ -261,8 +264,8 @@ bool Trace(
 						const uint4 PackedCP_P0 = InVF_PositionBuffer[P0];
 						const uint4 PackedCP_P1 = InVF_PositionBuffer[P1];
 
-						const FHairControlPoint CP_Curr = UnpackHairControlPoint(PackedCP_P0, InVF_PositionOffset, InVF_Radius, InVF_Length);
-						const FHairControlPoint CP_Next = UnpackHairControlPoint(PackedCP_P1, InVF_PositionOffset, InVF_Radius, InVF_Length);
+						const FHairControlPoint CP_Curr = UnpackHairControlPoint(PackedCP_P0, InVF_PositionOffset, InVF_Radius, InVF_RootScale, InVF_TipScale);
+						const FHairControlPoint CP_Next = UnpackHairControlPoint(PackedCP_P1, InVF_PositionOffset, InVF_Radius, InVF_RootScale, InVF_TipScale);
 
 						const float3 Tangent = ComputeTangent(CP_Curr, CP_Next);
 
@@ -279,12 +282,12 @@ bool Trace(
 							OutDepth = Distance;
 
 
-							const FHairAttributes Attributes0 = DecodeHairAttribute(InVF_AttributeBuffer[P0]);
-							const FHairAttributes Attributes1 = DecodeHairAttribute(InVF_AttributeBuffer[P1]);
+							const FHairAttributes Attributes0 = DecodeHairAttribute(InVF_AttributeBuffer[P0], InVF_Length);
+							const FHairAttributes Attributes1 = DecodeHairAttribute(InVF_AttributeBuffer[P1], InVF_Length);
 
 							OutSeed		= Attributes0.Seed;
 							OutRootUV	= Attributes0.RootUV;
-							OutU		= lerp(Attributes0.UCoord, Attributes1.UCoord, VertexU);
+							OutU		= lerp(CP_Curr.UCoord, CP_Next.UCoord, VertexU);
 
 							bIsValid = true;
 
diff --git a/Engine/Shaders/Private/HairStrands/HairStrandsVertexFactory.ush b/Engine/Shaders/Private/HairStrands/HairStrandsVertexFactory.ush
--- a/Engine/Shaders/Private/HairStrands/HairStrandsVertexFactory.ush
+++ b/Engine/Shaders/Private/HairStrands/HairStrandsVertexFactory.ush
@@ -195,7 +195,8 @@ FHairControlPoint GetVertexPosition(const FVertexInfo VertexInfo, bool bInvalidJ
 		PackedData, 
 		PositionOffset,
 		HairStrandsVF_Radius * VertexInfo.LodRadiusScale, 
-		HairStrandsVF_Length);
+		HairStrandsVF_RootScale,
+		HairStrandsVF_TipScale);
 	
 	// Create a degenerated quad the end of each strand to cut between each strands
 	// #hair_todo: This is not efficient for short strand like fur for instance. Need to revisit that at some point
@@ -216,11 +217,6 @@ float GetWorldStrandRadius(FVertexFactoryInput Input)
 	return GetVertexPosition(Input).WorldRadius;
 }
 
-float GetWorldStrandLength(FVertexFactoryInput Input)
-{
-	return GetVertexPosition(Input).WorldLength;
-}
-
 float3 GetVertexPreviousPosition(FVertexFactoryInput Input, bool bInvalidJointVertex = true)
 {
 	const float3 PreviousPositionOffset = HairStrandsVF_PreviousPositionOffsetBuffer[0].xyz;
@@ -231,7 +227,8 @@ float3 GetVertexPreviousPosition(FVertexFactoryInput Input, bool bInvalidJointVe
 		PackedData,
 		PreviousPositionOffset,
 		HairStrandsVF_Radius * VertexInfo.LodRadiusScale,
-		HairStrandsVF_Length);
+		HairStrandsVF_RootScale,
+		HairStrandsVF_TipScale);
 
 	// Create a degenerated quad the end of each strand to cut between each strands
 	// #hair_todo: This is not efficient for short strand like fur for instance. Need to revisit that at some point
@@ -363,7 +360,7 @@ FVertexFactoryIntermediates GetVertexFactoryIntermediates(FVertexFactoryInput In
 	Intermediates.HairPrimitiveId	= GetHairPrimitiveId(Input);
 	Intermediates.HairPrimitiveUV	= GetSegmentUVCoord(Input);
 	Intermediates.HairDensity		= HairStrandsVF_Density;
-	Intermediates.HairDimensions	= float2(GetWorldStrandLength(Input), GetWorldStrandRadius(Input));
+	Intermediates.HairDimensions	= float2(0.f, GetWorldStrandRadius(Input));
 
 	float TangentSign;
 	Intermediates.TangentToLocal = CalcTangentToLocal(Input, TangentSign);
diff --git a/Engine/Shaders/Private/HairStrands/HairStrandsVertexFactoryCommon.ush b/Engine/Shaders/Private/HairStrands/HairStrandsVertexFactoryCommon.ush
--- a/Engine/Shaders/Private/HairStrands/HairStrandsVertexFactoryCommon.ush
+++ b/Engine/Shaders/Private/HairStrands/HairStrandsVertexFactoryCommon.ush
@@ -13,7 +13,7 @@ struct FHairControlPoint
 {
 	float3 Position;
 	float  WorldRadius;
-	float  WorldLength;
+	float  UCoord;
 	uint   Type;
 };
 
@@ -21,17 +21,20 @@ FHairControlPoint UnpackHairControlPoint(
 	uint4 InPackedData,
 	float3 InVF_PositionOffset,
 	float InVF_Radius,
-	float InVF_Length)
+	float InVF_RootScale,
+	float InVF_TipScale)
 {
 	const uint PackedAlpha	= asuint(InPackedData.a);
 
 	FHairControlPoint Out = (FHairControlPoint)0;
 	Out.Position	= f16tof32(InPackedData.xyz) + InVF_PositionOffset;
-	Out.WorldRadius	= ((PackedAlpha & 0x00FC) >> 2) / 63.f * InVF_Radius;
-	Out.WorldLength	= ((PackedAlpha & 0xFF00) >> 8) / 255.f * InVF_Length;
+	Out.UCoord		= ((PackedAlpha & 0xFF00) >> 8) / 255.f;
+	Out.WorldRadius	= ((PackedAlpha & 0x00FC) >> 2) / 63.f;
 	Out.Type		= PackedAlpha & 0x3;
 
-	return Out;
+	Out.WorldRadius *= InVF_Radius * lerp(InVF_RootScale, InVF_TipScale, Out.UCoord);
+
+	return Out;	
 }
 
 bool IsValidHairStrandsSegment(const FHairControlPoint CP0, const FHairControlPoint CP1)
@@ -45,11 +48,11 @@ bool IsValidHairStrandsSegment(const FHairControlPoint CP0, const FHairControlPo
 struct FHairAttributes
 {
 	float2 RootUV;
-	float	Seed;
-	float	UCoord;
+	float Seed;
+	float WorldLength;
 };
 
-FHairAttributes DecodeHairAttribute(uint4 Attributes)
+FHairAttributes DecodeHairAttribute(uint4 Attributes, float InVF_Length)
 {
 	float2 RootUV, RootIndex;
 	RootUV.x = float(Attributes.x & 0xFF) / 255.f;
@@ -59,7 +62,7 @@ FHairAttributes DecodeHairAttribute(uint4 Attributes)
 
 	FHairAttributes Out;
 	Out.RootUV = RootUV + RootIndex;
-	Out.UCoord = (Attributes.y & 0xFF) / 255.f;
+	Out.WorldLength = (Attributes.y & 0xFF) / 255.f * InVF_Length;
 	Out.Seed = ((Attributes.y >> 8) & 0xFF) / 255.f;
 	return Out;
 }
\ No newline at end of file
diff --git a/Engine/Shaders/Private/HairStrands/HairStrandsVisibilityRasterCompute.usf b/Engine/Shaders/Private/HairStrands/HairStrandsVisibilityRasterCompute.usf
--- a/Engine/Shaders/Private/HairStrands/HairStrandsVisibilityRasterCompute.usf
+++ b/Engine/Shaders/Private/HairStrands/HairStrandsVisibilityRasterCompute.usf
@@ -74,6 +74,8 @@ int2			OutputResolution;
 uint			HairMaterialId;
 uint			HairStrandsVF_bIsCullingEnable;
 float			HairStrandsVF_Radius;
+float			HairStrandsVF_RootScale;
+float			HairStrandsVF_TipScale;
 float			HairStrandsVF_Length;
 float			HairStrandsVF_Density;
 uint			HairStrandsVF_bUseStableRasterization;
@@ -279,7 +281,8 @@ void MainCS(uint3 DispatchThreadID : SV_DispatchThreadID)
 		HairStrandsVF_PositionBuffer[VertexIndex0],
 		PositionOffset,
 		HairStrandsVF_Radius * LodRadiusScale0,
-		HairStrandsVF_Length);
+		HairStrandsVF_RootScale,
+		HairStrandsVF_TipScale);
 
 	if (CP0.Type == HAIR_CONTROLPOINT_END)
 		return;
@@ -288,7 +291,8 @@ void MainCS(uint3 DispatchThreadID : SV_DispatchThreadID)
 		HairStrandsVF_PositionBuffer[VertexIndex1],
 		PositionOffset,
 		HairStrandsVF_Radius * LodRadiusScale1,
-		HairStrandsVF_Length);
+		HairStrandsVF_RootScale,
+		HairStrandsVF_TipScale);
 
 
 	const FHairRenderInfo HairRenderInfo = GetHairRenderInfo(ResolvedView.HairRenderInfo, ResolvedView.HairRenderInfoBits, HairStrandsVF_bUseStableRasterization > 0);
diff --git a/Engine/Shaders/Private/HairStrands/HairStrandsVoxelRasterCompute.usf b/Engine/Shaders/Private/HairStrands/HairStrandsVoxelRasterCompute.usf
--- a/Engine/Shaders/Private/HairStrands/HairStrandsVoxelRasterCompute.usf
+++ b/Engine/Shaders/Private/HairStrands/HairStrandsVoxelRasterCompute.usf
@@ -13,7 +13,8 @@ uint			FrameIdMod8;
 uint			MacroGroupId;
 uint			HairStrandsVF_bIsCullingEnable;
 float			HairStrandsVF_Radius;
-float			HairStrandsVF_Length;
+float			HairStrandsVF_RootScale;
+float			HairStrandsVF_TipScale;
 float			HairStrandsVF_Density;
 float3			HairStrandsVF_PositionOffset;
 uint			HairStrandsVF_VertexCount;
@@ -103,7 +104,8 @@ void MainCS(uint3 DispatchThreadID : SV_DispatchThreadID)
 		HairStrandsVF_PositionBuffer[VertexIndex0],
 		PositionOffset,
 		HairStrandsVF_Radius * LodRadiusScale0,
-		HairStrandsVF_Length);
+		HairStrandsVF_RootScale,
+		HairStrandsVF_TipScale);
 
 	if (CP0.Type == HAIR_CONTROLPOINT_END)
 		return;
@@ -112,7 +114,8 @@ void MainCS(uint3 DispatchThreadID : SV_DispatchThreadID)
 		HairStrandsVF_PositionBuffer[VertexIndex1],
 		PositionOffset,
 		HairStrandsVF_Radius * LodRadiusScale1,
-		HairStrandsVF_Length);
+		HairStrandsVF_RootScale,
+		HairStrandsVF_TipScale);
 
 	const float DiameterToRadius = 0.5f;
 	const float3 WP0 = mul(float4(CP0.Position, 1), HairStrandsVF_LocalToWorldPrimitiveTransform).xyz;
diff --git a/Engine/Source/Runtime/Core/Public/UObject/UE5ReleaseStreamObjectVersion.h b/Engine/Source/Runtime/Core/Public/UObject/UE5ReleaseStreamObjectVersion.h
--- a/Engine/Source/Runtime/Core/Public/UObject/UE5ReleaseStreamObjectVersion.h
+++ b/Engine/Source/Runtime/Core/Public/UObject/UE5ReleaseStreamObjectVersion.h
@@ -30,6 +30,9 @@ struct CORE_API FUE5ReleaseStreamObjectVersion
 		// Serialize the folder path of actor descs
 		WorldPartitionActorDescSerializeActorFolderPath,
 
+		// Change hair strands vertex format
+		HairStrandsVertexFormatChange,
+
 		// -----<new versions can be added above this line>-------------------------------------------------
 		VersionPlusOne,
 		LatestVersion = VersionPlusOne - 1
diff --git a/Engine/Source/Runtime/Renderer/Private/HairStrands/HairStrandsVisibility.cpp b/Engine/Source/Runtime/Renderer/Private/HairStrands/HairStrandsVisibility.cpp
--- a/Engine/Source/Runtime/Renderer/Private/HairStrands/HairStrandsVisibility.cpp
+++ b/Engine/Source/Runtime/Renderer/Private/HairStrands/HairStrandsVisibility.cpp
@@ -3055,6 +3055,8 @@ class FVisiblityRasterComputeCS : public FGlobalShader
 		SHADER_PARAMETER(uint32, HairStrandsVF_bIsCullingEnable)
 		SHADER_PARAMETER(float, HairStrandsVF_Density)
 		SHADER_PARAMETER(float, HairStrandsVF_Radius)
+		SHADER_PARAMETER(float, HairStrandsVF_RootScale)
+		SHADER_PARAMETER(float, HairStrandsVF_TipScale)
 		SHADER_PARAMETER(float, HairStrandsVF_Length)
 		SHADER_PARAMETER(uint32, HairStrandsVF_bUseStableRasterization)
 		SHADER_PARAMETER(uint32, HairStrandsVF_VertexCount)
@@ -3242,6 +3244,8 @@ static FRasterComputeOutput AddVisibilityComputeRasterPass(
 			PassParameters->HairStrandsVF_PositionOffsetBuffer = VFInput.Strands.PositionOffsetBuffer.SRV;
 			PassParameters->HairStrandsVF_VertexCount		= VFInput.Strands.VertexCount;
 			PassParameters->HairStrandsVF_Radius			= VFInput.Strands.HairRadius;
+			PassParameters->HairStrandsVF_RootScale			= VFInput.Strands.HairRootScale;
+			PassParameters->HairStrandsVF_TipScale			= VFInput.Strands.HairTipScale;
 			PassParameters->HairStrandsVF_Length			= VFInput.Strands.HairLength;
 			PassParameters->HairStrandsVF_bUseStableRasterization = VFInput.Strands.bUseStableRasterization ? 1 : 0;
 			PassParameters->HairStrandsVF_Density			= VFInput.Strands.HairDensity;
diff --git a/Engine/Source/Runtime/Renderer/Private/HairStrands/HairStrandsVoxelization.cpp b/Engine/Source/Runtime/Renderer/Private/HairStrands/HairStrandsVoxelization.cpp
--- a/Engine/Source/Runtime/Renderer/Private/HairStrands/HairStrandsVoxelization.cpp
+++ b/Engine/Source/Runtime/Renderer/Private/HairStrands/HairStrandsVoxelization.cpp
@@ -989,7 +989,8 @@ class FVoxelRasterComputeCS : public FGlobalShader
 		SHADER_PARAMETER(uint32, HairStrandsVF_bIsCullingEnable)
 		SHADER_PARAMETER(float,	  HairStrandsVF_Density)
 		SHADER_PARAMETER(float,   HairStrandsVF_Radius)
-		SHADER_PARAMETER(float,	  HairStrandsVF_Length)
+		SHADER_PARAMETER(float,   HairStrandsVF_RootScale)
+		SHADER_PARAMETER(float,   HairStrandsVF_TipScale)
 		SHADER_PARAMETER(FVector, HairStrandsVF_PositionOffset)
 		SHADER_PARAMETER(uint32,  HairStrandsVF_VertexCount)
 		SHADER_PARAMETER(FMatrix, HairStrandsVF_LocalToWorldPrimitiveTransform)
@@ -1069,7 +1070,8 @@ static void AddVirtualVoxelizationComputeRasterPass(
 				PassParameters->HairStrandsVF_PositionOffsetBuffer = VFInput.Strands.PositionOffsetBuffer.SRV;
 				PassParameters->HairStrandsVF_VertexCount = VFInput.Strands.VertexCount;
 				PassParameters->HairStrandsVF_Radius = VFInput.Strands.HairRadius;
-				PassParameters->HairStrandsVF_Length = VFInput.Strands.HairLength;
+				PassParameters->HairStrandsVF_RootScale = VFInput.Strands.HairRootScale;
+				PassParameters->HairStrandsVF_TipScale = VFInput.Strands.HairTipScale;
 				PassParameters->HairStrandsVF_Density = VFInput.Strands.HairDensity;
 				PassParameters->HairStrandsVF_LocalToWorldPrimitiveTransform = VFInput.LocalToWorldTransform.ToMatrixWithScale();
 
diff --git a/Engine/Source/Runtime/Renderer/Public/HairStrandsInterface.h b/Engine/Source/Runtime/Renderer/Public/HairStrandsInterface.h
--- a/Engine/Source/Runtime/Renderer/Public/HairStrandsInterface.h
+++ b/Engine/Source/Runtime/Renderer/Public/HairStrandsInterface.h
@@ -235,6 +235,8 @@ public:
 
 			uint32 VertexCount = 0;
 			float HairRadius = 0;
+			float HairRootScale = 0;
+			float HairTipScale = 0;
 			float HairLength = 0;
 			float HairDensity = 0;
 			bool bUseStableRasterization = false;
