Fix CIS static analysis issue.

Static analysis get confused with nullptr pointer testing due to a unecessary check later down within a block of code.

diff --git a/Engine/Plugins/Runtime/HairStrands/Source/HairStrandsCore/Private/GroomComponent.cpp b/Engine/Plugins/Runtime/HairStrands/Source/HairStrandsCore/Private/GroomComponent.cpp
--- a/Engine/Plugins/Runtime/HairStrands/Source/HairStrandsCore/Private/GroomComponent.cpp
+++ b/Engine/Plugins/Runtime/HairStrands/Source/HairStrandsCore/Private/GroomComponent.cpp
@@ -2114,7 +2114,7 @@ void UGroomComponent::InitResources(bool bIsBindingReloading)
 		// LODBias is in the Modifier which is needed for LOD selection regardless if the strands are there or not
 		HairGroupInstance->Strands.Modifier = GetGroomGroupsDesc(GroomAsset, this, GroupIt);
 		#if WITH_EDITORONLY_DATA
-		HairGroupInstance->Debug.DebugMode = GroomAsset ? GroomAsset->GetDebugMode() : EHairStrandsDebugMode::NoneDebug;
+		HairGroupInstance->Debug.DebugMode = GroomAsset->GetDebugMode();
 		#endif// #if WITH_EDITORONLY_DATA
 
 		// Strands data/resources
