Fix hair groom resources naming macro.

diff --git a/Engine/Plugins/Runtime/HairStrands/Source/HairStrandsCore/Private/GroomResources.cpp b/Engine/Plugins/Runtime/HairStrands/Source/HairStrandsCore/Private/GroomResources.cpp
--- a/Engine/Plugins/Runtime/HairStrands/Source/HairStrandsCore/Private/GroomResources.cpp
+++ b/Engine/Plugins/Runtime/HairStrands/Source/HairStrandsCore/Private/GroomResources.cpp
@@ -29,7 +29,7 @@ enum class EHairResourceUsageType : uint8
 	Dynamic
 };
 
-#define HAIRSTRANDS_RESOUCE_NAME(Type, Name) (Type == EHairStrandsResourcesType::Guides  ? TEXT("Name##(Guides)") : (Type == EHairStrandsResourcesType::Strands ? TEXT("Name##(Strands)") : TEXT("Name##(Cards)")))
+#define HAIRSTRANDS_RESOUCE_NAME(Type, Name) (Type == EHairStrandsResourcesType::Guides  ? TEXT(#Name "(Guides)") : (Type == EHairStrandsResourcesType::Strands ? TEXT(#Name "(Strands)") : TEXT(#Name "(Cards)")))
 
 /////////////////////////////////////////////////////////////////////////////////////////
 // FRWBuffer utils 
