Fix issue when changing groom asset on the component

diff --git a/Engine/Plugins/Runtime/HairStrands/Source/HairStrandsCore/Private/GroomComponent.cpp b/Engine/Plugins/Runtime/HairStrands/Source/HairStrandsCore/Private/GroomComponent.cpp
--- a/Engine/Plugins/Runtime/HairStrands/Source/HairStrandsCore/Private/GroomComponent.cpp
+++ b/Engine/Plugins/Runtime/HairStrands/Source/HairStrandsCore/Private/GroomComponent.cpp
@@ -2528,6 +2528,7 @@ void UGroomComponent::InvalidateAndRecreate()
 void UGroomComponent::OnRegister()
 {
 	Super::OnRegister();
+	UpdateHairGroupsDesc();
 	UpdateHairSimulation();
 
 	// Insure the parent skeletal mesh is the same than the registered skeletal mesh, and if not reinitialized resources
