Integrate 15823512 and 15823993:

diff --git a/Engine/Plugins/Runtime/HairStrands/Source/HairStrandsCore/Private/Niagara/NiagaraDataInterfaceHairStrands.cpp b/Engine/Plugins/Runtime/HairStrands/Source/HairStrandsCore/Private/Niagara/NiagaraDataInterfaceHairStrands.cpp
--- a/Engine/Plugins/Runtime/HairStrands/Source/HairStrandsCore/Private/Niagara/NiagaraDataInterfaceHairStrands.cpp
+++ b/Engine/Plugins/Runtime/HairStrands/Source/HairStrandsCore/Private/Niagara/NiagaraDataInterfaceHairStrands.cpp
@@ -16,7 +16,7 @@
 #include "GlobalShader.h"
 
 #include "GroomComponent.h"
-#include "GroomAsset.h"  
+#include "GroomAsset.h"
 
 #define LOCTEXT_NAMESPACE "NiagaraDataInterfaceHairStrands"
 
@@ -148,7 +148,7 @@ static const FName ComputeMaterialFrameName(TEXT("ComputeMaterialFrame"));
 
 //------------------------------------------------------------------------------------------------------------
 
-static const FName ComputeAirDragForceName(TEXT("ComputeAirDragForce"));  
+static const FName ComputeAirDragForceName(TEXT("ComputeAirDragForce"));
 
 //------------------------------------------------------------------------------------------------------------
 
@@ -283,9 +283,9 @@ struct FNDIHairStrandsParametersName
 //------------------------------------------------------------------------------------------------------------
 
 void FNDIHairStrandsBuffer::Initialize(
-	const FHairStrandsDatas*  HairStrandsDatas, 
-	const FHairStrandsRestResource* HairStrandsRestResource, 
-	const FHairStrandsDeformedResource*  HairStrandsDeformedResource, 
+	const FHairStrandsDatas*  HairStrandsDatas,
+	const FHairStrandsRestResource* HairStrandsRestResource,
+	const FHairStrandsDeformedResource*  HairStrandsDeformedResource,
 	const FHairStrandsRestRootResource* HairStrandsRestRootResource,
 	const FHairStrandsDeformedRootResource* HairStrandsDeformedRootResource,
 	const TStaticArray<float, 32 * NumScales>& InParamsScale)
@@ -351,13 +351,13 @@ void FNDIHairStrandsBuffer::InitRHI()
 
 			BoundingBoxBuffer.Initialize(TEXT("BoundingBoxBuffer"), sizeof(uint32), BoundCount, EPixelFormat::PF_R32_UINT, BUF_Static);
 			void* BoundBufferData = RHILockBuffer(BoundingBoxBuffer.Buffer, 0, BoundBytes, RLM_WriteOnly);
-			
+
 			FMemory::Memcpy(BoundBufferData, ZeroData.GetData(), BoundBytes);
 			RHIUnlockBuffer(BoundingBoxBuffer.Buffer);
 		}
 		{
 			const uint32 ScaleCount = 32 * NumScales;
-			const uint32 ScaleBytes = sizeof(float) * ScaleCount;  
+			const uint32 ScaleBytes = sizeof(float) * ScaleCount;
 
 			ParamsScaleBuffer.Initialize(TEXT("ParamsScaleBuffer"), sizeof(float), ScaleCount, EPixelFormat::PF_R32_FLOAT, BUF_Static);
 			void* ScaleBufferData = RHILockBuffer(ParamsScaleBuffer.Buffer, 0, ScaleBytes, RLM_WriteOnly);
@@ -412,7 +412,7 @@ void FNDIHairStrandsData::Release()
 	}
 }
 
-void FNDIHairStrandsData::Update(UNiagaraDataInterfaceHairStrands* Interface, FNiagaraSystemInstance* SystemInstance, const FHairStrandsDatas* StrandsDatas, 
+void FNDIHairStrandsData::Update(UNiagaraDataInterfaceHairStrands* Interface, FNiagaraSystemInstance* SystemInstance, const FHairStrandsDatas* StrandsDatas,
 	UGroomAsset* GroomAsset, const int32 GroupIndex, const int32 LODIndex, const FTransform& LocalToWorld)
 {
 	if (Interface != nullptr)
@@ -504,7 +504,7 @@ bool FNDIHairStrandsData::Init(UNiagaraDataInterfaceHairStrands* Interface, FNia
 			HairStrandsBuffer->Initialize(StrandsDatas, StrandsRestResource, StrandsDeformedResource, StrandsRestRootResource, StrandsDeformedRootResource, ParamsScale);
 
 			BeginInitResource(HairStrandsBuffer);
-		
+
 			TickCount = 0;
 			ForceReset = true;
 		}
@@ -533,7 +533,7 @@ struct FNDIHairStrandsParametersCS : public FNiagaraDataInterfaceParametersCS
 		WorldInverse.Bind(ParameterMap, *ParamNames.WorldInverseName);
 		WorldRotation.Bind(ParameterMap, *ParamNames.WorldRotationName);
 		NumStrands.Bind(ParameterMap, *ParamNames.NumStrandsName);
-		StrandSize.Bind(ParameterMap, *ParamNames.StrandSizeName); 
+		StrandSize.Bind(ParameterMap, *ParamNames.StrandSizeName);
 
 		DeformedPositionBuffer.Bind(ParameterMap, *ParamNames.DeformedPositionBufferName);
 		CurvesOffsetsBuffer.Bind(ParameterMap, *ParamNames.CurvesOffsetsBufferName);
@@ -565,12 +565,12 @@ struct FNDIHairStrandsParametersCS : public FNiagaraDataInterfaceParametersCS
 		BoundingBoxBuffer.Bind(ParameterMap, *ParamNames.BoundingBoxBufferName);
 		ParamsScaleBuffer.Bind(ParameterMap, *ParamNames.ParamsScaleBufferName);
 
-		
+
 	}
 
 	void Set(FRHICommandList& RHICmdList, const FNiagaraDataInterfaceSetArgs& Context) const
 	{
-		check(IsInRenderingThread()); 
+		check(IsInRenderingThread());
 
 		FRHIComputeShader* ComputeShaderRHI = RHICmdList.GetBoundComputeShader();
 
@@ -657,7 +657,7 @@ struct FNDIHairStrandsParametersCS : public FNiagaraDataInterfaceParametersCS
 			RHICmdList.Transition(MakeArrayView(Transitions, UE_ARRAY_COUNT(Transitions)));
 
 			SetUAVParameter(RHICmdList, ComputeShaderRHI, DeformedPositionBuffer, PointPositionsUAV);
-	
+
 			SetUAVParameter(RHICmdList, ComputeShaderRHI, BoundingBoxBuffer, HairStrandsBuffer->BoundingBoxBuffer.UAV);
 
 			SetSRVParameter(RHICmdList, ComputeShaderRHI, CurvesOffsetsBuffer, HairStrandsBuffer->CurvesOffsetsBuffer.SRV);
@@ -910,16 +910,16 @@ void UNiagaraDataInterfaceHairStrands::ExtractSourceComponent(FNiagaraSystemInst
 					}
 				}
 			}
-		}		
+		}
 	}
 }
 
 void UNiagaraDataInterfaceHairStrands::ExtractDatasAndResources(
 	FNiagaraSystemInstance* SystemInstance,
-	FHairStrandsDatas*& OutStrandsDatas, 
-	FHairStrandsRestResource*& OutStrandsRestResource, 
-	FHairStrandsDeformedResource*& OutStrandsDeformedResource, 
-	FHairStrandsRestRootResource*& OutStrandsRestRootResource, 
+	FHairStrandsDatas*& OutStrandsDatas,
+	FHairStrandsRestResource*& OutStrandsRestResource,
+	FHairStrandsDeformedResource*& OutStrandsDeformedResource,
+	FHairStrandsRestRootResource*& OutStrandsRestRootResource,
 	FHairStrandsDeformedRootResource*& OutStrandsDeformedRootResource,
 	UGroomAsset*& OutGroomAsset,
 	int32& OutGroupIndex, 
@@ -942,10 +942,13 @@ void UNiagaraDataInterfaceHairStrands::ExtractDatasAndResources(
 		{
 			if (UNiagaraComponent* NiagaraComponent = SourceComponent->NiagaraComponents[NiagaraIndex])
 			{
-				if (NiagaraComponent->GetSystemInstance() == SystemInstance)
+				if (FNiagaraSystemInstanceControllerConstPtr SystemInstanceController = NiagaraComponent->GetSystemInstanceController())
 				{
-					OutGroupIndex = NiagaraIndex;
-					break;
+					if (SystemInstanceController->GetSystemInstanceID() == SystemInstance->GetId())
+					{
+						OutGroupIndex = NiagaraIndex;
+						break;
+					}
 				}
 			}
 		}
@@ -1539,7 +1542,7 @@ void UNiagaraDataInterfaceHairStrands::GetFunctions(TArray<FNiagaraFunctionSigna
 		Sig.Inputs.Add(FNiagaraVariable(FNiagaraTypeDefinition::GetVec3Def(), TEXT("Angular Velocity")));
 		Sig.Inputs.Add(FNiagaraVariable(FNiagaraTypeDefinition::GetQuatDef(), TEXT("Node Orientation")));
 		Sig.Outputs.Add(FNiagaraVariable(FNiagaraTypeDefinition::GetVec3Def(), TEXT("Angular Velocity")));
-		Sig.Outputs.Add(FNiagaraVariable(FNiagaraTypeDefinition::GetQuatDef(), TEXT("Node Orientation"))); 
+		Sig.Outputs.Add(FNiagaraVariable(FNiagaraTypeDefinition::GetQuatDef(), TEXT("Node Orientation")));
 
 		OutFunctions.Add(Sig);
 	}
@@ -2925,242 +2928,242 @@ void UNiagaraDataInterfaceHairStrands::GetWorldInverse(FVectorVMContext& Context
 
 void UNiagaraDataInterfaceHairStrands::GetBoundingBox(FVectorVMContext& Context)
 {
-	// @todo : implement function for cpu 
+	// @todo : implement function for cpu
 }
 
 void UNiagaraDataInterfaceHairStrands::ResetBoundingBox(FVectorVMContext& Context)
 {
-	// @todo : implement function for cpu 
+	// @todo : implement function for cpu
 }
 
 void UNiagaraDataInterfaceHairStrands::BuildBoundingBox(FVectorVMContext& Context)
 {
-	// @todo : implement function for cpu 
+	// @todo : implement function for cpu
 }
 
 void UNiagaraDataInterfaceHairStrands::GetPointPosition(FVectorVMContext& Context)
 {
-	// @todo : implement function for cpu 
+	// @todo : implement function for cpu
 }
 
 void UNiagaraDataInterfaceHairStrands::ComputeNodePosition(FVectorVMContext& Context)
 {
-	// @todo : implement function for cpu 
+	// @todo : implement function for cpu
 }
 
 void UNiagaraDataInterfaceHairStrands::ComputeNodeOrientation(FVectorVMContext& Context)
 {
-	// @todo : implement function for cpu 
+	// @todo : implement function for cpu
 }
 
 void UNiagaraDataInterfaceHairStrands::ComputeNodeMass(FVectorVMContext& Context)
 {
-	// @todo : implement function for cpu 
+	// @todo : implement function for cpu
 }
 
 void UNiagaraDataInterfaceHairStrands::ComputeNodeInertia(FVectorVMContext& Context)
 {
-	// @todo : implement function for cpu 
+	// @todo : implement function for cpu
 }
 
 void UNiagaraDataInterfaceHairStrands::ComputeEdgeLength(FVectorVMContext& Context)
 {
-	// @todo : implement function for cpu 
+	// @todo : implement function for cpu
 }
 
 void UNiagaraDataInterfaceHairStrands::ComputeEdgeRotation(FVectorVMContext& Context)
 {
-	// @todo : implement function for cpu 
+	// @todo : implement function for cpu
 }
 
 void UNiagaraDataInterfaceHairStrands::ComputeRestPosition(FVectorVMContext& Context)
 {
-	// @todo : implement function for cpu 
+	// @todo : implement function for cpu
 }
 
 void UNiagaraDataInterfaceHairStrands::ComputeRestOrientation(FVectorVMContext& Context)
 {
-	// @todo : implement function for cpu 
+	// @todo : implement function for cpu
 }
 
 void UNiagaraDataInterfaceHairStrands::ComputeLocalState(FVectorVMContext& Context)
 {
-	// @todo : implement function for cpu 
+	// @todo : implement function for cpu
 }
 
 void UNiagaraDataInterfaceHairStrands::UpdatePointPosition(FVectorVMContext& Context)
 {
-	// @todo : implement function for cpu 
+	// @todo : implement function for cpu
 }
 
 void UNiagaraDataInterfaceHairStrands::ResetPointPosition(FVectorVMContext& Context)
 {
-	// @todo : implement function for cpu 
+	// @todo : implement function for cpu
 }
 
 void UNiagaraDataInterfaceHairStrands::AttachNodePosition(FVectorVMContext& Context)
 {
-	// @todo : implement function for cpu 
+	// @todo : implement function for cpu
 }
 
 void UNiagaraDataInterfaceHairStrands::EvalSkinnedPosition(FVectorVMContext& Context)
 {
-	// @todo : implement function for cpu 
+	// @todo : implement function for cpu
 }
 
 void UNiagaraDataInterfaceHairStrands::AttachNodeOrientation(FVectorVMContext& Context)
 {
-	// @todo : implement function for cpu 
+	// @todo : implement function for cpu
 }
 
 void UNiagaraDataInterfaceHairStrands::AttachNodeState(FVectorVMContext& Context)
 {
-	// @todo : implement function for cpu 
+	// @todo : implement function for cpu
 }
 
 void UNiagaraDataInterfaceHairStrands::UpdateNodeState(FVectorVMContext& Context)
 {
-	// @todo : implement function for cpu 
+	// @todo : implement function for cpu
 }
 
 void UNiagaraDataInterfaceHairStrands::AdvectNodePosition(FVectorVMContext& Context)
 {
-	// @todo : implement function for cpu 
+	// @todo : implement function for cpu
 }
 
 void UNiagaraDataInterfaceHairStrands::AdvectNodeOrientation(FVectorVMContext& Context)
 {
-	// @todo : implement function for cpu 
+	// @todo : implement function for cpu
 }
 
 void UNiagaraDataInterfaceHairStrands::UpdateLinearVelocity(FVectorVMContext& Context)
 {
-	// @todo : implement function for cpu 
+	// @todo : implement function for cpu
 }
 
 void UNiagaraDataInterfaceHairStrands::UpdateAngularVelocity(FVectorVMContext& Context)
 {
-	// @todo : implement function for cpu 
+	// @todo : implement function for cpu
 }
 
 void UNiagaraDataInterfaceHairStrands::SetupDistanceSpringMaterial(FVectorVMContext& Context)
 {
-	// @todo : implement function for cpu 
+	// @todo : implement function for cpu
 }
 
 void UNiagaraDataInterfaceHairStrands::SolveDistanceSpringMaterial(FVectorVMContext& Context)
 {
-	// @todo : implement function for cpu 
+	// @todo : implement function for cpu
 }
 
 void UNiagaraDataInterfaceHairStrands::ProjectDistanceSpringMaterial(FVectorVMContext& Context)
 {
-	// @todo : implement function for cpu 
+	// @todo : implement function for cpu
 }
 
 void UNiagaraDataInterfaceHairStrands::SetupAngularSpringMaterial(FVectorVMContext& Context)
 {
-	// @todo : implement function for cpu 
+	// @todo : implement function for cpu
 }
 
 void UNiagaraDataInterfaceHairStrands::SolveAngularSpringMaterial(FVectorVMContext& Context)
 {
-	// @todo : implement function for cpu 
+	// @todo : implement function for cpu
 }
 
 void UNiagaraDataInterfaceHairStrands::ProjectAngularSpringMaterial(FVectorVMContext& Context)
 {
-	// @todo : implement function for cpu 
+	// @todo : implement function for cpu
 }
 
 
 void UNiagaraDataInterfaceHairStrands::SetupStretchRodMaterial(FVectorVMContext& Context)
 {
-	// @todo : implement function for cpu 
+	// @todo : implement function for cpu
 }
 
 void UNiagaraDataInterfaceHairStrands::SolveStretchRodMaterial(FVectorVMContext& Context)
 {
-	// @todo : implement function for cpu 
+	// @todo : implement function for cpu
 }
 
 void UNiagaraDataInterfaceHairStrands::ProjectStretchRodMaterial(FVectorVMContext& Context)
 {
-	// @todo : implement function for cpu 
+	// @todo : implement function for cpu
 }
 
 void UNiagaraDataInterfaceHairStrands::SetupBendRodMaterial(FVectorVMContext& Context)
 {
-	// @todo : implement function for cpu 
+	// @todo : implement function for cpu
 }
 
 void UNiagaraDataInterfaceHairStrands::SolveBendRodMaterial(FVectorVMContext& Context)
 {
-	// @todo : implement function for cpu 
+	// @todo : implement function for cpu
 }
 
 void UNiagaraDataInterfaceHairStrands::ProjectBendRodMaterial(FVectorVMContext& Context)
 {
-	// @todo : implement function for cpu 
+	// @todo : implement function for cpu
 }
 
 void UNiagaraDataInterfaceHairStrands::ComputeEdgeDirection(FVectorVMContext& Context)
 {
-	// @todo : implement function for cpu 
+	// @todo : implement function for cpu
 }
 
 void UNiagaraDataInterfaceHairStrands::UpdateMaterialFrame(FVectorVMContext& Context)
 {
-	// @todo : implement function for cpu 
+	// @todo : implement function for cpu
 }
 
 void UNiagaraDataInterfaceHairStrands::ComputeMaterialFrame(FVectorVMContext& Context)
 {
-	// @todo : implement function for cpu 
+	// @todo : implement function for cpu
 }
 
 void UNiagaraDataInterfaceHairStrands::SolveHardCollisionConstraint(FVectorVMContext& Context)
 {
-	// @todo : implement function for cpu 
+	// @todo : implement function for cpu
 }
 
 void UNiagaraDataInterfaceHairStrands::ProjectHardCollisionConstraint(FVectorVMContext& Context)
 {
-	// @todo : implement function for cpu 
+	// @todo : implement function for cpu
 }
 
 void UNiagaraDataInterfaceHairStrands::SolveSoftCollisionConstraint(FVectorVMContext& Context)
 {
-	// @todo : implement function for cpu 
+	// @todo : implement function for cpu
 }
 
 void UNiagaraDataInterfaceHairStrands::ProjectSoftCollisionConstraint(FVectorVMContext& Context)
 {
-	// @todo : implement function for cpu 
+	// @todo : implement function for cpu
 }
 void UNiagaraDataInterfaceHairStrands::SetupSoftCollisionConstraint(FVectorVMContext& Context)
 {
-	// @todo : implement function for cpu 
+	// @todo : implement function for cpu
 }
 
 void UNiagaraDataInterfaceHairStrands::ComputeAirDragForce(FVectorVMContext& Context)
 {
-	// @todo : implement function for cpu 
+	// @todo : implement function for cpu
 }
 
 void UNiagaraDataInterfaceHairStrands::NeedSimulationReset(FVectorVMContext& Context)
 {
-	// @todo : implement function for cpu 
+	// @todo : implement function for cpu
 }
 
 void UNiagaraDataInterfaceHairStrands::InitGridSamples(FVectorVMContext& Context)
 {
-	// @todo : implement function for cpu 
+	// @todo : implement function for cpu
 }
 
 void UNiagaraDataInterfaceHairStrands::GetSampleState(FVectorVMContext& Context)
 {
-	// @todo : implement function for cpu 
+	// @todo : implement function for cpu
 }
 
 void UNiagaraDataInterfaceHairStrands::HasGlobalInterpolation(FVectorVMContext& Context)
@@ -3340,7 +3343,7 @@ bool UNiagaraDataInterfaceHairStrands::GetFunctionHLSL(const FNiagaraDataInterfa
 		static const TCHAR *FormatSample = TEXT(R"(
 			void {InstanceFunctionName} (in float3 NodePosition, in int NodeOffset, out float OutEdgeLength)
 			{
-				{HairStrandsContextName} 
+				{HairStrandsContextName}
 				if(NodeOffset == 2)
 				{
 					DIHairStrands_ComputeEdgeVolume(DIContext,NodePosition,OutEdgeLength);
@@ -3419,7 +3422,7 @@ bool UNiagaraDataInterfaceHairStrands::GetFunctionHLSL(const FNiagaraDataInterfa
 				)");
 		OutHLSL += FString::Format(FormatSample, ArgsSample);
 		return true;
-	}	
+	}
 	else if (FunctionInfo.DefinitionName == AttachNodeStateName)
 	{
 		static const TCHAR *FormatSample = TEXT(R"(
@@ -3467,7 +3470,7 @@ bool UNiagaraDataInterfaceHairStrands::GetFunctionHLSL(const FNiagaraDataInterfa
 	else if (FunctionInfo.DefinitionName == AdvectNodePositionName)
 	{
 		static const TCHAR *FormatSample = TEXT(R"(
-			void {InstanceFunctionName} (in float NodeMass, in bool IsPositionMobile, in float3 ExternalForce, in float3 ForceGradient, in float DeltaTime, 
+			void {InstanceFunctionName} (in float NodeMass, in bool IsPositionMobile, in float3 ExternalForce, in float3 ForceGradient, in float DeltaTime,
 									     in float3 LinearVelocity, in float3 NodePosition, out float3 OutLinearVelocity, out float3 OutNodePosition)
 			{
 				OutLinearVelocity = LinearVelocity;
@@ -3481,7 +3484,7 @@ bool UNiagaraDataInterfaceHairStrands::GetFunctionHLSL(const FNiagaraDataInterfa
 	else if (FunctionInfo.DefinitionName == AdvectNodeOrientationName)
 	{
 		static const TCHAR *FormatSample = TEXT(R"(
-			void {InstanceFunctionName} (in float3 NodeInertia, in bool IsOrientationMobile, in float3 ExternalTorque, in float3 TorqueGradient, in float DeltaTime, 
+			void {InstanceFunctionName} (in float3 NodeInertia, in bool IsOrientationMobile, in float3 ExternalTorque, in float3 TorqueGradient, in float DeltaTime,
 										 in float3 AngularVelocity, in float4 NodeOrientation, out float3 OutAngularVelocity, out float4 OutNodeOrientation)
 			{
 				OutAngularVelocity = AngularVelocity;
@@ -3550,10 +3553,10 @@ bool UNiagaraDataInterfaceHairStrands::GetFunctionHLSL(const FNiagaraDataInterfa
 	else if (FunctionInfo.DefinitionName == SetupDistanceSpringMaterialName)
 	{
 		static const TCHAR *FormatSample = TEXT(R"(
-				void {InstanceFunctionName} (in float YoungModulus, in float RodThickness, 
+				void {InstanceFunctionName} (in float YoungModulus, in float RodThickness,
 in float RestLength, in float DeltaTime, in int NodeOffset, in float MaterialDamping, out float OutMaterialCompliance, out float OutMaterialWeight, out float OutMaterialMultiplier)
 				{
-					{HairStrandsContextName} 
+					{HairStrandsContextName}
 					if(NodeOffset == 0)
 					{
 						SetupStretchSpringMaterial(DIContext.StrandSize,YoungModulus,RodThickness,RestLength,DeltaTime,false,MaterialDamping,OutMaterialCompliance,OutMaterialWeight,OutMaterialMultiplier);
@@ -3574,10 +3577,10 @@ in float RestLength, in float DeltaTime, in int NodeOffset, in float MaterialDam
 	else if (FunctionInfo.DefinitionName == SolveDistanceSpringMaterialName)
 	{
 		static const TCHAR *FormatSample = TEXT(R"(
-				void {InstanceFunctionName} (in bool EnableConstraint, in float RestLength, in float DeltaTime, in int NodeOffset, in float MaterialDamping, 
+				void {InstanceFunctionName} (in bool EnableConstraint, in float RestLength, in float DeltaTime, in int NodeOffset, in float MaterialDamping,
 		in float MaterialCompliance, in float MaterialWeight, in float MaterialMultiplier, out float OutMaterialMultiplier)
 				{
-					{HairStrandsContextName} 
+					{HairStrandsContextName}
 					if(NodeOffset == 0)
 					{
 						SolveStretchSpringMaterial(EnableConstraint,DIContext.StrandSize,RestLength,DeltaTime,MaterialDamping,MaterialCompliance,MaterialWeight,MaterialMultiplier,OutMaterialMultiplier);
@@ -3600,7 +3603,7 @@ in float RestLength, in float DeltaTime, in int NodeOffset, in float MaterialDam
 		static const TCHAR *FormatSample = TEXT(R"(
 				void {InstanceFunctionName} (in bool EnableConstraint, in float YoungModulus, in float RodThickness, in float RestLength, in float DeltaTime, in int NodeOffset, out float3 OutNodePosition)
 				{
-					{HairStrandsContextName} 
+					{HairStrandsContextName}
 					if(NodeOffset == 0)
 					{
 						ProjectStretchSpringMaterial(EnableConstraint,DIContext.StrandSize,YoungModulus,RodThickness,RestLength,DeltaTime,OutNodePosition);
@@ -3621,7 +3624,7 @@ in float RestLength, in float DeltaTime, in int NodeOffset, in float MaterialDam
 	else if (FunctionInfo.DefinitionName == SetupAngularSpringMaterialName)
 	{
 		static const TCHAR *FormatSample = TEXT(R"(
-					void {InstanceFunctionName} (in float YoungModulus, in float RodThickness, 
+					void {InstanceFunctionName} (in float YoungModulus, in float RodThickness,
 	in float RestLength, in float DeltaTime, in float MaterialDamping, out float OutMaterialCompliance, out float OutMaterialWeight, out float3 OutMaterialMultiplier)
 					{
 						{HairStrandsContextName} SetupAngularSpringMaterial(DIContext.StrandSize,YoungModulus,RodThickness,RestLength,DeltaTime,false,MaterialDamping,OutMaterialCompliance,OutMaterialWeight,OutMaterialMultiplier);
@@ -3633,7 +3636,7 @@ in float RestLength, in float DeltaTime, in int NodeOffset, in float MaterialDam
 	else if (FunctionInfo.DefinitionName == SolveAngularSpringMaterialName)
 	{
 		static const TCHAR *FormatSample = TEXT(R"(
-					void {InstanceFunctionName} (in bool EnableConstraint, in float RestLength, in float3 RestDirection, in float DeltaTime, in float MaterialDamping, 
+					void {InstanceFunctionName} (in bool EnableConstraint, in float RestLength, in float3 RestDirection, in float DeltaTime, in float MaterialDamping,
 			in float MaterialCompliance, in float MaterialWeight, in float3 MaterialMultiplier, out float3 OutMaterialMultiplier)
 					{
 						{HairStrandsContextName} SolveAngularSpringMaterial(EnableConstraint,DIContext.StrandSize,RestLength, RestDirection,DeltaTime,MaterialDamping,MaterialCompliance,MaterialWeight,MaterialMultiplier,OutMaterialMultiplier);
@@ -3656,7 +3659,7 @@ in float RestLength, in float DeltaTime, in int NodeOffset, in float MaterialDam
 	else if (FunctionInfo.DefinitionName == SetupStretchRodMaterialName)
 	{
 		static const TCHAR *FormatSample = TEXT(R"(
-					void {InstanceFunctionName} (in float YoungModulus, in float RodThickness, 
+					void {InstanceFunctionName} (in float YoungModulus, in float RodThickness,
 	in float RestLength, in float DeltaTime, in float MaterialDamping, out float OutMaterialCompliance, out float OutMaterialWeight, out float3 OutMaterialMultiplier)
 					{
 						{HairStrandsContextName} SetupStretchRodMaterial(DIContext.StrandSize,YoungModulus,RodThickness,RestLength,DeltaTime,false,MaterialDamping,OutMaterialCompliance,OutMaterialWeight,OutMaterialMultiplier);
@@ -3668,7 +3671,7 @@ in float RestLength, in float DeltaTime, in int NodeOffset, in float MaterialDam
 	else if (FunctionInfo.DefinitionName == SolveStretchRodMaterialName)
 	{
 		static const TCHAR *FormatSample = TEXT(R"(
-					void {InstanceFunctionName} (in bool EnableConstraint, in float RestLength, in float DeltaTime, in float MaterialDamping, 
+					void {InstanceFunctionName} (in bool EnableConstraint, in float RestLength, in float DeltaTime, in float MaterialDamping,
 			in float MaterialCompliance, in float MaterialWeight, in float3 MaterialMultiplier, out float3 OutMaterialMultiplier)
 					{
 						{HairStrandsContextName} SolveStretchRodMaterial(EnableConstraint,DIContext.StrandSize,RestLength,DeltaTime,MaterialDamping,MaterialCompliance,MaterialWeight,MaterialMultiplier,OutMaterialMultiplier);
@@ -3691,7 +3694,7 @@ in float RestLength, in float DeltaTime, in int NodeOffset, in float MaterialDam
 	else if (FunctionInfo.DefinitionName == SetupBendRodMaterialName)
 	{
 		static const TCHAR *FormatSample = TEXT(R"(
-						void {InstanceFunctionName} (in float YoungModulus, in float RodThickness, 
+						void {InstanceFunctionName} (in float YoungModulus, in float RodThickness,
 		in float RestLength, in float DeltaTime, in float MaterialDamping, out float OutMaterialCompliance, out float OutMaterialWeight, out float3 OutMaterialMultiplier)
 						{
 							{HairStrandsContextName} SetupBendRodMaterial(DIContext.StrandSize,YoungModulus,RodThickness,RestLength,DeltaTime,false,MaterialDamping,OutMaterialCompliance,OutMaterialWeight,OutMaterialMultiplier);
@@ -3703,7 +3706,7 @@ in float RestLength, in float DeltaTime, in int NodeOffset, in float MaterialDam
 	else if (FunctionInfo.DefinitionName == SolveBendRodMaterialName)
 	{
 		static const TCHAR *FormatSample = TEXT(R"(
-						void {InstanceFunctionName} (in bool EnableConstraint, in float RestLength, in float4 RestRotation, in float DeltaTime, in float MaterialDamping, 
+						void {InstanceFunctionName} (in bool EnableConstraint, in float RestLength, in float4 RestRotation, in float DeltaTime, in float MaterialDamping,
 				in float MaterialCompliance, in float MaterialWeight, in float3 MaterialMultiplier, out float3 OutMaterialMultiplier)
 						{
 							{HairStrandsContextName} SolveBendRodMaterial(EnableConstraint,DIContext.StrandSize,RestLength,RestRotation,DeltaTime,MaterialDamping,MaterialCompliance,MaterialWeight,MaterialMultiplier,OutMaterialMultiplier);
@@ -3726,7 +3729,7 @@ in float RestLength, in float DeltaTime, in int NodeOffset, in float MaterialDam
 	else if (FunctionInfo.DefinitionName == SolveHardCollisionConstraintName)
 	{
 		static const TCHAR *FormatSample = TEXT(R"(
-							void {InstanceFunctionName} (in bool EnableConstraint, in float PenetrationDepth, in float3 CollisionPosition, in float3 CollisionVelocity, in float3 CollisionNormal, 
+							void {InstanceFunctionName} (in bool EnableConstraint, in float PenetrationDepth, in float3 CollisionPosition, in float3 CollisionVelocity, in float3 CollisionNormal,
 				in float StaticFriction, in float KineticFriction, in float DeltaTime, out float3 OutMaterialMultiplier )
 							{
 								OutMaterialMultiplier = float3(0,0,0);
@@ -3740,7 +3743,7 @@ in float RestLength, in float DeltaTime, in int NodeOffset, in float MaterialDam
 	else if (FunctionInfo.DefinitionName == ProjectHardCollisionConstraintName)
 	{
 			static const TCHAR *FormatSample = TEXT(R"(
-						void {InstanceFunctionName} (in bool EnableConstraint, in float PenetrationDepth, in float3 CollisionPosition, in float3 CollisionVelocity, in float3 CollisionNormal, 
+						void {InstanceFunctionName} (in bool EnableConstraint, in float PenetrationDepth, in float3 CollisionPosition, in float3 CollisionVelocity, in float3 CollisionNormal,
 			in float StaticFriction, in float KineticFriction, in float DeltaTime, out float3 OutNodePosition )
 						{
 							{HairStrandsContextName} ProjectHardCollisionConstraint(EnableConstraint,DIContext.StrandSize,PenetrationDepth,
@@ -3753,8 +3756,8 @@ in float RestLength, in float DeltaTime, in int NodeOffset, in float MaterialDam
 	else if (FunctionInfo.DefinitionName == SolveSoftCollisionConstraintName)
 	{
 		static const TCHAR *FormatSample = TEXT(R"(
-								void {InstanceFunctionName} (in bool EnableConstraint, in float PenetrationDepth, in float3 CollisionPosition, in float3 CollisionVelocity, in float3 CollisionNormal, 
-					in float StaticFriction, in float KineticFriction, in float DeltaTime, in float MaterialDamping, 
+								void {InstanceFunctionName} (in bool EnableConstraint, in float PenetrationDepth, in float3 CollisionPosition, in float3 CollisionVelocity, in float3 CollisionNormal,
+					in float StaticFriction, in float KineticFriction, in float DeltaTime, in float MaterialDamping,
 			in float MaterialCompliance, in float MaterialWeight, in float3 MaterialMultiplier, out float3 OutMaterialMultiplier )
 								{
 									{HairStrandsContextName} SolveSoftCollisionConstraint(EnableConstraint,DIContext.StrandSize,PenetrationDepth,
@@ -3768,7 +3771,7 @@ in float RestLength, in float DeltaTime, in int NodeOffset, in float MaterialDam
 	else if (FunctionInfo.DefinitionName == ProjectSoftCollisionConstraintName)
 	{
 		static const TCHAR *FormatSample = TEXT(R"(
-							void {InstanceFunctionName} (in bool EnableConstraint, in float ConstraintStiffness, in float PenetrationDepth, in float3 CollisionPosition, in float3 CollisionVelocity, in float3 CollisionNormal, 
+							void {InstanceFunctionName} (in bool EnableConstraint, in float ConstraintStiffness, in float PenetrationDepth, in float3 CollisionPosition, in float3 CollisionVelocity, in float3 CollisionNormal,
 					in float StaticFriction, in float KineticFriction, in float DeltaTime, out float3 OutNodePosition )
 							{
 								{HairStrandsContextName} ProjectSoftCollisionConstraint(EnableConstraint,DIContext.StrandSize,ConstraintStiffness,PenetrationDepth,
@@ -3827,7 +3830,7 @@ in float RestLength, in float DeltaTime, in int NodeOffset, in float MaterialDam
 	else if (FunctionInfo.DefinitionName == ComputeAirDragForceName)
 	{
 		static const TCHAR *FormatSample = TEXT(R"(
-						void {InstanceFunctionName} (  in float AirDensity, in float AirViscosity, in float AirDrag, 
+						void {InstanceFunctionName} (  in float AirDensity, in float AirViscosity, in float AirDrag,
 		in float3 AirVelocity, in float NodeThickness, in float3 NodePosition, in float3 NodeVelocity, out float3 OutAirDrag, out float3 OutDragGradient )
 						{
 							{HairStrandsContextName} ComputeAirDragForce(DIContext.StrandSize,AirDensity,AirViscosity,AirDrag,AirVelocity,NodeThickness,NodePosition,NodeVelocity,OutAirDrag,OutDragGradient);
@@ -3839,12 +3842,12 @@ in float RestLength, in float DeltaTime, in int NodeOffset, in float MaterialDam
 	else if (FunctionInfo.DefinitionName == InitGridSamplesName)
 	{
 		static const TCHAR* FormatSample = TEXT(R"(
-							void {InstanceFunctionName} ( in float3 NodePosition, in float3 NodeVelocity, 
+							void {InstanceFunctionName} ( in float3 NodePosition, in float3 NodeVelocity,
 	in float NodeMass, in float GridLength, out int OutNumSamples,
 						out float3 OutDeltaPosition, out float3 OutDeltaVelocity, out float OutSampleMass)
 							{
 								{HairStrandsContextName} DIHairStrands_InitGridSamples(DIContext,
-										 NodePosition, NodeVelocity, NodeMass, GridLength, 
+										 NodePosition, NodeVelocity, NodeMass, GridLength,
 											OutNumSamples, OutDeltaPosition, OutDeltaVelocity, OutSampleMass);
 							}
 							)");
@@ -3854,11 +3857,11 @@ in float RestLength, in float DeltaTime, in int NodeOffset, in float MaterialDam
 	else if (FunctionInfo.DefinitionName == GetSampleStateName)
 	{
 		static const TCHAR* FormatSample = TEXT(R"(
-								void {InstanceFunctionName} ( in float3 NodePosition, in float3 NodeVelocity, in float3 DeltaPosition, in float3 DeltaVelocity, 
+								void {InstanceFunctionName} ( in float3 NodePosition, in float3 NodeVelocity, in float3 DeltaPosition, in float3 DeltaVelocity,
 			in int NumSamples, in int SampleIndex, out float3 OutSamplePosition, out float3 OutSampleVelocity)
 								{
 									{HairStrandsContextName} DIHairStrands_GetSampleState(DIContext,
-											 NodePosition, NodeVelocity, DeltaPosition, DeltaVelocity, 
+											 NodePosition, NodeVelocity, DeltaPosition, DeltaVelocity,
 											 NumSamples, SampleIndex, OutSamplePosition, OutSampleVelocity);
 								}
 								)");
