Fix hair strands invalid resources access.

This CL:
* Move hair strands data onto the scene rather than a global manager
* Move registration/unregistration onto the proxy resources creation/deletion
* Change instance resources to defer their deletion so that resources used by proxies are valid until the proxy are released.

diff --git a/Engine/Plugins/Runtime/HairStrands/Source/HairStrandsCore/Private/GroomComponent.cpp b/Engine/Plugins/Runtime/HairStrands/Source/HairStrandsCore/Private/GroomComponent.cpp
--- a/Engine/Plugins/Runtime/HairStrands/Source/HairStrandsCore/Private/GroomComponent.cpp
+++ b/Engine/Plugins/Runtime/HairStrands/Source/HairStrandsCore/Private/GroomComponent.cpp
@@ -606,12 +606,53 @@ public:
 
 	virtual ~FHairStrandsSceneProxy()
 	{
-
 	}
 	
+	/**
+	 *	Called when the rendering thread adds the proxy to the scene.
+	 *	This function allows for generating renderer-side resources.
+	 *	Called in the rendering thread.
+	 */
+	virtual void CreateRenderThreadResources() 
+	{
+		FPrimitiveSceneProxy::CreateRenderThreadResources();
+
+		// Register the data to the scene
+		FSceneInterface& LocalScene = GetScene();
+		TArray<FHairGroupInstance*> LocalInstances = HairGroupInstances;
+		for (FHairGroupInstance* Instance : LocalInstances)
+		{
+			if (Instance->IsValid() || Instance->Strands.ClusterCullingResource)
+			{
+				check(Instance->HairGroupPublicData != nullptr);
+				Instance->AddRef();
+				LocalScene.AddHairStrands(Instance);
+			}
+		}
+	}
+
+	/**
+	 *	Called when the rendering thread removes the proxy from the scene.
+	 *	This function allows for removing renderer-side resources.
+	 *	Called in the rendering thread.
+	 */
 	virtual void DestroyRenderThreadResources() override
 	{
 		FPrimitiveSceneProxy::DestroyRenderThreadResources();
+
+		// Unregister the data to the scene
+		FSceneInterface& LocalScene = GetScene();
+		TArray<FHairGroupInstance*> LocalInstances = HairGroupInstances;
+		for (FHairGroupInstance* Instance : LocalInstances)
+		{
+			if (Instance->IsValid() || Instance->Strands.ClusterCullingResource)
+			{
+				check(Instance->GetRefCount() > 0);
+				LocalScene.RemoveHairStrands(Instance);
+				Instance->Release();
+			}
+		}
+
 		CardsAndMeshesVertexFactories.ReleaseResources();
 		StrandsVertexFactory.ReleaseResource();
 	}
@@ -663,6 +704,7 @@ public:
 			const FHairGroup& GroupData = HairGroups[GroupIt];
 
 			FHairGroupInstance* Instance = HairGroupInstances[GroupIt];
+			check(Instance->GetRefCount() > 0);
 			FMatrix OverrideLocalToWorld = GetLocalToWorld();
 			GetOverrideLocalToTransform(Instance, OverrideLocalToWorld);
 
@@ -756,6 +798,8 @@ public:
 			{
 				for (uint32 GroupIt = 0; GroupIt < GroupCount; ++GroupIt)
 				{
+					check(Instances[GroupIt]->GetRefCount() > 0);
+
 					FMaterialRenderProxy* Strands_MaterialProxy = nullptr;
 					const EHairStrandsDebugMode DebugMode = GetHairStrandsGeometryDebugMode(Instances[GroupIt]);
 					if (DebugMode != EHairStrandsDebugMode::NoneDebug)
@@ -822,6 +866,8 @@ public:
 			return nullptr;
 		}
 
+		check(Instance->GetRefCount());
+
 		const int32 IntLODIndex = Instance->HairGroupPublicData->GetIntLODIndex();
 		const FVertexFactory* VertexFactory = nullptr;
 		FIndexBuffer* IndexBuffer = nullptr;
@@ -957,6 +1003,7 @@ public:
 		const uint32 GroupCount = Instances.Num();
 		for (uint32 GroupIt = 0; GroupIt < GroupCount; ++GroupIt)
 		{
+			check(Instances[GroupIt]->GetRefCount());
 			const EHairGeometryType GeometryType = Instances[GroupIt]->GeometryType;
 			const FHairGroup& GroupData = HairGroups[GroupIt];
 			bUseCardsOrMesh = bUseCardsOrMesh || GeometryType == EHairGeometryType::Cards || GeometryType == EHairGeometryType::Meshes;
@@ -1409,6 +1456,8 @@ void UGroomComponent::UpdateHairGroupsDescAndInvalidateRenderState()
 
 FPrimitiveSceneProxy* UGroomComponent::CreateSceneProxy()
 {
+	DeleteDeferredHairGroupInstances();
+
 	if (!GroomAsset || GroomAsset->GetNumHairGroups() == 0 || HairGroupInstances.Num() == 0)
 		return nullptr;
 
@@ -2003,6 +2052,7 @@ void UGroomComponent::InitResources(bool bIsBindingReloading)
 	for (int32 GroupIt = 0, GroupCount = GroomAsset->HairGroupsData.Num(); GroupIt < GroupCount; ++GroupIt)
 	{
 		FHairGroupInstance* HairGroupInstance = new FHairGroupInstance();
+		HairGroupInstance->AddRef();
 		HairGroupInstances.Add(HairGroupInstance);
 		HairGroupInstance->WorldType = WorldType;
 		HairGroupInstance->Debug.ComponentId = ComponentId.PrimIDValue;
@@ -2349,22 +2399,6 @@ void UGroomComponent::InitResources(bool bIsBindingReloading)
 	{
 		return;
 	}
-
-	TArray<FHairGroupInstance*> LocalInstances = HairGroupInstances;
-	ENQUEUE_RENDER_COMMAND(FHairStrandsBuffers)(
-	[ LocalInstances ] (FRHICommandListImmediate& RHICmdList)
-	{
-		const uint32 GroupCount = LocalInstances.Num();
-		for (uint32 GroupIt = 0; GroupIt < GroupCount; ++GroupIt)
-		{
-			FHairGroupInstance* HairGroupInstance = LocalInstances[GroupIt];
-			// An empty groom doesn't have a ClusterCullingResource and shouldn't be registered
-			if (HairGroupInstance->IsValid() || HairGroupInstance->Strands.ClusterCullingResource)
-			{
-				RegisterHairStrands(HairGroupInstance);
-			}
-		}
-	});
 }
 
 template<typename T>
@@ -2380,23 +2414,36 @@ void InternalResourceRelease(T*& In)
 
 void UGroomComponent::ReleaseResources()
 {
-	// Unregister component interpolation resources
-	const FPrimitiveComponentId LocalComponentId = ComponentId;
-	const uint32 Id = LocalComponentId.PrimIDValue;
-	ENQUEUE_RENDER_COMMAND(StaticMeshVertexBuffersLegacyInit)(
-		[Id](FRHICommandListImmediate& RHICmdList)
-	{
-		UnregisterHairStrands(Id);
-	});
 
+	FHairStrandsSceneProxy* GroomSceneProxy = (FHairStrandsSceneProxy*)SceneProxy;
 	InitializedResources = nullptr;
-	for (FHairGroupInstance* Instance : HairGroupInstances)
+
+	// Deferring instances deletion to insure scene proxy are done with the rendering data
+	DeferredDeleteHairGroupInstances = HairGroupInstances;
+	HairGroupInstances.Empty();
+
+	// Insure the ticking of the Groom component always happens after the skeletalMeshComponent.
+	if (RegisteredMeshComponent)
+	{
+		RemoveTickPrerequisiteComponent(RegisteredMeshComponent);
+	}
+	SkeletalPreviousPositionOffset = FVector::ZeroVector;
+	RegisteredMeshComponent = nullptr;
+
+	MarkRenderStateDirty();
+}
+
+void UGroomComponent::DeleteDeferredHairGroupInstances()
+{
+	FHairStrandsSceneProxy* GroomSceneProxy = (FHairStrandsSceneProxy*)SceneProxy;
+	for (FHairGroupInstance* Instance : DeferredDeleteHairGroupInstances)
 	{
 		FHairGroupInstance* LocalInstance = Instance;
 		ENQUEUE_RENDER_COMMAND(FHairStrandsBuffers)(
-		[LocalInstance](FRHICommandListImmediate& RHICmdList)
+		[LocalInstance, GroomSceneProxy](FRHICommandListImmediate& RHICmdList)
 		{
-			// Release the root resources only if they have been created internally (vs. being created by external asset)
+			// Sanity check
+			check(LocalInstance->GetRefCount() == 1);
 
 			// Guides
 			if (LocalInstance->Guides.IsValid())
@@ -2450,20 +2497,11 @@ void UGroomComponent::ReleaseResources()
 			}
 
 			InternalResourceRelease(LocalInstance->HairGroupPublicData);
+			LocalInstance->Release();
 			delete LocalInstance;
 		});
 	}
-	HairGroupInstances.Empty();
-
-	// Insure the ticking of the Groom component always happens after the skeletalMeshComponent.
-	if (RegisteredMeshComponent)
-	{
-		RemoveTickPrerequisiteComponent(RegisteredMeshComponent);
-	}
-	SkeletalPreviousPositionOffset = FVector::ZeroVector;
-	RegisteredMeshComponent = nullptr;
-
-	MarkRenderStateDirty();
+	DeferredDeleteHairGroupInstances.Empty();
 }
 
 void UGroomComponent::PostLoad()
@@ -2565,9 +2603,16 @@ void UGroomComponent::BeginDestroy()
 	Super::BeginDestroy();
 }
 
+void UGroomComponent::FinishDestroy()
+{
+	DeleteDeferredHairGroupInstances();
+	Super::FinishDestroy();
+}
+
 void UGroomComponent::OnComponentDestroyed(bool bDestroyingHierarchy)
 {
 	ReleaseResources();
+	DeleteDeferredHairGroupInstances();
 
 #if WITH_EDITOR
 	if (bIsGroomAssetCallbackRegistered && GroomAsset)
diff --git a/Engine/Plugins/Runtime/HairStrands/Source/HairStrandsCore/Private/GroomDebug.cpp b/Engine/Plugins/Runtime/HairStrands/Source/HairStrandsCore/Private/GroomDebug.cpp
--- a/Engine/Plugins/Runtime/HairStrands/Source/HairStrandsCore/Private/GroomDebug.cpp
+++ b/Engine/Plugins/Runtime/HairStrands/Source/HairStrandsCore/Private/GroomDebug.cpp
@@ -56,14 +56,16 @@ bool IsHairStrandsSkinCacheEnable();
 static void GetGroomInterpolationData(
 	FRDGBuilder& GraphBuilder,
 	FGlobalShaderMap* ShaderMap,
-	const TArray<FHairGroupInstance*> Instances,
+	const FHairStrandsInstances& Instances,
 	const EWorldType::Type WorldType,
 	const EHairStrandsProjectionMeshType MeshType,
 	const FGPUSkinCache* SkinCache,
 	FHairStrandsProjectionMeshData::LOD& OutGeometries)
 {
-	for (const FHairGroupInstance* Instance : Instances)
+	for (FHairStrandsInstance* AbstractInstance : Instances)
 	{
+		FHairGroupInstance* Instance = static_cast<FHairGroupInstance*>(AbstractInstance);
+
 		if (!Instance || Instance->WorldType != WorldType || !Instance->Debug.MeshComponent)
 			continue;
 
@@ -773,9 +775,9 @@ void RunHairStrandsDebug(
 	FGlobalShaderMap* ShaderMap,
 	EWorldType::Type WorldType,
 	const FSceneView& View,
+	const FHairStrandsInstances& Instances,
 	const FGPUSkinCache* SkinCache,
 	const FShaderDrawDebugData* ShaderDrawData,
-	const TArray<FHairGroupInstance*>& Instances,
 	FRDGTextureRef SceneColorTexture,
 	FIntRect Viewport,
 	const TUniformBufferRef<FViewUniformShaderParameters>& ViewUniformBuffer)
@@ -810,8 +812,10 @@ void RunHairStrandsDebug(
 			Line = FString::Printf(TEXT("Registered hair groups count : %d"), Instances.Num());
 			Canvas.DrawShadowedString(X, Y += YStep, *Line, GetStatsFont(), DebugColor);
 
-			for (FHairGroupInstance* Instance : Instances)
+			for (FHairStrandsInstance* AbstractInstance : Instances)
 			{
+				FHairGroupInstance* Instance = static_cast<FHairGroupInstance*>(AbstractInstance);
+
 				const bool bIsActive = Instance->WorldType == WorldType;
 				if (!bIsActive)
 					continue;
@@ -871,8 +875,10 @@ void RunHairStrandsDebug(
 			auto RenderProjectionData = [&GraphBuilder, ShaderMap, Viewport, &ViewUniformBuffer, Instances, WorldType, SkinCache, &bClearDepth, SceneColorTexture, DepthTexture](EHairStrandsInterpolationType StrandType, bool bRestTriangle, bool bRestFrame, bool bDeformedTriangle, bool bDeformedFrame)
 			{
 				TArray<int32> HairLODIndices;
-				for (FHairGroupInstance* Instance : Instances)
+				for (FHairStrandsInstance* AbstractInstance : Instances)
 				{
+					FHairGroupInstance* Instance = static_cast<FHairGroupInstance*>(AbstractInstance);
+
 					if (Instance->WorldType != WorldType || !Instance->HairGroupPublicData || Instance->Guides.RestRootResource == nullptr || Instance->Guides.DeformedRootResource == nullptr || Instance->BindingType != EHairBindingType::Skinning)
 						continue;
 
@@ -986,22 +992,28 @@ void RunHairStrandsDebug(
 
 	if (GHairCardsVoxelDebug > 0)
 	{
-		for (FHairGroupInstance* Instance : Instances)
+		for (FHairStrandsInstance* AbstractInstance : Instances)
 		{
+			FHairGroupInstance* Instance = static_cast<FHairGroupInstance*>(AbstractInstance);
+
 			AddVoxelPlainRaymarchingPass(GraphBuilder, View, ShaderMap, Instance, ShaderDrawData, SceneColorTexture);
 		}
 	}
 
 	if (GHairCardsAtlasDebug > 0)
 	{
-		for (FHairGroupInstance* Instance : Instances)
+		for (FHairStrandsInstance* AbstractInstance : Instances)
 		{
+			FHairGroupInstance* Instance = static_cast<FHairGroupInstance*>(AbstractInstance);
+
 			AddDrawDebugCardsAtlasPass(GraphBuilder, View, ShaderMap, Instance, ShaderDrawData, SceneColorTexture);
 		}
 	}
 
-	for (FHairGroupInstance* Instance : Instances)
+	for (FHairStrandsInstance* AbstractInstance : Instances)
 	{
+		FHairGroupInstance* Instance = static_cast<FHairGroupInstance*>(AbstractInstance);
+
 		if (GHairCardsGuidesDebug_Ren > 0 || Instance->Debug.bDrawCardsGuides)
 			AddDrawDebugCardsGuidesPass(GraphBuilder, View, ShaderMap, Instance, ShaderDrawData, Instance->Debug.bDrawCardsGuides ? false : GHairCardsGuidesDebug_Ren == 1, true);
 		if (GHairCardsGuidesDebug_Sim > 0)
diff --git a/Engine/Plugins/Runtime/HairStrands/Source/HairStrandsCore/Private/GroomManager.cpp b/Engine/Plugins/Runtime/HairStrands/Source/HairStrandsCore/Private/GroomManager.cpp
--- a/Engine/Plugins/Runtime/HairStrands/Source/HairStrandsCore/Private/GroomManager.cpp
+++ b/Engine/Plugins/Runtime/HairStrands/Source/HairStrandsCore/Private/GroomManager.cpp
@@ -39,68 +39,6 @@ bool IsHairStrandsSkinCacheEnable()
 
 DEFINE_LOG_CATEGORY_STATIC(LogGroomManager, Log, All);
 
-////////////////////////////////////////////////////////////////////////////////////////////////////////////////////////
-// Runtime execution order (on the render thread):
-//  * Register
-//  * For each frame
-//		* Update
-//		* Update triangles information for dynamic meshes
-//		* RunHairStrandsInterpolation (Interpolation callback)
-//  * UnRegister
-//
-// This code supposed a  small number of instance (~10), and won't scale to large crowed (linear loop, lot of cache misses, ...)
-////////////////////////////////////////////////////////////////////////////////////////////////////////////////////////
-
-struct FHairStrandsManager
-{
-	// #hair_todo: change this array to a queue update, in order make processing/update thread safe.
-	TArray<FHairGroupInstance*> Instances;
-
-	FHairStrandsManager()
-	{
-		// Reserve a large a mount of object to avoid any potential memory reallocation, which 
-		// could cause some thread safety issue. This is a workaround against the non-thread-safe array
-		Instances.Reserve(64);
-	}
-};
-
-FHairStrandsManager GHairManager;
-
-void RegisterHairStrands(FHairGroupInstance* InInstance)
-{
-	for (const FHairGroupInstance* Instance : GHairManager.Instances)
-	{
-		if (Instance->Debug.ComponentId == InInstance->Debug.ComponentId && 
-			Instance->Debug.GroupIndex == InInstance->Debug.GroupIndex)
-		{
-			// Component already registered. This should not happen. 
-			UE_LOG(LogGroomManager, Warning, TEXT("Component already register. This should't happen. Please report this to a rendering engineer."))
-			return;
-		}
-	}
-
-	check(InInstance->HairGroupPublicData);
-	GHairManager.Instances.Add(InInstance);
-}
-
-void UnregisterHairStrands(uint32 ComponentId)
-{
-	for (int32 Index=0;Index< GHairManager.Instances.Num();)
-	{
-		const FHairGroupInstance* Instance = GHairManager.Instances[Index];
-
-		if (Instance->Debug.ComponentId == ComponentId)
-		{
-			GHairManager.Instances[Index] = GHairManager.Instances[GHairManager.Instances.Num()-1];
-			GHairManager.Instances.SetNum(GHairManager.Instances.Num() - 1, false);
-		}
-		else
-		{
-			++Index;
-		}
-	}
-}
-
 ///////////////////////////////////////////////////////////////////////////////////////////////////
 
 static bool IsInstanceVisible(
@@ -131,6 +69,7 @@ bool NeedsUpdateCardsMeshTriangles();
 static void RunInternalHairStrandsInterpolation(
 	FRDGBuilder& GraphBuilder,
 	const FSceneView* View,
+	const FHairStrandsInstances& Instances,
 	EWorldType::Type WorldType, 
 	const FGPUSkinCache* SkinCache,
 	const FShaderDrawDebugData* ShaderDrawData,
@@ -141,15 +80,21 @@ static void RunInternalHairStrandsInterpolation(
 	check(IsInRenderingThread());
 
 	// Update dynamic mesh triangles
-	for (FHairGroupInstance* Instance : GHairManager.Instances)
+	for (FHairStrandsInstance* AbstractInstance : Instances)
 	{
+		FHairGroupInstance* Instance = static_cast<FHairGroupInstance*>(AbstractInstance);
+
+		// HAIR_TODO: Do culling on the output of the view rather than the macro group cluster, in order to get consistent result with draw logic
 		// Frustum culling for rendering strands. Update position only for visible/in camera frustum
+		#if 0
 		if (Type == EHairStrandsInterpolationType::RenderStrands && !IsInstanceVisible(View, Instance))
 		{
 			continue;
 		}
+		#endif
 
 		int32 MeshLODIndex = -1;
+		check(Instance->WorldType == WorldType);
 		if (Instance->WorldType != WorldType || Instance->GeometryType == EHairGeometryType::NoneGeometry)
 			continue;
 	
@@ -177,8 +122,6 @@ static void RunInternalHairStrandsInterpolation(
 		{
 			BuildCacheGeometry(GraphBuilder, ShaderMap, GeometryCacheComponent, CachedGeometry);
 		}
-		if (CachedGeometry.Sections.Num() == 0)
-			continue;
 
 		FHairStrandsProjectionMeshData::LOD MeshDataLOD;
 		for (const FCachedGeometry::Section& Section : CachedGeometry.Sections)
@@ -351,16 +294,22 @@ static void RunInternalHairStrandsInterpolation(
 	// Reset deformation
 	if (EHairStrandsInterpolationType::SimulationStrands == Type)
 	{
-		for (FHairGroupInstance* Instance : GHairManager.Instances)
+		for (FHairStrandsInstance* AbstractInstance : Instances)
 		{
+			FHairGroupInstance* Instance = static_cast<FHairGroupInstance*>(AbstractInstance);
+
+			check(Instance->WorldType == WorldType);
 			if (Instance->WorldType != WorldType || Instance->GeometryType == EHairGeometryType::NoneGeometry)
 				continue;
 
+			// HAIR_TODO: Do culling on the output of the view rather than the macro group cluster, in order to get consistent result with draw logic
 			// Frustum culling guide deformation if the instance does not have any simulation
+			#if 0
 			if (Instance->GeometryType == EHairGeometryType::Strands && !Instance->Guides.bIsSimulationEnable  && !IsInstanceVisible(View, Instance))
 			{
 				continue;
 			}
+			#endif
 
 			ResetHairStrandsInterpolation(GraphBuilder, ShaderMap, Instance, Instance->Debug.MeshLODIndex);
 		}
@@ -369,16 +318,22 @@ static void RunInternalHairStrandsInterpolation(
 	// Hair interpolation
 	if (EHairStrandsInterpolationType::RenderStrands == Type)
 	{
-		for (FHairGroupInstance* Instance : GHairManager.Instances)
+		for (FHairStrandsInstance* AbstractInstance : Instances)
 		{
+			FHairGroupInstance* Instance = static_cast<FHairGroupInstance*>(AbstractInstance);
+
+			check(Instance->WorldType == WorldType);
 			if (Instance->WorldType != WorldType || Instance->GeometryType == EHairGeometryType::NoneGeometry)
 				continue;
 
+			// HAIR_TODO: Do culling on the output of the view rather than the macro group cluster, in order to get consistent result with draw logic
 			// Frustum culling for rendering strands. Update position only for visible/in camera frustum
+			#if 0
 			if (!IsInstanceVisible(View, Instance))
 			{
 				continue;
 			}
+			#endif
 
  			ComputeHairStrandsInterpolation(
 				GraphBuilder, 
@@ -394,6 +349,7 @@ static void RunInternalHairStrandsInterpolation(
 static void RunHairStrandsInterpolation_Guide(
 	FRDGBuilder& GraphBuilder,
 	const FSceneView* View,
+	const FHairStrandsInstances& Instances,
 	EWorldType::Type WorldType,
 	const FGPUSkinCache* SkinCache,
 	const FShaderDrawDebugData* ShaderDrawData,
@@ -409,6 +365,7 @@ static void RunHairStrandsInterpolation_Guide(
 	RunInternalHairStrandsInterpolation(
 		GraphBuilder,
 		View,
+		Instances,
 		WorldType,
 		SkinCache,
 		ShaderDrawData,
@@ -420,6 +377,7 @@ static void RunHairStrandsInterpolation_Guide(
 static void RunHairStrandsInterpolation_Strands(
 	FRDGBuilder& GraphBuilder,
 	const FSceneView* View,
+	const FHairStrandsInstances& Instances,
 	EWorldType::Type WorldType,
 	const FGPUSkinCache* SkinCache,
 	const FShaderDrawDebugData* ShaderDrawData,
@@ -435,6 +393,7 @@ static void RunHairStrandsInterpolation_Strands(
 	RunInternalHairStrandsInterpolation(
 		GraphBuilder,
 		View,
+		Instances,
 		WorldType,
 		SkinCache,
 		ShaderDrawData,
@@ -445,11 +404,15 @@ static void RunHairStrandsInterpolation_Strands(
 
 
 static void RunHairStrandsGatherCluster(
+	const FHairStrandsInstances& Instances,
 	EWorldType::Type WorldType,
 	FHairStrandClusterData* ClusterData)
 {
-	for (FHairGroupInstance* Instance : GHairManager.Instances)
+	for (FHairStrandsInstance* AbstractInstance : Instances)
 	{
+		FHairGroupInstance* Instance = static_cast<FHairGroupInstance*>(AbstractInstance);
+
+		check(Instance->WorldType == WorldType);
 		if (Instance->WorldType != WorldType || Instance->GeometryType != EHairGeometryType::Strands)
 			continue;
 
@@ -511,7 +474,7 @@ static EHairGeometryType ConvertLODGeometryType(EHairGeometryType Type, bool Inb
 	return EHairGeometryType::NoneGeometry;
 }
 
-static void RunHairBufferSwap(EWorldType::Type WorldType, const TArray<const FSceneView*> Views)
+static void RunHairBufferSwap(const FHairStrandsInstances& Instances, EWorldType::Type WorldType, const TArray<const FSceneView*> Views)
 {
 	EShaderPlatform ShaderPlatform = EShaderPlatform::SP_NumPlatforms;
 	if (Views.Num() > 0)
@@ -519,9 +482,13 @@ static void RunHairBufferSwap(EWorldType::Type WorldType, const TArray<const FSc
 		ShaderPlatform = Views[0]->GetShaderPlatform();
 	}
 
-	for (FHairGroupInstance* Instance : GHairManager.Instances)
+	for (FHairStrandsInstance* AbstractInstance : Instances)
 	{
+		FHairGroupInstance* Instance = static_cast<FHairGroupInstance*>(AbstractInstance);
+
 		int32 MeshLODIndex = -1;
+		check(Instance);
+		check(Instance->WorldType == WorldType);
 		if (!Instance || Instance->WorldType != WorldType)
 			continue;
 
@@ -590,7 +557,7 @@ static void RunHairBufferSwap(EWorldType::Type WorldType, const TArray<const FSc
 	}
 }
 
-static void RunHairLODSelection(EWorldType::Type WorldType, const TArray<const FSceneView*> Views)
+static void RunHairLODSelection(const FHairStrandsInstances& Instances, EWorldType::Type WorldType, const TArray<const FSceneView*> Views)
 {
 	EShaderPlatform ShaderPlatform = EShaderPlatform::SP_NumPlatforms;
 	if (Views.Num() > 0)
@@ -598,11 +565,13 @@ static void RunHairLODSelection(EWorldType::Type WorldType, const TArray<const F
 		ShaderPlatform = Views[0]->GetShaderPlatform();
 	}
 
-	for (FHairGroupInstance* Instance : GHairManager.Instances)
+	for (FHairStrandsInstance* AbstractInstance : Instances)
 	{
+		FHairGroupInstance* Instance = static_cast<FHairGroupInstance*>(AbstractInstance);
+
 		int32 MeshLODIndex = -1;
-		if (!Instance || Instance->WorldType != WorldType)
-			continue;
+		check(Instance);
+		check(Instance->WorldType == WorldType);
 
 		check(Instance->HairGroupPublicData);
 
@@ -750,9 +719,9 @@ void RunHairStrandsDebug(
 	FGlobalShaderMap* ShaderMap,
 	EWorldType::Type WorldType,
 	const FSceneView& View,
+	const FHairStrandsInstances& Instances,
 	const FGPUSkinCache* SkinCache,
 	const FShaderDrawDebugData* ShaderDrawData,
-	const TArray<FHairGroupInstance*>& Instances,
 	FRDGTextureRef SceneColor,
 	FIntRect Viewport,
 	const TUniformBufferRef<FViewUniformShaderParameters>& ViewUniformBuffer);
@@ -765,6 +734,8 @@ void ProcessHairStrandsBookmark(
 	EHairStrandsBookmark Bookmark,
 	FHairStrandsBookmarkParameters& Parameters)
 {
+	check(Parameters.Instances != nullptr);
+
 	if (Bookmark == EHairStrandsBookmark::ProcessTasks)
 	{
 		const bool bHasHairStardsnProcess =
@@ -786,10 +757,12 @@ void ProcessHairStrandsBookmark(
 		if (GHairStrands_SwapBufferEndOfFrame <= 0)
 		{
 			RunHairBufferSwap(
+				*Parameters.Instances,
 				Parameters.WorldType,
 				Parameters.AllViews);
 		}
 		RunHairLODSelection(
+			*Parameters.Instances,
 			Parameters.WorldType,
 			Parameters.AllViews);
 	}
@@ -798,6 +771,7 @@ void ProcessHairStrandsBookmark(
 		if (GHairStrands_SwapBufferEndOfFrame > 0)
 		{
 			RunHairBufferSwap(
+				*Parameters.Instances,
 				Parameters.WorldType,
 				Parameters.AllViews);
 		}
@@ -808,6 +782,7 @@ void ProcessHairStrandsBookmark(
 		RunHairStrandsInterpolation_Guide(
 			*GraphBuilder,
 			Parameters.View,
+			*Parameters.Instances,
 			Parameters.WorldType,
 			Parameters.SkinCache,
 			Parameters.DebugShaderData,
@@ -817,6 +792,7 @@ void ProcessHairStrandsBookmark(
 	else if (Bookmark == EHairStrandsBookmark::ProcessGatherCluster)
 	{
 		RunHairStrandsGatherCluster(
+			*Parameters.Instances,
 			Parameters.WorldType,
 			&Parameters.HairClusterData);
 	}
@@ -826,6 +802,7 @@ void ProcessHairStrandsBookmark(
 		RunHairStrandsInterpolation_Strands(
 			*GraphBuilder,
 			Parameters.View,
+			*Parameters.Instances,
 			Parameters.WorldType,
 			Parameters.SkinCache,
 			Parameters.DebugShaderData,
@@ -840,9 +817,9 @@ void ProcessHairStrandsBookmark(
 			Parameters.ShaderMap,
 			Parameters.WorldType,
 			*Parameters.View,
+			*Parameters.Instances,
 			Parameters.SkinCache,
 			Parameters.DebugShaderData,
-			GHairManager.Instances,
 			Parameters.SceneColorTexture,
 			Parameters.View->UnscaledViewRect,
 			Parameters.View->ViewUniformBuffer);
@@ -851,5 +828,5 @@ void ProcessHairStrandsBookmark(
 
 void ProcessHairStrandsParameters(FHairStrandsBookmarkParameters& Parameters)
 {
-	Parameters.bHasElements = GHairManager.Instances.Num() > 0;
+	Parameters.bHasElements = Parameters.Instances && Parameters.Instances->Num() > 0;
 }
diff --git a/Engine/Plugins/Runtime/HairStrands/Source/HairStrandsCore/Public/GroomComponent.h b/Engine/Plugins/Runtime/HairStrands/Source/HairStrandsCore/Public/GroomComponent.h
--- a/Engine/Plugins/Runtime/HairStrands/Source/HairStrandsCore/Public/GroomComponent.h
+++ b/Engine/Plugins/Runtime/HairStrands/Source/HairStrandsCore/Public/GroomComponent.h
@@ -93,6 +93,7 @@ public:
 	virtual void OnUnregister() override;
 	virtual void OnComponentDestroyed(bool bDestroyingHierarchy) override;
 	virtual void BeginDestroy() override;
+	virtual void FinishDestroy() override;
 	virtual void OnAttachmentChanged() override;
 	virtual void DetachFromComponent(const FDetachmentTransformRules& DetachmentRules) override;
 	virtual void TickComponent(float DeltaTime, enum ELevelTick TickType, FActorComponentTickFunction *ThisTickFunction) override;
@@ -195,6 +196,7 @@ public:
 
 private:
 	TArray<FHairGroupInstance*> HairGroupInstances;
+	TArray<FHairGroupInstance*> DeferredDeleteHairGroupInstances;
 
 #if WITH_EDITORONLY_DATA
 	UPROPERTY(Transient)
@@ -210,6 +212,7 @@ protected:
 	virtual void OnChildDetached(USceneComponent* ChildComponent) override;
 
 private:
+	void DeleteDeferredHairGroupInstances();
 	void* InitializedResources;
 	class UMeshComponent* RegisteredMeshComponent;
 	FVector SkeletalPreviousPositionOffset;
diff --git a/Engine/Plugins/Runtime/HairStrands/Source/HairStrandsCore/Public/GroomInstance.h b/Engine/Plugins/Runtime/HairStrands/Source/HairStrandsCore/Public/GroomInstance.h
--- a/Engine/Plugins/Runtime/HairStrands/Source/HairStrandsCore/Public/GroomInstance.h
+++ b/Engine/Plugins/Runtime/HairStrands/Source/HairStrandsCore/Public/GroomInstance.h
@@ -37,7 +37,7 @@ END_GLOBAL_SHADER_PARAMETER_STRUCT()
 typedef TUniformBufferRef<FHairCardsVertexFactoryUniformShaderParameters> FHairCardsUniformBuffer;
 
 // Represent/Describe data & resources of a hair group belonging to a groom
-struct HAIRSTRANDSCORE_API FHairGroupInstance
+struct HAIRSTRANDSCORE_API FHairGroupInstance : public FHairStrandsInstance
 {
 	//////////////////////////////////////////////////////////////////////////////////////////
 	// Helper struct which aggregate strands based data/resources
diff --git a/Engine/Source/Runtime/Engine/Public/SceneInterface.h b/Engine/Source/Runtime/Engine/Public/SceneInterface.h
--- a/Engine/Source/Runtime/Engine/Public/SceneInterface.h
+++ b/Engine/Source/Runtime/Engine/Public/SceneInterface.h
@@ -35,6 +35,7 @@ class FViewInfo;
 class FSceneRenderer;
 class FComputeFramework;
 class FInstanceCullingManager;
+struct FHairStrandsInstance;
 
 enum EBasePassDrawListType
 {
@@ -310,6 +311,20 @@ public:
 	 */
 	virtual void RemoveVolumetricCloud(FVolumetricCloudSceneProxy* VolumetricCloudSceneProxy) = 0;
 
+	/**
+	 * Adds a hair strands proxy to the scene
+	 *
+	 * @param Proxy - the hair strands proxy
+	 */
+	virtual void AddHairStrands(FHairStrandsInstance* Proxy) = 0;
+
+	/**
+	 * Removes a hair strands proxy to the scene
+	 *
+	 * @param Proxy - the hair strands proxy
+	 */
+	virtual void RemoveHairStrands(FHairStrandsInstance* Proxy) = 0;
+
 	/**
 	 * Set the physics field scene proxy to the scene
 	 *
diff --git a/Engine/Source/Runtime/Renderer/Private/DeferredShadingRenderer.cpp b/Engine/Source/Runtime/Renderer/Private/DeferredShadingRenderer.cpp
--- a/Engine/Source/Runtime/Renderer/Private/DeferredShadingRenderer.cpp
+++ b/Engine/Source/Runtime/Renderer/Private/DeferredShadingRenderer.cpp
@@ -1933,8 +1933,7 @@ void FDeferredShadingSceneRenderer::Render(FRDGBuilder& GraphBuilder)
 	FHairStrandsBookmarkParameters& HairStrandsBookmarkParameters = *GraphBuilder.AllocObject<FHairStrandsBookmarkParameters>();
 	if (IsHairStrandsEnabled(EHairStrandsShaderType::All, Scene->GetShaderPlatform()))
 	{
-
-		HairStrandsBookmarkParameters = CreateHairStrandsBookmarkParameters(Views[0]);
+		HairStrandsBookmarkParameters = CreateHairStrandsBookmarkParameters(Scene, Views[0]);
 		RunHairStrandsBookmark(GraphBuilder, EHairStrandsBookmark::ProcessTasks, HairStrandsBookmarkParameters);
 
 		// Interpolation needs to happen after the skin cache run as there is a dependency 
@@ -2800,7 +2799,7 @@ void FDeferredShadingSceneRenderer::Render(FRDGBuilder& GraphBuilder)
 	{
 		if (HairStrandsBookmarkParameters.bHasElements)
 		{
-			RenderHairStrandsDebugInfo(GraphBuilder, Views, HairStrandsBookmarkParameters.HairClusterData, SceneTextures.Color.Target);
+			RenderHairStrandsDebugInfo(GraphBuilder, Scene, Views, HairStrandsBookmarkParameters.HairClusterData, SceneTextures.Color.Target);
 		}
 	}
 
diff --git a/Engine/Source/Runtime/Renderer/Private/HairStrands/HairStrandsDebug.cpp b/Engine/Source/Runtime/Renderer/Private/HairStrands/HairStrandsDebug.cpp
--- a/Engine/Source/Runtime/Renderer/Private/HairStrands/HairStrandsDebug.cpp
+++ b/Engine/Source/Runtime/Renderer/Private/HairStrands/HairStrandsDebug.cpp
@@ -1044,6 +1044,7 @@ static void AddDrawDebugClusterPass(
 uint32 GetHairStrandsMeanSamplePerPixel();
 static void InternalRenderHairStrandsDebugInfo(
 	FRDGBuilder& GraphBuilder,
+	FScene* Scene,
 	FViewInfo& View,
 	const struct FHairStrandClusterData& HairClusterData,
 	FRDGTextureRef SceneColorTexture)
@@ -1066,7 +1067,7 @@ static void InternalRenderHairStrandsDebugInfo(
 	const FSceneTextures& SceneTextures = FSceneTextures::Get(GraphBuilder);
 
 	{
-		FHairStrandsBookmarkParameters Params = CreateHairStrandsBookmarkParameters(View);
+		FHairStrandsBookmarkParameters Params = CreateHairStrandsBookmarkParameters(Scene, View);
 		Params.SceneColorTexture = SceneColorTexture;
 		RunHairStrandsBookmark(GraphBuilder, EHairStrandsBookmark::ProcessDebug, Params);
 	}
@@ -1365,6 +1366,7 @@ static void InternalRenderHairStrandsDebugInfo(
 
 void RenderHairStrandsDebugInfo(
 	FRDGBuilder& GraphBuilder,
+	FScene* Scene,
 	TArrayView<FViewInfo> Views,
 	const struct FHairStrandClusterData& HairClusterData,
 	FRDGTextureRef SceneColorTexture)
@@ -1374,7 +1376,7 @@ void RenderHairStrandsDebugInfo(
 	{
 		if (HairStrands::HasViewHairStrandsData(View))
 		{
-			InternalRenderHairStrandsDebugInfo(GraphBuilder, View, HairClusterData, SceneColorTexture);
+			InternalRenderHairStrandsDebugInfo(GraphBuilder, Scene, View, HairClusterData, SceneColorTexture);
 		}
 	}
 }
\ No newline at end of file
diff --git a/Engine/Source/Runtime/Renderer/Private/HairStrands/HairStrandsDebug.h b/Engine/Source/Runtime/Renderer/Private/HairStrands/HairStrandsDebug.h
--- a/Engine/Source/Runtime/Renderer/Private/HairStrands/HairStrandsDebug.h
+++ b/Engine/Source/Runtime/Renderer/Private/HairStrands/HairStrandsDebug.h
@@ -11,6 +11,7 @@
 
 void RenderHairStrandsDebugInfo(
 	FRDGBuilder& GraphBuilder,
+	FScene* Scene, 
 	TArrayView<FViewInfo> Views,
 	const struct FHairStrandClusterData& HairClusterData,
 	FRDGTextureRef SceneColorTexture);
\ No newline at end of file
diff --git a/Engine/Source/Runtime/Renderer/Private/HairStrands/HairStrandsInterface.cpp b/Engine/Source/Runtime/Renderer/Private/HairStrands/HairStrandsInterface.cpp
--- a/Engine/Source/Runtime/Renderer/Private/HairStrands/HairStrandsInterface.cpp
+++ b/Engine/Source/Runtime/Renderer/Private/HairStrands/HairStrandsInterface.cpp
@@ -16,6 +16,7 @@
 #include "SkeletalRenderPublic.h"
 #include "SceneRendering.h"
 #include "SystemTextures.h"
+#include "ScenePrivate.h"
 
 DEFINE_LOG_CATEGORY_STATIC(LogHairRendering, Log, All);
 
@@ -56,6 +57,25 @@ static TAutoConsoleVariable<int32> CVarHairStrandsSimulation(
 	TEXT("Enable/disable hair simulation"),
 	ECVF_RenderThreadSafe | ECVF_Scalability);
 
+////////////////////////////////////////////////////////////////////////////////////////////////////////////////////////
+// Hair strands instance ref. counting for debug purpose only
+uint32 FHairStrandsInstance::GetRefCount() const
+{
+	return RefCount;
+}
+
+uint32 FHairStrandsInstance::AddRef() const
+{
+	return ++RefCount;
+}
+
+uint32 FHairStrandsInstance::Release() const
+{
+	check(RefCount > 0);
+	uint32 LocalRefCount = --RefCount;
+	return LocalRefCount;
+}
+
 ////////////////////////////////////////////////////////////////////////////////////////////////////////////////////////
 // Import/export utils function for hair resources
 void FRDGExternalBuffer::Release()
@@ -366,13 +386,14 @@ void RunHairStrandsBookmark(EHairStrandsBookmark Bookmark, FHairStrandsBookmarkP
 }
 
 bool IsHairStrandsClusterCullingUseHzb();
-FHairStrandsBookmarkParameters CreateHairStrandsBookmarkParameters(FViewInfo& View)
+FHairStrandsBookmarkParameters CreateHairStrandsBookmarkParameters(FScene* Scene, FViewInfo& View)
 {
 	FHairStrandsBookmarkParameters Out;
 	Out.DebugShaderData			= &View.ShaderDrawData;
 	Out.SkinCache				= View.Family->Scene->GetGPUSkinCache();
 	Out.WorldType				= View.Family->Scene->GetWorld()->WorldType;
 	Out.ShaderMap				= View.ShaderMap;
+	Out.Instances				= &Scene->HairStrandsSceneData.RegisteredProxies;
 	Out.View					= &View;
 	Out.ViewRect				= View.ViewRect;
 	Out.SceneColorTexture		= nullptr;
@@ -386,10 +407,10 @@ FHairStrandsBookmarkParameters CreateHairStrandsBookmarkParameters(FViewInfo& Vi
 	return Out;
 }
 
-FHairStrandsBookmarkParameters CreateHairStrandsBookmarkParameters(TArray<FViewInfo>& Views)
+FHairStrandsBookmarkParameters CreateHairStrandsBookmarkParameters(FScene* Scene, TArray<FViewInfo>& Views)
 {
 	FHairStrandsBookmarkParameters Out;
-	Out = CreateHairStrandsBookmarkParameters(Views[0]);
+	Out = CreateHairStrandsBookmarkParameters(Scene, Views[0]);
 	Out.AllViews.Reserve(Views.Num());
 	for (const FViewInfo& View : Views)
 	{
diff --git a/Engine/Source/Runtime/Renderer/Private/HairStrands/HairStrandsRendering.h b/Engine/Source/Runtime/Renderer/Private/HairStrands/HairStrandsRendering.h
--- a/Engine/Source/Runtime/Renderer/Private/HairStrands/HairStrandsRendering.h
+++ b/Engine/Source/Runtime/Renderer/Private/HairStrands/HairStrandsRendering.h
@@ -45,5 +45,5 @@ void RunHairStrandsBookmark(
 	EHairStrandsBookmark Bookmark,
 	FHairStrandsBookmarkParameters& Parameters);
 
-FHairStrandsBookmarkParameters CreateHairStrandsBookmarkParameters(FViewInfo& View);
-FHairStrandsBookmarkParameters CreateHairStrandsBookmarkParameters(TArray<FViewInfo>& View);
\ No newline at end of file
+FHairStrandsBookmarkParameters CreateHairStrandsBookmarkParameters(FScene* Scene, FViewInfo& View);
+FHairStrandsBookmarkParameters CreateHairStrandsBookmarkParameters(FScene* Scene, TArray<FViewInfo>& View);
\ No newline at end of file
diff --git a/Engine/Source/Runtime/Renderer/Private/RendererScene.cpp b/Engine/Source/Runtime/Renderer/Private/RendererScene.cpp
--- a/Engine/Source/Runtime/Renderer/Private/RendererScene.cpp
+++ b/Engine/Source/Runtime/Renderer/Private/RendererScene.cpp
@@ -60,6 +60,8 @@
 #include "GPUScene.h"
 #include "HAL/LowLevelMemTracker.h"
 #include "VT/RuntimeVirtualTextureSceneProxy.h"
+#include "HairStrandsInterface.h"
+
 #if RHI_RAYTRACING
 #include "RayTracingDynamicGeometryCollection.h"
 #include "RayTracing/RayTracingScene.h"
@@ -1974,6 +1976,34 @@ void FScene::UpdateDecalFadeInTime(UDecalComponent* Decal)
 	}
 }
 
+void FScene::AddHairStrands(FHairStrandsInstance* Proxy)
+{
+	if (Proxy)
+	{
+		check(IsInRenderingThread());
+		const int32 PackedIndex = HairStrandsSceneData.RegisteredProxies.Add(Proxy);
+		Proxy->RegisteredIndex = PackedIndex;
+	}
+}
+
+void FScene::RemoveHairStrands(FHairStrandsInstance* Proxy)
+{
+	if (Proxy)
+	{
+		check(IsInRenderingThread());
+		int32 ProxyIndex = Proxy->RegisteredIndex;
+		if (HairStrandsSceneData.RegisteredProxies.IsValidIndex(ProxyIndex))
+		{
+			HairStrandsSceneData.RegisteredProxies.RemoveAtSwap(ProxyIndex);
+		}
+		Proxy->RegisteredIndex = -1;
+		if (HairStrandsSceneData.RegisteredProxies.IsValidIndex(ProxyIndex))
+		{
+			FHairStrandsInstance* Other = HairStrandsSceneData.RegisteredProxies[ProxyIndex];
+			Other->RegisteredIndex = ProxyIndex;
+		}
+	}
+}
 
 void FScene::AddReflectionCapture(UReflectionCaptureComponent* Component)
 {
@@ -4362,6 +4392,9 @@ public:
 	virtual FSkyAtmosphereRenderSceneInfo* GetSkyAtmosphereSceneInfo() override { return NULL; }
 	virtual const FSkyAtmosphereRenderSceneInfo* GetSkyAtmosphereSceneInfo() const override { return NULL; }
 
+	virtual void AddHairStrands(FHairStrandsInstance* Proxy) override {}
+	virtual void RemoveHairStrands(FHairStrandsInstance* Proxy) override {}
+
 	virtual void SetPhysicsField(FPhysicsFieldSceneProxy* PhysicsFieldSceneProxy) override {}
 	virtual void ResetPhysicsField() override {}
 	virtual void UpdatePhysicsField(FRDGBuilder& GraphBuilder, FViewInfo& View) override {}
diff --git a/Engine/Source/Runtime/Renderer/Private/ScenePrivate.h b/Engine/Source/Runtime/Renderer/Private/ScenePrivate.h
--- a/Engine/Source/Runtime/Renderer/Private/ScenePrivate.h
+++ b/Engine/Source/Runtime/Renderer/Private/ScenePrivate.h
@@ -90,6 +90,7 @@ class FRuntimeVirtualTextureSceneProxy;
 class FLumenSceneData;
 class FVirtualShadowMapArrayCacheManager;
 class FComputeFramework;
+struct FHairStrandsInstance;
 
 /** Holds information about a single primitive's occlusion. */
 class FPrimitiveOcclusionHistory
@@ -1671,6 +1672,13 @@ public:
 	void ResizeCubemapArrayGPU(uint32 InMaxCubemaps, int32 InCubemapSize);
 };
 
+/** Scene state used to manage hair strands. */
+class FHairStrandsSceneData
+{
+public:
+	TArray<FHairStrandsInstance*> RegisteredProxies;
+};
+
 class FVolumetricLightmapInterpolation
 {
 public:
@@ -2712,6 +2720,9 @@ public:
 	/** State needed for the reflection environment feature. */
 	FReflectionEnvironmentSceneData ReflectionSceneData;
 
+	/** The hair strands in the scene. */
+	FHairStrandsSceneData HairStrandsSceneData;
+
 	/** 
 	 * Precomputed lighting volumes in the scene, used for interpolating dynamic object lighting.
 	 * These are typically one per streaming level and they store volume lighting samples computed by Lightmass. 
@@ -2927,6 +2938,9 @@ public:
 	virtual void RemoveExponentialHeightFog(UExponentialHeightFogComponent* FogComponent) override;
 	virtual bool HasAnyExponentialHeightFog() const override;
 
+	virtual void AddHairStrands(FHairStrandsInstance* Proxy) override;
+	virtual void RemoveHairStrands(FHairStrandsInstance* Proxy) override;
+
 	virtual void AddSkyAtmosphere(FSkyAtmosphereSceneProxy* SkyAtmosphereSceneProxy, bool bStaticLightingBuilt) override;
 	virtual void RemoveSkyAtmosphere(FSkyAtmosphereSceneProxy* SkyAtmosphereSceneProxy) override;
 	virtual FSkyAtmosphereRenderSceneInfo* GetSkyAtmosphereSceneInfo() override { return SkyAtmosphere; }
diff --git a/Engine/Source/Runtime/Renderer/Private/SceneRendering.cpp b/Engine/Source/Runtime/Renderer/Private/SceneRendering.cpp
--- a/Engine/Source/Runtime/Renderer/Private/SceneRendering.cpp
+++ b/Engine/Source/Runtime/Renderer/Private/SceneRendering.cpp
@@ -3690,7 +3690,7 @@ static void RenderViewFamily_RenderThread(FRHICommandListImmediate& RHICmdList,
 			
 			if (IsHairStrandsEnabled(EHairStrandsShaderType::All, SceneRenderer->Scene->GetShaderPlatform()) && SceneRenderer->Views.Num() > 0 && !ViewFamily.EngineShowFlags.HitProxies)
 			{
-				FHairStrandsBookmarkParameters Parameters = CreateHairStrandsBookmarkParameters(SceneRenderer->Views);
+				FHairStrandsBookmarkParameters Parameters = CreateHairStrandsBookmarkParameters(SceneRenderer->Scene, SceneRenderer->Views);
 				if (Parameters.bHasElements)
 				{
 					RunHairStrandsBookmark(EHairStrandsBookmark::ProcessEndOfFrame, Parameters);
diff --git a/Engine/Source/Runtime/Renderer/Private/SceneVisibility.cpp b/Engine/Source/Runtime/Renderer/Private/SceneVisibility.cpp
--- a/Engine/Source/Runtime/Renderer/Private/SceneVisibility.cpp
+++ b/Engine/Source/Runtime/Renderer/Private/SceneVisibility.cpp
@@ -3147,7 +3147,7 @@ void FSceneRenderer::PreVisibilityFrameSetup(FRDGBuilder& GraphBuilder, const FS
 
 	if (Views.Num() > 0 && !ViewFamily.EngineShowFlags.HitProxies)
 	{
-		FHairStrandsBookmarkParameters Parameters = CreateHairStrandsBookmarkParameters(Views);
+		FHairStrandsBookmarkParameters Parameters = CreateHairStrandsBookmarkParameters(Scene, Views);
 		if (Parameters.bHasElements)
 		{
 			RunHairStrandsBookmark(EHairStrandsBookmark::ProcessLODSelection, Parameters);
@@ -3156,7 +3156,7 @@ void FSceneRenderer::PreVisibilityFrameSetup(FRDGBuilder& GraphBuilder, const FS
 
 	if (IsHairStrandsEnabled(EHairStrandsShaderType::All, Scene->GetShaderPlatform()) && Views.Num() > 0 && !ViewFamily.EngineShowFlags.HitProxies)
 	{
-		FHairStrandsBookmarkParameters Parameters = CreateHairStrandsBookmarkParameters(Views);
+		FHairStrandsBookmarkParameters Parameters = CreateHairStrandsBookmarkParameters(Scene, Views);
 		RunHairStrandsBookmark(GraphBuilder, EHairStrandsBookmark::ProcessGuideInterpolation, Parameters);
 	}
 
diff --git a/Engine/Source/Runtime/Renderer/Public/HairStrandsInterface.h b/Engine/Source/Runtime/Renderer/Public/HairStrandsInterface.h
--- a/Engine/Source/Runtime/Renderer/Public/HairStrandsInterface.h
+++ b/Engine/Source/Runtime/Renderer/Public/HairStrandsInterface.h
@@ -132,6 +132,19 @@ enum EHairInterpolationType
 
 ////////////////////////////////////////////////////////////////////////////////////////////////////
 // Public group data 
+
+struct RENDERER_API FHairStrandsInstance
+{
+	virtual ~FHairStrandsInstance() = default;
+	uint32 GetRefCount() const;
+	uint32 AddRef() const;
+	uint32 Release() const;
+	int32 RegisteredIndex = -1;
+protected:
+	mutable uint32 RefCount = 0;
+};
+typedef RENDERER_API TArray<FHairStrandsInstance*> FHairStrandsInstances;
+
 class RENDERER_API FHairGroupPublicData : public FRenderResource
 {
 public:
@@ -416,6 +429,7 @@ struct FHairStrandsBookmarkParameters
 	class FGlobalShaderMap* ShaderMap = nullptr;
 
 	FIntRect ViewRect; // View 0
+	FHairStrandsInstances* Instances = nullptr;
 	const FSceneView* View = nullptr;// // View 0
 	TArray<const FSceneView*> AllViews;
 	FRDGTextureRef SceneColorTexture{};
