Remove WorldType access/check within hair instance code.

This is a simple/straightforward clean up following CL 15915567.

diff --git a/Engine/Plugins/Runtime/HairStrands/Source/HairStrandsCore/Private/GroomComponent.cpp b/Engine/Plugins/Runtime/HairStrands/Source/HairStrandsCore/Private/GroomComponent.cpp
--- a/Engine/Plugins/Runtime/HairStrands/Source/HairStrandsCore/Private/GroomComponent.cpp
+++ b/Engine/Plugins/Runtime/HairStrands/Source/HairStrandsCore/Private/GroomComponent.cpp
@@ -1767,18 +1767,11 @@ FHairStrandsDeformedRootResource* UGroomComponent::GetGuideStrandsDeformedRootRe
 	return HairGroupInstances[GroupIndex]->Guides.DeformedRootResource;
 }
 
-EWorldType::Type UGroomComponent::GetWorldType() const
-{
-	EWorldType::Type WorldType = GetWorld() ? EWorldType::Type(GetWorld()->WorldType) : EWorldType::None;
-	return WorldType == EWorldType::Inactive ? EWorldType::Editor : WorldType;
-}
-
 void UGroomComponent::UpdateSimulatedGroups()
 {
 	if (HairGroupInstances.Num()>0)
 	{
 		const uint32 Id = ComponentId.PrimIDValue;
-		const EWorldType::Type WorldType = GetWorldType();
 		
 		const bool bIsStrandsEnabled = IsHairStrandsEnabled(EHairStrandsShaderType::Strands);
 
@@ -1789,7 +1782,7 @@ void UGroomComponent::UpdateSimulatedGroups()
 		UGroomAsset* LocalGroomAsset = GroomAsset;
 		UGroomBindingAsset* LocalBindingAsset = BindingAsset;
 		ENQUEUE_RENDER_COMMAND(FHairStrandsTick_UEnableSimulatedGroups)(
-			[LocalInstances, LocalGroomAsset, LocalBindingAsset, Id, WorldType, bIsStrandsEnabled, LODIndex](FRHICommandListImmediate& RHICmdList)
+			[LocalInstances, LocalGroomAsset, LocalBindingAsset, Id, bIsStrandsEnabled, LODIndex](FRHICommandListImmediate& RHICmdList)
 		{
 			int32 GroupIt = 0;
 			for (FHairGroupInstance* Instance : LocalInstances)
@@ -2048,13 +2041,11 @@ void UGroomComponent::InitResources(bool bIsBindingReloading)
 	}
 
 	const bool bIsStrandsEnabled = IsHairStrandsEnabled(EHairStrandsShaderType::Strands);
-	const EWorldType::Type WorldType = GetWorldType();
 	for (int32 GroupIt = 0, GroupCount = GroomAsset->HairGroupsData.Num(); GroupIt < GroupCount; ++GroupIt)
 	{
 		FHairGroupInstance* HairGroupInstance = new FHairGroupInstance();
 		HairGroupInstance->AddRef();
 		HairGroupInstances.Add(HairGroupInstance);
-		HairGroupInstance->WorldType = WorldType;
 		HairGroupInstance->Debug.ComponentId = ComponentId.PrimIDValue;
 		HairGroupInstance->Debug.GroupIndex = GroupIt;
 		HairGroupInstance->Debug.GroupCount = GroupCount;
@@ -2578,17 +2569,6 @@ void UGroomComponent::OnRegister()
 	{
 		InitResources();
 	}
-
-	const EWorldType::Type WorldType = GetWorldType();
-	TArray<FHairGroupInstance*> LocalHairGroupInstances = HairGroupInstances;
-	ENQUEUE_RENDER_COMMAND(FHairStrandsRegister)(
-		[WorldType, LocalHairGroupInstances](FRHICommandListImmediate& RHICmdList)
-	{
-		for (FHairGroupInstance* Instance : LocalHairGroupInstances)
-		{
-			Instance->WorldType = WorldType;
-		}
-	});
 }
 
 void UGroomComponent::OnUnregister()
@@ -2657,7 +2637,6 @@ void UGroomComponent::TickComponent(float DeltaTime, enum ELevelTick TickType, F
 {
 	Super::TickComponent(DeltaTime, TickType, ThisTickFunction);	
 	
-	const EWorldType::Type WorldType = GetWorldType();
 	const uint32 Id = ComponentId.PrimIDValue;
 	const ERHIFeatureLevel::Type FeatureLevel = GetWorld() ? ERHIFeatureLevel::Type(GetWorld()->FeatureLevel) : ERHIFeatureLevel::Num;
 
@@ -2722,14 +2701,13 @@ void UGroomComponent::TickComponent(float DeltaTime, enum ELevelTick TickType, F
 	const FTransform SkinLocalToWorld = RegisteredMeshComponent ? RegisteredMeshComponent->GetComponentTransform() : FTransform();
 	TArray<FHairGroupInstance*> LocalHairGroupInstances = HairGroupInstances;
 	ENQUEUE_RENDER_COMMAND(FHairStrandsTick_TransformUpdate)(
-		[Id, WorldType, SkinLocalToWorld, FeatureLevel, LocalHairGroupInstances](FRHICommandListImmediate& RHICmdList)
+		[Id, SkinLocalToWorld, FeatureLevel, LocalHairGroupInstances](FRHICommandListImmediate& RHICmdList)
 	{		
 		if (ERHIFeatureLevel::Num == FeatureLevel)
 			return;
 
 		for (FHairGroupInstance* Instance : LocalHairGroupInstances)
 		{
-			Instance->WorldType = WorldType;
 			Instance->Debug.SkeletalPreviousLocalToWorld = Instance->Debug.SkeletalLocalToWorld;
 			Instance->Debug.SkeletalLocalToWorld = SkinLocalToWorld;
 		}
diff --git a/Engine/Plugins/Runtime/HairStrands/Source/HairStrandsCore/Private/GroomDebug.cpp b/Engine/Plugins/Runtime/HairStrands/Source/HairStrandsCore/Private/GroomDebug.cpp
--- a/Engine/Plugins/Runtime/HairStrands/Source/HairStrandsCore/Private/GroomDebug.cpp
+++ b/Engine/Plugins/Runtime/HairStrands/Source/HairStrandsCore/Private/GroomDebug.cpp
@@ -57,7 +57,6 @@ static void GetGroomInterpolationData(
 	FRDGBuilder& GraphBuilder,
 	FGlobalShaderMap* ShaderMap,
 	const FHairStrandsInstances& Instances,
-	const EWorldType::Type WorldType,
 	const EHairStrandsProjectionMeshType MeshType,
 	const FGPUSkinCache* SkinCache,
 	FHairStrandsProjectionMeshData::LOD& OutGeometries)
@@ -66,7 +65,7 @@ static void GetGroomInterpolationData(
 	{
 		FHairGroupInstance* Instance = static_cast<FHairGroupInstance*>(AbstractInstance);
 
-		if (!Instance || Instance->WorldType != WorldType || !Instance->Debug.MeshComponent)
+		if (!Instance || !Instance->Debug.MeshComponent)
 			continue;
 
 		FCachedGeometry CachedGeometry;
@@ -773,7 +772,6 @@ static const TCHAR* ToString(EHairBindingType In)
 void RunHairStrandsDebug(
 	FRDGBuilder& GraphBuilder,
 	FGlobalShaderMap* ShaderMap,
-	EWorldType::Type WorldType,
 	const FSceneView& View,
 	const FHairStrandsInstances& Instances,
 	const FGPUSkinCache* SkinCache,
@@ -816,10 +814,6 @@ void RunHairStrandsDebug(
 			{
 				FHairGroupInstance* Instance = static_cast<FHairGroupInstance*>(AbstractInstance);
 
-				const bool bIsActive = Instance->WorldType == WorldType;
-				if (!bIsActive)
-					continue;
-
 				Line = FString::Printf(TEXT(" * Group:%d/%d | LOD:%1.2f/%d | GeometryType:%s | BindingType:%s | Sim:%d | RBF:%d | VertexCount:%d | Name: %s"),
 					Instance->Debug.GroupIndex,
 					Instance->Debug.GroupCount,
@@ -833,7 +827,7 @@ void RunHairStrandsDebug(
 					Instance->Guides.bHasGlobalInterpolation,
 					Instance->HairGroupPublicData->VertexCount,
 					*Instance->Debug.GroomAssetName);
-				Canvas.DrawShadowedString(X, Y += YStep, *Line, GetStatsFont(), bIsActive ? DebugGroupColor : InactiveColor);
+				Canvas.DrawShadowedString(X, Y += YStep, *Line, GetStatsFont(), DebugGroupColor);
 			}
 
 			const bool bFlush = false;
@@ -855,10 +849,10 @@ void RunHairStrandsDebug(
 
 			if (GHairDebugMeshProjection_SkinCacheMesh > 0)
 			{
-				auto RenderMeshProjection = [&bClearDepth, WorldType, ShaderMap, Viewport, &ViewUniformBuffer, Instances, SkinCache, &SceneColorTexture, &DepthTexture, &GraphBuilder](FRDGBuilder& LocalGraphBuilder, EHairStrandsProjectionMeshType MeshType)
+				auto RenderMeshProjection = [&bClearDepth, ShaderMap, Viewport, &ViewUniformBuffer, Instances, SkinCache, &SceneColorTexture, &DepthTexture, &GraphBuilder](FRDGBuilder& LocalGraphBuilder, EHairStrandsProjectionMeshType MeshType)
 				{
 					FHairStrandsProjectionMeshData::LOD MeshProjectionLODData;
-					GetGroomInterpolationData(GraphBuilder, ShaderMap, Instances, WorldType, MeshType, SkinCache, MeshProjectionLODData);
+					GetGroomInterpolationData(GraphBuilder, ShaderMap, Instances, MeshType, SkinCache, MeshProjectionLODData);
 					for (FHairStrandsProjectionMeshData::Section& Section : MeshProjectionLODData.Sections)
 					{
 						AddDebugProjectionMeshPass(LocalGraphBuilder, ShaderMap, Viewport, ViewUniformBuffer, MeshType, bClearDepth, Section, SceneColorTexture, DepthTexture);
@@ -872,14 +866,14 @@ void RunHairStrandsDebug(
 				RenderMeshProjection(GraphBuilder, EHairStrandsProjectionMeshType::TargetMesh);
 			}
 
-			auto RenderProjectionData = [&GraphBuilder, ShaderMap, Viewport, &ViewUniformBuffer, Instances, WorldType, SkinCache, &bClearDepth, SceneColorTexture, DepthTexture](EHairStrandsInterpolationType StrandType, bool bRestTriangle, bool bRestFrame, bool bDeformedTriangle, bool bDeformedFrame)
+			auto RenderProjectionData = [&GraphBuilder, ShaderMap, Viewport, &ViewUniformBuffer, Instances, SkinCache, &bClearDepth, SceneColorTexture, DepthTexture](EHairStrandsInterpolationType StrandType, bool bRestTriangle, bool bRestFrame, bool bDeformedTriangle, bool bDeformedFrame)
 			{
 				TArray<int32> HairLODIndices;
 				for (FHairStrandsInstance* AbstractInstance : Instances)
 				{
 					FHairGroupInstance* Instance = static_cast<FHairGroupInstance*>(AbstractInstance);
 
-					if (Instance->WorldType != WorldType || !Instance->HairGroupPublicData || Instance->Guides.RestRootResource == nullptr || Instance->Guides.DeformedRootResource == nullptr || Instance->BindingType != EHairBindingType::Skinning)
+					if (!Instance->HairGroupPublicData || Instance->Guides.RestRootResource == nullptr || Instance->Guides.DeformedRootResource == nullptr || Instance->BindingType != EHairBindingType::Skinning)
 						continue;
 
 					const int32 MeshLODIndex = Instance->Debug.MeshLODIndex;
diff --git a/Engine/Plugins/Runtime/HairStrands/Source/HairStrandsCore/Private/GroomManager.cpp b/Engine/Plugins/Runtime/HairStrands/Source/HairStrandsCore/Private/GroomManager.cpp
--- a/Engine/Plugins/Runtime/HairStrands/Source/HairStrandsCore/Private/GroomManager.cpp
+++ b/Engine/Plugins/Runtime/HairStrands/Source/HairStrandsCore/Private/GroomManager.cpp
@@ -70,7 +70,6 @@ static void RunInternalHairStrandsInterpolation(
 	FRDGBuilder& GraphBuilder,
 	const FSceneView* View,
 	const FHairStrandsInstances& Instances,
-	EWorldType::Type WorldType, 
 	const FGPUSkinCache* SkinCache,
 	const FShaderDrawDebugData* ShaderDrawData,
 	FGlobalShaderMap* ShaderMap, 
@@ -94,8 +93,7 @@ static void RunInternalHairStrandsInterpolation(
 		#endif
 
 		int32 MeshLODIndex = -1;
-		check(Instance->WorldType == WorldType);
-		if (Instance->WorldType != WorldType || Instance->GeometryType == EHairGeometryType::NoneGeometry)
+		if (Instance->GeometryType == EHairGeometryType::NoneGeometry)
 			continue;
 	
 		check(Instance->HairGroupPublicData);
@@ -298,8 +296,7 @@ static void RunInternalHairStrandsInterpolation(
 		{
 			FHairGroupInstance* Instance = static_cast<FHairGroupInstance*>(AbstractInstance);
 
-			check(Instance->WorldType == WorldType);
-			if (Instance->WorldType != WorldType || Instance->GeometryType == EHairGeometryType::NoneGeometry)
+			if (Instance->GeometryType == EHairGeometryType::NoneGeometry)
 				continue;
 
 			// HAIR_TODO: Do culling on the output of the view rather than the macro group cluster, in order to get consistent result with draw logic
@@ -322,8 +319,7 @@ static void RunInternalHairStrandsInterpolation(
 		{
 			FHairGroupInstance* Instance = static_cast<FHairGroupInstance*>(AbstractInstance);
 
-			check(Instance->WorldType == WorldType);
-			if (Instance->WorldType != WorldType || Instance->GeometryType == EHairGeometryType::NoneGeometry)
+			if (Instance->GeometryType == EHairGeometryType::NoneGeometry)
 				continue;
 
 			// HAIR_TODO: Do culling on the output of the view rather than the macro group cluster, in order to get consistent result with draw logic
@@ -350,7 +346,6 @@ static void RunHairStrandsInterpolation_Guide(
 	FRDGBuilder& GraphBuilder,
 	const FSceneView* View,
 	const FHairStrandsInstances& Instances,
-	EWorldType::Type WorldType,
 	const FGPUSkinCache* SkinCache,
 	const FShaderDrawDebugData* ShaderDrawData,
 	FGlobalShaderMap* ShaderMap,
@@ -366,7 +361,6 @@ static void RunHairStrandsInterpolation_Guide(
 		GraphBuilder,
 		View,
 		Instances,
-		WorldType,
 		SkinCache,
 		ShaderDrawData,
 		ShaderMap,
@@ -378,7 +372,6 @@ static void RunHairStrandsInterpolation_Strands(
 	FRDGBuilder& GraphBuilder,
 	const FSceneView* View,
 	const FHairStrandsInstances& Instances,
-	EWorldType::Type WorldType,
 	const FGPUSkinCache* SkinCache,
 	const FShaderDrawDebugData* ShaderDrawData,
 	FGlobalShaderMap* ShaderMap,
@@ -394,7 +387,6 @@ static void RunHairStrandsInterpolation_Strands(
 		GraphBuilder,
 		View,
 		Instances,
-		WorldType,
 		SkinCache,
 		ShaderDrawData,
 		ShaderMap,
@@ -405,15 +397,13 @@ static void RunHairStrandsInterpolation_Strands(
 
 static void RunHairStrandsGatherCluster(
 	const FHairStrandsInstances& Instances,
-	EWorldType::Type WorldType,
 	FHairStrandClusterData* ClusterData)
 {
 	for (FHairStrandsInstance* AbstractInstance : Instances)
 	{
 		FHairGroupInstance* Instance = static_cast<FHairGroupInstance*>(AbstractInstance);
 
-		check(Instance->WorldType == WorldType);
-		if (Instance->WorldType != WorldType || Instance->GeometryType != EHairGeometryType::Strands)
+		if (Instance->GeometryType != EHairGeometryType::Strands)
 			continue;
 
 		if (Instance->Strands.IsValid())
@@ -474,7 +464,7 @@ static EHairGeometryType ConvertLODGeometryType(EHairGeometryType Type, bool Inb
 	return EHairGeometryType::NoneGeometry;
 }
 
-static void RunHairBufferSwap(const FHairStrandsInstances& Instances, EWorldType::Type WorldType, const TArray<const FSceneView*> Views)
+static void RunHairBufferSwap(const FHairStrandsInstances& Instances, const TArray<const FSceneView*> Views)
 {
 	EShaderPlatform ShaderPlatform = EShaderPlatform::SP_NumPlatforms;
 	if (Views.Num() > 0)
@@ -488,8 +478,7 @@ static void RunHairBufferSwap(const FHairStrandsInstances& Instances, EWorldType
 
 		int32 MeshLODIndex = -1;
 		check(Instance);
-		check(Instance->WorldType == WorldType);
-		if (!Instance || Instance->WorldType != WorldType)
+		if (!Instance)
 			continue;
 
 		check(Instance->HairGroupPublicData);
@@ -557,7 +546,7 @@ static void RunHairBufferSwap(const FHairStrandsInstances& Instances, EWorldType
 	}
 }
 
-static void RunHairLODSelection(const FHairStrandsInstances& Instances, EWorldType::Type WorldType, const TArray<const FSceneView*> Views)
+static void RunHairLODSelection(const FHairStrandsInstances& Instances, const TArray<const FSceneView*> Views)
 {
 	EShaderPlatform ShaderPlatform = EShaderPlatform::SP_NumPlatforms;
 	if (Views.Num() > 0)
@@ -571,8 +560,6 @@ static void RunHairLODSelection(const FHairStrandsInstances& Instances, EWorldTy
 
 		int32 MeshLODIndex = -1;
 		check(Instance);
-		check(Instance->WorldType == WorldType);
-
 		check(Instance->HairGroupPublicData);
 
 		//if (Instance->ProxyLocalToWorld)
@@ -717,7 +704,6 @@ static void RunHairStrandsProcess(FRDGBuilder& GraphBuilder, FGlobalShaderMap* S
 void RunHairStrandsDebug(
 	FRDGBuilder& GraphBuilder,
 	FGlobalShaderMap* ShaderMap,
-	EWorldType::Type WorldType,
 	const FSceneView& View,
 	const FHairStrandsInstances& Instances,
 	const FGPUSkinCache* SkinCache,
@@ -758,12 +744,10 @@ void ProcessHairStrandsBookmark(
 		{
 			RunHairBufferSwap(
 				*Parameters.Instances,
-				Parameters.WorldType,
 				Parameters.AllViews);
 		}
 		RunHairLODSelection(
 			*Parameters.Instances,
-			Parameters.WorldType,
 			Parameters.AllViews);
 	}
 	else if (Bookmark == EHairStrandsBookmark::ProcessEndOfFrame)
@@ -772,7 +756,6 @@ void ProcessHairStrandsBookmark(
 		{
 			RunHairBufferSwap(
 				*Parameters.Instances,
-				Parameters.WorldType,
 				Parameters.AllViews);
 		}
 	}
@@ -783,7 +766,6 @@ void ProcessHairStrandsBookmark(
 			*GraphBuilder,
 			Parameters.View,
 			*Parameters.Instances,
-			Parameters.WorldType,
 			Parameters.SkinCache,
 			Parameters.DebugShaderData,
 			Parameters.ShaderMap,
@@ -793,7 +775,6 @@ void ProcessHairStrandsBookmark(
 	{
 		RunHairStrandsGatherCluster(
 			*Parameters.Instances,
-			Parameters.WorldType,
 			&Parameters.HairClusterData);
 	}
 	else if (Bookmark == EHairStrandsBookmark::ProcessStrandsInterpolation)
@@ -803,7 +784,6 @@ void ProcessHairStrandsBookmark(
 			*GraphBuilder,
 			Parameters.View,
 			*Parameters.Instances,
-			Parameters.WorldType,
 			Parameters.SkinCache,
 			Parameters.DebugShaderData,
 			Parameters.ShaderMap,
@@ -815,7 +795,6 @@ void ProcessHairStrandsBookmark(
 		RunHairStrandsDebug(
 			*GraphBuilder,
 			Parameters.ShaderMap,
-			Parameters.WorldType,
 			*Parameters.View,
 			*Parameters.Instances,
 			Parameters.SkinCache,
diff --git a/Engine/Plugins/Runtime/HairStrands/Source/HairStrandsCore/Public/GroomComponent.h b/Engine/Plugins/Runtime/HairStrands/Source/HairStrandsCore/Public/GroomComponent.h
--- a/Engine/Plugins/Runtime/HairStrands/Source/HairStrandsCore/Public/GroomComponent.h
+++ b/Engine/Plugins/Runtime/HairStrands/Source/HairStrandsCore/Public/GroomComponent.h
@@ -222,7 +222,6 @@ private:
 	bool bValidationEnable = true;
 	bool bPreviewMode = false;
 
-	EWorldType::Type GetWorldType() const; 
 	void InitResources(bool bIsBindingReloading=false);
 	void ReleaseResources();
 
diff --git a/Engine/Plugins/Runtime/HairStrands/Source/HairStrandsCore/Public/GroomInstance.h b/Engine/Plugins/Runtime/HairStrands/Source/HairStrandsCore/Public/GroomInstance.h
--- a/Engine/Plugins/Runtime/HairStrands/Source/HairStrandsCore/Public/GroomInstance.h
+++ b/Engine/Plugins/Runtime/HairStrands/Source/HairStrandsCore/Public/GroomInstance.h
@@ -187,7 +187,6 @@ struct HAIRSTRANDSCORE_API FHairGroupInstance : public FHairStrandsInstance
 	} Debug;
 
 	FTransform				LocalToWorld = FTransform::Identity;
-	EWorldType::Type		WorldType;
 	FHairGroupPublicData*	HairGroupPublicData = nullptr;
 	EHairGeometryType		GeometryType = EHairGeometryType::NoneGeometry;
 	EHairBindingType		BindingType = EHairBindingType::NoneBinding;
diff --git a/Engine/Source/Runtime/Renderer/Private/HairStrands/HairStrandsInterface.cpp b/Engine/Source/Runtime/Renderer/Private/HairStrands/HairStrandsInterface.cpp
--- a/Engine/Source/Runtime/Renderer/Private/HairStrands/HairStrandsInterface.cpp
+++ b/Engine/Source/Runtime/Renderer/Private/HairStrands/HairStrandsInterface.cpp
@@ -391,7 +391,6 @@ FHairStrandsBookmarkParameters CreateHairStrandsBookmarkParameters(FScene* Scene
 	FHairStrandsBookmarkParameters Out;
 	Out.DebugShaderData			= &View.ShaderDrawData;
 	Out.SkinCache				= View.Family->Scene->GetGPUSkinCache();
-	Out.WorldType				= View.Family->Scene->GetWorld()->WorldType;
 	Out.ShaderMap				= View.ShaderMap;
 	Out.Instances				= &Scene->HairStrandsSceneData.RegisteredProxies;
 	Out.View					= &View;
diff --git a/Engine/Source/Runtime/Renderer/Public/HairStrandsInterface.h b/Engine/Source/Runtime/Renderer/Public/HairStrandsInterface.h
--- a/Engine/Source/Runtime/Renderer/Public/HairStrandsInterface.h
+++ b/Engine/Source/Runtime/Renderer/Public/HairStrandsInterface.h
@@ -425,7 +425,6 @@ struct FHairStrandsBookmarkParameters
 {
 	class FGPUSkinCache* SkinCache = nullptr;
 	FShaderDrawDebugData* DebugShaderData = nullptr;
-	EWorldType::Type WorldType = EWorldType::None;
 	class FGlobalShaderMap* ShaderMap = nullptr;
 
 	FIntRect ViewRect; // View 0
