Add frustum culling for updating hair strands instance back.

It was disabled in a previous CL due to the instance data refactoring.

diff --git a/Engine/Plugins/Runtime/HairStrands/Source/HairStrandsCore/Private/GroomComponent.cpp b/Engine/Plugins/Runtime/HairStrands/Source/HairStrandsCore/Private/GroomComponent.cpp
--- a/Engine/Plugins/Runtime/HairStrands/Source/HairStrandsCore/Private/GroomComponent.cpp
+++ b/Engine/Plugins/Runtime/HairStrands/Source/HairStrandsCore/Private/GroomComponent.cpp
@@ -2066,6 +2066,7 @@ void UGroomComponent::InitResources(bool bIsBindingReloading)
 		// Initialize LOD screen size & visibility
 		{
 			HairGroupInstance->HairGroupPublicData = new FHairGroupPublicData(GroupIt);
+			HairGroupInstance->HairGroupPublicData->Instance = HairGroupInstance;
 			TArray<float> CPULODScreenSize;
 			TArray<bool> LODVisibility;
 			TArray<EHairGeometryType> LODGeometryTypes;
diff --git a/Engine/Plugins/Runtime/HairStrands/Source/HairStrandsCore/Private/GroomManager.cpp b/Engine/Plugins/Runtime/HairStrands/Source/HairStrandsCore/Private/GroomManager.cpp
--- a/Engine/Plugins/Runtime/HairStrands/Source/HairStrandsCore/Private/GroomManager.cpp
+++ b/Engine/Plugins/Runtime/HairStrands/Source/HairStrandsCore/Private/GroomManager.cpp
@@ -41,27 +41,9 @@ DEFINE_LOG_CATEGORY_STATIC(LogGroomManager, Log, All);
 
 ///////////////////////////////////////////////////////////////////////////////////////////////////
 
-static bool IsInstanceVisible(
-	const FSceneView* View,
-	FHairGroupInstance* Instance)
+static bool IsInstanceFrustumCullingEnable()
 {
-	// Frustum culling for rendering strands. Update position only for visible/in camera frustum
-	if (GHairStrands_InterpolationFrustumCullingEnable > 0)
-	{
-		if (!Instance->ProxyBounds)
-		{
-			return false;
-		}
-
-		const FSphere InstanceBound = Instance->ProxyBounds->GetSphere();
-		bool bFullyContained = false;
-		if (!View->ViewFrustum.IntersectSphere(InstanceBound.Center, InstanceBound.W, bFullyContained))
-		{
-			return false;
-		}
-	}
-
-	return true;
+	return GHairStrands_InterpolationFrustumCullingEnable > 0;
 }
 
 bool NeedsUpdateCardsMeshTriangles();
@@ -83,15 +65,6 @@ static void RunInternalHairStrandsInterpolation(
 	{
 		FHairGroupInstance* Instance = static_cast<FHairGroupInstance*>(AbstractInstance);
 
-		// HAIR_TODO: Do culling on the output of the view rather than the macro group cluster, in order to get consistent result with draw logic
-		// Frustum culling for rendering strands. Update position only for visible/in camera frustum
-		#if 0
-		if (Type == EHairStrandsInterpolationType::RenderStrands && !IsInstanceVisible(View, Instance))
-		{
-			continue;
-		}
-		#endif
-
 		int32 MeshLODIndex = -1;
 		if (Instance->GeometryType == EHairGeometryType::NoneGeometry)
 			continue;
@@ -299,15 +272,6 @@ static void RunInternalHairStrandsInterpolation(
 			if (Instance->GeometryType == EHairGeometryType::NoneGeometry)
 				continue;
 
-			// HAIR_TODO: Do culling on the output of the view rather than the macro group cluster, in order to get consistent result with draw logic
-			// Frustum culling guide deformation if the instance does not have any simulation
-			#if 0
-			if (Instance->GeometryType == EHairGeometryType::Strands && !Instance->Guides.bIsSimulationEnable  && !IsInstanceVisible(View, Instance))
-			{
-				continue;
-			}
-			#endif
-
 			ResetHairStrandsInterpolation(GraphBuilder, ShaderMap, Instance, Instance->Debug.MeshLODIndex);
 		}
 	}
@@ -322,15 +286,6 @@ static void RunInternalHairStrandsInterpolation(
 			if (Instance->GeometryType == EHairGeometryType::NoneGeometry)
 				continue;
 
-			// HAIR_TODO: Do culling on the output of the view rather than the macro group cluster, in order to get consistent result with draw logic
-			// Frustum culling for rendering strands. Update position only for visible/in camera frustum
-			#if 0
-			if (!IsInstanceVisible(View, Instance))
-			{
-				continue;
-			}
-			#endif
-
  			ComputeHairStrandsInterpolation(
 				GraphBuilder, 
 				ShaderMap,
@@ -721,6 +676,13 @@ void ProcessHairStrandsBookmark(
 	FHairStrandsBookmarkParameters& Parameters)
 {
 	check(Parameters.Instances != nullptr);
+	const bool bCulling =
+		IsInstanceFrustumCullingEnable() && (
+			Bookmark == EHairStrandsBookmark::ProcessGuideInterpolation ||
+			Bookmark == EHairStrandsBookmark::ProcessGatherCluster ||
+			Bookmark == EHairStrandsBookmark::ProcessStrandsInterpolation ||
+			Bookmark == EHairStrandsBookmark::ProcessDebug);
+	FHairStrandsInstances& Instances = bCulling ? Parameters.VisibleInstances : *Parameters.Instances;
 
 	if (Bookmark == EHairStrandsBookmark::ProcessTasks)
 	{
@@ -765,7 +727,7 @@ void ProcessHairStrandsBookmark(
 		RunHairStrandsInterpolation_Guide(
 			*GraphBuilder,
 			Parameters.View,
-			*Parameters.Instances,
+			Instances,
 			Parameters.SkinCache,
 			Parameters.DebugShaderData,
 			Parameters.ShaderMap,
@@ -774,7 +736,7 @@ void ProcessHairStrandsBookmark(
 	else if (Bookmark == EHairStrandsBookmark::ProcessGatherCluster)
 	{
 		RunHairStrandsGatherCluster(
-			*Parameters.Instances,
+			Instances,
 			&Parameters.HairClusterData);
 	}
 	else if (Bookmark == EHairStrandsBookmark::ProcessStrandsInterpolation)
@@ -783,7 +745,7 @@ void ProcessHairStrandsBookmark(
 		RunHairStrandsInterpolation_Strands(
 			*GraphBuilder,
 			Parameters.View,
-			*Parameters.Instances,
+			Instances,
 			Parameters.SkinCache,
 			Parameters.DebugShaderData,
 			Parameters.ShaderMap,
@@ -796,7 +758,7 @@ void ProcessHairStrandsBookmark(
 			*GraphBuilder,
 			Parameters.ShaderMap,
 			*Parameters.View,
-			*Parameters.Instances,
+			Instances,
 			Parameters.SkinCache,
 			Parameters.DebugShaderData,
 			Parameters.SceneColorTexture,
diff --git a/Engine/Source/Runtime/Renderer/Private/HairStrands/HairStrandsInterface.cpp b/Engine/Source/Runtime/Renderer/Private/HairStrands/HairStrandsInterface.cpp
--- a/Engine/Source/Runtime/Renderer/Private/HairStrands/HairStrandsInterface.cpp
+++ b/Engine/Source/Runtime/Renderer/Private/HairStrands/HairStrandsInterface.cpp
@@ -389,6 +389,18 @@ bool IsHairStrandsClusterCullingUseHzb();
 FHairStrandsBookmarkParameters CreateHairStrandsBookmarkParameters(FScene* Scene, FViewInfo& View)
 {
 	FHairStrandsBookmarkParameters Out;
+	Out.VisibleInstances.Reserve(View.HairStrandsMeshElements.Num());
+	for (const FMeshBatchAndRelevance& MeshBatch : View.HairStrandsMeshElements)
+	{
+		if (MeshBatch.Mesh && MeshBatch.Mesh->Elements.Num() > 0)
+		{
+			FHairGroupPublicData* HairGroupPublicData = reinterpret_cast<FHairGroupPublicData*>(MeshBatch.Mesh->Elements[0].VertexFactoryUserData);
+			if (HairGroupPublicData && HairGroupPublicData->Instance)
+			{
+				Out.VisibleInstances.Add(HairGroupPublicData->Instance);
+			}
+		}
+	}
 	Out.DebugShaderData			= &View.ShaderDrawData;
 	Out.SkinCache				= View.Family->Scene->GetGPUSkinCache();
 	Out.ShaderMap				= View.ShaderMap;
diff --git a/Engine/Source/Runtime/Renderer/Public/HairStrandsInterface.h b/Engine/Source/Runtime/Renderer/Public/HairStrandsInterface.h
--- a/Engine/Source/Runtime/Renderer/Public/HairStrandsInterface.h
+++ b/Engine/Source/Runtime/Renderer/Public/HairStrandsInterface.h
@@ -271,9 +271,13 @@ public:
 		EHairBindingType BindingType = EHairBindingType::NoneBinding;
 		FTransform LocalToWorldTransform;
 	};
+
+	// Current hook for retriving instance data. 
+	// This needs to be refactor to merge FHairGroupPublicData & HairStrandsInstance
+	FHairStrandsInstance* Instance = nullptr;
+
 	FVertexFactoryInput VFInput;
 	uint32 ClusterDataIndex = ~0; // #hair_todo: move this into instance data, or remove FHairStrandClusterData
-//private:
 
 	uint32 GroupControlTriangleStripVertexCount;
 	uint32 GroupIndex;
@@ -428,6 +432,7 @@ struct FHairStrandsBookmarkParameters
 	class FGlobalShaderMap* ShaderMap = nullptr;
 
 	FIntRect ViewRect; // View 0
+	FHairStrandsInstances VisibleInstances;
 	FHairStrandsInstances* Instances = nullptr;
 	const FSceneView* View = nullptr;// // View 0
 	TArray<const FSceneView*> AllViews;
