Add (improted) strands max width into the groom asset stats.

diff --git a/Engine/Plugins/Runtime/HairStrands/Source/HairStrandsEditor/Private/GroomAssetDetails.cpp b/Engine/Plugins/Runtime/HairStrands/Source/HairStrandsEditor/Private/GroomAssetDetails.cpp
--- a/Engine/Plugins/Runtime/HairStrands/Source/HairStrandsEditor/Private/GroomAssetDetails.cpp
+++ b/Engine/Plugins/Runtime/HairStrands/Source/HairStrandsEditor/Private/GroomAssetDetails.cpp
@@ -109,7 +109,7 @@ static uint32 GetHairCardsAtlasResolution(int32 InLODIndex, int32 PrevResolution
 
 ///////////////////////////////////////////////////////////////////////////////////////////////////////////
 // Array panel for hair strands infos
-TSharedRef<SUniformGridPanel> MakeHairStrandsInfoGrid(const FSlateFontInfo& DetailFontInfo, FHairGroupInfoWithVisibility& CurrentAsset)
+TSharedRef<SUniformGridPanel> MakeHairStrandsInfoGrid(const FSlateFontInfo& DetailFontInfo, FHairGroupInfoWithVisibility& CurrentAsset, float MaxRadius)
 {
 	TSharedRef<SUniformGridPanel> Grid = SNew(SUniformGridPanel).SlotPadding(2.0f);
 
@@ -176,6 +176,28 @@ TSharedRef<SUniformGridPanel> MakeHairStrandsInfoGrid(const FSlateFontInfo& Deta
 		.Text(FText::AsNumber(CurrentAsset.NumGuideVertices))
 	];
 
+	// Width (mm)
+	Grid->AddSlot(0, 3) // x, y
+	.HAlign(HAlign_Left)
+	[
+		SNew(STextBlock)
+		.Font(DetailFontInfo)
+		.Text(LOCTEXT("HairInfo_Width", "Max. Width"))
+	];
+	Grid->AddSlot(1, 3) // x, y
+	.HAlign(HAlign_Right)
+	[
+		SNew(STextBlock)
+		.Font(DetailFontInfo)
+		.Text(LOCTEXT("HairInfo_GuideWidth", ""))
+	];
+	Grid->AddSlot(2, 3) // x, y
+	.HAlign(HAlign_Right)
+	[
+		SNew(STextBlock)
+		.Font(DetailFontInfo)
+		.Text(FText::AsNumber(MaxRadius*2.0f))
+	];
 	return Grid;
 }
 
@@ -1216,7 +1238,7 @@ void FGroomRenderingDetails::OnGenerateElementForHairGroup(TSharedRef<IPropertyH
 		.ValueContent()
 		.HAlign(HAlign_Fill)
 		[
-			MakeHairStrandsInfoGrid(DetailFontInfo, GroomAsset->HairGroupsInfo[GroupIndex])
+			MakeHairStrandsInfoGrid(DetailFontInfo, GroomAsset->HairGroupsInfo[GroupIndex], GroomAsset->HairGroupsData[GroupIndex].Strands.Data.StrandsCurves.MaxRadius)
 		];
 	}
 
