Avoid unecesary allocation during hair component initialization.

diff --git a/Engine/Plugins/Runtime/HairStrands/Source/HairStrandsCore/Private/GroomResources.cpp b/Engine/Plugins/Runtime/HairStrands/Source/HairStrandsCore/Private/GroomResources.cpp
--- a/Engine/Plugins/Runtime/HairStrands/Source/HairStrandsCore/Private/GroomResources.cpp
+++ b/Engine/Plugins/Runtime/HairStrands/Source/HairStrandsCore/Private/GroomResources.cpp
@@ -722,6 +722,7 @@ FHairStrandsRestRootResource::FHairStrandsRestRootResource(const FHairStrandsDat
 void FHairStrandsRestRootResource::PopulateFromRootData()
 {
 	uint32 LODIndex = 0;
+	LODs.Reserve(RootData.MeshProjectionLODs.Num());
 	for (FHairStrandsRootData::FMeshProjectionLOD& MeshProjectionLOD : RootData.MeshProjectionLODs)
 	{
 		FLOD& LOD = LODs.AddDefaulted_GetRef();
@@ -844,6 +845,7 @@ CurveType(InCurveType)
 	check(InRestResources);
 	uint32 LODIndex = 0;
 	RootCount = InRestResources->RootData.RootCount;
+	LODs.Reserve(InRestResources->LODs.Num());
 	for (const FHairStrandsRestRootResource::FLOD& InLOD : InRestResources->LODs)
 	{
 		FLOD& LOD = LODs.AddDefaulted_GetRef();
