Improve groom component creation time when using only cards setup.

diff --git a/Engine/Plugins/Runtime/HairStrands/Source/HairStrandsCore/Private/GroomComponent.cpp b/Engine/Plugins/Runtime/HairStrands/Source/HairStrandsCore/Private/GroomComponent.cpp
--- a/Engine/Plugins/Runtime/HairStrands/Source/HairStrandsCore/Private/GroomComponent.cpp
+++ b/Engine/Plugins/Runtime/HairStrands/Source/HairStrandsCore/Private/GroomComponent.cpp
@@ -1953,6 +1953,11 @@ static UMeshComponent* ValidateBindingAsset(
 
 void CreateHairStrandsDebugAttributeBuffer(FRDGExternalBuffer* DebugAttributeBuffer, uint32 VertexCount);
 
+static EGroomGeometryType GetEffectiveGeometryType(EGroomGeometryType Type, bool bUseCards)
+{
+	return Type == EGroomGeometryType::Strands && (!IsHairStrandsEnabled(EHairStrandsShaderType::Strands) || bUseCards) ? EGroomGeometryType::Cards : Type;
+}
+
 void UGroomComponent::InitResources(bool bIsBindingReloading)
 {
 	LLM_SCOPE(ELLMTag::Meshes) // This should be a Groom LLM tag, but there is no LLM tag bit left
@@ -1988,7 +1993,7 @@ void UGroomComponent::InitResources(bool bIsBindingReloading)
 	{
 		for (uint32 LODIt = 0, LODCount = GroomAsset->GetLODCount(); LODIt < LODCount; ++LODIt)
 		{
-			const EGroomGeometryType GeometryType = GroomAsset->GetGeometryType(GroupIt, LODIt);
+			const EGroomGeometryType GeometryType = GetEffectiveGeometryType(GroomAsset->GetGeometryType(GroupIt, LODIt), bUseCards);
 			const EGroomBindingType BindingType = GroomAsset->GetBindingType(GroupIt, LODIt);
 
 			// Note on Global Deformation:
@@ -2040,7 +2045,6 @@ void UGroomComponent::InitResources(bool bIsBindingReloading)
 		bHasNeedGlobalDeformation = false;
 	}
 
-	const bool bIsStrandsEnabled = IsHairStrandsEnabled(EHairStrandsShaderType::Strands);
 	for (int32 GroupIt = 0, GroupCount = GroomAsset->HairGroupsData.Num(); GroupIt < GroupCount; ++GroupIt)
 	{
 		FHairGroupInstance* HairGroupInstance = new FHairGroupInstance();
@@ -2064,6 +2068,7 @@ void UGroomComponent::InitResources(bool bIsBindingReloading)
 		const EHairInterpolationType HairInterpolationType = ToHairInterpolationType(GroomAsset->HairInterpolationType);
 
 		// Initialize LOD screen size & visibility
+		bool bNeedStrandsData = false;
 		{
 			HairGroupInstance->HairGroupPublicData = new FHairGroupPublicData(GroupIt);
 			HairGroupInstance->HairGroupPublicData->Instance = HairGroupInstance;
@@ -2086,9 +2091,10 @@ void UGroomComponent::InitResources(bool bIsBindingReloading)
 
 				// * Force global interpolation to be enable for meshes with skinning binding as we use RBF defomation for 'sticking' meshes onto skel. mesh surface
 				// * Global deformation are allowed only with 'Skinning' binding type
-				const EHairGeometryType GeometryType = ToHairGeometryType(GroomAsset->GetGeometryType(GroupIt, LODIt));
+				const EHairGeometryType GeometryType = ToHairGeometryType(GetEffectiveGeometryType(GroomAsset->GetGeometryType(GroupIt, LODIt), bUseCards));
 				const bool LODSimulation = GroomAsset->IsSimulationEnable(GroupIt, LODIt);
 				const bool LODGlobalInterpolation = LocalBindingAsset && (BindingType == EHairBindingType::Skinning && (GroomAsset->IsGlobalInterpolationEnable(GroupIt, LODIt) || GeometryType == EHairGeometryType::Meshes));
+				bNeedStrandsData = bNeedStrandsData || GeometryType == EHairGeometryType::Strands;
 
 				CPULODScreenSize.Add(LODSettings.ScreenSize);
 				LODVisibility.Add(LODSettings.bVisible);
@@ -2160,7 +2166,7 @@ void UGroomComponent::InitResources(bool bIsBindingReloading)
 		#endif// #if WITH_EDITORONLY_DATA
 
 		// Strands data/resources
-		if (bIsStrandsEnabled && GroupData.Strands.IsValid())
+		if (bNeedStrandsData && GroupData.Strands.IsValid())
 		{
 			check(GroupIt < GroomGroupsDesc.Num());
 
