Shared RT groom resources across instances when groom data are static (e.g., no sim, no skinning).

This allows to save memory and initialization time.

diff --git a/Engine/Plugins/Runtime/HairStrands/Source/HairStrandsCore/Private/GroomAsset.cpp b/Engine/Plugins/Runtime/HairStrands/Source/HairStrandsCore/Private/GroomAsset.cpp
--- a/Engine/Plugins/Runtime/HairStrands/Source/HairStrandsCore/Private/GroomAsset.cpp
+++ b/Engine/Plugins/Runtime/HairStrands/Source/HairStrandsCore/Private/GroomAsset.cpp
@@ -490,6 +490,14 @@ void UGroomAsset::UpdateResource()
 			{
 				InternalUpdateResource(GroupData.Strands.ClusterCullingResource);
 			}
+
+			#if RHI_RAYTRACING
+			if (IsHairRayTracingEnabled() && GroupData.Strands.RaytracingResource)
+			{
+				// Release raytracing resources, so that it will be recreate by the component.
+				InternalReleaseResource(GroupData.Strands.RaytracingResource);
+			}
+			#endif
 		}
 	}
 
@@ -523,6 +531,9 @@ void UGroomAsset::ReleaseResource()
 			InternalReleaseResource(GroupData.Strands.RestResource);
 			InternalReleaseResource(GroupData.Strands.ClusterCullingResource);
 			InternalReleaseResource(GroupData.Strands.InterpolationResource);
+			#if RHI_RAYTRACING
+			InternalReleaseResource(GroupData.Strands.RaytracingResource);
+			#endif
 		}
 
 		for (FHairGroupData::FCards::FLOD& LOD : GroupData.Cards.LODs)
@@ -532,10 +543,16 @@ void UGroomAsset::ReleaseResource()
 			InternalReleaseResource(LOD.InterpolationResource);
 			InternalReleaseResource(LOD.Guides.RestResource);
 			InternalReleaseResource(LOD.Guides.InterpolationResource);
+			#if RHI_RAYTRACING
+			InternalReleaseResource(LOD.RaytracingResource);
+			#endif
 		}
 		for (FHairGroupData::FMeshes::FLOD& LOD : GroupData.Meshes.LODs)
 		{
 			InternalReleaseResource(LOD.RestResource);
+			#if RHI_RAYTRACING
+			InternalReleaseResource(LOD.RaytracingResource);
+			#endif
 		}
 	}
 
@@ -924,6 +941,8 @@ void UGroomAsset::PostEditChangeProperty(FPropertyChangedEvent& PropertyChangedE
 
 	SavePendingProceduralAssets();
 
+	const bool bHairStrandsRaytracingRadiusChanged = PropertyName == GET_MEMBER_NAME_CHECKED(FHairShadowSettings, HairRaytracingRadiusScale);
+	
 	// By pass update for all procedural cards parameters, as we don't want them to invalidate the cards data. 
 	// Cards should be refresh only under user action
 	// By pass update if bStrandsInterpolationChanged has the resources have already been recreated
@@ -2205,28 +2224,97 @@ void UGroomAsset::InitGuideResources()
 	//}
 }
 
-void UGroomAsset::AllocateGuidesResources(uint32 GroupIndex)
+FHairStrandsRestResource* UGroomAsset::AllocateGuidesResources(uint32 GroupIndex)
 {
 	if (GroupIndex < uint32(GetNumHairGroups()))
 	{
 		FHairGroupData& GroupData = HairGroupsData[GroupIndex];
 		if (GroupData.Guides.HasValidData())
 		{
-			GroupData.Guides.RestResource = new FHairStrandsRestResource(GroupData.Guides.Data.RenderData, GroupData.Guides.Data.BoundingBox.GetCenter(), EHairStrandsResourcesType::Guides);
-			BeginInitResource(GroupData.Guides.RestResource);
+			if (GroupData.Guides.RestResource == nullptr)
+			{
+				GroupData.Guides.RestResource = new FHairStrandsRestResource(GroupData.Guides.Data.RenderData, GroupData.Guides.Data.BoundingBox.GetCenter(), EHairStrandsResourcesType::Guides);
+				BeginInitResource(GroupData.Guides.RestResource);
+			}
+			return GroupData.Guides.RestResource;
 		}
 	}
+	return nullptr;
 }
 
-void UGroomAsset::AllocateInterpolationResources(uint32 GroupIndex)
+FHairStrandsInterpolationResource* UGroomAsset::AllocateInterpolationResources(uint32 GroupIndex)
 {
 	if (GroupIndex < uint32(GetNumHairGroups()))
 	{
 		FHairGroupData& GroupData = HairGroupsData[GroupIndex];
 		check(GroupData.Guides.IsValid());
-		GroupData.Strands.InterpolationResource = new FHairStrandsInterpolationResource(GroupData.Strands.InterpolationData.RenderData, GroupData.Guides.Data);
-		BeginInitResource(GroupData.Strands.InterpolationResource);
+		if (GroupData.Strands.InterpolationResource == nullptr)
+		{
+			GroupData.Strands.InterpolationResource = new FHairStrandsInterpolationResource(GroupData.Strands.InterpolationData.RenderData, GroupData.Guides.Data);
+			BeginInitResource(GroupData.Strands.InterpolationResource);
+		}
+		return GroupData.Strands.InterpolationResource;
+	}
+	return nullptr;
+}
+
+FHairStrandsRaytracingResource* UGroomAsset::AllocateCardsRaytracingResources(uint32 GroupIndex, uint32 LODIndex)
+{
+#if RHI_RAYTRACING
+	if (IsRayTracingEnabled() && GroupIndex < uint32(GetNumHairGroups()) && LODIndex < uint32(HairGroupsLOD[GroupIndex].LODs.Num()))
+	{
+		FHairGroupData& GroupData = HairGroupsData[GroupIndex];
+		FHairGroupData::FCards::FLOD& LOD = GroupData.Cards.LODs[LODIndex];
+		check(LOD.Data.IsValid());
+
+		if (LOD.RaytracingResource == nullptr)
+		{
+			LOD.RaytracingResource = new FHairStrandsRaytracingResource(LOD.Data);
+			BeginInitResource(LOD.RaytracingResource);
+		}
+		return LOD.RaytracingResource;
 	}
+#endif
+	return nullptr;
+}
+
+FHairStrandsRaytracingResource* UGroomAsset::AllocateMeshesRaytracingResources(uint32 GroupIndex, uint32 LODIndex)
+{
+#if RHI_RAYTRACING
+	if (IsRayTracingEnabled() && GroupIndex < uint32(GetNumHairGroups()) && LODIndex < uint32(HairGroupsLOD[GroupIndex].LODs.Num()))
+	{
+		FHairGroupData& GroupData = HairGroupsData[GroupIndex];
+		FHairGroupData::FMeshes::FLOD& LOD = GroupData.Meshes.LODs[LODIndex];
+		check(LOD.Data.IsValid());
+
+		if (LOD.RaytracingResource == nullptr)
+		{
+			LOD.RaytracingResource = new FHairStrandsRaytracingResource(LOD.Data);
+			BeginInitResource(LOD.RaytracingResource);
+		}
+		return LOD.RaytracingResource;
+	}
+#endif
+	return nullptr;
+}
+
+FHairStrandsRaytracingResource* UGroomAsset::AllocateStrandsRaytracingResources(uint32 GroupIndex)
+{
+#if RHI_RAYTRACING
+	if (IsRayTracingEnabled() && GroupIndex < uint32(GetNumHairGroups()))
+	{
+		FHairGroupData& GroupData = HairGroupsData[GroupIndex];
+		check(GroupData.Strands.HasValidData());
+
+		if (GroupData.Strands.RaytracingResource == nullptr)
+		{
+			GroupData.Strands.RaytracingResource = new FHairStrandsRaytracingResource(GroupData.Strands.Data);
+			BeginInitResource(GroupData.Strands.RaytracingResource);
+		}
+		return GroupData.Strands.RaytracingResource;
+	}
+#endif
+	return nullptr;
 }
 
 void UGroomAsset::InitStrandsResources()
diff --git a/Engine/Plugins/Runtime/HairStrands/Source/HairStrandsCore/Private/GroomComponent.cpp b/Engine/Plugins/Runtime/HairStrands/Source/HairStrandsCore/Private/GroomComponent.cpp
--- a/Engine/Plugins/Runtime/HairStrands/Source/HairStrandsCore/Private/GroomComponent.cpp
+++ b/Engine/Plugins/Runtime/HairStrands/Source/HairStrandsCore/Private/GroomComponent.cpp
@@ -2137,14 +2137,9 @@ void UGroomComponent::InitResources(bool bIsBindingReloading)
 				BeginInitResource(HairGroupInstance->Guides.DeformedRootResource);
 			}
 
-			// (Lazy) Allocation of the guide rest resources
-			if (!GroupData.Guides.RestResource)
-			{
-				// Lazy allocation of the guide resources
-				GroomAsset->AllocateGuidesResources(GroupIt);
-				check(GroupData.Guides.RestResource);
-			}
-			HairGroupInstance->Guides.RestResource = GroupData.Guides.RestResource;
+			// Lazy allocation of the guide resources
+			HairGroupInstance->Guides.RestResource = GroomAsset->AllocateGuidesResources(GroupIt);
+			check(GroupData.Guides.RestResource);
 
 			// If guides are allocated, deformed resources are always needs since they are either used with simulation, or RBF deformation. Both are dynamics, and require deformed positions
 			HairGroupInstance->Guides.DeformedResource = new FHairStrandsDeformedResource(GroupData.Guides.Data.RenderData, true, HairGroupInstance->Guides.RestResource->PositionOffset, EHairStrandsResourcesType::Guides);
@@ -2176,25 +2171,13 @@ void UGroomComponent::InitResources(bool bIsBindingReloading)
 			// (Lazy) Allocate interpolation resources, only if guides are required
 			if (bNeedGuides)
 			{
-				if (!GroupData.Strands.InterpolationResource)
-				{
-					// Lazy allocation of the guide resources
-					GroomAsset->AllocateInterpolationResources(GroupIt);
-					check(GroupData.Strands.InterpolationResource);
-				}
-				HairGroupInstance->Strands.InterpolationResource = GroupData.Strands.InterpolationResource;
+				HairGroupInstance->Strands.InterpolationResource = GroomAsset->AllocateInterpolationResources(GroupIt);
+				check(HairGroupInstance->Strands.InterpolationResource);
 			}
 
 			// Material
 			HairGroupInstance->Strands.Material = nullptr;
 
-			#if RHI_RAYTRACING
-			if (IsHairRayTracingEnabled() && HairGroupInstance->Strands.Modifier.bUseHairRaytracingGeometry)
-			{
-				HairGroupInstance->Strands.RenRaytracingResource = new FHairStrandsRaytracingResource(GroupData.Strands.Data);
-				BeginInitResource(HairGroupInstance->Strands.RenRaytracingResource);
-			}
-			#endif
 
 			// Check if strands needs to load/allocate root resources
 			bool bNeedRootResources = false;
@@ -2221,6 +2204,26 @@ void UGroomComponent::InitResources(bool bIsBindingReloading)
 				}
 			}
 
+			#if RHI_RAYTRACING
+			if (IsHairRayTracingEnabled() && HairGroupInstance->Strands.Modifier.bUseHairRaytracingGeometry)
+			{
+				if (bNeedDynamicResources)
+				{
+					// Allocate dynamic raytracing resources (owned by the groom component/instance)
+					HairGroupInstance->Strands.RenRaytracingResource = new FHairStrandsRaytracingResource(GroupData.Strands.Data);
+					BeginInitResource(HairGroupInstance->Strands.RenRaytracingResource);
+					HairGroupInstance->Strands.RenRaytracingResourceOwned = true;
+				}
+				else
+				{
+					// (Lazy) Allocate static raytracing resources (owned by the grooom asset)
+					HairGroupInstance->Strands.RenRaytracingResourceOwned = false;
+					HairGroupInstance->Strands.RenRaytracingResource = GroomAsset->AllocateStrandsRaytracingResources(GroupIt);
+					check(HairGroupInstance->Strands.RenRaytracingResource);
+				}
+			}
+			#endif
+
 			if (bNeedRootResources)
 			{
 				check(LocalBindingAsset);
@@ -2240,7 +2243,7 @@ void UGroomComponent::InitResources(bool bIsBindingReloading)
 			if (bNeedDynamicResources)
 			{
 				HairGroupInstance->Strands.DeformedResource = new FHairStrandsDeformedResource(GroupData.Strands.Data.RenderData, false, HairGroupInstance->Strands.RestResource->PositionOffset, EHairStrandsResourcesType::Strands);
-			BeginInitResource(HairGroupInstance->Strands.DeformedResource);
+				BeginInitResource(HairGroupInstance->Strands.DeformedResource);
 
 				// Initialize deformed position relative position offset to the rest pose offset
 				HairGroupInstance->Strands.DeformedResource->GetPositionOffset(FHairStrandsDeformedResource::EFrameType::Current)  = HairGroupInstance->Strands.RestResource->PositionOffset;
@@ -2309,14 +2312,26 @@ void UGroomComponent::InitResources(bool bIsBindingReloading)
 				if (bNeedDeformedPositions)
 				{
 					InstanceLOD.DeformedResource = new FHairCardsDeformedResource(LOD.Data.RenderData, false);
-				BeginInitResource(InstanceLOD.DeformedResource);
+					BeginInitResource(InstanceLOD.DeformedResource);
 				}
 
 				#if RHI_RAYTRACING
 				if (IsRayTracingEnabled())
 				{
-					InstanceLOD.RaytracingResource = new FHairStrandsRaytracingResource(*InstanceLOD.Data);
-					BeginInitResource(InstanceLOD.RaytracingResource);
+					if (bNeedDeformedPositions)
+					{
+						// Allocate dynamic raytracing resources (owned by the groom component/instance)
+						InstanceLOD.RaytracingResource = new FHairStrandsRaytracingResource(*InstanceLOD.Data);
+						BeginInitResource(InstanceLOD.RaytracingResource);
+						InstanceLOD.RaytracingResourceOwned = true;
+					}
+					else
+					{
+						// (Lazy) Allocate static raytracing resources (owned by the grooom asset)
+						InstanceLOD.RaytracingResourceOwned = false;
+						InstanceLOD.RaytracingResource = GroomAsset->AllocateCardsRaytracingResources(GroupIt, CardsLODIndex);
+						check(InstanceLOD.RaytracingResource);
+					}
 				}
 				#endif
 
@@ -2375,14 +2390,26 @@ void UGroomComponent::InitResources(bool bIsBindingReloading)
 				if (bNeedDeformedPositions)
 				{
 					InstanceLOD.DeformedResource = new FHairMeshesDeformedResource(LOD.Data.RenderData, true);
-				BeginInitResource(InstanceLOD.DeformedResource);
+					BeginInitResource(InstanceLOD.DeformedResource);
 				}
 
 				#if RHI_RAYTRACING
-				if (IsRayTracingEnabled())
+				if (IsRayTracingEnabled() && bNeedDeformedPositions)
 				{
-					InstanceLOD.RaytracingResource = new FHairStrandsRaytracingResource(*InstanceLOD.Data);
-					BeginInitResource(InstanceLOD.RaytracingResource);
+					if (bNeedDeformedPositions)
+					{
+						// Allocate dynamic raytracing resources (owned by the groom component/instance)
+						InstanceLOD.RaytracingResource = new FHairStrandsRaytracingResource(*InstanceLOD.Data);
+						BeginInitResource(InstanceLOD.RaytracingResource);
+						InstanceLOD.RaytracingResourceOwned = true;
+					}
+					else
+					{
+						// (Lazy) Allocate static raytracing resources (owned by the grooom asset)
+						InstanceLOD.RaytracingResourceOwned = false;
+						InstanceLOD.RaytracingResource = GroomAsset->AllocateMeshesRaytracingResources(GroupIt, MeshLODIndex);
+						check(InstanceLOD.RaytracingResource);
+					}
 				}
 				#endif
 
@@ -2457,7 +2484,10 @@ void UGroomComponent::DeleteDeferredHairGroupInstances()
 				InternalResourceRelease(LocalInstance->Strands.DeformedResource);
 
 				#if RHI_RAYTRACING
-				InternalResourceRelease(LocalInstance->Strands.RenRaytracingResource);
+				if (LocalInstance->Strands.RenRaytracingResourceOwned)
+				{
+					InternalResourceRelease(LocalInstance->Strands.RenRaytracingResource);
+				}
 				#endif
 
 				#if WITH_EDITOR
@@ -2476,7 +2506,10 @@ void UGroomComponent::DeleteDeferredHairGroupInstances()
 						InternalResourceRelease(LocalInstance->Cards.LODs[CardLODIt].Guides.DeformedResource);
 						InternalResourceRelease(LocalInstance->Cards.LODs[CardLODIt].DeformedResource);
 						#if RHI_RAYTRACING
-						InternalResourceRelease(LocalInstance->Cards.LODs[CardLODIt].RaytracingResource);
+						if (LocalInstance->Cards.LODs[CardLODIt].RaytracingResourceOwned)
+						{
+							InternalResourceRelease(LocalInstance->Cards.LODs[CardLODIt].RaytracingResource);
+						}
 						#endif
 					}
 				}
@@ -2490,6 +2523,12 @@ void UGroomComponent::DeleteDeferredHairGroupInstances()
 					if (LocalInstance->Meshes.IsValid(MeshesLODIt))
 					{
 						InternalResourceRelease(LocalInstance->Meshes.LODs[MeshesLODIt].DeformedResource);
+						#if RHI_RAYTRACING
+						if (LocalInstance->Meshes.LODs[MeshesLODIt].RaytracingResourceOwned)
+						{
+							InternalResourceRelease(LocalInstance->Meshes.LODs[MeshesLODIt].RaytracingResource);
+						}
+						#endif
 					}
 				}
 			}
diff --git a/Engine/Plugins/Runtime/HairStrands/Source/HairStrandsCore/Private/HairStrandsRendering.cpp b/Engine/Plugins/Runtime/HairStrands/Source/HairStrandsCore/Private/HairStrandsRendering.cpp
--- a/Engine/Plugins/Runtime/HairStrands/Source/HairStrandsCore/Private/HairStrandsRendering.cpp
+++ b/Engine/Plugins/Runtime/HairStrands/Source/HairStrandsCore/Private/HairStrandsRendering.cpp
@@ -1375,7 +1375,7 @@ void ComputeHairStrandsInterpolation(
 	if (InstanceGeometryType == EHairGeometryType::Strands)
 	{
 		DECLARE_GPU_STAT(HairStrandsInterpolation);
-		RDG_EVENT_SCOPE(GraphBuilder, "HairStrandsInterpolation");
+		RDG_EVENT_SCOPE(GraphBuilder, "HairInterpolation(Strands)");
 		RDG_GPU_STAT_SCOPE(GraphBuilder, HairStrandsInterpolation);
 
 		// Debug mode:
@@ -1516,8 +1516,9 @@ void ComputeHairStrandsInterpolation(
 			if (Instance->Strands.RenRaytracingResource)
 			{
 				const float HairRadiusScaleRT = (GHairRaytracingRadiusScale > 0 ? GHairRaytracingRadiusScale : Instance->Strands.Modifier.HairRaytracingRadiusScale);
-				const bool bNeedUpdate = !Instance->Strands.RenRaytracingResource->bIsRTGeometryInitialized || Instance->BindingType == EHairBindingType::Skinning;
-				if (bNeedUpdate)
+				const bool bNeedUpdate = Instance->Strands.DeformedResource != nullptr;
+				const bool bNeedBuild = !Instance->Strands.RenRaytracingResource->bIsRTGeometryInitialized;
+				if (bNeedBuild || bNeedUpdate)
 				{
 					FRDGImportedBuffer Raytracing_PositionBuffer = Register(GraphBuilder, Instance->Strands.RenRaytracingResource->PositionBuffer, ERDGImportedBufferFlags::CreateViews);
 					AddGenerateRaytracingGeometryPass(
@@ -1533,17 +1534,17 @@ void ComputeHairStrandsInterpolation(
 						Raytracing_PositionBuffer.UAV);
 
 					GraphBuilder.AddPass(
-						RDG_EVENT_NAME("HairStrandsUpdateBLAS"),
+						RDG_EVENT_NAME("HairUpdateBLAS(Strands)"),
 						ERDGPassFlags::NeverCull,
-					[Instance](FRHICommandList& RHICmdList)
+					[Instance, bNeedUpdate](FRHICommandList& RHICmdList)
 					{
-						const bool bNeedFullBuild = !Instance->Strands.RenRaytracingResource->bIsRTGeometryInitialized;
-						if (bNeedFullBuild)
+						const bool bLocalNeedBuild = !Instance->Strands.RenRaytracingResource->bIsRTGeometryInitialized;
+						if (bLocalNeedBuild)
 						{
 							FBufferRHIRef PositionBuffer(Instance->Strands.RenRaytracingResource->PositionBuffer.Buffer->GetRHI());
 							BuildHairAccelerationStructure_Strands(RHICmdList, Instance->Strands.RenRaytracingResource->VertexCount, PositionBuffer, &Instance->Strands.RenRaytracingResource->RayTracingGeometry);
 						}
-						else
+						else if (bNeedUpdate)
 						{
 							UpdateHairAccelerationStructure(RHICmdList, &Instance->Strands.RenRaytracingResource->RayTracingGeometry);
 						}
@@ -1560,7 +1561,7 @@ void ComputeHairStrandsInterpolation(
 	else if (InstanceGeometryType == EHairGeometryType::Cards)
 	{	
 		DECLARE_GPU_STAT(HairCardsInterpolation);
-		RDG_EVENT_SCOPE(GraphBuilder, "HairCardsInterpolation");
+		RDG_EVENT_SCOPE(GraphBuilder, "HairInterpolation(Cards)");
 		RDG_GPU_STAT_SCOPE(GraphBuilder, HairCardsInterpolation);
 
 		const uint32 HairLODIndex = Instance->HairGroupPublicData->GetIntLODIndex();
@@ -1646,26 +1647,30 @@ void ComputeHairStrandsInterpolation(
 			#if RHI_RAYTRACING
 			if (LOD.RaytracingResource)
 			{
-				const bool bNeedUpdate = !LOD.RaytracingResource->bIsRTGeometryInitialized || Instance->BindingType == EHairBindingType::Skinning;
-				if (bNeedUpdate)
+				// * When cards are dynamic, RT geometry is built once and then update.
+				// * When cards are static (i.e., no sim, not attached to  skinning()), in which case RT geometry needs only to be built once. This static geometry is shared between all the instances, 
+				//   and owned by the asset, rathter than the component. In this case the RT geometry will be built only once for all the instances
+				const bool bNeedUpdate = bNeedDeformation;
+				const bool bNeedBuild = !LOD.RaytracingResource->bIsRTGeometryInitialized;
+				if (bNeedBuild || bNeedUpdate)
 				{
 					GraphBuilder.AddPass(
-						RDG_EVENT_NAME("HairCardsUpdateBLAS"),
+						RDG_EVENT_NAME("HairUpdateBLAS(Cards)"),
 						ERDGPassFlags::NeverCull,
-						[Instance, HairLODIndex](FRHICommandList& RHICmdList)
+						[Instance, HairLODIndex, bNeedUpdate](FRHICommandList& RHICmdList)
 					{
 						FHairGroupInstance::FCards::FLOD& LocalLOD = Instance->Cards.LODs[HairLODIndex];
 
-						const bool bNeedFullBuild = !LocalLOD.RaytracingResource->bIsRTGeometryInitialized;
-						if (bNeedFullBuild)
+						const bool bLocalNeedBuild = !LocalLOD.RaytracingResource->bIsRTGeometryInitialized;
+						if (bLocalNeedBuild)
 						{
 							BuildHairAccelerationStructure_Cards(RHICmdList, LocalLOD.RestResource, LocalLOD.DeformedResource, &LocalLOD.RaytracingResource->RayTracingGeometry);
+							LocalLOD.RaytracingResource->bIsRTGeometryInitialized = true;
 						}
-						else
+						else if (bNeedUpdate)
 						{
 							UpdateHairAccelerationStructure(RHICmdList, &LocalLOD.RaytracingResource->RayTracingGeometry);
 						}
-						LocalLOD.RaytracingResource->bIsRTGeometryInitialized = true;
 					});
 				}
 			}
@@ -1675,7 +1680,7 @@ void ComputeHairStrandsInterpolation(
 	else if (InstanceGeometryType == EHairGeometryType::Meshes)
 	{
 		DECLARE_GPU_STAT(HairMeshesInterpolation);
-		RDG_EVENT_SCOPE(GraphBuilder, "HairMeshesInterpolation");
+		RDG_EVENT_SCOPE(GraphBuilder, "HairInterpolation(Meshes)");
 		RDG_GPU_STAT_SCOPE(GraphBuilder, HairMeshesInterpolation);
 
 		const uint32 HairLODIndex = Instance->HairGroupPublicData->GetIntLODIndex();
@@ -1708,21 +1713,22 @@ void ComputeHairStrandsInterpolation(
 			FHairGroupInstance::FMeshes::FLOD& LOD = Instance->Meshes.LODs[HairLODIndex];
 			if (LOD.RaytracingResource)
 			{
-				const bool bNeedUpdate = !LOD.RaytracingResource->bIsRTGeometryInitialized || bNeedDeformation;
-				if (bNeedUpdate)
+				const bool bNeedUpdate = bNeedDeformation;
+				const bool bNeedBuid = !LOD.RaytracingResource->bIsRTGeometryInitialized;
+				if (bNeedBuid || bNeedUpdate)
 				{
 					GraphBuilder.AddPass(
-						RDG_EVENT_NAME("HairMeshesUpdateBLAS"),
+						RDG_EVENT_NAME("HairUpdateBLAS(Meshes)"),
 						ERDGPassFlags::NeverCull,
-						[Instance, HairLODIndex](FRHICommandList& RHICmdList)
+						[Instance, HairLODIndex, bNeedUpdate](FRHICommandList& RHICmdList)
 					{
 						FHairGroupInstance::FMeshes::FLOD& LocalLOD = Instance->Meshes.LODs[HairLODIndex];				
-						const bool bNeedFullBuild = !LocalLOD.RaytracingResource->bIsRTGeometryInitialized;
-						if (bNeedFullBuild)
+						const bool bLocalNeedBuild = !LocalLOD.RaytracingResource->bIsRTGeometryInitialized;
+						if (bLocalNeedBuild)
 						{
 							BuildHairAccelerationStructure_Meshes(RHICmdList, LocalLOD.RestResource, LocalLOD.DeformedResource,  &LocalLOD.RaytracingResource->RayTracingGeometry);
 						}
-						else
+						else if (bNeedUpdate)
 						{
 							UpdateHairAccelerationStructure(RHICmdList, &LocalLOD.RaytracingResource->RayTracingGeometry);
 						}
diff --git a/Engine/Plugins/Runtime/HairStrands/Source/HairStrandsCore/Public/GroomAsset.h b/Engine/Plugins/Runtime/HairStrands/Source/HairStrandsCore/Public/GroomAsset.h
--- a/Engine/Plugins/Runtime/HairStrands/Source/HairStrandsCore/Public/GroomAsset.h
+++ b/Engine/Plugins/Runtime/HairStrands/Source/HairStrandsCore/Public/GroomAsset.h
@@ -28,6 +28,9 @@
 class UAssetUserData;
 class UMaterialInterface;
 class UNiagaraSystem;
+struct FHairStrandsRestResource;
+struct FHairStrandsInterpolationResource;
+struct FHairStrandsRaytracingResource;
 
 USTRUCT(BlueprintType)
 struct HAIRSTRANDSCORE_API FHairGroupInfo
@@ -109,11 +112,17 @@ struct HAIRSTRANDSCORE_API FHairGroupData
 			if (RestResource) Total += RestResource->GetResourcesSize();
 			if (InterpolationResource) Total += InterpolationResource->GetResourcesSize();
 			if (ClusterCullingResource) Total += ClusterCullingResource->GetResourcesSize();
+			if (RaytracingResource) Total += RaytracingResource->GetResourcesSize();
 			return Total;
 		}
 
 		FHairStrandsClusterCullingData		ClusterCullingData;
 		FHairStrandsClusterCullingResource* ClusterCullingResource = nullptr;
+
+		#if RHI_RAYTRACING
+		FHairStrandsRaytracingResource*		RaytracingResource = nullptr;
+		#endif
+
 	} Strands;
 
 	struct FCards
@@ -159,6 +168,7 @@ struct HAIRSTRANDSCORE_API FHairGroupData
 				if (RestResource) Total += RestResource->GetResourcesSize();
 				if (ProceduralResource) Total += ProceduralResource->GetResourcesSize();
 				if (InterpolationResource) Total += InterpolationResource->GetResourcesSize();
+				if (RaytracingResource) Total += RaytracingResource->GetResourcesSize();
 				return Total;
 			}
 
@@ -178,6 +188,10 @@ struct HAIRSTRANDSCORE_API FHairGroupData
 
 			FBaseWithInterpolation				Guides;
 
+			#if RHI_RAYTRACING
+			FHairStrandsRaytracingResource*		RaytracingResource = nullptr;
+			#endif
+
 			bool bIsCookedOut = false;
 		};
 		TArray<FLOD> LODs;
@@ -223,6 +237,7 @@ struct HAIRSTRANDSCORE_API FHairGroupData
 			{
 				uint32 Total = 0;
 				if (RestResource) Total += RestResource->GetResourcesSize();
+				if (RaytracingResource) Total += RaytracingResource->GetResourcesSize();
 				return Total;
 			}
 
@@ -230,6 +245,9 @@ struct HAIRSTRANDSCORE_API FHairGroupData
 			bool IsValid() const { return Data.IsValid() && RestResource != nullptr; }
 			FHairMeshesDatas Data;
 			FHairMeshesRestResource* RestResource = nullptr;
+			#if RHI_RAYTRACING
+			FHairStrandsRaytracingResource* RaytracingResource = nullptr;
+			#endif
 			bool bIsCookedOut = false;
 		};
 		TArray<FLOD> LODs;
@@ -461,8 +479,12 @@ public:
 private:
 	void ApplyStripFlags(uint8 StripFlags, const class ITargetPlatform* CookTarget);
 
-	void AllocateGuidesResources(uint32 GroupIndex);
-	void AllocateInterpolationResources(uint32 GroupIndex);
+	// Functions allocating lazily/on-demand resources (guides, interpolation, RT geometry, ...)
+	FHairStrandsRestResource*			AllocateGuidesResources(uint32 GroupIndex);
+	FHairStrandsInterpolationResource*	AllocateInterpolationResources(uint32 GroupIndex);
+	FHairStrandsRaytracingResource*		AllocateCardsRaytracingResources(uint32 GroupIndex, uint32 LODIndex);
+	FHairStrandsRaytracingResource*		AllocateMeshesRaytracingResources(uint32 GroupIndex, uint32 LODIndex);
+	FHairStrandsRaytracingResource*		AllocateStrandsRaytracingResources(uint32 GroupIndex);
 	friend class UGroomComponent;
 
 #if WITH_EDITORONLY_DATA
diff --git a/Engine/Plugins/Runtime/HairStrands/Source/HairStrandsCore/Public/GroomInstance.h b/Engine/Plugins/Runtime/HairStrands/Source/HairStrandsCore/Public/GroomInstance.h
--- a/Engine/Plugins/Runtime/HairStrands/Source/HairStrandsCore/Public/GroomInstance.h
+++ b/Engine/Plugins/Runtime/HairStrands/Source/HairStrandsCore/Public/GroomInstance.h
@@ -91,6 +91,7 @@ struct HAIRSTRANDSCORE_API FHairGroupInstance : public FHairStrandsInstance
 		// Resources - Raytracing data when enabling (expensive) raytracing method
 		#if RHI_RAYTRACING
 		FHairStrandsRaytracingResource* RenRaytracingResource = nullptr;
+		bool RenRaytracingResourceOwned = false;
 		#endif
 
 		FRDGExternalBuffer DebugAttributeBuffer;
@@ -117,6 +118,7 @@ struct HAIRSTRANDSCORE_API FHairGroupInstance : public FHairStrandsInstance
 			FHairCardsDeformedResource* DeformedResource = nullptr;
 			#if RHI_RAYTRACING
 			FHairStrandsRaytracingResource* RaytracingResource = nullptr;
+			bool RaytracingResourceOwned = false;
 			#endif
 
 			// Interpolation data/resources
@@ -149,6 +151,7 @@ struct HAIRSTRANDSCORE_API FHairGroupInstance : public FHairStrandsInstance
 			FHairMeshesDeformedResource* DeformedResource = nullptr;
 			#if RHI_RAYTRACING
 			FHairStrandsRaytracingResource* RaytracingResource = nullptr;
+			bool RaytracingResourceOwned = false;
 			#endif
 
 			UMaterialInterface* Material = nullptr;
