Fix CIS issue.

diff --git a/Engine/Plugins/Runtime/HairStrands/Source/HairStrandsCore/Private/GroomAsset.cpp b/Engine/Plugins/Runtime/HairStrands/Source/HairStrandsCore/Private/GroomAsset.cpp
--- a/Engine/Plugins/Runtime/HairStrands/Source/HairStrandsCore/Private/GroomAsset.cpp
+++ b/Engine/Plugins/Runtime/HairStrands/Source/HairStrandsCore/Private/GroomAsset.cpp
@@ -2258,9 +2258,9 @@ FHairStrandsInterpolationResource* UGroomAsset::AllocateInterpolationResources(u
 	return nullptr;
 }
 
+#if RHI_RAYTRACING
 FHairStrandsRaytracingResource* UGroomAsset::AllocateCardsRaytracingResources(uint32 GroupIndex, uint32 LODIndex)
 {
-#if RHI_RAYTRACING
 	if (IsRayTracingEnabled() && GroupIndex < uint32(GetNumHairGroups()) && LODIndex < uint32(HairGroupsLOD[GroupIndex].LODs.Num()))
 	{
 		FHairGroupData& GroupData = HairGroupsData[GroupIndex];
@@ -2274,13 +2274,11 @@ FHairStrandsRaytracingResource* UGroomAsset::AllocateCardsRaytracingResources(ui
 		}
 		return LOD.RaytracingResource;
 	}
-#endif
 	return nullptr;
 }
 
 FHairStrandsRaytracingResource* UGroomAsset::AllocateMeshesRaytracingResources(uint32 GroupIndex, uint32 LODIndex)
 {
-#if RHI_RAYTRACING
 	if (IsRayTracingEnabled() && GroupIndex < uint32(GetNumHairGroups()) && LODIndex < uint32(HairGroupsLOD[GroupIndex].LODs.Num()))
 	{
 		FHairGroupData& GroupData = HairGroupsData[GroupIndex];
@@ -2294,13 +2292,11 @@ FHairStrandsRaytracingResource* UGroomAsset::AllocateMeshesRaytracingResources(u
 		}
 		return LOD.RaytracingResource;
 	}
-#endif
 	return nullptr;
 }
 
 FHairStrandsRaytracingResource* UGroomAsset::AllocateStrandsRaytracingResources(uint32 GroupIndex)
 {
-#if RHI_RAYTRACING
 	if (IsRayTracingEnabled() && GroupIndex < uint32(GetNumHairGroups()))
 	{
 		FHairGroupData& GroupData = HairGroupsData[GroupIndex];
@@ -2313,9 +2309,9 @@ FHairStrandsRaytracingResource* UGroomAsset::AllocateStrandsRaytracingResources(
 		}
 		return GroupData.Strands.RaytracingResource;
 	}
-#endif
 	return nullptr;
 }
+#endif // RHI_RAYTRACING
 
 void UGroomAsset::InitStrandsResources()
 {
diff --git a/Engine/Plugins/Runtime/HairStrands/Source/HairStrandsCore/Public/GroomAsset.h b/Engine/Plugins/Runtime/HairStrands/Source/HairStrandsCore/Public/GroomAsset.h
--- a/Engine/Plugins/Runtime/HairStrands/Source/HairStrandsCore/Public/GroomAsset.h
+++ b/Engine/Plugins/Runtime/HairStrands/Source/HairStrandsCore/Public/GroomAsset.h
@@ -112,7 +112,9 @@ struct HAIRSTRANDSCORE_API FHairGroupData
 			if (RestResource) Total += RestResource->GetResourcesSize();
 			if (InterpolationResource) Total += InterpolationResource->GetResourcesSize();
 			if (ClusterCullingResource) Total += ClusterCullingResource->GetResourcesSize();
+			#if RHI_RAYTRACING
 			if (RaytracingResource) Total += RaytracingResource->GetResourcesSize();
+			#endif
 			return Total;
 		}
 
@@ -168,7 +170,9 @@ struct HAIRSTRANDSCORE_API FHairGroupData
 				if (RestResource) Total += RestResource->GetResourcesSize();
 				if (ProceduralResource) Total += ProceduralResource->GetResourcesSize();
 				if (InterpolationResource) Total += InterpolationResource->GetResourcesSize();
+				#if RHI_RAYTRACING
 				if (RaytracingResource) Total += RaytracingResource->GetResourcesSize();
+				#endif
 				return Total;
 			}
 
@@ -237,7 +241,9 @@ struct HAIRSTRANDSCORE_API FHairGroupData
 			{
 				uint32 Total = 0;
 				if (RestResource) Total += RestResource->GetResourcesSize();
+				#if RHI_RAYTRACING
 				if (RaytracingResource) Total += RaytracingResource->GetResourcesSize();
+				#endif
 				return Total;
 			}
 
@@ -482,9 +488,11 @@ private:
 	// Functions allocating lazily/on-demand resources (guides, interpolation, RT geometry, ...)
 	FHairStrandsRestResource*			AllocateGuidesResources(uint32 GroupIndex);
 	FHairStrandsInterpolationResource*	AllocateInterpolationResources(uint32 GroupIndex);
+#if RHI_RAYTRACING
 	FHairStrandsRaytracingResource*		AllocateCardsRaytracingResources(uint32 GroupIndex, uint32 LODIndex);
 	FHairStrandsRaytracingResource*		AllocateMeshesRaytracingResources(uint32 GroupIndex, uint32 LODIndex);
 	FHairStrandsRaytracingResource*		AllocateStrandsRaytracingResources(uint32 GroupIndex);
+#endif
 	friend class UGroomComponent;
 
 #if WITH_EDITORONLY_DATA
