Fix hair resources memory leak.

diff --git a/Engine/Plugins/Runtime/HairStrands/Source/HairStrandsCore/Private/GroomComponent.cpp b/Engine/Plugins/Runtime/HairStrands/Source/HairStrandsCore/Private/GroomComponent.cpp
--- a/Engine/Plugins/Runtime/HairStrands/Source/HairStrandsCore/Private/GroomComponent.cpp
+++ b/Engine/Plugins/Runtime/HairStrands/Source/HairStrandsCore/Private/GroomComponent.cpp
@@ -2444,7 +2444,7 @@ void UGroomComponent::ReleaseResources()
 	InitializedResources = nullptr;
 
 	// Deferring instances deletion to insure scene proxy are done with the rendering data
-	DeferredDeleteHairGroupInstances = HairGroupInstances;
+	DeferredDeleteHairGroupInstances.Append(HairGroupInstances);
 	HairGroupInstances.Empty();
 
 	// Insure the ticking of the Groom component always happens after the skeletalMeshComponent.
