Fix guides culling + IsSimulateEnabled option

diff --git a/Engine/Plugins/Runtime/HairStrands/Source/HairStrandsCore/Private/GroomComponent.cpp b/Engine/Plugins/Runtime/HairStrands/Source/HairStrandsCore/Private/GroomComponent.cpp
--- a/Engine/Plugins/Runtime/HairStrands/Source/HairStrandsCore/Private/GroomComponent.cpp
+++ b/Engine/Plugins/Runtime/HairStrands/Source/HairStrandsCore/Private/GroomComponent.cpp
@@ -1795,8 +1795,8 @@ void UGroomComponent::UpdateSimulatedGroups()
 				if (Instance->Guides.IsValid())
 				{
 					check(Instance->HairGroupPublicData);
-					Instance->Guides.bIsSimulationEnable	 = Instance->HairGroupPublicData->IsGlobalInterpolationEnable(LODIndex);
-					Instance->Guides.bHasGlobalInterpolation = Instance->HairGroupPublicData->IsSimulationEnable(LODIndex);
+					Instance->Guides.bIsSimulationEnable	 = Instance->HairGroupPublicData->IsSimulationEnable(LODIndex);
+					Instance->Guides.bHasGlobalInterpolation = Instance->HairGroupPublicData->IsGlobalInterpolationEnable(LODIndex);
 				}
 				++GroupIt;
 			}
@@ -2150,7 +2150,7 @@ void UGroomComponent::InitResources(bool bIsBindingReloading)
 
 			// Initialize the simulation and the global deformation to its default behavior by setting it with LODIndex = -1
 			const int32 LODIndex = -1;
-			HairGroupInstance->Guides.bIsSimulationEnable = GroomAsset->IsGlobalInterpolationEnable(GroupIt,LODIndex);
+			HairGroupInstance->Guides.bIsSimulationEnable = GroomAsset->IsSimulationEnable(GroupIt,LODIndex);
 			HairGroupInstance->Guides.bHasGlobalInterpolation = LocalBindingAsset && GroomAsset->IsGlobalInterpolationEnable(GroupIt,LODIndex);
 		}
 
diff --git a/Engine/Plugins/Runtime/HairStrands/Source/HairStrandsCore/Private/GroomManager.cpp b/Engine/Plugins/Runtime/HairStrands/Source/HairStrandsCore/Private/GroomManager.cpp
--- a/Engine/Plugins/Runtime/HairStrands/Source/HairStrandsCore/Private/GroomManager.cpp
+++ b/Engine/Plugins/Runtime/HairStrands/Source/HairStrandsCore/Private/GroomManager.cpp
@@ -676,9 +676,9 @@ void ProcessHairStrandsBookmark(
 	FHairStrandsBookmarkParameters& Parameters)
 {
 	check(Parameters.Instances != nullptr);
+
 	const bool bCulling =
 		IsInstanceFrustumCullingEnable() && (
-			Bookmark == EHairStrandsBookmark::ProcessGuideInterpolation ||
 			Bookmark == EHairStrandsBookmark::ProcessGatherCluster ||
 			Bookmark == EHairStrandsBookmark::ProcessStrandsInterpolation ||
 			Bookmark == EHairStrandsBookmark::ProcessDebug);
