Change non-deformable hair (no-simuation, no-globaldeformation, no-skinning, rigid attachement, ..) skip cluster AABB update, and have plain voxel page allocated based on their group AABB.

When hitting these conditions, and groom having a single LOD: culling, cluster AABB,  voxel page allocation passes can be avoided/simplied which saves performance.

diff --git a/Engine/Plugins/Runtime/HairStrands/Source/HairStrandsCore/Private/GroomComponent.cpp b/Engine/Plugins/Runtime/HairStrands/Source/HairStrandsCore/Private/GroomComponent.cpp
--- a/Engine/Plugins/Runtime/HairStrands/Source/HairStrandsCore/Private/GroomComponent.cpp
+++ b/Engine/Plugins/Runtime/HairStrands/Source/HairStrandsCore/Private/GroomComponent.cpp
@@ -919,8 +919,8 @@ public:
 			VertexFactory = (FVertexFactory*)&StrandsVertexFactory;
 			HairVertexCount = Instance->Strands.RestResource->GetVertexCount();
 			MaxVertexIndex = HairVertexCount * 6;
-			NumPrimitive = 0;
-			bUseCulling = true;
+			bUseCulling = Instance->Strands.bIsCullingEnabled;
+			NumPrimitive = bUseCulling ? 0 : HairVertexCount * 2;
 			MaterialRenderProxy = Strands_MaterialProxy == nullptr ? Instance->Strands.Material->GetRenderProxy() : Strands_MaterialProxy;
 			bWireframe = AllowDebugViewmodes() && ViewFamily.EngineShowFlags.Wireframe;
 		}
diff --git a/Engine/Plugins/Runtime/HairStrands/Source/HairStrandsCore/Private/GroomManager.cpp b/Engine/Plugins/Runtime/HairStrands/Source/HairStrandsCore/Private/GroomManager.cpp
--- a/Engine/Plugins/Runtime/HairStrands/Source/HairStrandsCore/Private/GroomManager.cpp
+++ b/Engine/Plugins/Runtime/HairStrands/Source/HairStrandsCore/Private/GroomManager.cpp
@@ -363,7 +363,14 @@ static void RunHairStrandsGatherCluster(
 
 		if (Instance->Strands.IsValid())
 		{
-			RegisterClusterData(Instance, ClusterData);
+			if (Instance->Strands.bIsCullingEnabled)
+			{
+				RegisterClusterData(Instance, ClusterData);
+			}
+			else
+			{
+				Instance->HairGroupPublicData->bCullingResultAvailable = false;
+			}
 		}
 	}
 }
@@ -569,6 +576,8 @@ static void RunHairLODSelection(const FHairStrandsInstances& Instances, const TA
 		const bool bForceCards = GHairStrands_UseCards > 0 || Instance->bForceCards; // todo
 		EHairGeometryType GeometryType = ConvertLODGeometryType(LODGeometryTypes[IntLODIndex], bForceCards, ShaderPlatform);
 
+		// Skip cluster culling if strands have a single LOD (i.e.: LOD0), and the instance is static (i.e., no skinning, no simulation, ...)
+		bool bCullingEnable = false;
 		if (GeometryType == EHairGeometryType::Meshes)
 		{
 			if (!Instance->Meshes.IsValid(IntLODIndex))
@@ -589,6 +598,12 @@ static void RunHairLODSelection(const FHairStrandsInstances& Instances, const TA
 			{
 				GeometryType = EHairGeometryType::NoneGeometry;
 			}
+			else
+			{
+				const bool bNeedLODing		= LODIndex > 0;
+				const bool bNeedDeformation = Instance->Strands.DeformedResource != nullptr;
+				bCullingEnable = bNeedLODing || bNeedDeformation;
+			}
 		}
 
 		if (!bIsVisible)
@@ -605,6 +620,7 @@ static void RunHairLODSelection(const FHairStrandsInstances& Instances, const TA
 		Instance->BindingType = Instance->HairGroupPublicData->GetBindingType(IntLODIndex);
 		Instance->Guides.bIsSimulationEnable = Instance->HairGroupPublicData->IsSimulationEnable(IntLODIndex);
 		Instance->Guides.bHasGlobalInterpolation = Instance->HairGroupPublicData->IsGlobalInterpolationEnable(IntLODIndex);
+		Instance->Strands.bIsCullingEnabled = bCullingEnable;
 
 	}
 }
diff --git a/Engine/Plugins/Runtime/HairStrands/Source/HairStrandsCore/Private/HairStrandsRendering.cpp b/Engine/Plugins/Runtime/HairStrands/Source/HairStrandsCore/Private/HairStrandsRendering.cpp
--- a/Engine/Plugins/Runtime/HairStrands/Source/HairStrandsCore/Private/HairStrandsRendering.cpp
+++ b/Engine/Plugins/Runtime/HairStrands/Source/HairStrandsCore/Private/HairStrandsRendering.cpp
@@ -324,9 +324,12 @@ FRDGHairStrandsCullingData ImportCullingData(FRDGBuilder& GraphBuilder, FHairGro
 	FRDGHairStrandsCullingData Out;
 	Out.bCullingResultAvailable					= In->GetCullingResultAvailable();
 
-	Out.HairStrandsVF_CullingIndirectBuffer		= Register(GraphBuilder, In->GetDrawIndirectRasterComputeBuffer(), ERDGImportedBufferFlags::CreateViews);
-	Out.HairStrandsVF_CullingIndexBuffer		= Register(GraphBuilder, In->GetCulledVertexIdBuffer(), ERDGImportedBufferFlags::CreateViews);
-	Out.HairStrandsVF_CullingRadiusScaleBuffer	= Register(GraphBuilder, In->GetCulledVertexRadiusScaleBuffer(), ERDGImportedBufferFlags::CreateViews);
+	if (Out.bCullingResultAvailable)
+	{
+		Out.HairStrandsVF_CullingIndirectBuffer		= Register(GraphBuilder, In->GetDrawIndirectRasterComputeBuffer(), ERDGImportedBufferFlags::CreateViews);
+		Out.HairStrandsVF_CullingIndexBuffer		= Register(GraphBuilder, In->GetCulledVertexIdBuffer(), ERDGImportedBufferFlags::CreateViews);
+		Out.HairStrandsVF_CullingRadiusScaleBuffer	= Register(GraphBuilder, In->GetCulledVertexRadiusScaleBuffer(), ERDGImportedBufferFlags::CreateViews);
+	}
 
 	Out.ClusterCount		= In->GetClusterCount();
 	Out.ClusterAABBBuffer	= Register(GraphBuilder, In->GetClusterAABBBuffer(), ERDGImportedBufferFlags::CreateViews);
@@ -665,18 +668,20 @@ static void AddHairStrandsInterpolationPass(
 
 ///////////////////////////////////////////////////////////////////////////////////////////////////
 
-
 class FHairClusterAABBCS : public FGlobalShader
 {
 	DECLARE_GLOBAL_SHADER(FHairClusterAABBCS);
 	SHADER_USE_PARAMETER_STRUCT(FHairClusterAABBCS, FGlobalShader);
 
 	class FGroupSize : SHADER_PERMUTATION_SPARSE_INT("PERMUTATION_GROUP_SIZE", 32, 64);
-	using FPermutationDomain = TShaderPermutationDomain<FGroupSize>;
+	class FCPUAABB : SHADER_PERMUTATION_BOOL("PERMUTATION_USE_CPU_AABB");
+	using FPermutationDomain = TShaderPermutationDomain<FGroupSize, FCPUAABB>;
 
 	BEGIN_SHADER_PARAMETER_STRUCT(FParameters, )
 		SHADER_PARAMETER(uint32, DispatchCountX)
 		SHADER_PARAMETER(uint32, ClusterCount)
+		SHADER_PARAMETER(FVector, CPUBoundMin)
+		SHADER_PARAMETER(FVector, CPUBoundMax)
 		SHADER_PARAMETER(FMatrix, LocalToWorldMatrix)
 
 		SHADER_PARAMETER_RDG_BUFFER_SRV(Buffer, RenderDeformedPositionBuffer)
@@ -701,51 +706,75 @@ public:
 
 IMPLEMENT_GLOBAL_SHADER(FHairClusterAABBCS, "/Engine/Private/HairStrands/HairStrandsInterpolation.usf", "ClusterAABBEvaluationCS", SF_Compute);
 
+enum class EHairAABBUpdateType
+{
+	UpdateClusterAABB,
+	UpdateGroupAABB
+};
+
 static void AddHairClusterAABBPass(
 	FRDGBuilder& GraphBuilder,
 	FGlobalShaderMap* ShaderMap,
-	const FTransform& InRenLocalToWorld,
+	const EHairAABBUpdateType UpdateType,
+	FHairGroupInstance* Instance,
 	FRDGBufferSRVRef RenderDeformedOffsetBuffer,
-	FHairStrandClusterData::FHairGroup& ClusterData,
-
+	FHairStrandClusterData::FHairGroup* ClusterData,
 	FRDGHairStrandsCullingData& ClusterAABBData,
-
 	FRDGBufferSRVRef RenderPositionBufferSRV,
 	FRDGBufferSRVRef& DrawIndirectRasterComputeBuffer)
 {
+	// Clusters AABB are only update if the groom is deformed.
+	const FBoxSphereBounds& Bounds = Instance->Strands.Data->BoundingBox;
+	const FTransform& InRenLocalToWorld = Instance->LocalToWorld;
+	const FBoxSphereBounds TransformedBounds = Bounds.TransformBy(InRenLocalToWorld);
+	Instance->HairGroupPublicData->bClusterAABBValid = UpdateType == EHairAABBUpdateType::UpdateClusterAABB;
+
+	// bNeedDeformation: 
+	// * If the instance has deformatin (simulation, global interpolation, skinning, then we update all clusters (for voxelization purpose) and the instance AABB
+	// * If the instance is static, we only update the instance AABB
 	const uint32 GroupSize = ComputeGroupSize();
-	const FIntVector DispatchCount = ComputeDispatchGroupCount2D(ClusterData.ClusterCount);
+	const FIntVector DispatchCount = (UpdateType == EHairAABBUpdateType::UpdateClusterAABB && ClusterData) ? ComputeDispatchGroupCount2D(ClusterData->ClusterCount) : FIntVector(1,1,1);
 
-	FRDGBufferRef ClusterIdBuffer = GraphBuilder.RegisterExternalBuffer(ClusterData.ClusterIdBuffer);
-	FRDGBufferRef ClusterIndexOffsetBuffer = GraphBuilder.RegisterExternalBuffer(ClusterData.ClusterIndexOffsetBuffer);
-	FRDGBufferRef ClusterIndexCountBuffer  = GraphBuilder.RegisterExternalBuffer(ClusterData.ClusterIndexCountBuffer);
 	FHairClusterAABBCS::FParameters* Parameters = GraphBuilder.AllocParameters<FHairClusterAABBCS::FParameters>();
+	Parameters->CPUBoundMin = TransformedBounds.GetBox().Min;
+	Parameters->CPUBoundMax = TransformedBounds.GetBox().Max;
 	Parameters->DispatchCountX = DispatchCount.X;
-	Parameters->ClusterCount = ClusterData.ClusterCount;
 	Parameters->LocalToWorldMatrix = InRenLocalToWorld.ToMatrixWithScale();
 	Parameters->RenderDeformedPositionBuffer = RenderPositionBufferSRV;
 	Parameters->RenderDeformedOffsetBuffer = RenderDeformedOffsetBuffer;
-	Parameters->ClusterVertexIdBuffer = RegisterAsSRV(GraphBuilder, *ClusterData.ClusterVertexIdBuffer);
-	Parameters->ClusterIdBuffer = GraphBuilder.CreateSRV(ClusterIdBuffer, PF_R32_UINT);
-	Parameters->ClusterIndexOffsetBuffer = GraphBuilder.CreateSRV(ClusterIndexOffsetBuffer, PF_R32_UINT);
-	Parameters->ClusterIndexCountBuffer = GraphBuilder.CreateSRV(ClusterIndexCountBuffer, PF_R32_UINT);
+
+	if (UpdateType == EHairAABBUpdateType::UpdateClusterAABB)
+	{
+		check(ClusterData);
+
+		FRDGBufferRef ClusterIdBuffer = GraphBuilder.RegisterExternalBuffer(ClusterData->ClusterIdBuffer);
+		FRDGBufferRef ClusterIndexOffsetBuffer = GraphBuilder.RegisterExternalBuffer(ClusterData->ClusterIndexOffsetBuffer);
+		FRDGBufferRef ClusterIndexCountBuffer  = GraphBuilder.RegisterExternalBuffer(ClusterData->ClusterIndexCountBuffer);
+		Parameters->ClusterVertexIdBuffer = RegisterAsSRV(GraphBuilder, *ClusterData->ClusterVertexIdBuffer);
+		Parameters->ClusterIdBuffer = GraphBuilder.CreateSRV(ClusterIdBuffer, PF_R32_UINT);
+		Parameters->ClusterIndexOffsetBuffer = GraphBuilder.CreateSRV(ClusterIndexOffsetBuffer, PF_R32_UINT);
+		Parameters->ClusterIndexCountBuffer = GraphBuilder.CreateSRV(ClusterIndexCountBuffer, PF_R32_UINT);
+		Parameters->ClusterCount = ClusterData->ClusterCount;
+
+		// Sanity check
+		check(ClusterData->ClusterCount == ClusterIdBuffer->Desc.NumElements);
+		check(ClusterData->ClusterCount == ClusterIndexOffsetBuffer->Desc.NumElements);
+		check(ClusterData->ClusterCount == ClusterIndexCountBuffer->Desc.NumElements);
+		check(ClusterData->ClusterCount * 6 == ClusterAABBData.ClusterAABBBuffer.Buffer->Desc.NumElements);
+	}
 	Parameters->HairStrandsVF_CullingIndirectBuffer = DrawIndirectRasterComputeBuffer; // Used for checking max vertex count
 	Parameters->OutClusterAABBBuffer = ClusterAABBData.ClusterAABBBuffer.UAV;
 	Parameters->OutGroupAABBBuffer = ClusterAABBData.GroupAABBBuffer.UAV;
 
-	// Sanity check
-	check(ClusterData.ClusterCount == ClusterIdBuffer->Desc.NumElements);
-	check(ClusterData.ClusterCount == ClusterIndexOffsetBuffer->Desc.NumElements);
-	check(ClusterData.ClusterCount == ClusterIndexCountBuffer->Desc.NumElements);
-	check(ClusterData.ClusterCount * 6 == ClusterAABBData.ClusterAABBBuffer.Buffer->Desc.NumElements);
 
 	FHairClusterAABBCS::FPermutationDomain PermutationVector;
 	PermutationVector.Set<FHairClusterAABBCS::FGroupSize>(GroupSize);
+	PermutationVector.Set<FHairClusterAABBCS::FCPUAABB>(UpdateType == EHairAABBUpdateType::UpdateClusterAABB ? 0u : 1u);
 	TShaderMapRef<FHairClusterAABBCS> ComputeShader(ShaderMap, PermutationVector);
 
 	FComputeShaderUtils::AddPass(
 		GraphBuilder,
-		RDG_EVENT_NAME("HairStrandsClusterAABB"),
+		RDG_EVENT_NAME("HairStrandsClusterAABB(%s)", TEXT("Cluster,Group"), TEXT("Group Only")),
 		ComputeShader,
 		Parameters,
 		DispatchCount);
@@ -1055,6 +1084,7 @@ class FClearClusterAABBCS : public FGlobalShader
 		SHADER_PARAMETER_RDG_BUFFER_UAV(RWBuffer, OutClusterAABBBuffer)
 		SHADER_PARAMETER_RDG_BUFFER_UAV(RWBuffer, OutGroupAABBBuffer)
 		SHADER_PARAMETER(uint32, ClusterCount)
+		SHADER_PARAMETER(uint32, bClearClusterAABBs)
 	END_SHADER_PARAMETER_STRUCT()
 
 public:
@@ -1071,6 +1101,7 @@ IMPLEMENT_GLOBAL_SHADER(FClearClusterAABBCS, "/Engine/Private/HairStrands/HairSt
 static void AddClearClusterAABBPass(
 	FRDGBuilder& GraphBuilder,
 	FGlobalShaderMap* ShaderMap,
+	const EHairAABBUpdateType UpdateType,
 	uint32 ClusterCount,
 	FRDGImportedBuffer& OutClusterAABBBuffer,
 	FRDGImportedBuffer& OutGroupAABBBuffer)
@@ -1079,15 +1110,20 @@ static void AddClearClusterAABBPass(
 
 	FClearClusterAABBCS::FParameters* Parameters = GraphBuilder.AllocParameters<FClearClusterAABBCS::FParameters>();
 	Parameters->ClusterCount = ClusterCount;
+	Parameters->bClearClusterAABBs = UpdateType == EHairAABBUpdateType::UpdateClusterAABB ? 1 : 0;
 	Parameters->OutClusterAABBBuffer = OutClusterAABBBuffer.UAV;
 	Parameters->OutGroupAABBBuffer = OutGroupAABBBuffer.UAV;
 
 	TShaderMapRef<FClearClusterAABBCS> ComputeShader(ShaderMap);
 
-	const FIntVector DispatchCount = DispatchCount.DivideAndRoundUp(FIntVector(ClusterCount * 6, 1, 1) , FIntVector(64, 1, 1));
+	const FIntVector DispatchCount = 
+		UpdateType == EHairAABBUpdateType::UpdateClusterAABB ? 
+		FIntVector(FMath::DivideAndRoundUp(ClusterCount * 6u, 64u), 1, 1) :
+		FIntVector(1,1,1);
+
 	FComputeShaderUtils::AddPass(
 		GraphBuilder,
-		RDG_EVENT_NAME("HairStrandsClearClusterAABB"),
+		RDG_EVENT_NAME("HairStrandsClearClusterAABB(%s)", UpdateType == EHairAABBUpdateType::UpdateClusterAABB ? TEXT("Cluster,Group") : TEXT("Group Only")),
 		ComputeShader,
 		Parameters,
 		DispatchCount);
@@ -1404,15 +1440,19 @@ void ComputeHairStrandsInterpolation(
 		else
 		{
 			const bool bNeedDeformation = Instance->Strands.DeformedResource != nullptr;
+			const EHairAABBUpdateType UpdateType = bNeedDeformation ? EHairAABBUpdateType::UpdateClusterAABB : EHairAABBUpdateType::UpdateGroupAABB;
 
-			// 1. Clear cluster AABBs (use for voxel allocation & optionally for culling)
+			// 1. Clear cluster AABBs (used optionally  for voxel allocation & for culling)
 			FRDGHairStrandsCullingData CullingData = ImportCullingData(GraphBuilder, Instance->HairGroupPublicData);
-			AddClearClusterAABBPass(
-				GraphBuilder,
-				ShaderMap,
-				CullingData.ClusterCount,
-				CullingData.ClusterAABBBuffer,
-				CullingData.GroupAABBBuffer);
+			{
+				AddClearClusterAABBPass(
+					GraphBuilder,
+					ShaderMap,
+					UpdateType,
+					CullingData.ClusterCount,
+					CullingData.ClusterAABBBuffer,
+					CullingData.GroupAABBBuffer);
+			}
 
 			// 2. Deform hair if needed (e.g.: skinning, simulation, RBF) and recompute tangent
 			FRDGBufferSRVRef Strands_PositionSRV = nullptr;
@@ -1483,18 +1523,29 @@ void ComputeHairStrandsInterpolation(
 				Strands_PositionOffsetSRV = RegisterAsSRV(GraphBuilder, Instance->Strands.RestResource->PositionOffsetBuffer);
 				// Generated the static tangent if they haven;t been generated yet
 				Instance->Strands.RestResource->GetTangentBuffer(GraphBuilder, ShaderMap);
-
 			}
 
 			// 2.1 Update the VF input with the update resources
 			Instance->HairGroupPublicData->VFInput = InternalComputeHairStrandsVertexInputData(&GraphBuilder, Instance);
 
 			// 3. Compute cluster AABBs (used for LODing and voxelization)
-			FHairStrandClusterData::FHairGroup& HairGroupCluster = InClusterData->HairGroups[Instance->HairGroupPublicData->ClusterDataIndex];
-			if (HairGroupCluster.bVisible)
 			{
 				// Optim: If an instance is using CPU selection and does not voxelize it's data, then there is no need for having valid AABB
-				const bool bNeedAABB = !Instance->bUseCPULODSelection || (Instance->Strands.Modifier.bSupportVoxelization && Instance->bCastShadow);
+				bool bNeedAABB = !Instance->bUseCPULODSelection || (Instance->Strands.Modifier.bSupportVoxelization && Instance->bCastShadow);
+
+				FHairStrandClusterData::FHairGroup* HairGroupCluster = nullptr;
+				if (CullingData.bCullingResultAvailable)
+				{
+					// Sanity check
+					check(Instance->Strands.bIsCullingEnabled);
+					HairGroupCluster = &InClusterData->HairGroups[Instance->HairGroupPublicData->ClusterDataIndex];
+					if (!HairGroupCluster->bVisible)
+					{
+						bNeedAABB = false;
+					}
+				}
+				Instance->HairGroupPublicData->bClusterAABBValid = false;
+				Instance->HairGroupPublicData->bGroupAABBValid = false;
 
 				if (bNeedAABB)
 				{
@@ -1502,12 +1553,16 @@ void ComputeHairStrandsInterpolation(
 					AddHairClusterAABBPass(
 						GraphBuilder,
 						ShaderMap,
-						Instance->LocalToWorld,
+						UpdateType,
+						Instance,
 						Strands_PositionOffsetSRV,
 						HairGroupCluster,
 						CullingData,
 						Strands_PositionSRV,
 						Strands_CulledVertexCount.SRV);
+
+					Instance->HairGroupPublicData->bClusterAABBValid = UpdateType == EHairAABBUpdateType::UpdateClusterAABB;
+					Instance->HairGroupPublicData->bGroupAABBValid = true;
 				}
 			}
 
diff --git a/Engine/Plugins/Runtime/HairStrands/Source/HairStrandsCore/Public/GroomInstance.h b/Engine/Plugins/Runtime/HairStrands/Source/HairStrandsCore/Public/GroomInstance.h
--- a/Engine/Plugins/Runtime/HairStrands/Source/HairStrandsCore/Public/GroomInstance.h
+++ b/Engine/Plugins/Runtime/HairStrands/Source/HairStrandsCore/Public/GroomInstance.h
@@ -71,6 +71,9 @@ struct HAIRSTRANDSCORE_API FHairGroupInstance : public FHairStrandsInstance
 		FHairStrandsInterpolationResource* InterpolationResource = nullptr;
 
 		EHairInterpolationType HairInterpolationType = EHairInterpolationType::NoneSkinning;
+
+		// Indicates if culling is enabled for this hair strands data.
+		bool bIsCullingEnabled = false;
 	};
 
 	//////////////////////////////////////////////////////////////////////////////////////////
diff --git a/Engine/Shaders/Private/HairStrands/HairStrandsClusterCulling.usf b/Engine/Shaders/Private/HairStrands/HairStrandsClusterCulling.usf
--- a/Engine/Shaders/Private/HairStrands/HairStrandsClusterCulling.usf
+++ b/Engine/Shaders/Private/HairStrands/HairStrandsClusterCulling.usf
@@ -39,6 +39,8 @@ uint3 ComputeDispatchGroupCount2D(uint GroupCount)
 /// Clear clusters AABB
 #ifdef SHADER_CLEARCLUSTERAABB
 
+uint bClearClusterAABBs;
+
 RWBuffer<int> OutClusterAABBBuffer; // Cluster data packed as {uint3 AABBMin, uint3 AABBMax}
 RWBuffer<int> OutGroupAABBBuffer; 	// Group data packed as {uint3 AABBMin, uint3 AABBMax}
 
@@ -46,7 +48,7 @@ RWBuffer<int> OutGroupAABBBuffer; 	// Group data packed as {uint3 AABBMin, uint3
 void MainClearClusterAABBCS(uint2 DispatchThreadId : SV_DispatchThreadID)
 {
 	// Clear 'clusters' AABB
-	if (DispatchThreadId.x < ClusterCount*6)
+	if (bClearClusterAABBs > 0 && DispatchThreadId.x < ClusterCount*6)
 	{
 		const int ClusterId = DispatchThreadId.x / 6;
 		const bool bIsMin = (DispatchThreadId.x - ClusterId * 6) < 3 ? 1 : 0;
diff --git a/Engine/Shaders/Private/HairStrands/HairStrandsDeepShadowMask.usf b/Engine/Shaders/Private/HairStrands/HairStrandsDeepShadowMask.usf
--- a/Engine/Shaders/Private/HairStrands/HairStrandsDeepShadowMask.usf
+++ b/Engine/Shaders/Private/HairStrands/HairStrandsDeepShadowMask.usf
@@ -493,6 +493,7 @@ void MainPS(
 
 		const FPackedVirtualVoxelNodeDesc PackedNode = VirtualVoxel.NodeDescBuffer[Voxel_MacroGroupId];
 		const FVirtualVoxelNodeDesc NodeDesc = UnpackVoxelNode(PackedNode, VirtualVoxel.PageResolution);
+		const bool bIsValid = all(NodeDesc.PageIndexResolution != 0);
 
 		#if VOXEL_TRAVERSAL_DEBUG
 		const bool bDebugEnabled = PixelCoord.x == GetCursorPos().x && PixelCoord.y == GetCursorPos().y;
@@ -509,7 +510,7 @@ void MainPS(
 		TraversalSettings.Random = Random;
 		TraversalSettings.PixelRadius = SceneDepth * VirtualVoxel.HairCoveragePixelRadiusAtDepth1;
 		
-		if (RayMarchMaskTexture.Load(uint3(PixelCoord, 0)).r > 0.0f &&  Categorization.PixelCoverage < CoverageThreshold)
+		if (bIsValid && RayMarchMaskTexture.Load(uint3(PixelCoord, 0)).r > 0.0f &&  Categorization.PixelCoverage < CoverageThreshold)
 		{
 			Result = ComputeHairCountVirtualVoxel(
 				WorldPosition,
diff --git a/Engine/Shaders/Private/HairStrands/HairStrandsInterpolation.usf b/Engine/Shaders/Private/HairStrands/HairStrandsInterpolation.usf
--- a/Engine/Shaders/Private/HairStrands/HairStrandsInterpolation.usf
+++ b/Engine/Shaders/Private/HairStrands/HairStrandsInterpolation.usf
@@ -549,6 +549,9 @@ void MainCS(uint2 DispatchThreadId : SV_DispatchThreadID)
 #error Unknown group size	
 #endif
 
+float3 CPUBoundMin;
+float3 CPUBoundMax;
+
 uint DispatchCountX;
 uint ClusterCount;
 //float3 OutHairPositionOffset;
@@ -576,6 +579,16 @@ groupshared float3 SharedClusterAABBMax[GROUP_SIZE];
 [numthreads(GROUP_SIZE, 1, 1)]
 void ClusterAABBEvaluationCS(uint3 DispatchThreadId : SV_DispatchThreadID, uint3 GroupId : SV_GroupID, uint3 GroupThreadId : SV_GroupThreadID)
 {
+#if PERMUTATION_USE_CPU_AABB
+	// Only update the full group AABB
+	if (all(DispatchThreadId == 0))
+	{
+		FHairAABB Bound;
+		Bound.Min = CPUBoundMin;
+		Bound.Max = CPUBoundMax;
+		WriteHairAABB(0, Bound, OutGroupAABBBuffer);
+	}
+#else // PERMUTATION_USE_CPU_AABB
 	const uint ClusterFetchIndex = GroupId.x + GroupId.y * DispatchCountX;
 	if (ClusterFetchIndex >= ClusterCount)
 	{
@@ -683,6 +696,7 @@ void ClusterAABBEvaluationCS(uint3 DispatchThreadId : SV_DispatchThreadID, uint3
 	}
 
 	// Min/Max group AABB is done later by another pass
+#endif // PERMUTATION_USE_CPU_AABB
 }
 
 #endif // SHADER_CLUSTERAABB
diff --git a/Engine/Shaders/Private/HairStrands/HairStrandsVoxelPageAllocation.usf b/Engine/Shaders/Private/HairStrands/HairStrandsVoxelPageAllocation.usf
--- a/Engine/Shaders/Private/HairStrands/HairStrandsVoxelPageAllocation.usf
+++ b/Engine/Shaders/Private/HairStrands/HairStrandsVoxelPageAllocation.usf
@@ -35,7 +35,7 @@ void AllocatePageIndex(uint2 DispatchThreadId : SV_DispatchThreadID)
 	const uint MacroGroupId = DispatchThreadId.x;
 
 	FHairAABB Bound = InitHairAABB();
-	const bool bIsValid = MacroGroupId < MacroGroupCount;
+	bool bIsValid = MacroGroupId < MacroGroupCount;
 	if (bIsValid)
 	{
 		Bound = ReadHairAABB(MacroGroupId, MacroGroupAABBBuffer);
@@ -43,6 +43,7 @@ void AllocatePageIndex(uint2 DispatchThreadId : SV_DispatchThreadID)
 		{
 			Bound.Min = 0;
 			Bound.Max = 0;
+			bIsValid = false;
 		}
 	}
 
@@ -228,7 +229,7 @@ void UpdatePackScatterCoord(inout uint2 Out, uint PartIndex)
 #define SCATTER_GROUP_SIZE 32
 #define SCATTER_GROUP_LINE 4
 
-#endif
+#endif // SHADER_MARKVALID_PREPARE || SHADER_MARKVALID_INDIRECTARG || SHADER_MARKVALID_SCATTER
 
 #if SHADER_MARKVALID_PREPARE
 uint			MaxClusterCount;
@@ -236,6 +237,7 @@ uint			MacroGroupId;
 uint			MaxScatterAllocationCount;
 uint			bForceDirectPageAllocation;
 
+Buffer<int>		GroupAABBsBuffer;
 Buffer<int>		ClusterAABBsBuffer;
 Buffer<int>		MacroGroupAABBBuffer;
 Buffer<uint4>	PageIndexResolutionAndOffsetBuffer;
@@ -248,6 +250,7 @@ RWBuffer<uint2> OutDeferredScatterBuffer;
 // PageIndexBuffer      is sampled with linear coordinate computed from the 3d page coordinate. VALID NODE ARE NOT COMPACTED. It contains the LINEAR PAGE INDEX (to map to the 3d volume).
 // PageIndexCoordBuffer is sampled with linear coordinate for allocated nodes. VALID NODE ARE COMPACTED. It contains the 3d page coordinate and ClustedId. Only used for opaque voxel injection.
 
+#if PERMUTATION_USE_CLUSTER
 [numthreads(SCATTER_GROUP_SIZE, 1, 1)]
 void MarkValid_PrepareCS(uint2 DispatchThreadId : SV_DispatchThreadID)
 {
@@ -339,7 +342,45 @@ void MarkValid_PrepareCS(uint2 DispatchThreadId : SV_DispatchThreadID)
 		}
 	}
 }
-#endif
+#else // PERMUTATION_USE_CLUSTER
+[numthreads(8, 8, 8)]
+void MarkValid_PrepareCS(uint3 DispatchThreadId : SV_DispatchThreadID)
+{
+	const uint4 PageIndexResolutionAndOffset = PageIndexResolutionAndOffsetBuffer.Load(MacroGroupId);
+	const FHairAABB MacroGroupBound = ReadHairAABB(MacroGroupId, MacroGroupAABBBuffer);
+	const int3 PageIndexResolution = PageIndexResolutionAndOffset.xyz;
+	const uint PageIndexOffset = PageIndexResolutionAndOffset.w;
+
+	const uint3 Coord = DispatchThreadId;
+	if (any(Coord >= PageIndexResolution))
+		return;
+
+	const FHairAABB GroupBound = ReadHairAABB(0, GroupAABBsBuffer);
+	const uint3 MinCoord = PositionToCoord(GroupBound.Min, MacroGroupBound.Min, MacroGroupBound.Max, PageIndexResolution);
+	const uint3 MaxCoord = PositionToCoord(GroupBound.Max, MacroGroupBound.Min, MacroGroupBound.Max, PageIndexResolution);
+
+	if (any(!IsFinite(float3(MinCoord))))
+		return;
+
+	if (any(!IsFinite(float3(MaxCoord))))
+		return;
+
+	if (all(Coord >= MinCoord) && all(Coord <=MaxCoord))
+	{
+		const uint3 PageIndexCoord = uint3(Coord.x, Coord.y, Coord.z);
+		const uint LinearPageIndexCoord = CoordToIndex(PageIndexCoord, PageIndexResolution, PageIndexOffset);
+		InterlockedOr(OutValidPageIndexBuffer[LinearPageIndexCoord], 1u);
+	}
+
+	// Fake writting to insure RDG allocate buffers even there is only Group-based allocated elements (vs. cluster-based allocated elements)
+	if (Coord.x > 65000)
+	{
+		OutDeferredScatterCounter[0] = 0;
+		OutDeferredScatterBuffer[0] = uint2(0,0);
+	}
+}
+#endif // PERMUTATION_USE_CLUSTER
+#endif // SHADER_MARKVALID_PREPARE
 
 #if SHADER_MARKVALID_INDIRECTARG
 Buffer<uint> DeferredScatterCounter;
diff --git a/Engine/Shaders/Private/HairStrands/HairStrandsVoxelPageRayMarching.usf b/Engine/Shaders/Private/HairStrands/HairStrandsVoxelPageRayMarching.usf
--- a/Engine/Shaders/Private/HairStrands/HairStrandsVoxelPageRayMarching.usf
+++ b/Engine/Shaders/Private/HairStrands/HairStrandsVoxelPageRayMarching.usf
@@ -69,6 +69,7 @@ void MainCS(uint3 DispatchThreadId : SV_DispatchThreadID)
 
 	const FPackedVirtualVoxelNodeDesc PackedNode = VirtualVoxel.NodeDescBuffer[MacroGroupId];
 	const FVirtualVoxelNodeDesc NodeDesc = UnpackVoxelNode(PackedNode, VirtualVoxel.PageResolution);
+	const bool bIsNodeValid = all(NodeDesc.PageIndexResolution != 0);
 
 	// Draw macro group AABB & debug infos
 	if (uint(PixelCoord.x) == 0 && uint(PixelCoord.y) == 0)
@@ -582,33 +583,33 @@ void MainCS(uint3 DispatchThreadId : SV_DispatchThreadID)
 		CommonDesc.PageTextureResolution= VirtualVoxel.PageTextureResolution;
 		CommonDesc.PageResolution		= VirtualVoxel.PageResolution;
 
-		const FPackedVirtualVoxelNodeDesc PackedNode = VirtualVoxel.NodeDescBuffer[MacroGroupId];
-		const FVirtualVoxelNodeDesc NodeDesc = UnpackVoxelNode(PackedNode, VirtualVoxel.PageResolution);
-
-		FHairTraversalSettings TraversalSettings = InitHairTraversalSettings();
-		TraversalSettings.DensityScale		= VirtualVoxel.DensityScale;
-		TraversalSettings.CountThreshold	= GetOpaqueVoxelValue();
-		TraversalSettings.DistanceThreshold = DistanceThreshold;
-		TraversalSettings.bDebugEnabled		= bDebugEnabled;
-		TraversalSettings.SteppingScale		= VirtualVoxel.SteppingScale_Shadow;
-		TraversalSettings.Random			= SampleRandom;
-		TraversalSettings.TanConeAngle		= 0;
-		TraversalSettings.bIsPrimaryRay		= true;
-		TraversalSettings.PixelRadius		= -1;
-		TraversalSettings.ForcedMip			= ForcedMipLevel;
-
-		TraversalResult = ComputeHairCountVirtualVoxel(
-			WP0,
-			WP1,
-			CommonDesc,
-			NodeDesc,
-			VirtualVoxel.PageIndexBuffer,
-			VirtualVoxel.PageIndexOccupancyBuffer,
-			VirtualVoxel.PageTexture,
-			TraversalSettings);
-
-		HitColor = GetColorCode(saturate(TraversalResult.HairCount / 16));
-		bIsValid = TraversalResult.HairCount > 0;
+		if (bIsNodeValid)
+		{
+			FHairTraversalSettings TraversalSettings = InitHairTraversalSettings();
+			TraversalSettings.DensityScale		= VirtualVoxel.DensityScale;
+			TraversalSettings.CountThreshold	= GetOpaqueVoxelValue();
+			TraversalSettings.DistanceThreshold = DistanceThreshold;
+			TraversalSettings.bDebugEnabled		= bDebugEnabled;
+			TraversalSettings.SteppingScale		= VirtualVoxel.SteppingScale_Shadow;
+			TraversalSettings.Random			= SampleRandom;
+			TraversalSettings.TanConeAngle		= 0;
+			TraversalSettings.bIsPrimaryRay		= true;
+			TraversalSettings.PixelRadius		= -1;
+			TraversalSettings.ForcedMip			= ForcedMipLevel;
+
+			TraversalResult = ComputeHairCountVirtualVoxel(
+				WP0,
+				WP1,
+				CommonDesc,
+				NodeDesc,
+				VirtualVoxel.PageIndexBuffer,
+				VirtualVoxel.PageIndexOccupancyBuffer,
+				VirtualVoxel.PageTexture,
+				TraversalSettings);
+
+			HitColor = GetColorCode(saturate(TraversalResult.HairCount / 16));
+			bIsValid = TraversalResult.HairCount > 0;
+		}
 
 		const bool bIsPixelEnable = all(int2(PixelCoord) == GetCursorPos());
 		if (bIsPixelEnable && bIsValid)
diff --git a/Engine/Source/Runtime/Renderer/Private/HairStrands/HairStrandsVoxelization.cpp b/Engine/Source/Runtime/Renderer/Private/HairStrands/HairStrandsVoxelization.cpp
--- a/Engine/Source/Runtime/Renderer/Private/HairStrands/HairStrandsVoxelization.cpp
+++ b/Engine/Source/Runtime/Renderer/Private/HairStrands/HairStrandsVoxelization.cpp
@@ -233,12 +233,16 @@ class FVoxelMarkValidPageIndex_PrepareCS : public FGlobalShader
 	DECLARE_GLOBAL_SHADER(FVoxelMarkValidPageIndex_PrepareCS);
 	SHADER_USE_PARAMETER_STRUCT(FVoxelMarkValidPageIndex_PrepareCS, FGlobalShader);
 
+	class FUseCluster : SHADER_PERMUTATION_BOOL("PERMUTATION_USE_CLUSTER");
+	using FPermutationDomain = TShaderPermutationDomain<FUseCluster>;
+
 	BEGIN_SHADER_PARAMETER_STRUCT(FParameters, )
 		SHADER_PARAMETER(uint32, MaxClusterCount)
 		SHADER_PARAMETER(uint32, MacroGroupId)
 		SHADER_PARAMETER(uint32, MaxScatterAllocationCount)
 		SHADER_PARAMETER(uint32, bForceDirectPageAllocation)
 
+		SHADER_PARAMETER_RDG_BUFFER_SRV(Buffer, GroupAABBsBuffer)
 		SHADER_PARAMETER_RDG_BUFFER_SRV(Buffer, ClusterAABBsBuffer)
 		SHADER_PARAMETER_RDG_BUFFER_SRV(Buffer, MacroGroupAABBBuffer)
 		SHADER_PARAMETER_RDG_BUFFER_SRV(Buffer, PageIndexResolutionAndOffsetBuffer)
@@ -635,66 +639,88 @@ static void AddAllocateVoxelPagesPass(
 
 			AddClearUAVPass(GraphBuilder, ScatterCounterUAV, 0);
 
+			const uint32 bForceDirectPageAllocation = GHairVirtualVoxelUseImmediatePageAllocation > 0 ? 1 : 0;
 			// Prepare
 			for (const FHairStrandsMacroGroupData::PrimitiveInfo& PrimitiveInfo : MacroGroup.PrimitivesInfos)
 			{
 				FHairGroupPublicData* HairGroupData = PrimitiveInfo.PublicDataPtr;
 
+				if (!HairGroupData->DoesSupportVoxelization())
+					continue;
+
+				// For allocating/marking the used pages, a primitive/group can:
+				// * Either use its group AABB (for static groom)
+				// * Or use its clusters AABBs (for dynamic groom)
+				// Even if the bGroupAABBValid is invalid, it has been reset to invalid AABB, so the code below won't mark any page as used.
+				const bool bUseClusterAABB = HairGroupData->bClusterAABBValid;
+				
 				FVoxelMarkValidPageIndex_PrepareCS::FParameters* Parameters = GraphBuilder.AllocParameters<FVoxelMarkValidPageIndex_PrepareCS::FParameters>();
 				Parameters->MaxClusterCount				= HairGroupData->GetClusterCount();
 				Parameters->MacroGroupId				= MacroGroup.MacroGroupId;
 				Parameters->MaxScatterAllocationCount	= MaxAllocationCount;
-				Parameters->bForceDirectPageAllocation	= GHairVirtualVoxelUseImmediatePageAllocation > 0 ? 1 : 0;
+				Parameters->bForceDirectPageAllocation	= bForceDirectPageAllocation;
 
-				Parameters->ClusterAABBsBuffer		= RegisterAsSRV(GraphBuilder, HairGroupData->GetClusterAABBBuffer());
-				Parameters->MacroGroupAABBBuffer	= GraphBuilder.CreateSRV(MacroGroupResources.MacroGroupAABBsBuffer, PF_R32_SINT);
+				Parameters->GroupAABBsBuffer			= RegisterAsSRV(GraphBuilder, HairGroupData->GetGroupAABBBuffer());
+				Parameters->ClusterAABBsBuffer			= RegisterAsSRV(GraphBuilder, HairGroupData->GetClusterAABBBuffer());
+				Parameters->MacroGroupAABBBuffer		= GraphBuilder.CreateSRV(MacroGroupResources.MacroGroupAABBsBuffer, PF_R32_SINT);
 				Parameters->PageIndexResolutionAndOffsetBuffer = PageIndexResolutionAndOffsetBufferSRV;
 
 				Parameters->OutDeferredScatterCounter	= ScatterCounterUAV;
 				Parameters->OutDeferredScatterBuffer	= ScatterBufferUAV;
 				Parameters->OutValidPageIndexBuffer		= PageIndexBufferUAV;
 
-				FIntVector DispatchCount((Parameters->MaxClusterCount + GroupSize - 1) / GroupSize, 1, 1);
+				FIntVector DispatchCount = bUseClusterAABB ?
+					FIntVector((Parameters->MaxClusterCount + GroupSize - 1) / GroupSize, 1, 1) :
+					FIntVector(
+						FMath::DivideAndRoundUp(MacroGroup.VirtualVoxelNodeDesc.PageIndexResolution.X, 8),
+						FMath::DivideAndRoundUp(MacroGroup.VirtualVoxelNodeDesc.PageIndexResolution.Y, 8),
+						FMath::DivideAndRoundUp(MacroGroup.VirtualVoxelNodeDesc.PageIndexResolution.Z, 8));
+
 				check(DispatchCount.X < 65535);
-				TShaderMapRef<FVoxelMarkValidPageIndex_PrepareCS> ComputeShader(View.ShaderMap);
+				FVoxelMarkValidPageIndex_PrepareCS::FPermutationDomain PermutationVector;
+				PermutationVector.Set<FVoxelMarkValidPageIndex_PrepareCS::FUseCluster>(bUseClusterAABB ? 1 : 0);
+				TShaderMapRef<FVoxelMarkValidPageIndex_PrepareCS> ComputeShader(View.ShaderMap, PermutationVector);
 				FComputeShaderUtils::AddPass(
 					GraphBuilder,
-					RDG_EVENT_NAME("HairStrandsMarkValidPageIndex_Prepare"),
+					RDG_EVENT_NAME("HairStrandsMarkValidPageIndex_Prepare(%s,%s)", bUseClusterAABB ? TEXT("Cluster") : TEXT("Group"), bForceDirectPageAllocation ? TEXT("Direct") : TEXT("Indirect")),
 					ComputeShader,
 					Parameters,
 					DispatchCount);
 			}
-			
-			FRDGBufferSRVRef ScatterCounterSRV = GraphBuilder.CreateSRV(ScatterCounter, PF_R32_UINT);
 
-			// Build indirect buffer args
-			FRDGBufferRef ScatterIndirectArgsBuffer = GraphBuilder.CreateBuffer(FRDGBufferDesc::CreateIndirectDesc<FRHIDispatchIndirectParameters>(1), TEXT("Hair.PageScatterIndirectArgs"));
+			if (!bForceDirectPageAllocation)
 			{
-				check(MacroGroup.MacroGroupId < MacroGroupCount);
+				FRDGBufferSRVRef ScatterCounterSRV = GraphBuilder.CreateSRV(ScatterCounter, PF_R32_UINT);
 
-				FVoxelMarkValidPageIndex_IndirectArgsCS::FParameters* Parameters = GraphBuilder.AllocParameters<FVoxelMarkValidPageIndex_IndirectArgsCS::FParameters>();
-				Parameters->DeferredScatterCounter = ScatterCounterSRV;
-				Parameters->OutIndirectArgsBuffer = GraphBuilder.CreateUAV(ScatterIndirectArgsBuffer);
+				// Build indirect buffer args
+				FRDGBufferRef ScatterIndirectArgsBuffer = GraphBuilder.CreateBuffer(FRDGBufferDesc::CreateIndirectDesc<FRHIDispatchIndirectParameters>(1), TEXT("Hair.PageScatterIndirectArgs"));
+				{
+					check(MacroGroup.MacroGroupId < MacroGroupCount);
 
-				const FIntVector DispatchCount(1, 1, 1);
-				TShaderMapRef<FVoxelMarkValidPageIndex_IndirectArgsCS> ComputeShader(View.ShaderMap);
-				FComputeShaderUtils::AddPass(GraphBuilder, RDG_EVENT_NAME("HairStrandsMarkValidPageIndex_IndirectArgs"), ComputeShader, Parameters, DispatchCount);
-			}
+					FVoxelMarkValidPageIndex_IndirectArgsCS::FParameters* Parameters = GraphBuilder.AllocParameters<FVoxelMarkValidPageIndex_IndirectArgsCS::FParameters>();
+					Parameters->DeferredScatterCounter = ScatterCounterSRV;
+					Parameters->OutIndirectArgsBuffer = GraphBuilder.CreateUAV(ScatterIndirectArgsBuffer);
 
-			// Scatter
-			{
-				FRDGBufferSRVRef ScatterBufferSRV = GraphBuilder.CreateSRV(ScatterBuffer, PF_R32G32_UINT);
-				check(MacroGroup.MacroGroupId < MacroGroupCount);
+					const FIntVector DispatchCount(1, 1, 1);
+					TShaderMapRef<FVoxelMarkValidPageIndex_IndirectArgsCS> ComputeShader(View.ShaderMap);
+					FComputeShaderUtils::AddPass(GraphBuilder, RDG_EVENT_NAME("HairStrandsMarkValidPageIndex_IndirectArgs"), ComputeShader, Parameters, DispatchCount);
+				}
 
-				FVoxelMarkValidPageIndex_ScatterCS::FParameters* Parameters = GraphBuilder.AllocParameters<FVoxelMarkValidPageIndex_ScatterCS::FParameters>();
-				Parameters->IndirectBufferArgs = ScatterIndirectArgsBuffer;
-				Parameters->PageIndexResolutionAndOffsetBuffer = PageIndexResolutionAndOffsetBufferSRV;
-				Parameters->DeferredScatterCounter = ScatterCounterSRV;
-				Parameters->DeferredScatterBuffer = ScatterBufferSRV;
-				Parameters->OutValidPageIndexBuffer = PageIndexBufferUAV;
+				// Scatter
+				{
+					FRDGBufferSRVRef ScatterBufferSRV = GraphBuilder.CreateSRV(ScatterBuffer, PF_R32G32_UINT);
+					check(MacroGroup.MacroGroupId < MacroGroupCount);
 
-				TShaderMapRef<FVoxelMarkValidPageIndex_ScatterCS> ComputeShader(View.ShaderMap);
-				FComputeShaderUtils::AddPass(GraphBuilder, RDG_EVENT_NAME("HairStrandsMarkValidPageIndex_Scatter"), ComputeShader, Parameters, ScatterIndirectArgsBuffer, 0);
+					FVoxelMarkValidPageIndex_ScatterCS::FParameters* Parameters = GraphBuilder.AllocParameters<FVoxelMarkValidPageIndex_ScatterCS::FParameters>();
+					Parameters->IndirectBufferArgs = ScatterIndirectArgsBuffer;
+					Parameters->PageIndexResolutionAndOffsetBuffer = PageIndexResolutionAndOffsetBufferSRV;
+					Parameters->DeferredScatterCounter = ScatterCounterSRV;
+					Parameters->DeferredScatterBuffer = ScatterBufferSRV;
+					Parameters->OutValidPageIndexBuffer = PageIndexBufferUAV;
+
+					TShaderMapRef<FVoxelMarkValidPageIndex_ScatterCS> ComputeShader(View.ShaderMap);
+					FComputeShaderUtils::AddPass(GraphBuilder, RDG_EVENT_NAME("HairStrandsMarkValidPageIndex_Scatter"), ComputeShader, Parameters, ScatterIndirectArgsBuffer, 0);
+				}
 			}
 		}
 		else
@@ -702,6 +728,8 @@ static void AddAllocateVoxelPagesPass(
 			for (const FHairStrandsMacroGroupData::PrimitiveInfo& PrimitiveInfo : MacroGroup.PrimitivesInfos)
 			{
 				FHairGroupPublicData* HairGroupData = PrimitiveInfo.PublicDataPtr;
+				if (!HairGroupData->DoesSupportVoxelization())
+					continue;
 
 				FVoxelMarkValidPageIndexCS::FParameters* Parameters = GraphBuilder.AllocParameters<FVoxelMarkValidPageIndexCS::FParameters>();
 				Parameters->MacroGroupId = MacroGroup.MacroGroupId;
diff --git a/Engine/Source/Runtime/Renderer/Public/HairStrandsInterface.h b/Engine/Source/Runtime/Renderer/Public/HairStrandsInterface.h
--- a/Engine/Source/Runtime/Renderer/Public/HairStrandsInterface.h
+++ b/Engine/Source/Runtime/Renderer/Public/HairStrandsInterface.h
@@ -291,6 +291,8 @@ public:
 	/* Hair Cluster & Hair Group bounding box buffer */
 	FRDGExternalBuffer ClusterAABBBuffer;
 	FRDGExternalBuffer GroupAABBBuffer;
+	bool bGroupAABBValid = false;
+	bool bClusterAABBValid = false;
 
 	/* Culling & LODing results for a hair group */ // Better to be transient?
 	FRDGExternalBuffer CulledVertexIdBuffer;
