Fix use after free value when initializing groom (rest) offset position buffer.

diff --git a/Engine/Plugins/Runtime/HairStrands/Source/HairStrandsCore/Private/GroomResources.cpp b/Engine/Plugins/Runtime/HairStrands/Source/HairStrandsCore/Private/GroomResources.cpp
--- a/Engine/Plugins/Runtime/HairStrands/Source/HairStrandsCore/Private/GroomResources.cpp
+++ b/Engine/Plugins/Runtime/HairStrands/Source/HairStrandsCore/Private/GroomResources.cpp
@@ -116,7 +116,7 @@ static FRDGBufferRef InternalCreateVertexBuffer(
 }
 
 template<typename FormatType>
-void InternalCreateVertexBufferRDG(FRDGBuilder& GraphBuilder, const TArray<typename FormatType::Type>& InData, FRDGExternalBuffer& Out, const TCHAR* DebugName, EHairResourceUsageType UsageType)
+void InternalCreateVertexBufferRDG(FRDGBuilder& GraphBuilder, const TArray<typename FormatType::Type>& InData, FRDGExternalBuffer& Out, const TCHAR* DebugName, EHairResourceUsageType UsageType, ERDGInitialDataFlags InitialDataFlags= ERDGInitialDataFlags::NoCopy)
 {
 	FRDGBufferRef Buffer = nullptr;
 
@@ -139,7 +139,7 @@ void InternalCreateVertexBufferRDG(FRDGBuilder& GraphBuilder, const TArray<typen
 		Desc,
 		InData.GetData(),
 		DataSizeInBytes,
-		ERDGInitialDataFlags::NoCopy);
+		InitialDataFlags);
 
 	ConvertToExternalBufferWithViews(GraphBuilder, Buffer, Out, FormatType::Format);
 }
@@ -560,7 +560,7 @@ void FHairStrandsRestResource::InitRHI()
 
 	TArray<FVector4> RestOffset;
 	RestOffset.Add(PositionOffset);
-	InternalCreateVertexBufferRDG<FHairStrandsPositionOffsetFormat>(GraphBuilder, RestOffset, PositionOffsetBuffer, HAIRSTRANDS_RESOUCE_NAME(CurveType, Hair.StrandsRest_PositionOffsetBuffer), EHairResourceUsageType::Static);
+	InternalCreateVertexBufferRDG<FHairStrandsPositionOffsetFormat>(GraphBuilder, RestOffset, PositionOffsetBuffer, HAIRSTRANDS_RESOUCE_NAME(CurveType, Hair.StrandsRest_PositionOffsetBuffer), EHairResourceUsageType::Static, ERDGInitialDataFlags::None);
 	
 	GraphBuilder.Execute();
 }
