Deferred groom component resources allocation to first usage in order to reduce memory footprint and initialization time.

diff --git a/Engine/Plugins/Runtime/HairStrands/Source/HairStrandsCore/Private/GroomComponent.cpp b/Engine/Plugins/Runtime/HairStrands/Source/HairStrandsCore/Private/GroomComponent.cpp
--- a/Engine/Plugins/Runtime/HairStrands/Source/HairStrandsCore/Private/GroomComponent.cpp
+++ b/Engine/Plugins/Runtime/HairStrands/Source/HairStrandsCore/Private/GroomComponent.cpp
@@ -40,7 +40,7 @@ static FAutoConsoleVariableRef CVarHairBindingValidationEnable(TEXT("r.HairStran
 
 ////////////////////////////////////////////////////////////////////////////////////////////////////////////////////////
 
-EHairStrandsDebugMode GetHairStrandsGeometryDebugMode(FHairGroupInstance* Instance);
+EHairStrandsDebugMode GetHairStrandsGeometryDebugMode(const FHairGroupInstance* Instance);
 
 const FLinearColor GetHairGroupDebugColor(int32 GroupIt)
 {
@@ -257,154 +257,6 @@ inline int32 GetMaterialIndexWithFallback(int32 SlotIndex)
 	return SlotIndex != INDEX_NONE ? SlotIndex : 0;
 }
 
-struct FHairCardsVertexFactories
-{
-	struct FInstanceFactory
-	{
-		TArray<FHairCardsVertexFactory> VertexFactories0;
-		TArray<FHairCardsVertexFactory> VertexFactories1;
-		TArray<int32> LODMapping;
-	};
-	TArray<FInstanceFactory> Instances;
-
-	FHairCardsVertexFactories()
-	{
-	}
-
-	~FHairCardsVertexFactories()
-	{
-		BeginReleaseResource();
-	}
-
-	void AddInstances(TArray<FHairGroupInstance*> In, EShaderPlatform InShaderPlatform, ERHIFeatureLevel::Type InFeatureLevel)
-	{
-		FeatureLevel = InFeatureLevel;
-		ShaderPlatform = InShaderPlatform;
-		Instances.Reserve(Instances.Num());
-		for (uint32 InstanceIt = 0, InstanceCount = In.Num(); InstanceIt < InstanceCount; ++InstanceIt)
-		{
-			FInstanceFactory& Out = Instances.AddDefaulted_GetRef();
-			InternalAddInstance(Out, In[InstanceIt], InstanceIt);
-		}
-	}
-
-	void InternalAddInstance(FInstanceFactory& Out, FHairGroupInstance* Instance, uint32 GroupIndex)
-	{
-		const uint32 LODCount = Instance->Cards.LODs.Num();
-
-		Out.VertexFactories0.Empty(LODCount);
-		Out.VertexFactories1.Empty(LODCount);
-		for (uint32 LODIt = 0; LODIt < LODCount; ++LODIt)
-		{
-			int32 LODIndirectIndex = -1;
-			if (Instance->Cards.IsValid(LODIt))
-			{
-				LODIndirectIndex = Out.VertexFactories0.Num();
-				new (Out.VertexFactories0) FHairCardsVertexFactory(Instance, GroupIndex, LODIt, 0u, EHairGeometryType::Cards, ShaderPlatform, FeatureLevel, "HairCardsVertexFactory");
-				new (Out.VertexFactories1) FHairCardsVertexFactory(Instance, GroupIndex, LODIt, 1u, EHairGeometryType::Cards, ShaderPlatform, FeatureLevel, "HairCardsVertexFactory");
-
-			}
-			else if (Instance->Meshes.IsValid(LODIt))
-			{
-				LODIndirectIndex = Out.VertexFactories0.Num();
-				new (Out.VertexFactories0) FHairCardsVertexFactory(Instance, GroupIndex, LODIt, 0u, EHairGeometryType::Meshes, ShaderPlatform, FeatureLevel, "HairMeshesVertexFactory");
-				new (Out.VertexFactories1) FHairCardsVertexFactory(Instance, GroupIndex, LODIt, 1u, EHairGeometryType::Meshes, ShaderPlatform, FeatureLevel, "HairMeshesVertexFactory");
-
-			}
-			Out.LODMapping.Add(LODIndirectIndex);
-		}
-	}
-
-	const FHairCardsVertexFactory* GetVertexFactory(uint32 GroupIndex, uint32 LODIndex, bool bIsCurrent0, EHairGeometryType GeometryType) const
-	{
-		const FHairCardsVertexFactory* Out = nullptr;
-		if (bIsCurrent0)
-		{
-			if (GroupIndex < uint32(Instances.Num()) && LODIndex < uint32(Instances[GroupIndex].LODMapping.Num()))
-			{
-				const int32 IndirectLODIndex = Instances[GroupIndex].LODMapping[LODIndex];
-				if (IndirectLODIndex != -1)
-				{
-					Out = &Instances[GroupIndex].VertexFactories0[IndirectLODIndex];
-				}
-			}
-		}
-		else
-		{
-			if (GroupIndex < uint32(Instances.Num()) && LODIndex < uint32(Instances[GroupIndex].LODMapping.Num()))
-			{
-				const int32 IndirectLODIndex = Instances[GroupIndex].LODMapping[LODIndex];
-				if (IndirectLODIndex != -1)
-				{
-					Out = &Instances[GroupIndex].VertexFactories1[IndirectLODIndex];
-				}
-			}
-		}
-
-		// Sanity check
-		if (Out)
-		{
-			check(Out->Data.GeometryType == GeometryType);
-		}
-
-		return Out;
-	}
-
-	void BeginInitResource()
-	{
-		for (FInstanceFactory& Instance : Instances)
-		{
-			for (FHairCardsVertexFactory& VertexFactory : Instance.VertexFactories0)
-			{
-				::BeginInitResource(&VertexFactory);
-			}
-			for (FHairCardsVertexFactory& VertexFactory : Instance.VertexFactories1)
-			{
-				::BeginInitResource(&VertexFactory);
-			}
-		}
-	}
-
-	void BeginReleaseResource()
-	{
-		for (FInstanceFactory& Instance : Instances)
-		{
-			for (FHairCardsVertexFactory& VertexFactory : Instance.VertexFactories0)
-			{
-				::BeginReleaseResource(&VertexFactory);
-			}
-			Instance.VertexFactories0.Empty();
-			for (FHairCardsVertexFactory& VertexFactory : Instance.VertexFactories1)
-			{
-				::BeginReleaseResource(&VertexFactory);
-			}
-			Instance.VertexFactories1.Empty();
-		}
-		Instances.Empty();
-	}
-
-	void ReleaseResources()
-	{
-		for (FInstanceFactory& Instance : Instances)
-		{
-			for (FHairCardsVertexFactory& VertexFactory : Instance.VertexFactories0)
-			{
-				VertexFactory.ReleaseResource();
-			}
-			Instance.VertexFactories0.Empty();
-			for (FHairCardsVertexFactory& VertexFactory : Instance.VertexFactories1)
-			{
-				VertexFactory.ReleaseResource();
-			}
-			Instance.VertexFactories1.Empty();
-		}
-		Instances.Empty();
-	}
-
-	ERHIFeatureLevel::Type FeatureLevel = ERHIFeatureLevel::Num;
-	EShaderPlatform ShaderPlatform = EShaderPlatform::SP_NumPlatforms;
-};
-
 static EHairGeometryType ToHairGeometryType(EGroomGeometryType Type)
 {
 	switch (Type)
@@ -440,7 +292,6 @@ public:
 
 	FHairStrandsSceneProxy(UGroomComponent* Component)
 		: FPrimitiveSceneProxy(Component)
-		, StrandsVertexFactory(Component->HairGroupInstances, GetScene().GetFeatureLevel(), "FStrandsHairSceneProxy")
 		, MaterialRelevance(Component->GetMaterialRelevance(GetScene().GetFeatureLevel()))
 	{
 		// Forcing primitive uniform as we don't support robustly GPU scene data
@@ -459,15 +310,10 @@ public:
 			bAlwaysHasVelocity = true;
 		}
 
-		FHairStrandsVertexFactory::FDataType StrandsVFData;
-		StrandsVFData.Instances = Component->HairGroupInstances;
-		CardsAndMeshesVertexFactories.AddInstances(Component->HairGroupInstances, Component->GetScene()->GetShaderPlatform(), Component->GetScene()->GetFeatureLevel());
-
 		check(Component->HairGroupInstances.Num());
 
-		GroomLODSelection = Component->GroomAsset->LODSelectionType;
-
-		const EShaderPlatform Platform = GetScene().GetShaderPlatform();
+		const ERHIFeatureLevel::Type FeatureLevel = GetScene().GetFeatureLevel();
+		const EShaderPlatform ShaderPlatform = GetScene().GetShaderPlatform();
 
 		const int32 GroupCount = Component->GroomAsset->GetNumHairGroups();
 		check(Component->GroomAsset->HairGroupsData.Num() == Component->HairGroupInstances.Num());
@@ -484,6 +330,40 @@ public:
 			HairInstance->bForceCards = Component->bUseCards;
 			HairInstance->bUpdatePositionOffset = Component->RegisteredMeshComponent != nullptr;
 			HairInstance->bCastShadow = Component->CastShadow;
+
+			if (HairInstance->Strands.IsValid() && HairInstance->Strands.VertexFactory == nullptr)
+			{
+				HairInstance->Strands.VertexFactory = new FHairStrandsVertexFactory(HairInstance, FeatureLevel, "FStrandsHairSceneProxy");
+			}
+
+			if (HairInstance->Cards.IsValid())
+			{
+				const uint32 LODCount = HairInstance->Cards.LODs.Num();
+				for (uint32 LODIt = 0; LODIt < LODCount; ++LODIt)
+				{
+					if (HairInstance->Cards.IsValid(LODIt))
+					{
+						if (HairInstance->Cards.LODs[LODIt].VertexFactory[0] == nullptr) HairInstance->Cards.LODs[LODIt].VertexFactory[0] = new FHairCardsVertexFactory(HairInstance, GroupIt, LODIt, 0u, EHairGeometryType::Cards, ShaderPlatform, FeatureLevel, "HairCardsVertexFactory");
+						if (HairInstance->Cards.LODs[LODIt].VertexFactory[1] == nullptr) HairInstance->Cards.LODs[LODIt].VertexFactory[1] = new FHairCardsVertexFactory(HairInstance, GroupIt, LODIt, 1u, EHairGeometryType::Cards, ShaderPlatform, FeatureLevel, "HairCardsVertexFactory");
+					}
+				}
+			}
+
+			if (HairInstance->Meshes.IsValid())
+			{
+				const uint32 LODCount = HairInstance->Meshes.LODs.Num();
+
+				for (uint32 LODIt = 0; LODIt < LODCount; ++LODIt)
+				{
+					if (HairInstance->Meshes.IsValid(LODIt))
+					{
+						if (HairInstance->Meshes.LODs[LODIt].VertexFactory[0] == nullptr) HairInstance->Meshes.LODs[LODIt].VertexFactory[0] = new FHairCardsVertexFactory(HairInstance, GroupIt, LODIt, 0u, EHairGeometryType::Meshes, ShaderPlatform, FeatureLevel, "HairMeshesVertexFactory");
+						if (HairInstance->Meshes.LODs[LODIt].VertexFactory[1] == nullptr) HairInstance->Meshes.LODs[LODIt].VertexFactory[1] = new FHairCardsVertexFactory(HairInstance, GroupIt, LODIt, 1u, EHairGeometryType::Meshes, ShaderPlatform, FeatureLevel, "HairMeshesVertexFactory");
+
+					}
+				}
+			}
+
 			{
 				FHairGroup& OutGroupData = HairGroups.AddDefaulted_GetRef();
 				OutGroupData.bIsVsibible = bIsVisible;
@@ -532,14 +412,14 @@ public:
 			}
 
 			// Material - Strands
-			if (IsHairStrandsEnabled(EHairStrandsShaderType::Strands, Platform))
+			if (IsHairStrandsEnabled(EHairStrandsShaderType::Strands, ShaderPlatform))
 			{
 				const int32 SlotIndex = Component->GroomAsset->GetMaterialIndex(Component->GroomAsset->HairGroupsRendering[GroupIt].MaterialSlotName);
 				HairInstance->Strands.Material = Component->GetMaterial(GetMaterialIndexWithFallback(SlotIndex), EHairGeometryType::Strands, true);
 			}
 
 			// Material - Cards
-			if (IsHairStrandsEnabled(EHairStrandsShaderType::Cards, Platform))
+			if (IsHairStrandsEnabled(EHairStrandsShaderType::Cards, ShaderPlatform))
 			{
 				uint32 CardsLODIndex = 0;
 				for (const FHairGroupData::FCards::FLOD& LOD : InGroupData.Cards.LODs)
@@ -563,7 +443,7 @@ public:
 			}
 
 			// Material - Meshes
-			if (IsHairStrandsEnabled(EHairStrandsShaderType::Meshes, Platform))
+			if (IsHairStrandsEnabled(EHairStrandsShaderType::Meshes, ShaderPlatform))
 			{
 				uint32 MeshesLODIndex = 0;
 				for (const FHairGroupData::FMeshes::FLOD& LOD : InGroupData.Meshes.LODs)
@@ -586,22 +466,6 @@ public:
 				}
 			}
 		}
-
-		if (IsHairStrandsEnabled(EHairStrandsShaderType::Cards, Platform) || IsHairStrandsEnabled(EHairStrandsShaderType::Meshes, Platform))
-		{
-			CardsAndMeshesVertexFactories.BeginInitResource();
-		}
-
-		FHairStrandsVertexFactory* LocalStrandsVertexFactory = &StrandsVertexFactory;
-		ENQUEUE_RENDER_COMMAND(InitHairStrandsVertexFactory)(
-			[Platform, LocalStrandsVertexFactory, StrandsVFData](FRHICommandListImmediate& RHICmdList)
-		{
-			if (IsHairStrandsEnabled(EHairStrandsShaderType::Strands, Platform))
-			{
-				LocalStrandsVertexFactory->SetData(StrandsVFData);
-				LocalStrandsVertexFactory->InitResource();
-			}
-		});
 	}
 
 	virtual ~FHairStrandsSceneProxy()
@@ -652,9 +516,6 @@ public:
 				Instance->Release();
 			}
 		}
-
-		CardsAndMeshesVertexFactories.ReleaseResources();
-		StrandsVertexFactory.ReleaseResource();
 	}
 
 	virtual void OnTransformChanged() override
@@ -884,9 +745,7 @@ public:
 			{
 				return nullptr;
 			}
-			const bool bIsDynamic  = Instance->Meshes.LODs[IntLODIndex].DeformedResource != nullptr;
-			const bool bIsCurrent0 = bIsDynamic ? Instance->Meshes.LODs[IntLODIndex].DeformedResource->CurrentIndex == 0 : true;
-			VertexFactory = (FVertexFactory*)CardsAndMeshesVertexFactories.GetVertexFactory(GroupIndex, IntLODIndex, bIsCurrent0, GeometryType);
+			VertexFactory = (FVertexFactory*)Instance->Meshes.LODs[IntLODIndex].GetVertexFactory();
 			check(VertexFactory);
 			HairVertexCount = Instance->Meshes.LODs[IntLODIndex].RestResource->PrimitiveCount * 3;
 			MaxVertexIndex = HairVertexCount;
@@ -902,9 +761,7 @@ public:
 				return nullptr;
 			}
 
-			const bool bIsDynamic  = Instance->Cards.LODs[IntLODIndex].DeformedResource != nullptr;
-			const bool bIsCurrent0 = bIsDynamic ? Instance->Cards.LODs[IntLODIndex].DeformedResource->CurrentIndex == 0 : true;
-			VertexFactory = (FVertexFactory*)CardsAndMeshesVertexFactories.GetVertexFactory(GroupIndex, IntLODIndex, bIsCurrent0, GeometryType);
+			VertexFactory = (FVertexFactory*)Instance->Cards.LODs[IntLODIndex].GetVertexFactory();
 			check(VertexFactory);
 			HairVertexCount = Instance->Cards.LODs[IntLODIndex].RestResource->PrimitiveCount * 3;
 			MaxVertexIndex = HairVertexCount;
@@ -916,7 +773,7 @@ public:
 		}
 		else // if (GeometryType == EHairGeometryType::Strands)
 		{
-			VertexFactory = (FVertexFactory*)&StrandsVertexFactory;
+			VertexFactory = (FVertexFactory*)Instance->Strands.VertexFactory;
 			HairVertexCount = Instance->Strands.RestResource->GetVertexCount();
 			MaxVertexIndex = HairVertexCount * 6;
 			bUseCulling = Instance->Strands.bIsCullingEnabled;
@@ -997,15 +854,11 @@ public:
 
 	virtual FPrimitiveViewRelevance GetViewRelevance(const FSceneView* View) const override
 	{
-		TArray<FHairGroupInstance*> Instances = StrandsVertexFactory.GetData().Instances;
-
 		bool bUseCardsOrMesh = false;
-		const uint32 GroupCount = Instances.Num();
-		for (uint32 GroupIt = 0; GroupIt < GroupCount; ++GroupIt)
+		for (FHairGroupInstance* Instance : HairGroupInstances)
 		{
-			check(Instances[GroupIt]->GetRefCount());
-			const EHairGeometryType GeometryType = Instances[GroupIt]->GeometryType;
-			const FHairGroup& GroupData = HairGroups[GroupIt];
+			check(Instance->GetRefCount());
+			const EHairGeometryType GeometryType = Instance->GeometryType;
 			bUseCardsOrMesh = bUseCardsOrMesh || GeometryType == EHairGeometryType::Cards || GeometryType == EHairGeometryType::Meshes;
 		}
 
@@ -1039,9 +892,6 @@ public:
 
 private:
 	uint32 ComponentId = 0;
-	EHairLODSelectionType GroomLODSelection = EHairLODSelectionType::Cpu;
-	FHairStrandsVertexFactory StrandsVertexFactory;
-	FHairCardsVertexFactories CardsAndMeshesVertexFactories; // One factory per instance, per LOD, and per buffer type (current/previous)
 	FMaterialRelevance MaterialRelevance;
 	UMaterialInterface* Strands_DebugMaterial = nullptr;
 	int32* PredictedLODIndex = nullptr;
@@ -2132,9 +1982,7 @@ void UGroomComponent::InitResources(bool bIsBindingReloading)
 				}
 
 				HairGroupInstance->Guides.RestRootResource = LocalBindingAsset->HairGroupResources[GroupIt].SimRootResources;
-
 				HairGroupInstance->Guides.DeformedRootResource = new FHairStrandsDeformedRootResource(HairGroupInstance->Guides.RestRootResource, EHairStrandsResourcesType::Guides);
-				BeginInitResource(HairGroupInstance->Guides.DeformedRootResource);
 			}
 
 			// Lazy allocation of the guide resources
@@ -2143,8 +1991,6 @@ void UGroomComponent::InitResources(bool bIsBindingReloading)
 
 			// If guides are allocated, deformed resources are always needs since they are either used with simulation, or RBF deformation. Both are dynamics, and require deformed positions
 			HairGroupInstance->Guides.DeformedResource = new FHairStrandsDeformedResource(GroupData.Guides.Data.RenderData, true, HairGroupInstance->Guides.RestResource->PositionOffset, EHairStrandsResourcesType::Guides);
-			BeginInitResource(HairGroupInstance->Guides.DeformedResource);
-
 			HairGroupInstance->Guides.DeformedResource->GetPositionOffset(FHairStrandsDeformedResource::EFrameType::Current)  = HairGroupInstance->Guides.RestResource->PositionOffset;
 			HairGroupInstance->Guides.DeformedResource->GetPositionOffset(FHairStrandsDeformedResource::EFrameType::Previous) = HairGroupInstance->Guides.RestResource->PositionOffset;
 
@@ -2211,7 +2057,6 @@ void UGroomComponent::InitResources(bool bIsBindingReloading)
 				{
 					// Allocate dynamic raytracing resources (owned by the groom component/instance)
 					HairGroupInstance->Strands.RenRaytracingResource = new FHairStrandsRaytracingResource(GroupData.Strands.Data);
-					BeginInitResource(HairGroupInstance->Strands.RenRaytracingResource);
 					HairGroupInstance->Strands.RenRaytracingResourceOwned = true;
 				}
 				else
@@ -2236,15 +2081,12 @@ void UGroomComponent::InitResources(bool bIsBindingReloading)
 
 				HairGroupInstance->Strands.RestRootResource = LocalBindingAsset->HairGroupResources[GroupIt].RenRootResources;
 				HairGroupInstance->Strands.DeformedRootResource = new FHairStrandsDeformedRootResource(HairGroupInstance->Strands.RestRootResource, EHairStrandsResourcesType::Strands);
-				BeginInitResource(HairGroupInstance->Strands.DeformedRootResource);
 			}
 
 			HairGroupInstance->Strands.RestResource = GroupData.Strands.RestResource;
 			if (bNeedDynamicResources)
 			{
 				HairGroupInstance->Strands.DeformedResource = new FHairStrandsDeformedResource(GroupData.Strands.Data.RenderData, false, HairGroupInstance->Strands.RestResource->PositionOffset, EHairStrandsResourcesType::Strands);
-				BeginInitResource(HairGroupInstance->Strands.DeformedResource);
-
 				// Initialize deformed position relative position offset to the rest pose offset
 				HairGroupInstance->Strands.DeformedResource->GetPositionOffset(FHairStrandsDeformedResource::EFrameType::Current)  = HairGroupInstance->Strands.RestResource->PositionOffset;
 				HairGroupInstance->Strands.DeformedResource->GetPositionOffset(FHairStrandsDeformedResource::EFrameType::Previous) = HairGroupInstance->Strands.RestResource->PositionOffset;
@@ -2312,7 +2154,6 @@ void UGroomComponent::InitResources(bool bIsBindingReloading)
 				if (bNeedDeformedPositions)
 				{
 					InstanceLOD.DeformedResource = new FHairCardsDeformedResource(LOD.Data.RenderData, false);
-					BeginInitResource(InstanceLOD.DeformedResource);
 				}
 
 				#if RHI_RAYTRACING
@@ -2322,7 +2163,6 @@ void UGroomComponent::InitResources(bool bIsBindingReloading)
 					{
 						// Allocate dynamic raytracing resources (owned by the groom component/instance)
 						InstanceLOD.RaytracingResource = new FHairStrandsRaytracingResource(*InstanceLOD.Data);
-						BeginInitResource(InstanceLOD.RaytracingResource);
 						InstanceLOD.RaytracingResourceOwned = true;
 					}
 					else
@@ -2356,13 +2196,11 @@ void UGroomComponent::InitResources(bool bIsBindingReloading)
 
 						InstanceLOD.Guides.RestRootResource = LocalBindingAsset->HairGroupResources[GroupIt].CardsRootResources[CardsLODIndex];
 						InstanceLOD.Guides.DeformedRootResource = new FHairStrandsDeformedRootResource(InstanceLOD.Guides.RestRootResource, EHairStrandsResourcesType::Cards);
-						BeginInitResource(InstanceLOD.Guides.DeformedRootResource);
 					}
 
 					InstanceLOD.Guides.RestResource = LOD.Guides.RestResource;
 					{
 						InstanceLOD.Guides.DeformedResource = new FHairStrandsDeformedResource(LOD.Guides.Data.RenderData, false, InstanceLOD.Guides.RestResource->PositionOffset, EHairStrandsResourcesType::Cards);
-						BeginInitResource(InstanceLOD.Guides.DeformedResource);
 
 						// Initialize deformed position relative position offset to the rest pose offset
 						InstanceLOD.Guides.DeformedResource->GetPositionOffset(FHairStrandsDeformedResource::EFrameType::Current) = InstanceLOD.Guides.RestResource->PositionOffset;
@@ -2390,7 +2228,6 @@ void UGroomComponent::InitResources(bool bIsBindingReloading)
 				if (bNeedDeformedPositions)
 				{
 					InstanceLOD.DeformedResource = new FHairMeshesDeformedResource(LOD.Data.RenderData, true);
-					BeginInitResource(InstanceLOD.DeformedResource);
 				}
 
 				#if RHI_RAYTRACING
@@ -2400,7 +2237,6 @@ void UGroomComponent::InitResources(bool bIsBindingReloading)
 					{
 						// Allocate dynamic raytracing resources (owned by the groom component/instance)
 						InstanceLOD.RaytracingResource = new FHairStrandsRaytracingResource(*InstanceLOD.Data);
-						BeginInitResource(InstanceLOD.RaytracingResource);
 						InstanceLOD.RaytracingResourceOwned = true;
 					}
 					else
@@ -2493,6 +2329,8 @@ void UGroomComponent::DeleteDeferredHairGroupInstances()
 				#if WITH_EDITOR
 				LocalInstance->Strands.DebugAttributeBuffer.Release();
 				#endif
+
+				InternalResourceRelease(LocalInstance->Strands.VertexFactory);
 			}
 
 			// Cards
@@ -2511,6 +2349,9 @@ void UGroomComponent::DeleteDeferredHairGroupInstances()
 							InternalResourceRelease(LocalInstance->Cards.LODs[CardLODIt].RaytracingResource);
 						}
 						#endif
+
+						InternalResourceRelease(LocalInstance->Cards.LODs[CardLODIt].VertexFactory[0]);
+						InternalResourceRelease(LocalInstance->Cards.LODs[CardLODIt].VertexFactory[1]);
 					}
 				}
 			}
@@ -2529,6 +2370,9 @@ void UGroomComponent::DeleteDeferredHairGroupInstances()
 							InternalResourceRelease(LocalInstance->Meshes.LODs[MeshesLODIt].RaytracingResource);
 						}
 						#endif
+
+						InternalResourceRelease(LocalInstance->Meshes.LODs[MeshesLODIt].VertexFactory[0]);
+						InternalResourceRelease(LocalInstance->Meshes.LODs[MeshesLODIt].VertexFactory[1]);
 					}
 				}
 			}
diff --git a/Engine/Plugins/Runtime/HairStrands/Source/HairStrandsCore/Private/GroomManager.cpp b/Engine/Plugins/Runtime/HairStrands/Source/HairStrandsCore/Private/GroomManager.cpp
--- a/Engine/Plugins/Runtime/HairStrands/Source/HairStrandsCore/Private/GroomManager.cpp
+++ b/Engine/Plugins/Runtime/HairStrands/Source/HairStrandsCore/Private/GroomManager.cpp
@@ -16,6 +16,8 @@
 #include "GroomGeometryCache.h"
 #include "HAL/IConsoleManager.h"
 #include "SceneView.h"
+#include "HairCardsVertexFactory.h"
+#include "HairStrandsVertexFactory.h"
 
 static int32 GHairStrandsMinLOD = 0;
 static FAutoConsoleVariableRef CVarGHairStrandsMinLOD(TEXT("r.HairStrands.MinLOD"), GHairStrandsMinLOD, TEXT("Clamp the min hair LOD to this value, preventing to reach lower/high-quality LOD."), ECVF_Scalability);
@@ -508,7 +510,7 @@ static void RunHairBufferSwap(const FHairStrandsInstances& Instances, const TArr
 	}
 }
 
-static void RunHairLODSelection(const FHairStrandsInstances& Instances, const TArray<const FSceneView*> Views)
+static void RunHairLODSelection(FRDGBuilder& GraphBuilder, const FHairStrandsInstances& Instances, const TArray<const FSceneView*> Views)
 {
 	EShaderPlatform ShaderPlatform = EShaderPlatform::SP_NumPlatforms;
 	if (Views.Num() > 0)
@@ -611,6 +613,47 @@ static void RunHairLODSelection(const FHairStrandsInstances& Instances, const TA
 			GeometryType = EHairGeometryType::NoneGeometry;
 		}
 
+		// Lazy allocation of resources
+		// Note: Allocation will only be done if the resources it not yet initialized.
+		if (Instance->Guides.Data)
+		{
+			if (Instance->Guides.DeformedRootResource)	Instance->Guides.DeformedRootResource->Allocate(GraphBuilder);
+			if (Instance->Guides.DeformedResource)		Instance->Guides.DeformedResource->Allocate(GraphBuilder);
+		}
+
+		if (GeometryType == EHairGeometryType::Meshes)
+		{
+			FHairGroupInstance::FMeshes::FLOD& InstanceLOD = Instance->Meshes.LODs[IntLODIndex];
+
+			if (InstanceLOD.DeformedResource)	InstanceLOD.DeformedResource->Allocate(GraphBuilder);
+			#if RHI_RAYTRACING
+			if (InstanceLOD.RaytracingResource)	InstanceLOD.RaytracingResource->Allocate(GraphBuilder);
+			#endif
+
+			InstanceLOD.GetVertexFactory()->InitResources();
+		}
+		else if (GeometryType == EHairGeometryType::Cards)
+		{
+			FHairGroupInstance::FCards::FLOD& InstanceLOD = Instance->Cards.LODs[IntLODIndex];
+
+			if (InstanceLOD.DeformedResource)				InstanceLOD.DeformedResource->Allocate(GraphBuilder);
+			if (InstanceLOD.Guides.DeformedRootResource)	InstanceLOD.Guides.DeformedRootResource->Allocate(GraphBuilder);
+			if (InstanceLOD.Guides.DeformedResource)		InstanceLOD.Guides.DeformedResource->Allocate(GraphBuilder);
+			#if RHI_RAYTRACING
+			if (InstanceLOD.RaytracingResource)				InstanceLOD.RaytracingResource->Allocate(GraphBuilder);
+			#endif
+			InstanceLOD.GetVertexFactory()->InitResources();
+		}
+		else if (GeometryType == EHairGeometryType::Strands)
+		{
+			if (Instance->Strands.DeformedRootResource)	Instance->Strands.DeformedRootResource->Allocate(GraphBuilder);
+			if (Instance->Strands.DeformedResource)		Instance->Strands.DeformedResource->Allocate(GraphBuilder);
+			#if RHI_RAYTRACING
+			if (Instance->Strands.RenRaytracingResource)Instance->Strands.RenRaytracingResource->Allocate(GraphBuilder);
+			#endif
+			Instance->Strands.VertexFactory->InitResources();
+		}
+
 		Instance->HairGroupPublicData->SetLODVisibility(bIsVisible);
 		Instance->HairGroupPublicData->SetLODIndex(LODIndex);
 		Instance->HairGroupPublicData->SetLODBias(0);
@@ -621,7 +664,6 @@ static void RunHairLODSelection(const FHairStrandsInstances& Instances, const TA
 		Instance->Guides.bIsSimulationEnable = Instance->HairGroupPublicData->IsSimulationEnable(IntLODIndex);
 		Instance->Guides.bHasGlobalInterpolation = Instance->HairGroupPublicData->IsGlobalInterpolationEnable(IntLODIndex);
 		Instance->Strands.bIsCullingEnabled = bCullingEnable;
-
 	}
 }
 
@@ -724,7 +766,9 @@ void ProcessHairStrandsBookmark(
 				*Parameters.Instances,
 				Parameters.AllViews);
 		}
+		check(GraphBuilder);
 		RunHairLODSelection(
+			*GraphBuilder,
 			*Parameters.Instances,
 			Parameters.AllViews);
 	}
diff --git a/Engine/Plugins/Runtime/HairStrands/Source/HairStrandsCore/Private/GroomResources.cpp b/Engine/Plugins/Runtime/HairStrands/Source/HairStrandsCore/Private/GroomResources.cpp
--- a/Engine/Plugins/Runtime/HairStrands/Source/HairStrandsCore/Private/GroomResources.cpp
+++ b/Engine/Plugins/Runtime/HairStrands/Source/HairStrandsCore/Private/GroomResources.cpp
@@ -312,6 +312,44 @@ void CreateHairStrandsDebugAttributeBuffer(FRDGExternalBuffer* DebugAttributeBuf
 
 /////////////////////////////////////////////////////////////////////////////////////////
 
+FHairCommonResource::FHairCommonResource(EHairStrandsAllocationType InAllocationType, bool bInUseRenderGraph):
+bUseRenderGraph(bInUseRenderGraph),
+bIsInitialized(false),
+AllocationType(InAllocationType)
+{
+}
+
+void FHairCommonResource::InitRHI()
+{
+	if (bIsInitialized || AllocationType == EHairStrandsAllocationType::Deferred || GUsingNullRHI) { return; }
+
+	if (bUseRenderGraph)
+	{
+		FMemMark Mark(FMemStack::Get());
+		FRHICommandListImmediate& RHICmdList = FRHICommandListExecutor::GetImmediateCommandList();
+		FRDGBuilder GraphBuilder(RHICmdList);
+		InternalAllocate(GraphBuilder);
+		GraphBuilder.Execute();
+	}
+	else
+	{
+		InternalAllocate();
+	}
+	bIsInitialized = true;
+}
+
+void FHairCommonResource::Allocate(FRDGBuilder& GraphBuilder)
+{
+	if (bIsInitialized) { return; }
+
+	check(AllocationType == EHairStrandsAllocationType::Deferred);
+	FRenderResource::InitResource(); // Call RenderResource InitResource() so that the resources is marked as initialized
+	InternalAllocate(GraphBuilder);
+	bIsInitialized = true;
+}
+
+/////////////////////////////////////////////////////////////////////////////////////////
+
 void FHairCardIndexBuffer::InitRHI()
 {
 	const uint32 DataSizeInBytes = FHairCardsIndexFormat::SizeInByte * Indices.Num();
@@ -324,6 +362,7 @@ void FHairCardIndexBuffer::InitRHI()
 }
 
 FHairCardsRestResource::FHairCardsRestResource(const FHairCardsDatas::FRenderData& InRenderData, uint32 InVertexCount, uint32 InPrimitiveCount) :
+	FHairCommonResource(EHairStrandsAllocationType::Immediate, false),
 	RestPositionBuffer(),
 	RestIndexBuffer(InRenderData.Indices),
 	VertexCount(InVertexCount),
@@ -335,7 +374,7 @@ FHairCardsRestResource::FHairCardsRestResource(const FHairCardsDatas::FRenderDat
 
 }
 
-void FHairCardsRestResource::InitRHI()
+void FHairCardsRestResource::InternalAllocate()
 {
 	// These resources are kept as regular (i.e., non-RDG resources) as they need to be bound at the input assembly stage by the Vertex declaraction which requires FVertexBuffer type
 	CreateBuffer<FHairCardsPositionFormat>(RenderData.Positions, RestPositionBuffer, TEXT("Hair.CardsRest_PositionBuffer"));
@@ -373,6 +412,7 @@ void FHairCardsRestResource::ReleaseResource()
 
 /////////////////////////////////////////////////////////////////////////////////////////
 FHairCardsProceduralResource::FHairCardsProceduralResource(const FHairCardsProceduralDatas::FRenderData& InRenderData, const FIntPoint& InAtlasResolution, const FHairCardsVoxel& InVoxel):
+	FHairCommonResource(EHairStrandsAllocationType::Immediate),
 	CardBoundCount(InRenderData.ClusterBounds.Num()),
 	AtlasResolution(InAtlasResolution),
 	AtlasRectBuffer(),
@@ -387,13 +427,8 @@ FHairCardsProceduralResource::FHairCardsProceduralResource(const FHairCardsProce
 	CardVoxel = InVoxel;
 }
 
-void FHairCardsProceduralResource::InitRHI()
+void FHairCardsProceduralResource::InternalAllocate(FRDGBuilder& GraphBuilder)
 {
-	if (GUsingNullRHI) { return; }
-	FMemMark Mark(FMemStack::Get());
-	FRHICommandListImmediate& RHICmdList = FRHICommandListExecutor::GetImmediateCommandList();
-	FRDGBuilder GraphBuilder(RHICmdList);
-
 	InternalCreateVertexBufferRDG<FHairCardsAtlasRectFormat>(GraphBuilder, RenderData.CardsRect, AtlasRectBuffer, TEXT("Hair.CardsProcedural_AtlasRectBuffer"), EHairResourceUsageType::Static);
 	InternalCreateVertexBufferRDG<FHairCardsDimensionFormat>(GraphBuilder, RenderData.CardsLengths, LengthBuffer, TEXT("Hair.CardsProcedural_LengthBuffer"), EHairResourceUsageType::Static);
 
@@ -407,8 +442,6 @@ void FHairCardsProceduralResource::InitRHI()
 
 	InternalCreateVertexBufferRDG<FHairCardsStrandsPositionFormat>(GraphBuilder, RenderData.CardsStrandsPositions, CardsStrandsPositions, TEXT("Hair.CardsProcedural_CardsStrandsPositions"), EHairResourceUsageType::Static);
 	InternalCreateVertexBufferRDG<FHairCardsStrandsAttributeFormat>(GraphBuilder, RenderData.CardsStrandsAttributes, CardsStrandsAttributes, TEXT("Hair.CardsProcedural_CardsStrandsAttributes"), EHairResourceUsageType::Static);
-
-	GraphBuilder.Execute();
 }
 
 void FHairCardsProceduralResource::ReleaseRHI()
@@ -430,16 +463,12 @@ void FHairCardsProceduralResource::ReleaseRHI()
 /////////////////////////////////////////////////////////////////////////////////////////
 
 FHairCardsDeformedResource::FHairCardsDeformedResource(const FHairCardsDatas::FRenderData& HairCardsRenderData, bool bInInitializedData) :
+	FHairCommonResource(EHairStrandsAllocationType::Deferred),
 	RenderData(HairCardsRenderData), bInitializedData(bInInitializedData)
 {}
 
-void FHairCardsDeformedResource::InitRHI()
+void FHairCardsDeformedResource::InternalAllocate(FRDGBuilder& GraphBuilder)
 {
-	if (GUsingNullRHI) { return; }
-	FMemMark Mark(FMemStack::Get());
-	FRHICommandListImmediate& RHICmdList = FRHICommandListExecutor::GetImmediateCommandList();
-	FRDGBuilder GraphBuilder(RHICmdList);
-
 	const uint32 VertexCount = RenderData.Positions.Num();
 	if (bInitializedData)
 	{
@@ -451,7 +480,6 @@ void FHairCardsDeformedResource::InitRHI()
 		InternalCreateVertexBufferRDG<FHairCardsPositionFormat>(GraphBuilder, VertexCount, DeformedPositionBuffer[0], TEXT("Hair.CardsDeformed_Position0"), EHairResourceUsageType::Dynamic);
 		InternalCreateVertexBufferRDG<FHairCardsPositionFormat>(GraphBuilder, VertexCount, DeformedPositionBuffer[1], TEXT("Hair.CardsDeformed_Position1"), EHairResourceUsageType::Dynamic);
 	}
-	GraphBuilder.Execute();
 }
 
 void FHairCardsDeformedResource::ReleaseRHI()
@@ -463,6 +491,7 @@ void FHairCardsDeformedResource::ReleaseRHI()
 /////////////////////////////////////////////////////////////////////////////////////////
 
 FHairMeshesRestResource::FHairMeshesRestResource(const FHairMeshesDatas::FRenderData& InRenderData, uint32 InVertexCount, uint32 InPrimitiveCount) :
+	FHairCommonResource(EHairStrandsAllocationType::Immediate, false),
 	RestPositionBuffer(),
 	IndexBuffer(InRenderData.Indices),
 	VertexCount(InVertexCount),
@@ -475,7 +504,7 @@ FHairMeshesRestResource::FHairMeshesRestResource(const FHairMeshesDatas::FRender
 	check(IndexBuffer.Indices.Num() > 0);
 }
 
-void FHairMeshesRestResource::InitRHI()
+void FHairMeshesRestResource::InternalAllocate()
 {
 	// These resources are kept as regular (i.e., non-RDG resources) as they need to be bound at the input assembly stage by the Vertex declaraction which requires FVertexBuffer type
 	CreateBuffer<FHairCardsPositionFormat>(RenderData.Positions, RestPositionBuffer, TEXT("Hair.MeshesRest_Positions"));
@@ -510,16 +539,12 @@ void FHairMeshesRestResource::ReleaseResource()
 /////////////////////////////////////////////////////////////////////////////////////////
 
 FHairMeshesDeformedResource::FHairMeshesDeformedResource(const FHairMeshesDatas::FRenderData& HairMeshesData, bool bInInitializedData) :
+	FHairCommonResource(EHairStrandsAllocationType::Deferred),
 	RenderData(HairMeshesData), bInitializedData(bInInitializedData)
 {}
 
-void FHairMeshesDeformedResource::InitRHI()
+void FHairMeshesDeformedResource::InternalAllocate(FRDGBuilder& GraphBuilder)
 {
-	if (GUsingNullRHI) { return; }
-	FMemMark Mark(FMemStack::Get());
-	FRHICommandListImmediate& RHICmdList = FRHICommandListExecutor::GetImmediateCommandList();
-	FRDGBuilder GraphBuilder(RHICmdList);
-
 	const uint32 VertexCount = RenderData.Positions.Num();
 	if (bInitializedData)
 	{
@@ -531,8 +556,6 @@ void FHairMeshesDeformedResource::InitRHI()
 		InternalCreateVertexBufferRDG<FHairCardsPositionFormat>(GraphBuilder, VertexCount, DeformedPositionBuffer[0], TEXT("Hair.MeshesDeformed_Positions0"), EHairResourceUsageType::Dynamic);
 		InternalCreateVertexBufferRDG<FHairCardsPositionFormat>(GraphBuilder, VertexCount, DeformedPositionBuffer[1], TEXT("Hair.MeshesDeformed_Positions1"), EHairResourceUsageType::Dynamic);
 	}
-
-	GraphBuilder.Execute();
 }
 
 void FHairMeshesDeformedResource::ReleaseRHI()
@@ -544,16 +567,12 @@ void FHairMeshesDeformedResource::ReleaseRHI()
 /////////////////////////////////////////////////////////////////////////////////////////
 
 FHairStrandsRestResource::FHairStrandsRestResource(const FHairStrandsDatas::FRenderData& HairStrandRenderData, const FVector& InPositionOffset, EHairStrandsResourcesType InCurveType) :
+	FHairCommonResource(EHairStrandsAllocationType::Immediate),
 	PositionBuffer(), AttributeBuffer(), MaterialBuffer(), PositionOffset(InPositionOffset), RenderData(HairStrandRenderData), CurveType(InCurveType)
 {}
 
-void FHairStrandsRestResource::InitRHI()
+void FHairStrandsRestResource::InternalAllocate(FRDGBuilder& GraphBuilder)
 {
-	if (GUsingNullRHI) { return; }
-	FMemMark Mark(FMemStack::Get());
-	FRHICommandListImmediate& RHICmdList = FRHICommandListExecutor::GetImmediateCommandList();
-	FRDGBuilder GraphBuilder(RHICmdList);
-
 	InternalCreateVertexBufferRDG<FHairStrandsPositionFormat>(GraphBuilder, RenderData.Positions, PositionBuffer, HAIRSTRANDS_RESOUCE_NAME(CurveType, Hair.StrandsRest_PositionBuffer), EHairResourceUsageType::Static);
 	InternalCreateVertexBufferRDG<FHairStrandsAttributeFormat>(GraphBuilder, RenderData.Attributes, AttributeBuffer, HAIRSTRANDS_RESOUCE_NAME(CurveType, Hair.StrandsRest_AttributeBuffer), EHairResourceUsageType::Static);
 	InternalCreateVertexBufferRDG<FHairStrandsMaterialFormat>(GraphBuilder, RenderData.Materials, MaterialBuffer, HAIRSTRANDS_RESOUCE_NAME(CurveType, Hair.StrandsRest_MaterialBuffer), EHairResourceUsageType::Static);
@@ -561,8 +580,6 @@ void FHairStrandsRestResource::InitRHI()
 	TArray<FVector4> RestOffset;
 	RestOffset.Add(PositionOffset);
 	InternalCreateVertexBufferRDG<FHairStrandsPositionOffsetFormat>(GraphBuilder, RestOffset, PositionOffsetBuffer, HAIRSTRANDS_RESOUCE_NAME(CurveType, Hair.StrandsRest_PositionOffsetBuffer), EHairResourceUsageType::Static, ERDGInitialDataFlags::None);
-	
-	GraphBuilder.Execute();
 }
 
 void AddHairTangentPass(
@@ -605,16 +622,12 @@ void FHairStrandsRestResource::ReleaseRHI()
 /////////////////////////////////////////////////////////////////////////////////////////
 
 FHairStrandsDeformedResource::FHairStrandsDeformedResource(const FHairStrandsDatas::FRenderData& HairStrandRenderData, bool bInInitializedData, const FVector& InDefaultOffset, EHairStrandsResourcesType InCurveType) :
+	FHairCommonResource(EHairStrandsAllocationType::Deferred),
 	RenderData(HairStrandRenderData), bInitializedData(bInInitializedData), DefaultOffset(InDefaultOffset), CurveType(InCurveType)
 {}
 
-void FHairStrandsDeformedResource::InitRHI()
+void FHairStrandsDeformedResource::InternalAllocate(FRDGBuilder& GraphBuilder)
 {
-	if (GUsingNullRHI) { return; }
-	FMemMark Mark(FMemStack::Get());
-	FRHICommandListImmediate& RHICmdList = FRHICommandListExecutor::GetImmediateCommandList();
-	FRDGBuilder GraphBuilder(RHICmdList);
-
 	const uint32 VertexCount = RenderData.Positions.Num();
 	if (bInitializedData)
 	{
@@ -632,8 +645,6 @@ void FHairStrandsDeformedResource::InitRHI()
 	DefaultOffsets.Add(DefaultOffset);
 	InternalCreateVertexBufferRDG<FHairStrandsPositionOffsetFormat>(GraphBuilder, DefaultOffsets, DeformedOffsetBuffer[0], HAIRSTRANDS_RESOUCE_NAME(CurveType, Hair.StrandsDeformed_DeformedOffsetBuffer0), EHairResourceUsageType::Dynamic);
 	InternalCreateVertexBufferRDG<FHairStrandsPositionOffsetFormat>(GraphBuilder, DefaultOffsets, DeformedOffsetBuffer[1], HAIRSTRANDS_RESOUCE_NAME(CurveType, Hair.StrandsDeformed_DeformedOffsetBuffer1), EHairResourceUsageType::Dynamic);
-
-	GraphBuilder.Execute();
 }
 
 void FHairStrandsDeformedResource::ReleaseRHI()
@@ -675,26 +686,20 @@ void FHairStrandsClusterCullingData::Serialize(FArchive& Ar)
 	Ar << PackedClusterInfos;
 }
 
-FHairStrandsClusterCullingResource::FHairStrandsClusterCullingResource(const FHairStrandsClusterCullingData& InData)
-: Data(InData) 
+FHairStrandsClusterCullingResource::FHairStrandsClusterCullingResource(const FHairStrandsClusterCullingData& InData): 
+	FHairCommonResource(EHairStrandsAllocationType::Immediate),
+	Data(InData) 
 {
 
 }
 
-void FHairStrandsClusterCullingResource::InitRHI()
+void FHairStrandsClusterCullingResource::InternalAllocate(FRDGBuilder& GraphBuilder)
 {
-	if (GUsingNullRHI) { return; }
-	FMemMark Mark(FMemStack::Get());
-	FRHICommandListImmediate& RHICmdList = FRHICommandListExecutor::GetImmediateCommandList();
-	FRDGBuilder GraphBuilder(RHICmdList);
-
 	InternalCreateStructuredBufferRDG(GraphBuilder, Data.PackedClusterInfos, sizeof(FHairStrandsClusterCullingData::FHairClusterInfo::Packed), ClusterInfoBuffer, TEXT("Hair.StrandsClusterCulling_ClusterInfoBuffer"), EHairResourceUsageType::Static);
 	InternalCreateStructuredBufferRDG(GraphBuilder, Data.ClusterLODInfos, sizeof(FHairStrandsClusterCullingData::FHairClusterLODInfo), ClusterLODInfoBuffer, TEXT("Hair.StrandsClusterCulling_ClusterLODInfoBuffer"), EHairResourceUsageType::Static);
 
 	InternalCreateVertexBufferRDG(GraphBuilder, Data.ClusterVertexIds, PF_R32_UINT, ClusterVertexIdBuffer, TEXT("Hair.StrandsClusterCulling_ClusterVertexIds"), EHairResourceUsageType::Static);
 	InternalCreateVertexBufferRDG(GraphBuilder, Data.VertexToClusterIds, PF_R32_UINT, VertexToClusterIdBuffer, TEXT("Hair.StrandsClusterCulling_VertexToClusterIds"), EHairResourceUsageType::Static);
-
-	GraphBuilder.Execute();
 }
 
 void FHairStrandsClusterCullingResource::ReleaseRHI()
@@ -708,12 +713,14 @@ void FHairStrandsClusterCullingResource::ReleaseRHI()
 /////////////////////////////////////////////////////////////////////////////////////////
 
 FHairStrandsRestRootResource::FHairStrandsRestRootResource(const FHairStrandsRootData& InRootData, EHairStrandsResourcesType InCurveType) :
-RootData(InRootData), CurveType(InCurveType)
+	FHairCommonResource(EHairStrandsAllocationType::Immediate),
+	RootData(InRootData), CurveType(InCurveType)
 {
 	PopulateFromRootData();
 }
 
 FHairStrandsRestRootResource::FHairStrandsRestRootResource(const FHairStrandsDatas* HairStrandsDatas, uint32 LODCount, const TArray<uint32>& NumSamples, EHairStrandsResourcesType InCurveType):
+	FHairCommonResource(EHairStrandsAllocationType::Immediate),
 	RootData(HairStrandsDatas, LODCount, NumSamples), CurveType(InCurveType)
 {
 	PopulateFromRootData();
@@ -733,15 +740,10 @@ void FHairStrandsRestRootResource::PopulateFromRootData()
 	}
 }
 
-void FHairStrandsRestRootResource::InitRHI()
+void FHairStrandsRestRootResource::InternalAllocate(FRDGBuilder& GraphBuilder)
 {
 	if (RootData.VertexToCurveIndexBuffer.Num() > 0)
 	{
-		if (GUsingNullRHI) { return; }
-		FMemMark Mark(FMemStack::Get());
-		FRHICommandListImmediate& RHICmdList = FRHICommandListExecutor::GetImmediateCommandList();
-		FRDGBuilder GraphBuilder(RHICmdList);
-
 		InternalCreateVertexBufferRDG<FHairStrandsIndexFormat>(GraphBuilder, RootData.VertexToCurveIndexBuffer, VertexToCurveIndexBuffer, HAIRSTRANDS_RESOUCE_NAME(CurveType, Hair.StrandsRestRoot_VertexToCurveIndexBuffer), EHairResourceUsageType::Static);
 		InternalCreateVertexBufferRDG<FHairStrandsRootPositionFormat>(GraphBuilder, RootData.RootPositionBuffer, RootPositionBuffer, HAIRSTRANDS_RESOUCE_NAME(CurveType, Hair.StrandsRestRoot_RootPositionBuffer), EHairResourceUsageType::Static);
 		InternalCreateVertexBufferRDG<FHairStrandsRootNormalFormat>(GraphBuilder, RootData.RootNormalBuffer, RootNormalBuffer, HAIRSTRANDS_RESOUCE_NAME(CurveType, Hair.StrandsRestRoot_RootNormalBuffer), EHairResourceUsageType::Static);
@@ -802,7 +804,6 @@ void FHairStrandsRestRootResource::InitRHI()
 				InternalCreateVertexBufferRDG<FHairStrandsMeshTrianglePositionFormat>(GraphBuilder, CPUData.SampleCount, GPUData.RestSamplePositionsBuffer, HAIRSTRANDS_RESOUCE_NAME(CurveType, Hair.StrandsRestRoot_RestSamplePositionsBuffer), EHairResourceUsageType::Dynamic);
 			}
 		}
-		GraphBuilder.Execute();
 	}
 }
 
@@ -834,13 +835,16 @@ void FHairStrandsRestRootResource::ReleaseRHI()
 
 /////////////////////////////////////////////////////////////////////////////////////////
 
-FHairStrandsDeformedRootResource::FHairStrandsDeformedRootResource(EHairStrandsResourcesType InCurveType):CurveType(InCurveType)
+FHairStrandsDeformedRootResource::FHairStrandsDeformedRootResource(EHairStrandsResourcesType InCurveType):
+	FHairCommonResource(EHairStrandsAllocationType::Deferred),
+	CurveType(InCurveType)
 {
 
 }
 
 FHairStrandsDeformedRootResource::FHairStrandsDeformedRootResource(const FHairStrandsRestRootResource* InRestResources, EHairStrandsResourcesType InCurveType):
-CurveType(InCurveType)
+	FHairCommonResource(EHairStrandsAllocationType::Deferred),
+	CurveType(InCurveType)
 {
 	check(InRestResources);
 	uint32 LODIndex = 0;
@@ -856,13 +860,8 @@ CurveType(InCurveType)
 	}
 }
 
-void FHairStrandsDeformedRootResource::InitRHI()
+void FHairStrandsDeformedRootResource::InternalAllocate(FRDGBuilder& GraphBuilder)
 {
-	if (GUsingNullRHI) { return; }
-	FMemMark Mark(FMemStack::Get());
-	FRHICommandListImmediate& RHICmdList = FRHICommandListExecutor::GetImmediateCommandList();
-	FRDGBuilder GraphBuilder(RHICmdList);
-
 	if (RootCount > 0)
 	{
 		for (FLOD& LOD : LODs)
@@ -876,7 +875,6 @@ void FHairStrandsDeformedRootResource::InitRHI()
 			InternalCreateVertexBufferRDG<FHairStrandsMeshTrianglePositionFormat>(GraphBuilder, RootCount, LOD.DeformedRootTrianglePosition2Buffer, HAIRSTRANDS_RESOUCE_NAME(CurveType, Hair.StrandsRootDeformed_DeformedRootTrianglePosition2Buffer), EHairResourceUsageType::Dynamic);
 		}
 	}
-	GraphBuilder.Execute();
 }
 
 void FHairStrandsDeformedRootResource::ReleaseRHI()
@@ -1027,6 +1025,7 @@ void FHairStrandsRootData::Reset()
 /////////////////////////////////////////////////////////////////////////////////////////
 
 FHairStrandsInterpolationResource::FHairStrandsInterpolationResource(const FHairStrandsInterpolationDatas::FRenderData& InterpolationRenderData, const FHairStrandsDatas& SimDatas) :
+	FHairCommonResource(EHairStrandsAllocationType::Immediate),
 	Interpolation0Buffer(), Interpolation1Buffer(), RenderData(InterpolationRenderData)
 {
 	const uint32 RootCount = SimDatas.GetNumCurves();
@@ -1042,17 +1041,11 @@ FHairStrandsInterpolationResource::FHairStrandsInterpolationResource(const FHair
 	}
 }
 
-void FHairStrandsInterpolationResource::InitRHI()
+void FHairStrandsInterpolationResource::InternalAllocate(FRDGBuilder& GraphBuilder)
 {
-	if (GUsingNullRHI) { return; }
-	FMemMark Mark(FMemStack::Get());
-	FRHICommandListImmediate& RHICmdList = FRHICommandListExecutor::GetImmediateCommandList();
-	FRDGBuilder GraphBuilder(RHICmdList);
-
 	InternalCreateVertexBufferRDG<FHairStrandsInterpolation0Format>(GraphBuilder, RenderData.Interpolation0, Interpolation0Buffer, TEXT("Hair.StrandsInterpolation_Interpolation0Buffer"), EHairResourceUsageType::Static);
 	InternalCreateVertexBufferRDG<FHairStrandsInterpolation1Format>(GraphBuilder, RenderData.Interpolation1, Interpolation1Buffer, TEXT("Hair.StrandsInterpolation_Interpolation1Buffer"), EHairResourceUsageType::Static);
 	InternalCreateVertexBufferRDG<FHairStrandsRootIndexFormat>(GraphBuilder, SimRootPointIndex, SimRootPointIndexBuffer, TEXT("Hair.StrandsInterpolation_SimRootPointIndex"), EHairResourceUsageType::Static);
-	GraphBuilder.Execute();
 	//SimRootPointIndex.SetNum(0);
 }
 
@@ -1089,19 +1082,14 @@ void FHairCardsInterpolationDatas::Reset()
 }
 
 FHairCardsInterpolationResource::FHairCardsInterpolationResource(const FHairCardsInterpolationDatas::FRenderData& InterpolationRenderData) :
+	FHairCommonResource(EHairStrandsAllocationType::Immediate),
 	InterpolationBuffer(), RenderData(InterpolationRenderData)
 {
 }
 
-void FHairCardsInterpolationResource::InitRHI()
+void FHairCardsInterpolationResource::InternalAllocate(FRDGBuilder& GraphBuilder)
 {
-	if (GUsingNullRHI) { return; }
-	FMemMark Mark(FMemStack::Get());
-	FRHICommandListImmediate& RHICmdList = FRHICommandListExecutor::GetImmediateCommandList();
-	FRDGBuilder GraphBuilder(RHICmdList);
-
 	InternalCreateVertexBufferRDG<FHairCardsInterpolationFormat>(GraphBuilder, RenderData.Interpolation, InterpolationBuffer, TEXT("Hair.CardsInterpolation_InterpolationBuffer"), EHairResourceUsageType::Static);
-	GraphBuilder.Execute();
 }
 
 void FHairCardsInterpolationResource::ReleaseRHI()
@@ -1115,26 +1103,23 @@ void FHairCardsInterpolationResource::ReleaseRHI()
 // RT geometry is built to for a cross around the fiber.
 // 4 triangles per hair vertex => 12 vertices per hair vertex
 FHairStrandsRaytracingResource::FHairStrandsRaytracingResource(const FHairStrandsDatas& InData) :
+	FHairCommonResource(EHairStrandsAllocationType::Deferred),
 	PositionBuffer(), VertexCount(InData.GetNumPoints()*12)  
 {}
 
 FHairStrandsRaytracingResource::FHairStrandsRaytracingResource(const FHairCardsDatas& InData) :
+	FHairCommonResource(EHairStrandsAllocationType::Deferred),
 	PositionBuffer(), VertexCount(InData.Cards.GetNumVertices())
 {}
 
 FHairStrandsRaytracingResource::FHairStrandsRaytracingResource(const FHairMeshesDatas& InData) :
+	FHairCommonResource(EHairStrandsAllocationType::Deferred),
 	PositionBuffer(), VertexCount(InData.Meshes.GetNumVertices())
 {}
 
-void FHairStrandsRaytracingResource::InitRHI()
+void FHairStrandsRaytracingResource::InternalAllocate(FRDGBuilder& GraphBuilder)
 {
-	if (GUsingNullRHI) { return; }
-	FMemMark Mark(FMemStack::Get());
-	FRHICommandListImmediate& RHICmdList = FRHICommandListExecutor::GetImmediateCommandList();
-	FRDGBuilder GraphBuilder(RHICmdList);
-
 	InternalCreateVertexBufferRDG<FHairStrandsRaytracingFormat>(GraphBuilder, VertexCount, PositionBuffer, TEXT("Hair.StrandsRaytracing_PositionBuffer"), EHairResourceUsageType::Dynamic);
-	GraphBuilder.Execute();
 }
 
 void FHairStrandsRaytracingResource::ReleaseRHI()
diff --git a/Engine/Plugins/Runtime/HairStrands/Source/HairStrandsCore/Private/HairCardsVertexFactory.cpp b/Engine/Plugins/Runtime/HairStrands/Source/HairStrandsCore/Private/HairCardsVertexFactory.cpp
--- a/Engine/Plugins/Runtime/HairStrands/Source/HairStrandsCore/Private/HairCardsVertexFactory.cpp
+++ b/Engine/Plugins/Runtime/HairStrands/Source/HairStrandsCore/Private/HairCardsVertexFactory.cpp
@@ -258,8 +258,14 @@ void FHairCardsVertexFactory::Copy(const FHairCardsVertexFactory& Other)
 	BeginUpdateResourceRHI(this);
 }
 
-void FHairCardsVertexFactory::InitRHI()
+void FHairCardsVertexFactory::InitResources()
 {
+	if (bIsInitialized)
+		return;
+
+	FVertexFactory::InitResource(); //Call VertexFactory/RenderResources::InitResource() to mark the resource as initialized();
+
+	bIsInitialized = true;
 	bNeedsDeclaration = true;
 	// We create different streams based on feature level
 	check(HasValidFeatureLevel());
@@ -329,6 +335,11 @@ void FHairCardsVertexFactory::InitRHI()
 	}
 }
 
+void FHairCardsVertexFactory::InitRHI()
+{
+	// Nothing as the initialization is done when needed
+}
+
 IMPLEMENT_VERTEX_FACTORY_PARAMETER_TYPE(FHairCardsVertexFactory, SF_Vertex,		FHairCardsVertexFactoryShaderParameters);
 IMPLEMENT_VERTEX_FACTORY_PARAMETER_TYPE(FHairCardsVertexFactory, SF_Pixel,		FHairCardsVertexFactoryShaderParameters);
 IMPLEMENT_VERTEX_FACTORY_PARAMETER_TYPE(FHairCardsVertexFactory, SF_Compute,	FHairCardsVertexFactoryShaderParameters);
diff --git a/Engine/Plugins/Runtime/HairStrands/Source/HairStrandsCore/Private/HairStrandsRendering.cpp b/Engine/Plugins/Runtime/HairStrands/Source/HairStrandsCore/Private/HairStrandsRendering.cpp
--- a/Engine/Plugins/Runtime/HairStrands/Source/HairStrandsCore/Private/HairStrandsRendering.cpp
+++ b/Engine/Plugins/Runtime/HairStrands/Source/HairStrandsCore/Private/HairStrandsRendering.cpp
@@ -1270,7 +1270,7 @@ void RegisterClusterData(FHairGroupInstance* Instance, FHairStrandClusterData* I
 	HairGroupCluster.HairGroupPublicPtr->ClusterDataIndex = ClusterDataGroupIndex;
 }
 
-EHairStrandsDebugMode GetHairStrandsGeometryDebugMode(FHairGroupInstance* Instance)
+EHairStrandsDebugMode GetHairStrandsGeometryDebugMode(const FHairGroupInstance* Instance)
 {
 	// 1. Use the per-instance debug mode by default, and fallback on the global debug otherwise
 	EHairStrandsDebugMode DebugMode = Instance->Debug.DebugMode != EHairStrandsDebugMode::NoneDebug ? Instance->Debug.DebugMode : GetHairStrandsDebugStrandsMode();
@@ -1306,7 +1306,7 @@ static void ConvertHairStrandsVFParameters(
 #define CONVERT_HAIRSSTRANDS_VF_PARAMETERS(OutName, ExternalBuffer) ConvertHairStrandsVFParameters(GraphBuilder, OutName, OutName##RHISRV, ExternalBuffer);
 
 // Compute/Update the hair strands description which will be used for rendering (VF) / voxelization & co.
-static FHairGroupPublicData::FVertexFactoryInput InternalComputeHairStrandsVertexInputData(FRDGBuilder* GraphBuilder, FHairGroupInstance* Instance)
+static FHairGroupPublicData::FVertexFactoryInput InternalComputeHairStrandsVertexInputData(FRDGBuilder* GraphBuilder, const FHairGroupInstance* Instance)
 {
 	FHairGroupPublicData::FVertexFactoryInput OutVFInput;
 	if (!Instance || Instance->GeometryType != EHairGeometryType::Strands)
@@ -1338,7 +1338,7 @@ static FHairGroupPublicData::FVertexFactoryInput InternalComputeHairStrandsVerte
 		CONVERT_HAIRSSTRANDS_VF_PARAMETERS(OutVFInput.Strands.PositionBuffer,			Instance->Strands.DeformedResource->GetBuffer(FHairStrandsDeformedResource::EFrameType::Current));
 		CONVERT_HAIRSSTRANDS_VF_PARAMETERS(OutVFInput.Strands.PrevPositionBuffer,		Instance->Strands.DeformedResource->GetBuffer(FHairStrandsDeformedResource::EFrameType::Previous));
 		CONVERT_HAIRSSTRANDS_VF_PARAMETERS(OutVFInput.Strands.TangentBuffer,			Instance->Strands.DeformedResource->TangentBuffer);
-		CONVERT_HAIRSSTRANDS_VF_PARAMETERS(OutVFInput.Strands.AttributeBuffer,			bDebugModePatchedAttributeBuffer ? Instance->Strands.DebugAttributeBuffer : Instance->Strands.RestResource->AttributeBuffer);
+		CONVERT_HAIRSSTRANDS_VF_PARAMETERS(OutVFInput.Strands.AttributeBuffer,			/*bDebugModePatchedAttributeBuffer ? Instance->Strands.DebugAttributeBuffer : */Instance->Strands.RestResource->AttributeBuffer); // TODO
 		CONVERT_HAIRSSTRANDS_VF_PARAMETERS(OutVFInput.Strands.MaterialBuffer,			Instance->Strands.RestResource->MaterialBuffer);
 		CONVERT_HAIRSSTRANDS_VF_PARAMETERS(OutVFInput.Strands.PositionOffsetBuffer,		Instance->Strands.DeformedResource->GetPositionOffsetBuffer(FHairStrandsDeformedResource::EFrameType::Current));
 		CONVERT_HAIRSSTRANDS_VF_PARAMETERS(OutVFInput.Strands.PrevPositionOffsetBuffer,	Instance->Strands.DeformedResource->GetPositionOffsetBuffer(FHairStrandsDeformedResource::EFrameType::Previous));
@@ -1374,7 +1374,7 @@ static FHairGroupPublicData::FVertexFactoryInput InternalComputeHairStrandsVerte
 	return OutVFInput;
 }
 
-FHairGroupPublicData::FVertexFactoryInput ComputeHairStrandsVertexInputData(FHairGroupInstance* Instance)
+FHairGroupPublicData::FVertexFactoryInput ComputeHairStrandsVertexInputData(const FHairGroupInstance* Instance)
 {
 	return InternalComputeHairStrandsVertexInputData(nullptr, Instance);
 }
diff --git a/Engine/Plugins/Runtime/HairStrands/Source/HairStrandsCore/Private/HairStrandsVertexFactory.cpp b/Engine/Plugins/Runtime/HairStrands/Source/HairStrandsCore/Private/HairStrandsVertexFactory.cpp
--- a/Engine/Plugins/Runtime/HairStrands/Source/HairStrandsCore/Private/HairStrandsVertexFactory.cpp
+++ b/Engine/Plugins/Runtime/HairStrands/Source/HairStrandsCore/Private/HairStrandsVertexFactory.cpp
@@ -68,7 +68,7 @@ inline float GetHairDensity(const FHairGroupPublicData::FVertexFactoryInput& VFI
 
 /////////////////////////////////////////////////////////////////////////////////////////
 
-FHairGroupPublicData::FVertexFactoryInput ComputeHairStrandsVertexInputData(FHairGroupInstance* Instance);
+FHairGroupPublicData::FVertexFactoryInput ComputeHairStrandsVertexInputData(const FHairGroupInstance* Instance);
 
 class FHairStrandsVertexFactoryShaderParameters : public FVertexFactoryShaderParameters
 {
@@ -135,11 +135,7 @@ public:
 	) const
 	{
 		const FHairStrandsVertexFactory* VF = static_cast<const FHairStrandsVertexFactory*>(VertexFactory);
-
-		const FHairGroupPublicData* GroupPublicData = reinterpret_cast<const FHairGroupPublicData*>(BatchElement.VertexFactoryUserData);
-		check(GroupPublicData);
-		const uint64 GroupIndex = GroupPublicData->GetGroupIndex();
-		const FHairGroupPublicData::FVertexFactoryInput VFInput = ComputeHairStrandsVertexInputData(VF->Data.Instances[GroupIndex]);
+		const FHairGroupPublicData::FVertexFactoryInput VFInput = ComputeHairStrandsVertexInputData(VF->Data.Instance);
 
 		VFS_BindParam(ShaderBindings, PositionBuffer, GetPositionSRV(VFInput));
 		VFS_BindParam(ShaderBindings, PreviousPositionBuffer, GetPreviousPositionSRV(VFInput));
@@ -159,6 +155,7 @@ public:
 		FShaderResourceViewRHIRef CulledDispatchVertexIdsSRV = GDummyCulledDispatchVertexIdsBuffer.SRVUint;
 		FShaderResourceViewRHIRef CulledCompactedRadiusScaleBufferSRV = GDummyCulledDispatchVertexIdsBuffer.SRVFloat;
 
+		const FHairGroupPublicData* GroupPublicData = VF->Data.Instance->HairGroupPublicData;
 		const bool bCulling = GroupPublicData->GetCullingResultAvailable();
 		if (bCulling)
 		{
@@ -227,8 +224,14 @@ void FHairStrandsVertexFactory::Copy(const FHairStrandsVertexFactory& Other)
 	BeginUpdateResourceRHI(this);
 }
 
-void FHairStrandsVertexFactory::InitRHI()
+void FHairStrandsVertexFactory::InitResources()
 {
+	if (bIsInitialized)
+		return;
+
+	FVertexFactory::InitResource(); //Call VertexFactory/RenderResources::InitResource() to mark the resource as initialized();
+
+	bIsInitialized = true;
 	bNeedsDeclaration = false;
 	bSupportsManualVertexFetch = true;
 
@@ -265,6 +268,11 @@ IMPLEMENT_VERTEX_FACTORY_PARAMETER_TYPE(FHairStrandsVertexFactory, SF_Compute,
 IMPLEMENT_VERTEX_FACTORY_PARAMETER_TYPE(FHairStrandsVertexFactory, SF_RayHitGroup,	FHairStrandsVertexFactoryShaderParameters);
 #endif
 
+void FHairStrandsVertexFactory::InitRHI()
+{
+	// Nothing as the initialize runs only on first use
+}
+
 void FHairStrandsVertexFactory::ReleaseRHI()
 {
 	FVertexFactory::ReleaseRHI();
diff --git a/Engine/Plugins/Runtime/HairStrands/Source/HairStrandsCore/Public/GroomInstance.h b/Engine/Plugins/Runtime/HairStrands/Source/HairStrandsCore/Public/GroomInstance.h
--- a/Engine/Plugins/Runtime/HairStrands/Source/HairStrandsCore/Public/GroomInstance.h
+++ b/Engine/Plugins/Runtime/HairStrands/Source/HairStrandsCore/Public/GroomInstance.h
@@ -10,6 +10,8 @@
 #include "HairStrandsInterface.h"
 
 class UMeshComponent;
+class FHairCardsVertexFactory;
+class FHairStrandsVertexFactory;
 
 // @hair_todo: pack card ID + card UV in 32Bits alpha channel's of the position buffer:
 //  * 10/10 bits for UV -> max 1024/1024 rect resolution
@@ -101,6 +103,7 @@ struct HAIRSTRANDSCORE_API FHairGroupInstance : public FHairStrandsInstance
 		FHairGroupInstanceModifer Modifier;
 
 		UMaterialInterface* Material = nullptr;
+		FHairStrandsVertexFactory* VertexFactory = nullptr;
 	} Strands;
 
 	//////////////////////////////////////////////////////////////////////////////////////////
@@ -132,6 +135,13 @@ struct HAIRSTRANDSCORE_API FHairGroupInstance : public FHairStrandsInstance
 
 			UMaterialInterface* Material = nullptr;
 			FHairCardsUniformBuffer UniformBuffer[2];
+			FHairCardsVertexFactory* VertexFactory[2] = { nullptr, nullptr };
+			FHairCardsVertexFactory* GetVertexFactory() const
+			{
+				const uint32 Index = DeformedResource ? DeformedResource->CurrentIndex : 0;
+				return VertexFactory[Index];
+			} 
+
 		};
 		TArray<FLOD> LODs;
 	} Cards;
@@ -159,6 +169,12 @@ struct HAIRSTRANDSCORE_API FHairGroupInstance : public FHairStrandsInstance
 
 			UMaterialInterface* Material = nullptr;
 			FHairCardsUniformBuffer UniformBuffer[2];
+			FHairCardsVertexFactory* VertexFactory[2] = { nullptr, nullptr };
+			FHairCardsVertexFactory* GetVertexFactory() const 
+			{
+				const uint32 Index = DeformedResource ? DeformedResource->CurrentIndex : 0;
+				return VertexFactory[Index];
+			}
 		};
 		TArray<FLOD> LODs;
 	} Meshes;
diff --git a/Engine/Plugins/Runtime/HairStrands/Source/HairStrandsCore/Public/GroomResources.h b/Engine/Plugins/Runtime/HairStrands/Source/HairStrandsCore/Public/GroomResources.h
--- a/Engine/Plugins/Runtime/HairStrands/Source/HairStrandsCore/Public/GroomResources.h
+++ b/Engine/Plugins/Runtime/HairStrands/Source/HairStrandsCore/Public/GroomResources.h
@@ -23,15 +23,38 @@ enum class EHairStrandsResourcesType : uint8
 	Cards		// Guides used for deforming the cards geometry (which is different from the simulation guides)
 };
 
+enum class EHairStrandsAllocationType : uint8
+{
+	Immediate,	// Resources are allocated immediately
+	Deferred	// Resources allocation is deferred to first usage
+};
+
+/* Hair resouces which whom allocation can be deferred */
+struct FHairCommonResource : public FRenderResource
+{
+	/** Build the hair strands resource */
+	FHairCommonResource(EHairStrandsAllocationType AllocationType, bool bUseRenderGraph=true);
+
+	/* Init the buffer */
+	virtual void InitRHI() override;
+	void Allocate(FRDGBuilder& GraphBuilder);
+	virtual void InternalAllocate() {};
+	virtual void InternalAllocate(FRDGBuilder& GraphBuilder) {};
+
+	bool bUseRenderGraph = true;
+	bool bIsInitialized = false;
+	EHairStrandsAllocationType AllocationType = EHairStrandsAllocationType::Deferred;
+};
+
 /* Render buffers for root deformation for dynamic meshes */
-struct FHairStrandsRestRootResource : public FRenderResource
+struct FHairStrandsRestRootResource : public FHairCommonResource
 {
 	/** Build the hair strands resource */
 	FHairStrandsRestRootResource(const FHairStrandsRootData& RootData, EHairStrandsResourcesType CurveType);
 	FHairStrandsRestRootResource(const FHairStrandsDatas* HairStrandsDatas, uint32 LODCount, const TArray<uint32>& NumSamples, EHairStrandsResourcesType CurveType);
 
 	/* Init the buffer */
-	virtual void InitRHI() override;
+	void InternalAllocate(FRDGBuilder& GraphBuilder) override;
 
 	/* Release the buffer */
 	virtual void ReleaseRHI() override;
@@ -104,14 +127,14 @@ struct FHairStrandsRestRootResource : public FRenderResource
 };
 
 /* Render buffers for root deformation for dynamic meshes */
-struct FHairStrandsDeformedRootResource : public FRenderResource
+struct FHairStrandsDeformedRootResource : public FHairCommonResource
 {
 	/** Build the hair strands resource */
 	FHairStrandsDeformedRootResource(EHairStrandsResourcesType CurveType);
 	FHairStrandsDeformedRootResource(const FHairStrandsRestRootResource* InRestResources, EHairStrandsResourcesType CurveType);
 
 	/* Init the buffer */
-	virtual void InitRHI() override;
+	void InternalAllocate(FRDGBuilder& GraphBuilder) override;
 
 	/* Release the buffer */
 	virtual void ReleaseRHI() override;
@@ -172,13 +195,13 @@ struct FHairStrandsDeformedRootResource : public FRenderResource
 };
 
 /* Render buffers that will be used for rendering */
-struct FHairStrandsRestResource : public FRenderResource
+struct FHairStrandsRestResource : public FHairCommonResource
 {
 	/** Build the hair strands resource */
 	FHairStrandsRestResource(const FHairStrandsDatas::FRenderData& HairStrandRenderData, const FVector& PositionOffset, EHairStrandsResourcesType CurveType);
 
 	/* Init the buffer */
-	virtual void InitRHI() override;
+	void InternalAllocate(FRDGBuilder& GraphBuilder) override;
 
 	/* Release the buffer */
 	virtual void ReleaseRHI() override;
@@ -227,13 +250,13 @@ struct FHairStrandsRestResource : public FRenderResource
 	inline uint32 GetVertexCount() const { return RenderData.Positions.Num() / FHairStrandsPositionFormat::ComponentCount; }
 };
 
-struct FHairStrandsDeformedResource : public FRenderResource
+struct FHairStrandsDeformedResource : public FHairCommonResource
 {
 	/** Build the hair strands resource */
 	FHairStrandsDeformedResource(const FHairStrandsDatas::FRenderData& HairStrandRenderData, bool bInitializeData, const FVector& InDefaultOffset, EHairStrandsResourcesType CurveType);
 
 	/* Init the buffer */
-	virtual void InitRHI() override;
+	void InternalAllocate(FRDGBuilder& GraphBuilder) override;
 
 	/* Release the buffer */
 	virtual void ReleaseRHI() override;
@@ -294,12 +317,12 @@ struct FHairStrandsDeformedResource : public FRenderResource
 	//bool NeedsToUpdateTangent();
 };
 
-struct FHairStrandsClusterCullingResource : public FRenderResource
+struct FHairStrandsClusterCullingResource : public FHairCommonResource
 {
 	FHairStrandsClusterCullingResource(const FHairStrandsClusterCullingData& Data);
 
 	/* Init the buffer */
-	virtual void InitRHI() override;
+	void InternalAllocate(FRDGBuilder& GraphBuilder) override;
 
 	/* Release the buffer */
 	virtual void ReleaseRHI() override;
@@ -331,13 +354,13 @@ struct FHairStrandsClusterCullingResource : public FRenderResource
 	const FHairStrandsClusterCullingData Data;
 };
 
-struct FHairStrandsInterpolationResource : public FRenderResource
+struct FHairStrandsInterpolationResource : public FHairCommonResource
 {
 	/** Build the hair strands resource */
 	FHairStrandsInterpolationResource(const FHairStrandsInterpolationDatas::FRenderData& InterpolationRenderData, const FHairStrandsDatas& SimDatas);
 
 	/* Init the buffer */
-	virtual void InitRHI() override;
+	void InternalAllocate(FRDGBuilder& GraphBuilder) override;
 
 	/* Release the buffer */
 	virtual void ReleaseRHI() override;
@@ -367,7 +390,7 @@ struct FHairStrandsInterpolationResource : public FRenderResource
 };
 
 #if RHI_RAYTRACING
-struct FHairStrandsRaytracingResource : public FRenderResource
+struct FHairStrandsRaytracingResource : public FHairCommonResource
 {
 	/** Build the hair strands resource */
 	FHairStrandsRaytracingResource(const FHairStrandsDatas& HairStrandsDatas);
@@ -375,7 +398,7 @@ struct FHairStrandsRaytracingResource : public FRenderResource
 	FHairStrandsRaytracingResource(const FHairMeshesDatas& InData);
 
 	/* Init the buffer */
-	virtual void InitRHI() override;
+	void InternalAllocate(FRDGBuilder& GraphBuilder) override;
 
 	/* Release the buffer */
 	virtual void ReleaseRHI() override;
@@ -416,7 +439,7 @@ public:
 };
 
 /* Render buffers that will be used for rendering */
-struct FHairCardsRestResource : public FRenderResource
+struct FHairCardsRestResource : public FHairCommonResource
 {
 	/** Build the hair strands resource */
 	FHairCardsRestResource(const FHairCardsDatas::FRenderData& HairCardsRenderData, uint32 VertexCount, uint32 PrimitiveCount);
@@ -425,7 +448,7 @@ struct FHairCardsRestResource : public FRenderResource
 	virtual void ReleaseResource() override;
 
 	/* Init the buffer */
-	virtual void InitRHI() override;
+	void InternalAllocate() override;
 
 	/* Release the buffer */
 	virtual void ReleaseRHI() override;
@@ -471,13 +494,13 @@ struct FHairCardsRestResource : public FRenderResource
 };
 
 /* Render buffers that will be used for rendering */
-struct FHairCardsProceduralResource : public FRenderResource
+struct FHairCardsProceduralResource : public FHairCommonResource
 {
 	/** Build the hair strands resource */
 	FHairCardsProceduralResource(const FHairCardsProceduralDatas::FRenderData& HairCardsRenderData, const FIntPoint& AtlasResolution, const FHairCardsVoxel& InVoxel);
 
 	/* Init the buffer */
-	virtual void InitRHI() override;
+	void InternalAllocate(FRDGBuilder& GraphBuilder) override;
 
 	/* Release the buffer */
 	virtual void ReleaseRHI() override;
@@ -512,13 +535,13 @@ struct FHairCardsProceduralResource : public FRenderResource
 	const FHairCardsProceduralDatas::FRenderData& RenderData;
 };
 
-struct FHairCardsDeformedResource : public FRenderResource
+struct FHairCardsDeformedResource : public FHairCommonResource
 {
 	/** Build the hair strands resource */
 	FHairCardsDeformedResource(const FHairCardsDatas::FRenderData& HairStrandRenderData, bool bInitializeData);
 
 	/* Init the buffer */
-	virtual void InitRHI() override;
+	void InternalAllocate(FRDGBuilder& GraphBuilder) override;
 
 	/* Release the buffer */
 	virtual void ReleaseRHI() override;
@@ -588,13 +611,13 @@ struct HAIRSTRANDSCORE_API FHairCardsInterpolationDatas
 
 FArchive& operator<<(FArchive& Ar, FHairCardsInterpolationDatas& CardInterpData);
 
-struct FHairCardsInterpolationResource : public FRenderResource
+struct FHairCardsInterpolationResource : public FHairCommonResource
 {
 	/** Build the hair strands resource */
 	FHairCardsInterpolationResource(const FHairCardsInterpolationDatas::FRenderData& InterpolationRenderData);
 
 	/* Init the buffer */
-	virtual void InitRHI() override;
+	void InternalAllocate(FRDGBuilder& GraphBuilder) override;
 
 	/* Release the buffer */
 	virtual void ReleaseRHI() override;
@@ -620,7 +643,7 @@ struct FHairCardsInterpolationResource : public FRenderResource
 // Meshes
 
 /* Render buffers that will be used for rendering */
-struct FHairMeshesRestResource : public FRenderResource
+struct FHairMeshesRestResource : public FHairCommonResource
 {
 	/** Build the hair strands resource */
 	FHairMeshesRestResource(const FHairMeshesDatas::FRenderData& HairMeshesRenderData, uint32 VertexCount, uint32 PrimitiveCount);
@@ -629,7 +652,7 @@ struct FHairMeshesRestResource : public FRenderResource
 	virtual void ReleaseResource() override;
 
 	/* Init the buffer */
-	virtual void InitRHI() override;
+	void InternalAllocate() override;
 
 	/* Release the buffer */
 	virtual void ReleaseRHI() override;
@@ -675,13 +698,13 @@ struct FHairMeshesRestResource : public FRenderResource
 
 
 /* Render buffers that will be used for rendering */
-struct FHairMeshesDeformedResource : public FRenderResource
+struct FHairMeshesDeformedResource : public FHairCommonResource
 {
 	/** Build the hair strands resource */
 	FHairMeshesDeformedResource(const FHairMeshesDatas::FRenderData& HairMeshesRenderData, bool bInInitializedData);
 
 	/* Init the buffer */
-	virtual void InitRHI() override;
+	void InternalAllocate(FRDGBuilder& GraphBuilder) override;
 
 	/* Release the buffer */
 	virtual void ReleaseRHI() override;
diff --git a/Engine/Plugins/Runtime/HairStrands/Source/HairStrandsCore/Public/HairCardsVertexFactory.h b/Engine/Plugins/Runtime/HairStrands/Source/HairStrandsCore/Public/HairCardsVertexFactory.h
--- a/Engine/Plugins/Runtime/HairStrands/Source/HairStrandsCore/Public/HairCardsVertexFactory.h
+++ b/Engine/Plugins/Runtime/HairStrands/Source/HairStrandsCore/Public/HairCardsVertexFactory.h
@@ -52,6 +52,8 @@ public:
 	*/
 	void Copy(const FHairCardsVertexFactory& Other);
 
+	void InitResources();
+
 	// FRenderResource interface.
 	virtual void InitRHI() override;
 	virtual void ReleaseRHI() override;
@@ -59,6 +61,7 @@ public:
 	FDataType Data;
 protected:
 
+	bool bIsInitialized = false;
 
 	struct FDebugName
 	{
diff --git a/Engine/Plugins/Runtime/HairStrands/Source/HairStrandsCore/Public/HairStrandsVertexFactory.h b/Engine/Plugins/Runtime/HairStrands/Source/HairStrandsCore/Public/HairStrandsVertexFactory.h
--- a/Engine/Plugins/Runtime/HairStrands/Source/HairStrandsCore/Public/HairStrandsVertexFactory.h
+++ b/Engine/Plugins/Runtime/HairStrands/Source/HairStrandsCore/Public/HairStrandsVertexFactory.h
@@ -26,15 +26,15 @@ public:
 	
 	struct FDataType
 	{
-		TArray<FHairGroupInstance*> Instances;
+		const FHairGroupInstance* Instance = nullptr;
 	};
 
-	FHairStrandsVertexFactory(const TArray<FHairGroupInstance*>& InitialData, ERHIFeatureLevel::Type InFeatureLevel, const char* InDebugName)
+	FHairStrandsVertexFactory(const FHairGroupInstance* Instance, ERHIFeatureLevel::Type InFeatureLevel, const char* InDebugName)
 		: FVertexFactory(InFeatureLevel)
 		, DebugName(InDebugName)
 	{
 		bSupportsManualVertexFetch = true;
-		Data.Instances = InitialData;
+		Data.Instance = Instance;
 	}
 
 	/**
@@ -58,10 +58,13 @@ public:
 	// FRenderResource interface.
 	virtual void InitRHI() override;
 	virtual void ReleaseRHI() override;
+	void InitResources();
 
 	const FDataType& GetData() const { return Data; }
 	FDataType Data;
 protected:
+	bool bIsInitialized = false;
+
 	struct FDebugName
 	{
 		FDebugName(const char* InDebugName)
diff --git a/Engine/Source/Runtime/Renderer/Private/SceneVisibility.cpp b/Engine/Source/Runtime/Renderer/Private/SceneVisibility.cpp
--- a/Engine/Source/Runtime/Renderer/Private/SceneVisibility.cpp
+++ b/Engine/Source/Runtime/Renderer/Private/SceneVisibility.cpp
@@ -3151,7 +3151,7 @@ void FSceneRenderer::PreVisibilityFrameSetup(FRDGBuilder& GraphBuilder, const FS
 		FHairStrandsBookmarkParameters Parameters = CreateHairStrandsBookmarkParameters(Scene, Views);
 		if (Parameters.bHasElements)
 		{
-			RunHairStrandsBookmark(EHairStrandsBookmark::ProcessLODSelection, Parameters);
+			RunHairStrandsBookmark(GraphBuilder, EHairStrandsBookmark::ProcessLODSelection, Parameters);
 		}
 	}
 
