Remove unecessary copy during groom asset loading. Save 500ms on MassCrowd.

diff --git a/Engine/Plugins/Runtime/HairStrands/Source/HairStrandsCore/Private/GroomResources.cpp b/Engine/Plugins/Runtime/HairStrands/Source/HairStrandsCore/Private/GroomResources.cpp
--- a/Engine/Plugins/Runtime/HairStrands/Source/HairStrandsCore/Private/GroomResources.cpp
+++ b/Engine/Plugins/Runtime/HairStrands/Source/HairStrandsCore/Private/GroomResources.cpp
@@ -719,18 +719,11 @@ FHairStrandsRestRootResource::FHairStrandsRestRootResource(const FHairStrandsRoo
 	PopulateFromRootData();
 }
 
-FHairStrandsRestRootResource::FHairStrandsRestRootResource(const FHairStrandsDatas* HairStrandsDatas, uint32 LODCount, const TArray<uint32>& NumSamples, EHairStrandsResourcesType InCurveType):
-	FHairCommonResource(EHairStrandsAllocationType::Immediate),
-	RootData(HairStrandsDatas, LODCount, NumSamples), CurveType(InCurveType)
-{
-	PopulateFromRootData();
-}
-
 void FHairStrandsRestRootResource::PopulateFromRootData()
 {
 	uint32 LODIndex = 0;
 	LODs.Reserve(RootData.MeshProjectionLODs.Num());
-	for (FHairStrandsRootData::FMeshProjectionLOD& MeshProjectionLOD : RootData.MeshProjectionLODs)
+	for (const FHairStrandsRootData::FMeshProjectionLOD& MeshProjectionLOD : RootData.MeshProjectionLODs)
 	{
 		FLOD& LOD = LODs.AddDefaulted_GetRef();
 
diff --git a/Engine/Plugins/Runtime/HairStrands/Source/HairStrandsCore/Public/GroomResources.h b/Engine/Plugins/Runtime/HairStrands/Source/HairStrandsCore/Public/GroomResources.h
--- a/Engine/Plugins/Runtime/HairStrands/Source/HairStrandsCore/Public/GroomResources.h
+++ b/Engine/Plugins/Runtime/HairStrands/Source/HairStrandsCore/Public/GroomResources.h
@@ -51,7 +51,6 @@ struct FHairStrandsRestRootResource : public FHairCommonResource
 {
 	/** Build the hair strands resource */
 	FHairStrandsRestRootResource(const FHairStrandsRootData& RootData, EHairStrandsResourcesType CurveType);
-	FHairStrandsRestRootResource(const FHairStrandsDatas* HairStrandsDatas, uint32 LODCount, const TArray<uint32>& NumSamples, EHairStrandsResourcesType CurveType);
 
 	/* Init the buffer */
 	void InternalAllocate(FRDGBuilder& GraphBuilder) override;
@@ -120,7 +119,7 @@ struct FHairStrandsRestRootResource : public FHairCommonResource
 	TArray<FLOD> LODs;
 
 	/* Store CPU data for root info & root binding */
-	FHairStrandsRootData RootData;
+	const FHairStrandsRootData& RootData;
 
 	/* Type of curves */
 	const EHairStrandsResourcesType CurveType;
