Fix crash with groom lazy allocation when updating/reloading resources.

diff --git a/Engine/Plugins/Runtime/HairStrands/Source/HairStrandsCore/Private/GroomComponent.cpp b/Engine/Plugins/Runtime/HairStrands/Source/HairStrandsCore/Private/GroomComponent.cpp
--- a/Engine/Plugins/Runtime/HairStrands/Source/HairStrandsCore/Private/GroomComponent.cpp
+++ b/Engine/Plugins/Runtime/HairStrands/Source/HairStrandsCore/Private/GroomComponent.cpp
@@ -1801,8 +1801,6 @@ static UMeshComponent* ValidateBindingAsset(
 	return ValidateBindingAsset(GroomAsset, BindingAsset, Cast<UGeometryCacheComponent>(MeshComponent), bIsBindingReloading, bValidationEnable, Component);
 }
 
-void CreateHairStrandsDebugAttributeBuffer(FRDGExternalBuffer* DebugAttributeBuffer, uint32 VertexCount);
-
 static EGroomGeometryType GetEffectiveGeometryType(EGroomGeometryType Type, bool bUseCards)
 {
 	return Type == EGroomGeometryType::Strands && (!IsHairStrandsEnabled(EHairStrandsShaderType::Strands) || bUseCards) ? EGroomGeometryType::Cards : Type;
@@ -2092,12 +2090,6 @@ void UGroomComponent::InitResources(bool bIsBindingReloading)
 				HairGroupInstance->Strands.DeformedResource->GetPositionOffset(FHairStrandsDeformedResource::EFrameType::Previous) = HairGroupInstance->Strands.RestResource->PositionOffset;
 			} 
 
-			// Create an debug buffer for storing cluster visalization data. This is only used for debug purpose, hence only enable in editor build.
-			// Special case for debug mode were the attribute buffer is patch with some custom data to show hair properties (strands belonging to the same cluster, ...)
-			#if WITH_EDITOR
-			CreateHairStrandsDebugAttributeBuffer(&HairGroupInstance->Strands.DebugAttributeBuffer, HairGroupInstance->Strands.RestResource->GetVertexCount());
-			#endif
-
 			// An empty groom doesn't have a ClusterCullingResource
 			HairGroupInstance->Strands.ClusterCullingResource = GroupData.Strands.ClusterCullingResource;
 			if (HairGroupInstance->Strands.ClusterCullingResource)
diff --git a/Engine/Plugins/Runtime/HairStrands/Source/HairStrandsCore/Private/GroomResources.cpp b/Engine/Plugins/Runtime/HairStrands/Source/HairStrandsCore/Private/GroomResources.cpp
--- a/Engine/Plugins/Runtime/HairStrands/Source/HairStrandsCore/Private/GroomResources.cpp
+++ b/Engine/Plugins/Runtime/HairStrands/Source/HairStrandsCore/Private/GroomResources.cpp
@@ -293,21 +293,11 @@ static UTexture2D* CreateCardTexture(FIntPoint Resolution)
 }
 
 /////////////////////////////////////////////////////////////////////////////////////////
-void CreateHairStrandsDebugAttributeBuffer(FRDGExternalBuffer* DebugAttributeBuffer, uint32 VertexCount)
+void CreateHairStrandsDebugAttributeBuffer(FRDGBuilder& GraphBuilder, FRDGExternalBuffer* DebugAttributeBuffer, uint32 VertexCount)
 {
 	if (VertexCount == 0 || !DebugAttributeBuffer)
 		return;
-
-	FRDGExternalBuffer* LocalAttributeBuffer = DebugAttributeBuffer;
-	ENQUEUE_RENDER_COMMAND(FHairStrandsCreateDebugAttributeBuffer)(
-	[LocalAttributeBuffer, VertexCount](FRHICommandListImmediate& RHICmdList)
-	{
-		if (GUsingNullRHI) { return; }
-		FMemMark Mark(FMemStack::Get());
-		FRDGBuilder GraphBuilder(RHICmdList);
-		InternalCreateVertexBufferRDG<FHairStrandsAttributeFormat>(GraphBuilder, VertexCount, *LocalAttributeBuffer, TEXT("Hair.Strands_DebugAttributeBuffer"), EHairResourceUsageType::Dynamic);
-		GraphBuilder.Execute();
-	});
+	InternalCreateVertexBufferRDG<FHairStrandsAttributeFormat>(GraphBuilder, VertexCount, *DebugAttributeBuffer, TEXT("Hair.Strands_DebugAttributeBuffer"), EHairResourceUsageType::Dynamic);
 }
 
 /////////////////////////////////////////////////////////////////////////////////////////
@@ -338,6 +328,12 @@ void FHairCommonResource::InitRHI()
 	bIsInitialized = true;
 }
 
+void FHairCommonResource::ReleaseRHI()
+{
+	InternalRelease();
+	bIsInitialized = false;
+}
+
 void FHairCommonResource::Allocate(FRDGBuilder& GraphBuilder)
 {
 	if (bIsInitialized) { return; }
@@ -388,7 +384,7 @@ void FHairCardsRestResource::InternalAllocate()
 	AttributeSampler = DefaultSampler;
 }
 
-void FHairCardsRestResource::ReleaseRHI()
+void FHairCardsRestResource::InternalRelease()
 {
 }
 
@@ -444,7 +440,7 @@ void FHairCardsProceduralResource::InternalAllocate(FRDGBuilder& GraphBuilder)
 	InternalCreateVertexBufferRDG<FHairCardsStrandsAttributeFormat>(GraphBuilder, RenderData.CardsStrandsAttributes, CardsStrandsAttributes, TEXT("Hair.CardsProcedural_CardsStrandsAttributes"), EHairResourceUsageType::Static);
 }
 
-void FHairCardsProceduralResource::ReleaseRHI()
+void FHairCardsProceduralResource::InternalRelease()
 {
 	AtlasRectBuffer.Release();
 	LengthBuffer.Release();
@@ -482,7 +478,7 @@ void FHairCardsDeformedResource::InternalAllocate(FRDGBuilder& GraphBuilder)
 	}
 }
 
-void FHairCardsDeformedResource::ReleaseRHI()
+void FHairCardsDeformedResource::InternalRelease()
 {
 	DeformedPositionBuffer[0].Release();
 	DeformedPositionBuffer[1].Release();
@@ -512,7 +508,7 @@ void FHairMeshesRestResource::InternalAllocate()
 	CreateBuffer<FHairCardsUVFormat>(RenderData.UVs, UVsBuffer, TEXT("Hair.MeshesRest_UVs"));
 }
 
-void FHairMeshesRestResource::ReleaseRHI()
+void FHairMeshesRestResource::InternalRelease()
 {
 
 }
@@ -558,7 +554,7 @@ void FHairMeshesDeformedResource::InternalAllocate(FRDGBuilder& GraphBuilder)
 	}
 }
 
-void FHairMeshesDeformedResource::ReleaseRHI()
+void FHairMeshesDeformedResource::InternalRelease()
 {
 	DeformedPositionBuffer[0].Release();
 	DeformedPositionBuffer[1].Release();
@@ -610,7 +606,7 @@ FRDGExternalBuffer FHairStrandsRestResource::GetTangentBuffer(FRDGBuilder& Graph
 	return TangentBuffer;
 }
 
-void FHairStrandsRestResource::ReleaseRHI()
+void FHairStrandsRestResource::InternalRelease()
 {
 	PositionBuffer.Release();
 	PositionOffsetBuffer.Release();
@@ -647,7 +643,7 @@ void FHairStrandsDeformedResource::InternalAllocate(FRDGBuilder& GraphBuilder)
 	InternalCreateVertexBufferRDG<FHairStrandsPositionOffsetFormat>(GraphBuilder, DefaultOffsets, DeformedOffsetBuffer[1], HAIRSTRANDS_RESOUCE_NAME(CurveType, Hair.StrandsDeformed_DeformedOffsetBuffer1), EHairResourceUsageType::Dynamic);
 }
 
-void FHairStrandsDeformedResource::ReleaseRHI()
+void FHairStrandsDeformedResource::InternalRelease()
 {
 	DeformedPositionBuffer[0].Release();
 	DeformedPositionBuffer[1].Release();
@@ -702,7 +698,7 @@ void FHairStrandsClusterCullingResource::InternalAllocate(FRDGBuilder& GraphBuil
 	InternalCreateVertexBufferRDG(GraphBuilder, Data.VertexToClusterIds, PF_R32_UINT, VertexToClusterIdBuffer, TEXT("Hair.StrandsClusterCulling_VertexToClusterIds"), EHairResourceUsageType::Static);
 }
 
-void FHairStrandsClusterCullingResource::ReleaseRHI()
+void FHairStrandsClusterCullingResource::InternalRelease()
 {
 	ClusterInfoBuffer.Release();
 	ClusterLODInfoBuffer.Release();
@@ -800,7 +796,7 @@ void FHairStrandsRestRootResource::InternalAllocate(FRDGBuilder& GraphBuilder)
 	}
 }
 
-void FHairStrandsRestRootResource::ReleaseRHI()
+void FHairStrandsRestRootResource::InternalRelease()
 {
 	RootPositionBuffer.Release();
 	RootNormalBuffer.Release();
@@ -870,7 +866,7 @@ void FHairStrandsDeformedRootResource::InternalAllocate(FRDGBuilder& GraphBuilde
 	}
 }
 
-void FHairStrandsDeformedRootResource::ReleaseRHI()
+void FHairStrandsDeformedRootResource::InternalRelease()
 {
 	for (FLOD& GPUData : LODs)
 	{
@@ -1042,7 +1038,7 @@ void FHairStrandsInterpolationResource::InternalAllocate(FRDGBuilder& GraphBuild
 	//SimRootPointIndex.SetNum(0);
 }
 
-void FHairStrandsInterpolationResource::ReleaseRHI()
+void FHairStrandsInterpolationResource::InternalRelease()
 {
 	Interpolation0Buffer.Release();
 	Interpolation1Buffer.Release();
@@ -1085,7 +1081,7 @@ void FHairCardsInterpolationResource::InternalAllocate(FRDGBuilder& GraphBuilder
 	InternalCreateVertexBufferRDG<FHairCardsInterpolationFormat>(GraphBuilder, RenderData.Interpolation, InterpolationBuffer, TEXT("Hair.CardsInterpolation_InterpolationBuffer"), EHairResourceUsageType::Static);
 }
 
-void FHairCardsInterpolationResource::ReleaseRHI()
+void FHairCardsInterpolationResource::InternalRelease()
 {
 	InterpolationBuffer.Release();
 }
@@ -1115,7 +1111,7 @@ void FHairStrandsRaytracingResource::InternalAllocate(FRDGBuilder& GraphBuilder)
 	InternalCreateVertexBufferRDG<FHairStrandsRaytracingFormat>(GraphBuilder, VertexCount, PositionBuffer, TEXT("Hair.StrandsRaytracing_PositionBuffer"), EHairResourceUsageType::Dynamic);
 }
 
-void FHairStrandsRaytracingResource::ReleaseRHI()
+void FHairStrandsRaytracingResource::InternalRelease()
 {
 	PositionBuffer.Release();
 	RayTracingGeometry.ReleaseResource();
diff --git a/Engine/Plugins/Runtime/HairStrands/Source/HairStrandsCore/Private/HairStrandsRendering.cpp b/Engine/Plugins/Runtime/HairStrands/Source/HairStrandsCore/Private/HairStrandsRendering.cpp
--- a/Engine/Plugins/Runtime/HairStrands/Source/HairStrandsCore/Private/HairStrandsRendering.cpp
+++ b/Engine/Plugins/Runtime/HairStrands/Source/HairStrandsCore/Private/HairStrandsRendering.cpp
@@ -1295,7 +1295,7 @@ static void ConvertHairStrandsVFParameters(
 	FRDGBuilder* GraphBuilder,
 	FRDGImportedBuffer& Buffer,
 	FShaderResourceViewRHIRef& BufferRHISRV,
-	FRDGExternalBuffer& ExternalBuffer)
+	const FRDGExternalBuffer& ExternalBuffer)
 {
 	if (GraphBuilder)
 	{
@@ -1387,6 +1387,8 @@ void AddBufferTransitionToReadablePass(FRDGBuilder& GraphBuilder, FRHIUnorderedA
 	});
 }
 
+void CreateHairStrandsDebugAttributeBuffer(FRDGBuilder& GraphBuilder, FRDGExternalBuffer* DebugAttributeBuffer, uint32 VertexCount);
+
 void ComputeHairStrandsInterpolation(
 	FRDGBuilder& GraphBuilder,
 	FGlobalShaderMap* ShaderMap,
@@ -1462,11 +1464,21 @@ void ComputeHairStrandsInterpolation(
 				FRDGImportedBuffer Strands_DeformedPosition = Register(GraphBuilder, Instance->Strands.DeformedResource->GetBuffer(FHairStrandsDeformedResource::Current), ERDGImportedBufferFlags::CreateViews);
 				FRDGImportedBuffer Strands_DeformedPrevPosition = Register(GraphBuilder, Instance->Strands.DeformedResource->GetBuffer(FHairStrandsDeformedResource::Previous), ERDGImportedBufferFlags::CreateViews);
 				FRDGImportedBuffer Strands_DeformedTangent = Register(GraphBuilder, Instance->Strands.DeformedResource->TangentBuffer, ERDGImportedBufferFlags::CreateViews);
-			#if WITH_EDITOR
-				FRDGImportedBuffer Strands_DebugAttributeReg = Register(GraphBuilder, Instance->Strands.DebugAttributeBuffer, ERDGImportedBufferFlags::CreateUAV);
-				FRDGImportedBuffer* Strands_DebugAttribute = &Strands_DebugAttributeReg;
-			#else
 				FRDGImportedBuffer* Strands_DebugAttribute = nullptr;
+			#if WITH_EDITOR
+				FRDGImportedBuffer Strands_DebugAttributeReg;
+				if (bDebugModePatchedAttributeBuffer)
+				{
+					// Create an debug buffer for storing cluster visalization data. This is only used for debug purpose, hence only enable in editor build.
+					// Special case for debug mode were the attribute buffer is patch with some custom data to show hair properties (strands belonging to the same cluster, ...)
+					if (Instance->Strands.DebugAttributeBuffer.Buffer == nullptr)
+					{
+						CreateHairStrandsDebugAttributeBuffer(GraphBuilder, &Instance->Strands.DebugAttributeBuffer, Instance->Strands.Data->GetNumPoints());
+					}
+
+					Strands_DebugAttributeReg = Register(GraphBuilder, Instance->Strands.DebugAttributeBuffer, ERDGImportedBufferFlags::CreateUAV);
+					Strands_DebugAttribute = &Strands_DebugAttributeReg;
+				}
 			#endif
 				Strands_PositionSRV = Strands_DeformedPosition.SRV;
 				Strands_PositionOffsetSRV = RegisterAsSRV(GraphBuilder, Instance->Strands.DeformedResource->GetPositionOffsetBuffer(FHairStrandsDeformedResource::EFrameType::Current));
diff --git a/Engine/Plugins/Runtime/HairStrands/Source/HairStrandsCore/Public/GroomResources.h b/Engine/Plugins/Runtime/HairStrands/Source/HairStrandsCore/Public/GroomResources.h
--- a/Engine/Plugins/Runtime/HairStrands/Source/HairStrandsCore/Public/GroomResources.h
+++ b/Engine/Plugins/Runtime/HairStrands/Source/HairStrandsCore/Public/GroomResources.h
@@ -35,11 +35,15 @@ struct FHairCommonResource : public FRenderResource
 	/** Build the hair strands resource */
 	FHairCommonResource(EHairStrandsAllocationType AllocationType, bool bUseRenderGraph=true);
 
-	/* Init the buffer */
+	/* Init/Release buffers (FRenderResource) */
 	virtual void InitRHI() override;
+	virtual void ReleaseRHI() override;
+
+	/* Init/Release buffers (FHairCommonResource) */
 	void Allocate(FRDGBuilder& GraphBuilder);
 	virtual void InternalAllocate() {};
 	virtual void InternalAllocate(FRDGBuilder& GraphBuilder) {};
+	virtual void InternalRelease() {};
 
 	bool bUseRenderGraph = true;
 	bool bIsInitialized = false;
@@ -52,11 +56,9 @@ struct FHairStrandsRestRootResource : public FHairCommonResource
 	/** Build the hair strands resource */
 	FHairStrandsRestRootResource(const FHairStrandsRootData& RootData, EHairStrandsResourcesType CurveType);
 
-	/* Init the buffer */
-	void InternalAllocate(FRDGBuilder& GraphBuilder) override;
-
-	/* Release the buffer */
-	virtual void ReleaseRHI() override;
+	/* Init/Release buffers */
+	virtual void InternalAllocate(FRDGBuilder& GraphBuilder) override;
+	virtual void InternalRelease() override;
 
 	/* Get the resource name */
 	virtual FString GetFriendlyName() const override { return TEXT("FHairStrandsRestRootResource"); }
@@ -132,11 +134,9 @@ struct FHairStrandsDeformedRootResource : public FHairCommonResource
 	FHairStrandsDeformedRootResource(EHairStrandsResourcesType CurveType);
 	FHairStrandsDeformedRootResource(const FHairStrandsRestRootResource* InRestResources, EHairStrandsResourcesType CurveType);
 
-	/* Init the buffer */
-	void InternalAllocate(FRDGBuilder& GraphBuilder) override;
-
-	/* Release the buffer */
-	virtual void ReleaseRHI() override;
+	/* Init/Release buffers */
+	virtual void InternalAllocate(FRDGBuilder& GraphBuilder) override;
+	virtual void InternalRelease() override;
 
 	/* Get the resource name */
 	virtual FString GetFriendlyName() const override { return TEXT("FHairStrandsDeformedRootResource"); }
@@ -199,11 +199,9 @@ struct FHairStrandsRestResource : public FHairCommonResource
 	/** Build the hair strands resource */
 	FHairStrandsRestResource(const FHairStrandsDatas::FRenderData& HairStrandRenderData, const FVector& PositionOffset, EHairStrandsResourcesType CurveType);
 
-	/* Init the buffer */
-	void InternalAllocate(FRDGBuilder& GraphBuilder) override;
-
-	/* Release the buffer */
-	virtual void ReleaseRHI() override;
+	/* Init/Release buffers */
+	virtual void InternalAllocate(FRDGBuilder& GraphBuilder) override;
+	virtual void InternalRelease() override;
 
 	/* Get the resource name */
 	virtual FString GetFriendlyName() const override { return TEXT("FHairStrandsResource"); }
@@ -254,11 +252,9 @@ struct FHairStrandsDeformedResource : public FHairCommonResource
 	/** Build the hair strands resource */
 	FHairStrandsDeformedResource(const FHairStrandsDatas::FRenderData& HairStrandRenderData, bool bInitializeData, const FVector& InDefaultOffset, EHairStrandsResourcesType CurveType);
 
-	/* Init the buffer */
-	void InternalAllocate(FRDGBuilder& GraphBuilder) override;
-
-	/* Release the buffer */
-	virtual void ReleaseRHI() override;
+	/* Init/Release buffers */
+	virtual void InternalAllocate(FRDGBuilder& GraphBuilder) override;
+	virtual void InternalRelease() override;
 
 	/* Get the resource name */
 	virtual FString GetFriendlyName() const override { return TEXT("FHairStrandsDeformedResource"); }
@@ -320,11 +316,9 @@ struct FHairStrandsClusterCullingResource : public FHairCommonResource
 {
 	FHairStrandsClusterCullingResource(const FHairStrandsClusterCullingData& Data);
 
-	/* Init the buffer */
-	void InternalAllocate(FRDGBuilder& GraphBuilder) override;
-
-	/* Release the buffer */
-	virtual void ReleaseRHI() override;
+	/* Init/Release buffers */
+	virtual void InternalAllocate(FRDGBuilder& GraphBuilder) override;
+	virtual void InternalRelease() override;
 
 	/* Get the resource name */
 	virtual FString GetFriendlyName() const override { return TEXT("FHairStrandsClusterResource"); }
@@ -357,12 +351,10 @@ struct FHairStrandsInterpolationResource : public FHairCommonResource
 {
 	/** Build the hair strands resource */
 	FHairStrandsInterpolationResource(const FHairStrandsInterpolationDatas::FRenderData& InterpolationRenderData, const FHairStrandsDatas& SimDatas);
-
-	/* Init the buffer */
-	void InternalAllocate(FRDGBuilder& GraphBuilder) override;
-
-	/* Release the buffer */
-	virtual void ReleaseRHI() override;
+	
+	/* Init/Release buffers */
+	virtual void InternalAllocate(FRDGBuilder& GraphBuilder) override;
+	virtual void InternalRelease() override;
 
 	/* Get the resource name */
 	virtual FString GetFriendlyName() const override { return TEXT("FHairStrandsInterplationResource"); }
@@ -396,11 +388,9 @@ struct FHairStrandsRaytracingResource : public FHairCommonResource
 	FHairStrandsRaytracingResource(const FHairCardsDatas& InData);
 	FHairStrandsRaytracingResource(const FHairMeshesDatas& InData);
 
-	/* Init the buffer */
-	void InternalAllocate(FRDGBuilder& GraphBuilder) override;
-
-	/* Release the buffer */
-	virtual void ReleaseRHI() override;
+	/* Init/Release buffers */
+	virtual void InternalAllocate(FRDGBuilder& GraphBuilder) override;
+	virtual void InternalRelease() override;
 
 	/* Get the resource name */
 	virtual FString GetFriendlyName() const override { return TEXT("FHairStrandsRaytracingResource"); }
@@ -446,11 +436,9 @@ struct FHairCardsRestResource : public FHairCommonResource
 	virtual void InitResource() override;
 	virtual void ReleaseResource() override;
 
-	/* Init the buffer */
-	void InternalAllocate() override;
-
-	/* Release the buffer */
-	virtual void ReleaseRHI() override;
+	/* Init/release buffers */
+	virtual void InternalAllocate() override;
+	virtual void InternalRelease() override;
 
 	/* Get the resource name */
 	virtual FString GetFriendlyName() const override { return TEXT("FHairCardsResource"); }
@@ -498,11 +486,9 @@ struct FHairCardsProceduralResource : public FHairCommonResource
 	/** Build the hair strands resource */
 	FHairCardsProceduralResource(const FHairCardsProceduralDatas::FRenderData& HairCardsRenderData, const FIntPoint& AtlasResolution, const FHairCardsVoxel& InVoxel);
 
-	/* Init the buffer */
-	void InternalAllocate(FRDGBuilder& GraphBuilder) override;
-
-	/* Release the buffer */
-	virtual void ReleaseRHI() override;
+	/* Init/release buffers */
+	virtual void InternalAllocate(FRDGBuilder& GraphBuilder) override;
+	virtual void InternalRelease() override;
 
 	/* Get the resource name */
 	virtual FString GetFriendlyName() const override { return TEXT("FHairCardsResource"); }
@@ -539,11 +525,9 @@ struct FHairCardsDeformedResource : public FHairCommonResource
 	/** Build the hair strands resource */
 	FHairCardsDeformedResource(const FHairCardsDatas::FRenderData& HairStrandRenderData, bool bInitializeData);
 
-	/* Init the buffer */
-	void InternalAllocate(FRDGBuilder& GraphBuilder) override;
-
-	/* Release the buffer */
-	virtual void ReleaseRHI() override;
+	/* Init/release buffers */
+	virtual void InternalAllocate(FRDGBuilder& GraphBuilder) override;
+	virtual void InternalRelease() override;
 
 	/* Get the resource name */
 	virtual FString GetFriendlyName() const override { return TEXT("FHairCardsDeformedResource"); }
@@ -615,11 +599,9 @@ struct FHairCardsInterpolationResource : public FHairCommonResource
 	/** Build the hair strands resource */
 	FHairCardsInterpolationResource(const FHairCardsInterpolationDatas::FRenderData& InterpolationRenderData);
 
-	/* Init the buffer */
-	void InternalAllocate(FRDGBuilder& GraphBuilder) override;
-
-	/* Release the buffer */
-	virtual void ReleaseRHI() override;
+	/* Init/release buffers */
+	virtual void InternalAllocate(FRDGBuilder& GraphBuilder) override;
+	virtual void InternalRelease() override;
 
 	/* Get the resource name */
 	virtual FString GetFriendlyName() const override { return TEXT("FHairCardsInterplationResource"); }
@@ -650,11 +632,9 @@ struct FHairMeshesRestResource : public FHairCommonResource
 	virtual void InitResource() override;
 	virtual void ReleaseResource() override;
 
-	/* Init the buffer */
-	void InternalAllocate() override;
-
-	/* Release the buffer */
-	virtual void ReleaseRHI() override;
+	/* Init/release buffers */
+	virtual void InternalAllocate() override;
+	virtual void InternalRelease() override;
 
 	/* Get the resource name */
 	virtual FString GetFriendlyName() const override { return TEXT("FHairMeshesRestResource"); }
@@ -702,11 +682,9 @@ struct FHairMeshesDeformedResource : public FHairCommonResource
 	/** Build the hair strands resource */
 	FHairMeshesDeformedResource(const FHairMeshesDatas::FRenderData& HairMeshesRenderData, bool bInInitializedData);
 
-	/* Init the buffer */
-	void InternalAllocate(FRDGBuilder& GraphBuilder) override;
-
-	/* Release the buffer */
-	virtual void ReleaseRHI() override;
+	/* Init/release buffers */
+	virtual void InternalAllocate(FRDGBuilder& GraphBuilder) override;
+	virtual void InternalRelease() override;
 
 	/* Get the resource name */
 	virtual FString GetFriendlyName() const override { return TEXT("FHairMeshesDeformedResource"); }
