Separate 'bulk' data from hair data to later trim data serialization.

diff --git a/Engine/Plugins/Runtime/HairStrands/Source/HairStrandsCore/Private/GroomAsset.cpp b/Engine/Plugins/Runtime/HairStrands/Source/HairStrandsCore/Private/GroomAsset.cpp
--- a/Engine/Plugins/Runtime/HairStrands/Source/HairStrandsCore/Private/GroomAsset.cpp
+++ b/Engine/Plugins/Runtime/HairStrands/Source/HairStrandsCore/Private/GroomAsset.cpp
@@ -172,6 +172,24 @@ bool IsHairStrandsAssetLoadingEnable()
 	return GHairStrandsLoadAsset > 0;
 }
 
+static void PatchHairStrandsInterpolationBulkData(FHairStrandsInterpolationBulkData& BulkData, const FHairStrandsDatas& SimDatas)
+{
+	// This is now move into the builder, but kept here for compatibility.
+	if (BulkData.SimRootPointIndex.Num() == 0)
+	{
+		const uint32 RootCount = SimDatas.GetNumCurves();
+		BulkData.SimRootPointIndex.SetNum(SimDatas.GetNumPoints());
+		for (uint32 CurveIndex = 0; CurveIndex < RootCount; ++CurveIndex)
+		{
+			const uint16 PointCount = SimDatas.StrandsCurves.CurvesCount[CurveIndex];
+			const uint32 PointOffset = SimDatas.StrandsCurves.CurvesOffset[CurveIndex];
+			for (uint32 PointIndex = 0; PointIndex < PointCount; ++PointIndex)
+			{
+				BulkData.SimRootPointIndex[PointIndex + PointOffset] = PointOffset;
+			}
+		}
+	}
+}
 /////////////////////////////////////////////////////////////////////////////////////////
 
 uint8 UGroomAsset::GenerateClassStripFlags(FArchive& Ar)
@@ -1014,24 +1032,26 @@ FArchive& operator<<(FArchive& Ar, FHairGroupData::FCards::FLOD& CardLODData)
 {
 	if (!Ar.IsCooking() || !CardLODData.bIsCookedOut)
 	{
-		Ar << CardLODData.Data;
-		Ar << CardLODData.InterpolationData;
+		CardLODData.Data.Serialize(Ar, CardLODData.BulkData);
+		CardLODData.InterpolationData.Serialize(Ar, CardLODData.InterpolationBulkData);
 
-		CardLODData.Guides.Data.Serialize(Ar);
-		CardLODData.Guides.InterpolationData.Serialize(Ar);
+		CardLODData.Guides.Data.Serialize(Ar, CardLODData.Guides.BulkData);
+		CardLODData.Guides.InterpolationData.Serialize(Ar, CardLODData.Guides.InterpolationBulkData);
 	}
 	else
 	{
 		// LOD has been marked to be cooked out so serialize empty data
 		FHairCardsDatas NoCardsData;
-		Ar << NoCardsData;
+		FHairCardsBulkData NoCardsBulkData;
+		NoCardsData.Serialize(Ar, NoCardsBulkData);
 
 		FHairCardsInterpolationDatas NoInterpolationData;
-		Ar << NoInterpolationData;
+		FHairCardsInterpolationBulkData NoInterpolationBulkData;
+		NoInterpolationData.Serialize(Ar, NoInterpolationBulkData);
 
 		FHairGroupData::FBaseWithInterpolation NoGuideData;
-		NoGuideData.Data.Serialize(Ar);
-		NoGuideData.InterpolationData.Serialize(Ar);
+		NoGuideData.Data.Serialize(Ar, NoGuideData.BulkData);
+		NoGuideData.InterpolationData.Serialize(Ar, NoGuideData.InterpolationBulkData);
 	}
 
 	return Ar;
@@ -1041,13 +1061,14 @@ FArchive& operator<<(FArchive& Ar, FHairGroupData::FMeshes::FLOD& MeshLODData)
 {
 	if (!Ar.IsCooking() || !MeshLODData.bIsCookedOut)
 	{
-		Ar << MeshLODData.Data;
+		MeshLODData.Data.Serialize(Ar, MeshLODData.BulkData);
 	}
 	else
 	{
 		// LOD has been marked to be cooked out so serialize empty data
 		FHairMeshesDatas NoMeshesData;
-		Ar << NoMeshesData;
+		FHairMeshesBulkData NoMeshesBulkData;
+		NoMeshesData.Serialize(Ar, NoMeshesBulkData);
 	}
 
 	return Ar;
@@ -1060,9 +1081,9 @@ FArchive& operator<<(FArchive& Ar, FHairGroupData& GroupData)
 	FHairGroupData NoStrandsData;
 	if (!Ar.IsCooking() || !GroupData.bIsCookedOut)
 	{
-		GroupData.Strands.Data.Serialize(Ar);
-		GroupData.Guides.Data.Serialize(Ar);
-		GroupData.Strands.InterpolationData.Serialize(Ar);
+		GroupData.Strands.Data.Serialize(Ar, GroupData.Strands.BulkData);
+		GroupData.Guides.Data.Serialize(Ar, GroupData.Guides.BulkData);
+		GroupData.Strands.InterpolationData.Serialize(Ar, GroupData.Strands.InterpolationBulkData);
 
 	}
 	else
@@ -1070,9 +1091,9 @@ FArchive& operator<<(FArchive& Ar, FHairGroupData& GroupData)
 		// Fall back to no data, but still serialize guide data as they an be used for cards simulation
 		// Theoritically, we should have something to detect if we are going to use or not guide (for 
 		// simulation or RBF deformation) on the target platform
-		NoStrandsData.Strands.Data.Serialize(Ar);
-		GroupData.Guides.Data.Serialize(Ar);
-		NoStrandsData.Strands.InterpolationData.Serialize(Ar);
+		NoStrandsData.Strands.Data.Serialize(Ar, NoStrandsData.Strands.BulkData);
+		GroupData.Guides.Data.Serialize(Ar, GroupData.Guides.BulkData);
+		NoStrandsData.Strands.InterpolationData.Serialize(Ar, NoStrandsData.Strands.InterpolationBulkData);
 	}
 
 	if (Ar.CustomVer(FAnimObjectVersion::GUID) >= FAnimObjectVersion::SerializeHairClusterCullingData)
@@ -1941,16 +1962,13 @@ bool UGroomAsset::BuildCardsGeometry(uint32 GroupIndex)
 			if (CardsMesh != nullptr)
 			{
 				CardsMesh->ConditionalPostLoad();
-				bInitResources = FHairCardsBuilder::ImportGeometry(CardsMesh, LOD.Data, LOD.Guides.Data, LOD.InterpolationData);
+				bInitResources = FHairCardsBuilder::ImportGeometry(CardsMesh, LOD.Data, LOD.BulkData, LOD.Guides.Data, LOD.InterpolationData, LOD.InterpolationBulkData);
 				if (!bInitResources)
 				{
 					UE_LOG(LogHairStrands, Warning, TEXT("Failed to import cards from %s for Group %d LOD %d."), *CardsMesh->GetName(), GroupIndex, LODIt);
 
 					LOD.Data.Cards.Reset();
-					LOD.Data.RenderData.Positions.Empty();
-					LOD.Data.RenderData.Normals.Empty();
-					LOD.Data.RenderData.UVs.Empty();
-					LOD.Data.RenderData.Indices.Empty();
+					LOD.BulkData.Reset();
 				}
 			}
 
@@ -1963,10 +1981,7 @@ bool UGroomAsset::BuildCardsGeometry(uint32 GroupIndex)
 			// 2. Init geometry & texture resources, and generate interpolation data/resources
 			if (bInitResources)
 			{
-				LOD.RestResource = new FHairCardsRestResource(
-					LOD.Data.RenderData,
-					LOD.Data.Cards.GetNumVertices(),
-					LOD.Data.Cards.GetNumTriangles());
+				LOD.RestResource = new FHairCardsRestResource(LOD.BulkData);
 				BeginInitResource(LOD.RestResource);
 
 				// 2.1 Load atlas textures
@@ -1978,7 +1993,7 @@ bool UGroomAsset::BuildCardsGeometry(uint32 GroupIndex)
 				LOD.RestResource->bInvertUV = Desc->SourceType == EHairCardsSourceType::Procedural;
 				
 				// 2.2 Load interoplatino resources
-				LOD.InterpolationResource = new FHairCardsInterpolationResource(LOD.InterpolationData.RenderData);
+				LOD.InterpolationResource = new FHairCardsInterpolationResource(LOD.InterpolationBulkData);
 				BeginInitResource(LOD.InterpolationResource);
 
 				// Create own interpolation settings for cards.
@@ -2004,7 +2019,10 @@ bool UGroomAsset::BuildCardsGeometry(uint32 GroupIndex)
 				FGroomBuilder::BuildData(
 					LOD.Guides.Data,
 					GroupData.Guides.Data,
+					LOD.Guides.BulkData,
+					GroupData.Guides.BulkData,
 					LOD.Guides.InterpolationData,
+					LOD.Guides.InterpolationBulkData,
 					CardsInterpolationSettings,
 					true,
 					false,
@@ -2013,14 +2031,15 @@ bool UGroomAsset::BuildCardsGeometry(uint32 GroupIndex)
 
 				if (bCopyRenderData)
 				{
-					GroupData.Strands.Data.RenderData = LOD.Guides.Data.RenderData;
-					GroupData.Guides.Data.RenderData = LOD.Guides.Data.RenderData;
+					GroupData.Strands.BulkData = LOD.Guides.BulkData;
+					GroupData.Guides.BulkData = LOD.Guides.BulkData;
 				}
 
-				LOD.Guides.RestResource = new FHairStrandsRestResource(LOD.Guides.Data.RenderData, LOD.Guides.Data.BoundingBox.GetCenter(), EHairStrandsResourcesType::Cards);
+				LOD.Guides.RestResource = new FHairStrandsRestResource(LOD.Guides.BulkData, EHairStrandsResourcesType::Cards);
 				BeginInitResource(LOD.Guides.RestResource);
 
-				LOD.Guides.InterpolationResource = new FHairStrandsInterpolationResource(LOD.Guides.InterpolationData.RenderData, GroupData.Guides.Data);
+				PatchHairStrandsInterpolationBulkData(LOD.Guides.InterpolationBulkData, GroupData.Guides.Data);
+				LOD.Guides.InterpolationResource = new FHairStrandsInterpolationResource(LOD.Guides.InterpolationBulkData);
 				BeginInitResource(LOD.Guides.InterpolationResource);
 
 				// Update card stats to display
@@ -2173,7 +2192,8 @@ bool UGroomAsset::BuildMeshesGeometry(uint32 GroupIndex)
 				Desc->ImportedMesh->ConditionalPostLoad();
 				FHairMeshesBuilder::ImportGeometry(
 					Desc->ImportedMesh,
-					LOD.Data);
+					LOD.Data,
+					LOD.BulkData);
 			}
 			else
 			{
@@ -2181,7 +2201,8 @@ bool UGroomAsset::BuildMeshesGeometry(uint32 GroupIndex)
 				FHairMeshesBuilder::BuildGeometry(
 					GroupData.Strands.Data,
 					GroupData.Guides.Data,
-					LOD.Data);
+					LOD.Data,
+					LOD.BulkData);
 			}
 
 			bDataBuilt |= true;
@@ -2233,7 +2254,7 @@ FHairStrandsRestResource* UGroomAsset::AllocateGuidesResources(uint32 GroupIndex
 		{
 			if (GroupData.Guides.RestResource == nullptr)
 			{
-				GroupData.Guides.RestResource = new FHairStrandsRestResource(GroupData.Guides.Data.RenderData, GroupData.Guides.Data.BoundingBox.GetCenter(), EHairStrandsResourcesType::Guides);
+				GroupData.Guides.RestResource = new FHairStrandsRestResource(GroupData.Guides.BulkData, EHairStrandsResourcesType::Guides);
 				BeginInitResource(GroupData.Guides.RestResource);
 			}
 			return GroupData.Guides.RestResource;
@@ -2250,7 +2271,8 @@ FHairStrandsInterpolationResource* UGroomAsset::AllocateInterpolationResources(u
 		check(GroupData.Guides.IsValid());
 		if (GroupData.Strands.InterpolationResource == nullptr)
 		{
-			GroupData.Strands.InterpolationResource = new FHairStrandsInterpolationResource(GroupData.Strands.InterpolationData.RenderData, GroupData.Guides.Data);
+			PatchHairStrandsInterpolationBulkData(GroupData.Strands.InterpolationBulkData, GroupData.Guides.Data);
+			GroupData.Strands.InterpolationResource = new FHairStrandsInterpolationResource(GroupData.Strands.InterpolationBulkData);
 			BeginInitResource(GroupData.Strands.InterpolationResource);
 		}
 		return GroupData.Strands.InterpolationResource;
@@ -2269,7 +2291,7 @@ FHairStrandsRaytracingResource* UGroomAsset::AllocateCardsRaytracingResources(ui
 
 		if (LOD.RaytracingResource == nullptr)
 		{
-			LOD.RaytracingResource = new FHairStrandsRaytracingResource(LOD.Data);
+			LOD.RaytracingResource = new FHairStrandsRaytracingResource(LOD.BulkData);
 			BeginInitResource(LOD.RaytracingResource);
 		}
 		return LOD.RaytracingResource;
@@ -2287,7 +2309,7 @@ FHairStrandsRaytracingResource* UGroomAsset::AllocateMeshesRaytracingResources(u
 
 		if (LOD.RaytracingResource == nullptr)
 		{
-			LOD.RaytracingResource = new FHairStrandsRaytracingResource(LOD.Data);
+			LOD.RaytracingResource = new FHairStrandsRaytracingResource(LOD.BulkData);
 			BeginInitResource(LOD.RaytracingResource);
 		}
 		return LOD.RaytracingResource;
@@ -2304,7 +2326,7 @@ FHairStrandsRaytracingResource* UGroomAsset::AllocateStrandsRaytracingResources(
 
 		if (GroupData.Strands.RaytracingResource == nullptr)
 		{
-			GroupData.Strands.RaytracingResource = new FHairStrandsRaytracingResource(GroupData.Strands.Data);
+			GroupData.Strands.RaytracingResource = new FHairStrandsRaytracingResource(GroupData.Strands.BulkData);
 			BeginInitResource(GroupData.Strands.RaytracingResource);
 		}
 		return GroupData.Strands.RaytracingResource;
@@ -2333,7 +2355,7 @@ void UGroomAsset::InitStrandsResources()
 
 		if (GroupData.Strands.HasValidData())
 		{
-			GroupData.Strands.RestResource = new FHairStrandsRestResource(GroupData.Strands.Data.RenderData, GroupData.Strands.Data.BoundingBox.GetCenter(), EHairStrandsResourcesType::Strands);
+			GroupData.Strands.RestResource = new FHairStrandsRestResource(GroupData.Strands.BulkData, EHairStrandsResourcesType::Strands);
 			BeginInitResource(GroupData.Strands.RestResource);
 
 			if (GroupData.Strands.ClusterCullingData.IsValid())
@@ -2426,19 +2448,17 @@ void UGroomAsset::InitCardsResources()
 			if (LOD.RestResource == nullptr &&			// don't initialize again if they were previously initialized during the BuildCardsGeometry
 				LOD.HasValidData())
 			{
-				LOD.RestResource = new FHairCardsRestResource(
-					LOD.Data.RenderData,
-					LOD.Data.Cards.GetNumVertices(),
-					LOD.Data.Cards.GetNumTriangles());
+				LOD.RestResource = new FHairCardsRestResource(LOD.BulkData);
 				BeginInitResource(LOD.RestResource);
 
-				LOD.InterpolationResource = new FHairCardsInterpolationResource(LOD.InterpolationData.RenderData);
+				LOD.InterpolationResource = new FHairCardsInterpolationResource(LOD.InterpolationBulkData);
 				BeginInitResource(LOD.InterpolationResource);
 
-				LOD.Guides.RestResource = new FHairStrandsRestResource(LOD.Guides.Data.RenderData, LOD.Guides.Data.BoundingBox.GetCenter(), EHairStrandsResourcesType::Cards);
+				LOD.Guides.RestResource = new FHairStrandsRestResource(LOD.Guides.BulkData, EHairStrandsResourcesType::Cards);
 				BeginInitResource(LOD.Guides.RestResource);
 
-				LOD.Guides.InterpolationResource = new FHairStrandsInterpolationResource(LOD.Guides.InterpolationData.RenderData, GroupData.Guides.Data);
+				PatchHairStrandsInterpolationBulkData(LOD.Guides.InterpolationBulkData, GroupData.Guides.Data);
+				LOD.Guides.InterpolationResource = new FHairStrandsInterpolationResource(LOD.Guides.InterpolationBulkData);
 				BeginInitResource(LOD.Guides.InterpolationResource);
 
 				if (Desc)
@@ -2489,10 +2509,7 @@ void UGroomAsset::InitMeshesResources()
 
 			if (LOD.HasValidData())
 			{
-				LOD.RestResource = new FHairMeshesRestResource(
-					LOD.Data.RenderData,
-					LOD.Data.Meshes.GetNumVertices(),
-					LOD.Data.Meshes.GetNumTriangles());
+				LOD.RestResource = new FHairMeshesRestResource(LOD.BulkData);
 				BeginInitResource(LOD.RestResource);
 
 				if (Desc)
@@ -2758,9 +2775,11 @@ bool UGroomAsset::IsMaterialUsed(int32 MaterialIndex) const
 struct FHairProceduralCardsQuery
 {
 	FHairCardsInterpolationDatas InterpolationData;
+	FHairCardsInterpolationBulkData InterpolationBulkData;
 	FHairCardsProceduralDatas ProceduralData;
 	FHairStrandsDatas GuideData;
 	FHairCardsDatas Data;
+	FHairCardsBulkData BulkData;
 
 	UGroomAsset* Asset = nullptr;
 	FHairCardsRestResource* Resources = nullptr;
@@ -2826,13 +2845,14 @@ void UGroomAsset::SaveProceduralCards(uint32 DescIndex)
 		Q.ProceduralData,
 		Q.GuideData,
 		Q.InterpolationData,
+		Q.InterpolationBulkData,
 		Desc->Textures);
 	Q.Textures = &Desc->Textures;
 
-	FHairCardsBuilder::Convert(Q.ProceduralData, Q.Data);
+	FHairCardsBuilder::Convert(Q.ProceduralData, Q.Data, Q.BulkData);
 
 	// 3. Create resources and enqueue texture generation (GPU, kicked by the render thread) 
-	Q.Resources = new FHairCardsRestResource(Q.Data.RenderData, Q.Data.Cards.GetNumVertices(), Q.Data.Cards.GetNumTriangles());
+	Q.Resources = new FHairCardsRestResource(Q.BulkData);
 	BeginInitResource(Q.Resources);
 	Q.ProceduralResources = new FHairCardsProceduralResource(Q.ProceduralData.RenderData, Q.ProceduralData.Atlas.Resolution, Q.ProceduralData.Voxels);
 	BeginInitResource(Q.ProceduralResources);
diff --git a/Engine/Plugins/Runtime/HairStrands/Source/HairStrandsCore/Private/GroomBuilder.cpp b/Engine/Plugins/Runtime/HairStrands/Source/HairStrandsCore/Private/GroomBuilder.cpp
--- a/Engine/Plugins/Runtime/HairStrands/Source/HairStrandsCore/Private/GroomBuilder.cpp
+++ b/Engine/Plugins/Runtime/HairStrands/Source/HairStrandsCore/Private/GroomBuilder.cpp
@@ -204,7 +204,7 @@ namespace HairStrandsBuilder
 	}
 
 	/** Build the packed datas for gpu rendering/simulation */
-	void BuildRenderData(FHairStrandsDatas& HairStrands, const TArray<uint8>& RandomSeeds)
+	void BuildRenderData(const FHairStrandsDatas& HairStrands, const TArray<uint8>& RandomSeeds, FHairStrandsBulkData& OutBulkData)
 	{
 		TRACE_CPUPROFILER_EVENT_SCOPE(HairStrandsBuilder::BuildRenderData);
 
@@ -213,9 +213,9 @@ namespace HairStrandsBuilder
 		if (!(NumCurves > 0 && NumPoints > 0))
 			return;
 
-		TArray<FHairStrandsPositionFormat::Type>& OutPackedPositions = HairStrands.RenderData.Positions;
-		TArray<FHairStrandsAttributeFormat::Type>& OutPackedAttributes = HairStrands.RenderData.Attributes;
-		TArray<FHairStrandsMaterialFormat::Type>& OutPackedMaterials = HairStrands.RenderData.Materials;
+		TArray<FHairStrandsPositionFormat::Type>& OutPackedPositions = OutBulkData.Positions;
+		TArray<FHairStrandsAttributeFormat::Type>& OutPackedAttributes = OutBulkData.Attributes;
+		TArray<FHairStrandsMaterialFormat::Type>& OutPackedMaterials = OutBulkData.Materials;
 
 		OutPackedPositions.SetNum(NumPoints * FHairStrandsPositionFormat::ComponentCount);
 		OutPackedAttributes.SetNum(NumPoints * FHairStrandsAttributeFormat::ComponentCount);
@@ -223,8 +223,8 @@ namespace HairStrandsBuilder
 
 		const FVector HairBoxCenter = HairStrands.BoundingBox.GetCenter();
 
-		FHairStrandsCurves& Curves = HairStrands.StrandsCurves;
-		FHairStrandsPoints& Points = HairStrands.StrandsPoints;
+		const FHairStrandsCurves& Curves = HairStrands.StrandsCurves;
+		const FHairStrandsPoints& Points = HairStrands.StrandsPoints;
 
 		const bool bSeedValid = RandomSeeds.Num() > 0;
 		for (uint32 CurveIndex = 0; CurveIndex < NumCurves; ++CurveIndex)
@@ -273,12 +273,18 @@ namespace HairStrandsBuilder
 				Material.Roughness  = FMath::Clamp(uint32(Points.PointsRoughness[PointIndex + IndexOffset] * 0xFF), 0u, 0xFFu);
 			}
 		}
+
+		OutBulkData.BoundingBox = HairStrands.BoundingBox;
+		OutBulkData.CurveCount = HairStrands.GetNumCurves();
+		OutBulkData.PointCount = HairStrands.GetNumPoints();
+		OutBulkData.MaxLength = HairStrands.StrandsCurves.MaxLength;
+		OutBulkData.MaxRadius = HairStrands.StrandsCurves.MaxRadius;
 	}
 
-	void BuildRenderData(FHairStrandsDatas& HairStrands)
+	void BuildRenderData(const FHairStrandsDatas& HairStrands, FHairStrandsBulkData& OutBulkData)
 	{
 		TArray<uint8> RandomSeeds;
-		BuildRenderData(HairStrands, RandomSeeds);
+		BuildRenderData(HairStrands, RandomSeeds, OutBulkData);
 	}
 
 } // namespace HairStrandsBuilder
@@ -1099,7 +1105,7 @@ namespace HairInterpolationBuilder
 	}
 
 	/** Build data for interpolation between simulation and rendering */
-	void BuildRenderData(FHairStrandsInterpolationDatas& HairInterpolation)
+	void BuildRenderData(const FHairStrandsInterpolationDatas& HairInterpolation, const FHairStrandsDatas& SimDatas, FHairStrandsInterpolationBulkData& OutBulkData)
 	{
 		TRACE_CPUPROFILER_EVENT_SCOPE(HairInterpolationBuilder::BuildRenderData);
 
@@ -1110,8 +1116,8 @@ namespace HairInterpolationBuilder
 		auto LowerPart = [](uint32 Index) { return uint16(Index & 0xFFFF); };
 		auto UpperPart = [](uint32 Index) { return uint8((Index >> 16) & 0xFF); };
 
-		TArray<FHairStrandsInterpolation0Format::Type>& OutPointsInterpolation0 = HairInterpolation.RenderData.Interpolation0;
-		TArray<FHairStrandsInterpolation1Format::Type>& OutPointsInterpolation1 = HairInterpolation.RenderData.Interpolation1;
+		TArray<FHairStrandsInterpolation0Format::Type>& OutPointsInterpolation0 = OutBulkData.Interpolation0;
+		TArray<FHairStrandsInterpolation1Format::Type>& OutPointsInterpolation1 = OutBulkData.Interpolation1;
 
 		OutPointsInterpolation0.SetNum(PointCount * FHairStrandsInterpolation0Format::ComponentCount);
 		OutPointsInterpolation1.SetNum(PointCount * FHairStrandsInterpolation1Format::ComponentCount);
@@ -1139,6 +1145,18 @@ namespace HairInterpolationBuilder
 			OutInterp1.Pad0			= 0;
 			OutInterp1.Pad1			= 0;
 		}	
+
+		const uint32 RootCount = SimDatas.GetNumCurves();
+		OutBulkData.SimRootPointIndex.SetNum(SimDatas.GetNumPoints());
+		for (uint32 CurveIndex = 0; CurveIndex < RootCount; ++CurveIndex)
+		{
+			const uint16 SimPointCount = SimDatas.StrandsCurves.CurvesCount[CurveIndex];
+			const uint32 SimPointOffset = SimDatas.StrandsCurves.CurvesOffset[CurveIndex];
+			for (uint32 PointIndex = 0; PointIndex < SimPointCount; ++PointIndex)
+			{
+				OutBulkData.SimRootPointIndex[PointIndex + SimPointOffset] = SimPointOffset;
+			}
+		}		
 	}
 
 	/** Fill the GroomAsset with the interpolation data that exists in the HairDescription */
@@ -1620,9 +1638,12 @@ void FGroomBuilder::BuildData(UGroomAsset* GroomAsset)
 }
 
 void FGroomBuilder::BuildData(
-	FHairStrandsDatas& RenData,
-	FHairStrandsDatas& SimData,
+	const FHairStrandsDatas& RenData,
+	const FHairStrandsDatas& SimData,
+	FHairStrandsBulkData& RenBulkData,
+	FHairStrandsBulkData& SimBulkData,
 	FHairStrandsInterpolationDatas& InterpolationData,
+	FHairStrandsInterpolationBulkData& InterpolationBulkData,
 	const FHairInterpolationSettings& InterpolationSettings,
 	const bool bBuildRen,
 	const bool bBuildSim,
@@ -1638,8 +1659,8 @@ void FGroomBuilder::BuildData(
 	FGroomDataRandomizer Random(Seed, RenData.GetNumCurves(), SimData.GetNumCurves());
 
 	// Build RenderData for HairStrandsDatas
-	if (bBuildRen) { HairStrandsBuilder::BuildRenderData(RenData, Random.GetRenderCurveSeeds()); }
-	if (bBuildSim) { HairStrandsBuilder::BuildRenderData(SimData, Random.GetSimCurveSeeds()); }
+	if (bBuildRen) { HairStrandsBuilder::BuildRenderData(RenData, Random.GetRenderCurveSeeds(),	RenBulkData); }
+	if (bBuildSim) { HairStrandsBuilder::BuildRenderData(SimData, Random.GetSimCurveSeeds(),	SimBulkData); }
 
 	// Build Rendering data for InterpolationData
 	if (bBuildInterpolation)
@@ -1650,7 +1671,7 @@ void FGroomBuilder::BuildData(
 		{
 			HairInterpolationBuilder::BuildInterpolationData(InterpolationData, SimData, RenData, InterpolationSettings, Random.GetRandomGuideIndices());
 		}
-		HairInterpolationBuilder::BuildRenderData(InterpolationData);
+		HairInterpolationBuilder::BuildRenderData(InterpolationData, SimData, InterpolationBulkData);
 	}
 }
 
@@ -1659,7 +1680,10 @@ void FGroomBuilder::BuildData(FHairGroupData& GroupData, const FHairGroupsInterp
 	BuildData(
 		GroupData.Strands.Data,
 		GroupData.Guides.Data,
+		GroupData.Strands.BulkData,
+		GroupData.Guides.BulkData,
 		GroupData.Strands.InterpolationData,
+		GroupData.Strands.InterpolationBulkData,
 		InSettings.InterpolationSettings, true, true, true, GroupIndex);
 }
 
diff --git a/Engine/Plugins/Runtime/HairStrands/Source/HairStrandsCore/Private/GroomComponent.cpp b/Engine/Plugins/Runtime/HairStrands/Source/HairStrandsCore/Private/GroomComponent.cpp
--- a/Engine/Plugins/Runtime/HairStrands/Source/HairStrandsCore/Private/GroomComponent.cpp
+++ b/Engine/Plugins/Runtime/HairStrands/Source/HairStrandsCore/Private/GroomComponent.cpp
@@ -747,7 +747,7 @@ public:
 			}
 			VertexFactory = (FVertexFactory*)Instance->Meshes.LODs[IntLODIndex].GetVertexFactory();
 			check(VertexFactory);
-			HairVertexCount = Instance->Meshes.LODs[IntLODIndex].RestResource->PrimitiveCount * 3;
+			HairVertexCount = Instance->Meshes.LODs[IntLODIndex].RestResource->GetPrimitiveCount() * 3;
 			MaxVertexIndex = HairVertexCount;
 			NumPrimitive = HairVertexCount / 3;
 			IndexBuffer = &Instance->Meshes.LODs[IntLODIndex].RestResource->IndexBuffer;
@@ -763,7 +763,7 @@ public:
 
 			VertexFactory = (FVertexFactory*)Instance->Cards.LODs[IntLODIndex].GetVertexFactory();
 			check(VertexFactory);
-			HairVertexCount = Instance->Cards.LODs[IntLODIndex].RestResource->PrimitiveCount * 3;
+			HairVertexCount = Instance->Cards.LODs[IntLODIndex].RestResource->GetPrimitiveCount() * 3;
 			MaxVertexIndex = HairVertexCount;
 			NumPrimitive = HairVertexCount / 3;
 			IndexBuffer = &Instance->Cards.LODs[IntLODIndex].RestResource->RestIndexBuffer;
@@ -1967,7 +1967,7 @@ void UGroomComponent::InitResources(bool bIsBindingReloading)
 		const bool bNeedGuides = GroupData.Guides.HasValidData() && (bHasNeedSimulation || bHasNeedGlobalDeformation || bPreviewMode);
 		if (bNeedGuides)
 		{
-			HairGroupInstance->Guides.Data = &GroupData.Guides.Data;
+			HairGroupInstance->Guides.Data = &GroupData.Guides.BulkData;
 
 			if (LocalBindingAsset)
 			{
@@ -1988,9 +1988,7 @@ void UGroomComponent::InitResources(bool bIsBindingReloading)
 			check(GroupData.Guides.RestResource);
 
 			// If guides are allocated, deformed resources are always needs since they are either used with simulation, or RBF deformation. Both are dynamics, and require deformed positions
-			HairGroupInstance->Guides.DeformedResource = new FHairStrandsDeformedResource(GroupData.Guides.Data.RenderData, true, HairGroupInstance->Guides.RestResource->PositionOffset, EHairStrandsResourcesType::Guides);
-			HairGroupInstance->Guides.DeformedResource->GetPositionOffset(FHairStrandsDeformedResource::EFrameType::Current)  = HairGroupInstance->Guides.RestResource->PositionOffset;
-			HairGroupInstance->Guides.DeformedResource->GetPositionOffset(FHairStrandsDeformedResource::EFrameType::Previous) = HairGroupInstance->Guides.RestResource->PositionOffset;
+			HairGroupInstance->Guides.DeformedResource = new FHairStrandsDeformedResource(GroupData.Guides.BulkData, true, EHairStrandsResourcesType::Guides);
 
 			// Initialize the simulation and the global deformation to its default behavior by setting it with LODIndex = -1
 			const int32 LODIndex = -1;
@@ -2009,8 +2007,8 @@ void UGroomComponent::InitResources(bool bIsBindingReloading)
 		{
 			check(GroupIt < GroomGroupsDesc.Num());
 
-			HairGroupInstance->Strands.Data = &GroupData.Strands.Data;
-			HairGroupInstance->Strands.InterpolationData = &GroupData.Strands.InterpolationData;
+			HairGroupInstance->Strands.Data = &GroupData.Strands.BulkData;
+//			HairGroupInstance->Strands.InterpolationData = &GroupData.Strands.InterpolationBulkData;
 
 			// (Lazy) Allocate interpolation resources, only if guides are required
 			if (bNeedGuides)
@@ -2054,7 +2052,7 @@ void UGroomComponent::InitResources(bool bIsBindingReloading)
 				if (bNeedDynamicResources)
 				{
 					// Allocate dynamic raytracing resources (owned by the groom component/instance)
-					HairGroupInstance->Strands.RenRaytracingResource = new FHairStrandsRaytracingResource(GroupData.Strands.Data);
+					HairGroupInstance->Strands.RenRaytracingResource = new FHairStrandsRaytracingResource(GroupData.Strands.BulkData);
 					HairGroupInstance->Strands.RenRaytracingResourceOwned = true;
 				}
 				else
@@ -2084,10 +2082,7 @@ void UGroomComponent::InitResources(bool bIsBindingReloading)
 			HairGroupInstance->Strands.RestResource = GroupData.Strands.RestResource;
 			if (bNeedDynamicResources)
 			{
-				HairGroupInstance->Strands.DeformedResource = new FHairStrandsDeformedResource(GroupData.Strands.Data.RenderData, false, HairGroupInstance->Strands.RestResource->PositionOffset, EHairStrandsResourcesType::Strands);
-				// Initialize deformed position relative position offset to the rest pose offset
-				HairGroupInstance->Strands.DeformedResource->GetPositionOffset(FHairStrandsDeformedResource::EFrameType::Current)  = HairGroupInstance->Strands.RestResource->PositionOffset;
-				HairGroupInstance->Strands.DeformedResource->GetPositionOffset(FHairStrandsDeformedResource::EFrameType::Previous) = HairGroupInstance->Strands.RestResource->PositionOffset;
+				HairGroupInstance->Strands.DeformedResource = new FHairStrandsDeformedResource(GroupData.Strands.BulkData, false, EHairStrandsResourcesType::Strands);
 			} 
 
 			// An empty groom doesn't have a ClusterCullingResource
@@ -2138,14 +2133,13 @@ void UGroomComponent::InitResources(bool bIsBindingReloading)
 					check(BindingType == EHairBindingType::Skinning);
 				}
 
-				InstanceLOD.Data = &LOD.Data;
+				InstanceLOD.Data = &LOD.BulkData;
 				InstanceLOD.RestResource = LOD.RestResource;
-				InstanceLOD.InterpolationData = &LOD.InterpolationData;
 				InstanceLOD.InterpolationResource = LOD.InterpolationResource;
 
 				if (bNeedDeformedPositions)
 				{
-					InstanceLOD.DeformedResource = new FHairCardsDeformedResource(LOD.Data.RenderData, false);
+					InstanceLOD.DeformedResource = new FHairCardsDeformedResource(LOD.BulkData, false);
 				}
 
 				#if RHI_RAYTRACING
@@ -2173,8 +2167,7 @@ void UGroomComponent::InitResources(bool bIsBindingReloading)
 				// Strands data/resources
 				if (bNeedDeformedPositions)
 				{
-					InstanceLOD.Guides.Data = &LOD.Guides.Data;
-					InstanceLOD.Guides.InterpolationData = &LOD.Guides.InterpolationData;
+					InstanceLOD.Guides.Data = &LOD.Guides.BulkData;
 					InstanceLOD.Guides.InterpolationResource = LOD.Guides.InterpolationResource;
 
 					if (bNeedRootData)
@@ -2192,11 +2185,7 @@ void UGroomComponent::InitResources(bool bIsBindingReloading)
 
 					InstanceLOD.Guides.RestResource = LOD.Guides.RestResource;
 					{
-						InstanceLOD.Guides.DeformedResource = new FHairStrandsDeformedResource(LOD.Guides.Data.RenderData, false, InstanceLOD.Guides.RestResource->PositionOffset, EHairStrandsResourcesType::Cards);
-
-						// Initialize deformed position relative position offset to the rest pose offset
-						InstanceLOD.Guides.DeformedResource->GetPositionOffset(FHairStrandsDeformedResource::EFrameType::Current) = InstanceLOD.Guides.RestResource->PositionOffset;
-						InstanceLOD.Guides.DeformedResource->GetPositionOffset(FHairStrandsDeformedResource::EFrameType::Previous) = InstanceLOD.Guides.RestResource->PositionOffset;
+						InstanceLOD.Guides.DeformedResource = new FHairStrandsDeformedResource(LOD.Guides.BulkData, false, EHairStrandsResourcesType::Cards);
 					}
 
 					InstanceLOD.Guides.HairInterpolationType = HairInterpolationType;
@@ -2215,11 +2204,11 @@ void UGroomComponent::InitResources(bool bIsBindingReloading)
 				const bool bHasGlobalDeformation   = HairGroupInstance->HairGroupPublicData->IsGlobalInterpolationEnable(MeshLODIndex);
 				const bool bNeedDeformedPositions  = bHasGlobalDeformation && BindingType == EHairBindingType::Skinning;
 
-				InstanceLOD.Data = &LOD.Data;
+				InstanceLOD.Data = &LOD.BulkData;
 				InstanceLOD.RestResource = LOD.RestResource;
 				if (bNeedDeformedPositions)
 				{
-					InstanceLOD.DeformedResource = new FHairMeshesDeformedResource(LOD.Data.RenderData, true);
+					InstanceLOD.DeformedResource = new FHairMeshesDeformedResource(LOD.BulkData, true);
 				}
 
 				#if RHI_RAYTRACING
diff --git a/Engine/Plugins/Runtime/HairStrands/Source/HairStrandsCore/Private/GroomDebug.cpp b/Engine/Plugins/Runtime/HairStrands/Source/HairStrandsCore/Private/GroomDebug.cpp
--- a/Engine/Plugins/Runtime/HairStrands/Source/HairStrandsCore/Private/GroomDebug.cpp
+++ b/Engine/Plugins/Runtime/HairStrands/Source/HairStrandsCore/Private/GroomDebug.cpp
@@ -714,13 +714,13 @@ static void AddDrawDebugCardsGuidesPass(
 	Parameters->RenVertexCount = LOD.Guides.RestResource->GetVertexCount();
 	Parameters->SimVertexCount = Instance->Guides.RestResource->GetVertexCount();
 
-	Parameters->RenRestOffset = LOD.Guides.RestResource->PositionOffset;
+	Parameters->RenRestOffset = LOD.Guides.RestResource->GetPositionOffset();
 	Parameters->RenRestPosition = RegisterAsSRV(GraphBuilder, LOD.Guides.RestResource->PositionBuffer);
 
 	Parameters->RenDeformedOffset = LOD.Guides.DeformedResource->GetPositionOffset(FHairStrandsDeformedResource::Current);
 	Parameters->RenDeformedPosition = RegisterAsSRV(GraphBuilder, LOD.Guides.DeformedResource->GetBuffer(FHairStrandsDeformedResource::Current));
 
-	Parameters->SimRestOffset = Instance->Guides.RestResource->PositionOffset;
+	Parameters->SimRestOffset = Instance->Guides.RestResource->GetPositionOffset();
 	Parameters->SimRestPosition = RegisterAsSRV(GraphBuilder, Instance->Guides.RestResource->PositionBuffer);
 
 	Parameters->SimDeformedOffset = Instance->Guides.DeformedResource->GetPositionOffset(FHairStrandsDeformedResource::Current);
diff --git a/Engine/Plugins/Runtime/HairStrands/Source/HairStrandsCore/Private/GroomResources.cpp b/Engine/Plugins/Runtime/HairStrands/Source/HairStrandsCore/Private/GroomResources.cpp
--- a/Engine/Plugins/Runtime/HairStrands/Source/HairStrandsCore/Private/GroomResources.cpp
+++ b/Engine/Plugins/Runtime/HairStrands/Source/HairStrandsCore/Private/GroomResources.cpp
@@ -357,15 +357,13 @@ void FHairCardIndexBuffer::InitRHI()
 	RHIUnlockBuffer(IndexBufferRHI);
 }
 
-FHairCardsRestResource::FHairCardsRestResource(const FHairCardsDatas::FRenderData& InRenderData, uint32 InVertexCount, uint32 InPrimitiveCount) :
+FHairCardsRestResource::FHairCardsRestResource(const FHairCardsBulkData& InBulkData) :
 	FHairCommonResource(EHairStrandsAllocationType::Immediate, false),
 	RestPositionBuffer(),
-	RestIndexBuffer(InRenderData.Indices),
-	VertexCount(InVertexCount),
-	PrimitiveCount(InPrimitiveCount),
+	RestIndexBuffer(InBulkData.Indices),
 	NormalsBuffer(),
 	UVsBuffer(),
-	RenderData(InRenderData)
+	BulkData(InBulkData)
 {
 
 }
@@ -373,9 +371,9 @@ FHairCardsRestResource::FHairCardsRestResource(const FHairCardsDatas::FRenderDat
 void FHairCardsRestResource::InternalAllocate()
 {
 	// These resources are kept as regular (i.e., non-RDG resources) as they need to be bound at the input assembly stage by the Vertex declaraction which requires FVertexBuffer type
-	CreateBuffer<FHairCardsPositionFormat>(RenderData.Positions, RestPositionBuffer, TEXT("Hair.CardsRest_PositionBuffer"));
-	CreateBuffer<FHairCardsNormalFormat>(RenderData.Normals, NormalsBuffer, TEXT("Hair.CardsRest_NormalBuffer"));
-	CreateBuffer<FHairCardsUVFormat>(RenderData.UVs, UVsBuffer, TEXT("Hair.CardsRest_UVBuffer"));
+	CreateBuffer<FHairCardsPositionFormat>(BulkData.Positions, RestPositionBuffer, TEXT("Hair.CardsRest_PositionBuffer"));
+	CreateBuffer<FHairCardsNormalFormat>(BulkData.Normals, NormalsBuffer, TEXT("Hair.CardsRest_NormalBuffer"));
+	CreateBuffer<FHairCardsUVFormat>(BulkData.UVs, UVsBuffer, TEXT("Hair.CardsRest_UVBuffer"));
 
 	FSamplerStateRHIRef DefaultSampler = TStaticSamplerState<SF_Bilinear, AM_Clamp, AM_Clamp, AM_Clamp>::GetRHI();
 	DepthSampler = DefaultSampler;
@@ -458,23 +456,22 @@ void FHairCardsProceduralResource::InternalRelease()
 
 /////////////////////////////////////////////////////////////////////////////////////////
 
-FHairCardsDeformedResource::FHairCardsDeformedResource(const FHairCardsDatas::FRenderData& HairCardsRenderData, bool bInInitializedData) :
+FHairCardsDeformedResource::FHairCardsDeformedResource(const FHairCardsBulkData& InBulkData, bool bInInitializedData) :
 	FHairCommonResource(EHairStrandsAllocationType::Deferred),
-	RenderData(HairCardsRenderData), bInitializedData(bInInitializedData)
+	BulkData(InBulkData), bInitializedData(bInInitializedData)
 {}
 
 void FHairCardsDeformedResource::InternalAllocate(FRDGBuilder& GraphBuilder)
 {
-	const uint32 VertexCount = RenderData.Positions.Num();
 	if (bInitializedData)
 	{
-		InternalCreateVertexBufferRDG<FHairCardsPositionFormat>(GraphBuilder, RenderData.Positions, DeformedPositionBuffer[0], TEXT("Hair.CardsDeformed_Position0"), EHairResourceUsageType::Dynamic);
-		InternalCreateVertexBufferRDG<FHairCardsPositionFormat>(GraphBuilder, RenderData.Positions, DeformedPositionBuffer[1], TEXT("Hair.CardsDeformed_Position1"), EHairResourceUsageType::Dynamic);
+		InternalCreateVertexBufferRDG<FHairCardsPositionFormat>(GraphBuilder, BulkData.Positions, DeformedPositionBuffer[0], TEXT("Hair.CardsDeformed_Position0"), EHairResourceUsageType::Dynamic);
+		InternalCreateVertexBufferRDG<FHairCardsPositionFormat>(GraphBuilder, BulkData.Positions, DeformedPositionBuffer[1], TEXT("Hair.CardsDeformed_Position1"), EHairResourceUsageType::Dynamic);
 	}
 	else
 	{
-		InternalCreateVertexBufferRDG<FHairCardsPositionFormat>(GraphBuilder, VertexCount, DeformedPositionBuffer[0], TEXT("Hair.CardsDeformed_Position0"), EHairResourceUsageType::Dynamic);
-		InternalCreateVertexBufferRDG<FHairCardsPositionFormat>(GraphBuilder, VertexCount, DeformedPositionBuffer[1], TEXT("Hair.CardsDeformed_Position1"), EHairResourceUsageType::Dynamic);
+		InternalCreateVertexBufferRDG<FHairCardsPositionFormat>(GraphBuilder, BulkData.GetNumVertices(), DeformedPositionBuffer[0], TEXT("Hair.CardsDeformed_Position0"), EHairResourceUsageType::Dynamic);
+		InternalCreateVertexBufferRDG<FHairCardsPositionFormat>(GraphBuilder, BulkData.GetNumVertices(), DeformedPositionBuffer[1], TEXT("Hair.CardsDeformed_Position1"), EHairResourceUsageType::Dynamic);
 	}
 }
 
@@ -486,26 +483,24 @@ void FHairCardsDeformedResource::InternalRelease()
 
 /////////////////////////////////////////////////////////////////////////////////////////
 
-FHairMeshesRestResource::FHairMeshesRestResource(const FHairMeshesDatas::FRenderData& InRenderData, uint32 InVertexCount, uint32 InPrimitiveCount) :
+FHairMeshesRestResource::FHairMeshesRestResource(const FHairMeshesBulkData& InBulkData) :
 	FHairCommonResource(EHairStrandsAllocationType::Immediate, false),
 	RestPositionBuffer(),
-	IndexBuffer(InRenderData.Indices),
-	VertexCount(InVertexCount),
-	PrimitiveCount(InPrimitiveCount),
+	IndexBuffer(InBulkData.Indices),
 	NormalsBuffer(),
 	UVsBuffer(),
-	RenderData(InRenderData)
+	BulkData(InBulkData)
 {
-	check(VertexCount > 0);
+	check(BulkData.GetNumVertices() > 0);
 	check(IndexBuffer.Indices.Num() > 0);
 }
 
 void FHairMeshesRestResource::InternalAllocate()
 {
 	// These resources are kept as regular (i.e., non-RDG resources) as they need to be bound at the input assembly stage by the Vertex declaraction which requires FVertexBuffer type
-	CreateBuffer<FHairCardsPositionFormat>(RenderData.Positions, RestPositionBuffer, TEXT("Hair.MeshesRest_Positions"));
-	CreateBuffer<FHairCardsNormalFormat>(RenderData.Normals, NormalsBuffer, TEXT("Hair.MeshesRest_Normals"));
-	CreateBuffer<FHairCardsUVFormat>(RenderData.UVs, UVsBuffer, TEXT("Hair.MeshesRest_UVs"));
+	CreateBuffer<FHairCardsPositionFormat>(BulkData.Positions, RestPositionBuffer, TEXT("Hair.MeshesRest_Positions"));
+	CreateBuffer<FHairCardsNormalFormat>(BulkData.Normals, NormalsBuffer, TEXT("Hair.MeshesRest_Normals"));
+	CreateBuffer<FHairCardsUVFormat>(BulkData.UVs, UVsBuffer, TEXT("Hair.MeshesRest_UVs"));
 }
 
 void FHairMeshesRestResource::InternalRelease()
@@ -534,23 +529,22 @@ void FHairMeshesRestResource::ReleaseResource()
 
 /////////////////////////////////////////////////////////////////////////////////////////
 
-FHairMeshesDeformedResource::FHairMeshesDeformedResource(const FHairMeshesDatas::FRenderData& HairMeshesData, bool bInInitializedData) :
+FHairMeshesDeformedResource::FHairMeshesDeformedResource(const FHairMeshesBulkData& InBulkData, bool bInInitializedData) :
 	FHairCommonResource(EHairStrandsAllocationType::Deferred),
-	RenderData(HairMeshesData), bInitializedData(bInInitializedData)
+	BulkData(InBulkData), bInitializedData(bInInitializedData)
 {}
 
 void FHairMeshesDeformedResource::InternalAllocate(FRDGBuilder& GraphBuilder)
 {
-	const uint32 VertexCount = RenderData.Positions.Num();
 	if (bInitializedData)
 	{
-		InternalCreateVertexBufferRDG<FHairCardsPositionFormat>(GraphBuilder, RenderData.Positions, DeformedPositionBuffer[0], TEXT("Hair.MeshesDeformed_Positions0"), EHairResourceUsageType::Dynamic);
-		InternalCreateVertexBufferRDG<FHairCardsPositionFormat>(GraphBuilder, RenderData.Positions, DeformedPositionBuffer[1], TEXT("Hair.MeshesDeformed_Positions1"), EHairResourceUsageType::Dynamic);
+		InternalCreateVertexBufferRDG<FHairCardsPositionFormat>(GraphBuilder, BulkData.Positions, DeformedPositionBuffer[0], TEXT("Hair.MeshesDeformed_Positions0"), EHairResourceUsageType::Dynamic);
+		InternalCreateVertexBufferRDG<FHairCardsPositionFormat>(GraphBuilder, BulkData.Positions, DeformedPositionBuffer[1], TEXT("Hair.MeshesDeformed_Positions1"), EHairResourceUsageType::Dynamic);
 	}
 	else
 	{
-		InternalCreateVertexBufferRDG<FHairCardsPositionFormat>(GraphBuilder, VertexCount, DeformedPositionBuffer[0], TEXT("Hair.MeshesDeformed_Positions0"), EHairResourceUsageType::Dynamic);
-		InternalCreateVertexBufferRDG<FHairCardsPositionFormat>(GraphBuilder, VertexCount, DeformedPositionBuffer[1], TEXT("Hair.MeshesDeformed_Positions1"), EHairResourceUsageType::Dynamic);
+		InternalCreateVertexBufferRDG<FHairCardsPositionFormat>(GraphBuilder, BulkData.GetNumVertices(), DeformedPositionBuffer[0], TEXT("Hair.MeshesDeformed_Positions0"), EHairResourceUsageType::Dynamic);
+		InternalCreateVertexBufferRDG<FHairCardsPositionFormat>(GraphBuilder, BulkData.GetNumVertices(), DeformedPositionBuffer[1], TEXT("Hair.MeshesDeformed_Positions1"), EHairResourceUsageType::Dynamic);
 	}
 }
 
@@ -562,19 +556,19 @@ void FHairMeshesDeformedResource::InternalRelease()
 
 /////////////////////////////////////////////////////////////////////////////////////////
 
-FHairStrandsRestResource::FHairStrandsRestResource(const FHairStrandsDatas::FRenderData& HairStrandRenderData, const FVector& InPositionOffset, EHairStrandsResourcesType InCurveType) :
+FHairStrandsRestResource::FHairStrandsRestResource(const FHairStrandsBulkData& InBulkData, EHairStrandsResourcesType InCurveType) :
 	FHairCommonResource(EHairStrandsAllocationType::Immediate),
-	PositionBuffer(), AttributeBuffer(), MaterialBuffer(), PositionOffset(InPositionOffset), RenderData(HairStrandRenderData), CurveType(InCurveType)
+	PositionBuffer(), AttributeBuffer(), MaterialBuffer(), BulkData(InBulkData), CurveType(InCurveType)
 {}
 
 void FHairStrandsRestResource::InternalAllocate(FRDGBuilder& GraphBuilder)
 {
-	InternalCreateVertexBufferRDG<FHairStrandsPositionFormat>(GraphBuilder, RenderData.Positions, PositionBuffer, HAIRSTRANDS_RESOUCE_NAME(CurveType, Hair.StrandsRest_PositionBuffer), EHairResourceUsageType::Static);
-	InternalCreateVertexBufferRDG<FHairStrandsAttributeFormat>(GraphBuilder, RenderData.Attributes, AttributeBuffer, HAIRSTRANDS_RESOUCE_NAME(CurveType, Hair.StrandsRest_AttributeBuffer), EHairResourceUsageType::Static);
-	InternalCreateVertexBufferRDG<FHairStrandsMaterialFormat>(GraphBuilder, RenderData.Materials, MaterialBuffer, HAIRSTRANDS_RESOUCE_NAME(CurveType, Hair.StrandsRest_MaterialBuffer), EHairResourceUsageType::Static);
+	InternalCreateVertexBufferRDG<FHairStrandsPositionFormat>(GraphBuilder, BulkData.Positions, PositionBuffer, HAIRSTRANDS_RESOUCE_NAME(CurveType, Hair.StrandsRest_PositionBuffer), EHairResourceUsageType::Static);
+	InternalCreateVertexBufferRDG<FHairStrandsAttributeFormat>(GraphBuilder, BulkData.Attributes, AttributeBuffer, HAIRSTRANDS_RESOUCE_NAME(CurveType, Hair.StrandsRest_AttributeBuffer), EHairResourceUsageType::Static);
+	InternalCreateVertexBufferRDG<FHairStrandsMaterialFormat>(GraphBuilder, BulkData.Materials, MaterialBuffer, HAIRSTRANDS_RESOUCE_NAME(CurveType, Hair.StrandsRest_MaterialBuffer), EHairResourceUsageType::Static);
 
 	TArray<FVector4> RestOffset;
-	RestOffset.Add(PositionOffset);
+	RestOffset.Add(BulkData.GetPositionOffset());
 	InternalCreateVertexBufferRDG<FHairStrandsPositionOffsetFormat>(GraphBuilder, RestOffset, PositionOffsetBuffer, HAIRSTRANDS_RESOUCE_NAME(CurveType, Hair.StrandsRest_PositionOffsetBuffer), EHairResourceUsageType::Static, ERDGInitialDataFlags::None);
 }
 
@@ -591,13 +585,12 @@ FRDGExternalBuffer FHairStrandsRestResource::GetTangentBuffer(FRDGBuilder& Graph
 	// Lazy allocation and update
 	if (TangentBuffer.Buffer == nullptr)
 	{
-		const uint32 VertexCount = RenderData.Positions.Num();
-		InternalCreateVertexBufferRDG<FHairStrandsTangentFormat>(GraphBuilder, VertexCount * FHairStrandsTangentFormat::ComponentCount, TangentBuffer, TEXT("Hair.StrandsRest_TangentBuffer"), EHairResourceUsageType::Dynamic);
+		InternalCreateVertexBufferRDG<FHairStrandsTangentFormat>(GraphBuilder, BulkData.PointCount * FHairStrandsTangentFormat::ComponentCount, TangentBuffer, TEXT("Hair.StrandsRest_TangentBuffer"), EHairResourceUsageType::Dynamic);
 
 		AddHairTangentPass(
 			GraphBuilder,
 			ShaderMap,
-			VertexCount,
+			BulkData.PointCount,
 			nullptr,
 			RegisterAsSRV(GraphBuilder, PositionBuffer),
 			Register(GraphBuilder, TangentBuffer, ERDGImportedBufferFlags::CreateUAV));
@@ -617,28 +610,30 @@ void FHairStrandsRestResource::InternalRelease()
 
 /////////////////////////////////////////////////////////////////////////////////////////
 
-FHairStrandsDeformedResource::FHairStrandsDeformedResource(const FHairStrandsDatas::FRenderData& HairStrandRenderData, bool bInInitializedData, const FVector& InDefaultOffset, EHairStrandsResourcesType InCurveType) :
+FHairStrandsDeformedResource::FHairStrandsDeformedResource(const FHairStrandsBulkData& InBulkData, bool bInInitializedData, EHairStrandsResourcesType InCurveType) :
 	FHairCommonResource(EHairStrandsAllocationType::Deferred),
-	RenderData(HairStrandRenderData), bInitializedData(bInInitializedData), DefaultOffset(InDefaultOffset), CurveType(InCurveType)
-{}
+	BulkData(InBulkData), bInitializedData(bInInitializedData), CurveType(InCurveType)
+{
+	GetPositionOffset(FHairStrandsDeformedResource::EFrameType::Current)  = BulkData.GetPositionOffset();
+	GetPositionOffset(FHairStrandsDeformedResource::EFrameType::Previous) = BulkData.GetPositionOffset();
+}
 
 void FHairStrandsDeformedResource::InternalAllocate(FRDGBuilder& GraphBuilder)
 {
-	const uint32 VertexCount = RenderData.Positions.Num();
 	if (bInitializedData)
 	{
-		InternalCreateVertexBufferRDG<FHairStrandsPositionFormat>(GraphBuilder, RenderData.Positions, DeformedPositionBuffer[0], HAIRSTRANDS_RESOUCE_NAME(CurveType, Hair.StrandsDeformed_DeformedPositionBuffer0), EHairResourceUsageType::Dynamic);
-		InternalCreateVertexBufferRDG<FHairStrandsPositionFormat>(GraphBuilder, RenderData.Positions, DeformedPositionBuffer[1], HAIRSTRANDS_RESOUCE_NAME(CurveType, Hair.StrandsDeformed_DeformedPositionBuffer1), EHairResourceUsageType::Dynamic);
+		InternalCreateVertexBufferRDG<FHairStrandsPositionFormat>(GraphBuilder, BulkData.Positions, DeformedPositionBuffer[0], HAIRSTRANDS_RESOUCE_NAME(CurveType, Hair.StrandsDeformed_DeformedPositionBuffer0), EHairResourceUsageType::Dynamic);
+		InternalCreateVertexBufferRDG<FHairStrandsPositionFormat>(GraphBuilder, BulkData.Positions, DeformedPositionBuffer[1], HAIRSTRANDS_RESOUCE_NAME(CurveType, Hair.StrandsDeformed_DeformedPositionBuffer1), EHairResourceUsageType::Dynamic);
 	}
 	else
 	{
-		InternalCreateVertexBufferRDG<FHairStrandsPositionFormat>(GraphBuilder, VertexCount, DeformedPositionBuffer[0], HAIRSTRANDS_RESOUCE_NAME(CurveType, Hair.StrandsDeformed_DeformedPositionBuffer0), EHairResourceUsageType::Dynamic);
-		InternalCreateVertexBufferRDG<FHairStrandsPositionFormat>(GraphBuilder, VertexCount, DeformedPositionBuffer[1], HAIRSTRANDS_RESOUCE_NAME(CurveType, Hair.StrandsDeformed_DeformedPositionBuffer1), EHairResourceUsageType::Dynamic);
+		InternalCreateVertexBufferRDG<FHairStrandsPositionFormat>(GraphBuilder, BulkData.PointCount, DeformedPositionBuffer[0], HAIRSTRANDS_RESOUCE_NAME(CurveType, Hair.StrandsDeformed_DeformedPositionBuffer0), EHairResourceUsageType::Dynamic);
+		InternalCreateVertexBufferRDG<FHairStrandsPositionFormat>(GraphBuilder, BulkData.PointCount, DeformedPositionBuffer[1], HAIRSTRANDS_RESOUCE_NAME(CurveType, Hair.StrandsDeformed_DeformedPositionBuffer1), EHairResourceUsageType::Dynamic);
 	}
-	InternalCreateVertexBufferRDG<FHairStrandsTangentFormat>(GraphBuilder, VertexCount * FHairStrandsTangentFormat::ComponentCount, TangentBuffer, HAIRSTRANDS_RESOUCE_NAME(CurveType, Hair.StrandsDeformed_TangentBuffer), EHairResourceUsageType::Dynamic);
+	InternalCreateVertexBufferRDG<FHairStrandsTangentFormat>(GraphBuilder, BulkData.PointCount * FHairStrandsTangentFormat::ComponentCount, TangentBuffer, HAIRSTRANDS_RESOUCE_NAME(CurveType, Hair.StrandsDeformed_TangentBuffer), EHairResourceUsageType::Dynamic);
 
 	TArray<FVector4> DefaultOffsets;
-	DefaultOffsets.Add(DefaultOffset);
+	DefaultOffsets.Add(BulkData.GetPositionOffset());
 	InternalCreateVertexBufferRDG<FHairStrandsPositionOffsetFormat>(GraphBuilder, DefaultOffsets, DeformedOffsetBuffer[0], HAIRSTRANDS_RESOUCE_NAME(CurveType, Hair.StrandsDeformed_DeformedOffsetBuffer0), EHairResourceUsageType::Dynamic);
 	InternalCreateVertexBufferRDG<FHairStrandsPositionOffsetFormat>(GraphBuilder, DefaultOffsets, DeformedOffsetBuffer[1], HAIRSTRANDS_RESOUCE_NAME(CurveType, Hair.StrandsDeformed_DeformedOffsetBuffer1), EHairResourceUsageType::Dynamic);
 }
@@ -1013,29 +1008,17 @@ void FHairStrandsRootData::Reset()
 
 /////////////////////////////////////////////////////////////////////////////////////////
 
-FHairStrandsInterpolationResource::FHairStrandsInterpolationResource(const FHairStrandsInterpolationDatas::FRenderData& InterpolationRenderData, const FHairStrandsDatas& SimDatas) :
+FHairStrandsInterpolationResource::FHairStrandsInterpolationResource(const FHairStrandsInterpolationBulkData& InBulkData) :
 	FHairCommonResource(EHairStrandsAllocationType::Immediate),
-	Interpolation0Buffer(), Interpolation1Buffer(), RenderData(InterpolationRenderData)
+	Interpolation0Buffer(), Interpolation1Buffer(), BulkData(InBulkData)
 {
-	const uint32 RootCount = SimDatas.GetNumCurves();
-	SimRootPointIndex.SetNum(SimDatas.GetNumPoints());
-	for (uint32 CurveIndex = 0; CurveIndex < RootCount; ++CurveIndex)
-	{
-		const uint16 PointCount = SimDatas.StrandsCurves.CurvesCount[CurveIndex];
-		const uint32 PointOffset = SimDatas.StrandsCurves.CurvesOffset[CurveIndex];
-		for (uint32 PointIndex = 0; PointIndex < PointCount; ++PointIndex)
-		{
-			SimRootPointIndex[PointIndex + PointOffset] = PointOffset;
-		}
-	}
 }
 
 void FHairStrandsInterpolationResource::InternalAllocate(FRDGBuilder& GraphBuilder)
 {
-	InternalCreateVertexBufferRDG<FHairStrandsInterpolation0Format>(GraphBuilder, RenderData.Interpolation0, Interpolation0Buffer, TEXT("Hair.StrandsInterpolation_Interpolation0Buffer"), EHairResourceUsageType::Static);
-	InternalCreateVertexBufferRDG<FHairStrandsInterpolation1Format>(GraphBuilder, RenderData.Interpolation1, Interpolation1Buffer, TEXT("Hair.StrandsInterpolation_Interpolation1Buffer"), EHairResourceUsageType::Static);
-	InternalCreateVertexBufferRDG<FHairStrandsRootIndexFormat>(GraphBuilder, SimRootPointIndex, SimRootPointIndexBuffer, TEXT("Hair.StrandsInterpolation_SimRootPointIndex"), EHairResourceUsageType::Static);
-	//SimRootPointIndex.SetNum(0);
+	InternalCreateVertexBufferRDG<FHairStrandsInterpolation0Format>(GraphBuilder, BulkData.Interpolation0, Interpolation0Buffer, TEXT("Hair.StrandsInterpolation_Interpolation0Buffer"), EHairResourceUsageType::Static);
+	InternalCreateVertexBufferRDG<FHairStrandsInterpolation1Format>(GraphBuilder, BulkData.Interpolation1, Interpolation1Buffer, TEXT("Hair.StrandsInterpolation_Interpolation1Buffer"), EHairResourceUsageType::Static);
+	InternalCreateVertexBufferRDG<FHairStrandsRootIndexFormat>(GraphBuilder, BulkData.SimRootPointIndex, SimRootPointIndexBuffer, TEXT("Hair.StrandsInterpolation_SimRootPointIndex"), EHairResourceUsageType::Static);
 }
 
 void FHairStrandsInterpolationResource::InternalRelease()
@@ -1046,16 +1029,6 @@ void FHairStrandsInterpolationResource::InternalRelease()
 }
 
 /////////////////////////////////////////////////////////////////////////////////////////
-FArchive& operator<<(FArchive& Ar, FHairCardsInterpolationDatas& CardInterpData)
-{
-	Ar << CardInterpData.PointsSimCurvesIndex;
-	Ar << CardInterpData.PointsSimCurvesVertexIndex;
-	Ar << CardInterpData.PointsSimCurvesVertexLerp;
-	Ar << CardInterpData.RenderData.Interpolation;
-
-	return Ar;
-}
-
 void FHairCardsInterpolationDatas::SetNum(const uint32 NumPoints)
 {
 	PointsSimCurvesIndex.SetNum(NumPoints);
@@ -1070,15 +1043,28 @@ void FHairCardsInterpolationDatas::Reset()
 	PointsSimCurvesVertexLerp.Empty();
 }
 
-FHairCardsInterpolationResource::FHairCardsInterpolationResource(const FHairCardsInterpolationDatas::FRenderData& InterpolationRenderData) :
+void FHairCardsInterpolationDatas::Serialize(FArchive& Ar, FHairCardsInterpolationBulkData& BulkData)
+{
+	Ar << PointsSimCurvesIndex;
+	Ar << PointsSimCurvesVertexIndex;
+	Ar << PointsSimCurvesVertexLerp;
+	BulkData.Serialize(Ar);
+}
+
+void FHairCardsInterpolationBulkData::Serialize(FArchive& Ar)
+{
+	Ar << Interpolation;
+}
+
+FHairCardsInterpolationResource::FHairCardsInterpolationResource(const FHairCardsInterpolationBulkData& InBulkData) :
 	FHairCommonResource(EHairStrandsAllocationType::Immediate),
-	InterpolationBuffer(), RenderData(InterpolationRenderData)
+	InterpolationBuffer(), BulkData(InBulkData)
 {
 }
 
 void FHairCardsInterpolationResource::InternalAllocate(FRDGBuilder& GraphBuilder)
 {
-	InternalCreateVertexBufferRDG<FHairCardsInterpolationFormat>(GraphBuilder, RenderData.Interpolation, InterpolationBuffer, TEXT("Hair.CardsInterpolation_InterpolationBuffer"), EHairResourceUsageType::Static);
+	InternalCreateVertexBufferRDG<FHairCardsInterpolationFormat>(GraphBuilder, BulkData.Interpolation, InterpolationBuffer, TEXT("Hair.CardsInterpolation_InterpolationBuffer"), EHairResourceUsageType::Static);
 }
 
 void FHairCardsInterpolationResource::InternalRelease()
@@ -1091,19 +1077,19 @@ void FHairCardsInterpolationResource::InternalRelease()
 #if RHI_RAYTRACING
 // RT geometry is built to for a cross around the fiber.
 // 4 triangles per hair vertex => 12 vertices per hair vertex
-FHairStrandsRaytracingResource::FHairStrandsRaytracingResource(const FHairStrandsDatas& InData) :
+FHairStrandsRaytracingResource::FHairStrandsRaytracingResource(const FHairStrandsBulkData& InData) :
 	FHairCommonResource(EHairStrandsAllocationType::Deferred),
 	PositionBuffer(), VertexCount(InData.GetNumPoints()*12)  
 {}
 
-FHairStrandsRaytracingResource::FHairStrandsRaytracingResource(const FHairCardsDatas& InData) :
+FHairStrandsRaytracingResource::FHairStrandsRaytracingResource(const FHairCardsBulkData& InData) :
 	FHairCommonResource(EHairStrandsAllocationType::Deferred),
-	PositionBuffer(), VertexCount(InData.Cards.GetNumVertices())
+	PositionBuffer(), VertexCount(InData.GetNumVertices())
 {}
 
-FHairStrandsRaytracingResource::FHairStrandsRaytracingResource(const FHairMeshesDatas& InData) :
+FHairStrandsRaytracingResource::FHairStrandsRaytracingResource(const FHairMeshesBulkData& InData) :
 	FHairCommonResource(EHairStrandsAllocationType::Deferred),
-	PositionBuffer(), VertexCount(InData.Meshes.GetNumVertices())
+	PositionBuffer(), VertexCount(InData.GetNumVertices())
 {}
 
 void FHairStrandsRaytracingResource::InternalAllocate(FRDGBuilder& GraphBuilder)
diff --git a/Engine/Plugins/Runtime/HairStrands/Source/HairStrandsCore/Private/GroomTextureBuilder.cpp b/Engine/Plugins/Runtime/HairStrands/Source/HairStrandsCore/Private/GroomTextureBuilder.cpp
--- a/Engine/Plugins/Runtime/HairStrands/Source/HairStrandsCore/Private/GroomTextureBuilder.cpp
+++ b/Engine/Plugins/Runtime/HairStrands/Source/HairStrandsCore/Private/GroomTextureBuilder.cpp
@@ -964,10 +964,10 @@ static void InternalBuildStrandsTextures_GPU(
 
 				GroupData.Strands.RestResource->PositionBuffer.SRV,
 				GroupData.Strands.RestResource->AttributeBuffer.SRV,
-				GroupData.Strands.RestResource->PositionOffset,
+				GroupData.Strands.RestResource->GetPositionOffset(),
 				RenderingData.GeometrySettings.HairWidth * 0.5f,
-				GroupData.Strands.Data.StrandsCurves.MaxLength,
-				GroupData.Strands.Data.RenderData.Positions.Num(),
+				GroupData.Strands.RestResource->BulkData.MaxLength,
+				GroupData.Strands.RestResource->BulkData.PointCount,
 
 				DepthTestTexture,
 				DepthTexture[0],
diff --git a/Engine/Plugins/Runtime/HairStrands/Source/HairStrandsCore/Private/HairCardsBuilder.cpp b/Engine/Plugins/Runtime/HairStrands/Source/HairStrandsCore/Private/HairCardsBuilder.cpp
--- a/Engine/Plugins/Runtime/HairStrands/Source/HairStrandsCore/Private/HairCardsBuilder.cpp
+++ b/Engine/Plugins/Runtime/HairStrands/Source/HairStrandsCore/Private/HairCardsBuilder.cpp
@@ -3163,7 +3163,7 @@ namespace HairCards
 
 					const float StrandsU = InStrands.StrandsPoints.PointsCoordU[PointIndex];
 					const FVector2D RootUV = InStrands.StrandsCurves.CurvesRootUV[CurveIndex];
-					const float Seed = float(InStrands.RenderData.Attributes[PointIndex].Seed) / 255.f;
+					const float Seed = 0.f; // TODO float(InStrands.RenderData.Attributes[PointIndex].Seed) / 255.f;
 
 					const FVector& P0_Rest = InStrands.StrandsPoints.PointsPosition[PointIndex];
 					const float Radius = InStrands.StrandsPoints.PointsRadius[PointIndex] * InStrands.StrandsCurves.MaxRadius;
@@ -3383,6 +3383,7 @@ void BuildGeometry(
 	FHairCardsProceduralDatas& Out,
 	FHairStrandsDatas& OutGuides,
 	FHairCardsInterpolationDatas& OutInterpolation,
+	FHairCardsInterpolationBulkData& OutInterpolationBulk,
 	FHairGroupCardsTextures& OutTextures)
 {
 	// Basic algo:
@@ -3529,7 +3530,7 @@ void BuildGeometry(
 	}
 
 	// Interpolation data
-	OutInterpolation.RenderData.Interpolation.SetNum(PointCount);
+	OutInterpolationBulk.Interpolation.SetNum(PointCount);
 	for (uint32 PointIt = 0; PointIt < PointCount; ++PointIt)
 	{
 		const uint32 VertexIndex = OutInterpolation.PointsSimCurvesVertexIndex[PointIt];
@@ -3537,7 +3538,7 @@ void BuildGeometry(
 		FHairCardsInterpolationVertex PackedData;
 		PackedData.VertexIndex = VertexIndex;
 		PackedData.VertexLerp  = FMath::Clamp(uint32(VertexLerp * 0xFF), 0u, 0xFFu);
-		OutInterpolation.RenderData.Interpolation[PointIt] = PackedData;
+		OutInterpolationBulk.Interpolation[PointIt] = PackedData;
 	}
 }
 
@@ -3563,8 +3564,10 @@ void SanitizeMeshDescription(FMeshDescription* MeshDescription)
 bool ImportGeometry(
 	const UStaticMesh* StaticMesh,
 	FHairCardsDatas& Out,
+	FHairCardsBulkData& OutBulk,
 	FHairStrandsDatas& OutGuides,
-	FHairCardsInterpolationDatas& OutInterpolationData)
+	FHairCardsInterpolationDatas& OutInterpolationData,
+	FHairCardsInterpolationBulkData& OutInterpolationBulkData)
 {
 	const uint32 MeshLODIndex = 0;
 
@@ -3685,27 +3688,27 @@ bool ImportGeometry(
 
 	// Fill in render resources (do we need to keep it separated? e.g, format compression, packing)
 	const uint32 PointCount = Out.Cards.Positions.Num();
-	Out.RenderData.Positions.SetNum(PointCount);
-	Out.RenderData.Normals.SetNum(PointCount * FHairCardsNormalFormat::ComponentCount);
-	Out.RenderData.UVs.SetNum(PointCount);
+	OutBulk.Positions.SetNum(PointCount);
+	OutBulk.Normals.SetNum(PointCount * FHairCardsNormalFormat::ComponentCount);
+	OutBulk.UVs.SetNum(PointCount);
 	for (uint32 PointIt = 0; PointIt < PointCount; ++PointIt)
 	{
-		Out.RenderData.Positions[PointIt] = FVector4(Out.Cards.Positions[PointIt], 0);
-		Out.RenderData.UVs[PointIt] = Out.Cards.UVs[PointIt];
-		Out.RenderData.Normals[PointIt * 2] = FVector4(Out.Cards.Tangents[PointIt], 0);
-		Out.RenderData.Normals[PointIt * 2 + 1] = FVector4(Out.Cards.Normals[PointIt], TangentFrameSigns[PointIt]);
+		OutBulk.Positions[PointIt] = FVector4(Out.Cards.Positions[PointIt], 0);
+		OutBulk.UVs[PointIt] = Out.Cards.UVs[PointIt];
+		OutBulk.Normals[PointIt * 2] = FVector4(Out.Cards.Tangents[PointIt], 0);
+		OutBulk.Normals[PointIt * 2 + 1] = FVector4(Out.Cards.Normals[PointIt], TangentFrameSigns[PointIt]);
 	}
 
-	Out.RenderData.Indices.SetNum(IndexCount);
+	OutBulk.Indices.SetNum(IndexCount);
 	for (uint32 IndexIt = 0; IndexIt < IndexCount; ++IndexIt)
 	{
-		Out.RenderData.Indices[IndexIt] = Out.Cards.Indices[IndexIt];
+		OutBulk.Indices[IndexIt] = Out.Cards.Indices[IndexIt];
 	}
 
-	Out.DepthTexture = nullptr;
-	Out.TangentTexture = nullptr;
-	Out.CoverageTexture = nullptr;
-	Out.AttributeTexture = nullptr;
+	OutBulk.DepthTexture = nullptr;
+	OutBulk.TangentTexture = nullptr;
+	OutBulk.CoverageTexture = nullptr;
+	OutBulk.AttributeTexture = nullptr;
 
 	TArray<float> CardLengths;
 	CardLengths.Reserve(Out.Cards.IndexOffsets.Num());
@@ -3715,7 +3718,7 @@ bool ImportGeometry(
 		HairCards::CreateCardsInterpolation(Out.Cards, OutGuides, OutInterpolationData, CardLengths);
 
 		// Fill out the interpolation data
-		OutInterpolationData.RenderData.Interpolation.SetNum(PointCount);
+		OutInterpolationBulkData.Interpolation.SetNum(PointCount);
 		for (uint32 PointIt = 0; PointIt < PointCount; ++PointIt)
 		{
 			const uint32 InterpVertexIndex = OutInterpolationData.PointsSimCurvesVertexIndex[PointIt];
@@ -3723,21 +3726,21 @@ bool ImportGeometry(
 			FHairCardsInterpolationVertex PackedData;
 			PackedData.VertexIndex = InterpVertexIndex;
 			PackedData.VertexLerp = FMath::Clamp(uint32(VertexLerp * 0xFF), 0u, 0xFFu);
-			OutInterpolationData.RenderData.Interpolation[PointIt] = PackedData;
+			OutInterpolationBulkData.Interpolation[PointIt] = PackedData;
 		}
 	}
 
 	return bSuccess;
 }
 
-void Convert(const FHairCardsProceduralDatas& In, FHairCardsDatas& Out)
+void Convert(const FHairCardsProceduralDatas& In, FHairCardsDatas& Out, FHairCardsBulkData& OutBulk)
 {
 	Out.Cards = In.Cards;
 
-	Out.RenderData.Positions = In.RenderData.Positions;
-	Out.RenderData.Normals = In.RenderData.Normals;
-	Out.RenderData.Indices = In.RenderData.Indices;
-	Out.RenderData.UVs = In.RenderData.UVs;
+	OutBulk.Positions = In.RenderData.Positions;
+	OutBulk.Normals = In.RenderData.Normals;
+	OutBulk.Indices = In.RenderData.Indices;
+	OutBulk.UVs = In.RenderData.UVs;
 }
 
 ///////////////////////////////////////////////////////////////////////////////////////////////////
@@ -3899,7 +3902,8 @@ FString GetVersion()
 void BuildGeometry(
 	const FHairStrandsDatas& InRen,
 	const FHairStrandsDatas& InSim,
-	FHairMeshesDatas& Out)
+	FHairMeshesDatas& Out,
+	FHairMeshesBulkData& OutBulk)
 {
 	const FVector Center = InRen.BoundingBox.GetCenter();
 	const FVector Extent = InRen.BoundingBox.GetExtent();
@@ -3996,30 +4000,31 @@ void BuildGeometry(
 
 	// Fill in render resources (do we need to keep it separated? e.g, format compression, packing)
 	const uint32 PointCount = Out.Meshes.Positions.Num();
-	Out.RenderData.Positions.SetNum(PointCount);
-	Out.RenderData.Normals.SetNum(PointCount * FHairCardsNormalFormat::ComponentCount);
-	Out.RenderData.UVs.SetNum(PointCount);
+	OutBulk.Positions.SetNum(PointCount);
+	OutBulk.Normals.SetNum(PointCount * FHairCardsNormalFormat::ComponentCount);
+	OutBulk.UVs.SetNum(PointCount);
 	for (uint32 PointIt = 0; PointIt < PointCount; ++PointIt)
 	{
-		Out.RenderData.Positions[PointIt] = FVector4(Out.Meshes.Positions[PointIt], 0);
-		Out.RenderData.UVs[PointIt] = FVector4(Out.Meshes.UVs[PointIt].X, Out.Meshes.UVs[PointIt].Y, 0, 0);
-		Out.RenderData.Normals[PointIt * 2] = FVector4(Out.Meshes.Tangents[PointIt], 0);
-		Out.RenderData.Normals[PointIt * 2 + 1] = FVector4(Out.Meshes.Normals[PointIt], 1);
+		OutBulk.Positions[PointIt] = FVector4(Out.Meshes.Positions[PointIt], 0);
+		OutBulk.UVs[PointIt] = FVector4(Out.Meshes.UVs[PointIt].X, Out.Meshes.UVs[PointIt].Y, 0, 0);
+		OutBulk.Normals[PointIt * 2] = FVector4(Out.Meshes.Tangents[PointIt], 0);
+		OutBulk.Normals[PointIt * 2 + 1] = FVector4(Out.Meshes.Normals[PointIt], 1);
 
-		Out.Meshes.BoundingBox += Out.RenderData.Positions[PointIt];
+		Out.Meshes.BoundingBox += OutBulk.Positions[PointIt];
 	}
 
 	const uint32 IndexCount = Out.Meshes.Indices.Num();
-	Out.RenderData.Indices.SetNum(IndexCount);
+	OutBulk.Indices.SetNum(IndexCount);
 	for (uint32 IndexIt = 0; IndexIt < IndexCount; ++IndexIt)
 	{
-		Out.RenderData.Indices[IndexIt] = Out.Meshes.Indices[IndexIt];
+		OutBulk.Indices[IndexIt] = Out.Meshes.Indices[IndexIt];
 	}
 }
 
 void ImportGeometry(
 	const UStaticMesh* StaticMesh,
-	FHairMeshesDatas& Out)
+	FHairMeshesDatas& Out,
+	FHairMeshesBulkData& OutBulk)
 {
 	const uint32 MeshLODIndex = 0;
 
@@ -4052,21 +4057,21 @@ void ImportGeometry(
 
 	// Fill in render resources (do we need to keep it separated? e.g, format compression, packing)
 	const uint32 PointCount = Out.Meshes.Positions.Num();
-	Out.RenderData.Positions.SetNum(PointCount);
-	Out.RenderData.Normals.SetNum(PointCount * FHairCardsNormalFormat::ComponentCount);
-	Out.RenderData.UVs.SetNum(PointCount);
+	OutBulk.Positions.SetNum(PointCount);
+	OutBulk.Normals.SetNum(PointCount * FHairCardsNormalFormat::ComponentCount);
+	OutBulk.UVs.SetNum(PointCount);
 	for (uint32 PointIt = 0; PointIt < PointCount; ++PointIt)
 	{
-		Out.RenderData.Positions[PointIt] = FVector4(Out.Meshes.Positions[PointIt], 0);
-		Out.RenderData.UVs[PointIt] = FVector4(Out.Meshes.UVs[PointIt].X, Out.Meshes.UVs[PointIt].Y, 0, 0);
-		Out.RenderData.Normals[PointIt * 2] = FVector4(Out.Meshes.Tangents[PointIt], 0);
-		Out.RenderData.Normals[PointIt * 2 + 1] = FVector4(Out.Meshes.Normals[PointIt], 1);
+		OutBulk.Positions[PointIt] = FVector4(Out.Meshes.Positions[PointIt], 0);
+		OutBulk.UVs[PointIt] = FVector4(Out.Meshes.UVs[PointIt].X, Out.Meshes.UVs[PointIt].Y, 0, 0);
+		OutBulk.Normals[PointIt * 2] = FVector4(Out.Meshes.Tangents[PointIt], 0);
+		OutBulk.Normals[PointIt * 2 + 1] = FVector4(Out.Meshes.Normals[PointIt], 1);
 	}
 
-	Out.RenderData.Indices.SetNum(IndexCount);
+	OutBulk.Indices.SetNum(IndexCount);
 	for (uint32 IndexIt = 0; IndexIt < IndexCount; ++IndexIt)
 	{
-		Out.RenderData.Indices[IndexIt] = Out.Meshes.Indices[IndexIt];
+		OutBulk.Indices[IndexIt] = Out.Meshes.Indices[IndexIt];
 	}
 }
 
diff --git a/Engine/Plugins/Runtime/HairStrands/Source/HairStrandsCore/Private/HairCardsDatas.cpp b/Engine/Plugins/Runtime/HairStrands/Source/HairStrandsCore/Private/HairCardsDatas.cpp
--- a/Engine/Plugins/Runtime/HairStrands/Source/HairStrandsCore/Private/HairCardsDatas.cpp
+++ b/Engine/Plugins/Runtime/HairStrands/Source/HairStrandsCore/Private/HairCardsDatas.cpp
@@ -42,22 +42,18 @@ FArchive& operator<<(FArchive& Ar, FHairCardsGeometry& CardGeometry)
 	return Ar;
 }
 
-FArchive& operator<<(FArchive& Ar, FHairCardsDatas::FRenderData& CardRenderData)
+void FHairCardsBulkData::Serialize(FArchive& Ar)
 {
-	Ar << CardRenderData.Positions;
-	Ar << CardRenderData.Normals;
-	Ar << CardRenderData.UVs;
-	Ar << CardRenderData.Indices;
-
-	return Ar;
+	Ar << Positions;
+	Ar << Normals;
+	Ar << UVs;
+	Ar << Indices;
 }
 
-FArchive& operator<<(FArchive& Ar, FHairCardsDatas& CardData)
+void FHairCardsDatas::Serialize(FArchive& Ar, FHairCardsBulkData& BulkData)
 {
-	Ar << CardData.Cards;
-	Ar << CardData.RenderData;
-
-	return Ar;
+	Ar << Cards;
+	BulkData.Serialize(Ar);
 }
 
 FArchive& operator<<(FArchive& Ar, FHairCardsInterpolationVertex& CardInterpVertex)
@@ -99,20 +95,17 @@ FArchive& operator<<(FArchive& Ar, FHairMeshes& HairMesh)
 	return Ar;
 }
 
-FArchive& operator<<(FArchive& Ar, FHairMeshesDatas::FRenderData& MeshRenderData)
+void FHairMeshesDatas::Serialize(FArchive& Ar, FHairMeshesBulkData& InBulkData)
 {
-	Ar << MeshRenderData.Positions;
-	Ar << MeshRenderData.Normals;
-	Ar << MeshRenderData.UVs;
-	Ar << MeshRenderData.Indices;
-
-	return Ar;
+	Ar << Meshes;
+	InBulkData.Serialize(Ar);
 }
 
-FArchive& operator<<(FArchive& Ar, FHairMeshesDatas& MeshData)
+void FHairMeshesBulkData::Serialize(FArchive& Ar)
 {
-	Ar << MeshData.Meshes;
-	Ar << MeshData.RenderData;
-
-	return Ar;
+	Ar << Positions;
+	Ar << Normals;
+	Ar << UVs;
+	Ar << Indices;
 }
+
diff --git a/Engine/Plugins/Runtime/HairStrands/Source/HairStrandsCore/Private/HairCardsVertexFactory.cpp b/Engine/Plugins/Runtime/HairStrands/Source/HairStrandsCore/Private/HairCardsVertexFactory.cpp
--- a/Engine/Plugins/Runtime/HairStrands/Source/HairStrandsCore/Private/HairCardsVertexFactory.cpp
+++ b/Engine/Plugins/Runtime/HairStrands/Source/HairStrandsCore/Private/HairCardsVertexFactory.cpp
@@ -39,7 +39,7 @@ FHairCardsUniformBuffer CreateHairCardsVFUniformBuffer(
 
 		// Cards atlas UV are inverted so fetching needs to be inverted on the y-axis
 		UniformParameters.bInvertUV = LOD.RestResource->bInvertUV;
-		UniformParameters.MaxVertexCount = LOD.RestResource->VertexCount;
+		UniformParameters.MaxVertexCount = LOD.RestResource->GetVertexCount();
 
 		// When the geometry is not-dynamic (no binding to skeletal mesh, no simulation), only a single vertex buffer is allocated. 
 		// In this case we force the buffer index to 0
@@ -89,7 +89,7 @@ FHairCardsUniformBuffer CreateHairCardsVFUniformBuffer(
 
 		// Meshes UV are not inverted so no need to invert the y-axis
 		UniformParameters.bInvertUV = 0;
-		UniformParameters.MaxVertexCount = LOD.RestResource->VertexCount;
+		UniformParameters.MaxVertexCount = LOD.RestResource->GetVertexCount();
 		UniformParameters.NormalsBuffer = LOD.RestResource->NormalsBuffer.ShaderResourceViewRHI.GetReference();
 		UniformParameters.UVsBuffer = LOD.RestResource->UVsBuffer.ShaderResourceViewRHI.GetReference();
 
diff --git a/Engine/Plugins/Runtime/HairStrands/Source/HairStrandsCore/Private/HairStrandsDatas.cpp b/Engine/Plugins/Runtime/HairStrands/Source/HairStrandsCore/Private/HairStrandsDatas.cpp
--- a/Engine/Plugins/Runtime/HairStrands/Source/HairStrandsCore/Private/HairStrandsDatas.cpp
+++ b/Engine/Plugins/Runtime/HairStrands/Source/HairStrandsCore/Private/HairStrandsDatas.cpp
@@ -170,13 +170,15 @@ FArchive& operator<<(FArchive& Ar, FHairInterpolation1Vertex& Vertex)
 	return Ar;
 }
 
-void FHairStrandsInterpolationDatas::FRenderData::Serialize(FArchive& Ar)
+void FHairStrandsInterpolationBulkData::Serialize(FArchive& Ar)
 {
 	Ar << Interpolation0;
 	Ar << Interpolation1;
+	// TODO add here SimRootIndices serializaztion once changed.
+//	Ar << SimRootPointIndex;
 }
 
-void FHairStrandsInterpolationDatas::Serialize(FArchive& Ar)
+void FHairStrandsInterpolationDatas::Serialize(FArchive& Ar, FHairStrandsInterpolationBulkData& BulkData)
 {
 	Ar.UsingCustomVersion(FReleaseObjectVersion::GUID);
 
@@ -198,7 +200,7 @@ void FHairStrandsInterpolationDatas::Serialize(FArchive& Ar)
 		}
 	}
 
-	RenderData.Serialize(Ar);
+	BulkData.Serialize(Ar);
 }
 
 void FHairStrandsPoints::Serialize(FArchive& Ar)
@@ -232,7 +234,7 @@ void FHairStrandsCurves::Serialize(FArchive& Ar)
 	Ar << MaxRadius;
 }
 
-void FHairStrandsDatas::FRenderData::Serialize(FArchive& Ar)
+void FHairStrandsBulkData::Serialize(FArchive& Ar)
 {
 	Ar.UsingCustomVersion(FReleaseObjectVersion::GUID);
 	Ar.UsingCustomVersion(FUE5ReleaseStreamObjectVersion::GUID);
@@ -265,14 +267,23 @@ void FHairStrandsDatas::FRenderData::Serialize(FArchive& Ar)
 	}
 }
 
-void FHairStrandsDatas::Serialize(FArchive& Ar)
+void FHairStrandsDatas::Serialize(FArchive& Ar, FHairStrandsBulkData& BulkData)
 {
 	StrandsPoints.Serialize(Ar);
 	StrandsCurves.Serialize(Ar);
 	Ar << HairDensity;
 	Ar << BoundingBox;
 
-	RenderData.Serialize(Ar);
+	BulkData.Serialize(Ar);
+
+	if (Ar.IsLoading())
+	{
+		BulkData.CurveCount = StrandsCurves.Num();
+		BulkData.PointCount = StrandsPoints.Num();
+		BulkData.MaxLength  = StrandsCurves.MaxLength;
+		BulkData.MaxRadius  = StrandsCurves.MaxRadius;
+		BulkData.BoundingBox= BoundingBox;
+	}
 }
 
 void FHairStrandsDatas::Reset()
@@ -308,4 +319,4 @@ FArchive& operator<<(FArchive& Ar, FHairStrandsClusterCullingData::FHairClusterI
 	Ar << Uint32Data[2];
 	Ar << Uint32Data[3];
 	return Ar;
-}
\ No newline at end of file
+}
diff --git a/Engine/Plugins/Runtime/HairStrands/Source/HairStrandsCore/Private/HairStrandsMeshProjection.cpp b/Engine/Plugins/Runtime/HairStrands/Source/HairStrandsCore/Private/HairStrandsMeshProjection.cpp
--- a/Engine/Plugins/Runtime/HairStrands/Source/HairStrandsCore/Private/HairStrandsMeshProjection.cpp
+++ b/Engine/Plugins/Runtime/HairStrands/Source/HairStrandsCore/Private/HairStrandsMeshProjection.cpp
@@ -1006,7 +1006,7 @@ void InternalAddHairRBFInterpolationPass(
 	FHairStrandsRestRootResource* RestRootResources,
 	FHairStrandsDeformedRootResource* DeformedRootResources)
 {
-	const uint32 VertexCount = RestResources ? RestResources->VertexCount : 0;
+	const uint32 VertexCount = RestResources ? RestResources->GetVertexCount() : 0;
 	if (!RestResources || !DeformedResources || !RestRootResources || !DeformedRootResources || VertexCount == 0 || MeshLODIndex >= RestRootResources->LODs.Num())
 	{
 		return;
diff --git a/Engine/Plugins/Runtime/HairStrands/Source/HairStrandsCore/Private/Niagara/NiagaraDataInterfaceHairStrands.cpp b/Engine/Plugins/Runtime/HairStrands/Source/HairStrandsCore/Private/Niagara/NiagaraDataInterfaceHairStrands.cpp
--- a/Engine/Plugins/Runtime/HairStrands/Source/HairStrandsCore/Private/Niagara/NiagaraDataInterfaceHairStrands.cpp
+++ b/Engine/Plugins/Runtime/HairStrands/Source/HairStrandsCore/Private/Niagara/NiagaraDataInterfaceHairStrands.cpp
@@ -657,7 +657,7 @@ struct FNDIHairStrandsParametersCS : public FNiagaraDataInterfaceParametersCS
 			FVector RestRootOffsetValue = FVector::ZeroVector;
 			FVector DeformedRootOffsetValue = FVector::ZeroVector;
 			
-			FVector RestPositionOffsetValue = ProxyData->HairStrandsBuffer->SourceRestResources->PositionOffset;
+			FVector RestPositionOffsetValue = ProxyData->HairStrandsBuffer->SourceRestResources->GetPositionOffset();
 
 			FRHITransitionInfo Transitions[] = {
 				FRHITransitionInfo(PointPositionsUAV, ERHIAccess::Unknown, ERHIAccess::UAVCompute),
diff --git a/Engine/Plugins/Runtime/HairStrands/Source/HairStrandsCore/Public/GroomAsset.h b/Engine/Plugins/Runtime/HairStrands/Source/HairStrandsCore/Public/GroomAsset.h
--- a/Engine/Plugins/Runtime/HairStrands/Source/HairStrandsCore/Public/GroomAsset.h
+++ b/Engine/Plugins/Runtime/HairStrands/Source/HairStrandsCore/Public/GroomAsset.h
@@ -80,12 +80,14 @@ struct HAIRSTRANDSCORE_API FHairGroupData
 		bool HasValidData() const	{ return Data.GetNumPoints() > 0;}
 		bool IsValid() const		{ return RestResource != nullptr; }
 		FHairStrandsDatas					Data;
+		FHairStrandsBulkData				BulkData;
 		FHairStrandsRestResource*			RestResource = nullptr;
 	};
 
 	struct FBaseWithInterpolation : FBase
 	{
 		FHairStrandsInterpolationDatas		InterpolationData;
+		FHairStrandsInterpolationBulkData	InterpolationBulkData;
 		FHairStrandsInterpolationResource*	InterpolationResource = nullptr;
 	};
 
@@ -180,6 +182,7 @@ struct HAIRSTRANDSCORE_API FHairGroupData
 			bool IsValid() const { return Data.IsValid() && RestResource != nullptr; }
 			// Main data & Resources
 			FHairCardsDatas						Data;
+			FHairCardsBulkData					BulkData;
 			FHairCardsRestResource*				RestResource = nullptr;
 
 			// Procedural data & resources
@@ -188,6 +191,7 @@ struct HAIRSTRANDSCORE_API FHairGroupData
 
 			// Interpolation data & resources
 			FHairCardsInterpolationDatas		InterpolationData;
+			FHairCardsInterpolationBulkData		InterpolationBulkData;
 			FHairCardsInterpolationResource*	InterpolationResource = nullptr;
 
 			FBaseWithInterpolation				Guides;
@@ -250,6 +254,7 @@ struct HAIRSTRANDSCORE_API FHairGroupData
 			bool HasValidData() const { return Data.IsValid(); }
 			bool IsValid() const { return Data.IsValid() && RestResource != nullptr; }
 			FHairMeshesDatas Data;
+			FHairMeshesBulkData BulkData;
 			FHairMeshesRestResource* RestResource = nullptr;
 			#if RHI_RAYTRACING
 			FHairStrandsRaytracingResource* RaytracingResource = nullptr;
diff --git a/Engine/Plugins/Runtime/HairStrands/Source/HairStrandsCore/Public/GroomBuilder.h b/Engine/Plugins/Runtime/HairStrands/Source/HairStrandsCore/Public/GroomBuilder.h
--- a/Engine/Plugins/Runtime/HairStrands/Source/HairStrandsCore/Public/GroomBuilder.h
+++ b/Engine/Plugins/Runtime/HairStrands/Source/HairStrandsCore/Public/GroomBuilder.h
@@ -15,6 +15,8 @@ struct FProcessedHairDescription;
 struct FHairGroupsInterpolation;
 struct FHairStrandsInterpolationDatas;
 struct FHairInterpolationSettings;
+struct FHairStrandsBulkData;
+struct FHairStrandsInterpolationBulkData;
 class FHairDescription;
 class UGroomAsset;
 
@@ -27,7 +29,18 @@ struct HAIRSTRANDSCORE_API FGroomBuilder
 
 	static void BuildData(FHairGroupData& GroupData, const FHairGroupsInterpolation& InterpolationSettings, uint32 GroupIndex);
 	static void BuildData(UGroomAsset* GroomAsset);
-	static void BuildData(FHairStrandsDatas& RenData, FHairStrandsDatas& SimData, FHairStrandsInterpolationDatas& InterpolationData, const FHairInterpolationSettings& InterpolationSettings, const bool bBuildRen, const bool bBuildSim, const bool bBuildInterpolation, uint32 Seed);
+	static void BuildData(
+		const FHairStrandsDatas& RenData,
+		const FHairStrandsDatas& SimData,
+		FHairStrandsBulkData& RenBulkData,
+		FHairStrandsBulkData& SimBulkData,
+		FHairStrandsInterpolationDatas& InterpolationData,
+		FHairStrandsInterpolationBulkData& InterpolationBulkData,
+		const FHairInterpolationSettings& InterpolationSettings,
+		const bool bBuildRen,
+		const bool bBuildSim,
+		const bool bBuildInterpolation,
+		uint32 Seed);
 
 	static bool BuildGroom(FProcessedHairDescription& ProcessedHairDescription, UGroomAsset* GroomAsset, uint32 GroupIndex);
 	static bool BuildGroom(const FHairDescription& HairDescription, UGroomAsset* GroomAsset);
diff --git a/Engine/Plugins/Runtime/HairStrands/Source/HairStrandsCore/Public/GroomInstance.h b/Engine/Plugins/Runtime/HairStrands/Source/HairStrandsCore/Public/GroomInstance.h
--- a/Engine/Plugins/Runtime/HairStrands/Source/HairStrandsCore/Public/GroomInstance.h
+++ b/Engine/Plugins/Runtime/HairStrands/Source/HairStrandsCore/Public/GroomInstance.h
@@ -50,7 +50,7 @@ struct HAIRSTRANDSCORE_API FHairGroupInstance : public FHairStrandsInstance
 		bool IsValid() const { return RestResource != nullptr; }
 
 		// Data - Render & sim (rest) data
-		FHairStrandsDatas* Data = nullptr;
+		FHairStrandsBulkData* Data = nullptr;
 
 		// Resources - Strands rest position data for sim & render strands
 		// Resources - Strands deformed position data for sim & render strands
@@ -69,7 +69,6 @@ struct HAIRSTRANDSCORE_API FHairGroupInstance : public FHairStrandsInstance
 	{
 		// Data - Interpolation data (weights/Id/...) for transfering sim strands (i.e. guide) motion to render strands
 		// Resources - Strands deformed position data for sim & render strands
-		FHairStrandsInterpolationDatas* InterpolationData = nullptr;
 		FHairStrandsInterpolationResource* InterpolationResource = nullptr;
 
 		EHairInterpolationType HairInterpolationType = EHairInterpolationType::NoneSkinning;
@@ -117,7 +116,7 @@ struct HAIRSTRANDSCORE_API FHairGroupInstance : public FHairStrandsInstance
 			bool IsValid() const { return RestResource != nullptr; }
 
 			// Data
-			FHairCardsDatas* Data = nullptr;
+			FHairCardsBulkData* Data = nullptr;
 
 			// Resources
 			FHairCardsRestResource* RestResource = nullptr;
@@ -128,7 +127,6 @@ struct HAIRSTRANDSCORE_API FHairGroupInstance : public FHairStrandsInstance
 			#endif
 
 			// Interpolation data/resources
-			FHairCardsInterpolationDatas* InterpolationData = nullptr;
 			FHairCardsInterpolationResource* InterpolationResource = nullptr;
 
 			FStrandsBaseWithInterpolation Guides;
@@ -157,7 +155,7 @@ struct HAIRSTRANDSCORE_API FHairGroupInstance : public FHairStrandsInstance
 			bool IsValid() const { return RestResource != nullptr; }
 
 			// Data
-			FHairMeshesDatas* Data = nullptr;
+			FHairMeshesBulkData* Data = nullptr;
 
 			// Resources
 			FHairMeshesRestResource* RestResource = nullptr;
@@ -223,4 +221,4 @@ struct HAIRSTRANDSCORE_API FHairGroupInstance : public FHairStrandsInstance
 	{
 		return Meshes.IsValid() || Cards.IsValid() || Strands.IsValid();
 	}
-};
\ No newline at end of file
+};
diff --git a/Engine/Plugins/Runtime/HairStrands/Source/HairStrandsCore/Public/GroomResources.h b/Engine/Plugins/Runtime/HairStrands/Source/HairStrandsCore/Public/GroomResources.h
--- a/Engine/Plugins/Runtime/HairStrands/Source/HairStrandsCore/Public/GroomResources.h
+++ b/Engine/Plugins/Runtime/HairStrands/Source/HairStrandsCore/Public/GroomResources.h
@@ -197,7 +197,7 @@ struct FHairStrandsDeformedRootResource : public FHairCommonResource
 struct FHairStrandsRestResource : public FHairCommonResource
 {
 	/** Build the hair strands resource */
-	FHairStrandsRestResource(const FHairStrandsDatas::FRenderData& HairStrandRenderData, const FVector& PositionOffset, EHairStrandsResourcesType CurveType);
+	FHairStrandsRestResource(const FHairStrandsBulkData& InBulkData, EHairStrandsResourcesType CurveType);
 
 	/* Init/Release buffers */
 	virtual void InternalAllocate(FRDGBuilder& GraphBuilder) override;
@@ -207,6 +207,8 @@ struct FHairStrandsRestResource : public FHairCommonResource
 	virtual FString GetFriendlyName() const override { return TEXT("FHairStrandsResource"); }
 
 	FRDGExternalBuffer GetTangentBuffer(class FRDGBuilder& GraphBuilder, class FGlobalShaderMap* ShaderMap);
+	
+	FVector GetPositionOffset() const { return BulkData.GetPositionOffset(); }
 
 	/* Return the memory size for GPU resources */
 	uint32 GetResourcesSize() const
@@ -234,23 +236,20 @@ struct FHairStrandsRestResource : public FHairCommonResource
 
 	/* Strand hair material buffer */
 	FRDGExternalBuffer MaterialBuffer;
-	
-	/* Position offset as the rest positions are expressed in relative coordinate (16bits) */
-	FVector PositionOffset = FVector::ZeroVector;
 
 	/* Reference to the hair strands render data */
-	const FHairStrandsDatas::FRenderData& RenderData;
+	const FHairStrandsBulkData& BulkData;
 
 	/* Type of curves */
 	const EHairStrandsResourcesType CurveType;
 
-	inline uint32 GetVertexCount() const { return RenderData.Positions.Num() / FHairStrandsPositionFormat::ComponentCount; }
+	inline uint32 GetVertexCount() const { return BulkData.PointCount; }
 };
 
 struct FHairStrandsDeformedResource : public FHairCommonResource
 {
 	/** Build the hair strands resource */
-	FHairStrandsDeformedResource(const FHairStrandsDatas::FRenderData& HairStrandRenderData, bool bInitializeData, const FVector& InDefaultOffset, EHairStrandsResourcesType CurveType);
+	FHairStrandsDeformedResource(const FHairStrandsBulkData& BulkData, bool bInitializeData, EHairStrandsResourcesType CurveType);
 
 	/* Init/Release buffers */
 	virtual void InternalAllocate(FRDGBuilder& GraphBuilder) override;
@@ -282,7 +281,7 @@ struct FHairStrandsDeformedResource : public FHairCommonResource
 	FVector PositionOffset[2] = {FVector::ZeroVector, FVector::ZeroVector};
 
 	/* Reference to the hair strands render data */
-	const FHairStrandsDatas::FRenderData& RenderData;
+	const FHairStrandsBulkData& BulkData;
 
 	/* Whether the GPU data should be initialized with the asset data or not */
 	const bool bInitializedData = false;
@@ -290,9 +289,6 @@ struct FHairStrandsDeformedResource : public FHairCommonResource
 	/* Whether the GPU data should be initialized with the asset data or not */
 	uint32 CurrentIndex = 0;
 
-	/* Whether the underlying resource is dynamic not (single or double buffer allocated) */
-	FVector DefaultOffset = FVector::ZeroVector;
-
 	/* Type of curves */
 	const EHairStrandsResourcesType CurveType;
 
@@ -350,7 +346,7 @@ struct FHairStrandsClusterCullingResource : public FHairCommonResource
 struct FHairStrandsInterpolationResource : public FHairCommonResource
 {
 	/** Build the hair strands resource */
-	FHairStrandsInterpolationResource(const FHairStrandsInterpolationDatas::FRenderData& InterpolationRenderData, const FHairStrandsDatas& SimDatas);
+	FHairStrandsInterpolationResource(const FHairStrandsInterpolationBulkData& InBulkData);
 	
 	/* Init/Release buffers */
 	virtual void InternalAllocate(FRDGBuilder& GraphBuilder) override;
@@ -371,22 +367,19 @@ struct FHairStrandsInterpolationResource : public FHairCommonResource
 
 	FRDGExternalBuffer Interpolation0Buffer;
 	FRDGExternalBuffer Interpolation1Buffer;
-
-	// For debug purpose only (should be remove once all hair simulation is correctly handled)
 	FRDGExternalBuffer SimRootPointIndexBuffer;
-	TArray<FHairStrandsRootIndexFormat::Type> SimRootPointIndex;
 
 	/* Reference to the hair strands interpolation render data */
-	const FHairStrandsInterpolationDatas::FRenderData& RenderData;
+	const FHairStrandsInterpolationBulkData& BulkData;
 };
 
 #if RHI_RAYTRACING
 struct FHairStrandsRaytracingResource : public FHairCommonResource
 {
 	/** Build the hair strands resource */
-	FHairStrandsRaytracingResource(const FHairStrandsDatas& HairStrandsDatas);
-	FHairStrandsRaytracingResource(const FHairCardsDatas& InData);
-	FHairStrandsRaytracingResource(const FHairMeshesDatas& InData);
+	FHairStrandsRaytracingResource(const FHairStrandsBulkData& InData);
+	FHairStrandsRaytracingResource(const FHairCardsBulkData& InData);
+	FHairStrandsRaytracingResource(const FHairMeshesBulkData& InData);
 
 	/* Init/Release buffers */
 	virtual void InternalAllocate(FRDGBuilder& GraphBuilder) override;
@@ -427,11 +420,13 @@ public:
 	virtual void InitRHI() override;
 };
 
+struct FHairCardsBulkData;
+
 /* Render buffers that will be used for rendering */
 struct FHairCardsRestResource : public FHairCommonResource
 {
 	/** Build the hair strands resource */
-	FHairCardsRestResource(const FHairCardsDatas::FRenderData& HairCardsRenderData, uint32 VertexCount, uint32 PrimitiveCount);
+	FHairCardsRestResource(const FHairCardsBulkData& InBulkData);
 
 	virtual void InitResource() override;
 	virtual void ReleaseResource() override;
@@ -447,18 +442,19 @@ struct FHairCardsRestResource : public FHairCommonResource
 	uint32 GetResourcesSize() const
 	{
 		uint32 Total = 0;
-		Total += RenderData.Positions.GetAllocatedSize();
-		Total += RenderData.Normals.GetAllocatedSize();
-		Total += RenderData.Indices.GetAllocatedSize();
-		Total += RenderData.UVs.GetAllocatedSize();
+		Total += BulkData.Positions.GetAllocatedSize();
+		Total += BulkData.Normals.GetAllocatedSize();
+		Total += BulkData.Indices.GetAllocatedSize();
+		Total += BulkData.UVs.GetAllocatedSize();
 		return Total;
 	}
 
+	uint32 GetVertexCount() const { return BulkData.GetNumVertices();  }
+	uint32 GetPrimitiveCount() const { return BulkData.GetNumTriangles(); }
+
 	/* Strand hair rest position buffer */
 	FHairCardsVertexBuffer RestPositionBuffer;
 	FHairCardIndexBuffer RestIndexBuffer;
-	uint32 VertexCount;
-	uint32 PrimitiveCount;
 	bool bInvertUV = false;
 
 	FHairCardsVertexBuffer NormalsBuffer;
@@ -477,7 +473,7 @@ struct FHairCardsRestResource : public FHairCommonResource
 	FTextureReferenceRHIRef	AuxilaryDataTexture = nullptr;
 
 	/* Reference to the hair strands render data */
-	const FHairCardsDatas::FRenderData& RenderData;
+	const FHairCardsBulkData& BulkData;
 };
 
 /* Render buffers that will be used for rendering */
@@ -523,7 +519,7 @@ struct FHairCardsProceduralResource : public FHairCommonResource
 struct FHairCardsDeformedResource : public FHairCommonResource
 {
 	/** Build the hair strands resource */
-	FHairCardsDeformedResource(const FHairCardsDatas::FRenderData& HairStrandRenderData, bool bInitializeData);
+	FHairCardsDeformedResource(const FHairCardsBulkData& BulkData, bool bInitializeData);
 
 	/* Init/release buffers */
 	virtual void InternalAllocate(FRDGBuilder& GraphBuilder) override;
@@ -545,7 +541,7 @@ struct FHairCardsDeformedResource : public FHairCommonResource
 	FRDGExternalBuffer DeformedPositionBuffer[2];
 
 	/* Reference to the hair strands render data */
-	const FHairCardsDatas::FRenderData& RenderData;
+	const FHairCardsBulkData& BulkData;
 
 	/* Whether the GPU data should be initialized with the asset data or not */
 	const bool bInitializedData = false;
@@ -565,6 +561,8 @@ struct FHairCardsDeformedResource : public FHairCommonResource
 	inline void SwapBuffer()							{ CurrentIndex = 1u - CurrentIndex; }
 };
 
+struct FHairCardsInterpolationBulkData;
+
 /** Hair cards points interpolation attributes */
 struct HAIRSTRANDSCORE_API FHairCardsInterpolationDatas
 {
@@ -586,18 +584,20 @@ struct HAIRSTRANDSCORE_API FHairCardsInterpolationDatas
 	/** Lerp value between the closest vertex indices and the next one, ordered by closest influence */
 	TArray<float> PointsSimCurvesVertexLerp;
 
-	struct FRenderData
-	{
-		TArray<FHairCardsInterpolationFormat::Type> Interpolation;
-	} RenderData;
+	void Serialize(FArchive& Ar, FHairCardsInterpolationBulkData& BulkData);
 };
 
-FArchive& operator<<(FArchive& Ar, FHairCardsInterpolationDatas& CardInterpData);
+struct HAIRSTRANDSCORE_API FHairCardsInterpolationBulkData
+{
+	TArray<FHairCardsInterpolationFormat::Type> Interpolation;
+
+	void Serialize(FArchive& Ar);
+};
 
 struct FHairCardsInterpolationResource : public FHairCommonResource
 {
 	/** Build the hair strands resource */
-	FHairCardsInterpolationResource(const FHairCardsInterpolationDatas::FRenderData& InterpolationRenderData);
+	FHairCardsInterpolationResource(const FHairCardsInterpolationBulkData& InBulkData);
 
 	/* Init/release buffers */
 	virtual void InternalAllocate(FRDGBuilder& GraphBuilder) override;
@@ -617,7 +617,7 @@ struct FHairCardsInterpolationResource : public FHairCommonResource
 	FRDGExternalBuffer InterpolationBuffer;
 
 	/* Reference to the hair strands interpolation render data */
-	const FHairCardsInterpolationDatas::FRenderData& RenderData;
+	const FHairCardsInterpolationBulkData& BulkData;
 };
 
 ////////////////////////////////////////////////////////////////////////////////////////////////////////////////////////
@@ -627,7 +627,7 @@ struct FHairCardsInterpolationResource : public FHairCommonResource
 struct FHairMeshesRestResource : public FHairCommonResource
 {
 	/** Build the hair strands resource */
-	FHairMeshesRestResource(const FHairMeshesDatas::FRenderData& HairMeshesRenderData, uint32 VertexCount, uint32 PrimitiveCount);
+	FHairMeshesRestResource(const FHairMeshesBulkData& BulkData);
 
 	virtual void InitResource() override;
 	virtual void ReleaseResource() override;
@@ -643,18 +643,19 @@ struct FHairMeshesRestResource : public FHairCommonResource
 	uint32 GetResourcesSize() const
 	{
 		uint32 Total = 0;
-		Total += RenderData.Positions.GetAllocatedSize();
-		Total += RenderData.Normals.GetAllocatedSize();
-		Total += RenderData.Indices.GetAllocatedSize();
-		Total += RenderData.UVs.GetAllocatedSize();
+		Total += BulkData.Positions.GetAllocatedSize();
+		Total += BulkData.Normals.GetAllocatedSize();
+		Total += BulkData.Indices.GetAllocatedSize();
+		Total += BulkData.UVs.GetAllocatedSize();
 		return Total;
 	}
 
+	uint32 GetVertexCount() const { return BulkData.GetNumVertices(); }
+	uint32 GetPrimitiveCount() const { return BulkData.GetNumTriangles(); }
+
 	/* Strand hair rest position buffer */
 	FHairCardsVertexBuffer RestPositionBuffer;
-	FHairCardIndexBuffer IndexBuffer;
-	uint32 VertexCount;
-	uint32 PrimitiveCount;
+	FHairCardIndexBuffer IndexBuffer;	
 
 	FHairCardsVertexBuffer NormalsBuffer;
 	FHairCardsVertexBuffer UVsBuffer;
@@ -672,7 +673,7 @@ struct FHairMeshesRestResource : public FHairCommonResource
 	FTextureReferenceRHIRef	AuxilaryDataTexture = nullptr;
 
 	/* Reference to the hair strands render data */
-	const FHairMeshesDatas::FRenderData& RenderData;
+	const FHairMeshesBulkData& BulkData;
 };
 
 
@@ -680,7 +681,7 @@ struct FHairMeshesRestResource : public FHairCommonResource
 struct FHairMeshesDeformedResource : public FHairCommonResource
 {
 	/** Build the hair strands resource */
-	FHairMeshesDeformedResource(const FHairMeshesDatas::FRenderData& HairMeshesRenderData, bool bInInitializedData);
+	FHairMeshesDeformedResource(const FHairMeshesBulkData& InBulkData, bool bInInitializedData);
 
 	/* Init/release buffers */
 	virtual void InternalAllocate(FRDGBuilder& GraphBuilder) override;
@@ -702,7 +703,7 @@ struct FHairMeshesDeformedResource : public FHairCommonResource
 	FRDGExternalBuffer DeformedPositionBuffer[2];
 
 	/* Reference to the hair strands render data */
-	const FHairMeshesDatas::FRenderData& RenderData;
+	const FHairMeshesBulkData& BulkData;
 
 	/* Whether the GPU data should be initialized with the asset data or not */
 	const bool bInitializedData = false;
diff --git a/Engine/Plugins/Runtime/HairStrands/Source/HairStrandsCore/Public/HairCardsBuilder.h b/Engine/Plugins/Runtime/HairStrands/Source/HairStrandsCore/Public/HairCardsBuilder.h
--- a/Engine/Plugins/Runtime/HairStrands/Source/HairStrandsCore/Public/HairCardsBuilder.h
+++ b/Engine/Plugins/Runtime/HairStrands/Source/HairStrandsCore/Public/HairCardsBuilder.h
@@ -20,8 +20,10 @@ namespace FHairCardsBuilder
 	bool ImportGeometry(
 		const UStaticMesh* StaticMesh,
 		FHairCardsDatas& OutCards,
+		FHairCardsBulkData& OutBulk,
 		FHairStrandsDatas& OutGuides,
-		FHairCardsInterpolationDatas& OutInterpolationData);
+		FHairCardsInterpolationDatas& OutInterpolationData,
+		FHairCardsInterpolationBulkData& OutInterpolationBulkData);
 
 	void ExportGeometry(
 		const FHairCardsDatas& InCardsData, 
@@ -29,12 +31,13 @@ namespace FHairCardsBuilder
 
 	void BuildGeometry(
 		const FString& LODName,
-		const struct FHairStrandsDatas& In,
+		const struct FHairStrandsDatas& InRen,
 		const struct FHairStrandsDatas& InSim,
 		const struct FHairGroupsProceduralCards& Settings,
 		FHairCardsProceduralDatas& Out,
 		FHairStrandsDatas& OutGuides,
 		FHairCardsInterpolationDatas& OutInterpolation,
+		FHairCardsInterpolationBulkData& OutInterpolationBulk,
 		FHairGroupCardsTextures& OutTextures);
 
 	void BuildTextureAtlas(
@@ -43,7 +46,7 @@ namespace FHairCardsBuilder
 		FHairCardsProceduralResource* AtlasResource,
 		FHairGroupCardsTextures* Textures);
 
-	void Convert(const FHairCardsProceduralDatas& In, FHairCardsDatas& Out);
+	void Convert(const FHairCardsProceduralDatas& In, FHairCardsDatas& Out, FHairCardsBulkData& OutBulk);
 
 	FString GetVersion();
 }
@@ -53,11 +56,13 @@ namespace FHairMeshesBuilder
 	void BuildGeometry(
 		const struct FHairStrandsDatas& In,
 		const struct FHairStrandsDatas& InSim,
-		FHairMeshesDatas& Out);
+		FHairMeshesDatas& Out,
+		FHairMeshesBulkData& OutBulk);
 
 	void ImportGeometry(
 		const UStaticMesh* StaticMesh,
-		FHairMeshesDatas& Out);
+		FHairMeshesDatas& Out,
+		FHairMeshesBulkData& OutBulk);
 
 	FString GetVersion();
 }
diff --git a/Engine/Plugins/Runtime/HairStrands/Source/HairStrandsCore/Public/HairCardsDatas.h b/Engine/Plugins/Runtime/HairStrands/Source/HairStrandsCore/Public/HairCardsDatas.h
--- a/Engine/Plugins/Runtime/HairStrands/Source/HairStrandsCore/Public/HairCardsDatas.h
+++ b/Engine/Plugins/Runtime/HairStrands/Source/HairStrandsCore/Public/HairCardsDatas.h
@@ -240,27 +240,47 @@ struct FHairCardsGeometry
 	}
 };
 
+struct FHairCardsBulkData;
 struct FHairCardsDatas
 {
-	bool IsValid() const { return RenderData.Positions.Num() > 0; }
+	bool IsValid() const { return Cards.Positions.Num() > 0; }
 
 	FHairCardsGeometry Cards;
 
+	void Serialize(FArchive& Ar, FHairCardsBulkData& BulkData);
+};
+
+struct FHairCardsBulkData
+{
+	bool IsValid() const { return Positions.Num() > 0; }
+	void Reset()
+	{
+		DepthTexture = nullptr;
+		TangentTexture = nullptr;
+		CoverageTexture = nullptr;
+		AttributeTexture = nullptr;
+
+		Positions.Empty();
+		Normals.Empty();
+		UVs.Empty();
+		Indices.Empty();
+	}
+
+	uint32 GetNumTriangles() const { return Indices.Num() / 3; }
+	uint32 GetNumVertices() const { return Positions.Num(); }
+
 	UTexture2D* DepthTexture = nullptr;
 	UTexture2D* TangentTexture = nullptr;
 	UTexture2D* CoverageTexture = nullptr;
 	UTexture2D* AttributeTexture = nullptr;
 
-	struct FRenderData
-	{
-		TArray<FHairCardsPositionFormat::Type> Positions;
-		TArray<FHairCardsNormalFormat::Type> Normals;
-		TArray<FHairCardsUVFormat::Type> UVs;
-		TArray<FHairCardsIndexFormat::Type> Indices;
-	} RenderData;
-};
+	TArray<FHairCardsPositionFormat::Type> Positions;
+	TArray<FHairCardsNormalFormat::Type> Normals;
+	TArray<FHairCardsUVFormat::Type> UVs;
+	TArray<FHairCardsIndexFormat::Type> Indices;
 
-FArchive& operator<<(FArchive& Ar, FHairCardsDatas& CardData);
+	void Serialize(FArchive& Ar);
+};
 
 ////////////////////////////////////////////////////////////////////////////////////////////////
 // Hair Cards (Procedural generation)
@@ -406,30 +426,31 @@ struct FHairMeshes
 		BoundingBox.Init();
 	}
 
-	uint32 GetNumTriangles() const
-	{
-		return Indices.Num() / 3;
-	}
-
-	uint32 GetNumVertices() const
-	{
-		return Positions.Num();
-	}
+	uint32 GetNumTriangles() const { return Indices.Num() / 3; }
+	uint32 GetNumVertices() const { return Positions.Num(); }
 };
 
+struct FHairMeshesBulkData;
+
 struct FHairMeshesDatas
 {
-	bool IsValid() const { return RenderData.Positions.Num() > 0; }
+	bool IsValid() const { return Meshes.Positions.Num() > 0; }
 
 	FHairMeshes Meshes;
+	void Serialize(FArchive& Ar, FHairMeshesBulkData& InBulkData);
+};
 
-	struct FRenderData
-	{
-		TArray<FHairCardsPositionFormat::Type> Positions;
-		TArray<FHairCardsNormalFormat::Type> Normals;
-		TArray<FHairCardsUVFormat::Type> UVs;
-		TArray<FHairCardsIndexFormat::Type> Indices;		
-	} RenderData;
+struct FHairMeshesBulkData
+{
+	bool IsValid() const { return Positions.Num() > 0; }
+
+	uint32 GetNumTriangles() const { return Indices.Num() / 3; }
+	uint32 GetNumVertices() const { return Positions.Num(); }
+
+	TArray<FHairCardsPositionFormat::Type> Positions;
+	TArray<FHairCardsNormalFormat::Type> Normals;
+	TArray<FHairCardsUVFormat::Type> UVs;
+	TArray<FHairCardsIndexFormat::Type> Indices;
+	void Serialize(FArchive& Ar);
 };
 
-FArchive& operator<<(FArchive& Ar, FHairMeshesDatas& MeshData);
diff --git a/Engine/Plugins/Runtime/HairStrands/Source/HairStrandsCore/Public/HairStrandsDatas.h b/Engine/Plugins/Runtime/HairStrands/Source/HairStrandsCore/Public/HairStrandsDatas.h
--- a/Engine/Plugins/Runtime/HairStrands/Source/HairStrandsCore/Public/HairStrandsDatas.h
+++ b/Engine/Plugins/Runtime/HairStrands/Source/HairStrandsCore/Public/HairStrandsDatas.h
@@ -12,6 +12,8 @@ DECLARE_LOG_CATEGORY_EXTERN(LogHairStrands, Log, All);
 
 struct FPackedHairVertex
 {
+	typedef uint64 BulkType;
+
 	FFloat16 X, Y, Z;
 	uint8 ControlPointType : 2;
 	uint8 NormalizedRadius : 6;
@@ -20,6 +22,8 @@ struct FPackedHairVertex
 
 struct FPackedHairAttributeVertex
 {
+	typedef uint64 BulkType;
+
 	uint8 RootU;
 	uint8 RootV;
 	uint8 NormalizedLength;
@@ -33,6 +37,8 @@ struct FPackedHairAttributeVertex
 
 struct FHairMaterialVertex
 {
+	typedef uint32 BulkType;
+
 	// sRGB color space
 	uint8 BaseColorR;
 	uint8 BaseColorG;
@@ -43,6 +49,8 @@ struct FHairMaterialVertex
 
 struct FHairInterpolation0Vertex
 {
+	typedef uint64 BulkType;
+
 	uint16 Index0;
 	uint16 Index1;
 	uint16 Index2;
@@ -53,6 +61,8 @@ struct FHairInterpolation0Vertex
 
 struct FHairInterpolation1Vertex
 {
+	typedef uint64 BulkType;
+
 	uint8 VertexIndex0;
 	uint8 VertexIndex1;
 	uint8 VertexIndex2;
@@ -77,6 +87,7 @@ FArchive& operator<<(FArchive& Ar, FVector4_16& Vertex);
 struct FHairStrandsPositionFormat
 {
 	typedef FPackedHairVertex Type;
+	typedef FPackedHairVertex::BulkType BulkType;
 	static const uint32 ComponentCount = 1;
 	static const uint32 SizeInByte = sizeof(Type);
 	static const EVertexElementType VertexElementType = VET_UShort4;
@@ -95,6 +106,7 @@ struct FHairStrandsPositionOffsetFormat
 struct FHairStrandsAttributeFormat
 {
 	typedef FPackedHairAttributeVertex Type;
+	typedef FPackedHairAttributeVertex::BulkType BulkType;
 	static const uint32 ComponentCount = 1;
 	static const uint32 SizeInByte = sizeof(Type);
 	static const EVertexElementType VertexElementType = VET_UShort4;
@@ -104,6 +116,7 @@ struct FHairStrandsAttributeFormat
 struct FHairStrandsMaterialFormat
 {
 	typedef FHairMaterialVertex Type;
+	typedef FHairMaterialVertex::BulkType BulkType;
 	static const uint32 ComponentCount = 1;
 	static const uint32 SizeInByte = sizeof(Type);
 	static const EVertexElementType VertexElementType = VET_UByte4;
@@ -113,7 +126,7 @@ struct FHairStrandsMaterialFormat
 struct FHairStrandsTangentFormat
 {
 	typedef FPackedNormal Type;
-
+	typedef uint32 BulkType;
 	// TangentX & tangentZ are packed into 2 * PF_R8G8B8A8_SNORM
 	static const uint32 ComponentCount = 2;
 	static const uint32 SizeInByte = sizeof(Type);
@@ -124,6 +137,7 @@ struct FHairStrandsTangentFormat
 struct FHairStrandsInterpolation0Format
 {
 	typedef FHairInterpolation0Vertex Type;
+	typedef uint64 BulkType;
 
 	static const uint32 ComponentCount = 1;
 	static const uint32 SizeInByte = sizeof(Type);
@@ -134,6 +148,7 @@ struct FHairStrandsInterpolation0Format
 struct FHairStrandsInterpolation1Format
 {
 	typedef FHairInterpolation1Vertex Type;
+	typedef uint64 BulkType;
 
 	static const uint32 ComponentCount = 1;
 	static const uint32 SizeInByte = sizeof(Type);
@@ -144,6 +159,7 @@ struct FHairStrandsInterpolation1Format
 struct FHairStrandsRootIndexFormat
 {
 	typedef uint32 Type;
+	typedef uint32 BulkType;
 
 	static const uint32 ComponentCount = 1;
 	static const uint32 SizeInByte = sizeof(Type);
@@ -239,11 +255,13 @@ struct FHairStrandsRootUtils
 	static float	 PackUVsToFloat(const FVector2D& UV);
 };
 
+struct FHairStrandsInterpolationBulkData;
+
 /** Hair strands points interpolation attributes */
 struct HAIRSTRANDSCORE_API FHairStrandsInterpolationDatas
 {
 	/** Serialize the interpolated points */
-	void Serialize(FArchive& Ar);
+	void Serialize(FArchive& Ar, FHairStrandsInterpolationBulkData& BulkData);
 
 	/** Set the number of interpolated points */
 	void SetNum(const uint32 NumPoints);
@@ -254,6 +272,8 @@ struct HAIRSTRANDSCORE_API FHairStrandsInterpolationDatas
 	/** Get the number of interpolated points */
 	uint32 Num() const { return PointsSimCurvesVertexIndex.Num(); }
 
+	bool IsValid() const { return PointsSimCurvesIndex.Num() > 0; }
+
 	/** Simulation curve indices, ordered by closest influence */
 	TArray<FIntVector> PointsSimCurvesIndex;
 
@@ -265,14 +285,15 @@ struct HAIRSTRANDSCORE_API FHairStrandsInterpolationDatas
 
 	/** Weight of vertex indices on simulation curve, ordered by closest influence */
 	TArray<FVector>	PointsSimCurvesVertexWeights;
+};
 
-	struct FRenderData
-	{
-		TArray<FHairStrandsInterpolation0Format::Type> Interpolation0;
-		TArray<FHairStrandsInterpolation1Format::Type> Interpolation1;
-
-		void Serialize(FArchive& Ar);
-	} RenderData;
+struct HAIRSTRANDSCORE_API FHairStrandsInterpolationBulkData
+{
+	void Serialize(FArchive& Ar);
+	uint32 GetPointCount() const { return Interpolation0.Num(); };
+	TArray<FHairStrandsInterpolation0Format::Type> Interpolation0;
+	TArray<FHairStrandsInterpolation1Format::Type> Interpolation1;
+	TArray<FHairStrandsRootIndexFormat::Type> SimRootPointIndex;
 };
 
 /** Hair strands points attribute */
@@ -346,11 +367,13 @@ struct HAIRSTRANDSCORE_API FHairStrandsCurves
 	float MaxRadius = 0;
 };
 
+struct FHairStrandsBulkData;
+
 /** Hair strands datas that are stored on CPU */
 struct HAIRSTRANDSCORE_API FHairStrandsDatas
 {
 	/** Serialize all the hair strands datas */
-	void Serialize(FArchive& Ar);
+	void Serialize(FArchive& Ar, FHairStrandsBulkData& BulkData);
 
 	/* Get the total number of points */
 	uint32 GetNumPoints() const { return StrandsPoints.Num(); }
@@ -360,6 +383,8 @@ struct HAIRSTRANDSCORE_API FHairStrandsDatas
 
 	void Reset();
 
+	bool IsValid() const { return StrandsCurves.Num() > 0 && StrandsPoints.Num() > 0; }
+
 	/** List of all the strands points */
 	FHairStrandsPoints StrandsPoints;
 
@@ -371,16 +396,30 @@ struct HAIRSTRANDSCORE_API FHairStrandsDatas
 
 	/* Strands bounding box */
 	FBox BoundingBox;
+};
 
-	struct FRenderData
-	{
-		TArray<FHairStrandsPositionFormat::Type>	Positions;
-		TArray<FHairStrandsAttributeFormat::Type>	Attributes;
-		TArray<FHairStrandsMaterialFormat::Type>	Materials;
-		TArray<	FHairStrandsRootIndexFormat::Type>	RootIndices; // Vertex to strands Index
+struct HAIRSTRANDSCORE_API FHairStrandsBulkData
+{
+	void Serialize(FArchive& Ar);
+
+	bool IsValid() const { return CurveCount > 0 && PointCount > 0;	}
+	
+	uint32 GetNumCurves() const { return CurveCount;  };
+	uint32 GetNumPoints() const { return PointCount; };
+	float  GetMaxLength() const	{ return MaxLength; };
+	float  GetMaxRadius() const { return MaxRadius; }
+	FVector GetPositionOffset() const { return BoundingBox.GetCenter(); }
+
+	uint32 CurveCount = 0;
+	uint32 PointCount = 0;
+	float MaxLength = 0;
+	float MaxRadius = 0;
+	FBox BoundingBox;
 
-		void Serialize(FArchive& Ar);
-	} RenderData;
+	TArray<FHairStrandsPositionFormat::Type>	Positions;
+	TArray<FHairStrandsAttributeFormat::Type>	Attributes;
+	TArray<FHairStrandsMaterialFormat::Type>	Materials;
+	TArray<FHairStrandsRootIndexFormat::Type>	RootIndices; // Vertex to strands Index
 };
 
 /** Hair strands debug data */
