Add missing file.

diff --git a/Engine/Plugins/Runtime/HairStrands/Source/HairStrandsCore/Private/HairStrandsRendering.cpp b/Engine/Plugins/Runtime/HairStrands/Source/HairStrandsCore/Private/HairStrandsRendering.cpp
--- a/Engine/Plugins/Runtime/HairStrands/Source/HairStrandsCore/Private/HairStrandsRendering.cpp
+++ b/Engine/Plugins/Runtime/HairStrands/Source/HairStrandsCore/Private/HairStrandsRendering.cpp
@@ -845,12 +845,12 @@ static void AddHairCardsDeformationPass(
 
 	FHairCardsDeformationCS::FParameters* Parameters = GraphBuilder.AllocParameters<FHairCardsDeformationCS::FParameters>();
 	Parameters->GuideVertexCount			= LOD.Guides.RestResource->GetVertexCount();
-	Parameters->GuideRestPositionOffset		= LOD.Guides.RestResource->PositionOffset;
+	Parameters->GuideRestPositionOffset		= LOD.Guides.RestResource->GetPositionOffset();
 	Parameters->GuideRestPositionBuffer		= RegisterAsSRV(GraphBuilder, LOD.Guides.RestResource->PositionBuffer);
 	Parameters->GuideDeformedPositionBuffer = RegisterAsSRV(GraphBuilder, LOD.Guides.DeformedResource->GetBuffer(FHairStrandsDeformedResource::Current));
 	Parameters->GuideDeformedPositionOffsetBuffer = RegisterAsSRV(GraphBuilder, LOD.Guides.DeformedResource->GetPositionOffsetBuffer(FHairStrandsDeformedResource::Current));
 	
-	Parameters->CardsVertexCount			= LOD.RestResource->VertexCount;
+	Parameters->CardsVertexCount			= LOD.RestResource->GetVertexCount();
 	Parameters->CardsRestPositionBuffer		= LOD.RestResource->RestPositionBuffer.ShaderResourceViewRHI;
 	Parameters->CardsDeformedPositionBuffer = CardsDeformedPositionBuffer.UAV;
 
@@ -1184,7 +1184,7 @@ static void BuildHairAccelerationStructure_Cards(FRHICommandList& RHICmdList,
 	Initializer.IndexBuffer = RestResource->RestIndexBuffer.IndexBufferRHI;
 	Initializer.IndexBufferOffset = 0;
 	Initializer.GeometryType = RTGT_Triangles;
-	Initializer.TotalPrimitiveCount = RestResource->PrimitiveCount;
+	Initializer.TotalPrimitiveCount = RestResource->GetPrimitiveCount();
 	Initializer.bFastBuild = true;
 	Initializer.bAllowUpdate = true;
 
@@ -1198,8 +1198,8 @@ static void BuildHairAccelerationStructure_Cards(FRHICommandList& RHICmdList,
 	Segment.VertexBuffer = PositionBuffer;
 	Segment.VertexBufferStride = FHairCardsPositionFormat::SizeInByte;
 	Segment.VertexBufferElementType = FHairCardsPositionFormat::VertexElementType;
-	Segment.NumPrimitives = RestResource->PrimitiveCount;
-	Segment.MaxVertices = RestResource->VertexCount;
+	Segment.NumPrimitives = RestResource->GetPrimitiveCount();
+	Segment.MaxVertices = RestResource->GetVertexCount();
 	Initializer.Segments.Add(Segment);
 
 	OutRayTracingGeometry->SetInitializer(Initializer);
@@ -1218,7 +1218,7 @@ static void BuildHairAccelerationStructure_Meshes(FRHICommandList& RHICmdList,
 	Initializer.IndexBuffer = RestResource->IndexBuffer.IndexBufferRHI;
 	Initializer.IndexBufferOffset = 0;
 	Initializer.GeometryType = RTGT_Triangles;
-	Initializer.TotalPrimitiveCount = RestResource->PrimitiveCount;
+	Initializer.TotalPrimitiveCount = RestResource->GetPrimitiveCount();
 	Initializer.bFastBuild = true;
 	Initializer.bAllowUpdate = true;
 
@@ -1232,8 +1232,8 @@ static void BuildHairAccelerationStructure_Meshes(FRHICommandList& RHICmdList,
 	Segment.VertexBuffer = PositionBuffer;
 	Segment.VertexBufferStride = FHairCardsPositionFormat::SizeInByte;
 	Segment.VertexBufferElementType = FHairCardsPositionFormat::VertexElementType;
-	Segment.NumPrimitives = RestResource->PrimitiveCount;
-	Segment.MaxVertices = RestResource->VertexCount;
+	Segment.NumPrimitives = RestResource->GetPrimitiveCount();
+	Segment.MaxVertices = RestResource->GetVertexCount();
 	Initializer.Segments.Add(Segment);
 
 	OutRayTracingGeometry->SetInitializer(Initializer);
@@ -1358,8 +1358,8 @@ static FHairGroupPublicData::FVertexFactoryInput InternalComputeHairStrandsVerte
 		CONVERT_HAIRSSTRANDS_VF_PARAMETERS(OutVFInput.Strands.PositionOffsetBuffer,		Instance->Strands.RestResource->PositionOffsetBuffer);
 		CONVERT_HAIRSSTRANDS_VF_PARAMETERS(OutVFInput.Strands.PrevPositionOffsetBuffer, Instance->Strands.RestResource->PositionOffsetBuffer);
 
-		OutVFInput.Strands.PositionOffset = Instance->Strands.RestResource->PositionOffset;
-		OutVFInput.Strands.PrevPositionOffset = Instance->Strands.RestResource->PositionOffset;
+		OutVFInput.Strands.PositionOffset = Instance->Strands.RestResource->GetPositionOffset();
+		OutVFInput.Strands.PrevPositionOffset = Instance->Strands.RestResource->GetPositionOffset();
 		OutVFInput.Strands.VertexCount = Instance->Strands.RestResource->GetVertexCount();
 	}
 
@@ -1499,8 +1499,8 @@ void ComputeHairStrandsInterpolation(
 					Instance->Strands.HairInterpolationType,
 					InstanceGeometryType,
 					CullingData,
-					Instance->Strands.RestResource->PositionOffset,
-					bValidGuide ? Instance->Guides.RestResource->PositionOffset : FVector::ZeroVector,
+					Instance->Strands.RestResource->GetPositionOffset(),
+					bValidGuide ? Instance->Guides.RestResource->GetPositionOffset() : FVector::ZeroVector,
 					Strands_PositionOffsetSRV,
 					bValidGuide ? RegisterAsSRV(GraphBuilder, Instance->Guides.DeformedResource->GetPositionOffsetBuffer(FHairStrandsDeformedResource::Current)) : nullptr,
 					bHasSkinning ? Instance->Strands.RestRootResource : nullptr,
@@ -1667,8 +1667,8 @@ void ComputeHairStrandsInterpolation(
 						LOD.Guides.HairInterpolationType,
 						InstanceGeometryType,
 						CullingData,
-						LOD.Guides.RestResource->PositionOffset,
-						bValidGuide ? Instance->Guides.RestResource->PositionOffset : FVector::ZeroVector,
+						LOD.Guides.RestResource->GetPositionOffset(),
+						bValidGuide ? Instance->Guides.RestResource->GetPositionOffset() : FVector::ZeroVector,
 						RegisterAsSRV(GraphBuilder, LOD.Guides.DeformedResource->GetPositionOffsetBuffer(FHairStrandsDeformedResource::Current)),
 						bValidGuide ? RegisterAsSRV(GraphBuilder, Instance->Guides.DeformedResource->GetPositionOffsetBuffer(FHairStrandsDeformedResource::Current)) : nullptr,
 						bHasSkinning ? LOD.Guides.RestRootResource : nullptr ,
@@ -1845,7 +1845,7 @@ void ResetHairStrandsInterpolation(
 		Instance->Guides.DeformedRootResource,
 		RegisterAsSRV(GraphBuilder, Instance->Guides.RestResource->PositionBuffer),
 		DeformedPositionBuffer,
-		Instance->Guides.RestResource->PositionOffset,
+		Instance->Guides.RestResource->GetPositionOffset(),
 		RegisterAsSRV(GraphBuilder, Instance->Guides.DeformedResource->GetPositionOffsetBuffer(FHairStrandsDeformedResource::Current)),
 		Instance->Guides.bHasGlobalInterpolation);
 }
