Fix crash when opening ROD emitter solver.

diff --git a/Engine/Plugins/Runtime/HairStrands/Source/HairStrandsCore/Public/GroomComponent.h b/Engine/Plugins/Runtime/HairStrands/Source/HairStrandsCore/Public/GroomComponent.h
--- a/Engine/Plugins/Runtime/HairStrands/Source/HairStrandsCore/Public/GroomComponent.h
+++ b/Engine/Plugins/Runtime/HairStrands/Source/HairStrandsCore/Public/GroomComponent.h
@@ -175,8 +175,8 @@ public:
 
 	/** Hair group instance access */
 	uint32 GetGroupCount() const { return HairGroupInstances.Num();  }
-	FHairGroupInstance* GetGroupInstance(uint32 Index) { return Index < uint32(HairGroupInstances.Num()) ? HairGroupInstances[Index] : nullptr; }
-	const FHairGroupInstance* GetGroupInstance(uint32 Index) const { return Index < uint32(HairGroupInstances.Num()) ? HairGroupInstances[Index] : nullptr; } 
+	FHairGroupInstance* GetGroupInstance(int32 Index) { return Index >= 0 && Index < HairGroupInstances.Num() ? HairGroupInstances[Index] : nullptr; }
+	const FHairGroupInstance* GetGroupInstance(int32 Index) const { return Index >= 0 && Index < HairGroupInstances.Num() ? HairGroupInstances[Index] : nullptr; } 
 
 	//~ Begin UPrimitiveComponent Interface
 	EHairGeometryType GetMaterialGeometryType(int32 ElementIndex) const;
