Drop support for legacy groom assets not having HairDescription embedded, preventing them to be rebuilt & cached.

Now all compatible assets are built from HairDescription and stored through DDC.

This allows flexible data restructuration without object version

diff --git a/Engine/Plugins/Runtime/HairStrands/Source/HairStrandsCore/Private/GroomAsset.cpp b/Engine/Plugins/Runtime/HairStrands/Source/HairStrandsCore/Private/GroomAsset.cpp
--- a/Engine/Plugins/Runtime/HairStrands/Source/HairStrandsCore/Private/GroomAsset.cpp
+++ b/Engine/Plugins/Runtime/HairStrands/Source/HairStrandsCore/Private/GroomAsset.cpp
@@ -372,19 +372,13 @@ void UGroomAsset::Serialize(FArchive& Ar)
 			}
 
 			HairDescriptionBulkData->Serialize(Ar, this);
-
-			// Serialize the HairGroupsData directly into the asset if it couldn't be cached in the DDC
-			if (!bIsCacheable)
-			{
-				Ar << HairGroupsData;
-			}
 		}
 #endif // WITH_EDITORONLY_DATA
 	}
 	else
 	{
-		// Old format serialized the computed groom data directly
-		Ar << HairGroupsData;
+		// Old format serialized the computed groom data directly, but are no longer supported
+		UE_LOG(LogHairStrands, Error, TEXT("[Groom] The groom asset %s is too old. Please reimported the groom from its original source file."));
 	}
 }
 
@@ -669,8 +663,7 @@ void UGroomAsset::PostLoad()
 #endif
 
 #if WITH_EDITORONLY_DATA
-	bool bSucceed = bIsCacheable;
-	if (bIsCacheable)
+	bool bSucceed = true;
 	{
 		// Interpolation settings are used for building the interpolation data, and per se defined the number of groups
 		const uint32 GroupCount = HairGroupsInterpolation.Num();
@@ -1705,7 +1698,7 @@ bool UGroomAsset::CacheStrandsData(uint32 GroupIndex, FProcessedHairDescription&
 			// we can compute a size threshold to skip the caching when
 			// the uncompressed size exceeds the threshold
 			static constexpr const int64 SizeThreshold = (int64)MAX_int32 * 2.5;
-			bIsCacheable = UncompressedSize < SizeThreshold;
+			const bool bIsCacheable = UncompressedSize < SizeThreshold;
 			if (bIsCacheable)
 			{
 				FMemoryWriter CompressedArchive(DerivedData, true);
@@ -1715,6 +1708,11 @@ bool UGroomAsset::CacheStrandsData(uint32 GroupIndex, FProcessedHairDescription&
 
 				GetDerivedDataCacheRef().Put(*DerivedDataKey, DerivedData, GetPathName());
 			}
+			else
+			{
+				UE_LOG(LogHairStrands, Error, TEXT("[Groom/DDC] Strands - The groom asset is too large to be cached into the Derived Data Cache. (Groom:%s, Group:%d, Size %d bytes, DDC Limit: SizeThreshold %d bytes)) ."), *GetName(), GroupIndex, UncompressedSize, SizeThreshold);
+				return false;
+			}
 		}
 	}
 
diff --git a/Engine/Plugins/Runtime/HairStrands/Source/HairStrandsCore/Public/GroomAsset.h b/Engine/Plugins/Runtime/HairStrands/Source/HairStrandsCore/Public/GroomAsset.h
--- a/Engine/Plugins/Runtime/HairStrands/Source/HairStrandsCore/Public/GroomAsset.h
+++ b/Engine/Plugins/Runtime/HairStrands/Source/HairStrandsCore/Public/GroomAsset.h
@@ -538,9 +538,6 @@ private:
 	TArray<FString> CardsDerivedDataKey;
 	TArray<FString> MeshesDerivedDataKey;
 
-	UPROPERTY()
-	bool bIsCacheable = true;
-
 	TStrongObjectPtr<UGroomAsset> GroomAssetStrongPtr;
 	bool bRetryLoadFromGameThread = false;
 #endif
