First pass for removing hair strands/cards data dependency.

The goal is to make all strands/cards data transients, and only leave bulk data during runtime.

diff --git a/Engine/Plugins/Runtime/HairStrands/Source/HairStrandsCore/Private/GroomAsset.cpp b/Engine/Plugins/Runtime/HairStrands/Source/HairStrandsCore/Private/GroomAsset.cpp
--- a/Engine/Plugins/Runtime/HairStrands/Source/HairStrandsCore/Private/GroomAsset.cpp
+++ b/Engine/Plugins/Runtime/HairStrands/Source/HairStrandsCore/Private/GroomAsset.cpp
@@ -581,11 +581,11 @@ void UGroomAsset::UpdateHairGroupsInfo()
 	{
 		FHairGroupInfoWithVisibility& Info = HairGroupsInfo[GroupIndex];
 		Info.GroupID = GroupIndex;
-		Info.NumCurves = Data.Strands.Data.GetNumCurves();
-		Info.NumGuides = Data.Guides.Data.GetNumCurves();
-		Info.NumCurveVertices = Data.Strands.Data.GetNumPoints();
-		Info.NumGuideVertices = Data.Guides.Data.GetNumPoints();
-		Info.MaxCurveLength = Data.Strands.Data.StrandsCurves.MaxLength;
+		Info.NumCurves = Data.Strands.BulkData.GetNumCurves();
+		Info.NumGuides = Data.Guides.BulkData.GetNumCurves();
+		Info.NumCurveVertices = Data.Strands.BulkData.GetNumPoints();
+		Info.NumGuideVertices = Data.Guides.BulkData.GetNumPoints();
+		Info.MaxCurveLength = Data.Strands.BulkData.GetMaxLength();
 		if (bForceReset)
 		{
 			Info.bIsVisible = true;
@@ -745,9 +745,9 @@ void UGroomAsset::PostLoad()
 		for (uint32 GroupIndex = 0; GroupIndex < GroupCount; ++GroupIndex)
 		{
 			const bool bNeedToBuildData = 
-				(HairGroupsData[GroupIndex].Guides.Data.GetNumCurves() == 0 ||
-				HairGroupsData[GroupIndex].Strands.InterpolationData.Num() == 0) &&
-				HairGroupsData[GroupIndex].Strands.Data.GetNumCurves() > 0; // Empty groom has no data to build
+				(HairGroupsData[GroupIndex].Guides.BulkData.GetNumCurves() == 0 ||
+				 HairGroupsData[GroupIndex].Strands.InterpolationBulkData.GetPointCount() == 0) &&
+				 HairGroupsData[GroupIndex].Strands.BulkData.GetNumCurves() > 0; // Empty groom has no data to build
 			if (bNeedToBuildData)
 			{
 				FGroomBuilder::BuildData(HairGroupsData[GroupIndex], HairGroupsInterpolation[GroupIndex], GroupIndex );
@@ -1767,7 +1767,7 @@ bool UGroomAsset::CacheCardsGeometry(uint32 GroupIndex, const FString& StrandsKe
 	// Handle the case where the cards data is already cached in the DDC
 	// Need to populate the strands data with it
 	FHairGroupData& GroupData = HairGroupsData[GroupIndex];
-	if (StrandsDerivedDataKey[GroupIndex].IsEmpty() && GroupData.Strands.Data.GetNumPoints() == 0)
+	if (StrandsDerivedDataKey[GroupIndex].IsEmpty() && GroupData.Strands.BulkData.GetNumPoints() == 0)
 	{
 		for (int32 LODIt = 0; LODIt < GroupData.Cards.LODs.Num(); ++LODIt)
 		{
@@ -1779,6 +1779,9 @@ bool UGroomAsset::CacheCardsGeometry(uint32 GroupIndex, const FString& StrandsKe
 				GroupData.Strands.Data = LOD.Guides.Data;
 				GroupData.Guides.Data = LOD.Guides.Data;
 
+				GroupData.Strands.BulkData = LOD.Guides.BulkData;
+				GroupData.Guides.BulkData = LOD.Guides.BulkData;
+
 				break;
 			}
 		}
@@ -2005,11 +2008,14 @@ bool UGroomAsset::BuildCardsGeometry(uint32 GroupIndex)
 
 				// There could be no strands data when importing cards into an empty groom so get them from the card guides
 				bool bCopyRenderData = false;
-				if (GroupData.Guides.Data.GetNumPoints() == 0)
+				if (GroupData.Guides.BulkData.GetNumPoints() == 0)
 				{
 					GroupData.Strands.Data = LOD.Guides.Data;
 					GroupData.Guides.Data = LOD.Guides.Data;
 
+					GroupData.Strands.BulkData = LOD.Guides.BulkData;
+					GroupData.Guides.BulkData = LOD.Guides.BulkData;
+
 					// The RenderData is filled out by BuildData
 					bCopyRenderData = true;
 				}
diff --git a/Engine/Plugins/Runtime/HairStrands/Source/HairStrandsCore/Private/GroomBindingAsset.cpp b/Engine/Plugins/Runtime/HairStrands/Source/HairStrandsCore/Private/GroomBindingAsset.cpp
--- a/Engine/Plugins/Runtime/HairStrands/Source/HairStrandsCore/Private/GroomBindingAsset.cpp
+++ b/Engine/Plugins/Runtime/HairStrands/Source/HairStrandsCore/Private/GroomBindingAsset.cpp
@@ -512,7 +512,7 @@ bool UGroomBindingAsset::IsCompatible(const UGroomAsset* InGroom, const UGroomBi
 		for (uint32 GroupIt = 0; GroupIt < GroupCount; ++GroupIt)
 		{
 			{
-				const uint32 GroomCount = InGroom->HairGroupsData[GroupIt].Guides.Data.GetNumCurves();
+				const uint32 GroomCount = InGroom->HairGroupsData[GroupIt].Guides.BulkData.GetNumCurves();
 				const uint32 BindingCount = InBinding->GroupInfos[GroupIt].SimRootCount;
 
 				if (GroomCount != 0 && GroomCount != BindingCount)
@@ -532,7 +532,7 @@ bool UGroomBindingAsset::IsCompatible(const UGroomAsset* InGroom, const UGroomBi
 
 			if (IsHairStrandsEnabled(EHairStrandsShaderType::Strands))
 			{
-				const uint32 GroomCount = InGroom->HairGroupsData[GroupIt].Strands.Data.GetNumCurves();
+				const uint32 GroomCount = InGroom->HairGroupsData[GroupIt].Strands.BulkData.GetNumCurves();
 				const uint32 BindingCount = InBinding->GroupInfos[GroupIt].RenRootCount;
 
 				// Groom may have stripped strands data so GroomCount would be 0
diff --git a/Engine/Plugins/Runtime/HairStrands/Source/HairStrandsCore/Private/GroomComponent.cpp b/Engine/Plugins/Runtime/HairStrands/Source/HairStrandsCore/Private/GroomComponent.cpp
--- a/Engine/Plugins/Runtime/HairStrands/Source/HairStrandsCore/Private/GroomComponent.cpp
+++ b/Engine/Plugins/Runtime/HairStrands/Source/HairStrandsCore/Private/GroomComponent.cpp
@@ -102,7 +102,7 @@ static FHairGroupDesc GetGroomGroupsDesc(const UGroomAsset* Asset, UGroomCompone
 	}
 
 	FHairGroupDesc O = Component->GroomGroupsDesc[GroupIndex];
-	O.HairLength = Asset->HairGroupsData[GroupIndex].Strands.Data.StrandsCurves.MaxLength;
+	O.HairLength = Asset->HairGroupsData[GroupIndex].Strands.BulkData.MaxLength;
 	O.LODBias 	 = Asset->EffectiveLODBias[GroupIndex] > 0 ? FMath::Max(O.LODBias, Asset->EffectiveLODBias[GroupIndex]) : O.LODBias;
 
 	if (!O.HairWidth_Override)					{ O.HairWidth					= Asset->HairGroupsRendering[GroupIndex].GeometrySettings.HairWidth;					}
@@ -1340,7 +1340,7 @@ FBoxSphereBounds UGroomComponent::CalcBounds(const FTransform& InLocalToWorld) c
 			{
 				if (IsHairStrandsEnabled(EHairStrandsShaderType::Strands) && GroupData.Strands.HasValidData())
 				{
-					LocalBounds += GroupData.Strands.Data.BoundingBox;
+					LocalBounds += GroupData.Strands.GetBounds();
 				}
 				else if (IsHairStrandsEnabled(EHairStrandsShaderType::Cards) && GroupData.Cards.HasValidData())
 				{
@@ -1352,7 +1352,7 @@ FBoxSphereBounds UGroomComponent::CalcBounds(const FTransform& InLocalToWorld) c
 				}
 				else if (GroupData.Guides.HasValidData())
 				{
-					LocalBounds += GroupData.Guides.Data.BoundingBox;
+					LocalBounds += GroupData.Guides.GetBounds();
 				}
 			}
 			return FBoxSphereBounds(LocalBounds.TransformBy(InLocalToWorld));
@@ -2106,7 +2106,7 @@ void UGroomComponent::InitResources(bool bIsBindingReloading)
 						check(GroupData.Strands.ClusterCullingResource->Data.LODVisibility[LODIt] == LODVisibilities[LODIt]);
 					}
 				}
-				HairGroupInstance->HairGroupPublicData->SetClusters(HairGroupInstance->Strands.ClusterCullingResource->Data.ClusterCount, GroupData.Strands.Data.GetNumPoints());
+				HairGroupInstance->HairGroupPublicData->SetClusters(HairGroupInstance->Strands.ClusterCullingResource->Data.ClusterCount, GroupData.Strands.BulkData.GetNumPoints());
 				BeginInitResource(HairGroupInstance->HairGroupPublicData);
 			}
 
diff --git a/Engine/Plugins/Runtime/HairStrands/Source/HairStrandsCore/Private/HairCardsBuilder.cpp b/Engine/Plugins/Runtime/HairStrands/Source/HairStrandsCore/Private/HairCardsBuilder.cpp
--- a/Engine/Plugins/Runtime/HairStrands/Source/HairStrandsCore/Private/HairCardsBuilder.cpp
+++ b/Engine/Plugins/Runtime/HairStrands/Source/HairStrandsCore/Private/HairCardsBuilder.cpp
@@ -3685,6 +3685,7 @@ bool ImportGeometry(
 
 		Out.Cards.BoundingBox += Out.Cards.Positions[VertexIndex];
 	}
+	OutBulk.BoundingBox = Out.Cards.BoundingBox;
 
 	// Fill in render resources (do we need to keep it separated? e.g, format compression, packing)
 	const uint32 PointCount = Out.Cards.Positions.Num();
@@ -3741,6 +3742,7 @@ void Convert(const FHairCardsProceduralDatas& In, FHairCardsDatas& Out, FHairCar
 	OutBulk.Normals = In.RenderData.Normals;
 	OutBulk.Indices = In.RenderData.Indices;
 	OutBulk.UVs = In.RenderData.UVs;
+	OutBulk.BoundingBox = In.Cards.BoundingBox;
 }
 
 ///////////////////////////////////////////////////////////////////////////////////////////////////
@@ -4012,6 +4014,7 @@ void BuildGeometry(
 
 		Out.Meshes.BoundingBox += OutBulk.Positions[PointIt];
 	}
+	OutBulk.BoundingBox = Out.Meshes.BoundingBox;
 
 	const uint32 IndexCount = Out.Meshes.Indices.Num();
 	OutBulk.Indices.SetNum(IndexCount);
@@ -4056,6 +4059,7 @@ void ImportGeometry(
 	}
 
 	// Fill in render resources (do we need to keep it separated? e.g, format compression, packing)
+	OutBulk.BoundingBox = Out.Meshes.BoundingBox;
 	const uint32 PointCount = Out.Meshes.Positions.Num();
 	OutBulk.Positions.SetNum(PointCount);
 	OutBulk.Normals.SetNum(PointCount * FHairCardsNormalFormat::ComponentCount);
diff --git a/Engine/Plugins/Runtime/HairStrands/Source/HairStrandsCore/Private/HairCardsDatas.cpp b/Engine/Plugins/Runtime/HairStrands/Source/HairStrandsCore/Private/HairCardsDatas.cpp
--- a/Engine/Plugins/Runtime/HairStrands/Source/HairStrandsCore/Private/HairCardsDatas.cpp
+++ b/Engine/Plugins/Runtime/HairStrands/Source/HairStrandsCore/Private/HairCardsDatas.cpp
@@ -50,10 +50,14 @@ void FHairCardsBulkData::Serialize(FArchive& Ar)
 	Ar << Indices;
 }
 
-void FHairCardsDatas::Serialize(FArchive& Ar, FHairCardsBulkData& BulkData)
+void FHairCardsDatas::Serialize(FArchive& Ar, FHairCardsBulkData& InBulkData)
 {
 	Ar << Cards;
-	BulkData.Serialize(Ar);
+	InBulkData.Serialize(Ar);
+	if (Ar.IsLoading())
+	{
+		InBulkData.BoundingBox = Cards.BoundingBox;
+	}
 }
 
 FArchive& operator<<(FArchive& Ar, FHairCardsInterpolationVertex& CardInterpVertex)
@@ -99,6 +103,10 @@ void FHairMeshesDatas::Serialize(FArchive& Ar, FHairMeshesBulkData& InBulkData)
 {
 	Ar << Meshes;
 	InBulkData.Serialize(Ar);
+	if (Ar.IsLoading())
+	{
+		InBulkData.BoundingBox = Meshes.BoundingBox;
+	}
 }
 
 void FHairMeshesBulkData::Serialize(FArchive& Ar)
diff --git a/Engine/Plugins/Runtime/HairStrands/Source/HairStrandsCore/Public/GroomAsset.h b/Engine/Plugins/Runtime/HairStrands/Source/HairStrandsCore/Public/GroomAsset.h
--- a/Engine/Plugins/Runtime/HairStrands/Source/HairStrandsCore/Public/GroomAsset.h
+++ b/Engine/Plugins/Runtime/HairStrands/Source/HairStrandsCore/Public/GroomAsset.h
@@ -77,8 +77,10 @@ struct HAIRSTRANDSCORE_API FHairGroupData
 
 	struct FBase
 	{
-		bool HasValidData() const	{ return Data.GetNumPoints() > 0;}
-		bool IsValid() const		{ return RestResource != nullptr; }
+		bool HasValidData() const		{ return BulkData.GetNumPoints() > 0;}
+		bool IsValid() const			{ return RestResource != nullptr; }
+		const FBox& GetBounds() const	{ return BulkData.GetBounds(); }
+
 		FHairStrandsDatas					Data;
 		FHairStrandsBulkData				BulkData;
 		FHairStrandsRestResource*			RestResource = nullptr;
@@ -147,7 +149,7 @@ struct HAIRSTRANDSCORE_API FHairGroupData
 		{
 			for (const FLOD& LOD : LODs)
 			{
-				if (LOD.IsValid()) return LOD.Data.Cards.BoundingBox;
+				if (LOD.IsValid()) return LOD.BulkData.BoundingBox;
 			}
 			return FBox();
 		}
@@ -178,8 +180,9 @@ struct HAIRSTRANDSCORE_API FHairGroupData
 				return Total;
 			}
 
-			bool HasValidData() const { return Data.IsValid(); }
-			bool IsValid() const { return Data.IsValid() && RestResource != nullptr; }
+			bool HasValidData() const { return BulkData.IsValid(); }
+			bool IsValid() const { return BulkData.IsValid() && RestResource != nullptr; }
+
 			// Main data & Resources
 			FHairCardsDatas						Data;
 			FHairCardsBulkData					BulkData;
@@ -222,7 +225,7 @@ struct HAIRSTRANDSCORE_API FHairGroupData
 		{
 			for (const FLOD& LOD : LODs)
 			{
-				if (LOD.IsValid()) return LOD.Data.Meshes.BoundingBox;
+				if (LOD.IsValid()) return LOD.BulkData.BoundingBox;
 			}
 			return FBox();
 		}
@@ -251,8 +254,9 @@ struct HAIRSTRANDSCORE_API FHairGroupData
 				return Total;
 			}
 
-			bool HasValidData() const { return Data.IsValid(); }
-			bool IsValid() const { return Data.IsValid() && RestResource != nullptr; }
+			bool HasValidData() const { return BulkData.IsValid(); }
+			bool IsValid() const { return BulkData.IsValid() && RestResource != nullptr; }
+
 			FHairMeshesDatas Data;
 			FHairMeshesBulkData BulkData;
 			FHairMeshesRestResource* RestResource = nullptr;
diff --git a/Engine/Plugins/Runtime/HairStrands/Source/HairStrandsCore/Public/HairCardsDatas.h b/Engine/Plugins/Runtime/HairStrands/Source/HairStrandsCore/Public/HairCardsDatas.h
--- a/Engine/Plugins/Runtime/HairStrands/Source/HairStrandsCore/Public/HairCardsDatas.h
+++ b/Engine/Plugins/Runtime/HairStrands/Source/HairStrandsCore/Public/HairCardsDatas.h
@@ -268,6 +268,7 @@ struct FHairCardsBulkData
 
 	uint32 GetNumTriangles() const { return Indices.Num() / 3; }
 	uint32 GetNumVertices() const { return Positions.Num(); }
+	const FBox& GetBounds() const { return BoundingBox; }
 
 	UTexture2D* DepthTexture = nullptr;
 	UTexture2D* TangentTexture = nullptr;
@@ -278,6 +279,7 @@ struct FHairCardsBulkData
 	TArray<FHairCardsNormalFormat::Type> Normals;
 	TArray<FHairCardsUVFormat::Type> UVs;
 	TArray<FHairCardsIndexFormat::Type> Indices;
+	FBox BoundingBox;
 
 	void Serialize(FArchive& Ar);
 };
@@ -446,11 +448,14 @@ struct FHairMeshesBulkData
 
 	uint32 GetNumTriangles() const { return Indices.Num() / 3; }
 	uint32 GetNumVertices() const { return Positions.Num(); }
+	const FBox& GetBounds() const { return BoundingBox; }
 
 	TArray<FHairCardsPositionFormat::Type> Positions;
 	TArray<FHairCardsNormalFormat::Type> Normals;
 	TArray<FHairCardsUVFormat::Type> UVs;
 	TArray<FHairCardsIndexFormat::Type> Indices;
+	FBox BoundingBox;
+
 	void Serialize(FArchive& Ar);
 };
 
diff --git a/Engine/Plugins/Runtime/HairStrands/Source/HairStrandsCore/Public/HairStrandsDatas.h b/Engine/Plugins/Runtime/HairStrands/Source/HairStrandsCore/Public/HairStrandsDatas.h
--- a/Engine/Plugins/Runtime/HairStrands/Source/HairStrandsCore/Public/HairStrandsDatas.h
+++ b/Engine/Plugins/Runtime/HairStrands/Source/HairStrandsCore/Public/HairStrandsDatas.h
@@ -409,6 +409,7 @@ struct HAIRSTRANDSCORE_API FHairStrandsBulkData
 	float  GetMaxLength() const	{ return MaxLength; };
 	float  GetMaxRadius() const { return MaxRadius; }
 	FVector GetPositionOffset() const { return BoundingBox.GetCenter(); }
+	const FBox& GetBounds() const { return BoundingBox; }
 
 	uint32 CurveCount = 0;
 	uint32 PointCount = 0;
diff --git a/Engine/Plugins/Runtime/HairStrands/Source/HairStrandsEditor/Private/GroomActions.cpp b/Engine/Plugins/Runtime/HairStrands/Source/HairStrandsEditor/Private/GroomActions.cpp
--- a/Engine/Plugins/Runtime/HairStrands/Source/HairStrandsEditor/Private/GroomActions.cpp
+++ b/Engine/Plugins/Runtime/HairStrands/Source/HairStrandsEditor/Private/GroomActions.cpp
@@ -187,8 +187,8 @@ void FGroomActions::ExecuteRebuild(TArray<TWeakObjectPtr<UGroomAsset>> Objects)
 					{
 						FGroomHairGroupPreview& OutGroup = GroupsPreview->Groups.AddDefaulted_GetRef();
 						OutGroup.GroupID    = GroupIndex;
-						OutGroup.CurveCount = GroomAsset->HairGroupsData[GroupIndex].Strands.Data.GetNumCurves();
-						OutGroup.GuideCount = GroomAsset->HairGroupsData[GroupIndex].Guides.Data.GetNumCurves();
+						OutGroup.CurveCount = GroomAsset->HairGroupsData[GroupIndex].Strands.BulkData.GetNumCurves();
+						OutGroup.GuideCount = GroomAsset->HairGroupsData[GroupIndex].Guides.BulkData.GetNumCurves();
 						OutGroup.InterpolationSettings = GroomAsset->HairGroupsInterpolation[GroupIndex];
 					}
 				}
diff --git a/Engine/Plugins/Runtime/HairStrands/Source/HairStrandsEditor/Private/GroomAssetDetails.cpp b/Engine/Plugins/Runtime/HairStrands/Source/HairStrandsEditor/Private/GroomAssetDetails.cpp
--- a/Engine/Plugins/Runtime/HairStrands/Source/HairStrandsEditor/Private/GroomAssetDetails.cpp
+++ b/Engine/Plugins/Runtime/HairStrands/Source/HairStrandsEditor/Private/GroomAssetDetails.cpp
@@ -1238,7 +1238,7 @@ void FGroomRenderingDetails::OnGenerateElementForHairGroup(TSharedRef<IPropertyH
 		.ValueContent()
 		.HAlign(HAlign_Fill)
 		[
-			MakeHairStrandsInfoGrid(DetailFontInfo, GroomAsset->HairGroupsInfo[GroupIndex], GroomAsset->HairGroupsData[GroupIndex].Strands.Data.StrandsCurves.MaxRadius)
+			MakeHairStrandsInfoGrid(DetailFontInfo, GroomAsset->HairGroupsInfo[GroupIndex], GroomAsset->HairGroupsData[GroupIndex].Strands.BulkData.MaxRadius)
 		];
 	}
 
