Refactor hair data building to well separate input / output.

This will allows to make hairstrands data transient and preserve only bulk data in subsequent CLs.

diff --git a/Engine/Plugins/Runtime/HairStrands/Source/HairStrandsCore/Private/GroomAsset.cpp b/Engine/Plugins/Runtime/HairStrands/Source/HairStrandsCore/Private/GroomAsset.cpp
--- a/Engine/Plugins/Runtime/HairStrands/Source/HairStrandsCore/Private/GroomAsset.cpp
+++ b/Engine/Plugins/Runtime/HairStrands/Source/HairStrandsCore/Private/GroomAsset.cpp
@@ -190,8 +190,78 @@ static void PatchHairStrandsInterpolationBulkData(FHairStrandsInterpolationBulkD
 		}
 	}
 }
+
+static bool BuildHairGroup(
+	const int32 GroupIndex,
+	const FHairDescriptionGroups& HairDescriptionGroups, 
+	const TArray<FHairGroupsInterpolation>& HairGroupsInterpolation,
+	const TArray<FHairGroupsLOD>& HairGroupsLOD,
+	TArray<FHairGroupInfoWithVisibility>& OutHairGroupsInfo,
+	TArray<FHairGroupData>& OutHairGroupsData)
+{
+	const bool bIsValid = HairDescriptionGroups.IsValid();
+	if (bIsValid)
+	{
+		const FHairDescriptionGroup& HairGroup = HairDescriptionGroups.HairGroups[GroupIndex];
+		check(GroupIndex <= HairDescriptionGroups.HairGroups.Num());
+		check(GroupIndex == HairGroup.Info.GroupID);
+
+		FGroomBuilder::BuildData(HairGroup, HairGroupsInterpolation[GroupIndex], OutHairGroupsInfo[GroupIndex], OutHairGroupsData[GroupIndex].Strands.Data, OutHairGroupsData[GroupIndex].Guides.Data);
+
+		FGroomBuilder::BuildBulkData(HairGroup.Info, OutHairGroupsData[GroupIndex].Guides.Data, OutHairGroupsData[GroupIndex].Guides.BulkData);
+		FGroomBuilder::BuildBulkData(HairGroup.Info, OutHairGroupsData[GroupIndex].Strands.Data, OutHairGroupsData[GroupIndex].Strands.BulkData);
+
+		FGroomBuilder::BuildInterplationData(HairGroup.Info, OutHairGroupsData[GroupIndex].Strands.Data, OutHairGroupsData[GroupIndex].Guides.Data, HairGroupsInterpolation[GroupIndex].InterpolationSettings, OutHairGroupsData[GroupIndex].Strands.InterpolationData);
+		FGroomBuilder::BuildInterplationBulkData(OutHairGroupsData[GroupIndex].Guides.Data, OutHairGroupsData[GroupIndex].Strands.InterpolationData, OutHairGroupsData[GroupIndex].Strands.InterpolationBulkData);
+
+		FGroomBuilder::BuildClusterData(OutHairGroupsData[GroupIndex].Strands.Data, HairDescriptionGroups.BoundRadius, HairGroupsLOD[GroupIndex], OutHairGroupsData[GroupIndex].Strands.ClusterCullingData);
+	}
+	return bIsValid;
+}
+
+static bool BuildHairGroupCluster(
+	const int32 GroupIndex,
+	const FHairDescriptionGroups& HairDescriptionGroups,
+	const TArray<FHairGroupsInterpolation>& HairGroupsInterpolation,
+	const TArray<FHairGroupsLOD>& HairGroupsLOD,
+	TArray<FHairGroupInfoWithVisibility>& OutHairGroupsInfo,
+	TArray<FHairGroupData>& OutHairGroupsData)
+{
+	const bool bIsValid = HairDescriptionGroups.IsValid();
+	if (bIsValid)
+	{
+		const FHairDescriptionGroup& HairGroup = HairDescriptionGroups.HairGroups[GroupIndex];
+		check(GroupIndex <= HairDescriptionGroups.HairGroups.Num());
+		check(GroupIndex == HairGroup.Info.GroupID);
+
+		FGroomBuilder::BuildData(HairGroup, HairGroupsInterpolation[GroupIndex], OutHairGroupsInfo[GroupIndex], OutHairGroupsData[GroupIndex].Strands.Data, OutHairGroupsData[GroupIndex].Guides.Data);
+		FGroomBuilder::BuildClusterData(OutHairGroupsData[GroupIndex].Strands.Data, HairDescriptionGroups.BoundRadius, HairGroupsLOD[GroupIndex], OutHairGroupsData[GroupIndex].Strands.ClusterCullingData);
+	}
+	return bIsValid;
+}
+
 /////////////////////////////////////////////////////////////////////////////////////////
 
+#if WITH_EDITORONLY_DATA
+const FHairDescriptionGroups& UGroomAsset::GetHairDescriptionGroups()
+{
+	check(HairDescriptionBulkData);
+
+	if (!HairDescription)
+	{
+		HairDescription = MakeUnique<FHairDescription>();
+		HairDescriptionBulkData->LoadHairDescription(*HairDescription);
+	}
+	if (!HairDescriptionGroups)
+	{
+		HairDescriptionGroups = MakeUnique<FHairDescriptionGroups>();
+		FGroomBuilder::BuildHairDescriptionGroups(*HairDescription, *HairDescriptionGroups);
+	}
+
+	return *HairDescriptionGroups;
+}
+#endif // WITH_EDITORONLY_DATA
+
 uint8 UGroomAsset::GenerateClassStripFlags(FArchive& Ar)
 {
 #if WITH_EDITOR
@@ -465,7 +535,6 @@ void UGroomAsset::UpdateResource()
 {
 	LLM_SCOPE(ELLMTag::Meshes) // This should be a Groom LLM tag, but there is no LLM tag bit left
 
-	float GroomBoundRadius = -1; 
 	uint32 AllChangeType = 0;
 	for (uint32 GroupIndex = 0, GroupCount = GetNumHairGroups(); GroupIndex < GroupCount; ++GroupIndex)
 	{
@@ -489,12 +558,11 @@ void UGroomAsset::UpdateResource()
 
 			if ((ChangeType & GroomChangeType_LOD) && GroupData.Strands.HasValidData())
 			{
-				if (GroomBoundRadius < 0)
-				{
-					GroomBoundRadius = FGroomBuilder::ComputeGroomBoundRadius(HairGroupsData);
-				}
+				const FHairDescriptionGroups& LocalHairDescriptionGroups = GetHairDescriptionGroups();
+
 				// Force rebuilding the LOD data as the LOD settings has changed
-				FGroomBuilder::BuildClusterData(this, GroomBoundRadius, GroupIndex);
+				FGroomBuilder::BuildClusterData(HairGroupsData[GroupIndex].Strands.Data, LocalHairDescriptionGroups.BoundRadius, HairGroupsLOD[GroupIndex], HairGroupsData[GroupIndex].Strands.ClusterCullingData);
+
 				GroupData.Strands.ClusterCullingResource = new FHairStrandsClusterCullingResource(GroupData.Strands.ClusterCullingData);
 				BeginInitResource(GroupData.Strands.ClusterCullingResource);
 			}
@@ -642,22 +710,16 @@ void UGroomAsset::PostLoad()
 	{
 		if (HairDescriptionBulkData)
 		{
-			FProcessedHairDescription ProcessedHairDescription;
-			if (!HairDescription)
-			{
-				HairDescription = MakeUnique<FHairDescription>();
-				HairDescriptionBulkData->LoadHairDescription(*HairDescription);
-			}
-			const bool bValidDescription = FGroomBuilder::ProcessHairDescription(*HairDescription, ProcessedHairDescription);
-			check(bValidDescription);
+			const FHairDescriptionGroups& LocalHairDescriptionGroups = GetHairDescriptionGroups();
 
-			const uint32 GroupCount = ProcessedHairDescription.HairGroups.Num();
+			const uint32 GroupCount = LocalHairDescriptionGroups.HairGroups.Num();
 			SetNumGroup(GroupCount);
 		}
 		else
 		{
-			const uint32 GroupCount = HairGroupsData.Num();
-			SetNumGroup(GroupCount, false);
+			// Old format serialized the computed groom data directly, but are no longer supported
+			UE_LOG(LogHairStrands, Error, TEXT("[Groom] The groom asset %s is too old. Please reimported the groom from its original source file."));
+			SetNumGroup(0, false);
 		}
 	}
 #endif
@@ -739,10 +801,14 @@ void UGroomAsset::PostLoad()
 	check(GetNumHairGroups() > 0);
 
 	// Build hair strands if needed
+#if WITH_EDITORONLY_DATA
 	if (!bSucceed && IsHairStrandsEnabled(EHairStrandsShaderType::Strands))
 	{
-		const uint32 GroupCount = HairGroupsInterpolation.Num();
-		for (uint32 GroupIndex = 0; GroupIndex < GroupCount; ++GroupIndex)
+		const FHairDescriptionGroups& LocalHairDescriptionGroups = GetHairDescriptionGroups();
+
+		const int32 GroupCount = HairGroupsInterpolation.Num();
+		check(LocalHairDescriptionGroups.HairGroups.Num() == GroupCount);
+		for (int32 GroupIndex = 0; GroupIndex < GroupCount; ++GroupIndex)
 		{
 			const bool bNeedToBuildData = 
 				(HairGroupsData[GroupIndex].Guides.BulkData.GetNumCurves() == 0 ||
@@ -750,11 +816,17 @@ void UGroomAsset::PostLoad()
 				 HairGroupsData[GroupIndex].Strands.BulkData.GetNumCurves() > 0; // Empty groom has no data to build
 			if (bNeedToBuildData)
 			{
-				FGroomBuilder::BuildData(HairGroupsData[GroupIndex], HairGroupsInterpolation[GroupIndex], GroupIndex );
-				FGroomBuilder::BuildClusterData(this, FGroomBuilder::ComputeGroomBoundRadius(HairGroupsData), GroupIndex);
+				BuildHairGroup(
+					GroupIndex,
+					LocalHairDescriptionGroups,
+					HairGroupsInterpolation,
+					HairGroupsLOD,
+					HairGroupsInfo,
+					HairGroupsData);
 			}
 		}
 	}
+#endif
 
 	// Convert all material data to new format
 	{
@@ -818,32 +890,28 @@ void UGroomAsset::PreSave(FObjectPreSaveContext ObjectSaveContext)
 
 	if ((ChangeType & GroomChangeType_Interpolation) || (ChangeType & GroomChangeType_LOD))
 	{
-		if (!HairDescription)
-		{
-			HairDescription = MakeUnique<FHairDescription>();
-			HairDescriptionBulkData->LoadHairDescription(*HairDescription);
-		}
-		FProcessedHairDescription ProcessedHairDescription;
-		const bool bValidDescription = FGroomBuilder::ProcessHairDescription(*HairDescription, ProcessedHairDescription);
-		check(bValidDescription);
+		const FHairDescriptionGroups& LocalHairDescriptionGroups = GetHairDescriptionGroups();
 
 		FGroomComponentRecreateRenderStateContext RecreateRenderContext(this);
-		if (ChangeType & GroomChangeType_Interpolation)
+		for (uint32 GroupIt = 0; GroupIt < GroupCount; ++GroupIt)
 		{
-			for (uint32 GroupIt = 0; GroupIt < GroupCount; ++GroupIt)
+			const bool bHasInterpolationChanged = !(CachedHairGroupsInterpolation[GroupIt] == HairGroupsInterpolation[GroupIt]);
+			const bool bHasLODChanged = !(CachedHairGroupsLOD[GroupIt] == HairGroupsLOD[GroupIt]);
+
+			if (bHasInterpolationChanged)
 			{
-				const bool bHasChanged = !(CachedHairGroupsInterpolation[GroupIt] == HairGroupsInterpolation[GroupIt]);
-				if (bHasChanged)
-				{
-					FGroomBuilder::BuildGroom(ProcessedHairDescription, this, GroupIt);
-				}
+				BuildHairGroup(
+					GroupIt,
+					LocalHairDescriptionGroups,
+					HairGroupsInterpolation,
+					HairGroupsLOD,
+					HairGroupsInfo,
+					HairGroupsData);
+			}
+			else if (bHasLODChanged)
+			{
+				FGroomBuilder::BuildClusterData(HairGroupsData[GroupIt].Strands.Data, HairDescriptionGroups->BoundRadius, HairGroupsLOD[GroupIt], HairGroupsData[GroupIt].Strands.ClusterCullingData);
 			}
-		}
-
-		if (ChangeType & GroomChangeType_LOD)
-		{
-			const float GroomBoundRadius = FGroomBuilder::ComputeGroomBoundRadius(ProcessedHairDescription);
-			FGroomBuilder::BuildClusterData(this, ProcessedHairDescription);
 		}
 		InitResources();
 	}
@@ -1576,11 +1644,10 @@ bool UGroomAsset::CacheDerivedDatas()
 	const bool bIsGameThread = IsInGameThread();
 	FGroomComponentRecreateRenderStateContext RecreateContext(bIsGameThread ? this : nullptr);
 
-	FProcessedHairDescription ProcessedHairDescription;
 	const uint32 GroupCount = HairGroupsInterpolation.Num();
 	for (uint32 GroupIndex = 0; GroupIndex < GroupCount; ++GroupIndex)
 	{
-		bool bSucceed = CacheDerivedData(GroupIndex, ProcessedHairDescription);
+		bool bSucceed = CacheDerivedData(GroupIndex);
 		if (!bSucceed)
 			return false;
 	}
@@ -1593,7 +1660,7 @@ bool UGroomAsset::CacheDerivedDatas()
 	return true;
 }
 
-bool UGroomAsset::CacheDerivedData(uint32 GroupIndex, FProcessedHairDescription& ProcessedHairDescription)
+bool UGroomAsset::CacheDerivedData(uint32 GroupIndex)
 {
 	// Check if the asset correctly initialized prior building
 	if (!AreGroupsValid())
@@ -1609,7 +1676,7 @@ bool UGroomAsset::CacheDerivedData(uint32 GroupIndex, FProcessedHairDescription&
 	}
 
 	FString DerivedDataKey;
-	bool bSuccess = CacheStrandsData(GroupIndex, ProcessedHairDescription, DerivedDataKey);
+	bool bSuccess = CacheStrandsData(GroupIndex, DerivedDataKey);
 
 	// Cache the cards and meshes separately from the strands data
 	bSuccess |= CacheCardsGeometry(GroupIndex, DerivedDataKey);
@@ -1620,10 +1687,11 @@ bool UGroomAsset::CacheDerivedData(uint32 GroupIndex, FProcessedHairDescription&
 	return bSuccess;
 }
 
-bool UGroomAsset::CacheStrandsData(uint32 GroupIndex, FProcessedHairDescription& ProcessedHairDescription, FString& OutDerivedDataKey)
+bool UGroomAsset::CacheStrandsData(uint32 GroupIndex, FString& OutDerivedDataKey)
 {
 	if (!HairDescriptionBulkData)
 	{
+		UE_LOG(LogHairStrands, Error, TEXT("[Groom] The groom asset %s is too old. Please reimported the groom from its original source file."));
 		return false;
 	}
 
@@ -1662,33 +1730,29 @@ bool UGroomAsset::CacheStrandsData(uint32 GroupIndex, FProcessedHairDescription&
 		{
 			UE_LOG(LogHairStrands, Log, TEXT("[Groom/DDC] Strands - Not found (Groom:%s Group:%d)."), *GetName(), GroupIndex);
 		}
-		// Load the HairDescription from the bulk data if needed
-		if (!HairDescription)
-		{
-			HairDescription = MakeUnique<FHairDescription>();
-			HairDescriptionBulkData->LoadHairDescription(*HairDescription);
-		}
 
-		if (!ProcessedHairDescription.IsValid())
+		const FHairDescriptionGroups& LocalHairDescriptionGroups = GetHairDescriptionGroups();
+
+		if (!LocalHairDescriptionGroups.IsValid())
 		{
-			if (!FGroomBuilder::ProcessHairDescription(*HairDescription, ProcessedHairDescription))
-			{
-				return false;
-			}
+			UE_LOG(LogHairStrands, Warning, TEXT("[Groom] The groom asset %s does not have valid hair groups."));
+			return false;
 		}
 
 		// Build groom data with the new build settings
-		bSuccess = FGroomBuilder::BuildGroom(ProcessedHairDescription, this, GroupIndex);
-
-		FHairGroupData& HairGroupData = HairGroupsData[GroupIndex];
+		bSuccess = BuildHairGroup(
+			GroupIndex,
+			LocalHairDescriptionGroups,
+			HairGroupsInterpolation,
+			HairGroupsLOD,
+			HairGroupsInfo,
+			HairGroupsData);
+		
 		if (bSuccess)
 		{
-			// Build cluster data
-			FGroomBuilder::BuildClusterData(this, ProcessedHairDescription, GroupIndex);
-
 			// Using a LargeMemoryWriter for serialization since the data can be bigger than 2 GB
 			FLargeMemoryWriter LargeMemWriter(0, /*bIsPersistent=*/ true);
-			LargeMemWriter << HairGroupData;
+			LargeMemWriter << HairGroupsData[GroupIndex];
 
 			int64 UncompressedSize = LargeMemWriter.TotalSize();
 
@@ -2011,32 +2075,24 @@ bool UGroomAsset::BuildCardsGeometry(uint32 GroupIndex)
 				if (GroupData.Guides.BulkData.GetNumPoints() == 0)
 				{
 					GroupData.Strands.Data = LOD.Guides.Data;
-					GroupData.Guides.Data = LOD.Guides.Data;
-
-					GroupData.Strands.BulkData = LOD.Guides.BulkData;
-					GroupData.Guides.BulkData = LOD.Guides.BulkData;
+					GroupData.Guides.Data  = LOD.Guides.Data;
 
 					// The RenderData is filled out by BuildData
 					bCopyRenderData = true;
 				}
 
-				FGroomBuilder::BuildData(
-					LOD.Guides.Data,
-					GroupData.Guides.Data,
-					LOD.Guides.BulkData,
-					GroupData.Guides.BulkData,
-					LOD.Guides.InterpolationData,
-					LOD.Guides.InterpolationBulkData,
-					CardsInterpolationSettings,
-					true,
-					false,
-					true,
-					GroupIndex);
+				FHairGroupInfo DummyGroupInfo;
+
+				FGroomBuilder::BuildData(LOD.Guides.Data);
+				FGroomBuilder::BuildBulkData(DummyGroupInfo, LOD.Guides.Data, LOD.Guides.BulkData);
+
+				FGroomBuilder::BuildInterplationData(DummyGroupInfo, LOD.Guides.Data, GroupData.Guides.Data, CardsInterpolationSettings, LOD.Guides.InterpolationData);
+				FGroomBuilder::BuildInterplationBulkData(GroupData.Guides.Data, LOD.Guides.InterpolationData, LOD.Guides.InterpolationBulkData);
 
 				if (bCopyRenderData)
 				{
 					GroupData.Strands.BulkData = LOD.Guides.BulkData;
-					GroupData.Guides.BulkData = LOD.Guides.BulkData;
+					GroupData.Guides.BulkData  = LOD.Guides.BulkData;
 				}
 
 				LOD.Guides.RestResource = new FHairStrandsRestResource(LOD.Guides.BulkData, EHairStrandsResourcesType::Cards);
@@ -2612,13 +2668,11 @@ bool UGroomAsset::IsGlobalInterpolationEnable(int32 GroupIndex, int32 LODIndex)
 		(HairGroupsLOD[GroupIndex].LODs[LODIndex].GlobalInterpolation == EGroomOverrideType::Auto && EnableGlobalInterpolation);
 }
 
-bool FProcessedHairDescription::IsValid() const
+bool FHairDescriptionGroups::IsValid() const
 {
-	for (TPair<int32, FProcessedHairDescription::FHairGroup> HairGroupIt : HairGroups)
+	for (const FHairDescriptionGroup& HairGroup : HairGroups)
 	{
-		const FProcessedHairDescription::FHairGroup& Group = HairGroupIt.Value;
-		const FHairGroupInfo& GroupInfo = Group.Key;
-		if (GroupInfo.NumCurves == 0)
+		if (HairGroup.Info.NumCurves == 0)
 		{
 			return false;
 		}
@@ -2673,8 +2727,8 @@ void UGroomAsset::StripLODs(const TArray<int32>& LODsToKeep, bool bRebuildResour
 		{
 			FHairGroupData& GroupData = HairGroupsData[GroupIt];
 			InternalReleaseResource(GroupData.Strands.ClusterCullingResource);
+			FGroomBuilder::BuildClusterData(GroupData.Strands.Data, HairDescriptionGroups->BoundRadius, HairGroupsLOD[GroupIt], GroupData.Strands.ClusterCullingData);
 		}
-		FGroomBuilder::BuildClusterData(this, FGroomBuilder::ComputeGroomBoundRadius(HairGroupsData));
 	}
 }
 
diff --git a/Engine/Plugins/Runtime/HairStrands/Source/HairStrandsCore/Private/GroomBuilder.cpp b/Engine/Plugins/Runtime/HairStrands/Source/HairStrandsCore/Private/GroomBuilder.cpp
--- a/Engine/Plugins/Runtime/HairStrands/Source/HairStrandsCore/Private/GroomBuilder.cpp
+++ b/Engine/Plugins/Runtime/HairStrands/Source/HairStrandsCore/Private/GroomBuilder.cpp
@@ -35,6 +35,15 @@ FString FGroomBuilder::GetVersion()
 	return TEXT("1g");
 }
 
+namespace FHairStrandsDecimation
+{
+	void Decimate(
+		const FHairStrandsDatas& InData,
+		float CurveDecimationPercentage,
+		float VertexDecimationPercentage,
+		FHairStrandsDatas& OutData);
+}
+
 namespace HairStrandsBuilder
 {
 	FVector2D SignNotZero(const FVector2D& v)
@@ -1105,7 +1114,10 @@ namespace HairInterpolationBuilder
 	}
 
 	/** Build data for interpolation between simulation and rendering */
-	void BuildRenderData(const FHairStrandsInterpolationDatas& HairInterpolation, const FHairStrandsDatas& SimDatas, FHairStrandsInterpolationBulkData& OutBulkData)
+	void BuildRenderData(
+		const FHairStrandsDatas& SimDatas, 
+		const FHairStrandsInterpolationDatas& HairInterpolation, 
+		FHairStrandsInterpolationBulkData& OutBulkData)
 	{
 		TRACE_CPUPROFILER_EVENT_SCOPE(HairInterpolationBuilder::BuildRenderData);
 
@@ -1293,9 +1305,9 @@ private:
 	TArray<FIntVector> GuideIndices;
 };
 
-bool FGroomBuilder::ProcessHairDescription(const FHairDescription& HairDescription, FProcessedHairDescription& Out)
+bool FGroomBuilder::BuildHairDescriptionGroups(const FHairDescription& HairDescription, FHairDescriptionGroups& Out)
 {
-	TRACE_CPUPROFILER_EVENT_SCOPE(FGroomBuilder::ProcessHairDescription);
+	TRACE_CPUPROFILER_EVENT_SCOPE(FGroomBuilder::BuildHairDescriptionGroups);
 
 	// Convert HairDescription to HairStrandsDatas
 	// For now, just convert HairDescription to HairStrandsDatas
@@ -1346,7 +1358,7 @@ bool FGroomBuilder::ProcessHairDescription(const FHairDescription& HairDescripti
 	TStrandAttributesConstRef<float> StrandWidths = HairDescription.StrandAttributes().GetAttributesRef<float>(HairAttribute::Strand::Width);
 
 	TStrandAttributesConstRef<FVector2D> StrandRootUV = HairDescription.StrandAttributes().GetAttributesRef<FVector2D>(HairAttribute::Strand::RootUV);
-	Out.bHasUVData = StrandRootUV.IsValid();
+	const bool bHasUVData = StrandRootUV.IsValid();
 
 	TStrandAttributesConstRef<int> StrandGuides = HairDescription.StrandAttributes().GetAttributesRef<int>(HairAttribute::Strand::Guide);
 	TStrandAttributesConstRef<int> GroupIDs = HairDescription.StrandAttributes().GetAttributesRef<int>(HairAttribute::Strand::GroupID);
@@ -1359,7 +1371,21 @@ bool FGroomBuilder::ProcessHairDescription(const FHairDescription& HairDescripti
 
 	// To use ClosestGuides and GuideWeights attributes, guides must be imported from HairDescription and
 	// must include StrandID attribute since ClosestGuides references those IDs
-	Out.bCanUseClosestGuidesAndWeights = bImportGuides && StrandIDs.IsValid() && ClosestGuides.IsValid() && GuideWeights.IsValid();
+	const bool bCanUseClosestGuidesAndWeights = bImportGuides && StrandIDs.IsValid() && ClosestGuides.IsValid() && GuideWeights.IsValid();
+
+	auto FindOrAdd = [&Out](int32 GroupID) -> FHairDescriptionGroup&
+	{
+		for (FHairDescriptionGroup& Group : Out.HairGroups)
+		{
+			if (Group.Info.GroupID == GroupID)
+			{
+				return Group;
+			}
+		}
+		FHairDescriptionGroup& Group = Out.HairGroups.AddDefaulted_GetRef();
+		Group.Info.GroupID = GroupID;
+		return Group;
+	};
 
 	int32 GlobalVertexIndex = 0;
 	int32 NumHairPoints = 0;
@@ -1389,23 +1415,21 @@ bool FGroomBuilder::ProcessHairDescription(const FHairDescription& HairDescripti
 		}
 
 		FHairStrandsDatas* CurrentHairStrandsDatas = nullptr;
-		FProcessedHairDescription::FHairGroup& Group = Out.HairGroups.FindOrAdd(GroupID);
-		FHairGroupInfo& GroupInfo = Group.Key;
-		FHairGroupData& GroupData = Group.Value;
-		GroupInfo.GroupID = GroupID;
+		FHairDescriptionGroup& Group = FindOrAdd(GroupID);
+		check(Group.Info.GroupID == GroupID);
 		if (!bIsGuide)
 		{
 			NumHairPoints += CurveNumVertices;
-			CurrentHairStrandsDatas = &GroupData.Strands.Data;
+			CurrentHairStrandsDatas = &Group.Strands;
 
-			++GroupInfo.NumCurves;
+			++Group.Info.NumCurves;
 		}
 		else if (bImportGuides)
 		{
 			NumGuidePoints += CurveNumVertices;
-			CurrentHairStrandsDatas = &GroupData.Guides.Data;
+			CurrentHairStrandsDatas = &Group.Guides;
 
-			++GroupInfo.NumGuides;
+			++Group.Info.NumGuides;
 		}
 		else
 		{
@@ -1419,7 +1443,7 @@ bool FGroomBuilder::ProcessHairDescription(const FHairDescription& HairDescripti
 
 		CurrentHairStrandsDatas->StrandsCurves.CurvesCount.Add(CurveNumVertices);
 
-		if (Out.bCanUseClosestGuidesAndWeights)
+		if (bCanUseClosestGuidesAndWeights)
 		{
 			// ClosesGuides needs mapping of StrandID (the strand index in the HairDescription)
 			CurrentHairStrandsDatas->StrandsCurves.StrandIDs.Add(CurveIndex);
@@ -1430,7 +1454,7 @@ bool FGroomBuilder::ProcessHairDescription(const FHairDescription& HairDescripti
 			CurrentHairStrandsDatas->StrandsCurves.GroomIDToIndex.Add(ImportedGroomID, StrandCurveIndex);
 		}
 
-		if (Out.bHasUVData)
+		if (bHasUVData)
 		{
 			CurrentHairStrandsDatas->StrandsCurves.CurvesRootUV.Add(StrandRootUV[StrandID]);
 		}
@@ -1464,69 +1488,79 @@ bool FGroomBuilder::ProcessHairDescription(const FHairDescription& HairDescripti
 			CurrentHairStrandsDatas->StrandsPoints.PointsRadius.Add(VertexWidth * 0.5f);
 		}
 	}
+
+	// Sort groups by ascending GroupID
+	Out.HairGroups.Sort([](const FHairDescriptionGroup& A, const FHairDescriptionGroup& B)
+	{
+		return A.Info.GroupID < B.Info.GroupID;
+	});
+
 	// Sparse->Dense Groups
-	FProcessedHairDescription::FHairGroups CompactGroups;
-	int32 GroupIndex = 0;
-	for (auto& HairGroupIt : Out.HairGroups)
+	// Change GroupID to be contiguous
 	{
-		FProcessedHairDescription::FHairGroup& Group = CompactGroups.FindOrAdd(GroupIndex);
-		Group = HairGroupIt.Value;
+		int32 GroupIndex = 0;
+		for (FHairDescriptionGroup& Group : Out.HairGroups)
+		{
+			Group.Info.GroupID = GroupIndex++;
+			Group.bCanUseClosestGuidesAndWeights = bCanUseClosestGuidesAndWeights;
+			Group.bHasUVData = bHasUVData;
+		}
+	}
 
-		FHairGroupInfo& GroupInfo = Group.Key;
-		FHairGroupData& GroupData = Group.Value;
+	// Compute bound radius
+	{
+		FVector GroomBoundMin(FLT_MAX);
+		FVector GroomBoundMax(-FLT_MAX);
+		for (const FHairDescriptionGroup& Group : Out.HairGroups)
+		{
+			GroomBoundMin.X = FMath::Min(GroomBoundMin.X, Group.Strands.BoundingBox.Min.X);
+			GroomBoundMin.Y = FMath::Min(GroomBoundMin.Y, Group.Strands.BoundingBox.Min.Y);
+			GroomBoundMin.Z = FMath::Min(GroomBoundMin.Z, Group.Strands.BoundingBox.Min.Z);
 
-		GroupInfo.GroupID = GroupIndex;
+			GroomBoundMax.X = FMath::Max(GroomBoundMax.X, Group.Strands.BoundingBox.Max.X);
+			GroomBoundMax.Y = FMath::Max(GroomBoundMax.Y, Group.Strands.BoundingBox.Max.Y);
+			GroomBoundMax.Z = FMath::Max(GroomBoundMax.Z, Group.Strands.BoundingBox.Max.Z);
+		}
 
-		++GroupIndex;
+		Out.BoundRadius = FVector::Distance(GroomBoundMax, GroomBoundMin) * 0.5f;
 	}
-	Out.HairGroups = CompactGroups;
-	return true;
-}
 
-bool FGroomBuilder::BuildGroom(const FHairDescription& HairDescription, UGroomAsset* GroomAsset)
-{
-	if (!GroomAsset)
+	// Update GroupInfo
+	for (FHairDescriptionGroup& Group : Out.HairGroups)
 	{
-		return false;
-	}
+		// Sanity check
+		check(Group.Strands.GetNumCurves() == Group.Info.NumCurves);
+		check(Group.Guides.GetNumCurves() == Group.Info.NumGuides);
 
-	TRACE_CPUPROFILER_EVENT_SCOPE(FGroomBuilder::BuildGroom);
+		// Update infos
+		Group.Info.NumCurveVertices = Group.Strands.GetNumPoints();
+		Group.Info.NumGuideVertices = Group.Guides.GetNumPoints();
 
-	FProcessedHairDescription ProcessedHairDescription;
-	if (!ProcessHairDescription(HairDescription, ProcessedHairDescription))
-	{
-		return false;
-	}
+		// Prepare/Resize all attributes based on position/curves counts
+		Group.Strands.StrandsPoints.SetNum(Group.Strands.GetNumPoints());
+		Group.Strands.StrandsCurves.SetNum(Group.Strands.GetNumCurves());
 
-	const float GroomBoundRadius = FGroomBuilder::ComputeGroomBoundRadius(ProcessedHairDescription);
-	const uint32 GroupCount = ProcessedHairDescription.HairGroups.Num();
-	for (uint32 GroupIndex = 0; GroupIndex < GroupCount; ++GroupIndex)
-	{
-		bool bSucceed = BuildGroom(ProcessedHairDescription, GroomAsset, GroupIndex);
-		if (!bSucceed)
-		{
-			return false;
-		}
-		FGroomBuilder::BuildClusterData(GroomAsset, GroomBoundRadius, GroupIndex);
+		// Prepare/Resize all attributes based on position/curves counts
+		Group.Guides.StrandsPoints.SetNum(Group.Guides.GetNumPoints());
+		Group.Guides.StrandsCurves.SetNum(Group.Guides.GetNumCurves());
 	}
 
 	return true;
 }
 
-bool FGroomBuilder::BuildGroom(FProcessedHairDescription& ProcessedHairDescription, UGroomAsset* GroomAsset, uint32 GroupIndex)
+void FGroomBuilder::BuildData(FHairStrandsDatas& OutStrands)
 {
-	if (!GroomAsset)
-	{
-		return false;
-	}
-
-	TRACE_CPUPROFILER_EVENT_SCOPE(FGroomBuilder::BuildGroom);
-
-	const uint32 GroupCount = ProcessedHairDescription.HairGroups.Num();
-	check(GroupIndex < GroupCount);
-	check(uint32(GroomAsset->GetNumHairGroups()) == GroupCount);
+	HairStrandsBuilder::BuildInternalData(OutStrands, false);
+}
 
-	const FHairGroupsInterpolation& InSettings = GroomAsset->HairGroupsInterpolation[GroupIndex];
+void FGroomBuilder::BuildData(
+	const FHairDescriptionGroup& InHairDescriptionGroup, 
+	const FHairGroupsInterpolation& InSettings, 
+	FHairGroupInfo& OutGroupInfo, 
+	FHairStrandsDatas& OutRen,
+	FHairStrandsDatas& OutSim)
+{
+	TRACE_CPUPROFILER_EVENT_SCOPE(FGroomBuilder::BuildData);
 
 	// Sanitize decimation values. Do not update the 'InSettings' values directly as this would change 
 	// the groom asset and thus would change the DDC key
@@ -1534,159 +1568,120 @@ bool FGroomBuilder::BuildGroom(FProcessedHairDescription& ProcessedHairDescripti
 	const float VertexDecimation	= FMath::Clamp(InSettings.DecimationSettings.VertexDecimation, 0.f, 1.f);
 	const float HairToGuideDensity	= FMath::Clamp(InSettings.InterpolationSettings.HairToGuideDensity, 0.f, 1.f);	
 
-	for (TPair<int32, FProcessedHairDescription::FHairGroup>& HairGroupIt : ProcessedHairDescription.HairGroups)
 	{
-		int32 GroupID = HairGroupIt.Key;
-		if (GroupIndex != GroupID)
-		{
-			continue;
-		}
-
-		FProcessedHairDescription::FHairGroup& Group = HairGroupIt.Value;
-		FHairGroupInfo& GroupInfo = Group.Key;
-		FHairGroupData& GroupData = Group.Value;
+		OutGroupInfo.GroupID = InHairDescriptionGroup.Info.GroupID;
 
 		// Rendering data
 		{
-			FHairStrandsDatas& RenData = GroupData.Strands.Data;
+			OutRen = InHairDescriptionGroup.Strands;
 
-			const int32 RenGroupNumCurves = RenData.StrandsCurves.Num();
-			RenData.StrandsCurves.SetNum(RenGroupNumCurves);
-			GroupInfo.NumCurves = RenGroupNumCurves;
-
-			const int32 RenGroupNumPoints = RenData.StrandsPoints.Num();
-			RenData.StrandsPoints.SetNum(RenGroupNumPoints);
-
-			HairStrandsBuilder::BuildInternalData(RenData, !ProcessedHairDescription.bHasUVData);
+			HairStrandsBuilder::BuildInternalData(OutRen, !InHairDescriptionGroup.bHasUVData);
 
 			// Decimate
 			if (CurveDecimation < 1 || VertexDecimation < 1)
 			{
-				FHairStrandsDatas FullData = RenData;
-				RenData.Reset();
-				Decimate(FullData, CurveDecimation, VertexDecimation, RenData);
+				FHairStrandsDatas FullData = OutRen;
+				OutRen.Reset();
+				FHairStrandsDecimation::Decimate(FullData, CurveDecimation, VertexDecimation, OutRen);
 			}
+			OutGroupInfo.NumCurves			= OutRen.GetNumCurves();
+			OutGroupInfo.NumCurveVertices	= OutRen.GetNumPoints();
+
+			// Sanity check
+			check(OutGroupInfo.NumCurves		<= InHairDescriptionGroup.Info.NumCurves);
+			check(OutGroupInfo.NumCurveVertices <= InHairDescriptionGroup.Info.NumCurveVertices);
+
 		}
 
 		// Simulation data
 		{
-			FHairStrandsDatas& SimData = GroupData.Guides.Data;
-			const int32 SimGroupNumGuides = SimData.StrandsCurves.Num();
-
-			if (SimGroupNumGuides > 0 && !InSettings.InterpolationSettings.bOverrideGuides)
+			OutSim = InHairDescriptionGroup.Guides;
+			if (InHairDescriptionGroup.Info.NumCurves > 0 && !InSettings.InterpolationSettings.bOverrideGuides)
 			{
-				GroupInfo.NumGuides = SimGroupNumGuides;
-				SimData.StrandsCurves.SetNum(SimGroupNumGuides);
-
-				const int32 GroupNumPoints = SimData.StrandsPoints.Num();
-				SimData.StrandsPoints.SetNum(GroupNumPoints);
-
-				HairStrandsBuilder::BuildInternalData(SimData, true); // Imported guides don't currently have root UVs so force computing them
+				HairStrandsBuilder::BuildInternalData(OutSim, true); // Imported guides don't currently have root UVs so force computing them
 			}
 			else
 			{
-				SimData.Reset();
-				Decimate(GroupData.Strands.Data, HairToGuideDensity, 1, SimData);
+				OutSim.Reset();
+				FHairStrandsDecimation::Decimate(OutRen, HairToGuideDensity, 1, OutSim);
 			}
-		}
-	}
 
-	for (TPair<int32, FProcessedHairDescription::FHairGroup> HairGroupIt : ProcessedHairDescription.HairGroups)
-	{
-		int32 GroupID = HairGroupIt.Key;
-		if (GroupIndex != GroupID)
-		{
-			continue;
-		}
+			OutGroupInfo.NumGuides			= OutSim.GetNumCurves();
+			OutGroupInfo.NumGuideVertices	= OutSim.GetNumPoints();
 
-		FProcessedHairDescription::FHairGroup& Group = HairGroupIt.Value;
-		FHairGroupInfo& GroupInfo = Group.Key;
-		FHairGroupData& GroupData = Group.Value;
-		GroomAsset->HairGroupsData[GroupIndex] = MoveTemp(GroupData);
+			// Sanity check
+			check(OutGroupInfo.NumCurves		<= InHairDescriptionGroup.Info.NumCurves);
+			check(OutGroupInfo.NumCurveVertices <= InHairDescriptionGroup.Info.NumCurveVertices);
+		}
 	}
 
 	// If there's usable closest guides and guide weights attributes, fill them into the asset
 	// This step requires the HairSimulationData (guides) to be filled prior to this
-	// TODO
-	//if (ProcessedHairDescription.bCanUseClosestGuidesAndWeights)
+	//if (InHairDescriptionGroup.bCanUseClosestGuidesAndWeights)
 	//{
 	//	HairInterpolationBuilder::FillInterpolationData(GroomAsset, HairDescription);
 	//}
-
-	
-	FHairGroupData& GroupData = GroomAsset->HairGroupsData[GroupIndex];
-	BuildData(GroupData, InSettings, GroupIndex);
-
-	return true;
 }
 
-void FGroomBuilder::BuildData(UGroomAsset* GroomAsset)
+void FGroomBuilder::BuildBulkData(
+	const FHairGroupInfo& InInfo,
+	const FHairStrandsDatas& InData,
+	FHairStrandsBulkData& OutBulkData)
 {
-	if (!GroomAsset)
-	{
-		return;
-	}
-
-	TRACE_CPUPROFILER_EVENT_SCOPE(FGroomBuilder::BuildData);
-
-	for (int32 Index = 0, GroupCount = GroomAsset->GetNumHairGroups(); Index < GroupCount; ++Index)
+	const int32 NumCurves = InData.GetNumCurves();
+	FRandomStream Random;
+	Random.Initialize(InInfo.GroupID);
+	TArray<uint8> CurveSeeds;
+	CurveSeeds.SetNumUninitialized(NumCurves);
+	for (int32 Index = 0; Index < NumCurves; ++Index)
 	{
-		FHairGroupsInterpolation& IntepolationSettings = GroomAsset->HairGroupsInterpolation[Index];
-		FHairGroupData& GroupData = GroomAsset->HairGroupsData[Index];
-		BuildData(GroupData, IntepolationSettings, Index);
+		CurveSeeds[Index] = Random.RandHelper(255);
 	}
+	HairStrandsBuilder::BuildRenderData(InData, CurveSeeds, OutBulkData);
 }
 
-void FGroomBuilder::BuildData(
-	const FHairStrandsDatas& RenData,
-	const FHairStrandsDatas& SimData,
-	FHairStrandsBulkData& RenBulkData,
-	FHairStrandsBulkData& SimBulkData,
-	FHairStrandsInterpolationDatas& InterpolationData,
-	FHairStrandsInterpolationBulkData& InterpolationBulkData,
-	const FHairInterpolationSettings& InterpolationSettings,
-	const bool bBuildRen,
-	const bool bBuildSim,
-	const bool bBuildInterpolation,
-	uint32 Seed)
+void FGroomBuilder::BuildInterplationData(
+	const FHairGroupInfo& InInfo,
+	const FHairStrandsDatas& InRenData,
+	const FHairStrandsDatas& InSimData,
+	const FHairInterpolationSettings& InInterpolationSettings,
+	FHairStrandsInterpolationDatas& OutInterpolationData)
 {
-	// Sanity check
-	check(RenData.GetNumCurves() > 0);
-	check(SimData.GetNumCurves() > 0);
-
-	// Precompute the random values used to build the data
-	// In particular, BuildInterpolationData needs this because it parallelizes its computation
-	FGroomDataRandomizer Random(Seed, RenData.GetNumCurves(), SimData.GetNumCurves());
-
-	// Build RenderData for HairStrandsDatas
-	if (bBuildRen) { HairStrandsBuilder::BuildRenderData(RenData, Random.GetRenderCurveSeeds(),	RenBulkData); }
-	if (bBuildSim) { HairStrandsBuilder::BuildRenderData(SimData, Random.GetSimCurveSeeds(),	SimBulkData); }
-
 	// Build Rendering data for InterpolationData
-	if (bBuildInterpolation)
+	// Build InterpolationData from render and simulation HairStrandsDatas
+	// Skip building if interpolation data was provided by the source file
+	if (OutInterpolationData.Num() == 0)
 	{
-		// Build InterpolationData from render and simulation HairStrandsDatas
-		// Skip building if interpolation data was provided by the source file
-		if (InterpolationData.Num() == 0)
+		// This randomization makes certain strands being affected by 1, 2, or 3 guides
+		FRandomStream Random;
+		Random.Initialize(InInfo.GroupID);
+		const int32 NumRenderCurves = InRenData.GetNumCurves();
+		TArray<FIntVector> GuideIndices;
+		GuideIndices.SetNumUninitialized(NumRenderCurves);
+		for (int32 Index = 0; Index < NumRenderCurves; ++Index)
 		{
-			HairInterpolationBuilder::BuildInterpolationData(InterpolationData, SimData, RenData, InterpolationSettings, Random.GetRandomGuideIndices());
+			FIntVector& RandomIndices = GuideIndices[Index];
+			for (int GuideIndex = 0; GuideIndex < HairInterpolationBuilder::FMetrics::Count; ++GuideIndex)
+			{
+				RandomIndices[GuideIndex] = Random.RandRange(0, HairInterpolationBuilder::FMetrics::Count - 1);
+			}
 		}
-		HairInterpolationBuilder::BuildRenderData(InterpolationData, SimData, InterpolationBulkData);
+
+		HairInterpolationBuilder::BuildInterpolationData(OutInterpolationData, InSimData, InRenData, InInterpolationSettings, GuideIndices);
 	}
 }
 
-void FGroomBuilder::BuildData(FHairGroupData& GroupData, const FHairGroupsInterpolation& InSettings, uint32 GroupIndex)
+void FGroomBuilder::BuildInterplationBulkData(
+	const FHairStrandsDatas& InSimData,
+	const FHairStrandsInterpolationDatas& InInterpolationData,
+	FHairStrandsInterpolationBulkData& OutInterpolationData)
 {
-	BuildData(
-		GroupData.Strands.Data,
-		GroupData.Guides.Data,
-		GroupData.Strands.BulkData,
-		GroupData.Guides.BulkData,
-		GroupData.Strands.InterpolationData,
-		GroupData.Strands.InterpolationBulkData,
-		InSettings.InterpolationSettings, true, true, true, GroupIndex);
+	HairInterpolationBuilder::BuildRenderData(InSimData, InInterpolationData, OutInterpolationData);
 }
 
+namespace FHairStrandsDecimation
+{
+
 inline uint32 DecimatePointCount(uint32 InCount, float InDecimationFactor)
 {
 	return FMath::Clamp(uint32(FMath::CeilToInt(InCount * InDecimationFactor)), 2u, InCount);
@@ -1731,7 +1726,7 @@ static void DecimateCurve(
 	}
 }
 
-void FGroomBuilder::Decimate(
+void Decimate(
 	const FHairStrandsDatas& InData, 
 	float CurveDecimationPercentage, 
 	float VertexDecimationPercentage, 
@@ -1840,6 +1835,7 @@ void FGroomBuilder::Decimate(
 
 	HairStrandsBuilder::BuildInternalData(OutData, false);
 }
+} // namespace FHairStrandsDecimation
 
 
 ///////////////////////////////////////////////////////////////////////////////////////////////////
@@ -2042,48 +2038,7 @@ inline uint32 to10Bits(float V)
 	return FMath::Clamp(uint32(V * 1024), 0u, 1023u);
 }
 
-float FGroomBuilder::ComputeGroomBoundRadius(const FProcessedHairDescription& Description)
-{
-	FVector GroomBoundMin(FLT_MAX);
-	FVector GroomBoundMax(-FLT_MAX);
-	for (auto& Group : Description.HairGroups)
-	{
-		const FHairGroupData::FStrands& Strands = Group.Value.Value.Strands;
-		GroomBoundMin.X = FMath::Min(GroomBoundMin.X, Strands.Data.BoundingBox.Min.X);
-		GroomBoundMin.Y = FMath::Min(GroomBoundMin.Y, Strands.Data.BoundingBox.Min.Y);
-		GroomBoundMin.Z = FMath::Min(GroomBoundMin.Z, Strands.Data.BoundingBox.Min.Z);
-
-		GroomBoundMax.X = FMath::Max(GroomBoundMax.X, Strands.Data.BoundingBox.Max.X);
-		GroomBoundMax.Y = FMath::Max(GroomBoundMax.Y, Strands.Data.BoundingBox.Max.Y);
-		GroomBoundMax.Z = FMath::Max(GroomBoundMax.Z, Strands.Data.BoundingBox.Max.Z);
-	}
-
-	const float GroomBoundRadius = FVector::Distance(GroomBoundMax, GroomBoundMin) * 0.5f;
-	return GroomBoundRadius;
-}
-
-float FGroomBuilder::ComputeGroomBoundRadius(const TArray<FHairGroupData>& HairGroupsData)
-{
-	// Compute the bounding box of all the groups. This is used for scaling LOD sceensize 
-	// for each group & cluster respectively to their relative size
-	FVector GroomBoundMin(FLT_MAX);
-	FVector GroomBoundMax(-FLT_MAX);
-	for (const FHairGroupData& LocalGroupData : HairGroupsData)
-	{
-		GroomBoundMin.X = FMath::Min(GroomBoundMin.X, LocalGroupData.Strands.Data.BoundingBox.Min.X);
-		GroomBoundMin.Y = FMath::Min(GroomBoundMin.Y, LocalGroupData.Strands.Data.BoundingBox.Min.Y);
-		GroomBoundMin.Z = FMath::Min(GroomBoundMin.Z, LocalGroupData.Strands.Data.BoundingBox.Min.Z);
-
-		GroomBoundMax.X = FMath::Max(GroomBoundMax.X, LocalGroupData.Strands.Data.BoundingBox.Max.X);
-		GroomBoundMax.Y = FMath::Max(GroomBoundMax.Y, LocalGroupData.Strands.Data.BoundingBox.Max.Y);
-		GroomBoundMax.Z = FMath::Max(GroomBoundMax.Z, LocalGroupData.Strands.Data.BoundingBox.Max.Z);
-	}
-
-	const float GroomBoundRadius = FVector::Distance(GroomBoundMax, GroomBoundMin) * 0.5f;
-	return GroomBoundRadius;
-}
-
-static void InternalBuildClusterData(
+void FGroomBuilder::BuildClusterData(
 	const FHairStrandsDatas& InRenStrandsData,
 	const float InGroomAssetRadius, 
 	const FHairGroupsLOD& InSettings, 
@@ -2495,51 +2450,4 @@ static void InternalBuildClusterData(
 	}
 }
 
-void FGroomBuilder::BuildClusterData(UGroomAsset* GroomAsset, const float GroomBoundRadius)
-{
-	if (GroomAsset)
-	{
-		const uint32 GroupCount = GroomAsset->HairGroupsData.Num();
-		for (uint32 GroupIt = 0; GroupIt < GroupCount; GroupIt++)
-		{
-			FHairGroupData& GroupData = GroomAsset->HairGroupsData[GroupIt];
-			InternalBuildClusterData(
-				GroupData.Strands.Data,
-				GroomBoundRadius,
-				GroomAsset->HairGroupsLOD[GroupIt],
-				GroupData.Strands.ClusterCullingData);
-		}
-	}
-}
-void FGroomBuilder::BuildClusterData(UGroomAsset* GroomAsset, const float GroomBoundRadius, uint32 GroupIndex)
-{
-	if (GroomAsset)
-	{
-		FHairGroupData& GroupData = GroomAsset->HairGroupsData[GroupIndex];
-		InternalBuildClusterData(
-			GroupData.Strands.Data,
-			GroomBoundRadius,
-			GroomAsset->HairGroupsLOD[GroupIndex],
-			GroupData.Strands.ClusterCullingData);
-	}
-}
-
-void FGroomBuilder::BuildClusterData(UGroomAsset* GroomAsset, const FProcessedHairDescription& ProcessedHairDescription)
-{
-	if (GroomAsset)
-	{
-		const float GroomBoundRadius = FGroomBuilder::ComputeGroomBoundRadius(ProcessedHairDescription);
-		FGroomBuilder::BuildClusterData(GroomAsset, GroomBoundRadius);
-	}
-}
-void FGroomBuilder::BuildClusterData(UGroomAsset* GroomAsset, const FProcessedHairDescription& ProcessedHairDescription, uint32 GroupIndex)
-{
-	if (GroomAsset)
-	{
-		const float GroomBoundRadius = FGroomBuilder::ComputeGroomBoundRadius(ProcessedHairDescription);
-		FGroomBuilder::BuildClusterData(GroomAsset, GroomBoundRadius, GroupIndex);
-	}
-}
-
-
 #undef LOCTEXT_NAMESPACE
diff --git a/Engine/Plugins/Runtime/HairStrands/Source/HairStrandsCore/Private/GroomRBFDeformer.cpp b/Engine/Plugins/Runtime/HairStrands/Source/HairStrandsCore/Private/GroomRBFDeformer.cpp
--- a/Engine/Plugins/Runtime/HairStrands/Source/HairStrandsCore/Private/GroomRBFDeformer.cpp
+++ b/Engine/Plugins/Runtime/HairStrands/Source/HairStrandsCore/Private/GroomRBFDeformer.cpp
@@ -376,8 +376,38 @@ static void ApplyDeformationToGroom(const TArray<FRBFDeformedPositions>& Deforme
 		VertexPositions[VertexID] = FlattenedDeformedPositions[VertexIndex];
 	}
 
-	FGroomBuilder::BuildGroom(HairDescription, GroomAsset);
-	FGroomBuilder::BuildClusterData(GroomAsset, FGroomBuilder::ComputeGroomBoundRadius(GroomAsset->HairGroupsData));
+	{
+		FHairDescriptionGroups HairDescriptionGroups;
+		FGroomBuilder::BuildHairDescriptionGroups(HairDescription, HairDescriptionGroups);
+
+		const int32 GroupCount = GroomAsset->HairGroupsInterpolation.Num();
+		check(HairDescriptionGroups.HairGroups.Num() == GroupCount);
+		for (int32 GroupIndex = 0; GroupIndex < GroupCount; ++GroupIndex)
+		{
+			const bool bIsValid = HairDescriptionGroups.IsValid();
+			if (bIsValid)
+			{
+				const FHairDescriptionGroup& HairGroup = HairDescriptionGroups.HairGroups[GroupIndex];
+				check(GroupIndex <= HairDescriptionGroups.HairGroups.Num());
+				check(GroupIndex == HairGroup.Info.GroupID);
+				const FHairInterpolationSettings& InterpolationSettings = GroomAsset->HairGroupsInterpolation[GroupIndex].InterpolationSettings;
+				const FHairGroupsLOD& HairGroupLOD = GroomAsset->HairGroupsLOD[GroupIndex];
+
+				FHairGroupInfo& HairGroupsInfo = GroomAsset->HairGroupsInfo[GroupIndex];
+				FHairGroupData& HairGroupsData = GroomAsset->HairGroupsData[GroupIndex];
+
+				FGroomBuilder::BuildData(HairGroup, GroomAsset->HairGroupsInterpolation[GroupIndex], HairGroupsInfo, HairGroupsData.Strands.Data, HairGroupsData.Guides.Data);
+
+				FGroomBuilder::BuildBulkData(HairGroup.Info, HairGroupsData.Guides.Data,  HairGroupsData.Guides.BulkData);
+				FGroomBuilder::BuildBulkData(HairGroup.Info, HairGroupsData.Strands.Data, HairGroupsData.Strands.BulkData);
+
+				FGroomBuilder::BuildInterplationData(HairGroup.Info, HairGroupsData.Strands.Data, HairGroupsData.Guides.Data, InterpolationSettings, HairGroupsData.Strands.InterpolationData);
+				FGroomBuilder::BuildInterplationBulkData(HairGroupsData.Guides.Data, HairGroupsData.Strands.InterpolationData, HairGroupsData.Strands.InterpolationBulkData);
+
+				FGroomBuilder::BuildClusterData(HairGroupsData.Strands.Data, HairDescriptionGroups.BoundRadius, HairGroupLOD, HairGroupsData.Strands.ClusterCullingData);
+			}
+		}
+	}
 	
 	GroomAsset->CommitHairDescription(MoveTemp(HairDescription));
 	GroomAsset->UpdateHairGroupsInfo();
@@ -469,8 +499,8 @@ void FGroomRBFDeformer::GetRBFDeformedGroomAsset(const UGroomAsset* InGroomAsset
 
 		// Use the vertices positions from the HairDescription instead of the GroomAsset since the latter
 		// may contain decimated or auto-generated guides depending on the import settings
-		FProcessedHairDescription ProcessedHairDescription;
-		FGroomBuilder::ProcessHairDescription(InGroomAsset->GetHairDescription(), ProcessedHairDescription);
+		FHairDescriptionGroups HairDescriptionGroups;
+		FGroomBuilder::BuildHairDescriptionGroups(InGroomAsset->GetHairDescription(), HairDescriptionGroups);
 
 		TArray<FRBFDeformedPositions> DeformedPositions;
 		DeformedPositions.SetNum(NumGroups);
@@ -484,17 +514,17 @@ void FGroomRBFDeformer::GetRBFDeformedGroomAsset(const UGroomAsset* InGroomAsset
 		// Note that the GroupID from the HairGroups cannot be used as the GroupIndex since 
 		// the former may not be strictly increasing nor consecutive
 		// but the ordering of the groups does represent the GroupIndex		
-		int32 GroupIndex = 0;
-		for (FProcessedHairDescription::FHairGroups::TConstIterator GroupIt = ProcessedHairDescription.HairGroups.CreateConstIterator(); GroupIt; ++GroupIt)
+		for (const FHairDescriptionGroup& Group : HairDescriptionGroups.HairGroups)
 		{
-			const FHairGroupData& OriginalHairGroupData = GroupIt->Value.Value;
+			const uint32 GroupIndex = Group.Info.GroupID;
+			const FHairStrandsDatas& OriginalGuides = Group.Guides;
 			const FHairGroupData& HairGroupData = InGroomAsset->HairGroupsData[GroupIndex];
 
 			// Get deformed guides
 			// If the groom override the value, we output dummy value for the guides, since they won't be used
 			if (InGroomAsset->HairGroupsInterpolation[GroupIndex].InterpolationSettings.bOverrideGuides)
 			{
-				const uint32 OriginalVertexCount = OriginalHairGroupData.Guides.Data.StrandsPoints.Num();
+				const uint32 OriginalVertexCount = OriginalGuides.StrandsPoints.Num();
 				DeformedPositions[GroupIndex].GuideStrands.Init(FVector::ZeroVector, OriginalVertexCount);
 			}
 			else
@@ -516,8 +546,6 @@ void FGroomRBFDeformer::GetRBFDeformedGroomAsset(const UGroomAsset* InGroomAsset
 				SkeletalMeshData_Target,
 				BindingAsset->HairGroupDatas[GroupIndex].RenRootData.VertexToCurveIndexBuffer,
 				BindingAsset->HairGroupDatas[GroupIndex].RenRootData.MeshProjectionLODs[MeshLODIndex]);
-
-			++GroupIndex;
 		}
 
 		// Apply changes onto cards and meshes (OutGroomASset already contain duplicated mesh asset
diff --git a/Engine/Plugins/Runtime/HairStrands/Source/HairStrandsCore/Public/GroomAsset.h b/Engine/Plugins/Runtime/HairStrands/Source/HairStrandsCore/Public/GroomAsset.h
--- a/Engine/Plugins/Runtime/HairStrands/Source/HairStrandsCore/Public/GroomAsset.h
+++ b/Engine/Plugins/Runtime/HairStrands/Source/HairStrandsCore/Public/GroomAsset.h
@@ -277,14 +277,21 @@ struct HAIRSTRANDSCORE_API FHairGroupData
 	bool bIsCookedOut = false;
 };
 
-struct FProcessedHairDescription
+struct FHairDescriptionGroup
 {
-	typedef TPair<FHairGroupInfo, FHairGroupData> FHairGroup;
-	typedef TMap<int32, FHairGroup> FHairGroups;
-	FHairGroups HairGroups;
-	bool bCanUseClosestGuidesAndWeights = false;
-	bool bHasUVData = false;
-	bool IsValid() const;
+	FHairGroupInfo		Info;
+	FHairStrandsDatas	Strands;
+	FHairStrandsDatas	Guides;
+
+	bool  bCanUseClosestGuidesAndWeights = false;
+	bool  bHasUVData = false;
+};
+
+struct FHairDescriptionGroups
+{
+	TArray<FHairDescriptionGroup> HairGroups;
+	float BoundRadius = 0;
+	bool  IsValid() const;
 };
 
 USTRUCT(BlueprintType)
@@ -501,7 +508,7 @@ private:
 	FHairStrandsRaytracingResource*		AllocateCardsRaytracingResources(uint32 GroupIndex, uint32 LODIndex);
 	FHairStrandsRaytracingResource*		AllocateMeshesRaytracingResources(uint32 GroupIndex, uint32 LODIndex);
 	FHairStrandsRaytracingResource*		AllocateStrandsRaytracingResources(uint32 GroupIndex);
-#endif
+#endif // RHI_RAYTRACING
 	friend class UGroomComponent;
 
 #if WITH_EDITORONLY_DATA
@@ -518,9 +525,9 @@ public:
 	/** Caches the computed (group) groom data with the given build settings from/to the Derived Data Cache, building it if needed.
 	 *  This function assumes the interpolation settings are properly populated, as they will be used to build the asset.
 	 */
-	bool CacheDerivedData(uint32 GroupIndex, FProcessedHairDescription& ProcessedHairDescription);
 	bool CacheDerivedDatas();
-	bool CacheStrandsData(uint32 GroupIndex, FProcessedHairDescription& ProcessedHairDescription, FString& OutDerivedDataKey);
+	bool CacheDerivedData(uint32 GroupIndex);
+	bool CacheStrandsData(uint32 GroupIndex, FString& OutDerivedDataKey);
 	bool CacheCardsGeometry(uint32 GroupIndex, const FString& StrandsKey);
 	bool CacheMeshesGeometry(uint32 GroupIndex);
 
@@ -530,10 +537,12 @@ public:
 	FString GetDerivedDataKeyForMeshes(uint32 GroupIndex);
 
 private:
+	const FHairDescriptionGroups& GetHairDescriptionGroups();
 	FString BuildDerivedDataKeySuffix(uint32 GroupIndex, const FHairGroupsInterpolation& InterpolationSettings, const FHairGroupsLOD& LODSettings) const;
 	bool IsFullyCached();
 	TUniquePtr<FHairDescription> HairDescription;
 	TUniquePtr<FHairDescriptionBulkData> HairDescriptionBulkData;
+	TUniquePtr<FHairDescriptionGroups> HairDescriptionGroups;
 
 	// Transient property for visualizing the groom in a certain debug mode. This is used by the groom editor
 	EHairStrandsDebugMode DebugMode = EHairStrandsDebugMode::NoneDebug;
@@ -544,7 +553,7 @@ private:
 
 	TStrongObjectPtr<UGroomAsset> GroomAssetStrongPtr;
 	bool bRetryLoadFromGameThread = false;
-#endif
+#endif // WITH_EDITORONLY_DATA
 	bool bIsInitialized = false;
 
 #if WITH_EDITOR
@@ -565,5 +574,5 @@ private:
 	// Queue of procedural assets which needs to be saved
 	TQueue<UStaticMesh*> AssetToSave_Meshes;
 	TQueue<FHairGroupCardsTextures*> AssetToSave_Textures;
-#endif
+#endif // WITH_EDITOR
 };
diff --git a/Engine/Plugins/Runtime/HairStrands/Source/HairStrandsCore/Public/GroomBuilder.h b/Engine/Plugins/Runtime/HairStrands/Source/HairStrandsCore/Public/GroomBuilder.h
--- a/Engine/Plugins/Runtime/HairStrands/Source/HairStrandsCore/Public/GroomBuilder.h
+++ b/Engine/Plugins/Runtime/HairStrands/Source/HairStrandsCore/Public/GroomBuilder.h
@@ -11,45 +11,71 @@ struct FHairStrandsDatas;
 struct FHairStrandsClusterCullingData;
 struct FHairGroupsLOD;
 struct FHairGroupData;
-struct FProcessedHairDescription;
+struct FHairDescriptionGroups;
 struct FHairGroupsInterpolation;
 struct FHairStrandsInterpolationDatas;
 struct FHairInterpolationSettings;
 struct FHairStrandsBulkData;
 struct FHairStrandsInterpolationBulkData;
+struct FHairDescriptionGroup;
+struct FHairGroupInfo;
 class FHairDescription;
 class UGroomAsset;
 
+// Data flow overview
+// ==================
+// HairDescription -> HairDescriptionGroups -> HairStrandsData -> HairStrandsBulkData*
+//															   -> HairStrandsInterpolationData ->HairStrandsInterpolationBulkData*
+//															   -> HairStrandsClusterData*
+//
+// * Data used at runtime. Other type of data are intermediate data used only during building/within the editor
 struct HAIRSTRANDSCORE_API FGroomBuilder
 {
 	static FString GetVersion();
 
-	static bool ProcessHairDescription(const FHairDescription& HairDescription, FProcessedHairDescription& Out);
-	static void Decimate(const FHairStrandsDatas& InData, float CurveDecimationPercentage, float VertexDecimationPercentage, FHairStrandsDatas& OutData);
+	// 1. Build hair group based on the hair description
+	static bool BuildHairDescriptionGroups(
+		const FHairDescription& HairDescription, 
+		FHairDescriptionGroups& Out);
 
-	static void BuildData(FHairGroupData& GroupData, const FHairGroupsInterpolation& InterpolationSettings, uint32 GroupIndex);
-	static void BuildData(UGroomAsset* GroomAsset);
+	// 2.a Build FHairStrandsDatas for Strands & Guides, based on HairDescriptionGroups and DecimationSettings
 	static void BuildData(
-		const FHairStrandsDatas& RenData,
-		const FHairStrandsDatas& SimData,
-		FHairStrandsBulkData& RenBulkData,
-		FHairStrandsBulkData& SimBulkData,
-		FHairStrandsInterpolationDatas& InterpolationData,
-		FHairStrandsInterpolationBulkData& InterpolationBulkData,
-		const FHairInterpolationSettings& InterpolationSettings,
-		const bool bBuildRen,
-		const bool bBuildSim,
-		const bool bBuildInterpolation,
-		uint32 Seed);
-
-	static bool BuildGroom(FProcessedHairDescription& ProcessedHairDescription, UGroomAsset* GroomAsset, uint32 GroupIndex);
-	static bool BuildGroom(const FHairDescription& HairDescription, UGroomAsset* GroomAsset);
-
-	static float ComputeGroomBoundRadius(const FProcessedHairDescription& Description);
-	static float ComputeGroomBoundRadius(const TArray<FHairGroupData>& HairGroupsData);
-
-	static void BuildClusterData(UGroomAsset* GroomAsset, const float InGroomAssetRadius);
-	static void BuildClusterData(UGroomAsset* GroomAsset, const float InGroomAssetRadius, uint32 GroupIndex);
-	static void BuildClusterData(UGroomAsset* GroomAsset, const FProcessedHairDescription& ProcessedHairDescription);
-	static void BuildClusterData(UGroomAsset* GroomAsset, const FProcessedHairDescription& ProcessedHairDescription, uint32 GroupIndex);
+		const FHairDescriptionGroup& InHairDescriptionGroup,
+		const FHairGroupsInterpolation& InSettings,
+		FHairGroupInfo& OutGroupInfo,
+		FHairStrandsDatas& OutStrands,
+		FHairStrandsDatas& OutGuides);
+
+	// 2.b Build FHairStrandsDatas for Strands or Guides. 
+	// This version:
+	// * Suppose OutStrands already contains curves & points data. 
+	// * Compute only offset data/bounding box/max length/max radius/...
+	static void BuildData(FHairStrandsDatas& OutStrands);
+
+	// 3. Build bulk data for Strands / Guides
+	static void BuildBulkData(
+		const FHairGroupInfo& InInfo,
+		const FHairStrandsDatas& InData,
+		FHairStrandsBulkData& OutBulkData);
+
+	// 4. Build interplation data based on the hairStrands data
+	static void BuildInterplationData(
+		const FHairGroupInfo& InInfo,
+		const FHairStrandsDatas& InRenData,
+		const FHairStrandsDatas& InSimData,
+		const FHairInterpolationSettings& InInterpolationSettings,
+		FHairStrandsInterpolationDatas& OutInterpolationData);
+
+	// 5. Build interplation bulk data
+	static void BuildInterplationBulkData(
+		const FHairStrandsDatas& InSimData,
+		const FHairStrandsInterpolationDatas& InInterpolationData,
+		FHairStrandsInterpolationBulkData& OutInterpolationData);
+
+	// 6. Build cluster data
+	static void BuildClusterData(
+		const FHairStrandsDatas& InRenData,
+		const float InGroomAssetRadius,
+		const FHairGroupsLOD& InSettings,
+		FHairStrandsClusterCullingData& OutClusterCullingData);
 };
\ No newline at end of file
diff --git a/Engine/Plugins/Runtime/HairStrands/Source/HairStrandsEditor/Private/HairStrandsFactory.cpp b/Engine/Plugins/Runtime/HairStrands/Source/HairStrandsEditor/Private/HairStrandsFactory.cpp
--- a/Engine/Plugins/Runtime/HairStrands/Source/HairStrandsEditor/Private/HairStrandsFactory.cpp
+++ b/Engine/Plugins/Runtime/HairStrands/Source/HairStrandsEditor/Private/HairStrandsFactory.cpp
@@ -75,7 +75,7 @@ UObject* UHairStrandsFactory::FactoryCreateFile(UClass* InClass, UObject* InPare
 
 	{
 		// Load the alembic file upfront to preview & report any potential issue
-		FProcessedHairDescription OutDescription;
+		FHairDescriptionGroups OutDescription;
 		{
 			FScopedSlowTask Progress((float)1, LOCTEXT("ImportHairAssetForPreview", "Importing hair asset for preview..."), true);
 			Progress.MakeDialog(true);
@@ -86,7 +86,7 @@ UObject* UHairStrandsFactory::FactoryCreateFile(UClass* InClass, UObject* InPare
 				return nullptr;
 			}
 
-			FGroomBuilder::ProcessHairDescription(HairDescription, OutDescription);
+			FGroomBuilder::BuildHairDescriptionGroups(HairDescription, OutDescription);
 		
 			// Populate the interpolation settings based on the group count, as this is used later during the ImportHair() to define 
 			// the exact number of group to create
@@ -100,16 +100,12 @@ UObject* UHairStrandsFactory::FactoryCreateFile(UClass* InClass, UObject* InPare
 		// Convert the process hair description into hair groups
 		UGroomHairGroupsPreview* GroupsPreview = NewObject<UGroomHairGroupsPreview>();
 		{
-			uint32 GroupIndex = 0;
-			for (TPair<int32, FProcessedHairDescription::FHairGroup> HairGroupIt : OutDescription.HairGroups)
+			for (const FHairDescriptionGroup& Group : OutDescription.HairGroups)
 			{
-				const FProcessedHairDescription::FHairGroup& Group = HairGroupIt.Value;
-				const FHairGroupInfo& GroupInfo = Group.Key;
-
 				FGroomHairGroupPreview& OutGroup = GroupsPreview->Groups.AddDefaulted_GetRef();
-				OutGroup.GroupID = GroupInfo.GroupID;
-				OutGroup.CurveCount = GroupInfo.NumCurves;
-				OutGroup.GuideCount = GroupInfo.NumGuides;
+				OutGroup.GroupID	= Group.Info.GroupID;
+				OutGroup.CurveCount = Group.Info.NumCurves;
+				OutGroup.GuideCount = Group.Info.NumGuides;
 
 				if (OutGroup.GroupID < OutDescription.HairGroups.Num())
 				{				
diff --git a/Engine/Plugins/Runtime/HairStrands/Source/HairStrandsEditor/Private/HairStrandsImporter.cpp b/Engine/Plugins/Runtime/HairStrands/Source/HairStrandsEditor/Private/HairStrandsImporter.cpp
--- a/Engine/Plugins/Runtime/HairStrands/Source/HairStrandsEditor/Private/HairStrandsImporter.cpp
+++ b/Engine/Plugins/Runtime/HairStrands/Source/HairStrandsEditor/Private/HairStrandsImporter.cpp
@@ -44,7 +44,6 @@ UGroomAsset* FHairStrandsImporter::ImportHair(const FHairImportContext& ImportCo
 	OutHairAsset->CommitHairDescription(MoveTemp(HairDescription));
 
 	// Populate the interpolation settings with the new settings from the importer	
-	FProcessedHairDescription ProcessedHairDescription;
 	for (uint32 GroupIndex = 0; GroupIndex < GroupCount; ++GroupIndex)
 	{
 		// Insure the interpolation settings matches between the importer and the actual asset
diff --git a/Engine/Plugins/Runtime/HairStrands/Source/HairStrandsEditor/Private/ReimportHairStrandsFactory.cpp b/Engine/Plugins/Runtime/HairStrands/Source/HairStrandsEditor/Private/ReimportHairStrandsFactory.cpp
--- a/Engine/Plugins/Runtime/HairStrands/Source/HairStrandsEditor/Private/ReimportHairStrandsFactory.cpp
+++ b/Engine/Plugins/Runtime/HairStrands/Source/HairStrandsEditor/Private/ReimportHairStrandsFactory.cpp
@@ -103,7 +103,7 @@ EReimportResult::Type UReimportHairStrandsFactory::Reimport(UObject* Obj)
 		}
 
 		// Load the alembic file upfront to preview & report any potential issue
-		FProcessedHairDescription OutDescription;
+		FHairDescriptionGroups OutDescription;
 		{
 			FScopedSlowTask Progress((float)1, LOCTEXT("ReimportHairAsset", "Reimporting hair asset for preview..."), true);
 			Progress.MakeDialog(true);
@@ -114,7 +114,7 @@ EReimportResult::Type UReimportHairStrandsFactory::Reimport(UObject* Obj)
 				return EReimportResult::Failed;
 			}
 
-			FGroomBuilder::ProcessHairDescription(HairDescription, OutDescription);
+			FGroomBuilder::BuildHairDescriptionGroups(HairDescription, OutDescription);
 
 			// Populate the interpolation settings based on the group count from the description
 			const uint32 GroupCount = OutDescription.HairGroups.Num();
@@ -137,16 +137,12 @@ EReimportResult::Type UReimportHairStrandsFactory::Reimport(UObject* Obj)
 			// Convert the process hair description into hair groups
 			UGroomHairGroupsPreview* GroupsPreview = NewObject<UGroomHairGroupsPreview>();
 			{
-				uint32 GroupIndex = 0;
-				for (TPair<int32, FProcessedHairDescription::FHairGroup> HairGroupIt : OutDescription.HairGroups)
+				for (const FHairDescriptionGroup& Group : OutDescription.HairGroups)
 				{
-					const FProcessedHairDescription::FHairGroup& Group = HairGroupIt.Value;
-					const FHairGroupInfo& GroupInfo = Group.Key;
-
 					FGroomHairGroupPreview& OutGroup = GroupsPreview->Groups.AddDefaulted_GetRef();
-					OutGroup.GroupID = GroupInfo.GroupID;
-					OutGroup.CurveCount = GroupInfo.NumCurves;
-					OutGroup.GuideCount = GroupInfo.NumGuides;
+					OutGroup.GroupID	= Group.Info.GroupID;
+					OutGroup.CurveCount = Group.Info.NumCurves;
+					OutGroup.GuideCount = Group.Info.NumGuides;
 
 					if (OutGroup.GroupID < OutDescription.HairGroups.Num())
 					{
