CIS fix.

diff --git a/Engine/Plugins/Runtime/HairStrands/Source/HairStrandsCore/Private/GroomAsset.cpp b/Engine/Plugins/Runtime/HairStrands/Source/HairStrandsCore/Private/GroomAsset.cpp
--- a/Engine/Plugins/Runtime/HairStrands/Source/HairStrandsCore/Private/GroomAsset.cpp
+++ b/Engine/Plugins/Runtime/HairStrands/Source/HairStrandsCore/Private/GroomAsset.cpp
@@ -2691,6 +2691,7 @@ int32 UGroomAsset::GetLODCount() const
 	return MaxLODCount;
 }
 
+#if WITH_EDITORONLY_DATA
 void UGroomAsset::StripLODs(const TArray<int32>& LODsToKeep, bool bRebuildResources)
 {
 	// Assume that the LOD are ordered from 0 ... Max
@@ -2731,6 +2732,7 @@ void UGroomAsset::StripLODs(const TArray<int32>& LODsToKeep, bool bRebuildResour
 		}
 	}
 }
+#endif // WITH_EDITORONLY_DATA
 
 bool UGroomAsset::HasDebugData() const
 {
diff --git a/Engine/Plugins/Runtime/HairStrands/Source/HairStrandsCore/Public/GroomAsset.h b/Engine/Plugins/Runtime/HairStrands/Source/HairStrandsCore/Public/GroomAsset.h
--- a/Engine/Plugins/Runtime/HairStrands/Source/HairStrandsCore/Public/GroomAsset.h
+++ b/Engine/Plugins/Runtime/HairStrands/Source/HairStrandsCore/Public/GroomAsset.h
@@ -434,7 +434,9 @@ public:
 	int32 GetNumHairGroups() const;
 
 	int32 GetLODCount() const;
+#if WITH_EDITORONLY_DATA
 	void StripLODs(const TArray<int32>& LODsToKeep, bool bRebuildResources);
+#endif // WITH_EDITORONLY_DATA
 
 	/** Debug data for derived asset generation (strands textures, ...). */
 	bool HasDebugData() const;
