Fix crash due to invalid GroupInfo size when a groom is not cached in DDC.

diff --git a/Engine/Plugins/Runtime/HairStrands/Source/HairStrandsCore/Private/GroomAsset.cpp b/Engine/Plugins/Runtime/HairStrands/Source/HairStrandsCore/Private/GroomAsset.cpp
--- a/Engine/Plugins/Runtime/HairStrands/Source/HairStrandsCore/Private/GroomAsset.cpp
+++ b/Engine/Plugins/Runtime/HairStrands/Source/HairStrandsCore/Private/GroomAsset.cpp
@@ -1290,6 +1290,11 @@ void UGroomAsset::SetNumGroup(uint32 InGroupCount, bool bResetGroupData)
 		}
 	}
 
+	if (InGroupCount != HairGroupsInfo.Num())
+	{
+		HairGroupsInfo.Init(FHairGroupInfoWithVisibility(), InGroupCount);
+	}
+
 	EffectiveLODBias.SetNum(InGroupCount);
 
 #if WITH_EDITORONLY_DATA
