Fix crash when building grooms (invalid bounds & invalid guides data).

diff --git a/Engine/Plugins/Runtime/HairStrands/Source/HairStrandsCore/Private/GroomBuilder.cpp b/Engine/Plugins/Runtime/HairStrands/Source/HairStrandsCore/Private/GroomBuilder.cpp
--- a/Engine/Plugins/Runtime/HairStrands/Source/HairStrandsCore/Private/GroomBuilder.cpp
+++ b/Engine/Plugins/Runtime/HairStrands/Source/HairStrandsCore/Private/GroomBuilder.cpp
@@ -217,8 +217,8 @@ namespace HairStrandsBuilder
 	{
 		TRACE_CPUPROFILER_EVENT_SCOPE(HairStrandsBuilder::BuildRenderData);
 
-		uint32 NumCurves = HairStrands.GetNumCurves();
-		uint32 NumPoints = HairStrands.GetNumPoints();
+		const uint32 NumCurves = HairStrands.GetNumCurves();
+		const uint32 NumPoints = HairStrands.GetNumPoints();
 		if (!(NumCurves > 0 && NumPoints > 0))
 			return;
 
@@ -1513,13 +1513,16 @@ bool FGroomBuilder::BuildHairDescriptionGroups(const FHairDescription& HairDescr
 		FVector GroomBoundMax(-FLT_MAX);
 		for (const FHairDescriptionGroup& Group : Out.HairGroups)
 		{
-			GroomBoundMin.X = FMath::Min(GroomBoundMin.X, Group.Strands.BoundingBox.Min.X);
-			GroomBoundMin.Y = FMath::Min(GroomBoundMin.Y, Group.Strands.BoundingBox.Min.Y);
-			GroomBoundMin.Z = FMath::Min(GroomBoundMin.Z, Group.Strands.BoundingBox.Min.Z);
+			for (const FVector& P : Group.Strands.StrandsPoints.PointsPosition)
+			{
+				GroomBoundMin.X = FMath::Min(GroomBoundMin.X, P.X);
+				GroomBoundMin.Y = FMath::Min(GroomBoundMin.Y, P.Y);
+				GroomBoundMin.Z = FMath::Min(GroomBoundMin.Z, P.Z);
 
-			GroomBoundMax.X = FMath::Max(GroomBoundMax.X, Group.Strands.BoundingBox.Max.X);
-			GroomBoundMax.Y = FMath::Max(GroomBoundMax.Y, Group.Strands.BoundingBox.Max.Y);
-			GroomBoundMax.Z = FMath::Max(GroomBoundMax.Z, Group.Strands.BoundingBox.Max.Z);
+				GroomBoundMax.X = FMath::Max(GroomBoundMax.X, P.X);
+				GroomBoundMax.Y = FMath::Max(GroomBoundMax.Y, P.Y);
+				GroomBoundMax.Z = FMath::Max(GroomBoundMax.Z, P.Z);
+			}
 		}
 
 		Out.BoundRadius = FVector::Distance(GroomBoundMax, GroomBoundMin) * 0.5f;
@@ -1586,6 +1589,7 @@ void FGroomBuilder::BuildData(
 			}
 			OutGroupInfo.NumCurves			= OutRen.GetNumCurves();
 			OutGroupInfo.NumCurveVertices	= OutRen.GetNumPoints();
+			OutGroupInfo.MaxCurveLength		= OutRen.StrandsCurves.MaxLength;
 
 			// Sanity check
 			check(OutGroupInfo.NumCurves		<= InHairDescriptionGroup.Info.NumCurves);
@@ -1596,7 +1600,7 @@ void FGroomBuilder::BuildData(
 		// Simulation data
 		{
 			OutSim = InHairDescriptionGroup.Guides;
-			if (InHairDescriptionGroup.Info.NumCurves > 0 && !InSettings.InterpolationSettings.bOverrideGuides)
+			if (InHairDescriptionGroup.Info.NumGuides > 0 && !InSettings.InterpolationSettings.bOverrideGuides)
 			{
 				HairStrandsBuilder::BuildInternalData(OutSim, true); // Imported guides don't currently have root UVs so force computing them
 			}
@@ -1610,8 +1614,8 @@ void FGroomBuilder::BuildData(
 			OutGroupInfo.NumGuideVertices	= OutSim.GetNumPoints();
 
 			// Sanity check
-			check(OutGroupInfo.NumCurves		<= InHairDescriptionGroup.Info.NumCurves);
-			check(OutGroupInfo.NumCurveVertices <= InHairDescriptionGroup.Info.NumCurveVertices);
+			check(OutGroupInfo.NumGuides > 0);
+			check(OutGroupInfo.NumGuideVertices > 0);
 		}
 	}
 
