Remove hair group sorting by ID when building grooms as it prevents corrent mapping of previous serialized parameters on existing grooms.

diff --git a/Engine/Plugins/Runtime/HairStrands/Source/HairStrandsCore/Private/GroomBuilder.cpp b/Engine/Plugins/Runtime/HairStrands/Source/HairStrandsCore/Private/GroomBuilder.cpp
--- a/Engine/Plugins/Runtime/HairStrands/Source/HairStrandsCore/Private/GroomBuilder.cpp
+++ b/Engine/Plugins/Runtime/HairStrands/Source/HairStrandsCore/Private/GroomBuilder.cpp
@@ -1489,12 +1489,6 @@ bool FGroomBuilder::BuildHairDescriptionGroups(const FHairDescription& HairDescr
 		}
 	}
 
-	// Sort groups by ascending GroupID
-	Out.HairGroups.Sort([](const FHairDescriptionGroup& A, const FHairDescriptionGroup& B)
-	{
-		return A.Info.GroupID < B.Info.GroupID;
-	});
-
 	// Sparse->Dense Groups
 	// Change GroupID to be contiguous
 	{
