Initialize all Hair Cards/Meshes vertex factory resources for a given LOD when used.

Previously, for hair cards/meshes having deformation (i.e. skinning/simulation/global interpolation) only the current VF was initialize instead of boths. This prevents issue with hit proxy creating mesh batch with non-initialized VF resources.

diff --git a/Engine/Plugins/Runtime/HairStrands/Source/HairStrandsCore/Private/GroomManager.cpp b/Engine/Plugins/Runtime/HairStrands/Source/HairStrandsCore/Private/GroomManager.cpp
--- a/Engine/Plugins/Runtime/HairStrands/Source/HairStrandsCore/Private/GroomManager.cpp
+++ b/Engine/Plugins/Runtime/HairStrands/Source/HairStrandsCore/Private/GroomManager.cpp
@@ -43,6 +43,24 @@ DEFINE_LOG_CATEGORY_STATIC(LogGroomManager, Log, All);
 
 ///////////////////////////////////////////////////////////////////////////////////////////////////
 
+void FHairGroupInstance::FCards::FLOD::InitVertexFactory()
+{
+	VertexFactory[0]->InitResources();
+	if (DeformedResource)
+	{
+		VertexFactory[1]->InitResources();
+	}
+}
+
+void FHairGroupInstance::FMeshes::FLOD::InitVertexFactory()
+{
+	VertexFactory[0]->InitResources();
+	if (DeformedResource)
+	{
+		VertexFactory[1]->InitResources();
+	}
+}
+
 static bool IsInstanceFrustumCullingEnable()
 {
 	return GHairStrands_InterpolationFrustumCullingEnable > 0;
@@ -630,7 +648,7 @@ static void RunHairLODSelection(FRDGBuilder& GraphBuilder, const FHairStrandsIns
 			if (InstanceLOD.RaytracingResource)	InstanceLOD.RaytracingResource->Allocate(GraphBuilder);
 			#endif
 
-			InstanceLOD.GetVertexFactory()->InitResources();
+			InstanceLOD.InitVertexFactory();
 		}
 		else if (GeometryType == EHairGeometryType::Cards)
 		{
@@ -642,7 +660,7 @@ static void RunHairLODSelection(FRDGBuilder& GraphBuilder, const FHairStrandsIns
 			#if RHI_RAYTRACING
 			if (InstanceLOD.RaytracingResource)				InstanceLOD.RaytracingResource->Allocate(GraphBuilder);
 			#endif
-			InstanceLOD.GetVertexFactory()->InitResources();
+			InstanceLOD.InitVertexFactory();
 		}
 		else if (GeometryType == EHairGeometryType::Strands)
 		{
diff --git a/Engine/Plugins/Runtime/HairStrands/Source/HairStrandsCore/Public/GroomInstance.h b/Engine/Plugins/Runtime/HairStrands/Source/HairStrandsCore/Public/GroomInstance.h
--- a/Engine/Plugins/Runtime/HairStrands/Source/HairStrandsCore/Public/GroomInstance.h
+++ b/Engine/Plugins/Runtime/HairStrands/Source/HairStrandsCore/Public/GroomInstance.h
@@ -138,7 +138,8 @@ struct HAIRSTRANDSCORE_API FHairGroupInstance : public FHairStrandsInstance
 			{
 				const uint32 Index = DeformedResource ? DeformedResource->CurrentIndex : 0;
 				return VertexFactory[Index];
-			} 
+			}
+			void InitVertexFactory();
 
 		};
 		TArray<FLOD> LODs;
@@ -173,6 +174,7 @@ struct HAIRSTRANDSCORE_API FHairGroupInstance : public FHairStrandsInstance
 				const uint32 Index = DeformedResource ? DeformedResource->CurrentIndex : 0;
 				return VertexFactory[Index];
 			}
+			void InitVertexFactory();
 		};
 		TArray<FLOD> LODs;
 	} Meshes;
