Remove Niagara Hair solver dependency onto HairStrandsData.

diff --git a/Engine/Plugins/Runtime/HairStrands/Source/HairStrandsCore/Private/GroomBuilder.cpp b/Engine/Plugins/Runtime/HairStrands/Source/HairStrandsCore/Private/GroomBuilder.cpp
--- a/Engine/Plugins/Runtime/HairStrands/Source/HairStrandsCore/Private/GroomBuilder.cpp
+++ b/Engine/Plugins/Runtime/HairStrands/Source/HairStrandsCore/Private/GroomBuilder.cpp
@@ -32,7 +32,7 @@ static FAutoConsoleVariableRef CVarHairClusterBuilder_MaxVoxelResolution(TEXT("r
 FString FGroomBuilder::GetVersion()
 {
 	// Important to update the version when groom building changes
-	return TEXT("1g");
+	return TEXT("1hc");
 }
 
 namespace FHairStrandsDecimation
@@ -225,10 +225,12 @@ namespace HairStrandsBuilder
 		TArray<FHairStrandsPositionFormat::Type>& OutPackedPositions = OutBulkData.Positions;
 		TArray<FHairStrandsAttributeFormat::Type>& OutPackedAttributes = OutBulkData.Attributes;
 		TArray<FHairStrandsMaterialFormat::Type>& OutPackedMaterials = OutBulkData.Materials;
+		TArray<FHairStrandsRootIndexFormat::Type>& OutCurveOffsets = OutBulkData.CurveOffsets;
 
 		OutPackedPositions.SetNum(NumPoints * FHairStrandsPositionFormat::ComponentCount);
 		OutPackedAttributes.SetNum(NumPoints * FHairStrandsAttributeFormat::ComponentCount);
 		OutPackedMaterials.SetNum(NumPoints * FHairStrandsMaterialFormat::ComponentCount);
+		OutCurveOffsets.SetNum(NumCurves);
 
 		const FVector HairBoxCenter = HairStrands.BoundingBox.GetCenter();
 
@@ -281,6 +283,7 @@ namespace HairStrandsBuilder
 				Material.BaseColorB = FMath::Clamp(uint32(FMath::Sqrt(Points.PointsBaseColor[PointIndex + IndexOffset].B) * 0xFF), 0u, 0xFFu);
 				Material.Roughness  = FMath::Clamp(uint32(Points.PointsRoughness[PointIndex + IndexOffset] * 0xFF), 0u, 0xFFu);
 			}
+			OutCurveOffsets[CurveIndex] = IndexOffset;
 		}
 
 		OutBulkData.BoundingBox = HairStrands.BoundingBox;
diff --git a/Engine/Plugins/Runtime/HairStrands/Source/HairStrandsCore/Private/GroomComponent.cpp b/Engine/Plugins/Runtime/HairStrands/Source/HairStrandsCore/Private/GroomComponent.cpp
--- a/Engine/Plugins/Runtime/HairStrands/Source/HairStrandsCore/Private/GroomComponent.cpp
+++ b/Engine/Plugins/Runtime/HairStrands/Source/HairStrandsCore/Private/GroomComponent.cpp
@@ -1533,20 +1533,6 @@ UMaterialInterface* UGroomComponent::GetMaterial(int32 ElementIndex) const
 	return GetMaterial(ElementIndex, GeometryType, true);
 }
 
-FHairStrandsDatas* UGroomComponent::GetGuideStrandsDatas(uint32 GroupIndex)
-{
-	if (!GroomAsset || GroupIndex >= uint32(GroomAsset->GetNumHairGroups()))
-	{
-		return nullptr;
-	}
-
-	if (!GroomAsset->HairGroupsData[GroupIndex].Guides.HasValidData())
-	{
-		return nullptr;
-	}
-	return &GroomAsset->HairGroupsData[GroupIndex].Guides.Data;
-}
-
 FHairStrandsRestResource* UGroomComponent::GetGuideStrandsRestResource(uint32 GroupIndex)
 {
 	if (!GroomAsset || GroupIndex >= uint32(GroomAsset->GetNumHairGroups()))
diff --git a/Engine/Plugins/Runtime/HairStrands/Source/HairStrandsCore/Private/HairStrandsDatas.cpp b/Engine/Plugins/Runtime/HairStrands/Source/HairStrandsCore/Private/HairStrandsDatas.cpp
--- a/Engine/Plugins/Runtime/HairStrands/Source/HairStrandsCore/Private/HairStrandsDatas.cpp
+++ b/Engine/Plugins/Runtime/HairStrands/Source/HairStrandsCore/Private/HairStrandsDatas.cpp
@@ -241,30 +241,8 @@ void FHairStrandsBulkData::Serialize(FArchive& Ar)
 
 	Ar << Positions;
 	Ar << Attributes;
-
-	if (Ar.CustomVer(FUE5ReleaseStreamObjectVersion::GUID) < FUE5ReleaseStreamObjectVersion::HairStrandsVertexFormatChange && Ar.IsLoading())
-	{
-		check(Positions.Num() == Attributes.Num());
-		const uint32 VertexCount = Positions.Num();
-		for (uint32 VertexIt = 0; VertexIt < VertexCount; ++VertexIt)
-		{
-			// Prior to FUE5ReleaseStreamObjectVersion::HairStrandsVertexFormatChange, UCoord and NormalizedLength were swapped
-			const uint8 UCoord = Attributes[VertexIt].NormalizedLength;
-			const uint8 NormalizedLength = Positions[VertexIt].UCoord;
-			Attributes[VertexIt].NormalizedLength = NormalizedLength;
-			Positions[VertexIt].UCoord = UCoord;
-		}
-	}
-
-	if (Ar.CustomVer(FReleaseObjectVersion::GUID) >= FReleaseObjectVersion::GroomAssetVersion2)
-	{
-		Ar << Materials;
-	}
-	else
-	{
-		const uint32 ElementCount = Attributes.Num();
-		Materials.InsertZeroed(0, ElementCount);
-	}
+	Ar << Materials;
+	Ar << CurveOffsets;
 }
 
 void FHairStrandsDatas::Serialize(FArchive& Ar, FHairStrandsBulkData& BulkData)
diff --git a/Engine/Plugins/Runtime/HairStrands/Source/HairStrandsCore/Private/Niagara/NiagaraDataInterfaceHairStrands.cpp b/Engine/Plugins/Runtime/HairStrands/Source/HairStrandsCore/Private/Niagara/NiagaraDataInterfaceHairStrands.cpp
--- a/Engine/Plugins/Runtime/HairStrands/Source/HairStrandsCore/Private/Niagara/NiagaraDataInterfaceHairStrands.cpp
+++ b/Engine/Plugins/Runtime/HairStrands/Source/HairStrandsCore/Private/Niagara/NiagaraDataInterfaceHairStrands.cpp
@@ -290,14 +290,12 @@ struct FNDIHairStrandsParametersName
 //------------------------------------------------------------------------------------------------------------
 
 void FNDIHairStrandsBuffer::Initialize(
-	const FHairStrandsDatas*  HairStrandsDatas,
 	const FHairStrandsRestResource* HairStrandsRestResource,
 	const FHairStrandsDeformedResource*  HairStrandsDeformedResource,
 	const FHairStrandsRestRootResource* HairStrandsRestRootResource,
 	const FHairStrandsDeformedRootResource* HairStrandsDeformedRootResource,
 	const TStaticArray<float, 32 * NumScales>& InParamsScale)
 {
-	SourceDatas = HairStrandsDatas;
 	SourceRestResources = HairStrandsRestResource;
 	SourceDeformedResources = HairStrandsDeformedResource;
 	SourceRestRootResources = HairStrandsRestRootResource;
@@ -307,13 +305,11 @@ void FNDIHairStrandsBuffer::Initialize(
 }
 
 void FNDIHairStrandsBuffer::Update(
-	const FHairStrandsDatas* HairStrandsDatas,
 	const FHairStrandsRestResource* HairStrandsRestResource,
 	const FHairStrandsDeformedResource* HairStrandsDeformedResource,
 	const FHairStrandsRestRootResource* HairStrandsRestRootResource,
 	const FHairStrandsDeformedRootResource* HairStrandsDeformedRootResource)
 {
-	SourceDatas = HairStrandsDatas;
 	SourceRestResources = HairStrandsRestResource;
 	SourceDeformedResources = HairStrandsDeformedResource;
 	SourceRestRootResources = HairStrandsRestRootResource;
@@ -322,7 +318,7 @@ void FNDIHairStrandsBuffer::Update(
 
 void FNDIHairStrandsBuffer::Transfer(const TStaticArray<float, 32 * NumScales>& InParamsScale)
 {
-	if (SourceDatas != nullptr && SourceRestResources != nullptr)
+	if (SourceRestResources != nullptr)
 	{
 		const uint32 ScaleCount = 32 * NumScales;
 		const uint32 ScaleBytes = sizeof(float) * ScaleCount;
@@ -335,8 +331,9 @@ void FNDIHairStrandsBuffer::Transfer(const TStaticArray<float, 32 * NumScales>&
 
 void FNDIHairStrandsBuffer::InitRHI()
 {
-	if (SourceDatas != nullptr && SourceRestResources != nullptr)
+	if (SourceRestResources != nullptr)
 	{
+		const FHairStrandsBulkData* SourceDatas = &SourceRestResources->BulkData;
 		{
 			const uint32 OffsetCount = SourceDatas->GetNumCurves() + 1;
 			const uint32 OffsetBytes = sizeof(uint32)*OffsetCount;
@@ -344,7 +341,7 @@ void FNDIHairStrandsBuffer::InitRHI()
 			CurvesOffsetsBuffer.Initialize(TEXT("CurvesOffsetsBuffer"), sizeof(uint32), OffsetCount, EPixelFormat::PF_R32_UINT, BUF_Static);
 			void* OffsetBufferData = RHILockBuffer(CurvesOffsetsBuffer.Buffer, 0, OffsetBytes, RLM_WriteOnly);
 
-			FMemory::Memcpy(OffsetBufferData, SourceDatas->StrandsCurves.CurvesOffset.GetData(), OffsetBytes);
+			FMemory::Memcpy(OffsetBufferData, SourceDatas->CurveOffsets.GetData(), OffsetBytes);
 			RHIUnlockBuffer(CurvesOffsetsBuffer.Buffer);
 		}
 		{
@@ -419,7 +416,7 @@ void FNDIHairStrandsData::Release()
 	}
 }
 
-void FNDIHairStrandsData::Update(UNiagaraDataInterfaceHairStrands* Interface, FNiagaraSystemInstance* SystemInstance, const FHairStrandsDatas* StrandsDatas,
+void FNDIHairStrandsData::Update(UNiagaraDataInterfaceHairStrands* Interface, FNiagaraSystemInstance* SystemInstance, const FHairStrandsBulkData* StrandsDatas,
 	UGroomAsset* GroomAsset, const int32 GroupIndex, const int32 LODIndex, const FTransform& LocalToWorld)
 {
 	if (Interface != nullptr)
@@ -493,7 +490,6 @@ bool FNDIHairStrandsData::Init(UNiagaraDataInterfaceHairStrands* Interface, FNia
 
 	if (Interface != nullptr)
 	{
-		FHairStrandsDatas* StrandsDatas = nullptr;
 		FHairStrandsRestResource* StrandsRestResource = nullptr;
 		FHairStrandsDeformedResource* StrandsDeformedResource = nullptr;
 		FHairStrandsRestRootResource* StrandsRestRootResource = nullptr;
@@ -504,11 +500,11 @@ bool FNDIHairStrandsData::Init(UNiagaraDataInterfaceHairStrands* Interface, FNia
 
 		{
 			FTransform LocalToWorld = FTransform::Identity;
-			Interface->ExtractDatasAndResources(SystemInstance, StrandsDatas, StrandsRestResource, StrandsDeformedResource, StrandsRestRootResource, StrandsDeformedRootResource, GroomAsset, GroupIndex, LODIndex, LocalToWorld);
-			Update(Interface, SystemInstance, StrandsDatas, GroomAsset, GroupIndex, LODIndex, LocalToWorld);
+			Interface->ExtractDatasAndResources(SystemInstance, StrandsRestResource, StrandsDeformedResource, StrandsRestRootResource, StrandsDeformedRootResource, GroomAsset, GroupIndex, LODIndex, LocalToWorld);
+			Update(Interface, SystemInstance, StrandsRestResource ? &StrandsRestResource->BulkData : nullptr, GroomAsset, GroupIndex, LODIndex, LocalToWorld);
 
 			HairStrandsBuffer = new FNDIHairStrandsBuffer();
-			HairStrandsBuffer->Initialize(StrandsDatas, StrandsRestResource, StrandsDeformedResource, StrandsRestRootResource, StrandsDeformedRootResource, ParamsScale);
+			HairStrandsBuffer->Initialize(StrandsRestResource, StrandsDeformedResource, StrandsRestRootResource, StrandsDeformedRootResource, ParamsScale);
 
 			BeginInitResource(HairStrandsBuffer);
 
@@ -928,7 +924,6 @@ void UNiagaraDataInterfaceHairStrands::ExtractSourceComponent(FNiagaraSystemInst
 
 void UNiagaraDataInterfaceHairStrands::ExtractDatasAndResources(
 	FNiagaraSystemInstance* SystemInstance,
-	FHairStrandsDatas*& OutStrandsDatas,
 	FHairStrandsRestResource*& OutStrandsRestResource,
 	FHairStrandsDeformedResource*& OutStrandsDeformedResource,
 	FHairStrandsRestRootResource*& OutStrandsRestRootResource,
@@ -940,7 +935,6 @@ void UNiagaraDataInterfaceHairStrands::ExtractDatasAndResources(
 {
 	ExtractSourceComponent(SystemInstance);
 
-	OutStrandsDatas = nullptr;
 	OutStrandsRestResource = nullptr;
 	OutStrandsDeformedResource = nullptr;
 	OutStrandsRestRootResource = nullptr;
@@ -966,7 +960,6 @@ void UNiagaraDataInterfaceHairStrands::ExtractDatasAndResources(
 		}
 		if (OutGroupIndex >= 0 && OutGroupIndex < SourceComponent->NiagaraComponents.Num())
 		{
-			OutStrandsDatas = SourceComponent->GetGuideStrandsDatas(OutGroupIndex);
 			OutStrandsRestResource = SourceComponent->GetGuideStrandsRestResource(OutGroupIndex);
 			OutStrandsDeformedResource = SourceComponent->GetGuideStrandsDeformedResource(OutGroupIndex);
 			OutStrandsRestRootResource = SourceComponent->GetGuideStrandsRestRootResource(OutGroupIndex);
@@ -983,7 +976,6 @@ void UNiagaraDataInterfaceHairStrands::ExtractDatasAndResources(
 		OutLocalToWorld = FTransform::Identity;
 		if (OutGroupIndex < DefaultSource->GetNumHairGroups())
 		{
-			OutStrandsDatas = &DefaultSource->HairGroupsData[OutGroupIndex].Strands.Data;
 			OutStrandsRestResource = DefaultSource->HairGroupsData[OutGroupIndex].Strands.RestResource;
 			OutGroomAsset = DefaultSource;
 		}
@@ -1029,7 +1021,6 @@ bool UNiagaraDataInterfaceHairStrands::PerInstanceTick(void* PerInstanceData, FN
 {
 	FNDIHairStrandsData* InstanceData = static_cast<FNDIHairStrandsData*>(PerInstanceData);
 
-	FHairStrandsDatas* StrandsDatas = nullptr;
 	FHairStrandsRestResource* StrandsRestResource = nullptr;
 	FHairStrandsDeformedResource* StrandsDeformedResource = nullptr;
 	FHairStrandsRestRootResource* StrandsRestRootResource = nullptr;
@@ -1041,8 +1032,8 @@ bool UNiagaraDataInterfaceHairStrands::PerInstanceTick(void* PerInstanceData, FN
 	InstanceData->TickCount = FMath::Min(GHairSimulationMaxDelay + 1, InstanceData->TickCount + 1);
 
 	FTransform LocalToWorld = FTransform::Identity;
-	ExtractDatasAndResources(SystemInstance, StrandsDatas, StrandsRestResource, StrandsDeformedResource, StrandsRestRootResource, StrandsDeformedRootResource, GroomAsset, GroupIndex, LODIndex, LocalToWorld);
-	InstanceData->HairStrandsBuffer->Update(StrandsDatas, StrandsRestResource, StrandsDeformedResource, StrandsRestRootResource, StrandsDeformedRootResource);
+	ExtractDatasAndResources(SystemInstance, StrandsRestResource, StrandsDeformedResource, StrandsRestRootResource, StrandsDeformedRootResource, GroomAsset, GroupIndex, LODIndex, LocalToWorld);
+	InstanceData->HairStrandsBuffer->Update(StrandsRestResource, StrandsDeformedResource, StrandsRestRootResource, StrandsDeformedRootResource);
 
 	if (SourceComponent != nullptr)
 	{
@@ -1053,7 +1044,7 @@ bool UNiagaraDataInterfaceHairStrands::PerInstanceTick(void* PerInstanceData, FN
 		}
 		InstanceData->ForceReset = bRequiresReset;
 	}
-	InstanceData->Update(this, SystemInstance, StrandsDatas, GroomAsset, GroupIndex, LODIndex, LocalToWorld);
+	InstanceData->Update(this, SystemInstance, StrandsRestResource ? &StrandsRestResource->BulkData : nullptr, GroomAsset, GroupIndex, LODIndex, LocalToWorld);
 	return false;
 }
 
diff --git a/Engine/Plugins/Runtime/HairStrands/Source/HairStrandsCore/Public/GroomComponent.h b/Engine/Plugins/Runtime/HairStrands/Source/HairStrandsCore/Public/GroomComponent.h
--- a/Engine/Plugins/Runtime/HairStrands/Source/HairStrandsCore/Public/GroomComponent.h
+++ b/Engine/Plugins/Runtime/HairStrands/Source/HairStrandsCore/Public/GroomComponent.h
@@ -112,9 +112,6 @@ public:
 	virtual void PostLoad() override;
 	//~ End UMeshComponent Interface.
 
-	/** Return the guide hairs datas */
-	FHairStrandsDatas* GetGuideStrandsDatas(uint32 GroupIndex);
-
 	/** Return the guide hairs rest resources*/
 	FHairStrandsRestResource* GetGuideStrandsRestResource(uint32 GroupIndex);
 
diff --git a/Engine/Plugins/Runtime/HairStrands/Source/HairStrandsCore/Public/HairStrandsDatas.h b/Engine/Plugins/Runtime/HairStrands/Source/HairStrandsCore/Public/HairStrandsDatas.h
--- a/Engine/Plugins/Runtime/HairStrands/Source/HairStrandsCore/Public/HairStrandsDatas.h
+++ b/Engine/Plugins/Runtime/HairStrands/Source/HairStrandsCore/Public/HairStrandsDatas.h
@@ -417,10 +417,10 @@ struct HAIRSTRANDSCORE_API FHairStrandsBulkData
 	float MaxRadius = 0;
 	FBox BoundingBox;
 
-	TArray<FHairStrandsPositionFormat::Type>	Positions;
-	TArray<FHairStrandsAttributeFormat::Type>	Attributes;
-	TArray<FHairStrandsMaterialFormat::Type>	Materials;
-	TArray<FHairStrandsRootIndexFormat::Type>	RootIndices; // Vertex to strands Index
+	TArray<FHairStrandsPositionFormat::Type>	Positions;		// Size = PointCount
+	TArray<FHairStrandsAttributeFormat::Type>	Attributes;		// Size = PointCount
+	TArray<FHairStrandsMaterialFormat::Type>	Materials;		// Size = PointCount
+	TArray<FHairStrandsRootIndexFormat::Type>	CurveOffsets;	// Size = CurveCount  - Store the root point index for the curve
 };
 
 /** Hair strands debug data */
diff --git a/Engine/Plugins/Runtime/HairStrands/Source/HairStrandsCore/Public/Niagara/NiagaraDataInterfaceHairStrands.h b/Engine/Plugins/Runtime/HairStrands/Source/HairStrandsCore/Public/Niagara/NiagaraDataInterfaceHairStrands.h
--- a/Engine/Plugins/Runtime/HairStrands/Source/HairStrandsCore/Public/Niagara/NiagaraDataInterfaceHairStrands.h
+++ b/Engine/Plugins/Runtime/HairStrands/Source/HairStrandsCore/Public/Niagara/NiagaraDataInterfaceHairStrands.h
@@ -23,7 +23,7 @@ struct FNDIHairStrandsData;
 struct FNDIHairStrandsBuffer : public FRenderResource
 {
 	/** Set the asset that will be used to affect the buffer */
-	void Initialize(const FHairStrandsDatas*  HairStrandsDatas, 
+	void Initialize(
 		const FHairStrandsRestResource*  HairStrandsRestResource, 
 		const FHairStrandsDeformedResource*  HairStrandsDeformedResource, 
 		const FHairStrandsRestRootResource* HairStrandsRestRootResource, 
@@ -31,7 +31,7 @@ struct FNDIHairStrandsBuffer : public FRenderResource
 		const TStaticArray<float, 32 * NumScales>& InParamsScale);  
 
 	/** Set the asset that will be used to affect the buffer */
-	void Update(const FHairStrandsDatas* HairStrandsDatas,
+	void Update(
 		const FHairStrandsRestResource* HairStrandsRestResource,
 		const FHairStrandsDeformedResource* HairStrandsDeformedResource,
 		const FHairStrandsRestRootResource* HairStrandsRestRootResource,
@@ -61,9 +61,6 @@ struct FNDIHairStrandsBuffer : public FRenderResource
 	/** Params scale buffer */
 	FRWBuffer ParamsScaleBuffer;
 
-	/** The strand asset datas from which to sample */
-	const FHairStrandsDatas* SourceDatas;
-
 	/** The strand asset resource from which to sample */
 	const FHairStrandsRestResource* SourceRestResources;
 
@@ -97,7 +94,7 @@ struct FNDIHairStrandsData
 	void Release();
 
 	/** Update the buffers */
-	void Update(UNiagaraDataInterfaceHairStrands* Interface, FNiagaraSystemInstance* SystemInstance, const FHairStrandsDatas* HairStrandsDatas, UGroomAsset* GroomAsset, const int32 GroupIndex, const int32 LODIndex, const FTransform& LocalToWorld);
+	void Update(UNiagaraDataInterfaceHairStrands* Interface, FNiagaraSystemInstance* SystemInstance, const FHairStrandsBulkData* HairStrandsDatas, UGroomAsset* GroomAsset, const int32 GroupIndex, const int32 LODIndex, const FTransform& LocalToWorld);
 
 	inline void ResetDatas()
 	{
@@ -365,7 +362,6 @@ public:
 	/** Extract datas and resources */
 	void ExtractDatasAndResources(
 		FNiagaraSystemInstance* SystemInstance, 
-		FHairStrandsDatas*& OutStrandsDatas,
 		FHairStrandsRestResource*& OutStrandsRestResource, 
 		FHairStrandsDeformedResource*& OutStrandsDeformedResource, 
 		FHairStrandsRestRootResource*& OutStrandsRestRootResource, 
