Remove interpolation data from serialization.

diff --git a/Engine/Plugins/Runtime/HairStrands/Source/HairStrandsCore/Private/GroomAsset.cpp b/Engine/Plugins/Runtime/HairStrands/Source/HairStrandsCore/Private/GroomAsset.cpp
--- a/Engine/Plugins/Runtime/HairStrands/Source/HairStrandsCore/Private/GroomAsset.cpp
+++ b/Engine/Plugins/Runtime/HairStrands/Source/HairStrandsCore/Private/GroomAsset.cpp
@@ -211,8 +211,9 @@ static bool BuildHairGroup(
 		FGroomBuilder::BuildBulkData(HairGroup.Info, OutHairGroupsData[GroupIndex].Guides.Data, OutHairGroupsData[GroupIndex].Guides.BulkData);
 		FGroomBuilder::BuildBulkData(HairGroup.Info, OutHairGroupsData[GroupIndex].Strands.Data, OutHairGroupsData[GroupIndex].Strands.BulkData);
 
-		FGroomBuilder::BuildInterplationData(HairGroup.Info, OutHairGroupsData[GroupIndex].Strands.Data, OutHairGroupsData[GroupIndex].Guides.Data, HairGroupsInterpolation[GroupIndex].InterpolationSettings, OutHairGroupsData[GroupIndex].Strands.InterpolationData);
-		FGroomBuilder::BuildInterplationBulkData(OutHairGroupsData[GroupIndex].Guides.Data, OutHairGroupsData[GroupIndex].Strands.InterpolationData, OutHairGroupsData[GroupIndex].Strands.InterpolationBulkData);
+		FHairStrandsInterpolationDatas InterpolationData;
+		FGroomBuilder::BuildInterplationData(HairGroup.Info, OutHairGroupsData[GroupIndex].Strands.Data, OutHairGroupsData[GroupIndex].Guides.Data, HairGroupsInterpolation[GroupIndex].InterpolationSettings, InterpolationData);
+		FGroomBuilder::BuildInterplationBulkData(OutHairGroupsData[GroupIndex].Guides.Data, InterpolationData, OutHairGroupsData[GroupIndex].Strands.InterpolationBulkData);
 
 		FGroomBuilder::BuildClusterData(OutHairGroupsData[GroupIndex].Strands.Data, HairDescriptionGroups.BoundRadius, HairGroupsLOD[GroupIndex], OutHairGroupsData[GroupIndex].Strands.ClusterCullingData);
 	}
@@ -1094,10 +1095,10 @@ FArchive& operator<<(FArchive& Ar, FHairGroupData::FCards::FLOD& CardLODData)
 	if (!Ar.IsCooking() || !CardLODData.bIsCookedOut)
 	{
 		CardLODData.Data.Serialize(Ar, CardLODData.BulkData);
-		CardLODData.InterpolationData.Serialize(Ar, CardLODData.InterpolationBulkData);
+		CardLODData.InterpolationBulkData.Serialize(Ar);
 
 		CardLODData.Guides.Data.Serialize(Ar, CardLODData.Guides.BulkData);
-		CardLODData.Guides.InterpolationData.Serialize(Ar, CardLODData.Guides.InterpolationBulkData);
+		CardLODData.Guides.InterpolationBulkData.Serialize(Ar);
 	}
 	else
 	{
@@ -1106,13 +1107,12 @@ FArchive& operator<<(FArchive& Ar, FHairGroupData::FCards::FLOD& CardLODData)
 		FHairCardsBulkData NoCardsBulkData;
 		NoCardsData.Serialize(Ar, NoCardsBulkData);
 
-		FHairCardsInterpolationDatas NoInterpolationData;
 		FHairCardsInterpolationBulkData NoInterpolationBulkData;
-		NoInterpolationData.Serialize(Ar, NoInterpolationBulkData);
+		NoInterpolationBulkData.Serialize(Ar);
 
 		FHairGroupData::FBaseWithInterpolation NoGuideData;
 		NoGuideData.Data.Serialize(Ar, NoGuideData.BulkData);
-		NoGuideData.InterpolationData.Serialize(Ar, NoGuideData.InterpolationBulkData);
+		NoGuideData.InterpolationBulkData.Serialize(Ar);
 	}
 
 	return Ar;
@@ -1144,7 +1144,7 @@ FArchive& operator<<(FArchive& Ar, FHairGroupData& GroupData)
 	{
 		GroupData.Strands.Data.Serialize(Ar, GroupData.Strands.BulkData);
 		GroupData.Guides.Data.Serialize(Ar, GroupData.Guides.BulkData);
-		GroupData.Strands.InterpolationData.Serialize(Ar, GroupData.Strands.InterpolationBulkData);
+		GroupData.Strands.InterpolationBulkData.Serialize(Ar);
 
 	}
 	else
@@ -1154,7 +1154,7 @@ FArchive& operator<<(FArchive& Ar, FHairGroupData& GroupData)
 		// simulation or RBF deformation) on the target platform
 		NoStrandsData.Strands.Data.Serialize(Ar, NoStrandsData.Strands.BulkData);
 		GroupData.Guides.Data.Serialize(Ar, GroupData.Guides.BulkData);
-		NoStrandsData.Strands.InterpolationData.Serialize(Ar, NoStrandsData.Strands.InterpolationBulkData);
+		NoStrandsData.Strands.InterpolationBulkData.Serialize(Ar);
 	}
 
 	if (Ar.CustomVer(FAnimObjectVersion::GUID) >= FAnimObjectVersion::SerializeHairClusterCullingData)
@@ -1332,7 +1332,7 @@ void UGroomAsset::SetHairWidth(float Width)
 // differences, etc.) replace the version GUID below with a new one.
 // In case of merge conflicts with DDC versions, you MUST generate a new GUID
 // and set this new GUID as the version.
-#define GROOM_DERIVED_DATA_VERSION TEXT("04296541C35C4CE0999D08B6170BB67A")
+#define GROOM_DERIVED_DATA_VERSION TEXT("FCAFECFFD02A4509A0D86775505E97F2")
 
 #if WITH_EDITORONLY_DATA
 
@@ -2032,7 +2032,7 @@ bool UGroomAsset::BuildCardsGeometry(uint32 GroupIndex)
 			if (CardsMesh != nullptr)
 			{
 				CardsMesh->ConditionalPostLoad();
-				bInitResources = FHairCardsBuilder::ImportGeometry(CardsMesh, LOD.Data, LOD.BulkData, LOD.Guides.Data, LOD.InterpolationData, LOD.InterpolationBulkData);
+				bInitResources = FHairCardsBuilder::ImportGeometry(CardsMesh, LOD.Data, LOD.BulkData, LOD.Guides.Data, LOD.InterpolationBulkData);
 				if (!bInitResources)
 				{
 					UE_LOG(LogHairStrands, Warning, TEXT("Failed to import cards from %s for Group %d LOD %d."), *CardsMesh->GetName(), GroupIndex, LODIt);
@@ -2091,8 +2091,9 @@ bool UGroomAsset::BuildCardsGeometry(uint32 GroupIndex)
 				FGroomBuilder::BuildData(LOD.Guides.Data);
 				FGroomBuilder::BuildBulkData(DummyGroupInfo, LOD.Guides.Data, LOD.Guides.BulkData);
 
-				FGroomBuilder::BuildInterplationData(DummyGroupInfo, LOD.Guides.Data, GroupData.Guides.Data, CardsInterpolationSettings, LOD.Guides.InterpolationData);
-				FGroomBuilder::BuildInterplationBulkData(GroupData.Guides.Data, LOD.Guides.InterpolationData, LOD.Guides.InterpolationBulkData);
+				FHairStrandsInterpolationDatas LODGuideInterpolationData;
+				FGroomBuilder::BuildInterplationData(DummyGroupInfo, LOD.Guides.Data, GroupData.Guides.Data, CardsInterpolationSettings, LODGuideInterpolationData);
+				FGroomBuilder::BuildInterplationBulkData(GroupData.Guides.Data, LODGuideInterpolationData, LOD.Guides.InterpolationBulkData);
 
 				if (bCopyRenderData)
 				{
@@ -2839,7 +2840,6 @@ bool UGroomAsset::IsMaterialUsed(int32 MaterialIndex) const
 #if WITH_EDITOR
 struct FHairProceduralCardsQuery
 {
-	FHairCardsInterpolationDatas InterpolationData;
 	FHairCardsInterpolationBulkData InterpolationBulkData;
 	FHairCardsProceduralDatas ProceduralData;
 	FHairStrandsDatas GuideData;
@@ -2909,7 +2909,6 @@ void UGroomAsset::SaveProceduralCards(uint32 DescIndex)
 		Desc->ProceduralSettings,
 		Q.ProceduralData,
 		Q.GuideData,
-		Q.InterpolationData,
 		Q.InterpolationBulkData,
 		Desc->Textures);
 	Q.Textures = &Desc->Textures;
diff --git a/Engine/Plugins/Runtime/HairStrands/Source/HairStrandsCore/Private/GroomBuilder.cpp b/Engine/Plugins/Runtime/HairStrands/Source/HairStrandsCore/Private/GroomBuilder.cpp
--- a/Engine/Plugins/Runtime/HairStrands/Source/HairStrandsCore/Private/GroomBuilder.cpp
+++ b/Engine/Plugins/Runtime/HairStrands/Source/HairStrandsCore/Private/GroomBuilder.cpp
@@ -15,6 +15,7 @@
 DEFINE_LOG_CATEGORY_STATIC(LogGroomBuilder, Log, All);
 
 #define LOCTEXT_NAMESPACE "GroomBuilder"
+#define ENABLE_EXTERNAL_INTERPOLATION_DATA 0
 
 // For debug purpose
 static float GHairInterpolationMetric_Distance = 1;
@@ -1175,6 +1176,7 @@ namespace HairInterpolationBuilder
 	}
 
 	/** Fill the GroomAsset with the interpolation data that exists in the HairDescription */
+	#if ENABLE_EXTERNAL_INTERPOLATION_DATA
 	void FillInterpolationData(UGroomAsset* GroomAsset, const FHairDescription& HairDescription)
 	{
 		TRACE_CPUPROFILER_EVENT_SCOPE(HairInterpolationBuilder::FillInterpolationData);
@@ -1264,6 +1266,7 @@ namespace HairInterpolationBuilder
 			SimulationData.StrandsCurves.StrandIDs.Empty();
 		}
 	}
+	#endif // ENABLE_EXTERNAL_INTERPOLATION_DATA
 }
 
 class FGroomDataRandomizer
@@ -1618,10 +1621,12 @@ void FGroomBuilder::BuildData(
 
 	// If there's usable closest guides and guide weights attributes, fill them into the asset
 	// This step requires the HairSimulationData (guides) to be filled prior to this
-	//if (InHairDescriptionGroup.bCanUseClosestGuidesAndWeights)
-	//{
-	//	HairInterpolationBuilder::FillInterpolationData(GroomAsset, HairDescription);
-	//}
+#if ENABLE_EXTERNAL_INTERPOLATION_DATA
+	if (InHairDescriptionGroup.bCanUseClosestGuidesAndWeights)
+	{
+		HairInterpolationBuilder::FillInterpolationData(GroomAsset, HairDescription);
+	}
+#endif
 }
 
 void FGroomBuilder::BuildBulkData(
diff --git a/Engine/Plugins/Runtime/HairStrands/Source/HairStrandsCore/Private/GroomRBFDeformer.cpp b/Engine/Plugins/Runtime/HairStrands/Source/HairStrandsCore/Private/GroomRBFDeformer.cpp
--- a/Engine/Plugins/Runtime/HairStrands/Source/HairStrandsCore/Private/GroomRBFDeformer.cpp
+++ b/Engine/Plugins/Runtime/HairStrands/Source/HairStrandsCore/Private/GroomRBFDeformer.cpp
@@ -401,8 +401,9 @@ static void ApplyDeformationToGroom(const TArray<FRBFDeformedPositions>& Deforme
 				FGroomBuilder::BuildBulkData(HairGroup.Info, HairGroupsData.Guides.Data,  HairGroupsData.Guides.BulkData);
 				FGroomBuilder::BuildBulkData(HairGroup.Info, HairGroupsData.Strands.Data, HairGroupsData.Strands.BulkData);
 
-				FGroomBuilder::BuildInterplationData(HairGroup.Info, HairGroupsData.Strands.Data, HairGroupsData.Guides.Data, InterpolationSettings, HairGroupsData.Strands.InterpolationData);
-				FGroomBuilder::BuildInterplationBulkData(HairGroupsData.Guides.Data, HairGroupsData.Strands.InterpolationData, HairGroupsData.Strands.InterpolationBulkData);
+				FHairStrandsInterpolationDatas StrandsInterpolationData;
+				FGroomBuilder::BuildInterplationData(HairGroup.Info, HairGroupsData.Strands.Data, HairGroupsData.Guides.Data, InterpolationSettings, StrandsInterpolationData);
+				FGroomBuilder::BuildInterplationBulkData(HairGroupsData.Guides.Data, StrandsInterpolationData, HairGroupsData.Strands.InterpolationBulkData);
 
 				FGroomBuilder::BuildClusterData(HairGroupsData.Strands.Data, HairDescriptionGroups.BoundRadius, HairGroupLOD, HairGroupsData.Strands.ClusterCullingData);
 			}
diff --git a/Engine/Plugins/Runtime/HairStrands/Source/HairStrandsCore/Private/GroomResources.cpp b/Engine/Plugins/Runtime/HairStrands/Source/HairStrandsCore/Private/GroomResources.cpp
--- a/Engine/Plugins/Runtime/HairStrands/Source/HairStrandsCore/Private/GroomResources.cpp
+++ b/Engine/Plugins/Runtime/HairStrands/Source/HairStrandsCore/Private/GroomResources.cpp
@@ -1043,14 +1043,6 @@ void FHairCardsInterpolationDatas::Reset()
 	PointsSimCurvesVertexLerp.Empty();
 }
 
-void FHairCardsInterpolationDatas::Serialize(FArchive& Ar, FHairCardsInterpolationBulkData& BulkData)
-{
-	Ar << PointsSimCurvesIndex;
-	Ar << PointsSimCurvesVertexIndex;
-	Ar << PointsSimCurvesVertexLerp;
-	BulkData.Serialize(Ar);
-}
-
 void FHairCardsInterpolationBulkData::Serialize(FArchive& Ar)
 {
 	Ar << Interpolation;
diff --git a/Engine/Plugins/Runtime/HairStrands/Source/HairStrandsCore/Private/HairCardsBuilder.cpp b/Engine/Plugins/Runtime/HairStrands/Source/HairStrandsCore/Private/HairCardsBuilder.cpp
--- a/Engine/Plugins/Runtime/HairStrands/Source/HairStrandsCore/Private/HairCardsBuilder.cpp
+++ b/Engine/Plugins/Runtime/HairStrands/Source/HairStrandsCore/Private/HairCardsBuilder.cpp
@@ -3382,7 +3382,6 @@ void BuildGeometry(
 	const FHairGroupsProceduralCards& Settings,
 	FHairCardsProceduralDatas& Out,
 	FHairStrandsDatas& OutGuides,
-	FHairCardsInterpolationDatas& OutInterpolation,
 	FHairCardsInterpolationBulkData& OutInterpolationBulk,
 	FHairGroupCardsTextures& OutTextures)
 {
@@ -3411,6 +3410,7 @@ void BuildGeometry(
 	//
 	// 
 	// TODO: add some caching about the different par so that we don't reprocess the entire group when something change
+	FHairCardsInterpolationDatas OutInterpolation;
 	HairCards::FVoxelVolume Voxels;
 	{
 		// Build voxel structure
@@ -3566,7 +3566,6 @@ bool ImportGeometry(
 	FHairCardsDatas& Out,
 	FHairCardsBulkData& OutBulk,
 	FHairStrandsDatas& OutGuides,
-	FHairCardsInterpolationDatas& OutInterpolationData,
 	FHairCardsInterpolationBulkData& OutInterpolationBulkData)
 {
 	const uint32 MeshLODIndex = 0;
@@ -3716,14 +3715,15 @@ bool ImportGeometry(
 	bool bSuccess = HairCards::CreateCardsGuides(Out.Cards, OutGuides, CardLengths);
 	if (bSuccess)
 	{
-		HairCards::CreateCardsInterpolation(Out.Cards, OutGuides, OutInterpolationData, CardLengths);
+		FHairCardsInterpolationDatas InterpolationData;
+		HairCards::CreateCardsInterpolation(Out.Cards, OutGuides, InterpolationData, CardLengths);
 
 		// Fill out the interpolation data
 		OutInterpolationBulkData.Interpolation.SetNum(PointCount);
 		for (uint32 PointIt = 0; PointIt < PointCount; ++PointIt)
 		{
-			const uint32 InterpVertexIndex = OutInterpolationData.PointsSimCurvesVertexIndex[PointIt];
-			const float VertexLerp = OutInterpolationData.PointsSimCurvesVertexLerp[PointIt];
+			const uint32 InterpVertexIndex = InterpolationData.PointsSimCurvesVertexIndex[PointIt];
+			const float VertexLerp = InterpolationData.PointsSimCurvesVertexLerp[PointIt];
 			FHairCardsInterpolationVertex PackedData;
 			PackedData.VertexIndex = InterpVertexIndex;
 			PackedData.VertexLerp = FMath::Clamp(uint32(VertexLerp * 0xFF), 0u, 0xFFu);
diff --git a/Engine/Plugins/Runtime/HairStrands/Source/HairStrandsCore/Private/HairStrandsDatas.cpp b/Engine/Plugins/Runtime/HairStrands/Source/HairStrandsCore/Private/HairStrandsDatas.cpp
--- a/Engine/Plugins/Runtime/HairStrands/Source/HairStrandsCore/Private/HairStrandsDatas.cpp
+++ b/Engine/Plugins/Runtime/HairStrands/Source/HairStrandsCore/Private/HairStrandsDatas.cpp
@@ -174,33 +174,7 @@ void FHairStrandsInterpolationBulkData::Serialize(FArchive& Ar)
 {
 	Ar << Interpolation0;
 	Ar << Interpolation1;
-	// TODO add here SimRootIndices serializaztion once changed.
-//	Ar << SimRootPointIndex;
-}
-
-void FHairStrandsInterpolationDatas::Serialize(FArchive& Ar, FHairStrandsInterpolationBulkData& BulkData)
-{
-	Ar.UsingCustomVersion(FReleaseObjectVersion::GUID);
-
-	Ar << PointsSimCurvesVertexWeights;
-	Ar << PointsSimCurvesVertexIndex;
-	Ar << PointsSimCurvesIndex;
-
-	if (Ar.CustomVer(FReleaseObjectVersion::GUID) >= FReleaseObjectVersion::GroomAssetVersion3)
-	{
-		Ar << PointsSimCurvesVertexLerp;
-	}
-	else if (Ar.IsLoading())
-	{
-		const uint32 ElmentCount = PointsSimCurvesVertexIndex.Num();
-		PointsSimCurvesVertexLerp.SetNum(ElmentCount);
-		for (FVector& S : PointsSimCurvesVertexLerp)
-		{
-			S = FVector::ZeroVector;
-		}
-	}
-
-	BulkData.Serialize(Ar);
+	Ar << SimRootPointIndex;
 }
 
 void FHairStrandsPoints::Serialize(FArchive& Ar)
diff --git a/Engine/Plugins/Runtime/HairStrands/Source/HairStrandsCore/Public/GroomAsset.h b/Engine/Plugins/Runtime/HairStrands/Source/HairStrandsCore/Public/GroomAsset.h
--- a/Engine/Plugins/Runtime/HairStrands/Source/HairStrandsCore/Public/GroomAsset.h
+++ b/Engine/Plugins/Runtime/HairStrands/Source/HairStrandsCore/Public/GroomAsset.h
@@ -88,7 +88,6 @@ struct HAIRSTRANDSCORE_API FHairGroupData
 
 	struct FBaseWithInterpolation : FBase
 	{
-		FHairStrandsInterpolationDatas		InterpolationData;
 		FHairStrandsInterpolationBulkData	InterpolationBulkData;
 		FHairStrandsInterpolationResource*	InterpolationResource = nullptr;
 	};
@@ -193,7 +192,6 @@ struct HAIRSTRANDSCORE_API FHairGroupData
 			FHairCardsProceduralResource*		ProceduralResource = nullptr;
 
 			// Interpolation data & resources
-			FHairCardsInterpolationDatas		InterpolationData;
 			FHairCardsInterpolationBulkData		InterpolationBulkData;
 			FHairCardsInterpolationResource*	InterpolationResource = nullptr;
 
diff --git a/Engine/Plugins/Runtime/HairStrands/Source/HairStrandsCore/Public/GroomResources.h b/Engine/Plugins/Runtime/HairStrands/Source/HairStrandsCore/Public/GroomResources.h
--- a/Engine/Plugins/Runtime/HairStrands/Source/HairStrandsCore/Public/GroomResources.h
+++ b/Engine/Plugins/Runtime/HairStrands/Source/HairStrandsCore/Public/GroomResources.h
@@ -583,8 +583,6 @@ struct HAIRSTRANDSCORE_API FHairCardsInterpolationDatas
 
 	/** Lerp value between the closest vertex indices and the next one, ordered by closest influence */
 	TArray<float> PointsSimCurvesVertexLerp;
-
-	void Serialize(FArchive& Ar, FHairCardsInterpolationBulkData& BulkData);
 };
 
 struct HAIRSTRANDSCORE_API FHairCardsInterpolationBulkData
diff --git a/Engine/Plugins/Runtime/HairStrands/Source/HairStrandsCore/Public/HairCardsBuilder.h b/Engine/Plugins/Runtime/HairStrands/Source/HairStrandsCore/Public/HairCardsBuilder.h
--- a/Engine/Plugins/Runtime/HairStrands/Source/HairStrandsCore/Public/HairCardsBuilder.h
+++ b/Engine/Plugins/Runtime/HairStrands/Source/HairStrandsCore/Public/HairCardsBuilder.h
@@ -22,7 +22,6 @@ namespace FHairCardsBuilder
 		FHairCardsDatas& OutCards,
 		FHairCardsBulkData& OutBulk,
 		FHairStrandsDatas& OutGuides,
-		FHairCardsInterpolationDatas& OutInterpolationData,
 		FHairCardsInterpolationBulkData& OutInterpolationBulkData);
 
 	void ExportGeometry(
@@ -36,7 +35,6 @@ namespace FHairCardsBuilder
 		const struct FHairGroupsProceduralCards& Settings,
 		FHairCardsProceduralDatas& Out,
 		FHairStrandsDatas& OutGuides,
-		FHairCardsInterpolationDatas& OutInterpolation,
 		FHairCardsInterpolationBulkData& OutInterpolationBulk,
 		FHairGroupCardsTextures& OutTextures);
 
diff --git a/Engine/Plugins/Runtime/HairStrands/Source/HairStrandsCore/Public/HairStrandsDatas.h b/Engine/Plugins/Runtime/HairStrands/Source/HairStrandsCore/Public/HairStrandsDatas.h
--- a/Engine/Plugins/Runtime/HairStrands/Source/HairStrandsCore/Public/HairStrandsDatas.h
+++ b/Engine/Plugins/Runtime/HairStrands/Source/HairStrandsCore/Public/HairStrandsDatas.h
@@ -260,9 +260,6 @@ struct FHairStrandsInterpolationBulkData;
 /** Hair strands points interpolation attributes */
 struct HAIRSTRANDSCORE_API FHairStrandsInterpolationDatas
 {
-	/** Serialize the interpolated points */
-	void Serialize(FArchive& Ar, FHairStrandsInterpolationBulkData& BulkData);
-
 	/** Set the number of interpolated points */
 	void SetNum(const uint32 NumPoints);
 
@@ -291,9 +288,9 @@ struct HAIRSTRANDSCORE_API FHairStrandsInterpolationBulkData
 {
 	void Serialize(FArchive& Ar);
 	uint32 GetPointCount() const { return Interpolation0.Num(); };
-	TArray<FHairStrandsInterpolation0Format::Type> Interpolation0;
-	TArray<FHairStrandsInterpolation1Format::Type> Interpolation1;
-	TArray<FHairStrandsRootIndexFormat::Type> SimRootPointIndex;
+	TArray<FHairStrandsInterpolation0Format::Type> Interpolation0;	// Per-rendering-vertex interpolation data (closest guides, weight factors, ...)
+	TArray<FHairStrandsInterpolation1Format::Type> Interpolation1;	// Per-rendering-vertex interpolation data (closest guides, weight factors, ...)
+	TArray<FHairStrandsRootIndexFormat::Type> SimRootPointIndex;	// Per-rendering-vertex index of the sim-root vertex
 };
 
 /** Hair strands points attribute */
