Remove unsused resources.

diff --git a/Engine/Plugins/Runtime/HairStrands/Source/HairStrandsCore/Private/GroomAsset.cpp b/Engine/Plugins/Runtime/HairStrands/Source/HairStrandsCore/Private/GroomAsset.cpp
--- a/Engine/Plugins/Runtime/HairStrands/Source/HairStrandsCore/Private/GroomAsset.cpp
+++ b/Engine/Plugins/Runtime/HairStrands/Source/HairStrandsCore/Private/GroomAsset.cpp
@@ -620,7 +620,6 @@ void UGroomAsset::ReleaseResource()
 		for (FHairGroupData::FCards::FLOD& LOD : GroupData.Cards.LODs)
 		{
 			InternalReleaseResource(LOD.RestResource);
-			InternalReleaseResource(LOD.ProceduralResource);
 			InternalReleaseResource(LOD.InterpolationResource);
 			InternalReleaseResource(LOD.Guides.RestResource);
 			InternalReleaseResource(LOD.Guides.InterpolationResource);
@@ -2012,7 +2011,6 @@ bool UGroomAsset::BuildCardsGeometry(uint32 GroupIndex)
 
 			// 0. Release geometry resources
 			InternalReleaseResource(LOD.RestResource);
-			InternalReleaseResource(LOD.ProceduralResource);
 			InternalReleaseResource(LOD.InterpolationResource);
 			InternalReleaseResource(LOD.Guides.RestResource);
 			InternalReleaseResource(LOD.Guides.InterpolationResource);
diff --git a/Engine/Plugins/Runtime/HairStrands/Source/HairStrandsCore/Public/GroomAsset.h b/Engine/Plugins/Runtime/HairStrands/Source/HairStrandsCore/Public/GroomAsset.h
--- a/Engine/Plugins/Runtime/HairStrands/Source/HairStrandsCore/Public/GroomAsset.h
+++ b/Engine/Plugins/Runtime/HairStrands/Source/HairStrandsCore/Public/GroomAsset.h
@@ -171,7 +171,6 @@ struct HAIRSTRANDSCORE_API FHairGroupData
 			{
 				uint32 Total = 0;
 				if (RestResource) Total += RestResource->GetResourcesSize();
-				if (ProceduralResource) Total += ProceduralResource->GetResourcesSize();
 				if (InterpolationResource) Total += InterpolationResource->GetResourcesSize();
 				#if RHI_RAYTRACING
 				if (RaytracingResource) Total += RaytracingResource->GetResourcesSize();
@@ -187,10 +186,6 @@ struct HAIRSTRANDSCORE_API FHairGroupData
 			FHairCardsBulkData					BulkData;
 			FHairCardsRestResource*				RestResource = nullptr;
 
-			// Procedural data & resources
-			FHairCardsProceduralDatas			ProceduralData;
-			FHairCardsProceduralResource*		ProceduralResource = nullptr;
-
 			// Interpolation data & resources
 			FHairCardsInterpolationBulkData		InterpolationBulkData;
 			FHairCardsInterpolationResource*	InterpolationResource = nullptr;
