Remove mesh data from Groom asset.

diff --git a/Engine/Plugins/Runtime/HairStrands/Source/HairStrandsCore/Private/GroomAsset.cpp b/Engine/Plugins/Runtime/HairStrands/Source/HairStrandsCore/Private/GroomAsset.cpp
--- a/Engine/Plugins/Runtime/HairStrands/Source/HairStrandsCore/Private/GroomAsset.cpp
+++ b/Engine/Plugins/Runtime/HairStrands/Source/HairStrandsCore/Private/GroomAsset.cpp
@@ -1121,14 +1121,13 @@ FArchive& operator<<(FArchive& Ar, FHairGroupData::FMeshes::FLOD& MeshLODData)
 {
 	if (!Ar.IsCooking() || !MeshLODData.bIsCookedOut)
 	{
-		MeshLODData.Data.Serialize(Ar, MeshLODData.BulkData);
+		MeshLODData.BulkData.Serialize(Ar);
 	}
 	else
 	{
 		// LOD has been marked to be cooked out so serialize empty data
-		FHairMeshesDatas NoMeshesData;
 		FHairMeshesBulkData NoMeshesBulkData;
-		NoMeshesData.Serialize(Ar, NoMeshesBulkData);
+		NoMeshesBulkData.Serialize(Ar);
 	}
 
 	return Ar;
@@ -1331,7 +1330,7 @@ void UGroomAsset::SetHairWidth(float Width)
 // differences, etc.) replace the version GUID below with a new one.
 // In case of merge conflicts with DDC versions, you MUST generate a new GUID
 // and set this new GUID as the version.
-#define GROOM_DERIVED_DATA_VERSION TEXT("FCAFECFFD02A4509A0D86775505E97F2")
+#define GROOM_DERIVED_DATA_VERSION TEXT("81DC2B98E96241E2BDFC03DE84A6A8DD")
 
 #if WITH_EDITORONLY_DATA
 
@@ -2177,7 +2176,7 @@ bool UGroomAsset::BuildMeshesGeometry(uint32 GroupIndex)
 		int32 SourceIt = 0;
 		if (const FHairGroupsMeshesSourceDescription* Desc = GetSourceDescription(HairGroupsMeshes, GroupIndex, LODIt, SourceIt))
 		{
-			bIsAlreadyBuilt[LODIt] = GroupData.Meshes.LODs.IsValidIndex(LODIt) && GroupData.Meshes.LODs[LODIt].Data.Meshes.GetNumVertices() > 0 && !Desc->HasMeshChanged();
+			bIsAlreadyBuilt[LODIt] = GroupData.Meshes.LODs.IsValidIndex(LODIt) && GroupData.Meshes.LODs[LODIt].BulkData.GetNumVertices() > 0 && !Desc->HasMeshChanged();
 			bHasChanged |= !bIsAlreadyBuilt[LODIt];
 		}
 	}
@@ -2256,7 +2255,6 @@ bool UGroomAsset::BuildMeshesGeometry(uint32 GroupIndex)
 				Desc->ImportedMesh->ConditionalPostLoad();
 				FHairMeshesBuilder::ImportGeometry(
 					Desc->ImportedMesh,
-					LOD.Data,
 					LOD.BulkData);
 			}
 			else
@@ -2265,7 +2263,6 @@ bool UGroomAsset::BuildMeshesGeometry(uint32 GroupIndex)
 				FHairMeshesBuilder::BuildGeometry(
 					GroupData.Strands.Data,
 					GroupData.Guides.Data,
-					LOD.Data,
 					LOD.BulkData);
 			}
 
@@ -2369,7 +2366,7 @@ FHairStrandsRaytracingResource* UGroomAsset::AllocateMeshesRaytracingResources(u
 	{
 		FHairGroupData& GroupData = HairGroupsData[GroupIndex];
 		FHairGroupData::FMeshes::FLOD& LOD = GroupData.Meshes.LODs[LODIndex];
-		check(LOD.Data.IsValid());
+		check(LOD.BulkData.IsValid());
 
 		if (LOD.RaytracingResource == nullptr)
 		{
diff --git a/Engine/Plugins/Runtime/HairStrands/Source/HairStrandsCore/Private/HairCardsBuilder.cpp b/Engine/Plugins/Runtime/HairStrands/Source/HairStrandsCore/Private/HairCardsBuilder.cpp
--- a/Engine/Plugins/Runtime/HairStrands/Source/HairStrandsCore/Private/HairCardsBuilder.cpp
+++ b/Engine/Plugins/Runtime/HairStrands/Source/HairStrandsCore/Private/HairCardsBuilder.cpp
@@ -3904,7 +3904,6 @@ FString GetVersion()
 void BuildGeometry(
 	const FHairStrandsDatas& InRen,
 	const FHairStrandsDatas& InSim,
-	FHairMeshesDatas& Out,
 	FHairMeshesBulkData& OutBulk)
 {
 	const FVector Center = InRen.BoundingBox.GetCenter();
@@ -3914,6 +3913,8 @@ void BuildGeometry(
 	const uint32 TotalPointCount = 8;
 	const uint32 TotalIndexCount = 36;
 
+	FHairMeshesDatas Out;
+
 	Out.Meshes.Positions.SetNum(TotalPointCount);
 	Out.Meshes.Normals.SetNum(TotalPointCount);
 	Out.Meshes.Tangents.SetNum(TotalPointCount);
@@ -4026,9 +4027,10 @@ void BuildGeometry(
 
 void ImportGeometry(
 	const UStaticMesh* StaticMesh,
-	FHairMeshesDatas& Out,
 	FHairMeshesBulkData& OutBulk)
 {
+	FHairMeshesDatas Out;
+
 	const uint32 MeshLODIndex = 0;
 
 	// Note: if there are multiple section we only import the first one. Support for multiple section could be added later on. 
diff --git a/Engine/Plugins/Runtime/HairStrands/Source/HairStrandsCore/Private/HairCardsDatas.cpp b/Engine/Plugins/Runtime/HairStrands/Source/HairStrandsCore/Private/HairCardsDatas.cpp
--- a/Engine/Plugins/Runtime/HairStrands/Source/HairStrandsCore/Private/HairCardsDatas.cpp
+++ b/Engine/Plugins/Runtime/HairStrands/Source/HairStrandsCore/Private/HairCardsDatas.cpp
@@ -78,42 +78,12 @@ FArchive& operator<<(FArchive& Ar, FHairCardsInterpolationVertex& CardInterpVert
 	return Ar;
 }
 
-FArchive& operator<<(FArchive& Ar, FHairMeshes& HairMesh)
-{
-	Ar << HairMesh.UVs;
-	Ar << HairMesh.Normals;
-	Ar << HairMesh.Tangents;
-	Ar << HairMesh.Positions;
-	Ar << HairMesh.Indices;
-
-	// Bounds should be serialized
-	if (Ar.IsLoading())
-	{
-		HairMesh.BoundingBox.Init();
-		for (const FVector& P : HairMesh.Positions)
-		{
-			HairMesh.BoundingBox += P;
-		}
-	}
-
-	return Ar;
-}
-
-void FHairMeshesDatas::Serialize(FArchive& Ar, FHairMeshesBulkData& InBulkData)
-{
-	Ar << Meshes;
-	InBulkData.Serialize(Ar);
-	if (Ar.IsLoading())
-	{
-		InBulkData.BoundingBox = Meshes.BoundingBox;
-	}
-}
-
 void FHairMeshesBulkData::Serialize(FArchive& Ar)
 {
 	Ar << Positions;
 	Ar << Normals;
 	Ar << UVs;
 	Ar << Indices;
+	Ar << BoundingBox;
 }
 
diff --git a/Engine/Plugins/Runtime/HairStrands/Source/HairStrandsCore/Public/GroomAsset.h b/Engine/Plugins/Runtime/HairStrands/Source/HairStrandsCore/Public/GroomAsset.h
--- a/Engine/Plugins/Runtime/HairStrands/Source/HairStrandsCore/Public/GroomAsset.h
+++ b/Engine/Plugins/Runtime/HairStrands/Source/HairStrandsCore/Public/GroomAsset.h
@@ -250,7 +250,6 @@ struct HAIRSTRANDSCORE_API FHairGroupData
 			bool HasValidData() const { return BulkData.IsValid(); }
 			bool IsValid() const { return BulkData.IsValid() && RestResource != nullptr; }
 
-			FHairMeshesDatas Data;
 			FHairMeshesBulkData BulkData;
 			FHairMeshesRestResource* RestResource = nullptr;
 			#if RHI_RAYTRACING
diff --git a/Engine/Plugins/Runtime/HairStrands/Source/HairStrandsCore/Public/HairCardsBuilder.h b/Engine/Plugins/Runtime/HairStrands/Source/HairStrandsCore/Public/HairCardsBuilder.h
--- a/Engine/Plugins/Runtime/HairStrands/Source/HairStrandsCore/Public/HairCardsBuilder.h
+++ b/Engine/Plugins/Runtime/HairStrands/Source/HairStrandsCore/Public/HairCardsBuilder.h
@@ -54,12 +54,10 @@ namespace FHairMeshesBuilder
 	void BuildGeometry(
 		const struct FHairStrandsDatas& In,
 		const struct FHairStrandsDatas& InSim,
-		FHairMeshesDatas& Out,
 		FHairMeshesBulkData& OutBulk);
 
 	void ImportGeometry(
 		const UStaticMesh* StaticMesh,
-		FHairMeshesDatas& Out,
 		FHairMeshesBulkData& OutBulk);
 
 	FString GetVersion();
diff --git a/Engine/Plugins/Runtime/HairStrands/Source/HairStrandsCore/Public/HairCardsDatas.h b/Engine/Plugins/Runtime/HairStrands/Source/HairStrandsCore/Public/HairCardsDatas.h
--- a/Engine/Plugins/Runtime/HairStrands/Source/HairStrandsCore/Public/HairCardsDatas.h
+++ b/Engine/Plugins/Runtime/HairStrands/Source/HairStrandsCore/Public/HairCardsDatas.h
@@ -439,7 +439,6 @@ struct FHairMeshesDatas
 	bool IsValid() const { return Meshes.Positions.Num() > 0; }
 
 	FHairMeshes Meshes;
-	void Serialize(FArchive& Ar, FHairMeshesBulkData& InBulkData);
 };
 
 struct FHairMeshesBulkData
