Fix crash related to invalid curve offset data.

diff --git a/Engine/Plugins/Runtime/HairStrands/Source/HairStrandsCore/Private/GroomAsset.cpp b/Engine/Plugins/Runtime/HairStrands/Source/HairStrandsCore/Private/GroomAsset.cpp
--- a/Engine/Plugins/Runtime/HairStrands/Source/HairStrandsCore/Private/GroomAsset.cpp
+++ b/Engine/Plugins/Runtime/HairStrands/Source/HairStrandsCore/Private/GroomAsset.cpp
@@ -1330,7 +1330,7 @@ void UGroomAsset::SetHairWidth(float Width)
 // differences, etc.) replace the version GUID below with a new one.
 // In case of merge conflicts with DDC versions, you MUST generate a new GUID
 // and set this new GUID as the version.
-#define GROOM_DERIVED_DATA_VERSION TEXT("81DC2B98E96241E2BDFC03DE84A6A8DD")
+#define GROOM_DERIVED_DATA_VERSION TEXT("B136B9EEDA5C4FABAF263AD0B05C1D76")
 
 #if WITH_EDITORONLY_DATA
 
diff --git a/Engine/Plugins/Runtime/HairStrands/Source/HairStrandsCore/Private/GroomBuilder.cpp b/Engine/Plugins/Runtime/HairStrands/Source/HairStrandsCore/Private/GroomBuilder.cpp
--- a/Engine/Plugins/Runtime/HairStrands/Source/HairStrandsCore/Private/GroomBuilder.cpp
+++ b/Engine/Plugins/Runtime/HairStrands/Source/HairStrandsCore/Private/GroomBuilder.cpp
@@ -231,7 +231,7 @@ namespace HairStrandsBuilder
 		OutPackedPositions.SetNum(NumPoints * FHairStrandsPositionFormat::ComponentCount);
 		OutPackedAttributes.SetNum(NumPoints * FHairStrandsAttributeFormat::ComponentCount);
 		OutPackedMaterials.SetNum(NumPoints * FHairStrandsMaterialFormat::ComponentCount);
-		OutCurveOffsets.SetNum(NumCurves);
+		OutCurveOffsets = HairStrands.StrandsCurves.CurvesOffset;
 
 		const FVector HairBoxCenter = HairStrands.BoundingBox.GetCenter();
 
@@ -284,7 +284,6 @@ namespace HairStrandsBuilder
 				Material.BaseColorB = FMath::Clamp(uint32(FMath::Sqrt(Points.PointsBaseColor[PointIndex + IndexOffset].B) * 0xFF), 0u, 0xFFu);
 				Material.Roughness  = FMath::Clamp(uint32(Points.PointsRoughness[PointIndex + IndexOffset] * 0xFF), 0u, 0xFFu);
 			}
-			OutCurveOffsets[CurveIndex] = IndexOffset;
 		}
 
 		OutBulkData.BoundingBox = HairStrands.BoundingBox;
