Remove data from cards groom asset. Only cards bulk data are left.

diff --git a/Engine/Plugins/Experimental/MeshModelingToolset/Source/MeshModelingToolsEditorOnly/Private/Hair/GroomQueryUtil.cpp b/Engine/Plugins/Experimental/MeshModelingToolset/Source/MeshModelingToolsEditorOnly/Private/Hair/GroomQueryUtil.cpp
--- a/Engine/Plugins/Experimental/MeshModelingToolset/Source/MeshModelingToolsEditorOnly/Private/Hair/GroomQueryUtil.cpp
+++ b/Engine/Plugins/Experimental/MeshModelingToolset/Source/MeshModelingToolsEditorOnly/Private/Hair/GroomQueryUtil.cpp
@@ -6,6 +6,7 @@
 #include "MeshDescriptionToDynamicMesh.h"
 #include "Selections/MeshConnectedComponents.h"
 #include "FaceGroupUtil.h"
+#include "HairCardsBuilder.h"
 
 #include "ExplicitUseGeometryMathTypes.h"		// using UE::Geometry::(math types)
 using namespace UE::Geometry;
@@ -48,8 +49,32 @@ void UE::GroomQueries::ExtractAllHairCards(AGroomActor* GroomActor,
 			continue;
 		}
 
-		const FHairGroupData::FCards::FLOD& CardLOD = GroupData.Cards.LODs[LODIndex];
-		const FHairCardsDatas& CardData = CardLOD.Data;
+		UStaticMesh* StaticMesh = nullptr;
+		for (const FHairGroupsCardsSourceDescription& Desc : Asset->HairGroupsCards)
+		{
+			if (Desc.GroupIndex == GroupIdx && Desc.LODIndex == LODIndex)
+			{
+				if (Desc.SourceType == EHairCardsSourceType::Imported)
+				{
+					StaticMesh = Desc.ImportedMesh;
+				}
+				else if (Desc.SourceType == EHairCardsSourceType::Procedural)
+				{
+					StaticMesh = Desc.ProceduralMesh;
+				}
+			}
+		}
+		
+		if (!StaticMesh)
+		{
+			continue;
+		}
+
+		FHairCardsDatas CardData;
+		if (!FHairCardsBuilder::ExtractCardsData(StaticMesh, CardData))
+		{
+			continue;
+		}
 		const FHairCardsGeometry& CardGeo = CardData.Cards;
 
 		int32 NumVerts = CardGeo.GetNumVertices();
diff --git a/Engine/Plugins/Runtime/HairStrands/Source/HairStrandsCore/Private/GroomAsset.cpp b/Engine/Plugins/Runtime/HairStrands/Source/HairStrandsCore/Private/GroomAsset.cpp
--- a/Engine/Plugins/Runtime/HairStrands/Source/HairStrandsCore/Private/GroomAsset.cpp
+++ b/Engine/Plugins/Runtime/HairStrands/Source/HairStrandsCore/Private/GroomAsset.cpp
@@ -1093,7 +1093,7 @@ FArchive& operator<<(FArchive& Ar, FHairGroupData::FCards::FLOD& CardLODData)
 {
 	if (!Ar.IsCooking() || !CardLODData.bIsCookedOut)
 	{
-		CardLODData.Data.Serialize(Ar, CardLODData.BulkData);
+		CardLODData.BulkData.Serialize(Ar);
 		CardLODData.InterpolationBulkData.Serialize(Ar);
 
 		CardLODData.Guides.Data.Serialize(Ar, CardLODData.Guides.BulkData);
@@ -1102,9 +1102,8 @@ FArchive& operator<<(FArchive& Ar, FHairGroupData::FCards::FLOD& CardLODData)
 	else
 	{
 		// LOD has been marked to be cooked out so serialize empty data
-		FHairCardsDatas NoCardsData;
 		FHairCardsBulkData NoCardsBulkData;
-		NoCardsData.Serialize(Ar, NoCardsBulkData);
+		NoCardsBulkData.Serialize(Ar);
 
 		FHairCardsInterpolationBulkData NoInterpolationBulkData;
 		NoInterpolationBulkData.Serialize(Ar);
@@ -2006,7 +2005,7 @@ bool UGroomAsset::BuildCardsGeometry(uint32 GroupIndex)
 			}
 
 			FHairGroupData::FCards::FLOD& LOD = GroupData.Cards.LODs[LODIt];
-			LOD.Data.Cards.Reset();
+			LOD.BulkData.Reset();
 
 			// 0. Release geometry resources
 			InternalReleaseResource(LOD.RestResource);
@@ -2029,12 +2028,10 @@ bool UGroomAsset::BuildCardsGeometry(uint32 GroupIndex)
 			if (CardsMesh != nullptr)
 			{
 				CardsMesh->ConditionalPostLoad();
-				bInitResources = FHairCardsBuilder::ImportGeometry(CardsMesh, LOD.Data, LOD.BulkData, LOD.Guides.Data, LOD.InterpolationBulkData);
+				bInitResources = FHairCardsBuilder::ImportGeometry(CardsMesh, LOD.BulkData, LOD.Guides.Data, LOD.InterpolationBulkData);
 				if (!bInitResources)
 				{
 					UE_LOG(LogHairStrands, Warning, TEXT("Failed to import cards from %s for Group %d LOD %d."), *CardsMesh->GetName(), GroupIndex, LODIt);
-
-					LOD.Data.Cards.Reset();
 					LOD.BulkData.Reset();
 				}
 			}
@@ -2106,8 +2103,8 @@ bool UGroomAsset::BuildCardsGeometry(uint32 GroupIndex)
 				BeginInitResource(LOD.Guides.InterpolationResource);
 
 				// Update card stats to display
-				Desc->CardsInfo.NumCardVertices = LOD.Data.Cards.GetNumVertices();
-				Desc->CardsInfo.NumCards = LOD.Data.Cards.IndexOffsets.Num();
+				Desc->CardsInfo.NumCardVertices = LOD.BulkData.GetNumVertices();
+				Desc->CardsInfo.NumCards = LOD.Guides.BulkData.CurveCount;// LOD.Data.Cards.IndexOffsets.Num();
 			}
 		}
 	}
@@ -2348,7 +2345,7 @@ FHairStrandsRaytracingResource* UGroomAsset::AllocateCardsRaytracingResources(ui
 	{
 		FHairGroupData& GroupData = HairGroupsData[GroupIndex];
 		FHairGroupData::FCards::FLOD& LOD = GroupData.Cards.LODs[LODIndex];
-		check(LOD.Data.IsValid());
+		check(LOD.BulkData.IsValid());
 
 		if (LOD.RaytracingResource == nullptr)
 		{
@@ -2536,8 +2533,8 @@ void UGroomAsset::InitCardsResources()
 			if (Desc)
 			{
 				// Update card stats to display
-				Desc->CardsInfo.NumCardVertices = LOD.Data.Cards.GetNumVertices();
-				Desc->CardsInfo.NumCards = LOD.Data.Cards.GetNumCards();
+				Desc->CardsInfo.NumCardVertices = LOD.BulkData.GetNumVertices();
+				Desc->CardsInfo.NumCards = LOD.Guides.BulkData.CurveCount;// LOD.Data.Cards.GetNumCards();
 			}
 		}
 	}
@@ -2903,13 +2900,12 @@ void UGroomAsset::SaveProceduralCards(uint32 DescIndex)
 		GroupData.Guides.Data,
 		Desc->ProceduralSettings,
 		Q.ProceduralData,
+		Q.BulkData,
 		Q.GuideData,
 		Q.InterpolationBulkData,
 		Desc->Textures);
 	Q.Textures = &Desc->Textures;
 
-	FHairCardsBuilder::Convert(Q.ProceduralData, Q.Data, Q.BulkData);
-
 	// 3. Create resources and enqueue texture generation (GPU, kicked by the render thread) 
 	Q.Resources = new FHairCardsRestResource(Q.BulkData);
 	BeginInitResource(Q.Resources);
diff --git a/Engine/Plugins/Runtime/HairStrands/Source/HairStrandsCore/Private/HairCardsBuilder.cpp b/Engine/Plugins/Runtime/HairStrands/Source/HairStrandsCore/Private/HairCardsBuilder.cpp
--- a/Engine/Plugins/Runtime/HairStrands/Source/HairStrandsCore/Private/HairCardsBuilder.cpp
+++ b/Engine/Plugins/Runtime/HairStrands/Source/HairStrandsCore/Private/HairCardsBuilder.cpp
@@ -3327,7 +3327,7 @@ namespace FHairCardsBuilder
 FString GetVersion()
 {
 	// Important to update the version when cards building or importing changes
-	return TEXT("8");
+	return TEXT("9ai");
 }
 
 void AllocateAtlasTexture(UTexture2D* Out, const FIntPoint& Resolution, uint32 MipCount, EPixelFormat PixelFormat, ETextureSourceFormat SourceFormat)
@@ -3381,6 +3381,7 @@ void BuildGeometry(
 	const FHairStrandsDatas& InSim,
 	const FHairGroupsProceduralCards& Settings,
 	FHairCardsProceduralDatas& Out,
+	FHairCardsBulkData& OutBulk,
 	FHairStrandsDatas& OutGuides,
 	FHairCardsInterpolationBulkData& OutInterpolationBulk,
 	FHairGroupCardsTextures& OutTextures)
@@ -3468,6 +3469,13 @@ void BuildGeometry(
 		Out.RenderData.Indices[IndexIt] = Out.Cards.Indices[IndexIt];
 	}
 
+	// Copy transient Procedural cards bulk data into the final bulk data
+	OutBulk.Positions 	= Out.RenderData.Positions;
+	OutBulk.Normals 	= Out.RenderData.Normals;
+	OutBulk.Indices 	= Out.RenderData.Indices;
+	OutBulk.UVs 		= Out.RenderData.UVs;
+	OutBulk.BoundingBox = Out.Cards.BoundingBox;
+
 	const uint32 CardCount = Out.Cards.Lengths.Num();
 	Out.RenderData.CardsRect.SetNum(CardCount);
 	Out.RenderData.CardsLengths.SetNum(CardCount);
@@ -3561,7 +3569,7 @@ void SanitizeMeshDescription(FMeshDescription* MeshDescription)
 	}
 }
 
-bool ImportGeometry(
+bool InternalImportGeometry(
 	const UStaticMesh* StaticMesh,
 	FHairCardsDatas& Out,
 	FHairCardsBulkData& OutBulk,
@@ -3734,15 +3742,32 @@ bool ImportGeometry(
 	return bSuccess;
 }
 
-void Convert(const FHairCardsProceduralDatas& In, FHairCardsDatas& Out, FHairCardsBulkData& OutBulk)
+bool ImportGeometry(
+	const UStaticMesh* StaticMesh,
+	FHairCardsBulkData& OutBulk,
+	FHairStrandsDatas& OutGuides,
+	FHairCardsInterpolationBulkData& OutInterpolationBulkData)
 {
-	Out.Cards = In.Cards;
+	FHairCardsDatas CardData;
+	return InternalImportGeometry(
+		StaticMesh,
+		CardData,
+		OutBulk,
+		OutGuides,
+		OutInterpolationBulkData);
+}
 
-	OutBulk.Positions = In.RenderData.Positions;
-	OutBulk.Normals = In.RenderData.Normals;
-	OutBulk.Indices = In.RenderData.Indices;
-	OutBulk.UVs = In.RenderData.UVs;
-	OutBulk.BoundingBox = In.Cards.BoundingBox;
+bool ExtractCardsData(const UStaticMesh* StaticMesh, FHairCardsDatas& Out)
+{
+	FHairCardsBulkData OutBulk;
+	FHairStrandsDatas OutGuides;
+	FHairCardsInterpolationBulkData OutInterpolationBulkData;
+	return InternalImportGeometry(
+		StaticMesh,
+		Out,
+		OutBulk,
+		OutGuides,
+		OutInterpolationBulkData); 
 }
 
 ///////////////////////////////////////////////////////////////////////////////////////////////////
diff --git a/Engine/Plugins/Runtime/HairStrands/Source/HairStrandsCore/Private/HairCardsDatas.cpp b/Engine/Plugins/Runtime/HairStrands/Source/HairStrandsCore/Private/HairCardsDatas.cpp
--- a/Engine/Plugins/Runtime/HairStrands/Source/HairStrandsCore/Private/HairCardsDatas.cpp
+++ b/Engine/Plugins/Runtime/HairStrands/Source/HairStrandsCore/Private/HairCardsDatas.cpp
@@ -15,49 +15,13 @@ FArchive& operator<<(FArchive& Ar, FUIntPoint& Point)
 	return Ar;
 }
 
-FArchive& operator<<(FArchive& Ar, FHairCardsGeometry& CardGeometry)
-{
-	Ar << CardGeometry.UVs;
-	Ar << CardGeometry.Normals;
-	Ar << CardGeometry.Tangents;
-	Ar << CardGeometry.Positions;
-	Ar << CardGeometry.Indices;
-
-	Ar << CardGeometry.PointOffsets; // No longer used, kept only for backward compatibility
-	Ar << CardGeometry.PointCounts;  // No longer used, kept only for backward compatibility
-
-	Ar << CardGeometry.IndexOffsets;
-	Ar << CardGeometry.IndexCounts;
-
-	// Bounds should be serialized
-	if (Ar.IsLoading())
-	{
-		CardGeometry.BoundingBox.Init();
-		for (const FVector& P : CardGeometry.Positions)
-		{
-			CardGeometry.BoundingBox += P;
-		}
-	}
-
-	return Ar;
-}
-
 void FHairCardsBulkData::Serialize(FArchive& Ar)
 {
 	Ar << Positions;
 	Ar << Normals;
 	Ar << UVs;
 	Ar << Indices;
-}
-
-void FHairCardsDatas::Serialize(FArchive& Ar, FHairCardsBulkData& InBulkData)
-{
-	Ar << Cards;
-	InBulkData.Serialize(Ar);
-	if (Ar.IsLoading())
-	{
-		InBulkData.BoundingBox = Cards.BoundingBox;
-	}
+	Ar << BoundingBox;
 }
 
 FArchive& operator<<(FArchive& Ar, FHairCardsInterpolationVertex& CardInterpVertex)
diff --git a/Engine/Plugins/Runtime/HairStrands/Source/HairStrandsCore/Public/GroomAsset.h b/Engine/Plugins/Runtime/HairStrands/Source/HairStrandsCore/Public/GroomAsset.h
--- a/Engine/Plugins/Runtime/HairStrands/Source/HairStrandsCore/Public/GroomAsset.h
+++ b/Engine/Plugins/Runtime/HairStrands/Source/HairStrandsCore/Public/GroomAsset.h
@@ -182,7 +182,6 @@ struct HAIRSTRANDSCORE_API FHairGroupData
 			bool IsValid() const { return BulkData.IsValid() && RestResource != nullptr; }
 
 			// Main data & Resources
-			FHairCardsDatas						Data;
 			FHairCardsBulkData					BulkData;
 			FHairCardsRestResource*				RestResource = nullptr;
 
diff --git a/Engine/Plugins/Runtime/HairStrands/Source/HairStrandsCore/Public/HairCardsBuilder.h b/Engine/Plugins/Runtime/HairStrands/Source/HairStrandsCore/Public/HairCardsBuilder.h
--- a/Engine/Plugins/Runtime/HairStrands/Source/HairStrandsCore/Public/HairCardsBuilder.h
+++ b/Engine/Plugins/Runtime/HairStrands/Source/HairStrandsCore/Public/HairCardsBuilder.h
@@ -19,11 +19,14 @@ namespace FHairCardsBuilder
 {
 	bool ImportGeometry(
 		const UStaticMesh* StaticMesh,
-		FHairCardsDatas& OutCards,
 		FHairCardsBulkData& OutBulk,
 		FHairStrandsDatas& OutGuides,
 		FHairCardsInterpolationBulkData& OutInterpolationBulkData);
 
+	HAIRSTRANDSCORE_API bool ExtractCardsData(
+		const UStaticMesh* StaticMesh, 
+		FHairCardsDatas& Out);
+
 	void ExportGeometry(
 		const FHairCardsDatas& InCardsData, 
 		UStaticMesh* OutStaticMesh);
@@ -34,6 +37,7 @@ namespace FHairCardsBuilder
 		const struct FHairStrandsDatas& InSim,
 		const struct FHairGroupsProceduralCards& Settings,
 		FHairCardsProceduralDatas& Out,
+		FHairCardsBulkData& OutBulk,
 		FHairStrandsDatas& OutGuides,
 		FHairCardsInterpolationBulkData& OutInterpolationBulk,
 		FHairGroupCardsTextures& OutTextures);
@@ -44,8 +48,6 @@ namespace FHairCardsBuilder
 		FHairCardsProceduralResource* AtlasResource,
 		FHairGroupCardsTextures* Textures);
 
-	void Convert(const FHairCardsProceduralDatas& In, FHairCardsDatas& Out, FHairCardsBulkData& OutBulk);
-
 	FString GetVersion();
 }
 
diff --git a/Engine/Plugins/Runtime/HairStrands/Source/HairStrandsCore/Public/HairCardsDatas.h b/Engine/Plugins/Runtime/HairStrands/Source/HairStrandsCore/Public/HairCardsDatas.h
--- a/Engine/Plugins/Runtime/HairStrands/Source/HairStrandsCore/Public/HairCardsDatas.h
+++ b/Engine/Plugins/Runtime/HairStrands/Source/HairStrandsCore/Public/HairCardsDatas.h
@@ -246,8 +246,6 @@ struct FHairCardsDatas
 	bool IsValid() const { return Cards.Positions.Num() > 0; }
 
 	FHairCardsGeometry Cards;
-
-	void Serialize(FArchive& Ar, FHairCardsBulkData& BulkData);
 };
 
 struct FHairCardsBulkData
