Fix crash when running the engine in debug build.

diff --git a/Engine/Plugins/Runtime/HairStrands/Source/HairStrandsCore/Private/GroomComponent.cpp b/Engine/Plugins/Runtime/HairStrands/Source/HairStrandsCore/Private/GroomComponent.cpp
--- a/Engine/Plugins/Runtime/HairStrands/Source/HairStrandsCore/Private/GroomComponent.cpp
+++ b/Engine/Plugins/Runtime/HairStrands/Source/HairStrandsCore/Private/GroomComponent.cpp
@@ -824,7 +824,7 @@ public:
 
 		BatchElement.FirstIndex = 0;
 		BatchElement.NumInstances = 1;
-		BatchElement.PrimitiveIdMode = PrimID_DynamicPrimitiveShaderData; // Force using per primitive uniform buffer.
+		BatchElement.PrimitiveIdMode = PrimID_ForceZero;
 		if (bUseCulling)
 		{
 			BatchElement.NumPrimitives = 0;
