Remove FHairStrandsData from serialized data

diff --git a/Engine/Plugins/Experimental/MeshModelingToolset/Source/MeshModelingToolsEditorOnly/Private/Hair/GroomQueryUtil.cpp b/Engine/Plugins/Experimental/MeshModelingToolset/Source/MeshModelingToolsEditorOnly/Private/Hair/GroomQueryUtil.cpp
--- a/Engine/Plugins/Experimental/MeshModelingToolset/Source/MeshModelingToolsEditorOnly/Private/Hair/GroomQueryUtil.cpp
+++ b/Engine/Plugins/Experimental/MeshModelingToolset/Source/MeshModelingToolsEditorOnly/Private/Hair/GroomQueryUtil.cpp
@@ -525,8 +525,12 @@ void UE::GroomQueries::ProcessHairCurves(AGroomActor* GroomActor,
 		//int32 NumCurves = GroupInfo.NumCurves;
 		//int32 NumGuides = GroupInfo.NumGuides;
 
+		FHairStrandsDatas StrandsData;
+		FHairStrandsDatas GuidesData;
+		Asset->GetHairStrandsDatas(GroupIdx, StrandsData, GuidesData);
+
 		//const FHairStrandsDatas& GroupStrandData = (bUseGuides) ? GroupData.HairSimulationData : GroupData.HairRenderData;
-		const FHairStrandsDatas& GroupStrandData = (bUseGuides) ? GroupData.Guides.Data : GroupData.Strands.Data;
+		const FHairStrandsDatas& GroupStrandData = (bUseGuides) ? GuidesData : StrandsData;
 
 		const FHairStrandsPoints& GroupStrandPoints = GroupStrandData.StrandsPoints;
 		const TArray<FVector>& Positions = GroupStrandPoints.PointsPosition;
diff --git a/Engine/Plugins/Experimental/MeshModelingToolset/Source/MeshModelingToolsEditorOnly/Private/Hair/GroomToMeshTool.cpp b/Engine/Plugins/Experimental/MeshModelingToolset/Source/MeshModelingToolsEditorOnly/Private/Hair/GroomToMeshTool.cpp
--- a/Engine/Plugins/Experimental/MeshModelingToolset/Source/MeshModelingToolsEditorOnly/Private/Hair/GroomToMeshTool.cpp
+++ b/Engine/Plugins/Experimental/MeshModelingToolset/Source/MeshModelingToolsEditorOnly/Private/Hair/GroomToMeshTool.cpp
@@ -305,7 +305,10 @@ static void ProcessHairCurvePoints(AGroomActor* GroomActor,
 		//int32 NumGuides = GroupInfo.NumGuides;
 
 		//const FHairStrandsDatas& GroupStrandData = (bUseGuides) ? GroupData.HairSimulationData : GroupData.HairRenderData;
-		const FHairStrandsDatas& GroupStrandData = (bUseGuides) ? GroupData.Guides.Data : GroupData.Strands.Data;
+		FHairStrandsDatas StrandsData;
+		FHairStrandsDatas GuidesData;
+		Asset->GetHairStrandsDatas(GroupIdx, StrandsData, GuidesData);
+		const FHairStrandsDatas& GroupStrandData = (bUseGuides) ? GuidesData : StrandsData;
 
 		const FHairStrandsPoints& GroupStrandPoints = GroupStrandData.StrandsPoints;
 		const TArray<FVector>& Positions = GroupStrandPoints.PointsPosition;
@@ -347,8 +350,11 @@ static void ProcessHairCurves(AGroomActor* GroomActor,
 		//int32 NumCurves = GroupInfo.NumCurves;
 		//int32 NumGuides = GroupInfo.NumGuides;
 
+		FHairStrandsDatas StrandsData;
+		FHairStrandsDatas GuidesData;
+		Asset->GetHairStrandsDatas(GroupIdx, StrandsData, GuidesData);
 		//const FHairStrandsDatas& GroupStrandData = (bUseGuides) ? GroupData.HairSimulationData : GroupData.HairRenderData;
-		const FHairStrandsDatas& GroupStrandData = (bUseGuides) ? GroupData.Guides.Data : GroupData.Strands.Data;
+		const FHairStrandsDatas& GroupStrandData = (bUseGuides) ? GuidesData : StrandsData;
 
 		const FHairStrandsPoints& GroupStrandPoints = GroupStrandData.StrandsPoints;
 		const TArray<FVector>& Positions = GroupStrandPoints.PointsPosition;
diff --git a/Engine/Plugins/Runtime/HairStrands/Source/HairStrandsCore/Private/GroomAsset.cpp b/Engine/Plugins/Runtime/HairStrands/Source/HairStrandsCore/Private/GroomAsset.cpp
--- a/Engine/Plugins/Runtime/HairStrands/Source/HairStrandsCore/Private/GroomAsset.cpp
+++ b/Engine/Plugins/Runtime/HairStrands/Source/HairStrandsCore/Private/GroomAsset.cpp
@@ -172,25 +172,13 @@ bool IsHairStrandsAssetLoadingEnable()
 	return GHairStrandsLoadAsset > 0;
 }
 
-static void PatchHairStrandsInterpolationBulkData(FHairStrandsInterpolationBulkData& BulkData, const FHairStrandsDatas& SimDatas)
-{
-	// This is now move into the builder, but kept here for compatibility.
-	if (BulkData.SimRootPointIndex.Num() == 0)
-	{
-		const uint32 RootCount = SimDatas.GetNumCurves();
-		BulkData.SimRootPointIndex.SetNum(SimDatas.GetNumPoints());
-		for (uint32 CurveIndex = 0; CurveIndex < RootCount; ++CurveIndex)
-		{
-			const uint16 PointCount = SimDatas.StrandsCurves.CurvesCount[CurveIndex];
-			const uint32 PointOffset = SimDatas.StrandsCurves.CurvesOffset[CurveIndex];
-			for (uint32 PointIndex = 0; PointIndex < PointCount; ++PointIndex)
-			{
-				BulkData.SimRootPointIndex[PointIndex + PointOffset] = PointOffset;
-			}
-		}
-	}
-}
-
+// Build: 
+// * Guides bulk data
+// * Strands bulk data
+// * Interpolation bulk data
+// * Clusters (bulk) data
+//
+// Guides, strands, and interpolation data are all transient
 static bool BuildHairGroup(
 	const int32 GroupIndex,
 	const FHairDescriptionGroups& HairDescriptionGroups, 
@@ -206,20 +194,26 @@ static bool BuildHairGroup(
 		check(GroupIndex <= HairDescriptionGroups.HairGroups.Num());
 		check(GroupIndex == HairGroup.Info.GroupID);
 
-		FGroomBuilder::BuildData(HairGroup, HairGroupsInterpolation[GroupIndex], OutHairGroupsInfo[GroupIndex], OutHairGroupsData[GroupIndex].Strands.Data, OutHairGroupsData[GroupIndex].Guides.Data);
+		FHairStrandsDatas StrandsData;
+		FHairStrandsDatas GuidesData;
+		FGroomBuilder::BuildData(HairGroup, HairGroupsInterpolation[GroupIndex], OutHairGroupsInfo[GroupIndex], StrandsData, GuidesData);
 
-		FGroomBuilder::BuildBulkData(HairGroup.Info, OutHairGroupsData[GroupIndex].Guides.Data, OutHairGroupsData[GroupIndex].Guides.BulkData);
-		FGroomBuilder::BuildBulkData(HairGroup.Info, OutHairGroupsData[GroupIndex].Strands.Data, OutHairGroupsData[GroupIndex].Strands.BulkData);
+		FGroomBuilder::BuildBulkData(HairGroup.Info, GuidesData, OutHairGroupsData[GroupIndex].Guides.BulkData);
+		FGroomBuilder::BuildBulkData(HairGroup.Info, StrandsData, OutHairGroupsData[GroupIndex].Strands.BulkData);
 
 		FHairStrandsInterpolationDatas InterpolationData;
-		FGroomBuilder::BuildInterplationData(HairGroup.Info, OutHairGroupsData[GroupIndex].Strands.Data, OutHairGroupsData[GroupIndex].Guides.Data, HairGroupsInterpolation[GroupIndex].InterpolationSettings, InterpolationData);
-		FGroomBuilder::BuildInterplationBulkData(OutHairGroupsData[GroupIndex].Guides.Data, InterpolationData, OutHairGroupsData[GroupIndex].Strands.InterpolationBulkData);
+		FGroomBuilder::BuildInterplationData(HairGroup.Info, StrandsData, GuidesData, HairGroupsInterpolation[GroupIndex].InterpolationSettings, InterpolationData);
+		FGroomBuilder::BuildInterplationBulkData(GuidesData, InterpolationData, OutHairGroupsData[GroupIndex].Strands.InterpolationBulkData);
 
-		FGroomBuilder::BuildClusterData(OutHairGroupsData[GroupIndex].Strands.Data, HairDescriptionGroups.BoundRadius, HairGroupsLOD[GroupIndex], OutHairGroupsData[GroupIndex].Strands.ClusterCullingData);
+		FGroomBuilder::BuildClusterData(StrandsData, HairDescriptionGroups.BoundRadius, HairGroupsLOD[GroupIndex], OutHairGroupsData[GroupIndex].Strands.ClusterCullingData);
 	}
 	return bIsValid;
 }
 
+// Build: 
+// * Clusters (bulk) data
+//
+// Guides and strands data are all transient
 static bool BuildHairGroupCluster(
 	const int32 GroupIndex,
 	const FHairDescriptionGroups& HairDescriptionGroups,
@@ -235,8 +229,10 @@ static bool BuildHairGroupCluster(
 		check(GroupIndex <= HairDescriptionGroups.HairGroups.Num());
 		check(GroupIndex == HairGroup.Info.GroupID);
 
-		FGroomBuilder::BuildData(HairGroup, HairGroupsInterpolation[GroupIndex], OutHairGroupsInfo[GroupIndex], OutHairGroupsData[GroupIndex].Strands.Data, OutHairGroupsData[GroupIndex].Guides.Data);
-		FGroomBuilder::BuildClusterData(OutHairGroupsData[GroupIndex].Strands.Data, HairDescriptionGroups.BoundRadius, HairGroupsLOD[GroupIndex], OutHairGroupsData[GroupIndex].Strands.ClusterCullingData);
+		FHairStrandsDatas StrandsData;
+		FHairStrandsDatas GuidesData;
+		FGroomBuilder::BuildData(HairGroup, HairGroupsInterpolation[GroupIndex], OutHairGroupsInfo[GroupIndex], StrandsData, GuidesData);
+		FGroomBuilder::BuildClusterData(StrandsData, HairDescriptionGroups.BoundRadius, HairGroupsLOD[GroupIndex], OutHairGroupsData[GroupIndex].Strands.ClusterCullingData);
 	}
 	return bIsValid;
 }
@@ -562,7 +558,13 @@ void UGroomAsset::UpdateResource()
 				const FHairDescriptionGroups& LocalHairDescriptionGroups = GetHairDescriptionGroups();
 
 				// Force rebuilding the LOD data as the LOD settings has changed
-				FGroomBuilder::BuildClusterData(HairGroupsData[GroupIndex].Strands.Data, LocalHairDescriptionGroups.BoundRadius, HairGroupsLOD[GroupIndex], HairGroupsData[GroupIndex].Strands.ClusterCullingData);
+				BuildHairGroupCluster(
+					GroupIndex,
+					LocalHairDescriptionGroups,
+					HairGroupsInterpolation,
+					HairGroupsLOD,
+					HairGroupsInfo,
+					HairGroupsData);
 
 				GroupData.Strands.ClusterCullingResource = new FHairStrandsClusterCullingResource(GroupData.Strands.ClusterCullingData);
 				BeginInitResource(GroupData.Strands.ClusterCullingResource);
@@ -910,7 +912,13 @@ void UGroomAsset::PreSave(FObjectPreSaveContext ObjectSaveContext)
 			}
 			else if (bHasLODChanged)
 			{
-				FGroomBuilder::BuildClusterData(HairGroupsData[GroupIt].Strands.Data, HairDescriptionGroups->BoundRadius, HairGroupsLOD[GroupIt], HairGroupsData[GroupIt].Strands.ClusterCullingData);
+				BuildHairGroupCluster(
+					GroupIt,
+					LocalHairDescriptionGroups,
+					HairGroupsInterpolation,
+					HairGroupsLOD,
+					HairGroupsInfo,
+					HairGroupsData);
 			}
 		}
 		InitResources();
@@ -1096,7 +1104,7 @@ FArchive& operator<<(FArchive& Ar, FHairGroupData::FCards::FLOD& CardLODData)
 		CardLODData.BulkData.Serialize(Ar);
 		CardLODData.InterpolationBulkData.Serialize(Ar);
 
-		CardLODData.Guides.Data.Serialize(Ar, CardLODData.Guides.BulkData);
+		CardLODData.Guides.BulkData.Serialize(Ar);
 		CardLODData.Guides.InterpolationBulkData.Serialize(Ar);
 	}
 	else
@@ -1109,7 +1117,7 @@ FArchive& operator<<(FArchive& Ar, FHairGroupData::FCards::FLOD& CardLODData)
 		NoInterpolationBulkData.Serialize(Ar);
 
 		FHairGroupData::FBaseWithInterpolation NoGuideData;
-		NoGuideData.Data.Serialize(Ar, NoGuideData.BulkData);
+		NoGuideData.BulkData.Serialize(Ar);
 		NoGuideData.InterpolationBulkData.Serialize(Ar);
 	}
 
@@ -1139,8 +1147,8 @@ FArchive& operator<<(FArchive& Ar, FHairGroupData& GroupData)
 	FHairGroupData NoStrandsData;
 	if (!Ar.IsCooking() || !GroupData.bIsCookedOut)
 	{
-		GroupData.Strands.Data.Serialize(Ar, GroupData.Strands.BulkData);
-		GroupData.Guides.Data.Serialize(Ar, GroupData.Guides.BulkData);
+		GroupData.Strands.BulkData.Serialize(Ar);
+		GroupData.Guides.BulkData.Serialize(Ar);
 		GroupData.Strands.InterpolationBulkData.Serialize(Ar);
 
 	}
@@ -1149,8 +1157,8 @@ FArchive& operator<<(FArchive& Ar, FHairGroupData& GroupData)
 		// Fall back to no data, but still serialize guide data as they an be used for cards simulation
 		// Theoritically, we should have something to detect if we are going to use or not guide (for 
 		// simulation or RBF deformation) on the target platform
-		NoStrandsData.Strands.Data.Serialize(Ar, NoStrandsData.Strands.BulkData);
-		GroupData.Guides.Data.Serialize(Ar, GroupData.Guides.BulkData);
+		NoStrandsData.Strands.BulkData.Serialize(Ar);
+		GroupData.Guides.BulkData.Serialize(Ar);
 		NoStrandsData.Strands.InterpolationBulkData.Serialize(Ar);
 	}
 
@@ -1329,7 +1337,7 @@ void UGroomAsset::SetHairWidth(float Width)
 // differences, etc.) replace the version GUID below with a new one.
 // In case of merge conflicts with DDC versions, you MUST generate a new GUID
 // and set this new GUID as the version.
-#define GROOM_DERIVED_DATA_VERSION TEXT("B136B9EEDA5C4FABAF263AD0B05C1D76")
+#define GROOM_DERIVED_DATA_VERSION TEXT("22D18A66610F485D8AF4B7DB8DDA08E0")
 
 #if WITH_EDITORONLY_DATA
 
@@ -1841,13 +1849,8 @@ bool UGroomAsset::CacheCardsGeometry(uint32 GroupIndex, const FString& StrandsKe
 			if (FHairGroupsCardsSourceDescription* Desc = GetSourceDescription(HairGroupsCards, GroupIndex, LODIt, SourceIt))
 			{
 				FHairGroupData::FCards::FLOD& LOD = GroupData.Cards.LODs[LODIt];
-
-				GroupData.Strands.Data = LOD.Guides.Data;
-				GroupData.Guides.Data = LOD.Guides.Data;
-
 				GroupData.Strands.BulkData = LOD.Guides.BulkData;
 				GroupData.Guides.BulkData = LOD.Guides.BulkData;
-
 				break;
 			}
 		}
@@ -2025,10 +2028,11 @@ bool UGroomAsset::BuildCardsGeometry(uint32 GroupIndex)
 			}
 
 			bool bInitResources = false;
+			FHairStrandsDatas LODGuidesData;
 			if (CardsMesh != nullptr)
 			{
 				CardsMesh->ConditionalPostLoad();
-				bInitResources = FHairCardsBuilder::ImportGeometry(CardsMesh, LOD.BulkData, LOD.Guides.Data, LOD.InterpolationBulkData);
+				bInitResources = FHairCardsBuilder::ImportGeometry(CardsMesh, LOD.BulkData, LODGuidesData, LOD.InterpolationBulkData);
 				if (!bInitResources)
 				{
 					UE_LOG(LogHairStrands, Warning, TEXT("Failed to import cards from %s for Group %d LOD %d."), *CardsMesh->GetName(), GroupIndex, LODIt);
@@ -2073,21 +2077,28 @@ bool UGroomAsset::BuildCardsGeometry(uint32 GroupIndex)
 				bool bCopyRenderData = false;
 				if (GroupData.Guides.BulkData.GetNumPoints() == 0)
 				{
-					GroupData.Strands.Data = LOD.Guides.Data;
-					GroupData.Guides.Data  = LOD.Guides.Data;
-
 					// The RenderData is filled out by BuildData
 					bCopyRenderData = true;
 				}
 
 				FHairGroupInfo DummyGroupInfo;
 
-				FGroomBuilder::BuildData(LOD.Guides.Data);
-				FGroomBuilder::BuildBulkData(DummyGroupInfo, LOD.Guides.Data, LOD.Guides.BulkData);
+				// 1. Build data & bulk data for the cards guides
+				FGroomBuilder::BuildData(LODGuidesData);
+				FGroomBuilder::BuildBulkData(DummyGroupInfo, LODGuidesData, LOD.Guides.BulkData);
 
+				// 2. (Re)Build data for the sim guides (since there are transient/no-cached)
+				const FHairDescriptionGroups& LocalHairDescriptionGroups = GetHairDescriptionGroups();
+				FHairGroupInfo DummyInfo;
+				check(LocalHairDescriptionGroups.IsValid());
+				FHairStrandsDatas SimGuidesData;
+				FHairStrandsDatas RenStrandsData; // Dummy, not used
+				FGroomBuilder::BuildData(LocalHairDescriptionGroups.HairGroups[GroupIndex], HairGroupsInterpolation[GroupIndex], DummyInfo, RenStrandsData, SimGuidesData);
+
+				// 3. Compute interpolation data & bulk data between the sim guides and the cards guides
 				FHairStrandsInterpolationDatas LODGuideInterpolationData;
-				FGroomBuilder::BuildInterplationData(DummyGroupInfo, LOD.Guides.Data, GroupData.Guides.Data, CardsInterpolationSettings, LODGuideInterpolationData);
-				FGroomBuilder::BuildInterplationBulkData(GroupData.Guides.Data, LODGuideInterpolationData, LOD.Guides.InterpolationBulkData);
+				FGroomBuilder::BuildInterplationData(DummyGroupInfo, LODGuidesData, SimGuidesData, CardsInterpolationSettings, LODGuideInterpolationData);
+				FGroomBuilder::BuildInterplationBulkData(SimGuidesData, LODGuideInterpolationData, LOD.Guides.InterpolationBulkData);
 
 				if (bCopyRenderData)
 				{
@@ -2098,7 +2109,6 @@ bool UGroomAsset::BuildCardsGeometry(uint32 GroupIndex)
 				LOD.Guides.RestResource = new FHairStrandsRestResource(LOD.Guides.BulkData, EHairStrandsResourcesType::Cards);
 				BeginInitResource(LOD.Guides.RestResource);
 
-				PatchHairStrandsInterpolationBulkData(LOD.Guides.InterpolationBulkData, GroupData.Guides.Data);
 				LOD.Guides.InterpolationResource = new FHairStrandsInterpolationResource(LOD.Guides.InterpolationBulkData);
 				BeginInitResource(LOD.Guides.InterpolationResource);
 
@@ -2258,8 +2268,7 @@ bool UGroomAsset::BuildMeshesGeometry(uint32 GroupIndex)
 			{
 				// Build a default box
 				FHairMeshesBuilder::BuildGeometry(
-					GroupData.Strands.Data,
-					GroupData.Guides.Data,
+					GroupData.Strands.BulkData.BoundingBox,
 					LOD.BulkData);
 			}
 
@@ -2329,7 +2338,6 @@ FHairStrandsInterpolationResource* UGroomAsset::AllocateInterpolationResources(u
 		check(GroupData.Guides.IsValid());
 		if (GroupData.Strands.InterpolationResource == nullptr)
 		{
-			PatchHairStrandsInterpolationBulkData(GroupData.Strands.InterpolationBulkData, GroupData.Guides.Data);
 			GroupData.Strands.InterpolationResource = new FHairStrandsInterpolationResource(GroupData.Strands.InterpolationBulkData);
 			BeginInitResource(GroupData.Strands.InterpolationResource);
 		}
@@ -2515,7 +2523,6 @@ void UGroomAsset::InitCardsResources()
 				LOD.Guides.RestResource = new FHairStrandsRestResource(LOD.Guides.BulkData, EHairStrandsResourcesType::Cards);
 				BeginInitResource(LOD.Guides.RestResource);
 
-				PatchHairStrandsInterpolationBulkData(LOD.Guides.InterpolationBulkData, GroupData.Guides.Data);
 				LOD.Guides.InterpolationResource = new FHairStrandsInterpolationResource(LOD.Guides.InterpolationBulkData);
 				BeginInitResource(LOD.Guides.InterpolationResource);
 
@@ -2722,11 +2729,19 @@ void UGroomAsset::StripLODs(const TArray<int32>& LODsToKeep, bool bRebuildResour
 	// Rebuild the LOD data
 	if (bRebuildResources)
 	{
+		const FHairDescriptionGroups& LocalHairDescriptionGroups = GetHairDescriptionGroups();
 		for (int32 GroupIt = 0; GroupIt < GroupCount; ++GroupIt)
 		{
 			FHairGroupData& GroupData = HairGroupsData[GroupIt];
 			InternalReleaseResource(GroupData.Strands.ClusterCullingResource);
-			FGroomBuilder::BuildClusterData(GroupData.Strands.Data, HairDescriptionGroups->BoundRadius, HairGroupsLOD[GroupIt], GroupData.Strands.ClusterCullingData);
+
+			BuildHairGroupCluster(
+				GroupIt,
+				LocalHairDescriptionGroups,
+				HairGroupsInterpolation,
+				HairGroupsLOD,
+				HairGroupsInfo,
+				HairGroupsData);
 		}
 	}
 }
@@ -2752,8 +2767,15 @@ void UGroomAsset::CreateDebugData()
 
 	for (uint32 GroupIndex = 0, GroupCount = GetNumHairGroups(); GroupIndex < GroupCount; ++GroupIndex)
 	{
+		// 2.1 Generate transient strands data
+		const FHairDescriptionGroups& LocalHairDescriptionGroups = GetHairDescriptionGroups();
+		FHairGroupInfo DummyInfo;
+		FHairStrandsDatas StrandsData;
+		FHairStrandsDatas GuidesData;
+		FGroomBuilder::BuildData(LocalHairDescriptionGroups.HairGroups[GroupIndex], HairGroupsInterpolation[GroupIndex], DummyInfo, StrandsData, GuidesData);
+
 		FHairGroupData& GroupData = HairGroupsData[GroupIndex];
-		CreateHairStrandsDebugDatas(GroupData.Strands.Data, 1.f, GroupData.Debug.Data);
+		CreateHairStrandsDebugDatas(StrandsData, 1.f, GroupData.Debug.Data);
 
 		if (GroupData.Debug.Data.IsValid())
 		{
@@ -2875,7 +2897,7 @@ void UGroomAsset::SaveProceduralCards(uint32 DescIndex)
 		const bool bNeedConversion = ClusterSettings.ClusterDecimation > 0;
 		if (bNeedConversion)
 		{
-			const int32 MaxCardCount = HairGroupsData[GroupIndex].Strands.Data.GetNumCurves();
+			const int32 MaxCardCount = HairGroupsData[GroupIndex].Strands.BulkData.GetNumCurves();
 
 			GeometrySettings.GenerationType = ClusterSettings.bUseGuide ? EHairCardsGenerationType::UseGuides : EHairCardsGenerationType::CardsCount;
 			GeometrySettings.CardsCount = FMath::Clamp(FMath::CeilToInt(ClusterSettings.ClusterDecimation * MaxCardCount), 1, MaxCardCount);
@@ -2892,19 +2914,27 @@ void UGroomAsset::SaveProceduralCards(uint32 DescIndex)
 
 	FHairProceduralCardsQuery& Q = *QP;
 	Q.Asset = this;
-
-	const FHairGroupData& GroupData = HairGroupsData[GroupIndex];
-	FHairCardsBuilder::BuildGeometry(
-		GetLODName(this, Desc->LODIndex),
-		GroupData.Strands.Data,
-		GroupData.Guides.Data,
-		Desc->ProceduralSettings,
-		Q.ProceduralData,
-		Q.BulkData,
-		Q.GuideData,
-		Q.InterpolationBulkData,
-		Desc->Textures);
-	Q.Textures = &Desc->Textures;
+	{
+		// 2.1 Generate transient strands data
+		const FHairDescriptionGroups& LocalHairDescriptionGroups = GetHairDescriptionGroups();
+		FHairGroupInfo DummyInfo;
+		FHairStrandsDatas StrandsData;
+		FHairStrandsDatas GuidesData;
+		FGroomBuilder::BuildData(LocalHairDescriptionGroups.HairGroups[GroupIndex], HairGroupsInterpolation[GroupIndex], DummyInfo, StrandsData, GuidesData);
+
+		// 2.2 Build cards geometry
+		FHairCardsBuilder::BuildGeometry(
+			GetLODName(this, Desc->LODIndex),
+			StrandsData,
+			GuidesData,
+			Desc->ProceduralSettings,
+			Q.ProceduralData,
+			Q.BulkData,
+			Q.GuideData,
+			Q.InterpolationBulkData,
+			Desc->Textures);
+		Q.Textures = &Desc->Textures;
+	}
 
 	// 3. Create resources and enqueue texture generation (GPU, kicked by the render thread) 
 	Q.Resources = new FHairCardsRestResource(Q.BulkData);
@@ -2969,6 +2999,66 @@ void UGroomAsset::SavePendingProceduralAssets()
 }
 #endif // WITH_EDITOR
 
+#if WITH_EDITORONLY_DATA
+bool UGroomAsset::GetHairStrandsDatas(
+	const int32 GroupIndex,
+	FHairStrandsDatas& OutStrandsData,
+	FHairStrandsDatas& OutGuidesData)
+{
+	const FHairDescriptionGroups& DescriptionGroups = GetHairDescriptionGroups();
+	const bool bIsValid = DescriptionGroups.IsValid() && GroupIndex < DescriptionGroups.HairGroups.Num();
+	if (bIsValid)
+	{
+		const FHairDescriptionGroup& HairGroup = DescriptionGroups.HairGroups[GroupIndex];
+		check(GroupIndex == HairGroup.Info.GroupID);
+
+		FHairGroupInfo DummyInfo;
+		FGroomBuilder::BuildData(HairGroup, HairGroupsInterpolation[GroupIndex], DummyInfo, OutStrandsData, OutGuidesData);
+	}
+	return bIsValid;
+}
+
+bool UGroomAsset::GetHairCardsGuidesDatas(
+	const int32 GroupIndex,
+	const int32 LODIndex,
+	FHairStrandsDatas& OutCardsGuidesData)
+{
+	FHairStrandsDatas StrandsData;
+	FHairStrandsDatas GuidesData;
+	const bool bIsValid = GetHairStrandsDatas(GroupIndex, StrandsData, GuidesData);
+	if (!bIsValid)
+	{
+		return false;
+	}
+
+	int32 SourceIt = 0;
+	if (FHairGroupsCardsSourceDescription* Desc = GetSourceDescription(HairGroupsCards, GroupIndex, LODIndex, SourceIt))
+	{
+		// 1. Load geometry data, if any
+		UStaticMesh* CardsMesh = nullptr;
+		if (Desc->SourceType == EHairCardsSourceType::Procedural)
+		{
+			CardsMesh = Desc->ProceduralMesh;
+		}
+		else if (Desc->SourceType == EHairCardsSourceType::Imported)
+		{
+			CardsMesh = Desc->ImportedMesh;
+		}
+
+		if (CardsMesh)
+		{
+			CardsMesh->ConditionalPostLoad();
+
+			FHairCardsBulkData					DummyBulkData;
+			FHairCardsInterpolationBulkData		DummyInterpolationBulkData;
+			FHairCardsBuilder::ImportGeometry(CardsMesh, DummyBulkData, OutCardsGuidesData, DummyInterpolationBulkData);
+			return true;
+		}
+	}
+	return false;
+}
+#endif // WITH_EDITORONLY_DATA
+
 void UGroomAsset::GetResourceSizeEx(FResourceSizeEx& CumulativeResourceSize)
 {
 	CumulativeResourceSize.AddDedicatedSystemMemoryBytes(HairGroupsRendering.GetAllocatedSize());
diff --git a/Engine/Plugins/Runtime/HairStrands/Source/HairStrandsCore/Private/GroomBindingBuilder.cpp b/Engine/Plugins/Runtime/HairStrands/Source/HairStrandsCore/Private/GroomBindingBuilder.cpp
--- a/Engine/Plugins/Runtime/HairStrands/Source/HairStrandsCore/Private/GroomBindingBuilder.cpp
+++ b/Engine/Plugins/Runtime/HairStrands/Source/HairStrandsCore/Private/GroomBindingBuilder.cpp
@@ -1835,13 +1835,18 @@ static bool InternalBuildBinding_CPU(UGroomBindingAsset* BindingAsset, bool bIni
 	UGroomBindingAsset::FHairGroupDatas& OutHairGroupDatas = BindingAsset->HairGroupDatas;
 	OutHairGroupDatas.Empty();
 
+	uint32 GroupIndex = 0;
 	TArray<uint32> NumSamples;
 	NumSamples.Init(NumInterpolationPoints, MeshLODCount);
 	for (const FHairGroupData& GroupData : GroomAsset->HairGroupsData)
 	{
+		FHairStrandsDatas StrandsData;
+		FHairStrandsDatas GuidesData;
+		GroomAsset->GetHairStrandsDatas(GroupIndex, StrandsData, GuidesData);
+
 		UGroomBindingAsset::FHairGroupData& Data = OutHairGroupDatas.AddDefaulted_GetRef();
-		Data.RenRootData = FHairStrandsRootData(&GroupData.Strands.Data, MeshLODCount, NumSamples);
-		Data.SimRootData = FHairStrandsRootData(&GroupData.Guides.Data, MeshLODCount, NumSamples);
+		Data.RenRootData = FHairStrandsRootData(&StrandsData, MeshLODCount, NumSamples);
+		Data.SimRootData = FHairStrandsRootData(&GuidesData, MeshLODCount, NumSamples);
 
 		const uint32 CardsLODCount = GroupData.Cards.LODs.Num();
 		Data.CardsRootData.SetNum(GroupData.Cards.LODs.Num());
@@ -1849,9 +1854,13 @@ static bool InternalBuildBinding_CPU(UGroomBindingAsset* BindingAsset, bool bIni
 		{
 			if (GroupData.Cards.IsValid(CardsLODIt))
 			{
-				Data.CardsRootData[CardsLODIt] = FHairStrandsRootData(&GroupData.Cards.LODs[CardsLODIt].Guides.Data, MeshLODCount, NumSamples);
+				FHairStrandsDatas LODGuidesData;
+				const bool bIsValid = GroomAsset->GetHairCardsGuidesDatas(GroupIndex, CardsLODIt, LODGuidesData);
+				check(bIsValid);
+				Data.CardsRootData[CardsLODIt] = FHairStrandsRootData(&LODGuidesData, MeshLODCount, NumSamples);
 			}
 		}
+		++GroupIndex;
 	}
 
 	UGroomBindingAsset::FHairGroupResources& OutHairGroupResources = BindingAsset->HairGroupResources;
@@ -1910,8 +1919,12 @@ static bool InternalBuildBinding_CPU(UGroomBindingAsset* BindingAsset, bool bIni
 	bool bSucceed = false;
 	for (uint32 GroupIt=0; GroupIt < GroupCount; ++GroupIt)
 	{
+		FHairStrandsDatas StrandsData;
+		FHairStrandsDatas GuidesData;
+		GroomAsset->GetHairStrandsDatas(GroupIt, StrandsData, GuidesData);
+
 		bSucceed = GroomBinding_RootProjection::Project(
-			BindingAsset->Groom->HairGroupsData[GroupIt].Strands.Data,
+			StrandsData,
 			TargetMeshData.Get(),
 			TransferredPositions,
 			BindingAsset->HairGroupDatas[GroupIt].RenRootData);
@@ -1924,7 +1937,7 @@ static bool InternalBuildBinding_CPU(UGroomBindingAsset* BindingAsset, bool bIni
 		SlowTask.EnterProgressFrame();
 
 		bSucceed = GroomBinding_RootProjection::Project(
-			BindingAsset->Groom->HairGroupsData[GroupIt].Guides.Data,
+			GuidesData,
 			TargetMeshData.Get(),
 			TransferredPositions,
 			BindingAsset->HairGroupDatas[GroupIt].SimRootData);
@@ -1941,11 +1954,16 @@ static bool InternalBuildBinding_CPU(UGroomBindingAsset* BindingAsset, bool bIni
 		{
 			if (BindingAsset->Groom->HairGroupsData[GroupIt].Cards.IsValid(CardsLODIt))
 			{
+				FHairStrandsDatas LODGuidesData;
+				const bool bIsValid = GroomAsset->GetHairCardsGuidesDatas(GroupIt, CardsLODIt, LODGuidesData);
+				check(bIsValid);
+				
 				bSucceed = GroomBinding_RootProjection::Project(
-					BindingAsset->Groom->HairGroupsData[GroupIt].Cards.LODs[CardsLODIt].Guides.Data,
+					LODGuidesData,
 					TargetMeshData.Get(),
 					TransferredPositions,
 					BindingAsset->HairGroupDatas[GroupIt].CardsRootData[CardsLODIt]);
+
 				if (!bSucceed) 
 				{
 					UE_LOG(LogHairStrands, Error, TEXT("[Groom] Binding asset could not be built. Some cards guide roots are not close enough to the target mesh to be projected onto it."));
diff --git a/Engine/Plugins/Runtime/HairStrands/Source/HairStrandsCore/Private/GroomBuilder.cpp b/Engine/Plugins/Runtime/HairStrands/Source/HairStrandsCore/Private/GroomBuilder.cpp
--- a/Engine/Plugins/Runtime/HairStrands/Source/HairStrandsCore/Private/GroomBuilder.cpp
+++ b/Engine/Plugins/Runtime/HairStrands/Source/HairStrandsCore/Private/GroomBuilder.cpp
@@ -33,7 +33,7 @@ static FAutoConsoleVariableRef CVarHairClusterBuilder_MaxVoxelResolution(TEXT("r
 FString FGroomBuilder::GetVersion()
 {
 	// Important to update the version when groom building changes
-	return TEXT("1hc");
+	return TEXT("1i");
 }
 
 namespace FHairStrandsDecimation
diff --git a/Engine/Plugins/Runtime/HairStrands/Source/HairStrandsCore/Private/GroomRBFDeformer.cpp b/Engine/Plugins/Runtime/HairStrands/Source/HairStrandsCore/Private/GroomRBFDeformer.cpp
--- a/Engine/Plugins/Runtime/HairStrands/Source/HairStrandsCore/Private/GroomRBFDeformer.cpp
+++ b/Engine/Plugins/Runtime/HairStrands/Source/HairStrandsCore/Private/GroomRBFDeformer.cpp
@@ -396,16 +396,19 @@ static void ApplyDeformationToGroom(const TArray<FRBFDeformedPositions>& Deforme
 				FHairGroupInfo& HairGroupsInfo = GroomAsset->HairGroupsInfo[GroupIndex];
 				FHairGroupData& HairGroupsData = GroomAsset->HairGroupsData[GroupIndex];
 
-				FGroomBuilder::BuildData(HairGroup, GroomAsset->HairGroupsInterpolation[GroupIndex], HairGroupsInfo, HairGroupsData.Strands.Data, HairGroupsData.Guides.Data);
+				FHairStrandsDatas StrandsData;
+				FHairStrandsDatas GuidesData;
+				FGroomBuilder::BuildData(HairGroup, GroomAsset->HairGroupsInterpolation[GroupIndex], HairGroupsInfo, StrandsData, GuidesData);
+				//GroomAsset->GetHairStrandsDatas(GroupIndex, StrandsData, GuidesData);
 
-				FGroomBuilder::BuildBulkData(HairGroup.Info, HairGroupsData.Guides.Data,  HairGroupsData.Guides.BulkData);
-				FGroomBuilder::BuildBulkData(HairGroup.Info, HairGroupsData.Strands.Data, HairGroupsData.Strands.BulkData);
+				FGroomBuilder::BuildBulkData(HairGroup.Info, GuidesData,  HairGroupsData.Guides.BulkData);
+				FGroomBuilder::BuildBulkData(HairGroup.Info, StrandsData, HairGroupsData.Strands.BulkData);
 
 				FHairStrandsInterpolationDatas StrandsInterpolationData;
-				FGroomBuilder::BuildInterplationData(HairGroup.Info, HairGroupsData.Strands.Data, HairGroupsData.Guides.Data, InterpolationSettings, StrandsInterpolationData);
-				FGroomBuilder::BuildInterplationBulkData(HairGroupsData.Guides.Data, StrandsInterpolationData, HairGroupsData.Strands.InterpolationBulkData);
+				FGroomBuilder::BuildInterplationData(HairGroup.Info, StrandsData, GuidesData, InterpolationSettings, StrandsInterpolationData);
+				FGroomBuilder::BuildInterplationBulkData(GuidesData, StrandsInterpolationData, HairGroupsData.Strands.InterpolationBulkData);
 
-				FGroomBuilder::BuildClusterData(HairGroupsData.Strands.Data, HairDescriptionGroups.BoundRadius, HairGroupLOD, HairGroupsData.Strands.ClusterCullingData);
+				FGroomBuilder::BuildClusterData(StrandsData, HairDescriptionGroups.BoundRadius, HairGroupLOD, HairGroupsData.Strands.ClusterCullingData);
 			}
 		}
 	}
@@ -521,6 +524,11 @@ void FGroomRBFDeformer::GetRBFDeformedGroomAsset(const UGroomAsset* InGroomAsset
 			const FHairStrandsDatas& OriginalGuides = Group.Guides;
 			const FHairGroupData& HairGroupData = InGroomAsset->HairGroupsData[GroupIndex];
 
+			FHairStrandsDatas StrandsData;
+			FHairStrandsDatas GuidesData;
+			FHairGroupInfo DummyInfo;
+			FGroomBuilder::BuildData(Group, InGroomAsset->HairGroupsInterpolation[GroupIndex], DummyInfo, StrandsData, GuidesData);
+
 			// Get deformed guides
 			// If the groom override the value, we output dummy value for the guides, since they won't be used
 			if (InGroomAsset->HairGroupsInterpolation[GroupIndex].InterpolationSettings.bOverrideGuides)
@@ -532,7 +540,7 @@ void FGroomRBFDeformer::GetRBFDeformedGroomAsset(const UGroomAsset* InGroomAsset
 			{
 				DeformedPositions[GroupIndex].GuideStrands = GetDeformedHairStrandsPositions(
 					MeshVertexPositionsBuffer_Target,
-					HairGroupData.Guides.Data,
+					GuidesData,
 					MeshLODIndex,
 					SkeletalMeshData_Target,
 					BindingAsset->HairGroupDatas[GroupIndex].SimRootData.VertexToCurveIndexBuffer,
@@ -542,7 +550,7 @@ void FGroomRBFDeformer::GetRBFDeformedGroomAsset(const UGroomAsset* InGroomAsset
 			// Get deformed render strands
 			DeformedPositions[GroupIndex].RenderStrands = GetDeformedHairStrandsPositions(
 				MeshVertexPositionsBuffer_Target,
-				HairGroupData.Strands.Data,
+				StrandsData,
 				MeshLODIndex,
 				SkeletalMeshData_Target,
 				BindingAsset->HairGroupDatas[GroupIndex].RenRootData.VertexToCurveIndexBuffer,
diff --git a/Engine/Plugins/Runtime/HairStrands/Source/HairStrandsCore/Private/GroomTextureBuilder.cpp b/Engine/Plugins/Runtime/HairStrands/Source/HairStrandsCore/Private/GroomTextureBuilder.cpp
--- a/Engine/Plugins/Runtime/HairStrands/Source/HairStrandsCore/Private/GroomTextureBuilder.cpp
+++ b/Engine/Plugins/Runtime/HairStrands/Source/HairStrandsCore/Private/GroomTextureBuilder.cpp
@@ -258,9 +258,15 @@ static void InternalBuildFollicleTexture_CPU(const TArray<FFollicleInfo>& InInfo
 			case FFollicleInfo::A: Channel = 3; break;
 		}
 
+		uint32 GroupIndex = 0;
 		for (const FHairGroupData& HairGroupData : Info.GroomAsset->HairGroupsData)
 		{
-			RasterToTexture(Resolution.X, Info.KernelSizeInPixels, Channel, HairGroupData.Strands.Data, Pixels);
+			FHairStrandsDatas StrandsData;
+			FHairStrandsDatas GuidesData;
+			Info.GroomAsset->GetHairStrandsDatas(GroupIndex, StrandsData, GuidesData);
+
+			RasterToTexture(Resolution.X, Info.KernelSizeInPixels, Channel, StrandsData, Pixels);
+			++GroupIndex;
 		}
 	}
 	OutTexture->Source.UnlockMip(0);
@@ -297,10 +303,15 @@ static void InternalBuildFollicleTexture_GPU(
 		}
 
 		// Create root UVs buffer
+		uint32 GroupIndex = 0;
 		for (const FHairGroupData& GroupData : Info.GroomAsset->HairGroupsData)
 		{
-			GroupData.Strands.Data.StrandsCurves.CurvesRootUV.GetData();
-			const uint32 DataCount = GroupData.Strands.Data.StrandsCurves.CurvesRootUV.Num();
+			FHairStrandsDatas StrandsData;
+			FHairStrandsDatas GuidesData;
+			Info.GroomAsset->GetHairStrandsDatas(GroupIndex, StrandsData, GuidesData);
+
+			StrandsData.StrandsCurves.CurvesRootUV.GetData();
+			const uint32 DataCount = StrandsData.StrandsCurves.CurvesRootUV.Num();
 			const uint32 DataSizeInBytes = sizeof(FVector2D) * DataCount;
 			check(DataSizeInBytes != 0);
 
@@ -309,10 +320,12 @@ static void InternalBuildFollicleTexture_GPU(
 				GraphBuilder, 
 				TEXT("RootUVBuffer"),
 				Desc,
-				GroupData.Strands.Data.StrandsCurves.CurvesRootUV.GetData(),
+				StrandsData.StrandsCurves.CurvesRootUV.GetData(),
 				DataSizeInBytes,
 				ERDGInitialDataFlags::None);
 			RootUVBuffers[Info.Channel].Add(RootBuffer);
+
+			GroupIndex++;
 		}
 	}
 
diff --git a/Engine/Plugins/Runtime/HairStrands/Source/HairStrandsCore/Private/HairCardsBuilder.cpp b/Engine/Plugins/Runtime/HairStrands/Source/HairStrandsCore/Private/HairCardsBuilder.cpp
--- a/Engine/Plugins/Runtime/HairStrands/Source/HairStrandsCore/Private/HairCardsBuilder.cpp
+++ b/Engine/Plugins/Runtime/HairStrands/Source/HairStrandsCore/Private/HairCardsBuilder.cpp
@@ -3927,12 +3927,11 @@ FString GetVersion()
 }
 
 void BuildGeometry(
-	const FHairStrandsDatas& InRen,
-	const FHairStrandsDatas& InSim,
+	const FBox& InBox,
 	FHairMeshesBulkData& OutBulk)
 {
-	const FVector Center = InRen.BoundingBox.GetCenter();
-	const FVector Extent = InRen.BoundingBox.GetExtent();
+	const FVector Center = InBox.GetCenter();
+	const FVector Extent = InBox.GetExtent();
 
 	// Simple (incorrect normal/tangent) cube geomtry in place of the hair rendering
 	const uint32 TotalPointCount = 8;
diff --git a/Engine/Plugins/Runtime/HairStrands/Source/HairStrandsCore/Private/HairStrandsDatas.cpp b/Engine/Plugins/Runtime/HairStrands/Source/HairStrandsCore/Private/HairStrandsDatas.cpp
--- a/Engine/Plugins/Runtime/HairStrands/Source/HairStrandsCore/Private/HairStrandsDatas.cpp
+++ b/Engine/Plugins/Runtime/HairStrands/Source/HairStrandsCore/Private/HairStrandsDatas.cpp
@@ -177,37 +177,6 @@ void FHairStrandsInterpolationBulkData::Serialize(FArchive& Ar)
 	Ar << SimRootPointIndex;
 }
 
-void FHairStrandsPoints::Serialize(FArchive& Ar)
-{
-	Ar.UsingCustomVersion(FReleaseObjectVersion::GUID);
-
-	Ar << PointsPosition;
-	Ar << PointsRadius;
-	Ar << PointsCoordU;
-
-	if (Ar.CustomVer(FReleaseObjectVersion::GUID) >= FReleaseObjectVersion::GroomAssetVersion2)
-	{
-		Ar << PointsBaseColor;
-		Ar << PointsRoughness;
-	}
-	else
-	{
-		const uint32 ElementCount = PointsPosition.Num();
-		PointsBaseColor.InsertZeroed(0, ElementCount);
-		PointsRoughness.InsertZeroed(0, ElementCount);
-	}
-}
-
-void FHairStrandsCurves::Serialize(FArchive& Ar)
-{
-	Ar << CurvesCount;
-	Ar << CurvesOffset;
-	Ar << CurvesLength;
-	Ar << CurvesRootUV;
-	Ar << MaxLength;
-	Ar << MaxRadius;
-}
-
 void FHairStrandsBulkData::Serialize(FArchive& Ar)
 {
 	Ar.UsingCustomVersion(FReleaseObjectVersion::GUID);
@@ -217,25 +186,13 @@ void FHairStrandsBulkData::Serialize(FArchive& Ar)
 	Ar << Attributes;
 	Ar << Materials;
 	Ar << CurveOffsets;
-}
 
-void FHairStrandsDatas::Serialize(FArchive& Ar, FHairStrandsBulkData& BulkData)
-{
-	StrandsPoints.Serialize(Ar);
-	StrandsCurves.Serialize(Ar);
-	Ar << HairDensity;
+	Ar << CurveCount;
+	Ar << PointCount;
+	Ar << MaxLength;
+	Ar << MaxRadius;
 	Ar << BoundingBox;
 
-	BulkData.Serialize(Ar);
-
-	if (Ar.IsLoading())
-	{
-		BulkData.CurveCount = StrandsCurves.Num();
-		BulkData.PointCount = StrandsPoints.Num();
-		BulkData.MaxLength  = StrandsCurves.MaxLength;
-		BulkData.MaxRadius  = StrandsCurves.MaxRadius;
-		BulkData.BoundingBox= BoundingBox;
-	}
 }
 
 void FHairStrandsDatas::Reset()
diff --git a/Engine/Plugins/Runtime/HairStrands/Source/HairStrandsCore/Public/GroomAsset.h b/Engine/Plugins/Runtime/HairStrands/Source/HairStrandsCore/Public/GroomAsset.h
--- a/Engine/Plugins/Runtime/HairStrands/Source/HairStrandsCore/Public/GroomAsset.h
+++ b/Engine/Plugins/Runtime/HairStrands/Source/HairStrandsCore/Public/GroomAsset.h
@@ -81,7 +81,6 @@ struct HAIRSTRANDSCORE_API FHairGroupData
 		bool IsValid() const			{ return RestResource != nullptr; }
 		const FBox& GetBounds() const	{ return BulkData.GetBounds(); }
 
-		FHairStrandsDatas					Data;
 		FHairStrandsBulkData				BulkData;
 		FHairStrandsRestResource*			RestResource = nullptr;
 	};
@@ -515,6 +514,10 @@ public:
 	void CommitHairDescription(FHairDescription&& HairDescription);
 	FHairDescription GetHairDescription() const;
 
+	/** Get/Build render & guides data based on the hair description and interpolation settings */
+	bool GetHairStrandsDatas(const int32 GroupIndex, FHairStrandsDatas& OutStrandsData, FHairStrandsDatas& OutGuidesData);
+	bool GetHairCardsGuidesDatas(const int32 GroupIndex, const int32 LODIndex, FHairStrandsDatas& OutCardsGuidesData);
+
 	/** Caches the computed (group) groom data with the given build settings from/to the Derived Data Cache, building it if needed.
 	 *  This function assumes the interpolation settings are properly populated, as they will be used to build the asset.
 	 */
diff --git a/Engine/Plugins/Runtime/HairStrands/Source/HairStrandsCore/Public/HairCardsBuilder.h b/Engine/Plugins/Runtime/HairStrands/Source/HairStrandsCore/Public/HairCardsBuilder.h
--- a/Engine/Plugins/Runtime/HairStrands/Source/HairStrandsCore/Public/HairCardsBuilder.h
+++ b/Engine/Plugins/Runtime/HairStrands/Source/HairStrandsCore/Public/HairCardsBuilder.h
@@ -54,8 +54,7 @@ namespace FHairCardsBuilder
 namespace FHairMeshesBuilder
 {
 	void BuildGeometry(
-		const struct FHairStrandsDatas& In,
-		const struct FHairStrandsDatas& InSim,
+		const FBox& InBox,
 		FHairMeshesBulkData& OutBulk);
 
 	void ImportGeometry(
diff --git a/Engine/Plugins/Runtime/HairStrands/Source/HairStrandsCore/Public/HairStrandsDatas.h b/Engine/Plugins/Runtime/HairStrands/Source/HairStrandsCore/Public/HairStrandsDatas.h
--- a/Engine/Plugins/Runtime/HairStrands/Source/HairStrandsCore/Public/HairStrandsDatas.h
+++ b/Engine/Plugins/Runtime/HairStrands/Source/HairStrandsCore/Public/HairStrandsDatas.h
@@ -296,9 +296,6 @@ struct HAIRSTRANDSCORE_API FHairStrandsInterpolationBulkData
 /** Hair strands points attribute */
 struct HAIRSTRANDSCORE_API FHairStrandsPoints
 {
-	/** Serialize the points */
-	void Serialize(FArchive& Ar);
-
 	/** Set the number of points */
 	void SetNum(const uint32 NumPoints);
 
@@ -327,9 +324,6 @@ struct HAIRSTRANDSCORE_API FHairStrandsPoints
 /** Hair strands Curves attribute */
 struct HAIRSTRANDSCORE_API FHairStrandsCurves
 {
-	/** Serialize the curves */
-	void Serialize(FArchive& Ar);
-
 	/** Set the number of Curves */
 	void SetNum(const uint32 NumPoints);
 
@@ -369,9 +363,6 @@ struct FHairStrandsBulkData;
 /** Hair strands datas that are stored on CPU */
 struct HAIRSTRANDSCORE_API FHairStrandsDatas
 {
-	/** Serialize all the hair strands datas */
-	void Serialize(FArchive& Ar, FHairStrandsBulkData& BulkData);
-
 	/* Get the total number of points */
 	uint32 GetNumPoints() const { return StrandsPoints.Num(); }
 
