Fix CIS issue.

diff --git a/Engine/Plugins/Runtime/HairStrands/Source/HairStrandsCore/Private/GroomAsset.cpp b/Engine/Plugins/Runtime/HairStrands/Source/HairStrandsCore/Private/GroomAsset.cpp
--- a/Engine/Plugins/Runtime/HairStrands/Source/HairStrandsCore/Private/GroomAsset.cpp
+++ b/Engine/Plugins/Runtime/HairStrands/Source/HairStrandsCore/Private/GroomAsset.cpp
@@ -2762,6 +2762,7 @@ bool UGroomAsset::HasDebugData() const
 
 void UGroomAsset::CreateDebugData()
 {
+#if WITH_EDITORONLY_DATA
 	if (!IsHairStrandsEnabled(EHairStrandsShaderType::Strands))
 		return;
 
@@ -2792,6 +2793,7 @@ void UGroomAsset::CreateDebugData()
 				});
 		}
 	}
+#endif
 }
 
 int32 UGroomAsset::GetMaterialIndex(FName MaterialSlotName) const
diff --git a/Engine/Plugins/Runtime/HairStrands/Source/HairStrandsCore/Private/GroomTextureBuilder.cpp b/Engine/Plugins/Runtime/HairStrands/Source/HairStrandsCore/Private/GroomTextureBuilder.cpp
--- a/Engine/Plugins/Runtime/HairStrands/Source/HairStrandsCore/Private/GroomTextureBuilder.cpp
+++ b/Engine/Plugins/Runtime/HairStrands/Source/HairStrandsCore/Private/GroomTextureBuilder.cpp
@@ -308,8 +308,9 @@ static void InternalBuildFollicleTexture_GPU(
 		{
 			FHairStrandsDatas StrandsData;
 			FHairStrandsDatas GuidesData;
+#if WITH_EDITORONLY_DATA
 			Info.GroomAsset->GetHairStrandsDatas(GroupIndex, StrandsData, GuidesData);
-
+#endif
 			StrandsData.StrandsCurves.CurvesRootUV.GetData();
 			const uint32 DataCount = StrandsData.StrandsCurves.CurvesRootUV.Num();
 			const uint32 DataSizeInBytes = sizeof(FVector2D) * DataCount;
