Change strands data & interpolation data to use bulk serialization.

diff --git a/Engine/Plugins/Runtime/HairStrands/Source/HairStrandsCore/Private/GroomAsset.cpp b/Engine/Plugins/Runtime/HairStrands/Source/HairStrandsCore/Private/GroomAsset.cpp
--- a/Engine/Plugins/Runtime/HairStrands/Source/HairStrandsCore/Private/GroomAsset.cpp
+++ b/Engine/Plugins/Runtime/HairStrands/Source/HairStrandsCore/Private/GroomAsset.cpp
@@ -1337,7 +1337,7 @@ void UGroomAsset::SetHairWidth(float Width)
 // differences, etc.) replace the version GUID below with a new one.
 // In case of merge conflicts with DDC versions, you MUST generate a new GUID
 // and set this new GUID as the version.
-#define GROOM_DERIVED_DATA_VERSION TEXT("22D18A66610F485D8AF4B7DB8DDA08E0")
+#define GROOM_DERIVED_DATA_VERSION TEXT("4369C0EB2129496C8128610E601ECD31")
 
 #if WITH_EDITORONLY_DATA
 
diff --git a/Engine/Plugins/Runtime/HairStrands/Source/HairStrandsCore/Private/HairStrandsDatas.cpp b/Engine/Plugins/Runtime/HairStrands/Source/HairStrandsCore/Private/HairStrandsDatas.cpp
--- a/Engine/Plugins/Runtime/HairStrands/Source/HairStrandsCore/Private/HairStrandsDatas.cpp
+++ b/Engine/Plugins/Runtime/HairStrands/Source/HairStrandsCore/Private/HairStrandsDatas.cpp
@@ -172,9 +172,9 @@ FArchive& operator<<(FArchive& Ar, FHairInterpolation1Vertex& Vertex)
 
 void FHairStrandsInterpolationBulkData::Serialize(FArchive& Ar)
 {
-	Ar << Interpolation0;
-	Ar << Interpolation1;
-	Ar << SimRootPointIndex;
+	Interpolation0.BulkSerialize(Ar);
+	Interpolation1.BulkSerialize(Ar);
+	SimRootPointIndex.BulkSerialize(Ar);
 }
 
 void FHairStrandsBulkData::Serialize(FArchive& Ar)
@@ -182,10 +182,10 @@ void FHairStrandsBulkData::Serialize(FArchive& Ar)
 	Ar.UsingCustomVersion(FReleaseObjectVersion::GUID);
 	Ar.UsingCustomVersion(FUE5ReleaseStreamObjectVersion::GUID);
 
-	Ar << Positions;
-	Ar << Attributes;
-	Ar << Materials;
-	Ar << CurveOffsets;
+	Positions.BulkSerialize(Ar);
+	Attributes.BulkSerialize(Ar);
+	Materials.BulkSerialize(Ar);
+	CurveOffsets.BulkSerialize(Ar);
 
 	Ar << CurveCount;
 	Ar << PointCount;
