Exposing simulation parameters to the component

diff --git a/Engine/Plugins/Runtime/HairStrands/Source/HairStrandsCore/Private/GroomComponent.cpp b/Engine/Plugins/Runtime/HairStrands/Source/HairStrandsCore/Private/GroomComponent.cpp
--- a/Engine/Plugins/Runtime/HairStrands/Source/HairStrandsCore/Private/GroomComponent.cpp
+++ b/Engine/Plugins/Runtime/HairStrands/Source/HairStrandsCore/Private/GroomComponent.cpp
@@ -980,6 +980,16 @@ void UGroomComponent::ReleaseHairSimulation()
 	}
 	NiagaraComponents.Empty();
 }
+bool UGroomComponent::IsSimulationEnable(int32 GroupIndex, int32 LODIndex) const
+{
+	bool bIsSimulationEnable = GroomAsset ? GroomAsset->IsSimulationEnable(GroupIndex, LODIndex) : false;
+	if (SimulationSettings.bOverrideSettings)
+	{
+		bIsSimulationEnable = bIsSimulationEnable && SimulationSettings.SolverSettings.bEnableSimulation;
+	}
+
+	return bIsSimulationEnable;
+}
 
 static void InternalUpdateHairSimulation(UGroomComponent* GroomComponent, const bool bHasWorldReady)
 {
@@ -1001,7 +1011,7 @@ static void InternalUpdateHairSimulation(UGroomComponent* GroomComponent, const
 		for (int32 i = 0; i < NumGroups; ++i)
 		{
 			const int32 LODIndex = GroomComponent->GroomGroupsDesc[i].LODForcedIndex;
-			ValidComponents[i] = GroomAsset->IsSimulationEnable(i, LODIndex);
+			ValidComponents[i] = GroomComponent->IsSimulationEnable(i, LODIndex);
 			if (ValidComponents[i] && (GroomAsset->HairGroupsPhysics[i].SolverSettings.NiagaraSolver == EGroomNiagaraSolvers::AngularSprings))
 			{
 				NeedSpringsSolver = true;
@@ -1288,6 +1298,19 @@ void UGroomComponent::SetBindingAsset(UGroomBindingAsset* InBinding)
 	}
 }
 
+void UGroomComponent::SetEnableSimulation(bool bInEnableSimulation)
+{
+	if (SimulationSettings.SolverSettings.bEnableSimulation != bInEnableSimulation)
+	{
+		SimulationSettings.SolverSettings.bEnableSimulation = bInEnableSimulation;
+
+		ReleaseResources();
+		UpdateHairGroupsDesc();
+		UpdateHairSimulation();
+		InitResources();
+	}
+}
+
 void UGroomComponent::UpdateHairGroupsDescAndInvalidateRenderState()
 {
 	UpdateHairGroupsDesc();
@@ -1835,7 +1858,7 @@ void UGroomComponent::InitResources(bool bIsBindingReloading)
 			// * Force global interpolation to be enable for meshes with skinning binding as we use RBF defomation for 'sticking' meshes onto skel. mesh surface			
 			bHasNeedSkeletalMesh	  = bHasNeedSkeletalMesh || BindingType == EGroomBindingType::Rigid;
 			bHasNeedSkinningBinding   = bHasNeedSkinningBinding || BindingType == EGroomBindingType::Skinning;
-			bHasNeedSimulation		  = bHasNeedSimulation || GroomAsset->IsSimulationEnable(GroupIt, LODIt);
+			bHasNeedSimulation		  = bHasNeedSimulation || IsSimulationEnable(GroupIt, LODIt);
 			bHasNeedGlobalDeformation = bHasNeedGlobalDeformation || (BindingType == EGroomBindingType::Skinning && (GroomAsset->IsGlobalInterpolationEnable(GroupIt, LODIt) || GeometryType == EGroomGeometryType::Meshes));
 		}
 	}
@@ -1926,7 +1949,7 @@ void UGroomComponent::InitResources(bool bIsBindingReloading)
 				// * Force global interpolation to be enable for meshes with skinning binding as we use RBF defomation for 'sticking' meshes onto skel. mesh surface
 				// * Global deformation are allowed only with 'Skinning' binding type
 				const EHairGeometryType GeometryType = ToHairGeometryType(GetEffectiveGeometryType(GroomAsset->GetGeometryType(GroupIt, LODIt), bUseCards));
-				const bool LODSimulation = GroomAsset->IsSimulationEnable(GroupIt, LODIt);
+				const bool LODSimulation = IsSimulationEnable(GroupIt, LODIt);
 				const bool LODGlobalInterpolation = LocalBindingAsset && (BindingType == EHairBindingType::Skinning && (GroomAsset->IsGlobalInterpolationEnable(GroupIt, LODIt) || GeometryType == EHairGeometryType::Meshes));
 				bNeedStrandsData = bNeedStrandsData || GeometryType == EHairGeometryType::Strands;
 
@@ -1978,7 +2001,7 @@ void UGroomComponent::InitResources(bool bIsBindingReloading)
 
 			// Initialize the simulation and the global deformation to its default behavior by setting it with LODIndex = -1
 			const int32 LODIndex = -1;
-			HairGroupInstance->Guides.bIsSimulationEnable = GroomAsset->IsSimulationEnable(GroupIt,LODIndex);
+			HairGroupInstance->Guides.bIsSimulationEnable = IsSimulationEnable(GroupIt,LODIndex);
 			HairGroupInstance->Guides.bHasGlobalInterpolation = LocalBindingAsset && GroomAsset->IsGlobalInterpolationEnable(GroupIt,LODIndex);
 		}
 
@@ -2399,8 +2422,8 @@ void UGroomComponent::PostLoad()
 #if WITH_EDITOR
 void UGroomComponent::Invalidate()
 {
-	UpdateHairSimulation();
 	UpdateHairGroupsDescAndInvalidateRenderState();
+	UpdateHairSimulation();
 	ValidateMaterials(false);
 }
 
@@ -2540,6 +2563,7 @@ void UGroomComponent::TickComponent(float DeltaTime, enum ELevelTick TickType, F
 				PrevBoneMatrix = NextBoneMatrix;
 			}
 		}
+		bResetSimulation |= SimulationSettings.bResetSimulation;
 	}
 	bInitSimulation = false;
 
@@ -2648,6 +2672,7 @@ void UGroomComponent::PostEditChangeProperty(FPropertyChangedEvent& PropertyChan
 	const bool bSourceSkeletalMeshChanged = PropertyName == GET_MEMBER_NAME_CHECKED(UGroomComponent, SourceSkeletalMesh);
 	const bool bBindingAssetChanged = PropertyName == GET_MEMBER_NAME_CHECKED(UGroomComponent, BindingAsset);
 	const bool bIsBindingCompatible = UGroomBindingAsset::IsCompatible(GroomAsset, BindingAsset, bValidationEnable);
+	const bool bEnableSolverChanged = PropertyName == GET_MEMBER_NAME_CHECKED(FHairSimulationSolver, bEnableSimulation);
 	if (!bIsBindingCompatible || !UGroomBindingAsset::IsBindingAssetValid(BindingAsset, false, bValidationEnable))
 	{
 		BindingAsset = nullptr;
@@ -2683,7 +2708,7 @@ void UGroomComponent::PostEditChangeProperty(FPropertyChangedEvent& PropertyChan
 	}
 	#endif
 
-	const bool bRecreateResources = bAssetChanged || bBindingAssetChanged || PropertyThatChanged == nullptr || bSourceSkeletalMeshChanged || bRayTracingGeometryChanged;
+	const bool bRecreateResources = bAssetChanged || bBindingAssetChanged || PropertyThatChanged == nullptr || bSourceSkeletalMeshChanged || bRayTracingGeometryChanged || bEnableSolverChanged;
 	if (bRecreateResources)
 	{
 		// Release the resources before Super::PostEditChangeProperty so that they get
diff --git a/Engine/Plugins/Runtime/HairStrands/Source/HairStrandsCore/Private/Niagara/NiagaraDataInterfaceHairStrands.cpp b/Engine/Plugins/Runtime/HairStrands/Source/HairStrandsCore/Private/Niagara/NiagaraDataInterfaceHairStrands.cpp
--- a/Engine/Plugins/Runtime/HairStrands/Source/HairStrandsCore/Private/Niagara/NiagaraDataInterfaceHairStrands.cpp
+++ b/Engine/Plugins/Runtime/HairStrands/Source/HairStrandsCore/Private/Niagara/NiagaraDataInterfaceHairStrands.cpp
@@ -425,12 +425,14 @@ void FNDIHairStrandsData::Update(UNiagaraDataInterfaceHairStrands* Interface, FN
 
 		GlobalInterpolation = (Interface->IsComponentValid() && Interface->SourceComponent->BindingAsset && Interface->SourceComponent->GroomAsset) ?
 			Interface->SourceComponent->GroomAsset->EnableGlobalInterpolation : false;
-		HairGroupInstance = Interface->SourceComponent->GetGroupInstance(GroupIndex);
+		HairGroupInstance = Interface->IsComponentValid() ? Interface->SourceComponent->GetGroupInstance(GroupIndex) : nullptr;
 
 		TickingGroup = Interface->IsComponentValid() ? ComputeTickingGroup(Interface->SourceComponent) : NiagaraFirstTickGroup;
 
-		if (StrandsDatas != nullptr && GroomAsset != nullptr && GroupIndex >= 0 && GroupIndex < GroomAsset->HairGroupsPhysics.Num() &&
-			GroomAsset->IsSimulationEnable(GroupIndex, LODIndex))
+		const bool bIsSimulationEnable = Interface->IsComponentValid() ? Interface->SourceComponent->IsSimulationEnable(GroupIndex, LODIndex) : 
+																		    GroomAsset ? GroomAsset->IsSimulationEnable(GroupIndex, LODIndex) : false;
+
+		if (StrandsDatas != nullptr && GroomAsset != nullptr && GroupIndex >= 0 && GroupIndex < GroomAsset->HairGroupsPhysics.Num() && bIsSimulationEnable)
 		{
 			FHairGroupsPhysics& HairPhysics = GroomAsset->HairGroupsPhysics[GroupIndex];
 			StrandsSize = static_cast<uint8>(HairPhysics.StrandsParameters.StrandsSize);
@@ -476,6 +478,28 @@ void FNDIHairStrandsData::Update(UNiagaraDataInterfaceHairStrands* Interface, FN
 			const FBox& StrandsBox = StrandsDatas->BoundingBox;
 
 			NumStrands = StrandsDatas->GetNumCurves();
+
+			if (Interface->IsComponentValid())
+			{
+				const FHairSimulationSettings& SimulationSettings = Interface->SourceComponent->SimulationSettings;
+				if (SimulationSettings.bOverrideSettings)
+				{
+					GravityVector = SimulationSettings.ExternalForces.GravityVector;
+					AirDrag = SimulationSettings.ExternalForces.AirDrag;
+					AirVelocity = SimulationSettings.ExternalForces.AirVelocity;
+
+					BendDamping = SimulationSettings.MaterialConstraints.BendDamping;
+					BendStiffness = SimulationSettings.MaterialConstraints.BendStiffness;
+
+					StretchDamping = SimulationSettings.MaterialConstraints.StretchDamping;
+					StretchStiffness = SimulationSettings.MaterialConstraints.StretchStiffness;
+
+					StaticFriction = SimulationSettings.MaterialConstraints.StaticFriction;
+					KineticFriction = SimulationSettings.MaterialConstraints.KineticFriction;
+					StrandsViscosity = SimulationSettings.MaterialConstraints.StrandsViscosity;
+					CollisionRadius = SimulationSettings.MaterialConstraints.CollisionRadius;
+				}
+			}
 		}
 		else
 		{
diff --git a/Engine/Plugins/Runtime/HairStrands/Source/HairStrandsCore/Public/GroomAssetPhysics.h b/Engine/Plugins/Runtime/HairStrands/Source/HairStrandsCore/Public/GroomAssetPhysics.h
--- a/Engine/Plugins/Runtime/HairStrands/Source/HairStrandsCore/Public/GroomAssetPhysics.h
+++ b/Engine/Plugins/Runtime/HairStrands/Source/HairStrandsCore/Public/GroomAssetPhysics.h
@@ -273,4 +273,96 @@ struct HAIRSTRANDSCORE_API FHairGroupsPhysics
 	FHairStrandsParameters StrandsParameters;
 
 	bool operator==(const FHairGroupsPhysics& A) const;
+};
+
+USTRUCT(BlueprintType)
+struct HAIRSTRANDSCORE_API FHairSimulationSolver
+{
+	GENERATED_USTRUCT_BODY()
+
+	/** Enable the simulation for the group */
+	UPROPERTY(EditAnywhere, BlueprintReadWrite, Interp, Category = "SolverSettings", meta = (ToolTip = "Enable the simulation of the groom groups/LODs if both this boolean and the one in the groom asset are true"))
+	bool bEnableSimulation = false;
+};
+
+USTRUCT(BlueprintType)
+struct HAIRSTRANDSCORE_API FHairSimulationForces
+{
+	GENERATED_USTRUCT_BODY()
+
+	/** Acceleration vector in cm/s2 to be used for the gravity*/
+	UPROPERTY(EditAnywhere, BlueprintReadWrite, Interp, Category = "ExternalForces", meta = (ToolTip = "Acceleration vector in cm/s2 to be used for the gravity"))
+	FVector GravityVector = FVector(0.0, 0.0, -981.0);
+
+	/** Coefficient between 0 and 1 to be used for the air drag */
+	UPROPERTY(EditAnywhere, BlueprintReadWrite, Interp, Category = "ExternalForces", meta = (ToolTip = "Coefficient between 0 and 1 to be used for the air drag"))
+	float AirDrag = 0.1f;
+
+	/** Velocity of the surrounding air in cm/s */
+	UPROPERTY(EditAnywhere, BlueprintReadWrite, Interp, Category = "ExternalForces", meta = (ToolTip = "Velocity of the surrounding air in cm/s"))
+	FVector AirVelocity = FVector(0.0, 0.0, 0.0);
+};
+
+USTRUCT(BlueprintType)
+struct HAIRSTRANDSCORE_API FHairSimulationConstraints
+{
+	GENERATED_USTRUCT_BODY()
+
+	/** Damping for the bend constraint between 0 and 1 */
+	UPROPERTY(EditAnywhere, BlueprintReadWrite, Interp, AdvancedDisplay, Category = "BendConstraint", meta = (ToolTip = "Damping for the bend constraint between 0 and 1"))
+	float BendDamping = 0.001;
+
+	/** Stiffness for the bend constraint in GPa */
+	UPROPERTY(EditAnywhere, BlueprintReadWrite, Interp, Category = "BendConstraint", meta = (ToolTip = "Stiffness for the bend constraint in GPa", EditCondition = "bOverrideSettings"))
+	float BendStiffness = 0.01;
+
+	/** Damping for the stretch constraint between 0 and 1 */
+	UPROPERTY(EditAnywhere, BlueprintReadWrite, Interp, AdvancedDisplay, Category = "StretchConstraint", meta = (ToolTip = "Damping for the stretch constraint between 0 and 1"))
+	float StretchDamping = 0.001;
+
+	/** Stiffness for the stretch constraint in GPa */
+	UPROPERTY(EditAnywhere, BlueprintReadWrite, Interp, Category = "StretchConstraint", meta = (ToolTip = "Stiffness for the stretch constraint in GPa"))
+	float StretchStiffness = 1.0;
+
+	/** Static friction used for collision against the physics asset */
+	UPROPERTY(EditAnywhere, BlueprintReadWrite, Interp, AdvancedDisplay, Category = "Collision Constraint", meta = (ToolTip = "Static friction used for collision against the physics asset"))
+	float StaticFriction = 0.1;
+
+	/** Kinetic friction used for collision against the physics asset*/
+	UPROPERTY(EditAnywhere, BlueprintReadWrite, Interp, AdvancedDisplay, Category = "CollisionConstraint", meta = (ToolTip = "Kinetic friction used for collision against the physics asset"))
+	float KineticFriction = 0.1;
+
+	/** Viscosity parameter between 0 and 1 that will be used for self collision */
+	UPROPERTY(EditAnywhere, BlueprintReadWrite, Interp, Category = "CollisionConstraint", meta = (ToolTip = "Viscosity parameter between 0 and 1 that will be used for self collision"))
+	float StrandsViscosity = 1.0;
+
+	/** Radius that will be used for the collision detection against the physics asset */
+	UPROPERTY(EditAnywhere, BlueprintReadWrite, Interp, Category = "CollisionConstraint", meta = (ToolTip = "Radius that will be used for the collision detection against the physics asset"))
+	float CollisionRadius = 0.001;
+};
+
+USTRUCT(BlueprintType)
+struct HAIRSTRANDSCORE_API FHairSimulationSettings
+{
+	GENERATED_USTRUCT_BODY()
+
+	/** Reset the simulation in time*/
+	UPROPERTY(EditAnywhere, BlueprintReadWrite, Interp, Category = "GroupsSimulation", meta = (ToolTip = "Boolean to control if we want to reset trhe simulation at some point in time"))
+	bool bResetSimulation = false;
+
+	/** Override the asset simulation settings */
+	UPROPERTY(EditAnywhere, BlueprintReadWrite, Interp, Category = "GroupsSimulation", meta = (ToolTip = "Boolean to control if we are going to override the groom asset physics settings"))
+	bool bOverrideSettings = false;
+
+	/** Solver settings */
+	UPROPERTY(EditAnywhere, BlueprintReadWrite, Interp, Category = "GroupsPhysics", meta = (ToolTip = "Solver Settings for the hair physics", EditCondition = "bOverrideSettings"))
+	FHairSimulationSolver SolverSettings;
+
+	/** Externalk forces */
+	UPROPERTY(EditAnywhere, BlueprintReadWrite, Interp, Category = "GroupsPhysics", meta = (ToolTip = "External Forces for the hair physics", EditCondition = "bOverrideSettings"))
+	FHairSimulationForces ExternalForces;
+
+	/** Material Constraints */
+	UPROPERTY(EditAnywhere, BlueprintReadWrite, Interp, Category = "GroupsPhysics", meta = (ToolTip = "Material Constraints for the hair physics", EditCondition = "bOverrideSettings"))
+	FHairSimulationConstraints MaterialConstraints;
 };
\ No newline at end of file
diff --git a/Engine/Plugins/Runtime/HairStrands/Source/HairStrandsCore/Public/GroomComponent.h b/Engine/Plugins/Runtime/HairStrands/Source/HairStrandsCore/Public/GroomComponent.h
--- a/Engine/Plugins/Runtime/HairStrands/Source/HairStrandsCore/Public/GroomComponent.h
+++ b/Engine/Plugins/Runtime/HairStrands/Source/HairStrandsCore/Public/GroomComponent.h
@@ -15,7 +15,7 @@
 
 #include "GroomComponent.generated.h"
 
-UCLASS(HideCategories = (Object, Physics, Activation, Mobility, "Components|Activation"), editinlinenew, meta = (BlueprintSpawnableComponent), ClassGroup = Rendering)
+UCLASS(HideCategories = (Object, Physics, Activation, Mobility, Collision, "Components|Activation"), editinlinenew, meta = (BlueprintSpawnableComponent), ClassGroup = Rendering)
 class HAIRSTRANDSCORE_API UGroomComponent : public UMeshComponent, public ILODSyncInterface
 {
 	GENERATED_UCLASS_BODY()
@@ -39,9 +39,13 @@ public:
 	class UGroomBindingAsset* BindingAsset;
 
 	/** Physics asset to be used for hair simulation */
-	UPROPERTY(EditAnywhere, BlueprintReadOnly, Category = "Groom")
+	UPROPERTY(EditAnywhere, BlueprintReadOnly, Category = "Simulation")
 	class UPhysicsAsset* PhysicsAsset;
 
+	/** Groom's simulation settings */
+	UPROPERTY(EditAnywhere, BlueprintReadWrite, interp, Category = "Simulation")
+	FHairSimulationSettings SimulationSettings;
+
 	/* Reference of the default/debug materials for each geometric representation */
 	UPROPERTY()
 	UMaterialInterface* Strands_DebugMaterial;
@@ -77,6 +81,9 @@ public:
 	/** Release Niagara components */
 	void ReleaseHairSimulation();
 
+	/** Check if the simulation is enabled or not */
+	bool IsSimulationEnable(int32 GroupIndex, int32 LODIndex) const;
+
 	/** Update Group Description */
 	void UpdateHairGroupsDesc();
 	void UpdateHairGroupsDescAndInvalidateRenderState();
@@ -142,6 +149,10 @@ public:
 	UFUNCTION(BlueprintCallable, Category = "Groom")
 	void SetBindingAsset(UGroomBindingAsset* InBinding);
 
+	/* Accessor function for changing the enable simulation flag from blueprint/sequencer */
+	UFUNCTION(BlueprintCallable, Category = "Simulation")
+	void SetEnableSimulation(bool bInEnableSimulation);
+
 	void SetStableRasterization(bool bEnable);
 	void SetGroomAsset(UGroomAsset* Asset, UGroomBindingAsset* InBinding);
 	void SetHairRootScale(float Scale);
