Make groom Physics asset BPReadWrite

diff --git a/Engine/Plugins/Runtime/HairStrands/Source/HairStrandsCore/Public/GroomComponent.h b/Engine/Plugins/Runtime/HairStrands/Source/HairStrandsCore/Public/GroomComponent.h
--- a/Engine/Plugins/Runtime/HairStrands/Source/HairStrandsCore/Public/GroomComponent.h
+++ b/Engine/Plugins/Runtime/HairStrands/Source/HairStrandsCore/Public/GroomComponent.h
@@ -39,7 +39,7 @@ public:
 	class UGroomBindingAsset* BindingAsset;
 
 	/** Physics asset to be used for hair simulation */
-	UPROPERTY(EditAnywhere, BlueprintReadOnly, Category = "Simulation")
+	UPROPERTY(EditAnywhere, BlueprintReadWrite, Category = "Simulation")
 	class UPhysicsAsset* PhysicsAsset;
 
 	/** Groom's simulation settings */
