Fix bulk serialization with hair data.

diff --git a/Engine/Plugins/Runtime/HairStrands/Source/HairStrandsCore/Private/GroomAsset.cpp b/Engine/Plugins/Runtime/HairStrands/Source/HairStrandsCore/Private/GroomAsset.cpp
--- a/Engine/Plugins/Runtime/HairStrands/Source/HairStrandsCore/Private/GroomAsset.cpp
+++ b/Engine/Plugins/Runtime/HairStrands/Source/HairStrandsCore/Private/GroomAsset.cpp
@@ -1337,7 +1337,7 @@ void UGroomAsset::SetHairWidth(float Width)
 // differences, etc.) replace the version GUID below with a new one.
 // In case of merge conflicts with DDC versions, you MUST generate a new GUID
 // and set this new GUID as the version.
-#define GROOM_DERIVED_DATA_VERSION TEXT("4369C0EB2129496C8128610E601ECD31")
+#define GROOM_DERIVED_DATA_VERSION TEXT("1AA44249A62B435BB137500F2C880C91")
 
 #if WITH_EDITORONLY_DATA
 
diff --git a/Engine/Plugins/Runtime/HairStrands/Source/HairStrandsCore/Private/GroomBuilder.cpp b/Engine/Plugins/Runtime/HairStrands/Source/HairStrandsCore/Private/GroomBuilder.cpp
--- a/Engine/Plugins/Runtime/HairStrands/Source/HairStrandsCore/Private/GroomBuilder.cpp
+++ b/Engine/Plugins/Runtime/HairStrands/Source/HairStrandsCore/Private/GroomBuilder.cpp
@@ -238,6 +238,19 @@ namespace HairStrandsBuilder
 		const FHairStrandsCurves& Curves = HairStrands.StrandsCurves;
 		const FHairStrandsPoints& Points = HairStrands.StrandsPoints;
 
+		struct FPackedRadiusAndType
+		{
+			union
+			{
+				struct
+				{
+					uint8 ControlPointType : 2;
+					uint8 NormalizedRadius : 6;
+				} Data;
+				uint8 Packed;
+			};
+		};
+		static_assert(sizeof(FPackedRadiusAndType) == sizeof(uint8));
 		const bool bSeedValid = RandomSeeds.Num() > 0;
 		for (uint32 CurveIndex = 0; CurveIndex < NumCurves; ++CurveIndex)
 		{
@@ -256,8 +269,10 @@ namespace HairStrandsBuilder
 
 				FHairStrandsPositionFormat::Type& PackedPosition = OutPackedPositions[PointIndex + IndexOffset];
 				CopyVectorToPosition(PointPosition - HairBoxCenter, PackedPosition);
-				PackedPosition.ControlPointType = (PointIndex == 0) ? 1u : (PointIndex == (PointCount - 1) ? 2u : 0u);
-				PackedPosition.NormalizedRadius = uint8(FMath::Clamp(NormalizedRadius * 63.f, 0.f, 63.f));
+				FPackedRadiusAndType PackedRadiusAndType;
+				PackedRadiusAndType.Data.ControlPointType = (PointIndex == 0) ? 1u : (PointIndex == (PointCount - 1) ? 2u : 0u);
+				PackedRadiusAndType.Data.NormalizedRadius = uint8(FMath::Clamp(NormalizedRadius * 63.f, 0.f, 63.f));
+				PackedPosition.PackedRadiusAndType = PackedRadiusAndType.Packed;
 				PackedPosition.UCoord = uint8(FMath::Clamp(CoordU * 255.f, 0.f, 255.f));
 
 				const FVector2D RootUV = Curves.CurvesRootUV[CurveIndex];
diff --git a/Engine/Plugins/Runtime/HairStrands/Source/HairStrandsCore/Private/HairStrandsDatas.cpp b/Engine/Plugins/Runtime/HairStrands/Source/HairStrandsCore/Private/HairStrandsDatas.cpp
--- a/Engine/Plugins/Runtime/HairStrands/Source/HairStrandsCore/Private/HairStrandsDatas.cpp
+++ b/Engine/Plugins/Runtime/HairStrands/Source/HairStrandsCore/Private/HairStrandsDatas.cpp
@@ -66,21 +66,9 @@ FArchive& operator<<(FArchive& Ar, FPackedHairVertex& Vertex)
 	Ar << Vertex.X;
 	Ar << Vertex.Y;
 	Ar << Vertex.Z;
-
+	Ar << Vertex.PackedRadiusAndType;
 	Ar << Vertex.UCoord;
 
-	if (Ar.IsLoading())
-	{
-		uint8 Value;
-		Ar << Value;
-		Vertex.ControlPointType = Value & 0x03; // first 2 bits
-		Vertex.NormalizedRadius = Value >> 2;
-	}
-	else
-	{
-		uint8 Value = Vertex.ControlPointType | (Vertex.NormalizedRadius << 2);
-		Ar << Value;
-	}
 	return Ar;
 }
 
@@ -172,6 +160,10 @@ FArchive& operator<<(FArchive& Ar, FHairInterpolation1Vertex& Vertex)
 
 void FHairStrandsInterpolationBulkData::Serialize(FArchive& Ar)
 {
+	static_assert(sizeof(FHairInterpolation0Vertex::BulkType) == sizeof(FHairInterpolation0Vertex));
+	static_assert(sizeof(FHairInterpolation1Vertex::BulkType) == sizeof(FHairInterpolation1Vertex));
+	static_assert(sizeof(FHairStrandsRootIndexFormat::BulkType) == sizeof(FHairStrandsRootIndexFormat::Type));
+
 	Interpolation0.BulkSerialize(Ar);
 	Interpolation1.BulkSerialize(Ar);
 	SimRootPointIndex.BulkSerialize(Ar);
@@ -179,6 +171,11 @@ void FHairStrandsInterpolationBulkData::Serialize(FArchive& Ar)
 
 void FHairStrandsBulkData::Serialize(FArchive& Ar)
 {
+	static_assert(sizeof(FHairStrandsPositionFormat::BulkType) == sizeof(FHairStrandsPositionFormat::Type));
+	static_assert(sizeof(FHairStrandsAttributeFormat::BulkType) == sizeof(FHairStrandsAttributeFormat::Type));
+	static_assert(sizeof(FHairStrandsMaterialFormat::BulkType) == sizeof(FHairStrandsMaterialFormat::Type));
+	static_assert(sizeof(FHairStrandsRootIndexFormat::BulkType) == sizeof(FHairStrandsRootIndexFormat::Type)); 
+
 	Ar.UsingCustomVersion(FReleaseObjectVersion::GUID);
 	Ar.UsingCustomVersion(FUE5ReleaseStreamObjectVersion::GUID);
 
diff --git a/Engine/Plugins/Runtime/HairStrands/Source/HairStrandsCore/Public/HairStrandsDatas.h b/Engine/Plugins/Runtime/HairStrands/Source/HairStrandsCore/Public/HairStrandsDatas.h
--- a/Engine/Plugins/Runtime/HairStrands/Source/HairStrandsCore/Public/HairStrandsDatas.h
+++ b/Engine/Plugins/Runtime/HairStrands/Source/HairStrandsCore/Public/HairStrandsDatas.h
@@ -15,8 +15,7 @@ struct FPackedHairVertex
 	typedef uint64 BulkType;
 
 	FFloat16 X, Y, Z;
-	uint8 ControlPointType : 2;
-	uint8 NormalizedRadius : 6;
+	uint8 PackedRadiusAndType;
 	uint8 UCoord;
 };
 
@@ -148,7 +147,7 @@ struct FHairStrandsInterpolation0Format
 struct FHairStrandsInterpolation1Format
 {
 	typedef FHairInterpolation1Vertex Type;
-	typedef uint64 BulkType;
+	typedef uint32 BulkType;
 
 	static const uint32 ComponentCount = 1;
 	static const uint32 SizeInByte = sizeof(Type);
