Remove unecessary root position/root normal data from binding asset.

These attributes were used only for old GPU projection purpose (no-longer used), and debug rendering.

diff --git a/Engine/Plugins/Runtime/HairStrands/Source/HairStrandsCore/Private/GroomBindingAsset.cpp b/Engine/Plugins/Runtime/HairStrands/Source/HairStrandsCore/Private/GroomBindingAsset.cpp
--- a/Engine/Plugins/Runtime/HairStrands/Source/HairStrandsCore/Private/GroomBindingAsset.cpp
+++ b/Engine/Plugins/Runtime/HairStrands/Source/HairStrandsCore/Private/GroomBindingAsset.cpp
@@ -643,7 +643,7 @@ void UGroomBindingAsset::PostEditChangeProperty(FPropertyChangedEvent& PropertyC
 // differences, etc.) replace the version GUID below with a new one.
 // In case of merge conflicts with DDC versions, you MUST generate a new GUID
 // and set this new GUID as the version.
-#define GROOM_BINDING_DERIVED_DATA_VERSION TEXT("717523C33A6344D2995323533456DA76")
+#define GROOM_BINDING_DERIVED_DATA_VERSION TEXT("C0555759E09343F2AF8E6C673BF3CE38")
 
 namespace GroomBindingDerivedDataCacheUtils
 {
diff --git a/Engine/Plugins/Runtime/HairStrands/Source/HairStrandsCore/Private/GroomDebug.cpp b/Engine/Plugins/Runtime/HairStrands/Source/HairStrandsCore/Private/GroomDebug.cpp
--- a/Engine/Plugins/Runtime/HairStrands/Source/HairStrandsCore/Private/GroomDebug.cpp
+++ b/Engine/Plugins/Runtime/HairStrands/Source/HairStrandsCore/Private/GroomDebug.cpp
@@ -298,9 +298,6 @@ BEGIN_SHADER_PARAMETER_STRUCT(FHairProjectionHairDebugParameters, )
 	SHADER_PARAMETER_RDG_BUFFER_SRV(StructuredBuffer, DeformedPosition1Buffer)
 	SHADER_PARAMETER_RDG_BUFFER_SRV(StructuredBuffer, DeformedPosition2Buffer)
 
-	// Change for actual frame data (stored or computed only)
-	SHADER_PARAMETER_RDG_BUFFER_SRV(StructuredBuffer, RootPositionBuffer)
-	SHADER_PARAMETER_RDG_BUFFER_SRV(StructuredBuffer, RootNormalBuffer)
 	SHADER_PARAMETER_RDG_BUFFER_SRV(StructuredBuffer, RootBarycentricBuffer)
 
 	SHADER_PARAMETER_STRUCT_REF(FViewUniformShaderParameters, ViewUniformBuffer)
@@ -376,9 +373,7 @@ static void AddDebugProjectionHairPass(
 		return;
 
 	if (EDebugProjectionHairType::HairFrame == GeometryType &&
-		(!RestRootResources->RootPositionBuffer.Buffer ||
-			!RestRootResources->RootNormalBuffer.Buffer ||
-			!RestRootResources->LODs[MeshLODIndex].RootTriangleBarycentricBuffer.Buffer))
+		!RestRootResources->LODs[MeshLODIndex].RootTriangleBarycentricBuffer.Buffer)
 		return;
 
 	const FHairStrandsRestRootResource::FLOD& RestLODDatas = RestRootResources->LODs[MeshLODIndex];
@@ -402,8 +397,6 @@ static void AddDebugProjectionHairPass(
 
 	if (EDebugProjectionHairType::HairFrame == GeometryType)
 	{
-		Parameters->RootPositionBuffer		= RegisterAsSRV(GraphBuilder, RestRootResources->RootPositionBuffer);
-		Parameters->RootNormalBuffer		= RegisterAsSRV(GraphBuilder, RestRootResources->RootNormalBuffer);
 		Parameters->RootBarycentricBuffer	= RegisterAsSRV(GraphBuilder, RestLODDatas.RootTriangleBarycentricBuffer);
 	}
 
@@ -872,14 +865,16 @@ void RunHairStrandsDebug(
 				for (FHairStrandsInstance* AbstractInstance : Instances)
 				{
 					FHairGroupInstance* Instance = static_cast<FHairGroupInstance*>(AbstractInstance);
-
-					if (!Instance->HairGroupPublicData || Instance->Guides.RestRootResource == nullptr || Instance->Guides.DeformedRootResource == nullptr || Instance->BindingType != EHairBindingType::Skinning)
+					if (!Instance->HairGroupPublicData || Instance->BindingType != EHairBindingType::Skinning)
 						continue;
 
-					const int32 MeshLODIndex = Instance->Debug.MeshLODIndex;
 					const bool bRenderStrands = StrandType == EHairStrandsInterpolationType::RenderStrands;
 					FHairStrandsRestRootResource* RestRootResource = bRenderStrands ? Instance->Strands.RestRootResource : Instance->Guides.RestRootResource;
 					FHairStrandsDeformedRootResource* DeformedRootResource = bRenderStrands ? Instance->Strands.DeformedRootResource : Instance->Guides.DeformedRootResource;
+					if (RestRootResource == nullptr || DeformedRootResource == nullptr)
+						continue;
+
+					const int32 MeshLODIndex = Instance->Debug.MeshLODIndex;
 
 					if (bRestTriangle)
 					{
diff --git a/Engine/Plugins/Runtime/HairStrands/Source/HairStrandsCore/Private/GroomResources.cpp b/Engine/Plugins/Runtime/HairStrands/Source/HairStrandsCore/Private/GroomResources.cpp
--- a/Engine/Plugins/Runtime/HairStrands/Source/HairStrandsCore/Private/GroomResources.cpp
+++ b/Engine/Plugins/Runtime/HairStrands/Source/HairStrandsCore/Private/GroomResources.cpp
@@ -729,9 +729,7 @@ void FHairStrandsRestRootResource::InternalAllocate(FRDGBuilder& GraphBuilder)
 	if (RootData.VertexToCurveIndexBuffer.Num() > 0)
 	{
 		InternalCreateVertexBufferRDG<FHairStrandsIndexFormat>(GraphBuilder, RootData.VertexToCurveIndexBuffer, VertexToCurveIndexBuffer, HAIRSTRANDS_RESOUCE_NAME(CurveType, Hair.StrandsRestRoot_VertexToCurveIndexBuffer), EHairResourceUsageType::Static);
-		InternalCreateVertexBufferRDG<FHairStrandsRootPositionFormat>(GraphBuilder, RootData.RootPositionBuffer, RootPositionBuffer, HAIRSTRANDS_RESOUCE_NAME(CurveType, Hair.StrandsRestRoot_RootPositionBuffer), EHairResourceUsageType::Static);
-		InternalCreateVertexBufferRDG<FHairStrandsRootNormalFormat>(GraphBuilder, RootData.RootNormalBuffer, RootNormalBuffer, HAIRSTRANDS_RESOUCE_NAME(CurveType, Hair.StrandsRestRoot_RootNormalBuffer), EHairResourceUsageType::Static);
-		
+
 		check(LODs.Num() == RootData.MeshProjectionLODs.Num());
 		for (uint32 LODIt=0, LODCount = LODs.Num(); LODIt<LODCount; ++LODIt)
 		{
@@ -793,8 +791,6 @@ void FHairStrandsRestRootResource::InternalAllocate(FRDGBuilder& GraphBuilder)
 
 void FHairStrandsRestRootResource::InternalRelease()
 {
-	RootPositionBuffer.Release();
-	RootNormalBuffer.Release();
 	VertexToCurveIndexBuffer.Release();
 	
 	for (FLOD& GPUData : LODs)
@@ -890,9 +886,7 @@ FHairStrandsRootData::FHairStrandsRootData(const FHairStrandsDatas* HairStrandsD
 
 	const uint32 CurveCount = HairStrandsDatas->GetNumCurves();
 	VertexToCurveIndexBuffer.SetNum(HairStrandsDatas->GetNumPoints());
-	RootPositionBuffer.SetNum(RootCount);
-	RootNormalBuffer.SetNum(RootCount);
-
+	
 	for (uint32 CurveIndex = 0; CurveIndex < CurveCount; ++CurveIndex)
 	{
 		const uint32 RootIndex = HairStrandsDatas->StrandsCurves.CurvesOffset[CurveIndex];
@@ -901,33 +895,6 @@ FHairStrandsRootData::FHairStrandsRootData(const FHairStrandsDatas* HairStrandsD
 		{
 			VertexToCurveIndexBuffer[RootIndex + PointIndex] = CurveIndex; // RootIndex;
 		}
-
-		check(PointCount > 1);
-
-		const FVector P0 = HairStrandsDatas->StrandsPoints.PointsPosition[RootIndex];
-		const FVector P1 = HairStrandsDatas->StrandsPoints.PointsPosition[RootIndex + 1];
-		FVector N0 = (P1 - P0).GetSafeNormal();
-
-		// Fallback in case the initial points are too close (this happens on certain assets)
-		if (FVector::DotProduct(N0, N0) == 0)
-		{
-			N0 = FVector(0, 0, 1);
-		}
-
-		FHairStrandsRootPositionFormat::Type P;
-		P.X = P0.X;
-		P.Y = P0.Y;
-		P.Z = P0.Z;
-		P.W = 1;
-
-		FHairStrandsRootNormalFormat::Type N;
-		N.X = N0.X;
-		N.Y = N0.Y;
-		N.Z = N0.Z;
-		N.W = 0;
-
-		RootPositionBuffer[CurveIndex] = P;
-		RootNormalBuffer[CurveIndex] = N;
 	}
 	check(NumSamples.Num() == LODCount);
 
@@ -991,8 +958,6 @@ void FHairStrandsRootData::Serialize(FArchive& Ar)
 	{
 		Ar << RootCount;
 		Ar << VertexToCurveIndexBuffer;
-		Ar << RootPositionBuffer;
-		Ar << RootNormalBuffer;
 		Ar << MeshProjectionLODs;
 	}
 }
@@ -1001,8 +966,6 @@ void FHairStrandsRootData::Reset()
 {
 	RootCount = 0;
 	VertexToCurveIndexBuffer.Empty();
-	RootPositionBuffer.Empty();
-	RootNormalBuffer.Empty();
 	MeshProjectionLODs.Empty();
 }
 
diff --git a/Engine/Plugins/Runtime/HairStrands/Source/HairStrandsCore/Public/GroomResources.h b/Engine/Plugins/Runtime/HairStrands/Source/HairStrandsCore/Public/GroomResources.h
--- a/Engine/Plugins/Runtime/HairStrands/Source/HairStrandsCore/Public/GroomResources.h
+++ b/Engine/Plugins/Runtime/HairStrands/Source/HairStrandsCore/Public/GroomResources.h
@@ -70,8 +70,6 @@ struct FHairStrandsRestRootResource : public FHairCommonResource
 	uint32 GetResourcesSize() const 
 	{
 		uint32 Total = 0;
-		Total += GetBufferTotalNumBytes(RootNormalBuffer);
-		Total += GetBufferTotalNumBytes(RootNormalBuffer);
 		Total += GetBufferTotalNumBytes(VertexToCurveIndexBuffer);
 		for (const FLOD& LOD : LODs)
 		{
@@ -87,8 +85,6 @@ struct FHairStrandsRestRootResource : public FHairCommonResource
 		return Total;
 	}
 
-	FRDGExternalBuffer RootPositionBuffer;
-	FRDGExternalBuffer RootNormalBuffer;
 	FRDGExternalBuffer VertexToCurveIndexBuffer;
 
 	struct FLOD
diff --git a/Engine/Plugins/Runtime/HairStrands/Source/HairStrandsCore/Public/HairStrandsDatas.h b/Engine/Plugins/Runtime/HairStrands/Source/HairStrandsCore/Public/HairStrandsDatas.h
--- a/Engine/Plugins/Runtime/HairStrands/Source/HairStrandsCore/Public/HairStrandsDatas.h
+++ b/Engine/Plugins/Runtime/HairStrands/Source/HairStrandsCore/Public/HairStrandsDatas.h
@@ -226,24 +226,6 @@ struct FHairStrandsCurveTriangleBarycentricFormat
 	static const EPixelFormat Format = PF_R32_UINT;
 };
 
-struct FHairStrandsRootPositionFormat
-{
-	typedef FVector4 Type;
-	static const uint32 ComponentCount = 1;
-	static const uint32 SizeInByte = sizeof(Type);
-	static const EVertexElementType VertexElementType = VET_Float4;
-	static const EPixelFormat Format = PF_A32B32G32R32F;
-};
-
-struct FHairStrandsRootNormalFormat
-{
-	typedef FVector4_16 Type;
-	static const uint32 ComponentCount = 1;
-	static const uint32 SizeInByte = sizeof(Type);
-	static const EVertexElementType VertexElementType = VET_Float4;
-	static const EPixelFormat Format = PF_FloatRGBA;
-};
-
 struct FHairStrandsRootUtils
 {
 	static uint32	 EncodeTriangleIndex(uint32 TriangleIndex, uint32 SectionIndex);
@@ -407,7 +389,7 @@ struct HAIRSTRANDSCORE_API FHairStrandsBulkData
 	TArray<FHairStrandsPositionFormat::Type>	Positions;		// Size = PointCount
 	TArray<FHairStrandsAttributeFormat::Type>	Attributes;		// Size = PointCount
 	TArray<FHairStrandsMaterialFormat::Type>	Materials;		// Size = PointCount
-	TArray<FHairStrandsRootIndexFormat::Type>	CurveOffsets;	// Size = CurveCount  - Store the root point index for the curve
+	TArray<FHairStrandsRootIndexFormat::Type>	CurveOffsets;	// Size = CurveCount+1 - Store the root point index for the curve
 };
 
 /** Hair strands debug data */
@@ -497,12 +479,6 @@ struct FHairStrandsRootData
 	/* Curve index for every vertices */
 	TArray<FHairStrandsIndexFormat::Type> VertexToCurveIndexBuffer;
 
-	/* Curve root's positions */
-	TArray<FHairStrandsRootPositionFormat::Type> RootPositionBuffer;
-
-	/* Curve root's normal orientation */
-	TArray<FHairStrandsRootNormalFormat::Type> RootNormalBuffer;
-
 	/* Store the hair projection information for each mesh LOD */
 	TArray<FMeshProjectionLOD> MeshProjectionLODs;
 };
diff --git a/Engine/Shaders/Private/HairStrands/HairStrandsMeshProjectionHairDebug.usf b/Engine/Shaders/Private/HairStrands/HairStrandsMeshProjectionHairDebug.usf
--- a/Engine/Shaders/Private/HairStrands/HairStrandsMeshProjectionHairDebug.usf
+++ b/Engine/Shaders/Private/HairStrands/HairStrandsMeshProjectionHairDebug.usf
@@ -13,8 +13,6 @@ float4x4 RootLocalToWorld;
 uint DeformedFrameEnable;
 
 #if PERMUTATION_INPUT_TYPE == INPUT_ROOT_FRAME
-Buffer<float4> RootPositionBuffer;
-Buffer<float4> RootNormalBuffer;
 Buffer<uint>   RootBarycentricBuffer;
 #endif
 
@@ -47,38 +45,23 @@ void MainVS(
 	float3 N;
 	float3 P;
 	{
-		const FHairMeshTriangleWithUVs RestTriangleWithUVs     = GetTriangleWithUVsTransformation(LineIndex, RestPosition0Buffer, RestPosition1Buffer, RestPosition2Buffer);
-		const FHairMeshTriangleWithUVs DeformedTriangleWithUVs = GetTriangleWithUVsTransformation(LineIndex, DeformedPosition0Buffer, DeformedPosition1Buffer, DeformedPosition2Buffer);
-
-		const FHairMeshTriangle RestTriangle			= GetTriangleTransformation(LineIndex, RestPosition0Buffer, RestPosition1Buffer, RestPosition2Buffer);
-		const FHairMeshTriangle DeformedTriangle		= GetTriangleTransformation(LineIndex, DeformedPosition0Buffer, DeformedPosition1Buffer, DeformedPosition2Buffer);
-
-
-		float3 P0 = RootPositionBuffer[LineIndex].xyz;
-		float3 LN = RootNormalBuffer[LineIndex].xyz;
-		float3 B  = DecodeBarycentrics(RootBarycentricBuffer[LineIndex]);
-		float3 P1 = P0 + LN;
+		const FHairMeshTriangleWithUVs RestTriangle		= GetTriangleWithUVsTransformation(LineIndex, RestPosition0Buffer, RestPosition1Buffer, RestPosition2Buffer);
+		const FHairMeshTriangleWithUVs DeformedTriangle = GetTriangleWithUVsTransformation(LineIndex, DeformedPosition0Buffer, DeformedPosition1Buffer, DeformedPosition2Buffer);
 
+		FHairMeshTriangleWithUVs Tri;
 		if (DeformedFrameEnable > 0)
-		{
-			UVs = 
-				B.x * DeformedTriangleWithUVs.UV0 +
-				B.y * DeformedTriangleWithUVs.UV1 +
-				B.z * DeformedTriangleWithUVs.UV2;
-		}
+			Tri = DeformedTriangle;
 		else
-		{
-			UVs =
-				B.x * RestTriangleWithUVs.UV0 +
-				B.y * RestTriangleWithUVs.UV1 +
-				B.z * RestTriangleWithUVs.UV2;
-		}
+			Tri = RestTriangle;
 
-		if (DeformedFrameEnable > 0)
-		{
-			P0 = TransformPoint(P0, B, RestTriangle, DeformedTriangle);
-			P1 = TransformPoint(P1, B, RestTriangle, DeformedTriangle);
-		}
+		const float3 B  = DecodeBarycentrics(RootBarycentricBuffer[LineIndex]);
+		float3 P0 = B.x * Tri.P0 + B.y * Tri.P1 + B.z * Tri.P2;
+		float3 P1 = P0 + Tri.N;
+
+		UVs = 
+			B.x * Tri.UV0 +
+			B.y * Tri.UV1 +
+			B.z * Tri.UV2;
 
 		P0 = mul(float4(P0, 1), RootLocalToWorld).xyz;
 		P1 = mul(float4(P1, 1), RootLocalToWorld).xyz;
