Change root binding data to use bulk serialization.

diff --git a/Engine/Plugins/Runtime/HairStrands/Source/HairStrandsCore/Private/GroomBindingAsset.cpp b/Engine/Plugins/Runtime/HairStrands/Source/HairStrandsCore/Private/GroomBindingAsset.cpp
--- a/Engine/Plugins/Runtime/HairStrands/Source/HairStrandsCore/Private/GroomBindingAsset.cpp
+++ b/Engine/Plugins/Runtime/HairStrands/Source/HairStrandsCore/Private/GroomBindingAsset.cpp
@@ -643,7 +643,7 @@ void UGroomBindingAsset::PostEditChangeProperty(FPropertyChangedEvent& PropertyC
 // differences, etc.) replace the version GUID below with a new one.
 // In case of merge conflicts with DDC versions, you MUST generate a new GUID
 // and set this new GUID as the version.
-#define GROOM_BINDING_DERIVED_DATA_VERSION TEXT("C0555759E09343F2AF8E6C673BF3CE38")
+#define GROOM_BINDING_DERIVED_DATA_VERSION TEXT("C931E6356B4541B1BB3AE6FCE7C45DB7")
 
 namespace GroomBindingDerivedDataCacheUtils
 {
diff --git a/Engine/Plugins/Runtime/HairStrands/Source/HairStrandsCore/Private/GroomResources.cpp b/Engine/Plugins/Runtime/HairStrands/Source/HairStrandsCore/Private/GroomResources.cpp
--- a/Engine/Plugins/Runtime/HairStrands/Source/HairStrandsCore/Private/GroomResources.cpp
+++ b/Engine/Plugins/Runtime/HairStrands/Source/HairStrandsCore/Private/GroomResources.cpp
@@ -934,21 +934,18 @@ bool FHairStrandsRootData::HasProjectionData() const
 FArchive& operator<<(FArchive& Ar, FHairStrandsRootData::FMeshProjectionLOD& LOD)
 {
 	Ar << LOD.LODIndex;
-	Ar << LOD.RootTriangleIndexBuffer;
-	Ar << LOD.RootTriangleBarycentricBuffer;
-	Ar << LOD.RestRootTrianglePosition0Buffer;
-	Ar << LOD.RestRootTrianglePosition1Buffer;
-	Ar << LOD.RestRootTrianglePosition2Buffer;
+	LOD.RootTriangleIndexBuffer.BulkSerialize(Ar);
+	LOD.RootTriangleBarycentricBuffer.BulkSerialize(Ar);
+	LOD.RestRootTrianglePosition0Buffer.BulkSerialize(Ar);
+	LOD.RestRootTrianglePosition1Buffer.BulkSerialize(Ar);
+	LOD.RestRootTrianglePosition2Buffer.BulkSerialize(Ar);
 
 	Ar << LOD.SampleCount;
-	Ar << LOD.MeshInterpolationWeightsBuffer;
-	Ar << LOD.MeshSampleIndicesBuffer;
-	Ar << LOD.RestSamplePositionsBuffer;
+	LOD.MeshInterpolationWeightsBuffer.BulkSerialize(Ar);
+	LOD.MeshSampleIndicesBuffer.BulkSerialize(Ar);
+	LOD.RestSamplePositionsBuffer.BulkSerialize(Ar);
+	LOD.ValidSectionIndices.BulkSerialize(Ar);
 
-	if (Ar.IsLoading())
-	{
-		FGroomBindingBuilder::BuildUniqueSections(LOD);
-	}
 	return Ar;
 }
 
@@ -957,7 +954,7 @@ void FHairStrandsRootData::Serialize(FArchive& Ar)
 	if (!Ar.IsObjectReferenceCollector())
 	{
 		Ar << RootCount;
-		Ar << VertexToCurveIndexBuffer;
+		VertexToCurveIndexBuffer.BulkSerialize(Ar);
 		Ar << MeshProjectionLODs;
 	}
 }
diff --git a/Engine/Plugins/Runtime/HairStrands/Source/HairStrandsCore/Public/HairStrandsDatas.h b/Engine/Plugins/Runtime/HairStrands/Source/HairStrandsCore/Public/HairStrandsDatas.h
--- a/Engine/Plugins/Runtime/HairStrands/Source/HairStrandsCore/Public/HairStrandsDatas.h
+++ b/Engine/Plugins/Runtime/HairStrands/Source/HairStrandsCore/Public/HairStrandsDatas.h
@@ -190,6 +190,7 @@ struct HAIRSTRANDSCORE_API FHairStrandsIndexFormat
 struct FHairStrandsWeightFormat
 {
 	using Type = float;
+	using BulkType = float;
 	static const uint32 ComponentCount = 1;
 	static const uint32 SizeInByte = sizeof(Type);
 	static const EVertexElementType VertexElementType = VET_Float1;
@@ -202,6 +203,7 @@ struct FHairStrandsWeightFormat
 struct FHairStrandsMeshTrianglePositionFormat
 {
 	using Type = FVector4;
+	using BulkType = FVector4;
 	static const uint32 ComponentCount = 1;
 	static const uint32 SizeInByte = sizeof(Type);
 	static const EVertexElementType VertexElementType = VET_Float4;
@@ -211,6 +213,7 @@ struct FHairStrandsMeshTrianglePositionFormat
 struct FHairStrandsCurveTriangleIndexFormat
 {
 	using Type = uint32;
+	using BulkType = uint32;
 	static const uint32 ComponentCount = 1;
 	static const uint32 SizeInByte = sizeof(Type);
 	static const EVertexElementType VertexElementType = VET_UInt;
@@ -220,6 +223,7 @@ struct FHairStrandsCurveTriangleIndexFormat
 struct FHairStrandsCurveTriangleBarycentricFormat
 {
 	using Type = uint32;
+	using BulkType = uint32;
 	static const uint32 ComponentCount = 1;
 	static const uint32 SizeInByte = sizeof(Type);
 	static const EVertexElementType VertexElementType = VET_UInt;
