Fix hair sim init when the world is not yet initialized

diff --git a/Engine/Plugins/Runtime/HairStrands/Source/HairStrandsCore/Private/GroomComponent.cpp b/Engine/Plugins/Runtime/HairStrands/Source/HairStrandsCore/Private/GroomComponent.cpp
--- a/Engine/Plugins/Runtime/HairStrands/Source/HairStrandsCore/Private/GroomComponent.cpp
+++ b/Engine/Plugins/Runtime/HairStrands/Source/HairStrandsCore/Private/GroomComponent.cpp
@@ -1056,7 +1056,7 @@ static void InternalUpdateHairSimulation(UGroomComponent* GroomComponent, const
 			if (!NiagaraComponent)
 			{
 				NiagaraComponent = NewObject<UNiagaraComponent>(GroomComponent, NAME_None, RF_Transient);
-				if (GroomComponent->GetOwner() && GroomComponent->GetOwner()->GetWorld())
+				if (GroomComponent->GetOwner() && GroomComponent->GetOwner()->GetWorld() && GroomComponent->GetOwner()->GetWorld()->bIsWorldInitialized)
 				{
 					NiagaraComponent->AttachToComponent(GroomComponent, FAttachmentTransformRules::KeepRelativeTransform);
 					NiagaraComponent->RegisterComponent();
