Remove mandatory RBF deformation for hair meshes. This now needs to be opt-in at the asset level.

This saves cost when hair meshes are used, but BRDF deformation is not needed.

diff --git a/Engine/Plugins/Runtime/HairStrands/Source/HairStrandsCore/Private/GroomComponent.cpp b/Engine/Plugins/Runtime/HairStrands/Source/HairStrandsCore/Private/GroomComponent.cpp
--- a/Engine/Plugins/Runtime/HairStrands/Source/HairStrandsCore/Private/GroomComponent.cpp
+++ b/Engine/Plugins/Runtime/HairStrands/Source/HairStrandsCore/Private/GroomComponent.cpp
@@ -1859,7 +1859,7 @@ void UGroomComponent::InitResources(bool bIsBindingReloading)
 			bHasNeedSkeletalMesh	  = bHasNeedSkeletalMesh || BindingType == EGroomBindingType::Rigid;
 			bHasNeedSkinningBinding   = bHasNeedSkinningBinding || BindingType == EGroomBindingType::Skinning;
 			bHasNeedSimulation		  = bHasNeedSimulation || IsSimulationEnable(GroupIt, LODIt);
-			bHasNeedGlobalDeformation = bHasNeedGlobalDeformation || (BindingType == EGroomBindingType::Skinning && (GroomAsset->IsGlobalInterpolationEnable(GroupIt, LODIt) || GeometryType == EGroomGeometryType::Meshes));
+			bHasNeedGlobalDeformation = bHasNeedGlobalDeformation || (BindingType == EGroomBindingType::Skinning && GroomAsset->IsGlobalInterpolationEnable(GroupIt, LODIt));
 		}
 	}
 	const bool bHasNeedBindingData = bHasNeedSkinningBinding || bHasNeedGlobalDeformation;
@@ -1950,7 +1950,7 @@ void UGroomComponent::InitResources(bool bIsBindingReloading)
 				// * Global deformation are allowed only with 'Skinning' binding type
 				const EHairGeometryType GeometryType = ToHairGeometryType(GetEffectiveGeometryType(GroomAsset->GetGeometryType(GroupIt, LODIt), bUseCards));
 				const bool LODSimulation = IsSimulationEnable(GroupIt, LODIt);
-				const bool LODGlobalInterpolation = LocalBindingAsset && (BindingType == EHairBindingType::Skinning && (GroomAsset->IsGlobalInterpolationEnable(GroupIt, LODIt) || GeometryType == EHairGeometryType::Meshes));
+				const bool LODGlobalInterpolation = LocalBindingAsset && BindingType == EHairBindingType::Skinning && GroomAsset->IsGlobalInterpolationEnable(GroupIt, LODIt);
 				bNeedStrandsData = bNeedStrandsData || GeometryType == EHairGeometryType::Strands;
 
 				CPULODScreenSize.Add(LODSettings.ScreenSize);
