Change groom cache animated position buffer to fix alignment issue.

diff --git a/Engine/Plugins/Runtime/HairStrands/Source/HairStrandsCore/Private/HairStrandsRendering.cpp b/Engine/Plugins/Runtime/HairStrands/Source/HairStrandsCore/Private/HairStrandsRendering.cpp
--- a/Engine/Plugins/Runtime/HairStrands/Source/HairStrandsCore/Private/HairStrandsRendering.cpp
+++ b/Engine/Plugins/Runtime/HairStrands/Source/HairStrandsCore/Private/HairStrandsRendering.cpp
@@ -204,11 +204,14 @@ static void AddGroomCacheUpdatePass(
 	const uint32 DataSizeInBytes = sizeof(FVector) * DataCount;
 	if (DataSizeInBytes != 0)
 	{
-		VertexBuffer = CreateStructuredBuffer(
+		// Deformation are upload into a Buffer<float> as the original position are float3 which is is both
+		// 1) incompatible with structure buffer alignment (128bits), and 2) incompatible with vertex buffer 
+		// as R32G32B32_FLOAT format is not well supported for SRV accross HW.
+		// So instead the positions are uploaded into vertex buffer Buffer<float>
+		VertexBuffer = CreateVertexBuffer(
 			GraphBuilder,
 			TEXT("GroomCache_PositionBuffer"),
-			sizeof(FVector),
-			GroomCacheData.VertexData.PointsPosition.Num(),
+			FRDGBufferDesc::CreateBufferDesc(sizeof(float), GroomCacheData.VertexData.PointsPosition.Num() * 3),
 			GroomCacheData.VertexData.PointsPosition.GetData(),
 			DataSizeInBytes,
 			ERDGInitialDataFlags::None);
diff --git a/Engine/Shaders/Private/HairStrands/HairStrandsInterpolation.usf b/Engine/Shaders/Private/HairStrands/HairStrandsInterpolation.usf
--- a/Engine/Shaders/Private/HairStrands/HairStrandsInterpolation.usf
+++ b/Engine/Shaders/Private/HairStrands/HairStrandsInterpolation.usf
@@ -729,7 +729,7 @@ void MainCS(
 uint ElementCount;
 uint DispatchCountX;
 
-StructuredBuffer<float3> InAnimatedBuffer;
+Buffer<float> InAnimatedBuffer;
 Buffer<uint4> InRestPoseBuffer;
 Buffer<float4> InDeformedOffsetBuffer;
 RWBuffer<uint4> OutDeformedBuffer;
@@ -745,9 +745,15 @@ void MainCS(
 		return;
 	}
 
+	const uint Index3 = Index * 3;
+	const float3 InAnimatedPosition = float3(
+		InAnimatedBuffer[Index3],
+		InAnimatedBuffer[Index3+1],
+		InAnimatedBuffer[Index3+2]);
+		
 	uint4 PackedRestPosition = InRestPoseBuffer[Index];
 	float3 HairPositionOffset = InDeformedOffsetBuffer[0];
-	float3 WorldVertexPosition = InAnimatedBuffer[Index] - HairPositionOffset;
+	float3 WorldVertexPosition = InAnimatedPosition - HairPositionOffset;
 	OutDeformedBuffer[Index] = uint4(f32tof16(WorldVertexPosition) , PackedRestPosition.w);
 }
 
