Fix SRV creation error when uploading groom cache position.

diff --git a/Engine/Plugins/Runtime/HairStrands/Source/HairStrandsCore/Private/HairStrandsRendering.cpp b/Engine/Plugins/Runtime/HairStrands/Source/HairStrandsCore/Private/HairStrandsRendering.cpp
--- a/Engine/Plugins/Runtime/HairStrands/Source/HairStrandsCore/Private/HairStrandsRendering.cpp
+++ b/Engine/Plugins/Runtime/HairStrands/Source/HairStrandsCore/Private/HairStrandsRendering.cpp
@@ -224,7 +224,7 @@ static void AddGroomCacheUpdatePass(
 	FGroomCacheUpdatePassCS::FParameters* Parameters = GraphBuilder.AllocParameters<FGroomCacheUpdatePassCS::FParameters>();
 	Parameters->DispatchCountX = DispatchCount.X;
 	Parameters->ElementCount = ElementCount;
-	Parameters->InAnimatedBuffer = GraphBuilder.CreateSRV(VertexBuffer);
+	Parameters->InAnimatedBuffer = GraphBuilder.CreateSRV(VertexBuffer, PF_R32_FLOAT);
 	Parameters->InRestPoseBuffer = InBuffer;
 	Parameters->InDeformedOffsetBuffer = InDeformedOffsetBuffer;
 	Parameters->OutDeformedBuffer = OutBuffer;
