Temporary remove sanity check.

The underlying issue comes from RDG returning buffers whose size can be larger than the requested allocation. Since the returned buffer's desc is used after allocation (sanity check, pass parameters, ...) this can case assert or heavy rendering pass workload.

diff --git a/Engine/Plugins/Runtime/HairStrands/Source/HairStrandsCore/Private/HairStrandsRendering.cpp b/Engine/Plugins/Runtime/HairStrands/Source/HairStrandsCore/Private/HairStrandsRendering.cpp
--- a/Engine/Plugins/Runtime/HairStrands/Source/HairStrandsCore/Private/HairStrandsRendering.cpp
+++ b/Engine/Plugins/Runtime/HairStrands/Source/HairStrandsCore/Private/HairStrandsRendering.cpp
@@ -843,7 +843,8 @@ static void AddHairClusterAABBPass(
 		check(ClusterData->ClusterCount == ClusterIdBuffer->Desc.NumElements);
 		check(ClusterData->ClusterCount == ClusterIndexOffsetBuffer->Desc.NumElements);
 		check(ClusterData->ClusterCount == ClusterIndexCountBuffer->Desc.NumElements);
-		check(ClusterData->ClusterCount * 6 == ClusterAABBData.ClusterAABBBuffer.Buffer->Desc.NumElements);
+		// Currently disabled this sanity check as the RDG allocation can return a larger buffer (for reuse purpose)
+		//check(ClusterData->ClusterCount * 6 == ClusterAABBData.ClusterAABBBuffer.Buffer->Desc.NumElements);
 	}
 	Parameters->HairStrandsVF_CullingIndirectBuffer = DrawIndirectRasterComputeBuffer; // Used for checking max vertex count
 	Parameters->OutClusterAABBBuffer = ClusterAABBData.ClusterAABBBuffer.UAV;
