Change hair rest resources, culling resources, and interpolation resources to be lazily allocated.

This save GPU memory for groom having hair strands data, but never using them due to force cards option, or finer LOD levels never been reached.

diff --git a/Engine/Plugins/Runtime/HairStrands/Source/HairStrandsCore/Private/GroomManager.cpp b/Engine/Plugins/Runtime/HairStrands/Source/HairStrandsCore/Private/GroomManager.cpp
--- a/Engine/Plugins/Runtime/HairStrands/Source/HairStrandsCore/Private/GroomManager.cpp
+++ b/Engine/Plugins/Runtime/HairStrands/Source/HairStrandsCore/Private/GroomManager.cpp
@@ -635,6 +635,8 @@ static void RunHairLODSelection(FRDGBuilder& GraphBuilder, const FHairStrandsIns
 		// Note: Allocation will only be done if the resources it not yet initialized.
 		if (Instance->Guides.Data)
 		{
+			if (Instance->Guides.RestRootResource)		Instance->Guides.RestRootResource->Allocate(GraphBuilder);
+			if (Instance->Guides.RestResource)			Instance->Guides.RestResource->Allocate(GraphBuilder);
 			if (Instance->Guides.DeformedRootResource)	Instance->Guides.DeformedRootResource->Allocate(GraphBuilder);
 			if (Instance->Guides.DeformedResource)		Instance->Guides.DeformedResource->Allocate(GraphBuilder);
 		}
@@ -664,10 +666,15 @@ static void RunHairLODSelection(FRDGBuilder& GraphBuilder, const FHairStrandsIns
 		}
 		else if (GeometryType == EHairGeometryType::Strands)
 		{
-			if (Instance->Strands.DeformedRootResource)	Instance->Strands.DeformedRootResource->Allocate(GraphBuilder);
-			if (Instance->Strands.DeformedResource)		Instance->Strands.DeformedResource->Allocate(GraphBuilder);
+			if (Instance->Strands.RestRootResource)			Instance->Strands.RestRootResource->Allocate(GraphBuilder);
+			if (Instance->Strands.RestResource)				Instance->Strands.RestResource->Allocate(GraphBuilder);
+			if (Instance->Strands.ClusterCullingResource)	Instance->Strands.ClusterCullingResource->Allocate(GraphBuilder);
+			if (Instance->Strands.InterpolationResource)	Instance->Strands.InterpolationResource->Allocate(GraphBuilder);
+
+			if (Instance->Strands.DeformedRootResource)		Instance->Strands.DeformedRootResource->Allocate(GraphBuilder);
+			if (Instance->Strands.DeformedResource)			Instance->Strands.DeformedResource->Allocate(GraphBuilder);
 			#if RHI_RAYTRACING
-			if (Instance->Strands.RenRaytracingResource)Instance->Strands.RenRaytracingResource->Allocate(GraphBuilder);
+			if (Instance->Strands.RenRaytracingResource)	Instance->Strands.RenRaytracingResource->Allocate(GraphBuilder);
 			#endif
 			Instance->Strands.VertexFactory->InitResources();
 		}
diff --git a/Engine/Plugins/Runtime/HairStrands/Source/HairStrandsCore/Private/GroomResources.cpp b/Engine/Plugins/Runtime/HairStrands/Source/HairStrandsCore/Private/GroomResources.cpp
--- a/Engine/Plugins/Runtime/HairStrands/Source/HairStrandsCore/Private/GroomResources.cpp
+++ b/Engine/Plugins/Runtime/HairStrands/Source/HairStrandsCore/Private/GroomResources.cpp
@@ -568,7 +568,7 @@ void FHairMeshesDeformedResource::InternalRelease()
 /////////////////////////////////////////////////////////////////////////////////////////
 
 FHairStrandsRestResource::FHairStrandsRestResource(const FHairStrandsBulkData& InBulkData, EHairStrandsResourcesType InCurveType) :
-	FHairCommonResource(EHairStrandsAllocationType::Immediate),
+	FHairCommonResource(EHairStrandsAllocationType::Deferred),
 	PositionBuffer(), AttributeBuffer(), MaterialBuffer(), BulkData(InBulkData), CurveType(InCurveType)
 {}
 
@@ -697,7 +697,7 @@ void FHairStrandsClusterCullingData::Serialize(FArchive& Ar)
 }
 
 FHairStrandsClusterCullingResource::FHairStrandsClusterCullingResource(const FHairStrandsClusterCullingData& InData): 
-	FHairCommonResource(EHairStrandsAllocationType::Immediate),
+	FHairCommonResource(EHairStrandsAllocationType::Deferred),
 	Data(InData) 
 {
 
@@ -727,7 +727,7 @@ void FHairStrandsClusterCullingResource::InternalRelease()
 /////////////////////////////////////////////////////////////////////////////////////////
 
 FHairStrandsRestRootResource::FHairStrandsRestRootResource(const FHairStrandsRootData& InRootData, EHairStrandsResourcesType InCurveType) :
-	FHairCommonResource(EHairStrandsAllocationType::Immediate),
+	FHairCommonResource(EHairStrandsAllocationType::Deferred),
 	RootData(InRootData), CurveType(InCurveType)
 {
 	PopulateFromRootData();
@@ -995,7 +995,7 @@ void FHairStrandsRootData::Reset()
 /////////////////////////////////////////////////////////////////////////////////////////
 
 FHairStrandsInterpolationResource::FHairStrandsInterpolationResource(const FHairStrandsInterpolationBulkData& InBulkData) :
-	FHairCommonResource(EHairStrandsAllocationType::Immediate),
+	FHairCommonResource(EHairStrandsAllocationType::Deferred),
 	Interpolation0Buffer(), Interpolation1Buffer(), BulkData(InBulkData)
 {
 }
