Fix crash due to missing initialization on cards resources.

diff --git a/Engine/Plugins/Runtime/HairStrands/Source/HairStrandsCore/Private/GroomManager.cpp b/Engine/Plugins/Runtime/HairStrands/Source/HairStrandsCore/Private/GroomManager.cpp
--- a/Engine/Plugins/Runtime/HairStrands/Source/HairStrandsCore/Private/GroomManager.cpp
+++ b/Engine/Plugins/Runtime/HairStrands/Source/HairStrandsCore/Private/GroomManager.cpp
@@ -657,6 +657,8 @@ static void RunHairLODSelection(FRDGBuilder& GraphBuilder, const FHairStrandsIns
 			FHairGroupInstance::FCards::FLOD& InstanceLOD = Instance->Cards.LODs[IntLODIndex];
 
 			if (InstanceLOD.DeformedResource)				InstanceLOD.DeformedResource->Allocate(GraphBuilder);
+			if (InstanceLOD.Guides.RestRootResource)		InstanceLOD.Guides.RestRootResource->Allocate(GraphBuilder);
+			if (InstanceLOD.Guides.RestResource)			InstanceLOD.Guides.RestResource->Allocate(GraphBuilder);
 			if (InstanceLOD.Guides.DeformedRootResource)	InstanceLOD.Guides.DeformedRootResource->Allocate(GraphBuilder);
 			if (InstanceLOD.Guides.DeformedResource)		InstanceLOD.Guides.DeformedResource->Allocate(GraphBuilder);
 			#if RHI_RAYTRACING
