Fix buffer transfer in case the buffer is not valid

diff --git a/Engine/Plugins/Runtime/HairStrands/Source/HairStrandsCore/Private/Niagara/NiagaraDataInterfaceHairStrands.cpp b/Engine/Plugins/Runtime/HairStrands/Source/HairStrandsCore/Private/Niagara/NiagaraDataInterfaceHairStrands.cpp
--- a/Engine/Plugins/Runtime/HairStrands/Source/HairStrandsCore/Private/Niagara/NiagaraDataInterfaceHairStrands.cpp
+++ b/Engine/Plugins/Runtime/HairStrands/Source/HairStrandsCore/Private/Niagara/NiagaraDataInterfaceHairStrands.cpp
@@ -319,7 +319,7 @@ void FNDIHairStrandsBuffer::Update(
 
 void FNDIHairStrandsBuffer::Transfer(const TStaticArray<float, 32 * NumScales>& InParamsScale)
 {
-	if (SourceRestResources != nullptr)
+	if (SourceRestResources != nullptr && ParamsScaleBuffer.Buffer.IsValid())
 	{
 		const uint32 ScaleCount = 32 * NumScales;
 		const uint32 ScaleBytes = sizeof(float) * ScaleCount;
diff --git a/Engine/Plugins/Runtime/HairStrands/Source/HairStrandsCore/Private/Niagara/NiagaraDataInterfacePhysicsAsset.cpp b/Engine/Plugins/Runtime/HairStrands/Source/HairStrandsCore/Private/Niagara/NiagaraDataInterfacePhysicsAsset.cpp
--- a/Engine/Plugins/Runtime/HairStrands/Source/HairStrandsCore/Private/Niagara/NiagaraDataInterfacePhysicsAsset.cpp
+++ b/Engine/Plugins/Runtime/HairStrands/Source/HairStrandsCore/Private/Niagara/NiagaraDataInterfacePhysicsAsset.cpp
@@ -97,7 +97,7 @@ void CreateInternalBuffer(FRWBuffer& OutputBuffer)
 template<typename BufferType, EPixelFormat PixelFormat, uint32 ElementCount>
 void UpdateInternalBuffer(const TStaticArray<BufferType,ElementCount>& InputData, FRWBuffer& OutputBuffer)
 {
-	if (ElementCount > 0)
+	if (ElementCount > 0 && OutputBuffer.Buffer.IsValid())
 	{
 		const uint32 BufferCount = ElementCount;
 		const uint32 BufferBytes = sizeof(BufferType) * BufferCount;
