Unchecked cast in groom component crashing on PIE exit

diff --git a/Engine/Plugins/Runtime/HairStrands/Source/HairStrandsCore/Private/GroomComponent.cpp b/Engine/Plugins/Runtime/HairStrands/Source/HairStrandsCore/Private/GroomComponent.cpp
--- a/Engine/Plugins/Runtime/HairStrands/Source/HairStrandsCore/Private/GroomComponent.cpp
+++ b/Engine/Plugins/Runtime/HairStrands/Source/HairStrandsCore/Private/GroomComponent.cpp
@@ -1946,7 +1946,10 @@ void UGroomComponent::InitResources(bool bIsBindingReloading)
 	}
 	else if (bHasNeedSkeletalMesh && ParentMeshComponent)
 	{
-		ValidatedMeshComponent = Cast<USkeletalMeshComponent>(ParentMeshComponent)->SkeletalMesh ? ParentMeshComponent : nullptr;
+		if (USkeletalMeshComponent* ParentSkelMeshComponent = Cast<USkeletalMeshComponent>(ParentMeshComponent))
+		{
+			ValidatedMeshComponent = ParentSkelMeshComponent->SkeletalMesh ? ParentMeshComponent : nullptr;
+		}
 	}
 
 	// Insure the ticking of the Groom component always happens after the skeletalMeshComponent.
