Add missing initialization of cards guides interpolation resources.

diff --git a/Engine/Plugins/Runtime/HairStrands/Source/HairStrandsCore/Private/GroomManager.cpp b/Engine/Plugins/Runtime/HairStrands/Source/HairStrandsCore/Private/GroomManager.cpp
--- a/Engine/Plugins/Runtime/HairStrands/Source/HairStrandsCore/Private/GroomManager.cpp
+++ b/Engine/Plugins/Runtime/HairStrands/Source/HairStrandsCore/Private/GroomManager.cpp
@@ -661,6 +661,7 @@ static void RunHairLODSelection(FRDGBuilder& GraphBuilder, const FHairStrandsIns
 			if (InstanceLOD.Guides.RestResource)			InstanceLOD.Guides.RestResource->Allocate(GraphBuilder);
 			if (InstanceLOD.Guides.DeformedRootResource)	InstanceLOD.Guides.DeformedRootResource->Allocate(GraphBuilder);
 			if (InstanceLOD.Guides.DeformedResource)		InstanceLOD.Guides.DeformedResource->Allocate(GraphBuilder);
+			if (InstanceLOD.Guides.InterpolationResource)	InstanceLOD.Guides.InterpolationResource->Allocate(GraphBuilder);
 			#if RHI_RAYTRACING
 			if (InstanceLOD.RaytracingResource)				InstanceLOD.RaytracingResource->Allocate(GraphBuilder);
 			#endif
