Fix invalid logic for updating groom offset point.

diff --git a/Engine/Plugins/Runtime/HairStrands/Source/HairStrandsCore/Private/HairStrandsMeshProjection.cpp b/Engine/Plugins/Runtime/HairStrands/Source/HairStrandsCore/Private/HairStrandsMeshProjection.cpp
--- a/Engine/Plugins/Runtime/HairStrands/Source/HairStrandsCore/Private/HairStrandsMeshProjection.cpp
+++ b/Engine/Plugins/Runtime/HairStrands/Source/HairStrandsCore/Private/HairStrandsMeshProjection.cpp
@@ -1768,12 +1768,11 @@ private:
 	DECLARE_GLOBAL_SHADER(FHairUpdatePositionOffsetCS);
 	SHADER_USE_PARAMETER_STRUCT(FHairUpdatePositionOffsetCS, FGlobalShader);
 
-	class FSourceType : SHADER_PERMUTATION_INT("PERMUTATION_SOURCE_TYPE", 2);
-	using FPermutationDomain = TShaderPermutationDomain<FSourceType>;
+	class FUseGPUOffset : SHADER_PERMUTATION_BOOL("PERMUTATION_USE_GPU_OFFSET");
+	using FPermutationDomain = TShaderPermutationDomain<FUseGPUOffset>;
 
 	BEGIN_SHADER_PARAMETER_STRUCT(FParameters, )
 		SHADER_PARAMETER(FVector3f, CPUPositionOffset)
-		SHADER_PARAMETER(uint32, bUseCPUOffset)
 		SHADER_PARAMETER_RDG_BUFFER_SRV(StructuredBuffer, RootTrianglePosition0Buffer)
 		SHADER_PARAMETER_RDG_BUFFER_UAV(RWBuffer, OutOffsetBuffer)
 	END_SHADER_PARAMETER_STRUCT()
@@ -1811,15 +1810,15 @@ void AddHairStrandUpdatePositionOffsetPass(
 		RootTrianglePositionBuffer = Register(GraphBuilder, DeformedRootResources->LODs[LODIndex].DeformedRootTrianglePosition0Buffer, ERDGImportedBufferFlags::CreateSRV);
 	}
 
+	const bool bUseGPUOffset = DeformedRootResources != nullptr && GHairStrandsUseGPUPositionOffset > 0;
 	const uint32 OffsetIndex = DeformedResources->GetIndex(FHairStrandsDeformedResource::Current);
 	FHairUpdatePositionOffsetCS::FParameters* PassParameters = GraphBuilder.AllocParameters<FHairUpdatePositionOffsetCS::FParameters>();
 	PassParameters->CPUPositionOffset = DeformedResources->PositionOffset[OffsetIndex];
-	PassParameters->bUseCPUOffset = (DeformedRootResources == nullptr || GHairStrandsUseGPUPositionOffset>0) ? 1 : 0;
 	PassParameters->RootTrianglePosition0Buffer = RootTrianglePositionBuffer.SRV;
 	PassParameters->OutOffsetBuffer = OutPositionOffsetBuffer.UAV;
 
 	FHairUpdatePositionOffsetCS::FPermutationDomain PermutationVector;
-	PermutationVector.Set<FHairUpdatePositionOffsetCS::FSourceType>(PassParameters->bUseCPUOffset ? 0 : 1);
+	PermutationVector.Set<FHairUpdatePositionOffsetCS::FUseGPUOffset>(bUseGPUOffset);
 	TShaderMapRef<FHairUpdatePositionOffsetCS> ComputeShader(ShaderMap, PermutationVector);
 	FComputeShaderUtils::AddPass(
 		GraphBuilder,
diff --git a/Engine/Shaders/Private/HairStrands/HairStrandsMeshUpdate.usf b/Engine/Shaders/Private/HairStrands/HairStrandsMeshUpdate.usf
--- a/Engine/Shaders/Private/HairStrands/HairStrandsMeshUpdate.usf
+++ b/Engine/Shaders/Private/HairStrands/HairStrandsMeshUpdate.usf
@@ -208,11 +208,20 @@ void MainCS(uint3 DispatchThreadId : SV_DispatchThreadID)
 
 #if SHADER_OFFSET_UPDATE
 
-float3 CPUPositionOffset;
-uint bUseCPUOffset;
-#if PERMUTATION_SOURCE_TYPE == 1
+#if PERMUTATION_USE_GPU_OFFSET
 Buffer<float4> RootTrianglePosition0Buffer;
+float3 GetOffsetPosition()
+{
+	return RootTrianglePosition0Buffer[0];
+}
+#else
+float3 CPUPositionOffset;
+float3 GetOffsetPosition()
+{
+	return CPUPositionOffset;
+}
 #endif
+
 RWBuffer<float4> OutOffsetBuffer;
 
 [numthreads(1, 1, 1)]
@@ -221,15 +230,7 @@ void MainCS(uint3 DispatchThreadId : SV_DispatchThreadID)
 	if (DispatchThreadId.x > 0)
 		return;
 
-	float3 RefPosition = CPUPositionOffset;
-	#if PERMUTATION_SOURCE_TYPE == 1
-	if (bUseCPUOffset == 0)
-	{
-		RefPosition = RootTrianglePosition0Buffer[0];
-	}
-	#endif
-
-	OutOffsetBuffer[0] = float4(RefPosition, 1);
+	OutOffsetBuffer[0] = float4(GetOffsetPosition(), 1);
 }
 
 #endif // SHADER_OFFSET_UPDATE
\ No newline at end of file
