Fix refresh of groom component transform when viewport's realtime option is turned off

diff --git a/Engine/Plugins/Runtime/HairStrands/Source/HairStrandsCore/Private/GroomComponent.cpp b/Engine/Plugins/Runtime/HairStrands/Source/HairStrandsCore/Private/GroomComponent.cpp
--- a/Engine/Plugins/Runtime/HairStrands/Source/HairStrandsCore/Private/GroomComponent.cpp
+++ b/Engine/Plugins/Runtime/HairStrands/Source/HairStrandsCore/Private/GroomComponent.cpp
@@ -466,6 +466,10 @@ public:
 					++MeshesLODIndex;
 				}
 			}
+
+			// Initialization of the default skel. mesh transform (refresh during ticking).
+			HairInstance->Debug.SkeletalLocalToWorld = HairInstance->Debug.SkeletalComponent ? HairInstance->Debug.SkeletalComponent->GetComponentTransform() : FTransform();
+			HairInstance->Debug.SkeletalPreviousLocalToWorld = HairInstance->Debug.SkeletalLocalToWorld;
 		}
 	}
 
