Fix merging issue.

diff --git a/Engine/Plugins/Runtime/HairStrands/Source/HairStrandsCore/Private/GroomComponent.cpp b/Engine/Plugins/Runtime/HairStrands/Source/HairStrandsCore/Private/GroomComponent.cpp
--- a/Engine/Plugins/Runtime/HairStrands/Source/HairStrandsCore/Private/GroomComponent.cpp
+++ b/Engine/Plugins/Runtime/HairStrands/Source/HairStrandsCore/Private/GroomComponent.cpp
@@ -468,7 +468,7 @@ public:
 			}
 
 			// Initialization of the default skel. mesh transform (refresh during ticking).
-			HairInstance->Debug.SkeletalLocalToWorld = HairInstance->Debug.SkeletalComponent ? HairInstance->Debug.SkeletalComponent->GetComponentTransform() : FTransform();
+			HairInstance->Debug.SkeletalLocalToWorld = HairInstance->Debug.MeshComponent ? HairInstance->Debug.MeshComponent->GetComponentTransform() : FTransform();
 			HairInstance->Debug.SkeletalPreviousLocalToWorld = HairInstance->Debug.SkeletalLocalToWorld;
 		}
 	}
