Add report of groom asset bulk data size.

diff --git a/Engine/Plugins/Runtime/HairStrands/Source/HairStrandsCore/Private/GroomAsset.cpp b/Engine/Plugins/Runtime/HairStrands/Source/HairStrandsCore/Private/GroomAsset.cpp
--- a/Engine/Plugins/Runtime/HairStrands/Source/HairStrandsCore/Private/GroomAsset.cpp
+++ b/Engine/Plugins/Runtime/HairStrands/Source/HairStrandsCore/Private/GroomAsset.cpp
@@ -3079,6 +3079,11 @@ void UGroomAsset::GetResourceSizeEx(FResourceSizeEx& CumulativeResourceSize)
 		CumulativeResourceSize.AddDedicatedVideoMemoryBytes(GroupData.Strands.GetResourcesSize());
 		CumulativeResourceSize.AddDedicatedVideoMemoryBytes(GroupData.Cards.GetResourcesSize());
 		CumulativeResourceSize.AddDedicatedVideoMemoryBytes(GroupData.Meshes.GetResourcesSize());
+
+		CumulativeResourceSize.AddDedicatedSystemMemoryBytes(GroupData.Guides.GetDataSize());
+		CumulativeResourceSize.AddDedicatedSystemMemoryBytes(GroupData.Strands.GetDataSize());
+		CumulativeResourceSize.AddDedicatedSystemMemoryBytes(GroupData.Cards.GetDataSize());
+		CumulativeResourceSize.AddDedicatedSystemMemoryBytes(GroupData.Meshes.GetDataSize());
 	}
 }
 
diff --git a/Engine/Plugins/Runtime/HairStrands/Source/HairStrandsCore/Public/GroomAsset.h b/Engine/Plugins/Runtime/HairStrands/Source/HairStrandsCore/Public/GroomAsset.h
--- a/Engine/Plugins/Runtime/HairStrands/Source/HairStrandsCore/Public/GroomAsset.h
+++ b/Engine/Plugins/Runtime/HairStrands/Source/HairStrandsCore/Public/GroomAsset.h
@@ -81,12 +81,32 @@ struct HAIRSTRANDSCORE_API FHairGroupData
 		bool IsValid() const			{ return RestResource != nullptr; }
 		const FBox& GetBounds() const	{ return BulkData.GetBounds(); }
 
+		uint32 GetDataSize() const
+		{
+			uint32 Total = 0;
+			Total += BulkData.Positions.GetAllocatedSize();
+			Total += BulkData.Attributes.GetAllocatedSize();
+			Total += BulkData.Materials.GetAllocatedSize();
+			Total += BulkData.CurveOffsets.GetAllocatedSize();
+			return Total;
+		}
+
 		FHairStrandsBulkData				BulkData;
 		FHairStrandsRestResource*			RestResource = nullptr;
 	};
 
 	struct FBaseWithInterpolation : FBase
 	{
+		uint32 GetDataSize() const
+		{
+			uint32 Total = 0;
+			Total += FBase::GetDataSize();
+			Total += InterpolationBulkData.Interpolation0.GetAllocatedSize();
+			Total += InterpolationBulkData.Interpolation1.GetAllocatedSize();
+			Total += InterpolationBulkData.SimRootPointIndex.GetAllocatedSize();
+			return Total;
+		}
+
 		FHairStrandsInterpolationBulkData	InterpolationBulkData;
 		FHairStrandsInterpolationResource*	InterpolationResource = nullptr;
 	};
@@ -103,6 +123,14 @@ struct HAIRSTRANDSCORE_API FHairGroupData
 			if (RestResource) Total += RestResource->GetResourcesSize();
 			return Total;
 		}
+
+		uint32 GetDataSize() const
+		{
+			uint32 Total = 0;
+			Total += FBase::GetDataSize();
+			return Total;
+		}
+
 	} Guides;
 
 	struct FStrands : FBaseWithInterpolation
@@ -120,6 +148,20 @@ struct HAIRSTRANDSCORE_API FHairGroupData
 			return Total;
 		}
 
+		uint32 GetDataSize() const
+		{
+			uint32 Total = 0;
+			Total += FBaseWithInterpolation::GetDataSize();
+
+			Total += ClusterCullingData.LODVisibility.GetAllocatedSize();
+			Total += ClusterCullingData.CPULODScreenSize.GetAllocatedSize();
+			Total += ClusterCullingData.ClusterInfos.GetAllocatedSize();
+			Total += ClusterCullingData.ClusterLODInfos.GetAllocatedSize();
+			Total += ClusterCullingData.VertexToClusterIds.GetAllocatedSize();
+			Total += ClusterCullingData.ClusterVertexIds.GetAllocatedSize();
+			return Total;
+		}
+
 		FHairStrandsClusterCullingData		ClusterCullingData;
 		FHairStrandsClusterCullingResource* ClusterCullingResource = nullptr;
 
@@ -163,6 +205,16 @@ struct HAIRSTRANDSCORE_API FHairGroupData
 			return Total;
 		}
 
+		uint32 GetDataSize() const
+		{
+			uint32 Total = 0;
+			for (const FLOD& LOD : LODs)
+			{
+				Total += LOD.GetDataSize();
+			}
+			return Total;
+		}
+
 		struct FLOD
 		{
 			/* Return the memory size for GPU resources */
@@ -177,6 +229,20 @@ struct HAIRSTRANDSCORE_API FHairGroupData
 				return Total;
 			}
 
+			uint32 GetDataSize() const
+			{
+				uint32 Total = 0;
+				Total += BulkData.Positions.GetAllocatedSize();
+				Total += BulkData.Normals.GetAllocatedSize();
+				Total += BulkData.UVs.GetAllocatedSize();
+				Total += BulkData.Indices.GetAllocatedSize();
+
+				Total += InterpolationBulkData.Interpolation.GetAllocatedSize();
+
+				Total += Guides.GetDataSize();
+				return Total;
+			}
+
 			bool HasValidData() const { return BulkData.IsValid(); }
 			bool IsValid() const { return BulkData.IsValid() && RestResource != nullptr; }
 
@@ -232,6 +298,16 @@ struct HAIRSTRANDSCORE_API FHairGroupData
 			return Total;
 		}
 
+		uint32 GetDataSize() const
+		{
+			uint32 Total = 0;
+			for (const FLOD& LOD : LODs)
+			{
+				Total += LOD.GetDataSize();
+			}
+			return Total;
+		}
+
 		struct FLOD
 		{
 			/* Return the memory size for GPU resources */
@@ -245,6 +321,16 @@ struct HAIRSTRANDSCORE_API FHairGroupData
 				return Total;
 			}
 
+			uint32 GetDataSize() const
+			{
+				uint32 Total = 0;
+				Total += BulkData.Positions.GetAllocatedSize();
+				Total += BulkData.Normals.GetAllocatedSize();
+				Total += BulkData.UVs.GetAllocatedSize();
+				Total += BulkData.Indices.GetAllocatedSize();
+				return Total;
+			}
+
 			bool HasValidData() const { return BulkData.IsValid(); }
 			bool IsValid() const { return BulkData.IsValid() && RestResource != nullptr; }
 
