Reduce groom asset size.

This CL splits attribute data into two Attribute0 and Attribute1 , with Attribute1 being optional,
GroomAsset Attribute1 and Material data are now optionnaly serialize, loaded, and added as resources. This is based if the hair description contains or not such a data.

This allows to save up to 8bytes per hair vertex. This reduces the GroomAsset footprint by ~18%

diff --git a/Engine/Plugins/Runtime/HairStrands/Source/HairStrandsCore/Private/GroomBuilder.cpp b/Engine/Plugins/Runtime/HairStrands/Source/HairStrandsCore/Private/GroomBuilder.cpp
--- a/Engine/Plugins/Runtime/HairStrands/Source/HairStrandsCore/Private/GroomBuilder.cpp
+++ b/Engine/Plugins/Runtime/HairStrands/Source/HairStrandsCore/Private/GroomBuilder.cpp
@@ -33,7 +33,7 @@ static FAutoConsoleVariableRef CVarHairClusterBuilder_MaxVoxelResolution(TEXT("r
 FString FGroomBuilder::GetVersion()
 {
 	// Important to update the version when groom building changes
-	return TEXT("1i");
+	return TEXT("2a");
 }
 
 namespace FHairStrandsDecimation
@@ -224,12 +224,14 @@ namespace HairStrandsBuilder
 			return;
 
 		TArray<FHairStrandsPositionFormat::Type>& OutPackedPositions = OutBulkData.Positions;
-		TArray<FHairStrandsAttributeFormat::Type>& OutPackedAttributes = OutBulkData.Attributes;
+		TArray<FHairStrandsAttribute0Format::Type>& OutPackedAttributes0 = OutBulkData.Attributes0;
+		TArray<FHairStrandsAttribute1Format::Type>& OutPackedAttributes1 = OutBulkData.Attributes1;
 		TArray<FHairStrandsMaterialFormat::Type>& OutPackedMaterials = OutBulkData.Materials;
 		TArray<FHairStrandsRootIndexFormat::Type>& OutCurveOffsets = OutBulkData.CurveOffsets;
 
 		OutPackedPositions.SetNum(NumPoints * FHairStrandsPositionFormat::ComponentCount);
-		OutPackedAttributes.SetNum(NumPoints * FHairStrandsAttributeFormat::ComponentCount);
+		OutPackedAttributes0.SetNum(NumPoints * FHairStrandsAttribute0Format::ComponentCount);
+		OutPackedAttributes1.SetNum(NumPoints * FHairStrandsAttribute1Format::ComponentCount);
 		OutPackedMaterials.SetNum(NumPoints * FHairStrandsMaterialFormat::ComponentCount);
 		OutCurveOffsets = HairStrands.StrandsCurves.CurvesOffset;
 
@@ -238,6 +240,10 @@ namespace HairStrandsBuilder
 		const FHairStrandsCurves& Curves = HairStrands.StrandsCurves;
 		const FHairStrandsPoints& Points = HairStrands.StrandsPoints;
 
+		// Track what features/data the bulk data contains
+		bool bHasUDIMData = false;
+		bool bHasMaterialData = false;
+
 		struct FPackedRadiusAndType
 		{
 			union
@@ -276,28 +282,46 @@ namespace HairStrandsBuilder
 				PackedPosition.UCoord = uint8(FMath::Clamp(CoordU * 255.f, 0.f, 255.f));
 
 				const FVector2D RootUV = Curves.CurvesRootUV[CurveIndex];
-				FHairStrandsAttributeFormat::Type& PackedAttributes = OutPackedAttributes[PointIndex + IndexOffset];
-				PackedAttributes.Unused0 = 0;
-				PackedAttributes.Unused1 = 0;
-				PackedAttributes.NormalizedLength = uint8(FMath::Clamp(NormalizedLength * 255.f, 0.f, 255.f));
-				PackedAttributes.Seed = CurveSeed;
+				FHairStrandsAttribute0Format::Type& PackedAttributes0 = OutPackedAttributes0[PointIndex + IndexOffset];
+				PackedAttributes0.NormalizedLength = uint8(FMath::Clamp(NormalizedLength * 255.f, 0.f, 255.f));
+				PackedAttributes0.Seed = CurveSeed;
 
 				// Root UV support UDIM texture coordinate but limit the spans of the UDIM to be in 256x256 instead of 9999x9999.
 				// The internal UV coords are also limited to 8bits, which means if sampling need to be super precise, this is no enough.
 				const FVector2D TextureRootUV(FMath::Fractional(RootUV.X), FMath::Fractional(RootUV.Y));
 				const FVector2D TextureIndexUV = RootUV - TextureRootUV;
-				PackedAttributes.RootU  = uint32(FMath::Clamp(TextureRootUV.X*255.f, 0.f, 255.f));
-				PackedAttributes.RootV  = uint32(FMath::Clamp(TextureRootUV.Y*255.f, 0.f, 255.f));
-				PackedAttributes.IndexU = uint32(FMath::Clamp(TextureIndexUV.X, 0.f, 255.f));
-				PackedAttributes.IndexV = uint32(FMath::Clamp(TextureIndexUV.Y, 0.f, 255.f));
-
-				FHairStrandsMaterialFormat::Type& Material = OutPackedMaterials[PointIndex + IndexOffset];
-				// Cheap sRGB encoding instead of PointsBaseColor.ToFColor(true), as this makes the decompression 
-				// cheaper on GPU (since R8G8B8A8 sRGB format used/exposed not exposed)
-				Material.BaseColorR = FMath::Clamp(uint32(FMath::Sqrt(Points.PointsBaseColor[PointIndex + IndexOffset].R) * 0xFF), 0u, 0xFFu);
-				Material.BaseColorG = FMath::Clamp(uint32(FMath::Sqrt(Points.PointsBaseColor[PointIndex + IndexOffset].G) * 0xFF), 0u, 0xFFu);
-				Material.BaseColorB = FMath::Clamp(uint32(FMath::Sqrt(Points.PointsBaseColor[PointIndex + IndexOffset].B) * 0xFF), 0u, 0xFFu);
-				Material.Roughness  = FMath::Clamp(uint32(Points.PointsRoughness[PointIndex + IndexOffset] * 0xFF), 0u, 0xFFu);
+				PackedAttributes0.RootU  = uint8(FMath::Clamp(TextureRootUV.X*255.f, 0.f, 255.f));
+				PackedAttributes0.RootV  = uint8(FMath::Clamp(TextureRootUV.Y*255.f, 0.f, 255.f));
+
+				// UDIM
+				{
+					FHairStrandsAttribute1Format::Type& PackedAttributes1 = OutPackedAttributes1[PointIndex + IndexOffset];
+					PackedAttributes1.IndexU = uint8(FMath::Clamp(TextureIndexUV.X, 0.f, 255.f));
+					PackedAttributes1.IndexV = uint8(FMath::Clamp(TextureIndexUV.Y, 0.f, 255.f));
+					PackedAttributes1.Unused0 = 0;
+					PackedAttributes1.Unused1 = 0;
+
+					if (PackedAttributes1.IndexU != 0 || PackedAttributes1.IndexV != 0)
+					{
+						bHasUDIMData = true;
+					}
+				}
+
+				// Material
+				{
+					FHairStrandsMaterialFormat::Type& Material = OutPackedMaterials[PointIndex + IndexOffset];
+					// Cheap sRGB encoding instead of PointsBaseColor.ToFColor(true), as this makes the decompression 
+					// cheaper on GPU (since R8G8B8A8 sRGB format used/exposed not exposed)
+					Material.BaseColorR = FMath::Clamp(uint32(FMath::Sqrt(Points.PointsBaseColor[PointIndex + IndexOffset].R) * 0xFF), 0u, 0xFFu);
+					Material.BaseColorG = FMath::Clamp(uint32(FMath::Sqrt(Points.PointsBaseColor[PointIndex + IndexOffset].G) * 0xFF), 0u, 0xFFu);
+					Material.BaseColorB = FMath::Clamp(uint32(FMath::Sqrt(Points.PointsBaseColor[PointIndex + IndexOffset].B) * 0xFF), 0u, 0xFFu);
+					Material.Roughness  = FMath::Clamp(uint32(Points.PointsRoughness[PointIndex + IndexOffset] * 0xFF), 0u, 0xFFu);
+
+					if (Material.BaseColorR != 0 || Material.BaseColorG != 0 || Material.BaseColorB != 0 || Material.Roughness != 0)
+					{
+						bHasMaterialData = true;
+					}
+				}
 			}
 		}
 
@@ -306,6 +330,12 @@ namespace HairStrandsBuilder
 		OutBulkData.PointCount = HairStrands.GetNumPoints();
 		OutBulkData.MaxLength = HairStrands.StrandsCurves.MaxLength;
 		OutBulkData.MaxRadius = HairStrands.StrandsCurves.MaxRadius;
+		OutBulkData.Flags = 0;
+
+		if (bHasMaterialData)	{ OutBulkData.Flags |= FHairStrandsBulkData::DataFlags_HasMaterialData; }
+		if (bHasUDIMData)		{ OutBulkData.Flags |= FHairStrandsBulkData::DataFlags_HasUDIMData; }
+		if (!bHasMaterialData)	{ OutBulkData.Attributes1.Empty(); }
+		if (!bHasUDIMData)		{ OutBulkData.Materials.Empty(); }
 	}
 
 	void BuildRenderData(const FHairStrandsDatas& HairStrands, FHairStrandsBulkData& OutBulkData)
diff --git a/Engine/Plugins/Runtime/HairStrands/Source/HairStrandsCore/Private/GroomResources.cpp b/Engine/Plugins/Runtime/HairStrands/Source/HairStrandsCore/Private/GroomResources.cpp
--- a/Engine/Plugins/Runtime/HairStrands/Source/HairStrandsCore/Private/GroomResources.cpp
+++ b/Engine/Plugins/Runtime/HairStrands/Source/HairStrandsCore/Private/GroomResources.cpp
@@ -300,7 +300,7 @@ void CreateHairStrandsDebugAttributeBuffer(FRDGBuilder& GraphBuilder, FRDGExtern
 {
 	if (VertexCount == 0 || !DebugAttributeBuffer)
 		return;
-	InternalCreateVertexBufferRDG<FHairStrandsAttributeFormat>(GraphBuilder, VertexCount, *DebugAttributeBuffer, TEXT("Hair.Strands_DebugAttributeBuffer"), EHairResourceUsageType::Dynamic);
+	InternalCreateVertexBufferRDG<FHairStrandsAttribute0Format>(GraphBuilder, VertexCount, *DebugAttributeBuffer, TEXT("Hair.Strands_DebugAttributeBuffer"), EHairResourceUsageType::Dynamic);
 }
 
 /////////////////////////////////////////////////////////////////////////////////////////
@@ -569,7 +569,7 @@ void FHairMeshesDeformedResource::InternalRelease()
 
 FHairStrandsRestResource::FHairStrandsRestResource(const FHairStrandsBulkData& InBulkData, EHairStrandsResourcesType InCurveType) :
 	FHairCommonResource(EHairStrandsAllocationType::Deferred),
-	PositionBuffer(), AttributeBuffer(), MaterialBuffer(), BulkData(InBulkData), CurveType(InCurveType)
+	PositionBuffer(), Attribute0Buffer(), Attribute1Buffer(), MaterialBuffer(), BulkData(InBulkData), CurveType(InCurveType)
 {}
 
 void FHairStrandsRestResource::InternalAllocate(FRDGBuilder& GraphBuilder)
@@ -577,8 +577,17 @@ void FHairStrandsRestResource::InternalAllocate(FRDGBuilder& GraphBuilder)
 	FRDGBufferUploader BufferUploader;
 
 	InternalCreateVertexBufferRDG<FHairStrandsPositionFormat>(GraphBuilder, BufferUploader, BulkData.Positions, PositionBuffer, HAIRSTRANDS_RESOUCE_NAME(CurveType, Hair.StrandsRest_PositionBuffer), EHairResourceUsageType::Static);
-	InternalCreateVertexBufferRDG<FHairStrandsAttributeFormat>(GraphBuilder, BufferUploader, BulkData.Attributes, AttributeBuffer, HAIRSTRANDS_RESOUCE_NAME(CurveType, Hair.StrandsRest_AttributeBuffer), EHairResourceUsageType::Static);
-	InternalCreateVertexBufferRDG<FHairStrandsMaterialFormat>(GraphBuilder, BufferUploader, BulkData.Materials, MaterialBuffer, HAIRSTRANDS_RESOUCE_NAME(CurveType, Hair.StrandsRest_MaterialBuffer), EHairResourceUsageType::Static);
+	InternalCreateVertexBufferRDG<FHairStrandsAttribute0Format>(GraphBuilder, BufferUploader, BulkData.Attributes0, Attribute0Buffer, HAIRSTRANDS_RESOUCE_NAME(CurveType, Hair.StrandsRest_Attribute0Buffer), EHairResourceUsageType::Static);
+
+	if (!!(BulkData.Flags & FHairStrandsBulkData::DataFlags_HasUDIMData))
+	{
+		InternalCreateVertexBufferRDG<FHairStrandsAttribute1Format>(GraphBuilder, BufferUploader, BulkData.Attributes1, Attribute1Buffer, HAIRSTRANDS_RESOUCE_NAME(CurveType, Hair.StrandsRest_Attribute1Buffer), EHairResourceUsageType::Static);
+	}
+
+	if (!!(BulkData.Flags & FHairStrandsBulkData::DataFlags_HasMaterialData))
+	{
+		InternalCreateVertexBufferRDG<FHairStrandsMaterialFormat>(GraphBuilder, BufferUploader, BulkData.Materials, MaterialBuffer, HAIRSTRANDS_RESOUCE_NAME(CurveType, Hair.StrandsRest_MaterialBuffer), EHairResourceUsageType::Static);
+	}
 
 	TArray<FVector4> RestOffset;
 	RestOffset.Add(BulkData.GetPositionOffset());
@@ -618,7 +627,8 @@ void FHairStrandsRestResource::InternalRelease()
 {
 	PositionBuffer.Release();
 	PositionOffsetBuffer.Release();
-	AttributeBuffer.Release();
+	Attribute0Buffer.Release();
+	Attribute1Buffer.Release();
 	MaterialBuffer.Release();
 	TangentBuffer.Release();
 }
diff --git a/Engine/Plugins/Runtime/HairStrands/Source/HairStrandsCore/Private/GroomTextureBuilder.cpp b/Engine/Plugins/Runtime/HairStrands/Source/HairStrandsCore/Private/GroomTextureBuilder.cpp
--- a/Engine/Plugins/Runtime/HairStrands/Source/HairStrandsCore/Private/GroomTextureBuilder.cpp
+++ b/Engine/Plugins/Runtime/HairStrands/Source/HairStrandsCore/Private/GroomTextureBuilder.cpp
@@ -581,7 +581,7 @@ class FHairStrandsTexturePS : public FGlobalShader
 		SHADER_PARAMETER(float, InVF_Length)
 		SHADER_PARAMETER(FVector3f, InVF_PositionOffset)
 		SHADER_PARAMETER_SRV(Buffer, InVF_PositionBuffer)
-		SHADER_PARAMETER_SRV(Buffer, InVF_AttributeBuffer)
+		SHADER_PARAMETER_SRV(Buffer, InVF_Attribute0Buffer)
 		SHADER_PARAMETER(uint32, InVF_ControlPointCount)
 
 		SHADER_PARAMETER(FVector3f, Voxel_MinBound)
@@ -667,7 +667,7 @@ static void InternalGenerateHairStrandsTextures(
 
 	ParametersPS->InVF_PositionBuffer = InHairStrands_PositionBuffer;
 	ParametersPS->InVF_PositionOffset = InHairStrands_PositionOffset;
-	ParametersPS->InVF_AttributeBuffer = InHairStrands_AttributeBuffer;
+	ParametersPS->InVF_Attribute0Buffer = InHairStrands_AttributeBuffer;
 	ParametersPS->InVF_Radius = InHairStrands_Radius;
 	ParametersPS->InVF_Length = InHairStrands_Length;
 	ParametersPS->InVF_ControlPointCount = InHairStrands_ControlPointCount;
@@ -983,7 +983,7 @@ static void InternalBuildStrandsTextures_GPU(
 				VoxelData,
 
 				GroupData.Strands.RestResource->PositionBuffer.SRV,
-				GroupData.Strands.RestResource->AttributeBuffer.SRV,
+				GroupData.Strands.RestResource->Attribute0Buffer.SRV,
 				GroupData.Strands.RestResource->GetPositionOffset(),
 				RenderingData.GeometrySettings.HairWidth * 0.5f,
 				GroupData.Strands.RestResource->BulkData.MaxLength,
diff --git a/Engine/Plugins/Runtime/HairStrands/Source/HairStrandsCore/Private/HairStrandsDatas.cpp b/Engine/Plugins/Runtime/HairStrands/Source/HairStrandsCore/Private/HairStrandsDatas.cpp
--- a/Engine/Plugins/Runtime/HairStrands/Source/HairStrandsCore/Private/HairStrandsDatas.cpp
+++ b/Engine/Plugins/Runtime/HairStrands/Source/HairStrandsCore/Private/HairStrandsDatas.cpp
@@ -72,7 +72,7 @@ FArchive& operator<<(FArchive& Ar, FPackedHairVertex& Vertex)
 	return Ar;
 }
 
-FArchive& operator<<(FArchive& Ar, FPackedHairAttributeVertex& Vertex)
+FArchive& operator<<(FArchive& Ar, FPackedHairAttribute0Vertex& Vertex)
 {
 	Ar.UsingCustomVersion(FReleaseObjectVersion::GUID);
 
@@ -81,21 +81,17 @@ FArchive& operator<<(FArchive& Ar, FPackedHairAttributeVertex& Vertex)
 	Ar << Vertex.NormalizedLength;
 	Ar << Vertex.Seed;
 
-	if (Ar.CustomVer(FReleaseObjectVersion::GUID) >= FReleaseObjectVersion::GroomAssetVersion1)
-	{
-		Ar << Vertex.IndexU;
-		Ar << Vertex.IndexV;
-		Ar << Vertex.Unused0;
-		Ar << Vertex.Unused1;
-	}
-	else
-	{
-		Vertex.IndexU = 0;
-		Vertex.IndexV = 0;
-		Vertex.Unused0 = 0;
-		Vertex.Unused1 = 0;
-	}
-	
+	return Ar;
+}
+
+FArchive& operator<<(FArchive& Ar, FPackedHairAttribute1Vertex& Vertex)
+{
+	Ar.UsingCustomVersion(FReleaseObjectVersion::GUID);
+
+	Ar << Vertex.IndexU;
+	Ar << Vertex.IndexV;
+	Ar << Vertex.Unused0;
+	Ar << Vertex.Unused1;
 
 	return Ar;
 }
@@ -172,24 +168,32 @@ void FHairStrandsInterpolationBulkData::Serialize(FArchive& Ar)
 void FHairStrandsBulkData::Serialize(FArchive& Ar)
 {
 	static_assert(sizeof(FHairStrandsPositionFormat::BulkType) == sizeof(FHairStrandsPositionFormat::Type));
-	static_assert(sizeof(FHairStrandsAttributeFormat::BulkType) == sizeof(FHairStrandsAttributeFormat::Type));
+	static_assert(sizeof(FHairStrandsAttribute0Format::BulkType) == sizeof(FHairStrandsAttribute0Format::Type));
+	static_assert(sizeof(FHairStrandsAttribute1Format::BulkType) == sizeof(FHairStrandsAttribute1Format::Type));
 	static_assert(sizeof(FHairStrandsMaterialFormat::BulkType) == sizeof(FHairStrandsMaterialFormat::Type));
 	static_assert(sizeof(FHairStrandsRootIndexFormat::BulkType) == sizeof(FHairStrandsRootIndexFormat::Type)); 
 
 	Ar.UsingCustomVersion(FReleaseObjectVersion::GUID);
 	Ar.UsingCustomVersion(FUE5ReleaseStreamObjectVersion::GUID);
 
-	Positions.BulkSerialize(Ar);
-	Attributes.BulkSerialize(Ar);
-	Materials.BulkSerialize(Ar);
-	CurveOffsets.BulkSerialize(Ar);
-
 	Ar << CurveCount;
 	Ar << PointCount;
 	Ar << MaxLength;
 	Ar << MaxRadius;
 	Ar << BoundingBox;
+	Ar << Flags;
 
+	Positions.BulkSerialize(Ar);
+	Attributes0.BulkSerialize(Ar);
+	if (!!(Flags & DataFlags_HasUDIMData))
+	{
+		Attributes1.BulkSerialize(Ar);
+	}
+	if (!!(Flags & DataFlags_HasMaterialData))
+	{
+		Materials.BulkSerialize(Ar);
+	}
+	CurveOffsets.BulkSerialize(Ar);
 }
 
 void FHairStrandsDatas::Reset()
diff --git a/Engine/Plugins/Runtime/HairStrands/Source/HairStrandsCore/Private/HairStrandsRendering.cpp b/Engine/Plugins/Runtime/HairStrands/Source/HairStrandsCore/Private/HairStrandsRendering.cpp
--- a/Engine/Plugins/Runtime/HairStrands/Source/HairStrandsCore/Private/HairStrandsRendering.cpp
+++ b/Engine/Plugins/Runtime/HairStrands/Source/HairStrandsCore/Private/HairStrandsRendering.cpp
@@ -1407,7 +1407,8 @@ static FHairGroupPublicData::FVertexFactoryInput InternalComputeHairStrandsVerte
 		CONVERT_HAIRSSTRANDS_VF_PARAMETERS(OutVFInput.Strands.PositionBuffer,			Instance->Guides.DeformedResource->GetBuffer(FHairStrandsDeformedResource::EFrameType::Current));
 		CONVERT_HAIRSSTRANDS_VF_PARAMETERS(OutVFInput.Strands.PrevPositionBuffer,		Instance->Guides.DeformedResource->GetBuffer(FHairStrandsDeformedResource::EFrameType::Previous));
 		CONVERT_HAIRSSTRANDS_VF_PARAMETERS(OutVFInput.Strands.TangentBuffer,			Instance->Guides.DeformedResource->TangentBuffer);
-		CONVERT_HAIRSSTRANDS_VF_PARAMETERS(OutVFInput.Strands.AttributeBuffer,			Instance->Guides.RestResource->AttributeBuffer);
+		CONVERT_HAIRSSTRANDS_VF_PARAMETERS(OutVFInput.Strands.Attribute0Buffer,			Instance->Guides.RestResource->Attribute0Buffer);
+		CONVERT_HAIRSSTRANDS_VF_PARAMETERS(OutVFInput.Strands.Attribute1Buffer,			Instance->Guides.RestResource->Attribute1Buffer);
 		CONVERT_HAIRSSTRANDS_VF_PARAMETERS(OutVFInput.Strands.MaterialBuffer,			Instance->Guides.RestResource->MaterialBuffer);
 		CONVERT_HAIRSSTRANDS_VF_PARAMETERS(OutVFInput.Strands.PositionOffsetBuffer,		Instance->Guides.DeformedResource->GetPositionOffsetBuffer(FHairStrandsDeformedResource::EFrameType::Current));
 		CONVERT_HAIRSSTRANDS_VF_PARAMETERS(OutVFInput.Strands.PrevPositionOffsetBuffer,	Instance->Guides.DeformedResource->GetPositionOffsetBuffer(FHairStrandsDeformedResource::EFrameType::Previous));
@@ -1422,7 +1423,8 @@ static FHairGroupPublicData::FVertexFactoryInput InternalComputeHairStrandsVerte
 		CONVERT_HAIRSSTRANDS_VF_PARAMETERS(OutVFInput.Strands.PositionBuffer,			Instance->Strands.DeformedResource->GetBuffer(FHairStrandsDeformedResource::EFrameType::Current));
 		CONVERT_HAIRSSTRANDS_VF_PARAMETERS(OutVFInput.Strands.PrevPositionBuffer,		Instance->Strands.DeformedResource->GetBuffer(FHairStrandsDeformedResource::EFrameType::Previous));
 		CONVERT_HAIRSSTRANDS_VF_PARAMETERS(OutVFInput.Strands.TangentBuffer,			Instance->Strands.DeformedResource->TangentBuffer);
-		CONVERT_HAIRSSTRANDS_VF_PARAMETERS(OutVFInput.Strands.AttributeBuffer,			/*bDebugModePatchedAttributeBuffer ? Instance->Strands.DebugAttributeBuffer : */Instance->Strands.RestResource->AttributeBuffer); // TODO
+		CONVERT_HAIRSSTRANDS_VF_PARAMETERS(OutVFInput.Strands.Attribute0Buffer,			/*bDebugModePatchedAttributeBuffer ? Instance->Strands.DebugAttributeBuffer : */Instance->Strands.RestResource->Attribute0Buffer); // TODO
+		CONVERT_HAIRSSTRANDS_VF_PARAMETERS(OutVFInput.Strands.Attribute1Buffer,			Instance->Strands.RestResource->Attribute1Buffer);
 		CONVERT_HAIRSSTRANDS_VF_PARAMETERS(OutVFInput.Strands.MaterialBuffer,			Instance->Strands.RestResource->MaterialBuffer);
 		CONVERT_HAIRSSTRANDS_VF_PARAMETERS(OutVFInput.Strands.PositionOffsetBuffer,		Instance->Strands.DeformedResource->GetPositionOffsetBuffer(FHairStrandsDeformedResource::EFrameType::Current));
 		CONVERT_HAIRSSTRANDS_VF_PARAMETERS(OutVFInput.Strands.PrevPositionOffsetBuffer,	Instance->Strands.DeformedResource->GetPositionOffsetBuffer(FHairStrandsDeformedResource::EFrameType::Previous));
@@ -1437,7 +1439,8 @@ static FHairGroupPublicData::FVertexFactoryInput InternalComputeHairStrandsVerte
 		CONVERT_HAIRSSTRANDS_VF_PARAMETERS(OutVFInput.Strands.PositionBuffer,			Instance->Strands.RestResource->PositionBuffer);
 		CONVERT_HAIRSSTRANDS_VF_PARAMETERS(OutVFInput.Strands.PrevPositionBuffer,		Instance->Strands.RestResource->PositionBuffer);
 		CONVERT_HAIRSSTRANDS_VF_PARAMETERS(OutVFInput.Strands.TangentBuffer,			Instance->Strands.RestResource->TangentBuffer);
-		CONVERT_HAIRSSTRANDS_VF_PARAMETERS(OutVFInput.Strands.AttributeBuffer,			Instance->Strands.RestResource->AttributeBuffer);
+		CONVERT_HAIRSSTRANDS_VF_PARAMETERS(OutVFInput.Strands.Attribute0Buffer,			Instance->Strands.RestResource->Attribute0Buffer);
+		CONVERT_HAIRSSTRANDS_VF_PARAMETERS(OutVFInput.Strands.Attribute1Buffer,			Instance->Strands.RestResource->Attribute1Buffer);
 		CONVERT_HAIRSSTRANDS_VF_PARAMETERS(OutVFInput.Strands.MaterialBuffer,			Instance->Strands.RestResource->MaterialBuffer);
 		CONVERT_HAIRSSTRANDS_VF_PARAMETERS(OutVFInput.Strands.PositionOffsetBuffer,		Instance->Strands.RestResource->PositionOffsetBuffer);
 		CONVERT_HAIRSSTRANDS_VF_PARAMETERS(OutVFInput.Strands.PrevPositionOffsetBuffer, Instance->Strands.RestResource->PositionOffsetBuffer);
@@ -1620,12 +1623,12 @@ void ComputeHairStrandsInterpolation(
 						bHasSkinning ? Instance->Strands.DeformedRootResource : nullptr,
 						bHasSkinning && bValidGuide ? Instance->Guides.DeformedRootResource : nullptr,
 						RegisterAsSRV(GraphBuilder, Instance->Strands.RestResource->PositionBuffer),
-						RegisterAsSRV(GraphBuilder, Instance->Strands.RestResource->AttributeBuffer),
+						RegisterAsSRV(GraphBuilder, Instance->Strands.RestResource->Attribute0Buffer),
 						bValidGuide ? RegisterAsSRV(GraphBuilder, Instance->Strands.InterpolationResource->Interpolation0Buffer) : nullptr,
 						bValidGuide ? RegisterAsSRV(GraphBuilder, Instance->Strands.InterpolationResource->Interpolation1Buffer) : nullptr,
 						bValidGuide ? RegisterAsSRV(GraphBuilder, Instance->Guides.RestResource->PositionBuffer) : nullptr,
 						bValidGuide ? RegisterAsSRV(GraphBuilder, Instance->Guides.DeformedResource->GetBuffer(FHairStrandsDeformedResource::Current)) : nullptr,
-						bValidGuide ? RegisterAsSRV(GraphBuilder, Instance->Guides.RestResource->AttributeBuffer) : nullptr,
+						bValidGuide ? RegisterAsSRV(GraphBuilder, Instance->Guides.RestResource->Attribute0Buffer) : nullptr,
 						Strands_DeformedPosition,
 						&Strands_DeformedPrevPosition,
 						Strands_DebugAttribute,
@@ -1814,12 +1817,12 @@ void ComputeHairStrandsInterpolation(
 						bHasSkinning ? LOD.Guides.DeformedRootResource : nullptr,
 						bHasSkinning && bValidGuide ? Instance->Guides.DeformedRootResource : nullptr,
 						RegisterAsSRV(GraphBuilder, LOD.Guides.RestResource->PositionBuffer),
-						RegisterAsSRV(GraphBuilder, LOD.Guides.RestResource->AttributeBuffer),
+						RegisterAsSRV(GraphBuilder, LOD.Guides.RestResource->Attribute0Buffer),
 						RegisterAsSRV(GraphBuilder, LOD.Guides.InterpolationResource->Interpolation0Buffer),
 						RegisterAsSRV(GraphBuilder, LOD.Guides.InterpolationResource->Interpolation1Buffer),
 						bValidGuide ? RegisterAsSRV(GraphBuilder, Instance->Guides.RestResource->PositionBuffer) : nullptr,
 						bValidGuide ? RegisterAsSRV(GraphBuilder, Instance->Guides.DeformedResource->GetBuffer(FHairStrandsDeformedResource::Current)) : nullptr,
-						bValidGuide ? RegisterAsSRV(GraphBuilder, Instance->Guides.RestResource->AttributeBuffer) : nullptr,
+						bValidGuide ? RegisterAsSRV(GraphBuilder, Instance->Guides.RestResource->Attribute0Buffer) : nullptr,
 						Guides_DeformedPositionBuffer,
 						nullptr,
 						nullptr,
diff --git a/Engine/Plugins/Runtime/HairStrands/Source/HairStrandsCore/Private/HairStrandsVertexFactory.cpp b/Engine/Plugins/Runtime/HairStrands/Source/HairStrandsCore/Private/HairStrandsVertexFactory.cpp
--- a/Engine/Plugins/Runtime/HairStrands/Source/HairStrandsCore/Private/HairStrandsVertexFactory.cpp
+++ b/Engine/Plugins/Runtime/HairStrands/Source/HairStrandsCore/Private/HairStrandsVertexFactory.cpp
@@ -25,6 +25,7 @@ class FDummyCulledDispatchVertexIdsBuffer : public FVertexBuffer
 public:
 	FShaderResourceViewRHIRef SRVUint;
 	FShaderResourceViewRHIRef SRVFloat;
+	FShaderResourceViewRHIRef SRVRGBA;
 
 	virtual void InitRHI() override
 	{
@@ -37,6 +38,7 @@ public:
 
 		SRVUint = RHICreateShaderResourceView(VertexBufferRHI, sizeof(uint32), PF_R32_UINT);
 		SRVFloat = RHICreateShaderResourceView(VertexBufferRHI, sizeof(uint32), PF_R32_FLOAT);
+		SRVRGBA = RHICreateShaderResourceView(VertexBufferRHI, sizeof(uint32), PF_R8G8B8A8);
 	}
 
 	virtual void ReleaseRHI() override
@@ -44,6 +46,7 @@ public:
 		VertexBufferRHI.SafeRelease();
 		SRVUint.SafeRelease();
 		SRVFloat.SafeRelease();
+		SRVRGBA.SafeRelease();
 	}
 };
 TGlobalResource<FDummyCulledDispatchVertexIdsBuffer> GDummyCulledDispatchVertexIdsBuffer;
@@ -51,7 +54,8 @@ TGlobalResource<FDummyCulledDispatchVertexIdsBuffer> GDummyCulledDispatchVertexI
 /////////////////////////////////////////////////////////////////////////////////////////
 inline FRHIShaderResourceView* GetPositionSRV(const FHairGroupPublicData::FVertexFactoryInput& VFInput)			{ return VFInput.Strands.PositionBufferRHISRV; }; 
 inline FRHIShaderResourceView* GetPreviousPositionSRV(const FHairGroupPublicData::FVertexFactoryInput& VFInput)	{ return VFInput.Strands.PrevPositionBufferRHISRV; }
-inline FRHIShaderResourceView* GetAttributeSRV(const FHairGroupPublicData::FVertexFactoryInput& VFInput)		{ return VFInput.Strands.AttributeBufferRHISRV; }
+inline FRHIShaderResourceView* GetAttribute0SRV(const FHairGroupPublicData::FVertexFactoryInput& VFInput)		{ return VFInput.Strands.Attribute0BufferRHISRV; }
+inline FRHIShaderResourceView* GetAttribute1SRV(const FHairGroupPublicData::FVertexFactoryInput& VFInput)		{ return VFInput.Strands.Attribute1BufferRHISRV; }
 inline FRHIShaderResourceView* GetMaterialSRV(const FHairGroupPublicData::FVertexFactoryInput& VFInput)			{ return VFInput.Strands.MaterialBufferRHISRV; }
 inline FRHIShaderResourceView* GetTangentSRV(const FHairGroupPublicData::FVertexFactoryInput& VFInput)			{ return VFInput.Strands.TangentBufferRHISRV; }
 
@@ -65,6 +69,8 @@ inline float GetStrandRootScale(const FHairGroupPublicData::FVertexFactoryInput&
 inline float GetStrandTipScale(const FHairGroupPublicData::FVertexFactoryInput& VFInput)						{ return VFInput.Strands.HairTipScale; };
 inline float GetMaxStrandLength(const FHairGroupPublicData::FVertexFactoryInput& VFInput)						{ return VFInput.Strands.HairLength; };
 inline float GetHairDensity(const FHairGroupPublicData::FVertexFactoryInput& VFInput)							{ return VFInput.Strands.HairDensity; };
+//inline float GetHasAttribute1(const FHairGroupPublicData::FVertexFactoryInput& VFInput)							{ return VFInput.Strands.HasAttribute1; };
+//inline float GetHasMaterial(const FHairGroupPublicData::FVertexFactoryInput& VFInput)							{ return VFInput.Strands.HasMaterial; };
 
 /////////////////////////////////////////////////////////////////////////////////////////
 
@@ -83,6 +89,8 @@ public:
 	LAYOUT_FIELD(FShaderParameter, RadiusAtDepth1_Velocity);// unused
 	LAYOUT_FIELD(FShaderParameter, Density);
 	LAYOUT_FIELD(FShaderParameter, Culling);
+	LAYOUT_FIELD(FShaderParameter, HasAttribute1);
+	LAYOUT_FIELD(FShaderParameter, HasMaterial);
 	LAYOUT_FIELD(FShaderParameter, StableRasterization);
 	LAYOUT_FIELD(FShaderParameter, ScatterSceneLighing);
 
@@ -90,7 +98,8 @@ public:
 	LAYOUT_FIELD(FShaderResourceParameter, PreviousPositionOffsetBuffer);
 	LAYOUT_FIELD(FShaderResourceParameter, PositionBuffer);
 	LAYOUT_FIELD(FShaderResourceParameter, PreviousPositionBuffer);
-	LAYOUT_FIELD(FShaderResourceParameter, AttributeBuffer);
+	LAYOUT_FIELD(FShaderResourceParameter, Attribute0Buffer);
+	LAYOUT_FIELD(FShaderResourceParameter, Attribute1Buffer);
 	LAYOUT_FIELD(FShaderResourceParameter, MaterialBuffer);
 	LAYOUT_FIELD(FShaderResourceParameter, TangentBuffer);
 
@@ -104,8 +113,9 @@ public:
 		TipScale.Bind(ParameterMap, TEXT("HairStrandsVF_TipScale"));
 		Length.Bind(ParameterMap, TEXT("HairStrandsVF_Length"));
 		Density.Bind(ParameterMap, TEXT("HairStrandsVF_Density"));
-		Density.Bind(ParameterMap, TEXT("HairStrandsVF_Density"));	
 		Culling.Bind(ParameterMap, TEXT("HairStrandsVF_CullingEnable"));
+		HasAttribute1.Bind(ParameterMap, TEXT("HairStrandsVF_HasAttribute1"));
+		HasMaterial.Bind(ParameterMap, TEXT("HairStrandsVF_HasMaterial"));
 		StableRasterization.Bind(ParameterMap, TEXT("HairStrandsVF_bUseStableRasterization"));
 		ScatterSceneLighing.Bind(ParameterMap, TEXT("HairStrandsVF_bScatterSceneLighing"));
 
@@ -114,7 +124,8 @@ public:
 
 		PositionBuffer.Bind(ParameterMap, TEXT("HairStrandsVF_PositionBuffer"));
 		PreviousPositionBuffer.Bind(ParameterMap, TEXT("HairStrandsVF_PreviousPositionBuffer"));
-		AttributeBuffer.Bind(ParameterMap, TEXT("HairStrandsVF_AttributeBuffer"));
+		Attribute0Buffer.Bind(ParameterMap, TEXT("HairStrandsVF_Attribute0Buffer"));
+		Attribute1Buffer.Bind(ParameterMap, TEXT("HairStrandsVF_Attribute1Buffer"));
 		MaterialBuffer.Bind(ParameterMap, TEXT("HairStrandsVF_MaterialBuffer"));
 		TangentBuffer.Bind(ParameterMap, TEXT("HairStrandsVF_TangentBuffer"));
 
@@ -137,10 +148,25 @@ public:
 		const FHairStrandsVertexFactory* VF = static_cast<const FHairStrandsVertexFactory*>(VertexFactory);
 		const FHairGroupPublicData::FVertexFactoryInput VFInput = ComputeHairStrandsVertexInputData(VF->Data.Instance);
 
+		FShaderResourceViewRHIRef Attribute1SRV = GetAttribute1SRV(VFInput);
+		const bool bHasAttribute1 = Attribute1SRV != nullptr;
+		if (!bHasAttribute1)
+		{
+			Attribute1SRV = GDummyCulledDispatchVertexIdsBuffer.SRVRGBA;
+		}
+
+		FShaderResourceViewRHIRef MaterialSRV = GetMaterialSRV(VFInput);
+		const bool bHasMaterial = MaterialSRV != nullptr;
+		if (!bHasMaterial)
+		{
+			MaterialSRV = GDummyCulledDispatchVertexIdsBuffer.SRVRGBA;
+		}
+
 		VFS_BindParam(ShaderBindings, PositionBuffer, GetPositionSRV(VFInput));
 		VFS_BindParam(ShaderBindings, PreviousPositionBuffer, GetPreviousPositionSRV(VFInput));
-		VFS_BindParam(ShaderBindings, AttributeBuffer, GetAttributeSRV(VFInput));
-		VFS_BindParam(ShaderBindings, MaterialBuffer, GetMaterialSRV(VFInput));
+		VFS_BindParam(ShaderBindings, Attribute0Buffer, GetAttribute0SRV(VFInput));
+		VFS_BindParam(ShaderBindings, Attribute1Buffer, Attribute1SRV.GetReference());
+		VFS_BindParam(ShaderBindings, MaterialBuffer, MaterialSRV.GetReference());
 		VFS_BindParam(ShaderBindings, TangentBuffer, GetTangentSRV(VFInput));
 		VFS_BindParam(ShaderBindings, Radius, GetMaxStrandRadius(VFInput));
 		VFS_BindParam(ShaderBindings, RootScale, GetStrandRootScale(VFInput));
@@ -151,6 +177,8 @@ public:
 		VFS_BindParam(ShaderBindings, Density, GetHairDensity(VFInput));
 		VFS_BindParam(ShaderBindings, StableRasterization, UseStableRasterization(VFInput) ? 1u : 0u);
 		VFS_BindParam(ShaderBindings, ScatterSceneLighing, UseScatterSceneLighting(VFInput) ? 1u : 0u);
+		VFS_BindParam(ShaderBindings, HasAttribute1, bHasAttribute1 ? 1u : 0u);
+		VFS_BindParam(ShaderBindings, HasMaterial, bHasMaterial ? 1u : 0u);
 		
 		FShaderResourceViewRHIRef CulledDispatchVertexIdsSRV = GDummyCulledDispatchVertexIdsBuffer.SRVUint;
 		FShaderResourceViewRHIRef CulledCompactedRadiusScaleBufferSRV = GDummyCulledDispatchVertexIdsBuffer.SRVFloat;
diff --git a/Engine/Plugins/Runtime/HairStrands/Source/HairStrandsCore/Public/GroomAsset.h b/Engine/Plugins/Runtime/HairStrands/Source/HairStrandsCore/Public/GroomAsset.h
--- a/Engine/Plugins/Runtime/HairStrands/Source/HairStrandsCore/Public/GroomAsset.h
+++ b/Engine/Plugins/Runtime/HairStrands/Source/HairStrandsCore/Public/GroomAsset.h
@@ -85,7 +85,8 @@ struct HAIRSTRANDSCORE_API FHairGroupData
 		{
 			uint32 Total = 0;
 			Total += BulkData.Positions.GetAllocatedSize();
-			Total += BulkData.Attributes.GetAllocatedSize();
+			Total += BulkData.Attributes0.GetAllocatedSize();
+			Total += BulkData.Attributes1.GetAllocatedSize();
 			Total += BulkData.Materials.GetAllocatedSize();
 			Total += BulkData.CurveOffsets.GetAllocatedSize();
 			return Total;
diff --git a/Engine/Plugins/Runtime/HairStrands/Source/HairStrandsCore/Public/GroomResources.h b/Engine/Plugins/Runtime/HairStrands/Source/HairStrandsCore/Public/GroomResources.h
--- a/Engine/Plugins/Runtime/HairStrands/Source/HairStrandsCore/Public/GroomResources.h
+++ b/Engine/Plugins/Runtime/HairStrands/Source/HairStrandsCore/Public/GroomResources.h
@@ -212,7 +212,8 @@ struct FHairStrandsRestResource : public FHairCommonResource
 		uint32 Total = 0;
 		Total += GetBufferTotalNumBytes(PositionBuffer);
 		Total += GetBufferTotalNumBytes(PositionOffsetBuffer);
-		Total += GetBufferTotalNumBytes(AttributeBuffer);
+		Total += GetBufferTotalNumBytes(Attribute1Buffer);
+		Total += GetBufferTotalNumBytes(Attribute0Buffer);
 		Total += GetBufferTotalNumBytes(MaterialBuffer);
 		Total += GetBufferTotalNumBytes(TangentBuffer);
 		return Total;
@@ -228,7 +229,10 @@ struct FHairStrandsRestResource : public FHairCommonResource
 	FRDGExternalBuffer TangentBuffer;
 
 	/* Strand hair attribute buffer */
-	FRDGExternalBuffer AttributeBuffer;
+	FRDGExternalBuffer Attribute0Buffer;
+
+	/* Strand hair attribute buffer */
+	FRDGExternalBuffer Attribute1Buffer;
 
 	/* Strand hair material buffer */
 	FRDGExternalBuffer MaterialBuffer;
diff --git a/Engine/Plugins/Runtime/HairStrands/Source/HairStrandsCore/Public/HairStrandsDatas.h b/Engine/Plugins/Runtime/HairStrands/Source/HairStrandsCore/Public/HairStrandsDatas.h
--- a/Engine/Plugins/Runtime/HairStrands/Source/HairStrandsCore/Public/HairStrandsDatas.h
+++ b/Engine/Plugins/Runtime/HairStrands/Source/HairStrandsCore/Public/HairStrandsDatas.h
@@ -19,14 +19,19 @@ struct FPackedHairVertex
 	uint8 UCoord;
 };
 
-struct FPackedHairAttributeVertex
+struct FPackedHairAttribute0Vertex
 {
-	typedef uint64 BulkType;
+	typedef uint32 BulkType;
 
 	uint8 RootU;
 	uint8 RootV;
 	uint8 NormalizedLength;
 	uint8 Seed;
+};
+
+struct FPackedHairAttribute1Vertex
+{
+	typedef uint32 BulkType;
 
 	uint8 IndexU;
 	uint8 IndexV;
@@ -102,14 +107,24 @@ struct FHairStrandsPositionOffsetFormat
 	static const EPixelFormat Format = PF_A32B32G32R32F;
 };
 
-struct FHairStrandsAttributeFormat
+struct FHairStrandsAttribute0Format
 {
-	typedef FPackedHairAttributeVertex Type;
-	typedef FPackedHairAttributeVertex::BulkType BulkType;
+	typedef FPackedHairAttribute0Vertex Type;
+	typedef FPackedHairAttribute0Vertex::BulkType BulkType;
 	static const uint32 ComponentCount = 1;
 	static const uint32 SizeInByte = sizeof(Type);
-	static const EVertexElementType VertexElementType = VET_UShort4;
-	static const EPixelFormat Format = PF_R16G16B16A16_UINT;
+	static const EVertexElementType VertexElementType = VET_UByte4;
+	static const EPixelFormat Format = PF_R8G8B8A8;
+};
+
+struct FHairStrandsAttribute1Format
+{
+	typedef FPackedHairAttribute1Vertex Type;
+	typedef FPackedHairAttribute1Vertex::BulkType BulkType;
+	static const uint32 ComponentCount = 1;
+	static const uint32 SizeInByte = sizeof(Type);
+	static const EVertexElementType VertexElementType = VET_UByte4;
+	static const EPixelFormat Format = PF_R8G8B8A8_UINT;
 };
 
 struct FHairStrandsMaterialFormat
@@ -373,6 +388,12 @@ struct HAIRSTRANDSCORE_API FHairStrandsDatas
 
 struct HAIRSTRANDSCORE_API FHairStrandsBulkData
 {
+	enum EDataFlags
+	{
+		DataFlags_HasUDIMData = 1,
+		DataFlags_HasMaterialData = 2
+	};
+
 	void Serialize(FArchive& Ar);
 
 	bool IsValid() const { return CurveCount > 0 && PointCount > 0;	}
@@ -389,9 +410,11 @@ struct HAIRSTRANDSCORE_API FHairStrandsBulkData
 	float MaxLength = 0;
 	float MaxRadius = 0;
 	FBox BoundingBox;
+	uint32 Flags = 0;
 
 	TArray<FHairStrandsPositionFormat::Type>	Positions;		// Size = PointCount
-	TArray<FHairStrandsAttributeFormat::Type>	Attributes;		// Size = PointCount
+	TArray<FHairStrandsAttribute0Format::Type>	Attributes0;	// Size = PointCount
+	TArray<FHairStrandsAttribute1Format::Type>	Attributes1;	// Size = PointCount
 	TArray<FHairStrandsMaterialFormat::Type>	Materials;		// Size = PointCount
 	TArray<FHairStrandsRootIndexFormat::Type>	CurveOffsets;	// Size = CurveCount+1 - Store the root point index for the curve
 };
diff --git a/Engine/Shaders/Private/HairStrands/HairStrandsAttributeCommon.ush b/Engine/Shaders/Private/HairStrands/HairStrandsAttributeCommon.ush
--- a/Engine/Shaders/Private/HairStrands/HairStrandsAttributeCommon.ush
+++ b/Engine/Shaders/Private/HairStrands/HairStrandsAttributeCommon.ush
@@ -12,9 +12,12 @@ float  HairStrandsVF_Radius;
 float  HairStrandsVF_RootScale;
 float  HairStrandsVF_TipScale;
 float  HairStrandsVF_Length;
+uint   HairStrandsVF_HasAttribute1;
+uint   HairStrandsVF_HasMaterial;
 
 Buffer<uint4>  HairStrandsVF_PositionBuffer;
-Buffer<uint4>  HairStrandsVF_AttributeBuffer;
+Buffer<float4> HairStrandsVF_Attribute0Buffer;
+Buffer<uint4>  HairStrandsVF_Attribute1Buffer;
 Buffer<float4> HairStrandsVF_MaterialBuffer;
 
 struct FHairVertexAttributes
@@ -27,14 +30,34 @@ struct FHairVertexAttributes
 	float  HairRoughness;
 };
 
+uint4 GetHairStrandsOptionalAttribute1(uint InVertexId)
+{
+	uint4 Out = 0;
+	if (HairStrandsVF_HasAttribute1 > 0)
+	{
+		Out = HairStrandsVF_Attribute1Buffer[InVertexId];
+	}
+	return Out;
+}
+
+float4 GetHairStrandsOptionalMaterial(uint InVertexId)
+{
+	uint4 Out = 0;
+	if (HairStrandsVF_HasMaterial > 0)
+	{
+		Out = HairStrandsVF_MaterialBuffer[InVertexId];
+	}
+	return Out;
+}
+
 FHairVertexAttributes GetHairStrandsAttributes(uint HairPrimitiveId, float2 VertexUV)
 {
 	const uint VertexId0 = HairPrimitiveId;
 	const uint VertexId1 = HairPrimitiveId+1;
 
 	// Instead of doing the actual interpolation we could rely on dithering as an optimization
-	const FHairAttributes Attributes0 = DecodeHairAttribute(HairStrandsVF_AttributeBuffer[VertexId0], HairStrandsVF_Length);
-	const FHairAttributes Attributes1 = DecodeHairAttribute(HairStrandsVF_AttributeBuffer[VertexId1], HairStrandsVF_Length);
+	const FHairAttributes Attributes0 = DecodeHairAttribute(HairStrandsVF_Attribute0Buffer[VertexId0], GetHairStrandsOptionalAttribute1(VertexId0), HairStrandsVF_Length);
+	const FHairAttributes Attributes1 = DecodeHairAttribute(HairStrandsVF_Attribute0Buffer[VertexId1], GetHairStrandsOptionalAttribute1(VertexId1), HairStrandsVF_Length);
 
 	const FHairControlPoint ControlPoint0 = UnpackHairControlPoint(HairStrandsVF_PositionBuffer[VertexId0], float3(0, 0 ,0), HairStrandsVF_Radius, HairStrandsVF_RootScale, HairStrandsVF_TipScale);
 	const FHairControlPoint ControlPoint1 = UnpackHairControlPoint(HairStrandsVF_PositionBuffer[VertexId1], float3(0, 0, 0), HairStrandsVF_Radius, HairStrandsVF_RootScale, HairStrandsVF_TipScale);
@@ -52,8 +75,8 @@ FHairVertexAttributes GetHairStrandsAttributes(uint HairPrimitiveId, float2 Vert
 	Out.HairUV.y		= VertexUV.y;
 
 	// Instead of doing the actual interpolation we could rely on dithering as an optimization
-	const float4 Material0 = HairStrandsVF_MaterialBuffer[VertexId0];
-	const float4 Material1 = HairStrandsVF_MaterialBuffer[VertexId1];
+	const float4 Material0 = GetHairStrandsOptionalMaterial(VertexId0);
+	const float4 Material1 = GetHairStrandsOptionalMaterial(VertexId1);
 
 	// Cheap/manual sRGB decoding as the underlying format is linear quantization
 	Out.HairBaseColor	= lerp(Material0.xyz*Material0.xyz, Material1.xyz*Material1.xyz, VertexUV.x);
diff --git a/Engine/Shaders/Private/HairStrands/HairStrandsInterpolation.usf b/Engine/Shaders/Private/HairStrands/HairStrandsInterpolation.usf
--- a/Engine/Shaders/Private/HairStrands/HairStrandsInterpolation.usf
+++ b/Engine/Shaders/Private/HairStrands/HairStrandsInterpolation.usf
@@ -57,7 +57,7 @@ Buffer<float>	HairStrandsVF_CullingRadiusScaleBuffer;
 Buffer<uint>	HairStrandsVF_CullingIndirectBufferArgs;
 #endif
 
-Buffer<uint4>  AttributeBuffer;
+Buffer<float4>  AttributeBuffer;
 
 #if PERMUTATION_DYNAMIC_GEOMETRY >= 1 || PERMUTATION_DEBUG == 1 
 Buffer<uint>   RenVertexToRootIndexBuffer;
diff --git a/Engine/Shaders/Private/HairStrands/HairStrandsTexturesGeneration.usf b/Engine/Shaders/Private/HairStrands/HairStrandsTexturesGeneration.usf
--- a/Engine/Shaders/Private/HairStrands/HairStrandsTexturesGeneration.usf
+++ b/Engine/Shaders/Private/HairStrands/HairStrandsTexturesGeneration.usf
@@ -83,7 +83,7 @@ uint VertexCount;
 uint PrimitiveCount;
 
 Buffer<uint4> InVF_PositionBuffer;
-Buffer<uint4> InVF_AttributeBuffer;
+Buffer<float4> InVF_Attribute0Buffer;
 float3 InVF_PositionOffset;
 float InVF_Radius;
 float InVF_Length;
@@ -282,8 +282,9 @@ bool Trace(
 							OutDepth = Distance;
 
 
-							const FHairAttributes Attributes0 = DecodeHairAttribute(InVF_AttributeBuffer[P0], InVF_Length);
-							const FHairAttributes Attributes1 = DecodeHairAttribute(InVF_AttributeBuffer[P1], InVF_Length);
+							const uint4 Attribute1 = 0;
+							const FHairAttributes Attributes0 = DecodeHairAttribute(InVF_Attribute0Buffer[P0], Attribute1, InVF_Length);
+							const FHairAttributes Attributes1 = DecodeHairAttribute(InVF_Attribute0Buffer[P1], Attribute1, InVF_Length);
 
 							OutSeed		= Attributes0.Seed;
 							OutRootUV	= Attributes0.RootUV;
diff --git a/Engine/Shaders/Private/HairStrands/HairStrandsVertexFactoryCommon.ush b/Engine/Shaders/Private/HairStrands/HairStrandsVertexFactoryCommon.ush
--- a/Engine/Shaders/Private/HairStrands/HairStrandsVertexFactoryCommon.ush
+++ b/Engine/Shaders/Private/HairStrands/HairStrandsVertexFactoryCommon.ush
@@ -52,17 +52,17 @@ struct FHairAttributes
 	float WorldLength;
 };
 
-FHairAttributes DecodeHairAttribute(uint4 Attributes, float InVF_Length)
+FHairAttributes DecodeHairAttribute(float4 Attributes0, uint4 Attributes1, float InVF_Length)
 {
 	float2 RootUV, RootIndex;
-	RootUV.x = float(Attributes.x & 0xFF) / 255.f;
-	RootUV.y = 1.0 - (float((Attributes.x >> 8) & 0xFF) / 255.f); // Invert V to compensate image origin flip
-	RootIndex.x = float(Attributes.z & 0xFF);
-	RootIndex.y = float((Attributes.z >> 8) & 0xFF);
+	RootUV.x		= Attributes0.x;
+	RootUV.y		= 1.0 - Attributes0.y; // Invert V to compensate image origin flip
+	RootIndex.x		= float(Attributes1.x);
+	RootIndex.y		= float(Attributes1.y);
 
 	FHairAttributes Out;
-	Out.RootUV = RootUV + RootIndex;
-	Out.WorldLength = (Attributes.y & 0xFF) / 255.f * InVF_Length;
-	Out.Seed = ((Attributes.y >> 8) & 0xFF) / 255.f;
+	Out.RootUV		= RootUV + RootIndex;
+	Out.WorldLength = Attributes0.z * InVF_Length;
+	Out.Seed		= Attributes0.w;
 	return Out;
 }
\ No newline at end of file
diff --git a/Engine/Source/Runtime/Renderer/Private/HairStrands/HairStrandsVisibility.cpp b/Engine/Source/Runtime/Renderer/Private/HairStrands/HairStrandsVisibility.cpp
--- a/Engine/Source/Runtime/Renderer/Private/HairStrands/HairStrandsVisibility.cpp
+++ b/Engine/Source/Runtime/Renderer/Private/HairStrands/HairStrandsVisibility.cpp
@@ -3026,27 +3026,29 @@ static void AddMeshDrawTransitionPass(
 			FRDGResourceAccessFinalizer ResourceAccessFinalizer;
 
 			FHairGroupPublicData::FVertexFactoryInput& VFInput = HairGroupPublicData->VFInput;
-			ResourceAccessFinalizer.AddBuffer(VFInput.Strands.PositionBuffer.Buffer,				ERHIAccess::SRVMask);
-			ResourceAccessFinalizer.AddBuffer(VFInput.Strands.PrevPositionBuffer.Buffer,			ERHIAccess::SRVMask);
+			ResourceAccessFinalizer.AddBuffer(VFInput.Strands.PositionBuffer.Buffer,			ERHIAccess::SRVMask);
+			ResourceAccessFinalizer.AddBuffer(VFInput.Strands.PrevPositionBuffer.Buffer,		ERHIAccess::SRVMask);
 			ResourceAccessFinalizer.AddBuffer(VFInput.Strands.TangentBuffer.Buffer,				ERHIAccess::SRVMask);
-			ResourceAccessFinalizer.AddBuffer(VFInput.Strands.AttributeBuffer.Buffer,				ERHIAccess::SRVMask);
-			ResourceAccessFinalizer.AddBuffer(VFInput.Strands.MaterialBuffer.Buffer,				ERHIAccess::SRVMask);
+			ResourceAccessFinalizer.AddBuffer(VFInput.Strands.Attribute0Buffer.Buffer,			ERHIAccess::SRVMask);
+			ResourceAccessFinalizer.AddBuffer(VFInput.Strands.Attribute1Buffer.Buffer,			ERHIAccess::SRVMask);
+			ResourceAccessFinalizer.AddBuffer(VFInput.Strands.MaterialBuffer.Buffer,			ERHIAccess::SRVMask);
 			ResourceAccessFinalizer.AddBuffer(VFInput.Strands.PositionOffsetBuffer.Buffer,		ERHIAccess::SRVMask);
 			ResourceAccessFinalizer.AddBuffer(VFInput.Strands.PrevPositionOffsetBuffer.Buffer,	ERHIAccess::SRVMask);
 
 			FRDGBufferRef CulledVertexIdBuffer = Register(GraphBuilder, HairGroupPublicData->CulledVertexIdBuffer, ERDGImportedBufferFlags::None).Buffer;
 			FRDGBufferRef CulledVertexRadiusScaleBuffer = Register(GraphBuilder, HairGroupPublicData->CulledVertexRadiusScaleBuffer, ERDGImportedBufferFlags::None).Buffer;
 			FRDGBufferRef DrawIndirectBuffer = Register(GraphBuilder, HairGroupPublicData->DrawIndirectBuffer, ERDGImportedBufferFlags::None).Buffer;
-			ResourceAccessFinalizer.AddBuffer(CulledVertexIdBuffer, ERHIAccess::SRVMask);
-			ResourceAccessFinalizer.AddBuffer(CulledVertexRadiusScaleBuffer, ERHIAccess::SRVMask);
-			ResourceAccessFinalizer.AddBuffer(DrawIndirectBuffer, ERHIAccess::IndirectArgs);
+			ResourceAccessFinalizer.AddBuffer(CulledVertexIdBuffer,								ERHIAccess::SRVMask);
+			ResourceAccessFinalizer.AddBuffer(CulledVertexRadiusScaleBuffer,					ERHIAccess::SRVMask);
+			ResourceAccessFinalizer.AddBuffer(DrawIndirectBuffer,								ERHIAccess::IndirectArgs);
 
 			ResourceAccessFinalizer.Finalize(GraphBuilder);
 
 			VFInput.Strands.PositionBuffer				= FRDGImportedBuffer();
 			VFInput.Strands.PrevPositionBuffer			= FRDGImportedBuffer();
 			VFInput.Strands.TangentBuffer				= FRDGImportedBuffer();
-			VFInput.Strands.AttributeBuffer				= FRDGImportedBuffer();
+			VFInput.Strands.Attribute0Buffer			= FRDGImportedBuffer();
+			VFInput.Strands.Attribute1Buffer			= FRDGImportedBuffer();
 			VFInput.Strands.MaterialBuffer				= FRDGImportedBuffer();
 			VFInput.Strands.PositionOffsetBuffer		= FRDGImportedBuffer();
 			VFInput.Strands.PrevPositionOffsetBuffer	= FRDGImportedBuffer();
diff --git a/Engine/Source/Runtime/Renderer/Public/HairStrandsInterface.h b/Engine/Source/Runtime/Renderer/Public/HairStrandsInterface.h
--- a/Engine/Source/Runtime/Renderer/Public/HairStrandsInterface.h
+++ b/Engine/Source/Runtime/Renderer/Public/HairStrandsInterface.h
@@ -231,7 +231,8 @@ public:
 			FRDGImportedBuffer PrevPositionBuffer;
 			FRDGImportedBuffer TangentBuffer;
 			FRDGImportedBuffer MaterialBuffer;
-			FRDGImportedBuffer AttributeBuffer;
+			FRDGImportedBuffer Attribute0Buffer;
+			FRDGImportedBuffer Attribute1Buffer;
 			FRDGImportedBuffer PositionOffsetBuffer;
 			FRDGImportedBuffer PrevPositionOffsetBuffer;
 
@@ -239,7 +240,8 @@ public:
 			FShaderResourceViewRHIRef PrevPositionBufferRHISRV			= nullptr;
 			FShaderResourceViewRHIRef TangentBufferRHISRV				= nullptr;
 			FShaderResourceViewRHIRef MaterialBufferRHISRV				= nullptr;
-			FShaderResourceViewRHIRef AttributeBufferRHISRV				= nullptr;
+			FShaderResourceViewRHIRef Attribute0BufferRHISRV			= nullptr;
+			FShaderResourceViewRHIRef Attribute1BufferRHISRV			= nullptr;
 			FShaderResourceViewRHIRef PositionOffsetBufferRHISRV		= nullptr;
 			FShaderResourceViewRHIRef PrevPositionOffsetBufferRHISRV	= nullptr;
 
