Fix use-after-free in FHairStrandsDeformedResource::InternalAllocate.

InternalCreateVertexBufferRDG defaults to passing ERDGInitialDataFlags::NoCopy to FRDGBufferUploader::Upload() because the data is persistent for most of the buffers which are created, but the deformed offset buffers are initialized from an array allocated locally in the function, so by the time the RHI thread read the data, it was long gone.

#rnx

diff --git a/Engine/Plugins/Runtime/HairStrands/Source/HairStrandsCore/Private/GroomResources.cpp b/Engine/Plugins/Runtime/HairStrands/Source/HairStrandsCore/Private/GroomResources.cpp
--- a/Engine/Plugins/Runtime/HairStrands/Source/HairStrandsCore/Private/GroomResources.cpp
+++ b/Engine/Plugins/Runtime/HairStrands/Source/HairStrandsCore/Private/GroomResources.cpp
@@ -661,8 +661,8 @@ void FHairStrandsDeformedResource::InternalAllocate(FRDGBuilder& GraphBuilder)
 
 	TArray<FVector4> DefaultOffsets;
 	DefaultOffsets.Add(BulkData.GetPositionOffset());
-	InternalCreateVertexBufferRDG<FHairStrandsPositionOffsetFormat>(GraphBuilder, BufferUploader, DefaultOffsets, DeformedOffsetBuffer[0], HAIRSTRANDS_RESOUCE_NAME(CurveType, Hair.StrandsDeformed_DeformedOffsetBuffer0), EHairResourceUsageType::Dynamic);
-	InternalCreateVertexBufferRDG<FHairStrandsPositionOffsetFormat>(GraphBuilder, BufferUploader, DefaultOffsets, DeformedOffsetBuffer[1], HAIRSTRANDS_RESOUCE_NAME(CurveType, Hair.StrandsDeformed_DeformedOffsetBuffer1), EHairResourceUsageType::Dynamic);
+	InternalCreateVertexBufferRDG<FHairStrandsPositionOffsetFormat>(GraphBuilder, BufferUploader, DefaultOffsets, DeformedOffsetBuffer[0], HAIRSTRANDS_RESOUCE_NAME(CurveType, Hair.StrandsDeformed_DeformedOffsetBuffer0), EHairResourceUsageType::Dynamic, ERDGInitialDataFlags::None);
+	InternalCreateVertexBufferRDG<FHairStrandsPositionOffsetFormat>(GraphBuilder, BufferUploader, DefaultOffsets, DeformedOffsetBuffer[1], HAIRSTRANDS_RESOUCE_NAME(CurveType, Hair.StrandsDeformed_DeformedOffsetBuffer1), EHairResourceUsageType::Dynamic, ERDGInitialDataFlags::None);
 
 	BufferUploader.Submit(GraphBuilder);
 }
