Reduce interpolation data size when using single guide interpolation.

This save 33% of groom asset data.

diff --git a/Engine/Plugins/Runtime/HairStrands/Source/HairStrandsCore/Private/GroomBuilder.cpp b/Engine/Plugins/Runtime/HairStrands/Source/HairStrandsCore/Private/GroomBuilder.cpp
--- a/Engine/Plugins/Runtime/HairStrands/Source/HairStrandsCore/Private/GroomBuilder.cpp
+++ b/Engine/Plugins/Runtime/HairStrands/Source/HairStrandsCore/Private/GroomBuilder.cpp
@@ -33,7 +33,7 @@ static FAutoConsoleVariableRef CVarHairClusterBuilder_MaxVoxelResolution(TEXT("r
 FString FGroomBuilder::GetVersion()
 {
 	// Important to update the version when groom building changes
-	return TEXT("2a");
+	return TEXT("2b");
 }
 
 namespace FHairStrandsDecimation
@@ -1006,6 +1006,7 @@ namespace HairInterpolationBuilder
 		TRACE_CPUPROFILER_EVENT_SCOPE(HairInterpolationBuilder::BuildInterpolationData);
 
 		InterpolationData.SetNum(RenStrandsData.GetNumPoints());
+		InterpolationData.bUseUniqueGuide = Settings.bUseUniqueGuide;
 
 		typedef TArray<FHairRoot> FRoots;
 
@@ -1173,38 +1174,62 @@ namespace HairInterpolationBuilder
 		if (PointCount == 0)
 			return;
 
-		auto LowerPart = [](uint32 Index) { return uint16(Index & 0xFFFF); };
-		auto UpperPart = [](uint32 Index) { return uint8((Index >> 16) & 0xFF); };
-
-		TArray<FHairStrandsInterpolation0Format::Type>& OutPointsInterpolation0 = OutBulkData.Interpolation0;
-		TArray<FHairStrandsInterpolation1Format::Type>& OutPointsInterpolation1 = OutBulkData.Interpolation1;
-
-		OutPointsInterpolation0.SetNum(PointCount * FHairStrandsInterpolation0Format::ComponentCount);
-		OutPointsInterpolation1.SetNum(PointCount * FHairStrandsInterpolation1Format::ComponentCount);
-
-		for (uint32 PointIndex = 0; PointIndex < PointCount; ++PointIndex)
-		{
-			const FIntVector& Indices = HairInterpolation.PointsSimCurvesVertexIndex[PointIndex];
-			const FVector& Weights = HairInterpolation.PointsSimCurvesVertexWeights[PointIndex];
-			const FVector& S = HairInterpolation.PointsSimCurvesVertexLerp[PointIndex];
-
-			FHairStrandsInterpolation0Format::Type& OutInterp0 = OutPointsInterpolation0[PointIndex];
-			OutInterp0.Index0 = LowerPart(Indices[0]);
-			OutInterp0.Index1 = LowerPart(Indices[1]);
-			OutInterp0.Index2 = LowerPart(Indices[2]);
-			OutInterp0.VertexWeight0 = Weights[0] * 255.f;
-			OutInterp0.VertexWeight1 = Weights[1] * 255.f;
-
-			FHairStrandsInterpolation1Format::Type& OutInterp1 = OutPointsInterpolation1[PointIndex];
-			OutInterp1.VertexIndex0 = UpperPart(Indices[0]);
-			OutInterp1.VertexIndex1 = UpperPart(Indices[1]);
-			OutInterp1.VertexIndex2 = UpperPart(Indices[2]);
-			OutInterp1.VertexLerp0  = S[0] * 255.f;
-			OutInterp1.VertexLerp1  = S[1] * 255.f;
-			OutInterp1.VertexLerp2  = S[2] * 255.f;
-			OutInterp1.Pad0			= 0;
-			OutInterp1.Pad1			= 0;
-		}	
+		auto LowerPart = [](uint32 Index) -> uint16 { return uint16(Index & 0xFFFF); };
+		auto UpperPart = [](uint32 Index) -> uint8  { return uint8((Index >> 16) & 0xFF); };
+
+		if (HairInterpolation.bUseUniqueGuide)
+		{
+			OutBulkData.Flags = FHairStrandsInterpolationBulkData::DataFlags_HasSingleGuideData;
+
+			TArray<FHairStrandsInterpolationFormat::Type>& OutPointsInterpolation = OutBulkData.Interpolation;
+			OutPointsInterpolation.SetNum(PointCount * FHairStrandsInterpolationFormat::ComponentCount);
+
+			for (uint32 PointIndex = 0; PointIndex < PointCount; ++PointIndex)
+			{
+				const FIntVector& Indices = HairInterpolation.PointsSimCurvesVertexIndex[PointIndex];
+				const FVector& Weights = HairInterpolation.PointsSimCurvesVertexWeights[PointIndex];
+				const FVector& S = HairInterpolation.PointsSimCurvesVertexLerp[PointIndex];
+
+				FHairStrandsInterpolationFormat::Type& OutInterp = OutPointsInterpolation[PointIndex];
+				OutInterp.VertexGuideIndex0 = LowerPart(Indices[0]);
+				OutInterp.VertexGuideIndex1 = UpperPart(Indices[0]);
+				OutInterp.VertexLerp		= S[0] * 255.f;
+			}
+		}
+		else
+		{
+			OutBulkData.Flags = 0;
+
+			TArray<FHairStrandsInterpolation0Format::Type>& OutPointsInterpolation0 = OutBulkData.Interpolation0;
+			TArray<FHairStrandsInterpolation1Format::Type>& OutPointsInterpolation1 = OutBulkData.Interpolation1;
+
+			OutPointsInterpolation0.SetNum(PointCount * FHairStrandsInterpolation0Format::ComponentCount);
+			OutPointsInterpolation1.SetNum(PointCount * FHairStrandsInterpolation1Format::ComponentCount);
+
+			for (uint32 PointIndex = 0; PointIndex < PointCount; ++PointIndex)
+			{
+				const FIntVector& Indices = HairInterpolation.PointsSimCurvesVertexIndex[PointIndex];
+				const FVector& Weights = HairInterpolation.PointsSimCurvesVertexWeights[PointIndex];
+				const FVector& S = HairInterpolation.PointsSimCurvesVertexLerp[PointIndex];
+
+				FHairStrandsInterpolation0Format::Type& OutInterp0 = OutPointsInterpolation0[PointIndex];
+				OutInterp0.Index0 = LowerPart(Indices[0]);
+				OutInterp0.Index1 = LowerPart(Indices[1]);
+				OutInterp0.Index2 = LowerPart(Indices[2]);
+				OutInterp0.VertexWeight0 = Weights[0] * 255.f;
+				OutInterp0.VertexWeight1 = Weights[1] * 255.f;
+
+				FHairStrandsInterpolation1Format::Type& OutInterp1 = OutPointsInterpolation1[PointIndex];
+				OutInterp1.VertexIndex0 = UpperPart(Indices[0]);
+				OutInterp1.VertexIndex1 = UpperPart(Indices[1]);
+				OutInterp1.VertexIndex2 = UpperPart(Indices[2]);
+				OutInterp1.VertexLerp0  = S[0] * 255.f;
+				OutInterp1.VertexLerp1  = S[1] * 255.f;
+				OutInterp1.VertexLerp2  = S[2] * 255.f;
+				OutInterp1.Pad0			= 0;
+				OutInterp1.Pad1			= 0;
+			}
+		}
 
 		const uint32 RootCount = SimDatas.GetNumCurves();
 		OutBulkData.SimRootPointIndex.SetNum(SimDatas.GetNumPoints());
diff --git a/Engine/Plugins/Runtime/HairStrands/Source/HairStrandsCore/Private/GroomResources.cpp b/Engine/Plugins/Runtime/HairStrands/Source/HairStrandsCore/Private/GroomResources.cpp
--- a/Engine/Plugins/Runtime/HairStrands/Source/HairStrandsCore/Private/GroomResources.cpp
+++ b/Engine/Plugins/Runtime/HairStrands/Source/HairStrandsCore/Private/GroomResources.cpp
@@ -1006,21 +1006,30 @@ void FHairStrandsRootData::Reset()
 
 FHairStrandsInterpolationResource::FHairStrandsInterpolationResource(const FHairStrandsInterpolationBulkData& InBulkData) :
 	FHairCommonResource(EHairStrandsAllocationType::Deferred),
-	Interpolation0Buffer(), Interpolation1Buffer(), BulkData(InBulkData)
+	InterpolationBuffer(), Interpolation0Buffer(), Interpolation1Buffer(), BulkData(InBulkData)
 {
 }
 
 void FHairStrandsInterpolationResource::InternalAllocate(FRDGBuilder& GraphBuilder)
 {
 	FRDGBufferUploader BufferUploader;
-	InternalCreateVertexBufferRDG<FHairStrandsInterpolation0Format>(GraphBuilder, BufferUploader, BulkData.Interpolation0, Interpolation0Buffer, TEXT("Hair.StrandsInterpolation_Interpolation0Buffer"), EHairResourceUsageType::Static);
-	InternalCreateVertexBufferRDG<FHairStrandsInterpolation1Format>(GraphBuilder, BufferUploader, BulkData.Interpolation1, Interpolation1Buffer, TEXT("Hair.StrandsInterpolation_Interpolation1Buffer"), EHairResourceUsageType::Static);
+	const bool bUseSingleGuide = !!(BulkData.Flags & FHairStrandsInterpolationBulkData::DataFlags_HasSingleGuideData);
+	if (bUseSingleGuide)
+	{
+		InternalCreateVertexBufferRDG<FHairStrandsInterpolationFormat>(GraphBuilder, BufferUploader, BulkData.Interpolation, InterpolationBuffer, TEXT("Hair.StrandsInterpolation_InterpolationBuffer"), EHairResourceUsageType::Static);
+	}
+	else
+	{
+		InternalCreateVertexBufferRDG<FHairStrandsInterpolation0Format>(GraphBuilder, BufferUploader, BulkData.Interpolation0, Interpolation0Buffer, TEXT("Hair.StrandsInterpolation_Interpolation0Buffer"), EHairResourceUsageType::Static);
+		InternalCreateVertexBufferRDG<FHairStrandsInterpolation1Format>(GraphBuilder, BufferUploader, BulkData.Interpolation1, Interpolation1Buffer, TEXT("Hair.StrandsInterpolation_Interpolation1Buffer"), EHairResourceUsageType::Static);
+	}
 	InternalCreateVertexBufferRDG<FHairStrandsRootIndexFormat>(GraphBuilder, BufferUploader, BulkData.SimRootPointIndex, SimRootPointIndexBuffer, TEXT("Hair.StrandsInterpolation_SimRootPointIndex"), EHairResourceUsageType::Static);
 	BufferUploader.Submit(GraphBuilder);
 }
 
 void FHairStrandsInterpolationResource::InternalRelease()
 {
+	InterpolationBuffer.Release();
 	Interpolation0Buffer.Release();
 	Interpolation1Buffer.Release();
 	SimRootPointIndexBuffer.Release();
diff --git a/Engine/Plugins/Runtime/HairStrands/Source/HairStrandsCore/Private/HairStrandsDatas.cpp b/Engine/Plugins/Runtime/HairStrands/Source/HairStrandsCore/Private/HairStrandsDatas.cpp
--- a/Engine/Plugins/Runtime/HairStrands/Source/HairStrandsCore/Private/HairStrandsDatas.cpp
+++ b/Engine/Plugins/Runtime/HairStrands/Source/HairStrandsCore/Private/HairStrandsDatas.cpp
@@ -117,6 +117,15 @@ FArchive& operator<<(FArchive& Ar, FHairInterpolation0Vertex& Vertex)
 	return Ar;
 }
 
+FArchive& operator<<(FArchive& Ar, FHairInterpolationVertex& Vertex)
+{
+	Ar << Vertex.VertexGuideIndex0;
+	Ar << Vertex.VertexGuideIndex1;
+	Ar << Vertex.VertexLerp;
+
+	return Ar;
+}
+
 FArchive& operator<<(FArchive& Ar, FHairInterpolation1Vertex& Vertex)
 {
 	Ar.UsingCustomVersion(FReleaseObjectVersion::GUID);
@@ -156,12 +165,21 @@ FArchive& operator<<(FArchive& Ar, FHairInterpolation1Vertex& Vertex)
 
 void FHairStrandsInterpolationBulkData::Serialize(FArchive& Ar)
 {
+	static_assert(sizeof(FHairInterpolationVertex::BulkType) == sizeof(FHairInterpolationVertex));
 	static_assert(sizeof(FHairInterpolation0Vertex::BulkType) == sizeof(FHairInterpolation0Vertex));
 	static_assert(sizeof(FHairInterpolation1Vertex::BulkType) == sizeof(FHairInterpolation1Vertex));
 	static_assert(sizeof(FHairStrandsRootIndexFormat::BulkType) == sizeof(FHairStrandsRootIndexFormat::Type));
 
-	Interpolation0.BulkSerialize(Ar);
-	Interpolation1.BulkSerialize(Ar);
+	Ar << Flags;
+	if (!!(Flags & DataFlags_HasSingleGuideData))
+	{
+		Interpolation.BulkSerialize(Ar);
+	}
+	else
+	{
+		Interpolation0.BulkSerialize(Ar);
+		Interpolation1.BulkSerialize(Ar);
+	}
 	SimRootPointIndex.BulkSerialize(Ar);
 }
 
diff --git a/Engine/Plugins/Runtime/HairStrands/Source/HairStrandsCore/Private/HairStrandsRendering.cpp b/Engine/Plugins/Runtime/HairStrands/Source/HairStrandsCore/Private/HairStrandsRendering.cpp
--- a/Engine/Plugins/Runtime/HairStrands/Source/HairStrandsCore/Private/HairStrandsRendering.cpp
+++ b/Engine/Plugins/Runtime/HairStrands/Source/HairStrandsCore/Private/HairStrandsRendering.cpp
@@ -33,9 +33,6 @@ static FAutoConsoleVariableRef CVarStrandHairWidth(TEXT("r.HairStrands.StrandWid
 static int32 GStrandHairInterpolationDebug = 0;
 static FAutoConsoleVariableRef CVarStrandHairInterpolationDebug(TEXT("r.HairStrands.Interpolation.Debug"), GStrandHairInterpolationDebug, TEXT("Enable debug rendering for hair interpolation"));
 
-static int32 GHairStrandsUseSingleGuideInterpolation = 0;
-static FAutoConsoleVariableRef CVarHairStrandsUseSingleGuideInterpolation(TEXT("r.HairStrands.Interpolation.UseSingleGuide"), GHairStrandsUseSingleGuideInterpolation, TEXT("Hair interpolation will use a single guide for interpolating hair motion instead of 3. Save performance cost"), ECVF_Scalability | ECVF_RenderThreadSafe);
-
 static int32 GHairCardsInterpolationType = 1;
 static FAutoConsoleVariableRef CVarHairCardsInterpolationType(TEXT("r.HairStrands.Cards.InterpolationType"), GHairCardsInterpolationType, TEXT("Hair cards interpolation type: 0: None, 1:physics simulation, 2: RBF deformation"));
 
@@ -453,6 +450,7 @@ class FHairInterpolationCS : public FGlobalShader
 		SHADER_PARAMETER_RDG_BUFFER_SRV(Buffer, SimRestPosePositionBuffer)
 		SHADER_PARAMETER_RDG_BUFFER_SRV(Buffer, DeformedSimPositionBuffer)
 
+		SHADER_PARAMETER_RDG_BUFFER_SRV(Buffer, InterpolationBuffer)
 		SHADER_PARAMETER_RDG_BUFFER_SRV(Buffer, Interpolation0Buffer)
 		SHADER_PARAMETER_RDG_BUFFER_SRV(Buffer, Interpolation1Buffer)
 
@@ -533,6 +531,8 @@ static void AddHairStrandsInterpolationPass(
 	const FHairStrandsDeformedRootResource* SimDeformedRootResources,
 	const FRDGBufferSRVRef& RenderRestPosePositionBuffer,
 	const FRDGBufferSRVRef& RenderAttributeBuffer,
+	const bool bUseSingleGuide,
+	const FRDGBufferSRVRef& InterpolationBuffer,
 	const FRDGBufferSRVRef& Interpolation0Buffer,
 	const FRDGBufferSRVRef& Interpolation1Buffer,
 	const FRDGBufferSRVRef& SimRestPosePositionBuffer,
@@ -552,8 +552,15 @@ static void AddHairStrandsInterpolationPass(
 	Parameters->RenderRestPosePositionBuffer = RenderRestPosePositionBuffer;
 	Parameters->SimRestPosePositionBuffer = SimRestPosePositionBuffer;
 	Parameters->DeformedSimPositionBuffer = SimDeformedPositionBuffer;
-	Parameters->Interpolation0Buffer = Interpolation0Buffer;
-	Parameters->Interpolation1Buffer = Interpolation1Buffer;
+	if (bUseSingleGuide)
+	{
+		Parameters->InterpolationBuffer = InterpolationBuffer;
+	}
+	else
+	{
+		Parameters->Interpolation0Buffer = Interpolation0Buffer;
+		Parameters->Interpolation1Buffer = Interpolation1Buffer;
+	}
 	Parameters->OutRenderDeformedPositionBuffer = OutRenderPositionBuffer.UAV;
 	Parameters->HairStrandsCullIndex = FIntPoint(-1, -1);
 	Parameters->VertexCount = VertexCount;
@@ -679,7 +686,6 @@ static void AddHairStrandsInterpolationPass(
 		ShaderDrawDebug::SetParameters(GraphBuilder, *ShaderDrawData, Parameters->ShaderDrawParameters);
 	}
 
-	const bool bUseSingleGuide		= GHairStrandsUseSingleGuideInterpolation > 0;
 	const bool bHasLocalDeformation = Instance->Guides.bIsSimulationEnable || bSupportGlobalInterpolation;
 	const bool bCullingEnable		= InstanceGeometryType == EHairGeometryType::Strands && CullingData.bCullingResultAvailable;
 	Parameters->HairStrandsVF_bIsCullingEnable = bCullingEnable ? 1 : 0;
@@ -1603,6 +1609,7 @@ void ComputeHairStrandsInterpolation(
 				// 2.1 Compute deformation position based on simulation/skinning/RBF
 				if (Instance->Debug.GroomCacheType != EGroomCacheType::Strands)
 				{
+					const bool bUseSingleGuide = Instance->Strands.InterpolationResource->UseSingleGuide();
 					AddHairStrandsInterpolationPass(
 						GraphBuilder,
 						ShaderMap,
@@ -1624,8 +1631,10 @@ void ComputeHairStrandsInterpolation(
 						bHasSkinning && bValidGuide ? Instance->Guides.DeformedRootResource : nullptr,
 						RegisterAsSRV(GraphBuilder, Instance->Strands.RestResource->PositionBuffer),
 						RegisterAsSRV(GraphBuilder, Instance->Strands.RestResource->Attribute0Buffer),
-						bValidGuide ? RegisterAsSRV(GraphBuilder, Instance->Strands.InterpolationResource->Interpolation0Buffer) : nullptr,
-						bValidGuide ? RegisterAsSRV(GraphBuilder, Instance->Strands.InterpolationResource->Interpolation1Buffer) : nullptr,
+						bUseSingleGuide,
+						bValidGuide &&  bUseSingleGuide ? RegisterAsSRV(GraphBuilder, Instance->Strands.InterpolationResource->InterpolationBuffer) : nullptr,
+						bValidGuide && !bUseSingleGuide ? RegisterAsSRV(GraphBuilder, Instance->Strands.InterpolationResource->Interpolation0Buffer) : nullptr,
+						bValidGuide && !bUseSingleGuide ? RegisterAsSRV(GraphBuilder, Instance->Strands.InterpolationResource->Interpolation1Buffer) : nullptr,
 						bValidGuide ? RegisterAsSRV(GraphBuilder, Instance->Guides.RestResource->PositionBuffer) : nullptr,
 						bValidGuide ? RegisterAsSRV(GraphBuilder, Instance->Guides.DeformedResource->GetBuffer(FHairStrandsDeformedResource::Current)) : nullptr,
 						bValidGuide ? RegisterAsSRV(GraphBuilder, Instance->Guides.RestResource->Attribute0Buffer) : nullptr,
@@ -1796,6 +1805,8 @@ void ComputeHairStrandsInterpolation(
 				{
 					FRDGImportedBuffer Guides_DeformedPositionBuffer = Register(GraphBuilder, LOD.Guides.DeformedResource->GetBuffer(FHairStrandsDeformedResource::Current), ERDGImportedBufferFlags::CreateViews);
 
+					const bool bUseSingleGuide = LOD.Guides.InterpolationResource->UseSingleGuide();
+
 					FRDGHairStrandsCullingData CullingData;
 					AddHairStrandsInterpolationPass(
 						GraphBuilder,
@@ -1818,8 +1829,10 @@ void ComputeHairStrandsInterpolation(
 						bHasSkinning && bValidGuide ? Instance->Guides.DeformedRootResource : nullptr,
 						RegisterAsSRV(GraphBuilder, LOD.Guides.RestResource->PositionBuffer),
 						RegisterAsSRV(GraphBuilder, LOD.Guides.RestResource->Attribute0Buffer),
-						RegisterAsSRV(GraphBuilder, LOD.Guides.InterpolationResource->Interpolation0Buffer),
-						RegisterAsSRV(GraphBuilder, LOD.Guides.InterpolationResource->Interpolation1Buffer),
+						bUseSingleGuide,
+						bValidGuide &&  bUseSingleGuide ? RegisterAsSRV(GraphBuilder, LOD.Guides.InterpolationResource->InterpolationBuffer) : nullptr,
+						bValidGuide && !bUseSingleGuide ? RegisterAsSRV(GraphBuilder, LOD.Guides.InterpolationResource->Interpolation0Buffer) : nullptr,
+						bValidGuide && !bUseSingleGuide ? RegisterAsSRV(GraphBuilder, LOD.Guides.InterpolationResource->Interpolation1Buffer) : nullptr,
 						bValidGuide ? RegisterAsSRV(GraphBuilder, Instance->Guides.RestResource->PositionBuffer) : nullptr,
 						bValidGuide ? RegisterAsSRV(GraphBuilder, Instance->Guides.DeformedResource->GetBuffer(FHairStrandsDeformedResource::Current)) : nullptr,
 						bValidGuide ? RegisterAsSRV(GraphBuilder, Instance->Guides.RestResource->Attribute0Buffer) : nullptr,
diff --git a/Engine/Plugins/Runtime/HairStrands/Source/HairStrandsCore/Public/GroomAsset.h b/Engine/Plugins/Runtime/HairStrands/Source/HairStrandsCore/Public/GroomAsset.h
--- a/Engine/Plugins/Runtime/HairStrands/Source/HairStrandsCore/Public/GroomAsset.h
+++ b/Engine/Plugins/Runtime/HairStrands/Source/HairStrandsCore/Public/GroomAsset.h
@@ -102,6 +102,7 @@ struct HAIRSTRANDSCORE_API FHairGroupData
 		{
 			uint32 Total = 0;
 			Total += FBase::GetDataSize();
+			Total += InterpolationBulkData.Interpolation.GetAllocatedSize();
 			Total += InterpolationBulkData.Interpolation0.GetAllocatedSize();
 			Total += InterpolationBulkData.Interpolation1.GetAllocatedSize();
 			Total += InterpolationBulkData.SimRootPointIndex.GetAllocatedSize();
diff --git a/Engine/Plugins/Runtime/HairStrands/Source/HairStrandsCore/Public/GroomResources.h b/Engine/Plugins/Runtime/HairStrands/Source/HairStrandsCore/Public/GroomResources.h
--- a/Engine/Plugins/Runtime/HairStrands/Source/HairStrandsCore/Public/GroomResources.h
+++ b/Engine/Plugins/Runtime/HairStrands/Source/HairStrandsCore/Public/GroomResources.h
@@ -360,12 +360,16 @@ struct FHairStrandsInterpolationResource : public FHairCommonResource
 	uint32 GetResourcesSize() const
 	{
 		uint32 Total = 0;
+		Total += GetBufferTotalNumBytes(InterpolationBuffer);
 		Total += GetBufferTotalNumBytes(Interpolation0Buffer);
 		Total += GetBufferTotalNumBytes(Interpolation1Buffer);
 		Total += GetBufferTotalNumBytes(SimRootPointIndexBuffer);
 		return Total;
 	}
 
+	bool UseSingleGuide() const { return InterpolationBuffer.Buffer != nullptr; }
+
+	FRDGExternalBuffer InterpolationBuffer;
 	FRDGExternalBuffer Interpolation0Buffer;
 	FRDGExternalBuffer Interpolation1Buffer;
 	FRDGExternalBuffer SimRootPointIndexBuffer;
diff --git a/Engine/Plugins/Runtime/HairStrands/Source/HairStrandsCore/Public/HairStrandsDatas.h b/Engine/Plugins/Runtime/HairStrands/Source/HairStrandsCore/Public/HairStrandsDatas.h
--- a/Engine/Plugins/Runtime/HairStrands/Source/HairStrandsCore/Public/HairStrandsDatas.h
+++ b/Engine/Plugins/Runtime/HairStrands/Source/HairStrandsCore/Public/HairStrandsDatas.h
@@ -51,6 +51,16 @@ struct FHairMaterialVertex
 	uint8 Roughness;
 };
 
+struct FHairInterpolationVertex
+{
+	typedef uint32 BulkType;
+
+	// Guide's vertex index are stored onto 24bits: 1) VertexGuideIndex0 is the lower part, 2) VertexGuideIndex1 is the upper part
+	uint16 VertexGuideIndex0;
+	uint8  VertexGuideIndex1;
+	uint8  VertexLerp;
+};
+
 struct FHairInterpolation0Vertex
 {
 	typedef uint64 BulkType;
@@ -148,6 +158,17 @@ struct FHairStrandsTangentFormat
 	static const EPixelFormat Format = PF_R8G8B8A8_SNORM;
 };
 
+struct FHairStrandsInterpolationFormat
+{
+	typedef FHairInterpolationVertex Type;
+	typedef uint32 BulkType;
+
+	static const uint32 ComponentCount = 1;
+	static const uint32 SizeInByte = sizeof(Type);
+	static const EVertexElementType VertexElementType = VET_UShort2;
+	static const EPixelFormat Format = PF_R32_UINT;
+};
+
 struct FHairStrandsInterpolation0Format
 {
 	typedef FHairInterpolation0Vertex Type;
@@ -282,14 +303,24 @@ struct HAIRSTRANDSCORE_API FHairStrandsInterpolationDatas
 
 	/** Weight of vertex indices on simulation curve, ordered by closest influence */
 	TArray<FVector>	PointsSimCurvesVertexWeights;
+
+	/** True, if interpolation data are built using a single guide */
+	bool bUseUniqueGuide = false;
 };
 
 struct HAIRSTRANDSCORE_API FHairStrandsInterpolationBulkData
 {
+	enum EDataFlags
+	{
+		DataFlags_HasSingleGuideData = 1,
+	};
+
 	void Serialize(FArchive& Ar);
-	uint32 GetPointCount() const { return Interpolation0.Num(); };
-	TArray<FHairStrandsInterpolation0Format::Type> Interpolation0;	// Per-rendering-vertex interpolation data (closest guides, weight factors, ...)
-	TArray<FHairStrandsInterpolation1Format::Type> Interpolation1;	// Per-rendering-vertex interpolation data (closest guides, weight factors, ...)
+	uint32 Flags;
+	uint32 GetPointCount() const { return (!!(Flags & DataFlags_HasSingleGuideData)) ? Interpolation.Num() : Interpolation0.Num(); };
+	TArray<FHairStrandsInterpolationFormat::Type> Interpolation;	// Per-rendering-vertex interpolation data (closest guides, weight factors, ...). Data for a single guide
+	TArray<FHairStrandsInterpolation0Format::Type> Interpolation0;	// Per-rendering-vertex interpolation data (closest guides, weight factors, ...). Data for up to 3 guides
+	TArray<FHairStrandsInterpolation1Format::Type> Interpolation1;	// Per-rendering-vertex interpolation data (closest guides, weight factors, ...). Data for up to 3 guides
 	TArray<FHairStrandsRootIndexFormat::Type> SimRootPointIndex;	// Per-rendering-vertex index of the sim-root vertex
 };
 
diff --git a/Engine/Shaders/Private/HairStrands/HairStrandsInterpolation.usf b/Engine/Shaders/Private/HairStrands/HairStrandsInterpolation.usf
--- a/Engine/Shaders/Private/HairStrands/HairStrandsInterpolation.usf
+++ b/Engine/Shaders/Private/HairStrands/HairStrandsInterpolation.usf
@@ -45,8 +45,18 @@ float4x4 LocalToWorldMatrix;
 Buffer<uint4>  RenderRestPosePositionBuffer;
 Buffer<uint4>  SimRestPosePositionBuffer;
 Buffer<uint4>  DeformedSimPositionBuffer;
-Buffer<uint4>  Interpolation0Buffer;
-Buffer<uint4>  Interpolation1Buffer;
+
+// Hair interpolation support either 
+// * A single guide, 
+// * Or up to 3 guides. 
+// This option is setup on interpolation data of the groom asset
+#if PERMUTATION_USE_SINGLE_GUIDE == 1
+Buffer<uint>	InterpolationBuffer;
+#else
+Buffer<uint4>	Interpolation0Buffer;
+Buffer<uint4>	Interpolation1Buffer;
+#endif
+
 Buffer<uint>   SimRootPointIndexBuffer;
 
 uint HairStrandsVF_bIsCullingEnable;
@@ -192,32 +202,35 @@ void MainCS(uint2 DispatchThreadId : SV_DispatchThreadID)
 	uint4 OutPosition = 0;
 	if (VertexIndex < VertexCount)
 	{
-		uint3 GuideVertexIndices = 0;
+		uint3  GuideVertexIndices = 0;
 		float3 GuideVertexLerps = 0;
 		float3 GuideVertexWeights = 0;
 		#if PERMUTATION_SIMULATION == 1 || PERMUTATION_DEBUG == 1 || PERMUTATION_DEBUG == 2 
 		{
+
+		  #if PERMUTATION_USE_SINGLE_GUIDE == 1
+			GuideVertexWeights = float3(1, 0, 0);
+			GuideVertexIndices = 0;
+			GuideVertexLerps = 0;
+
+			const uint InterpolationData = InterpolationBuffer[VertexIndex];
+			GuideVertexIndices.x = InterpolationData & 0xFFFFFF;
+			GuideVertexLerps.x   = saturate(((InterpolationData >> 24) & 0xFF) / 255.f);
+		  #else
 			const uint4 InterpolationData0	= Interpolation0Buffer[VertexIndex];
 			const uint4 InterpolationData1	= Interpolation1Buffer[VertexIndex];
 
 			GuideVertexIndices.x = InterpolationData0.x | ( (InterpolationData1.x & 0xFF) << 16);
-		#if PERMUTATION_USE_SINGLE_GUIDE == 0
 			GuideVertexIndices.y = InterpolationData0.y | (((InterpolationData1.x >> 8) & 0xFF) << 16);
 			GuideVertexIndices.z = InterpolationData0.z | ( (InterpolationData1.y & 0xFF) << 16);
-		#endif
 
-		#if PERMUTATION_USE_SINGLE_GUIDE == 1
-			GuideVertexWeights = float3(1, 0, 0);
-		#else
 			GuideVertexWeights   = float3(saturate((InterpolationData0.w & 0xFF) / 255.f), saturate(((InterpolationData0.w >> 8) & 0xFF) / 255.f), 0);
 			GuideVertexWeights.z = saturate(1.0f - (GuideVertexWeights.x + GuideVertexWeights.y));
-		#endif
 
 			GuideVertexLerps.x = saturate(((InterpolationData1.y >> 8) & 0xFF)/255.f);
-		#if PERMUTATION_USE_SINGLE_GUIDE == 0
 			GuideVertexLerps.y = saturate(( InterpolationData1.z & 0xFF      )/255.f);
 			GuideVertexLerps.z = saturate(((InterpolationData1.z >> 8) & 0xFF)/255.f);
-		#endif
+		  #endif
 		}
 		#endif
 
