Fix crash when a groom does not have global interplation or simulation enabled.

diff --git a/Engine/Plugins/Runtime/HairStrands/Source/HairStrandsCore/Private/HairStrandsRendering.cpp b/Engine/Plugins/Runtime/HairStrands/Source/HairStrandsCore/Private/HairStrandsRendering.cpp
--- a/Engine/Plugins/Runtime/HairStrands/Source/HairStrandsCore/Private/HairStrandsRendering.cpp
+++ b/Engine/Plugins/Runtime/HairStrands/Source/HairStrandsCore/Private/HairStrandsRendering.cpp
@@ -1609,7 +1609,7 @@ void ComputeHairStrandsInterpolation(
 				// 2.1 Compute deformation position based on simulation/skinning/RBF
 				if (Instance->Debug.GroomCacheType != EGroomCacheType::Strands)
 				{
-					const bool bUseSingleGuide = Instance->Strands.InterpolationResource->UseSingleGuide();
+					const bool bUseSingleGuide = bValidGuide && Instance->Strands.InterpolationResource->UseSingleGuide();
 					AddHairStrandsInterpolationPass(
 						GraphBuilder,
 						ShaderMap,
