Do not serialize hair interpolation data when a groom asset does not use simulation, nor global interpolation.

diff --git a/Engine/Plugins/Runtime/HairStrands/Source/HairStrandsCore/Private/GroomAsset.cpp b/Engine/Plugins/Runtime/HairStrands/Source/HairStrandsCore/Private/GroomAsset.cpp
--- a/Engine/Plugins/Runtime/HairStrands/Source/HairStrandsCore/Private/GroomAsset.cpp
+++ b/Engine/Plugins/Runtime/HairStrands/Source/HairStrandsCore/Private/GroomAsset.cpp
@@ -181,6 +181,7 @@ bool IsHairStrandsAssetLoadingEnable()
 // Guides, strands, and interpolation data are all transient
 static bool BuildHairGroup(
 	const int32 GroupIndex,
+	const bool bNeedInterpolationData,
 	const FHairDescriptionGroups& HairDescriptionGroups, 
 	const TArray<FHairGroupsInterpolation>& HairGroupsInterpolation,
 	const TArray<FHairGroupsLOD>& HairGroupsLOD,
@@ -201,9 +202,17 @@ static bool BuildHairGroup(
 		FGroomBuilder::BuildBulkData(HairGroup.Info, GuidesData, OutHairGroupsData[GroupIndex].Guides.BulkData);
 		FGroomBuilder::BuildBulkData(HairGroup.Info, StrandsData, OutHairGroupsData[GroupIndex].Strands.BulkData);
 
-		FHairStrandsInterpolationDatas InterpolationData;
-		FGroomBuilder::BuildInterplationData(HairGroup.Info, StrandsData, GuidesData, HairGroupsInterpolation[GroupIndex].InterpolationSettings, InterpolationData);
-		FGroomBuilder::BuildInterplationBulkData(GuidesData, InterpolationData, OutHairGroupsData[GroupIndex].Strands.InterpolationBulkData);
+		// If there is no simulation or no global interpolation on that group there is no need for builder the interpolation data
+		if (bNeedInterpolationData)
+		{
+			FHairStrandsInterpolationDatas InterpolationData;
+			FGroomBuilder::BuildInterplationData(HairGroup.Info, StrandsData, GuidesData, HairGroupsInterpolation[GroupIndex].InterpolationSettings, InterpolationData);
+			FGroomBuilder::BuildInterplationBulkData(GuidesData, InterpolationData, OutHairGroupsData[GroupIndex].Strands.InterpolationBulkData);
+		}
+		else
+		{
+			OutHairGroupsData[GroupIndex].Strands.InterpolationBulkData.Reset();
+		}
 
 		FGroomBuilder::BuildClusterData(StrandsData, HairDescriptionGroups.BoundRadius, HairGroupsLOD[GroupIndex], OutHairGroupsData[GroupIndex].Strands.ClusterCullingData);
 	}
@@ -812,14 +821,16 @@ void UGroomAsset::PostLoad()
 		check(LocalHairDescriptionGroups.HairGroups.Num() == GroupCount);
 		for (int32 GroupIndex = 0; GroupIndex < GroupCount; ++GroupIndex)
 		{
+			const bool bNeedInterpoldationData = NeedsInterpolationData(GroupIndex);
 			const bool bNeedToBuildData = 
 				(HairGroupsData[GroupIndex].Guides.BulkData.GetNumCurves() == 0 ||
-				 HairGroupsData[GroupIndex].Strands.InterpolationBulkData.GetPointCount() == 0) &&
+				 (bNeedInterpoldationData && HairGroupsData[GroupIndex].Strands.InterpolationBulkData.GetPointCount() == 0)) &&
 				 HairGroupsData[GroupIndex].Strands.BulkData.GetNumCurves() > 0; // Empty groom has no data to build
 			if (bNeedToBuildData)
 			{
 				BuildHairGroup(
 					GroupIndex,
+					bNeedInterpoldationData,
 					LocalHairDescriptionGroups,
 					HairGroupsInterpolation,
 					HairGroupsLOD,
@@ -899,11 +910,13 @@ void UGroomAsset::PreSave(FObjectPreSaveContext ObjectSaveContext)
 		{
 			const bool bHasInterpolationChanged = !(CachedHairGroupsInterpolation[GroupIt] == HairGroupsInterpolation[GroupIt]);
 			const bool bHasLODChanged = !(CachedHairGroupsLOD[GroupIt] == HairGroupsLOD[GroupIt]);
+			const bool bNeedInterpoldationData = NeedsInterpolationData(GroupIt);
 
 			if (bHasInterpolationChanged)
 			{
 				BuildHairGroup(
 					GroupIt,
+					bNeedInterpoldationData,
 					LocalHairDescriptionGroups,
 					HairGroupsInterpolation,
 					HairGroupsLOD,
@@ -976,6 +989,11 @@ static bool IsStrandsInterpolationAttributes(const FName PropertyName)
 		   PropertyName == GET_MEMBER_NAME_CHECKED(FHairDecimationSettings, CurveDecimation)
 		|| PropertyName == GET_MEMBER_NAME_CHECKED(FHairDecimationSettings, VertexDecimation)
 
+		// Add dependency on simulation and per LOD-simulation/global-interoplation to strip-out interoplation data if there are not needed
+		|| PropertyName == GET_MEMBER_NAME_CHECKED(FHairSolverSettings, EnableSimulation)
+		|| PropertyName == GET_MEMBER_NAME_CHECKED(FHairLODSettings, Simulation)
+		|| PropertyName == GET_MEMBER_NAME_CHECKED(FHairLODSettings, GlobalInterpolation)
+
 		|| PropertyName == GET_MEMBER_NAME_CHECKED(FHairInterpolationSettings, bOverrideGuides)
 		|| PropertyName == GET_MEMBER_NAME_CHECKED(FHairInterpolationSettings, HairToGuideDensity)
 		|| PropertyName == GET_MEMBER_NAME_CHECKED(FHairInterpolationSettings, InterpolationQuality)
@@ -1337,7 +1355,7 @@ void UGroomAsset::SetHairWidth(float Width)
 // differences, etc.) replace the version GUID below with a new one.
 // In case of merge conflicts with DDC versions, you MUST generate a new GUID
 // and set this new GUID as the version.
-#define GROOM_DERIVED_DATA_VERSION TEXT("1AA44249A62B435BB137500F2C880C91")
+#define GROOM_DERIVED_DATA_VERSION TEXT("58A94E5B3E784BBBB5C1FF7D42540EB5")
 
 #if WITH_EDITORONLY_DATA
 
@@ -1354,10 +1372,11 @@ namespace GroomDerivedDataCacheUtils
 		return FDerivedDataCacheInterface::BuildCacheKey(*(TEXT("GROOM_V") + FGroomBuilder::GetVersion() + TEXT("_")), *GetGroomDerivedDataVersion(), *KeySuffix);
 	}
 
-	void SerializeHairInterpolationSettingsForDDC(FArchive& Ar, uint32 GroupIndex, FHairGroupsInterpolation& InterpolationSettings, FHairGroupsLOD& LODSettings)
+	void SerializeHairInterpolationSettingsForDDC(FArchive& Ar, uint32 GroupIndex, FHairGroupsInterpolation& InterpolationSettings, FHairGroupsLOD& LODSettings, bool bRequireInterpolationData)
 	{
 		// Note: this serializer is only used to build the groom DDC key, no versioning is required
 		Ar << GroupIndex;
+		Ar << bRequireInterpolationData;
 
 		InterpolationSettings.BuildDDCKey(Ar);
 		LODSettings.BuildDDCKey(Ar);
@@ -1558,8 +1577,11 @@ FString UGroomAsset::BuildDerivedDataKeySuffix(uint32 GroupIndex, const FHairGro
 	TArray<uint8> TempBytes;
 	TempBytes.Reserve(64);
 	FMemoryWriter Ar(TempBytes, /*bIsPersistent=*/ true);
+	
+	// If simulation or global interpolation is enabled, then interoplation data are required. Otherwise they can be skipped
+	const bool bNeedInterpolationData = NeedsInterpolationData(GroupIndex);
 
-	GroomDerivedDataCacheUtils::SerializeHairInterpolationSettingsForDDC(Ar, GroupIndex, const_cast<FHairGroupsInterpolation&>(InterpolationSettings), const_cast<FHairGroupsLOD&>(LODSettings));
+	GroomDerivedDataCacheUtils::SerializeHairInterpolationSettingsForDDC(Ar, GroupIndex, const_cast<FHairGroupsInterpolation&>(InterpolationSettings), const_cast<FHairGroupsLOD&>(LODSettings), bNeedInterpolationData);
 
 	FString KeySuffix;
 	if (HairDescriptionBulkData)
@@ -1749,9 +1771,12 @@ bool UGroomAsset::CacheStrandsData(uint32 GroupIndex, FString& OutDerivedDataKey
 			return false;
 		}
 
+		const bool bNeedInterpolationData = NeedsInterpolationData(GroupIndex);
+
 		// Build groom data with the new build settings
 		bSuccess = BuildHairGroup(
 			GroupIndex,
+			bNeedInterpolationData,
 			LocalHairDescriptionGroups,
 			HairGroupsInterpolation,
 			HairGroupsLOD,
@@ -2673,6 +2698,19 @@ bool UGroomAsset::IsGlobalInterpolationEnable(int32 GroupIndex, int32 LODIndex)
 		(HairGroupsLOD[GroupIndex].LODs[LODIndex].GlobalInterpolation == EGroomOverrideType::Auto && EnableGlobalInterpolation);
 }
 
+bool UGroomAsset::NeedsInterpolationData(int32 GroupIndex) const
+{
+	for (int32 LODIt = 0; LODIt < HairGroupsLOD.Num(); ++LODIt)
+	{
+		if (IsSimulationEnable(GroupIndex, LODIt) || IsGlobalInterpolationEnable(GroupIndex, LODIt))
+		{
+			return true;
+		}
+	}
+
+	return false;
+}
+
 bool FHairDescriptionGroups::IsValid() const
 {
 	for (const FHairDescriptionGroup& HairGroup : HairGroups)
diff --git a/Engine/Plugins/Runtime/HairStrands/Source/HairStrandsCore/Private/HairStrandsDatas.cpp b/Engine/Plugins/Runtime/HairStrands/Source/HairStrandsCore/Private/HairStrandsDatas.cpp
--- a/Engine/Plugins/Runtime/HairStrands/Source/HairStrandsCore/Private/HairStrandsDatas.cpp
+++ b/Engine/Plugins/Runtime/HairStrands/Source/HairStrandsCore/Private/HairStrandsDatas.cpp
@@ -163,6 +163,14 @@ FArchive& operator<<(FArchive& Ar, FHairInterpolation1Vertex& Vertex)
 	return Ar;
 }
 
+void FHairStrandsInterpolationBulkData::Reset()
+{
+	Interpolation.Empty();
+	Interpolation0.Empty();
+	Interpolation1.Empty();
+	SimRootPointIndex.Empty();
+}
+
 void FHairStrandsInterpolationBulkData::Serialize(FArchive& Ar)
 {
 	static_assert(sizeof(FHairInterpolationVertex::BulkType) == sizeof(FHairInterpolationVertex));
diff --git a/Engine/Plugins/Runtime/HairStrands/Source/HairStrandsCore/Public/GroomAsset.h b/Engine/Plugins/Runtime/HairStrands/Source/HairStrandsCore/Public/GroomAsset.h
--- a/Engine/Plugins/Runtime/HairStrands/Source/HairStrandsCore/Public/GroomAsset.h
+++ b/Engine/Plugins/Runtime/HairStrands/Source/HairStrandsCore/Public/GroomAsset.h
@@ -534,6 +534,7 @@ public:
 	EGroomBindingType GetBindingType(int32 GroupIndex, int32 LODIndex) const;
 	bool IsSimulationEnable(int32 GroupIndex, int32 LODIndex) const;
 	bool IsGlobalInterpolationEnable(int32 GroupIndex, int32 LODIndex) const;
+	bool NeedsInterpolationData(int32 GroupIndex) const;
 
 #if WITH_EDITORONLY_DATA
 	void SetDebugMode(EHairStrandsDebugMode InMode) { DebugMode = InMode; OnGroomAssetChanged.Broadcast(); }
diff --git a/Engine/Plugins/Runtime/HairStrands/Source/HairStrandsCore/Public/HairStrandsDatas.h b/Engine/Plugins/Runtime/HairStrands/Source/HairStrandsCore/Public/HairStrandsDatas.h
--- a/Engine/Plugins/Runtime/HairStrands/Source/HairStrandsCore/Public/HairStrandsDatas.h
+++ b/Engine/Plugins/Runtime/HairStrands/Source/HairStrandsCore/Public/HairStrandsDatas.h
@@ -315,9 +315,11 @@ struct HAIRSTRANDSCORE_API FHairStrandsInterpolationBulkData
 		DataFlags_HasSingleGuideData = 1,
 	};
 
+	void Reset();
 	void Serialize(FArchive& Ar);
-	uint32 Flags;
 	uint32 GetPointCount() const { return (!!(Flags & DataFlags_HasSingleGuideData)) ? Interpolation.Num() : Interpolation0.Num(); };
+
+	uint32 Flags;
 	TArray<FHairStrandsInterpolationFormat::Type> Interpolation;	// Per-rendering-vertex interpolation data (closest guides, weight factors, ...). Data for a single guide
 	TArray<FHairStrandsInterpolation0Format::Type> Interpolation0;	// Per-rendering-vertex interpolation data (closest guides, weight factors, ...). Data for up to 3 guides
 	TArray<FHairStrandsInterpolation1Format::Type> Interpolation1;	// Per-rendering-vertex interpolation data (closest guides, weight factors, ...). Data for up to 3 guides
