Improve niagara hair registration

diff --git a/Engine/Plugins/Runtime/HairStrands/Source/HairStrandsCore/Private/GroomComponent.cpp b/Engine/Plugins/Runtime/HairStrands/Source/HairStrandsCore/Private/GroomComponent.cpp
--- a/Engine/Plugins/Runtime/HairStrands/Source/HairStrandsCore/Private/GroomComponent.cpp
+++ b/Engine/Plugins/Runtime/HairStrands/Source/HairStrandsCore/Private/GroomComponent.cpp
@@ -470,6 +470,7 @@ public:
 			// Initialization of the default skel. mesh transform (refresh during ticking).
 			HairInstance->Debug.SkeletalLocalToWorld = HairInstance->Debug.MeshComponent ? HairInstance->Debug.MeshComponent->GetComponentTransform() : FTransform();
 			HairInstance->Debug.SkeletalPreviousLocalToWorld = HairInstance->Debug.SkeletalLocalToWorld;
+			HairInstance->LocalToWorld = HairInstance->Debug.SkeletalLocalToWorld;
 		}
 	}
 
@@ -1031,23 +1032,6 @@ void UGroomComponent::UpdateHairGroupsDesc()
 	}
 }
 
-void UGroomComponent::ReleaseHairSimulation()
-{
-	for (int32 i = 0; i < NiagaraComponents.Num(); ++i)
-	{
-		if (NiagaraComponents[i] && !NiagaraComponents[i]->IsBeingDestroyed())
-		{
-			if (GetWorld())
-			{
-				NiagaraComponents[i]->DetachFromComponent(FDetachmentTransformRules::KeepRelativeTransform);
-				NiagaraComponents[i]->UnregisterComponent();
-			}
-			NiagaraComponents[i]->DestroyComponent();
-			NiagaraComponents[i] = nullptr;
-		}
-	}
-	NiagaraComponents.Empty();
-}
 bool UGroomComponent::IsSimulationEnable(int32 GroupIndex, int32 LODIndex) const
 {
 	bool bIsSimulationEnable = GroomAsset ? GroomAsset->IsSimulationEnable(GroupIndex, LODIndex) : false;
@@ -1059,110 +1043,83 @@ bool UGroomComponent::IsSimulationEnable(int32 GroupIndex, int32 LODIndex) const
 	return bIsSimulationEnable;
 }
 
-static void InternalUpdateHairSimulation(UGroomComponent* GroomComponent, const bool bHasWorldReady)
+void UGroomComponent::ReleaseHairSimulation()
 {
-	if (!GroomComponent || GroomComponent->GroomGroupsDesc.Num() == 0)
+	for (int32 CompIndex = 0; CompIndex < NiagaraComponents.Num(); ++CompIndex)
 	{
-		return;
+		UNiagaraComponent*& NiagaraComponent = NiagaraComponents[CompIndex];
+		if (NiagaraComponent && !NiagaraComponent->IsBeingDestroyed())
+		{
+			if (GetWorld() && NiagaraComponent->IsRegistered())
+			{
+				NiagaraComponent->DetachFromComponent(FDetachmentTransformRules::KeepRelativeTransform);
+				NiagaraComponent->UnregisterComponent();
+			}
+			NiagaraComponent->DestroyComponent();
+			NiagaraComponent = nullptr;
+		}
 	}
-	UGroomAsset* GroomAsset = GroomComponent->GroomAsset;
+	NiagaraComponents.Empty();
+}
+
+void UGroomComponent::UpdateHairSimulation()
+{
 	const int32 NumGroups = GroomAsset ? GroomAsset->HairGroupsPhysics.Num() : 0;
-	const int32 NumComponents = FMath::Max(NumGroups, GroomComponent->NiagaraComponents.Num());
+	const int32 NumComponents = FMath::Max(NumGroups, NiagaraComponents.Num());
 
 	TArray<bool> ValidComponents;
 	ValidComponents.Init(false, NumComponents);
 
-	bool NeedSpringsSolver = false;
-	bool NeedRodsSolver = false;
 	if (GroomAsset)
 	{
-		for (int32 i = 0; i < NumGroups; ++i)
-		{
-			const int32 LODIndex = GroomComponent->GroomGroupsDesc[i].LODForcedIndex;
-			ValidComponents[i] = GroomComponent->IsSimulationEnable(i, LODIndex);
-			if (ValidComponents[i] && (GroomAsset->HairGroupsPhysics[i].SolverSettings.NiagaraSolver == EGroomNiagaraSolvers::AngularSprings))
-			{
-				NeedSpringsSolver = true;
-			}
-			if (ValidComponents[i] && (GroomAsset->HairGroupsPhysics[i].SolverSettings.NiagaraSolver == EGroomNiagaraSolvers::CosseratRods))
-			{
-				NeedRodsSolver = true;
-			}
-		}
-	}
-	if (bHasWorldReady)
-	{
-		if (IsHairAdaptiveSubstepsEnabled())
+		for (int32 GroupIndex = 0; GroupIndex < NumGroups; ++GroupIndex)
 		{
-			if (NeedSpringsSolver)
+			const int32 NumLODs = GroomAsset->HairGroupsLOD[GroupIndex].LODs.Num();
+			for (int32 LODIndex = 0; LODIndex < NumLODs; ++LODIndex)
 			{
-				GroomComponent->AngularSpringsSystem = LoadObject<UNiagaraSystem>(nullptr, TEXT("/HairStrands/Emitters/SimpleSpringsSystem.SimpleSpringsSystem"));
-			}
-			if (NeedRodsSolver)
-			{
-				GroomComponent->CosseratRodsSystem = LoadObject<UNiagaraSystem>(nullptr, TEXT("/HairStrands/Emitters/SimpleRodsSystem.SimpleRodsSystem"));
-			}
-		}
-		else
-		{
-			if (NeedSpringsSolver)
-			{
-				GroomComponent->AngularSpringsSystem = LoadObject<UNiagaraSystem>(nullptr, TEXT("/HairStrands/Emitters/StableSpringsSystem.StableSpringsSystem"));
-			}
-			if (NeedRodsSolver)
-			{
-				GroomComponent->CosseratRodsSystem = LoadObject<UNiagaraSystem>(nullptr, TEXT("/HairStrands/Emitters/StableRodsSystem.StableRodsSystem"));
+				ValidComponents[GroupIndex] = IsSimulationEnable(GroupIndex, LODIndex);
 			}
 		}
 	}
-	GroomComponent->NiagaraComponents.SetNumZeroed(NumComponents);
-	for (int32 i = 0; i < NumComponents; ++i)
+
+	NiagaraComponents.SetNumZeroed(NumComponents);
+	for (int32 CompIndex = 0; CompIndex < NumComponents; ++CompIndex)
 	{
-		UNiagaraComponent*& NiagaraComponent = GroomComponent->NiagaraComponents[i];
-		if (ValidComponents[i])
+		UNiagaraComponent*& NiagaraComponent = NiagaraComponents[CompIndex];
+		if (ValidComponents[CompIndex])
 		{
 			if (!NiagaraComponent)
 			{
-				NiagaraComponent = NewObject<UNiagaraComponent>(GroomComponent, NAME_None, RF_Transient);
-				if (GroomComponent->GetOwner() && GroomComponent->GetOwner()->GetWorld() && GroomComponent->GetOwner()->GetWorld()->bIsWorldInitialized)
-				{
-					NiagaraComponent->AttachToComponent(GroomComponent, FAttachmentTransformRules::KeepRelativeTransform);
-					NiagaraComponent->RegisterComponent();
-				}
-				else
-				{
-					NiagaraComponent->SetupAttachment(GroomComponent);
-				}
+				NiagaraComponent = NewObject<UNiagaraComponent>(this, NAME_None, RF_Transient);
 				NiagaraComponent->SetVisibleFlag(false);
 			}
-			if (bHasWorldReady)
+			if (GetWorld() && GetWorld()->bIsWorldInitialized)
 			{
-				if (GroomAsset->HairGroupsPhysics[i].SolverSettings.NiagaraSolver == EGroomNiagaraSolvers::AngularSprings)
+				if (!NiagaraComponent->IsRegistered())
 				{
-					NiagaraComponent->SetAsset(GroomComponent->AngularSpringsSystem);
-				}
-				else if (GroomAsset->HairGroupsPhysics[i].SolverSettings.NiagaraSolver == EGroomNiagaraSolvers::CosseratRods)
-				{
-					NiagaraComponent->SetAsset(GroomComponent->CosseratRodsSystem);
-				}
-				else
-				{
-					NiagaraComponent->SetAsset(GroomAsset->HairGroupsPhysics[i].SolverSettings.CustomSystem.LoadSynchronous());
+					NiagaraComponent->AttachToComponent(this, FAttachmentTransformRules::KeepRelativeTransform);
+					NiagaraComponent->RegisterComponentWithWorld(GetWorld());
 				}
+
+				if(!AngularSpringsSystem) AngularSpringsSystem = LoadObject<UNiagaraSystem>(nullptr, TEXT("/HairStrands/Emitters/StableSpringsSystem.StableSpringsSystem"));
+				if(!CosseratRodsSystem) CosseratRodsSystem = LoadObject<UNiagaraSystem>(nullptr, TEXT("/HairStrands/Emitters/StableRodsSystem.StableRodsSystem"));
+
+				UNiagaraSystem* NiagaraAsset =
+					(GroomAsset->HairGroupsPhysics[CompIndex].SolverSettings.NiagaraSolver == EGroomNiagaraSolvers::AngularSprings) ? AngularSpringsSystem :
+					(GroomAsset->HairGroupsPhysics[CompIndex].SolverSettings.NiagaraSolver == EGroomNiagaraSolvers::CosseratRods) ? CosseratRodsSystem :
+					GroomAsset->HairGroupsPhysics[CompIndex].SolverSettings.CustomSystem.LoadSynchronous();
+
+				NiagaraComponent->SetAsset(NiagaraAsset);
+				NiagaraComponent->ReinitializeSystem();
 			}
-			NiagaraComponent->ReinitializeSystem();
 		}
 		else if (NiagaraComponent && !NiagaraComponent->IsBeingDestroyed())
 		{
+			NiagaraComponent->SetAsset(nullptr);
 			NiagaraComponent->DeactivateImmediate();
 		}
 	}
-	GroomComponent->UpdateSimulatedGroups();
-}
-
-void UGroomComponent::UpdateHairSimulation()
-{
-	InternalUpdateHairSimulation(this, true);
+	UpdateSimulatedGroups();
 }
 
 void UGroomComponent::SetGroomAsset(UGroomAsset* Asset)
@@ -1216,8 +1173,7 @@ void UGroomComponent::SetGroomAsset(UGroomAsset* Asset, UGroomBindingAsset* InBi
 	}
 
 	UpdateHairGroupsDesc();
-	//UpdateHairSimulation();
-	InternalUpdateHairSimulation(this,false);
+	UpdateHairSimulation();
 	if (!GroomAsset || !GroomAsset->IsValid())
 	{
 		return;
@@ -1282,18 +1238,18 @@ void UGroomComponent::SetForcedLOD(int32 LODIndex)
 	{
 		LODIndex = -1;
 	}
+	const bool bHasLODChanged = (GroomGroupsDesc.Num() > 0) ? 
+		(IsSimulationEnable(0, GetForcedLOD()) != IsSimulationEnable(0, LODIndex)) : false;
 
-	bool bHasLODChanged = false;
 	for (FHairGroupDesc& HairDesc : GroomGroupsDesc)
 	{
-		bHasLODChanged = bHasLODChanged || HairDesc.LODForcedIndex != LODIndex;
 		HairDesc.LODForcedIndex = LODIndex;
 	}
 	UpdateHairGroupsDesc();
 	
 	if (bHasLODChanged)
 	{
-		InternalUpdateHairSimulation(this, true);
+		UpdateHairSimulation();
 	}
 
 	// Do not invalidate completly the proxy, but just update LOD index on the rendering thread
