Make binding asset to be ediable.

diff --git a/Engine/Plugins/Runtime/HairStrands/Source/HairStrandsCore/Private/GroomBindingAsset.cpp b/Engine/Plugins/Runtime/HairStrands/Source/HairStrandsCore/Private/GroomBindingAsset.cpp
--- a/Engine/Plugins/Runtime/HairStrands/Source/HairStrandsCore/Private/GroomBindingAsset.cpp
+++ b/Engine/Plugins/Runtime/HairStrands/Source/HairStrandsCore/Private/GroomBindingAsset.cpp
@@ -630,10 +630,8 @@ void UGroomBindingAsset::Build()
 void UGroomBindingAsset::PostEditChangeProperty(FPropertyChangedEvent& PropertyChangedEvent)
 {
 	Super::PostEditChangeProperty(PropertyChangedEvent);
-	UpdateResource();
-	OnGroomBindingAssetChanged.Broadcast();
-
 	CacheDerivedDatas();
+	OnGroomBindingAssetChanged.Broadcast();
 }
 #endif // WITH_EDITOR
 
diff --git a/Engine/Plugins/Runtime/HairStrands/Source/HairStrandsCore/Private/GroomResources.cpp b/Engine/Plugins/Runtime/HairStrands/Source/HairStrandsCore/Private/GroomResources.cpp
--- a/Engine/Plugins/Runtime/HairStrands/Source/HairStrandsCore/Private/GroomResources.cpp
+++ b/Engine/Plugins/Runtime/HairStrands/Source/HairStrandsCore/Private/GroomResources.cpp
@@ -759,6 +759,13 @@ void FHairStrandsRestRootResource::PopulateFromRootData()
 
 void FHairStrandsRestRootResource::InternalAllocate(FRDGBuilder& GraphBuilder)
 {
+	// Once empty, the MeshProjectionLODsneeds to be repopulate as it might be re-initialized. 
+	// E.g., when a resource is updated, it is first released, then re-init. 
+	if (LODs.IsEmpty())
+	{
+		PopulateFromRootData();
+	}
+
 	if (RootData.VertexToCurveIndexBuffer.Num() > 0)
 	{
 		FRDGBufferUploader BufferUploader;
@@ -843,10 +850,6 @@ void FHairStrandsRestRootResource::InternalRelease()
 		GPUData.RestSamplePositionsBuffer.Release();
 	}
 	LODs.Empty();
-
-	// Once empty, the MeshProjectionLODsneeds to be repopulate as it might be re-initialized. 
-	// E.g., when a resource is updated, it is first released, then re-init. 
-	PopulateFromRootData();
 }
 
 /////////////////////////////////////////////////////////////////////////////////////////
diff --git a/Engine/Plugins/Runtime/HairStrands/Source/HairStrandsCore/Public/GroomBindingAsset.h b/Engine/Plugins/Runtime/HairStrands/Source/HairStrandsCore/Public/GroomBindingAsset.h
--- a/Engine/Plugins/Runtime/HairStrands/Source/HairStrandsCore/Public/GroomBindingAsset.h
+++ b/Engine/Plugins/Runtime/HairStrands/Source/HairStrandsCore/Public/GroomBindingAsset.h
@@ -66,16 +66,16 @@ public:
 	EGroomBindingMeshType GroomBindingType = EGroomBindingMeshType::SkeletalMesh;
 
 	/** Groom to bind. */
-	UPROPERTY(VisibleAnywhere, BlueprintReadWrite, Category = "BuildSettings")
+	UPROPERTY(EditAnywhere, BlueprintReadWrite, Category = "BuildSettings")
 	UGroomAsset* Groom;
 
 	/** Skeletal mesh on which the groom has been authored. This is optional, and used only if the hair
 		binding is done a different mesh than the one which it has been authored */
-	UPROPERTY(VisibleAnywhere, BlueprintReadWrite, Category = "BuildSettings")
+	UPROPERTY(EditAnywhere, BlueprintReadWrite, Category = "BuildSettings")
 	USkeletalMesh* SourceSkeletalMesh;
 
 	/** Skeletal mesh on which the groom is attached to. */
-	UPROPERTY(VisibleAnywhere, BlueprintReadWrite, Category = "BuildSettings")
+	UPROPERTY(EditAnywhere, BlueprintReadWrite, Category = "BuildSettings")
 	USkeletalMesh* TargetSkeletalMesh;
 
 	UPROPERTY(VisibleAnywhere, BlueprintReadWrite, Category = "BuildSettings")
