Add dump info for groom asset and groom binding asset to report CPU/GPU memory usage.

diff --git a/Engine/Plugins/Runtime/HairStrands/Source/HairStrandsCore/Private/GroomAsset.cpp b/Engine/Plugins/Runtime/HairStrands/Source/HairStrandsCore/Private/GroomAsset.cpp
--- a/Engine/Plugins/Runtime/HairStrands/Source/HairStrandsCore/Private/GroomAsset.cpp
+++ b/Engine/Plugins/Runtime/HairStrands/Source/HairStrandsCore/Private/GroomAsset.cpp
@@ -67,6 +67,267 @@ static int32 GHairMaxSimulatedLOD = -1;
 static FAutoConsoleVariableRef CVarHairMaxSimulatedLOD(TEXT("r.HairStrands.MaxSimulatedLOD"), GHairMaxSimulatedLOD, TEXT("Maximum hair LOD to be simulated"));
 static bool IsHairLODSimulationEnabled(const int32 LODIndex) { return (LODIndex >= 0 && (GHairMaxSimulatedLOD < 0 || (GHairMaxSimulatedLOD >= 0 && LODIndex <= GHairMaxSimulatedLOD))); }
 
+/////////////////////////////////////////////////////////////////////////////////////////
+// Asset dump function
+
+void DumpLoadedGroomAssets(IConsoleVariable* InCVarPakTesterEnabled);
+void DumpLoadedGroomBindingAssets(IConsoleVariable* InCVarPakTesterEnabled);
+
+int32 GHairStrandsDump_GroomAsset = 0;
+static FAutoConsoleVariableRef CVarHairStrandsDump_GroomAsset(
+	TEXT("r.HairStrands.Dump.GroomAsset"),
+	GHairStrandsDump_GroomAsset,
+	TEXT("Dump information of all the loaded groom assets."),
+	FConsoleVariableDelegate::CreateStatic(DumpLoadedGroomAssets));
+
+int32 GHairStrandsDump_GroomBindingAsset = 0;
+static FAutoConsoleVariableRef CVarHairStrandsDump_GroomBindingAsset(
+	TEXT("r.HairStrands.Dump.GroomBindingAsset"),
+	GHairStrandsDump_GroomBindingAsset,
+	TEXT("Dump information of all the loaded groom binding assets."),
+	FConsoleVariableDelegate::CreateStatic(DumpLoadedGroomBindingAssets));
+
+void DumpLoadedGroomAssets(IConsoleVariable* InCVarPakTesterEnabled)
+{
+	const bool bDetails = true;
+	const float ToMb = 1.f / 1000000.f;
+
+	uint32 Total_Asset = 0;
+	uint32 Total_Group = 0;
+
+	uint32 Total_CPUMemorySize_Guides = 0;
+	uint32 Total_CPUMemorySize_Strands= 0;
+	uint32 Total_CPUMemorySize_Cards  = 0;
+	uint32 Total_CPUMemorySize_Meshes = 0;
+	uint32 Total_GPUMemorySize_Guides = 0;
+	uint32 Total_GPUMemorySize_Strands= 0;
+	uint32 Total_GPUMemorySize_Cards  = 0;
+	uint32 Total_GPUMemorySize_Meshes = 0;
+
+	UE_LOG(LogHairStrands, Log, TEXT("[Groom] ##### UGroomAssets #####"));
+	UE_LOG(LogHairStrands, Log, TEXT("----------------------------------------------------------------------------------------------------------------------------------------------------------------"));
+	UE_LOG(LogHairStrands, Log, TEXT("--  No.  - LOD -    CPU Total (     Guides|    Strands|      Cards|     Meshes) -    GPU Total (     Guides|    Strands|      Cards|     Meshes) - Asset Name "));
+	UE_LOG(LogHairStrands, Log, TEXT("----------------------------------------------------------------------------------------------------------------------------------------------------------------"));
+	FString OutputData;
+	for (TObjectIterator<UGroomAsset> AssetIt; AssetIt; ++AssetIt)
+	{
+		if (AssetIt)
+		{
+			const uint32 GroupCount = AssetIt->HairGroupsData.Num();
+			for (uint32 GroupIt = 0; GroupIt < GroupCount; ++GroupIt)
+			{
+				const uint32 CPUMemorySize_Guides  = AssetIt->HairGroupsData[GroupIt].Guides.GetDataSize();
+				const uint32 CPUMemorySize_Strands = AssetIt->HairGroupsData[GroupIt].Strands.GetDataSize();
+				const uint32 CPUMemorySize_Cards   = AssetIt->HairGroupsData[GroupIt].Cards.GetDataSize();
+				const uint32 CPUMemorySize_Meshes  = AssetIt->HairGroupsData[GroupIt].Meshes.GetDataSize();
+
+				const uint32 GPUMemorySize_Guides  = AssetIt->HairGroupsData[GroupIt].Guides.GetResourcesSize();
+				const uint32 GPUMemorySize_Strands = AssetIt->HairGroupsData[GroupIt].Strands.GetResourcesSize();
+				const uint32 GPUMemorySize_Cards   = AssetIt->HairGroupsData[GroupIt].Cards.GetResourcesSize();
+				const uint32 GPUMemorySize_Meshes  = AssetIt->HairGroupsData[GroupIt].Meshes.GetResourcesSize();
+
+				uint32 CPUMemorySize = 0;
+				uint32 GPUMemorySize = 0;
+				CPUMemorySize += CPUMemorySize_Guides;
+				CPUMemorySize += CPUMemorySize_Strands;
+				CPUMemorySize += CPUMemorySize_Cards;
+				CPUMemorySize += CPUMemorySize_Meshes;
+
+				GPUMemorySize += GPUMemorySize_Guides;
+				GPUMemorySize += GPUMemorySize_Strands;
+				GPUMemorySize += GPUMemorySize_Cards;
+				GPUMemorySize += GPUMemorySize_Meshes;
+
+				Total_CPUMemorySize_Guides += CPUMemorySize_Guides;
+				Total_CPUMemorySize_Strands+= CPUMemorySize_Strands;
+				Total_CPUMemorySize_Cards  += CPUMemorySize_Cards;
+				Total_CPUMemorySize_Meshes += CPUMemorySize_Meshes;
+				Total_GPUMemorySize_Guides += GPUMemorySize_Guides;
+				Total_GPUMemorySize_Strands+= GPUMemorySize_Strands;
+				Total_GPUMemorySize_Cards  += GPUMemorySize_Cards;
+				Total_GPUMemorySize_Meshes += GPUMemorySize_Meshes;
+
+				const uint32 LODCount = AssetIt->HairGroupsLOD[GroupIt].LODs.Num();
+				if (bDetails)
+				{
+//					UE_LOG(LogHairStrands, Log, TEXT("--  No.  - LOD -    CPU Total (     Guides|    Strands|      Cards|     Meshes) -    GPU Total (     Guides|    Strands|      Cards|     Meshes) - Asset Name "));
+//					UE_LOG(LogHairStrands, Log, TEXT("-- 00/00 -  00 -  0000.0Mb (0000.0Mb |0000.0Mb|0000.0Mb |0000.0Mb) -  0000.0Mb (0000.0Mb|0000.0Mb|0000.0Mb|0000.0Mb) - AssetName"));
+					UE_LOG(LogHairStrands, Log, TEXT("-- %2d/%2d -  %2d -  %9.3fMb (%9.3fMb|%9.3fMb|%9.3fMb|%9.3fMb) -  %9.3fMb (%9.3fMb|%9.3fMb|%9.3fMb|%9.3fMb) - %s"), 
+						GroupIt,
+						GroupCount,
+
+						LODCount, 
+
+						CPUMemorySize * ToMb,
+						CPUMemorySize_Guides * ToMb,
+						CPUMemorySize_Strands* ToMb,
+						CPUMemorySize_Cards  * ToMb,
+						CPUMemorySize_Meshes * ToMb,
+
+						GPUMemorySize* ToMb,
+						GPUMemorySize_Guides * ToMb,
+						GPUMemorySize_Strands* ToMb,
+						GPUMemorySize_Cards  * ToMb,
+						GPUMemorySize_Meshes * ToMb,
+						GroupIt == 0 ? *AssetIt->GetPathName() : TEXT("."));
+				}
+			}
+
+			Total_Asset++;
+			Total_Group += GroupCount;
+		}
+	}
+
+	uint32 Total_CPUMemorySize = 0;
+	Total_CPUMemorySize += Total_CPUMemorySize_Guides;
+	Total_CPUMemorySize += Total_CPUMemorySize_Strands;
+	Total_CPUMemorySize += Total_CPUMemorySize_Cards;
+	Total_CPUMemorySize += Total_CPUMemorySize_Meshes;
+
+	uint32 Total_GPUMemorySize = 0;
+	Total_GPUMemorySize += Total_GPUMemorySize_Guides;
+	Total_GPUMemorySize += Total_GPUMemorySize_Strands;
+	Total_GPUMemorySize += Total_GPUMemorySize_Cards;
+	Total_GPUMemorySize += Total_GPUMemorySize_Meshes;
+
+	UE_LOG(LogHairStrands, Log, TEXT("----------------------------------------------------------------------------------------------------------------------------------------------------------------"));
+	UE_LOG(LogHairStrands, Log, TEXT("-- A:%3d|G:%3d -  %9.3fMb (%9.3fMb|%9.3fMb|%9.3fMb|%9.3fMb) -  %9.3fMb (%9.3fMb|%9.3fMb|%9.3fMb|%9.3fMb)"),
+		Total_Asset,
+		Total_Group,
+
+		Total_CPUMemorySize * ToMb,
+		Total_CPUMemorySize_Guides * ToMb,
+		Total_CPUMemorySize_Strands * ToMb,
+		Total_CPUMemorySize_Cards * ToMb,
+		Total_CPUMemorySize_Meshes * ToMb,
+
+		Total_GPUMemorySize * ToMb,
+		Total_GPUMemorySize_Guides * ToMb,
+		Total_GPUMemorySize_Strands * ToMb,
+		Total_GPUMemorySize_Cards * ToMb,
+		Total_GPUMemorySize_Meshes * ToMb);
+}
+
+void DumpLoadedGroomBindingAssets(IConsoleVariable* InCVarPakTesterEnabled)
+{
+	const bool bDetails = true;
+	const float ToMb = 1.f / 1000000.f;
+
+	uint32 Total_Asset = 0;
+	uint32 Total_Group = 0;
+
+	uint32 Total_CPUMemorySize_Guides = 0;
+	uint32 Total_CPUMemorySize_Strands= 0;
+	uint32 Total_CPUMemorySize_Cards  = 0;
+	
+	uint32 Total_GPUMemorySize_Guides = 0;
+	uint32 Total_GPUMemorySize_Strands= 0;
+	uint32 Total_GPUMemorySize_Cards  = 0;
+	
+
+	UE_LOG(LogHairStrands, Log, TEXT("[Groom] ##### UGroomBindingAssets #####"));
+	UE_LOG(LogHairStrands, Log, TEXT("----------------------------------------------------------------------------------------------------------------------------------------------------------------"));
+	UE_LOG(LogHairStrands, Log, TEXT("--  No.  - LOD -    CPU Total (     Guides|    Strands|      Cards) -    GPU Total (     Guides|    Strands|      Cards) - Asset Name "));
+	UE_LOG(LogHairStrands, Log, TEXT("----------------------------------------------------------------------------------------------------------------------------------------------------------------"));
+	FString OutputData;
+	for (TObjectIterator<UGroomBindingAsset> AssetIt; AssetIt; ++AssetIt)
+	{
+		if (AssetIt)
+		{			
+			const uint32 GroupCount = AssetIt->HairGroupDatas.Num();
+			for (uint32 GroupIt = 0; GroupIt < GroupCount; ++GroupIt)
+			{
+				const uint32 CPUMemorySize_Guides  = AssetIt->HairGroupDatas[GroupIt].SimRootData.GetDataSize();
+				const uint32 CPUMemorySize_Strands = AssetIt->HairGroupDatas[GroupIt].RenRootData.GetDataSize();
+
+				uint32 GPUMemorySize_Guides = 0;
+				if (const FHairStrandsRestRootResource* RootResource = AssetIt->HairGroupResources[GroupIt].SimRootResources)
+				{
+					GPUMemorySize_Guides += RootResource->GetResourcesSize();
+				}
+				uint32 GPUMemorySize_Strands = 0;
+				if (const FHairStrandsRestRootResource* RootResource = AssetIt->HairGroupResources[GroupIt].RenRootResources)
+				{
+					GPUMemorySize_Strands += RootResource->GetResourcesSize();
+				}
+
+				uint32 CPUMemorySize_Cards = 0;
+				uint32 GPUMemorySize_Cards = 0;
+				const uint32 CardCount = AssetIt->HairGroupDatas[GroupIt].CardsRootData.Num();
+				for (uint32 CardIt = 0; CardIt < CardCount; ++CardIt)
+				{
+					CPUMemorySize_Cards += AssetIt->HairGroupDatas[GroupIt].CardsRootData[CardIt].GetDataSize();
+					if (const FHairStrandsRestRootResource* RootResource = AssetIt->HairGroupResources[GroupIt].CardsRootResources[CardIt])
+					{
+						GPUMemorySize_Cards += RootResource->GetResourcesSize();
+					}
+				}
+
+				const uint32 CPUMemorySize = CPUMemorySize_Guides + CPUMemorySize_Strands + CPUMemorySize_Cards;
+				const uint32 GPUMemorySize = GPUMemorySize_Guides + GPUMemorySize_Strands + GPUMemorySize_Cards;
+
+				Total_CPUMemorySize_Guides += CPUMemorySize_Guides;
+				Total_CPUMemorySize_Strands+= CPUMemorySize_Strands;
+				Total_CPUMemorySize_Cards  += CPUMemorySize_Cards;
+				
+				Total_GPUMemorySize_Guides += GPUMemorySize_Guides;
+				Total_GPUMemorySize_Strands+= GPUMemorySize_Strands;
+				Total_GPUMemorySize_Cards  += GPUMemorySize_Cards;
+
+				const uint32 SkelLODCount = AssetIt->HairGroupDatas[GroupIt].RenRootData.MeshProjectionLODs.Num();
+				if (bDetails)
+				{
+//					UE_LOG(LogHairStrands, Log, TEXT("--  No.  - LOD -    CPU Total (     Guides|    Strands|      Cards) -    GPU Total (     Guides|    Strands|      Cards) - Asset Name "));
+//					UE_LOG(LogHairStrands, Log, TEXT("-- 00/00 -  00 -  0000.0Mb (0000.0Mb |0000.0Mb|0000.0Mb) -  0000.0Mb (0000.0Mb|0000.0Mb|0000.0Mb) - AssetName"));
+					UE_LOG(LogHairStrands, Log, TEXT("-- %2d/%2d -  %2d -  %9.3fMb (%9.3fMb|%9.3fMb|%9.3fMb) -  %9.3fMb (%9.3fMb|%9.3fMb|%9.3fMb) - %s"), 
+						GroupIt,
+						GroupCount,
+
+						SkelLODCount,
+
+						CPUMemorySize * ToMb,
+						CPUMemorySize_Guides * ToMb,
+						CPUMemorySize_Strands* ToMb,
+						CPUMemorySize_Cards  * ToMb,
+
+						GPUMemorySize* ToMb,
+						GPUMemorySize_Guides * ToMb,
+						GPUMemorySize_Strands* ToMb,
+						GPUMemorySize_Cards  * ToMb,
+						GroupIt == 0 ? *AssetIt->GetPathName() : TEXT("."));
+				}
+			}
+
+			Total_Asset++;
+			Total_Group += GroupCount;
+		}
+	}
+
+	uint32 Total_CPUMemorySize = 0;
+	Total_CPUMemorySize += Total_CPUMemorySize_Guides;
+	Total_CPUMemorySize += Total_CPUMemorySize_Strands;
+	Total_CPUMemorySize += Total_CPUMemorySize_Cards;
+
+	uint32 Total_GPUMemorySize = 0;
+	Total_GPUMemorySize += Total_GPUMemorySize_Guides;
+	Total_GPUMemorySize += Total_GPUMemorySize_Strands;
+	Total_GPUMemorySize += Total_GPUMemorySize_Cards;
+
+	UE_LOG(LogHairStrands, Log, TEXT("----------------------------------------------------------------------------------------------------------------------------------------------------------------"));
+	UE_LOG(LogHairStrands, Log, TEXT("-- A:%3d|G:%3d -  %9.3fMb (%9.3fMb|%9.3fMb|%9.3fMb) -  %9.3fMb (%9.3fMb|%9.3fMb|%9.3fMb)"),
+		Total_Asset,
+		Total_Group,
+
+		Total_CPUMemorySize * ToMb,
+		Total_CPUMemorySize_Guides * ToMb,
+		Total_CPUMemorySize_Strands * ToMb,
+		Total_CPUMemorySize_Cards * ToMb,
+
+		Total_GPUMemorySize * ToMb,
+		Total_GPUMemorySize_Guides * ToMb,
+		Total_GPUMemorySize_Strands * ToMb,
+		Total_GPUMemorySize_Cards * ToMb);
+}
 /////////////////////////////////////////////////////////////////////////////////////////
 
 void UpdateHairStrandsLogVerbosity()
diff --git a/Engine/Plugins/Runtime/HairStrands/Source/HairStrandsCore/Public/HairStrandsDatas.h b/Engine/Plugins/Runtime/HairStrands/Source/HairStrandsCore/Public/HairStrandsDatas.h
--- a/Engine/Plugins/Runtime/HairStrands/Source/HairStrandsCore/Public/HairStrandsDatas.h
+++ b/Engine/Plugins/Runtime/HairStrands/Source/HairStrandsCore/Public/HairStrandsDatas.h
@@ -502,6 +502,25 @@ struct FHairStrandsRootData
 	bool HasProjectionData() const;
 	bool IsValid() const { return RootCount > 0; }
 
+	uint32 GetDataSize() const
+	{
+		uint32 Total = 0;
+		Total += VertexToCurveIndexBuffer.GetAllocatedSize();
+		for (const FMeshProjectionLOD& LOD : MeshProjectionLODs)
+		{
+			Total += LOD.RootTriangleIndexBuffer.GetAllocatedSize();
+			Total += LOD.RootTriangleBarycentricBuffer.GetAllocatedSize();
+			Total += LOD.RestRootTrianglePosition0Buffer.GetAllocatedSize();
+			Total += LOD.RestRootTrianglePosition1Buffer.GetAllocatedSize();
+			Total += LOD.RestRootTrianglePosition2Buffer.GetAllocatedSize();
+			Total += LOD.MeshInterpolationWeightsBuffer.GetAllocatedSize();
+			Total += LOD.MeshSampleIndicesBuffer.GetAllocatedSize();
+			Total += LOD.RestSamplePositionsBuffer.GetAllocatedSize();
+			Total += LOD.ValidSectionIndices.GetAllocatedSize();
+		}
+		return Total;
+	}
+
 	struct FMeshProjectionLOD
 	{
 		int32 LODIndex = -1;
