Fix CIS issue.

diff --git a/Engine/Plugins/Runtime/HairStrands/Source/HairStrandsCore/Private/GroomAsset.cpp b/Engine/Plugins/Runtime/HairStrands/Source/HairStrandsCore/Private/GroomAsset.cpp
--- a/Engine/Plugins/Runtime/HairStrands/Source/HairStrandsCore/Private/GroomAsset.cpp
+++ b/Engine/Plugins/Runtime/HairStrands/Source/HairStrandsCore/Private/GroomAsset.cpp
@@ -29,6 +29,7 @@
 #include "Logging/LogMacros.h"
 #include "HairStrandsCore.h"
 #include "UObject/ObjectSaveContext.h"
+#include "GroomBindingAsset.h"
 
 #if WITH_EDITOR
 #include "Interfaces/ITargetPlatform.h"
