Add dump information for CPU/GPU memory usage for groom component.

diff --git a/Engine/Plugins/Runtime/HairStrands/Source/HairStrandsCore/Private/GroomComponent.cpp b/Engine/Plugins/Runtime/HairStrands/Source/HairStrandsCore/Private/GroomComponent.cpp
--- a/Engine/Plugins/Runtime/HairStrands/Source/HairStrandsCore/Private/GroomComponent.cpp
+++ b/Engine/Plugins/Runtime/HairStrands/Source/HairStrandsCore/Private/GroomComponent.cpp
@@ -3022,6 +3022,164 @@ void UGroomComponent::GetResourceSizeEx(FResourceSizeEx& CumulativeResourceSize)
 	}
 }
 
+uint32 UGroomComponent::GetResourcesSize() const
+{
+	uint32 Total = 0;
+	for (const FHairGroupInstance* Instance : HairGroupInstances)
+	{
+		Total += Instance->HairGroupPublicData ? Instance->HairGroupPublicData->GetResourcesSize() : 0;
+
+		Total += Instance->Guides.DeformedResource ? Instance->Guides.DeformedResource->GetResourcesSize() : 0;
+		Total += Instance->Guides.DeformedRootResource ? Instance->Guides.DeformedRootResource->GetResourcesSize() : 0;
+
+		Total += Instance->Strands.DeformedResource ? Instance->Strands.DeformedResource->GetResourcesSize() : 0;
+		Total += Instance->Strands.DeformedRootResource ? Instance->Strands.DeformedRootResource->GetResourcesSize() : 0;
+	#if RHI_RAYTRACING
+		Total += (Instance->Strands.RenRaytracingResourceOwned && Instance->Strands.RenRaytracingResource) ? Instance->Strands.RenRaytracingResource->GetResourcesSize() : 0;
+	#endif
+
+		for (const FHairGroupInstance::FCards::FLOD& LOD : Instance->Cards.LODs)
+		{
+			Total += LOD.DeformedResource ? LOD.DeformedResource->GetResourcesSize() : 0;
+		}
+
+		for (const FHairGroupInstance::FMeshes::FLOD& LOD : Instance->Meshes.LODs)
+		{
+			Total += LOD.DeformedResource ? LOD.DeformedResource->GetResourcesSize() : 0;
+		}
+	}
+
+	return Total;
+}
+
+/////////////////////////////////////////////////////////////////////////////////////////
+// Asset dump function
+
+void DumpLoadedGroomComponent(IConsoleVariable* InCVarPakTesterEnabled);
+
+int32 GHairStrandsDump_GroomComponent = 0;
+static FAutoConsoleVariableRef CVarHairStrandsDump_GroomComponent(
+	TEXT("r.HairStrands.Dump.GroomComponent"),
+	GHairStrandsDump_GroomComponent,
+	TEXT("Dump information of all active groom components."),
+	FConsoleVariableDelegate::CreateStatic(DumpLoadedGroomComponent));
+
+void DumpLoadedGroomComponent(IConsoleVariable* InCVarPakTesterEnabled)
+{
+	const bool bDetails = true;
+	const float ToMb = 1.f / 1000000.f;
+
+	uint32 Total_Component = 0;
+	uint32 Total_Group = 0;
+
+	uint32 Total_GPUMemorySize_Common = 0;
+	uint32 Total_GPUMemorySize_Guides = 0;
+	uint32 Total_GPUMemorySize_Strands= 0;
+	uint32 Total_GPUMemorySize_Cards  = 0;
+	uint32 Total_GPUMemorySize_Meshes = 0;
+
+	UE_LOG(LogHairStrands, Log, TEXT("[Groom] ##### UGroomComponent #####"));
+	UE_LOG(LogHairStrands, Log, TEXT("----------------------------------------------------------------------------------------------------------------------------------------------------------------"));
+	UE_LOG(LogHairStrands, Log, TEXT("--  No.  - LOD -          GPU (     Common|     Guides|    Strands|      Cards|     Meshes) - Asset Name "));
+	UE_LOG(LogHairStrands, Log, TEXT("----------------------------------------------------------------------------------------------------------------------------------------------------------------"));
+	FString OutputData;
+	for (TObjectIterator<UGroomComponent> ComponentIt; ComponentIt; ++ComponentIt)
+	{
+		if (ComponentIt)
+		{
+			const uint32 GroupCount = ComponentIt->GetGroupCount();
+			for (uint32 GroupIt = 0; GroupIt < GroupCount; ++GroupIt)
+			{
+				FHairGroupInstance* Instance = ComponentIt->GetGroupInstance(GroupIt);
+
+				uint32 GPUMemorySize_Common = 0;
+				GPUMemorySize_Common += Instance->HairGroupPublicData ? Instance->HairGroupPublicData->GetResourcesSize() : 0;
+
+				uint32 GPUMemorySize_Guides = 0;
+				GPUMemorySize_Guides += Instance->Guides.DeformedResource ? Instance->Guides.DeformedResource->GetResourcesSize() : 0;
+				GPUMemorySize_Guides += Instance->Guides.DeformedRootResource ? Instance->Guides.DeformedRootResource->GetResourcesSize() : 0;
+
+				uint32 GPUMemorySize_Strands = 0;
+				GPUMemorySize_Strands += Instance->Strands.DeformedResource ? Instance->Strands.DeformedResource->GetResourcesSize() : 0;
+				GPUMemorySize_Strands += Instance->Strands.DeformedRootResource ? Instance->Strands.DeformedRootResource->GetResourcesSize() : 0;
+			#if RHI_RAYTRACING
+				GPUMemorySize_Strands += (Instance->Strands.RenRaytracingResourceOwned && Instance->Strands.RenRaytracingResource) ? Instance->Strands.RenRaytracingResource->GetResourcesSize() : 0;
+			#endif
+
+				uint32 GPUMemorySize_Cards  = 0;
+				for (const FHairGroupInstance::FCards::FLOD& LOD : Instance->Cards.LODs)
+				{
+					GPUMemorySize_Cards += LOD.DeformedResource ? LOD.DeformedResource->GetResourcesSize() : 0;
+				}
+
+				uint32 GPUMemorySize_Meshes = 0;
+				for (const FHairGroupInstance::FMeshes::FLOD& LOD : Instance->Meshes.LODs)
+				{
+					GPUMemorySize_Meshes += LOD.DeformedResource ? LOD.DeformedResource->GetResourcesSize() : 0;
+				}
+
+				uint32 LODCount = Instance->Cards.LODs.Num();
+
+				uint32 GPUMemorySize = 0;
+				GPUMemorySize += GPUMemorySize_Common;
+				GPUMemorySize += GPUMemorySize_Guides;
+				GPUMemorySize += GPUMemorySize_Strands;
+				GPUMemorySize += GPUMemorySize_Cards;
+				GPUMemorySize += GPUMemorySize_Meshes;
+
+				Total_GPUMemorySize_Common += GPUMemorySize_Common;
+				Total_GPUMemorySize_Guides += GPUMemorySize_Guides;
+				Total_GPUMemorySize_Strands+= GPUMemorySize_Strands;
+				Total_GPUMemorySize_Cards  += GPUMemorySize_Cards;
+				Total_GPUMemorySize_Meshes += GPUMemorySize_Meshes;
+
+				if (bDetails)
+				{
+//					UE_LOG(LogHairStrands, Log, TEXT("--  No.  - LOD -    GPU Total (    Common|     Guides|    Strands|      Cards|     Meshes) - Asset Name "));
+//					UE_LOG(LogHairStrands, Log, TEXT("-- 00/00 -  00 -  0000.0Mb (0000.0Mb |0000.0Mb|0000.0Mb |0000.0Mb) -  0000.0Mb (0000.0Mb|0000.0Mb|0000.0Mb|0000.0Mb) - AssetName"));
+					UE_LOG(LogHairStrands, Log, TEXT("-- %2d/%2d -  %2d -  %9.3fMb (%9.3fMb|%9.3fMb|%9.3fMb|%9.3fMb|%9.3fMb) - %s"), 
+						GroupIt,
+						GroupCount,
+
+						LODCount, 
+
+						GPUMemorySize* ToMb,
+						GPUMemorySize_Common * ToMb,
+						GPUMemorySize_Guides * ToMb,
+						GPUMemorySize_Strands* ToMb,
+						GPUMemorySize_Cards  * ToMb,
+						GPUMemorySize_Meshes * ToMb,
+						GroupIt == 0 && ComponentIt->GroomAsset ? *ComponentIt->GroomAsset->GetPathName() : TEXT("."));
+				}
+			}
+
+			Total_Component++;
+			Total_Group += GroupCount;
+		}
+	}
+
+	uint32 Total_GPUMemorySize = 0;
+	Total_GPUMemorySize += Total_GPUMemorySize_Common;
+	Total_GPUMemorySize += Total_GPUMemorySize_Guides;
+	Total_GPUMemorySize += Total_GPUMemorySize_Strands;
+	Total_GPUMemorySize += Total_GPUMemorySize_Cards;
+	Total_GPUMemorySize += Total_GPUMemorySize_Meshes;
+
+	UE_LOG(LogHairStrands, Log, TEXT("----------------------------------------------------------------------------------------------------------------------------------------------------------------"));
+	UE_LOG(LogHairStrands, Log, TEXT("-- C:%3d|G:%3d -  %9.3fMb (%9.3fMb|%9.3fMb|%9.3fMb|%9.3fMb|%9.3fMb)"),
+		Total_Component,
+		Total_Group,
+
+		Total_GPUMemorySize * ToMb,
+		Total_GPUMemorySize_Common* ToMb,
+		Total_GPUMemorySize_Guides * ToMb,
+		Total_GPUMemorySize_Strands * ToMb,
+		Total_GPUMemorySize_Cards * ToMb,
+		Total_GPUMemorySize_Meshes * ToMb);
+}
+
+/////////////////////////////////////////////////////////////////////////////////////////
+
 #if WITH_EDITORONLY_DATA
 FGroomComponentRecreateRenderStateContext::FGroomComponentRecreateRenderStateContext(UGroomAsset* GroomAsset)
 {
diff --git a/Engine/Plugins/Runtime/HairStrands/Source/HairStrandsCore/Public/GroomComponent.h b/Engine/Plugins/Runtime/HairStrands/Source/HairStrandsCore/Public/GroomComponent.h
--- a/Engine/Plugins/Runtime/HairStrands/Source/HairStrandsCore/Public/GroomComponent.h
+++ b/Engine/Plugins/Runtime/HairStrands/Source/HairStrandsCore/Public/GroomComponent.h
@@ -98,6 +98,7 @@ public:
 
 	//~ Begin UObject Interface.
 	virtual void GetResourceSizeEx(FResourceSizeEx& CumulativeResourceSize) override;
+	uint32 GetResourcesSize() const; 
 	//~ End UObject Interface.
 
 	//~ Begin UActorComponent Interface.
diff --git a/Engine/Source/Runtime/Renderer/Private/HairStrands/HairStrandsInterface.cpp b/Engine/Source/Runtime/Renderer/Private/HairStrands/HairStrandsInterface.cpp
--- a/Engine/Source/Runtime/Renderer/Private/HairStrands/HairStrandsInterface.cpp
+++ b/Engine/Source/Runtime/Renderer/Private/HairStrands/HairStrandsInterface.cpp
@@ -343,6 +343,23 @@ void FHairGroupPublicData::ReleaseRHI()
 	CulledVertexRadiusScaleBuffer.Release();
 }
 
+uint32 FHairGroupPublicData::GetResourcesSize() const
+{
+	auto ExtractSize = [](const TRefCountPtr<FRDGPooledBuffer>& InBuffer)
+	{
+		return InBuffer ? InBuffer->Desc.BytesPerElement * InBuffer->Desc.NumElements : 0; 
+	};
+
+	uint32 Total = 0;
+	Total += ExtractSize(DrawIndirectBuffer.Buffer);
+	Total += ExtractSize(DrawIndirectRasterComputeBuffer.Buffer);
+	Total += ExtractSize(ClusterAABBBuffer.Buffer);
+	Total += ExtractSize(GroupAABBBuffer.Buffer);
+	Total += ExtractSize(CulledVertexIdBuffer.Buffer);
+	Total += ExtractSize(CulledVertexRadiusScaleBuffer.Buffer);
+	return Total;
+}
+
 ///////////////////////////////////////////////////////////////////////////////////////////////////
 void TransitBufferToReadable(FRDGBuilder& GraphBuilder, FBufferTransitionQueue& BuffersToTransit)
 {
diff --git a/Engine/Source/Runtime/Renderer/Public/HairStrandsInterface.h b/Engine/Source/Runtime/Renderer/Public/HairStrandsInterface.h
--- a/Engine/Source/Runtime/Renderer/Public/HairStrandsInterface.h
+++ b/Engine/Source/Runtime/Renderer/Public/HairStrandsInterface.h
@@ -154,6 +154,7 @@ public:
 	virtual void InitRHI() override;
 	virtual void ReleaseRHI() override;
 	virtual FString GetFriendlyName() const override { return TEXT("FHairGroupPublicData"); }
+	uint32 GetResourcesSize() const;
 
 	// The primitive count when no culling and neither lod happens
 	uint32 GetGroupInstanceVertexCount() const { return GroupControlTriangleStripVertexCount; }
