Change hair strands data to be lazily loaded to avoid loading all data while only a portion will be used.

This CL changes the data to serialize into byte bulk data, and for its content to not be inlined, meaning it need to be explicitly loaded from the archive to acces its content. This works only for cook build. With editor data, the data are fetched and loaded from DDC as usual.

diff --git a/Engine/Plugins/Runtime/HairStrands/Source/HairStrandsCore/Private/GroomAsset.cpp b/Engine/Plugins/Runtime/HairStrands/Source/HairStrandsCore/Private/GroomAsset.cpp
--- a/Engine/Plugins/Runtime/HairStrands/Source/HairStrandsCore/Private/GroomAsset.cpp
+++ b/Engine/Plugins/Runtime/HairStrands/Source/HairStrandsCore/Private/GroomAsset.cpp
@@ -668,6 +668,16 @@ bool UGroomAsset::HasImportedStrandsData() const
 }
 #endif
 
+// Serialization for *array* of hair elements
+static void InternalSerialize(FArchive& Ar, UObject* Owner, TArray<FHairGroupData::FCards::FLOD>& CardLODData);
+static void InternalSerialize(FArchive& Ar, UObject* Owner, TArray<FHairGroupData::FMeshes::FLOD>& MeshLODData);
+static void InternalSerialize(FArchive& Ar, UObject* Owner, TArray<FHairGroupData>& GroupData);
+
+// Serialization for hair elements
+static void InternalSerialize(FArchive& Ar, UObject* Owner, FHairGroupData::FCards::FLOD& CardLODData);
+static void InternalSerialize(FArchive& Ar, UObject* Owner, FHairGroupData::FMeshes::FLOD& MeshLODData);
+static void InternalSerialize(FArchive& Ar, UObject* Owner, FHairGroupData& GroupData);
+
 void UGroomAsset::Serialize(FArchive& Ar)
 {
 	uint8 ClassDataStripFlags = GenerateClassStripFlags(Ar);
@@ -696,7 +706,7 @@ void UGroomAsset::Serialize(FArchive& Ar)
 		{
 			// When cooking data or serializing old format to new format,
 			// serialize the computed groom data
-			Ar << HairGroupsData;
+			InternalSerialize(Ar, this, HairGroupsData);
 		}
 #if WITH_EDITORONLY_DATA
 		else
@@ -1377,14 +1387,60 @@ int32 UGroomAsset::GetNumHairGroups() const
 	return HairGroupsData.Num();
 }
 
-FArchive& operator<<(FArchive& Ar, FHairGroupData::FCards::FLOD& CardLODData)
+////////////////////////////////////////////////////////////////////////////////////////////////////////////
+// Array serialize
+void InternalSerialize(FArchive& Ar, UObject* Owner, TArray<FHairGroupData::FMeshes::FLOD>& LODs)
+{
+	uint32 Count = LODs.Num();
+	Ar << Count;
+	if (Ar.IsLoading())
+	{
+		LODs.SetNum(Count);
+	}
+	for (uint32 MeshIt = 0; MeshIt < Count; ++MeshIt)
+	{
+		InternalSerialize(Ar, Owner, LODs[MeshIt]);
+	}
+}
+
+void InternalSerialize(FArchive& Ar, UObject* Owner, TArray<FHairGroupData::FCards::FLOD>& LODs)
+{
+	uint32 Count = LODs.Num();
+	Ar << Count;
+	if (Ar.IsLoading())
+	{
+		LODs.SetNum(Count);
+	}
+	for (uint32 MeshIt = 0; MeshIt < Count; ++MeshIt)
+	{
+		InternalSerialize(Ar, Owner, LODs[MeshIt]);
+	}
+}
+
+void InternalSerialize(FArchive& Ar, UObject* Owner, TArray<FHairGroupData>& GroupDatas)
+{
+	uint32 GroupCount = GroupDatas.Num();
+	Ar << GroupCount;
+	if (Ar.IsLoading())
+	{
+		GroupDatas.SetNum(GroupCount);
+	}
+	for (uint32 GroupIt = 0; GroupIt < GroupCount; ++GroupIt)
+	{
+		InternalSerialize(Ar, Owner, GroupDatas[GroupIt]);
+	}
+}
+
+////////////////////////////////////////////////////////////////////////////////////////////////////////////
+// Element serialize
+static void InternalSerialize(FArchive& Ar, UObject* Owner, FHairGroupData::FCards::FLOD& CardLODData)
 {
 	if (!Ar.IsCooking() || !CardLODData.bIsCookedOut)
 	{
 		CardLODData.BulkData.Serialize(Ar);
 		CardLODData.InterpolationBulkData.Serialize(Ar);
 
-		CardLODData.Guides.BulkData.Serialize(Ar);
+		CardLODData.Guides.BulkData.Serialize(Ar, Owner);
 		CardLODData.Guides.InterpolationBulkData.Serialize(Ar);
 	}
 	else
@@ -1397,14 +1453,12 @@ FArchive& operator<<(FArchive& Ar, FHairGroupData::FCards::FLOD& CardLODData)
 		NoInterpolationBulkData.Serialize(Ar);
 
 		FHairGroupData::FBaseWithInterpolation NoGuideData;
-		NoGuideData.BulkData.Serialize(Ar);
+		NoGuideData.BulkData.Serialize(Ar, Owner);
 		NoGuideData.InterpolationBulkData.Serialize(Ar);
 	}
-
-	return Ar;
 }
 
-FArchive& operator<<(FArchive& Ar, FHairGroupData::FMeshes::FLOD& MeshLODData)
+static void InternalSerialize(FArchive& Ar, UObject* Owner, FHairGroupData::FMeshes::FLOD& MeshLODData)
 {
 	if (!Ar.IsCooking() || !MeshLODData.bIsCookedOut)
 	{
@@ -1416,19 +1470,17 @@ FArchive& operator<<(FArchive& Ar, FHairGroupData::FMeshes::FLOD& MeshLODData)
 		FHairMeshesBulkData NoMeshesBulkData;
 		NoMeshesBulkData.Serialize(Ar);
 	}
-
-	return Ar;
 }
 
-FArchive& operator<<(FArchive& Ar, FHairGroupData& GroupData)
+static void InternalSerialize(FArchive& Ar, UObject* Owner, FHairGroupData& GroupData)
 {
 	Ar.UsingCustomVersion(FAnimObjectVersion::GUID);
 
 	FHairGroupData NoStrandsData;
 	if (!Ar.IsCooking() || !GroupData.bIsCookedOut)
 	{
-		GroupData.Strands.BulkData.Serialize(Ar);
-		GroupData.Guides.BulkData.Serialize(Ar);
+		GroupData.Strands.BulkData.Serialize(Ar, Owner);
+		GroupData.Guides.BulkData.Serialize(Ar, Owner);
 		GroupData.Strands.InterpolationBulkData.Serialize(Ar);
 
 	}
@@ -1437,8 +1489,8 @@ FArchive& operator<<(FArchive& Ar, FHairGroupData& GroupData)
 		// Fall back to no data, but still serialize guide data as they an be used for cards simulation
 		// Theoritically, we should have something to detect if we are going to use or not guide (for 
 		// simulation or RBF deformation) on the target platform
-		NoStrandsData.Strands.BulkData.Serialize(Ar);
-		GroupData.Guides.BulkData.Serialize(Ar);
+		NoStrandsData.Strands.BulkData.Serialize(Ar, Owner);
+		GroupData.Guides.BulkData.Serialize(Ar, Owner);
 		NoStrandsData.Strands.InterpolationBulkData.Serialize(Ar);
 	}
 
@@ -1460,13 +1512,11 @@ FArchive& operator<<(FArchive& Ar, FHairGroupData& GroupData)
 
 			if (bIsCooked)
 			{
-				Ar << GroupData.Cards.LODs;
-				Ar << GroupData.Meshes.LODs;
+				InternalSerialize(Ar, Owner, GroupData.Cards.LODs);
+				InternalSerialize(Ar, Owner, GroupData.Meshes.LODs);
 			}
 		}
 	}
-
-	return Ar;
 }
 
 void UGroomAsset::AddAssetUserData(UAssetUserData* InUserData)
@@ -1617,7 +1667,7 @@ void UGroomAsset::SetHairWidth(float Width)
 // differences, etc.) replace the version GUID below with a new one.
 // In case of merge conflicts with DDC versions, you MUST generate a new GUID
 // and set this new GUID as the version.
-#define GROOM_DERIVED_DATA_VERSION TEXT("58A94E5B3E784BBBB5C1FF7D42540EB5")
+#define GROOM_DERIVED_DATA_VERSION TEXT("7188891507314F3D9F79461BA21DF850")
 
 #if WITH_EDITORONLY_DATA
 
@@ -2009,7 +2059,7 @@ bool UGroomAsset::CacheStrandsData(uint32 GroupIndex, FString& OutDerivedDataKey
 
 		FHairGroupData& HairGroupData = HairGroupsData[GroupIndex];
 		FLargeMemoryReader LargeMemReader(DecompressionBuffer, UncompressedSize, ELargeMemoryReaderFlags::Persistent | ELargeMemoryReaderFlags::TakeOwnership);
-		LargeMemReader << HairGroupData;
+		InternalSerialize(LargeMemReader, this, HairGroupData);
 	}
 	else
 	{
@@ -2049,8 +2099,7 @@ bool UGroomAsset::CacheStrandsData(uint32 GroupIndex, FString& OutDerivedDataKey
 		{
 			// Using a LargeMemoryWriter for serialization since the data can be bigger than 2 GB
 			FLargeMemoryWriter LargeMemWriter(0, /*bIsPersistent=*/ true);
-			LargeMemWriter << HairGroupsData[GroupIndex];
-
+			InternalSerialize(LargeMemWriter, this, HairGroupsData[GroupIndex]);
 			int64 UncompressedSize = LargeMemWriter.TotalSize();
 
 			// Then the content of the LargeMemWriter is compressed into a MemoryWriter
@@ -2104,7 +2153,7 @@ bool UGroomAsset::CacheCardsGeometry(uint32 GroupIndex, const FString& StrandsKe
 			UE_LOG(LogHairStrands, Log, TEXT("[Groom/DDC] Cards - Found (Groom:%s Group:%d)."), *GetName(), GroupIndex);
 		}
 		FMemoryReader Ar(DerivedData, true);
-		Ar << HairGroupData.Cards.LODs;
+		InternalSerialize(Ar, this, HairGroupData.Cards.LODs);
 	}
 	else
 	{
@@ -2120,7 +2169,7 @@ bool UGroomAsset::CacheCardsGeometry(uint32 GroupIndex, const FString& StrandsKe
 		}
 
 		FMemoryWriter Ar(DerivedData, true);
-		Ar << HairGroupData.Cards.LODs;
+		InternalSerialize(Ar, this, HairGroupData.Cards.LODs);
 
 		GetDerivedDataCacheRef().Put(*DerivedDataKey, DerivedData, GetPathName());
 	}
@@ -2180,7 +2229,7 @@ bool UGroomAsset::CacheMeshesGeometry(uint32 GroupIndex)
 		}
 
 		FMemoryReader Ar(DerivedData, true);
-		Ar << HairGroupData.Meshes.LODs;
+		InternalSerialize(Ar, this, HairGroupData.Meshes.LODs);
 	}
 	else
 	{
@@ -2196,7 +2245,7 @@ bool UGroomAsset::CacheMeshesGeometry(uint32 GroupIndex)
 		}
 
 		FMemoryWriter Ar(DerivedData, true);
-		Ar << HairGroupData.Meshes.LODs;
+		InternalSerialize(Ar, this, HairGroupData.Meshes.LODs);
 
 		GetDerivedDataCacheRef().Put(*DerivedDataKey, DerivedData, GetPathName());
 	}
diff --git a/Engine/Plugins/Runtime/HairStrands/Source/HairStrandsCore/Private/GroomBuilder.cpp b/Engine/Plugins/Runtime/HairStrands/Source/HairStrandsCore/Private/GroomBuilder.cpp
--- a/Engine/Plugins/Runtime/HairStrands/Source/HairStrandsCore/Private/GroomBuilder.cpp
+++ b/Engine/Plugins/Runtime/HairStrands/Source/HairStrandsCore/Private/GroomBuilder.cpp
@@ -33,7 +33,7 @@ static FAutoConsoleVariableRef CVarHairClusterBuilder_MaxVoxelResolution(TEXT("r
 FString FGroomBuilder::GetVersion()
 {
 	// Important to update the version when groom building changes
-	return TEXT("2b");
+	return TEXT("2e");
 }
 
 namespace FHairStrandsDecimation
@@ -213,6 +213,18 @@ namespace HairStrandsBuilder
 		OutPosition.Z = InVector.Z;
 	}
 
+	template<typename TFormatType>
+	void CopyToBulkData(FByteBulkData& Out, const TArray<typename TFormatType::Type>& Data)
+	{
+		const uint32 DataSizeInByte = Data.Num() * sizeof(TFormatType::BulkType);
+
+		// The buffer is then stored into bulk data
+		Out.Lock(LOCK_READ_WRITE);
+		void* BulkBuffer = Out.Realloc(DataSizeInByte);
+		FMemory::Memcpy(BulkBuffer, Data.GetData(), DataSizeInByte);
+		Out.Unlock();
+	}
+
 	/** Build the packed datas for gpu rendering/simulation */
 	void BuildRenderData(const FHairStrandsDatas& HairStrands, const TArray<uint8>& RandomSeeds, FHairStrandsBulkData& OutBulkData)
 	{
@@ -223,11 +235,11 @@ namespace HairStrandsBuilder
 		if (!(NumCurves > 0 && NumPoints > 0))
 			return;
 
-		TArray<FHairStrandsPositionFormat::Type>& OutPackedPositions = OutBulkData.Positions;
-		TArray<FHairStrandsAttribute0Format::Type>& OutPackedAttributes0 = OutBulkData.Attributes0;
-		TArray<FHairStrandsAttribute1Format::Type>& OutPackedAttributes1 = OutBulkData.Attributes1;
-		TArray<FHairStrandsMaterialFormat::Type>& OutPackedMaterials = OutBulkData.Materials;
-		TArray<FHairStrandsRootIndexFormat::Type>& OutCurveOffsets = OutBulkData.CurveOffsets;
+		TArray<FHairStrandsPositionFormat::Type> OutPackedPositions;
+		TArray<FHairStrandsAttribute0Format::Type> OutPackedAttributes0;
+		TArray<FHairStrandsAttribute1Format::Type> OutPackedAttributes1;
+		TArray<FHairStrandsMaterialFormat::Type> OutPackedMaterials;
+		TArray<FHairStrandsRootIndexFormat::Type> OutCurveOffsets;
 
 		OutPackedPositions.SetNum(NumPoints * FHairStrandsPositionFormat::ComponentCount);
 		OutPackedAttributes0.SetNum(NumPoints * FHairStrandsAttribute0Format::ComponentCount);
@@ -281,13 +293,17 @@ namespace HairStrandsBuilder
 				PackedPosition.PackedRadiusAndType = PackedRadiusAndType.Packed;
 				PackedPosition.UCoord = uint8(FMath::Clamp(CoordU * 255.f, 0.f, 255.f));
 
-				const FVector2D RootUV = Curves.CurvesRootUV[CurveIndex];
+				FVector2D RootUV = Curves.CurvesRootUV[CurveIndex]; 
 				FHairStrandsAttribute0Format::Type& PackedAttributes0 = OutPackedAttributes0[PointIndex + IndexOffset];
 				PackedAttributes0.NormalizedLength = uint8(FMath::Clamp(NormalizedLength * 255.f, 0.f, 255.f));
 				PackedAttributes0.Seed = CurveSeed;
 
 				// Root UV support UDIM texture coordinate but limit the spans of the UDIM to be in 256x256 instead of 9999x9999.
 				// The internal UV coords are also limited to 8bits, which means if sampling need to be super precise, this is no enough.
+				// Specal case for UV == 1.0f, as we don't need UDIM data in this case, so force the value to be in [0..1)
+				const float SmallEpsilon = 1e-05f;
+				RootUV.X = RootUV.X == 1.0f ? RootUV.X - SmallEpsilon : RootUV.X;
+				RootUV.Y = RootUV.Y == 1.0f ? RootUV.Y - SmallEpsilon : RootUV.Y;
 				const FVector2D TextureRootUV(FMath::Fractional(RootUV.X), FMath::Fractional(RootUV.Y));
 				const FVector2D TextureIndexUV = RootUV - TextureRootUV;
 				PackedAttributes0.RootU  = uint8(FMath::Clamp(TextureRootUV.X*255.f, 0.f, 255.f));
@@ -332,10 +348,20 @@ namespace HairStrandsBuilder
 		OutBulkData.MaxRadius = HairStrands.StrandsCurves.MaxRadius;
 		OutBulkData.Flags = 0;
 
-		if (bHasMaterialData)	{ OutBulkData.Flags |= FHairStrandsBulkData::DataFlags_HasMaterialData; }
-		if (bHasUDIMData)		{ OutBulkData.Flags |= FHairStrandsBulkData::DataFlags_HasUDIMData; }
-		if (!bHasMaterialData)	{ OutBulkData.Attributes1.Empty(); }
-		if (!bHasUDIMData)		{ OutBulkData.Materials.Empty(); }
+		if (bHasMaterialData)	{ OutBulkData.Flags |= FHairStrandsBulkData::DataFlags_HasMaterialData; }	else { OutPackedMaterials.Empty(); }
+		if (bHasUDIMData)		{ OutBulkData.Flags |= FHairStrandsBulkData::DataFlags_HasUDIMData; }		else { OutPackedAttributes1.Empty(); }
+
+		CopyToBulkData<FHairStrandsPositionFormat>(OutBulkData.Positions, OutPackedPositions);
+		CopyToBulkData<FHairStrandsAttribute0Format>(OutBulkData.Attributes0, OutPackedAttributes0);
+		if (bHasUDIMData)
+		{
+			CopyToBulkData<FHairStrandsAttribute1Format>(OutBulkData.Attributes1, OutPackedAttributes1);
+		}
+		if (bHasMaterialData)
+		{
+			CopyToBulkData<FHairStrandsMaterialFormat>(OutBulkData.Materials, OutPackedMaterials);
+		}
+		CopyToBulkData<FHairStrandsRootIndexFormat>(OutBulkData.CurveOffsets, OutCurveOffsets);
 	}
 
 	void BuildRenderData(const FHairStrandsDatas& HairStrands, FHairStrandsBulkData& OutBulkData)
diff --git a/Engine/Plugins/Runtime/HairStrands/Source/HairStrandsCore/Private/GroomResources.cpp b/Engine/Plugins/Runtime/HairStrands/Source/HairStrandsCore/Private/GroomResources.cpp
--- a/Engine/Plugins/Runtime/HairStrands/Source/HairStrandsCore/Private/GroomResources.cpp
+++ b/Engine/Plugins/Runtime/HairStrands/Source/HairStrandsCore/Private/GroomResources.cpp
@@ -116,6 +116,35 @@ static FRDGBufferRef InternalCreateVertexBuffer(
 	return Buffer;
 }
 
+template<typename FormatType>
+void InternalCreateVertexBufferRDG(FRDGBuilder& GraphBuilder, FRDGBufferUploader& BufferUploader, typename const FormatType::Type* InData, uint32 InDataCount, FRDGExternalBuffer& Out, const TCHAR* DebugName, EHairResourceUsageType UsageType, ERDGInitialDataFlags InitialDataFlags)
+{
+	FRDGBufferRef Buffer = nullptr;
+
+	const uint32 DataSizeInBytes = FormatType::SizeInByte * InDataCount;
+	if (DataSizeInBytes == 0)
+	{
+		Out.Buffer = nullptr;
+		return;
+	}
+
+	FRDGBufferDesc Desc = FRDGBufferDesc::CreateBufferDesc(FormatType::SizeInByte, InDataCount);
+	if (UsageType != EHairResourceUsageType::Dynamic)
+	{
+		Desc.Usage = Desc.Usage & (~BUF_UnorderedAccess);
+	}
+	Buffer = InternalCreateVertexBuffer(
+		GraphBuilder,
+		BufferUploader,
+		DebugName,
+		Desc,
+		InData,
+		DataSizeInBytes,
+		InitialDataFlags);
+
+	ConvertToExternalBufferWithViews(GraphBuilder, Buffer, Out, FormatType::Format);
+}
+
 template<typename FormatType>
 void InternalCreateVertexBufferRDG(FRDGBuilder& GraphBuilder, FRDGBufferUploader& BufferUploader, const TArray<typename FormatType::Type>& InData, FRDGExternalBuffer& Out, const TCHAR* DebugName, EHairResourceUsageType UsageType, ERDGInitialDataFlags InitialDataFlags= ERDGInitialDataFlags::NoCopy)
 {
@@ -576,24 +605,46 @@ void FHairStrandsRestResource::InternalAllocate(FRDGBuilder& GraphBuilder)
 {
 	FRDGBufferUploader BufferUploader;
 
-	InternalCreateVertexBufferRDG<FHairStrandsPositionFormat>(GraphBuilder, BufferUploader, BulkData.Positions, PositionBuffer, HAIRSTRANDS_RESOUCE_NAME(CurveType, Hair.StrandsRest_PositionBuffer), EHairResourceUsageType::Static);
-	InternalCreateVertexBufferRDG<FHairStrandsAttribute0Format>(GraphBuilder, BufferUploader, BulkData.Attributes0, Attribute0Buffer, HAIRSTRANDS_RESOUCE_NAME(CurveType, Hair.StrandsRest_Attribute0Buffer), EHairResourceUsageType::Static);
+	const uint32 PointCount = BulkData.PointCount;
+
+	// 1. Lock data, which force the loading data from files (on non-editor build/cooked data). These data are then uploaded to the GPU
+	// 2. A local copy is done by the buffer uploader. This copy is discarded once the uploading is done.
+	const FHairStrandsPositionFormat::Type* BulkDataPositions = (const FHairStrandsPositionFormat::Type*)BulkData.Positions.LockReadOnly();
+	InternalCreateVertexBufferRDG<FHairStrandsPositionFormat>(GraphBuilder, BufferUploader, BulkDataPositions, PointCount, PositionBuffer, HAIRSTRANDS_RESOUCE_NAME(CurveType, Hair.StrandsRest_PositionBuffer), EHairResourceUsageType::Static, ERDGInitialDataFlags::None);
 
+	const FHairStrandsAttribute0Format::Type* BulkDataAttributes0 = (const FHairStrandsAttribute0Format::Type*)BulkData.Attributes0.LockReadOnly();
+	InternalCreateVertexBufferRDG<FHairStrandsAttribute0Format>(GraphBuilder, BufferUploader, BulkDataAttributes0, PointCount, Attribute0Buffer, HAIRSTRANDS_RESOUCE_NAME(CurveType, Hair.StrandsRest_Attribute0Buffer), EHairResourceUsageType::Static, ERDGInitialDataFlags::None);
+
+	const FHairStrandsAttribute1Format::Type* BulkDataAttributes1 = nullptr;
 	if (!!(BulkData.Flags & FHairStrandsBulkData::DataFlags_HasUDIMData))
 	{
-		InternalCreateVertexBufferRDG<FHairStrandsAttribute1Format>(GraphBuilder, BufferUploader, BulkData.Attributes1, Attribute1Buffer, HAIRSTRANDS_RESOUCE_NAME(CurveType, Hair.StrandsRest_Attribute1Buffer), EHairResourceUsageType::Static);
+		BulkDataAttributes1 = (const FHairStrandsAttribute1Format::Type*)BulkData.Attributes1.LockReadOnly();
+		InternalCreateVertexBufferRDG<FHairStrandsAttribute1Format>(GraphBuilder, BufferUploader, BulkDataAttributes1, PointCount, Attribute1Buffer, HAIRSTRANDS_RESOUCE_NAME(CurveType, Hair.StrandsRest_Attribute1Buffer), EHairResourceUsageType::Static, ERDGInitialDataFlags::None);
 	}
 
+	const FHairStrandsMaterialFormat::Type* BulkDataMaterials = nullptr;
 	if (!!(BulkData.Flags & FHairStrandsBulkData::DataFlags_HasMaterialData))
 	{
-		InternalCreateVertexBufferRDG<FHairStrandsMaterialFormat>(GraphBuilder, BufferUploader, BulkData.Materials, MaterialBuffer, HAIRSTRANDS_RESOUCE_NAME(CurveType, Hair.StrandsRest_MaterialBuffer), EHairResourceUsageType::Static);
+		BulkDataMaterials = (const FHairStrandsMaterialFormat::Type*)BulkData.Materials.LockReadOnly();
+		InternalCreateVertexBufferRDG<FHairStrandsMaterialFormat>(GraphBuilder, BufferUploader, BulkDataMaterials, PointCount, MaterialBuffer, HAIRSTRANDS_RESOUCE_NAME(CurveType, Hair.StrandsRest_MaterialBuffer), EHairResourceUsageType::Static, ERDGInitialDataFlags::None);
 	}
 
 	TArray<FVector4> RestOffset;
 	RestOffset.Add(BulkData.GetPositionOffset());
 	InternalCreateVertexBufferRDG<FHairStrandsPositionOffsetFormat>(GraphBuilder, BufferUploader, RestOffset, PositionOffsetBuffer, HAIRSTRANDS_RESOUCE_NAME(CurveType, Hair.StrandsRest_PositionOffsetBuffer), EHairResourceUsageType::Static, ERDGInitialDataFlags::None);
-
 	BufferUploader.Submit(GraphBuilder);
+
+	// Unlock data once the data has been copied/cached by the uploader
+	BulkData.Positions.Unlock();
+	BulkData.Attributes0.Unlock();
+	if (BulkDataAttributes1)
+	{
+		BulkData.Attributes1.Unlock();
+	}
+	if (BulkDataMaterials)
+	{
+		BulkData.Materials.Unlock();
+	}
 }
 
 void AddHairTangentPass(
@@ -647,10 +698,16 @@ void FHairStrandsDeformedResource::InternalAllocate(FRDGBuilder& GraphBuilder)
 {
 	FRDGBufferUploader BufferUploader;
 
+	const uint32 PointCount = BulkData.PointCount;
+
+	const FHairStrandsPositionFormat::Type* BulkDataPositions = nullptr;
 	if (bInitializedData)
 	{
-		InternalCreateVertexBufferRDG<FHairStrandsPositionFormat>(GraphBuilder, BufferUploader, BulkData.Positions, DeformedPositionBuffer[0], HAIRSTRANDS_RESOUCE_NAME(CurveType, Hair.StrandsDeformed_DeformedPositionBuffer0), EHairResourceUsageType::Dynamic);
-		InternalCreateVertexBufferRDG<FHairStrandsPositionFormat>(GraphBuilder, BufferUploader, BulkData.Positions, DeformedPositionBuffer[1], HAIRSTRANDS_RESOUCE_NAME(CurveType, Hair.StrandsDeformed_DeformedPositionBuffer1), EHairResourceUsageType::Dynamic);
+		// 1. Lock data, which force the loading data from files (on non-editor build/cooked data). These data are then uploaded to the GPU
+		// 2. A local copy is done by the buffer uploader. This copy is discarded once the uploading is done.
+		BulkDataPositions = (const FHairStrandsPositionFormat::Type*)BulkData.Positions.LockReadOnly();
+		InternalCreateVertexBufferRDG<FHairStrandsPositionFormat>(GraphBuilder, BufferUploader, BulkDataPositions, PointCount, DeformedPositionBuffer[0], HAIRSTRANDS_RESOUCE_NAME(CurveType, Hair.StrandsDeformed_DeformedPositionBuffer0), EHairResourceUsageType::Dynamic, ERDGInitialDataFlags::None);
+		InternalCreateVertexBufferRDG<FHairStrandsPositionFormat>(GraphBuilder, BufferUploader, BulkDataPositions, PointCount, DeformedPositionBuffer[1], HAIRSTRANDS_RESOUCE_NAME(CurveType, Hair.StrandsDeformed_DeformedPositionBuffer1), EHairResourceUsageType::Dynamic, ERDGInitialDataFlags::None);
 	}
 	else
 	{
@@ -665,6 +722,11 @@ void FHairStrandsDeformedResource::InternalAllocate(FRDGBuilder& GraphBuilder)
 	InternalCreateVertexBufferRDG<FHairStrandsPositionOffsetFormat>(GraphBuilder, BufferUploader, DefaultOffsets, DeformedOffsetBuffer[1], HAIRSTRANDS_RESOUCE_NAME(CurveType, Hair.StrandsDeformed_DeformedOffsetBuffer1), EHairResourceUsageType::Dynamic, ERDGInitialDataFlags::None);
 
 	BufferUploader.Submit(GraphBuilder);
+
+	if (BulkDataPositions)
+	{
+		BulkData.Positions.Unlock();
+	}
 }
 
 void FHairStrandsDeformedResource::InternalRelease()
diff --git a/Engine/Plugins/Runtime/HairStrands/Source/HairStrandsCore/Private/HairStrandsDatas.cpp b/Engine/Plugins/Runtime/HairStrands/Source/HairStrandsCore/Private/HairStrandsDatas.cpp
--- a/Engine/Plugins/Runtime/HairStrands/Source/HairStrandsCore/Private/HairStrandsDatas.cpp
+++ b/Engine/Plugins/Runtime/HairStrands/Source/HairStrandsCore/Private/HairStrandsDatas.cpp
@@ -191,7 +191,7 @@ void FHairStrandsInterpolationBulkData::Serialize(FArchive& Ar)
 	SimRootPointIndex.BulkSerialize(Ar);
 }
 
-void FHairStrandsBulkData::Serialize(FArchive& Ar)
+void FHairStrandsBulkData::Serialize(FArchive& Ar, UObject* Owner)
 {
 	static_assert(sizeof(FHairStrandsPositionFormat::BulkType) == sizeof(FHairStrandsPositionFormat::Type));
 	static_assert(sizeof(FHairStrandsAttribute0Format::BulkType) == sizeof(FHairStrandsAttribute0Format::Type));
@@ -209,17 +209,29 @@ void FHairStrandsBulkData::Serialize(FArchive& Ar)
 	Ar << BoundingBox;
 	Ar << Flags;
 
-	Positions.BulkSerialize(Ar);
-	Attributes0.BulkSerialize(Ar);
+	// Forced not inline means the bulk data won't automatically be loaded when we deserialize
+	// but only when we explicitly take action to load it
+	const uint32 BulkFlags = BULKDATA_Force_NOT_InlinePayload | BULKDATA_SerializeCompressed;
+	Positions.SetBulkDataFlags(BulkFlags);
+	Attributes0.SetBulkDataFlags(BulkFlags);
+	Attributes1.SetBulkDataFlags(BulkFlags);
+	Materials.SetBulkDataFlags(BulkFlags);
+	CurveOffsets.SetBulkDataFlags(BulkFlags);
+
+	const int32 ChunkIndex = 0;
+	bool bAttemptFileMapping = false;
+
+	Positions.Serialize(Ar, Owner, ChunkIndex, bAttemptFileMapping);
+	Attributes0.Serialize(Ar, Owner, ChunkIndex, bAttemptFileMapping);
 	if (!!(Flags & DataFlags_HasUDIMData))
 	{
-		Attributes1.BulkSerialize(Ar);
+		Attributes1.Serialize(Ar, Owner, ChunkIndex, bAttemptFileMapping);
 	}
 	if (!!(Flags & DataFlags_HasMaterialData))
 	{
-		Materials.BulkSerialize(Ar);
+		Materials.Serialize(Ar, Owner, ChunkIndex, bAttemptFileMapping);
 	}
-	CurveOffsets.BulkSerialize(Ar);
+	CurveOffsets.Serialize(Ar, Owner, ChunkIndex, bAttemptFileMapping);
 }
 
 void FHairStrandsDatas::Reset()
diff --git a/Engine/Plugins/Runtime/HairStrands/Source/HairStrandsCore/Private/Niagara/NiagaraDataInterfaceHairStrands.cpp b/Engine/Plugins/Runtime/HairStrands/Source/HairStrandsCore/Private/Niagara/NiagaraDataInterfaceHairStrands.cpp
--- a/Engine/Plugins/Runtime/HairStrands/Source/HairStrandsCore/Private/Niagara/NiagaraDataInterfaceHairStrands.cpp
+++ b/Engine/Plugins/Runtime/HairStrands/Source/HairStrandsCore/Private/Niagara/NiagaraDataInterfaceHairStrands.cpp
@@ -339,11 +339,12 @@ void FNDIHairStrandsBuffer::InitRHI()
 			const uint32 OffsetCount = SourceDatas->GetNumCurves() + 1;
 			const uint32 OffsetBytes = sizeof(uint32)*OffsetCount;
 
+			const FHairStrandsRootIndexFormat::Type* SourceData = (const FHairStrandsRootIndexFormat::Type*)SourceDatas->CurveOffsets.LockReadOnly();
 			CurvesOffsetsBuffer.Initialize(TEXT("CurvesOffsetsBuffer"), sizeof(uint32), OffsetCount, EPixelFormat::PF_R32_UINT, BUF_Static);
 			void* OffsetBufferData = RHILockBuffer(CurvesOffsetsBuffer.Buffer, 0, OffsetBytes, RLM_WriteOnly);
-
-			FMemory::Memcpy(OffsetBufferData, SourceDatas->CurveOffsets.GetData(), OffsetBytes);
+			FMemory::Memcpy(OffsetBufferData, SourceData, OffsetBytes);
 			RHIUnlockBuffer(CurvesOffsetsBuffer.Buffer);
+			SourceDatas->CurveOffsets.Unlock();
 		}
 		{
 			static const TArray<uint32> ZeroData = { UINT_MAX,UINT_MAX,UINT_MAX,0,0,0,
diff --git a/Engine/Plugins/Runtime/HairStrands/Source/HairStrandsCore/Public/GroomAsset.h b/Engine/Plugins/Runtime/HairStrands/Source/HairStrandsCore/Public/GroomAsset.h
--- a/Engine/Plugins/Runtime/HairStrands/Source/HairStrandsCore/Public/GroomAsset.h
+++ b/Engine/Plugins/Runtime/HairStrands/Source/HairStrandsCore/Public/GroomAsset.h
@@ -84,11 +84,11 @@ struct HAIRSTRANDSCORE_API FHairGroupData
 		uint32 GetDataSize() const
 		{
 			uint32 Total = 0;
-			Total += BulkData.Positions.GetAllocatedSize();
-			Total += BulkData.Attributes0.GetAllocatedSize();
-			Total += BulkData.Attributes1.GetAllocatedSize();
-			Total += BulkData.Materials.GetAllocatedSize();
-			Total += BulkData.CurveOffsets.GetAllocatedSize();
+			Total += BulkData.Positions.IsBulkDataLoaded() 		? BulkData.Positions.GetBulkDataSize()   : 0;
+			Total += BulkData.Attributes0.IsBulkDataLoaded() 	? BulkData.Attributes0.GetBulkDataSize() : 0;
+			Total += BulkData.Attributes1.IsBulkDataLoaded() 	? BulkData.Attributes1.GetBulkDataSize() : 0;
+			Total += BulkData.Materials.IsBulkDataLoaded() 		? BulkData.Materials.GetBulkDataSize()   : 0;
+			Total += BulkData.CurveOffsets.IsBulkDataLoaded() 	? BulkData.CurveOffsets.GetBulkDataSize(): 0;
 			return Total;
 		}
 
diff --git a/Engine/Plugins/Runtime/HairStrands/Source/HairStrandsCore/Public/HairStrandsDatas.h b/Engine/Plugins/Runtime/HairStrands/Source/HairStrandsCore/Public/HairStrandsDatas.h
--- a/Engine/Plugins/Runtime/HairStrands/Source/HairStrandsCore/Public/HairStrandsDatas.h
+++ b/Engine/Plugins/Runtime/HairStrands/Source/HairStrandsCore/Public/HairStrandsDatas.h
@@ -427,7 +427,7 @@ struct HAIRSTRANDSCORE_API FHairStrandsBulkData
 		DataFlags_HasMaterialData = 2
 	};
 
-	void Serialize(FArchive& Ar);
+	void Serialize(FArchive& Ar, UObject* Owner);
 
 	bool IsValid() const { return CurveCount > 0 && PointCount > 0;	}
 	
@@ -445,11 +445,11 @@ struct HAIRSTRANDSCORE_API FHairStrandsBulkData
 	FBox BoundingBox;
 	uint32 Flags = 0;
 
-	TArray<FHairStrandsPositionFormat::Type>	Positions;		// Size = PointCount
-	TArray<FHairStrandsAttribute0Format::Type>	Attributes0;	// Size = PointCount
-	TArray<FHairStrandsAttribute1Format::Type>	Attributes1;	// Size = PointCount
-	TArray<FHairStrandsMaterialFormat::Type>	Materials;		// Size = PointCount
-	TArray<FHairStrandsRootIndexFormat::Type>	CurveOffsets;	// Size = CurveCount+1 - Store the root point index for the curve
+	FByteBulkData Positions;	// Size = PointCount
+	FByteBulkData Attributes0;	// Size = PointCount
+	FByteBulkData Attributes1;	// Size = PointCount
+	FByteBulkData Materials;	// Size = PointCount
+	FByteBulkData CurveOffsets;	// Size = CurveCount+1 - Store the root point index for the curve
 };
 
 /** Hair strands debug data */
