Fix CIS issue.

diff --git a/Engine/Plugins/Runtime/HairStrands/Source/HairStrandsCore/Private/GroomResources.cpp b/Engine/Plugins/Runtime/HairStrands/Source/HairStrandsCore/Private/GroomResources.cpp
--- a/Engine/Plugins/Runtime/HairStrands/Source/HairStrandsCore/Private/GroomResources.cpp
+++ b/Engine/Plugins/Runtime/HairStrands/Source/HairStrandsCore/Private/GroomResources.cpp
@@ -117,7 +117,7 @@ static FRDGBufferRef InternalCreateVertexBuffer(
 }
 
 template<typename FormatType>
-void InternalCreateVertexBufferRDG(FRDGBuilder& GraphBuilder, FRDGBufferUploader& BufferUploader, typename const FormatType::Type* InData, uint32 InDataCount, FRDGExternalBuffer& Out, const TCHAR* DebugName, EHairResourceUsageType UsageType, ERDGInitialDataFlags InitialDataFlags)
+void InternalCreateVertexBufferRDG(FRDGBuilder& GraphBuilder, FRDGBufferUploader& BufferUploader, const typename FormatType::Type* InData, uint32 InDataCount, FRDGExternalBuffer& Out, const TCHAR* DebugName, EHairResourceUsageType UsageType, ERDGInitialDataFlags InitialDataFlags)
 {
 	FRDGBufferRef Buffer = nullptr;
 
