Fix another CIS issue.

diff --git a/Engine/Plugins/Runtime/HairStrands/Source/HairStrandsCore/Private/GroomBuilder.cpp b/Engine/Plugins/Runtime/HairStrands/Source/HairStrandsCore/Private/GroomBuilder.cpp
--- a/Engine/Plugins/Runtime/HairStrands/Source/HairStrandsCore/Private/GroomBuilder.cpp
+++ b/Engine/Plugins/Runtime/HairStrands/Source/HairStrandsCore/Private/GroomBuilder.cpp
@@ -216,7 +216,7 @@ namespace HairStrandsBuilder
 	template<typename TFormatType>
 	void CopyToBulkData(FByteBulkData& Out, const TArray<typename TFormatType::Type>& Data)
 	{
-		const uint32 DataSizeInByte = Data.Num() * sizeof(TFormatType::BulkType);
+		const uint32 DataSizeInByte = Data.Num() * sizeof(typename TFormatType::BulkType);
 
 		// The buffer is then stored into bulk data
 		Out.Lock(LOCK_READ_WRITE);
