Fix CIS issue.

diff --git a/Engine/Plugins/Runtime/HairStrands/Source/HairStrandsCore/Public/HairStrandsDatas.h b/Engine/Plugins/Runtime/HairStrands/Source/HairStrandsCore/Public/HairStrandsDatas.h
--- a/Engine/Plugins/Runtime/HairStrands/Source/HairStrandsCore/Public/HairStrandsDatas.h
+++ b/Engine/Plugins/Runtime/HairStrands/Source/HairStrandsCore/Public/HairStrandsDatas.h
@@ -7,6 +7,7 @@
 #include "RHIDefinitions.h"
 #include "PackedNormal.h"
 #include "RenderGraphResources.h"
+#include "Serialization/BulkData.h"
 
 DECLARE_LOG_CATEGORY_EXTERN(LogHairStrands, Log, All);
 
