Change binding asset LOD resources to be lazily created.

This allows to save GPU memory by sub-allocating only the used LOD, rather than all the LOD upfront.

diff --git a/Engine/Plugins/Runtime/HairStrands/Source/HairStrandsCore/Private/GroomManager.cpp b/Engine/Plugins/Runtime/HairStrands/Source/HairStrandsCore/Private/GroomManager.cpp
--- a/Engine/Plugins/Runtime/HairStrands/Source/HairStrandsCore/Private/GroomManager.cpp
+++ b/Engine/Plugins/Runtime/HairStrands/Source/HairStrandsCore/Private/GroomManager.cpp
@@ -68,6 +68,53 @@ static bool IsInstanceFrustumCullingEnable()
 
 bool NeedsUpdateCardsMeshTriangles();
 
+// Returns the cached geometry of the underlying geometry on which a hair instance is attached to
+static FCachedGeometry GetCacheGeometryForHair(
+	FRDGBuilder& GraphBuilder, 
+	FHairGroupInstance* Instance, 
+	const FGPUSkinCache* SkinCache, 
+	FGlobalShaderMap* ShaderMap)
+{
+	FCachedGeometry Out;
+	if (Instance->Debug.GroomBindingType == EGroomBindingMeshType::SkeletalMesh)
+	{
+		if (USkeletalMeshComponent* SkeletalMeshComponent = Cast<USkeletalMeshComponent>(Instance->Debug.MeshComponent))
+		{
+			if (SkinCache)
+			{
+				Out = SkinCache->GetCachedGeometry(SkeletalMeshComponent->ComponentId.PrimIDValue);
+			}
+
+			if (IsHairStrandsSkinCacheEnable() && Out.Sections.Num() == 0)
+			{
+				//#hair_todo: Need to have a (frame) cache to insure that we don't recompute the same projection several time
+				// Actual populate the cache with only the needed part basd on the groom projection data. At the moment it recompute everything ...
+				BuildCacheGeometry(GraphBuilder, ShaderMap, SkeletalMeshComponent, Out);
+			}
+		}
+	}
+	else if (UGeometryCacheComponent* GeometryCacheComponent = Cast<UGeometryCacheComponent>(Instance->Debug.MeshComponent))
+	{
+		BuildCacheGeometry(GraphBuilder, ShaderMap, GeometryCacheComponent, Out);
+	}
+	return Out;
+}
+
+// Return the LOD index of the cached geometry on which the hair are bound to. 
+// Return -1 if the hair are not bound of if the underlying geometry is invalid
+static int32 GetCacheGeometryLODIndex(const FCachedGeometry& In)
+{
+	int32 MeshLODIndex = -1;
+	for (const FCachedGeometry::Section& Section : In.Sections)
+	{
+		// Ensure all mesh's sections have the same LOD index
+		if (MeshLODIndex < 0) MeshLODIndex = Section.LODIndex;
+		check(MeshLODIndex == Section.LODIndex);
+	}
+
+	return MeshLODIndex;
+}
+
 static void RunInternalHairStrandsInterpolation(
 	FRDGBuilder& GraphBuilder,
 	const FSceneView* View,
@@ -91,30 +138,8 @@ static void RunInternalHairStrandsInterpolation(
 	
 		check(Instance->HairGroupPublicData);
 
-		FCachedGeometry CachedGeometry;
-		if (Instance->Debug.GroomBindingType == EGroomBindingMeshType::SkeletalMesh)
-		{
-			if (USkeletalMeshComponent* SkeletalMeshComponent = Cast<USkeletalMeshComponent>(Instance->Debug.MeshComponent))
-			{
-				if (SkinCache)
-				{
-					CachedGeometry = SkinCache->GetCachedGeometry(SkeletalMeshComponent->ComponentId.PrimIDValue);
-				}
-
-				if (IsHairStrandsSkinCacheEnable() && CachedGeometry.Sections.Num() == 0)
-				{
-					//#hair_todo: Need to have a (frame) cache to insure that we don't recompute the same projection several time
-					// Actual populate the cache with only the needed part basd on the groom projection data. At the moment it recompute everything ...
-					BuildCacheGeometry(GraphBuilder, ShaderMap, SkeletalMeshComponent, CachedGeometry);
-				}
-			}
-		}
-		else if (UGeometryCacheComponent* GeometryCacheComponent = Cast<UGeometryCacheComponent>(Instance->Debug.MeshComponent))
-		{
-			BuildCacheGeometry(GraphBuilder, ShaderMap, GeometryCacheComponent, CachedGeometry);
-		}
-
 		FHairStrandsProjectionMeshData::LOD MeshDataLOD;
+		const FCachedGeometry CachedGeometry = GetCacheGeometryForHair(GraphBuilder, Instance, SkinCache, ShaderMap);
 		for (const FCachedGeometry::Section& Section : CachedGeometry.Sections)
 		{
 			// Ensure all mesh's sections have the same LOD index
@@ -369,7 +394,6 @@ static void RunHairStrandsInterpolation_Strands(
 		ClusterData);
 }
 
-
 static void RunHairStrandsGatherCluster(
 	const FHairStrandsInstances& Instances,
 	FHairStrandClusterData* ClusterData)
@@ -395,7 +419,6 @@ static void RunHairStrandsGatherCluster(
 	}
 }
 
-
 // Return the LOD which should be used for a given screen size and LOD bias value
 // This function is mirrored in HairStrandsClusterCommon.ush
 static float GetHairInstanceLODIndex(const TArray<float>& InLODScreenSizes, float InScreenSize, float InLODBias)
@@ -528,7 +551,12 @@ static void RunHairBufferSwap(const FHairStrandsInstances& Instances, const TArr
 	}
 }
 
-static void RunHairLODSelection(FRDGBuilder& GraphBuilder, const FHairStrandsInstances& Instances, const TArray<const FSceneView*> Views)
+static void RunHairLODSelection(
+	FRDGBuilder& GraphBuilder, 
+	const FHairStrandsInstances& Instances, 
+	const TArray<const FSceneView*>& Views, 
+	const FGPUSkinCache* SkinCache, 
+	FGlobalShaderMap* ShaderMap)
 {
 	EShaderPlatform ShaderPlatform = EShaderPlatform::SP_NumPlatforms;
 	if (Views.Num() > 0)
@@ -540,9 +568,10 @@ static void RunHairLODSelection(FRDGBuilder& GraphBuilder, const FHairStrandsIns
 	{
 		FHairGroupInstance* Instance = static_cast<FHairGroupInstance*>(AbstractInstance);
 
-		int32 MeshLODIndex = -1;
 		check(Instance);
 		check(Instance->HairGroupPublicData);
+		const FCachedGeometry CachedGeometry = GetCacheGeometryForHair(GraphBuilder, Instance, SkinCache, ShaderMap);
+		const int32 MeshLODIndex = GetCacheGeometryLODIndex(CachedGeometry);
 
 		//if (Instance->ProxyLocalToWorld)
 		//{
@@ -635,19 +664,19 @@ static void RunHairLODSelection(FRDGBuilder& GraphBuilder, const FHairStrandsIns
 		// Note: Allocation will only be done if the resources it not yet initialized.
 		if (Instance->Guides.Data)
 		{
-			if (Instance->Guides.RestRootResource)		Instance->Guides.RestRootResource->Allocate(GraphBuilder);
-			if (Instance->Guides.RestResource)			Instance->Guides.RestResource->Allocate(GraphBuilder);
-			if (Instance->Guides.DeformedRootResource)	Instance->Guides.DeformedRootResource->Allocate(GraphBuilder);
-			if (Instance->Guides.DeformedResource)		Instance->Guides.DeformedResource->Allocate(GraphBuilder);
+			if (Instance->Guides.RestRootResource)			{ Instance->Guides.RestRootResource->Allocate(GraphBuilder); Instance->Guides.RestRootResource->AllocateLOD(GraphBuilder, MeshLODIndex); }
+			if (Instance->Guides.RestResource)				{ Instance->Guides.RestResource->Allocate(GraphBuilder); }
+			if (Instance->Guides.DeformedRootResource)		{ Instance->Guides.DeformedRootResource->Allocate(GraphBuilder); Instance->Guides.DeformedRootResource->AllocateLOD(GraphBuilder, MeshLODIndex); }
+			if (Instance->Guides.DeformedResource)			{ Instance->Guides.DeformedResource->Allocate(GraphBuilder); }
 		}
 
 		if (GeometryType == EHairGeometryType::Meshes)
 		{
 			FHairGroupInstance::FMeshes::FLOD& InstanceLOD = Instance->Meshes.LODs[IntLODIndex];
 
-			if (InstanceLOD.DeformedResource)	InstanceLOD.DeformedResource->Allocate(GraphBuilder);
+			if (InstanceLOD.DeformedResource)				{ InstanceLOD.DeformedResource->Allocate(GraphBuilder); }
 			#if RHI_RAYTRACING
-			if (InstanceLOD.RaytracingResource)	InstanceLOD.RaytracingResource->Allocate(GraphBuilder);
+			if (InstanceLOD.RaytracingResource)				{ InstanceLOD.RaytracingResource->Allocate(GraphBuilder);}
 			#endif
 
 			InstanceLOD.InitVertexFactory();
@@ -656,14 +685,14 @@ static void RunHairLODSelection(FRDGBuilder& GraphBuilder, const FHairStrandsIns
 		{
 			FHairGroupInstance::FCards::FLOD& InstanceLOD = Instance->Cards.LODs[IntLODIndex];
 
-			if (InstanceLOD.DeformedResource)				InstanceLOD.DeformedResource->Allocate(GraphBuilder);
-			if (InstanceLOD.Guides.RestRootResource)		InstanceLOD.Guides.RestRootResource->Allocate(GraphBuilder);
-			if (InstanceLOD.Guides.RestResource)			InstanceLOD.Guides.RestResource->Allocate(GraphBuilder);
-			if (InstanceLOD.Guides.DeformedRootResource)	InstanceLOD.Guides.DeformedRootResource->Allocate(GraphBuilder);
-			if (InstanceLOD.Guides.DeformedResource)		InstanceLOD.Guides.DeformedResource->Allocate(GraphBuilder);
-			if (InstanceLOD.Guides.InterpolationResource)	InstanceLOD.Guides.InterpolationResource->Allocate(GraphBuilder);
+			if (InstanceLOD.DeformedResource)				{ InstanceLOD.DeformedResource->Allocate(GraphBuilder); }
+			if (InstanceLOD.Guides.RestRootResource)		{ InstanceLOD.Guides.RestRootResource->Allocate(GraphBuilder); InstanceLOD.Guides.RestRootResource->AllocateLOD(GraphBuilder, MeshLODIndex); }
+			if (InstanceLOD.Guides.RestResource)			{ InstanceLOD.Guides.RestResource->Allocate(GraphBuilder); }
+			if (InstanceLOD.Guides.DeformedRootResource)	{ InstanceLOD.Guides.DeformedRootResource->Allocate(GraphBuilder); InstanceLOD.Guides.DeformedRootResource->AllocateLOD(GraphBuilder, MeshLODIndex); }
+			if (InstanceLOD.Guides.DeformedResource)		{ InstanceLOD.Guides.DeformedResource->Allocate(GraphBuilder); }
+			if (InstanceLOD.Guides.InterpolationResource)	{ InstanceLOD.Guides.InterpolationResource->Allocate(GraphBuilder); }
 			#if RHI_RAYTRACING
-			if (InstanceLOD.RaytracingResource)				InstanceLOD.RaytracingResource->Allocate(GraphBuilder);
+			if (InstanceLOD.RaytracingResource)				{ InstanceLOD.RaytracingResource->Allocate(GraphBuilder);}
 			#endif
 			InstanceLOD.InitVertexFactory();
 		}
@@ -672,15 +701,15 @@ static void RunHairLODSelection(FRDGBuilder& GraphBuilder, const FHairStrandsIns
 			check(Instance->HairGroupPublicData);
 			Instance->HairGroupPublicData->Allocate(GraphBuilder);
 
-			if (Instance->Strands.RestRootResource)			Instance->Strands.RestRootResource->Allocate(GraphBuilder);
-			if (Instance->Strands.RestResource)				Instance->Strands.RestResource->Allocate(GraphBuilder);
-			if (Instance->Strands.ClusterCullingResource)	Instance->Strands.ClusterCullingResource->Allocate(GraphBuilder);
-			if (Instance->Strands.InterpolationResource)	Instance->Strands.InterpolationResource->Allocate(GraphBuilder);
+			if (Instance->Strands.RestRootResource)			{ Instance->Strands.RestRootResource->Allocate(GraphBuilder); Instance->Strands.RestRootResource->AllocateLOD(GraphBuilder, MeshLODIndex); }
+			if (Instance->Strands.RestResource)				{ Instance->Strands.RestResource->Allocate(GraphBuilder); }
+			if (Instance->Strands.ClusterCullingResource)	{ Instance->Strands.ClusterCullingResource->Allocate(GraphBuilder); }
+			if (Instance->Strands.InterpolationResource)	{ Instance->Strands.InterpolationResource->Allocate(GraphBuilder); }
 
-			if (Instance->Strands.DeformedRootResource)		Instance->Strands.DeformedRootResource->Allocate(GraphBuilder);
-			if (Instance->Strands.DeformedResource)			Instance->Strands.DeformedResource->Allocate(GraphBuilder);
+			if (Instance->Strands.DeformedRootResource)		{ Instance->Strands.DeformedRootResource->Allocate(GraphBuilder); Instance->Strands.DeformedRootResource->AllocateLOD(GraphBuilder, MeshLODIndex); }
+			if (Instance->Strands.DeformedResource)			{ Instance->Strands.DeformedResource->Allocate(GraphBuilder); }
 			#if RHI_RAYTRACING
-			if (Instance->Strands.RenRaytracingResource)	Instance->Strands.RenRaytracingResource->Allocate(GraphBuilder);
+			if (Instance->Strands.RenRaytracingResource)	{ Instance->Strands.RenRaytracingResource->Allocate(GraphBuilder); }
 			#endif
 			Instance->Strands.VertexFactory->InitResources();
 		}
@@ -801,7 +830,9 @@ void ProcessHairStrandsBookmark(
 		RunHairLODSelection(
 			*GraphBuilder,
 			*Parameters.Instances,
-			Parameters.AllViews);
+			Parameters.AllViews,
+			Parameters.SkinCache, 
+			Parameters.ShaderMap);
 	}
 	else if (Bookmark == EHairStrandsBookmark::ProcessEndOfFrame)
 	{
diff --git a/Engine/Plugins/Runtime/HairStrands/Source/HairStrandsCore/Private/GroomResources.cpp b/Engine/Plugins/Runtime/HairStrands/Source/HairStrandsCore/Private/GroomResources.cpp
--- a/Engine/Plugins/Runtime/HairStrands/Source/HairStrandsCore/Private/GroomResources.cpp
+++ b/Engine/Plugins/Runtime/HairStrands/Source/HairStrandsCore/Private/GroomResources.cpp
@@ -376,6 +376,15 @@ void FHairCommonResource::Allocate(FRDGBuilder& GraphBuilder)
 	bIsInitialized = true;
 }
 
+void FHairCommonResource::AllocateLOD(FRDGBuilder& GraphBuilder, int32 LODIndex)
+{
+	// Sanity check. For allocation sub-resources (like LOD resources) the common/main resource needs to be already initialized.
+	check(bIsInitialized);
+	check(AllocationType == EHairStrandsAllocationType::Deferred);
+
+	InternalAllocateLOD(GraphBuilder, LODIndex);
+}
+
 /////////////////////////////////////////////////////////////////////////////////////////
 
 void FHairCardIndexBuffer::InitRHI()
@@ -832,64 +841,76 @@ void FHairStrandsRestRootResource::InternalAllocate(FRDGBuilder& GraphBuilder)
 	{
 		FRDGBufferUploader BufferUploader;
 		InternalCreateVertexBufferRDG<FHairStrandsIndexFormat>(GraphBuilder, BufferUploader, RootData.VertexToCurveIndexBuffer, VertexToCurveIndexBuffer, HAIRSTRANDS_RESOUCE_NAME(CurveType, Hair.StrandsRestRoot_VertexToCurveIndexBuffer), EHairResourceUsageType::Static);
+		BufferUploader.Submit(GraphBuilder);
+	}
+}
 
-		check(LODs.Num() == RootData.MeshProjectionLODs.Num());
-		for (uint32 LODIt=0, LODCount = LODs.Num(); LODIt<LODCount; ++LODIt)
+void FHairStrandsRestRootResource::InternalAllocateLOD(FRDGBuilder& GraphBuilder, int32 LODIndex)
+{
+	// Sanity check ti insure that the 'common' part of FHairStrandsRestRootResource is already initialized
+	check(bIsInitialized);
+	check(RootData.VertexToCurveIndexBuffer.Num() > 0);
+	check(LODs.Num() == RootData.MeshProjectionLODs.Num());
+	if (LODIndex >= 0 && LODIndex < LODs.Num())
+	{
+		FLOD& GPUData = LODs[LODIndex];
+		const bool bIsLODInitialized = GPUData.Status == FLOD::EStatus::Completed || GPUData.Status == FLOD::EStatus::Initialized;
+		if (bIsLODInitialized)
 		{
-			FLOD& GPUData = LODs[LODIt];
-			const FHairStrandsRootData::FMeshProjectionLOD& CPUData = RootData.MeshProjectionLODs[LODIt];
+			return;
+		}
 
-			const bool bHasValidCPUData = CPUData.RootTriangleBarycentricBuffer.Num() > 0;
-			if (bHasValidCPUData)
-			{
-				GPUData.Status = FLOD::EStatus::Completed;
+		FRDGBufferUploader BufferUploader;
+		const FHairStrandsRootData::FMeshProjectionLOD& CPUData = RootData.MeshProjectionLODs[LODIndex];
+		const bool bHasValidCPUData = CPUData.RootTriangleBarycentricBuffer.Num() > 0;
+		if (bHasValidCPUData)
+		{
+			GPUData.Status = FLOD::EStatus::Completed;
 
-				check(CPUData.RootTriangleBarycentricBuffer.Num() > 0);
-				InternalCreateVertexBufferRDG<FHairStrandsCurveTriangleBarycentricFormat>(GraphBuilder, BufferUploader, CPUData.RootTriangleBarycentricBuffer, GPUData.RootTriangleBarycentricBuffer, HAIRSTRANDS_RESOUCE_NAME(CurveType, Hair.StrandsRestRoot_RootTriangleBarycentricBuffer), EHairResourceUsageType::Static);
+			check(CPUData.RootTriangleBarycentricBuffer.Num() > 0);
+			InternalCreateVertexBufferRDG<FHairStrandsCurveTriangleBarycentricFormat>(GraphBuilder, BufferUploader, CPUData.RootTriangleBarycentricBuffer, GPUData.RootTriangleBarycentricBuffer, HAIRSTRANDS_RESOUCE_NAME(CurveType, Hair.StrandsRestRoot_RootTriangleBarycentricBuffer), EHairResourceUsageType::Static);
 
-				check(CPUData.RootTriangleIndexBuffer.Num() > 0);
-				InternalCreateVertexBufferRDG<FHairStrandsCurveTriangleIndexFormat>(GraphBuilder, BufferUploader, CPUData.RootTriangleIndexBuffer, GPUData.RootTriangleIndexBuffer, HAIRSTRANDS_RESOUCE_NAME(CurveType, Hair.StrandsRestRoot_RootTriangleIndexBuffer), EHairResourceUsageType::Static);
+			check(CPUData.RootTriangleIndexBuffer.Num() > 0);
+			InternalCreateVertexBufferRDG<FHairStrandsCurveTriangleIndexFormat>(GraphBuilder, BufferUploader, CPUData.RootTriangleIndexBuffer, GPUData.RootTriangleIndexBuffer, HAIRSTRANDS_RESOUCE_NAME(CurveType, Hair.StrandsRestRoot_RootTriangleIndexBuffer), EHairResourceUsageType::Static);
 
-				check(CPUData.RestRootTrianglePosition0Buffer.Num() > 0);
-				check(CPUData.RestRootTrianglePosition1Buffer.Num() > 0);
-				check(CPUData.RestRootTrianglePosition2Buffer.Num() > 0);
-				InternalCreateVertexBufferRDG<FHairStrandsMeshTrianglePositionFormat>(GraphBuilder, BufferUploader, CPUData.RestRootTrianglePosition0Buffer, GPUData.RestRootTrianglePosition0Buffer, HAIRSTRANDS_RESOUCE_NAME(CurveType, Hair.StrandsRestRoot_RestRootTrianglePosition0Buffer), EHairResourceUsageType::Static);
-				InternalCreateVertexBufferRDG<FHairStrandsMeshTrianglePositionFormat>(GraphBuilder, BufferUploader, CPUData.RestRootTrianglePosition1Buffer, GPUData.RestRootTrianglePosition1Buffer, HAIRSTRANDS_RESOUCE_NAME(CurveType, Hair.StrandsRestRoot_RestRootTrianglePosition1Buffer), EHairResourceUsageType::Static);
-				InternalCreateVertexBufferRDG<FHairStrandsMeshTrianglePositionFormat>(GraphBuilder, BufferUploader, CPUData.RestRootTrianglePosition2Buffer, GPUData.RestRootTrianglePosition2Buffer, HAIRSTRANDS_RESOUCE_NAME(CurveType, Hair.StrandsRestRoot_RestRootTrianglePosition2Buffer), EHairResourceUsageType::Static);
-			}
-			else
-			{
-				GPUData.Status = FLOD::EStatus::Initialized;
+			check(CPUData.RestRootTrianglePosition0Buffer.Num() > 0);
+			check(CPUData.RestRootTrianglePosition1Buffer.Num() > 0);
+			check(CPUData.RestRootTrianglePosition2Buffer.Num() > 0);
+			InternalCreateVertexBufferRDG<FHairStrandsMeshTrianglePositionFormat>(GraphBuilder, BufferUploader, CPUData.RestRootTrianglePosition0Buffer, GPUData.RestRootTrianglePosition0Buffer, HAIRSTRANDS_RESOUCE_NAME(CurveType, Hair.StrandsRestRoot_RestRootTrianglePosition0Buffer), EHairResourceUsageType::Static);
+			InternalCreateVertexBufferRDG<FHairStrandsMeshTrianglePositionFormat>(GraphBuilder, BufferUploader, CPUData.RestRootTrianglePosition1Buffer, GPUData.RestRootTrianglePosition1Buffer, HAIRSTRANDS_RESOUCE_NAME(CurveType, Hair.StrandsRestRoot_RestRootTrianglePosition1Buffer), EHairResourceUsageType::Static);
+			InternalCreateVertexBufferRDG<FHairStrandsMeshTrianglePositionFormat>(GraphBuilder, BufferUploader, CPUData.RestRootTrianglePosition2Buffer, GPUData.RestRootTrianglePosition2Buffer, HAIRSTRANDS_RESOUCE_NAME(CurveType, Hair.StrandsRestRoot_RestRootTrianglePosition2Buffer), EHairResourceUsageType::Static);
+		}
+		else
+		{
+			GPUData.Status = FLOD::EStatus::Initialized;
 			
-				InternalCreateVertexBufferRDG<FHairStrandsCurveTriangleBarycentricFormat>(GraphBuilder, RootData.RootCount, GPUData.RootTriangleBarycentricBuffer, HAIRSTRANDS_RESOUCE_NAME(CurveType, Hair.StrandsRestRoot_RootTriangleBarycentricBuffer), EHairResourceUsageType::Dynamic);
-				InternalCreateVertexBufferRDG<FHairStrandsCurveTriangleIndexFormat>(GraphBuilder, RootData.RootCount, GPUData.RootTriangleIndexBuffer, HAIRSTRANDS_RESOUCE_NAME(CurveType, Hair.StrandsRestRoot_RootTriangleIndexBuffer), EHairResourceUsageType::Dynamic);
+			InternalCreateVertexBufferRDG<FHairStrandsCurveTriangleBarycentricFormat>(GraphBuilder, RootData.RootCount, GPUData.RootTriangleBarycentricBuffer, HAIRSTRANDS_RESOUCE_NAME(CurveType, Hair.StrandsRestRoot_RootTriangleBarycentricBuffer), EHairResourceUsageType::Dynamic);
+			InternalCreateVertexBufferRDG<FHairStrandsCurveTriangleIndexFormat>(GraphBuilder, RootData.RootCount, GPUData.RootTriangleIndexBuffer, HAIRSTRANDS_RESOUCE_NAME(CurveType, Hair.StrandsRestRoot_RootTriangleIndexBuffer), EHairResourceUsageType::Dynamic);
 			
-				// Create buffers. Initialization will be done by render passes
-				InternalCreateVertexBufferRDG<FHairStrandsMeshTrianglePositionFormat>(GraphBuilder, RootData.RootCount, GPUData.RestRootTrianglePosition0Buffer, HAIRSTRANDS_RESOUCE_NAME(CurveType, Hair.StrandsRestRoot_RestRootTrianglePosition0Buffer), EHairResourceUsageType::Dynamic);
-				InternalCreateVertexBufferRDG<FHairStrandsMeshTrianglePositionFormat>(GraphBuilder, RootData.RootCount, GPUData.RestRootTrianglePosition1Buffer, HAIRSTRANDS_RESOUCE_NAME(CurveType, Hair.StrandsRestRoot_RestRootTrianglePosition1Buffer), EHairResourceUsageType::Dynamic);
-				InternalCreateVertexBufferRDG<FHairStrandsMeshTrianglePositionFormat>(GraphBuilder, RootData.RootCount, GPUData.RestRootTrianglePosition2Buffer, HAIRSTRANDS_RESOUCE_NAME(CurveType, Hair.StrandsRestRoot_RestRootTrianglePosition2Buffer), EHairResourceUsageType::Dynamic);
-			}
+			// Create buffers. Initialization will be done by render passes
+			InternalCreateVertexBufferRDG<FHairStrandsMeshTrianglePositionFormat>(GraphBuilder, RootData.RootCount, GPUData.RestRootTrianglePosition0Buffer, HAIRSTRANDS_RESOUCE_NAME(CurveType, Hair.StrandsRestRoot_RestRootTrianglePosition0Buffer), EHairResourceUsageType::Dynamic);
+			InternalCreateVertexBufferRDG<FHairStrandsMeshTrianglePositionFormat>(GraphBuilder, RootData.RootCount, GPUData.RestRootTrianglePosition1Buffer, HAIRSTRANDS_RESOUCE_NAME(CurveType, Hair.StrandsRestRoot_RestRootTrianglePosition1Buffer), EHairResourceUsageType::Dynamic);
+			InternalCreateVertexBufferRDG<FHairStrandsMeshTrianglePositionFormat>(GraphBuilder, RootData.RootCount, GPUData.RestRootTrianglePosition2Buffer, HAIRSTRANDS_RESOUCE_NAME(CurveType, Hair.StrandsRestRoot_RestRootTrianglePosition2Buffer), EHairResourceUsageType::Dynamic);
+		}
 
-			GPUData.SampleCount = CPUData.SampleCount;
-			const bool bHasValidCPUWeights = CPUData.MeshSampleIndicesBuffer.Num() > 0;
-			if (bHasValidCPUWeights)
-			{
-				//check(CPUData.MeshInterpolationWeightsBuffer.Num() == (CPUData.SampleCount+4) * (CPUData.SampleCount+4));
-				check(CPUData.MeshSampleIndicesBuffer.Num() == CPUData.SampleCount);
-				check(CPUData.RestSamplePositionsBuffer.Num() == CPUData.SampleCount);
+		GPUData.SampleCount = CPUData.SampleCount;
+		const bool bHasValidCPUWeights = CPUData.MeshSampleIndicesBuffer.Num() > 0;
+		if (bHasValidCPUWeights)
+		{
+			//check(CPUData.MeshInterpolationWeightsBuffer.Num() == (CPUData.SampleCount+4) * (CPUData.SampleCount+4));
+			check(CPUData.MeshSampleIndicesBuffer.Num() == CPUData.SampleCount);
+			check(CPUData.RestSamplePositionsBuffer.Num() == CPUData.SampleCount);
 
-				InternalCreateVertexBufferRDG<FHairStrandsWeightFormat>(GraphBuilder, BufferUploader, CPUData.MeshInterpolationWeightsBuffer, GPUData.MeshInterpolationWeightsBuffer, HAIRSTRANDS_RESOUCE_NAME(CurveType, Hair.StrandsRestRoot_MeshInterpolationWeightsBuffer), EHairResourceUsageType::Static);
-				InternalCreateVertexBufferRDG<FHairStrandsIndexFormat>(GraphBuilder, BufferUploader, CPUData.MeshSampleIndicesBuffer, GPUData.MeshSampleIndicesBuffer, HAIRSTRANDS_RESOUCE_NAME(CurveType, Hair.StrandsRestRoot_MeshSampleIndicesBuffer), EHairResourceUsageType::Static);
-				InternalCreateVertexBufferRDG<FHairStrandsMeshTrianglePositionFormat>(GraphBuilder, BufferUploader, CPUData.RestSamplePositionsBuffer, GPUData.RestSamplePositionsBuffer, HAIRSTRANDS_RESOUCE_NAME(CurveType, Hair.StrandsRestRoot_RestSamplePositionsBuffer), EHairResourceUsageType::Static);
-			}
-			else
-			{
-				InternalCreateVertexBufferRDG<FHairStrandsWeightFormat>(GraphBuilder, (CPUData.SampleCount+4) * (CPUData.SampleCount+4), GPUData.MeshInterpolationWeightsBuffer, HAIRSTRANDS_RESOUCE_NAME(CurveType, Hair.StrandsRestRoot_MeshInterpolationWeightsBuffer), EHairResourceUsageType::Dynamic);
-				InternalCreateVertexBufferRDG<FHairStrandsIndexFormat>(GraphBuilder, CPUData.SampleCount, GPUData.MeshSampleIndicesBuffer, HAIRSTRANDS_RESOUCE_NAME(CurveType, Hair.StrandsRestRoot_MeshSampleIndicesBuffer), EHairResourceUsageType::Dynamic);
-				InternalCreateVertexBufferRDG<FHairStrandsMeshTrianglePositionFormat>(GraphBuilder, CPUData.SampleCount, GPUData.RestSamplePositionsBuffer, HAIRSTRANDS_RESOUCE_NAME(CurveType, Hair.StrandsRestRoot_RestSamplePositionsBuffer), EHairResourceUsageType::Dynamic);
-			}
+			InternalCreateVertexBufferRDG<FHairStrandsWeightFormat>(GraphBuilder, BufferUploader, CPUData.MeshInterpolationWeightsBuffer, GPUData.MeshInterpolationWeightsBuffer, HAIRSTRANDS_RESOUCE_NAME(CurveType, Hair.StrandsRestRoot_MeshInterpolationWeightsBuffer), EHairResourceUsageType::Static);
+			InternalCreateVertexBufferRDG<FHairStrandsIndexFormat>(GraphBuilder, BufferUploader, CPUData.MeshSampleIndicesBuffer, GPUData.MeshSampleIndicesBuffer, HAIRSTRANDS_RESOUCE_NAME(CurveType, Hair.StrandsRestRoot_MeshSampleIndicesBuffer), EHairResourceUsageType::Static);
+			InternalCreateVertexBufferRDG<FHairStrandsMeshTrianglePositionFormat>(GraphBuilder, BufferUploader, CPUData.RestSamplePositionsBuffer, GPUData.RestSamplePositionsBuffer, HAIRSTRANDS_RESOUCE_NAME(CurveType, Hair.StrandsRestRoot_RestSamplePositionsBuffer), EHairResourceUsageType::Static);
+		}
+		else
+		{
+			InternalCreateVertexBufferRDG<FHairStrandsWeightFormat>(GraphBuilder, (CPUData.SampleCount+4) * (CPUData.SampleCount+4), GPUData.MeshInterpolationWeightsBuffer, HAIRSTRANDS_RESOUCE_NAME(CurveType, Hair.StrandsRestRoot_MeshInterpolationWeightsBuffer), EHairResourceUsageType::Dynamic);
+			InternalCreateVertexBufferRDG<FHairStrandsIndexFormat>(GraphBuilder, CPUData.SampleCount, GPUData.MeshSampleIndicesBuffer, HAIRSTRANDS_RESOUCE_NAME(CurveType, Hair.StrandsRestRoot_MeshSampleIndicesBuffer), EHairResourceUsageType::Dynamic);
+			InternalCreateVertexBufferRDG<FHairStrandsMeshTrianglePositionFormat>(GraphBuilder, CPUData.SampleCount, GPUData.RestSamplePositionsBuffer, HAIRSTRANDS_RESOUCE_NAME(CurveType, Hair.StrandsRestRoot_RestSamplePositionsBuffer), EHairResourceUsageType::Dynamic);
 		}
-
 		BufferUploader.Submit(GraphBuilder);
 	}
 }
diff --git a/Engine/Plugins/Runtime/HairStrands/Source/HairStrandsCore/Public/GroomResources.h b/Engine/Plugins/Runtime/HairStrands/Source/HairStrandsCore/Public/GroomResources.h
--- a/Engine/Plugins/Runtime/HairStrands/Source/HairStrandsCore/Public/GroomResources.h
+++ b/Engine/Plugins/Runtime/HairStrands/Source/HairStrandsCore/Public/GroomResources.h
@@ -41,8 +41,10 @@ struct FHairCommonResource : public FRenderResource
 
 	/* Init/Release buffers (FHairCommonResource) */
 	void Allocate(FRDGBuilder& GraphBuilder);
+	void AllocateLOD(FRDGBuilder& GraphBuilder, int32 LODIndex);
 	virtual void InternalAllocate() {};
 	virtual void InternalAllocate(FRDGBuilder& GraphBuilder) {};
+	virtual void InternalAllocateLOD(FRDGBuilder& GraphBuilder, int32 LODIndex) {};
 	virtual void InternalRelease() {};
 
 	bool bUseRenderGraph = true;
@@ -58,6 +60,7 @@ struct FHairStrandsRestRootResource : public FHairCommonResource
 
 	/* Init/Release buffers */
 	virtual void InternalAllocate(FRDGBuilder& GraphBuilder) override;
+	virtual void InternalAllocateLOD(FRDGBuilder& GraphBuilder, int32 LODIndex) override;
 	virtual void InternalRelease() override;
 
 	/* Get the resource name */
