Move hair strands root data initialization into hair strands root builder.

diff --git a/Engine/Plugins/Runtime/HairStrands/Source/HairStrandsCore/Private/GroomBindingBuilder.cpp b/Engine/Plugins/Runtime/HairStrands/Source/HairStrandsCore/Private/GroomBindingBuilder.cpp
--- a/Engine/Plugins/Runtime/HairStrands/Source/HairStrandsCore/Private/GroomBindingBuilder.cpp
+++ b/Engine/Plugins/Runtime/HairStrands/Source/HairStrandsCore/Private/GroomBindingBuilder.cpp
@@ -1779,6 +1779,38 @@ namespace GroomBinding_Transfer
 }
 // namespace GroomBinding_Transfer
 
+static void InitHairStrandsRootData(FHairStrandsRootData& Out, const FHairStrandsDatas* HairStrandsDatas, uint32 LODCount, const TArray<uint32>& NumSamples)
+{
+	check(HairStrandsDatas);
+
+	Out.RootCount = HairStrandsDatas->GetNumCurves();
+
+	const uint32 CurveCount = HairStrandsDatas->GetNumCurves();
+	Out.VertexToCurveIndexBuffer.SetNum(HairStrandsDatas->GetNumPoints());
+	
+	for (uint32 CurveIndex = 0; CurveIndex < CurveCount; ++CurveIndex)
+	{
+		const uint32 RootIndex = HairStrandsDatas->StrandsCurves.CurvesOffset[CurveIndex];
+		const uint32 PointCount = HairStrandsDatas->StrandsCurves.CurvesCount[CurveIndex];
+		for (uint32 PointIndex = 0; PointIndex < PointCount; ++PointIndex)
+		{
+			Out.VertexToCurveIndexBuffer[RootIndex + PointIndex] = CurveIndex; // RootIndex;
+		}
+	}
+	check(NumSamples.Num() == LODCount);
+
+	Out.MeshProjectionLODs.SetNum(LODCount);
+	uint32 LODIndex = 0;
+	for (FHairStrandsRootData::FMeshProjectionLOD& MeshProjectionLOD : Out.MeshProjectionLODs)
+	{
+		MeshProjectionLOD.SampleCount = NumSamples[LODIndex];
+		MeshProjectionLOD.LODIndex = LODIndex++;
+		MeshProjectionLOD.MeshInterpolationWeightsBuffer.Empty();
+		MeshProjectionLOD.MeshSampleIndicesBuffer.Empty();
+		MeshProjectionLOD.RestSamplePositionsBuffer.Empty();
+	}
+}
+
 ///////////////////////////////////////////////////////////////////////////////////////////////////
 // Main entry (CPU path)
 static bool InternalBuildBinding_CPU(UGroomBindingAsset* BindingAsset, bool bInitResources)
@@ -1845,8 +1877,8 @@ static bool InternalBuildBinding_CPU(UGroomBindingAsset* BindingAsset, bool bIni
 		GroomAsset->GetHairStrandsDatas(GroupIndex, StrandsData, GuidesData);
 
 		UGroomBindingAsset::FHairGroupData& Data = OutHairGroupDatas.AddDefaulted_GetRef();
-		Data.RenRootData = FHairStrandsRootData(&StrandsData, MeshLODCount, NumSamples);
-		Data.SimRootData = FHairStrandsRootData(&GuidesData, MeshLODCount, NumSamples);
+		InitHairStrandsRootData(Data.RenRootData, &StrandsData, MeshLODCount, NumSamples);
+		InitHairStrandsRootData(Data.SimRootData, &GuidesData, MeshLODCount, NumSamples);
 
 		const uint32 CardsLODCount = GroupData.Cards.LODs.Num();
 		Data.CardsRootData.SetNum(GroupData.Cards.LODs.Num());
@@ -1857,7 +1889,7 @@ static bool InternalBuildBinding_CPU(UGroomBindingAsset* BindingAsset, bool bIni
 				FHairStrandsDatas LODGuidesData;
 				const bool bIsValid = GroomAsset->GetHairCardsGuidesDatas(GroupIndex, CardsLODIt, LODGuidesData);
 				check(bIsValid);
-				Data.CardsRootData[CardsLODIt] = FHairStrandsRootData(&LODGuidesData, MeshLODCount, NumSamples);
+				InitHairStrandsRootData(Data.CardsRootData[CardsLODIt], &LODGuidesData, MeshLODCount, NumSamples);
 			}
 		}
 		++GroupIndex;
diff --git a/Engine/Plugins/Runtime/HairStrands/Source/HairStrandsCore/Private/GroomResources.cpp b/Engine/Plugins/Runtime/HairStrands/Source/HairStrandsCore/Private/GroomResources.cpp
--- a/Engine/Plugins/Runtime/HairStrands/Source/HairStrandsCore/Private/GroomResources.cpp
+++ b/Engine/Plugins/Runtime/HairStrands/Source/HairStrandsCore/Private/GroomResources.cpp
@@ -1000,38 +1000,6 @@ FHairStrandsRootData::FHairStrandsRootData()
 
 }
 
-FHairStrandsRootData::FHairStrandsRootData(const FHairStrandsDatas* HairStrandsDatas, uint32 LODCount, const TArray<uint32>& NumSamples):
-	RootCount(HairStrandsDatas ? HairStrandsDatas->GetNumCurves() : 0)
-{
-	if (!HairStrandsDatas)
-		return;
-
-	const uint32 CurveCount = HairStrandsDatas->GetNumCurves();
-	VertexToCurveIndexBuffer.SetNum(HairStrandsDatas->GetNumPoints());
-	
-	for (uint32 CurveIndex = 0; CurveIndex < CurveCount; ++CurveIndex)
-	{
-		const uint32 RootIndex = HairStrandsDatas->StrandsCurves.CurvesOffset[CurveIndex];
-		const uint32 PointCount = HairStrandsDatas->StrandsCurves.CurvesCount[CurveIndex];
-		for (uint32 PointIndex = 0; PointIndex < PointCount; ++PointIndex)
-		{
-			VertexToCurveIndexBuffer[RootIndex + PointIndex] = CurveIndex; // RootIndex;
-		}
-	}
-	check(NumSamples.Num() == LODCount);
-
-	MeshProjectionLODs.SetNum(LODCount);
-	uint32 LODIndex = 0;
-	for (FMeshProjectionLOD& MeshProjectionLOD : MeshProjectionLODs)
-	{
-		MeshProjectionLOD.SampleCount = NumSamples[LODIndex];
-		MeshProjectionLOD.LODIndex = LODIndex++;
-		MeshProjectionLOD.MeshInterpolationWeightsBuffer.Empty();
-		MeshProjectionLOD.MeshSampleIndicesBuffer.Empty();
-		MeshProjectionLOD.RestSamplePositionsBuffer.Empty();
-	}
-}
-
 bool FHairStrandsRootData::HasProjectionData() const
 {
 	bool bIsValid = MeshProjectionLODs.Num() > 0;
