Remove constructor declaration.

diff --git a/Engine/Plugins/Runtime/HairStrands/Source/HairStrandsCore/Public/HairStrandsDatas.h b/Engine/Plugins/Runtime/HairStrands/Source/HairStrandsCore/Public/HairStrandsDatas.h
--- a/Engine/Plugins/Runtime/HairStrands/Source/HairStrandsCore/Public/HairStrandsDatas.h
+++ b/Engine/Plugins/Runtime/HairStrands/Source/HairStrandsCore/Public/HairStrandsDatas.h
@@ -497,7 +497,6 @@ struct FHairStrandsRootData
 {
 	/** Build the hair strands resource */
 	FHairStrandsRootData();
-	FHairStrandsRootData(const FHairStrandsDatas* HairStrandsDatas, uint32 LODCount, const TArray<uint32>& NumSamples);
 	void Serialize(FArchive& Ar);
 	void Reset();
 	bool HasProjectionData() const;
