Remove GPU path for building binding hair data.

diff --git a/Engine/Plugins/Runtime/HairStrands/Source/HairStrandsCore/Private/GroomBindingAsset.cpp b/Engine/Plugins/Runtime/HairStrands/Source/HairStrandsCore/Private/GroomBindingAsset.cpp
--- a/Engine/Plugins/Runtime/HairStrands/Source/HairStrandsCore/Private/GroomBindingAsset.cpp
+++ b/Engine/Plugins/Runtime/HairStrands/Source/HairStrandsCore/Private/GroomBindingAsset.cpp
@@ -732,7 +732,7 @@ void UGroomBindingAsset::CacheDerivedDatas()
 			}
 
 			// Build groom binding data
-			bIsValid = FGroomBindingBuilder::BuildBinding(this, false, false);
+			bIsValid = FGroomBindingBuilder::BuildBinding(this, false);
 			if (bIsValid)
 			{
 				// Using a LargeMemoryWriter for serialization since the data can be bigger than 2 GB
diff --git a/Engine/Plugins/Runtime/HairStrands/Source/HairStrandsCore/Private/GroomBindingBuilder.cpp b/Engine/Plugins/Runtime/HairStrands/Source/HairStrandsCore/Private/GroomBindingBuilder.cpp
--- a/Engine/Plugins/Runtime/HairStrands/Source/HairStrandsCore/Private/GroomBindingBuilder.cpp
+++ b/Engine/Plugins/Runtime/HairStrands/Source/HairStrandsCore/Private/GroomBindingBuilder.cpp
@@ -2021,489 +2021,8 @@ static bool InternalBuildBinding_CPU(UGroomBindingAsset* BindingAsset, bool bIni
 	return true;
 }
 
-///////////////////////////////////////////////////////////////////////////////////////////////////
-// GPU path
-
-#define GPU_BINDING 0
-#if GPU_BINDING
-#if WITH_EDITORONLY_DATA
-namespace GroomBinding_GPU
-{
-	template<typename ReadBackType>
-	void ReadbackBuffer(TArray<ReadBackType>& OutData, FRWBuffer& InBuffer)
-	{
-		ReadBackType* Data = (ReadBackType*)RHILockBuffer(InBuffer.Buffer, 0, InBuffer.Buffer->GetSize(), RLM_ReadOnly);
-		const uint32 ElementCount = InBuffer.Buffer->GetSize() / sizeof(ReadBackType);
-		OutData.SetNum(ElementCount);
-		for (uint32 ElementIt = 0; ElementIt < ElementCount; ++ElementIt)
-		{
-			OutData[ElementIt] = Data[ElementIt];
-		}
-		RHIUnlockBuffer(InBuffer.Buffer);
-	}
-
-	template<typename WriteBackType>
-	void WritebackBuffer(TArray<WriteBackType>& InData, FRWBuffer& OutBuffer)
-	{
-		const uint32 DataSize = sizeof(WriteBackType) * InData.Num();
-		check(DataSize == OutBuffer.Buffer->GetSize());
-
-		WriteBackType* Data = (WriteBackType*)RHILockBuffer(OutBuffer.Buffer, 0, DataSize, RLM_WriteOnly);
-		FMemory::Memcpy(Data, InData.GetData(), DataSize);
-		RHIUnlockBuffer(OutBuffer.Buffer);
-	}
-
-	void ReadbackGroupData(
-		FHairStrandsRootData& OutCPUData,
-		FHairStrandsRestRootResource* InGPUData)
-	{
-		if (!InGPUData)
-		{
-			return;
-		}
-
-		check(InGPUData->LODs.Num() == OutCPUData.MeshProjectionLODs.Num());
-
-		const uint32 MeshLODCount = InGPUData->LODs.Num();
-		for (uint32 LODIt = 0; LODIt < MeshLODCount; ++LODIt)
-		{
-			FHairStrandsRootData::FMeshProjectionLOD& CPULOD = OutCPUData.MeshProjectionLODs[LODIt];
-			FHairStrandsRestRootResource::FLOD& GPULOD = InGPUData->LODs[LODIt];
-			check(CPULOD.LODIndex == GPULOD.LODIndex);
-
-			ReadbackBuffer(CPULOD.RootTriangleIndexBuffer, GPULOD.RootTriangleIndexBuffer);
-			ReadbackBuffer(CPULOD.RootTriangleBarycentricBuffer, GPULOD.RootTriangleBarycentricBuffer);
-			ReadbackBuffer(CPULOD.RestRootTrianglePosition0Buffer, GPULOD.RestRootTrianglePosition0Buffer);
-			ReadbackBuffer(CPULOD.RestRootTrianglePosition1Buffer, GPULOD.RestRootTrianglePosition1Buffer);
-			ReadbackBuffer(CPULOD.RestRootTrianglePosition2Buffer, GPULOD.RestRootTrianglePosition2Buffer);
-
-			InGPUData->RootData.MeshProjectionLODs[LODIt].RootTriangleIndexBuffer = CPULOD.RootTriangleIndexBuffer;
-			InGPUData->RootData.MeshProjectionLODs[LODIt].RootTriangleBarycentricBuffer = CPULOD.RootTriangleBarycentricBuffer;
-			InGPUData->RootData.MeshProjectionLODs[LODIt].RestRootTrianglePosition0Buffer = CPULOD.RestRootTrianglePosition0Buffer;
-			InGPUData->RootData.MeshProjectionLODs[LODIt].RestRootTrianglePosition1Buffer = CPULOD.RestRootTrianglePosition1Buffer;
-			InGPUData->RootData.MeshProjectionLODs[LODIt].RestRootTrianglePosition2Buffer = CPULOD.RestRootTrianglePosition2Buffer;
-		}
-	}
-
-	void WritebackGroupData(FHairStrandsRootData& InCPUData, FHairStrandsRestRootResource* OutGPUData)
-	{
-		if (!OutGPUData)
-		{
-			return;
-		}
-		check(OutGPUData->LODs.Num() == InCPUData.MeshProjectionLODs.Num());
-		const uint32 MeshLODCount = OutGPUData->LODs.Num();
-		for (uint32 LODIt = 0; LODIt < MeshLODCount; ++LODIt)
-		{
-			FHairStrandsRootData::FMeshProjectionLOD& CPULOD = InCPUData.MeshProjectionLODs[LODIt];
-			FHairStrandsRestRootResource::FLOD& GPULOD = OutGPUData->LODs[LODIt];
-			check(CPULOD.LODIndex == GPULOD.LODIndex);
-
-			if (CPULOD.SampleCount > 0)
-			{
-				WritebackBuffer(CPULOD.MeshInterpolationWeightsBuffer, GPULOD.MeshInterpolationWeightsBuffer);
-				WritebackBuffer(CPULOD.MeshSampleIndicesBuffer, GPULOD.MeshSampleIndicesBuffer);
-				WritebackBuffer(CPULOD.RestSamplePositionsBuffer, GPULOD.RestSamplePositionsBuffer);
-
-				OutGPUData->RootData.MeshProjectionLODs[LODIt].SampleCount = CPULOD.SampleCount;
-				OutGPUData->RootData.MeshProjectionLODs[LODIt].MeshInterpolationWeightsBuffer = CPULOD.MeshInterpolationWeightsBuffer;
-				OutGPUData->RootData.MeshProjectionLODs[LODIt].RestSamplePositionsBuffer = CPULOD.RestSamplePositionsBuffer;
-				OutGPUData->RootData.MeshProjectionLODs[LODIt].MeshSampleIndicesBuffer = CPULOD.MeshSampleIndicesBuffer;
-			}
-		}
-	}
-
-	void ComputeInterpolationWeights(UGroomBindingAsset* BindingAsset, FSkeletalMeshRenderData* TargetRenderData, TArray<FRWBuffer>& TransferedPositions)
-	{
-		UGroomAsset* GroomAsset = BindingAsset->Groom;
-		// Enforce GPU sync to read back data on CPU
-		GDynamicRHI->RHISubmitCommandsAndFlushGPU();
-		GDynamicRHI->RHIBlockUntilGPUIdle();
-
-		UGroomBindingAsset::FHairGroupDatas& OutHairGroupDatas = BindingAsset->HairGroupDatas;
-		UGroomBindingAsset::FHairGroupResources& OutHairGroupResources = BindingAsset->HairGroupResources;
-
-		const uint32 GroupCount  = OutHairGroupResources.Num();
-		const uint32 MeshLODCount= BindingAsset->TargetSkeletalMesh->GetLODNum();
-		const uint32 MaxSamples  = BindingAsset->NumInterpolationPoints;
-
-		for (uint32 GroupIt = 0; GroupIt < GroupCount; ++GroupIt)
-		{
-			ReadbackGroupData(OutHairGroupDatas[GroupIt].SimRootData, OutHairGroupResources[GroupIt].SimRootResources);
-			ReadbackGroupData(OutHairGroupDatas[GroupIt].RenRootData, OutHairGroupResources[GroupIt].RenRootResources);
-
-			const uint32 CardsLODCount = OutHairGroupDatas[GroupIt].CardsRootData.Num();
-			for (uint32 CardsLODIt = 0; CardsLODIt < CardsLODCount; ++CardsLODIt)
-			{
-				if (OutHairGroupDatas[GroupIt].CardsRootData[CardsLODIt].IsValid())
-				{
-					//check(OutHairGroupResources[GroupIt].CardsRootResources[CardsLODIt].IsValid());
-					ReadbackGroupData(OutHairGroupDatas[GroupIt].CardsRootData[CardsLODIt], OutHairGroupResources[GroupIt].CardsRootResources[CardsLODIt]);
-				}
-			}
-		}
-
-		const bool LocalSamples = false;
-		for (uint32 LODIndex = 0; LODIndex < MeshLODCount; ++LODIndex)
-		{
-			FSkeletalMeshLODRenderData& LODRenderData = TargetRenderData->LODRenderData[LODIndex];
-
-			int32 TargetSection = -1;
-			bool GlobalSamples = false;
-
-			TArray<FVector3f> SourcePositions;
-			FVector3f* PositionsPointer = nullptr;
-			if (TransferedPositions.Num() == MeshLODCount)
-			{
-				ReadbackBuffer(SourcePositions, TransferedPositions[LODIndex]);
-				PositionsPointer = SourcePositions.GetData();
-				GlobalSamples = true;
-				TargetSection = BindingAsset->MatchingSection;
-			}
-			else
-			{
-				FPositionVertexBuffer& VertexBuffer = LODRenderData.StaticVertexBuffers.PositionVertexBuffer;
-				PositionsPointer = static_cast<FVector3f*>(VertexBuffer.GetVertexData());
-			}
-
-			if (LocalSamples)
-			{
-				TArray<bool> ValidPoints;
-				for (uint32 GroupIt = 0; GroupIt < GroupCount; ++GroupIt)
-				{
-					GroomBinding_RBFWeighting::FillLocalValidPoints(LODRenderData, TargetSection, OutHairGroupDatas[GroupIt].RenRootData.MeshProjectionLODs[LODIndex], ValidPoints);
-
-					GroomBinding_RBFWeighting::FPointsSampler PointsSampler(ValidPoints, PositionsPointer, MaxSamples);
-					const uint32 SampleCount = PointsSampler.SamplePositions.Num();
-
-					GroomBinding_RBFWeighting::FWeightsBuilder InterpolationWeights(SampleCount, SampleCount,
-						PointsSampler.SamplePositions.GetData(), PointsSampler.SamplePositions.GetData());
-
-					GroomBinding_RBFWeighting::UpdateInterpolationWeights(InterpolationWeights, PointsSampler, LODIndex, OutHairGroupDatas[GroupIt].SimRootData);
-					GroomBinding_RBFWeighting::UpdateInterpolationWeights(InterpolationWeights, PointsSampler, LODIndex, OutHairGroupDatas[GroupIt].RenRootData);
-
-					const uint32 CardsLODCount = OutHairGroupDatas[GroupIt].CardsRootData.Num();
-					for (uint32 CardsLODIt = 0; CardsLODIt < CardsLODCount; ++CardsLODIt)
-					{
-						if (OutHairGroupDatas[GroupIt].CardsRootData[CardsLODIt].IsValid())
-						{
-							UpdateInterpolationWeights(InterpolationWeights, PointsSampler, LODIndex, OutHairGroupDatas[GroupIt].CardsRootData[CardsLODIt]);
-						}
-					}
-				}
-			}
-			else
-			{
-				TArray<bool> ValidPoints;
-
-				GroomBinding_RBFWeighting::FillGlobalValidPoints(LODRenderData, TargetSection, ValidPoints);
-
-				GroomBinding_RBFWeighting::FPointsSampler PointsSampler(ValidPoints, PositionsPointer, MaxSamples);
-				const uint32 SampleCount = PointsSampler.SamplePositions.Num();
-
-				GroomBinding_RBFWeighting::FWeightsBuilder InterpolationWeights(SampleCount, SampleCount,
-					PointsSampler.SamplePositions.GetData(), PointsSampler.SamplePositions.GetData());
-
-				for (uint32 GroupIt = 0; GroupIt < GroupCount; ++GroupIt)
-				{
-					GroomBinding_RBFWeighting::UpdateInterpolationWeights(InterpolationWeights, PointsSampler, LODIndex, OutHairGroupDatas[GroupIt].SimRootData);
-					GroomBinding_RBFWeighting::UpdateInterpolationWeights(InterpolationWeights, PointsSampler, LODIndex, OutHairGroupDatas[GroupIt].RenRootData);
-
-					const uint32 CardsLODCount = OutHairGroupDatas[GroupIt].CardsRootData.Num();
-					for (uint32 CardsLODIt = 0; CardsLODIt < CardsLODCount; ++CardsLODIt)
-					{
-						if (OutHairGroupDatas[GroupIt].CardsRootData[CardsLODIt].IsValid())
-						{
-							UpdateInterpolationWeights(InterpolationWeights, PointsSampler, LODIndex, OutHairGroupDatas[GroupIt].CardsRootData[CardsLODIt]);
-						}
-					}
-				}
-			}
-		}
-		for (uint32 GroupIt = 0; GroupIt < GroupCount; ++GroupIt)
-		{
-			WritebackGroupData(OutHairGroupDatas[GroupIt].SimRootData, OutHairGroupResources[GroupIt].SimRootResources);
-			WritebackGroupData(OutHairGroupDatas[GroupIt].RenRootData, OutHairGroupResources[GroupIt].RenRootResources);
-
-			const uint32 CardsLODCount = OutHairGroupDatas[GroupIt].CardsRootData.Num();
-			for (uint32 CardsLODIt = 0; CardsLODIt < CardsLODCount; ++CardsLODIt)
-			{
-				if (OutHairGroupDatas[GroupIt].CardsRootData[CardsLODIt].IsValid())
-				{
-					WritebackGroupData(OutHairGroupDatas[GroupIt].CardsRootData[CardsLODIt], OutHairGroupResources[GroupIt].CardsRootResources[CardsLODIt]);
-				}
-			}
-		}
-	}
-} // namespace GroomBinding_GPU
-
-#endif
-#endif
-
-///////////////////////////////////////////////////////////////////////////////////////////////////
-// Main entry (GPU path)
-#if GPU_BINDING
-static void InternalBuildBinding_GPU(FRDGBuilder& GraphBuilder, UGroomBindingAsset* BindingAsset)
-{
-#if WITH_EDITORONLY_DATA
-	if (!BindingAsset ||
-		!BindingAsset->Groom ||
-		!BindingAsset->TargetSkeletalMesh ||
-		BindingAsset->Groom->GetNumHairGroups() == 0)
-	{
-		if (GHairStrandsBindingBuilderWarningEnable > 0)
-		{
-			UE_LOG(LogHairStrands, Warning, TEXT("[Groom] Error - Binding asset can be created/rebuilt."));
-		}
-		return;
-	}
-
-	const int32 NumInterpolationPoints = BindingAsset->NumInterpolationPoints;
-	UGroomAsset* GroomAsset = BindingAsset->Groom;
-	USkeletalMesh* SourceSkeletalMesh = BindingAsset->SourceSkeletalMesh;
-	USkeletalMesh* TargetSkeletalMesh = BindingAsset->TargetSkeletalMesh;
-
-	const uint32 MeshLODCount = BindingAsset->TargetSkeletalMesh->GetLODNum();
-	UGroomBindingAsset::FHairGroupDatas& OutHairGroupDatas = BindingAsset->HairGroupDatas;
-	OutHairGroupDatas.Empty();
-	TArray<uint32> NumSamples;
-	NumSamples.Init(NumInterpolationPoints, MeshLODCount);
-	for (const FHairGroupData& GroupData : GroomAsset->HairGroupsData)
-	{
-		UGroomBindingAsset::FHairGroupData& Data = OutHairGroupDatas.AddDefaulted_GetRef();
-		Data.RenRootData = FHairStrandsRootData(&GroupData.Strands.Data, MeshLODCount, NumSamples);
-		Data.SimRootData = FHairStrandsRootData(&GroupData.Guides.Data, MeshLODCount, NumSamples);
-
-		const uint32 CardsLODCount = GroupData.Cards.LODs.Num();
-		Data.CardsRootData.SetNum(GroupData.Cards.LODs.Num());
-		for (uint32 CardsLODIt = 0; CardsLODIt < CardsLODCount; ++CardsLODIt)
-		{
-			Data.CardsRootData[CardsLODIt].Reset();
-			if (GroupData.Cards.IsValid(CardsLODIt))
-			{
-				Data.CardsRootData[CardsLODIt] = FHairStrandsRootData(&GroupData.Cards.LODs[CardsLODIt].Guides.Data, MeshLODCount, NumSamples);
-			}
-		}
-	}
-
-	UGroomBindingAsset::FHairGroupResources& OutHairGroupResources = BindingAsset->HairGroupResources;
-	if (BindingAsset->HairGroupResources.Num() > 0)
-	{
-		for (UGroomBindingAsset::FHairGroupResource& GroupResrouces : OutHairGroupResources)
-		{
-			BindingAsset->HairGroupResourcesToDelete.Enqueue(GroupResrouces);
-		}
-		OutHairGroupResources.Empty();
-	}
-
-	check(OutHairGroupResources.Num() == 0);
-	for (UGroomBindingAsset::FHairGroupData& GroupData : OutHairGroupDatas)
-	{
-		UGroomBindingAsset::FHairGroupResource& Resource = OutHairGroupResources.AddDefaulted_GetRef();
-		Resource.SimRootResources = new FHairStrandsRestRootResource(GroupData.SimRootData);
-		Resource.RenRootResources = new FHairStrandsRestRootResource(GroupData.RenRootData);
-
-		Resource.SimRootResources->InitRHI();
-		Resource.RenRootResources->InitRHI();
-
-		const uint32 CardsLODCount = GroupData.CardsRootData.Num();
-		Resource.CardsRootResources.SetNum(CardsLODCount);
-		for (uint32 CardsLODIt=0; CardsLODIt< CardsLODCount; ++CardsLODIt)
-		{
-			Resource.CardsRootResources[CardsLODIt] = nullptr;
-
-			if (GroupData.CardsRootData[CardsLODIt].IsValid())
-			{
-				Resource.CardsRootResources[CardsLODIt] = new FHairStrandsRestRootResource(GroupData.CardsRootData[CardsLODIt]);
-				Resource.CardsRootResources[CardsLODIt]->InitRHI();
-			}
-		}
-	}
-
-	UpdateGroupInfos(BindingAsset);
-
-	TArray<FGoomBindingGroupInfo>& OutGroupInfos = BindingAsset->GroupInfos;
-	OutGroupInfos.Empty();
-	for (const UGroomBindingAsset::FHairGroupData& Data : OutHairGroupDatas)
-	{
-		FGoomBindingGroupInfo& Info = OutGroupInfos.AddDefaulted_GetRef();
-		Info.SimRootCount = Data.SimRootData.RootCount;
-		Info.SimLODCount = Data.SimRootData.MeshProjectionLODs.Num();
-		Info.RenRootCount = Data.RenRootData.RootCount;
-		Info.RenLODCount = Data.RenRootData.MeshProjectionLODs.Num();
-	}
-
-	FSkeletalMeshRenderData* TargetRenderData = TargetSkeletalMesh->GetResourceForRendering();
-	FHairStrandsProjectionMeshData TargetMeshData = ExtractMeshData(TargetRenderData);
-
-	// Create mapping between the source & target using their UV
-	// The lifetime of 'TransferredPositions' needs to encompass RunProjection
-	struct FTransferData
-	{
-		TArray<FRWBuffer> TransferredPositions;
-	};
-	FTransferData* TransferData = new FTransferData();
-
-	if (FSkeletalMeshRenderData* SourceRenderData = SourceSkeletalMesh ? SourceSkeletalMesh->GetResourceForRendering() : nullptr)
-	{
-		FHairStrandsProjectionMeshData SourceMeshData = ExtractMeshData(SourceRenderData);
-		FGroomBindingBuilder::TransferMesh(
-			GraphBuilder,
-			SourceMeshData,
-			TargetMeshData,
-			TransferData->TransferredPositions);
-
-		for (uint32 LODIndex = 0; LODIndex < MeshLODCount; ++LODIndex)
-		{
-			for (FHairStrandsProjectionMeshData::Section& Section : TargetMeshData.LODs[LODIndex].Sections)
-			{
-				Section.PositionBuffer = TransferData->TransferredPositions[LODIndex].SRV;
-			}
-		}
-	}
-
-	TArray<FHairStrandsRestRootResource*> RootResources;
-	for (const UGroomBindingAsset::FHairGroupResource& GroupResources : OutHairGroupResources)
-	{
-		RootResources.Add(GroupResources.RenRootResources);
-		RootResources.Add(GroupResources.SimRootResources);
-
-		for (FHairStrandsRestRootResource* CardsRootResources : GroupResources.CardsRootResources)
-		{
-			if (CardsRootResources)
-			{
-				RootResources.Add(CardsRootResources);
-			}
-		}
-	}
-	FGroomBindingBuilder::ProjectStrands(GraphBuilder, FTransform::Identity, TargetMeshData, RootResources);
-
-	// Readback the data
-	GraphBuilder.AddPass(
-		RDG_EVENT_NAME("GroomBinding_Readback"),
-		ERDGPassFlags::None,
-		[BindingAsset, TargetRenderData, TransferData](FRHICommandList& RHICmdList)
-		{
-			GroomBinding_GPU::ComputeInterpolationWeights(BindingAsset, TargetRenderData, TransferData->TransferredPositions);
-			BindingAsset->QueryStatus = UGroomBindingAsset::EQueryStatus::Completed;
-		});
-
-#endif
-}
-#endif
-
-///////////////////////////////////////////////////////////////////////////////////////////////////
-// Asynchronous queuing for binding generation (GPU)
-struct FBindingQuery
+bool FGroomBindingBuilder::BuildBinding(UGroomBindingAsset* BindingAsset, bool bInitResources)
 {
-	UGroomBindingAsset* Asset = nullptr;
-};
-TQueue<FBindingQuery> GBindingQueries;
-
-bool HasHairStrandsBindigQueries()
-{
-	return !GBindingQueries.IsEmpty();
+	return InternalBuildBinding_CPU(BindingAsset, bInitResources);
 }
-
-void RunHairStrandsBindingQueries(FRDGBuilder& GraphBuilder, FGlobalShaderMap* ShaderMap)
-{
-	FBindingQuery Q;
-	while (GBindingQueries.Dequeue(Q))
-	{
-		if (Q.Asset)
-		{
-			#if GPU_BINDING
-			InternalBuildBinding_GPU(GraphBuilder, Q.Asset);
-			#endif
-		}
-	}
-}
-
-bool FGroomBindingBuilder::BuildBinding(UGroomBindingAsset* BindingAsset, bool bUseGPU, bool bInitResources)
-{
-	bool bSucceed = false;
-	if (!bUseGPU)
-	{
-		bSucceed = InternalBuildBinding_CPU(BindingAsset, bInitResources);
-	}
-	else
-	{
-		BindingAsset->QueryStatus = UGroomBindingAsset::EQueryStatus::Submitted;
-		GBindingQueries.Enqueue({ BindingAsset });
-		bSucceed = true;
-	}
-
-	return bSucceed;
-}
-
-////////////////////////////////////////////////////////////////////////////////////////////////////////////////////////
-// Immediate version
-#ifndef GPU_BINDING
-#define GPU_BINDING 0
-#endif
-#if GPU_BINDING
-void FGroomBindingBuilder::TransferMesh(
-	FRDGBuilder& GraphBuilder,
-	const FHairStrandsProjectionMeshData& SourceMeshData,
-	const FHairStrandsProjectionMeshData& TargetMeshData,
-	TArray<FRWBuffer>& OutTransferedPositions)
-{
-	const ERHIFeatureLevel::Type FeatureLevel = GMaxRHIFeatureLevel;
-	FGlobalShaderMap* ShaderMap = GetGlobalShaderMap(FeatureLevel);
-
-	FBufferTransitionQueue TransitionQueue;
-
-	const uint32 MeshLODCount = TargetMeshData.LODs.Num();
-	OutTransferedPositions.SetNum(MeshLODCount);
-	for (uint32 LODIndex = 0; LODIndex < MeshLODCount; ++LODIndex)
-	{
-		check(TargetMeshData.LODs[LODIndex].Sections.Num() > 0);
-
-		OutTransferedPositions[LODIndex].Initialize(sizeof(float), TargetMeshData.LODs[LODIndex].Sections[0].TotalVertexCount * 3, PF_R32_FLOAT);
-		::TransferMesh(GraphBuilder, ShaderMap, LODIndex, SourceMeshData, TargetMeshData, OutTransferedPositions[LODIndex], TransitionQueue);
-	}
-
-	TransitBufferToReadable(GraphBuilder, TransitionQueue);
-}
-
-void FGroomBindingBuilder::ProjectStrands(
-	FRDGBuilder& GraphBuilder,
-	const FTransform& LocalToWorld,
-	const FHairStrandsProjectionMeshData& TargetMeshData,
-	TArray<FHairStrandsRestRootResource*>& InRestRootResources)
-{
-	const ERHIFeatureLevel::Type FeatureLevel = GMaxRHIFeatureLevel;
-	FGlobalShaderMap* ShaderMap = GetGlobalShaderMap(FeatureLevel);
-
-	FBufferTransitionQueue TransitionQueue;
-
-	for (FHairStrandsRestRootResource* RestRootResource : InRestRootResources)
-	{
-		for (FHairStrandsRestRootResource::FLOD& LODData : RestRootResource->LODs)
-		{
-			const uint32 LODIndex = LODData.LODIndex;
-			ProjectHairStrandsOntoMesh(
-				GraphBuilder,
-				ShaderMap,
-				LODIndex,
-				TargetMeshData,
-				RestRootResource,
-				TransitionQueue);
-
-			AddHairStrandUpdateMeshTrianglesPass(
-				GraphBuilder,
-				ShaderMap,
-				LODIndex,
-				HairStrandsTriangleType::RestPose,
-				TargetMeshData.LODs[LODIndex],
-				RestRootResource,
-				nullptr,
-				TransitionQueue);
-		}
-	}
-
-	TransitBufferToReadable(GraphBuilder, TransitionQueue);
-}
-#endif
-
 #undef LOCTEXT_NAMESPACE
diff --git a/Engine/Plugins/Runtime/HairStrands/Source/HairStrandsCore/Private/GroomManager.cpp b/Engine/Plugins/Runtime/HairStrands/Source/HairStrandsCore/Private/GroomManager.cpp
--- a/Engine/Plugins/Runtime/HairStrands/Source/HairStrandsCore/Private/GroomManager.cpp
+++ b/Engine/Plugins/Runtime/HairStrands/Source/HairStrandsCore/Private/GroomManager.cpp
@@ -737,9 +737,6 @@ bool HasHairStrandsTexturesQueries();
 void RunHairStrandsTexturesQueries(FRDGBuilder& GraphBuilder, FGlobalShaderMap* ShaderMap, const struct FShaderDrawDebugData* DebugShaderData);
 #endif
 
-bool HasHairStrandsBindigQueries();
-void RunHairStrandsBindingQueries(FRDGBuilder& GraphBuilder, FGlobalShaderMap* ShaderMap);
-
 #if WITH_EDITOR
 bool HasHairCardsAtlasQueries();
 void RunHairCardsAtlasQueries(FRDGBuilder& GraphBuilder, FGlobalShaderMap* ShaderMap, const struct FShaderDrawDebugData* DebugShaderData);
@@ -759,11 +756,6 @@ static void RunHairStrandsProcess(FRDGBuilder& GraphBuilder, FGlobalShaderMap* S
 		RunHairStrandsFolliculeMaskQueries(GraphBuilder, ShaderMap);
 	}
 
-	if (HasHairStrandsBindigQueries())
-	{
-		RunHairStrandsBindingQueries(GraphBuilder, ShaderMap);
-	}
-
 #if WITH_EDITOR
 	if (HasHairCardsAtlasQueries())
 	{
@@ -809,8 +801,7 @@ void ProcessHairStrandsBookmark(
 			HasHairCardsAtlasQueries() ||
 			HasHairStrandsTexturesQueries() ||
 		#endif
-			HasHairStrandsFolliculeMaskQueries() ||
-			HasHairStrandsBindigQueries();
+			HasHairStrandsFolliculeMaskQueries();
 
 		if (bHasHairStardsnProcess)
 		{
diff --git a/Engine/Plugins/Runtime/HairStrands/Source/HairStrandsCore/Public/GroomBindingBuilder.h b/Engine/Plugins/Runtime/HairStrands/Source/HairStrandsCore/Public/GroomBindingBuilder.h
--- a/Engine/Plugins/Runtime/HairStrands/Source/HairStrandsCore/Public/GroomBindingBuilder.h
+++ b/Engine/Plugins/Runtime/HairStrands/Source/HairStrandsCore/Public/GroomBindingBuilder.h
@@ -17,22 +17,7 @@ struct HAIRSTRANDSCORE_API FGroomBindingBuilder
 	static FString GetVersion();
 
 	// Build binding asset data
-	static bool BuildBinding(class UGroomBindingAsset* BindingAsset, bool bUseGPU, bool bInitResources);
-
-	// Transfer mesh vertex position from a source mesh on to a target mesh.
-	// The topology of the two meshes can be different, but they need to shader the same UV layout
-	static void TransferMesh(
-		FRDGBuilder& GraphBuilder,
-		const FHairStrandsProjectionMeshData& SourceMeshData,
-		const FHairStrandsProjectionMeshData& TargetMeshData,
-		TArray<FRWBuffer>& TransferredLODsPositions);
-
-	// Project strands roots onto a mesh and compute closest triangle ID & barycentric
-	static void ProjectStrands(
-		FRDGBuilder& GraphBuilder,
-		const FTransform& LocalToWorld,
-		const FHairStrandsProjectionMeshData& TargetMeshData,
-		TArray<FHairStrandsRestRootResource*>& InRestRootResources);
+	static bool BuildBinding(class UGroomBindingAsset* BindingAsset, bool bInitResources);
 
 	// Update the root mesh projection data with unique valid mesh section IDs, based on the projection data
 	static void BuildUniqueSections(FHairStrandsRootData::FMeshProjectionLOD& LOD);
diff --git a/Engine/Plugins/Runtime/HairStrands/Source/HairStrandsEditor/Private/GroomBindingActions.cpp b/Engine/Plugins/Runtime/HairStrands/Source/HairStrandsEditor/Private/GroomBindingActions.cpp
--- a/Engine/Plugins/Runtime/HairStrands/Source/HairStrandsEditor/Private/GroomBindingActions.cpp
+++ b/Engine/Plugins/Runtime/HairStrands/Source/HairStrandsEditor/Private/GroomBindingActions.cpp
@@ -104,7 +104,7 @@ void FGroomBindingActions::ExecuteRebuildBindingAsset(TArray<TWeakObjectPtr<UGro
 					BindingAsset->SourceGeometryCache->ConditionalPostLoad();
 				}
 			}
-			FGroomBindingBuilder::BuildBinding(BindingAsset.Get(), false, true);
+			FGroomBindingBuilder::BuildBinding(BindingAsset.Get(), true);
 			BindingAsset->GetOutermost()->MarkPackageDirty();
 		}
 	}
