Fix crash when cooking groom data on platform not-supporting strands geometry.

This CL bypass the serialization of strands bulk data when a platform does not support strands geometry, instead of writing dummy data. The reasons behind this is that bulk data are serialized in two step: 1) the initial/regular serialization happens. At that time, the pointer on the bulk data is queued 2) then later the bulk data pointer is dequeued, and the bulk data are serialized. Because dummy resources where stack allocated, their memory was invalid by the time the second serialization step arises.

This CL fixes the issue by not serializating the bulkdata of these dummy resources.

diff --git a/Engine/Plugins/Runtime/HairStrands/Source/HairStrandsCore/Private/GroomBuilder.cpp b/Engine/Plugins/Runtime/HairStrands/Source/HairStrandsCore/Private/GroomBuilder.cpp
--- a/Engine/Plugins/Runtime/HairStrands/Source/HairStrandsCore/Private/GroomBuilder.cpp
+++ b/Engine/Plugins/Runtime/HairStrands/Source/HairStrandsCore/Private/GroomBuilder.cpp
@@ -33,7 +33,7 @@ static FAutoConsoleVariableRef CVarHairClusterBuilder_MaxVoxelResolution(TEXT("r
 FString FGroomBuilder::GetVersion()
 {
 	// Important to update the version when groom building changes
-	return TEXT("2e");
+	return TEXT("2f");
 }
 
 namespace FHairStrandsDecimation
@@ -346,7 +346,7 @@ namespace HairStrandsBuilder
 		OutBulkData.PointCount = HairStrands.GetNumPoints();
 		OutBulkData.MaxLength = HairStrands.StrandsCurves.MaxLength;
 		OutBulkData.MaxRadius = HairStrands.StrandsCurves.MaxRadius;
-		OutBulkData.Flags = 0;
+		OutBulkData.Flags = FHairStrandsBulkData::DataFlags_HasData;
 
 		if (bHasMaterialData)	{ OutBulkData.Flags |= FHairStrandsBulkData::DataFlags_HasMaterialData; }	else { OutPackedMaterials.Empty(); }
 		if (bHasUDIMData)		{ OutBulkData.Flags |= FHairStrandsBulkData::DataFlags_HasUDIMData; }		else { OutPackedAttributes1.Empty(); }
diff --git a/Engine/Plugins/Runtime/HairStrands/Source/HairStrandsCore/Private/HairStrandsDatas.cpp b/Engine/Plugins/Runtime/HairStrands/Source/HairStrandsCore/Private/HairStrandsDatas.cpp
--- a/Engine/Plugins/Runtime/HairStrands/Source/HairStrandsCore/Private/HairStrandsDatas.cpp
+++ b/Engine/Plugins/Runtime/HairStrands/Source/HairStrandsCore/Private/HairStrandsDatas.cpp
@@ -218,20 +218,23 @@ void FHairStrandsBulkData::Serialize(FArchive& Ar, UObject* Owner)
 	Materials.SetBulkDataFlags(BulkFlags);
 	CurveOffsets.SetBulkDataFlags(BulkFlags);
 
-	const int32 ChunkIndex = 0;
-	bool bAttemptFileMapping = false;
-
-	Positions.Serialize(Ar, Owner, ChunkIndex, bAttemptFileMapping);
-	Attributes0.Serialize(Ar, Owner, ChunkIndex, bAttemptFileMapping);
-	if (!!(Flags & DataFlags_HasUDIMData))
-	{
-		Attributes1.Serialize(Ar, Owner, ChunkIndex, bAttemptFileMapping);
-	}
-	if (!!(Flags & DataFlags_HasMaterialData))
+	if (!!(Flags & DataFlags_HasData))
 	{
-		Materials.Serialize(Ar, Owner, ChunkIndex, bAttemptFileMapping);
+		const int32 ChunkIndex = 0;
+		bool bAttemptFileMapping = false;
+
+		Positions.Serialize(Ar, Owner, ChunkIndex, bAttemptFileMapping);
+		Attributes0.Serialize(Ar, Owner, ChunkIndex, bAttemptFileMapping);
+		if (!!(Flags & DataFlags_HasUDIMData))
+		{
+			Attributes1.Serialize(Ar, Owner, ChunkIndex, bAttemptFileMapping);
+		}
+		if (!!(Flags & DataFlags_HasMaterialData))
+		{
+			Materials.Serialize(Ar, Owner, ChunkIndex, bAttemptFileMapping);
+		}
+		CurveOffsets.Serialize(Ar, Owner, ChunkIndex, bAttemptFileMapping);
 	}
-	CurveOffsets.Serialize(Ar, Owner, ChunkIndex, bAttemptFileMapping);
 }
 
 void FHairStrandsDatas::Reset()
diff --git a/Engine/Plugins/Runtime/HairStrands/Source/HairStrandsCore/Public/HairStrandsDatas.h b/Engine/Plugins/Runtime/HairStrands/Source/HairStrandsCore/Public/HairStrandsDatas.h
--- a/Engine/Plugins/Runtime/HairStrands/Source/HairStrandsCore/Public/HairStrandsDatas.h
+++ b/Engine/Plugins/Runtime/HairStrands/Source/HairStrandsCore/Public/HairStrandsDatas.h
@@ -424,8 +424,9 @@ struct HAIRSTRANDSCORE_API FHairStrandsBulkData
 {
 	enum EDataFlags
 	{
-		DataFlags_HasUDIMData = 1,
-		DataFlags_HasMaterialData = 2
+		DataFlags_HasData = 1,			// Contains valid data. Otherwise: Position, Attributes, ... are all empty
+		DataFlags_HasUDIMData = 2,		// Contains UDIM coords
+		DataFlags_HasMaterialData = 4	// Contains material data (albedo and/or roughness)
 	};
 
 	void Serialize(FArchive& Ar, UObject* Owner);
