Split Root data into data and bulk data.

This CL allows runtime data to be stored with byte bulk data which allows for partial/deferred loading (to be added).

diff --git a/Engine/Plugins/Runtime/HairStrands/Source/HairStrandsCore/Private/GroomAsset.cpp b/Engine/Plugins/Runtime/HairStrands/Source/HairStrandsCore/Private/GroomAsset.cpp
--- a/Engine/Plugins/Runtime/HairStrands/Source/HairStrandsCore/Private/GroomAsset.cpp
+++ b/Engine/Plugins/Runtime/HairStrands/Source/HairStrandsCore/Private/GroomAsset.cpp
@@ -235,11 +235,11 @@ void DumpLoadedGroomBindingAssets(IConsoleVariable* InCVarPakTesterEnabled)
 	{
 		if (AssetIt)
 		{			
-			const uint32 GroupCount = AssetIt->HairGroupDatas.Num();
+			const uint32 GroupCount = AssetIt->HairGroupBulkDatas.Num();
 			for (uint32 GroupIt = 0; GroupIt < GroupCount; ++GroupIt)
 			{
-				const uint32 CPUMemorySize_Guides  = AssetIt->HairGroupDatas[GroupIt].SimRootData.GetDataSize();
-				const uint32 CPUMemorySize_Strands = AssetIt->HairGroupDatas[GroupIt].RenRootData.GetDataSize();
+				const uint32 CPUMemorySize_Guides  = AssetIt->HairGroupBulkDatas[GroupIt].SimRootBulkData.GetDataSize();
+				const uint32 CPUMemorySize_Strands = AssetIt->HairGroupBulkDatas[GroupIt].RenRootBulkData.GetDataSize();
 
 				uint32 GPUMemorySize_Guides = 0;
 				if (const FHairStrandsRestRootResource* RootResource = AssetIt->HairGroupResources[GroupIt].SimRootResources)
@@ -254,10 +254,10 @@ void DumpLoadedGroomBindingAssets(IConsoleVariable* InCVarPakTesterEnabled)
 
 				uint32 CPUMemorySize_Cards = 0;
 				uint32 GPUMemorySize_Cards = 0;
-				const uint32 CardCount = AssetIt->HairGroupDatas[GroupIt].CardsRootData.Num();
+				const uint32 CardCount = AssetIt->HairGroupBulkDatas[GroupIt].CardsRootBulkData.Num();
 				for (uint32 CardIt = 0; CardIt < CardCount; ++CardIt)
 				{
-					CPUMemorySize_Cards += AssetIt->HairGroupDatas[GroupIt].CardsRootData[CardIt].GetDataSize();
+					CPUMemorySize_Cards += AssetIt->HairGroupBulkDatas[GroupIt].CardsRootBulkData[CardIt].GetDataSize();
 					if (const FHairStrandsRestRootResource* RootResource = AssetIt->HairGroupResources[GroupIt].CardsRootResources[CardIt])
 					{
 						GPUMemorySize_Cards += RootResource->GetResourcesSize();
@@ -275,7 +275,7 @@ void DumpLoadedGroomBindingAssets(IConsoleVariable* InCVarPakTesterEnabled)
 				Total_GPUMemorySize_Strands+= GPUMemorySize_Strands;
 				Total_GPUMemorySize_Cards  += GPUMemorySize_Cards;
 
-				const uint32 SkelLODCount = AssetIt->HairGroupDatas[GroupIt].RenRootData.MeshProjectionLODs.Num();
+				const uint32 SkelLODCount = AssetIt->HairGroupBulkDatas[GroupIt].RenRootBulkData.MeshProjectionLODs.Num();
 				if (bDetails)
 				{
 //					UE_LOG(LogHairStrands, Log, TEXT("--  No.  - LOD -    CPU Total (     Guides|    Strands|      Cards) -    GPU Total (     Guides|    Strands|      Cards) - Asset Name "));
diff --git a/Engine/Plugins/Runtime/HairStrands/Source/HairStrandsCore/Private/GroomBindingAsset.cpp b/Engine/Plugins/Runtime/HairStrands/Source/HairStrandsCore/Private/GroomBindingAsset.cpp
--- a/Engine/Plugins/Runtime/HairStrands/Source/HairStrandsCore/Private/GroomBindingAsset.cpp
+++ b/Engine/Plugins/Runtime/HairStrands/Source/HairStrandsCore/Private/GroomBindingAsset.cpp
@@ -20,24 +20,49 @@
 
 bool IsHairStrandsDDCLogEnable();
 
-FArchive& operator<<(FArchive& Ar, FHairStrandsRootData& RootData)
+////////////////////////////////////////////////////////////////////////////////////////////////////////////
+
+static void InternalSerialize(FArchive& Ar, UObject* Owner, TArray<FHairStrandsRootBulkData>& GroupDatas)
 {
-	RootData.Serialize(Ar);
-	return Ar;
+	uint32 GroupCount = GroupDatas.Num();
+	Ar << GroupCount;
+	if (Ar.IsLoading())
+	{
+		GroupDatas.SetNum(GroupCount);
+	}
+	for (uint32 GroupIt = 0; GroupIt < GroupCount; ++GroupIt)
+	{
+		GroupDatas[GroupIt].Serialize(Ar, Owner);
+	}
 }
 
-FArchive& operator<<(FArchive& Ar, UGroomBindingAsset::FHairGroupData& GroupData)
+static void InternalSerialize(FArchive& Ar, UObject* Owner, FHairGroupBulkData& GroupData)
 {
 	Ar.UsingCustomVersion(FAnimObjectVersion::GUID);
-	GroupData.SimRootData.Serialize(Ar);
-	GroupData.RenRootData.Serialize(Ar);
+	GroupData.SimRootBulkData.Serialize(Ar, Owner);
+	GroupData.RenRootBulkData.Serialize(Ar, Owner);
 	if (Ar.CustomVer(FAnimObjectVersion::GUID) >= FAnimObjectVersion::SerializeHairBindingAsset)
 	{
-		Ar << GroupData.CardsRootData;
+		InternalSerialize(Ar, Owner, GroupData.CardsRootBulkData);
 	}
-	return Ar;
 }
 
+static void InternalSerialize(FArchive& Ar, UObject* Owner, TArray<FHairGroupBulkData>& GroupDatas)
+{
+	uint32 GroupCount = GroupDatas.Num();
+	Ar << GroupCount;
+	if (Ar.IsLoading())
+	{
+		GroupDatas.SetNum(GroupCount);
+	}
+	for (uint32 GroupIt = 0; GroupIt < GroupCount; ++GroupIt)
+	{
+		InternalSerialize(Ar, Owner, GroupDatas[GroupIt]);
+	}
+}
+
+////////////////////////////////////////////////////////////////////////////////////////////////////////////
+
 void UGroomBindingAsset::Serialize(FArchive& Ar)
 {
 #if WITH_EDITOR
@@ -58,21 +83,21 @@ void UGroomBindingAsset::Serialize(FArchive& Ar)
 #if WITH_EDITOR
 		if (bIsStrandsStrippedForCook)
 		{
-			FHairGroupDatas StrippedHairGroupDatas;
-			StrippedHairGroupDatas.SetNum(HairGroupDatas.Num());
-			for (int32 Index = 0; Index < HairGroupDatas.Num(); ++Index)
+			TArray<FHairGroupBulkData> StrippedHairGroupDatas;
+			StrippedHairGroupDatas.SetNum(HairGroupBulkDatas.Num());
+			for (int32 Index = 0; Index < HairGroupBulkDatas.Num(); ++Index)
 			{
-				FHairGroupData& HairGroupData = HairGroupDatas[Index];
-				StrippedHairGroupDatas[Index].SimRootData = HairGroupData.SimRootData;
-				StrippedHairGroupDatas[Index].CardsRootData = HairGroupData.CardsRootData;
+				FHairGroupBulkData& HairGroupBulkData = HairGroupBulkDatas[Index];
+				StrippedHairGroupDatas[Index].SimRootBulkData = HairGroupBulkData.SimRootBulkData;
+				StrippedHairGroupDatas[Index].CardsRootBulkData = HairGroupBulkData.CardsRootBulkData;
 			}
 
-			Ar << StrippedHairGroupDatas;
+			InternalSerialize(Ar, this, StrippedHairGroupDatas);
 		}
 		else
 #endif
 		{
-			Ar << HairGroupDatas;
+			InternalSerialize(Ar, this, HairGroupBulkDatas);
 		}
 		bIsValid = true;
 	}
@@ -82,38 +107,38 @@ void UGroomBindingAsset::InitResource()
 {
 	LLM_SCOPE(ELLMTag::Meshes) // This should be a Groom LLM tag, but there is no LLM tag bit left
 
-	for (FHairGroupData& Data : HairGroupDatas)
+	for (FHairGroupBulkData& BulkData : HairGroupBulkDatas)
 	{
 		const int32 GroupIndex = HairGroupResources.Num();
 		FHairGroupResource& Resource = HairGroupResources.AddDefaulted_GetRef();
 
 		// Guides
 		Resource.SimRootResources = nullptr;
-		if (Data.SimRootData.IsValid())
+		if (BulkData.SimRootBulkData.IsValid())
 		{
-			Resource.SimRootResources = new FHairStrandsRestRootResource(Data.SimRootData, EHairStrandsResourcesType::Guides);
+			Resource.SimRootResources = new FHairStrandsRestRootResource(BulkData.SimRootBulkData, EHairStrandsResourcesType::Guides);
 			BeginInitResource(Resource.SimRootResources);
 		}
 
 		// Strands
 		Resource.RenRootResources = nullptr;
-		if (IsHairStrandsEnabled(EHairStrandsShaderType::Strands) && Data.RenRootData.IsValid())
+		if (IsHairStrandsEnabled(EHairStrandsShaderType::Strands) && BulkData.RenRootBulkData.IsValid())
 		{
-			Resource.RenRootResources = new FHairStrandsRestRootResource(Data.RenRootData, EHairStrandsResourcesType::Strands);
+			Resource.RenRootResources = new FHairStrandsRestRootResource(BulkData.RenRootBulkData, EHairStrandsResourcesType::Strands);
 			BeginInitResource(Resource.RenRootResources);
 		}
 
 		// Cards
 		if (IsHairStrandsEnabled(EHairStrandsShaderType::Cards))
 		{
-			const uint32 CardsLODCount = Data.CardsRootData.Num();
+			const uint32 CardsLODCount = BulkData.CardsRootBulkData.Num();
 			Resource.CardsRootResources.SetNum(CardsLODCount);
 			for (uint32 CardsLODIt=0; CardsLODIt<CardsLODCount; ++CardsLODIt)
 			{
 				Resource.CardsRootResources[CardsLODIt] = nullptr;
-				if (Data.CardsRootData[CardsLODIt].IsValid())
+				if (BulkData.CardsRootBulkData[CardsLODIt].IsValid())
 				{
-					Resource.CardsRootResources[CardsLODIt] = new FHairStrandsRestRootResource(Data.CardsRootData[CardsLODIt], EHairStrandsResourcesType::Cards);
+					Resource.CardsRootResources[CardsLODIt] = new FHairStrandsRestRootResource(BulkData.CardsRootBulkData[CardsLODIt], EHairStrandsResourcesType::Cards);
 					BeginInitResource(Resource.CardsRootResources[CardsLODIt]);
 				}
 			}
@@ -215,14 +240,14 @@ void UGroomBindingAsset::ReleaseResource()
 void UGroomBindingAsset::Reset()
 {
 	ReleaseResource();
-	for (FHairGroupData& Data : HairGroupDatas)
+	for (FHairGroupBulkData& Data : HairGroupBulkDatas)
 	{
-		Data.SimRootData.Reset();
-		Data.RenRootData.Reset();
+		Data.SimRootBulkData.Reset();
+		Data.RenRootBulkData.Reset();
 
-		for (FHairStrandsRootData& CardsRootData : Data.CardsRootData)
+		for (FHairStrandsRootBulkData& CardsRootBulkData : Data.CardsRootBulkData)
 		{
-			CardsRootData.Reset();
+			CardsRootBulkData.Reset();
 		}
 	}
 
@@ -391,26 +416,26 @@ bool UGroomBindingAsset::IsCompatible(const USkeletalMesh* InSkeletalMesh, const
 
 		for (const FHairGroupResource& Resource : InBinding->HairGroupResources)
 		{
-			if (Resource.SimRootResources && InSkeletalMesh->GetLODNum() != Resource.SimRootResources->RootData.MeshProjectionLODs.Num())
+			if (Resource.SimRootResources && InSkeletalMesh->GetLODNum() != Resource.SimRootResources->BulkData.MeshProjectionLODs.Num())
 			{
 				if (bIssueWarning)
 				{
 					UE_LOG(LogHairStrands, Warning, TEXT("[Groom] The Groom binding (%s) does not have the same have LOD count (LOD sim:%d) than the skeletal mesh (%s, LOD:%d). The binding asset will not be used."),
 						*InBinding->GetName(),
-						Resource.SimRootResources->RootData.MeshProjectionLODs.Num(),
+						Resource.SimRootResources->BulkData.MeshProjectionLODs.Num(),
 						*InSkeletalMesh->GetName(),
 						InSkeletalMesh->GetLODNum());
 				}
 				return false;
 			}
 
-			if (Resource.RenRootResources && InSkeletalMesh->GetLODNum() != Resource.RenRootResources->RootData.MeshProjectionLODs.Num() && IsHairStrandsEnabled(EHairStrandsShaderType::Strands))
+			if (Resource.RenRootResources && InSkeletalMesh->GetLODNum() != Resource.RenRootResources->BulkData.MeshProjectionLODs.Num() && IsHairStrandsEnabled(EHairStrandsShaderType::Strands))
 			{
 				if (bIssueWarning)
 				{
 					UE_LOG(LogHairStrands, Warning, TEXT("[Groom] The Groom binding (%s) does not have the same have LOD count (LOD render:%d) than the skeletal mesh (%s, LOD:%d). The binding asset will not be used."),
 						*InBinding->GetName(),
-						Resource.RenRootResources->RootData.MeshProjectionLODs.Num(),
+						Resource.RenRootResources->BulkData.MeshProjectionLODs.Num(),
 						*InSkeletalMesh->GetName(),
 						InSkeletalMesh->GetLODNum());
 				}
@@ -641,7 +666,7 @@ void UGroomBindingAsset::PostEditChangeProperty(FPropertyChangedEvent& PropertyC
 // differences, etc.) replace the version GUID below with a new one.
 // In case of merge conflicts with DDC versions, you MUST generate a new GUID
 // and set this new GUID as the version.
-#define GROOM_BINDING_DERIVED_DATA_VERSION TEXT("C931E6356B4541B1BB3AE6FCE7C45DB7")
+#define GROOM_BINDING_DERIVED_DATA_VERSION TEXT("5D1E46FE8AF848E792FFD94F6F00CB10")
 
 namespace GroomBindingDerivedDataCacheUtils
 {
@@ -720,7 +745,7 @@ void UGroomBindingAsset::CacheDerivedDatas()
 			Ar.SerializeCompressed(DecompressionBuffer, 0, NAME_Zlib);
 
 			FLargeMemoryReader LargeMemReader(DecompressionBuffer, UncompressedSize, ELargeMemoryReaderFlags::Persistent | ELargeMemoryReaderFlags::TakeOwnership);
-			LargeMemReader << HairGroupDatas;
+			InternalSerialize(LargeMemReader, this, HairGroupBulkDatas);
 
 			bIsValid = true;
 		}
@@ -737,7 +762,7 @@ void UGroomBindingAsset::CacheDerivedDatas()
 			{
 				// Using a LargeMemoryWriter for serialization since the data can be bigger than 2 GB
 				FLargeMemoryWriter LargeMemWriter(0, /*bIsPersistent=*/ true);
-				LargeMemWriter << HairGroupDatas;
+				InternalSerialize(LargeMemWriter, this, HairGroupBulkDatas);
 
 				int64 UncompressedSize = LargeMemWriter.TotalSize();
 
@@ -755,7 +780,7 @@ void UGroomBindingAsset::CacheDerivedDatas()
 		}
 
 		// Patch hair group info if it does not match the DDC-read/deserialized data
-		const uint32 GroupCount = HairGroupDatas.Num();
+		const uint32 GroupCount = HairGroupBulkDatas.Num();
 		if (GroupInfos.Num() != GroupCount)
 		{
 			GroupInfos.SetNum(GroupCount);
@@ -763,12 +788,12 @@ void UGroomBindingAsset::CacheDerivedDatas()
 		for (uint32 GroupIt=0; GroupIt< GroupCount; ++GroupIt)
 		{
 			FGoomBindingGroupInfo& Info = GroupInfos[GroupIt];
-			const UGroomBindingAsset::FHairGroupData& Data = HairGroupDatas[GroupIt];
+			const FHairGroupBulkData& Data = HairGroupBulkDatas[GroupIt];
 			{
-				Info.SimRootCount = Data.SimRootData.RootCount;
-				Info.SimLODCount  = Data.SimRootData.MeshProjectionLODs.Num();
-				Info.RenRootCount = Data.RenRootData.RootCount;
-				Info.RenLODCount  = Data.RenRootData.MeshProjectionLODs.Num();
+				Info.SimRootCount = Data.SimRootBulkData.RootCount;
+				Info.SimLODCount  = Data.SimRootBulkData.MeshProjectionLODs.Num();
+				Info.RenRootCount = Data.RenRootBulkData.RootCount;
+				Info.RenLODCount  = Data.RenRootBulkData.MeshProjectionLODs.Num();
 			}
 		}
 
@@ -787,7 +812,7 @@ void UGroomBindingAsset::CacheDerivedDatas()
 
 void UGroomBindingAsset::GetResourceSizeEx(FResourceSizeEx& CumulativeResourceSize)
 {
-	CumulativeResourceSize.AddDedicatedSystemMemoryBytes(HairGroupDatas.GetAllocatedSize());
+	CumulativeResourceSize.AddDedicatedSystemMemoryBytes(HairGroupBulkDatas.GetAllocatedSize());
 
 	for (const FHairGroupResource& Group : HairGroupResources)
 	{
diff --git a/Engine/Plugins/Runtime/HairStrands/Source/HairStrandsCore/Private/GroomBindingBuilder.cpp b/Engine/Plugins/Runtime/HairStrands/Source/HairStrandsCore/Private/GroomBindingBuilder.cpp
--- a/Engine/Plugins/Runtime/HairStrands/Source/HairStrandsCore/Private/GroomBindingBuilder.cpp
+++ b/Engine/Plugins/Runtime/HairStrands/Source/HairStrandsCore/Private/GroomBindingBuilder.cpp
@@ -57,7 +57,7 @@ static FAutoConsoleVariableRef CVarHairStrandsBindingBuilderWarningEnable(TEXT("
 FString FGroomBindingBuilder::GetVersion()
 {
 	// Important to update the version when groom building changes
-	return TEXT("1ce");
+	return TEXT("2b");
 }
 
 ///////////////////////////////////////////////////////////////////////////////////////////////////
@@ -102,22 +102,6 @@ float FHairStrandsRootUtils::PackUVsToFloat(const FVector2D& UV)
 	return *((float*)(&Encoded));
 }
 
-// Mirror the group info into the GroupInfos struct. This is used for groom binding validation at runtime
-static void UpdateGroupInfos(UGroomBindingAsset* BindingAsset)
-{
-	UGroomBindingAsset::FHairGroupDatas& OutHairGroupDatas = BindingAsset->HairGroupDatas;
-	TArray<FGoomBindingGroupInfo>& OutGroupInfos = BindingAsset->GroupInfos;
-	OutGroupInfos.Empty();
-	for (const UGroomBindingAsset::FHairGroupData& Data : OutHairGroupDatas)
-	{
-		FGoomBindingGroupInfo& Info = OutGroupInfos.AddDefaulted_GetRef();
-		Info.SimRootCount = Data.SimRootData.RootCount;
-		Info.SimLODCount = Data.SimRootData.MeshProjectionLODs.Num();
-		Info.RenRootCount = Data.RenRootData.RootCount;
-		Info.RenLODCount = Data.RenRootData.MeshProjectionLODs.Num();
-	}
-}
-
 namespace
 {
 //////////////////////////////////////////////////////////////////////////
@@ -723,15 +707,16 @@ namespace GroomBinding_RBFWeighting
 		}
 	}
 
-	void ComputeInterpolationWeights(UGroomBindingAsset* BindingAsset, IMeshData* MeshData, TArray<TArray<FVector3f>>& TransferedPositions)
+	void ComputeInterpolationWeights(
+		TArray<FHairRootGroupData>& Out, 
+		const uint32 NumInterpolationPoints, 
+		const int32 MatchingSection, 
+		const IMeshData* MeshData, 
+		const TArray<TArray<FVector3f>>& TransferedPositions)
 	{
-		UGroomAsset* GroomAsset = BindingAsset->Groom;
-
-		UGroomBindingAsset::FHairGroupDatas& OutHairGroupDatas = BindingAsset->HairGroupDatas;
-
-		const uint32 GroupCount  = OutHairGroupDatas.Num();
+		const uint32 GroupCount  = Out.Num();
 		const uint32 MeshLODCount= MeshData->GetNumLODs();
-		const uint32 MaxSamples  = BindingAsset->NumInterpolationPoints;
+		const uint32 MaxSamples  = NumInterpolationPoints;
 
 		for (uint32 LODIndex = 0; LODIndex < MeshLODCount; ++LODIndex)
 		{
@@ -744,7 +729,7 @@ namespace GroomBinding_RBFWeighting
 			{
 				PositionsPointer = TransferedPositions[LODIndex].GetData();
 				GlobalSamples = true;
-				TargetSection = BindingAsset->MatchingSection;
+				TargetSection = MatchingSection;
 			}
 			else
 			{
@@ -756,7 +741,7 @@ namespace GroomBinding_RBFWeighting
 				TArray<bool> ValidPoints;
 				for (uint32 GroupIt = 0; GroupIt < GroupCount; ++GroupIt)
 				{
-					FillLocalValidPoints(MeshLODData, TargetSection, OutHairGroupDatas[GroupIt].SimRootData.MeshProjectionLODs[LODIndex], ValidPoints);
+					FillLocalValidPoints(MeshLODData, TargetSection, Out[GroupIt].SimRootData.MeshProjectionLODs[LODIndex], ValidPoints);
 
 					FPointsSampler PointsSampler(ValidPoints, PositionsPointer, MaxSamples);
 					const uint32 SampleCount = PointsSampler.SamplePositions.Num();
@@ -764,8 +749,8 @@ namespace GroomBinding_RBFWeighting
 					FWeightsBuilder InterpolationWeights(SampleCount, SampleCount,
 						PointsSampler.SamplePositions.GetData(), PointsSampler.SamplePositions.GetData());
 
-					UpdateInterpolationWeights(InterpolationWeights, PointsSampler, LODIndex, OutHairGroupDatas[GroupIt].SimRootData);
-					UpdateInterpolationWeights(InterpolationWeights, PointsSampler, LODIndex, OutHairGroupDatas[GroupIt].RenRootData);
+					UpdateInterpolationWeights(InterpolationWeights, PointsSampler, LODIndex, Out[GroupIt].SimRootData);
+					UpdateInterpolationWeights(InterpolationWeights, PointsSampler, LODIndex, Out[GroupIt].RenRootData);
 				}
 			}
 			else
@@ -782,8 +767,8 @@ namespace GroomBinding_RBFWeighting
 
 				for (uint32 GroupIt = 0; GroupIt < GroupCount; ++GroupIt)
 				{
-					UpdateInterpolationWeights(InterpolationWeights, PointsSampler, LODIndex, OutHairGroupDatas[GroupIt].SimRootData);
-					UpdateInterpolationWeights(InterpolationWeights, PointsSampler, LODIndex, OutHairGroupDatas[GroupIt].RenRootData);
+					UpdateInterpolationWeights(InterpolationWeights, PointsSampler, LODIndex, Out[GroupIt].SimRootData);
+					UpdateInterpolationWeights(InterpolationWeights, PointsSampler, LODIndex, Out[GroupIt].RenRootData);
 				}
 			}
 		}
@@ -1779,11 +1764,12 @@ namespace GroomBinding_Transfer
 }
 // namespace GroomBinding_Transfer
 
-static void InitHairStrandsRootData(FHairStrandsRootData& Out, const FHairStrandsDatas* HairStrandsDatas, uint32 LODCount, const TArray<uint32>& NumSamples)
+static void InitHairStrandsRootData(FHairStrandsRootData& Out, const FHairStrandsDatas* HairStrandsDatas, uint32 LODCount, uint32 NumSamples)
 {
 	check(HairStrandsDatas);
 
 	Out.RootCount = HairStrandsDatas->GetNumCurves();
+	Out.PointCount = HairStrandsDatas->GetNumPoints();
 
 	const uint32 CurveCount = HairStrandsDatas->GetNumCurves();
 	Out.VertexToCurveIndexBuffer.SetNum(HairStrandsDatas->GetNumPoints());
@@ -1797,13 +1783,12 @@ static void InitHairStrandsRootData(FHairStrandsRootData& Out, const FHairStrand
 			Out.VertexToCurveIndexBuffer[RootIndex + PointIndex] = CurveIndex; // RootIndex;
 		}
 	}
-	check(NumSamples.Num() == LODCount);
 
 	Out.MeshProjectionLODs.SetNum(LODCount);
 	uint32 LODIndex = 0;
 	for (FHairStrandsRootData::FMeshProjectionLOD& MeshProjectionLOD : Out.MeshProjectionLODs)
 	{
-		MeshProjectionLOD.SampleCount = NumSamples[LODIndex];
+		MeshProjectionLOD.SampleCount = NumSamples;
 		MeshProjectionLOD.LODIndex = LODIndex++;
 		MeshProjectionLOD.MeshInterpolationWeightsBuffer.Empty();
 		MeshProjectionLOD.MeshSampleIndicesBuffer.Empty();
@@ -1812,207 +1797,356 @@ static void InitHairStrandsRootData(FHairStrandsRootData& Out, const FHairStrand
 }
 
 ///////////////////////////////////////////////////////////////////////////////////////////////////
-// Main entry (CPU path)
-static bool InternalBuildBinding_CPU(UGroomBindingAsset* BindingAsset, bool bInitResources)
+// Convert data into bulk data
+
+namespace GroomBinding_BulkCopy
 {
-#if WITH_EDITORONLY_DATA
-	if (!BindingAsset ||
-		!BindingAsset->Groom ||
-		!BindingAsset->HasValidTarget() ||
-		BindingAsset->Groom->GetNumHairGroups() == 0)
-	{
-		UE_LOG(LogHairStrands, Error, TEXT("[Groom] Binding asset cannot be created/rebuilt."));
-		return false;
-	}
 
-	BindingAsset->Groom->ConditionalPostLoad();
+template<typename TFormatType>
+void CopyToBulkData(FByteBulkData& Out, const TArray<typename TFormatType::Type>& Data)
+{
+	const uint32 DataSizeInByte = Data.Num() * sizeof(typename TFormatType::BulkType);
+
+	// The buffer is then stored into bulk data
+	Out.Lock(LOCK_READ_WRITE);
+	void* BulkBuffer = Out.Realloc(DataSizeInByte);
+	FMemory::Memcpy(BulkBuffer, Data.GetData(), DataSizeInByte);
+	Out.Unlock();
+}
+
+template<typename TFormatType>
+void CopyFromBulkData(TArray<typename TFormatType::Type>& Out, const FByteBulkData& In)
+{
+	const uint32 InDataSize = In.GetBulkDataSize();
+	const uint32 ElementCount = InDataSize / sizeof(typename TFormatType::BulkType);
+	Out.SetNum(ElementCount);
+
+	// The buffer is then stored into bulk data
+	const uint8* InData = (const uint8*)In.LockReadOnly();
+	FMemory::Memcpy(Out.GetData(), InData, InDataSize);
+	In.Unlock();
+}
+
+// Convert "root data" -> "root bulk data"
+static void BuildRootBulkData(
+	FHairStrandsRootBulkData& Out,
+	const FHairStrandsRootData& In)
+{
+	Out.RootCount = In.RootCount;
+	Out.PointCount = In.PointCount;
 
-	const int32 NumInterpolationPoints = BindingAsset->NumInterpolationPoints;
-	UGroomAsset* GroomAsset = BindingAsset->Groom;
+	CopyToBulkData<FHairStrandsIndexFormat>(Out.VertexToCurveIndexBuffer, In.VertexToCurveIndexBuffer);
 
-	TUniquePtr<IMeshData> SourceMeshData;
-	TUniquePtr<IMeshData> TargetMeshData;
-	if (BindingAsset->GroomBindingType == EGroomBindingMeshType::SkeletalMesh)
+	const uint32 MeshLODCount = In.MeshProjectionLODs.Num();
+	Out.MeshProjectionLODs.SetNum(MeshLODCount);
+	for (uint32 MeshLODIt = 0; MeshLODIt < MeshLODCount; ++MeshLODIt)
 	{
-		BindingAsset->TargetSkeletalMesh->ConditionalPostLoad();
-		if (BindingAsset->SourceSkeletalMesh)
+		const bool bHasValidSamples =
+			In.MeshProjectionLODs[MeshLODIt].MeshInterpolationWeightsBuffer.Num() > 0 &&
+			In.MeshProjectionLODs[MeshLODIt].MeshSampleIndicesBuffer.Num() > 0 &&
+			In.MeshProjectionLODs[MeshLODIt].RestSamplePositionsBuffer.Num() > 0;
+
+		Out.MeshProjectionLODs[MeshLODIt].LODIndex = In.MeshProjectionLODs[MeshLODIt].LODIndex;
+		Out.MeshProjectionLODs[MeshLODIt].SampleCount = bHasValidSamples ? In.MeshProjectionLODs[MeshLODIt].SampleCount : 0u;
+
+		CopyToBulkData<FHairStrandsCurveTriangleIndexFormat>(Out.MeshProjectionLODs[MeshLODIt].RootTriangleIndexBuffer, In.MeshProjectionLODs[MeshLODIt].RootTriangleIndexBuffer);
+		CopyToBulkData<FHairStrandsCurveTriangleBarycentricFormat>(Out.MeshProjectionLODs[MeshLODIt].RootTriangleBarycentricBuffer, In.MeshProjectionLODs[MeshLODIt].RootTriangleBarycentricBuffer);
+
+		CopyToBulkData<FHairStrandsMeshTrianglePositionFormat>(Out.MeshProjectionLODs[MeshLODIt].RestRootTrianglePosition0Buffer, In.MeshProjectionLODs[MeshLODIt].RestRootTrianglePosition0Buffer);
+		CopyToBulkData<FHairStrandsMeshTrianglePositionFormat>(Out.MeshProjectionLODs[MeshLODIt].RestRootTrianglePosition1Buffer, In.MeshProjectionLODs[MeshLODIt].RestRootTrianglePosition1Buffer);
+		CopyToBulkData<FHairStrandsMeshTrianglePositionFormat>(Out.MeshProjectionLODs[MeshLODIt].RestRootTrianglePosition2Buffer, In.MeshProjectionLODs[MeshLODIt].RestRootTrianglePosition2Buffer);
+
+		CopyToBulkData<FHairStrandsIndexFormat>(Out.MeshProjectionLODs[MeshLODIt].ValidSectionIndices, In.MeshProjectionLODs[MeshLODIt].ValidSectionIndices);
+
+		if (bHasValidSamples)
 		{
-			BindingAsset->SourceSkeletalMesh->ConditionalPostLoad();
+			CopyToBulkData<FHairStrandsWeightFormat>(Out.MeshProjectionLODs[MeshLODIt].MeshInterpolationWeightsBuffer, In.MeshProjectionLODs[MeshLODIt].MeshInterpolationWeightsBuffer);
+			CopyToBulkData<FHairStrandsIndexFormat>(Out.MeshProjectionLODs[MeshLODIt].MeshSampleIndicesBuffer, In.MeshProjectionLODs[MeshLODIt].MeshSampleIndicesBuffer);
+			CopyToBulkData<FHairStrandsMeshTrianglePositionFormat>(Out.MeshProjectionLODs[MeshLODIt].RestSamplePositionsBuffer, In.MeshProjectionLODs[MeshLODIt].RestSamplePositionsBuffer);
 		}
-
-		SourceMeshData = TUniquePtr<FSkeletalMeshData, TDefaultDelete<IMeshData>>(new FSkeletalMeshData(BindingAsset->SourceSkeletalMesh));
-		TargetMeshData = TUniquePtr<FSkeletalMeshData, TDefaultDelete<IMeshData>>(new FSkeletalMeshData(BindingAsset->TargetSkeletalMesh));
-	}
-	else
-	{
-		BindingAsset->TargetGeometryCache->ConditionalPostLoad();
-		if (BindingAsset->SourceGeometryCache)
+		else
 		{
-			BindingAsset->SourceGeometryCache->ConditionalPostLoad();
+			Out.MeshProjectionLODs[MeshLODIt].MeshInterpolationWeightsBuffer.RemoveBulkData();
+			Out.MeshProjectionLODs[MeshLODIt].MeshSampleIndicesBuffer.RemoveBulkData();
+			Out.MeshProjectionLODs[MeshLODIt].RestSamplePositionsBuffer.RemoveBulkData();
 		}
-
-		SourceMeshData = TUniquePtr<FGeometryCacheData, TDefaultDelete<IMeshData>>(new FGeometryCacheData(BindingAsset->SourceGeometryCache));
-		TargetMeshData = TUniquePtr<FGeometryCacheData, TDefaultDelete<IMeshData>>(new FGeometryCacheData(BindingAsset->TargetGeometryCache));
 	}
+}
 
-	if (!TargetMeshData->IsValid())
+// Convert "root data" <- "root bulk data"
+static void BuildRootData(
+	FHairStrandsRootData& Out,
+	const FHairStrandsRootBulkData& In)
+{
+	Out.RootCount = In.RootCount;
+	Out.PointCount = In.PointCount;
+
+	CopyFromBulkData<FHairStrandsIndexFormat>(Out.VertexToCurveIndexBuffer, In.VertexToCurveIndexBuffer);
+
+	const uint32 MeshLODCount = In.MeshProjectionLODs.Num();
+	Out.MeshProjectionLODs.SetNum(MeshLODCount);
+	for (uint32 MeshLODIt = 0; MeshLODIt < MeshLODCount; ++MeshLODIt)
 	{
-		UE_LOG(LogHairStrands, Error, TEXT("[Groom] Binding asset could not be built. Target mesh is not valid."));
-		return false;
-	}
+		const bool bHasValidSamples = In.MeshProjectionLODs[MeshLODIt].SampleCount > 0;
 
-	const uint32 GroupCount = GroomAsset->GetNumHairGroups();
+		Out.MeshProjectionLODs[MeshLODIt].LODIndex = In.MeshProjectionLODs[MeshLODIt].LODIndex;
+		Out.MeshProjectionLODs[MeshLODIt].SampleCount = bHasValidSamples ? In.MeshProjectionLODs[MeshLODIt].SampleCount : 0u;
 
-	const uint32 MeshLODCount = TargetMeshData->GetNumLODs();
-	UGroomBindingAsset::FHairGroupDatas& OutHairGroupDatas = BindingAsset->HairGroupDatas;
-	OutHairGroupDatas.Empty();
+		CopyFromBulkData<FHairStrandsCurveTriangleIndexFormat>(Out.MeshProjectionLODs[MeshLODIt].RootTriangleIndexBuffer, In.MeshProjectionLODs[MeshLODIt].RootTriangleIndexBuffer);
+		CopyFromBulkData<FHairStrandsCurveTriangleBarycentricFormat>(Out.MeshProjectionLODs[MeshLODIt].RootTriangleBarycentricBuffer, In.MeshProjectionLODs[MeshLODIt].RootTriangleBarycentricBuffer);
 
-	uint32 GroupIndex = 0;
-	TArray<uint32> NumSamples;
-	NumSamples.Init(NumInterpolationPoints, MeshLODCount);
-	for (const FHairGroupData& GroupData : GroomAsset->HairGroupsData)
-	{
-		FHairStrandsDatas StrandsData;
-		FHairStrandsDatas GuidesData;
-		GroomAsset->GetHairStrandsDatas(GroupIndex, StrandsData, GuidesData);
+		CopyFromBulkData<FHairStrandsMeshTrianglePositionFormat>(Out.MeshProjectionLODs[MeshLODIt].RestRootTrianglePosition0Buffer, In.MeshProjectionLODs[MeshLODIt].RestRootTrianglePosition0Buffer);
+		CopyFromBulkData<FHairStrandsMeshTrianglePositionFormat>(Out.MeshProjectionLODs[MeshLODIt].RestRootTrianglePosition1Buffer, In.MeshProjectionLODs[MeshLODIt].RestRootTrianglePosition1Buffer);
+		CopyFromBulkData<FHairStrandsMeshTrianglePositionFormat>(Out.MeshProjectionLODs[MeshLODIt].RestRootTrianglePosition2Buffer, In.MeshProjectionLODs[MeshLODIt].RestRootTrianglePosition2Buffer);
 
-		UGroomBindingAsset::FHairGroupData& Data = OutHairGroupDatas.AddDefaulted_GetRef();
-		InitHairStrandsRootData(Data.RenRootData, &StrandsData, MeshLODCount, NumSamples);
-		InitHairStrandsRootData(Data.SimRootData, &GuidesData, MeshLODCount, NumSamples);
+		CopyFromBulkData<FHairStrandsIndexFormat>(Out.MeshProjectionLODs[MeshLODIt].ValidSectionIndices, In.MeshProjectionLODs[MeshLODIt].ValidSectionIndices);
 
-		const uint32 CardsLODCount = GroupData.Cards.LODs.Num();
-		Data.CardsRootData.SetNum(GroupData.Cards.LODs.Num());
-		for (uint32 CardsLODIt = 0; CardsLODIt < CardsLODCount; ++CardsLODIt)
+		if (bHasValidSamples)
 		{
-			if (GroupData.Cards.IsValid(CardsLODIt))
-			{
-				FHairStrandsDatas LODGuidesData;
-				const bool bIsValid = GroomAsset->GetHairCardsGuidesDatas(GroupIndex, CardsLODIt, LODGuidesData);
-				check(bIsValid);
-				InitHairStrandsRootData(Data.CardsRootData[CardsLODIt], &LODGuidesData, MeshLODCount, NumSamples);
-			}
+			CopyFromBulkData<FHairStrandsWeightFormat>(Out.MeshProjectionLODs[MeshLODIt].MeshInterpolationWeightsBuffer, In.MeshProjectionLODs[MeshLODIt].MeshInterpolationWeightsBuffer);
+			CopyFromBulkData<FHairStrandsIndexFormat>(Out.MeshProjectionLODs[MeshLODIt].MeshSampleIndicesBuffer, In.MeshProjectionLODs[MeshLODIt].MeshSampleIndicesBuffer);
+			CopyFromBulkData<FHairStrandsMeshTrianglePositionFormat>(Out.MeshProjectionLODs[MeshLODIt].RestSamplePositionsBuffer, In.MeshProjectionLODs[MeshLODIt].RestSamplePositionsBuffer);
+		}
+		else
+		{
+			Out.MeshProjectionLODs[MeshLODIt].MeshInterpolationWeightsBuffer.Empty();
+			Out.MeshProjectionLODs[MeshLODIt].MeshSampleIndicesBuffer.Empty();
+			Out.MeshProjectionLODs[MeshLODIt].RestSamplePositionsBuffer.Empty();
 		}
-		++GroupIndex;
 	}
+}
 
-	UGroomBindingAsset::FHairGroupResources& OutHairGroupResources = BindingAsset->HairGroupResources;
-	if (BindingAsset->HairGroupResources.Num() > 0)
+// Convert the root data into root bulk data
+static void BuildRootBulkData(
+	TArray<FHairGroupBulkData>& Out,
+	const TArray<FHairRootGroupData>& In)
+{
+	const uint32 GroupCount = In.Num();
+	Out.SetNum(GroupCount);
+	for (uint32 GroupIt = 0; GroupIt < GroupCount; ++GroupIt)
 	{
-		for (UGroomBindingAsset::FHairGroupResource& GroupResrouces : OutHairGroupResources)
+		BuildRootBulkData(Out[GroupIt].SimRootBulkData, In[GroupIt].SimRootData);
+		BuildRootBulkData(Out[GroupIt].RenRootBulkData, In[GroupIt].RenRootData);
+
+		const uint32 LODCount = In[GroupIt].CardsRootData.Num();
+		Out[GroupIt].CardsRootBulkData.SetNum(LODCount);
+		for (uint32 LODIt = 0; LODIt < LODCount; ++LODIt)
 		{
-			BindingAsset->HairGroupResourcesToDelete.Enqueue(GroupResrouces);
+			BuildRootBulkData(Out[GroupIt].CardsRootBulkData[LODIt], In[GroupIt].CardsRootData[LODIt]);
 		}
-		OutHairGroupResources.Empty();
 	}
+}
 
-	check(OutHairGroupResources.Num() == 0);
-
-	TArray<FGoomBindingGroupInfo>& OutGroupInfos = BindingAsset->GroupInfos;
-	OutGroupInfos.Empty();
-	for (const UGroomBindingAsset::FHairGroupData& Data : OutHairGroupDatas)
-	{
-		FGoomBindingGroupInfo& Info = OutGroupInfos.AddDefaulted_GetRef();
-		Info.SimRootCount = Data.SimRootData.RootCount;
-		Info.SimLODCount  = Data.SimRootData.MeshProjectionLODs.Num();
-		Info.RenRootCount = Data.RenRootData.RootCount;
-		Info.RenLODCount  = Data.RenRootData.MeshProjectionLODs.Num();
-	}	
-	const bool bNeedTransferPosition = SourceMeshData->IsValid();
+} // namespace GroomBinding_BulkCopy
 
-	// Create mapping between the source & target using their UV
-	uint32 WorkItemCount = 1 + (bNeedTransferPosition ? 1 : 0); //RBF + optional position transfer
-	for (uint32 GroupIt = 0; GroupIt < GroupCount; ++GroupIt)
+  ///////////////////////////////////////////////////////////////////////////////////////////////////
+// Main entry (CPU path)
+static bool InternalBuildBinding_CPU(UGroomBindingAsset* BindingAsset, bool bInitResources)
+{
+#if WITH_EDITORONLY_DATA
+	if (!BindingAsset ||
+		!BindingAsset->Groom ||
+		!BindingAsset->HasValidTarget() ||
+		BindingAsset->Groom->GetNumHairGroups() == 0)
 	{
-		WorkItemCount += 2; // Sim & Render
-		const uint32 CardsLODCount = BindingAsset->HairGroupDatas[GroupIt].CardsRootData.Num();
-		WorkItemCount += CardsLODCount;
+		UE_LOG(LogHairStrands, Error, TEXT("[Groom] Binding asset cannot be created/rebuilt."));
+		return false;
 	}
-	
-	uint32 WorkItemIndex = 0;
-	FScopedSlowTask SlowTask(WorkItemCount, LOCTEXT("BuildBindingData", "Building groom binding data"));
-	SlowTask.MakeDialog();
 
-	TArray<TArray<FVector3f>> TransferredPositions;
-	if (bNeedTransferPosition)
+	// 1. Build groom root data
+	TArray<FHairRootGroupData> OutHairGroupDatas;
 	{
-		bool bSucceed = GroomBinding_Transfer::Transfer( 
-			SourceMeshData.Get(),
-			TargetMeshData.Get(),
-			TransferredPositions, BindingAsset->MatchingSection);
 
-		if (!bSucceed)
+		BindingAsset->Groom->ConditionalPostLoad();
+
+		const int32 NumInterpolationPoints = BindingAsset->NumInterpolationPoints;
+		UGroomAsset* GroomAsset = BindingAsset->Groom;
+
+		TUniquePtr<IMeshData> SourceMeshData;
+		TUniquePtr<IMeshData> TargetMeshData;
+		if (BindingAsset->GroomBindingType == EGroomBindingMeshType::SkeletalMesh)
 		{
-			return false;
-		}
+			BindingAsset->TargetSkeletalMesh->ConditionalPostLoad();
+			if (BindingAsset->SourceSkeletalMesh)
+			{
+				BindingAsset->SourceSkeletalMesh->ConditionalPostLoad();
+			}
 
-		SlowTask.EnterProgressFrame();
-	}
+			SourceMeshData = TUniquePtr<FSkeletalMeshData, TDefaultDelete<IMeshData>>(new FSkeletalMeshData(BindingAsset->SourceSkeletalMesh));
+			TargetMeshData = TUniquePtr<FSkeletalMeshData, TDefaultDelete<IMeshData>>(new FSkeletalMeshData(BindingAsset->TargetSkeletalMesh));
+		}
+		else
+		{
+			BindingAsset->TargetGeometryCache->ConditionalPostLoad();
+			if (BindingAsset->SourceGeometryCache)
+			{
+				BindingAsset->SourceGeometryCache->ConditionalPostLoad();
+			}
 
-	bool bSucceed = false;
-	for (uint32 GroupIt=0; GroupIt < GroupCount; ++GroupIt)
-	{
-		FHairStrandsDatas StrandsData;
-		FHairStrandsDatas GuidesData;
-		GroomAsset->GetHairStrandsDatas(GroupIt, StrandsData, GuidesData);
+			SourceMeshData = TUniquePtr<FGeometryCacheData, TDefaultDelete<IMeshData>>(new FGeometryCacheData(BindingAsset->SourceGeometryCache));
+			TargetMeshData = TUniquePtr<FGeometryCacheData, TDefaultDelete<IMeshData>>(new FGeometryCacheData(BindingAsset->TargetGeometryCache));
+		}
 
-		bSucceed = GroomBinding_RootProjection::Project(
-			StrandsData,
-			TargetMeshData.Get(),
-			TransferredPositions,
-			BindingAsset->HairGroupDatas[GroupIt].RenRootData);
-		if (!bSucceed)
+		if (!TargetMeshData->IsValid())
 		{
-			UE_LOG(LogHairStrands, Error, TEXT("[Groom] Binding asset could not be built. Some strand roots are not close enough to the target mesh to be projected onto it."));
+			UE_LOG(LogHairStrands, Error, TEXT("[Groom] Binding asset could not be built. Target mesh is not valid."));
 			return false;
 		}
+		const uint32 GroupCount = GroomAsset->GetNumHairGroups();
+		const uint32 MeshLODCount = TargetMeshData->GetNumLODs();
 
-		SlowTask.EnterProgressFrame();
+		uint32 GroupIndex = 0;
+		for (const FHairGroupData& GroupData : GroomAsset->HairGroupsData)
+		{
+			FHairStrandsDatas StrandsData;
+			FHairStrandsDatas GuidesData;
+			GroomAsset->GetHairStrandsDatas(GroupIndex, StrandsData, GuidesData);
 
-		bSucceed = GroomBinding_RootProjection::Project(
-			GuidesData,
-			TargetMeshData.Get(),
-			TransferredPositions,
-			BindingAsset->HairGroupDatas[GroupIt].SimRootData);
-		if (!bSucceed) 
+			FHairRootGroupData& OutData = OutHairGroupDatas.AddDefaulted_GetRef();
+			InitHairStrandsRootData(OutData.RenRootData, &StrandsData, MeshLODCount, NumInterpolationPoints);
+			InitHairStrandsRootData(OutData.SimRootData, &GuidesData, MeshLODCount, NumInterpolationPoints);
+
+			const uint32 CardsLODCount = GroupData.Cards.LODs.Num();
+			OutData.CardsRootData.SetNum(GroupData.Cards.LODs.Num());
+			for (uint32 CardsLODIt = 0; CardsLODIt < CardsLODCount; ++CardsLODIt)
+			{
+				if (GroupData.Cards.IsValid(CardsLODIt))
+				{
+					FHairStrandsDatas LODGuidesData;
+					const bool bIsValid = GroomAsset->GetHairCardsGuidesDatas(GroupIndex, CardsLODIt, LODGuidesData);
+					check(bIsValid);
+					InitHairStrandsRootData(OutData.CardsRootData[CardsLODIt], &LODGuidesData, MeshLODCount, NumInterpolationPoints);
+				}
+			}
+			++GroupIndex;
+		}
+
+		const bool bNeedTransferPosition = SourceMeshData->IsValid();
+
+		// Create mapping between the source & target using their UV
+		uint32 WorkItemCount = 1 + (bNeedTransferPosition ? 1 : 0); //RBF + optional position transfer
+		for (uint32 GroupIt = 0; GroupIt < GroupCount; ++GroupIt)
 		{
-			UE_LOG(LogHairStrands, Error, TEXT("[Groom] Binding asset could not be built. Some guide roots are not close enough to the target mesh to be projected onto it."));
-			return false; 
+			WorkItemCount += 2; // Sim & Render
+			const uint32 CardsLODCount = OutHairGroupDatas[GroupIt].CardsRootData.Num();
+			WorkItemCount += CardsLODCount;
 		}
+	
+		uint32 WorkItemIndex = 0;
+		FScopedSlowTask SlowTask(WorkItemCount, LOCTEXT("BuildBindingData", "Building groom binding data"));
+		SlowTask.MakeDialog();
 
-		SlowTask.EnterProgressFrame();
+		TArray<TArray<FVector3f>> TransferredPositions;
+		if (bNeedTransferPosition)
+		{
+			bool bSucceed = GroomBinding_Transfer::Transfer( 
+				SourceMeshData.Get(),
+				TargetMeshData.Get(),
+				TransferredPositions, BindingAsset->MatchingSection);
 
-		const uint32 CardsLODCount = BindingAsset->HairGroupDatas[GroupIt].CardsRootData.Num();
-		for (uint32 CardsLODIt = 0; CardsLODIt < CardsLODCount; ++CardsLODIt)
+			if (!bSucceed)
+			{
+				return false;
+			}
+
+			SlowTask.EnterProgressFrame();
+		}
+
+		bool bSucceed = false;
+		for (uint32 GroupIt=0; GroupIt < GroupCount; ++GroupIt)
 		{
-			if (BindingAsset->Groom->HairGroupsData[GroupIt].Cards.IsValid(CardsLODIt))
+			FHairStrandsDatas StrandsData;
+			FHairStrandsDatas GuidesData;
+			GroomAsset->GetHairStrandsDatas(GroupIt, StrandsData, GuidesData);
+
+			bSucceed = GroomBinding_RootProjection::Project(
+				StrandsData,
+				TargetMeshData.Get(),
+				TransferredPositions,
+				OutHairGroupDatas[GroupIt].RenRootData);
+			if (!bSucceed)
 			{
-				FHairStrandsDatas LODGuidesData;
-				const bool bIsValid = GroomAsset->GetHairCardsGuidesDatas(GroupIt, CardsLODIt, LODGuidesData);
-				check(bIsValid);
-				
-				bSucceed = GroomBinding_RootProjection::Project(
-					LODGuidesData,
-					TargetMeshData.Get(),
-					TransferredPositions,
-					BindingAsset->HairGroupDatas[GroupIt].CardsRootData[CardsLODIt]);
+				UE_LOG(LogHairStrands, Error, TEXT("[Groom] Binding asset could not be built. Some strand roots are not close enough to the target mesh to be projected onto it."));
+				return false;
+			}
+
+			SlowTask.EnterProgressFrame();
+
+			bSucceed = GroomBinding_RootProjection::Project(
+				GuidesData,
+				TargetMeshData.Get(),
+				TransferredPositions,
+				OutHairGroupDatas[GroupIt].SimRootData);
+			if (!bSucceed) 
+			{
+				UE_LOG(LogHairStrands, Error, TEXT("[Groom] Binding asset could not be built. Some guide roots are not close enough to the target mesh to be projected onto it."));
+				return false; 
+			}
 
-				if (!bSucceed) 
+			SlowTask.EnterProgressFrame();
+
+			const uint32 CardsLODCount = OutHairGroupDatas[GroupIt].CardsRootData.Num();
+			for (uint32 CardsLODIt = 0; CardsLODIt < CardsLODCount; ++CardsLODIt)
+			{
+				if (BindingAsset->Groom->HairGroupsData[GroupIt].Cards.IsValid(CardsLODIt))
 				{
-					UE_LOG(LogHairStrands, Error, TEXT("[Groom] Binding asset could not be built. Some cards guide roots are not close enough to the target mesh to be projected onto it."));
-					return false; 
+					FHairStrandsDatas LODGuidesData;
+					const bool bIsValid = GroomAsset->GetHairCardsGuidesDatas(GroupIt, CardsLODIt, LODGuidesData);
+					check(bIsValid);
+				
+					bSucceed = GroomBinding_RootProjection::Project(
+						LODGuidesData,
+						TargetMeshData.Get(),
+						TransferredPositions,
+						OutHairGroupDatas[GroupIt].CardsRootData[CardsLODIt]);
+
+					if (!bSucceed) 
+					{
+						UE_LOG(LogHairStrands, Error, TEXT("[Groom] Binding asset could not be built. Some cards guide roots are not close enough to the target mesh to be projected onto it."));
+						return false; 
+					}
 				}
+
+				SlowTask.EnterProgressFrame();
 			}
+		}
 
-			SlowTask.EnterProgressFrame();
+		GroomBinding_RBFWeighting::ComputeInterpolationWeights(OutHairGroupDatas, BindingAsset->NumInterpolationPoints, BindingAsset->MatchingSection, TargetMeshData.Get(), TransferredPositions);
+		SlowTask.EnterProgressFrame();
+	}
+
+	// 2. Release existing resources data
+	UGroomBindingAsset::FHairGroupResources& OutHairGroupResources = BindingAsset->HairGroupResources;
+	if (BindingAsset->HairGroupResources.Num() > 0)
+	{
+		for (UGroomBindingAsset::FHairGroupResource& GroupResrouces : OutHairGroupResources)
+		{
+			BindingAsset->HairGroupResourcesToDelete.Enqueue(GroupResrouces);
 		}
+		OutHairGroupResources.Empty();
 	}
+	check(OutHairGroupResources.Num() == 0);
 
-	GroomBinding_RBFWeighting::ComputeInterpolationWeights(BindingAsset, TargetMeshData.Get(), TransferredPositions);
-	SlowTask.EnterProgressFrame();
+	// 2. Convert data to bulk data
+	GroomBinding_BulkCopy::BuildRootBulkData(BindingAsset->HairGroupBulkDatas, OutHairGroupDatas);
+
+	// 3. Update GroomBindingAsset infos
+	{
+		TArray<FGoomBindingGroupInfo>& OutGroupInfos = BindingAsset->GroupInfos;
+		OutGroupInfos.Empty();
+		for (const FHairRootGroupData& Data : OutHairGroupDatas)
+		{
+			FGoomBindingGroupInfo& Info = OutGroupInfos.AddDefaulted_GetRef();
+			Info.SimRootCount = Data.SimRootData.RootCount;
+			Info.SimLODCount  = Data.SimRootData.MeshProjectionLODs.Num();
+			Info.RenRootCount = Data.RenRootData.RootCount;
+			Info.RenLODCount  = Data.RenRootData.MeshProjectionLODs.Num();
+		}
+	}
 
-	UpdateGroupInfos(BindingAsset);
 	BindingAsset->QueryStatus = UGroomBindingAsset::EQueryStatus::Completed;
 
+	// 4. Optionnally update resources
 	if (bInitResources)
 	{
 		BindingAsset->InitResource();
@@ -2025,4 +2159,12 @@ bool FGroomBindingBuilder::BuildBinding(UGroomBindingAsset* BindingAsset, bool b
 {
 	return InternalBuildBinding_CPU(BindingAsset, bInitResources);
 }
+
+void FGroomBindingBuilder::GetRootData(
+	FHairStrandsRootData& Out,
+	const FHairStrandsRootBulkData& In)
+{
+	GroomBinding_BulkCopy::BuildRootData(Out, In);
+}
+
 #undef LOCTEXT_NAMESPACE
diff --git a/Engine/Plugins/Runtime/HairStrands/Source/HairStrandsCore/Private/GroomComponent.cpp b/Engine/Plugins/Runtime/HairStrands/Source/HairStrandsCore/Private/GroomComponent.cpp
--- a/Engine/Plugins/Runtime/HairStrands/Source/HairStrandsCore/Private/GroomComponent.cpp
+++ b/Engine/Plugins/Runtime/HairStrands/Source/HairStrandsCore/Private/GroomComponent.cpp
@@ -1800,7 +1800,7 @@ static USkeletalMeshComponent* ValidateBindingAsset(
 					GroupIt < BindingAsset->HairGroupResources.Num() &&
 					CardsLODIndex < uint32(BindingAsset->HairGroupResources[GroupIt].CardsRootResources.Num()) &&
 					BindingAsset->HairGroupResources[GroupIt].CardsRootResources[CardsLODIndex] != nullptr &&
-					((SkeletalMeshComponent && SkeletalMeshComponent->SkeletalMesh) ? SkeletalMeshComponent->SkeletalMesh->GetLODInfoArray().Num() == BindingAsset->HairGroupResources[GroupIt].CardsRootResources[CardsLODIndex]->RootData.MeshProjectionLODs.Num() : false);
+					((SkeletalMeshComponent && SkeletalMeshComponent->SkeletalMesh) ? SkeletalMeshComponent->SkeletalMesh->GetLODInfoArray().Num() == BindingAsset->HairGroupResources[GroupIt].CardsRootResources[CardsLODIndex]->BulkData.MeshProjectionLODs.Num() : false);
 
 				if (!bIsCardsBindingCompatible)
 				{
@@ -2024,7 +2024,7 @@ void UGroomComponent::InitResources(bool bIsBindingReloading)
 				check(GroupIt < LocalBindingAsset->HairGroupResources.Num());
 				if (USkeletalMeshComponent* SkeletalMeshComponent = Cast<USkeletalMeshComponent>(RegisteredMeshComponent))
 				{
-					check(SkeletalMeshComponent->SkeletalMesh ? SkeletalMeshComponent->SkeletalMesh->GetLODInfoArray().Num() == LocalBindingAsset->HairGroupResources[GroupIt].SimRootResources->RootData.MeshProjectionLODs.Num() : false);
+					check(SkeletalMeshComponent->SkeletalMesh ? SkeletalMeshComponent->SkeletalMesh->GetLODInfoArray().Num() == LocalBindingAsset->HairGroupResources[GroupIt].SimRootResources->BulkData.MeshProjectionLODs.Num() : false);
 				}
 
 				HairGroupInstance->Guides.RestRootResource = LocalBindingAsset->HairGroupResources[GroupIt].SimRootResources;
@@ -2120,7 +2120,7 @@ void UGroomComponent::InitResources(bool bIsBindingReloading)
 				check(GroupIt < LocalBindingAsset->HairGroupResources.Num());
 				if (USkeletalMeshComponent* SkeletalMeshComponent = Cast<USkeletalMeshComponent>(RegisteredMeshComponent))
 				{
-					check(SkeletalMeshComponent->SkeletalMesh ? SkeletalMeshComponent->SkeletalMesh->GetLODInfoArray().Num() == LocalBindingAsset->HairGroupResources[GroupIt].RenRootResources->RootData.MeshProjectionLODs.Num() : false);
+					check(SkeletalMeshComponent->SkeletalMesh ? SkeletalMeshComponent->SkeletalMesh->GetLODInfoArray().Num() == LocalBindingAsset->HairGroupResources[GroupIt].RenRootResources->BulkData.MeshProjectionLODs.Num() : false);
 				}
 
 				HairGroupInstance->Strands.RestRootResource = LocalBindingAsset->HairGroupResources[GroupIt].RenRootResources;
@@ -2224,7 +2224,7 @@ void UGroomComponent::InitResources(bool bIsBindingReloading)
 						check(GroupIt < LocalBindingAsset->HairGroupResources.Num());
 						if (USkeletalMeshComponent* SkeletalMeshComponent = Cast<USkeletalMeshComponent>(RegisteredMeshComponent))
 						{
-							check(SkeletalMeshComponent->SkeletalMesh ? SkeletalMeshComponent->SkeletalMesh->GetLODInfoArray().Num() == LocalBindingAsset->HairGroupResources[GroupIt].CardsRootResources[CardsLODIndex]->RootData.MeshProjectionLODs.Num() : false);
+							check(SkeletalMeshComponent->SkeletalMesh ? SkeletalMeshComponent->SkeletalMesh->GetLODInfoArray().Num() == LocalBindingAsset->HairGroupResources[GroupIt].CardsRootResources[CardsLODIndex]->BulkData.MeshProjectionLODs.Num() : false);
 						}
 
 						InstanceLOD.Guides.RestRootResource = LocalBindingAsset->HairGroupResources[GroupIt].CardsRootResources[CardsLODIndex];
diff --git a/Engine/Plugins/Runtime/HairStrands/Source/HairStrandsCore/Private/GroomDebug.cpp b/Engine/Plugins/Runtime/HairStrands/Source/HairStrandsCore/Private/GroomDebug.cpp
--- a/Engine/Plugins/Runtime/HairStrands/Source/HairStrandsCore/Private/GroomDebug.cpp
+++ b/Engine/Plugins/Runtime/HairStrands/Source/HairStrandsCore/Private/GroomDebug.cpp
@@ -366,7 +366,7 @@ static void AddDebugProjectionHairPass(
 	FRDGTextureRef DepthTexture)
 {
 	const EPrimitiveType PrimitiveType = GeometryType == EDebugProjectionHairType::HairFrame ? PT_LineList : PT_TriangleList;
-	const uint32 RootCount = RestRootResources->RootData.RootCount;
+	const uint32 RootCount = RestRootResources->BulkData.RootCount;
 	const uint32 PrimitiveCount = RootCount;
 
 	if (PrimitiveCount == 0 || MeshLODIndex < 0 || MeshLODIndex >= RestRootResources->LODs.Num() || MeshLODIndex >= DeformedRootResources->LODs.Num())
diff --git a/Engine/Plugins/Runtime/HairStrands/Source/HairStrandsCore/Private/GroomRBFDeformer.cpp b/Engine/Plugins/Runtime/HairStrands/Source/HairStrandsCore/Private/GroomRBFDeformer.cpp
--- a/Engine/Plugins/Runtime/HairStrands/Source/HairStrandsCore/Private/GroomRBFDeformer.cpp
+++ b/Engine/Plugins/Runtime/HairStrands/Source/HairStrandsCore/Private/GroomRBFDeformer.cpp
@@ -8,6 +8,7 @@
 #include "Rendering/SkeletalMeshRenderData.h"
 #include "MeshAttributes.h"
 #include "Engine/StaticMesh.h"
+#include "GroomBindingBuilder.h"
 
 // HairStrandsSamplesInit.usf
 void InitMeshSamples(
@@ -499,7 +500,7 @@ void FGroomRBFDeformer::GetRBFDeformedGroomAsset(const UGroomAsset* InGroomAsset
 		ExtractSkeletalVertexPosition(SkeletalMeshData_Target, MeshLODIndex, MeshVertexPositionsBuffer_Target);
 
 		// Apply RBF deformation to each group of guides and render strands
-		const int32 NumGroups = BindingAsset->HairGroupDatas.Num();
+		const int32 NumGroups = BindingAsset->HairGroupBulkDatas.Num();
 
 		// Use the vertices positions from the HairDescription instead of the GroomAsset since the latter
 		// may contain decimated or auto-generated guides depending on the import settings
@@ -538,23 +539,31 @@ void FGroomRBFDeformer::GetRBFDeformedGroomAsset(const UGroomAsset* InGroomAsset
 			}
 			else
 			{
+				FHairStrandsRootData SimRootData;
+				FGroomBindingBuilder::GetRootData(SimRootData, BindingAsset->HairGroupBulkDatas[GroupIndex].SimRootBulkData);
+
 				DeformedPositions[GroupIndex].GuideStrands = GetDeformedHairStrandsPositions(
 					MeshVertexPositionsBuffer_Target,
 					GuidesData,
 					MeshLODIndex,
 					SkeletalMeshData_Target,
-					BindingAsset->HairGroupDatas[GroupIndex].SimRootData.VertexToCurveIndexBuffer,
-					BindingAsset->HairGroupDatas[GroupIndex].SimRootData.MeshProjectionLODs[MeshLODIndex]);
+					SimRootData.VertexToCurveIndexBuffer,
+					SimRootData.MeshProjectionLODs[MeshLODIndex]);
 			}
 
 			// Get deformed render strands
-			DeformedPositions[GroupIndex].RenderStrands = GetDeformedHairStrandsPositions(
-				MeshVertexPositionsBuffer_Target,
-				StrandsData,
-				MeshLODIndex,
-				SkeletalMeshData_Target,
-				BindingAsset->HairGroupDatas[GroupIndex].RenRootData.VertexToCurveIndexBuffer,
-				BindingAsset->HairGroupDatas[GroupIndex].RenRootData.MeshProjectionLODs[MeshLODIndex]);
+			{
+				FHairStrandsRootData RenRootData;
+				FGroomBindingBuilder::GetRootData(RenRootData, BindingAsset->HairGroupBulkDatas[GroupIndex].RenRootBulkData);
+
+				DeformedPositions[GroupIndex].RenderStrands = GetDeformedHairStrandsPositions(
+					MeshVertexPositionsBuffer_Target,
+					StrandsData,
+					MeshLODIndex,
+					SkeletalMeshData_Target,
+					RenRootData.VertexToCurveIndexBuffer,
+					RenRootData.MeshProjectionLODs[MeshLODIndex]);
+			}
 		}
 
 		// Apply changes onto cards and meshes (OutGroomASset already contain duplicated mesh asset
@@ -577,7 +586,11 @@ void FGroomRBFDeformer::GetRBFDeformedGroomAsset(const UGroomAsset* InGroomAsset
 			if (Desc.GroupIndex >= 0)
 			{
 				Mesh->ConditionalPostLoad();
-				DeformStaticMeshPositions(Mesh, MeshVertexPositionsBuffer_Target, BindingAsset->HairGroupDatas[Desc.GroupIndex].RenRootData.MeshProjectionLODs[MeshLODIndex]);
+
+				FHairStrandsRootData RenRootData;
+				FGroomBindingBuilder::GetRootData(RenRootData, BindingAsset->HairGroupBulkDatas[Desc.GroupIndex].RenRootBulkData);
+
+				DeformStaticMeshPositions(Mesh, MeshVertexPositionsBuffer_Target, RenRootData.MeshProjectionLODs[MeshLODIndex]);
 			}
 		} 
 
@@ -592,7 +605,11 @@ void FGroomRBFDeformer::GetRBFDeformedGroomAsset(const UGroomAsset* InGroomAsset
 				}
 
 				Mesh->ConditionalPostLoad();
-				DeformStaticMeshPositions(Mesh, MeshVertexPositionsBuffer_Target, BindingAsset->HairGroupDatas[Desc.GroupIndex].RenRootData.MeshProjectionLODs[MeshLODIndex]);
+				
+				FHairStrandsRootData RenRootData;
+				FGroomBindingBuilder::GetRootData(RenRootData, BindingAsset->HairGroupBulkDatas[Desc.GroupIndex].RenRootBulkData);
+
+				DeformStaticMeshPositions(Mesh, MeshVertexPositionsBuffer_Target, RenRootData.MeshProjectionLODs[MeshLODIndex]);
 			}
 		}
 
diff --git a/Engine/Plugins/Runtime/HairStrands/Source/HairStrandsCore/Private/GroomResources.cpp b/Engine/Plugins/Runtime/HairStrands/Source/HairStrandsCore/Private/GroomResources.cpp
--- a/Engine/Plugins/Runtime/HairStrands/Source/HairStrandsCore/Private/GroomResources.cpp
+++ b/Engine/Plugins/Runtime/HairStrands/Source/HairStrandsCore/Private/GroomResources.cpp
@@ -784,9 +784,9 @@ void FHairStrandsClusterCullingResource::InternalRelease()
 
 /////////////////////////////////////////////////////////////////////////////////////////
 
-FHairStrandsRestRootResource::FHairStrandsRestRootResource(const FHairStrandsRootData& InRootData, EHairStrandsResourcesType InCurveType) :
+FHairStrandsRestRootResource::FHairStrandsRestRootResource(const FHairStrandsRootBulkData& InBulkData, EHairStrandsResourcesType InCurveType) :
 	FHairCommonResource(EHairStrandsAllocationType::Deferred),
-	RootData(InRootData), CurveType(InCurveType)
+	BulkData(InBulkData), CurveType(InCurveType)
 {
 	PopulateFromRootData();
 }
@@ -794,8 +794,8 @@ FHairStrandsRestRootResource::FHairStrandsRestRootResource(const FHairStrandsRoo
 void FHairStrandsRestRootResource::PopulateFromRootData()
 {
 	uint32 LODIndex = 0;
-	LODs.Reserve(RootData.MeshProjectionLODs.Num());
-	for (const FHairStrandsRootData::FMeshProjectionLOD& MeshProjectionLOD : RootData.MeshProjectionLODs)
+	LODs.Reserve(BulkData.MeshProjectionLODs.Num());
+	for (const FHairStrandsRootBulkData::FMeshProjectionLOD& MeshProjectionLOD : BulkData.MeshProjectionLODs)
 	{
 		FLOD& LOD = LODs.AddDefaulted_GetRef();
 
@@ -814,9 +814,11 @@ void FHairStrandsRestRootResource::InternalAllocate(FRDGBuilder& GraphBuilder)
 		PopulateFromRootData();
 	}
 
-	if (RootData.VertexToCurveIndexBuffer.Num() > 0)
+	if (BulkData.PointCount > 0)
 	{
-		InternalCreateVertexBufferRDG<FHairStrandsIndexFormat>(GraphBuilder, RootData.VertexToCurveIndexBuffer, VertexToCurveIndexBuffer, HAIRSTRANDS_RESOUCE_NAME(CurveType, Hair.StrandsRestRoot_VertexToCurveIndexBuffer), EHairResourceUsageType::Static);
+		const FHairStrandsIndexFormat::Type* Data = (const FHairStrandsIndexFormat::Type*)BulkData.VertexToCurveIndexBuffer.LockReadOnly();
+		InternalCreateVertexBufferRDG<FHairStrandsIndexFormat>(GraphBuilder, Data, BulkData.PointCount, VertexToCurveIndexBuffer, HAIRSTRANDS_RESOUCE_NAME(CurveType, Hair.StrandsRestRoot_VertexToCurveIndexBuffer), EHairResourceUsageType::Static, ERDGInitialDataFlags::None);
+		BulkData.VertexToCurveIndexBuffer.Unlock();
 	}
 }
 
@@ -824,8 +826,8 @@ void FHairStrandsRestRootResource::InternalAllocateLOD(FRDGBuilder& GraphBuilder
 {
 	// Sanity check ti insure that the 'common' part of FHairStrandsRestRootResource is already initialized
 	check(bIsInitialized);
-	check(RootData.VertexToCurveIndexBuffer.Num() > 0);
-	check(LODs.Num() == RootData.MeshProjectionLODs.Num());
+	check(BulkData.PointCount > 0);
+	check(LODs.Num() == BulkData.MeshProjectionLODs.Num());
 	if (LODIndex >= 0 && LODIndex < LODs.Num())
 	{
 		FLOD& GPUData = LODs[LODIndex];
@@ -835,56 +837,100 @@ void FHairStrandsRestRootResource::InternalAllocateLOD(FRDGBuilder& GraphBuilder
 			return;
 		}
 
-		const FHairStrandsRootData::FMeshProjectionLOD& CPUData = RootData.MeshProjectionLODs[LODIndex];
-		const bool bHasValidCPUData = CPUData.RootTriangleBarycentricBuffer.Num() > 0;
+		const FHairStrandsCurveTriangleBarycentricFormat::Type* Data_Barycentric= nullptr;
+		const FHairStrandsCurveTriangleIndexFormat::Type* Data_Index			= nullptr;
+		const FHairStrandsMeshTrianglePositionFormat::Type* Data_Triangle0		= nullptr;
+		const FHairStrandsMeshTrianglePositionFormat::Type* Data_Triangle1		= nullptr;
+		const FHairStrandsMeshTrianglePositionFormat::Type* Data_Triangle2		= nullptr;
+
+		const FHairStrandsRootBulkData::FMeshProjectionLOD& CPUData = BulkData.MeshProjectionLODs[LODIndex];
+		const bool bHasValidCPUData = CPUData.RootTriangleBarycentricBuffer.GetBulkDataSize() > 0;
 		if (bHasValidCPUData)
 		{
 			GPUData.Status = FLOD::EStatus::Completed;
-
-			check(CPUData.RootTriangleBarycentricBuffer.Num() > 0);
-			InternalCreateVertexBufferRDG<FHairStrandsCurveTriangleBarycentricFormat>(GraphBuilder, CPUData.RootTriangleBarycentricBuffer, GPUData.RootTriangleBarycentricBuffer, HAIRSTRANDS_RESOUCE_NAME(CurveType, Hair.StrandsRestRoot_RootTriangleBarycentricBuffer), EHairResourceUsageType::Static);
-
-			check(CPUData.RootTriangleIndexBuffer.Num() > 0);
-			InternalCreateVertexBufferRDG<FHairStrandsCurveTriangleIndexFormat>(GraphBuilder, CPUData.RootTriangleIndexBuffer, GPUData.RootTriangleIndexBuffer, HAIRSTRANDS_RESOUCE_NAME(CurveType, Hair.StrandsRestRoot_RootTriangleIndexBuffer), EHairResourceUsageType::Static);
-
-			check(CPUData.RestRootTrianglePosition0Buffer.Num() > 0);
-			check(CPUData.RestRootTrianglePosition1Buffer.Num() > 0);
-			check(CPUData.RestRootTrianglePosition2Buffer.Num() > 0);
-			InternalCreateVertexBufferRDG<FHairStrandsMeshTrianglePositionFormat>(GraphBuilder, CPUData.RestRootTrianglePosition0Buffer, GPUData.RestRootTrianglePosition0Buffer, HAIRSTRANDS_RESOUCE_NAME(CurveType, Hair.StrandsRestRoot_RestRootTrianglePosition0Buffer), EHairResourceUsageType::Static);
-			InternalCreateVertexBufferRDG<FHairStrandsMeshTrianglePositionFormat>(GraphBuilder, CPUData.RestRootTrianglePosition1Buffer, GPUData.RestRootTrianglePosition1Buffer, HAIRSTRANDS_RESOUCE_NAME(CurveType, Hair.StrandsRestRoot_RestRootTrianglePosition1Buffer), EHairResourceUsageType::Static);
-			InternalCreateVertexBufferRDG<FHairStrandsMeshTrianglePositionFormat>(GraphBuilder, CPUData.RestRootTrianglePosition2Buffer, GPUData.RestRootTrianglePosition2Buffer, HAIRSTRANDS_RESOUCE_NAME(CurveType, Hair.StrandsRestRoot_RestRootTrianglePosition2Buffer), EHairResourceUsageType::Static);
+			{
+				check(CPUData.RootTriangleBarycentricBuffer.GetBulkDataSize() > 0);
+				Data_Barycentric = (FHairStrandsCurveTriangleBarycentricFormat::Type*)CPUData.RootTriangleBarycentricBuffer.LockReadOnly();
+				InternalCreateVertexBufferRDG<FHairStrandsCurveTriangleBarycentricFormat>(GraphBuilder, Data_Barycentric, BulkData.RootCount, GPUData.RootTriangleBarycentricBuffer, HAIRSTRANDS_RESOUCE_NAME(CurveType, Hair.StrandsRestRoot_RootTriangleBarycentricBuffer), EHairResourceUsageType::Static, ERDGInitialDataFlags::None);
+			}
+			{
+				check(CPUData.RootTriangleIndexBuffer.GetBulkDataSize() > 0);
+				Data_Index = (FHairStrandsCurveTriangleIndexFormat::Type*)CPUData.RootTriangleIndexBuffer.LockReadOnly();
+				InternalCreateVertexBufferRDG<FHairStrandsCurveTriangleIndexFormat>(GraphBuilder, Data_Index, BulkData.RootCount, GPUData.RootTriangleIndexBuffer, HAIRSTRANDS_RESOUCE_NAME(CurveType, Hair.StrandsRestRoot_RootTriangleIndexBuffer), EHairResourceUsageType::Static, ERDGInitialDataFlags::None);
+			}
+			{
+				check(CPUData.RestRootTrianglePosition0Buffer.GetBulkDataSize() > 0);
+				Data_Triangle0 = (FHairStrandsMeshTrianglePositionFormat::Type*)CPUData.RestRootTrianglePosition0Buffer.LockReadOnly();
+				InternalCreateVertexBufferRDG<FHairStrandsMeshTrianglePositionFormat>(GraphBuilder, Data_Triangle0, BulkData.RootCount, GPUData.RestRootTrianglePosition0Buffer, HAIRSTRANDS_RESOUCE_NAME(CurveType, Hair.StrandsRestRoot_RestRootTrianglePosition0Buffer), EHairResourceUsageType::Static, ERDGInitialDataFlags::None);
+			}
+			{
+				check(CPUData.RestRootTrianglePosition1Buffer.GetBulkDataSize() > 0);
+				Data_Triangle1 = (FHairStrandsMeshTrianglePositionFormat::Type*)CPUData.RestRootTrianglePosition1Buffer.LockReadOnly();
+				InternalCreateVertexBufferRDG<FHairStrandsMeshTrianglePositionFormat>(GraphBuilder, Data_Triangle1, BulkData.RootCount, GPUData.RestRootTrianglePosition1Buffer, HAIRSTRANDS_RESOUCE_NAME(CurveType, Hair.StrandsRestRoot_RestRootTrianglePosition1Buffer), EHairResourceUsageType::Static, ERDGInitialDataFlags::None);
+			}
+			{
+				check(CPUData.RestRootTrianglePosition2Buffer.GetBulkDataSize() > 0);
+				Data_Triangle2 = (FHairStrandsMeshTrianglePositionFormat::Type*)CPUData.RestRootTrianglePosition2Buffer.LockReadOnly();
+				InternalCreateVertexBufferRDG<FHairStrandsMeshTrianglePositionFormat>(GraphBuilder, Data_Triangle2, BulkData.RootCount, GPUData.RestRootTrianglePosition2Buffer, HAIRSTRANDS_RESOUCE_NAME(CurveType, Hair.StrandsRestRoot_RestRootTrianglePosition2Buffer), EHairResourceUsageType::Static, ERDGInitialDataFlags::None);
+			}
 		}
 		else
 		{
 			GPUData.Status = FLOD::EStatus::Initialized;
 			
-			InternalCreateVertexBufferRDG<FHairStrandsCurveTriangleBarycentricFormat>(GraphBuilder, RootData.RootCount, GPUData.RootTriangleBarycentricBuffer, HAIRSTRANDS_RESOUCE_NAME(CurveType, Hair.StrandsRestRoot_RootTriangleBarycentricBuffer), EHairResourceUsageType::Dynamic);
-			InternalCreateVertexBufferRDG<FHairStrandsCurveTriangleIndexFormat>(GraphBuilder, RootData.RootCount, GPUData.RootTriangleIndexBuffer, HAIRSTRANDS_RESOUCE_NAME(CurveType, Hair.StrandsRestRoot_RootTriangleIndexBuffer), EHairResourceUsageType::Dynamic);
+			InternalCreateVertexBufferRDG<FHairStrandsCurveTriangleBarycentricFormat>(GraphBuilder, BulkData.RootCount, GPUData.RootTriangleBarycentricBuffer, HAIRSTRANDS_RESOUCE_NAME(CurveType, Hair.StrandsRestRoot_RootTriangleBarycentricBuffer), EHairResourceUsageType::Dynamic);
+			InternalCreateVertexBufferRDG<FHairStrandsCurveTriangleIndexFormat>(GraphBuilder, BulkData.RootCount, GPUData.RootTriangleIndexBuffer, HAIRSTRANDS_RESOUCE_NAME(CurveType, Hair.StrandsRestRoot_RootTriangleIndexBuffer), EHairResourceUsageType::Dynamic);
 			
 			// Create buffers. Initialization will be done by render passes
-			InternalCreateVertexBufferRDG<FHairStrandsMeshTrianglePositionFormat>(GraphBuilder, RootData.RootCount, GPUData.RestRootTrianglePosition0Buffer, HAIRSTRANDS_RESOUCE_NAME(CurveType, Hair.StrandsRestRoot_RestRootTrianglePosition0Buffer), EHairResourceUsageType::Dynamic);
-			InternalCreateVertexBufferRDG<FHairStrandsMeshTrianglePositionFormat>(GraphBuilder, RootData.RootCount, GPUData.RestRootTrianglePosition1Buffer, HAIRSTRANDS_RESOUCE_NAME(CurveType, Hair.StrandsRestRoot_RestRootTrianglePosition1Buffer), EHairResourceUsageType::Dynamic);
-			InternalCreateVertexBufferRDG<FHairStrandsMeshTrianglePositionFormat>(GraphBuilder, RootData.RootCount, GPUData.RestRootTrianglePosition2Buffer, HAIRSTRANDS_RESOUCE_NAME(CurveType, Hair.StrandsRestRoot_RestRootTrianglePosition2Buffer), EHairResourceUsageType::Dynamic);
+			InternalCreateVertexBufferRDG<FHairStrandsMeshTrianglePositionFormat>(GraphBuilder, BulkData.RootCount, GPUData.RestRootTrianglePosition0Buffer, HAIRSTRANDS_RESOUCE_NAME(CurveType, Hair.StrandsRestRoot_RestRootTrianglePosition0Buffer), EHairResourceUsageType::Dynamic);
+			InternalCreateVertexBufferRDG<FHairStrandsMeshTrianglePositionFormat>(GraphBuilder, BulkData.RootCount, GPUData.RestRootTrianglePosition1Buffer, HAIRSTRANDS_RESOUCE_NAME(CurveType, Hair.StrandsRestRoot_RestRootTrianglePosition1Buffer), EHairResourceUsageType::Dynamic);
+			InternalCreateVertexBufferRDG<FHairStrandsMeshTrianglePositionFormat>(GraphBuilder, BulkData.RootCount, GPUData.RestRootTrianglePosition2Buffer, HAIRSTRANDS_RESOUCE_NAME(CurveType, Hair.StrandsRestRoot_RestRootTrianglePosition2Buffer), EHairResourceUsageType::Dynamic);
 		}
 
+		FHairStrandsIndexFormat::Type* DataSample_SampleIndex = nullptr;
+		FHairStrandsWeightFormat::Type* DataSample_InterpolationWeight = nullptr;
+		FHairStrandsMeshTrianglePositionFormat::Type* DataSample_RestSample = nullptr;
+
 		GPUData.SampleCount = CPUData.SampleCount;
-		const bool bHasValidCPUWeights = CPUData.MeshSampleIndicesBuffer.Num() > 0;
+		const bool bHasValidCPUWeights = CPUData.MeshSampleIndicesBuffer.GetBulkDataSize() > 0;
 		if (bHasValidCPUWeights)
 		{
-			//check(CPUData.MeshInterpolationWeightsBuffer.Num() == (CPUData.SampleCount+4) * (CPUData.SampleCount+4));
-			check(CPUData.MeshSampleIndicesBuffer.Num() == CPUData.SampleCount);
-			check(CPUData.RestSamplePositionsBuffer.Num() == CPUData.SampleCount);
-
-			InternalCreateVertexBufferRDG<FHairStrandsWeightFormat>(GraphBuilder, CPUData.MeshInterpolationWeightsBuffer, GPUData.MeshInterpolationWeightsBuffer, HAIRSTRANDS_RESOUCE_NAME(CurveType, Hair.StrandsRestRoot_MeshInterpolationWeightsBuffer), EHairResourceUsageType::Static);
-			InternalCreateVertexBufferRDG<FHairStrandsIndexFormat>(GraphBuilder, CPUData.MeshSampleIndicesBuffer, GPUData.MeshSampleIndicesBuffer, HAIRSTRANDS_RESOUCE_NAME(CurveType, Hair.StrandsRestRoot_MeshSampleIndicesBuffer), EHairResourceUsageType::Static);
-			InternalCreateVertexBufferRDG<FHairStrandsMeshTrianglePositionFormat>(GraphBuilder, CPUData.RestSamplePositionsBuffer, GPUData.RestSamplePositionsBuffer, HAIRSTRANDS_RESOUCE_NAME(CurveType, Hair.StrandsRestRoot_RestSamplePositionsBuffer), EHairResourceUsageType::Static);
+			{
+				const uint32 InteroplationWeightCount = CPUData.MeshInterpolationWeightsBuffer.GetBulkDataSize() / sizeof(FHairStrandsWeightFormat::Type);
+				//check(InteroplationWeightCount == (CPUData.SampleCount+4) * (CPUData.SampleCount+4));
+				DataSample_InterpolationWeight = (FHairStrandsWeightFormat::Type*)CPUData.MeshInterpolationWeightsBuffer.LockReadOnly();
+				InternalCreateVertexBufferRDG<FHairStrandsWeightFormat>(GraphBuilder, DataSample_InterpolationWeight, InteroplationWeightCount, GPUData.MeshInterpolationWeightsBuffer, HAIRSTRANDS_RESOUCE_NAME(CurveType, Hair.StrandsRestRoot_MeshInterpolationWeightsBuffer), EHairResourceUsageType::Static, ERDGInitialDataFlags::None);
+			}
+			{
+				check(CPUData.MeshSampleIndicesBuffer.GetBulkDataSize()/sizeof(FHairStrandsIndexFormat::Type) == CPUData.SampleCount);
+				DataSample_SampleIndex = (FHairStrandsIndexFormat::Type*)CPUData.MeshSampleIndicesBuffer.LockReadOnly();
+				InternalCreateVertexBufferRDG<FHairStrandsIndexFormat>(GraphBuilder, DataSample_SampleIndex, CPUData.SampleCount, GPUData.MeshSampleIndicesBuffer, HAIRSTRANDS_RESOUCE_NAME(CurveType, Hair.StrandsRestRoot_MeshSampleIndicesBuffer), EHairResourceUsageType::Static, ERDGInitialDataFlags::None);
+			}
+			{
+				check(CPUData.RestSamplePositionsBuffer.GetBulkDataSize() / sizeof(FHairStrandsMeshTrianglePositionFormat::Type) == CPUData.SampleCount);
+				DataSample_RestSample = (FHairStrandsMeshTrianglePositionFormat::Type*)CPUData.RestSamplePositionsBuffer.LockReadOnly();
+				InternalCreateVertexBufferRDG<FHairStrandsMeshTrianglePositionFormat>(GraphBuilder, DataSample_RestSample, CPUData.SampleCount, GPUData.RestSamplePositionsBuffer, HAIRSTRANDS_RESOUCE_NAME(CurveType, Hair.StrandsRestRoot_RestSamplePositionsBuffer), EHairResourceUsageType::Static, ERDGInitialDataFlags::None);
+			}
 		}
 		else
 		{
+			// TODO: do not allocate these resources, since they won't be used
 			InternalCreateVertexBufferRDG<FHairStrandsWeightFormat>(GraphBuilder, (CPUData.SampleCount+4) * (CPUData.SampleCount+4), GPUData.MeshInterpolationWeightsBuffer, HAIRSTRANDS_RESOUCE_NAME(CurveType, Hair.StrandsRestRoot_MeshInterpolationWeightsBuffer), EHairResourceUsageType::Dynamic);
 			InternalCreateVertexBufferRDG<FHairStrandsIndexFormat>(GraphBuilder, CPUData.SampleCount, GPUData.MeshSampleIndicesBuffer, HAIRSTRANDS_RESOUCE_NAME(CurveType, Hair.StrandsRestRoot_MeshSampleIndicesBuffer), EHairResourceUsageType::Dynamic);
 			InternalCreateVertexBufferRDG<FHairStrandsMeshTrianglePositionFormat>(GraphBuilder, CPUData.SampleCount, GPUData.RestSamplePositionsBuffer, HAIRSTRANDS_RESOUCE_NAME(CurveType, Hair.StrandsRestRoot_RestSamplePositionsBuffer), EHairResourceUsageType::Dynamic);
 		}
+
+		// Unlock mapped bulk buffer
+		if (Data_Barycentric)				{ CPUData.RootTriangleBarycentricBuffer.Unlock(); }
+		if (Data_Index)						{ CPUData.RootTriangleIndexBuffer.Unlock(); }
+		if (Data_Triangle0)					{ CPUData.RestRootTrianglePosition0Buffer.Unlock(); }
+		if (Data_Triangle1)					{ CPUData.RestRootTrianglePosition1Buffer.Unlock(); }
+		if (Data_Triangle2)					{ CPUData.RestRootTrianglePosition2Buffer.Unlock(); }
+
+		if (DataSample_InterpolationWeight) { CPUData.MeshInterpolationWeightsBuffer.Unlock(); }
+		if (DataSample_SampleIndex)			{ CPUData.MeshSampleIndicesBuffer.Unlock(); }
+		if (DataSample_RestSample)			{ CPUData.RestSamplePositionsBuffer.Unlock(); }
+
 	}
 }
 
@@ -923,7 +969,7 @@ FHairStrandsDeformedRootResource::FHairStrandsDeformedRootResource(const FHairSt
 {
 	check(InRestResources);
 	uint32 LODIndex = 0;
-	RootCount = InRestResources->RootData.RootCount;
+	RootCount = InRestResources->BulkData.RootCount;
 	LODs.Reserve(InRestResources->LODs.Num());
 	for (const FHairStrandsRestRootResource::FLOD& InLOD : InRestResources->LODs)
 	{
@@ -942,8 +988,11 @@ void FHairStrandsDeformedRootResource::InternalAllocate(FRDGBuilder& GraphBuilde
 		for (FLOD& LOD : LODs)
 		{		
 			LOD.Status = FLOD::EStatus::Initialized;
-			InternalCreateVertexBufferRDG<FHairStrandsMeshTrianglePositionFormat>(GraphBuilder, LOD.SampleCount, LOD.DeformedSamplePositionsBuffer, HAIRSTRANDS_RESOUCE_NAME(CurveType, Hair.StrandsRootDeformed_DeformedSamplePositionsBuffer), EHairResourceUsageType::Dynamic);
-			InternalCreateVertexBufferRDG<FHairStrandsMeshTrianglePositionFormat>(GraphBuilder, LOD.SampleCount + 4, LOD.MeshSampleWeightsBuffer, HAIRSTRANDS_RESOUCE_NAME(CurveType, Hair.StrandsRootDeformed_MeshSampleWeightsBuffer), EHairResourceUsageType::Dynamic);
+			if (LOD.SampleCount > 0)
+			{
+				InternalCreateVertexBufferRDG<FHairStrandsMeshTrianglePositionFormat>(GraphBuilder, LOD.SampleCount, LOD.DeformedSamplePositionsBuffer, HAIRSTRANDS_RESOUCE_NAME(CurveType, Hair.StrandsRootDeformed_DeformedSamplePositionsBuffer), EHairResourceUsageType::Dynamic);
+				InternalCreateVertexBufferRDG<FHairStrandsMeshTrianglePositionFormat>(GraphBuilder, LOD.SampleCount + 4, LOD.MeshSampleWeightsBuffer, HAIRSTRANDS_RESOUCE_NAME(CurveType, Hair.StrandsRootDeformed_MeshSampleWeightsBuffer), EHairResourceUsageType::Dynamic);
+			}
 
 			InternalCreateVertexBufferRDG<FHairStrandsMeshTrianglePositionFormat>(GraphBuilder, RootCount, LOD.DeformedRootTrianglePosition0Buffer, HAIRSTRANDS_RESOUCE_NAME(CurveType, Hair.StrandsRootDeformed_DeformedRootTrianglePosition0Buffer), EHairResourceUsageType::Dynamic);
 			InternalCreateVertexBufferRDG<FHairStrandsMeshTrianglePositionFormat>(GraphBuilder, RootCount, LOD.DeformedRootTrianglePosition1Buffer, HAIRSTRANDS_RESOUCE_NAME(CurveType, Hair.StrandsRootDeformed_DeformedRootTrianglePosition1Buffer), EHairResourceUsageType::Dynamic);
@@ -967,7 +1016,111 @@ void FHairStrandsDeformedRootResource::InternalRelease()
 }
 
 /////////////////////////////////////////////////////////////////////////////////////////
+// Root bulk data
+FHairStrandsRootBulkData::FHairStrandsRootBulkData()
+{
+
+}
+
+bool FHairStrandsRootBulkData::HasProjectionData() const
+{
+	bool bIsValid = MeshProjectionLODs.Num() > 0;
+	bIsValid = RootCount > 0;
+	for (const FMeshProjectionLOD& LOD : MeshProjectionLODs)
+	{
+		bIsValid = bIsValid && LOD.LODIndex > 0;
+		if (!bIsValid) break;
+	}
+	return bIsValid;
+}
+
+const TArray<uint32>& FHairStrandsRootBulkData::GetValidSectionIndices(int32 LODIndex) const
+{
+	check(LODIndex >= 0 && LODIndex < MeshProjectionLODs.Num());
+
+	if (MeshProjectionLODs[LODIndex].CachedValidSectionIndices.IsEmpty())
+	{
+		const uint32 SizeInBytes = MeshProjectionLODs[LODIndex].ValidSectionIndices.GetBulkDataSize();
+		const uint32 ElementCount = SizeInBytes / sizeof(uint32);
+		MeshProjectionLODs[LODIndex].CachedValidSectionIndices.SetNum(ElementCount);
+
+		uint32* CachedData = MeshProjectionLODs[LODIndex].CachedValidSectionIndices.GetData();
+		const uint32* Data = (uint32*)MeshProjectionLODs[LODIndex].ValidSectionIndices.LockReadOnly();
+		FMemory::Memcpy(CachedData, Data, ElementCount);
+		MeshProjectionLODs[LODIndex].ValidSectionIndices.Unlock();
+	}
+	return MeshProjectionLODs[LODIndex].CachedValidSectionIndices;
+}
 
+static void InternalSerialize(FArchive& Ar, UObject* Owner, FHairStrandsRootBulkData::FMeshProjectionLOD& LOD)
+{
+	const int32 ChunkIndex = 0;
+	bool bAttemptFileMapping = false;
+
+	Ar << LOD.LODIndex;
+	LOD.RootTriangleIndexBuffer.Serialize(Ar, Owner, ChunkIndex, bAttemptFileMapping);
+	LOD.RootTriangleBarycentricBuffer.Serialize(Ar, Owner, ChunkIndex, bAttemptFileMapping);
+	LOD.RestRootTrianglePosition0Buffer.Serialize(Ar, Owner, ChunkIndex, bAttemptFileMapping);
+	LOD.RestRootTrianglePosition1Buffer.Serialize(Ar, Owner, ChunkIndex, bAttemptFileMapping);
+	LOD.RestRootTrianglePosition2Buffer.Serialize(Ar, Owner, ChunkIndex, bAttemptFileMapping);
+
+	Ar << LOD.SampleCount;
+	LOD.MeshInterpolationWeightsBuffer.Serialize(Ar, Owner, ChunkIndex, bAttemptFileMapping);
+	LOD.MeshSampleIndicesBuffer.Serialize(Ar, Owner, ChunkIndex, bAttemptFileMapping);
+	LOD.RestSamplePositionsBuffer.Serialize(Ar, Owner, ChunkIndex, bAttemptFileMapping);
+	LOD.ValidSectionIndices.Serialize(Ar, Owner, ChunkIndex, bAttemptFileMapping);
+}
+
+static void InternalSerialize(FArchive& Ar, UObject* Owner, TArray<FHairStrandsRootBulkData::FMeshProjectionLOD>& LODs)
+{
+	uint32 LODCount = LODs.Num();
+	Ar << LODCount;
+	if (Ar.IsLoading())
+	{
+		LODs.SetNum(LODCount);
+	}
+	for (uint32 LODIt = 0; LODIt < LODCount; ++LODIt)
+	{
+		InternalSerialize(Ar, Owner, LODs[LODIt]);
+	}
+}
+
+void FHairStrandsRootBulkData::Serialize(FArchive& Ar, UObject* Owner)
+{
+	const int32 ChunkIndex = 0;
+	bool bAttemptFileMapping = false;
+
+	if (!Ar.IsObjectReferenceCollector())
+	{
+		Ar << RootCount;
+		Ar << PointCount;
+		VertexToCurveIndexBuffer.Serialize(Ar, Owner, ChunkIndex, bAttemptFileMapping);
+		InternalSerialize(Ar, Owner, MeshProjectionLODs);
+	}
+}
+
+void FHairStrandsRootBulkData::Reset()
+{
+	RootCount = 0;
+	VertexToCurveIndexBuffer.RemoveBulkData();
+	for (FMeshProjectionLOD& LOD : MeshProjectionLODs)
+	{
+		LOD.RootTriangleIndexBuffer.RemoveBulkData();
+		LOD.RootTriangleBarycentricBuffer.RemoveBulkData();
+		LOD.RestRootTrianglePosition0Buffer.RemoveBulkData();
+		LOD.RestRootTrianglePosition1Buffer.RemoveBulkData();
+		LOD.RestRootTrianglePosition2Buffer.RemoveBulkData();
+		LOD.MeshInterpolationWeightsBuffer.RemoveBulkData();
+		LOD.MeshSampleIndicesBuffer.RemoveBulkData();
+		LOD.RestSamplePositionsBuffer.RemoveBulkData();
+		LOD.ValidSectionIndices.RemoveBulkData();
+		LOD.CachedValidSectionIndices.Empty();
+	}
+	MeshProjectionLODs.Empty();
+}
+
+/////////////////////////////////////////////////////////////////////////////////////////
+// Root data
 FHairStrandsRootData::FHairStrandsRootData()
 {
 
@@ -994,34 +1147,6 @@ bool FHairStrandsRootData::HasProjectionData() const
 	return bIsValid;
 }
 
-FArchive& operator<<(FArchive& Ar, FHairStrandsRootData::FMeshProjectionLOD& LOD)
-{
-	Ar << LOD.LODIndex;
-	LOD.RootTriangleIndexBuffer.BulkSerialize(Ar);
-	LOD.RootTriangleBarycentricBuffer.BulkSerialize(Ar);
-	LOD.RestRootTrianglePosition0Buffer.BulkSerialize(Ar);
-	LOD.RestRootTrianglePosition1Buffer.BulkSerialize(Ar);
-	LOD.RestRootTrianglePosition2Buffer.BulkSerialize(Ar);
-
-	Ar << LOD.SampleCount;
-	LOD.MeshInterpolationWeightsBuffer.BulkSerialize(Ar);
-	LOD.MeshSampleIndicesBuffer.BulkSerialize(Ar);
-	LOD.RestSamplePositionsBuffer.BulkSerialize(Ar);
-	LOD.ValidSectionIndices.BulkSerialize(Ar);
-
-	return Ar;
-}
-
-void FHairStrandsRootData::Serialize(FArchive& Ar)
-{
-	if (!Ar.IsObjectReferenceCollector())
-	{
-		Ar << RootCount;
-		VertexToCurveIndexBuffer.BulkSerialize(Ar);
-		Ar << MeshProjectionLODs;
-	}
-}
-
 void FHairStrandsRootData::Reset()
 {
 	RootCount = 0;
diff --git a/Engine/Plugins/Runtime/HairStrands/Source/HairStrandsCore/Private/HairStrandsCore.cpp b/Engine/Plugins/Runtime/HairStrands/Source/HairStrandsCore/Private/HairStrandsCore.cpp
--- a/Engine/Plugins/Runtime/HairStrands/Source/HairStrandsCore/Private/HairStrandsCore.cpp
+++ b/Engine/Plugins/Runtime/HairStrands/Source/HairStrandsCore/Private/HairStrandsCore.cpp
@@ -225,7 +225,7 @@ UGroomBindingAsset* FHairStrandsCore::CreateGroomBindingAsset(EGroomBindingMeshT
 			Out->SourceGeometryCache = Cast<UGeometryCache>(Source);
 			Out->TargetGeometryCache = Cast<UGeometryCache>(Target);
 		}
-		Out->HairGroupDatas.Reserve(GroomAsset->HairGroupsData.Num());
+		Out->HairGroupBulkDatas.Reserve(GroomAsset->HairGroupsData.Num());
 		Out->NumInterpolationPoints = NumInterpolationPoints;
 		Out->MatchingSection = MatchingSection;
 		Out->MarkPackageDirty();
diff --git a/Engine/Plugins/Runtime/HairStrands/Source/HairStrandsCore/Private/HairStrandsMeshProjection.cpp b/Engine/Plugins/Runtime/HairStrands/Source/HairStrandsCore/Private/HairStrandsMeshProjection.cpp
--- a/Engine/Plugins/Runtime/HairStrands/Source/HairStrandsCore/Private/HairStrandsMeshProjection.cpp
+++ b/Engine/Plugins/Runtime/HairStrands/Source/HairStrandsCore/Private/HairStrandsMeshProjection.cpp
@@ -596,7 +596,7 @@ void AddHairStrandUpdateMeshTrianglesPass(
 	FHairStrandsRestRootResource* RestResources,
 	FHairStrandsDeformedRootResource* DeformedResources)
 {
-	const uint32 RootCount = RestResources->RootData.RootCount;
+	const uint32 RootCount = RestResources->BulkData.RootCount;
 	if (RootCount == 0 || LODIndex < 0)
 	{
 		return;
@@ -627,7 +627,7 @@ void AddHairStrandUpdateMeshTrianglesPass(
 	DeformedResources->MeshLODIndex = LODIndex;
 
 	// When the number of section of a mesh is above FHairUpdateMeshTriangleCS::SectionArrayCount, the update is split into several passes
-	const TArray<uint32>& ValidSectionIndices = RestResources->RootData.MeshProjectionLODs[LODIndex].ValidSectionIndices;
+	const TArray<uint32>& ValidSectionIndices = RestResources->BulkData.GetValidSectionIndices(LODIndex);
 	const uint32 ValidSectionCount = ValidSectionIndices.Num();
 	const uint32 PassCount = FMath::DivideAndRoundUp(ValidSectionCount, FHairUpdateMeshTriangleCS::SectionArrayCount);
 
@@ -919,7 +919,7 @@ void AddHairStrandInterpolateMeshTrianglesPass(
 	FHairStrandsRestRootResource* RestResources,
 	FHairStrandsDeformedRootResource* DeformedResources)
 {
-	const uint32 RootCount = RestResources->RootData.RootCount;
+	const uint32 RootCount = RestResources->BulkData.RootCount;
 	if (RootCount == 0 || LODIndex < 0 || LODIndex >= RestResources->LODs.Num() || LODIndex >= DeformedResources->LODs.Num())
 	{
 		return;
@@ -1167,7 +1167,7 @@ void AddHairStrandInitMeshSamplesPass(
 	}
 
 	// When the number of section of a mesh is above FHairUpdateMeshTriangleCS::SectionArrayCount, the update is split into several passes
-	const TArray<uint32>& ValidSectionIndices = RestResources->RootData.MeshProjectionLODs[LODIndex].ValidSectionIndices;
+	const TArray<uint32>& ValidSectionIndices = RestResources->BulkData.GetValidSectionIndices(LODIndex);
 	const uint32 ValidSectionCount = ValidSectionIndices.Num();
 	const uint32 PassCount = FMath::DivideAndRoundUp(ValidSectionCount, FHairInitMeshSamplesCS::SectionArrayCount);
 
@@ -1520,7 +1520,7 @@ static void AddFollicleMaskPass(
 	FHairStrandsRestRootResource* RestResources,
 	FRDGTextureRef OutTexture)
 {
-	const uint32 RootCount = RestResources->RootData.RootCount;
+	const uint32 RootCount = RestResources->BulkData.RootCount;
 	if (LODIndex >= uint32(RestResources->LODs.Num()) || RootCount == 0)
 		return;
 
diff --git a/Engine/Plugins/Runtime/HairStrands/Source/HairStrandsCore/Private/HairStrandsRendering.cpp b/Engine/Plugins/Runtime/HairStrands/Source/HairStrandsCore/Private/HairStrandsRendering.cpp
--- a/Engine/Plugins/Runtime/HairStrands/Source/HairStrandsCore/Private/HairStrandsRendering.cpp
+++ b/Engine/Plugins/Runtime/HairStrands/Source/HairStrandsCore/Private/HairStrandsRendering.cpp
@@ -330,7 +330,7 @@ static void AddDeformSimHairStrandsPass(
 			Parameters->SimVertexToRootIndexBuffer = RegisterAsSRV(GraphBuilder, SimRestRootResources->VertexToCurveIndexBuffer);
 		}
 
-		const uint32 RootCount = SimRestRootResources ? SimRestRootResources->RootData.RootCount : 0;
+		const uint32 RootCount = SimRestRootResources ? SimRestRootResources->BulkData.RootCount : 0;
 		const bool bSupportDynamicMesh = 
 			RootCount > 0 && 
 			MeshLODIndex >= 0 && 
@@ -620,9 +620,9 @@ static void AddHairStrandsInterpolationPass(
 			Parameters->HairDebugMode = 2;
 			
 			if (Info.CullMode == EHairCullMode::Sim)
-				Parameters->HairStrandsCullIndex.Y = Info.ExplicitIndex >= 0 ? Info.ExplicitIndex : FMath::Clamp(uint32(Info.NormalizedIndex * SimRestRootResources->RootData.RootCount), 0u, SimRestRootResources->RootData.RootCount - 1);
+				Parameters->HairStrandsCullIndex.Y = Info.ExplicitIndex >= 0 ? Info.ExplicitIndex : FMath::Clamp(uint32(Info.NormalizedIndex * SimRestRootResources->BulkData.RootCount), 0u, SimRestRootResources->BulkData.RootCount - 1);
 			if (Info.CullMode == EHairCullMode::Render)
-				Parameters->HairStrandsCullIndex.X = Info.ExplicitIndex >= 0 ? Info.ExplicitIndex : FMath::Clamp(uint32(Info.NormalizedIndex * RenRestRootResources->RootData.RootCount), 0u, RenRestRootResources->RootData.RootCount - 1);
+				Parameters->HairStrandsCullIndex.X = Info.ExplicitIndex >= 0 ? Info.ExplicitIndex : FMath::Clamp(uint32(Info.NormalizedIndex * RenRestRootResources->BulkData.RootCount), 0u, RenRestRootResources->BulkData.RootCount - 1);
 		}
 
 		if (Parameters->HairDebugMode > 0 && OutRenderAttributeBuffer)
@@ -635,7 +635,7 @@ static void AddHairStrandsInterpolationPass(
 	const bool bSupportDynamicMesh = 
 		bIsVertexToCurveBuffersValid_Ren &&
 		RenRestRootResources &&
-		RenRestRootResources->RootData.RootCount > 0 && 
+		RenRestRootResources->BulkData.RootCount > 0 && 
 		MeshLODIndex >= 0 && 
 		MeshLODIndex < RenRestRootResources->LODs.Num() &&
 		MeshLODIndex < RenDeformedRootResources->LODs.Num() &&
@@ -954,7 +954,7 @@ static void AddHairCardsDeformationPass(
 		Instance->BindingType == EHairBindingType::Skinning &&
 		bIsVertexToCurveBuffersValid &&
 		RestRootResources &&
-		RestRootResources->RootData.RootCount > 0 &&
+		RestRootResources->BulkData.RootCount > 0 &&
 		MeshLODIndex >= 0 &&
 		MeshLODIndex < RestRootResources->LODs.Num() &&
 		MeshLODIndex < DeformedRootResources->LODs.Num() &&
diff --git a/Engine/Plugins/Runtime/HairStrands/Source/HairStrandsCore/Public/GroomBindingAsset.h b/Engine/Plugins/Runtime/HairStrands/Source/HairStrandsCore/Public/GroomBindingAsset.h
--- a/Engine/Plugins/Runtime/HairStrands/Source/HairStrandsCore/Public/GroomBindingAsset.h
+++ b/Engine/Plugins/Runtime/HairStrands/Source/HairStrandsCore/Public/GroomBindingAsset.h
@@ -47,6 +47,22 @@ enum class EGroomBindingMeshType : uint8
 	GeometryCache
 };
 
+/** Binding data */
+struct FHairRootGroupData
+{
+	FHairStrandsRootData			SimRootData;
+	FHairStrandsRootData			RenRootData;
+	TArray<FHairStrandsRootData>	CardsRootData;
+};
+
+/** Binding bulk data */
+struct FHairGroupBulkData
+{
+	FHairStrandsRootBulkData		SimRootBulkData;
+	FHairStrandsRootBulkData		RenRootBulkData;
+	TArray<FHairStrandsRootBulkData>CardsRootBulkData;
+};
+
 /**
  * Implements an asset that can be used to store binding information between a groom and a skeletal mesh
  */
@@ -109,14 +125,8 @@ public:
 	   when the binding asset is recomputed */
 	TQueue<FHairGroupResource> HairGroupResourcesToDelete;
 
-	struct FHairGroupData
-	{
-		FHairStrandsRootData SimRootData;
-		FHairStrandsRootData RenRootData;
-		TArray<FHairStrandsRootData> CardsRootData;
-	};
-	typedef TArray<FHairGroupData> FHairGroupDatas;
-	FHairGroupDatas HairGroupDatas;
+	/** Root bulk data for each hair gruops */
+	TArray<FHairGroupBulkData> HairGroupBulkDatas;
 
 	//~ Begin UObject Interface.
 	virtual void GetResourceSizeEx(FResourceSizeEx& CumulativeResourceSize) override;
diff --git a/Engine/Plugins/Runtime/HairStrands/Source/HairStrandsCore/Public/GroomBindingBuilder.h b/Engine/Plugins/Runtime/HairStrands/Source/HairStrandsCore/Public/GroomBindingBuilder.h
--- a/Engine/Plugins/Runtime/HairStrands/Source/HairStrandsCore/Public/GroomBindingBuilder.h
+++ b/Engine/Plugins/Runtime/HairStrands/Source/HairStrandsCore/Public/GroomBindingBuilder.h
@@ -21,4 +21,7 @@ struct HAIRSTRANDSCORE_API FGroomBindingBuilder
 
 	// Update the root mesh projection data with unique valid mesh section IDs, based on the projection data
 	static void BuildUniqueSections(FHairStrandsRootData::FMeshProjectionLOD& LOD);
+
+	// Extract root data from bulk data
+	static void GetRootData(FHairStrandsRootData& Out, const FHairStrandsRootBulkData& In);
 };
\ No newline at end of file
diff --git a/Engine/Plugins/Runtime/HairStrands/Source/HairStrandsCore/Public/GroomResources.h b/Engine/Plugins/Runtime/HairStrands/Source/HairStrandsCore/Public/GroomResources.h
--- a/Engine/Plugins/Runtime/HairStrands/Source/HairStrandsCore/Public/GroomResources.h
+++ b/Engine/Plugins/Runtime/HairStrands/Source/HairStrandsCore/Public/GroomResources.h
@@ -56,7 +56,7 @@ struct FHairCommonResource : public FRenderResource
 struct FHairStrandsRestRootResource : public FHairCommonResource
 {
 	/** Build the hair strands resource */
-	FHairStrandsRestRootResource(const FHairStrandsRootData& RootData, EHairStrandsResourcesType CurveType);
+	FHairStrandsRestRootResource(const FHairStrandsRootBulkData& BulkData, EHairStrandsResourcesType CurveType);
 
 	/* Init/Release buffers */
 	virtual void InternalAllocate(FRDGBuilder& GraphBuilder) override;
@@ -120,7 +120,7 @@ struct FHairStrandsRestRootResource : public FHairCommonResource
 	TArray<FLOD> LODs;
 
 	/* Store CPU data for root info & root binding */
-	const FHairStrandsRootData& RootData;
+	const FHairStrandsRootBulkData& BulkData;
 
 	/* Type of curves */
 	const EHairStrandsResourcesType CurveType;
diff --git a/Engine/Plugins/Runtime/HairStrands/Source/HairStrandsCore/Public/HairStrandsDatas.h b/Engine/Plugins/Runtime/HairStrands/Source/HairStrandsCore/Public/HairStrandsDatas.h
--- a/Engine/Plugins/Runtime/HairStrands/Source/HairStrandsCore/Public/HairStrandsDatas.h
+++ b/Engine/Plugins/Runtime/HairStrands/Source/HairStrandsCore/Public/HairStrandsDatas.h
@@ -217,6 +217,7 @@ struct FHairStrandsRaytracingFormat
 struct HAIRSTRANDSCORE_API FHairStrandsIndexFormat
 {
 	using Type = uint32;
+	using BulkType = uint32;
 	static const uint32 ComponentCount = 1;
 	static const uint32 SizeInByte = sizeof(Type);
 	static const EVertexElementType VertexElementType = VET_UInt;
@@ -493,35 +494,90 @@ struct HAIRSTRANDSCORE_API FHairStrandsDebugDatas
 	};
 };
 
-/* Source/CPU data for root resources (GPU resources are stored into FHairStrandsRootResources) */
-struct FHairStrandsRootData
+/* Bulk data for root resources (GPU resources are stored into FHairStrandsRootResources) */
+struct FHairStrandsRootBulkData
 {
-	/** Build the hair strands resource */
-	FHairStrandsRootData();
-	void Serialize(FArchive& Ar);
+	FHairStrandsRootBulkData();
+	void Serialize(FArchive& Ar, UObject* Owner);
 	void Reset();
 	bool HasProjectionData() const;
 	bool IsValid() const { return RootCount > 0; }
+	const TArray<uint32>& GetValidSectionIndices(int32 LODIndex) const;
 
 	uint32 GetDataSize() const
 	{
 		uint32 Total = 0;
-		Total += VertexToCurveIndexBuffer.GetAllocatedSize();
+		Total += VertexToCurveIndexBuffer.IsBulkDataLoaded() ? VertexToCurveIndexBuffer.GetBulkDataSize() : 0u;
 		for (const FMeshProjectionLOD& LOD : MeshProjectionLODs)
 		{
-			Total += LOD.RootTriangleIndexBuffer.GetAllocatedSize();
-			Total += LOD.RootTriangleBarycentricBuffer.GetAllocatedSize();
-			Total += LOD.RestRootTrianglePosition0Buffer.GetAllocatedSize();
-			Total += LOD.RestRootTrianglePosition1Buffer.GetAllocatedSize();
-			Total += LOD.RestRootTrianglePosition2Buffer.GetAllocatedSize();
-			Total += LOD.MeshInterpolationWeightsBuffer.GetAllocatedSize();
-			Total += LOD.MeshSampleIndicesBuffer.GetAllocatedSize();
-			Total += LOD.RestSamplePositionsBuffer.GetAllocatedSize();
-			Total += LOD.ValidSectionIndices.GetAllocatedSize();
+			Total += LOD.RootTriangleIndexBuffer.IsBulkDataLoaded() ?			LOD.RootTriangleIndexBuffer.GetBulkDataSize() : 0u;
+			Total += LOD.RootTriangleBarycentricBuffer.IsBulkDataLoaded() ?		LOD.RootTriangleBarycentricBuffer.GetBulkDataSize() : 0u;
+			Total += LOD.RestRootTrianglePosition0Buffer.IsBulkDataLoaded() ?	LOD.RestRootTrianglePosition0Buffer.GetBulkDataSize() : 0u;
+			Total += LOD.RestRootTrianglePosition1Buffer.IsBulkDataLoaded() ?	LOD.RestRootTrianglePosition1Buffer.GetBulkDataSize() : 0u;
+			Total += LOD.RestRootTrianglePosition2Buffer.IsBulkDataLoaded() ?	LOD.RestRootTrianglePosition2Buffer.GetBulkDataSize() : 0u;
+			Total += LOD.MeshInterpolationWeightsBuffer.IsBulkDataLoaded() ?	LOD.MeshInterpolationWeightsBuffer.GetBulkDataSize() : 0u;
+			Total += LOD.MeshSampleIndicesBuffer.IsBulkDataLoaded() ?			LOD.MeshSampleIndicesBuffer.GetBulkDataSize() : 0u;
+			Total += LOD.RestSamplePositionsBuffer.IsBulkDataLoaded() ?			LOD.RestSamplePositionsBuffer.GetBulkDataSize() : 0u;
+			Total += LOD.ValidSectionIndices.IsBulkDataLoaded() ?				LOD.ValidSectionIndices.GetBulkDataSize() : 0u;
 		}
 		return Total;
 	}
 
+	struct FMeshProjectionLOD
+	{
+		int32 LODIndex = -1;
+
+		/* Triangle on which a root is attached */
+		/* When the projection is done with source to target mesh transfer, the projection indices does not match.
+		   In this case we need to separate index computation. The barycentric coords remain the same however. */
+		FByteBulkData RootTriangleIndexBuffer;
+		FByteBulkData RootTriangleBarycentricBuffer;
+
+		/* Strand hair roots translation and rotation in rest position relative to the bound triangle. Positions are relative to the rest root center */
+		FByteBulkData RestRootTrianglePosition0Buffer;
+		FByteBulkData RestRootTrianglePosition1Buffer;
+		FByteBulkData RestRootTrianglePosition2Buffer;
+
+		/* Number of samples used for the mesh interpolation */
+		uint32 SampleCount = 0;
+
+		/* Store the hair interpolation weights | Size = SamplesCount * SamplesCount */
+		FByteBulkData MeshInterpolationWeightsBuffer;
+
+		/* Store the samples vertex indices */
+		FByteBulkData MeshSampleIndicesBuffer;
+
+		/* Store the samples rest positions */
+		FByteBulkData RestSamplePositionsBuffer;
+
+		/* Store the mesh section indices which are relevant for this root LOD data */
+		FByteBulkData ValidSectionIndices;
+
+		/* When cached valid section indices from bulk to avoid lock/unlock operation everyframe */
+		mutable TArray<uint32> CachedValidSectionIndices;
+	};
+
+	/* Number of roots */
+	uint32 RootCount = 0;
+
+	/* Number of control points */
+	uint32 PointCount = 0;
+
+	/* Curve index for every vertices */
+	FByteBulkData VertexToCurveIndexBuffer;
+
+	/* Store the hair projection information for each mesh LOD */
+	TArray<FMeshProjectionLOD> MeshProjectionLODs;
+};
+
+/* Source data for building root bulk data */
+struct FHairStrandsRootData
+{
+	FHairStrandsRootData();
+	void Reset();
+	bool HasProjectionData() const;
+	bool IsValid() const { return RootCount > 0; }
+
 	struct FMeshProjectionLOD
 	{
 		int32 LODIndex = -1;
@@ -556,6 +612,9 @@ struct FHairStrandsRootData
 	/* Number of roots */
 	uint32 RootCount = 0;
 
+	/* Number of control points */
+	uint32 PointCount = 0;
+
 	/* Curve index for every vertices */
 	TArray<FHairStrandsIndexFormat::Type> VertexToCurveIndexBuffer;
 
