Compilation time optimization
  - Disable codegen optims on extremely long function with lots of FORCEINLINE calls in it
  - 1m14s -> 28.09s for NiagaraDataInterfaceHairStrands.cpp
  - 1m03s -> 20.12s for Module.HairStrandsCore.3_of_4.cpp

#rnx

diff --git a/Engine/Plugins/Runtime/HairStrands/Source/HairStrandsCore/Private/Niagara/NiagaraDataInterfaceHairStrands.cpp b/Engine/Plugins/Runtime/HairStrands/Source/HairStrandsCore/Private/Niagara/NiagaraDataInterfaceHairStrands.cpp
--- a/Engine/Plugins/Runtime/HairStrands/Source/HairStrandsCore/Private/Niagara/NiagaraDataInterfaceHairStrands.cpp
+++ b/Engine/Plugins/Runtime/HairStrands/Source/HairStrandsCore/Private/Niagara/NiagaraDataInterfaceHairStrands.cpp
@@ -1113,6 +1113,10 @@ void UNiagaraDataInterfaceHairStrands::PostInitProperties()
 	}
 }
 
+// Codegen optimization degenerates for very long functions like GetFunctions when combined with the invokation of lots of FORCEINLINE methods.
+// We don't need this code to be particularly fast anyway. The other way to improve this code compilation time would be to split it in multiple functions.
+BEGIN_FUNCTION_BUILD_OPTIMIZATION 
+
 void UNiagaraDataInterfaceHairStrands::GetFunctions(TArray<FNiagaraFunctionSignature>& OutFunctions)
 {
 	{
@@ -2110,6 +2114,7 @@ void UNiagaraDataInterfaceHairStrands::GetFunctions(TArray<FNiagaraFunctionSigna
 		OutFunctions.Add(Sig);
 	}
 }
+END_FUNCTION_BUILD_OPTIMIZATION
 
 DEFINE_NDI_DIRECT_FUNC_BINDER(UNiagaraDataInterfaceHairStrands, GetNumStrands);
 DEFINE_NDI_DIRECT_FUNC_BINDER(UNiagaraDataInterfaceHairStrands, GetStrandSize);
