Simplify how hair strands bulk data filled into GPU buffer.

diff --git a/Engine/Plugins/Runtime/HairStrands/Source/HairStrandsCore/Private/GroomResources.cpp b/Engine/Plugins/Runtime/HairStrands/Source/HairStrandsCore/Private/GroomResources.cpp
--- a/Engine/Plugins/Runtime/HairStrands/Source/HairStrandsCore/Private/GroomResources.cpp
+++ b/Engine/Plugins/Runtime/HairStrands/Source/HairStrandsCore/Private/GroomResources.cpp
@@ -115,6 +115,38 @@ static FRDGBufferRef InternalCreateVertexBuffer(
 	return Buffer;
 }
 
+template<typename FormatType>
+void InternalCreateVertexBufferRDG_FromBulkData(FRDGBuilder& GraphBuilder, const FByteBulkData& InBulkData, uint32 InDataCount, FRDGExternalBuffer& Out, const TCHAR* DebugName, EHairResourceUsageType UsageType)
+{
+	const uint32 InDataCount_Check = InBulkData.GetBulkDataSize() / sizeof(FormatType::BulkType);
+	check(InDataCount_Check == InDataCount);
+
+	const uint32 DataSizeInBytes = FormatType::SizeInByte * InDataCount;
+	if (DataSizeInBytes == 0)
+	{
+		Out.Buffer = nullptr;
+		return;
+	}
+
+	FRDGBufferDesc Desc = FRDGBufferDesc::CreateBufferDesc(FormatType::SizeInByte, InDataCount);
+	if (UsageType != EHairResourceUsageType::Dynamic)
+	{
+		Desc.Usage = Desc.Usage & (~BUF_UnorderedAccess);
+	}
+
+	const FormatType::BulkType* BulkData = (const FormatType::BulkType*)InBulkData.LockReadOnly();
+	FRDGBufferRef Buffer = InternalCreateVertexBuffer(
+		GraphBuilder,
+		DebugName,
+		Desc,
+		BulkData,
+		DataSizeInBytes,
+		ERDGInitialDataFlags::None); // Copy data internally
+	InBulkData.Unlock();
+
+	ConvertToExternalBufferWithViews(GraphBuilder, Buffer, Out, FormatType::Format);
+}
+
 template<typename FormatType>
 void InternalCreateVertexBufferRDG(FRDGBuilder& GraphBuilder, const typename FormatType::Type* InData, uint32 InDataCount, FRDGExternalBuffer& Out, const TCHAR* DebugName, EHairResourceUsageType UsageType, ERDGInitialDataFlags InitialDataFlags)
 {
@@ -604,41 +636,22 @@ void FHairStrandsRestResource::InternalAllocate(FRDGBuilder& GraphBuilder)
 
 	// 1. Lock data, which force the loading data from files (on non-editor build/cooked data). These data are then uploaded to the GPU
 	// 2. A local copy is done by the buffer uploader. This copy is discarded once the uploading is done.
-	const FHairStrandsPositionFormat::Type* BulkDataPositions = (const FHairStrandsPositionFormat::Type*)BulkData.Positions.LockReadOnly();
-	InternalCreateVertexBufferRDG<FHairStrandsPositionFormat>(GraphBuilder, BulkDataPositions, PointCount, PositionBuffer, HAIRSTRANDS_RESOUCE_NAME(CurveType, Hair.StrandsRest_PositionBuffer), EHairResourceUsageType::Static, ERDGInitialDataFlags::None);
-
-	const FHairStrandsAttribute0Format::Type* BulkDataAttributes0 = (const FHairStrandsAttribute0Format::Type*)BulkData.Attributes0.LockReadOnly();
-	InternalCreateVertexBufferRDG<FHairStrandsAttribute0Format>(GraphBuilder, BulkDataAttributes0, PointCount, Attribute0Buffer, HAIRSTRANDS_RESOUCE_NAME(CurveType, Hair.StrandsRest_Attribute0Buffer), EHairResourceUsageType::Static, ERDGInitialDataFlags::None);
+	InternalCreateVertexBufferRDG_FromBulkData<FHairStrandsPositionFormat>(GraphBuilder, BulkData.Positions, PointCount, PositionBuffer, HAIRSTRANDS_RESOUCE_NAME(CurveType, Hair.StrandsRest_PositionBuffer), EHairResourceUsageType::Static);
+	InternalCreateVertexBufferRDG_FromBulkData<FHairStrandsAttribute0Format>(GraphBuilder, BulkData.Attributes0, PointCount, Attribute0Buffer, HAIRSTRANDS_RESOUCE_NAME(CurveType, Hair.StrandsRest_Attribute0Buffer), EHairResourceUsageType::Static);
 
-	const FHairStrandsAttribute1Format::Type* BulkDataAttributes1 = nullptr;
 	if (!!(BulkData.Flags & FHairStrandsBulkData::DataFlags_HasUDIMData))
 	{
-		BulkDataAttributes1 = (const FHairStrandsAttribute1Format::Type*)BulkData.Attributes1.LockReadOnly();
-		InternalCreateVertexBufferRDG<FHairStrandsAttribute1Format>(GraphBuilder, BulkDataAttributes1, PointCount, Attribute1Buffer, HAIRSTRANDS_RESOUCE_NAME(CurveType, Hair.StrandsRest_Attribute1Buffer), EHairResourceUsageType::Static, ERDGInitialDataFlags::None);
+		InternalCreateVertexBufferRDG_FromBulkData<FHairStrandsAttribute1Format>(GraphBuilder, BulkData.Attributes1, PointCount, Attribute1Buffer, HAIRSTRANDS_RESOUCE_NAME(CurveType, Hair.StrandsRest_Attribute1Buffer), EHairResourceUsageType::Static);
 	}
 
-	const FHairStrandsMaterialFormat::Type* BulkDataMaterials = nullptr;
 	if (!!(BulkData.Flags & FHairStrandsBulkData::DataFlags_HasMaterialData))
 	{
-		BulkDataMaterials = (const FHairStrandsMaterialFormat::Type*)BulkData.Materials.LockReadOnly();
-		InternalCreateVertexBufferRDG<FHairStrandsMaterialFormat>(GraphBuilder, BulkDataMaterials, PointCount, MaterialBuffer, HAIRSTRANDS_RESOUCE_NAME(CurveType, Hair.StrandsRest_MaterialBuffer), EHairResourceUsageType::Static, ERDGInitialDataFlags::None);
+		InternalCreateVertexBufferRDG_FromBulkData<FHairStrandsMaterialFormat>(GraphBuilder, BulkData.Materials, PointCount, MaterialBuffer, HAIRSTRANDS_RESOUCE_NAME(CurveType, Hair.StrandsRest_MaterialBuffer), EHairResourceUsageType::Static);
 	}
 
 	TArray<FVector4> RestOffset;
 	RestOffset.Add(BulkData.GetPositionOffset());
-	InternalCreateVertexBufferRDG<FHairStrandsPositionOffsetFormat>(GraphBuilder, RestOffset, PositionOffsetBuffer, HAIRSTRANDS_RESOUCE_NAME(CurveType, Hair.StrandsRest_PositionOffsetBuffer), EHairResourceUsageType::Static, ERDGInitialDataFlags::None);
-
-	// Unlock data once the data has been copied/cached by the uploader
-	BulkData.Positions.Unlock();
-	BulkData.Attributes0.Unlock();
-	if (BulkDataAttributes1)
-	{
-		BulkData.Attributes1.Unlock();
-	}
-	if (BulkDataMaterials)
-	{
-		BulkData.Materials.Unlock();
-	}
+	InternalCreateVertexBufferRDG<FHairStrandsPositionOffsetFormat>(GraphBuilder, RestOffset, PositionOffsetBuffer, HAIRSTRANDS_RESOUCE_NAME(CurveType, Hair.StrandsRest_PositionOffsetBuffer), EHairResourceUsageType::Static);
 }
 
 void AddHairTangentPass(
@@ -692,14 +705,12 @@ void FHairStrandsDeformedResource::InternalAllocate(FRDGBuilder& GraphBuilder)
 {
 	const uint32 PointCount = BulkData.PointCount;
 
-	const FHairStrandsPositionFormat::Type* BulkDataPositions = nullptr;
 	if (bInitializedData)
 	{
 		// 1. Lock data, which force the loading data from files (on non-editor build/cooked data). These data are then uploaded to the GPU
 		// 2. A local copy is done by the buffer uploader. This copy is discarded once the uploading is done.
-		BulkDataPositions = (const FHairStrandsPositionFormat::Type*)BulkData.Positions.LockReadOnly();
-		InternalCreateVertexBufferRDG<FHairStrandsPositionFormat>(GraphBuilder, BulkDataPositions, PointCount, DeformedPositionBuffer[0], HAIRSTRANDS_RESOUCE_NAME(CurveType, Hair.StrandsDeformed_DeformedPositionBuffer0), EHairResourceUsageType::Dynamic, ERDGInitialDataFlags::None);
-		InternalCreateVertexBufferRDG<FHairStrandsPositionFormat>(GraphBuilder, BulkDataPositions, PointCount, DeformedPositionBuffer[1], HAIRSTRANDS_RESOUCE_NAME(CurveType, Hair.StrandsDeformed_DeformedPositionBuffer1), EHairResourceUsageType::Dynamic, ERDGInitialDataFlags::None);
+		InternalCreateVertexBufferRDG_FromBulkData<FHairStrandsPositionFormat>(GraphBuilder, BulkData.Positions, PointCount, DeformedPositionBuffer[0], HAIRSTRANDS_RESOUCE_NAME(CurveType, Hair.StrandsDeformed_DeformedPositionBuffer0), EHairResourceUsageType::Dynamic);
+		InternalCreateVertexBufferRDG_FromBulkData<FHairStrandsPositionFormat>(GraphBuilder, BulkData.Positions, PointCount, DeformedPositionBuffer[1], HAIRSTRANDS_RESOUCE_NAME(CurveType, Hair.StrandsDeformed_DeformedPositionBuffer1), EHairResourceUsageType::Dynamic);
 	}
 	else
 	{
@@ -713,10 +724,6 @@ void FHairStrandsDeformedResource::InternalAllocate(FRDGBuilder& GraphBuilder)
 	InternalCreateVertexBufferRDG<FHairStrandsPositionOffsetFormat>(GraphBuilder, DefaultOffsets, DeformedOffsetBuffer[0], HAIRSTRANDS_RESOUCE_NAME(CurveType, Hair.StrandsDeformed_DeformedOffsetBuffer0), EHairResourceUsageType::Dynamic, ERDGInitialDataFlags::None);
 	InternalCreateVertexBufferRDG<FHairStrandsPositionOffsetFormat>(GraphBuilder, DefaultOffsets, DeformedOffsetBuffer[1], HAIRSTRANDS_RESOUCE_NAME(CurveType, Hair.StrandsDeformed_DeformedOffsetBuffer1), EHairResourceUsageType::Dynamic, ERDGInitialDataFlags::None);
 
-	if (BulkDataPositions)
-	{
-		BulkData.Positions.Unlock();
-	}
 }
 
 void FHairStrandsDeformedResource::InternalRelease()
@@ -816,9 +823,7 @@ void FHairStrandsRestRootResource::InternalAllocate(FRDGBuilder& GraphBuilder)
 
 	if (BulkData.PointCount > 0)
 	{
-		const FHairStrandsIndexFormat::Type* Data = (const FHairStrandsIndexFormat::Type*)BulkData.VertexToCurveIndexBuffer.LockReadOnly();
-		InternalCreateVertexBufferRDG<FHairStrandsIndexFormat>(GraphBuilder, Data, BulkData.PointCount, VertexToCurveIndexBuffer, HAIRSTRANDS_RESOUCE_NAME(CurveType, Hair.StrandsRestRoot_VertexToCurveIndexBuffer), EHairResourceUsageType::Static, ERDGInitialDataFlags::None);
-		BulkData.VertexToCurveIndexBuffer.Unlock();
+		InternalCreateVertexBufferRDG_FromBulkData<FHairStrandsIndexFormat>(GraphBuilder, BulkData.VertexToCurveIndexBuffer, BulkData.PointCount, VertexToCurveIndexBuffer, HAIRSTRANDS_RESOUCE_NAME(CurveType, Hair.StrandsRestRoot_VertexToCurveIndexBuffer), EHairResourceUsageType::Static);
 	}
 }
 
@@ -837,42 +842,18 @@ void FHairStrandsRestRootResource::InternalAllocateLOD(FRDGBuilder& GraphBuilder
 			return;
 		}
 
-		const FHairStrandsCurveTriangleBarycentricFormat::Type* Data_Barycentric= nullptr;
-		const FHairStrandsCurveTriangleIndexFormat::Type* Data_Index			= nullptr;
-		const FHairStrandsMeshTrianglePositionFormat::Type* Data_Triangle0		= nullptr;
-		const FHairStrandsMeshTrianglePositionFormat::Type* Data_Triangle1		= nullptr;
-		const FHairStrandsMeshTrianglePositionFormat::Type* Data_Triangle2		= nullptr;
-
 		const FHairStrandsRootBulkData::FMeshProjectionLOD& CPUData = BulkData.MeshProjectionLODs[LODIndex];
 		const bool bHasValidCPUData = CPUData.RootTriangleBarycentricBuffer.GetBulkDataSize() > 0;
 		if (bHasValidCPUData)
 		{
 			GPUData.Status = FLOD::EStatus::Completed;
-			{
-				check(CPUData.RootTriangleBarycentricBuffer.GetBulkDataSize() > 0);
-				Data_Barycentric = (FHairStrandsCurveTriangleBarycentricFormat::Type*)CPUData.RootTriangleBarycentricBuffer.LockReadOnly();
-				InternalCreateVertexBufferRDG<FHairStrandsCurveTriangleBarycentricFormat>(GraphBuilder, Data_Barycentric, BulkData.RootCount, GPUData.RootTriangleBarycentricBuffer, HAIRSTRANDS_RESOUCE_NAME(CurveType, Hair.StrandsRestRoot_RootTriangleBarycentricBuffer), EHairResourceUsageType::Static, ERDGInitialDataFlags::None);
-			}
-			{
-				check(CPUData.RootTriangleIndexBuffer.GetBulkDataSize() > 0);
-				Data_Index = (FHairStrandsCurveTriangleIndexFormat::Type*)CPUData.RootTriangleIndexBuffer.LockReadOnly();
-				InternalCreateVertexBufferRDG<FHairStrandsCurveTriangleIndexFormat>(GraphBuilder, Data_Index, BulkData.RootCount, GPUData.RootTriangleIndexBuffer, HAIRSTRANDS_RESOUCE_NAME(CurveType, Hair.StrandsRestRoot_RootTriangleIndexBuffer), EHairResourceUsageType::Static, ERDGInitialDataFlags::None);
-			}
-			{
-				check(CPUData.RestRootTrianglePosition0Buffer.GetBulkDataSize() > 0);
-				Data_Triangle0 = (FHairStrandsMeshTrianglePositionFormat::Type*)CPUData.RestRootTrianglePosition0Buffer.LockReadOnly();
-				InternalCreateVertexBufferRDG<FHairStrandsMeshTrianglePositionFormat>(GraphBuilder, Data_Triangle0, BulkData.RootCount, GPUData.RestRootTrianglePosition0Buffer, HAIRSTRANDS_RESOUCE_NAME(CurveType, Hair.StrandsRestRoot_RestRootTrianglePosition0Buffer), EHairResourceUsageType::Static, ERDGInitialDataFlags::None);
-			}
-			{
-				check(CPUData.RestRootTrianglePosition1Buffer.GetBulkDataSize() > 0);
-				Data_Triangle1 = (FHairStrandsMeshTrianglePositionFormat::Type*)CPUData.RestRootTrianglePosition1Buffer.LockReadOnly();
-				InternalCreateVertexBufferRDG<FHairStrandsMeshTrianglePositionFormat>(GraphBuilder, Data_Triangle1, BulkData.RootCount, GPUData.RestRootTrianglePosition1Buffer, HAIRSTRANDS_RESOUCE_NAME(CurveType, Hair.StrandsRestRoot_RestRootTrianglePosition1Buffer), EHairResourceUsageType::Static, ERDGInitialDataFlags::None);
-			}
-			{
-				check(CPUData.RestRootTrianglePosition2Buffer.GetBulkDataSize() > 0);
-				Data_Triangle2 = (FHairStrandsMeshTrianglePositionFormat::Type*)CPUData.RestRootTrianglePosition2Buffer.LockReadOnly();
-				InternalCreateVertexBufferRDG<FHairStrandsMeshTrianglePositionFormat>(GraphBuilder, Data_Triangle2, BulkData.RootCount, GPUData.RestRootTrianglePosition2Buffer, HAIRSTRANDS_RESOUCE_NAME(CurveType, Hair.StrandsRestRoot_RestRootTrianglePosition2Buffer), EHairResourceUsageType::Static, ERDGInitialDataFlags::None);
-			}
+
+			InternalCreateVertexBufferRDG_FromBulkData<FHairStrandsCurveTriangleBarycentricFormat>(GraphBuilder, CPUData.RootTriangleBarycentricBuffer, BulkData.RootCount, GPUData.RootTriangleBarycentricBuffer, HAIRSTRANDS_RESOUCE_NAME(CurveType, Hair.StrandsRestRoot_RootTriangleBarycentricBuffer), EHairResourceUsageType::Static);
+			InternalCreateVertexBufferRDG_FromBulkData<FHairStrandsCurveTriangleIndexFormat>(GraphBuilder, CPUData.RootTriangleIndexBuffer, BulkData.RootCount, GPUData.RootTriangleIndexBuffer, HAIRSTRANDS_RESOUCE_NAME(CurveType, Hair.StrandsRestRoot_RootTriangleIndexBuffer), EHairResourceUsageType::Static);
+
+			InternalCreateVertexBufferRDG_FromBulkData<FHairStrandsMeshTrianglePositionFormat>(GraphBuilder, CPUData.RestRootTrianglePosition0Buffer, BulkData.RootCount, GPUData.RestRootTrianglePosition0Buffer, HAIRSTRANDS_RESOUCE_NAME(CurveType, Hair.StrandsRestRoot_RestRootTrianglePosition0Buffer), EHairResourceUsageType::Static);
+			InternalCreateVertexBufferRDG_FromBulkData<FHairStrandsMeshTrianglePositionFormat>(GraphBuilder, CPUData.RestRootTrianglePosition1Buffer, BulkData.RootCount, GPUData.RestRootTrianglePosition1Buffer, HAIRSTRANDS_RESOUCE_NAME(CurveType, Hair.StrandsRestRoot_RestRootTrianglePosition1Buffer), EHairResourceUsageType::Static);
+			InternalCreateVertexBufferRDG_FromBulkData<FHairStrandsMeshTrianglePositionFormat>(GraphBuilder, CPUData.RestRootTrianglePosition2Buffer, BulkData.RootCount, GPUData.RestRootTrianglePosition2Buffer, HAIRSTRANDS_RESOUCE_NAME(CurveType, Hair.StrandsRestRoot_RestRootTrianglePosition2Buffer), EHairResourceUsageType::Static);
 		}
 		else
 		{
@@ -887,30 +868,14 @@ void FHairStrandsRestRootResource::InternalAllocateLOD(FRDGBuilder& GraphBuilder
 			InternalCreateVertexBufferRDG<FHairStrandsMeshTrianglePositionFormat>(GraphBuilder, BulkData.RootCount, GPUData.RestRootTrianglePosition2Buffer, HAIRSTRANDS_RESOUCE_NAME(CurveType, Hair.StrandsRestRoot_RestRootTrianglePosition2Buffer), EHairResourceUsageType::Dynamic);
 		}
 
-		FHairStrandsIndexFormat::Type* DataSample_SampleIndex = nullptr;
-		FHairStrandsWeightFormat::Type* DataSample_InterpolationWeight = nullptr;
-		FHairStrandsMeshTrianglePositionFormat::Type* DataSample_RestSample = nullptr;
-
 		GPUData.SampleCount = CPUData.SampleCount;
 		const bool bHasValidCPUWeights = CPUData.MeshSampleIndicesBuffer.GetBulkDataSize() > 0;
 		if (bHasValidCPUWeights)
 		{
-			{
-				const uint32 InteroplationWeightCount = CPUData.MeshInterpolationWeightsBuffer.GetBulkDataSize() / sizeof(FHairStrandsWeightFormat::Type);
-				//check(InteroplationWeightCount == (CPUData.SampleCount+4) * (CPUData.SampleCount+4));
-				DataSample_InterpolationWeight = (FHairStrandsWeightFormat::Type*)CPUData.MeshInterpolationWeightsBuffer.LockReadOnly();
-				InternalCreateVertexBufferRDG<FHairStrandsWeightFormat>(GraphBuilder, DataSample_InterpolationWeight, InteroplationWeightCount, GPUData.MeshInterpolationWeightsBuffer, HAIRSTRANDS_RESOUCE_NAME(CurveType, Hair.StrandsRestRoot_MeshInterpolationWeightsBuffer), EHairResourceUsageType::Static, ERDGInitialDataFlags::None);
-			}
-			{
-				check(CPUData.MeshSampleIndicesBuffer.GetBulkDataSize()/sizeof(FHairStrandsIndexFormat::Type) == CPUData.SampleCount);
-				DataSample_SampleIndex = (FHairStrandsIndexFormat::Type*)CPUData.MeshSampleIndicesBuffer.LockReadOnly();
-				InternalCreateVertexBufferRDG<FHairStrandsIndexFormat>(GraphBuilder, DataSample_SampleIndex, CPUData.SampleCount, GPUData.MeshSampleIndicesBuffer, HAIRSTRANDS_RESOUCE_NAME(CurveType, Hair.StrandsRestRoot_MeshSampleIndicesBuffer), EHairResourceUsageType::Static, ERDGInitialDataFlags::None);
-			}
-			{
-				check(CPUData.RestSamplePositionsBuffer.GetBulkDataSize() / sizeof(FHairStrandsMeshTrianglePositionFormat::Type) == CPUData.SampleCount);
-				DataSample_RestSample = (FHairStrandsMeshTrianglePositionFormat::Type*)CPUData.RestSamplePositionsBuffer.LockReadOnly();
-				InternalCreateVertexBufferRDG<FHairStrandsMeshTrianglePositionFormat>(GraphBuilder, DataSample_RestSample, CPUData.SampleCount, GPUData.RestSamplePositionsBuffer, HAIRSTRANDS_RESOUCE_NAME(CurveType, Hair.StrandsRestRoot_RestSamplePositionsBuffer), EHairResourceUsageType::Static, ERDGInitialDataFlags::None);
-			}
+			const uint32 InteroplationWeightCount = CPUData.MeshInterpolationWeightsBuffer.GetBulkDataSize() / sizeof(FHairStrandsWeightFormat::Type);
+			InternalCreateVertexBufferRDG_FromBulkData<FHairStrandsWeightFormat>(GraphBuilder, CPUData.MeshInterpolationWeightsBuffer, InteroplationWeightCount, GPUData.MeshInterpolationWeightsBuffer, HAIRSTRANDS_RESOUCE_NAME(CurveType, Hair.StrandsRestRoot_MeshInterpolationWeightsBuffer), EHairResourceUsageType::Static);
+			InternalCreateVertexBufferRDG_FromBulkData<FHairStrandsIndexFormat>(GraphBuilder, CPUData.MeshSampleIndicesBuffer, CPUData.SampleCount, GPUData.MeshSampleIndicesBuffer, HAIRSTRANDS_RESOUCE_NAME(CurveType, Hair.StrandsRestRoot_MeshSampleIndicesBuffer), EHairResourceUsageType::Static);
+			InternalCreateVertexBufferRDG_FromBulkData<FHairStrandsMeshTrianglePositionFormat>(GraphBuilder, CPUData.RestSamplePositionsBuffer, CPUData.SampleCount, GPUData.RestSamplePositionsBuffer, HAIRSTRANDS_RESOUCE_NAME(CurveType, Hair.StrandsRestRoot_RestSamplePositionsBuffer), EHairResourceUsageType::Static);
 		}
 		else
 		{
@@ -919,18 +884,6 @@ void FHairStrandsRestRootResource::InternalAllocateLOD(FRDGBuilder& GraphBuilder
 			InternalCreateVertexBufferRDG<FHairStrandsIndexFormat>(GraphBuilder, CPUData.SampleCount, GPUData.MeshSampleIndicesBuffer, HAIRSTRANDS_RESOUCE_NAME(CurveType, Hair.StrandsRestRoot_MeshSampleIndicesBuffer), EHairResourceUsageType::Dynamic);
 			InternalCreateVertexBufferRDG<FHairStrandsMeshTrianglePositionFormat>(GraphBuilder, CPUData.SampleCount, GPUData.RestSamplePositionsBuffer, HAIRSTRANDS_RESOUCE_NAME(CurveType, Hair.StrandsRestRoot_RestSamplePositionsBuffer), EHairResourceUsageType::Dynamic);
 		}
-
-		// Unlock mapped bulk buffer
-		if (Data_Barycentric)				{ CPUData.RootTriangleBarycentricBuffer.Unlock(); }
-		if (Data_Index)						{ CPUData.RootTriangleIndexBuffer.Unlock(); }
-		if (Data_Triangle0)					{ CPUData.RestRootTrianglePosition0Buffer.Unlock(); }
-		if (Data_Triangle1)					{ CPUData.RestRootTrianglePosition1Buffer.Unlock(); }
-		if (Data_Triangle2)					{ CPUData.RestRootTrianglePosition2Buffer.Unlock(); }
-
-		if (DataSample_InterpolationWeight) { CPUData.MeshInterpolationWeightsBuffer.Unlock(); }
-		if (DataSample_SampleIndex)			{ CPUData.MeshSampleIndicesBuffer.Unlock(); }
-		if (DataSample_RestSample)			{ CPUData.RestSamplePositionsBuffer.Unlock(); }
-
 	}
 }
 
@@ -1046,7 +999,7 @@ const TArray<uint32>& FHairStrandsRootBulkData::GetValidSectionIndices(int32 LOD
 
 		uint32* CachedData = MeshProjectionLODs[LODIndex].CachedValidSectionIndices.GetData();
 		const uint32* Data = (uint32*)MeshProjectionLODs[LODIndex].ValidSectionIndices.LockReadOnly();
-		FMemory::Memcpy(CachedData, Data, ElementCount);
+		FMemory::Memcpy(CachedData, Data, SizeInBytes);
 		MeshProjectionLODs[LODIndex].ValidSectionIndices.Unlock();
 	}
 	return MeshProjectionLODs[LODIndex].CachedValidSectionIndices;
