Fix CIS issue.

diff --git a/Engine/Plugins/Runtime/HairStrands/Source/HairStrandsCore/Private/GroomResources.cpp b/Engine/Plugins/Runtime/HairStrands/Source/HairStrandsCore/Private/GroomResources.cpp
--- a/Engine/Plugins/Runtime/HairStrands/Source/HairStrandsCore/Private/GroomResources.cpp
+++ b/Engine/Plugins/Runtime/HairStrands/Source/HairStrandsCore/Private/GroomResources.cpp
@@ -118,7 +118,7 @@ static FRDGBufferRef InternalCreateVertexBuffer(
 template<typename FormatType>
 void InternalCreateVertexBufferRDG_FromBulkData(FRDGBuilder& GraphBuilder, const FByteBulkData& InBulkData, uint32 InDataCount, FRDGExternalBuffer& Out, const TCHAR* DebugName, EHairResourceUsageType UsageType)
 {
-	const uint32 InDataCount_Check = InBulkData.GetBulkDataSize() / sizeof(FormatType::BulkType);
+	const uint32 InDataCount_Check = InBulkData.GetBulkDataSize() / sizeof(typename FormatType::BulkType);
 	check(InDataCount_Check == InDataCount);
 
 	const uint32 DataSizeInBytes = FormatType::SizeInByte * InDataCount;
@@ -134,7 +134,7 @@ void InternalCreateVertexBufferRDG_FromBulkData(FRDGBuilder& GraphBuilder, const
 		Desc.Usage = Desc.Usage & (~BUF_UnorderedAccess);
 	}
 
-	const FormatType::BulkType* BulkData = (const FormatType::BulkType*)InBulkData.LockReadOnly();
+	const FormatType::BulkType* BulkData = (const typename FormatType::BulkType*)InBulkData.LockReadOnly();
 	FRDGBufferRef Buffer = InternalCreateVertexBuffer(
 		GraphBuilder,
 		DebugName,
