Change hair interpolation data to byte bulk data and on-demand data loading.

diff --git a/Engine/Plugins/Runtime/HairStrands/Source/HairStrandsCore/Private/GroomAsset.cpp b/Engine/Plugins/Runtime/HairStrands/Source/HairStrandsCore/Private/GroomAsset.cpp
--- a/Engine/Plugins/Runtime/HairStrands/Source/HairStrandsCore/Private/GroomAsset.cpp
+++ b/Engine/Plugins/Runtime/HairStrands/Source/HairStrandsCore/Private/GroomAsset.cpp
@@ -1441,7 +1441,7 @@ static void InternalSerialize(FArchive& Ar, UObject* Owner, FHairGroupData::FCar
 		CardLODData.InterpolationBulkData.Serialize(Ar);
 
 		CardLODData.Guides.BulkData.Serialize(Ar, Owner);
-		CardLODData.Guides.InterpolationBulkData.Serialize(Ar);
+		CardLODData.Guides.InterpolationBulkData.Serialize(Ar, Owner);
 	}
 	else
 	{
@@ -1454,7 +1454,7 @@ static void InternalSerialize(FArchive& Ar, UObject* Owner, FHairGroupData::FCar
 
 		FHairGroupData::FBaseWithInterpolation NoGuideData;
 		NoGuideData.BulkData.Serialize(Ar, Owner);
-		NoGuideData.InterpolationBulkData.Serialize(Ar);
+		NoGuideData.InterpolationBulkData.Serialize(Ar, Owner);
 	}
 }
 
@@ -1481,7 +1481,7 @@ static void InternalSerialize(FArchive& Ar, UObject* Owner, FHairGroupData& Grou
 	{
 		GroupData.Strands.BulkData.Serialize(Ar, Owner);
 		GroupData.Guides.BulkData.Serialize(Ar, Owner);
-		GroupData.Strands.InterpolationBulkData.Serialize(Ar);
+		GroupData.Strands.InterpolationBulkData.Serialize(Ar, Owner);
 
 	}
 	else
@@ -1491,7 +1491,7 @@ static void InternalSerialize(FArchive& Ar, UObject* Owner, FHairGroupData& Grou
 		// simulation or RBF deformation) on the target platform
 		NoStrandsData.Strands.BulkData.Serialize(Ar, Owner);
 		GroupData.Guides.BulkData.Serialize(Ar, Owner);
-		NoStrandsData.Strands.InterpolationBulkData.Serialize(Ar);
+		NoStrandsData.Strands.InterpolationBulkData.Serialize(Ar, Owner);
 	}
 
 	if (Ar.CustomVer(FAnimObjectVersion::GUID) >= FAnimObjectVersion::SerializeHairClusterCullingData)
@@ -1667,7 +1667,7 @@ void UGroomAsset::SetHairWidth(float Width)
 // differences, etc.) replace the version GUID below with a new one.
 // In case of merge conflicts with DDC versions, you MUST generate a new GUID
 // and set this new GUID as the version.
-#define GROOM_DERIVED_DATA_VERSION TEXT("7188891507314F3D9F79461BA21DF850")
+#define GROOM_DERIVED_DATA_VERSION TEXT("DDFF529D572543C1BB7BD4E49E7B7B5E")
 
 #if WITH_EDITORONLY_DATA
 
diff --git a/Engine/Plugins/Runtime/HairStrands/Source/HairStrandsCore/Private/GroomBuilder.cpp b/Engine/Plugins/Runtime/HairStrands/Source/HairStrandsCore/Private/GroomBuilder.cpp
--- a/Engine/Plugins/Runtime/HairStrands/Source/HairStrandsCore/Private/GroomBuilder.cpp
+++ b/Engine/Plugins/Runtime/HairStrands/Source/HairStrandsCore/Private/GroomBuilder.cpp
@@ -33,7 +33,7 @@ static FAutoConsoleVariableRef CVarHairClusterBuilder_MaxVoxelResolution(TEXT("r
 FString FGroomBuilder::GetVersion()
 {
 	// Important to update the version when groom building changes
-	return TEXT("2f");
+	return TEXT("2j");
 }
 
 namespace FHairStrandsDecimation
@@ -1203,11 +1203,15 @@ namespace HairInterpolationBuilder
 		auto LowerPart = [](uint32 Index) -> uint16 { return uint16(Index & 0xFFFF); };
 		auto UpperPart = [](uint32 Index) -> uint8  { return uint8((Index >> 16) & 0xFF); };
 
+
+		OutBulkData.Flags = FHairStrandsInterpolationBulkData::DataFlags_HasData;
+		OutBulkData.PointCount = PointCount;
+
 		if (HairInterpolation.bUseUniqueGuide)
 		{
-			OutBulkData.Flags = FHairStrandsInterpolationBulkData::DataFlags_HasSingleGuideData;
+			OutBulkData.Flags |= FHairStrandsInterpolationBulkData::DataFlags_HasSingleGuideData;
 
-			TArray<FHairStrandsInterpolationFormat::Type>& OutPointsInterpolation = OutBulkData.Interpolation;
+			TArray<FHairStrandsInterpolationFormat::Type> OutPointsInterpolation;
 			OutPointsInterpolation.SetNum(PointCount * FHairStrandsInterpolationFormat::ComponentCount);
 
 			for (uint32 PointIndex = 0; PointIndex < PointCount; ++PointIndex)
@@ -1221,13 +1225,13 @@ namespace HairInterpolationBuilder
 				OutInterp.VertexGuideIndex1 = UpperPart(Indices[0]);
 				OutInterp.VertexLerp		= S[0] * 255.f;
 			}
+
+			HairStrandsBuilder::CopyToBulkData<FHairStrandsInterpolationFormat>(OutBulkData.Interpolation, OutPointsInterpolation);
 		}
 		else
 		{
-			OutBulkData.Flags = 0;
-
-			TArray<FHairStrandsInterpolation0Format::Type>& OutPointsInterpolation0 = OutBulkData.Interpolation0;
-			TArray<FHairStrandsInterpolation1Format::Type>& OutPointsInterpolation1 = OutBulkData.Interpolation1;
+			TArray<FHairStrandsInterpolation0Format::Type> OutPointsInterpolation0;
+			TArray<FHairStrandsInterpolation1Format::Type> OutPointsInterpolation1;
 
 			OutPointsInterpolation0.SetNum(PointCount * FHairStrandsInterpolation0Format::ComponentCount);
 			OutPointsInterpolation1.SetNum(PointCount * FHairStrandsInterpolation1Format::ComponentCount);
@@ -1255,19 +1259,29 @@ namespace HairInterpolationBuilder
 				OutInterp1.Pad0			= 0;
 				OutInterp1.Pad1			= 0;
 			}
+
+			HairStrandsBuilder::CopyToBulkData<FHairStrandsInterpolation0Format>(OutBulkData.Interpolation0, OutPointsInterpolation0);
+			HairStrandsBuilder::CopyToBulkData<FHairStrandsInterpolation1Format>(OutBulkData.Interpolation1, OutPointsInterpolation1);
 		}
 
-		const uint32 RootCount = SimDatas.GetNumCurves();
-		OutBulkData.SimRootPointIndex.SetNum(SimDatas.GetNumPoints());
-		for (uint32 CurveIndex = 0; CurveIndex < RootCount; ++CurveIndex)
 		{
-			const uint16 SimPointCount = SimDatas.StrandsCurves.CurvesCount[CurveIndex];
-			const uint32 SimPointOffset = SimDatas.StrandsCurves.CurvesOffset[CurveIndex];
-			for (uint32 PointIndex = 0; PointIndex < SimPointCount; ++PointIndex)
+			OutBulkData.SimPointCount = SimDatas.GetNumPoints();
+			const uint32 SimCurveCount = SimDatas.GetNumCurves();
+
+			TArray<FHairStrandsRootIndexFormat::Type> SimRootPointIndex;
+			SimRootPointIndex.SetNum(OutBulkData.SimPointCount);
+			for (uint32 CurveIndex = 0; CurveIndex < SimCurveCount; ++CurveIndex)
 			{
-				OutBulkData.SimRootPointIndex[PointIndex + SimPointOffset] = SimPointOffset;
+				const uint16 SimPointCount = SimDatas.StrandsCurves.CurvesCount[CurveIndex];
+				const uint32 SimPointOffset = SimDatas.StrandsCurves.CurvesOffset[CurveIndex];
+				for (uint32 PointIndex = 0; PointIndex < SimPointCount; ++PointIndex)
+				{
+					SimRootPointIndex[PointIndex + SimPointOffset] = SimPointOffset;
+				}
 			}
-		}		
+
+			HairStrandsBuilder::CopyToBulkData<FHairStrandsRootIndexFormat>(OutBulkData.SimRootPointIndex, SimRootPointIndex);
+		}
 	}
 
 	/** Fill the GroomAsset with the interpolation data that exists in the HairDescription */
diff --git a/Engine/Plugins/Runtime/HairStrands/Source/HairStrandsCore/Private/GroomResources.cpp b/Engine/Plugins/Runtime/HairStrands/Source/HairStrandsCore/Private/GroomResources.cpp
--- a/Engine/Plugins/Runtime/HairStrands/Source/HairStrandsCore/Private/GroomResources.cpp
+++ b/Engine/Plugins/Runtime/HairStrands/Source/HairStrandsCore/Private/GroomResources.cpp
@@ -134,7 +134,7 @@ void InternalCreateVertexBufferRDG_FromBulkData(FRDGBuilder& GraphBuilder, const
 		Desc.Usage = Desc.Usage & (~BUF_UnorderedAccess);
 	}
 
-	const FormatType::BulkType* BulkData = (const typename FormatType::BulkType*)InBulkData.LockReadOnly();
+	const typename FormatType::BulkType* BulkData = (const typename FormatType::BulkType*)InBulkData.LockReadOnly();
 	FRDGBufferRef Buffer = InternalCreateVertexBuffer(
 		GraphBuilder,
 		DebugName,
@@ -1120,14 +1120,14 @@ void FHairStrandsInterpolationResource::InternalAllocate(FRDGBuilder& GraphBuild
 	const bool bUseSingleGuide = !!(BulkData.Flags & FHairStrandsInterpolationBulkData::DataFlags_HasSingleGuideData);
 	if (bUseSingleGuide)
 	{
-		InternalCreateVertexBufferRDG<FHairStrandsInterpolationFormat>(GraphBuilder, BulkData.Interpolation, InterpolationBuffer, TEXT("Hair.StrandsInterpolation_InterpolationBuffer"), EHairResourceUsageType::Static);
+		InternalCreateVertexBufferRDG_FromBulkData<FHairStrandsInterpolationFormat>(GraphBuilder, BulkData.Interpolation, BulkData.PointCount, InterpolationBuffer, TEXT("Hair.StrandsInterpolation_InterpolationBuffer"), EHairResourceUsageType::Static);
 	}
 	else
 	{
-		InternalCreateVertexBufferRDG<FHairStrandsInterpolation0Format>(GraphBuilder, BulkData.Interpolation0, Interpolation0Buffer, TEXT("Hair.StrandsInterpolation_Interpolation0Buffer"), EHairResourceUsageType::Static);
-		InternalCreateVertexBufferRDG<FHairStrandsInterpolation1Format>(GraphBuilder, BulkData.Interpolation1, Interpolation1Buffer, TEXT("Hair.StrandsInterpolation_Interpolation1Buffer"), EHairResourceUsageType::Static);
+		InternalCreateVertexBufferRDG_FromBulkData<FHairStrandsInterpolation0Format>(GraphBuilder, BulkData.Interpolation0, BulkData.PointCount, Interpolation0Buffer, TEXT("Hair.StrandsInterpolation_Interpolation0Buffer"), EHairResourceUsageType::Static);
+		InternalCreateVertexBufferRDG_FromBulkData<FHairStrandsInterpolation1Format>(GraphBuilder, BulkData.Interpolation1, BulkData.PointCount, Interpolation1Buffer, TEXT("Hair.StrandsInterpolation_Interpolation1Buffer"), EHairResourceUsageType::Static);
 	}
-	InternalCreateVertexBufferRDG<FHairStrandsRootIndexFormat>(GraphBuilder, BulkData.SimRootPointIndex, SimRootPointIndexBuffer, TEXT("Hair.StrandsInterpolation_SimRootPointIndex"), EHairResourceUsageType::Static);
+	InternalCreateVertexBufferRDG_FromBulkData<FHairStrandsRootIndexFormat>(GraphBuilder, BulkData.SimRootPointIndex, BulkData.SimPointCount, SimRootPointIndexBuffer, TEXT("Hair.StrandsInterpolation_SimRootPointIndex"), EHairResourceUsageType::Static);
 }
 
 void FHairStrandsInterpolationResource::InternalRelease()
diff --git a/Engine/Plugins/Runtime/HairStrands/Source/HairStrandsCore/Private/HairStrandsDatas.cpp b/Engine/Plugins/Runtime/HairStrands/Source/HairStrandsCore/Private/HairStrandsDatas.cpp
--- a/Engine/Plugins/Runtime/HairStrands/Source/HairStrandsCore/Private/HairStrandsDatas.cpp
+++ b/Engine/Plugins/Runtime/HairStrands/Source/HairStrandsCore/Private/HairStrandsDatas.cpp
@@ -165,30 +165,48 @@ FArchive& operator<<(FArchive& Ar, FHairInterpolation1Vertex& Vertex)
 
 void FHairStrandsInterpolationBulkData::Reset()
 {
-	Interpolation.Empty();
-	Interpolation0.Empty();
-	Interpolation1.Empty();
-	SimRootPointIndex.Empty();
+	Flags = 0;
+	PointCount = 0;
+	SimPointCount = 0;
+	Interpolation.RemoveBulkData();
+	Interpolation0.RemoveBulkData();
+	Interpolation1.RemoveBulkData();
+	SimRootPointIndex.RemoveBulkData();
 }
 
-void FHairStrandsInterpolationBulkData::Serialize(FArchive& Ar)
+void FHairStrandsInterpolationBulkData::Serialize(FArchive& Ar, UObject* Owner)
 {
 	static_assert(sizeof(FHairInterpolationVertex::BulkType) == sizeof(FHairInterpolationVertex));
 	static_assert(sizeof(FHairInterpolation0Vertex::BulkType) == sizeof(FHairInterpolation0Vertex));
 	static_assert(sizeof(FHairInterpolation1Vertex::BulkType) == sizeof(FHairInterpolation1Vertex));
 	static_assert(sizeof(FHairStrandsRootIndexFormat::BulkType) == sizeof(FHairStrandsRootIndexFormat::Type));
 
+	const uint32 BulkFlags = BULKDATA_Force_NOT_InlinePayload | BULKDATA_SerializeCompressed;
+	Interpolation.SetBulkDataFlags(BulkFlags);
+	Interpolation0.SetBulkDataFlags(BulkFlags);
+	Interpolation1.SetBulkDataFlags(BulkFlags);
+	SimRootPointIndex.SetBulkDataFlags(BulkFlags);
+
 	Ar << Flags;
-	if (!!(Flags & DataFlags_HasSingleGuideData))
-	{
-		Interpolation.BulkSerialize(Ar);
-	}
-	else
+	Ar << PointCount;
+	Ar << SimPointCount;
+
+	if (!!(Flags & DataFlags_HasData))
 	{
-		Interpolation0.BulkSerialize(Ar);
-		Interpolation1.BulkSerialize(Ar);
+		const int32 ChunkIndex = 0;
+		bool bAttemptFileMapping = false;
+
+		if (!!(Flags & DataFlags_HasSingleGuideData))
+		{
+			Interpolation.Serialize(Ar, Owner, ChunkIndex, bAttemptFileMapping);
+		}
+		else
+		{
+			Interpolation0.Serialize(Ar, Owner, ChunkIndex, bAttemptFileMapping);
+			Interpolation1.Serialize(Ar, Owner, ChunkIndex, bAttemptFileMapping);
+		}
+		SimRootPointIndex.Serialize(Ar, Owner, ChunkIndex, bAttemptFileMapping);
 	}
-	SimRootPointIndex.BulkSerialize(Ar);
 }
 
 void FHairStrandsBulkData::Serialize(FArchive& Ar, UObject* Owner)
diff --git a/Engine/Plugins/Runtime/HairStrands/Source/HairStrandsCore/Public/GroomAsset.h b/Engine/Plugins/Runtime/HairStrands/Source/HairStrandsCore/Public/GroomAsset.h
--- a/Engine/Plugins/Runtime/HairStrands/Source/HairStrandsCore/Public/GroomAsset.h
+++ b/Engine/Plugins/Runtime/HairStrands/Source/HairStrandsCore/Public/GroomAsset.h
@@ -102,10 +102,10 @@ struct HAIRSTRANDSCORE_API FHairGroupData
 		{
 			uint32 Total = 0;
 			Total += FBase::GetDataSize();
-			Total += InterpolationBulkData.Interpolation.GetAllocatedSize();
-			Total += InterpolationBulkData.Interpolation0.GetAllocatedSize();
-			Total += InterpolationBulkData.Interpolation1.GetAllocatedSize();
-			Total += InterpolationBulkData.SimRootPointIndex.GetAllocatedSize();
+			Total += InterpolationBulkData.Interpolation.IsBulkDataLoaded()		? InterpolationBulkData.Interpolation.GetBulkDataSize() : 0;
+			Total += InterpolationBulkData.Interpolation0.IsBulkDataLoaded()	? InterpolationBulkData.Interpolation0.GetBulkDataSize() : 0;
+			Total += InterpolationBulkData.Interpolation1.IsBulkDataLoaded()	? InterpolationBulkData.Interpolation1.GetBulkDataSize() : 0;
+			Total += InterpolationBulkData.SimRootPointIndex.IsBulkDataLoaded()	? InterpolationBulkData.SimRootPointIndex.GetBulkDataSize() : 0;
 			return Total;
 		}
 
diff --git a/Engine/Plugins/Runtime/HairStrands/Source/HairStrandsCore/Public/HairStrandsDatas.h b/Engine/Plugins/Runtime/HairStrands/Source/HairStrandsCore/Public/HairStrandsDatas.h
--- a/Engine/Plugins/Runtime/HairStrands/Source/HairStrandsCore/Public/HairStrandsDatas.h
+++ b/Engine/Plugins/Runtime/HairStrands/Source/HairStrandsCore/Public/HairStrandsDatas.h
@@ -162,7 +162,7 @@ struct FHairStrandsTangentFormat
 struct FHairStrandsInterpolationFormat
 {
 	typedef FHairInterpolationVertex Type;
-	typedef uint32 BulkType;
+	typedef FHairInterpolationVertex::BulkType BulkType;
 
 	static const uint32 ComponentCount = 1;
 	static const uint32 SizeInByte = sizeof(Type);
@@ -173,7 +173,7 @@ struct FHairStrandsInterpolationFormat
 struct FHairStrandsInterpolation0Format
 {
 	typedef FHairInterpolation0Vertex Type;
-	typedef uint64 BulkType;
+	typedef FHairInterpolation0Vertex::BulkType BulkType;
 
 	static const uint32 ComponentCount = 1;
 	static const uint32 SizeInByte = sizeof(Type);
@@ -184,7 +184,7 @@ struct FHairStrandsInterpolation0Format
 struct FHairStrandsInterpolation1Format
 {
 	typedef FHairInterpolation1Vertex Type;
-	typedef uint32 BulkType;
+	typedef FHairInterpolation1Vertex::BulkType BulkType;
 
 	static const uint32 ComponentCount = 1;
 	static const uint32 SizeInByte = sizeof(Type);
@@ -314,18 +314,21 @@ struct HAIRSTRANDSCORE_API FHairStrandsInterpolationBulkData
 {
 	enum EDataFlags
 	{
-		DataFlags_HasSingleGuideData = 1,
+		DataFlags_HasData = 1,
+		DataFlags_HasSingleGuideData = 2,
 	};
 
 	void Reset();
-	void Serialize(FArchive& Ar);
-	uint32 GetPointCount() const { return (!!(Flags & DataFlags_HasSingleGuideData)) ? Interpolation.Num() : Interpolation0.Num(); };
+	void Serialize(FArchive& Ar, UObject* Owner);
+	uint32 GetPointCount() const { return PointCount; };
 
-	uint32 Flags;
-	TArray<FHairStrandsInterpolationFormat::Type> Interpolation;	// Per-rendering-vertex interpolation data (closest guides, weight factors, ...). Data for a single guide
-	TArray<FHairStrandsInterpolation0Format::Type> Interpolation0;	// Per-rendering-vertex interpolation data (closest guides, weight factors, ...). Data for up to 3 guides
-	TArray<FHairStrandsInterpolation1Format::Type> Interpolation1;	// Per-rendering-vertex interpolation data (closest guides, weight factors, ...). Data for up to 3 guides
-	TArray<FHairStrandsRootIndexFormat::Type> SimRootPointIndex;	// Per-rendering-vertex index of the sim-root vertex
+	uint32 Flags = 0;
+	uint32 PointCount = 0;
+	uint32 SimPointCount = 0;
+	FByteBulkData Interpolation;	// FHairStrandsInterpolationFormat  - Per-rendering-vertex interpolation data (closest guides, weight factors, ...). Data for a single guide
+	FByteBulkData Interpolation0;	// FHairStrandsInterpolation0Format - Per-rendering-vertex interpolation data (closest guides, weight factors, ...). Data for up to 3 guides
+	FByteBulkData Interpolation1;	// FHairStrandsInterpolation1Format - Per-rendering-vertex interpolation data (closest guides, weight factors, ...). Data for up to 3 guides
+	FByteBulkData SimRootPointIndex;// FHairStrandsRootIndexFormat      - Per-rendering-vertex index of the sim-root vertex
 };
 
 /** Hair strands points attribute */
