Change root deformed resources to only allocate GPU memory for active LOD.

diff --git a/Engine/Plugins/Runtime/HairStrands/Source/HairStrandsCore/Private/GroomResources.cpp b/Engine/Plugins/Runtime/HairStrands/Source/HairStrandsCore/Private/GroomResources.cpp
--- a/Engine/Plugins/Runtime/HairStrands/Source/HairStrandsCore/Private/GroomResources.cpp
+++ b/Engine/Plugins/Runtime/HairStrands/Source/HairStrandsCore/Private/GroomResources.cpp
@@ -934,11 +934,12 @@ FHairStrandsDeformedRootResource::FHairStrandsDeformedRootResource(const FHairSt
 	}
 }
 
-void FHairStrandsDeformedRootResource::InternalAllocate(FRDGBuilder& GraphBuilder)
+void FHairStrandsDeformedRootResource::InternalAllocateLOD(FRDGBuilder& GraphBuilder, int32 LODIndex)
 {
-	if (RootCount > 0)
+	if (RootCount > 0 && LODIndex >= 0 && LODIndex < LODs.Num())
 	{
-		for (FLOD& LOD : LODs)
+		FLOD& LOD = LODs[LODIndex];
+		if (LOD.Status == FLOD::EStatus::Invalid)
 		{		
 			LOD.Status = FLOD::EStatus::Initialized;
 			if (LOD.SampleCount > 0)
diff --git a/Engine/Plugins/Runtime/HairStrands/Source/HairStrandsCore/Public/GroomResources.h b/Engine/Plugins/Runtime/HairStrands/Source/HairStrandsCore/Public/GroomResources.h
--- a/Engine/Plugins/Runtime/HairStrands/Source/HairStrandsCore/Public/GroomResources.h
+++ b/Engine/Plugins/Runtime/HairStrands/Source/HairStrandsCore/Public/GroomResources.h
@@ -134,7 +134,7 @@ struct FHairStrandsDeformedRootResource : public FHairCommonResource
 	FHairStrandsDeformedRootResource(const FHairStrandsRestRootResource* InRestResources, EHairStrandsResourcesType CurveType);
 
 	/* Init/Release buffers */
-	virtual void InternalAllocate(FRDGBuilder& GraphBuilder) override;
+	virtual void InternalAllocateLOD(FRDGBuilder& GraphBuilder, int32 LODIndex) override;
 	virtual void InternalRelease() override;
 
 	/* Get the resource name */
