Change root bulk data to not inline their paylod to differ data loading with cook build

diff --git a/Engine/Plugins/Runtime/HairStrands/Source/HairStrandsCore/Private/GroomBindingAsset.cpp b/Engine/Plugins/Runtime/HairStrands/Source/HairStrandsCore/Private/GroomBindingAsset.cpp
--- a/Engine/Plugins/Runtime/HairStrands/Source/HairStrandsCore/Private/GroomBindingAsset.cpp
+++ b/Engine/Plugins/Runtime/HairStrands/Source/HairStrandsCore/Private/GroomBindingAsset.cpp
@@ -666,7 +666,7 @@ void UGroomBindingAsset::PostEditChangeProperty(FPropertyChangedEvent& PropertyC
 // differences, etc.) replace the version GUID below with a new one.
 // In case of merge conflicts with DDC versions, you MUST generate a new GUID
 // and set this new GUID as the version.
-#define GROOM_BINDING_DERIVED_DATA_VERSION TEXT("5D1E46FE8AF848E792FFD94F6F00CB10")
+#define GROOM_BINDING_DERIVED_DATA_VERSION TEXT("980323E294834AFB8AD013720FC5D19D")
 
 namespace GroomBindingDerivedDataCacheUtils
 {
diff --git a/Engine/Plugins/Runtime/HairStrands/Source/HairStrandsCore/Private/GroomResources.cpp b/Engine/Plugins/Runtime/HairStrands/Source/HairStrandsCore/Private/GroomResources.cpp
--- a/Engine/Plugins/Runtime/HairStrands/Source/HairStrandsCore/Private/GroomResources.cpp
+++ b/Engine/Plugins/Runtime/HairStrands/Source/HairStrandsCore/Private/GroomResources.cpp
@@ -1011,6 +1011,18 @@ static void InternalSerialize(FArchive& Ar, UObject* Owner, FHairStrandsRootBulk
 	const int32 ChunkIndex = 0;
 	bool bAttemptFileMapping = false;
 
+	const uint32 BulkFlags = BULKDATA_Force_NOT_InlinePayload | BULKDATA_SerializeCompressed;
+	LOD.RootTriangleIndexBuffer.SetBulkDataFlags(BulkFlags);
+	LOD.RootTriangleBarycentricBuffer.SetBulkDataFlags(BulkFlags);
+	LOD.RestRootTrianglePosition0Buffer.SetBulkDataFlags(BulkFlags);
+	LOD.RestRootTrianglePosition1Buffer.SetBulkDataFlags(BulkFlags);
+	LOD.RestRootTrianglePosition2Buffer.SetBulkDataFlags(BulkFlags);
+
+	LOD.MeshInterpolationWeightsBuffer.SetBulkDataFlags(BulkFlags);
+	LOD.MeshSampleIndicesBuffer.SetBulkDataFlags(BulkFlags);
+	LOD.RestSamplePositionsBuffer.SetBulkDataFlags(BulkFlags);
+	LOD.ValidSectionIndices.SetBulkDataFlags(BulkFlags);
+
 	Ar << LOD.LODIndex;
 	LOD.RootTriangleIndexBuffer.Serialize(Ar, Owner, ChunkIndex, bAttemptFileMapping);
 	LOD.RootTriangleBarycentricBuffer.Serialize(Ar, Owner, ChunkIndex, bAttemptFileMapping);
@@ -1044,6 +1056,9 @@ void FHairStrandsRootBulkData::Serialize(FArchive& Ar, UObject* Owner)
 	const int32 ChunkIndex = 0;
 	bool bAttemptFileMapping = false;
 
+	const uint32 BulkFlags = BULKDATA_Force_NOT_InlinePayload | BULKDATA_SerializeCompressed;
+	VertexToCurveIndexBuffer.SetBulkDataFlags(BulkFlags);
+
 	if (!Ar.IsObjectReferenceCollector())
 	{
 		Ar << RootCount;
