Fix crash when selecting hair cards panel.

diff --git a/Engine/Plugins/Runtime/HairStrands/Source/HairStrandsEditor/Private/GroomAssetDetails.cpp b/Engine/Plugins/Runtime/HairStrands/Source/HairStrandsEditor/Private/GroomAssetDetails.cpp
--- a/Engine/Plugins/Runtime/HairStrands/Source/HairStrandsEditor/Private/GroomAssetDetails.cpp
+++ b/Engine/Plugins/Runtime/HairStrands/Source/HairStrandsEditor/Private/GroomAssetDetails.cpp
@@ -739,7 +739,7 @@ bool AssignIfDifferent(T& Dest, const T& Src, bool bSetValue)
 bool FGroomRenderingDetails::CommonResetToDefault(TSharedPtr<IPropertyHandle> ChildHandle, int32 GroupIndex, int32 LODIndex, bool bSetValue)
 {
 	bool bHasChanged = false;
-	if (ChildHandle == nullptr || GroomAsset == nullptr || GroupIndex < 0 || GroupIndex >= GroomAsset->GetNumHairGroups())
+	if (ChildHandle == nullptr || GroomAsset == nullptr || GroupIndex < 0)
 	{
 		return bHasChanged;
 	}
@@ -747,6 +747,7 @@ bool FGroomRenderingDetails::CommonResetToDefault(TSharedPtr<IPropertyHandle> Ch
 	FName PropertyName = ChildHandle->GetProperty()->GetFName();
 
 	// Hair strands
+	if (GroupIndex < GroomAsset->HairGroupsRendering.Num())
 	{
 		{
 			FHairGeometrySettings Default;
@@ -771,6 +772,7 @@ bool FGroomRenderingDetails::CommonResetToDefault(TSharedPtr<IPropertyHandle> Ch
 	}
 
 	// Interpolation
+	if (GroupIndex < GroomAsset->HairGroupsInterpolation.Num())
 	{
 		{
 			FHairDecimationSettings Default;
@@ -788,7 +790,7 @@ bool FGroomRenderingDetails::CommonResetToDefault(TSharedPtr<IPropertyHandle> Ch
 	}
 
 	// LODs
-	if (LODIndex>=0)
+	if (GroupIndex < GroomAsset->HairGroupsLOD.Num())
 	{
 		{
 			FHairGroupsLOD Default;
@@ -796,6 +798,7 @@ bool FGroomRenderingDetails::CommonResetToDefault(TSharedPtr<IPropertyHandle> Ch
 			HAIR_RESET0(HairGroupsLOD, FHairGroupsLOD, ClusterScreenSizeScale);
 		}
 
+		if (LODIndex>=0 && LODIndex < GroomAsset->HairGroupsLOD[GroupIndex].LODs.Num())
 		{
 			FHairLODSettings Default;
 			HAIR_RESET1(HairGroupsLOD, FHairLODSettings, LODs[LODIndex], CurveDecimation);
@@ -808,6 +811,7 @@ bool FGroomRenderingDetails::CommonResetToDefault(TSharedPtr<IPropertyHandle> Ch
 	}
 
 	// Cards
+	if (GroupIndex < GroomAsset->HairGroupsCards.Num())
 	{
 		{
 			FHairGroupsCardsSourceDescription Default;
@@ -829,6 +833,7 @@ bool FGroomRenderingDetails::CommonResetToDefault(TSharedPtr<IPropertyHandle> Ch
 	}
 
 	// Meshes
+	if (GroupIndex < GroomAsset->HairGroupsMeshes.Num())
 	{
 		{
 			FHairGroupsMeshesSourceDescription Default;
@@ -848,6 +853,7 @@ bool FGroomRenderingDetails::CommonResetToDefault(TSharedPtr<IPropertyHandle> Ch
 	}
 
 	// Physics
+	if (GroupIndex < GroomAsset->HairGroupsPhysics.Num())
 	{
 		{
 			FHairSolverSettings Default;
