Fix missing CPU data when creating hair resources.

The LockReadOnly() method of FByteBulkData does not insure that data are loaded from disk. This CL change the loading to use Lock(LOCK_READ_ONLY) instead, which first load data from disk, then allo to access the data.

diff --git a/Engine/Plugins/Runtime/HairStrands/Source/HairStrandsCore/Private/GroomBindingBuilder.cpp b/Engine/Plugins/Runtime/HairStrands/Source/HairStrandsCore/Private/GroomBindingBuilder.cpp
--- a/Engine/Plugins/Runtime/HairStrands/Source/HairStrandsCore/Private/GroomBindingBuilder.cpp
+++ b/Engine/Plugins/Runtime/HairStrands/Source/HairStrandsCore/Private/GroomBindingBuilder.cpp
@@ -57,7 +57,7 @@ static FAutoConsoleVariableRef CVarHairStrandsBindingBuilderWarningEnable(TEXT("
 FString FGroomBindingBuilder::GetVersion()
 {
 	// Important to update the version when groom building changes
-	return TEXT("2b");
+	return TEXT("2d");
 }
 
 ///////////////////////////////////////////////////////////////////////////////////////////////////
@@ -1856,7 +1856,7 @@ static void BuildRootBulkData(
 		CopyToBulkData<FHairStrandsMeshTrianglePositionFormat>(Out.MeshProjectionLODs[MeshLODIt].RestRootTrianglePosition1Buffer, In.MeshProjectionLODs[MeshLODIt].RestRootTrianglePosition1Buffer);
 		CopyToBulkData<FHairStrandsMeshTrianglePositionFormat>(Out.MeshProjectionLODs[MeshLODIt].RestRootTrianglePosition2Buffer, In.MeshProjectionLODs[MeshLODIt].RestRootTrianglePosition2Buffer);
 
-		CopyToBulkData<FHairStrandsIndexFormat>(Out.MeshProjectionLODs[MeshLODIt].ValidSectionIndices, In.MeshProjectionLODs[MeshLODIt].ValidSectionIndices);
+		Out.MeshProjectionLODs[MeshLODIt].ValidSectionIndices = In.MeshProjectionLODs[MeshLODIt].ValidSectionIndices;
 
 		if (bHasValidSamples)
 		{
@@ -1899,7 +1899,7 @@ static void BuildRootData(
 		CopyFromBulkData<FHairStrandsMeshTrianglePositionFormat>(Out.MeshProjectionLODs[MeshLODIt].RestRootTrianglePosition1Buffer, In.MeshProjectionLODs[MeshLODIt].RestRootTrianglePosition1Buffer);
 		CopyFromBulkData<FHairStrandsMeshTrianglePositionFormat>(Out.MeshProjectionLODs[MeshLODIt].RestRootTrianglePosition2Buffer, In.MeshProjectionLODs[MeshLODIt].RestRootTrianglePosition2Buffer);
 
-		CopyFromBulkData<FHairStrandsIndexFormat>(Out.MeshProjectionLODs[MeshLODIt].ValidSectionIndices, In.MeshProjectionLODs[MeshLODIt].ValidSectionIndices);
+		Out.MeshProjectionLODs[MeshLODIt].ValidSectionIndices = In.MeshProjectionLODs[MeshLODIt].ValidSectionIndices;
 
 		if (bHasValidSamples)
 		{
diff --git a/Engine/Plugins/Runtime/HairStrands/Source/HairStrandsCore/Private/GroomResources.cpp b/Engine/Plugins/Runtime/HairStrands/Source/HairStrandsCore/Private/GroomResources.cpp
--- a/Engine/Plugins/Runtime/HairStrands/Source/HairStrandsCore/Private/GroomResources.cpp
+++ b/Engine/Plugins/Runtime/HairStrands/Source/HairStrandsCore/Private/GroomResources.cpp
@@ -116,7 +116,7 @@ static FRDGBufferRef InternalCreateVertexBuffer(
 }
 
 template<typename FormatType>
-void InternalCreateVertexBufferRDG_FromBulkData(FRDGBuilder& GraphBuilder, const FByteBulkData& InBulkData, uint32 InDataCount, FRDGExternalBuffer& Out, const TCHAR* DebugName, EHairResourceUsageType UsageType)
+void InternalCreateVertexBufferRDG_FromBulkData(FRDGBuilder& GraphBuilder, FByteBulkData& InBulkData, uint32 InDataCount, FRDGExternalBuffer& Out, const TCHAR* DebugName, EHairResourceUsageType UsageType)
 {
 	const uint32 InDataCount_Check = InBulkData.GetBulkDataSize() / sizeof(typename FormatType::BulkType);
 	check(InDataCount_Check == InDataCount);
@@ -134,7 +134,7 @@ void InternalCreateVertexBufferRDG_FromBulkData(FRDGBuilder& GraphBuilder, const
 		Desc.Usage = Desc.Usage & (~BUF_UnorderedAccess);
 	}
 
-	const typename FormatType::BulkType* BulkData = (const typename FormatType::BulkType*)InBulkData.LockReadOnly();
+	const typename FormatType::BulkType* BulkData = (const typename FormatType::BulkType*)InBulkData.Lock(LOCK_READ_ONLY);
 	FRDGBufferRef Buffer = InternalCreateVertexBuffer(
 		GraphBuilder,
 		DebugName,
@@ -625,7 +625,7 @@ void FHairMeshesDeformedResource::InternalRelease()
 
 /////////////////////////////////////////////////////////////////////////////////////////
 
-FHairStrandsRestResource::FHairStrandsRestResource(const FHairStrandsBulkData& InBulkData, EHairStrandsResourcesType InCurveType) :
+FHairStrandsRestResource::FHairStrandsRestResource(FHairStrandsBulkData& InBulkData, EHairStrandsResourcesType InCurveType) :
 	FHairCommonResource(EHairStrandsAllocationType::Deferred),
 	PositionBuffer(), Attribute0Buffer(), Attribute1Buffer(), MaterialBuffer(), BulkData(InBulkData), CurveType(InCurveType)
 {}
@@ -693,7 +693,7 @@ void FHairStrandsRestResource::InternalRelease()
 
 /////////////////////////////////////////////////////////////////////////////////////////
 
-FHairStrandsDeformedResource::FHairStrandsDeformedResource(const FHairStrandsBulkData& InBulkData, bool bInInitializedData, EHairStrandsResourcesType InCurveType) :
+FHairStrandsDeformedResource::FHairStrandsDeformedResource(FHairStrandsBulkData& InBulkData, bool bInInitializedData, EHairStrandsResourcesType InCurveType) :
 	FHairCommonResource(EHairStrandsAllocationType::Deferred),
 	BulkData(InBulkData), bInitializedData(bInInitializedData), CurveType(InCurveType)
 {
@@ -791,7 +791,7 @@ void FHairStrandsClusterCullingResource::InternalRelease()
 
 /////////////////////////////////////////////////////////////////////////////////////////
 
-FHairStrandsRestRootResource::FHairStrandsRestRootResource(const FHairStrandsRootBulkData& InBulkData, EHairStrandsResourcesType InCurveType) :
+FHairStrandsRestRootResource::FHairStrandsRestRootResource(FHairStrandsRootBulkData& InBulkData, EHairStrandsResourcesType InCurveType) :
 	FHairCommonResource(EHairStrandsAllocationType::Deferred),
 	BulkData(InBulkData), CurveType(InCurveType)
 {
@@ -842,7 +842,7 @@ void FHairStrandsRestRootResource::InternalAllocateLOD(FRDGBuilder& GraphBuilder
 			return;
 		}
 
-		const FHairStrandsRootBulkData::FMeshProjectionLOD& CPUData = BulkData.MeshProjectionLODs[LODIndex];
+		FHairStrandsRootBulkData::FMeshProjectionLOD& CPUData = BulkData.MeshProjectionLODs[LODIndex];
 		const bool bHasValidCPUData = CPUData.RootTriangleBarycentricBuffer.GetBulkDataSize() > 0;
 		if (bHasValidCPUData)
 		{
@@ -991,19 +991,7 @@ bool FHairStrandsRootBulkData::HasProjectionData() const
 const TArray<uint32>& FHairStrandsRootBulkData::GetValidSectionIndices(int32 LODIndex) const
 {
 	check(LODIndex >= 0 && LODIndex < MeshProjectionLODs.Num());
-
-	if (MeshProjectionLODs[LODIndex].CachedValidSectionIndices.IsEmpty())
-	{
-		const uint32 SizeInBytes = MeshProjectionLODs[LODIndex].ValidSectionIndices.GetBulkDataSize();
-		const uint32 ElementCount = SizeInBytes / sizeof(uint32);
-		MeshProjectionLODs[LODIndex].CachedValidSectionIndices.SetNum(ElementCount);
-
-		uint32* CachedData = MeshProjectionLODs[LODIndex].CachedValidSectionIndices.GetData();
-		const uint32* Data = (uint32*)MeshProjectionLODs[LODIndex].ValidSectionIndices.LockReadOnly();
-		FMemory::Memcpy(CachedData, Data, SizeInBytes);
-		MeshProjectionLODs[LODIndex].ValidSectionIndices.Unlock();
-	}
-	return MeshProjectionLODs[LODIndex].CachedValidSectionIndices;
+	return MeshProjectionLODs[LODIndex].ValidSectionIndices;
 }
 
 static void InternalSerialize(FArchive& Ar, UObject* Owner, FHairStrandsRootBulkData::FMeshProjectionLOD& LOD)
@@ -1021,7 +1009,6 @@ static void InternalSerialize(FArchive& Ar, UObject* Owner, FHairStrandsRootBulk
 	LOD.MeshInterpolationWeightsBuffer.SetBulkDataFlags(BulkFlags);
 	LOD.MeshSampleIndicesBuffer.SetBulkDataFlags(BulkFlags);
 	LOD.RestSamplePositionsBuffer.SetBulkDataFlags(BulkFlags);
-	LOD.ValidSectionIndices.SetBulkDataFlags(BulkFlags);
 
 	Ar << LOD.LODIndex;
 	LOD.RootTriangleIndexBuffer.Serialize(Ar, Owner, ChunkIndex, bAttemptFileMapping);
@@ -1034,7 +1021,7 @@ static void InternalSerialize(FArchive& Ar, UObject* Owner, FHairStrandsRootBulk
 	LOD.MeshInterpolationWeightsBuffer.Serialize(Ar, Owner, ChunkIndex, bAttemptFileMapping);
 	LOD.MeshSampleIndicesBuffer.Serialize(Ar, Owner, ChunkIndex, bAttemptFileMapping);
 	LOD.RestSamplePositionsBuffer.Serialize(Ar, Owner, ChunkIndex, bAttemptFileMapping);
-	LOD.ValidSectionIndices.Serialize(Ar, Owner, ChunkIndex, bAttemptFileMapping);
+	LOD.ValidSectionIndices.BulkSerialize(Ar);
 }
 
 static void InternalSerialize(FArchive& Ar, UObject* Owner, TArray<FHairStrandsRootBulkData::FMeshProjectionLOD>& LODs)
@@ -1082,8 +1069,7 @@ void FHairStrandsRootBulkData::Reset()
 		LOD.MeshInterpolationWeightsBuffer.RemoveBulkData();
 		LOD.MeshSampleIndicesBuffer.RemoveBulkData();
 		LOD.RestSamplePositionsBuffer.RemoveBulkData();
-		LOD.ValidSectionIndices.RemoveBulkData();
-		LOD.CachedValidSectionIndices.Empty();
+		LOD.ValidSectionIndices.Empty();
 	}
 	MeshProjectionLODs.Empty();
 }
@@ -1125,7 +1111,7 @@ void FHairStrandsRootData::Reset()
 
 /////////////////////////////////////////////////////////////////////////////////////////
 
-FHairStrandsInterpolationResource::FHairStrandsInterpolationResource(const FHairStrandsInterpolationBulkData& InBulkData) :
+FHairStrandsInterpolationResource::FHairStrandsInterpolationResource(FHairStrandsInterpolationBulkData& InBulkData) :
 	FHairCommonResource(EHairStrandsAllocationType::Deferred),
 	InterpolationBuffer(), Interpolation0Buffer(), Interpolation1Buffer(), BulkData(InBulkData)
 {
@@ -1174,7 +1160,7 @@ void FHairCardsInterpolationBulkData::Serialize(FArchive& Ar)
 	Ar << Interpolation;
 }
 
-FHairCardsInterpolationResource::FHairCardsInterpolationResource(const FHairCardsInterpolationBulkData& InBulkData) :
+FHairCardsInterpolationResource::FHairCardsInterpolationResource(FHairCardsInterpolationBulkData& InBulkData) :
 	FHairCommonResource(EHairStrandsAllocationType::Immediate),
 	InterpolationBuffer(), BulkData(InBulkData)
 {
diff --git a/Engine/Plugins/Runtime/HairStrands/Source/HairStrandsCore/Public/GroomResources.h b/Engine/Plugins/Runtime/HairStrands/Source/HairStrandsCore/Public/GroomResources.h
--- a/Engine/Plugins/Runtime/HairStrands/Source/HairStrandsCore/Public/GroomResources.h
+++ b/Engine/Plugins/Runtime/HairStrands/Source/HairStrandsCore/Public/GroomResources.h
@@ -56,7 +56,7 @@ struct FHairCommonResource : public FRenderResource
 struct FHairStrandsRestRootResource : public FHairCommonResource
 {
 	/** Build the hair strands resource */
-	FHairStrandsRestRootResource(const FHairStrandsRootBulkData& BulkData, EHairStrandsResourcesType CurveType);
+	FHairStrandsRestRootResource(FHairStrandsRootBulkData& BulkData, EHairStrandsResourcesType CurveType);
 
 	/* Init/Release buffers */
 	virtual void InternalAllocate(FRDGBuilder& GraphBuilder) override;
@@ -120,7 +120,7 @@ struct FHairStrandsRestRootResource : public FHairCommonResource
 	TArray<FLOD> LODs;
 
 	/* Store CPU data for root info & root binding */
-	const FHairStrandsRootBulkData& BulkData;
+	FHairStrandsRootBulkData& BulkData;
 
 	/* Type of curves */
 	const EHairStrandsResourcesType CurveType;
@@ -196,7 +196,7 @@ struct FHairStrandsDeformedRootResource : public FHairCommonResource
 struct FHairStrandsRestResource : public FHairCommonResource
 {
 	/** Build the hair strands resource */
-	FHairStrandsRestResource(const FHairStrandsBulkData& InBulkData, EHairStrandsResourcesType CurveType);
+	FHairStrandsRestResource(FHairStrandsBulkData& InBulkData, EHairStrandsResourcesType CurveType);
 
 	/* Init/Release buffers */
 	virtual void InternalAllocate(FRDGBuilder& GraphBuilder) override;
@@ -241,7 +241,7 @@ struct FHairStrandsRestResource : public FHairCommonResource
 	FRDGExternalBuffer MaterialBuffer;
 
 	/* Reference to the hair strands render data */
-	const FHairStrandsBulkData& BulkData;
+	FHairStrandsBulkData& BulkData;
 
 	/* Type of curves */
 	const EHairStrandsResourcesType CurveType;
@@ -252,7 +252,7 @@ struct FHairStrandsRestResource : public FHairCommonResource
 struct FHairStrandsDeformedResource : public FHairCommonResource
 {
 	/** Build the hair strands resource */
-	FHairStrandsDeformedResource(const FHairStrandsBulkData& BulkData, bool bInitializeData, EHairStrandsResourcesType CurveType);
+	FHairStrandsDeformedResource(FHairStrandsBulkData& BulkData, bool bInitializeData, EHairStrandsResourcesType CurveType);
 
 	/* Init/Release buffers */
 	virtual void InternalAllocate(FRDGBuilder& GraphBuilder) override;
@@ -284,7 +284,7 @@ struct FHairStrandsDeformedResource : public FHairCommonResource
 	FVector PositionOffset[2] = {FVector::ZeroVector, FVector::ZeroVector};
 
 	/* Reference to the hair strands render data */
-	const FHairStrandsBulkData& BulkData;
+	FHairStrandsBulkData& BulkData;
 
 	/* Whether the GPU data should be initialized with the asset data or not */
 	const bool bInitializedData = false;
@@ -350,7 +350,7 @@ struct FHairStrandsClusterCullingResource : public FHairCommonResource
 struct FHairStrandsInterpolationResource : public FHairCommonResource
 {
 	/** Build the hair strands resource */
-	FHairStrandsInterpolationResource(const FHairStrandsInterpolationBulkData& InBulkData);
+	FHairStrandsInterpolationResource(FHairStrandsInterpolationBulkData& InBulkData);
 	
 	/* Init/Release buffers */
 	virtual void InternalAllocate(FRDGBuilder& GraphBuilder) override;
@@ -378,7 +378,7 @@ struct FHairStrandsInterpolationResource : public FHairCommonResource
 	FRDGExternalBuffer SimRootPointIndexBuffer;
 
 	/* Reference to the hair strands interpolation render data */
-	const FHairStrandsInterpolationBulkData& BulkData;
+	FHairStrandsInterpolationBulkData& BulkData;
 };
 
 #if RHI_RAYTRACING
@@ -603,7 +603,7 @@ struct HAIRSTRANDSCORE_API FHairCardsInterpolationBulkData
 struct FHairCardsInterpolationResource : public FHairCommonResource
 {
 	/** Build the hair strands resource */
-	FHairCardsInterpolationResource(const FHairCardsInterpolationBulkData& InBulkData);
+	FHairCardsInterpolationResource(FHairCardsInterpolationBulkData& InBulkData);
 
 	/* Init/release buffers */
 	virtual void InternalAllocate(FRDGBuilder& GraphBuilder) override;
@@ -623,7 +623,7 @@ struct FHairCardsInterpolationResource : public FHairCommonResource
 	FRDGExternalBuffer InterpolationBuffer;
 
 	/* Reference to the hair strands interpolation render data */
-	const FHairCardsInterpolationBulkData& BulkData;
+	FHairCardsInterpolationBulkData& BulkData;
 };
 
 ////////////////////////////////////////////////////////////////////////////////////////////////////////////////////////
diff --git a/Engine/Plugins/Runtime/HairStrands/Source/HairStrandsCore/Public/HairStrandsDatas.h b/Engine/Plugins/Runtime/HairStrands/Source/HairStrandsCore/Public/HairStrandsDatas.h
--- a/Engine/Plugins/Runtime/HairStrands/Source/HairStrandsCore/Public/HairStrandsDatas.h
+++ b/Engine/Plugins/Runtime/HairStrands/Source/HairStrandsCore/Public/HairStrandsDatas.h
@@ -521,7 +521,7 @@ struct FHairStrandsRootBulkData
 			Total += LOD.MeshInterpolationWeightsBuffer.IsBulkDataLoaded() ?	LOD.MeshInterpolationWeightsBuffer.GetBulkDataSize() : 0u;
 			Total += LOD.MeshSampleIndicesBuffer.IsBulkDataLoaded() ?			LOD.MeshSampleIndicesBuffer.GetBulkDataSize() : 0u;
 			Total += LOD.RestSamplePositionsBuffer.IsBulkDataLoaded() ?			LOD.RestSamplePositionsBuffer.GetBulkDataSize() : 0u;
-			Total += LOD.ValidSectionIndices.IsBulkDataLoaded() ?				LOD.ValidSectionIndices.GetBulkDataSize() : 0u;
+			Total += LOD.ValidSectionIndices.GetAllocatedSize();
 		}
 		return Total;
 	}
@@ -554,10 +554,7 @@ struct FHairStrandsRootBulkData
 		FByteBulkData RestSamplePositionsBuffer;
 
 		/* Store the mesh section indices which are relevant for this root LOD data */
-		FByteBulkData ValidSectionIndices;
-
-		/* When cached valid section indices from bulk to avoid lock/unlock operation everyframe */
-		mutable TArray<uint32> CachedValidSectionIndices;
+		TArray<uint32> ValidSectionIndices;
 	};
 
 	/* Number of roots */
