This CL split data and bulk data. Data are serialized into FByteBulkData in order to skip load content on demand.

diff --git a/Engine/Plugins/Runtime/HairStrands/Source/HairStrandsCore/Private/GroomAsset.cpp b/Engine/Plugins/Runtime/HairStrands/Source/HairStrandsCore/Private/GroomAsset.cpp
--- a/Engine/Plugins/Runtime/HairStrands/Source/HairStrandsCore/Private/GroomAsset.cpp
+++ b/Engine/Plugins/Runtime/HairStrands/Source/HairStrandsCore/Private/GroomAsset.cpp
@@ -88,6 +88,22 @@ static FAutoConsoleVariableRef CVarHairStrandsDump_GroomBindingAsset(
 	TEXT("Dump information of all the loaded groom binding assets."),
 	FConsoleVariableDelegate::CreateStatic(DumpLoadedGroomBindingAssets));
 
+uint32 FHairGroupData::FStrands::GetDataSize() const
+{
+	uint32 Total = 0;
+	Total += FBaseWithInterpolation::GetDataSize();
+	uint32 CTotal = 0;
+	uint32 BTotal = 0;
+	CTotal += ClusterCullingBulkData.LODVisibility.GetAllocatedSize();
+	CTotal += ClusterCullingBulkData.CPULODScreenSize.GetAllocatedSize();
+	CTotal += ClusterCullingBulkData.ClusterLODInfos.IsBulkDataLoaded()		? ClusterCullingBulkData.ClusterLODInfos.GetBulkDataSize() : 0;
+	BTotal += ClusterCullingBulkData.VertexToClusterIds.IsBulkDataLoaded()	? ClusterCullingBulkData.VertexToClusterIds.GetBulkDataSize() : 0;
+	BTotal += ClusterCullingBulkData.ClusterVertexIds.IsBulkDataLoaded()	? ClusterCullingBulkData.ClusterVertexIds.GetBulkDataSize() : 0;
+	CTotal += ClusterCullingBulkData.PackedClusterInfos.IsBulkDataLoaded()	? ClusterCullingBulkData.PackedClusterInfos.GetBulkDataSize() : 0;
+
+	return BTotal + Total + CTotal;
+}
+
 void DumpLoadedGroomAssets(IConsoleVariable* InCVarPakTesterEnabled)
 {
 	const bool bDetails = true;
@@ -476,7 +492,7 @@ static bool BuildHairGroup(
 			OutHairGroupsData[GroupIndex].Strands.InterpolationBulkData.Reset();
 		}
 
-		FGroomBuilder::BuildClusterData(StrandsData, HairDescriptionGroups.BoundRadius, HairGroupsLOD[GroupIndex], OutHairGroupsData[GroupIndex].Strands.ClusterCullingData);
+		FGroomBuilder::BuildClusterBulkData(StrandsData, HairDescriptionGroups.BoundRadius, HairGroupsLOD[GroupIndex], OutHairGroupsData[GroupIndex].Strands.ClusterCullingBulkData);
 	}
 	return bIsValid;
 }
@@ -503,7 +519,7 @@ static bool BuildHairGroupCluster(
 		FHairStrandsDatas StrandsData;
 		FHairStrandsDatas GuidesData;
 		FGroomBuilder::BuildData(HairGroup, HairGroupsInterpolation[GroupIndex], OutHairGroupsInfo[GroupIndex], StrandsData, GuidesData);
-		FGroomBuilder::BuildClusterData(StrandsData, HairDescriptionGroups.BoundRadius, HairGroupsLOD[GroupIndex], OutHairGroupsData[GroupIndex].Strands.ClusterCullingData);
+		FGroomBuilder::BuildClusterBulkData(StrandsData, HairDescriptionGroups.BoundRadius, HairGroupsLOD[GroupIndex], OutHairGroupsData[GroupIndex].Strands.ClusterCullingBulkData);
 	}
 	return bIsValid;
 }
@@ -847,7 +863,7 @@ void UGroomAsset::UpdateResource()
 					HairGroupsInfo,
 					HairGroupsData);
 
-				GroupData.Strands.ClusterCullingResource = new FHairStrandsClusterCullingResource(GroupData.Strands.ClusterCullingData);
+				GroupData.Strands.ClusterCullingResource = new FHairStrandsClusterCullingResource(GroupData.Strands.ClusterCullingBulkData);
 				BeginInitResource(GroupData.Strands.ClusterCullingResource);
 			}
 			else
@@ -1498,11 +1514,11 @@ static void InternalSerialize(FArchive& Ar, UObject* Owner, FHairGroupData& Grou
 	{
 		if (!Ar.IsCooking() || !GroupData.bIsCookedOut)
 		{
-			GroupData.Strands.ClusterCullingData.Serialize(Ar);
+			GroupData.Strands.ClusterCullingBulkData.Serialize(Ar, Owner);
 		}
 		else
 		{
-			NoStrandsData.Strands.ClusterCullingData.Serialize(Ar);
+			NoStrandsData.Strands.ClusterCullingBulkData.Serialize(Ar, Owner);
 		}
 
 		if (Ar.CustomVer(FAnimObjectVersion::GUID) >= FAnimObjectVersion::SerializeGroomCardsAndMeshes)
@@ -2760,16 +2776,16 @@ void UGroomAsset::InitStrandsResources()
 			GroupData.Strands.RestResource = new FHairStrandsRestResource(GroupData.Strands.BulkData, EHairStrandsResourcesType::Strands);
 			BeginInitResource(GroupData.Strands.RestResource);
 
-			if (GroupData.Strands.ClusterCullingData.IsValid())
+			if (GroupData.Strands.ClusterCullingBulkData.IsValid())
 			{
-				// LOD visibility is not serialized into FHairStrandsClusterCullingData as it does not affect the actual generated data. For consistency we 
+				// LOD visibility is not serialized into FHairStrandsClusterCullingBulkData as it does not affect the actual generated data. For consistency we 
 				// patch the LOD visibility with the groom asset value, which might be different from what has been serialized
-				for (uint32 LODIt = 0, LODCount = GroupData.Strands.ClusterCullingData.LODVisibility.Num(); LODIt < LODCount; ++LODIt)
+				for (uint32 LODIt = 0, LODCount = GroupData.Strands.ClusterCullingBulkData.LODVisibility.Num(); LODIt < LODCount; ++LODIt)
 				{
-					GroupData.Strands.ClusterCullingData.LODVisibility[LODIt] = HairGroupsLOD[GroupIndex].LODs[LODIt].bVisible;
+					GroupData.Strands.ClusterCullingBulkData.LODVisibility[LODIt] = HairGroupsLOD[GroupIndex].LODs[LODIt].bVisible;
 				}
 
-				GroupData.Strands.ClusterCullingResource = new FHairStrandsClusterCullingResource(GroupData.Strands.ClusterCullingData);
+				GroupData.Strands.ClusterCullingResource = new FHairStrandsClusterCullingResource(GroupData.Strands.ClusterCullingBulkData);
 				BeginInitResource(GroupData.Strands.ClusterCullingResource);
 			}
 
diff --git a/Engine/Plugins/Runtime/HairStrands/Source/HairStrandsCore/Private/GroomBuilder.cpp b/Engine/Plugins/Runtime/HairStrands/Source/HairStrandsCore/Private/GroomBuilder.cpp
--- a/Engine/Plugins/Runtime/HairStrands/Source/HairStrandsCore/Private/GroomBuilder.cpp
+++ b/Engine/Plugins/Runtime/HairStrands/Source/HairStrandsCore/Private/GroomBuilder.cpp
@@ -33,7 +33,7 @@ static FAutoConsoleVariableRef CVarHairClusterBuilder_MaxVoxelResolution(TEXT("r
 FString FGroomBuilder::GetVersion()
 {
 	// Important to update the version when groom building changes
-	return TEXT("2j");
+	return TEXT("2m");
 }
 
 namespace FHairStrandsDecimation
@@ -1955,13 +1955,15 @@ void Decimate(
 ///////////////////////////////////////////////////////////////////////////////////////////////////
 // Cluster culling data
 
+namespace GroomBuilder_Cluster
+{
 struct FClusterGrid
 {
 	struct FCurve
 	{
 		FCurve()
 		{
-			for (uint8 LODIt = 0; LODIt < FHairStrandsClusterCullingData::MaxLOD; ++LODIt)
+			for (uint8 LODIt = 0; LODIt < FHairClusterInfo::MaxLOD; ++LODIt)
 			{
 				CountPerLOD[LODIt] = 0;
 			}
@@ -1971,7 +1973,7 @@ struct FClusterGrid
 		float Area = 0;
 		float AvgRadius = 0;
 		float MaxRadius = 0;
-		uint32 CountPerLOD[FHairStrandsClusterCullingData::MaxLOD];
+		uint32 CountPerLOD[FHairClusterInfo::MaxLOD];
 	};
 
 	struct FCluster
@@ -2082,7 +2084,7 @@ static void DecimateCurve(
 	}
 
 	const uint32 LODCount = Settings.Num();
-	check(LODCount <= FHairStrandsClusterCullingData::MaxLOD);
+	check(LODCount <= FHairClusterInfo::MaxLOD);
 
 	for (uint8 LODIt = 0; LODIt < LODCount; ++LODIt)
 	{
@@ -2152,7 +2154,7 @@ inline uint32 to10Bits(float V)
 	return FMath::Clamp(uint32(V * 1024), 0u, 1023u);
 }
 
-void FGroomBuilder::BuildClusterData(
+static void BuildClusterData(
 	const FHairStrandsDatas& InRenStrandsData,
 	const float InGroomAssetRadius, 
 	const FHairGroupsLOD& InSettings, 
@@ -2161,7 +2163,7 @@ void FGroomBuilder::BuildClusterData(
 	// 0. Rest existing culling data
 	Out.Reset();
 
-	const uint32 LODCount = FMath::Min(uint32(InSettings.LODs.Num()), FHairStrandsClusterCullingData::MaxLOD);
+	const uint32 LODCount = FMath::Min(uint32(InSettings.LODs.Num()), FHairClusterInfo::MaxLOD);
 	check(LODCount > 0);
 
 	const uint32 RenCurveCount = InRenStrandsData.GetNumCurves();
@@ -2244,7 +2246,7 @@ void FGroomBuilder::BuildClusterData(
 		}
 	}
 	Out.ClusterCount = ValidClusterIndices.Num();
-	Out.ClusterInfos.Init(FHairStrandsClusterCullingData::FHairClusterInfo(), Out.ClusterCount);
+	Out.ClusterInfos.Init(FHairClusterInfo(), Out.ClusterCount);
 	Out.VertexToClusterIds.SetNum(Out.VertexCount);
 
 	// Conservative allocation for inserting vertex indices for the various curves LOD
@@ -2257,8 +2259,8 @@ void FGroomBuilder::BuildClusterData(
 	VertexLODMasks.SetNum(InRenStrandsData.GetNumPoints());
 
 	// Local variable for being capture by the lambda
-	TArray<FHairStrandsClusterCullingData::FHairClusterInfo>& LocalClusterInfos = Out.ClusterInfos;
-	TArray<FHairStrandsClusterCullingData::FHairClusterLODInfo>& LocalClusterLODInfos = Out.ClusterLODInfos;
+	TArray<FHairClusterInfo>& LocalClusterInfos = Out.ClusterInfos;
+	TArray<FHairClusterLODInfo>& LocalClusterLODInfos = Out.ClusterLODInfos;
 	TArray<uint32>& LocalVertexToClusterIds = Out.VertexToClusterIds;
 #define USE_PARALLE_FOR 1
 #if USE_PARALLE_FOR
@@ -2402,12 +2404,12 @@ void FGroomBuilder::BuildClusterData(
 		TArray<uint32> LocalClusterVertexIds;
 		LocalClusterVertexIds.Reserve(LODCount * Cluster.ClusterCurves.Num() * 32); // Guestimate pre-allocation (32 points per curve in average)
 
-		FHairStrandsClusterCullingData::FHairClusterInfo& ClusterInfo = LocalClusterInfos[ClusterIt];
+		FHairClusterInfo& ClusterInfo = LocalClusterInfos[ClusterIt];
 		ClusterInfo.LODCount = LODCount;
 		ClusterInfo.LODInfoOffset = LODCount * ClusterIt;
 		for (uint8 LODIt = 0; LODIt < LODCount; ++LODIt)
 		{
-			FHairStrandsClusterCullingData::FHairClusterLODInfo& ClusterLODInfo = LocalClusterLODInfos[ClusterInfo.LODInfoOffset + LODIt];
+			FHairClusterLODInfo& ClusterLODInfo = LocalClusterLODInfos[ClusterInfo.LODInfoOffset + LODIt];
 			ClusterLODInfo.VertexOffset = LocalClusterVertexIds.Num(); // At the end, it will be the offset at which the data are inserted into ClusterVertexIds
 			ClusterLODInfo.VertexCount0 = 0;
 			ClusterLODInfo.VertexCount1 = 0;
@@ -2449,7 +2451,7 @@ void FGroomBuilder::BuildClusterData(
 		FMemory::Memcpy(RawClusterVertexIds + AllocOffset, LocalClusterVertexIds.GetData(), LocalClusterVertexIds.Num() * sizeof(uint32));
 		for (uint8 LODIt = 0; LODIt < LODCount; ++LODIt)
 		{
-			FHairStrandsClusterCullingData::FHairClusterLODInfo& ClusterLODInfo = LocalClusterLODInfos[ClusterInfo.LODInfoOffset + LODIt];
+			FHairClusterLODInfo& ClusterLODInfo = LocalClusterLODInfos[ClusterInfo.LODInfoOffset + LODIt];
 			ClusterLODInfo.VertexOffset += AllocOffset;
 		}
 
@@ -2498,7 +2500,7 @@ void FGroomBuilder::BuildClusterData(
 
 			ClusterInfo.ScreenSize[LODIt] = InSettings.LODs[LODIt].ScreenSize * ScreenSizeScale;
 			ClusterInfo.bIsVisible[LODIt] = InSettings.LODs[LODIt].bVisible;
-			FHairStrandsClusterCullingData::FHairClusterLODInfo& ClusterLODInfo = LocalClusterLODInfos[ClusterInfo.LODInfoOffset + LODIt];
+			FHairClusterLODInfo& ClusterLODInfo = LocalClusterLODInfos[ClusterInfo.LODInfoOffset + LODIt];
 			ClusterLODInfo.RadiusScale0 = LODScale;
 			ClusterLODInfo.RadiusScale1 = LODScale;
 		}
@@ -2506,8 +2508,8 @@ void FGroomBuilder::BuildClusterData(
 		// Fill in transition radius between LOD to insure that the interpolation is continuous
 		for (uint8 LODIt = 0; LODIt < LODCount - 1; ++LODIt)
 		{
-			FHairStrandsClusterCullingData::FHairClusterLODInfo& Curr = LocalClusterLODInfos[ClusterInfo.LODInfoOffset + LODIt];
-			FHairStrandsClusterCullingData::FHairClusterLODInfo& Next = LocalClusterLODInfos[ClusterInfo.LODInfoOffset + LODIt + 1];
+			FHairClusterLODInfo& Curr = LocalClusterLODInfos[ClusterInfo.LODInfoOffset + LODIt];
+			FHairClusterLODInfo& Next = LocalClusterLODInfos[ClusterInfo.LODInfoOffset + LODIt + 1];
 			Curr.RadiusScale1 = Next.RadiusScale0;
 		}
 	}
@@ -2528,16 +2530,40 @@ void FGroomBuilder::BuildClusterData(
 	FMemory::Memcpy(Out.ClusterVertexIds.GetData(), RawClusterVertexIds, RawClusterVertexCount * sizeof(uint32));
 
 	delete[] RawClusterVertexIds;
+}
+
+static void BuildClusterBulkData(
+	const FHairStrandsClusterCullingData& In,
+	FHairStrandsClusterCullingBulkData& Out)
+{
+	Out.ClusterCount	= In.ClusterCount;
+	Out.VertexCount		= In.VertexCount;
+	Out.VertexLODCount	= In.ClusterVertexIds.Num();
+	Out.ClusterLODCount = In.ClusterLODInfos.Num();
+
+	Out.CPULODScreenSize= In.CPULODScreenSize;
+	Out.LODVisibility	= In.LODVisibility;
+
+	// Sanity check
+	check(Out.ClusterCount			== uint32(In.ClusterInfos.Num()));
+	check(Out.ClusterLODCount		== uint32(In.ClusterLODInfos.Num()));
+	check(Out.VertexCount			== uint32(In.VertexToClusterIds.Num()));
+	check(Out.VertexLODCount		== uint32(In.ClusterVertexIds.Num()));
+	
+	HairStrandsBuilder::CopyToBulkData<FHairClusterLODInfoFormat>(Out.ClusterLODInfos, In.ClusterLODInfos);
+	HairStrandsBuilder::CopyToBulkData<FHairClusterIndexFormat>(Out.VertexToClusterIds, In.VertexToClusterIds);
+	HairStrandsBuilder::CopyToBulkData<FHairClusterIndexFormat>(Out.ClusterVertexIds, In.ClusterVertexIds);
 
 	// Pack LODInfo into GPU format
 	{
-		check(uint32(Out.ClusterInfos.Num()) == Out.ClusterCount);
-		check(uint32(Out.VertexToClusterIds.Num()) == Out.VertexCount);
-		
-		Out.PackedClusterInfos.Reserve(Out.ClusterInfos.Num());
-		for (const FHairStrandsClusterCullingData::FHairClusterInfo& Info : Out.ClusterInfos)
-		{
-			FHairStrandsClusterCullingData::FHairClusterInfo::Packed& PackedInfo = Out.PackedClusterInfos.AddDefaulted_GetRef();
+		check(uint32(In.ClusterInfos.Num()) == Out.ClusterCount);
+		check(uint32(In.VertexToClusterIds.Num()) == Out.VertexCount);
+	
+		TArray<FHairClusterInfo::Packed> PackedClusterInfos;
+		PackedClusterInfos.Reserve(In.ClusterInfos.Num());
+		for (const FHairClusterInfo& Info : In.ClusterInfos)
+		{
+			FHairClusterInfo::Packed& PackedInfo = PackedClusterInfos.AddDefaulted_GetRef();
 			PackedInfo.LODCount = FMath::Clamp(Info.LODCount, 0u, 0xFFu);
 			PackedInfo.LODInfoOffset = FMath::Clamp(Info.LODInfoOffset, 0u, (1u << 24u) - 1u);
 			PackedInfo.LOD_ScreenSize_0 = to10Bits(Info.ScreenSize[0]);
@@ -2549,7 +2575,7 @@ void FGroomBuilder::BuildClusterData(
 			PackedInfo.LOD_ScreenSize_6 = to10Bits(Info.ScreenSize[6]);
 			PackedInfo.LOD_ScreenSize_7 = to10Bits(Info.ScreenSize[7]);
 			PackedInfo.LOD_bIsVisible = 0;
-			for (uint32 LODIt = 0; LODIt < FHairStrandsClusterCullingData::MaxLOD; ++LODIt)
+			for (uint32 LODIt = 0; LODIt < FHairClusterInfo::MaxLOD; ++LODIt)
 			{
 				if (Info.bIsVisible[LODIt])
 				{
@@ -2561,7 +2587,22 @@ void FGroomBuilder::BuildClusterData(
 			PackedInfo.Pad1 = 0;
 			PackedInfo.Pad2 = 0;
 		}
+
+		HairStrandsBuilder::CopyToBulkData<FHairClusterInfoFormat>(Out.PackedClusterInfos, PackedClusterInfos);
 	}
 }
 
+} // namespace GroomBuilder_Cluster
+
+void FGroomBuilder::BuildClusterBulkData(
+	const FHairStrandsDatas& InRenStrandsData,
+	const float InGroomAssetRadius,
+	const FHairGroupsLOD& InSettings,
+	FHairStrandsClusterCullingBulkData& Out)
+{
+	FHairStrandsClusterCullingData ClusterData;
+	GroomBuilder_Cluster::BuildClusterData(InRenStrandsData, InGroomAssetRadius, InSettings, ClusterData);
+	GroomBuilder_Cluster::BuildClusterBulkData(ClusterData, Out);
+}
+
 #undef LOCTEXT_NAMESPACE
diff --git a/Engine/Plugins/Runtime/HairStrands/Source/HairStrandsCore/Private/GroomComponent.cpp b/Engine/Plugins/Runtime/HairStrands/Source/HairStrandsCore/Private/GroomComponent.cpp
--- a/Engine/Plugins/Runtime/HairStrands/Source/HairStrandsCore/Private/GroomComponent.cpp
+++ b/Engine/Plugins/Runtime/HairStrands/Source/HairStrandsCore/Private/GroomComponent.cpp
@@ -2139,7 +2139,7 @@ void UGroomComponent::InitResources(bool bIsBindingReloading)
 			{
 				// This codes assumes strands LOD are contigus and the highest (i.e., 0...x). Change this code to something more robust
 				check(HairGroupInstance->HairGroupPublicData);
-				const int32 StrandsLODCount = GroupData.Strands.ClusterCullingResource->Data.CPULODScreenSize.Num();
+				const int32 StrandsLODCount = GroupData.Strands.ClusterCullingResource->BulkData.CPULODScreenSize.Num();
 				const TArray<float>& LODScreenSizes = HairGroupInstance->HairGroupPublicData->GetLODScreenSizes();
 				const TArray<bool>& LODVisibilities = HairGroupInstance->HairGroupPublicData->GetLODVisibilities();
 				check(StrandsLODCount <= LODScreenSizes.Num());
@@ -2150,11 +2150,11 @@ void UGroomComponent::InitResources(bool bIsBindingReloading)
 					// Other type of geometry are not serizalized, and so won't match
 					if (GroomAsset->HairGroupsLOD[GroupIt].LODs[LODIt].GeometryType == EGroomGeometryType::Strands)
 					{
-						check(GroupData.Strands.ClusterCullingResource->Data.CPULODScreenSize[LODIt] == LODScreenSizes[LODIt]);
-						check(GroupData.Strands.ClusterCullingResource->Data.LODVisibility[LODIt] == LODVisibilities[LODIt]);
+						check(GroupData.Strands.ClusterCullingResource->BulkData.CPULODScreenSize[LODIt] == LODScreenSizes[LODIt]);
+						check(GroupData.Strands.ClusterCullingResource->BulkData.LODVisibility[LODIt] == LODVisibilities[LODIt]);
 					}
 				}
-				HairGroupInstance->HairGroupPublicData->SetClusters(HairGroupInstance->Strands.ClusterCullingResource->Data.ClusterCount, GroupData.Strands.BulkData.GetNumPoints());
+				HairGroupInstance->HairGroupPublicData->SetClusters(HairGroupInstance->Strands.ClusterCullingResource->BulkData.ClusterCount, GroupData.Strands.BulkData.GetNumPoints());
 				BeginInitResource(HairGroupInstance->HairGroupPublicData);
 			}
 
diff --git a/Engine/Plugins/Runtime/HairStrands/Source/HairStrandsCore/Private/GroomRBFDeformer.cpp b/Engine/Plugins/Runtime/HairStrands/Source/HairStrandsCore/Private/GroomRBFDeformer.cpp
--- a/Engine/Plugins/Runtime/HairStrands/Source/HairStrandsCore/Private/GroomRBFDeformer.cpp
+++ b/Engine/Plugins/Runtime/HairStrands/Source/HairStrandsCore/Private/GroomRBFDeformer.cpp
@@ -409,7 +409,7 @@ static void ApplyDeformationToGroom(const TArray<FRBFDeformedPositions>& Deforme
 				FGroomBuilder::BuildInterplationData(HairGroup.Info, StrandsData, GuidesData, InterpolationSettings, StrandsInterpolationData);
 				FGroomBuilder::BuildInterplationBulkData(GuidesData, StrandsInterpolationData, HairGroupsData.Strands.InterpolationBulkData);
 
-				FGroomBuilder::BuildClusterData(StrandsData, HairDescriptionGroups.BoundRadius, HairGroupLOD, HairGroupsData.Strands.ClusterCullingData);
+				FGroomBuilder::BuildClusterBulkData(StrandsData, HairDescriptionGroups.BoundRadius, HairGroupLOD, HairGroupsData.Strands.ClusterCullingBulkData);
 			}
 		}
 	}
diff --git a/Engine/Plugins/Runtime/HairStrands/Source/HairStrandsCore/Private/GroomResources.cpp b/Engine/Plugins/Runtime/HairStrands/Source/HairStrandsCore/Private/GroomResources.cpp
--- a/Engine/Plugins/Runtime/HairStrands/Source/HairStrandsCore/Private/GroomResources.cpp
+++ b/Engine/Plugins/Runtime/HairStrands/Source/HairStrandsCore/Private/GroomResources.cpp
@@ -311,6 +311,36 @@ void InternalCreateStructuredBufferRDG(FRDGBuilder& GraphBuilder, const TArray<F
 	ConvertToExternalBufferWithViews(GraphBuilder, Buffer, Out);
 }
 
+template<typename FormatType>
+void InternalCreateStructuredBufferRDG_FromBulkData(FRDGBuilder& GraphBuilder, FByteBulkData& InBulkData, uint32 InDataCount, FRDGExternalBuffer& Out, const TCHAR* DebugName, EHairResourceUsageType UsageType)
+{
+	const uint32 InSizeInByte = sizeof(FormatType::Type);
+	const uint32 DataSizeInBytes = InSizeInByte * InDataCount;
+	check(InBulkData.GetBulkDataSize() == DataSizeInBytes);
+
+	if (DataSizeInBytes == 0)
+	{
+		Out.Buffer = nullptr;
+		return;
+	}
+
+	FRDGBufferDesc Desc = FRDGBufferDesc::CreateStructuredDesc(InSizeInByte, InDataCount);
+	if (UsageType != EHairResourceUsageType::Dynamic)
+	{
+		Desc.Usage = Desc.Usage & (~BUF_UnorderedAccess);
+	}
+
+	const typename FormatType::Type* Data = (const typename FormatType::Type*)InBulkData.Lock(LOCK_READ_ONLY);
+	FRDGBufferRef Buffer = GraphBuilder.CreateBuffer(Desc, DebugName, ERDGBufferFlags::MultiFrame);
+	if (Data && DataSizeInBytes)
+	{
+		GraphBuilder.QueueBufferUpload(Buffer, Data, DataSizeInBytes, ERDGInitialDataFlags::None);  // Copy data internally
+	}
+	InBulkData.Unlock();
+
+	ConvertToExternalBufferWithViews(GraphBuilder, Buffer, Out);
+}
+
 /////////////////////////////////////////////////////////////////////////////////////////
 
 static UTexture2D* CreateCardTexture(FIntPoint Resolution)
@@ -737,10 +767,7 @@ void FHairStrandsDeformedResource::InternalRelease()
 }
 
 /////////////////////////////////////////////////////////////////////////////////////////
-
-FArchive& operator<<(FArchive& Ar, FHairStrandsClusterCullingData::FHairClusterInfo& Info);
-FArchive& operator<<(FArchive& Ar, FHairStrandsClusterCullingData::FHairClusterLODInfo& Info);
-FArchive& operator<<(FArchive& Ar, FHairStrandsClusterCullingData::FHairClusterInfo::Packed& Info);
+// Cluster culling data
 
 FHairStrandsClusterCullingData::FHairStrandsClusterCullingData()
 {
@@ -752,33 +779,69 @@ void FHairStrandsClusterCullingData::Reset()
 	*this = FHairStrandsClusterCullingData();
 }
 
-void FHairStrandsClusterCullingData::Serialize(FArchive& Ar)
+/////////////////////////////////////////////////////////////////////////////////////////
+// Cluster culling bulk data
+FHairStrandsClusterCullingBulkData::FHairStrandsClusterCullingBulkData()
+{
+
+}
+
+void FHairStrandsClusterCullingBulkData::Reset()
+{
+	ClusterCount = 0;
+	ClusterLODCount = 0;
+	VertexCount = 0;
+	VertexLODCount = 0;
+
+	LODVisibility.Empty();;
+	CPULODScreenSize.Empty();
+
+	ClusterLODInfos.RemoveBulkData();
+	VertexToClusterIds.RemoveBulkData();
+	ClusterVertexIds.RemoveBulkData();
+	PackedClusterInfos.RemoveBulkData();
+}
+
+void FHairStrandsClusterCullingBulkData::Serialize(FArchive& Ar, UObject* Owner)
 {
 	Ar << ClusterCount;
+	Ar << ClusterLODCount;
 	Ar << VertexCount;
+	Ar << VertexLODCount;
 	Ar << LODVisibility;
 	Ar << CPULODScreenSize;
-	Ar << ClusterInfos;
-	Ar << ClusterLODInfos;
-	Ar << VertexToClusterIds;
-	Ar << ClusterVertexIds;
-	Ar << PackedClusterInfos;
+
+	const int32 ChunkIndex = 0;
+	bool bAttemptFileMapping = false;
+
+	const uint32 BulkFlags = BULKDATA_Force_NOT_InlinePayload | BULKDATA_SerializeCompressed;
+	ClusterLODInfos.SetBulkDataFlags(BulkFlags);
+	VertexToClusterIds.SetBulkDataFlags(BulkFlags);
+	ClusterVertexIds.SetBulkDataFlags(BulkFlags);
+	PackedClusterInfos.SetBulkDataFlags(BulkFlags);
+
+	ClusterLODInfos.Serialize(Ar, Owner, ChunkIndex, bAttemptFileMapping);
+	VertexToClusterIds.Serialize(Ar, Owner, ChunkIndex, bAttemptFileMapping);
+	ClusterVertexIds.Serialize(Ar, Owner, ChunkIndex, bAttemptFileMapping);
+	PackedClusterInfos.Serialize(Ar, Owner, ChunkIndex, bAttemptFileMapping);
 }
 
-FHairStrandsClusterCullingResource::FHairStrandsClusterCullingResource(const FHairStrandsClusterCullingData& InData): 
+/////////////////////////////////////////////////////////////////////////////////////////
+// Cluster culling resources
+FHairStrandsClusterCullingResource::FHairStrandsClusterCullingResource(FHairStrandsClusterCullingBulkData& InBulkData): 
 	FHairCommonResource(EHairStrandsAllocationType::Deferred),
-	Data(InData) 
+	BulkData(InBulkData)
 {
 
 }
 
 void FHairStrandsClusterCullingResource::InternalAllocate(FRDGBuilder& GraphBuilder)
 {
-	InternalCreateStructuredBufferRDG(GraphBuilder, Data.PackedClusterInfos, sizeof(FHairStrandsClusterCullingData::FHairClusterInfo::Packed), ClusterInfoBuffer, TEXT("Hair.StrandsClusterCulling_ClusterInfoBuffer"), EHairResourceUsageType::Static);
-	InternalCreateStructuredBufferRDG(GraphBuilder, Data.ClusterLODInfos, sizeof(FHairStrandsClusterCullingData::FHairClusterLODInfo), ClusterLODInfoBuffer, TEXT("Hair.StrandsClusterCulling_ClusterLODInfoBuffer"), EHairResourceUsageType::Static);
+	InternalCreateStructuredBufferRDG_FromBulkData<FHairClusterInfoFormat>(GraphBuilder, BulkData.PackedClusterInfos, BulkData.ClusterCount, ClusterInfoBuffer, TEXT("Hair.StrandsClusterCulling_ClusterInfoBuffer"), EHairResourceUsageType::Static);
+	InternalCreateStructuredBufferRDG_FromBulkData<FHairClusterLODInfoFormat>(GraphBuilder, BulkData.ClusterLODInfos, BulkData.ClusterLODCount, ClusterLODInfoBuffer, TEXT("Hair.StrandsClusterCulling_ClusterLODInfoBuffer"), EHairResourceUsageType::Static);
 
-	InternalCreateVertexBufferRDG(GraphBuilder, Data.ClusterVertexIds, PF_R32_UINT, ClusterVertexIdBuffer, TEXT("Hair.StrandsClusterCulling_ClusterVertexIds"), EHairResourceUsageType::Static);
-	InternalCreateVertexBufferRDG(GraphBuilder, Data.VertexToClusterIds, PF_R32_UINT, VertexToClusterIdBuffer, TEXT("Hair.StrandsClusterCulling_VertexToClusterIds"), EHairResourceUsageType::Static);
+	InternalCreateVertexBufferRDG_FromBulkData<FHairClusterIndexFormat>(GraphBuilder, BulkData.VertexToClusterIds, BulkData.VertexCount, VertexToClusterIdBuffer, TEXT("Hair.StrandsClusterCulling_VertexToClusterIds"), EHairResourceUsageType::Static);
+	InternalCreateVertexBufferRDG_FromBulkData<FHairClusterIndexFormat>(GraphBuilder, BulkData.ClusterVertexIds, BulkData.VertexLODCount, ClusterVertexIdBuffer, TEXT("Hair.StrandsClusterCulling_ClusterVertexIds"), EHairResourceUsageType::Static);
 }
 
 void FHairStrandsClusterCullingResource::InternalRelease()
diff --git a/Engine/Plugins/Runtime/HairStrands/Source/HairStrandsCore/Private/HairStrandsDatas.cpp b/Engine/Plugins/Runtime/HairStrands/Source/HairStrandsCore/Private/HairStrandsDatas.cpp
--- a/Engine/Plugins/Runtime/HairStrands/Source/HairStrandsCore/Private/HairStrandsDatas.cpp
+++ b/Engine/Plugins/Runtime/HairStrands/Source/HairStrandsCore/Private/HairStrandsDatas.cpp
@@ -260,32 +260,3 @@ void FHairStrandsDatas::Reset()
 	StrandsCurves.Reset();
 	StrandsPoints.Reset();
 }
-
-FArchive& operator<<(FArchive& Ar, FHairStrandsClusterCullingData::FHairClusterInfo& Info)
-{
-	Ar << Info.LODCount;
-	Ar << Info.LODInfoOffset;
-	Ar << Info.ScreenSize;
-	Ar << Info.bIsVisible;
-	return Ar;
-}
-
-FArchive& operator<<(FArchive& Ar, FHairStrandsClusterCullingData::FHairClusterLODInfo& Info)
-{
-	Ar << Info.VertexOffset;
-	Ar << Info.VertexCount0;
-	Ar << Info.VertexCount1;
-	Ar << Info.RadiusScale0;
-	Ar << Info.RadiusScale1;
-	return Ar;
-}
-
-FArchive& operator<<(FArchive& Ar, FHairStrandsClusterCullingData::FHairClusterInfo::Packed& Info)
-{
-	uint32* Uint32Data = (uint32*)&Info;
-	Ar << Uint32Data[0];
-	Ar << Uint32Data[1];
-	Ar << Uint32Data[2];
-	Ar << Uint32Data[3];
-	return Ar;
-}
diff --git a/Engine/Plugins/Runtime/HairStrands/Source/HairStrandsCore/Public/GroomAsset.h b/Engine/Plugins/Runtime/HairStrands/Source/HairStrandsCore/Public/GroomAsset.h
--- a/Engine/Plugins/Runtime/HairStrands/Source/HairStrandsCore/Public/GroomAsset.h
+++ b/Engine/Plugins/Runtime/HairStrands/Source/HairStrandsCore/Public/GroomAsset.h
@@ -150,21 +150,22 @@ struct HAIRSTRANDSCORE_API FHairGroupData
 			return Total;
 		}
 
-		uint32 GetDataSize() const
+		uint32 GetDataSize() const;
+#if 0
 		{
 			uint32 Total = 0;
 			Total += FBaseWithInterpolation::GetDataSize();
 
-			Total += ClusterCullingData.LODVisibility.GetAllocatedSize();
-			Total += ClusterCullingData.CPULODScreenSize.GetAllocatedSize();
-			Total += ClusterCullingData.ClusterInfos.GetAllocatedSize();
-			Total += ClusterCullingData.ClusterLODInfos.GetAllocatedSize();
-			Total += ClusterCullingData.VertexToClusterIds.GetAllocatedSize();
-			Total += ClusterCullingData.ClusterVertexIds.GetAllocatedSize();
+			Total += ClusterCullingBulkData.LODVisibility.GetAllocatedSize();
+			Total += ClusterCullingBulkData.CPULODScreenSize.GetAllocatedSize();
+			Total += ClusterCullingBulkData.ClusterInfos.IsBulkDataLoaded()			? ClusterCullingBulkData.ClusterInfos.GetBulkDataSize() : 0;
+			Total += ClusterCullingBulkData.ClusterLODInfos.IsBulkDataLoaded()		? ClusterCullingBulkData.ClusterLODInfos.GetBulkDataSize() : 0;
+			Total += ClusterCullingBulkData.VertexToClusterIds.IsBulkDataLoaded()	? ClusterCullingBulkData.VertexToClusterIds.GetBulkDataSize() : 0;
+			Total += ClusterCullingBulkData.ClusterVertexIds.IsBulkDataLoaded()		? ClusterCullingBulkData.ClusterVertexIds.GetBulkDataSize() : 0;
 			return Total;
 		}
-
-		FHairStrandsClusterCullingData		ClusterCullingData;
+#endif
+		FHairStrandsClusterCullingBulkData	ClusterCullingBulkData;
 		FHairStrandsClusterCullingResource* ClusterCullingResource = nullptr;
 
 		#if RHI_RAYTRACING
diff --git a/Engine/Plugins/Runtime/HairStrands/Source/HairStrandsCore/Public/GroomBuilder.h b/Engine/Plugins/Runtime/HairStrands/Source/HairStrandsCore/Public/GroomBuilder.h
--- a/Engine/Plugins/Runtime/HairStrands/Source/HairStrandsCore/Public/GroomBuilder.h
+++ b/Engine/Plugins/Runtime/HairStrands/Source/HairStrandsCore/Public/GroomBuilder.h
@@ -9,6 +9,7 @@
 
 struct FHairStrandsDatas;
 struct FHairStrandsClusterCullingData;
+struct FHairStrandsClusterCullingBulkData;
 struct FHairGroupsLOD;
 struct FHairGroupData;
 struct FHairDescriptionGroups;
@@ -73,9 +74,9 @@ struct HAIRSTRANDSCORE_API FGroomBuilder
 		FHairStrandsInterpolationBulkData& OutInterpolationData);
 
 	// 6. Build cluster data
-	static void BuildClusterData(
+	static void BuildClusterBulkData(
 		const FHairStrandsDatas& InRenData,
 		const float InGroomAssetRadius,
 		const FHairGroupsLOD& InSettings,
-		FHairStrandsClusterCullingData& OutClusterCullingData);
+		FHairStrandsClusterCullingBulkData& OutClusterCullingData);
 };
\ No newline at end of file
diff --git a/Engine/Plugins/Runtime/HairStrands/Source/HairStrandsCore/Public/GroomResources.h b/Engine/Plugins/Runtime/HairStrands/Source/HairStrandsCore/Public/GroomResources.h
--- a/Engine/Plugins/Runtime/HairStrands/Source/HairStrandsCore/Public/GroomResources.h
+++ b/Engine/Plugins/Runtime/HairStrands/Source/HairStrandsCore/Public/GroomResources.h
@@ -314,7 +314,7 @@ struct FHairStrandsDeformedResource : public FHairCommonResource
 
 struct FHairStrandsClusterCullingResource : public FHairCommonResource
 {
-	FHairStrandsClusterCullingResource(const FHairStrandsClusterCullingData& Data);
+	FHairStrandsClusterCullingResource(FHairStrandsClusterCullingBulkData& Data);
 
 	/* Init/Release buffers */
 	virtual void InternalAllocate(FRDGBuilder& GraphBuilder) override;
@@ -344,7 +344,7 @@ struct FHairStrandsClusterCullingResource : public FHairCommonResource
 	/* Concatenated data for each cluster: list of VertexId pointed to by ClusterInfoBuffer */
 	FRDGExternalBuffer ClusterVertexIdBuffer;
 
-	const FHairStrandsClusterCullingData Data;
+	FHairStrandsClusterCullingBulkData& BulkData;
 };
 
 struct FHairStrandsInterpolationResource : public FHairCommonResource
diff --git a/Engine/Plugins/Runtime/HairStrands/Source/HairStrandsCore/Public/HairStrandsDatas.h b/Engine/Plugins/Runtime/HairStrands/Source/HairStrandsCore/Public/HairStrandsDatas.h
--- a/Engine/Plugins/Runtime/HairStrands/Source/HairStrandsCore/Public/HairStrandsDatas.h
+++ b/Engine/Plugins/Runtime/HairStrands/Source/HairStrandsCore/Public/HairStrandsDatas.h
@@ -622,66 +622,88 @@ struct FHairStrandsRootData
 	TArray<FMeshProjectionLOD> MeshProjectionLODs;
 };
 
-struct HAIRSTRANDSCORE_API FHairStrandsClusterCullingData
-{
-	FHairStrandsClusterCullingData();
-	void Reset();
-	void Serialize(FArchive& Ar);
-	bool IsValid() const { return ClusterCount > 0 && VertexCount > 0; }
 
+/* Structure describing the LOD settings (Screen size, vertex info, ...) for each clusters.
+	The packed version of this structure corresponds to the GPU data layout (HairStrandsClusterCommon.ush)
+	This uses by the GPU LOD selection. */
+struct FHairClusterInfo
+{
 	static const uint32 MaxLOD = 8;
 
-	/* Structure describing the LOD settings (Screen size, vertex info, ...) for each clusters.
-		The packed version of this structure corresponds to the GPU data layout (HairStrandsClusterCommon.ush)
-		This uses by the GPU LOD selection. */
-	struct FHairClusterInfo
+	struct Packed
 	{
-		struct Packed
-		{
-			uint32 LODInfoOffset : 24;
-			uint32 LODCount : 8;
-
-			uint32 LOD_ScreenSize_0 : 10;
-			uint32 LOD_ScreenSize_1 : 10;
-			uint32 LOD_ScreenSize_2 : 10;
-			uint32 Pad0 : 2;
-
-			uint32 LOD_ScreenSize_3 : 10;
-			uint32 LOD_ScreenSize_4 : 10;
-			uint32 LOD_ScreenSize_5 : 10;
-			uint32 Pad1 : 2;
-
-			uint32 LOD_ScreenSize_6 : 10;
-			uint32 LOD_ScreenSize_7 : 10;
-			uint32 LOD_bIsVisible : 8;
-			uint32 Pad2 : 4;
-		};
-
-		FHairClusterInfo()
+		uint32 LODInfoOffset : 24;
+		uint32 LODCount : 8;
+
+		uint32 LOD_ScreenSize_0 : 10;
+		uint32 LOD_ScreenSize_1 : 10;
+		uint32 LOD_ScreenSize_2 : 10;
+		uint32 Pad0 : 2;
+
+		uint32 LOD_ScreenSize_3 : 10;
+		uint32 LOD_ScreenSize_4 : 10;
+		uint32 LOD_ScreenSize_5 : 10;
+		uint32 Pad1 : 2;
+
+		uint32 LOD_ScreenSize_6 : 10;
+		uint32 LOD_ScreenSize_7 : 10;
+		uint32 LOD_bIsVisible : 8;
+		uint32 Pad2 : 4;
+	};
+
+	FHairClusterInfo()
+	{
+		for (uint32 LODIt = 0; LODIt < MaxLOD; ++LODIt)
 		{
-			for (uint32 LODIt = 0; LODIt < MaxLOD; ++LODIt)
-			{
-				ScreenSize[LODIt] = 0;
-				bIsVisible[LODIt] = true;
-			}
+			ScreenSize[LODIt] = 0;
+			bIsVisible[LODIt] = true;
 		}
+	}
 
-		uint32 LODCount = 0;
-		uint32 LODInfoOffset = 0;
-		TStaticArray<float, MaxLOD>  ScreenSize;
-		TStaticArray<bool, MaxLOD>   bIsVisible;
-	};
+	uint32 LODCount = 0;
+	uint32 LODInfoOffset = 0;
+	TStaticArray<float,MaxLOD> ScreenSize;
+	TStaticArray<bool, MaxLOD> bIsVisible;
+};
 
-	/* Structure describing the LOD settings common to all clusters. The layout of this structure is
-		identical the GPU data layout (HairStrandsClusterCommon.ush). This uses by the GPU LOD selection. */
-	struct FHairClusterLODInfo
-	{
-		uint32 VertexOffset = 0;
-		uint32 VertexCount0 = 0;
-		uint32 VertexCount1 = 0;
-		FFloat16 RadiusScale0 = 0;
-		FFloat16 RadiusScale1 = 0;
-	};
+/* Structure describing the LOD settings common to all clusters. The layout of this structure is
+	identical the GPU data layout (HairStrandsClusterCommon.ush). This uses by the GPU LOD selection. */
+struct FHairClusterLODInfo
+{
+	uint32 VertexOffset = 0;
+	uint32 VertexCount0 = 0;
+	uint32 VertexCount1 = 0;
+	FFloat16 RadiusScale0 = 0;
+	FFloat16 RadiusScale1 = 0;
+};
+
+struct FHairClusterInfoFormat
+{
+	typedef FHairClusterInfo::Packed Type;
+	typedef FHairClusterInfo::Packed BulkType;
+	static const uint32 SizeInByte = sizeof(Type);
+};
+
+struct FHairClusterLODInfoFormat
+{
+	typedef FHairClusterLODInfo Type;
+	typedef FHairClusterLODInfo BulkType;
+	static const uint32 SizeInByte = sizeof(Type);
+};
+
+struct FHairClusterIndexFormat
+{
+	typedef uint32 Type;
+	typedef uint32 BulkType;
+	static const uint32 SizeInByte = sizeof(Type);
+	static const EPixelFormat Format = PF_R32_UINT;
+};
+
+struct HAIRSTRANDSCORE_API FHairStrandsClusterCullingData
+{
+	FHairStrandsClusterCullingData();
+	void Reset();
+	bool IsValid() const { return ClusterCount > 0 && VertexCount > 0; }
 
 	/* Set LOD visibility, allowing to remove the simulation/rendering of certain LOD */
 	TArray<bool>				LODVisibility;
@@ -695,12 +717,35 @@ struct HAIRSTRANDSCORE_API FHairStrandsClusterCullingData
 	TArray<uint32>				VertexToClusterIds;
 	TArray<uint32>				ClusterVertexIds;
 
-	/* Packed LOD info packed into GPU format */
-	TArray<FHairClusterInfo::Packed> PackedClusterInfos;
-
 	/* Number of cluster  */
 	uint32 ClusterCount = 0;
 
 	/* Number of vertex  */
 	uint32 VertexCount = 0;
 };
+
+struct HAIRSTRANDSCORE_API FHairStrandsClusterCullingBulkData
+{
+	FHairStrandsClusterCullingBulkData();
+	void Reset();
+	void Serialize(FArchive& Ar, UObject* Owner);
+	bool IsValid() const { return ClusterCount > 0 && VertexCount > 0; }
+
+	/* Set LOD visibility, allowing to remove the simulation/rendering of certain LOD */
+	TArray<bool>	LODVisibility;
+
+	/* Screen size at which LOD should switches on CPU */
+	TArray<float>	CPULODScreenSize;
+
+	/* LOD info for the various clusters for LOD management on GPU */
+	FByteBulkData	PackedClusterInfos;		// Size - ClusterCount
+	FByteBulkData	ClusterLODInfos;		// Size - ClusterLODCount
+	FByteBulkData	VertexToClusterIds;		// Size - VertexCount
+	FByteBulkData	ClusterVertexIds;		// Size - VertexLODCount
+
+	uint32 ClusterCount = 0;
+	uint32 ClusterLODCount = 0;
+	uint32 VertexCount = 0;
+	uint32 VertexLODCount = 0;
+};
+
diff --git a/Engine/Plugins/Runtime/HairStrands/Source/HairStrandsEditor/Private/GroomCacheImporter.cpp b/Engine/Plugins/Runtime/HairStrands/Source/HairStrandsEditor/Private/GroomCacheImporter.cpp
--- a/Engine/Plugins/Runtime/HairStrands/Source/HairStrandsEditor/Private/GroomCacheImporter.cpp
+++ b/Engine/Plugins/Runtime/HairStrands/Source/HairStrandsEditor/Private/GroomCacheImporter.cpp
@@ -116,7 +116,7 @@ TArray<UGroomCache*> FGroomCacheImporter::ImportGroomCache(const FString& Source
 					FGroomBuilder::BuildInterplationData(HairGroup.Info, StrandsData, GuidesData, GroomAssetForCache->HairGroupsInterpolation[GroupIndex].InterpolationSettings, InterpolationData);
 					FGroomBuilder::BuildInterplationBulkData(GuidesData, InterpolationData, HairGroupsData[GroupIndex].Strands.InterpolationBulkData);
 
-					FGroomBuilder::BuildClusterData(StrandsData, HairDescriptionGroups.BoundRadius, GroomAssetForCache->HairGroupsLOD[GroupIndex], HairGroupsData[GroupIndex].Strands.ClusterCullingData);
+					FGroomBuilder::BuildClusterBulkData(StrandsData, HairDescriptionGroups.BoundRadius, GroomAssetForCache->HairGroupsLOD[GroupIndex], HairGroupsData[GroupIndex].Strands.ClusterCullingBulkData);
 				}
 
 				// The HairGroupData is converted into animated groom data by the GroomCacheProcessor
