Fix CIS issue.

diff --git a/Engine/Plugins/Runtime/HairStrands/Source/HairStrandsCore/Private/GroomResources.cpp b/Engine/Plugins/Runtime/HairStrands/Source/HairStrandsCore/Private/GroomResources.cpp
--- a/Engine/Plugins/Runtime/HairStrands/Source/HairStrandsCore/Private/GroomResources.cpp
+++ b/Engine/Plugins/Runtime/HairStrands/Source/HairStrandsCore/Private/GroomResources.cpp
@@ -314,7 +314,7 @@ void InternalCreateStructuredBufferRDG(FRDGBuilder& GraphBuilder, const TArray<F
 template<typename FormatType>
 void InternalCreateStructuredBufferRDG_FromBulkData(FRDGBuilder& GraphBuilder, FByteBulkData& InBulkData, uint32 InDataCount, FRDGExternalBuffer& Out, const TCHAR* DebugName, EHairResourceUsageType UsageType)
 {
-	const uint32 InSizeInByte = sizeof(FormatType::Type);
+	const uint32 InSizeInByte = sizeof(typename FormatType::Type);
 	const uint32 DataSizeInBytes = InSizeInByte * InDataCount;
 	check(InBulkData.GetBulkDataSize() == DataSizeInBytes);
 
