Fix incorrect loading of hair guides data with cooked data.

diff --git a/Engine/Plugins/Runtime/HairStrands/Source/HairStrandsCore/Private/Niagara/NiagaraDataInterfaceHairStrands.cpp b/Engine/Plugins/Runtime/HairStrands/Source/HairStrandsCore/Private/Niagara/NiagaraDataInterfaceHairStrands.cpp
--- a/Engine/Plugins/Runtime/HairStrands/Source/HairStrandsCore/Private/Niagara/NiagaraDataInterfaceHairStrands.cpp
+++ b/Engine/Plugins/Runtime/HairStrands/Source/HairStrandsCore/Private/Niagara/NiagaraDataInterfaceHairStrands.cpp
@@ -334,12 +334,12 @@ void FNDIHairStrandsBuffer::InitRHI()
 {
 	if (SourceRestResources != nullptr)
 	{
-		const FHairStrandsBulkData* SourceDatas = &SourceRestResources->BulkData;
+		FHairStrandsBulkData* SourceDatas = &SourceRestResources->BulkData;
 		{
 			const uint32 OffsetCount = SourceDatas->GetNumCurves() + 1;
 			const uint32 OffsetBytes = sizeof(uint32)*OffsetCount;
 
-			const FHairStrandsRootIndexFormat::Type* SourceData = (const FHairStrandsRootIndexFormat::Type*)SourceDatas->CurveOffsets.LockReadOnly();
+			const FHairStrandsRootIndexFormat::Type* SourceData = (const FHairStrandsRootIndexFormat::Type*)SourceDatas->CurveOffsets.Lock(LOCK_READ_ONLY);
 			CurvesOffsetsBuffer.Initialize(TEXT("CurvesOffsetsBuffer"), sizeof(uint32), OffsetCount, EPixelFormat::PF_R32_UINT, BUF_Static);
 			void* OffsetBufferData = RHILockBuffer(CurvesOffsetsBuffer.Buffer, 0, OffsetBytes, RLM_WriteOnly);
 			FMemory::Memcpy(OffsetBufferData, SourceData, OffsetBytes);
