Fix crash when serializing groom data with no strands data.

diff --git a/Engine/Plugins/Runtime/HairStrands/Source/HairStrandsCore/Private/GroomAsset.cpp b/Engine/Plugins/Runtime/HairStrands/Source/HairStrandsCore/Private/GroomAsset.cpp
--- a/Engine/Plugins/Runtime/HairStrands/Source/HairStrandsCore/Private/GroomAsset.cpp
+++ b/Engine/Plugins/Runtime/HairStrands/Source/HairStrandsCore/Private/GroomAsset.cpp
@@ -1683,7 +1683,7 @@ void UGroomAsset::SetHairWidth(float Width)
 // differences, etc.) replace the version GUID below with a new one.
 // In case of merge conflicts with DDC versions, you MUST generate a new GUID
 // and set this new GUID as the version.
-#define GROOM_DERIVED_DATA_VERSION TEXT("DDFF529D572543C1BB7BD4E49E7B7B5E")
+#define GROOM_DERIVED_DATA_VERSION TEXT("CFDBFBC049704DE4AADC19FE21F85F1D")
 
 #if WITH_EDITORONLY_DATA
 
diff --git a/Engine/Plugins/Runtime/HairStrands/Source/HairStrandsCore/Private/GroomBindingAsset.cpp b/Engine/Plugins/Runtime/HairStrands/Source/HairStrandsCore/Private/GroomBindingAsset.cpp
--- a/Engine/Plugins/Runtime/HairStrands/Source/HairStrandsCore/Private/GroomBindingAsset.cpp
+++ b/Engine/Plugins/Runtime/HairStrands/Source/HairStrandsCore/Private/GroomBindingAsset.cpp
@@ -666,7 +666,7 @@ void UGroomBindingAsset::PostEditChangeProperty(FPropertyChangedEvent& PropertyC
 // differences, etc.) replace the version GUID below with a new one.
 // In case of merge conflicts with DDC versions, you MUST generate a new GUID
 // and set this new GUID as the version.
-#define GROOM_BINDING_DERIVED_DATA_VERSION TEXT("980323E294834AFB8AD013720FC5D19D")
+#define GROOM_BINDING_DERIVED_DATA_VERSION TEXT("79008480605A4207B8E3CF5530321CD6")
 
 namespace GroomBindingDerivedDataCacheUtils
 {
diff --git a/Engine/Plugins/Runtime/HairStrands/Source/HairStrandsCore/Private/GroomResources.cpp b/Engine/Plugins/Runtime/HairStrands/Source/HairStrandsCore/Private/GroomResources.cpp
--- a/Engine/Plugins/Runtime/HairStrands/Source/HairStrandsCore/Private/GroomResources.cpp
+++ b/Engine/Plugins/Runtime/HairStrands/Source/HairStrandsCore/Private/GroomResources.cpp
@@ -816,10 +816,25 @@ void FHairStrandsClusterCullingBulkData::Serialize(FArchive& Ar, UObject* Owner)
 	ClusterVertexIds.SetBulkDataFlags(BulkFlags);
 	PackedClusterInfos.SetBulkDataFlags(BulkFlags);
 
-	ClusterLODInfos.Serialize(Ar, Owner, ChunkIndex, bAttemptFileMapping);
-	VertexToClusterIds.Serialize(Ar, Owner, ChunkIndex, bAttemptFileMapping);
-	ClusterVertexIds.Serialize(Ar, Owner, ChunkIndex, bAttemptFileMapping);
-	PackedClusterInfos.Serialize(Ar, Owner, ChunkIndex, bAttemptFileMapping);
+	if (ClusterLODCount)
+	{
+		ClusterLODInfos.Serialize(Ar, Owner, ChunkIndex, bAttemptFileMapping);
+	}
+
+	if (VertexCount)
+	{
+		VertexToClusterIds.Serialize(Ar, Owner, ChunkIndex, bAttemptFileMapping);
+	}
+
+	if (VertexLODCount)
+	{
+		ClusterVertexIds.Serialize(Ar, Owner, ChunkIndex, bAttemptFileMapping);
+	}
+
+	if (ClusterCount)
+	{
+		PackedClusterInfos.Serialize(Ar, Owner, ChunkIndex, bAttemptFileMapping);
+	}
 }
 
 /////////////////////////////////////////////////////////////////////////////////////////
@@ -1077,10 +1092,13 @@ static void InternalSerialize(FArchive& Ar, UObject* Owner, FHairStrandsRootBulk
 	LOD.RestRootTrianglePosition2Buffer.Serialize(Ar, Owner, ChunkIndex, bAttemptFileMapping);
 
 	Ar << LOD.SampleCount;
-	LOD.MeshInterpolationWeightsBuffer.Serialize(Ar, Owner, ChunkIndex, bAttemptFileMapping);
-	LOD.MeshSampleIndicesBuffer.Serialize(Ar, Owner, ChunkIndex, bAttemptFileMapping);
-	LOD.RestSamplePositionsBuffer.Serialize(Ar, Owner, ChunkIndex, bAttemptFileMapping);
-	LOD.ValidSectionIndices.BulkSerialize(Ar);
+	if (LOD.SampleCount)
+	{
+		LOD.MeshInterpolationWeightsBuffer.Serialize(Ar, Owner, ChunkIndex, bAttemptFileMapping);
+		LOD.MeshSampleIndicesBuffer.Serialize(Ar, Owner, ChunkIndex, bAttemptFileMapping);
+		LOD.RestSamplePositionsBuffer.Serialize(Ar, Owner, ChunkIndex, bAttemptFileMapping);
+	}
+		LOD.ValidSectionIndices.BulkSerialize(Ar);
 }
 
 static void InternalSerialize(FArchive& Ar, UObject* Owner, TArray<FHairStrandsRootBulkData::FMeshProjectionLOD>& LODs)
