Fix another bulk data serialization for groom binding asset having cards only data.

diff --git a/Engine/Plugins/Runtime/HairStrands/Source/HairStrandsCore/Private/GroomBindingAsset.cpp b/Engine/Plugins/Runtime/HairStrands/Source/HairStrandsCore/Private/GroomBindingAsset.cpp
--- a/Engine/Plugins/Runtime/HairStrands/Source/HairStrandsCore/Private/GroomBindingAsset.cpp
+++ b/Engine/Plugins/Runtime/HairStrands/Source/HairStrandsCore/Private/GroomBindingAsset.cpp
@@ -36,20 +36,24 @@ static void InternalSerialize(FArchive& Ar, UObject* Owner, TArray<FHairStrandsR
 	}
 }
 
-static void InternalSerialize(FArchive& Ar, UObject* Owner, FHairGroupBulkData& GroupData)
+static void InternalSerialize(FArchive& Ar, UObject* Owner, FHairGroupBulkData& GroupData, uint32 Flags)
 {
 	Ar.UsingCustomVersion(FAnimObjectVersion::GUID);
 	GroupData.SimRootBulkData.Serialize(Ar, Owner);
-	GroupData.RenRootBulkData.Serialize(Ar, Owner);
+	if (!(Flags & UGroomAsset::CDSF_StrandsStripped))
+	{
+		GroupData.RenRootBulkData.Serialize(Ar, Owner);
+	}
 	if (Ar.CustomVer(FAnimObjectVersion::GUID) >= FAnimObjectVersion::SerializeHairBindingAsset)
 	{
 		InternalSerialize(Ar, Owner, GroupData.CardsRootBulkData);
 	}
 }
 
-static void InternalSerialize(FArchive& Ar, UObject* Owner, TArray<FHairGroupBulkData>& GroupDatas)
+static void InternalSerialize(FArchive& Ar, UObject* Owner, TArray<FHairGroupBulkData>& GroupDatas, uint32 Flags)
 {
 	uint32 GroupCount = GroupDatas.Num();
+	Ar << Flags;
 	Ar << GroupCount;
 	if (Ar.IsLoading())
 	{
@@ -57,7 +61,7 @@ static void InternalSerialize(FArchive& Ar, UObject* Owner, TArray<FHairGroupBul
 	}
 	for (uint32 GroupIt = 0; GroupIt < GroupCount; ++GroupIt)
 	{
-		InternalSerialize(Ar, Owner, GroupDatas[GroupIt]);
+		InternalSerialize(Ar, Owner, GroupDatas[GroupIt], Flags);
 	}
 }
 
@@ -65,12 +69,11 @@ static void InternalSerialize(FArchive& Ar, UObject* Owner, TArray<FHairGroupBul
 
 void UGroomBindingAsset::Serialize(FArchive& Ar)
 {
+	uint8 Flags = 0;
 #if WITH_EDITOR
-	bool bIsStrandsStrippedForCook = false;
 	if (Groom)
 	{
-		uint8 StripFlags = Groom->GenerateClassStripFlags(Ar);
-		bIsStrandsStrippedForCook = !!(StripFlags & UGroomAsset::CDSF_StrandsStripped);
+		Flags = Groom->GenerateClassStripFlags(Ar);
 	}
 #endif
 
@@ -80,25 +83,7 @@ void UGroomBindingAsset::Serialize(FArchive& Ar)
 	if (Ar.CustomVer(FAnimObjectVersion::GUID) < FAnimObjectVersion::GroomBindingSerialization || Ar.IsCooking())
 #endif
 	{
-#if WITH_EDITOR
-		if (bIsStrandsStrippedForCook)
-		{
-			TArray<FHairGroupBulkData> StrippedHairGroupDatas;
-			StrippedHairGroupDatas.SetNum(HairGroupBulkDatas.Num());
-			for (int32 Index = 0; Index < HairGroupBulkDatas.Num(); ++Index)
-			{
-				FHairGroupBulkData& HairGroupBulkData = HairGroupBulkDatas[Index];
-				StrippedHairGroupDatas[Index].SimRootBulkData = HairGroupBulkData.SimRootBulkData;
-				StrippedHairGroupDatas[Index].CardsRootBulkData = HairGroupBulkData.CardsRootBulkData;
-			}
-
-			InternalSerialize(Ar, this, StrippedHairGroupDatas);
-		}
-		else
-#endif
-		{
-			InternalSerialize(Ar, this, HairGroupBulkDatas);
-		}
+		InternalSerialize(Ar, this, HairGroupBulkDatas, Flags);
 		bIsValid = true;
 	}
 }
@@ -666,7 +651,7 @@ void UGroomBindingAsset::PostEditChangeProperty(FPropertyChangedEvent& PropertyC
 // differences, etc.) replace the version GUID below with a new one.
 // In case of merge conflicts with DDC versions, you MUST generate a new GUID
 // and set this new GUID as the version.
-#define GROOM_BINDING_DERIVED_DATA_VERSION TEXT("79008480605A4207B8E3CF5530321CD6")
+#define GROOM_BINDING_DERIVED_DATA_VERSION TEXT("9B70B9F221C14F20910ED9B5B8AB740F")
 
 namespace GroomBindingDerivedDataCacheUtils
 {
@@ -745,7 +730,7 @@ void UGroomBindingAsset::CacheDerivedDatas()
 			Ar.SerializeCompressed(DecompressionBuffer, 0, NAME_Zlib);
 
 			FLargeMemoryReader LargeMemReader(DecompressionBuffer, UncompressedSize, ELargeMemoryReaderFlags::Persistent | ELargeMemoryReaderFlags::TakeOwnership);
-			InternalSerialize(LargeMemReader, this, HairGroupBulkDatas);
+			InternalSerialize(LargeMemReader, this, HairGroupBulkDatas, 0);
 
 			bIsValid = true;
 		}
@@ -762,7 +747,7 @@ void UGroomBindingAsset::CacheDerivedDatas()
 			{
 				// Using a LargeMemoryWriter for serialization since the data can be bigger than 2 GB
 				FLargeMemoryWriter LargeMemWriter(0, /*bIsPersistent=*/ true);
-				InternalSerialize(LargeMemWriter, this, HairGroupBulkDatas);
+				InternalSerialize(LargeMemWriter, this, HairGroupBulkDatas, 0);
 
 				int64 UncompressedSize = LargeMemWriter.TotalSize();
 
