Add command for dummping all Groom related memory usage at once (GroomAsset, GroomBindingAsset, and GroomComponent).

diff --git a/Engine/Plugins/Runtime/HairStrands/Source/HairStrandsCore/Private/GroomAsset.cpp b/Engine/Plugins/Runtime/HairStrands/Source/HairStrandsCore/Private/GroomAsset.cpp
--- a/Engine/Plugins/Runtime/HairStrands/Source/HairStrandsCore/Private/GroomAsset.cpp
+++ b/Engine/Plugins/Runtime/HairStrands/Source/HairStrandsCore/Private/GroomAsset.cpp
@@ -71,7 +71,9 @@ static bool IsHairLODSimulationEnabled(const int32 LODIndex) { return (LODIndex
 /////////////////////////////////////////////////////////////////////////////////////////
 // Asset dump function
 
+void DumpLoadedGroomData(IConsoleVariable* InCVarPakTesterEnabled);
 void DumpLoadedGroomAssets(IConsoleVariable* InCVarPakTesterEnabled);
+void DumpLoadedGroomComponent(IConsoleVariable* InCVarPakTesterEnabled);
 void DumpLoadedGroomBindingAssets(IConsoleVariable* InCVarPakTesterEnabled);
 
 int32 GHairStrandsDump_GroomAsset = 0;
@@ -88,6 +90,13 @@ static FAutoConsoleVariableRef CVarHairStrandsDump_GroomBindingAsset(
 	TEXT("Dump information of all the loaded groom binding assets."),
 	FConsoleVariableDelegate::CreateStatic(DumpLoadedGroomBindingAssets));
 
+int32 GHairStrandsDump_All = 0;
+static FAutoConsoleVariableRef CVarHairStrandsDump_All(
+	TEXT("r.HairStrands.Dump"),
+	GHairStrandsDump_All,
+	TEXT("Dump all the loaded groom assets, groom binding assets, and instanciated groom components."),
+	FConsoleVariableDelegate::CreateStatic(DumpLoadedGroomData));
+
 uint32 FHairGroupData::FStrands::GetDataSize() const
 {
 	uint32 Total = 0;
@@ -345,6 +354,14 @@ void DumpLoadedGroomBindingAssets(IConsoleVariable* InCVarPakTesterEnabled)
 		Total_GPUMemorySize_Strands * ToMb,
 		Total_GPUMemorySize_Cards * ToMb);
 }
+
+void DumpLoadedGroomData(IConsoleVariable* InCVarPakTesterEnabled)
+{
+	DumpLoadedGroomAssets(InCVarPakTesterEnabled);
+	DumpLoadedGroomBindingAssets(InCVarPakTesterEnabled);
+	DumpLoadedGroomComponent(InCVarPakTesterEnabled);
+}
+
 /////////////////////////////////////////////////////////////////////////////////////////
 
 void UpdateHairStrandsLogVerbosity()
