Add sanity check for hair cluster serialization.

diff --git a/Engine/Plugins/Runtime/HairStrands/Source/HairStrandsCore/Private/GroomBuilder.cpp b/Engine/Plugins/Runtime/HairStrands/Source/HairStrandsCore/Private/GroomBuilder.cpp
--- a/Engine/Plugins/Runtime/HairStrands/Source/HairStrandsCore/Private/GroomBuilder.cpp
+++ b/Engine/Plugins/Runtime/HairStrands/Source/HairStrandsCore/Private/GroomBuilder.cpp
@@ -33,7 +33,7 @@ static FAutoConsoleVariableRef CVarHairClusterBuilder_MaxVoxelResolution(TEXT("r
 FString FGroomBuilder::GetVersion()
 {
 	// Important to update the version when groom building changes
-	return TEXT("2m");
+	return TEXT("2n");
 }
 
 namespace FHairStrandsDecimation
@@ -2564,6 +2564,9 @@ static void BuildClusterBulkData(
 		for (const FHairClusterInfo& Info : In.ClusterInfos)
 		{
 			FHairClusterInfo::Packed& PackedInfo = PackedClusterInfos.AddDefaulted_GetRef();
+			// Sanity check
+			check(Info.LODCount <= 8);
+
 			PackedInfo.LODCount = FMath::Clamp(Info.LODCount, 0u, 0xFFu);
 			PackedInfo.LODInfoOffset = FMath::Clamp(Info.LODInfoOffset, 0u, (1u << 24u) - 1u);
 			PackedInfo.LOD_ScreenSize_0 = to10Bits(Info.ScreenSize[0]);
@@ -2586,6 +2589,8 @@ static void BuildClusterBulkData(
 			PackedInfo.Pad0 = 0;
 			PackedInfo.Pad1 = 0;
 			PackedInfo.Pad2 = 0;
+
+			static_assert(sizeof(FHairClusterInfo::Packed) == sizeof(FHairClusterInfo::BulkType));
 		}
 
 		HairStrandsBuilder::CopyToBulkData<FHairClusterInfoFormat>(Out.PackedClusterInfos, PackedClusterInfos);
diff --git a/Engine/Plugins/Runtime/HairStrands/Source/HairStrandsCore/Public/HairStrandsDatas.h b/Engine/Plugins/Runtime/HairStrands/Source/HairStrandsCore/Public/HairStrandsDatas.h
--- a/Engine/Plugins/Runtime/HairStrands/Source/HairStrandsCore/Public/HairStrandsDatas.h
+++ b/Engine/Plugins/Runtime/HairStrands/Source/HairStrandsCore/Public/HairStrandsDatas.h
@@ -650,6 +650,7 @@ struct FHairClusterInfo
 		uint32 LOD_bIsVisible : 8;
 		uint32 Pad2 : 4;
 	};
+	typedef FUintVector4 BulkType;
 
 	FHairClusterInfo()
 	{
