Uniformize all hair strands shader group size permutation.

diff --git a/Engine/Plugins/Runtime/HairStrands/Source/HairStrandsCore/Private/HairStrandsRendering.cpp b/Engine/Plugins/Runtime/HairStrands/Source/HairStrandsCore/Private/HairStrandsRendering.cpp
--- a/Engine/Plugins/Runtime/HairStrands/Source/HairStrandsCore/Private/HairStrandsRendering.cpp
+++ b/Engine/Plugins/Runtime/HairStrands/Source/HairStrandsCore/Private/HairStrandsRendering.cpp
@@ -89,11 +89,6 @@ inline uint32 ComputeGroupSize()
 	return GroupSize;
 }
 
-inline uint32 GetGroupSizePermutation(uint32 GroupSize)
-{
-	return GroupSize == 64 ? 0 : (GroupSize == 32 ? 1 : 2);
-}
-
 enum class EDeformationType : uint8
 {
 	Simulation,		// Use the output of the hair simulation
@@ -242,7 +237,7 @@ class FDeformGuideCS : public FGlobalShader
 	DECLARE_GLOBAL_SHADER(FDeformGuideCS);
 	SHADER_USE_PARAMETER_STRUCT(FDeformGuideCS, FGlobalShader);
 
-	class FGroupSize : SHADER_PERMUTATION_INT("PERMUTATION_GROUP_SIZE", 2);
+	class FGroupSize : SHADER_PERMUTATION_SPARSE_INT("PERMUTATION_GROUP_SIZE", 32, 64);
 	class FDeformationType : SHADER_PERMUTATION_INT("PERMUTATION_DEFORMATION", 4);
 	using FPermutationDomain = TShaderPermutationDomain<FGroupSize, FDeformationType>;
 
@@ -371,7 +366,7 @@ static void AddDeformSimHairStrandsPass(
 	}
 
 	FDeformGuideCS::FPermutationDomain PermutationVector;
-	PermutationVector.Set<FDeformGuideCS::FGroupSize>(GetGroupSizePermutation(GroupSize));
+	PermutationVector.Set<FDeformGuideCS::FGroupSize>(GroupSize);
 	PermutationVector.Set<FDeformGuideCS::FDeformationType>(InternalDeformationType);
 
 	TShaderMapRef<FDeformGuideCS> ComputeShader(ShaderMap, PermutationVector);
@@ -880,7 +875,7 @@ class FHairCardsDeformationCS : public FGlobalShader
 	DECLARE_GLOBAL_SHADER(FHairCardsDeformationCS);
 	SHADER_USE_PARAMETER_STRUCT(FHairCardsDeformationCS, FGlobalShader);
 
-	class FGroupSize : SHADER_PERMUTATION_INT("PERMUTATION_GROUP_SIZE", 2);
+	class FGroupSize : SHADER_PERMUTATION_SPARSE_INT("PERMUTATION_GROUP_SIZE", 32, 64);
 	class FDynamicGeometry : SHADER_PERMUTATION_INT("PERMUTATION_DYNAMIC_GEOMETRY", 2);
 	using FPermutationDomain = TShaderPermutationDomain<FGroupSize, FDynamicGeometry>;
 
@@ -985,7 +980,7 @@ static void AddHairCardsDeformationPass(
 	const uint32 GroupSize = ComputeGroupSize();
 	FHairCardsDeformationCS::FPermutationDomain PermutationVector;
 	PermutationVector.Set<FHairCardsDeformationCS::FDynamicGeometry>(bSupportDynamicMesh ? 1 : 0);
-	PermutationVector.Set<FHairCardsDeformationCS::FGroupSize>(GetGroupSizePermutation(GroupSize));
+	PermutationVector.Set<FHairCardsDeformationCS::FGroupSize>(GroupSize);
 
 	TShaderMapRef<FHairCardsDeformationCS> ComputeShader(ShaderMap, PermutationVector);
 
@@ -1007,7 +1002,7 @@ class FHairTangentCS : public FGlobalShader
 	DECLARE_GLOBAL_SHADER(FHairTangentCS);
 	SHADER_USE_PARAMETER_STRUCT(FHairTangentCS, FGlobalShader);
 
-	class FGroupSize : SHADER_PERMUTATION_INT("PERMUTATION_GROUP_SIZE", 2);
+	class FGroupSize : SHADER_PERMUTATION_SPARSE_INT("PERMUTATION_GROUP_SIZE", 32, 64);
 	class FCulling : SHADER_PERMUTATION_INT("PERMUTATION_CULLING", 2);
 	using FPermutationDomain = TShaderPermutationDomain<FGroupSize, FCulling>;
 
@@ -1048,7 +1043,7 @@ void AddHairTangentPass(
 	Parameters->HairStrandsVF_bIsCullingEnable = bCullingEnable ? 1 : 0;
 
 	FHairTangentCS::FPermutationDomain PermutationVector;
-	PermutationVector.Set<FHairTangentCS::FGroupSize>(GetGroupSizePermutation(GroupSize));
+	PermutationVector.Set<FHairTangentCS::FGroupSize>(GroupSize);
 	PermutationVector.Set<FHairTangentCS::FCulling>(bCullingEnable ? 1 : 0);
 
 	TShaderMapRef<FHairTangentCS> ComputeShader(ShaderMap, PermutationVector);
@@ -1087,7 +1082,7 @@ class FHairRaytracingGeometryCS : public FGlobalShader
 	DECLARE_GLOBAL_SHADER(FHairRaytracingGeometryCS);
 	SHADER_USE_PARAMETER_STRUCT(FHairRaytracingGeometryCS, FGlobalShader);
 
-	class FGroupSize : SHADER_PERMUTATION_INT("PERMUTATION_GROUP_SIZE", 2);
+	class FGroupSize : SHADER_PERMUTATION_SPARSE_INT("PERMUTATION_GROUP_SIZE", 32, 64);
 	class FCulling : SHADER_PERMUTATION_INT("PERMUTATION_CULLING", 2);
 	using FPermutationDomain = TShaderPermutationDomain<FGroupSize, FCulling>;
 
@@ -1151,7 +1146,7 @@ static void AddGenerateRaytracingGeometryPass(
 	}
 
 	FHairRaytracingGeometryCS::FPermutationDomain PermutationVector;
-	PermutationVector.Set<FHairRaytracingGeometryCS::FGroupSize>(GetGroupSizePermutation(GroupSize));
+	PermutationVector.Set<FHairRaytracingGeometryCS::FGroupSize>(GroupSize);
 	PermutationVector.Set<FHairRaytracingGeometryCS::FCulling>(bCullingEnable ? 1 : 0);
 
 	TShaderMapRef<FHairRaytracingGeometryCS> ComputeShader(ShaderMap, PermutationVector);
diff --git a/Engine/Shaders/Private/HairStrands/HairCardsDeformation.usf b/Engine/Shaders/Private/HairStrands/HairCardsDeformation.usf
--- a/Engine/Shaders/Private/HairStrands/HairCardsDeformation.usf
+++ b/Engine/Shaders/Private/HairStrands/HairCardsDeformation.usf
@@ -28,9 +28,10 @@ Buffer<float4>  TriangleDeformedPosition2Buffer;
 Buffer<uint>	GuideRootBarycentricBuffer;
 Buffer<uint>    GuideVertexToRootIndexBuffer;
 
-#if PERMUTATION_GROUP_SIZE == 0
+
+#if PERMUTATION_GROUP_SIZE == 64
 	#define GROUP_SIZE 64
-#elif PERMUTATION_GROUP_SIZE == 1
+#elif PERMUTATION_GROUP_SIZE == 32
 	#define GROUP_SIZE 32
 #else
 	#error Unknown group size	
diff --git a/Engine/Shaders/Private/HairStrands/HairStrandsGuideDeform.usf b/Engine/Shaders/Private/HairStrands/HairStrandsGuideDeform.usf
--- a/Engine/Shaders/Private/HairStrands/HairStrandsGuideDeform.usf
+++ b/Engine/Shaders/Private/HairStrands/HairStrandsGuideDeform.usf
@@ -12,16 +12,14 @@ Buffer<float4> SimDeformedOffsetBuffer;
 Buffer<uint4> SimRestPosePositionBuffer;
 RWBuffer<uint4> OutSimDeformedPositionBuffer;
 
-
-#if PERMUTATION_GROUP_SIZE == 0
-#define GROUP_SIZE 64
-#elif PERMUTATION_GROUP_SIZE == 1
-#define GROUP_SIZE 32
+#if PERMUTATION_GROUP_SIZE == 64
+	#define GROUP_SIZE 64
+#elif PERMUTATION_GROUP_SIZE == 32
+	#define GROUP_SIZE 32
 #else
-#error Unknown group size	
+	#error Unknown group size	
 #endif
 
-
 #define DEFORMATION_BYPASS 0
 #define DEFORMATION_OFFSET_GUIDE 1
 #define DEFORMATION_OFFSET_GUIDE_DYNAMIC 2
diff --git a/Engine/Shaders/Private/HairStrands/HairStrandsLightChannelMask.usf b/Engine/Shaders/Private/HairStrands/HairStrandsLightChannelMask.usf
--- a/Engine/Shaders/Private/HairStrands/HairStrandsLightChannelMask.usf
+++ b/Engine/Shaders/Private/HairStrands/HairStrandsLightChannelMask.usf
@@ -13,14 +13,14 @@ Texture2D<uint> 			NodeOffsetAndCount;
 StructuredBuffer<FPackedHairSample> NodeData;
 RWTexture2D<uint>			OutLightChannelMaskTexture;
 
-#if PERMUTATION_GROUP_SIZE == 0
-#define GROUP_SIZE_X 8
-#define GROUP_SIZE_Y 8
-#elif PERMUTATION_GROUP_SIZE == 1
-#define GROUP_SIZE_X 8
-#define GROUP_SIZE_Y 4
+#if PERMUTATION_GROUP_SIZE == 64
+	#define GROUP_SIZE_X 8
+	#define GROUP_SIZE_Y 8
+#elif PERMUTATION_GROUP_SIZE == 32
+	#define GROUP_SIZE_X 8
+	#define GROUP_SIZE_Y 4	
 #else
-#error Unknown group size	
+	#error Unknown group size	
 #endif
 
 [numthreads(GROUP_SIZE_X, GROUP_SIZE_Y, 1)]
diff --git a/Engine/Shaders/Private/HairStrands/HairStrandsRaytracingGeometry.usf b/Engine/Shaders/Private/HairStrands/HairStrandsRaytracingGeometry.usf
--- a/Engine/Shaders/Private/HairStrands/HairStrandsRaytracingGeometry.usf
+++ b/Engine/Shaders/Private/HairStrands/HairStrandsRaytracingGeometry.usf
@@ -22,9 +22,9 @@ Buffer<float>	HairStrandsVF_CullingRadiusScaleBuffer;
 Buffer<uint>	HairStrandsVF_CullingIndirectBufferArgs;
 #endif
 
-#if PERMUTATION_GROUP_SIZE == 0
+#if PERMUTATION_GROUP_SIZE == 64
 	#define GROUP_SIZE 64
-#elif PERMUTATION_GROUP_SIZE == 1
+#elif PERMUTATION_GROUP_SIZE == 32
 	#define GROUP_SIZE 32
 #else
 	#error Unknown group size	
diff --git a/Engine/Shaders/Private/HairStrands/HairStrandsTangent.usf b/Engine/Shaders/Private/HairStrands/HairStrandsTangent.usf
--- a/Engine/Shaders/Private/HairStrands/HairStrandsTangent.usf
+++ b/Engine/Shaders/Private/HairStrands/HairStrandsTangent.usf
@@ -19,9 +19,9 @@ Buffer<uint>	HairStrandsVF_CullingIndexBuffer;
 Buffer<uint>	IndirectBufferArgs;
 #endif
 
-#if PERMUTATION_GROUP_SIZE == 0
+#if PERMUTATION_GROUP_SIZE == 64
 	#define GROUP_SIZE 64
-#elif PERMUTATION_GROUP_SIZE == 1
+#elif PERMUTATION_GROUP_SIZE == 32
 	#define GROUP_SIZE 32
 #else
 	#error Unknown group size	
diff --git a/Engine/Shaders/Private/HairStrands/HairStrandsVisibilityCompaction.usf b/Engine/Shaders/Private/HairStrands/HairStrandsVisibilityCompaction.usf
--- a/Engine/Shaders/Private/HairStrands/HairStrandsVisibilityCompaction.usf
+++ b/Engine/Shaders/Private/HairStrands/HairStrandsVisibilityCompaction.usf
@@ -5,12 +5,14 @@
 #include "../Common.ush"
 #include "HairStrandsVisibilityCommon.ush" 
 
-#if PERMUTATION_GROUPSIZE == 0
-#define TILE_PIXEL_SIZE_X 8
-#define TILE_PIXEL_SIZE_Y 4
+#if PERMUTATION_GROUPSIZE ==32
+	#define TILE_PIXEL_SIZE_X 8
+	#define TILE_PIXEL_SIZE_Y 4
+#elif PERMUTATION_GROUPSIZE == 64
+	#define TILE_PIXEL_SIZE_X 8
+	#define TILE_PIXEL_SIZE_Y 8
 #else
-#define TILE_PIXEL_SIZE_X 8
-#define TILE_PIXEL_SIZE_Y 8
+	#error Unknown group size	
 #endif
 
 int2 OutputResolution;
diff --git a/Engine/Source/Runtime/Renderer/Private/HairStrands/HairStrandsVisibility.cpp b/Engine/Source/Runtime/Renderer/Private/HairStrands/HairStrandsVisibility.cpp
--- a/Engine/Source/Runtime/Renderer/Private/HairStrands/HairStrandsVisibility.cpp
+++ b/Engine/Source/Runtime/Renderer/Private/HairStrands/HairStrandsVisibility.cpp
@@ -1381,8 +1381,8 @@ class FHairLightChannelMaskCS : public FGlobalShader
 	DECLARE_GLOBAL_SHADER(FHairLightChannelMaskCS);
 	SHADER_USE_PARAMETER_STRUCT(FHairLightChannelMaskCS, FGlobalShader);
 
-	class FVendor : SHADER_PERMUTATION_INT("PERMUTATION_VENDOR", HairVisibilityVendorCount);
-	using FPermutationDomain = TShaderPermutationDomain<FVendor>;
+	class FGroupSize : SHADER_PERMUTATION_SPARSE_INT("PERMUTATION_GROUP_SIZE", 32, 64);
+	using FPermutationDomain = TShaderPermutationDomain<FGroupSize>;
 
 	BEGIN_SHADER_PARAMETER_STRUCT(FParameters, )
 		SHADER_PARAMETER(FIntPoint, OutputResolution)
@@ -1412,7 +1412,7 @@ static FRDGTextureRef AddHairLightChannelMaskPass(
 	FRDGTextureRef OutLightChannelMaskTexture = GraphBuilder.CreateTexture(Desc, TEXT("Hair.LightChannelMask"));
 
 	FHairLightChannelMaskCS::FPermutationDomain PermutationVector;
-	PermutationVector.Set<FHairLightChannelMaskCS::FVendor>(GetVendor());
+	PermutationVector.Set<FHairLightChannelMaskCS::FGroupSize>(GetVendorOptimalGroupSize1D());
 
 	FHairLightChannelMaskCS::FParameters* PassParameters = GraphBuilder.AllocParameters<FHairLightChannelMaskCS::FParameters>();
 	PassParameters->OutputResolution = Resolution;
@@ -1940,7 +1940,7 @@ class FHairVisibilityPrimitiveIdCompactionCS : public FGlobalShader
 	DECLARE_GLOBAL_SHADER(FHairVisibilityPrimitiveIdCompactionCS);
 	SHADER_USE_PARAMETER_STRUCT(FHairVisibilityPrimitiveIdCompactionCS, FGlobalShader);
 
-	class FGroupSize	: SHADER_PERMUTATION_INT("PERMUTATION_GROUPSIZE", 2);
+	class FGroupSize	: SHADER_PERMUTATION_SPARSE_INT("PERMUTATION_GROUPSIZE", 32, 64);
 	class FVelocity		: SHADER_PERMUTATION_INT("PERMUTATION_VELOCITY", 2);
 	class FViewTransmittance : SHADER_PERMUTATION_INT("PERMUTATION_VIEWTRANSMITTANCE", 2);
 	class FMaterial 	: SHADER_PERMUTATION_INT("PERMUTATION_MATERIAL_COMPACTION", 2);
@@ -2082,7 +2082,7 @@ static void AddHairVisibilityPrimitiveIdCompactionPass(
 	const bool bWriteOutVelocity = OutVelocityTexture != nullptr;
 	const uint32 VelocityPermutation = bWriteOutVelocity ? FMath::Clamp(GHairVelocityType + 1, 0, 3) : 0;
 	FHairVisibilityPrimitiveIdCompactionCS::FPermutationDomain PermutationVector;
-	PermutationVector.Set<FHairVisibilityPrimitiveIdCompactionCS::FGroupSize>(GetVendor() == HairVisibilityVendor_NVIDIA ? 0 : 1);
+	PermutationVector.Set<FHairVisibilityPrimitiveIdCompactionCS::FGroupSize>(GetVendorOptimalGroupSize1D());
 	PermutationVector.Set<FHairVisibilityPrimitiveIdCompactionCS::FVelocity>(VelocityPermutation > 0 ? 1 : 0);
 	PermutationVector.Set<FHairVisibilityPrimitiveIdCompactionCS::FViewTransmittance>(PassParameters->ViewTransmittanceTexture ? 1 : 0);
 	PermutationVector.Set<FHairVisibilityPrimitiveIdCompactionCS::FMaterial>(GHairStrandsMaterialCompactionEnable ? 1 : 0);
@@ -2140,7 +2140,7 @@ class FHairVisibilityCompactionComputeRasterCS : public FGlobalShader
 	DECLARE_GLOBAL_SHADER(FHairVisibilityCompactionComputeRasterCS);
 	SHADER_USE_PARAMETER_STRUCT(FHairVisibilityCompactionComputeRasterCS, FGlobalShader);
 
-	class FGroupSize : SHADER_PERMUTATION_INT("PERMUTATION_GROUPSIZE", 2);
+	class FGroupSize : SHADER_PERMUTATION_SPARSE_INT("PERMUTATION_GROUPSIZE", 32, 64);
 	using FPermutationDomain = TShaderPermutationDomain<FGroupSize>;
 
 	BEGIN_SHADER_PARAMETER_STRUCT(FParameters, )
@@ -2246,7 +2246,7 @@ static void AddHairVisibilityCompactionComputeRasterPass(
 
 	const FIntPoint GroupSize = GetVendorOptimalGroupSize2D();
 	FHairVisibilityCompactionComputeRasterCS::FPermutationDomain PermutationVector;
-	PermutationVector.Set<FHairVisibilityCompactionComputeRasterCS::FGroupSize>(GetVendor() == HairVisibilityVendor_NVIDIA ? 0 : 1);
+	PermutationVector.Set<FHairVisibilityCompactionComputeRasterCS::FGroupSize>(GetVendorOptimalGroupSize1D());
 	TShaderMapRef<FHairVisibilityCompactionComputeRasterCS> ComputeShader(View.ShaderMap, PermutationVector);
 	FComputeShaderUtils::AddPass(
 		GraphBuilder,
