Unload hair asset CPU data once the GPU resources have been created.

This is supported only with cooked build.

diff --git a/Engine/Plugins/Runtime/HairStrands/Source/HairStrandsCore/Private/GroomResources.cpp b/Engine/Plugins/Runtime/HairStrands/Source/HairStrandsCore/Private/GroomResources.cpp
--- a/Engine/Plugins/Runtime/HairStrands/Source/HairStrandsCore/Private/GroomResources.cpp
+++ b/Engine/Plugins/Runtime/HairStrands/Source/HairStrandsCore/Private/GroomResources.cpp
@@ -23,6 +23,9 @@
 #include "RenderTargetPool.h"
 #include "GroomBindingBuilder.h"
 
+static int32 GHairStrandsRelaseBulkDataAfterUse = 0;
+static FAutoConsoleVariableRef CVarHairStrandsRelaseBulkDataAfterUse(TEXT("r.HairStrands.Strands.RelaseBulkDataAfterUse"), GHairStrandsRelaseBulkDataAfterUse, TEXT("Release CPU bulk data once hair groom/groom binding asset GPU resources are created. This saves memory"));
+
 enum class EHairResourceUsageType : uint8
 {
 	Static,
@@ -144,6 +147,15 @@ void InternalCreateVertexBufferRDG_FromBulkData(FRDGBuilder& GraphBuilder, FByte
 		ERDGInitialDataFlags::None); // Copy data internally
 	InBulkData.Unlock();
 
+	// Unloading of the bulk data is only supported on cooked build, as we can reload the data from the file/archieve
+	#if !WITH_EDITORONLY_DATA
+	const bool bReleaseCPUData = GHairStrandsRelaseBulkDataAfterUse > 1;
+	if (bReleaseCPUData)
+	{
+		InBulkData.RemoveBulkData();
+	}
+	#endif
+
 	ConvertToExternalBufferWithViews(GraphBuilder, Buffer, Out, FormatType::Format);
 }
 
@@ -334,6 +346,14 @@ void InternalCreateStructuredBufferRDG_FromBulkData(FRDGBuilder& GraphBuilder, F
 	}
 	InBulkData.Unlock();
 
+	#if !WITH_EDITORONLY_DATA
+	const bool bReleaseCPUData = GHairStrandsRelaseBulkDataAfterUse > 1;
+	if (bReleaseCPUData)
+	{
+		InBulkData.RemoveBulkData();
+	}
+	#endif
+
 	ConvertToExternalBufferWithViews(GraphBuilder, Buffer, Out);
 }
 
