Fix hair cards generation

diff --git a/Engine/Plugins/Runtime/HairStrands/Source/HairStrandsCore/Private/GroomAsset.cpp b/Engine/Plugins/Runtime/HairStrands/Source/HairStrandsCore/Private/GroomAsset.cpp
--- a/Engine/Plugins/Runtime/HairStrands/Source/HairStrandsCore/Private/GroomAsset.cpp
+++ b/Engine/Plugins/Runtime/HairStrands/Source/HairStrandsCore/Private/GroomAsset.cpp
@@ -3342,7 +3342,13 @@ void UGroomAsset::SaveProceduralCards(uint32 DescIndex)
 			HairGroupsCards[DescIndex].ProceduralMesh = FHairStrandsCore::CreateStaticMesh(PackageName, SuffixName);
 		}
 
-		FHairCardsBuilder::ExportGeometry(Q.Data, HairGroupsCards[DescIndex].ProceduralMesh);
+		// Convert procedural cards data to cards data prior to export
+		FHairCardsDatas CardData;
+		{
+			CardData.Cards = Q.ProceduralData.Cards;
+		}
+
+		FHairCardsBuilder::ExportGeometry(CardData, HairGroupsCards[DescIndex].ProceduralMesh);
 		FHairStrandsCore::SaveAsset(HairGroupsCards[DescIndex].ProceduralMesh);
 		HairGroupsCards[DescIndex].ProceduralMeshKey = GroomDerivedDataCacheUtils::BuildCardsDerivedDataKeySuffix(this, *Desc);
 	}
diff --git a/Engine/Plugins/Runtime/HairStrands/Source/HairStrandsCore/Private/HairCardsBuilder.cpp b/Engine/Plugins/Runtime/HairStrands/Source/HairStrandsCore/Private/HairCardsBuilder.cpp
--- a/Engine/Plugins/Runtime/HairStrands/Source/HairStrandsCore/Private/HairCardsBuilder.cpp
+++ b/Engine/Plugins/Runtime/HairStrands/Source/HairStrandsCore/Private/HairCardsBuilder.cpp
@@ -3350,12 +3350,15 @@ void AllocateAtlasTexture(UTexture2D* Out, const FIntPoint& Resolution, uint32 M
 	Out->SetLayerFormatSettings(0, FormatSettings);
 #endif // #if WITH_EDITORONLY_DATA
 
+	// No need to allocate the platform data as they will be allocating once the source data are filled in
+#if 0
 	Out->PlatformData = new FTexturePlatformData();
 	Out->PlatformData->SizeX = Resolution.X;
 	Out->PlatformData->SizeY = Resolution.Y;
 	Out->PlatformData->PixelFormat = PixelFormat;
 
 	Out->UpdateResource();
+#endif
 }
 
 void AllocateAtlasTexture_Depth(UTexture2D* Out, const FIntPoint& Resolution, uint32 MipCount)		{ AllocateAtlasTexture(Out, Resolution, MipCount, PF_R8G8B8A8, ETextureSourceFormat::TSF_BGRA8); }
@@ -3816,8 +3819,8 @@ static void AddCardsTextureReadbackPass(
 	[InputTexture, OutTexture, BytePerPixel](FRHICommandListImmediate& RHICmdList)
 	{
 		const FIntPoint Resolution = InputTexture->Desc.Extent;
-		check(OutTexture->GetSurfaceWidth() == Resolution.X);
-		check(OutTexture->GetSurfaceHeight() == Resolution.Y);
+		check(OutTexture->Source.GetSizeX() == Resolution.X);
+		check(OutTexture->Source.GetSizeY() == Resolution.Y);
 
 		FRHIResourceCreateInfo CreateInfo(TEXT("CardsTextureReadbackPass_StagingTexture"));
 		FTexture2DRHIRef StagingTexture = RHICreateTexture2D(
