Change hair cards sampler to use anisotropic sampler for better accuracy.

diff --git a/Engine/Plugins/Runtime/HairStrands/Source/HairStrandsCore/Private/GroomAsset.cpp b/Engine/Plugins/Runtime/HairStrands/Source/HairStrandsCore/Private/GroomAsset.cpp
--- a/Engine/Plugins/Runtime/HairStrands/Source/HairStrandsCore/Private/GroomAsset.cpp
+++ b/Engine/Plugins/Runtime/HairStrands/Source/HairStrandsCore/Private/GroomAsset.cpp
@@ -398,7 +398,7 @@ static void InitAtlasTexture(ResourceType* InResource, UTexture2D* InTexture, EH
 	ENQUEUE_RENDER_COMMAND(HairStrandsCardsTextureCommand)(
 	[InResource, InTexture, InType](FRHICommandListImmediate& RHICmdList)
 	{
-		FSamplerStateRHIRef DefaultSampler = TStaticSamplerState<SF_Bilinear, AM_Clamp, AM_Clamp, AM_Clamp>::GetRHI();
+		FSamplerStateRHIRef DefaultSampler = TStaticSamplerState<SF_AnisotropicLinear, AM_Clamp, AM_Clamp, AM_Clamp, 0, 0>::GetRHI();
 		switch (InType)
 		{
 		case EHairAtlasTextureType::Depth:
