Add validation for cluster bulk data serialization and loading to track source of corruption.

diff --git a/Engine/Plugins/Runtime/HairStrands/Source/HairStrandsCore/Private/GroomBuilder.cpp b/Engine/Plugins/Runtime/HairStrands/Source/HairStrandsCore/Private/GroomBuilder.cpp
--- a/Engine/Plugins/Runtime/HairStrands/Source/HairStrandsCore/Private/GroomBuilder.cpp
+++ b/Engine/Plugins/Runtime/HairStrands/Source/HairStrandsCore/Private/GroomBuilder.cpp
@@ -216,7 +216,9 @@ namespace HairStrandsBuilder
 	template<typename TFormatType>
 	void CopyToBulkData(FByteBulkData& Out, const TArray<typename TFormatType::Type>& Data)
 	{
-		const uint32 DataSizeInByte = Data.Num() * sizeof(typename TFormatType::BulkType);
+		static_assert(TFormatType::SizeInByte == sizeof(typename TFormatType::BulkType));
+		const uint32 ElementSizeInByte = sizeof(typename TFormatType::BulkType);
+		const uint32 DataSizeInByte = Data.Num() * ElementSizeInByte;
 
 		// The buffer is then stored into bulk data
 		Out.Lock(LOCK_READ_WRITE);
diff --git a/Engine/Plugins/Runtime/HairStrands/Source/HairStrandsCore/Private/GroomResources.cpp b/Engine/Plugins/Runtime/HairStrands/Source/HairStrandsCore/Private/GroomResources.cpp
--- a/Engine/Plugins/Runtime/HairStrands/Source/HairStrandsCore/Private/GroomResources.cpp
+++ b/Engine/Plugins/Runtime/HairStrands/Source/HairStrandsCore/Private/GroomResources.cpp
@@ -23,8 +23,13 @@
 #include "RenderTargetPool.h"
 #include "GroomBindingBuilder.h"
 
-static int32 GHairStrandsRelaseBulkDataAfterUse = 0;
-static FAutoConsoleVariableRef CVarHairStrandsRelaseBulkDataAfterUse(TEXT("r.HairStrands.Strands.RelaseBulkDataAfterUse"), GHairStrandsRelaseBulkDataAfterUse, TEXT("Release CPU bulk data once hair groom/groom binding asset GPU resources are created. This saves memory"));
+static int32 GHairStrandsBulkData_ReleaseAfterUse = 0;
+static FAutoConsoleVariableRef CVarHairStrandsBulkData_ReleaseAfterUse(TEXT("r.HairStrands.Strands.BulkData.ReleaseAfterUse"), GHairStrandsBulkData_ReleaseAfterUse, TEXT("Release CPU bulk data once hair groom/groom binding asset GPU resources are created. This saves memory"));
+
+static int32 GHairStrandsBulkData_Validation = 1;
+static FAutoConsoleVariableRef CVarHairStrandsBulkData_Validation(TEXT("r.HairStrands.Strands.BulkData.Validation"), GHairStrandsBulkData_Validation, TEXT("Validate some hair strands data at serialization/loading time."));
+
+////////////////////////////////////////////////////////////////////////////////////
 
 enum class EHairResourceUsageType : uint8
 {
@@ -149,7 +154,7 @@ void InternalCreateVertexBufferRDG_FromBulkData(FRDGBuilder& GraphBuilder, FByte
 
 	// Unloading of the bulk data is only supported on cooked build, as we can reload the data from the file/archieve
 	#if !WITH_EDITORONLY_DATA
-	const bool bReleaseCPUData = GHairStrandsRelaseBulkDataAfterUse > 1;
+	const bool bReleaseCPUData = GHairStrandsBulkData_ReleaseAfterUse > 0;
 	if (bReleaseCPUData)
 	{
 		InBulkData.RemoveBulkData();
@@ -351,7 +356,7 @@ void InternalCreateStructuredBufferRDG_FromBulkData(FRDGBuilder& GraphBuilder, F
 	InBulkData.Unlock();
 
 	#if !WITH_EDITORONLY_DATA
-	const bool bReleaseCPUData = GHairStrandsRelaseBulkDataAfterUse > 1;
+	const bool bReleaseCPUData = GHairStrandsBulkData_ReleaseAfterUse > 0;
 	if (bReleaseCPUData)
 	{
 		InBulkData.RemoveBulkData();
@@ -859,6 +864,35 @@ void FHairStrandsClusterCullingBulkData::Serialize(FArchive& Ar, UObject* Owner)
 	{
 		PackedClusterInfos.Serialize(Ar, Owner, ChunkIndex, bAttemptFileMapping);
 	}
+
+	if (GHairStrandsBulkData_Validation > 0 && Ar.IsSaving())
+	{
+		Validate(true);
+	}
+}
+
+void FHairStrandsClusterCullingBulkData::Validate(bool bIsSaving)
+{
+	if (ClusterCount == 0)
+		return;
+
+	const FHairClusterInfo::Packed* Datas = (const FHairClusterInfo::Packed*)PackedClusterInfos.Lock(LOCK_READ_ONLY);
+	
+	// Simple heuristic to check if the data are valid
+	const uint32 MaxCount = FMath::Min(ClusterCount, 128u);
+	bool bIsValid = true;
+	for (uint32 It = 0; It < MaxCount; ++It)
+	{
+		bIsValid = bIsValid && Datas[It].LODCount <= 8;
+		if (!bIsValid) break;
+	}
+	if (!bIsValid)
+	{
+		FString Name = ClusterLODInfos.GetPackagePath().GetPackageName();
+		UE_LOG(LogHairStrands, Error, TEXT("[Groom/DDC] Strands - Invalid ClusterCullingBulkData when %s bulk data - %s"), bIsSaving ? TEXT("Saving") : TEXT("Loading"), *Name);
+	}
+
+	PackedClusterInfos.Unlock();
 }
 
 /////////////////////////////////////////////////////////////////////////////////////////
@@ -872,6 +906,11 @@ FHairStrandsClusterCullingResource::FHairStrandsClusterCullingResource(FHairStra
 
 void FHairStrandsClusterCullingResource::InternalAllocate(FRDGBuilder& GraphBuilder)
 {
+	if (GHairStrandsBulkData_Validation > 0)
+	{
+		BulkData.Validate(false);
+	}
+
 	InternalCreateStructuredBufferRDG_FromBulkData<FHairClusterInfoFormat>(GraphBuilder, BulkData.PackedClusterInfos, BulkData.ClusterCount, ClusterInfoBuffer, TEXT("Hair.StrandsClusterCulling_ClusterInfoBuffer"), EHairResourceUsageType::Static);
 	InternalCreateStructuredBufferRDG_FromBulkData<FHairClusterLODInfoFormat>(GraphBuilder, BulkData.ClusterLODInfos, BulkData.ClusterLODCount, ClusterLODInfoBuffer, TEXT("Hair.StrandsClusterCulling_ClusterLODInfoBuffer"), EHairResourceUsageType::Static);
 
diff --git a/Engine/Plugins/Runtime/HairStrands/Source/HairStrandsCore/Public/HairStrandsDatas.h b/Engine/Plugins/Runtime/HairStrands/Source/HairStrandsCore/Public/HairStrandsDatas.h
--- a/Engine/Plugins/Runtime/HairStrands/Source/HairStrandsCore/Public/HairStrandsDatas.h
+++ b/Engine/Plugins/Runtime/HairStrands/Source/HairStrandsCore/Public/HairStrandsDatas.h
@@ -731,6 +731,7 @@ struct HAIRSTRANDSCORE_API FHairStrandsClusterCullingBulkData
 	void Reset();
 	void Serialize(FArchive& Ar, UObject* Owner);
 	bool IsValid() const { return ClusterCount > 0 && VertexCount > 0; }
+	void Validate(bool bIsSaving);
 
 	/* Set LOD visibility, allowing to remove the simulation/rendering of certain LOD */
 	TArray<bool>	LODVisibility;
