Add GroomComponent reference to FHairGroupInstance for debug purpose.

diff --git a/Engine/Plugins/Runtime/HairStrands/Source/HairStrandsCore/Private/GroomComponent.cpp b/Engine/Plugins/Runtime/HairStrands/Source/HairStrandsCore/Private/GroomComponent.cpp
--- a/Engine/Plugins/Runtime/HairStrands/Source/HairStrandsCore/Private/GroomComponent.cpp
+++ b/Engine/Plugins/Runtime/HairStrands/Source/HairStrandsCore/Private/GroomComponent.cpp
@@ -2031,6 +2031,7 @@ void UGroomComponent::InitResources(bool bIsBindingReloading)
 		HairGroupInstance->Debug.GroupCount = GroupCount;
 		HairGroupInstance->Debug.GroomAssetName = GroomAsset->GetName();
 		HairGroupInstance->Debug.MeshComponent = RegisteredMeshComponent;
+		HairGroupInstance->Debug.GroomComponentForDebug = this;
 		HairGroupInstance->Debug.GroomBindingType = BindingAsset ? BindingAsset->GroomBindingType : EGroomBindingMeshType::SkeletalMesh;
 		HairGroupInstance->Debug.GroomCacheType = GroomCache ? GroomCache->GetType() : EGroomCacheType::None;
 		HairGroupInstance->Debug.GroomCacheBuffers = GroomCacheBuffers;
diff --git a/Engine/Plugins/Runtime/HairStrands/Source/HairStrandsCore/Public/GroomInstance.h b/Engine/Plugins/Runtime/HairStrands/Source/HairStrandsCore/Public/GroomInstance.h
--- a/Engine/Plugins/Runtime/HairStrands/Source/HairStrandsCore/Public/GroomInstance.h
+++ b/Engine/Plugins/Runtime/HairStrands/Source/HairStrandsCore/Public/GroomInstance.h
@@ -10,6 +10,7 @@
 #include "HairStrandsInterface.h"
 
 class UMeshComponent;
+class UGroomComponent;
 class FHairCardsVertexFactory;
 class FHairStrandsVertexFactory;
 enum class EGroomCacheType : uint8;
@@ -197,6 +198,7 @@ struct HAIRSTRANDSCORE_API FHairGroupInstance : public FHairStrandsInstance
 		EGroomCacheType			GroomCacheType;
 		UMeshComponent*			MeshComponent = nullptr;
 		FString					MeshComponentName;
+		const UGroomComponent*	GroomComponentForDebug = nullptr; // For debug only, shouldn't be deferred on the rendering thread
 		FTransform				SkeletalLocalToWorld = FTransform::Identity;
 		FTransform				SkeletalPreviousLocalToWorld = FTransform::Identity;
 		bool					bDrawCardsGuides = false;
