Fix invalid cluster AABB when culling is enabled.

The cluster AABB update can run onto non-initialized data due to the cluster info not being filled by the culling pass, if this one does not happen/or has culled the particular cluster, causing GPU crashes.

diff --git a/Engine/Plugins/Runtime/HairStrands/Source/HairStrandsCore/Private/HairStrandsRendering.cpp b/Engine/Plugins/Runtime/HairStrands/Source/HairStrandsCore/Private/HairStrandsRendering.cpp
--- a/Engine/Plugins/Runtime/HairStrands/Source/HairStrandsCore/Private/HairStrandsRendering.cpp
+++ b/Engine/Plugins/Runtime/HairStrands/Source/HairStrandsCore/Private/HairStrandsRendering.cpp
@@ -763,7 +763,6 @@ class FHairClusterAABBCS : public FGlobalShader
 
 	BEGIN_SHADER_PARAMETER_STRUCT(FParameters, )
 		SHADER_PARAMETER(uint32, DispatchCountX)
-		SHADER_PARAMETER(uint32, ClusterCount)
 		SHADER_PARAMETER(FVector3f, CPUBoundMin)
 		SHADER_PARAMETER(FVector3f, CPUBoundMax)
 		SHADER_PARAMETER(FMatrix44f, LocalToWorldMatrix)
@@ -774,9 +773,12 @@ class FHairClusterAABBCS : public FGlobalShader
 		SHADER_PARAMETER_RDG_BUFFER_SRV(Buffer, ClusterIdBuffer)
 		SHADER_PARAMETER_RDG_BUFFER_SRV(Buffer, ClusterIndexOffsetBuffer)
 		SHADER_PARAMETER_RDG_BUFFER_SRV(Buffer, ClusterIndexCountBuffer)
+		SHADER_PARAMETER_RDG_BUFFER_SRV(Buffer, CulledClusterCountBuffer)
 		SHADER_PARAMETER_RDG_BUFFER_SRV(Buffer<uint>, HairStrandsVF_CullingIndirectBuffer)
+		SHADER_PARAMETER_RDG_BUFFER_SRV(Buffer<uint>, IndirectBuffer)
 		SHADER_PARAMETER_RDG_BUFFER_UAV(RWBuffer, OutClusterAABBBuffer)
 		SHADER_PARAMETER_RDG_BUFFER_UAV(RWBuffer, OutGroupAABBBuffer)
+		RDG_BUFFER_ACCESS(IndirectArgsBuffer, ERHIAccess::IndirectArgs)
 	END_SHADER_PARAMETER_STRUCT()
 
 public:
@@ -817,12 +819,10 @@ static void AddHairClusterAABBPass(
 	// * If the instance has deformatin (simulation, global interpolation, skinning, then we update all clusters (for voxelization purpose) and the instance AABB
 	// * If the instance is static, we only update the instance AABB
 	const uint32 GroupSize = ComputeGroupSize();
-	const FIntVector DispatchCount = (UpdateType == EHairAABBUpdateType::UpdateClusterAABB && ClusterData) ? ComputeDispatchCount(ClusterData->ClusterCount) : FIntVector(1,1,1);
 
 	FHairClusterAABBCS::FParameters* Parameters = GraphBuilder.AllocParameters<FHairClusterAABBCS::FParameters>();
 	Parameters->CPUBoundMin = TransformedBounds.GetBox().Min;
 	Parameters->CPUBoundMax = TransformedBounds.GetBox().Max;
-	Parameters->DispatchCountX = DispatchCount.X;
 	Parameters->LocalToWorldMatrix = InRenLocalToWorld.ToMatrixWithScale();
 	Parameters->RenderDeformedPositionBuffer = RenderPositionBufferSRV;
 	Parameters->RenderDeformedOffsetBuffer = RenderDeformedOffsetBuffer;
@@ -830,15 +830,21 @@ static void AddHairClusterAABBPass(
 	if (UpdateType == EHairAABBUpdateType::UpdateClusterAABB)
 	{
 		check(ClusterData);
+		check(ClusterData->GroupSize1D == GroupSize);
+		check(ClusterData->CulledCluster2DIndirectArgsBuffer);
+		check(ClusterData->CulledClusterCountBuffer);
 
 		FRDGBufferRef ClusterIdBuffer = GraphBuilder.RegisterExternalBuffer(ClusterData->ClusterIdBuffer);
 		FRDGBufferRef ClusterIndexOffsetBuffer = GraphBuilder.RegisterExternalBuffer(ClusterData->ClusterIndexOffsetBuffer);
 		FRDGBufferRef ClusterIndexCountBuffer  = GraphBuilder.RegisterExternalBuffer(ClusterData->ClusterIndexCountBuffer);
+
 		Parameters->ClusterVertexIdBuffer = RegisterAsSRV(GraphBuilder, *ClusterData->ClusterVertexIdBuffer);
 		Parameters->ClusterIdBuffer = GraphBuilder.CreateSRV(ClusterIdBuffer, PF_R32_UINT);
 		Parameters->ClusterIndexOffsetBuffer = GraphBuilder.CreateSRV(ClusterIndexOffsetBuffer, PF_R32_UINT);
 		Parameters->ClusterIndexCountBuffer = GraphBuilder.CreateSRV(ClusterIndexCountBuffer, PF_R32_UINT);
-		Parameters->ClusterCount = ClusterData->ClusterCount;
+		Parameters->CulledClusterCountBuffer = GraphBuilder.CreateSRV(ClusterData->CulledClusterCountBuffer, PF_R32_UINT);
+		Parameters->IndirectBuffer = GraphBuilder.CreateSRV(ClusterData->CulledCluster2DIndirectArgsBuffer, PF_R32_UINT);
+		Parameters->IndirectArgsBuffer = ClusterData->CulledCluster2DIndirectArgsBuffer;
 
 		// Sanity check
 		check(ClusterData->ClusterCount == ClusterIdBuffer->Desc.NumElements);
@@ -857,12 +863,24 @@ static void AddHairClusterAABBPass(
 	PermutationVector.Set<FHairClusterAABBCS::FCPUAABB>(UpdateType == EHairAABBUpdateType::UpdateClusterAABB ? 0u : 1u);
 	TShaderMapRef<FHairClusterAABBCS> ComputeShader(ShaderMap, PermutationVector);
 
-	FComputeShaderUtils::AddPass(
-		GraphBuilder,
-		RDG_EVENT_NAME("HairStrandsClusterAABB(%s)", TEXT("Cluster,Group"), TEXT("Group Only")),
-		ComputeShader,
-		Parameters,
-		DispatchCount);
+	if (UpdateType == EHairAABBUpdateType::UpdateClusterAABB)
+	{
+		FComputeShaderUtils::AddPass(
+			GraphBuilder,
+			RDG_EVENT_NAME("HairStrandsClusterAABB(%s)", TEXT("Cluster,Group"), TEXT("Group Only")),
+			ComputeShader,
+			Parameters,
+			Parameters->IndirectArgsBuffer, 0);
+	}
+	else
+	{
+		FComputeShaderUtils::AddPass(
+			GraphBuilder,
+			RDG_EVENT_NAME("HairStrandsClusterAABB(%s)", TEXT("Cluster,Group"), TEXT("Group Only")),
+			ComputeShader,
+			Parameters,
+			FIntVector(1,1,1));
+	}
 
 	GraphBuilder.SetBufferAccessFinal(ClusterAABBData.ClusterAABBBuffer.Buffer, ERHIAccess::SRVMask);
 	GraphBuilder.SetBufferAccessFinal(ClusterAABBData.GroupAABBBuffer.Buffer, ERHIAccess::SRVMask);
diff --git a/Engine/Shaders/Private/HairStrands/HairStrandsClusterCulling.usf b/Engine/Shaders/Private/HairStrands/HairStrandsClusterCulling.usf
--- a/Engine/Shaders/Private/HairStrands/HairStrandsClusterCulling.usf
+++ b/Engine/Shaders/Private/HairStrands/HairStrandsClusterCulling.usf
@@ -676,8 +676,10 @@ void MainClusterCullingPrepareIndirectDispatchCS(
 /// 1. Indirect dispatch buffer preparation for rendering and prefix sums
 #ifdef SHADER_PREPAREINDIRECTDRAW
 
+uint GroupSize1D;
 Buffer<uint> DispatchIndirectParametersClusterCount;
 
+RWBuffer<uint> DispatchIndirectParametersClusterCount1D;
 RWBuffer<uint> DispatchIndirectParametersClusterCount2D;
 RWBuffer<uint> DispatchIndirectParametersClusterCountDiv512;
 RWBuffer<uint> DispatchIndirectParametersClusterCountDiv512Div512;
@@ -703,6 +705,10 @@ void MainClusterCullingPrepareIndirectDrawsCS(uint3 DispatchThreadId : SV_Dispat
 		DispatchIndirectParametersClusterCount2D[0] = DispatchGroupCount2D.x;
 		DispatchIndirectParametersClusterCount2D[1] = DispatchGroupCount2D.y;
 		DispatchIndirectParametersClusterCount2D[2] = 1.0;
+
+		DispatchIndirectParametersClusterCount1D[0] = (ClusterCount + GroupSize1D - 1) / GroupSize1D;
+		DispatchIndirectParametersClusterCount1D[1] = 1;
+		DispatchIndirectParametersClusterCount1D[2] = 1;
 	}
 }
 
diff --git a/Engine/Shaders/Private/HairStrands/HairStrandsInterpolation.usf b/Engine/Shaders/Private/HairStrands/HairStrandsInterpolation.usf
--- a/Engine/Shaders/Private/HairStrands/HairStrandsInterpolation.usf
+++ b/Engine/Shaders/Private/HairStrands/HairStrandsInterpolation.usf
@@ -562,9 +562,6 @@ void MainCS(uint2 DispatchThreadId : SV_DispatchThreadID)
 float3 CPUBoundMin;
 float3 CPUBoundMax;
 
-uint DispatchCountX;
-uint ClusterCount;
-//float3 OutHairPositionOffset;
 float4x4 LocalToWorldMatrix;
 
 Buffer<uint4> RenderDeformedPositionBuffer;
@@ -572,7 +569,9 @@ Buffer<uint> ClusterVertexIdBuffer;
 Buffer<uint> ClusterIdBuffer;
 Buffer<uint> ClusterIndexOffsetBuffer;
 Buffer<uint> ClusterIndexCountBuffer;
+Buffer<uint> CulledClusterCountBuffer;
 
+Buffer<uint> IndirectBuffer;
 Buffer<uint> HairStrandsVF_CullingIndirectBuffer;
 Buffer<float4> RenderDeformedOffsetBuffer;
 
@@ -599,8 +598,10 @@ void ClusterAABBEvaluationCS(uint3 DispatchThreadId : SV_DispatchThreadID, uint3
 		WriteHairAABB(0, Bound, OutGroupAABBBuffer);
 	}
 #else // PERMUTATION_USE_CPU_AABB
+	const uint CulledClusterCount = CulledClusterCountBuffer[0];
+	const uint DispatchCountX = IndirectBuffer[0];
 	const uint ClusterFetchIndex = GroupId.x + GroupId.y * DispatchCountX;
-	if (ClusterFetchIndex >= ClusterCount)
+	if (ClusterFetchIndex >= CulledClusterCount)
 	{
 		return;
 	}
diff --git a/Engine/Source/Runtime/Renderer/Private/HairStrands/HairStrandsClusters.cpp b/Engine/Source/Runtime/Renderer/Private/HairStrands/HairStrandsClusters.cpp
--- a/Engine/Source/Runtime/Renderer/Private/HairStrands/HairStrandsClusters.cpp
+++ b/Engine/Source/Runtime/Renderer/Private/HairStrands/HairStrandsClusters.cpp
@@ -104,7 +104,9 @@ class FMainClusterCullingPrepareIndirectDrawsCS : public FGlobalShader
 	SHADER_USE_PARAMETER_STRUCT(FMainClusterCullingPrepareIndirectDrawsCS, FGlobalShader);
 
 	BEGIN_SHADER_PARAMETER_STRUCT(FParameters, )
-		SHADER_PARAMETER_RDG_BUFFER_SRV(Buffer, DispatchIndirectParametersClusterCount)
+		SHADER_PARAMETER(uint32, GroupSize1D)
+		SHADER_PARAMETER_RDG_BUFFER_SRV(Buffer,   DispatchIndirectParametersClusterCount)
+		SHADER_PARAMETER_RDG_BUFFER_UAV(RWBuffer, DispatchIndirectParametersClusterCount1D)
 		SHADER_PARAMETER_RDG_BUFFER_UAV(RWBuffer, DispatchIndirectParametersClusterCount2D)
 		SHADER_PARAMETER_RDG_BUFFER_UAV(RWBuffer, DispatchIndirectParametersClusterCountDiv512)
 		SHADER_PARAMETER_RDG_BUFFER_UAV(RWBuffer, DispatchIndirectParametersClusterCountDiv512Div512)
@@ -248,6 +250,7 @@ static void AddClusterCullingPass(
 	FHairStrandClusterData::FHairGroup& ClusterData)
 {
 	FRDGBufferRef DispatchIndirectParametersClusterCount = GraphBuilder.CreateBuffer(FRDGBufferDesc::CreateIndirectDesc<FRHIDispatchIndirectParameters>(), TEXT("Hair.DispatchIndirectParametersClusterCount"));
+	FRDGBufferRef DispatchIndirectParametersClusterCount1D = GraphBuilder.CreateBuffer(FRDGBufferDesc::CreateIndirectDesc<FRHIDispatchIndirectParameters>(), TEXT("Hair.DispatchIndirectParametersClusterCount1D"));
 	FRDGBufferRef DispatchIndirectParametersClusterCount2D = GraphBuilder.CreateBuffer(FRDGBufferDesc::CreateIndirectDesc<FRHIDispatchIndirectParameters>(), TEXT("Hair.DispatchIndirectParametersClusterCount2D"));
 	FRDGBufferRef DispatchIndirectParametersClusterCountDiv512 = GraphBuilder.CreateBuffer(FRDGBufferDesc::CreateIndirectDesc<FRHIDispatchIndirectParameters>(), TEXT("Hair.DispatchIndirectParametersClusterCountDiv512"));
 	FRDGBufferRef DispatchIndirectParametersClusterCountDiv512Div512 = GraphBuilder.CreateBuffer(FRDGBufferDesc::CreateIndirectDesc<FRHIDispatchIndirectParameters>(), TEXT("Hair.DispatchIndirectParametersClusterCountDiv512Div512"));
@@ -369,9 +372,12 @@ static void AddClusterCullingPass(
 	}
 
 	/// Prepare some indirect draw buffers for specific compute group size
+	const uint32 GroupSize1D = GetVendorOptimalGroupSize1D();
 	{
 		FMainClusterCullingPrepareIndirectDrawsCS::FParameters* Parameters = GraphBuilder.AllocParameters<FMainClusterCullingPrepareIndirectDrawsCS::FParameters>();
+		Parameters->GroupSize1D = GroupSize1D;
 		Parameters->DispatchIndirectParametersClusterCount = GraphBuilder.CreateSRV(DispatchIndirectParametersClusterCount, PF_R32_UINT);
+		Parameters->DispatchIndirectParametersClusterCount1D = GraphBuilder.CreateUAV(DispatchIndirectParametersClusterCount1D, PF_R32_UINT);
 		Parameters->DispatchIndirectParametersClusterCount2D = GraphBuilder.CreateUAV(DispatchIndirectParametersClusterCount2D, PF_R32_UINT);
 		Parameters->DispatchIndirectParametersClusterCountDiv512 = GraphBuilder.CreateUAV(DispatchIndirectParametersClusterCountDiv512, PF_R32_UINT);
 		Parameters->DispatchIndirectParametersClusterCountDiv512Div512 = GraphBuilder.CreateUAV(DispatchIndirectParametersClusterCountDiv512Div512, PF_R32_UINT);
@@ -466,6 +472,10 @@ static void AddClusterCullingPass(
 		ClusterData.ClusterIdBuffer = GraphBuilder.ConvertToExternalBuffer(GlobalClusterIdBuffer);
 		ClusterData.ClusterIndexOffsetBuffer = GraphBuilder.ConvertToExternalBuffer(GlobalIndexStartBuffer);
 		ClusterData.ClusterIndexCountBuffer = GraphBuilder.ConvertToExternalBuffer(GlobalIndexCountBuffer);
+		ClusterData.CulledCluster2DIndirectArgsBuffer = DispatchIndirectParametersClusterCount2D;
+		ClusterData.CulledCluster1DIndirectArgsBuffer = DispatchIndirectParametersClusterCount1D;
+		ClusterData.CulledClusterCountBuffer = DispatchIndirectParametersClusterCount;
+		ClusterData.GroupSize1D = GroupSize1D;
 	}
 
 	/// Prepare some indirect dispatch for compute raster visibility buffers
@@ -491,10 +501,6 @@ static void AddClusterCullingPass(
 	{
 		ClusterData.ClusterDebugInfoBuffer = GraphBuilder.ConvertToExternalBuffer(ClusterDebugInfoBuffer);
 	}
-	if (bClusterDebug)
-	{
-		ClusterData.CulledDispatchIndirectParametersClusterCount = GraphBuilder.ConvertToExternalBuffer(DispatchIndirectParametersClusterCount);
-	}
 #endif
 
 	GraphBuilder.SetBufferAccessFinal(DrawIndirectParametersBuffer.Buffer, ERHIAccess::IndirectArgs | ERHIAccess::SRVMask);
diff --git a/Engine/Source/Runtime/Renderer/Private/HairStrands/HairStrandsDebug.cpp b/Engine/Source/Runtime/Renderer/Private/HairStrands/HairStrandsDebug.cpp
--- a/Engine/Source/Runtime/Renderer/Private/HairStrands/HairStrandsDebug.cpp
+++ b/Engine/Source/Runtime/Renderer/Private/HairStrands/HairStrandsDebug.cpp
@@ -999,9 +999,8 @@ static void AddDrawDebugClusterPass(
 						if (HairGroupPublicData != HairGroupClusters.HairGroupPublicPtr)
 							continue;
 
-						if (ShaderDrawDebug::IsEnabled(View) && HairGroupClusters.CulledDispatchIndirectParametersClusterCount)
+						if (ShaderDrawDebug::IsEnabled(View) && HairGroupClusters.CulledClusterCountBuffer)
 						{
-							FRDGBufferRef CulledDispatchIndirectParametersClusterCount = GraphBuilder.RegisterExternalBuffer(HairGroupClusters.CulledDispatchIndirectParametersClusterCount);
 							FRDGExternalBuffer& DrawIndirectBuffer = HairGroupClusters.HairGroupPublicPtr->GetDrawIndirectBuffer();
 
 							FDrawDebugClusterAABBCS::FPermutationDomain Permutation;
@@ -1015,7 +1014,7 @@ static void AddDrawDebugClusterPass(
 							Parameters->HairGroupId = DataIndex;
 							Parameters->ClusterDebugMode = GHairStrandsClusterDebug;
 							Parameters->ClusterAABBBuffer = HairGroupClusters.ClusterAABBBuffer->SRV;
-							Parameters->CulledDispatchIndirectParametersClusterCountBuffer = GraphBuilder.CreateSRV(CulledDispatchIndirectParametersClusterCount, EPixelFormat::PF_R32_UINT);
+							Parameters->CulledDispatchIndirectParametersClusterCountBuffer = GraphBuilder.CreateSRV(HairGroupClusters.CulledClusterCountBuffer, EPixelFormat::PF_R32_UINT);
 							Parameters->CulledDrawIndirectParameters = DrawIndirectBuffer.SRV;
 							Parameters->GroupAABBBuffer = HairGroupClusters.GroupAABBBuffer->SRV;
 
diff --git a/Engine/Source/Runtime/Renderer/Public/HairStrandsInterface.h b/Engine/Source/Runtime/Renderer/Public/HairStrandsInterface.h
--- a/Engine/Source/Runtime/Renderer/Public/HairStrandsInterface.h
+++ b/Engine/Source/Runtime/Renderer/Public/HairStrandsInterface.h
@@ -355,8 +355,11 @@ struct FHairStrandClusterData
 		bool GetCullingResultAvailable() const							{ return HairGroupPublicPtr ? HairGroupPublicPtr->GetCullingResultAvailable() : false; }
 		void SetCullingResultAvailable(bool b)							{ if (HairGroupPublicPtr) HairGroupPublicPtr->SetCullingResultAvailable(b); }
 
-		TRefCountPtr<FRDGPooledBuffer> ClusterDebugInfoBuffer;							// Null if this debug is not enabled.
-		TRefCountPtr<FRDGPooledBuffer> CulledDispatchIndirectParametersClusterCount;	// Null if this debug is not enabled.
+		TRefCountPtr<FRDGPooledBuffer> ClusterDebugInfoBuffer;	// Null if this debug is not enabled.
+		FRDGBufferRef CulledClusterCountBuffer = nullptr;
+		FRDGBufferRef CulledCluster1DIndirectArgsBuffer = nullptr;
+		FRDGBufferRef CulledCluster2DIndirectArgsBuffer = nullptr;
+		uint32 GroupSize1D = 0;
 
 		FHairGroupPublicData* HairGroupPublicPtr = nullptr;
 	};
