Workaround for fixing crash happening on packaged build with IOStore and compressed archieved.

diff --git a/Engine/Plugins/Runtime/HairStrands/Source/HairStrandsCore/Private/GroomResources.cpp b/Engine/Plugins/Runtime/HairStrands/Source/HairStrandsCore/Private/GroomResources.cpp
--- a/Engine/Plugins/Runtime/HairStrands/Source/HairStrandsCore/Private/GroomResources.cpp
+++ b/Engine/Plugins/Runtime/HairStrands/Source/HairStrandsCore/Private/GroomResources.cpp
@@ -23,6 +23,10 @@
 #include "RenderTargetPool.h"
 #include "GroomBindingBuilder.h"
 
+#ifndef CLEAR_COMPRESSION
+#define CLEAR_COMPRESSION 1
+#endif
+
 static int32 GHairStrandsBulkData_ReleaseAfterUse = 0;
 static FAutoConsoleVariableRef CVarHairStrandsBulkData_ReleaseAfterUse(TEXT("r.HairStrands.Strands.BulkData.ReleaseAfterUse"), GHairStrandsBulkData_ReleaseAfterUse, TEXT("Release CPU bulk data once hair groom/groom binding asset GPU resources are created. This saves memory"));
 
@@ -839,11 +843,24 @@ void FHairStrandsClusterCullingBulkData::Serialize(FArchive& Ar, UObject* Owner)
 	const int32 ChunkIndex = 0;
 	bool bAttemptFileMapping = false;
 
-	const uint32 BulkFlags = BULKDATA_Force_NOT_InlinePayload | BULKDATA_SerializeCompressed;
-	ClusterLODInfos.SetBulkDataFlags(BulkFlags);
-	VertexToClusterIds.SetBulkDataFlags(BulkFlags);
-	ClusterVertexIds.SetBulkDataFlags(BulkFlags);
-	PackedClusterInfos.SetBulkDataFlags(BulkFlags);
+	if (Ar.IsSaving())
+	{
+		const uint32 BulkFlags = BULKDATA_Force_NOT_InlinePayload | BULKDATA_SerializeCompressed;
+		ClusterLODInfos.SetBulkDataFlags(BulkFlags);
+		VertexToClusterIds.SetBulkDataFlags(BulkFlags);
+		ClusterVertexIds.SetBulkDataFlags(BulkFlags);
+		PackedClusterInfos.SetBulkDataFlags(BulkFlags);
+
+		#if CLEAR_COMPRESSION
+		if (Ar.IsCooking())
+		{
+			ClusterLODInfos.ClearBulkDataFlags(BULKDATA_SerializeCompressed);
+			VertexToClusterIds.ClearBulkDataFlags(BULKDATA_SerializeCompressed);
+			ClusterVertexIds.ClearBulkDataFlags(BULKDATA_SerializeCompressed);
+			PackedClusterInfos.ClearBulkDataFlags(BULKDATA_SerializeCompressed);
+		}
+		#endif
+	}
 
 	if (ClusterLODCount)
 	{
@@ -1136,16 +1153,34 @@ static void InternalSerialize(FArchive& Ar, UObject* Owner, FHairStrandsRootBulk
 	const int32 ChunkIndex = 0;
 	bool bAttemptFileMapping = false;
 
-	const uint32 BulkFlags = BULKDATA_Force_NOT_InlinePayload | BULKDATA_SerializeCompressed;
-	LOD.RootTriangleIndexBuffer.SetBulkDataFlags(BulkFlags);
-	LOD.RootTriangleBarycentricBuffer.SetBulkDataFlags(BulkFlags);
-	LOD.RestRootTrianglePosition0Buffer.SetBulkDataFlags(BulkFlags);
-	LOD.RestRootTrianglePosition1Buffer.SetBulkDataFlags(BulkFlags);
-	LOD.RestRootTrianglePosition2Buffer.SetBulkDataFlags(BulkFlags);
-
-	LOD.MeshInterpolationWeightsBuffer.SetBulkDataFlags(BulkFlags);
-	LOD.MeshSampleIndicesBuffer.SetBulkDataFlags(BulkFlags);
-	LOD.RestSamplePositionsBuffer.SetBulkDataFlags(BulkFlags);
+	if (Ar.IsSaving())
+	{
+		const uint32 BulkFlags = BULKDATA_Force_NOT_InlinePayload | BULKDATA_SerializeCompressed;
+		LOD.RootTriangleIndexBuffer.SetBulkDataFlags(BulkFlags);
+		LOD.RootTriangleBarycentricBuffer.SetBulkDataFlags(BulkFlags);
+		LOD.RestRootTrianglePosition0Buffer.SetBulkDataFlags(BulkFlags);
+		LOD.RestRootTrianglePosition1Buffer.SetBulkDataFlags(BulkFlags);
+		LOD.RestRootTrianglePosition2Buffer.SetBulkDataFlags(BulkFlags);
+
+		LOD.MeshInterpolationWeightsBuffer.SetBulkDataFlags(BulkFlags);
+		LOD.MeshSampleIndicesBuffer.SetBulkDataFlags(BulkFlags);
+		LOD.RestSamplePositionsBuffer.SetBulkDataFlags(BulkFlags);
+
+		#if CLEAR_COMPRESSION
+		if (Ar.IsCooking())
+		{
+			LOD.RootTriangleIndexBuffer.ClearBulkDataFlags(BULKDATA_SerializeCompressed);
+			LOD.RootTriangleBarycentricBuffer.ClearBulkDataFlags(BULKDATA_SerializeCompressed);
+			LOD.RestRootTrianglePosition0Buffer.ClearBulkDataFlags(BULKDATA_SerializeCompressed);
+			LOD.RestRootTrianglePosition1Buffer.ClearBulkDataFlags(BULKDATA_SerializeCompressed);
+			LOD.RestRootTrianglePosition2Buffer.ClearBulkDataFlags(BULKDATA_SerializeCompressed);
+
+			LOD.MeshInterpolationWeightsBuffer.ClearBulkDataFlags(BULKDATA_SerializeCompressed);
+			LOD.MeshSampleIndicesBuffer.ClearBulkDataFlags(BULKDATA_SerializeCompressed);
+			LOD.RestSamplePositionsBuffer.ClearBulkDataFlags(BULKDATA_SerializeCompressed);
+		}
+		#endif
+	}
 
 	Ar << LOD.LODIndex;
 	LOD.RootTriangleIndexBuffer.Serialize(Ar, Owner, ChunkIndex, bAttemptFileMapping);
@@ -1183,8 +1218,16 @@ void FHairStrandsRootBulkData::Serialize(FArchive& Ar, UObject* Owner)
 	const int32 ChunkIndex = 0;
 	bool bAttemptFileMapping = false;
 
-	const uint32 BulkFlags = BULKDATA_Force_NOT_InlinePayload | BULKDATA_SerializeCompressed;
-	VertexToCurveIndexBuffer.SetBulkDataFlags(BulkFlags);
+	{
+		const uint32 BulkFlags = BULKDATA_Force_NOT_InlinePayload | BULKDATA_SerializeCompressed;
+		VertexToCurveIndexBuffer.SetBulkDataFlags(BulkFlags);
+		#if CLEAR_COMPRESSION
+		if (Ar.IsCooking())
+		{
+			VertexToCurveIndexBuffer.ClearBulkDataFlags(BULKDATA_SerializeCompressed);
+		}
+		#endif
+	}
 
 	if (!Ar.IsObjectReferenceCollector())
 	{
diff --git a/Engine/Plugins/Runtime/HairStrands/Source/HairStrandsCore/Private/HairStrandsDatas.cpp b/Engine/Plugins/Runtime/HairStrands/Source/HairStrandsCore/Private/HairStrandsDatas.cpp
--- a/Engine/Plugins/Runtime/HairStrands/Source/HairStrandsCore/Private/HairStrandsDatas.cpp
+++ b/Engine/Plugins/Runtime/HairStrands/Source/HairStrandsCore/Private/HairStrandsDatas.cpp
@@ -4,6 +4,10 @@
 #include "UObject/ReleaseObjectVersion.h"
 #include "UObject/UE5ReleaseStreamObjectVersion.h"
 
+#ifndef CLEAR_COMPRESSION
+#define CLEAR_COMPRESSION 1
+#endif
+
 void FHairStrandsInterpolationDatas::SetNum(const uint32 NumCurves)
 {
 	PointsSimCurvesVertexWeights.SetNum(NumCurves);
@@ -181,11 +185,24 @@ void FHairStrandsInterpolationBulkData::Serialize(FArchive& Ar, UObject* Owner)
 	static_assert(sizeof(FHairInterpolation1Vertex::BulkType) == sizeof(FHairInterpolation1Vertex));
 	static_assert(sizeof(FHairStrandsRootIndexFormat::BulkType) == sizeof(FHairStrandsRootIndexFormat::Type));
 
-	const uint32 BulkFlags = BULKDATA_Force_NOT_InlinePayload | BULKDATA_SerializeCompressed;
-	Interpolation.SetBulkDataFlags(BulkFlags);
-	Interpolation0.SetBulkDataFlags(BulkFlags);
-	Interpolation1.SetBulkDataFlags(BulkFlags);
-	SimRootPointIndex.SetBulkDataFlags(BulkFlags);
+	if (Ar.IsSaving())
+	{
+		const uint32 BulkFlags = BULKDATA_Force_NOT_InlinePayload | BULKDATA_SerializeCompressed;
+		Interpolation.SetBulkDataFlags(BulkFlags);
+		Interpolation0.SetBulkDataFlags(BulkFlags);
+		Interpolation1.SetBulkDataFlags(BulkFlags);
+		SimRootPointIndex.SetBulkDataFlags(BulkFlags);
+
+		#if CLEAR_COMPRESSION
+		if (Ar.IsCooking())
+		{
+			Interpolation.ClearBulkDataFlags(BULKDATA_SerializeCompressed);
+			Interpolation0.ClearBulkDataFlags(BULKDATA_SerializeCompressed);
+			Interpolation1.ClearBulkDataFlags(BULKDATA_SerializeCompressed);
+			SimRootPointIndex.ClearBulkDataFlags(BULKDATA_SerializeCompressed);
+		}
+		#endif
+	}
 
 	Ar << Flags;
 	Ar << PointCount;
@@ -229,12 +246,26 @@ void FHairStrandsBulkData::Serialize(FArchive& Ar, UObject* Owner)
 
 	// Forced not inline means the bulk data won't automatically be loaded when we deserialize
 	// but only when we explicitly take action to load it
-	const uint32 BulkFlags = BULKDATA_Force_NOT_InlinePayload | BULKDATA_SerializeCompressed;
-	Positions.SetBulkDataFlags(BulkFlags);
-	Attributes0.SetBulkDataFlags(BulkFlags);
-	Attributes1.SetBulkDataFlags(BulkFlags);
-	Materials.SetBulkDataFlags(BulkFlags);
-	CurveOffsets.SetBulkDataFlags(BulkFlags);
+	if (Ar.IsSaving())
+	{
+		const uint32 BulkFlags = BULKDATA_Force_NOT_InlinePayload | BULKDATA_SerializeCompressed;
+		Positions.SetBulkDataFlags(BulkFlags);
+		Attributes0.SetBulkDataFlags(BulkFlags);
+		Attributes1.SetBulkDataFlags(BulkFlags);
+		Materials.SetBulkDataFlags(BulkFlags);
+		CurveOffsets.SetBulkDataFlags(BulkFlags);
+
+		#if CLEAR_COMPRESSION
+		if (Ar.IsCooking())
+		{
+			Positions.ClearBulkDataFlags(BULKDATA_SerializeCompressed);
+			Attributes0.ClearBulkDataFlags(BULKDATA_SerializeCompressed);
+			Attributes1.ClearBulkDataFlags(BULKDATA_SerializeCompressed);
+			Materials.ClearBulkDataFlags(BULKDATA_SerializeCompressed);
+			CurveOffsets.ClearBulkDataFlags(BULKDATA_SerializeCompressed);
+		}
+		#endif
+	}
 
 	if (!!(Flags & DataFlags_HasData))
 	{
