Fix GroomAsset invalid DDC request that are guaranteed to fail every run causing loading slow-down especially on high-latency backend like Jupiter
 - Saves 2.1 seconds for 14 grooms considering a 20ms latency to Jupiter

diff --git a/Engine/Plugins/Runtime/HairStrands/Source/HairStrandsCore/Private/GroomAsset.cpp b/Engine/Plugins/Runtime/HairStrands/Source/HairStrandsCore/Private/GroomAsset.cpp
--- a/Engine/Plugins/Runtime/HairStrands/Source/HairStrandsCore/Private/GroomAsset.cpp
+++ b/Engine/Plugins/Runtime/HairStrands/Source/HairStrandsCore/Private/GroomAsset.cpp
@@ -55,7 +55,6 @@ static TAutoConsoleVariable<int32> GHairStrandsWarningLogVerbosity(
 	-1,
 	TEXT("Enable warning log report for groom related asset (0: no logging, 1: error only, 2: error & warning only, other: all logs). By default all logging are enabled (-1). Value needs to be set at startup time."));
 
-// Editor async groom load can be useful in a workflow that consists mostly of loading grooms from a hot DDC
 static int32 GHairStrandsDDCLogEnable = 0;
 static FAutoConsoleVariableRef CVarHairStrandsDDCLogEnable(TEXT("r.HairStrands.DDCLog"), GHairStrandsDDCLogEnable, TEXT("Enable DDC logging for groom assets and groom binding assets"));
 
@@ -2078,10 +2077,8 @@ bool UGroomAsset::CacheStrandsData(uint32 GroupIndex, FString& OutDerivedDataKey
 	TArray<uint8> DerivedData;
 	if (GetDerivedDataCacheRef().GetSynchronous(*DerivedDataKey, DerivedData, GetPathName()))
 	{
-		if (IsHairStrandsDDCLogEnable())
-		{
-			UE_LOG(LogHairStrands, Log, TEXT("[Groom/DDC] Strands - Found (Groom:%s Group6:%d)."), *GetName(), GroupIndex);
-		}
+		UE_CLOG(IsHairStrandsDDCLogEnable(), LogHairStrands, Log, TEXT("[Groom/DDC] Strands - Found (Groom:%s Group6:%d)."), *GetName(), GroupIndex);
+
 		FMemoryReader Ar(DerivedData, /*bIsPersistent=*/ true);
 
 		int64 UncompressedSize = 0;
@@ -2103,10 +2100,7 @@ bool UGroomAsset::CacheStrandsData(uint32 GroupIndex, FString& OutDerivedDataKey
 			return false;
 		}
 
-		if (IsHairStrandsDDCLogEnable())
-		{
-			UE_LOG(LogHairStrands, Log, TEXT("[Groom/DDC] Strands - Not found (Groom:%s Group:%d)."), *GetName(), GroupIndex);
-		}
+		UE_CLOG(IsHairStrandsDDCLogEnable(), LogHairStrands, Log, TEXT("[Groom/DDC] Strands - Not found (Groom:%s Group:%d)."), *GetName(), GroupIndex);
 
 		const FHairDescriptionGroups& LocalHairDescriptionGroups = GetHairDescriptionGroups();
 
@@ -2178,22 +2172,44 @@ bool UGroomAsset::CacheCardsGeometry(uint32 GroupIndex, const FString& StrandsKe
 
 	FHairGroupData& HairGroupData = HairGroupsData[GroupIndex];
 
-	TArray<uint8> DerivedData;
-	if (GetDerivedDataCacheRef().GetSynchronous(*DerivedDataKey, DerivedData, GetPathName()))
+	// We should never query the DDC for something that we won't be able to build. So check that first.
+	// The goal is to avoid paying the price of a guaranteed failure on high-latency DDC backends every run.
+	bool bDataCanBeBuilt = false;
+
+	const FHairGroupsLOD& GroupsLOD = HairGroupsLOD[GroupIndex];
+	for (int32 LODIt = 0; LODIt < GroupsLOD.LODs.Num(); ++LODIt)
 	{
-		if (IsHairStrandsDDCLogEnable())
+		int32 SourceIt = 0;
+		if (FHairGroupsCardsSourceDescription* Desc = GetSourceDescription(HairGroupsCards, GroupIndex, LODIt, SourceIt))
 		{
-			UE_LOG(LogHairStrands, Log, TEXT("[Groom/DDC] Cards - Found (Groom:%s Group:%d)."), *GetName(), GroupIndex);
+			if (GroupsLOD.LODs[LODIt].GeometryType == EGroomGeometryType::Cards)
+			{
+				UStaticMesh* CardsMesh = nullptr;
+				if (Desc->SourceType == EHairCardsSourceType::Procedural)
+				{
+					CardsMesh = Desc->ProceduralMesh;
+				}
+				else if (Desc->SourceType == EHairCardsSourceType::Imported)
+				{
+					CardsMesh = Desc->ImportedMesh;
+				}
+				
+				bDataCanBeBuilt |= CardsMesh != nullptr;
+			}
 		}
+	}
+
+	TArray<uint8> DerivedData;
+	if (bDataCanBeBuilt && GetDerivedDataCacheRef().GetSynchronous(*DerivedDataKey, DerivedData, GetPathName()))
+	{
+		UE_CLOG(IsHairStrandsDDCLogEnable(), LogHairStrands, Log, TEXT("[Groom/DDC] Cards - Found (Groom:%s Group:%d)."), *GetName(), GroupIndex);
+		
 		FMemoryReader Ar(DerivedData, true);
 		InternalSerialize(Ar, this, HairGroupData.Cards.LODs);
 	}
 	else
 	{
-		if (IsHairStrandsDDCLogEnable())
-		{
-			UE_LOG(LogHairStrands, Log, TEXT("[Groom/DDC] Cards - Not found (Groom:%s Group:%d)."), *GetName(), GroupIndex);
-		}
+		UE_CLOG(IsHairStrandsDDCLogEnable(), LogHairStrands, Log, TEXT("[Groom/DDC] Cards - Not found (Groom:%s Group:%d)."), *GetName(), GroupIndex);
 
 		if (!BuildCardsGeometry(GroupIndex))
 		{
@@ -2253,23 +2269,34 @@ bool UGroomAsset::CacheMeshesGeometry(uint32 GroupIndex)
 
 	FHairGroupData& HairGroupData = HairGroupsData[GroupIndex];
 
-	TArray<uint8> DerivedData;
-	if (GetDerivedDataCacheRef().GetSynchronous(*DerivedDataKey, DerivedData, GetPathName()))
+	// We should never query the DDC for something that we won't be able to build. So check that first.
+	// The goal is to avoid paying the price of a guaranteed failure on high-latency DDC backends every run.
+	bool bDataCanBeBuilt = false;
+
+	const FHairGroupsLOD& GroupsLOD = HairGroupsLOD[GroupIndex];
+	for (int32 LODIt = 0; LODIt < GroupsLOD.LODs.Num(); ++LODIt)
 	{
-		if (IsHairStrandsDDCLogEnable())
+		int32 SourceIt = 0;
+		if (FHairGroupsMeshesSourceDescription* Desc = GetSourceDescription(HairGroupsMeshes, GroupIndex, LODIt, SourceIt))
 		{
-			UE_LOG(LogHairStrands, Log, TEXT("[Groom/DDC] Meshes - Found (Groom:%s Group:%d)."), *GetName(), GroupIndex);
+			if (GroupsLOD.LODs[LODIt].GeometryType == EGroomGeometryType::Meshes)
+			{
+				bDataCanBeBuilt |= Desc->ImportedMesh != nullptr;
+			}
 		}
+	}
+
+	TArray<uint8> DerivedData;
+	if (bDataCanBeBuilt && GetDerivedDataCacheRef().GetSynchronous(*DerivedDataKey, DerivedData, GetPathName()))
+	{
+		UE_CLOG(IsHairStrandsDDCLogEnable(), LogHairStrands, Log, TEXT("[Groom/DDC] Meshes - Found (Groom:%s Group:%d)."), *GetName(), GroupIndex);
 
 		FMemoryReader Ar(DerivedData, true);
 		InternalSerialize(Ar, this, HairGroupData.Meshes.LODs);
 	}
 	else
 	{
-		if (IsHairStrandsDDCLogEnable())
-		{
-			UE_LOG(LogHairStrands, Log, TEXT("[Groom/DDC] Meshes - Not found (Groom:%s Group:%d)."), *GetName(), GroupIndex);
-		}
+		UE_CLOG(IsHairStrandsDDCLogEnable(), LogHairStrands, Log, TEXT("[Groom/DDC] Meshes - Not found (Groom:%s Group:%d)."), *GetName(), GroupIndex);
 
 		if (!BuildMeshesGeometry(GroupIndex))
 		{
