Fix crash when changing certain groom asset property within the groom editor.

This CL releaeses the resource priori to recreate them, insuring that the new release are in the correct state.

diff --git a/Engine/Plugins/Runtime/HairStrands/Source/HairStrandsCore/Private/GroomAsset.cpp b/Engine/Plugins/Runtime/HairStrands/Source/HairStrandsCore/Private/GroomAsset.cpp
--- a/Engine/Plugins/Runtime/HairStrands/Source/HairStrandsCore/Private/GroomAsset.cpp
+++ b/Engine/Plugins/Runtime/HairStrands/Source/HairStrandsCore/Private/GroomAsset.cpp
@@ -2031,6 +2031,7 @@ bool UGroomAsset::CacheDerivedDatas()
 
 	if (bIsGameThread)
 	{
+		ReleaseResource();
 		InitResources();
 	}
 	return true;
