Update hair instance local to world transform during the LOD selection to insure the proper transform is selected.

diff --git a/Engine/Plugins/Runtime/HairStrands/Source/HairStrandsCore/Private/GroomManager.cpp b/Engine/Plugins/Runtime/HairStrands/Source/HairStrandsCore/Private/GroomManager.cpp
--- a/Engine/Plugins/Runtime/HairStrands/Source/HairStrandsCore/Private/GroomManager.cpp
+++ b/Engine/Plugins/Runtime/HairStrands/Source/HairStrandsCore/Private/GroomManager.cpp
@@ -724,6 +724,9 @@ static void RunHairLODSelection(
 		Instance->Guides.bIsSimulationEnable = Instance->HairGroupPublicData->IsSimulationEnable(IntLODIndex);
 		Instance->Guides.bHasGlobalInterpolation = Instance->HairGroupPublicData->IsGlobalInterpolationEnable(IntLODIndex);
 		Instance->Strands.bIsCullingEnabled = bCullingEnable;
+
+		// Update the local-to-world transform based on the binding type 
+		Instance->LocalToWorld = Instance->BindingType == EHairBindingType::Skinning ? Instance->Debug.SkeletalLocalToWorld : Instance->LocalToWorld;
 	}
 }
 
