Add optional debug code for rendering cluster bounding box.

diff --git a/Engine/Plugins/Runtime/HairStrands/Source/HairStrandsCore/Private/HairStrandsRendering.cpp b/Engine/Plugins/Runtime/HairStrands/Source/HairStrandsCore/Private/HairStrandsRendering.cpp
--- a/Engine/Plugins/Runtime/HairStrands/Source/HairStrandsCore/Private/HairStrandsRendering.cpp
+++ b/Engine/Plugins/Runtime/HairStrands/Source/HairStrandsCore/Private/HairStrandsRendering.cpp
@@ -762,7 +762,8 @@ class FHairClusterAABBCS : public FGlobalShader
 	using FPermutationDomain = TShaderPermutationDomain<FGroupSize, FCPUAABB>;
 
 	BEGIN_SHADER_PARAMETER_STRUCT(FParameters, )
-		SHADER_PARAMETER(uint32, DispatchCountX)
+		SHADER_PARAMETER_STRUCT_INCLUDE(ShaderDrawDebug::FShaderParameters, ShaderDrawParameters)
+		SHADER_PARAMETER(uint32, TotalClusterCount)
 		SHADER_PARAMETER(FVector3f, CPUBoundMin)
 		SHADER_PARAMETER(FVector3f, CPUBoundMax)
 		SHADER_PARAMETER(FMatrix44f, LocalToWorldMatrix)
@@ -801,6 +802,7 @@ enum class EHairAABBUpdateType
 static void AddHairClusterAABBPass(
 	FRDGBuilder& GraphBuilder,
 	FGlobalShaderMap* ShaderMap,
+	const FShaderDrawDebugData* ShaderDrawData,
 	const EHairAABBUpdateType UpdateType,
 	FHairGroupInstance* Instance,
 	FRDGBufferSRVRef RenderDeformedOffsetBuffer,
@@ -826,6 +828,12 @@ static void AddHairClusterAABBPass(
 	Parameters->LocalToWorldMatrix = InRenLocalToWorld.ToMatrixWithScale();
 	Parameters->RenderDeformedPositionBuffer = RenderPositionBufferSRV;
 	Parameters->RenderDeformedOffsetBuffer = RenderDeformedOffsetBuffer;
+	Parameters->TotalClusterCount = 1;
+	if (ShaderDrawDebug::IsEnabled() && ShaderDrawData)
+	{
+		ShaderDrawDebug::SetParameters(GraphBuilder, *ShaderDrawData, Parameters->ShaderDrawParameters);
+		
+	}
 
 	if (UpdateType == EHairAABBUpdateType::UpdateClusterAABB)
 	{
@@ -838,6 +846,7 @@ static void AddHairClusterAABBPass(
 		FRDGBufferRef ClusterIndexOffsetBuffer = GraphBuilder.RegisterExternalBuffer(ClusterData->ClusterIndexOffsetBuffer);
 		FRDGBufferRef ClusterIndexCountBuffer  = GraphBuilder.RegisterExternalBuffer(ClusterData->ClusterIndexCountBuffer);
 
+		Parameters->TotalClusterCount = ClusterData->ClusterCount;
 		Parameters->ClusterVertexIdBuffer = RegisterAsSRV(GraphBuilder, *ClusterData->ClusterVertexIdBuffer);
 		Parameters->ClusterIdBuffer = GraphBuilder.CreateSRV(ClusterIdBuffer, PF_R32_UINT);
 		Parameters->ClusterIndexOffsetBuffer = GraphBuilder.CreateSRV(ClusterIndexOffsetBuffer, PF_R32_UINT);
@@ -1728,6 +1737,7 @@ void ComputeHairStrandsInterpolation(
 					AddHairClusterAABBPass(
 						GraphBuilder,
 						ShaderMap,
+						ShaderDrawData,
 						UpdateType,
 						Instance,
 						Strands_PositionOffsetSRV,
diff --git a/Engine/Shaders/Private/HairStrands/HairStrandsInterpolation.usf b/Engine/Shaders/Private/HairStrands/HairStrandsInterpolation.usf
--- a/Engine/Shaders/Private/HairStrands/HairStrandsInterpolation.usf
+++ b/Engine/Shaders/Private/HairStrands/HairStrandsInterpolation.usf
@@ -7,6 +7,13 @@
 #include "HairStrandsClusterCommon.ush"
 #include "HairStrandsVertexFactoryCommon.ush"
 
+#define DEBUG_ENABLE 0
+
+#if DEBUG_ENABLE
+#include "../ShaderDrawDebug.ush"
+#include "../ColorMap.ush"
+#endif
+
 ///////////////////////////////////////////////////////////////////////////////////////////////////
 // Common
 
@@ -559,6 +566,7 @@ void MainCS(uint2 DispatchThreadId : SV_DispatchThreadID)
 	#error Unknown group size	
 #endif
 
+uint TotalClusterCount;
 float3 CPUBoundMin;
 float3 CPUBoundMax;
 
@@ -586,7 +594,7 @@ groupshared float3 SharedClusterAABBMin[GROUP_SIZE];
 groupshared float3 SharedClusterAABBMax[GROUP_SIZE];
 
 [numthreads(GROUP_SIZE, 1, 1)]
-void ClusterAABBEvaluationCS(uint3 DispatchThreadId : SV_DispatchThreadID, uint3 GroupId : SV_GroupID, uint3 GroupThreadId : SV_GroupThreadID)
+void ClusterAABBEvaluationCS(uint3 DispatchThreadId : SV_DispatchThreadID, uint3 GroupId : SV_GroupID, uint3 GroupThreadId : SV_GroupThreadID, uint ThreadIndex : SV_GroupIndex)
 {
 #if PERMUTATION_USE_CPU_AABB
 	// Only update the full group AABB
@@ -706,6 +714,15 @@ void ClusterAABBEvaluationCS(uint3 DispatchThreadId : SV_DispatchThreadID, uint3
 		InterlockHairAABB(0, Bound, OutGroupAABBBuffer);
 	}
 
+	#if DEBUG_ENABLE
+	if (ThreadIndex == 0)
+	{
+		const float4 DebugColor = float4(ColorMapViridis(ClusterId / float(TotalClusterCount)),1);
+		AddAABB(SharedClusterAABBMin[0], SharedClusterAABBMax[0], DebugColor);
+		AddReferential(LocalToWorldMatrix, 10);
+	}
+	#endif
+
 	// Min/Max group AABB is done later by another pass
 #endif // PERMUTATION_USE_CPU_AABB
 }
diff --git a/Engine/Shaders/Private/HairStrands/HairStrandsVoxelPageRayMarching.usf b/Engine/Shaders/Private/HairStrands/HairStrandsVoxelPageRayMarching.usf
--- a/Engine/Shaders/Private/HairStrands/HairStrandsVoxelPageRayMarching.usf
+++ b/Engine/Shaders/Private/HairStrands/HairStrandsVoxelPageRayMarching.usf
@@ -85,7 +85,7 @@ void MainCS(uint3 DispatchThreadId : SV_DispatchThreadID)
 			TotalAllocatedPageCount += LocalNodeDesc.PageIndexResolution.x * LocalNodeDesc.PageIndexResolution.y * LocalNodeDesc.PageIndexResolution.z;
 		}
 
-		AddAABB(NodeDesc.MinAABB, NodeDesc.MaxAABB, float4(1, 0.5f, 0, 1));
+		AddAABB(NodeDesc.MinAABB, NodeDesc.MaxAABB, ColorOrange);
 
 		// General information
 		// General
diff --git a/Engine/Shaders/Private/ShaderDrawDebug.ush b/Engine/Shaders/Private/ShaderDrawDebug.ush
--- a/Engine/Shaders/Private/ShaderDrawDebug.ush
+++ b/Engine/Shaders/Private/ShaderDrawDebug.ush
@@ -172,18 +172,14 @@ void AddCross(float3 Pos, float Size, float4 Color)
 	AddLine(Pos - float3(0,0,Size), Pos + float3(0,0,Size), Color, Color);
 }
 
-void AddReferential(float3 Pos, float3 T, float3 B, float3 N)
+void AddReferential(float3 Pos, float3 T, float3 B, float3 N, float Scale=1)
 {
-	const float4 Red = float4(1, 0, 0, 1);
-	const float4 Green = float4(0, 1, 0, 1);
-	const float4 Blue = float4(0, 0, 1, 1);
-
-	AddLine(Pos, Pos + T, Red, Red);
-	AddLine(Pos, Pos + B, Green, Green);
-	AddLine(Pos, Pos + N, Blue, Blue);
+	AddLine(Pos, Pos + T*Scale, ColorRed);
+	AddLine(Pos, Pos + B*Scale, ColorGreen);
+	AddLine(Pos, Pos + N*Scale, ColorBlue);
 }
 
-void AddReferential(float3 Pos, float3 TangentZ)
+void AddReferential(float3 Pos, float3 TangentZ, float Scale=1)
 {
 	const float Sign = TangentZ.z >= 0 ? 1 : -1;
 	const float a = -rcp(Sign + TangentZ.z);
@@ -192,7 +188,12 @@ void AddReferential(float3 Pos, float3 TangentZ)
 	const float3 TangentX = { 1 + Sign * a * Pow2(TangentZ.x), Sign * b, -Sign * TangentZ.x };
 	const float3 TangentY = { b,  Sign + a * Pow2(TangentZ.y), -TangentZ.y };
 
-	AddReferential(Pos, TangentX, TangentY, TangentZ);
+	AddReferential(Pos, TangentX, TangentY, TangentZ, Scale);
+}
+
+void AddReferential(float4x4 InM, float Scale = 1)
+{
+	AddReferential(InM[3].xyz, InM[0].xyz, InM[1].xyz, InM[2].xyz, Scale);
 }
 
 void AddLineTriangle(float3 P0, float3 P1, float3 P2, float4 Color)
