Early out groom caching for cards/meshes if the groom don't contain any cards/meshes data.

This avoids to send a DDC/high-latency request for fetching non-existing data.

diff --git a/Engine/Plugins/Runtime/HairStrands/Source/HairStrandsCore/Private/GroomAsset.cpp b/Engine/Plugins/Runtime/HairStrands/Source/HairStrandsCore/Private/GroomAsset.cpp
--- a/Engine/Plugins/Runtime/HairStrands/Source/HairStrandsCore/Private/GroomAsset.cpp
+++ b/Engine/Plugins/Runtime/HairStrands/Source/HairStrandsCore/Private/GroomAsset.cpp
@@ -2160,8 +2160,59 @@ bool UGroomAsset::CacheStrandsData(uint32 GroupIndex, FString& OutDerivedDataKey
 	return bSuccess;
 }
 
+bool UGroomAsset::HasValidMeshesData(uint32 GroupIndex) const
+{
+	const FHairGroupData& HairGroupData = HairGroupsData[GroupIndex];
+	const FHairGroupsLOD& GroupsLOD = HairGroupsLOD[GroupIndex];
+	for (int32 LODIt = 0; LODIt < GroupsLOD.LODs.Num(); ++LODIt)
+	{
+		if (GroupsLOD.LODs[LODIt].GeometryType == EGroomGeometryType::Meshes)
+		{
+			int32 SourceIt = 0;
+			if (const FHairGroupsMeshesSourceDescription* Desc = GetSourceDescription(HairGroupsMeshes, GroupIndex, LODIt, SourceIt))
+			{
+				if (Desc->ImportedMesh)
+				{
+					return true;
+				}
+			}
+		}
+	}
+	return false;
+}
+
+bool UGroomAsset::HasValidCardsData(uint32 GroupIndex) const
+{
+	const FHairGroupData& HairGroupData = HairGroupsData[GroupIndex];
+	const FHairGroupsLOD& GroupsLOD = HairGroupsLOD[GroupIndex];
+	for (int32 LODIt = 0; LODIt < GroupsLOD.LODs.Num(); ++LODIt)
+	{
+		if (GroupsLOD.LODs[LODIt].GeometryType == EGroomGeometryType::Cards)
+		{
+			int32 SourceIt = 0;
+			if (const FHairGroupsCardsSourceDescription* Desc = GetSourceDescription(HairGroupsCards, GroupIndex, LODIt, SourceIt))
+			{
+				if (Desc->GetMesh())
+				{
+					return true;
+				}
+			}
+		}
+	}
+	return false;
+}
+
+// Return true if the meshes data have changed
 bool UGroomAsset::CacheCardsGeometry(uint32 GroupIndex, const FString& StrandsKey)
 {
+	const bool bHasValidCardsData = HasValidCardsData(GroupIndex);
+	if (!bHasValidCardsData)
+	{
+		const bool bNeedReset = !CardsDerivedDataKey[GroupIndex].IsEmpty();
+		CardsDerivedDataKey[GroupIndex].Empty();
+		return bNeedReset;
+	}
+
 	const FString KeySuffix = GroomDerivedDataCacheUtils::BuildCardsDerivedDataKeySuffix(GroupIndex, HairGroupsLOD[GroupIndex].LODs, HairGroupsCards);
 	const FString DerivedDataKey = StrandsKey + KeySuffix;
 
@@ -2257,8 +2308,17 @@ bool UGroomAsset::CacheCardsGeometry(uint32 GroupIndex, const FString& StrandsKe
 	return true;
 }
 
+// Return true if the meshes data have changed
 bool UGroomAsset::CacheMeshesGeometry(uint32 GroupIndex)
 {
+	const bool bHasValidMeshesData = UGroomAsset::HasValidMeshesData(GroupIndex);
+	if (!bHasValidMeshesData)
+	{
+		const bool bNeedReset = !MeshesDerivedDataKey[GroupIndex].IsEmpty();
+		MeshesDerivedDataKey[GroupIndex].Empty();
+		return bNeedReset;
+	}
+
 	const FString KeySuffix = GroomDerivedDataCacheUtils::BuildMeshesDerivedDataKeySuffix(GroupIndex, HairGroupsLOD[GroupIndex].LODs, HairGroupsMeshes);
 	const FString DerivedDataKey = GroomDerivedDataCacheUtils::BuildGroomDerivedDataKey(KeySuffix);
 
diff --git a/Engine/Plugins/Runtime/HairStrands/Source/HairStrandsCore/Private/GroomAssetCards.cpp b/Engine/Plugins/Runtime/HairStrands/Source/HairStrandsCore/Private/GroomAssetCards.cpp
--- a/Engine/Plugins/Runtime/HairStrands/Source/HairStrandsCore/Private/GroomAssetCards.cpp
+++ b/Engine/Plugins/Runtime/HairStrands/Source/HairStrandsCore/Private/GroomAssetCards.cpp
@@ -119,20 +119,10 @@ bool FHairGroupsCardsSourceDescription::operator==(const FHairGroupsCardsSourceD
 FString FHairGroupsCardsSourceDescription::GetMeshKey() const
 {
 #if WITH_EDITORONLY_DATA
-	if (SourceType == EHairCardsSourceType::Imported && ImportedMesh)
-	{
-		ImportedMesh->ConditionalPostLoad();
-		FStaticMeshSourceModel& SourceModel = ImportedMesh->GetSourceModel(0);
-		if (SourceModel.GetMeshDescriptionBulkData())
-		{
-			return SourceModel.GetMeshDescriptionBulkData()->GetIdString();
-		}
-
-	}
-	else if (SourceType == EHairCardsSourceType::Procedural && ProceduralMesh)
+	if (UStaticMesh* Mesh = GetMesh())
 	{
-		ProceduralMesh->ConditionalPostLoad();
-		FStaticMeshSourceModel& SourceModel = ProceduralMesh->GetSourceModel(0);
+		Mesh->ConditionalPostLoad();
+		FStaticMeshSourceModel& SourceModel = Mesh->GetSourceModel(0);
 		if (SourceModel.GetMeshDescriptionBulkData())
 		{
 			return SourceModel.GetMeshDescriptionBulkData()->GetIdString();
@@ -170,4 +160,20 @@ void FHairGroupsCardsSourceDescription::UpdateMeshKey()
 		ProceduralMeshKey = GetMeshKey();
 	}
 #endif
+}
+
+UStaticMesh* FHairGroupsCardsSourceDescription::GetMesh() const
+{
+#if WITH_EDITORONLY_DATA
+	if (SourceType == EHairCardsSourceType::Imported)
+	{
+		return ImportedMesh;
+
+	}
+	else if (SourceType == EHairCardsSourceType::Procedural)
+	{
+		return ProceduralMesh;
+	}
+#endif
+	return nullptr;
 }
\ No newline at end of file
diff --git a/Engine/Plugins/Runtime/HairStrands/Source/HairStrandsCore/Public/GroomAsset.h b/Engine/Plugins/Runtime/HairStrands/Source/HairStrandsCore/Public/GroomAsset.h
--- a/Engine/Plugins/Runtime/HairStrands/Source/HairStrandsCore/Public/GroomAsset.h
+++ b/Engine/Plugins/Runtime/HairStrands/Source/HairStrandsCore/Public/GroomAsset.h
@@ -599,6 +599,8 @@ private:
 	bool BuildCardsGeometry(uint32 GroupIndex);
 	bool BuildMeshesGeometry(uint32 GroupIndex);
 
+	bool HasValidCardsData(uint32 GroupIndex) const;
+	bool HasValidMeshesData(uint32 GroupIndex) const;
 public:
 	/** Commits a HairDescription to buffer for serialization */
 	void CommitHairDescription(FHairDescription&& HairDescription);
diff --git a/Engine/Plugins/Runtime/HairStrands/Source/HairStrandsCore/Public/GroomAssetCards.h b/Engine/Plugins/Runtime/HairStrands/Source/HairStrandsCore/Public/GroomAssetCards.h
--- a/Engine/Plugins/Runtime/HairStrands/Source/HairStrandsCore/Public/GroomAssetCards.h
+++ b/Engine/Plugins/Runtime/HairStrands/Source/HairStrandsCore/Public/GroomAssetCards.h
@@ -232,6 +232,7 @@ struct HAIRSTRANDSCORE_API FHairGroupsCardsSourceDescription
 
 	bool operator==(const FHairGroupsCardsSourceDescription& A) const;
 
+	UStaticMesh* GetMesh() const;
 	FString GetMeshKey() const;
 	bool HasMeshChanged() const;
 	void UpdateMeshKey();
