Insure the skel. mesh section is not used, if it's data is not streamed yet.

When a skel. mesh is not stream, its SRV can be null, causing a crash when transforming/interpolating hair data.

This is more aworkaround that an actual solution as the SRV might never be created.

diff --git a/Engine/Plugins/Runtime/HairStrands/Source/HairStrandsCore/Private/HairStrandsMeshProjection.cpp b/Engine/Plugins/Runtime/HairStrands/Source/HairStrandsCore/Private/HairStrandsMeshProjection.cpp
--- a/Engine/Plugins/Runtime/HairStrands/Source/HairStrandsCore/Private/HairStrandsMeshProjection.cpp
+++ b/Engine/Plugins/Runtime/HairStrands/Source/HairStrandsCore/Private/HairStrandsMeshProjection.cpp
@@ -768,6 +768,12 @@ void AddHairStrandUpdateMeshTrianglesPass(
 
 			const FHairStrandsProjectionMeshData::Section& MeshSectionData = MeshData.Sections[SectionIndex];
 
+			// If a skel. mesh is not streaming yet, its SRV will be null
+			if (MeshSectionData.IndexBuffer == nullptr)
+			{
+				continue;
+			}
+
 			const FMeshSectionBuffers* Buffers = nullptr;
 			if (MeshSectionData.PositionBuffer)
 			{
