Implemented streaming manager with look-ahead buffering for GroomCache.

diff --git a/Engine/Plugins/Runtime/HairStrands/Source/HairStrandsCore/Private/GroomCache.cpp b/Engine/Plugins/Runtime/HairStrands/Source/HairStrandsCore/Private/GroomCache.cpp
--- a/Engine/Plugins/Runtime/HairStrands/Source/HairStrandsCore/Private/GroomCache.cpp
+++ b/Engine/Plugins/Runtime/HairStrands/Source/HairStrandsCore/Private/GroomCache.cpp
@@ -70,7 +70,7 @@ int32 UGroomCache::GetFrameIndexAtTime(const float Time, bool bLooping) const
 	return NormalizedFrame; 
 }
 
-void UGroomCache::FindSampleIndexesFromTime(float Time, bool bLooping, bool bIsPlayingBackwards, int32 &OutFrameIndex, int32 &OutNextFrameIndex, float &InterpolationFactor)
+void UGroomCache::GetFrameIndicesAtTime(float Time, bool bLooping, bool bIsPlayingBackwards, int32 &OutFrameIndex, int32 &OutNextFrameIndex, float &InterpolationFactor)
 {
 	const int32 NumFrames = GroomCacheInfo.AnimationInfo.NumFrames;
 	const float Duration = GroomCacheInfo.AnimationInfo.Duration;
@@ -117,6 +117,61 @@ void UGroomCache::FindSampleIndexesFromTime(float Time, bool bLooping, bool bIsP
 	}
 }
 
+void UGroomCache::GetFrameIndicesForTimeRange(float StartTime, float EndTime, bool bLooping, TArray<int32>& OutFrameIndices)
+{
+	// Sanity check
+	if (StartTime > EndTime)
+	{
+		EndTime = StartTime;
+	}
+
+	// Ensure the time range covers at least one frame
+	if (EndTime - StartTime < GroomCacheInfo.AnimationInfo.SecondsPerFrame)
+	{
+		EndTime = StartTime + GroomCacheInfo.AnimationInfo.SecondsPerFrame;
+	}
+
+	int32 StartIndex = GetFrameIndexAtTime(StartTime, bLooping);
+	int32 EndIndex = GetFrameIndexAtTime(EndTime, bLooping);
+
+	if (bLooping)
+	{
+		// Special cases to handle with looping enabled
+		if (((EndTime - StartTime) >= GroomCacheInfo.AnimationInfo.Duration) || 
+			(StartIndex == EndIndex))
+		{
+			// Requested time range is longer than the animation or exactly matches the duration so include all the frames
+			for (int32 FrameIndex = 0; FrameIndex < Chunks.Num(); ++FrameIndex)
+			{
+				OutFrameIndices.Add(FrameIndex);
+			}
+			return;
+		}
+
+		if (EndIndex < StartIndex)
+		{
+			// The requested time range is wrapping so add as two intervals
+			// From start index to end of animation
+			for (int32 FrameIndex = StartIndex; FrameIndex < (int32)GroomCacheInfo.AnimationInfo.NumFrames; ++FrameIndex)
+			{
+				OutFrameIndices.Add(FrameIndex);
+			}
+			// From 0 to end index
+			for (int32 FrameIndex = 0; FrameIndex <= EndIndex; ++FrameIndex)
+			{
+				OutFrameIndices.Add(FrameIndex);
+			}
+			return;
+		}
+	}
+
+	// Time range is a simple interval within the animation
+	for (int32 FrameIndex = StartIndex; FrameIndex <= EndIndex; ++FrameIndex)
+	{
+		OutFrameIndices.Add(FrameIndex);
+	}
+}
+
 bool UGroomCache::GetGroomDataAtTime(float Time, bool bLooping, FGroomCacheAnimationData& AnimData)
 {
 	const int32 FrameIndex = GetFrameIndexAtTime(Time, bLooping);
diff --git a/Engine/Plugins/Runtime/HairStrands/Source/HairStrandsCore/Private/GroomCacheStreamingData.cpp b/Engine/Plugins/Runtime/HairStrands/Source/HairStrandsCore/Private/GroomCacheStreamingData.cpp
new file mode 100644
--- /dev/null
+++ b/Engine/Plugins/Runtime/HairStrands/Source/HairStrandsCore/Private/GroomCacheStreamingData.cpp
@@ -0,0 +1,424 @@
+// Copyright Epic Games, Inc. All Rights Reserved.
+
+#include "GroomCacheStreamingData.h"
+
+#include "Async/Async.h"
+#include "GroomCache.h"
+#include "GroomComponent.h"
+#include "GroomPluginSettings.h"
+#include "Misc/ScopeLock.h"
+
+DEFINE_LOG_CATEGORY(LogGroomCacheStreaming);
+
+FGroomCacheStreamingData::FGroomCacheStreamingData(UGroomCache* InGroomCache)
+: GroomCache(InGroomCache)
+{
+}
+
+FGroomCacheStreamingData::~FGroomCacheStreamingData()
+{
+	// Make sure that all read requests are deleted including those in-flight, so block until they are finished
+	BlockTillAllRequestsFinished();
+
+	// Delete the read requests in the background as it can be time-consuming
+	TSet<IBulkDataIORequest*> ReadyForDeletion(DelayedDeleteReadRequests);
+	Async(EAsyncExecution::ThreadPool, [ReadyForDeletion]()
+	{
+		for (IBulkDataIORequest* ReadRequest : ReadyForDeletion)
+		{
+			delete ReadRequest;
+		}
+	});
+
+	// Delete AnimData buffers that are resident in memory
+	for (TMap<int32, FResidentGroomCacheChunk>::TIterator Iter = Chunks.CreateIterator(); Iter; ++Iter)
+	{
+		delete Iter.Value().AnimDataPtr;
+	}
+}
+
+void FGroomCacheStreamingData::ResetNeededChunks()
+{
+	ChunksNeeded.Empty();
+}
+
+void FGroomCacheStreamingData::AddNeededChunk(uint32 ChunkIndex)
+{
+	ChunksNeeded.AddUnique(ChunkIndex);
+}
+
+FResidentGroomCacheChunk& FGroomCacheStreamingData::AddResidentChunk(int32 ChunkId, const FGroomCacheChunk &ChunkInfo)
+{
+	FScopeLock Lock(&CriticalSection);
+
+	FResidentGroomCacheChunk& Chunk = Chunks.FindOrAdd(ChunkId);
+
+	Chunk.DataSize = ChunkInfo.DataSize;
+	if (!Chunk.AnimDataPtr)
+	{
+		// The resident chunks are re-used so allocate the AnimData buffer the first time only or when RemoveResidentChunk has been called
+		Chunk.AnimDataPtr = new FGroomCacheAnimationData();
+	}
+	return Chunk;
+}
+
+void FGroomCacheStreamingData::RemoveResidentChunk(FResidentGroomCacheChunk& LoadedChunk)
+{
+	checkf(LoadedChunk.Refcount == 0, TEXT("Tried to remove a chunk wich was still mapped. Make sure there is an unmap for every map."));
+
+	if (LoadedChunk.IORequest == nullptr)
+	{
+		// There's no read request in-flight for this chunk so it's safe to delete the AnimData buffer
+		// Otherwise, let ProcessCompletedChunks delete it once the read is completed
+		delete LoadedChunk.AnimDataPtr;
+	}
+	LoadedChunk.AnimDataPtr = nullptr;
+	LoadedChunk.IORequest = nullptr;
+	LoadedChunk.DataSize = 0;
+	LoadedChunk.Refcount = 0;
+}
+
+void FGroomCacheStreamingData::PrefetchData(UGroomComponent *Component)
+{
+	TRACE_CPUPROFILER_EVENT_SCOPE(FGroomCacheStreamingData::PrefetchData);
+
+	float RequestStartTime = Component->GetAnimationTime();
+	float RequestEndTime = RequestStartTime + GetDefault<UGroomPluginSettings>()->GroomCacheLookAheadBuffer;
+
+	TArray<int32> NewChunksNeeded;
+	GroomCache->GetFrameIndicesForTimeRange(RequestStartTime, RequestEndTime, Component->IsLooping(), NewChunksNeeded);
+
+	for (int32 ChunkId : NewChunksNeeded)
+	{
+		ChunksNeeded.AddUnique(ChunkId);
+	}
+
+	// Synchronously load the first 2 frames
+	// This ensures we have something to display initially
+	for (int32 ChunkIndex = 0; ChunkIndex < 2; ++ChunkIndex)
+	{
+		// We just check here in case anything got loaded asynchronously last minute
+		// to avoid unnecessary loading it synchronously again
+		ProcessCompletedChunks();
+
+		int32 ChunkId = NewChunksNeeded[ChunkIndex];
+
+		// Already got it
+		if (ChunksAvailable.Contains(ChunkId))
+		{
+			continue;
+		}
+
+		// Still waiting for eviction, revive it
+		if (ChunksEvicted.Contains(ChunkId))
+		{
+			ChunksEvicted.Remove(ChunkId);
+			ChunksAvailable.Add(ChunkId);
+			continue;
+		}
+
+		const FGroomCacheChunk& Chunk = GroomCache->GetChunks()[ChunkId];
+		FResidentGroomCacheChunk &ResidentChunk = AddResidentChunk(ChunkId, Chunk);
+
+		// Already requested an async load but not complete yet 
+		if (ChunksRequested.Contains(ChunkId))
+		{
+			ChunksRequested.Remove(ChunkId);
+
+			// Let the in-flight read request complete with the initial AnimDataPtr to avoid threading issues
+			// The initial buffer will be orphaned but will be deleted once the completed read request is processed
+			// This new buffer is used immediately to load the data synchronously
+			ResidentChunk.AnimDataPtr = new FGroomCacheAnimationData();
+		}
+
+		// Load chunk from bulk data
+		GroomCache->GetGroomDataAtFrameIndex(ChunkId, *ResidentChunk.AnimDataPtr);
+		ChunksAvailable.Add(ChunkId);
+	}
+
+	// Schedule the rest of the needed frames to be loaded asynchronously
+	UpdateStreamingStatus();
+}
+
+void FGroomCacheStreamingData::UpdateStreamingStatus()
+{
+	TRACE_CPUPROFILER_EVENT_SCOPE(FGroomCacheStreamingData::UpdateStreamingStatus);
+
+	// Find any chunks that aren't available yet
+	for (int32 NeededIndex : ChunksNeeded)
+	{
+		if (!ChunksAvailable.Contains(NeededIndex))
+		{
+			// Revive it if it was still referenced
+			if (ChunksEvicted.Contains(NeededIndex))
+			{
+				ChunksEvicted.Remove(NeededIndex);
+				ChunksAvailable.Add(NeededIndex);
+				continue;
+			}
+
+			// If not requested yet, then request a load
+			if (!ChunksRequested.Contains(NeededIndex))
+			{
+				const FGroomCacheChunk& Chunk = GroomCache->GetChunks()[NeededIndex];
+
+				// This can happen in the editor if the asset hasn't been saved yet
+				if (Chunk.BulkData.IsBulkDataLoaded())
+				{
+					FResidentGroomCacheChunk &ResidentChunk = AddResidentChunk(NeededIndex, Chunk);
+					GroomCache->GetGroomDataAtFrameIndex(NeededIndex, *ResidentChunk.AnimDataPtr);
+					ChunksAvailable.Add(NeededIndex);
+					continue;
+				}
+
+				checkf(Chunk.BulkData.CanLoadFromDisk(), TEXT("Bulk data is not loaded and cannot be loaded from disk!"));
+				check(!Chunk.BulkData.IsStoredCompressedOnDisk()); // We do not support compressed Bulkdata for this system. Limitation of the streaming request/bulk data
+
+				FResidentGroomCacheChunk& ResidentChunk = AddResidentChunk(NeededIndex, Chunk);
+				FGroomCacheAnimationData* AnimData = ResidentChunk.AnimDataPtr;
+				int32 DataSize = ResidentChunk.DataSize;
+
+				checkf(ResidentChunk.IORequest == nullptr, TEXT("Groom cache read request for chunk %d is already in progress. Make sure its bookkeeping is up to date."), NeededIndex);
+				if (ResidentChunk.IORequest != nullptr)
+				{
+					continue;
+				}
+
+				FBulkDataIORequestCallBack AsyncFileCallBack = [this, NeededIndex, DataSize, AnimData](bool bWasCancelled, IBulkDataIORequest* ReadRequest)
+				{
+					TRACE_CPUPROFILER_EVENT_SCOPE(FGroomCacheStreamingData::CompletedCallBack);
+
+					// Now that the read request has completed, we can process the read buffer
+					FCompletedGroomCacheChunk CompletedChunk(NeededIndex, AnimData, ReadRequest);
+
+					if (uint8* ReadBuffer = ReadRequest->GetReadResults())
+					{
+						TRACE_CPUPROFILER_EVENT_SCOPE(FGroomCacheStreamingData::CompletedCallBack_BufferSerialize);
+
+						// The bulk data buffer is then serialized into GroomCacheAnimationData
+						TArrayView<uint8> TempView(ReadBuffer, DataSize);
+						FMemoryReaderView Ar(TempView, true);
+						AnimData->Serialize(Ar);
+
+						// We became the owner of ReadBuffer when GetReadResults was called so free it now
+						FMemory::Free(ReadBuffer);
+
+						CompletedChunk.bReadSuccessful = true;
+					}
+
+					CompletedChunks.Enqueue(MoveTemp(CompletedChunk));
+				};
+				
+				// Start the read request
+				EAsyncIOPriorityAndFlags AsyncIOPriority = AIOP_BelowNormal;
+				ResidentChunk.IORequest = Chunk.BulkData.CreateStreamingRequest(AsyncIOPriority, &AsyncFileCallBack, nullptr);
+				if (!ResidentChunk.IORequest)
+				{
+					UE_LOG(LogGroomCacheStreaming, Error, TEXT("Groom cache streaming read request couldn't be created."));
+					return;
+				}
+
+				// Add it to the list for bookkeeping
+				ChunksRequested.Add(NeededIndex);
+			}
+
+			// Read request is in-flight, just let it complete
+		}
+	}
+
+	// Update bookkeeping with any recently completed chunks
+	ProcessCompletedChunks();
+
+	// Find chunks that are currently loaded but aren't needed anymore and add them to the list of chunks to evict
+	TSet<int32> ChunksToRemove;
+	for (int32 ChunkId : ChunksAvailable)
+	{
+		if (!ChunksNeeded.Contains(ChunkId))
+		{
+			ChunksEvicted.Add(ChunkId);
+			ChunksToRemove.Add(ChunkId);
+		}
+	}
+	ChunksAvailable = ChunksAvailable.Difference(ChunksToRemove);
+
+	// Try to evict a bunch of chunks. Chunks that are still mapped (ie. in use) can't be evicted but others are free to go
+	TSet<int32> EvictedChunks;
+	{
+		FScopeLock Lock(&CriticalSection);
+
+		for (int32 ChunkId : ChunksEvicted)
+		{
+			FResidentGroomCacheChunk& ResidentChunk = Chunks[ChunkId];
+			if (ResidentChunk.Refcount == 0)
+			{
+				RemoveResidentChunk(ResidentChunk);
+				EvictedChunks.Add(ChunkId);
+			}
+		}
+	}
+	ChunksEvicted = ChunksEvicted.Difference(EvictedChunks);
+
+	// Delete the read requests that were delayed
+	TSet<IBulkDataIORequest*> ReadyForDeletion;
+	for (IBulkDataIORequest* ReadRequest : DelayedDeleteReadRequests)
+	{
+		if (ReadRequest->PollCompletion())
+		{
+			ReadyForDeletion.Add(ReadRequest);
+		}
+	}
+	DelayedDeleteReadRequests = DelayedDeleteReadRequests.Difference(ReadyForDeletion);
+	
+	// Delete the read requests in the background as it can be time-consuming
+	Async(EAsyncExecution::ThreadPool, [ReadyForDeletion]()
+	{
+		for (IBulkDataIORequest* ReadRequest : ReadyForDeletion)
+		{
+			delete ReadRequest;
+		}
+	});
+}
+
+bool FGroomCacheStreamingData::BlockTillAllRequestsFinished(float TimeLimit)
+{
+	FScopeLock Lock(&CriticalSection);
+
+	if (TimeLimit == 0.0f)
+	{
+		for (TMap<int32, FResidentGroomCacheChunk>::TIterator Iter = Chunks.CreateIterator(); Iter; ++Iter)
+		{
+			if (Iter.Value().IORequest)
+			{
+				Iter.Value().IORequest->WaitCompletion();
+				ProcessCompletedChunks();
+			}
+		}
+	}
+	else
+	{
+		double EndTime = FPlatformTime::Seconds() + TimeLimit;
+		for (TMap<int32, FResidentGroomCacheChunk>::TIterator Iter = Chunks.CreateIterator(); Iter; ++Iter)
+		{
+			if (Iter.Value().IORequest)
+			{
+				float ThisTimeLimit = EndTime - FPlatformTime::Seconds();
+				if (ThisTimeLimit < .001f || // one ms is the granularity of the platform event system
+					!Iter.Value().IORequest->WaitCompletion(ThisTimeLimit))
+				{
+					return false;
+				}
+				ProcessCompletedChunks();
+			}
+		}
+	}
+
+	return true;
+}
+
+void FGroomCacheStreamingData::ProcessCompletedChunks()
+{
+	TRACE_CPUPROFILER_EVENT_SCOPE(FGroomCacheStreamingData::ProcessCompletedChunks);
+	FScopeLock Lock(&CriticalSection);
+
+	FCompletedGroomCacheChunk CompletedChunk;
+	while (CompletedChunks.Dequeue(CompletedChunk))
+	{
+		FResidentGroomCacheChunk* Chunk = Chunks.Find(CompletedChunk.LoadedChunkIndex);
+		if (!Chunk)
+		{
+			return;
+		}
+
+		// Check to see if we successfully managed to load anything
+		// If the chunk's IORequest is null, it means RemoveResidentChunk has been called on it during the read request
+		// in which case, we don't care about this chunk anymore 
+		if (Chunk->IORequest != nullptr)
+		{
+			if (CompletedChunk.bReadSuccessful)
+			{
+				ChunksAvailable.Add(CompletedChunk.LoadedChunkIndex);
+			}
+			else
+			{
+				UE_LOG(LogGroomCacheStreaming, Log, TEXT("Async loading request failed!"));
+			}
+		}
+		ChunksRequested.Remove(CompletedChunk.LoadedChunkIndex);
+
+		// The AnimData buffer can change while the read request was in-flight because the associated chunk was prefetched synchronously instead
+		// or it's been removed from resident chunks, so just delete the orphaned buffer
+		if (CompletedChunk.AnimDataPtr != Chunk->AnimDataPtr)
+		{
+			delete CompletedChunk.AnimDataPtr;
+		}
+
+		Chunk->IORequest = nullptr;
+
+		// Process the read request deletion later as in can be quite time-consuming
+		DelayedDeleteReadRequests.Add(CompletedChunk.ReadRequest);
+	}
+}
+
+const FGroomCacheAnimationData* FGroomCacheStreamingData::MapAnimationData(uint32 ChunkIndex)
+{
+	TRACE_CPUPROFILER_EVENT_SCOPE(FGroomCacheStreamingData::MapAnimationData);
+
+	// Quickly check before mapping if maybe something new arrived we haven't done bookkeeping for yet
+	ProcessCompletedChunks();
+
+	if (!ChunksAvailable.Contains(ChunkIndex))
+	{
+		if (!ChunksRequested.Contains(ChunkIndex))
+		{
+			if (ChunksEvicted.Contains(ChunkIndex))
+			{
+				UE_LOG(LogGroomCacheStreaming, Log, TEXT("Tried to map an evicted chunk: %i."), ChunkIndex);
+			}
+			else
+			{
+				UE_LOG(LogGroomCacheStreaming, Log, TEXT("Tried to map an unavailabe non-requested chunk: %i."), ChunkIndex);
+			}
+		}
+		else
+		{
+			UE_LOG(LogGroomCacheStreaming, Log, TEXT("Tried to map a chunk (%i) that is still being streamed in."), ChunkIndex);
+		}
+		return nullptr;
+	}
+	else
+	{
+		FScopeLock Lock(&CriticalSection);
+
+		FResidentGroomCacheChunk *ResidentChunk = Chunks.Find(ChunkIndex);
+		if (ResidentChunk)
+		{
+			checkf(ResidentChunk->AnimDataPtr, TEXT("Tried to map a chunk without data. Make sure its bookkeeping is up to date."));
+			++ResidentChunk->Refcount;
+			return ResidentChunk->AnimDataPtr;
+		}
+		UE_LOG(LogGroomCacheStreaming, Log, TEXT("Tried to map a chunk (%i) that is not resident in-memory."), ChunkIndex);
+		return nullptr;
+	}
+}
+
+void FGroomCacheStreamingData::UnmapAnimationData(uint32 ChunkIndex)
+{
+	TRACE_CPUPROFILER_EVENT_SCOPE(FGroomCacheStreamingData::UnmapAnimationData);
+
+	// The unmap could happen from the render thread when it releases the GroomCache buffers
+	// so access to StreamingGroomCaches has to be protected
+	FScopeLock Lock(&CriticalSection);
+
+	FResidentGroomCacheChunk* ResidentChunk = Chunks.Find(ChunkIndex);
+	if (ResidentChunk != nullptr )
+	{
+		checkf(ResidentChunk->Refcount > 0, TEXT("Map/Unmap out of balance. Make sure you unmap once fore every map."));
+		checkf(ChunksAvailable.Contains(ChunkIndex) || ChunksEvicted.Contains(ChunkIndex), TEXT("Tried to unmap a chunk in an invalid state."));
+		--ResidentChunk->Refcount;
+	}
+}
+
+bool FGroomCacheStreamingData::IsStreamingInProgress()
+{
+	return ChunksRequested.Num() > 0;
+}
diff --git a/Engine/Plugins/Runtime/HairStrands/Source/HairStrandsCore/Private/GroomCacheStreamingData.h b/Engine/Plugins/Runtime/HairStrands/Source/HairStrandsCore/Private/GroomCacheStreamingData.h
new file mode 100644
--- /dev/null
+++ b/Engine/Plugins/Runtime/HairStrands/Source/HairStrandsCore/Private/GroomCacheStreamingData.h
@@ -0,0 +1,100 @@
+// Copyright Epic Games, Inc. All Rights Reserved.
+
+#pragma once
+
+#include "CoreMinimal.h"
+#include "Containers/Queue.h"
+
+DECLARE_LOG_CATEGORY_EXTERN(LogGroomCacheStreaming, Verbose, All);
+
+struct FGroomCacheAnimationData;
+class IBulkDataIORequest;
+class UGroomCache;
+class UGroomComponent;
+
+/** Info about a chunk, loaded or not */
+struct FResidentGroomCacheChunk
+{
+	FGroomCacheAnimationData* AnimDataPtr = nullptr;
+	IBulkDataIORequest* IORequest = nullptr; // null when resident, non-null when in flight
+	int32 Refcount = 0;
+	int32 DataSize = 0;
+};
+
+/** Info about read request status of a chunk */
+struct FCompletedGroomCacheChunk
+{
+	FGroomCacheAnimationData* AnimDataPtr = nullptr;
+	IBulkDataIORequest* ReadRequest;
+	int32 LoadedChunkIndex;
+	bool bReadSuccessful;
+
+	FCompletedGroomCacheChunk() : AnimDataPtr(nullptr), ReadRequest(nullptr), LoadedChunkIndex(0), bReadSuccessful(false) {}
+	FCompletedGroomCacheChunk(int32 SetLoadedChunkIndex, FGroomCacheAnimationData* AnimBuffer, IBulkDataIORequest* SetReadRequest)
+	: AnimDataPtr(AnimBuffer), ReadRequest(SetReadRequest), LoadedChunkIndex(SetLoadedChunkIndex), bReadSuccessful(false) {}
+};
+
+/** Class that manages the streaming and cached data for a GroomCache */
+class FGroomCacheStreamingData
+{
+public:
+	FGroomCacheStreamingData(UGroomCache* GroomCache);
+	~FGroomCacheStreamingData();
+
+	/** Updates the internal state of the streaming data. */
+	void UpdateStreamingStatus();
+
+	/** Returns a pointer to the cache data for ChunkIndex if it's loaded; nullptr otherwise. */
+	const FGroomCacheAnimationData* MapAnimationData(uint32 ChunkIndex);
+
+	/** Informs the streaming data that the ChunkIndex data is not used anymore.  */
+	void UnmapAnimationData(uint32 ChunkIndex);
+
+	/** Returns true if there are read requests in flight; false otherwise. */
+	bool IsStreamingInProgress();
+
+	/** Blocks until all read requests are completed. Returns true if they were all completed before TimeLimit (0 means no limit); false otherwise. */
+	bool BlockTillAllRequestsFinished(float TimeLimit = 0.0f);
+
+	/** Synchronously loads the initial frames needed to display the given component */
+	void PrefetchData(UGroomComponent* Component);
+
+	/** Resets the list of chunks needed by users of this GroomCache */
+	void ResetNeededChunks();
+
+	/** Adds ChunkIndex to the list of chunks needed to be loaded */
+	void AddNeededChunk(uint32 ChunkIndex);
+
+private:
+
+	FResidentGroomCacheChunk& AddResidentChunk(int32 ChunkId, const struct FGroomCacheChunk& ChunkInfo);
+	void RemoveResidentChunk(FResidentGroomCacheChunk& LoadedChunk);
+	void ProcessCompletedChunks();
+
+	/** The GroomCache being streamed */
+	UGroomCache* GroomCache;
+
+	/** List of chunks that need to be loaded */
+	TArray<int32> ChunksNeeded;
+
+	/** List of chunks currently resident in memory */
+	TSet<int32> ChunksAvailable;
+
+	/** List of chunks being streamed, ie. their read requests are in-flight */
+	TSet<int32> ChunksRequested;
+
+	/** List of chunks to be evicted, not needed anymore so they can be unloaded */
+	TSet<int32> ChunksEvicted;
+
+	/** ChunkIndex to chunk info */
+	TMap<int32, FResidentGroomCacheChunk> Chunks;
+
+	/** List of read requests that are completed that need to be processed for bookkeeping (to be marked as available) */
+	TQueue<FCompletedGroomCacheChunk, EQueueMode::Mpsc> CompletedChunks;
+
+	/** List of read requests to be deleted after completion */
+	TSet<IBulkDataIORequest*> DelayedDeleteReadRequests;
+
+	FCriticalSection CriticalSection;
+};
+
diff --git a/Engine/Plugins/Runtime/HairStrands/Source/HairStrandsCore/Private/GroomCacheStreamingManager.cpp b/Engine/Plugins/Runtime/HairStrands/Source/HairStrandsCore/Private/GroomCacheStreamingManager.cpp
new file mode 100644
--- /dev/null
+++ b/Engine/Plugins/Runtime/HairStrands/Source/HairStrandsCore/Private/GroomCacheStreamingManager.cpp
@@ -0,0 +1,318 @@
+// Copyright Epic Games, Inc. All Rights Reserved.
+
+#include "GroomCacheStreamingManager.h"
+
+#include "Async/Async.h"
+#include "GroomCache.h"
+#include "GroomCacheStreamingData.h"
+#include "GroomComponent.h"
+#include "GroomPluginSettings.h"
+#include "Misc/ScopeLock.h"
+
+struct FGroomCacheStreamingManager : public IGroomCacheStreamingManager
+{
+	//~ Begin IStreamingManager interface
+	virtual void UpdateResourceStreaming(float DeltaTime, bool bProcessEverything = false) override;
+	virtual int32 BlockTillAllRequestsFinished(float TimeLimit = 0.0f, bool bLogResults = false) override;
+	virtual void CancelForcedResources() override {}
+	virtual void NotifyLevelChange() override {}
+	virtual void SetDisregardWorldResourcesForFrames(int32 NumFrames) override {}
+	virtual void AddLevel(class ULevel* Level) override {}
+	virtual void RemoveLevel(class ULevel* Level) override {}
+	virtual void NotifyLevelOffset(class ULevel* Level, const FVector& Offset) override {}
+	//~ End IStreamingManager interface
+
+	//~ Begin IGeometryCacheStreamingManager interface
+	virtual void RegisterComponent(UGroomComponent* GroomComponent) override;
+	virtual void UnregisterComponent(UGroomComponent* GroomComponent) override;
+	virtual void PrefetchData(UGroomComponent* GroomComponent) override;
+	virtual const FGroomCacheAnimationData* MapAnimationData(const UGroomCache* GroomCache, uint32 ChunkIndex) override;
+	virtual void UnmapAnimationData(const UGroomCache* GroomCache, uint32 ChunkIndex) override;
+	virtual void Shutdown() override;
+	//~ End IGeometryCacheStreamingManager interface
+
+private:
+
+	void RegisterGroomCache(UGroomCache* GroomCache, UGroomComponent* GroomComponent);
+	void UnregisterGroomCache(UGroomCache* GroomCache, UGroomComponent* GroomComponent);
+	void PrefetchDataInternal(UGroomComponent* GroomComponent);
+	void DeleteStreamingData();
+
+	// StreamingData for registered GroomCaches
+	TMap<UGroomCache*, FGroomCacheStreamingData*> StreamingGroomCaches;
+
+	// Set of GroomComponent using the registered GroomCaches
+	TMap<UGroomCache*, TSet<UGroomComponent*>> GroomCacheUsers;
+
+	// Set of StreamingData to delete when ready
+	TSet<FGroomCacheStreamingData*> StreamingGroomCachesToDelete;
+
+	// GroomComponent that are registered to stream GroomCache data
+	TArray<UGroomComponent*> StreamingComponents;
+
+	FCriticalSection CriticalSection;
+};
+
+IGroomCacheStreamingManager &IGroomCacheStreamingManager::Get()
+{
+	static FGroomCacheStreamingManager Instance;
+	return Instance;
+}
+
+void IGroomCacheStreamingManager::Register()
+{
+	IStreamingManager::Get().AddStreamingManager(&IGroomCacheStreamingManager::Get());
+}
+
+void IGroomCacheStreamingManager::Unregister()
+{
+	IStreamingManager::Get().RemoveStreamingManager(&IGroomCacheStreamingManager::Get());
+	IGroomCacheStreamingManager::Get().Shutdown();
+}
+
+void FGroomCacheStreamingManager::Shutdown()
+{
+	// Delete the StreamingData that were already queued for deletion
+	while (StreamingGroomCachesToDelete.Num() > 0)
+	{
+		DeleteStreamingData();
+	}
+
+	FScopeLock Lock(&CriticalSection);
+
+	// At this point, make sure that the StreamingData have finished their read requests before deleting them
+	for (TMap<UGroomCache*, FGroomCacheStreamingData*>::TIterator Iter = StreamingGroomCaches.CreateIterator(); Iter; ++Iter)
+	{
+		Iter.Value()->BlockTillAllRequestsFinished();
+		delete Iter.Value();
+	}
+
+	StreamingGroomCaches.Reset();
+	GroomCacheUsers.Reset();
+	StreamingComponents.Reset();
+}
+
+void FGroomCacheStreamingManager::UpdateResourceStreaming(float DeltaTime, bool bProcessEverything)
+{
+	TRACE_CPUPROFILER_EVENT_SCOPE(FGroomCacheStreamingManager::UpdateResourceStreaming);
+
+	const UGroomPluginSettings* Settings = GetDefault<UGroomPluginSettings>();
+	{
+		FScopeLock Lock(&CriticalSection);
+
+		// Clear needed chunks, they will be recomputed in the next step
+		for (TMap<UGroomCache*, FGroomCacheStreamingData*>::TIterator Iter = StreamingGroomCaches.CreateIterator(); Iter; ++Iter)
+		{
+			Iter.Value()->ResetNeededChunks();
+		}
+
+		// Gather all the chunks that need to be streamed from all playing instances based on their current animation time
+		for (UGroomComponent* Component : StreamingComponents)
+		{
+			if (UGroomCache* GroomCache = Component->GetGroomCache())
+			{
+				FGroomCacheStreamingData** StreamingData = StreamingGroomCaches.Find(GroomCache);
+				if (StreamingData)
+				{
+					float RequestStartTime = Component->GetAnimationTime();
+					float RequestEndTime = RequestStartTime + Settings->GroomCacheLookAheadBuffer;
+
+					TArray<int32> ChunksNeeded;
+					GroomCache->GetFrameIndicesForTimeRange(RequestStartTime, RequestEndTime, Component->IsLooping(), ChunksNeeded);
+
+					for (int32 ChunkId : ChunksNeeded)
+					{
+						(*StreamingData)->AddNeededChunk(ChunkId);
+					}
+				}
+			}
+		}
+
+		// Update the internal state of the StreamingData
+		for (TMap<UGroomCache*, FGroomCacheStreamingData*>::TIterator Iter = StreamingGroomCaches.CreateIterator(); Iter; ++Iter)
+		{
+			Iter.Value()->UpdateStreamingStatus();
+		}
+	}
+
+	DeleteStreamingData();
+}
+
+void FGroomCacheStreamingManager::DeleteStreamingData()
+{
+	TSet<FGroomCacheStreamingData*> ReadyForDeletion;
+	for (FGroomCacheStreamingData* StreamingData : StreamingGroomCachesToDelete)
+	{
+		StreamingData->UpdateStreamingStatus();
+		if (!StreamingData->IsStreamingInProgress())
+		{
+			ReadyForDeletion.Add(StreamingData);
+		}
+	}
+
+	StreamingGroomCachesToDelete = StreamingGroomCachesToDelete.Difference(ReadyForDeletion);
+
+	for (FGroomCacheStreamingData* StreamingData : ReadyForDeletion)
+	{
+		delete StreamingData;
+	}
+}
+
+int32 FGroomCacheStreamingManager::BlockTillAllRequestsFinished(float TimeLimit, bool)
+{
+	FScopeLock Lock(&CriticalSection);
+
+	int32 Result = 0;
+
+	if (TimeLimit == 0.0f)
+	{
+		for (TMap<UGroomCache*, FGroomCacheStreamingData*>::TIterator Iter = StreamingGroomCaches.CreateIterator(); Iter; ++Iter)
+		{
+			Iter.Value()->BlockTillAllRequestsFinished();
+		}
+	}
+	else
+	{
+		double EndTime = FPlatformTime::Seconds() + TimeLimit;
+		for (TMap<UGroomCache*, FGroomCacheStreamingData*>::TIterator Iter = StreamingGroomCaches.CreateIterator(); Iter; ++Iter)
+		{
+			float ThisTimeLimit = EndTime - FPlatformTime::Seconds();
+			if (ThisTimeLimit < .001f || // one ms is the granularity of the platform event system
+				!Iter.Value()->BlockTillAllRequestsFinished(ThisTimeLimit))
+			{
+				Result = 1; // we don't report the actual number, just 1 for any number of outstanding requests
+				break;
+			}
+		}
+	}
+
+	return Result;
+}
+
+void FGroomCacheStreamingManager::RegisterGroomCache(UGroomCache* GroomCache, UGroomComponent* GroomComponent)
+{
+	FScopeLock Lock(&CriticalSection);
+
+	if (!StreamingGroomCaches.Contains(GroomCache))
+	{
+		StreamingGroomCaches.Add(GroomCache, new FGroomCacheStreamingData(GroomCache));
+		GroomCacheUsers.Add(GroomCache, {});
+	}
+
+	GroomCacheUsers.FindChecked(GroomCache).Add(GroomComponent);
+}
+
+void FGroomCacheStreamingManager::UnregisterGroomCache(UGroomCache* GroomCache, UGroomComponent* GroomComponent)
+{
+	FScopeLock Lock(&CriticalSection);
+
+	TSet<UGroomComponent*>& Users = GroomCacheUsers.FindChecked(GroomCache);
+	Users.Remove(GroomComponent);
+
+	if (Users.Num() == 0)
+	{
+		FGroomCacheStreamingData** StreamingData = StreamingGroomCaches.Find(GroomCache);
+		if (StreamingData)
+		{
+			// The needed chunks are reset to prevent starting new read requests
+			(*StreamingData)->ResetNeededChunks();
+
+			// The StreamingData is deleted later once all the in-flight read requests are processed
+			StreamingGroomCachesToDelete.Add(*StreamingData);
+
+			StreamingGroomCaches.Remove(GroomCache);
+		}
+
+		GroomCacheUsers.Remove(GroomCache);
+	}
+}
+
+void FGroomCacheStreamingManager::RegisterComponent(UGroomComponent* GroomComponent)
+{
+	if (GroomComponent && GroomComponent->GetGroomCache())
+	{
+		StreamingComponents.AddUnique(GroomComponent);
+
+		RegisterGroomCache(GroomComponent->GetGroomCache(), GroomComponent);
+
+		PrefetchData(GroomComponent);
+	}
+}
+
+void FGroomCacheStreamingManager::UnregisterComponent(UGroomComponent* GroomComponent)
+{
+	if (GroomComponent)
+	{
+		StreamingComponents.Remove(GroomComponent);
+
+		if (UGroomCache* GroomCache = GroomComponent->GetGroomCache())
+		{
+			UnregisterGroomCache(GroomCache, GroomComponent);
+		}
+	}
+}
+
+void FGroomCacheStreamingManager::PrefetchDataInternal(UGroomComponent* GroomComponent)
+{
+	if (UGroomCache* GroomCache = GroomComponent->GetGroomCache())
+	{
+		FScopeLock Lock(&CriticalSection);
+
+		FGroomCacheStreamingData** StreamingData = StreamingGroomCaches.Find(GroomCache);
+		if (StreamingData)
+		{
+			(*StreamingData)->PrefetchData(GroomComponent);
+		}
+	}
+}
+
+void FGroomCacheStreamingManager::PrefetchData(UGroomComponent* GroomComponent)
+{
+	if (IsInGameThread())
+	{
+		PrefetchDataInternal(GroomComponent);
+	}
+	else
+	{
+		// The prefetch doesn't need to be executed right now, so schedule it for the game thread
+		FWeakObjectPtr WeakComponent(GroomComponent);
+		AsyncTask(ENamedThreads::GameThread, [this, WeakComponent]()
+		{
+			UGroomComponent* Component = Cast<UGroomComponent>(WeakComponent.Get());
+			if (Component)
+			{
+				PrefetchDataInternal(Component);
+			}
+		});
+	}
+}
+
+const FGroomCacheAnimationData* FGroomCacheStreamingManager::MapAnimationData(const UGroomCache* GroomCache, uint32 ChunkIndex)
+{
+	TRACE_CPUPROFILER_EVENT_SCOPE(FGroomCacheStreamingManager::MapAnimationData);
+
+	FScopeLock Lock(&CriticalSection);
+
+	FGroomCacheStreamingData** StreamingData = StreamingGroomCaches.Find(GroomCache);
+	if (StreamingData)
+	{
+		return (*StreamingData)->MapAnimationData(ChunkIndex);
+	}
+
+	return nullptr;
+}
+
+void FGroomCacheStreamingManager::UnmapAnimationData(const UGroomCache* GroomCache, uint32 ChunkIndex)
+{
+	TRACE_CPUPROFILER_EVENT_SCOPE(FGroomCacheStreamingManager::UnmapAnimationData);
+
+	// The unmap could happen from the render thread when it releases the GroomCache buffers
+	// so access to StreamingGroomCaches has to be protected
+	FScopeLock Lock(&CriticalSection);
+
+	FGroomCacheStreamingData** StreamingData = StreamingGroomCaches.Find(GroomCache);
+	if (StreamingData)
+	{
+		(*StreamingData)->UnmapAnimationData(ChunkIndex);
+	}
+}
diff --git a/Engine/Plugins/Runtime/HairStrands/Source/HairStrandsCore/Private/GroomComponent.cpp b/Engine/Plugins/Runtime/HairStrands/Source/HairStrandsCore/Private/GroomComponent.cpp
--- a/Engine/Plugins/Runtime/HairStrands/Source/HairStrandsCore/Private/GroomComponent.cpp
+++ b/Engine/Plugins/Runtime/HairStrands/Source/HairStrandsCore/Private/GroomComponent.cpp
@@ -29,6 +29,8 @@
 #include "GroomManager.h"
 #include "GroomInstance.h"
 #include "GroomCache.h"
+#include "GroomCacheStreamingManager.h"
+#include "GroomPluginSettings.h"
 #include "Async/ParallelFor.h"
 
 static int32 GHairEnableAdaptiveSubsteps = 0;  
@@ -38,6 +40,9 @@ bool IsHairAdaptiveSubstepsEnabled() { return (GHairEnableAdaptiveSubsteps == 1)
 static int32 GHairBindingValidationEnable = 0;
 static FAutoConsoleVariableRef CVarHairBindingValidationEnable(TEXT("r.HairStrands.BindingValidation"), GHairBindingValidationEnable, TEXT("Enable groom binding validation, which report error/warnings with details about the cause."));
 
+static bool GUseGroomCacheStreaming = true;
+static FAutoConsoleVariableRef CVarGroomCacheStreamingEnable(TEXT("GroomCache.EnableStreaming"), GUseGroomCacheStreaming, TEXT("Enable groom cache streaming and prebuffering. Do not switch while groom caches are in use."));
+
 #define LOCTEXT_NAMESPACE "GroomComponent"
 
 ////////////////////////////////////////////////////////////////////////////////////////////////////////////////////////
@@ -917,65 +922,59 @@ private:
 class FGroomCacheBuffers : public IGroomCacheBuffers
 {
 public:
-	virtual FGroomCacheAnimationData& GetCurrentFrameBuffer() override
+	FGroomCacheBuffers(UGroomCache* InGroomCache)
+	: GroomCache(InGroomCache)
 	{
-		return CurrentFrame;
 	}
 
-	virtual FGroomCacheAnimationData& GetNextFrameBuffer() override
+	virtual ~FGroomCacheBuffers()
 	{
-		return NextFrame;
+		Reset();
 	}
 
-	virtual FGroomCacheAnimationData& GetInterpolatedFrameBuffer() override
+	virtual void Reset()
 	{
-		return InterpolatedFrame;
 	}
 
-	virtual int32 GetCurrentFrameIndex() const override
+	virtual const FGroomCacheAnimationData& GetCurrentFrameBuffer() override
 	{
-		return CurrentFrameIndex;
+		return CurrentFrame;
 	}
 
-	virtual int32 GetNextFrameIndex() const override
+	virtual const FGroomCacheAnimationData& GetNextFrameBuffer() override
 	{
-		return NextFrameIndex;
+		return NextFrame;
 	}
 
-	virtual float GetInterpolationFactor() const override
+	virtual const FGroomCacheAnimationData& GetInterpolatedFrameBuffer() override
 	{
-		return InterpolationFactor;
+		return InterpolatedFrame;
 	}
 
-	virtual void SetCurrentFrameIndex(int32 FrameIndex)
+	virtual int32 GetCurrentFrameIndex() const override
 	{
-		CurrentFrameIndex = FrameIndex;
+		return CurrentFrameIndex;
 	}
 
-	virtual void SetNextFrameIndex(int32 FrameIndex)
+	virtual int32 GetNextFrameIndex() const override
 	{
-		NextFrameIndex = FrameIndex;
+		return NextFrameIndex;
 	}
 
-	virtual void SetInterpolationFactor(float Factor)
+	virtual float GetInterpolationFactor() const override
 	{
-		InterpolationFactor = Factor;
+		return InterpolationFactor;
 	}
 
-	void UpdateBuffersAtTime(UGroomCache* GroomCache, float Time, bool bIsLooping)
+	virtual void UpdateBuffersAtTime(float Time, bool bIsLooping)
 	{
 		TRACE_CPUPROFILER_EVENT_SCOPE(FGroomCacheBuffers::UpdateBuffersAtTime);
 
-		if (!GroomCache)
-		{
-			return;
-		}
-
 		// Find the frame indices and interpolation factor to interpolate between
 		int32 FrameIndexA = 0;
 		int32 FrameIndexB = 0;
 		float OutInterpolationFactor = 0.0f;
-		GroomCache->FindSampleIndexesFromTime(Time, bIsLooping, false, FrameIndexA, FrameIndexB, OutInterpolationFactor);
+		GroomCache->GetFrameIndicesAtTime(Time, bIsLooping, false, FrameIndexA, FrameIndexB, OutInterpolationFactor);
 
 		// Update and cache the frame data as needed
 		bool bComputeInterpolation = false;
@@ -1008,37 +1007,65 @@ public:
 		}
 
 		// Do interpolation of vertex positions if needed
-		if (bComputeInterpolation || !FMath::IsNearlyEqual(OutInterpolationFactor, InterpolationFactor, KINDA_SMALL_NUMBER))
+		if (bComputeInterpolation)
+		{
+			Interpolate(CurrentFrame, NextFrame, OutInterpolationFactor);
+		}
+	}
+
+	void Interpolate(const FGroomCacheAnimationData& FrameA, const FGroomCacheAnimationData& FrameB, float InInterpolationFactor)
+	{
+		TRACE_CPUPROFILER_EVENT_SCOPE(FGroomCacheBuffers::InterpolateCPU);
+
+		if (FMath::IsNearlyEqual(InInterpolationFactor, InterpolationFactor, KINDA_SMALL_NUMBER))
+		{
+			return;
+		}
+
+		InterpolationFactor = InInterpolationFactor;
+
+		FScopeLock Lock(GetCriticalSection());
+
+		const int32 NumGroups = FrameA.GroupsData.Num();
+		for (int32 GroupIndex = 0; GroupIndex < NumGroups; ++GroupIndex)
 		{
-			TRACE_CPUPROFILER_EVENT_SCOPE(FGroomCacheBuffers::UpdateBuffersAtTime_InterpolateCPU);
+			TRACE_CPUPROFILER_EVENT_SCOPE(FGroomCacheBuffers::InterpolateCPU_Group);
 
-			InterpolationFactor = OutInterpolationFactor;
+			const FGroomCacheGroupData& CurrentGroupData = FrameA.GroupsData[GroupIndex];
+			const FGroomCacheGroupData& NextGroupData = FrameB.GroupsData[GroupIndex];
+			FGroomCacheGroupData& InterpolatedGroupData = InterpolatedFrame.GroupsData[GroupIndex];
+			const int32 NumVertices = CurrentGroupData.VertexData.PointsPosition.Num();
 
-			FScopeLock Lock(GetCriticalSection());
+			// Parallel batched interpolation
+			const int32 BatchSize = 1024;
+			const int32 BatchCount = (NumVertices + BatchSize - 1) / BatchSize;
 
-			const int32 NumGroups = CurrentFrame.GroupsData.Num();
-			for (int32 GroupIndex = 0; GroupIndex < NumGroups; ++GroupIndex)
+			ParallelFor(BatchCount, [&](int32 BatchIndex)
 			{
-				const FGroomCacheGroupData& CurrentGroupData = CurrentFrame.GroupsData[GroupIndex];
-				const FGroomCacheGroupData& NextGroupData = NextFrame.GroupsData[GroupIndex];
-				FGroomCacheGroupData& InterpolatedGroupData = InterpolatedFrame.GroupsData[GroupIndex];
-				const int32 NumVertices = CurrentGroupData.VertexData.PointsPosition.Num();
+				const int32 Start = BatchIndex * BatchSize;
+				const int32 End = FMath::Min(Start + BatchSize, NumVertices); // one-past end index
 
-				ParallelFor(NumVertices,
-				[&CurrentGroupData, &NextGroupData, &InterpolatedGroupData, this](int32 VertexIndex)
+				for (int32 VertexIndex = Start; VertexIndex < End; ++VertexIndex)
 				{
 					const FVector& CurrentPosition = CurrentGroupData.VertexData.PointsPosition[VertexIndex];
 					const FVector& NextPosition = NextGroupData.VertexData.PointsPosition[VertexIndex];
+
 					InterpolatedGroupData.VertexData.PointsPosition[VertexIndex] = FMath::Lerp(CurrentPosition, NextPosition, InterpolationFactor);
 				}
-				);
-			}
+			});
 		}
 	}
 
-private:
+protected:
+	static FGroomCacheAnimationData EmptyFrame;
+
+	UGroomCache* GroomCache;
+
+	/** Used with synchronous loading */
 	FGroomCacheAnimationData CurrentFrame;
 	FGroomCacheAnimationData NextFrame;
+
+	/** Used for CPU interpolation */
 	FGroomCacheAnimationData InterpolatedFrame;
 
 	int32 CurrentFrameIndex = -1;
@@ -1046,6 +1073,128 @@ private:
 	float InterpolationFactor = 0.0f;
 };
 
+FGroomCacheAnimationData FGroomCacheBuffers::EmptyFrame;
+
+class FGroomCacheStreamedBuffers : public FGroomCacheBuffers
+{
+public:
+	FGroomCacheStreamedBuffers(UGroomCache* InGroomCache)
+		: FGroomCacheBuffers(InGroomCache)
+		, CurrentFramePtr(nullptr)
+		, NextFramePtr(nullptr)
+	{
+	}
+
+	virtual void Reset() override
+	{
+		// Unmap the frames that are currently mapped
+		if (CurrentFrameIndex != -1)
+		{
+			IGroomCacheStreamingManager::Get().UnmapAnimationData(GroomCache, CurrentFrameIndex);
+			CurrentFrameIndex = -1;
+			CurrentFramePtr = nullptr;
+		}
+
+		if (NextFrameIndex != -1)
+		{
+			IGroomCacheStreamingManager::Get().UnmapAnimationData(GroomCache, NextFrameIndex);
+			NextFrameIndex = -1;
+			NextFramePtr = nullptr;
+		}
+	}
+
+	virtual const FGroomCacheAnimationData& GetCurrentFrameBuffer() override
+	{
+		if (CurrentFramePtr)
+		{
+			return *CurrentFramePtr;
+		}
+		return EmptyFrame;
+	}
+
+	virtual const FGroomCacheAnimationData& GetNextFrameBuffer() override
+	{
+		if (NextFramePtr)
+		{
+			return *NextFramePtr;
+		}
+		return EmptyFrame;
+	}
+
+	virtual void UpdateBuffersAtTime(float Time, bool bIsLooping) override
+	{
+		TRACE_CPUPROFILER_EVENT_SCOPE(FGroomCacheStreamedBuffers::UpdateBuffersAtTime);
+
+		// Find the frame indices and interpolation factor to interpolate between
+		int32 FrameIndexA = 0;
+		int32 FrameIndexB = 0;
+		float OutInterpolationFactor = 0.0f;
+		GroomCache->GetFrameIndicesAtTime(Time, bIsLooping, false, FrameIndexA, FrameIndexB, OutInterpolationFactor);
+
+		// Update and cache the frame data as needed
+		bool bComputeInterpolation = false;
+		if (FrameIndexA != CurrentFrameIndex)
+		{
+			bComputeInterpolation = true;
+			if (FrameIndexA == NextFrameIndex)
+			{
+				// At this point, the NextFrame is already mapped so we know the pointer is valid
+				// It is mapped to increment its ref count
+				const FGroomCacheAnimationData* DataPtr = IGroomCacheStreamingManager::Get().MapAnimationData(GroomCache, NextFrameIndex);
+				IGroomCacheStreamingManager::Get().UnmapAnimationData(GroomCache, CurrentFrameIndex);
+
+				CurrentFramePtr = NextFramePtr;
+				CurrentFrameIndex = NextFrameIndex;
+			}
+			else
+			{
+				const FGroomCacheAnimationData* DataPtr = IGroomCacheStreamingManager::Get().MapAnimationData(GroomCache, FrameIndexA);
+				if (DataPtr)
+				{
+					IGroomCacheStreamingManager::Get().UnmapAnimationData(GroomCache, CurrentFrameIndex);
+					CurrentFramePtr = DataPtr;
+					CurrentFrameIndex = FrameIndexA;
+				}
+			}
+		}
+
+		if (FrameIndexB != NextFrameIndex)
+		{
+			bComputeInterpolation = true;
+
+			const FGroomCacheAnimationData* DataPtr = IGroomCacheStreamingManager::Get().MapAnimationData(GroomCache, FrameIndexB);
+			if (DataPtr)
+			{
+				IGroomCacheStreamingManager::Get().UnmapAnimationData(GroomCache, NextFrameIndex);
+				NextFramePtr = DataPtr;
+				NextFrameIndex = FrameIndexB;
+			}
+		}
+
+		if (CurrentFramePtr == nullptr || NextFramePtr == nullptr)
+		{
+			return;
+		}
+
+		// Make sure the initial interpolated frame is populated with valid data
+		if (InterpolatedFrame.GroupsData.Num() != CurrentFramePtr->GroupsData.Num())
+		{
+			InterpolatedFrame = *CurrentFramePtr;
+		}
+
+		// Do interpolation of vertex positions if needed
+		if (bComputeInterpolation)
+		{
+			Interpolate(*CurrentFramePtr, *NextFramePtr, OutInterpolationFactor);
+		}
+	}
+
+private:
+	/** Used with GroomCache streaming. Point to cached data in the manager */
+	const FGroomCacheAnimationData* CurrentFramePtr;
+	const FGroomCacheAnimationData* NextFramePtr;
+};
+
 /////////////////////////////////////////////////////////////////////////////////////////
 /////////////////////////////////////////////////////////////////////////////////////////
 // UComponent
@@ -2013,7 +2162,7 @@ void UGroomComponent::InitResources(bool bIsBindingReloading)
 
 	if (GroomCache)
 	{
-		GroomCacheBuffers = MakeShared<FGroomCacheBuffers, ESPMode::ThreadSafe>();
+		GroomCacheBuffers = GUseGroomCacheStreaming ? MakeShared<FGroomCacheStreamedBuffers, ESPMode::ThreadSafe>(GroomCache) : MakeShared<FGroomCacheBuffers, ESPMode::ThreadSafe>(GroomCache);
 		UpdateGroomCache(ElapsedTime);
 	}
 	else
@@ -2564,6 +2713,11 @@ void UGroomComponent::OnRegister()
 	UpdateHairGroupsDesc();
 	UpdateHairSimulation();
 
+	if (GUseGroomCacheStreaming)
+	{
+		IGroomCacheStreamingManager::Get().RegisterComponent(this);
+	}
+
 	// Insure the parent skeletal mesh is the same than the registered skeletal mesh, and if not reinitialized resources
 	// This can happens when the OnAttachment callback is not called, but the skeletal mesh change (e.g., the skeletal mesh get recompiled within a blueprint)
 	UMeshComponent* MeshComponent = GetAttachParent() ? Cast<UMeshComponent>(GetAttachParent()) : nullptr;
@@ -2573,12 +2727,28 @@ void UGroomComponent::OnRegister()
 	{
 		InitResources();
 	}
+	else if (GroomCache)
+	{
+		// Buffers already initialized, just need to update them
+		UpdateGroomCache(ElapsedTime);
+	}
 }
 
 void UGroomComponent::OnUnregister()
 {
 	Super::OnUnregister();
 	ReleaseHairSimulation();
+
+	if (GUseGroomCacheStreaming)
+	{
+		if (GroomCacheBuffers.IsValid())
+		{
+			// Reset the buffers so they can be updated at OnRegister
+			FGroomCacheBuffers* Buffers = static_cast<FGroomCacheBuffers*>(GroomCacheBuffers.Get());
+			Buffers->Reset();
+		}
+		IGroomCacheStreamingManager::Get().UnregisterComponent(this);
+	}
 }
 
 void UGroomComponent::BeginDestroy()
@@ -2642,7 +2812,7 @@ void UGroomComponent::UpdateGroomCache(float Time)
 	if (GroomCache && GroomCacheBuffers.IsValid() && bRunning)
 	{
 		FGroomCacheBuffers* Buffers = static_cast<FGroomCacheBuffers*>(GroomCacheBuffers.Get());
-		Buffers->UpdateBuffersAtTime(GroomCache, Time, bLooping);
+		Buffers->UpdateBuffersAtTime(Time, bLooping);
 	}
 }
 
@@ -2652,7 +2822,18 @@ void UGroomComponent::SetGroomCache(UGroomCache* InGroomCache)
 	{
 		ReleaseResources();
 		ResetAnimationTime();
-		GroomCache = InGroomCache;
+
+		if (GUseGroomCacheStreaming)
+		{
+			IGroomCacheStreamingManager::Get().UnregisterComponent(this);
+			GroomCache = InGroomCache;
+			IGroomCacheStreamingManager::Get().RegisterComponent(this);
+		}
+		else
+		{
+			GroomCache = InGroomCache;
+		}
+
 		InitResources();
 	}
 }
@@ -2677,10 +2858,26 @@ void UGroomComponent::ResetAnimationTime()
 	ElapsedTime = 0.0f;
 }
 
+float UGroomComponent::GetAnimationTime() const
+{
+	return ElapsedTime;
+}
+
 void UGroomComponent::TickAtThisTime(const float Time, bool bInIsRunning, bool bInBackwards, bool bInIsLooping)
 {
 	if (GroomCache && bRunning && bManualTick)
 	{
+		float DeltaTime = Time - ElapsedTime;
+		ElapsedTime = Time;
+		if (GUseGroomCacheStreaming)
+		{
+			// Scrubbing forward (or backward) can induce large (or negative) delta time, so force a prefetch
+			if ((DeltaTime > GetDefault<UGroomPluginSettings>()->GroomCacheLookAheadBuffer) ||
+				(DeltaTime < 0))
+			{
+				IGroomCacheStreamingManager::Get().PrefetchData(this);
+			}
+		}
 		UpdateGroomCache(Time);
 	}
 }
diff --git a/Engine/Plugins/Runtime/HairStrands/Source/HairStrandsCore/Private/GroomPluginSettings.cpp b/Engine/Plugins/Runtime/HairStrands/Source/HairStrandsCore/Private/GroomPluginSettings.cpp
new file mode 100644
--- /dev/null
+++ b/Engine/Plugins/Runtime/HairStrands/Source/HairStrandsCore/Private/GroomPluginSettings.cpp
@@ -0,0 +1,8 @@
+// Copyright Epic Games, Inc. All Rights Reserved.
+
+#include "GroomPluginSettings.h"
+	 
+UGroomPluginSettings::UGroomPluginSettings()
+: GroomCacheLookAheadBuffer(1.0f)
+{
+}
diff --git a/Engine/Plugins/Runtime/HairStrands/Source/HairStrandsCore/Public/GroomCache.h b/Engine/Plugins/Runtime/HairStrands/Source/HairStrandsCore/Public/GroomCache.h
--- a/Engine/Plugins/Runtime/HairStrands/Source/HairStrandsCore/Public/GroomCache.h
+++ b/Engine/Plugins/Runtime/HairStrands/Source/HairStrandsCore/Public/GroomCache.h
@@ -29,18 +29,22 @@ public:
 	float GetDuration() const;
 
 	/** Get the frame number at the specified time within the animation range which might not start at 0 */
-	int32  GetFrameNumberAtTime(const float Time, bool bLooping) const;
+	int32 GetFrameNumberAtTime(const float Time, bool bLooping) const;
 
 	/** Get the (floored) frame index at the specified time with the index 0 being the start of the animation */
-	int32  GetFrameIndexAtTime(const float Time, bool bLooping) const;
+	int32 GetFrameIndexAtTime(const float Time, bool bLooping) const;
 
 	/** Get the frame indices and interpolation factor between them that correspond to the specified time */
-	void FindSampleIndexesFromTime(float Time, bool bLooping, bool bIsPlayingBackwards, int32 &OutFrameIndex, int32 &OutNextFrameIndex, float &InterpolationFactor);
+	void GetFrameIndicesAtTime(float Time, bool bLooping, bool bIsPlayingBackwards, int32 &OutFrameIndex, int32 &OutNextFrameIndex, float &InterpolationFactor);
+
+	/** Get the frame indices that correspond to the specified time range */
+	void GetFrameIndicesForTimeRange(float StartTime, float EndTime, bool Looping, TArray<int32>& OutFrameIndices);
 
 	bool GetGroomDataAtTime(float Time, bool bLooping, FGroomCacheAnimationData& AnimData);
 	bool GetGroomDataAtFrameIndex(int32 FrameIndex, FGroomCacheAnimationData& AnimData);
 
 	void SetGroomAnimationInfo(const FGroomAnimationInfo& AnimInfo);
+	const FGroomAnimationInfo& GetGroomAnimationInfo() const { return GroomCacheInfo.AnimationInfo; }
 
 	EGroomCacheType GetType() const;
 
diff --git a/Engine/Plugins/Runtime/HairStrands/Source/HairStrandsCore/Public/GroomCacheData.h b/Engine/Plugins/Runtime/HairStrands/Source/HairStrandsCore/Public/GroomCacheData.h
--- a/Engine/Plugins/Runtime/HairStrands/Source/HairStrandsCore/Public/GroomCacheData.h
+++ b/Engine/Plugins/Runtime/HairStrands/Source/HairStrandsCore/Public/GroomCacheData.h
@@ -144,18 +144,14 @@ class IGroomCacheBuffers
 public:
 	virtual ~IGroomCacheBuffers() {}
 
-	virtual FGroomCacheAnimationData& GetCurrentFrameBuffer() = 0;
-	virtual FGroomCacheAnimationData& GetNextFrameBuffer() = 0;
-	virtual FGroomCacheAnimationData& GetInterpolatedFrameBuffer() = 0;
+	virtual const FGroomCacheAnimationData& GetCurrentFrameBuffer() = 0;
+	virtual const FGroomCacheAnimationData& GetNextFrameBuffer() = 0;
+	virtual const FGroomCacheAnimationData& GetInterpolatedFrameBuffer() = 0;
 
 	virtual int32 GetCurrentFrameIndex() const = 0;
 	virtual int32 GetNextFrameIndex() const = 0;
 	virtual float GetInterpolationFactor() const = 0;
 
-	virtual void SetCurrentFrameIndex(int32 FrameIndex) = 0;
-	virtual void SetNextFrameIndex(int32 FrameIndex) = 0;
-	virtual void SetInterpolationFactor(float Factor) = 0;
-
 	FCriticalSection* GetCriticalSection() { return &CriticalSection; }
 
 private:
diff --git a/Engine/Plugins/Runtime/HairStrands/Source/HairStrandsCore/Public/GroomCacheStreamingManager.h b/Engine/Plugins/Runtime/HairStrands/Source/HairStrandsCore/Public/GroomCacheStreamingManager.h
new file mode 100644
--- /dev/null
+++ b/Engine/Plugins/Runtime/HairStrands/Source/HairStrandsCore/Public/GroomCacheStreamingManager.h
@@ -0,0 +1,51 @@
+// Copyright Epic Games, Inc. All Rights Reserved.
+
+#pragma once
+
+#include "ContentStreaming.h"
+
+class UGroomCache;
+class UGroomComponent;
+
+struct IGroomCacheStreamingManager : public IStreamingManager
+{
+	virtual ~IGroomCacheStreamingManager() {}
+
+	HAIRSTRANDSCORE_API static struct IGroomCacheStreamingManager& Get();
+	HAIRSTRANDSCORE_API static void Register();
+	HAIRSTRANDSCORE_API static void Unregister();
+
+	/** Registers a new component to the streaming manager. */
+	virtual void RegisterComponent(UGroomComponent* GroomComponent) = 0;
+
+	/** Removes the component from the streaming manager. */
+	virtual void UnregisterComponent(UGroomComponent* GroomComponent) = 0;
+
+	/** 
+	 * Prefetch data for the component. Data is automatically prefetched when initially registering the component
+	 * but exposing this may be useful in other cases (eg. when seeking, etc.)
+	 */
+	virtual void PrefetchData(UGroomComponent* GroomComponent) = 0;
+
+	/**
+	 * Returns a pointer to the cached GroomCacheAnimationData.
+	 *
+	 * @param GroomCache	The GroomCache for which to retrieve the data
+	 * @param ChunkIndex	Index of the chunk we want
+	 * @return The cached data or null if it's not loaded yet
+	 */
+	virtual const struct FGroomCacheAnimationData* MapAnimationData(const UGroomCache* GroomCache, uint32 ChunkIndex) = 0;
+
+	/**
+	 * Releases pointer to the cachedGroomCacheAnimationData.
+	 * Must be called for every successful call to MapAnimationData.
+	 *
+	 * @param GroomCache	The GroomCache for which the data was retrieved
+	 * @param ChunkIndex	Index of the chunk to release
+	 */
+	virtual void UnmapAnimationData(const UGroomCache* GroomCache, uint32 ChunkIndex) = 0;
+
+	/** Cleans up the resources managed by the streaming manager */
+	virtual void Shutdown() = 0;
+};
+
diff --git a/Engine/Plugins/Runtime/HairStrands/Source/HairStrandsCore/Public/GroomComponent.h b/Engine/Plugins/Runtime/HairStrands/Source/HairStrandsCore/Public/GroomComponent.h
--- a/Engine/Plugins/Runtime/HairStrands/Source/HairStrandsCore/Public/GroomComponent.h
+++ b/Engine/Plugins/Runtime/HairStrands/Source/HairStrandsCore/Public/GroomComponent.h
@@ -219,6 +219,8 @@ public:
 	void TickAtThisTime(const float Time, bool bInIsRunning, bool bInBackwards, bool bInIsLooping);
 
 	void ResetAnimationTime();
+	float GetAnimationTime() const;
+	bool IsLooping() const { return bLooping; }
 
 private:
 	void UpdateGroomCache(float Time);
diff --git a/Engine/Plugins/Runtime/HairStrands/Source/HairStrandsCore/Public/GroomPluginSettings.h b/Engine/Plugins/Runtime/HairStrands/Source/HairStrandsCore/Public/GroomPluginSettings.h
new file mode 100644
--- /dev/null
+++ b/Engine/Plugins/Runtime/HairStrands/Source/HairStrandsCore/Public/GroomPluginSettings.h
@@ -0,0 +1,23 @@
+// Copyright Epic Games, Inc. All Rights Reserved.
+
+#pragma once
+
+#include "UObject/ObjectMacros.h"
+#include "UObject/Object.h"
+
+#include "GroomPluginSettings.generated.h"
+
+/** Settings for the groom plug-in */
+UCLASS(config=Engine)
+class HAIRSTRANDSCORE_API UGroomPluginSettings : public UObject
+{
+	GENERATED_BODY()
+
+public:
+	 
+	UGroomPluginSettings();
+
+	/** The amount of groom cache animation (in seconds) to pre-load ahead of time by the streaming manager */
+	UPROPERTY(config, EditAnywhere, Category=GroomCache, meta = (ClampMin = "0.05", ClampMax = "60.0", DisplayName = "Look-Ahead Buffer (in seconds)"))
+	float GroomCacheLookAheadBuffer;
+};
diff --git a/Engine/Plugins/Runtime/HairStrands/Source/HairStrandsEditor/Private/HairStrandsEditor.cpp b/Engine/Plugins/Runtime/HairStrands/Source/HairStrandsEditor/Private/HairStrandsEditor.cpp
--- a/Engine/Plugins/Runtime/HairStrands/Source/HairStrandsEditor/Private/HairStrandsEditor.cpp
+++ b/Engine/Plugins/Runtime/HairStrands/Source/HairStrandsEditor/Private/HairStrandsEditor.cpp
@@ -19,14 +19,18 @@
 #include "GroomBindingDetailsCustomization.h"
 #include "GroomCacheImportOptions.h"
 #include "GroomCacheImportSettingsCustomization.h"
+#include "GroomCacheStreamingManager.h"
 #include "GroomCacheTrackEditor.h"
 #include "GroomComponentDetailsCustomization.h"
 #include "GroomCreateBindingOptions.h"
 #include "GroomEditorCommands.h"
 #include "GroomEditorMode.h"
+#include "GroomPluginSettings.h"
 
 #include "AssetRegistryModule.h"
 #include "FileHelpers.h"
+#include "ISettingsModule.h"
+#include "ISettingsSection.h"
 #include "ISequencerModule.h"
 
 IMPLEMENT_MODULE(FGroomEditor, HairStrandsEditor);
@@ -134,10 +138,30 @@ void FGroomEditor::StartupModule()
 
 	ISequencerModule& SequencerModule = FModuleManager::Get().LoadModuleChecked<ISequencerModule>("Sequencer");
 	TrackEditorBindingHandle = SequencerModule.RegisterTrackEditor(FOnCreateTrackEditor::CreateStatic(&FGroomCacheTrackEditor::CreateTrackEditor));
+
+	ISettingsModule* SettingsModule = FModuleManager::GetModulePtr<ISettingsModule>("Settings");
+	if (SettingsModule != nullptr)
+	{
+		ISettingsSectionPtr SettingsSection = SettingsModule->RegisterSettings("Project", "Plugins", "Groom",
+			LOCTEXT("GroomPluginSettingsName", "Groom"),
+			LOCTEXT("GroomPluginSettingsDescription", "Configure the Groom plug-in."),
+			GetMutableDefault<UGroomPluginSettings>()
+		);
+	}
+
+	IGroomCacheStreamingManager::Register();
 }
 
 void FGroomEditor::ShutdownModule()
 {
+	IGroomCacheStreamingManager::Unregister();
+
+	ISettingsModule* SettingsModule = FModuleManager::GetModulePtr<ISettingsModule>("Settings");
+	if (SettingsModule != nullptr)
+	{
+		SettingsModule->UnregisterSettings("Project", "Plugins", "Groom");
+	}
+
 	ISequencerModule* SequencerModulePtr = FModuleManager::Get().GetModulePtr<ISequencerModule>("Sequencer");
 	if (SequencerModulePtr)
 	{
