Fixed issue where GroomCache could disappear if it moved far from its initial position. Its bounding box wasn't updated to follow the animation.

diff --git a/Engine/Plugins/Runtime/HairStrands/Source/HairStrandsCore/Private/GroomComponent.cpp b/Engine/Plugins/Runtime/HairStrands/Source/HairStrandsCore/Private/GroomComponent.cpp
--- a/Engine/Plugins/Runtime/HairStrands/Source/HairStrandsCore/Private/GroomComponent.cpp
+++ b/Engine/Plugins/Runtime/HairStrands/Source/HairStrandsCore/Private/GroomComponent.cpp
@@ -1036,6 +1036,10 @@ public:
 			FGroomCacheGroupData& InterpolatedGroupData = InterpolatedFrame.GroupsData[GroupIndex];
 			const int32 NumVertices = CurrentGroupData.VertexData.PointsPosition.Num();
 
+			// Update the bounding box used for hair strands rendering computation
+			FVector InterpolatedCenter = FMath::Lerp(CurrentGroupData.BoundingBox.GetCenter(), NextGroupData.BoundingBox.GetCenter(), InterpolationFactor);
+			InterpolatedGroupData.BoundingBox = CurrentGroupData.BoundingBox.MoveTo(InterpolatedCenter) + NextGroupData.BoundingBox.MoveTo(InterpolatedCenter);
+
 			// Parallel batched interpolation
 			const int32 BatchSize = 1024;
 			const int32 BatchCount = (NumVertices + BatchSize - 1) / BatchSize;
@@ -1056,6 +1060,24 @@ public:
 		}
 	}
 
+	FBox GetBoundingBox()
+	{
+		// Approximate bounding box used for visibility culling
+		FBox BBox(EForceInit::ForceInitToZero);
+		for (const FGroomCacheGroupData& GroupData : GetCurrentFrameBuffer().GroupsData)
+		{
+			BBox += GroupData.BoundingBox;
+		}
+
+		for (const FGroomCacheGroupData& GroupData : GetNextFrameBuffer().GroupsData)
+		{
+			BBox += GroupData.BoundingBox;
+		}
+
+		return BBox;
+	}
+
+
 protected:
 	static FGroomCacheAnimationData EmptyFrame;
 
@@ -1614,25 +1636,33 @@ FBoxSphereBounds UGroomComponent::CalcBounds(const FTransform& InLocalToWorld) c
 		else
 		{
 			FBox LocalBounds(EForceInit::ForceInitToZero);
-			for (const FHairGroupData& GroupData : GroomAsset->HairGroupsData)
+			if (!GroomCacheBuffers.IsValid())
 			{
-				if (IsHairStrandsEnabled(EHairStrandsShaderType::Strands) && GroupData.Strands.HasValidData())
-				{
-					LocalBounds += GroupData.Strands.GetBounds();
-				}
-				else if (IsHairStrandsEnabled(EHairStrandsShaderType::Cards) && GroupData.Cards.HasValidData())
-				{
-					LocalBounds += GroupData.Cards.GetBounds();
-				}
-				else if (IsHairStrandsEnabled(EHairStrandsShaderType::Meshes) && GroupData.Meshes.HasValidData())
-				{
-					LocalBounds += GroupData.Meshes.GetBounds();
-				}
-				else if (GroupData.Guides.HasValidData())
+				for (const FHairGroupData& GroupData : GroomAsset->HairGroupsData)
 				{
-					LocalBounds += GroupData.Guides.GetBounds();
+					if (IsHairStrandsEnabled(EHairStrandsShaderType::Strands) && GroupData.Strands.HasValidData())
+					{
+						LocalBounds += GroupData.Strands.GetBounds();
+					}
+					else if (IsHairStrandsEnabled(EHairStrandsShaderType::Cards) && GroupData.Cards.HasValidData())
+					{ 
+						LocalBounds += GroupData.Cards.GetBounds();
+					}
+					else if (IsHairStrandsEnabled(EHairStrandsShaderType::Meshes) && GroupData.Meshes.HasValidData())
+					{
+						LocalBounds += GroupData.Meshes.GetBounds();
+					}
+					else if (GroupData.Guides.HasValidData())
+					{
+						LocalBounds += GroupData.Guides.GetBounds();
+					}
 				}
 			}
+			else
+			{
+				FGroomCacheBuffers* Buffers = static_cast<FGroomCacheBuffers*>(GroomCacheBuffers.Get());
+				LocalBounds = Buffers->GetBoundingBox();
+			}
 			return FBoxSphereBounds(LocalBounds.TransformBy(InLocalToWorld));
 		}
 	}
@@ -2813,6 +2843,9 @@ void UGroomComponent::UpdateGroomCache(float Time)
 	{
 		FGroomCacheBuffers* Buffers = static_cast<FGroomCacheBuffers*>(GroomCacheBuffers.Get());
 		Buffers->UpdateBuffersAtTime(Time, bLooping);
+
+		// Trigger an update of the bounds so that it follows the GroomCache
+		MarkRenderTransformDirty();
 	}
 }
 
diff --git a/Engine/Plugins/Runtime/HairStrands/Source/HairStrandsCore/Private/HairStrandsMeshProjection.cpp b/Engine/Plugins/Runtime/HairStrands/Source/HairStrandsCore/Private/HairStrandsMeshProjection.cpp
--- a/Engine/Plugins/Runtime/HairStrands/Source/HairStrandsCore/Private/HairStrandsMeshProjection.cpp
+++ b/Engine/Plugins/Runtime/HairStrands/Source/HairStrandsCore/Private/HairStrandsMeshProjection.cpp
@@ -1804,7 +1804,7 @@ void AddHairStrandUpdatePositionOffsetPass(
 	FHairStrandsDeformedRootResource* DeformedRootResources,
 	FHairStrandsDeformedResource* DeformedResources)
 {
-	if (LODIndex < 0 || DeformedResources == nullptr)
+	if ((DeformedRootResources && LODIndex < 0) || DeformedResources == nullptr)
 	{
 		return;
 	}
