Fixed import and playback of GroomCache.

diff --git a/Engine/Plugins/Runtime/HairStrands/Source/HairStrandsCore/Private/GroomCache.cpp b/Engine/Plugins/Runtime/HairStrands/Source/HairStrandsCore/Private/GroomCache.cpp
--- a/Engine/Plugins/Runtime/HairStrands/Source/HairStrandsCore/Private/GroomCache.cpp
+++ b/Engine/Plugins/Runtime/HairStrands/Source/HairStrandsCore/Private/GroomCache.cpp
@@ -243,7 +243,7 @@ FGroomCacheProcessor::FGroomCacheProcessor(EGroomCacheType InType, EGroomCacheAt
 {
 }
 
-void FGroomCacheProcessor::AddGroomSample(TArray<FHairGroupData>&& GroupData)
+void FGroomCacheProcessor::AddGroomSample(TArray<FHairDescriptionGroup>&& GroupData)
 {
 	TArray<uint8> TempBytes;
 	FMemoryWriter Ar(TempBytes, true);
diff --git a/Engine/Plugins/Runtime/HairStrands/Source/HairStrandsCore/Private/GroomCacheData.cpp b/Engine/Plugins/Runtime/HairStrands/Source/HairStrandsCore/Private/GroomCacheData.cpp
--- a/Engine/Plugins/Runtime/HairStrands/Source/HairStrandsCore/Private/GroomCacheData.cpp
+++ b/Engine/Plugins/Runtime/HairStrands/Source/HairStrandsCore/Private/GroomCacheData.cpp
@@ -81,19 +81,14 @@ void FGroomCacheGroupData::Serialize(FArchive& Ar, int32 Version, EGroomCacheAtt
 	StrandData.Serialize(Ar, Version, Attributes);
 }
 
-FGroomCacheAnimationData::FGroomCacheAnimationData(TArray<FHairGroupData>&& HairGroupData, int32 InVersion, EGroomCacheType Type, EGroomCacheAttributes InAttributes)
+FGroomCacheAnimationData::FGroomCacheAnimationData(TArray<FHairDescriptionGroup>&& HairGroupData, int32 InVersion, EGroomCacheType Type, EGroomCacheAttributes InAttributes)
 : Version(InVersion)
 {
 	Attributes = Type == EGroomCacheType::Strands ? InAttributes : InAttributes & EGroomCacheAttributes::Position;
 
-	for (FHairGroupData& GroupData : HairGroupData)
+	for (FHairDescriptionGroup& GroupData : HairGroupData)
 	{
-		FHairStrandsDatas StrandsData;
-		FHairStrandsDatas GuidesData;
-
-		// TODO FIX THIS BY PASSING THE GUIDES/STRANDS DATAS as they are not part of FHairGroupData anymore.
-		//GroupsData.Add(FGroomCacheGroupData(MoveTemp(Type == EGroomCacheType::Strands ? GroupData.Strands.Data : GroupData.Guides.Data)));
-		GroupsData.Add(FGroomCacheGroupData(MoveTemp(Type == EGroomCacheType::Strands ? StrandsData : GuidesData)));
+		GroupsData.Add(FGroomCacheGroupData(MoveTemp(Type == EGroomCacheType::Strands ? GroupData.Strands : GroupData.Guides)));
 	}
 }
 
diff --git a/Engine/Plugins/Runtime/HairStrands/Source/HairStrandsCore/Private/GroomComponent.cpp b/Engine/Plugins/Runtime/HairStrands/Source/HairStrandsCore/Private/GroomComponent.cpp
--- a/Engine/Plugins/Runtime/HairStrands/Source/HairStrandsCore/Private/GroomComponent.cpp
+++ b/Engine/Plugins/Runtime/HairStrands/Source/HairStrandsCore/Private/GroomComponent.cpp
@@ -2274,7 +2274,7 @@ void UGroomComponent::InitResources(bool bIsBindingReloading)
 		// * Physics simulation
 		// * RBF deformation.
 		// Therefore, even if simulation is disabled, we need to run partially the update if the binding system is enabled (skin deformation + RBF correction)
-		const bool bNeedGuides = GroupData.Guides.HasValidData() && (bHasNeedSimulation || bHasNeedGlobalDeformation || bPreviewMode);
+		const bool bNeedGuides = GroupData.Guides.HasValidData() && (bHasNeedSimulation || bHasNeedGlobalDeformation || bPreviewMode) || HairGroupInstance->Debug.GroomCacheType == EGroomCacheType::Guides;
 		if (bNeedGuides)
 		{
 			HairGroupInstance->Guides.Data = &GroupData.Guides.BulkData;
@@ -2356,6 +2356,9 @@ void UGroomComponent::InitResources(bool bIsBindingReloading)
 				}
 			}
 
+			// The strands GroomCache needs the Strands.DeformedResource, but the guides GroomCache also needs it for the HairTangentPass
+			bNeedDynamicResources = bNeedDynamicResources || HairGroupInstance->Debug.GroomCacheType != EGroomCacheType::None;
+
 			#if RHI_RAYTRACING
 			if (IsHairRayTracingEnabled() && HairGroupInstance->Strands.Modifier.bUseHairRaytracingGeometry)
 			{
diff --git a/Engine/Plugins/Runtime/HairStrands/Source/HairStrandsCore/Public/GroomCache.h b/Engine/Plugins/Runtime/HairStrands/Source/HairStrandsCore/Public/GroomCache.h
--- a/Engine/Plugins/Runtime/HairStrands/Source/HairStrandsCore/Public/GroomCache.h
+++ b/Engine/Plugins/Runtime/HairStrands/Source/HairStrandsCore/Public/GroomCache.h
@@ -89,7 +89,7 @@ class HAIRSTRANDSCORE_API FGroomCacheProcessor
 public:
 	FGroomCacheProcessor(EGroomCacheType InType, EGroomCacheAttributes InAttributes);
 
-	void AddGroomSample(TArray<FHairGroupData>&& GroomData);
+	void AddGroomSample(TArray<FHairDescriptionGroup>&& GroomData);
 	void TransferChunks(UGroomCache* GroomCache);
 	EGroomCacheType GetType() const { return Type; }
 
diff --git a/Engine/Plugins/Runtime/HairStrands/Source/HairStrandsCore/Public/GroomCacheData.h b/Engine/Plugins/Runtime/HairStrands/Source/HairStrandsCore/Public/GroomCacheData.h
--- a/Engine/Plugins/Runtime/HairStrands/Source/HairStrandsCore/Public/GroomCacheData.h
+++ b/Engine/Plugins/Runtime/HairStrands/Source/HairStrandsCore/Public/GroomCacheData.h
@@ -129,7 +129,7 @@ struct FGroomCacheGroupData
 struct FGroomCacheAnimationData
 {
 	FGroomCacheAnimationData() = default;
-	FGroomCacheAnimationData(TArray<struct FHairGroupData>&& HairGroupData, int32 Version, EGroomCacheType Type, EGroomCacheAttributes Attributes);
+	FGroomCacheAnimationData(TArray<struct FHairDescriptionGroup>&& HairGroupData, int32 Version, EGroomCacheType Type, EGroomCacheAttributes Attributes);
 
 	TArray<FGroomCacheGroupData> GroupsData;
 	EGroomCacheAttributes Attributes;
diff --git a/Engine/Plugins/Runtime/HairStrands/Source/HairStrandsEditor/Private/GroomCacheImporter.cpp b/Engine/Plugins/Runtime/HairStrands/Source/HairStrandsEditor/Private/GroomCacheImporter.cpp
--- a/Engine/Plugins/Runtime/HairStrands/Source/HairStrandsEditor/Private/GroomCacheImporter.cpp
+++ b/Engine/Plugins/Runtime/HairStrands/Source/HairStrandsEditor/Private/GroomCacheImporter.cpp
@@ -124,26 +124,13 @@ TArray<UGroomCache*> FGroomCacheImporter::ImportGroomCache(const FString& Source
 				const uint32 GroupCount = HairDescriptionGroups.HairGroups.Num();
 
 				TArray<FHairGroupInfoWithVisibility> HairGroupsInfo = GroomAssetForCache->HairGroupsInfo;
-				TArray<FHairGroupData> HairGroupsData;
+				TArray<FHairDescriptionGroup> HairGroupsData;
 				HairGroupsData.SetNum(GroupCount);
 				for (uint32 GroupIndex = 0; GroupIndex < GroupCount; ++GroupIndex)
 				{
-					//FGroomBuilder::BuildHairGroupData(ProcessedHairDescription, GroomAssetForCache->HairGroupsInterpolation[GroupIndex], GroupIndex, HairGroupsData[GroupIndex]);
-
 					const FHairDescriptionGroup& HairGroup = HairDescriptionGroups.HairGroups[GroupIndex];
-
-					FHairStrandsDatas StrandsData;
-					FHairStrandsDatas GuidesData;
-					FGroomBuilder::BuildData(HairGroup, GroomAssetForCache->HairGroupsInterpolation[GroupIndex], HairGroupsInfo[GroupIndex], StrandsData, GuidesData);
-
-					FGroomBuilder::BuildBulkData(HairGroup.Info, GuidesData, HairGroupsData[GroupIndex].Guides.BulkData);
-					FGroomBuilder::BuildBulkData(HairGroup.Info, StrandsData, HairGroupsData[GroupIndex].Strands.BulkData);
-
-					FHairStrandsInterpolationDatas InterpolationData;
-					FGroomBuilder::BuildInterplationData(HairGroup.Info, StrandsData, GuidesData, GroomAssetForCache->HairGroupsInterpolation[GroupIndex].InterpolationSettings, InterpolationData);
-					FGroomBuilder::BuildInterplationBulkData(GuidesData, InterpolationData, HairGroupsData[GroupIndex].Strands.InterpolationBulkData);
-
-					FGroomBuilder::BuildClusterBulkData(StrandsData, HairDescriptionGroups.BoundRadius, GroomAssetForCache->HairGroupsLOD[GroupIndex], HairGroupsData[GroupIndex].Strands.ClusterCullingBulkData);
+					FHairDescriptionGroup& HairGroupData = HairGroupsData[GroupIndex];
+					FGroomBuilder::BuildData(HairGroup, GroomAssetForCache->HairGroupsInterpolation[GroupIndex], HairGroupsInfo[GroupIndex], HairGroupData.Strands, HairGroupData.Guides);
 				}
 
 				// The HairGroupData is converted into animated groom data by the GroomCacheProcessor
