Fixed build error (Issue 80899).

diff --git a/Engine/Plugins/Runtime/HairStrands/Source/HairStrandsCore/Private/GroomComponent.cpp b/Engine/Plugins/Runtime/HairStrands/Source/HairStrandsCore/Private/GroomComponent.cpp
--- a/Engine/Plugins/Runtime/HairStrands/Source/HairStrandsCore/Private/GroomComponent.cpp
+++ b/Engine/Plugins/Runtime/HairStrands/Source/HairStrandsCore/Private/GroomComponent.cpp
@@ -2274,7 +2274,7 @@ void UGroomComponent::InitResources(bool bIsBindingReloading)
 		// * Physics simulation
 		// * RBF deformation.
 		// Therefore, even if simulation is disabled, we need to run partially the update if the binding system is enabled (skin deformation + RBF correction)
-		const bool bNeedGuides = GroupData.Guides.HasValidData() && (bHasNeedSimulation || bHasNeedGlobalDeformation || bPreviewMode) || HairGroupInstance->Debug.GroomCacheType == EGroomCacheType::Guides;
+		const bool bNeedGuides = (GroupData.Guides.HasValidData() && (bHasNeedSimulation || bHasNeedGlobalDeformation || bPreviewMode)) || (HairGroupInstance->Debug.GroomCacheType == EGroomCacheType::Guides);
 		if (bNeedGuides)
 		{
 			HairGroupInstance->Guides.Data = &GroupData.Guides.BulkData;
