Fixed issue at shutdown with streaming GroomCache.

diff --git a/Engine/Plugins/Runtime/HairStrands/Source/HairStrandsCore/Private/GroomCacheStreamingData.cpp b/Engine/Plugins/Runtime/HairStrands/Source/HairStrandsCore/Private/GroomCacheStreamingData.cpp
--- a/Engine/Plugins/Runtime/HairStrands/Source/HairStrandsCore/Private/GroomCacheStreamingData.cpp
+++ b/Engine/Plugins/Runtime/HairStrands/Source/HairStrandsCore/Private/GroomCacheStreamingData.cpp
@@ -22,13 +22,10 @@ FGroomCacheStreamingData::~FGroomCacheStreamingData()
 
 	// Delete the read requests in the background as it can be time-consuming
 	TSet<IBulkDataIORequest*> ReadyForDeletion(DelayedDeleteReadRequests);
-	Async(EAsyncExecution::ThreadPool, [ReadyForDeletion]()
+	for (IBulkDataIORequest* ReadRequest : ReadyForDeletion)
 	{
-		for (IBulkDataIORequest* ReadRequest : ReadyForDeletion)
-		{
-			delete ReadRequest;
-		}
-	});
+		delete ReadRequest;
+	}
 
 	// Delete AnimData buffers that are resident in memory
 	for (TMap<int32, FResidentGroomCacheChunk>::TIterator Iter = Chunks.CreateIterator(); Iter; ++Iter)
@@ -137,10 +134,10 @@ void FGroomCacheStreamingData::PrefetchData(UGroomComponent *Component)
 	}
 
 	// Schedule the rest of the needed frames to be loaded asynchronously
-	UpdateStreamingStatus();
+	UpdateStreamingStatus(true);
 }
 
-void FGroomCacheStreamingData::UpdateStreamingStatus()
+void FGroomCacheStreamingData::UpdateStreamingStatus(bool bAsyncDeletionAllowed)
 {
 	TRACE_CPUPROFILER_EVENT_SCOPE(FGroomCacheStreamingData::UpdateStreamingStatus);
 
@@ -269,14 +266,27 @@ void FGroomCacheStreamingData::UpdateStreamingStatus()
 	}
 	DelayedDeleteReadRequests = DelayedDeleteReadRequests.Difference(ReadyForDeletion);
 	
-	// Delete the read requests in the background as it can be time-consuming
-	Async(EAsyncExecution::ThreadPool, [ReadyForDeletion]()
+	if (ReadyForDeletion.Num() > 0)
 	{
-		for (IBulkDataIORequest* ReadRequest : ReadyForDeletion)
+		if (bAsyncDeletionAllowed)
 		{
-			delete ReadRequest;
+			// Delete the read requests in the background as it can be time-consuming
+			Async(EAsyncExecution::ThreadPool, [ReadyForDeletion]()
+			{
+				for (IBulkDataIORequest* ReadRequest : ReadyForDeletion)
+				{
+					delete ReadRequest;
+				}
+			});
 		}
-	});
+		else
+		{
+			for (IBulkDataIORequest* ReadRequest : ReadyForDeletion)
+			{
+				delete ReadRequest;
+			}
+		}
+	}
 }
 
 bool FGroomCacheStreamingData::BlockTillAllRequestsFinished(float TimeLimit)
diff --git a/Engine/Plugins/Runtime/HairStrands/Source/HairStrandsCore/Private/GroomCacheStreamingData.h b/Engine/Plugins/Runtime/HairStrands/Source/HairStrandsCore/Private/GroomCacheStreamingData.h
--- a/Engine/Plugins/Runtime/HairStrands/Source/HairStrandsCore/Private/GroomCacheStreamingData.h
+++ b/Engine/Plugins/Runtime/HairStrands/Source/HairStrandsCore/Private/GroomCacheStreamingData.h
@@ -42,7 +42,7 @@ public:
 	~FGroomCacheStreamingData();
 
 	/** Updates the internal state of the streaming data. */
-	void UpdateStreamingStatus();
+	void UpdateStreamingStatus(bool bAsyncDeletionAllowed);
 
 	/** Returns a pointer to the cache data for ChunkIndex if it's loaded; nullptr otherwise. */
 	const FGroomCacheAnimationData* MapAnimationData(uint32 ChunkIndex);
diff --git a/Engine/Plugins/Runtime/HairStrands/Source/HairStrandsCore/Private/GroomCacheStreamingManager.cpp b/Engine/Plugins/Runtime/HairStrands/Source/HairStrandsCore/Private/GroomCacheStreamingManager.cpp
--- a/Engine/Plugins/Runtime/HairStrands/Source/HairStrandsCore/Private/GroomCacheStreamingManager.cpp
+++ b/Engine/Plugins/Runtime/HairStrands/Source/HairStrandsCore/Private/GroomCacheStreamingManager.cpp
@@ -36,7 +36,7 @@ private:
 	void RegisterGroomCache(UGroomCache* GroomCache, UGroomComponent* GroomComponent);
 	void UnregisterGroomCache(UGroomCache* GroomCache, UGroomComponent* GroomComponent);
 	void PrefetchDataInternal(UGroomComponent* GroomComponent);
-	void DeleteStreamingData();
+	void DeleteStreamingData(bool bAsyncDeletionAllowed);
 
 	// StreamingData for registered GroomCaches
 	TMap<UGroomCache*, FGroomCacheStreamingData*> StreamingGroomCaches;
@@ -75,7 +75,7 @@ void FGroomCacheStreamingManager::Shutdown()
 	// Delete the StreamingData that were already queued for deletion
 	while (StreamingGroomCachesToDelete.Num() > 0)
 	{
-		DeleteStreamingData();
+		DeleteStreamingData(false);
 	}
 
 	FScopeLock Lock(&CriticalSection);
@@ -131,19 +131,19 @@ void FGroomCacheStreamingManager::UpdateResourceStreaming(float DeltaTime, bool
 		// Update the internal state of the StreamingData
 		for (TMap<UGroomCache*, FGroomCacheStreamingData*>::TIterator Iter = StreamingGroomCaches.CreateIterator(); Iter; ++Iter)
 		{
-			Iter.Value()->UpdateStreamingStatus();
+			Iter.Value()->UpdateStreamingStatus(true);
 		}
 	}
 
-	DeleteStreamingData();
+	DeleteStreamingData(true);
 }
 
-void FGroomCacheStreamingManager::DeleteStreamingData()
+void FGroomCacheStreamingManager::DeleteStreamingData(bool bAsyncDeletionAllowed)
 {
 	TSet<FGroomCacheStreamingData*> ReadyForDeletion;
 	for (FGroomCacheStreamingData* StreamingData : StreamingGroomCachesToDelete)
 	{
-		StreamingData->UpdateStreamingStatus();
+		StreamingData->UpdateStreamingStatus(bAsyncDeletionAllowed);
 		if (!StreamingData->IsStreamingInProgress())
 		{
 			ReadyForDeletion.Add(StreamingData);
