Fix hair cards transition issues.

diff --git a/Engine/Plugins/Runtime/HairStrands/Source/HairStrandsCore/Private/GroomComponent.cpp b/Engine/Plugins/Runtime/HairStrands/Source/HairStrandsCore/Private/GroomComponent.cpp
--- a/Engine/Plugins/Runtime/HairStrands/Source/HairStrandsCore/Private/GroomComponent.cpp
+++ b/Engine/Plugins/Runtime/HairStrands/Source/HairStrandsCore/Private/GroomComponent.cpp
@@ -607,7 +607,7 @@ public:
 					AddOpaqueRaytracingInstance(OverrideLocalToWorld, GroupData.RayTracingGeometries_Cards[LODIndex], RaytracingInstanceMask_ThinShadow, OutRayTracingInstances);
 				#else
 
-					FMeshBatch* MeshBatch = CreateMeshBatch(Context.ReferenceView, Context.ReferenceViewFamily, Context.RayTracingMeshResourceCollector, GroupData, Instance, GroupIt, /*LODIndex,*/ nullptr);
+					FMeshBatch* MeshBatch = CreateMeshBatch(Context.ReferenceView, Context.ReferenceViewFamily, Context.RayTracingMeshResourceCollector, GroupData, Instance, GroupIt, nullptr);
 					TArray<FMeshBatch> MeshBatches;
 					MeshBatches.Add(*MeshBatch);
 					AddOpaqueRaytracingInstance(OverrideLocalToWorld, GroupData.RayTracingGeometries_Cards[LODIndex], RaytracingInstanceMask_Opaque, MeshBatches, GetScene().GetFeatureLevel(), OutRayTracingInstances);
@@ -628,7 +628,7 @@ public:
 					AddOpaqueRaytracingInstance(OverrideLocalToWorld, GroupData.RayTracingGeometries_Meshes[LODIndex], RaytracingInstanceMask_ThinShadow, OutRayTracingInstances);
 				#else
 
-					FMeshBatch* MeshBatch = CreateMeshBatch(Context.ReferenceView, Context.ReferenceViewFamily, Context.RayTracingMeshResourceCollector, GroupData, Instance, GroupIt, /*LODIndex,*/ nullptr);
+					FMeshBatch* MeshBatch = CreateMeshBatch(Context.ReferenceView, Context.ReferenceViewFamily, Context.RayTracingMeshResourceCollector, GroupData, Instance, GroupIt, nullptr);
 					TArray<FMeshBatch> MeshBatches;
 					MeshBatches.Add(*MeshBatch);
 					AddOpaqueRaytracingInstance(OverrideLocalToWorld, GroupData.RayTracingGeometries_Meshes[LODIndex], RaytracingInstanceMask_Opaque, MeshBatches, GetScene().GetFeatureLevel(), OutRayTracingInstances);
diff --git a/Engine/Plugins/Runtime/HairStrands/Source/HairStrandsCore/Private/GroomResources.cpp b/Engine/Plugins/Runtime/HairStrands/Source/HairStrandsCore/Private/GroomResources.cpp
--- a/Engine/Plugins/Runtime/HairStrands/Source/HairStrandsCore/Private/GroomResources.cpp
+++ b/Engine/Plugins/Runtime/HairStrands/Source/HairStrandsCore/Private/GroomResources.cpp
@@ -582,6 +582,29 @@ void FHairCardsProceduralResource::InternalRelease()
 	CardVoxel.NormalBuffer.Release();
 }
 
+////////////////////////////////////////////////////////////////////////////////////////////////////////////////////////
+namespace HairCardsTransition
+{
+	BEGIN_SHADER_PARAMETER_STRUCT(FTransitionParameters, )
+		SHADER_PARAMETER_RDG_BUFFER_SRV(Buffer, Buffer)
+	END_SHADER_PARAMETER_STRUCT()
+
+	void TransitToSRV(FRDGBuilder& GraphBuilder, FRDGBufferSRVRef InBuffer)
+	{
+		FTransitionParameters* Parameters = GraphBuilder.AllocParameters<FTransitionParameters>();
+		Parameters->Buffer = InBuffer;
+
+		GraphBuilder.AddPass(
+			RDG_EVENT_NAME("HairCardsDeformedBufferTransitionPass"),
+			Parameters,
+			ERDGPassFlags::Raster | ERDGPassFlags::SkipRenderPass | ERDGPassFlags::NeverCull,
+			[Parameters](FRHICommandList& RHICmdList)
+			{
+				// Nothing to do, we just want GraphBuilder to force the buffer to be transited to SRV Graphics state.
+			});
+	}
+}
+
 /////////////////////////////////////////////////////////////////////////////////////////
 
 FHairCardsDeformedResource::FHairCardsDeformedResource(const FHairCardsBulkData& InBulkData, bool bInInitializedData) :
@@ -594,12 +617,17 @@ void FHairCardsDeformedResource::InternalAllocate(FRDGBuilder& GraphBuilder)
 	if (bInitializedData)
 	{
 		InternalCreateVertexBufferRDG<FHairCardsPositionFormat>(GraphBuilder, BulkData.Positions, DeformedPositionBuffer[0], TEXT("Hair.CardsDeformed_Position0"), EHairResourceUsageType::Dynamic);
-		InternalCreateVertexBufferRDG<FHairCardsPositionFormat>(GraphBuilder,BulkData.Positions, DeformedPositionBuffer[1], TEXT("Hair.CardsDeformed_Position1"), EHairResourceUsageType::Dynamic);
+		InternalCreateVertexBufferRDG<FHairCardsPositionFormat>(GraphBuilder, BulkData.Positions, DeformedPositionBuffer[1], TEXT("Hair.CardsDeformed_Position1"), EHairResourceUsageType::Dynamic);
 	}
 	else
 	{
 		InternalCreateVertexBufferRDG<FHairCardsPositionFormat>(GraphBuilder, BulkData.GetNumVertices(), DeformedPositionBuffer[0], TEXT("Hair.CardsDeformed_Position0"), EHairResourceUsageType::Dynamic);
 		InternalCreateVertexBufferRDG<FHairCardsPositionFormat>(GraphBuilder, BulkData.GetNumVertices(), DeformedPositionBuffer[1], TEXT("Hair.CardsDeformed_Position1"), EHairResourceUsageType::Dynamic);
+
+		// Manually transit to SRVs, in case of the cards are rendered not visible, but still rendered (in shadows for instance). In such a case, the cards deformation pass is not called, and thus the 
+		// buffers are never transit from UAV (clear) to SRV for rasterization. 
+		HairCardsTransition::TransitToSRV(GraphBuilder, Register(GraphBuilder, DeformedPositionBuffer[0], ERDGImportedBufferFlags::CreateSRV).SRV);
+		HairCardsTransition::TransitToSRV(GraphBuilder, Register(GraphBuilder, DeformedPositionBuffer[1], ERDGImportedBufferFlags::CreateSRV).SRV);
 	}
 }
 
diff --git a/Engine/Plugins/Runtime/HairStrands/Source/HairStrandsCore/Private/HairStrandsRendering.cpp b/Engine/Plugins/Runtime/HairStrands/Source/HairStrandsCore/Private/HairStrandsRendering.cpp
--- a/Engine/Plugins/Runtime/HairStrands/Source/HairStrandsCore/Private/HairStrandsRendering.cpp
+++ b/Engine/Plugins/Runtime/HairStrands/Source/HairStrandsCore/Private/HairStrandsRendering.cpp
@@ -1886,6 +1886,7 @@ void ComputeHairStrandsInterpolation(
 				}
 
 				ResourceAccessFinalizer.AddBuffer(Register(GraphBuilder, LOD.DeformedResource->GetBuffer(FHairCardsDeformedResource::Current), ERDGImportedBufferFlags::None).Buffer, ERHIAccess::SRVMask);
+				ResourceAccessFinalizer.AddBuffer(Register(GraphBuilder, LOD.DeformedResource->GetBuffer(FHairCardsDeformedResource::Previous), ERDGImportedBufferFlags::None).Buffer, ERHIAccess::SRVMask);
 			}
 
 			#if RHI_RAYTRACING
