Fixed GroomCache not interpolating when scrubbing between 2 frames in Sequencer.

diff --git a/Engine/Plugins/Runtime/HairStrands/Source/HairStrandsCore/Private/GroomComponent.cpp b/Engine/Plugins/Runtime/HairStrands/Source/HairStrandsCore/Private/GroomComponent.cpp
--- a/Engine/Plugins/Runtime/HairStrands/Source/HairStrandsCore/Private/GroomComponent.cpp
+++ b/Engine/Plugins/Runtime/HairStrands/Source/HairStrandsCore/Private/GroomComponent.cpp
@@ -1006,6 +1006,8 @@ public:
 			InterpolatedFrame = CurrentFrame;
 		}
 
+		bComputeInterpolation = bComputeInterpolation || !FMath::IsNearlyEqual(InterpolationFactor, OutInterpolationFactor, KINDA_SMALL_NUMBER);
+
 		// Do interpolation of vertex positions if needed
 		if (bComputeInterpolation)
 		{
@@ -1017,11 +1019,6 @@ public:
 	{
 		TRACE_CPUPROFILER_EVENT_SCOPE(FGroomCacheBuffers::InterpolateCPU);
 
-		if (FMath::IsNearlyEqual(InInterpolationFactor, InterpolationFactor, KINDA_SMALL_NUMBER))
-		{
-			return;
-		}
-
 		InterpolationFactor = InInterpolationFactor;
 
 		FScopeLock Lock(GetCriticalSection());
@@ -1204,6 +1201,8 @@ public:
 			InterpolatedFrame = *CurrentFramePtr;
 		}
 
+		bComputeInterpolation = bComputeInterpolation || !FMath::IsNearlyEqual(InterpolationFactor, OutInterpolationFactor, KINDA_SMALL_NUMBER);
+
 		// Do interpolation of vertex positions if needed
 		if (bComputeInterpolation)
 		{
