Fix hair strands transition issue.

diff --git a/Engine/Plugins/Runtime/HairStrands/Source/HairStrandsCore/Private/GroomManager.cpp b/Engine/Plugins/Runtime/HairStrands/Source/HairStrandsCore/Private/GroomManager.cpp
--- a/Engine/Plugins/Runtime/HairStrands/Source/HairStrandsCore/Private/GroomManager.cpp
+++ b/Engine/Plugins/Runtime/HairStrands/Source/HairStrandsCore/Private/GroomManager.cpp
@@ -43,6 +43,13 @@ DEFINE_LOG_CATEGORY_STATIC(LogGroomManager, Log, All);
 
 ///////////////////////////////////////////////////////////////////////////////////////////////////
 
+namespace HairTransition
+{
+	void TransitToSRV(FRDGBuilder& GraphBuilder, FRDGBufferSRVRef InBuffer, ERDGPassFlags Flags);
+}
+
+///////////////////////////////////////////////////////////////////////////////////////////////////
+
 void FHairGroupInstance::FCards::FLOD::InitVertexFactory()
 {
 	VertexFactory[0]->InitResources();
@@ -290,6 +297,21 @@ static void RunInternalHairStrandsInterpolation(
 							Instance->Debug.MeshLODIndex,
 							Instance->Guides.DeformedRootResource,
 							Instance->Guides.DeformedResource);
+
+						// Add manual transition for the GPU solver as Niagara does not track properly the RDG buffer, and so doesn't issue the correct transitions
+						if (MeshLODIndex >= 0)
+						{
+							HairTransition::TransitToSRV(GraphBuilder, Register(GraphBuilder, Instance->Guides.DeformedRootResource->LODs[MeshLODIndex].DeformedRootTrianglePosition0Buffer, ERDGImportedBufferFlags::CreateSRV).SRV, ERDGPassFlags::Compute);
+							HairTransition::TransitToSRV(GraphBuilder, Register(GraphBuilder, Instance->Guides.DeformedRootResource->LODs[MeshLODIndex].DeformedRootTrianglePosition1Buffer, ERDGImportedBufferFlags::CreateSRV).SRV, ERDGPassFlags::Compute);
+							HairTransition::TransitToSRV(GraphBuilder, Register(GraphBuilder, Instance->Guides.DeformedRootResource->LODs[MeshLODIndex].DeformedRootTrianglePosition2Buffer, ERDGImportedBufferFlags::CreateSRV).SRV, ERDGPassFlags::Compute);
+
+							if (bGlobalDeformationEnable)
+							{
+								HairTransition::TransitToSRV(GraphBuilder, Register(GraphBuilder, Instance->Guides.DeformedRootResource->LODs[MeshLODIndex].DeformedSamplePositionsBuffer, ERDGImportedBufferFlags::CreateSRV).SRV, ERDGPassFlags::Compute);
+								HairTransition::TransitToSRV(GraphBuilder, Register(GraphBuilder, Instance->Guides.DeformedRootResource->LODs[MeshLODIndex].MeshSampleWeightsBuffer, ERDGImportedBufferFlags::CreateSRV).SRV, ERDGPassFlags::Compute);
+							}
+						}
+						HairTransition::TransitToSRV(GraphBuilder, Register(GraphBuilder, Instance->Guides.DeformedResource->GetPositionOffsetBuffer(FHairStrandsDeformedResource::EFrameType::Current), ERDGImportedBufferFlags::CreateSRV).SRV, ERDGPassFlags::Compute);
 					}
 					else if (bSimulationEnable)
 					{
@@ -301,6 +323,9 @@ static void RunInternalHairStrandsInterpolation(
 							Instance->Debug.MeshLODIndex,
 							nullptr,
 							Instance->Guides.DeformedResource);
+
+						// Add manual transition for the GPU solver as Niagara does not track properly the RDG buffer, and so doesn't issue the correct transitions
+						HairTransition::TransitToSRV(GraphBuilder, Register(GraphBuilder, Instance->Guides.DeformedResource->GetPositionOffsetBuffer(FHairStrandsDeformedResource::EFrameType::Current), ERDGImportedBufferFlags::CreateSRV).SRV, ERDGPassFlags::Compute);
 					}
 				}
 			}
diff --git a/Engine/Plugins/Runtime/HairStrands/Source/HairStrandsCore/Private/GroomResources.cpp b/Engine/Plugins/Runtime/HairStrands/Source/HairStrandsCore/Private/GroomResources.cpp
--- a/Engine/Plugins/Runtime/HairStrands/Source/HairStrandsCore/Private/GroomResources.cpp
+++ b/Engine/Plugins/Runtime/HairStrands/Source/HairStrandsCore/Private/GroomResources.cpp
@@ -35,6 +35,11 @@ static FAutoConsoleVariableRef CVarHairStrandsBulkData_Validation(TEXT("r.HairSt
 
 ////////////////////////////////////////////////////////////////////////////////////
 
+namespace HairTransition
+{
+	void TransitToSRV(FRDGBuilder& GraphBuilder, FRDGBufferSRVRef InBuffer, ERDGPassFlags Flags);	
+}
+
 enum class EHairResourceUsageType : uint8
 {
 	Static,
@@ -582,29 +587,6 @@ void FHairCardsProceduralResource::InternalRelease()
 	CardVoxel.NormalBuffer.Release();
 }
 
-////////////////////////////////////////////////////////////////////////////////////////////////////////////////////////
-namespace HairCardsTransition
-{
-	BEGIN_SHADER_PARAMETER_STRUCT(FTransitionParameters, )
-		SHADER_PARAMETER_RDG_BUFFER_SRV(Buffer, Buffer)
-	END_SHADER_PARAMETER_STRUCT()
-
-	void TransitToSRV(FRDGBuilder& GraphBuilder, FRDGBufferSRVRef InBuffer)
-	{
-		FTransitionParameters* Parameters = GraphBuilder.AllocParameters<FTransitionParameters>();
-		Parameters->Buffer = InBuffer;
-
-		GraphBuilder.AddPass(
-			RDG_EVENT_NAME("HairCardsDeformedBufferTransitionPass"),
-			Parameters,
-			ERDGPassFlags::Raster | ERDGPassFlags::SkipRenderPass | ERDGPassFlags::NeverCull,
-			[Parameters](FRHICommandList& RHICmdList)
-			{
-				// Nothing to do, we just want GraphBuilder to force the buffer to be transited to SRV Graphics state.
-			});
-	}
-}
-
 /////////////////////////////////////////////////////////////////////////////////////////
 
 FHairCardsDeformedResource::FHairCardsDeformedResource(const FHairCardsBulkData& InBulkData, bool bInInitializedData) :
@@ -626,8 +608,8 @@ void FHairCardsDeformedResource::InternalAllocate(FRDGBuilder& GraphBuilder)
 
 		// Manually transit to SRVs, in case of the cards are rendered not visible, but still rendered (in shadows for instance). In such a case, the cards deformation pass is not called, and thus the 
 		// buffers are never transit from UAV (clear) to SRV for rasterization. 
-		HairCardsTransition::TransitToSRV(GraphBuilder, Register(GraphBuilder, DeformedPositionBuffer[0], ERDGImportedBufferFlags::CreateSRV).SRV);
-		HairCardsTransition::TransitToSRV(GraphBuilder, Register(GraphBuilder, DeformedPositionBuffer[1], ERDGImportedBufferFlags::CreateSRV).SRV);
+		HairTransition::TransitToSRV(GraphBuilder, Register(GraphBuilder, DeformedPositionBuffer[0], ERDGImportedBufferFlags::CreateSRV).SRV, ERDGPassFlags::Raster);
+		HairTransition::TransitToSRV(GraphBuilder, Register(GraphBuilder, DeformedPositionBuffer[1], ERDGImportedBufferFlags::CreateSRV).SRV, ERDGPassFlags::Raster);
 	}
 }
 
@@ -739,6 +721,7 @@ void FHairStrandsRestResource::InternalAllocate(FRDGBuilder& GraphBuilder)
 	TArray<FVector4> RestOffset;
 	RestOffset.Add(BulkData.GetPositionOffset());
 	InternalCreateVertexBufferRDG<FHairStrandsPositionOffsetFormat>(GraphBuilder, RestOffset, PositionOffsetBuffer, HAIRSTRANDS_RESOUCE_NAME(CurveType, Hair.StrandsRest_PositionOffsetBuffer), EHairResourceUsageType::Static, ERDGInitialDataFlags::None);
+	HairTransition::TransitToSRV(GraphBuilder, Register(GraphBuilder, PositionOffsetBuffer, ERDGImportedBufferFlags::CreateSRV).SRV, ERDGPassFlags::Compute);
 }
 
 void AddHairTangentPass(
diff --git a/Engine/Plugins/Runtime/HairStrands/Source/HairStrandsCore/Private/HairStrandsRendering.cpp b/Engine/Plugins/Runtime/HairStrands/Source/HairStrandsCore/Private/HairStrandsRendering.cpp
--- a/Engine/Plugins/Runtime/HairStrands/Source/HairStrandsCore/Private/HairStrandsRendering.cpp
+++ b/Engine/Plugins/Runtime/HairStrands/Source/HairStrandsCore/Private/HairStrandsRendering.cpp
@@ -39,6 +39,39 @@ static FAutoConsoleVariableRef CVarHairCardsInterpolationType(TEXT("r.HairStrand
 static int32 GHairStrandsTransferPositionOnLODChange = 0;
 static FAutoConsoleVariableRef CVarHairStrandsTransferPositionOnLODChange(TEXT("r.HairStrands.Strands.TransferPrevPos"), GHairStrandsTransferPositionOnLODChange, TEXT("Transfer strands prev. position to current position on LOD switching to avoid large discrepancy causing large motion vector"));
 
+////////////////////////////////////////////////////////////////////////////////////////////////////////////////////////
+
+namespace HairTransition
+{
+	BEGIN_SHADER_PARAMETER_STRUCT(FTransitionParameters, )
+		SHADER_PARAMETER_RDG_BUFFER_SRV(Buffer, Buffer)
+	END_SHADER_PARAMETER_STRUCT()
+
+	void TransitToSRV(FRDGBuilder& GraphBuilder, FRDGBufferSRVRef InBuffer, ERDGPassFlags Flags)
+	{
+		FTransitionParameters* Parameters = GraphBuilder.AllocParameters<FTransitionParameters>();
+		Parameters->Buffer = InBuffer;
+
+		ensure(Flags == ERDGPassFlags::Raster || Flags == ERDGPassFlags::Compute);
+		if (Flags == ERDGPassFlags::Raster)
+		{
+			Flags |= ERDGPassFlags::SkipRenderPass;
+		}
+		Flags |= ERDGPassFlags::NeverCull;
+
+		GraphBuilder.AddPass(
+			RDG_EVENT_NAME("HairTransitionToSRVPass(%s)", Flags == ERDGPassFlags::Raster ? TEXT("Raster") : TEXT("Compute")),
+			Parameters,
+			Flags,
+			[Parameters](FRHICommandList& RHICmdList)
+			{
+				// Nothing to do, we just want GraphBuilder to force the buffer to be transited to SRV Graphics state.
+			});
+	}
+}
+
+////////////////////////////////////////////////////////////////////////////////////////////////////////////////////////
+
 enum class EHairCardsSimulationType
 {
 	None,
@@ -1098,8 +1131,6 @@ void AddHairTangentPass(
 			Parameters,
 			DispatchCount);
 	}
-
-	GraphBuilder.SetBufferAccessFinal(OutTangentBuffer.Buffer, ERHIAccess::SRVMask);
 }
 
 ///////////////////////////////////////////////////////////////////////////////////////////////////
@@ -1494,14 +1525,6 @@ FHairGroupPublicData::FVertexFactoryInput ComputeHairStrandsVertexInputData(cons
 	return InternalComputeHairStrandsVertexInputData(nullptr, Instance);
 }
 
-void AddBufferTransitionToReadablePass(FRDGBuilder& GraphBuilder, FRHIUnorderedAccessView* UAV)
-{
-	AddPass(GraphBuilder, [UAV](FRHICommandList& RHICmdList)
-	{
-		RHICmdList.Transition(FRHITransitionInfo(UAV, ERHIAccess::UAVCompute, ERHIAccess::SRVMask));
-	});
-}
-
 void CreateHairStrandsDebugAttributeBuffer(FRDGBuilder& GraphBuilder, FRDGExternalBuffer* DebugAttributeBuffer, uint32 VertexCount);
 
 void ComputeHairStrandsInterpolation(
