Split hair resources release per geometry type.

diff --git a/Engine/Plugins/Runtime/HairStrands/Source/HairStrandsCore/Private/GroomAsset.cpp b/Engine/Plugins/Runtime/HairStrands/Source/HairStrandsCore/Private/GroomAsset.cpp
--- a/Engine/Plugins/Runtime/HairStrands/Source/HairStrandsCore/Private/GroomAsset.cpp
+++ b/Engine/Plugins/Runtime/HairStrands/Source/HairStrandsCore/Private/GroomAsset.cpp
@@ -796,18 +796,6 @@ bool UGroomAsset::HasGeometryType(EGroomGeometryType Type) const
 	return false;
 }
 
-void UGroomAsset::InitResources()
-{
-	LLM_SCOPE(ELLMTag::Meshes) // This should be a Groom LLM tag, but there is no LLM tag bit left
-
-	InitGuideResources();
-	InitStrandsResources();
-	InitCardsResources();
-	InitMeshesResources();
-
-	bIsInitialized = true;
-}
-
 enum EGroomAssetChangeType
 {
 	GroomChangeType_Interpolation = 1,
@@ -910,49 +898,93 @@ void UGroomAsset::UpdateResource()
 }
 #endif // #if WITH_EDITOR
 
+////////////////////////////////////////////////////////////////////////////////////////////////////
+// Initialize resources
+
+void UGroomAsset::InitResources()
+{
+	LLM_SCOPE(ELLMTag::Meshes) // This should be a Groom LLM tag, but there is no LLM tag bit left
+
+	InitGuideResources();
+	InitStrandsResources();
+	InitCardsResources();
+	InitMeshesResources();
+
+	bIsInitialized = true;
+}
+
+////////////////////////////////////////////////////////////////////////////////////////////////////
+// Release resources
+
 void UGroomAsset::ReleaseResource()
 {
 	bIsInitialized = false;
 	for (uint32 GroupIndex = 0, GroupCount = GetNumHairGroups(); GroupIndex < GroupCount; ++GroupIndex)
 	{
-		FHairGroupData& GroupData = HairGroupsData[GroupIndex];
+		ReleaseGuidesResource(GroupIndex);
+		ReleaseStrandsResource(GroupIndex);
+		ReleaseCardsResource(GroupIndex);
+		ReleaseMeshesResource(GroupIndex);
+	}
+}
 
-		if (GroupData.Guides.IsValid())
-		{
-			InternalReleaseResource(GroupData.Guides.RestResource);
-		}
+void UGroomAsset::ReleaseGuidesResource(uint32 GroupIndex)
+{
+	if (GroupIndex >= uint32(HairGroupsData.Num())) return;
+	FHairGroupData& GroupData = HairGroupsData[GroupIndex];
+	if (GroupData.Guides.IsValid())
+	{
+		InternalReleaseResource(GroupData.Guides.RestResource);
+	}
+}
 
-		if (GroupData.Strands.IsValid())
-		{
-			InternalReleaseResource(GroupData.Strands.RestResource);
-			InternalReleaseResource(GroupData.Strands.ClusterCullingResource);
-			InternalReleaseResource(GroupData.Strands.InterpolationResource);
-			#if RHI_RAYTRACING
-			InternalReleaseResource(GroupData.Strands.RaytracingResource);
-			#endif
-		}
+void UGroomAsset::ReleaseStrandsResource(uint32 GroupIndex)
+{
+	if (GroupIndex >= uint32(HairGroupsData.Num())) return;
+	FHairGroupData& GroupData = HairGroupsData[GroupIndex];
+	if (GroupData.Strands.IsValid())
+	{
+		InternalReleaseResource(GroupData.Strands.RestResource);
+		InternalReleaseResource(GroupData.Strands.ClusterCullingResource);
+		InternalReleaseResource(GroupData.Strands.InterpolationResource);
+		#if RHI_RAYTRACING
+		InternalReleaseResource(GroupData.Strands.RaytracingResource);
+		#endif
+	}
+}
 
-		for (FHairGroupData::FCards::FLOD& LOD : GroupData.Cards.LODs)
-		{
-			InternalReleaseResource(LOD.RestResource);
-			InternalReleaseResource(LOD.InterpolationResource);
-			InternalReleaseResource(LOD.Guides.RestResource);
-			InternalReleaseResource(LOD.Guides.InterpolationResource);
-			#if RHI_RAYTRACING
-			InternalReleaseResource(LOD.RaytracingResource);
-			#endif
-		}
-		for (FHairGroupData::FMeshes::FLOD& LOD : GroupData.Meshes.LODs)
-		{
-			InternalReleaseResource(LOD.RestResource);
-			#if RHI_RAYTRACING
-			InternalReleaseResource(LOD.RaytracingResource);
-			#endif
-		}
+void UGroomAsset::ReleaseCardsResource(uint32 GroupIndex)
+{
+	if (GroupIndex >= uint32(HairGroupsData.Num())) return;
+	FHairGroupData& GroupData = HairGroupsData[GroupIndex];
+	for (FHairGroupData::FCards::FLOD& LOD : GroupData.Cards.LODs)
+	{
+		InternalReleaseResource(LOD.RestResource);
+		InternalReleaseResource(LOD.InterpolationResource);
+		InternalReleaseResource(LOD.Guides.RestResource);
+		InternalReleaseResource(LOD.Guides.InterpolationResource);
+		#if RHI_RAYTRACING
+		InternalReleaseResource(LOD.RaytracingResource);
+		#endif
 	}
+}
 
+void UGroomAsset::ReleaseMeshesResource(uint32 GroupIndex)
+{
+	if (GroupIndex >= uint32(HairGroupsData.Num())) return;
+	FHairGroupData& GroupData = HairGroupsData[GroupIndex];
+	for (FHairGroupData::FMeshes::FLOD& LOD : GroupData.Meshes.LODs)
+	{
+		InternalReleaseResource(LOD.RestResource);
+		#if RHI_RAYTRACING
+		InternalReleaseResource(LOD.RaytracingResource);
+		#endif
+	}
 }
 
+////////////////////////////////////////////////////////////////////////////////////////////////////
+// 
+
 void UGroomAsset::UpdateHairGroupsInfo()
 {
 	const uint32 GroupCount = GetNumHairGroups();
@@ -2053,9 +2085,10 @@ bool UGroomAsset::CacheDerivedData(uint32 GroupIndex)
 	}
 
 	FString DerivedDataKey;
-	bool bSuccess = CacheStrandsData(GroupIndex, DerivedDataKey);
+	bool bSuccess = false;
 
-	// Cache the cards and meshes separately from the strands data
+	// Cache each geometric representation separately
+	bSuccess |= CacheStrandsData(GroupIndex, DerivedDataKey);
 	bSuccess |= CacheCardsGeometry(GroupIndex, DerivedDataKey);
 	bSuccess |= CacheMeshesGeometry(GroupIndex);
 
diff --git a/Engine/Plugins/Runtime/HairStrands/Source/HairStrandsCore/Public/GroomAsset.h b/Engine/Plugins/Runtime/HairStrands/Source/HairStrandsCore/Public/GroomAsset.h
--- a/Engine/Plugins/Runtime/HairStrands/Source/HairStrandsCore/Public/GroomAsset.h
+++ b/Engine/Plugins/Runtime/HairStrands/Source/HairStrandsCore/Public/GroomAsset.h
@@ -507,6 +507,10 @@ public:
 	void UpdateResource();
 #endif
 	void ReleaseResource();
+	void ReleaseGuidesResource(uint32 GroupIndex);
+	void ReleaseStrandsResource(uint32 GroupIndex);
+	void ReleaseCardsResource(uint32 GroupIndex);
+	void ReleaseMeshesResource(uint32 GroupIndex);
 
 	void SetNumGroup(uint32 InGroupCount, bool bResetGroupData=true);
 	bool AreGroupsValid() const;
