Local strands simulation + better collision initialisation + fix bad buffer transitions + small cleanup

diff --git a/Engine/Plugins/FX/Niagara/Source/NiagaraCore/Private/NiagaraCustomVersion.cpp b/Engine/Plugins/FX/Niagara/Source/NiagaraCore/Private/NiagaraCustomVersion.cpp
--- a/Engine/Plugins/FX/Niagara/Source/NiagaraCore/Private/NiagaraCustomVersion.cpp
+++ b/Engine/Plugins/FX/Niagara/Source/NiagaraCore/Private/NiagaraCustomVersion.cpp
@@ -8,4 +8,5 @@ const FGuid FNiagaraCustomVersion::GUID(0xFCF57AFA, 0x50764283, 0xB9A9E658, 0xFF
 // Register the custom version with core
 FCustomVersionRegistration GRegisterNiagaraCustomVersion(FNiagaraCustomVersion::GUID, FNiagaraCustomVersion::LatestVersion, TEXT("NiagaraVer"));
 
-const FGuid FNiagaraCustomVersion::LatestScriptCompileVersion(0xA83D40AB, 0xE1F64D9F, 0x9F5451FB, 0x8C954559);
+const FGuid FNiagaraCustomVersion::LatestScriptCompileVersion(0xE1236282, 0x02E2477E, 0x808C6EC6, 0xAC96EE7F);
+
diff --git a/Engine/Plugins/Runtime/HairStrands/Shaders/Private/NiagaraDataInterfaceHairStrands.ush b/Engine/Plugins/Runtime/HairStrands/Shaders/Private/NiagaraDataInterfaceHairStrands.ush
--- a/Engine/Plugins/Runtime/HairStrands/Shaders/Private/NiagaraDataInterfaceHairStrands.ush
+++ b/Engine/Plugins/Runtime/HairStrands/Shaders/Private/NiagaraDataInterfaceHairStrands.ush
@@ -29,6 +29,9 @@ struct FDIHairStrandsContext
 	float4x4			WorldTransform;
 	float4x4			WorldInverse;
 	float4				WorldRotation;
+	float4x4			BoneTransform;
+	float4x4			BoneInverse;
+	float4				BoneRotation;
 	int					NumStrands;
 	int					StrandSize;
 	int					InterpolationMode;
@@ -38,6 +41,7 @@ struct FDIHairStrandsContext
 	Buffer<float4>		DeformedPositionOffset;
 	bool				ResetSimulation;
 	bool				RestUpdate;
+	bool				LocalSimulation;
 	int					SampleCount;
 	int4				BoundingBoxOffsets;
 	Buffer<float4>		RestSamplePositionsBuffer;
@@ -59,6 +63,9 @@ Buffer<uint>		RootBarycentricCoordinatesBuffer_##NAME;\
 float4x4			WorldTransform_##NAME;\
 float4x4			WorldInverse_##NAME;\
 float4				WorldRotation_##NAME;\
+float4x4			BoneTransform_##NAME;\
+float4x4			BoneInverse_##NAME;\
+float4				BoneRotation_##NAME;\
 int					NumStrands_##NAME;\
 int					StrandSize_##NAME;\
 int					InterpolationMode_##NAME;\
@@ -68,6 +75,7 @@ float3				RestPositionOffset_##NAME;\
 Buffer<float4>		DeformedPositionOffset_##NAME;\
 bool				ResetSimulation_##NAME;\
 bool				RestUpdate_##NAME;\
+bool				LocalSimulation_##NAME;\
 RWBuffer<uint>	    BoundingBoxBuffer_##NAME;\
 int					SampleCount_##NAME;\
 int4				BoundingBoxOffsets_##NAME;\
@@ -91,6 +99,9 @@ DIContext.RootBarycentricCoordinatesBuffer = RootBarycentricCoordinatesBuffer_##
 DIContext.WorldTransform = WorldTransform_##NAME;\
 DIContext.WorldInverse = WorldInverse_##NAME;\
 DIContext.WorldRotation = WorldRotation_##NAME;\
+DIContext.BoneTransform = BoneTransform_##NAME;\
+DIContext.BoneInverse = BoneInverse_##NAME;\
+DIContext.BoneRotation = BoneRotation_##NAME;\
 DIContext.NumStrands = NumStrands_##NAME;\
 DIContext.StrandSize = StrandSize_##NAME;\
 DIContext.InterpolationMode =	InterpolationMode_##NAME;\
@@ -100,6 +111,7 @@ DIContext.RestPositionOffset =	RestPositionOffset_##NAME;\
 DIContext.DeformedPositionOffset = DeformedPositionOffset_##NAME;\
 DIContext.ResetSimulation = ResetSimulation_##NAME;\
 DIContext.RestUpdate = RestUpdate_##NAME;\
+DIContext.LocalSimulation = LocalSimulation_##NAME;\
 RWBuffer<uint> DIContext_BoundingBoxBuffer = BoundingBoxBuffer_##NAME;\
 DIContext.SampleCount = SampleCount_##NAME;\
 DIContext.BoundingBoxOffsets = BoundingBoxOffsets_##NAME;\
@@ -322,6 +334,11 @@ void DIHairStrands_SetNodePosition(in FDIHairStrandsContext DIContext, in bool R
 						   f16tof32(DIContext.RestPositionBuffer[PointOffset+PointNext].xyz) * NodeAlpha;
 	}
 	OutNodePosition = mul(float4( OutNodePosition+DIContext.RestPositionOffset, 1.0), DIContext.WorldTransform).xyz;
+
+	if(DIContext.LocalSimulation)
+	{
+		OutNodePosition = mul(float4( OutNodePosition, 1.0), DIContext.BoneInverse).xyz;
+	}
 }
 
 // Compute the node position by linear interpolation over the points
@@ -701,8 +718,6 @@ void DIHairStrands_UpdatePointPosition(in FDIHairStrandsContext DIContext, in fl
 	const int PointEnd = DIContext.CurvesOffsetsBuffer[StrandIndex + 1];
 	const int PointCount = PointEnd - PointBegin;
 
-	const float3 PositionOffset = DIContext.RestPositionOffset - DIContext.DeformedPositionOffset[0].xyz;
-
 	for (int PointIndex = 0; PointIndex < PointCount; ++PointIndex)
 	{
 		const float PointCoord = (float) (PointIndex) / (PointCount - 1.0);
@@ -800,27 +815,50 @@ void DIHairStrands_UpdateAngularVelocity(in FDIHairStrandsContext DIContext, in
 // Compute rest position
 void DIHairStrands_ComputeRestPosition(in FDIHairStrandsContext DIContext, in float3 NodePosition, out float3 OutRestPosition)
 {	
-	OutRestPosition = mul(float4( NodePosition, 1.0), DIContext.WorldInverse).xyz;
+	OutRestPosition = DIContext.LocalSimulation ? mul(float4( NodePosition, 1.0), DIContext.BoneTransform).xyz : NodePosition;
+	OutRestPosition = mul(float4( OutRestPosition, 1.0), DIContext.WorldInverse).xyz;
 }
 
 // Compute rest orientation
 void DIHairStrands_ComputeRestOrientation(in FDIHairStrandsContext DIContext, in float4 NodeOrientation, out float4 OutRestOrientation)
 {	
-	OutRestOrientation = NormalizeQuat( MultiplyQuat(InverseQuat(DIContext.WorldRotation),NodeOrientation) );
+	OutRestOrientation = DIContext.LocalSimulation ? NormalizeQuat( MultiplyQuat(DIContext.BoneRotation,NodeOrientation) ) : NodeOrientation;
+	OutRestOrientation = NormalizeQuat( MultiplyQuat(InverseQuat(DIContext.WorldRotation),OutRestOrientation) );
 }
 
 // Update Node Position
 void DIHairStrands_AttachNodePosition(in FDIHairStrandsContext DIContext, in float3 RestPosition, out float3 OutNodePosition)
 {
 	OutNodePosition = mul(float4( RestPosition, 1.0), DIContext.WorldTransform).xyz;
+	if(DIContext.LocalSimulation)
+	{
+		OutNodePosition = mul(float4( OutNodePosition, 1.0), DIContext.BoneInverse).xyz;
+	}
 }
 
 // Update Node Orientation
 void DIHairStrands_AttachNodeOrientation(in FDIHairStrandsContext DIContext, in float4 RestOrientation, out float4 OutNodeOrientation)
 {
 	OutNodeOrientation = NormalizeQuat( MultiplyQuat(DIContext.WorldRotation,RestOrientation) );
+	if(DIContext.LocalSimulation)
+	{
+		OutNodeOrientation = NormalizeQuat( MultiplyQuat(InverseQuat(DIContext.BoneRotation),OutNodeOrientation) );
+	}
 }
 
+// Get World Vector
+void DIHairStrands_GetWorldVector(in FDIHairStrandsContext DIContext, in float3 LocalVector, in bool IsPosition, out float3 OutWorldVector)
+{	
+	OutWorldVector = DIContext.LocalSimulation ? mul(float4( LocalVector, (float)IsPosition), DIContext.BoneTransform).xyz : LocalVector;
+}
+
+// Get Local Vector
+void DIHairStrands_GetLocalVector(in FDIHairStrandsContext DIContext, in float3 WorldVector, in bool IsPosition, out float3 OutLocalVector)
+{
+	OutLocalVector = DIContext.LocalSimulation ? mul(float4( WorldVector, (float)IsPosition), DIContext.BoneInverse).xyz : WorldVector;
+}
+
+
 /* -----------------------------------------------------------------
  * Attached the root to the skinned cache
  * -----------------------------------------------------------------
diff --git a/Engine/Plugins/Runtime/HairStrands/Shaders/Private/NiagaraDataInterfacePhysicsAsset.ush b/Engine/Plugins/Runtime/HairStrands/Shaders/Private/NiagaraDataInterfacePhysicsAsset.ush
--- a/Engine/Plugins/Runtime/HairStrands/Shaders/Private/NiagaraDataInterfacePhysicsAsset.ush
+++ b/Engine/Plugins/Runtime/HairStrands/Shaders/Private/NiagaraDataInterfacePhysicsAsset.ush
@@ -261,7 +261,7 @@ float3 DIPhysicsAsset_GetLocalPosition(in FDIPhysicsAssetContext DIContext, in f
 
 // Given a world space position (NodePosition) compute the sphere closest point (position,normal,velocity)
 int DIPhysicsAsset_GetClosestElement(in FDIPhysicsAssetContext DIContext, in float3 NodePosition, out float3 OutClosestPosition, 
-										  out float3 OutClosestNormal, out float OutMinDistance, in float TimeFraction )
+										  out float3 OutClosestNormal, out float OutMinDistance, in float TimeFraction, in bool IsRestSpace )
 {	
 	float MinDistance = MAX_DISTANCE;
 	int ElementIndex = -1;
@@ -272,7 +272,8 @@ int DIPhysicsAsset_GetClosestElement(in FDIPhysicsAssetContext DIContext, in flo
 	const int SpheresEnd = DIContext.ElementOffsets[SPHERE_INDEX+1];
 	for (int SphereIndex = SpheresBegin; SphereIndex < SpheresEnd; ++SphereIndex)
 	{
-		const float3 LocalPosition = DIPhysicsAsset_GetLocalPosition(DIContext,NodePosition,SphereIndex,TimeFraction);
+		const float3 LocalPosition = IsRestSpace ? mul(DIPhysicsAsset_GetRestInverse(DIContext,SphereIndex), float4(NodePosition,1.0)).xyz : 
+						DIPhysicsAsset_GetLocalPosition(DIContext,NodePosition,SphereIndex,TimeFraction);
 
 		DIPhysicsAsset_GetSphereProjection(LocalPosition, float3(0,0,0), DIContext.ElementExtentBuffer[SphereIndex].x, SphereIndex,
 				CollisionPosition, CollisionNormal, ElementIndex, MinDistance);
@@ -282,7 +283,8 @@ int DIPhysicsAsset_GetClosestElement(in FDIPhysicsAssetContext DIContext, in flo
 	const int CapsulesEnd = DIContext.ElementOffsets[CAPSULE_INDEX+1];
 	for (int CapsuleIndex = CapsulesBegin; CapsuleIndex < CapsulesEnd; ++CapsuleIndex)
 	{
-		const float3 LocalPosition = DIPhysicsAsset_GetLocalPosition(DIContext,NodePosition,CapsuleIndex,TimeFraction);
+		const float3 LocalPosition = IsRestSpace ? mul(DIPhysicsAsset_GetRestInverse(DIContext,CapsuleIndex), float4(NodePosition,1.0)).xyz : 
+						DIPhysicsAsset_GetLocalPosition(DIContext,NodePosition,CapsuleIndex,TimeFraction);
 
 		DIPhysicsAsset_GetCapsuleProjection(LocalPosition, DIContext.ElementExtentBuffer[CapsuleIndex].xy, CapsuleIndex,
 				CollisionPosition, CollisionNormal, ElementIndex, MinDistance);
@@ -292,7 +294,8 @@ int DIPhysicsAsset_GetClosestElement(in FDIPhysicsAssetContext DIContext, in flo
 	const int BoxesEnd = DIContext.ElementOffsets[BOX_INDEX+1];
 	for (int BoxIndex = BoxesBegin; BoxIndex < BoxesEnd; ++BoxIndex)
 	{
-		const float3 LocalPosition = DIPhysicsAsset_GetLocalPosition(DIContext,NodePosition,BoxIndex,TimeFraction);
+		const float3 LocalPosition = IsRestSpace ? mul(DIPhysicsAsset_GetRestInverse(DIContext,BoxIndex), float4(NodePosition,1.0)).xyz : 
+						DIPhysicsAsset_GetLocalPosition(DIContext,NodePosition,BoxIndex,TimeFraction);
 
 		DIPhysicsAsset_GetBoxProjection(LocalPosition, DIContext.ElementExtentBuffer[BoxIndex].xyz, BoxIndex,
 				CollisionPosition, CollisionNormal, ElementIndex, MinDistance);
@@ -395,7 +398,7 @@ void DIPhysicsAsset_GetClosestElement(in FDIPhysicsAssetContext DIContext, in fl
 	float3 CollisionNormal = float3(0,0,0);
 	float MinDistance = 0.0;
 
-	OutElementIndex = DIPhysicsAsset_GetClosestElement(DIContext,NodePosition,CollisionPosition,CollisionNormal,MinDistance,TimeFraction);
+	OutElementIndex = DIPhysicsAsset_GetClosestElement(DIContext,NodePosition,CollisionPosition,CollisionNormal,MinDistance,TimeFraction, false);
 }
 
 // Given a world space position (NodePosition) compute the physics asset closest point (position,normal,velocity)
@@ -406,7 +409,7 @@ void DIPhysicsAsset_GetClosestPoint(in FDIPhysicsAssetContext DIContext, in floa
 	float3 CollisionNormal = float3(0,0,0);
 	float MinDistance = 0.0;
 
-	const int ElementIndex = DIPhysicsAsset_GetClosestElement(DIContext,NodePosition,CollisionPosition,CollisionNormal,MinDistance,TimeFraction);
+	const int ElementIndex = DIPhysicsAsset_GetClosestElement(DIContext,NodePosition,CollisionPosition,CollisionNormal,MinDistance,TimeFraction, false);
 
 	OutClosestNormal = float3(0,0,0);
 	OutClosestPosition = float3(0,0,0);
@@ -434,9 +437,7 @@ void DIPhysicsAsset_GetClosestDistance(in FDIPhysicsAssetContext DIContext, in f
 	float3 CollisionNormal = float3(0,0,0);
 	float MinDistance = 0.0;
 
-	OutClosestDistance = 0.0;
-
-	const int ElementIndex = DIPhysicsAsset_GetClosestElement(DIContext,NodePosition,CollisionPosition,CollisionNormal,MinDistance,TimeFraction);
+	const int ElementIndex = DIPhysicsAsset_GetClosestElement(DIContext,NodePosition,CollisionPosition,CollisionNormal,MinDistance,TimeFraction, false);
 
 	OutClosestDistance = 0.0;
 
@@ -452,6 +453,22 @@ void DIPhysicsAsset_GetClosestDistance(in FDIPhysicsAssetContext DIContext, in f
  * -----------------------------------------------------------------
  */
 
+void DIPhysicsAsset_GetRestDistance(in FDIPhysicsAssetContext DIContext, in float3 RestPosition, out float OutRestDistance)
+{
+	float3 CollisionPosition = float3(0,0,0);
+	float3 CollisionNormal = float3(0,0,0);
+	float MinDistance = 0.0;
+
+	const int ElementIndex = DIPhysicsAsset_GetClosestElement(DIContext,RestPosition,CollisionPosition,CollisionNormal,MinDistance, 1.0, true);
+
+	OutRestDistance = 0.0;
+
+	if (ElementIndex != -1)
+	{
+		OutRestDistance = MinDistance;
+	}
+}
+
 void DIPhysicsAsset_GetTexturePoint(in FDIPhysicsAssetContext DIContext, in float3 NodePosition, out int OutElementIndex, out float3 OutTexturePosition)
 {
 	float3 CollisionPosition = float3(0,0,0);
@@ -460,7 +477,7 @@ void DIPhysicsAsset_GetTexturePoint(in FDIPhysicsAssetContext DIContext, in floa
 
 	OutTexturePosition = float3(0,0,0);
 
-	OutElementIndex = DIPhysicsAsset_GetClosestElement(DIContext,NodePosition,CollisionPosition,CollisionNormal,MinDistance, 1.0);
+	OutElementIndex = DIPhysicsAsset_GetClosestElement(DIContext,NodePosition,CollisionPosition,CollisionNormal,MinDistance, 1.0, false);
 
 	const float AxisExtent = max( max(DIContext.BoxExtent.x, DIContext.BoxExtent.y), DIContext.BoxExtent.z );
 	const float3 BoxExtent = float3(AxisExtent,AxisExtent,AxisExtent);
diff --git a/Engine/Plugins/Runtime/HairStrands/Shaders/Private/NiagaraDirectSolver.ush b/Engine/Plugins/Runtime/HairStrands/Shaders/Private/NiagaraDirectSolver.ush
--- a/Engine/Plugins/Runtime/HairStrands/Shaders/Private/NiagaraDirectSolver.ush
+++ b/Engine/Plugins/Runtime/HairStrands/Shaders/Private/NiagaraDirectSolver.ush
@@ -1,124 +1,19 @@
 // Copyright Epic Games, Inc. All Rights Reserved.
 
 /*=============================================================================
-	NiagaraStrandsPhysics.ush
+	NiagaraDirectSolver.ush
 =============================================================================*/
 
 #if GPU_SIMULATION
 
 /* -----------------------------------------------------------------
- * Shared memory for datas that will be accessed within the constraints
+ * Shared memory for datas that will be accessed within the direct solver 
  * -----------------------------------------------------------------
  */
 
-groupshared float3 SharedNodePosition[THREADGROUP_SIZE];
-groupshared float4 SharedNodeOrientation[THREADGROUP_SIZE];
-
-groupshared float SharedInverseMass[THREADGROUP_SIZE];
-groupshared float SharedInverseInertia[THREADGROUP_SIZE];
-
-groupshared float3 SharedPreviousPosition[THREADGROUP_SIZE];
-groupshared float4 SharedPreviousOrientation[THREADGROUP_SIZE];
-
-//groupshared float3x3 SharedMatrixValues[THREADGROUP_SIZE*3];
 groupshared float4 SharedTridiagValues[THREADGROUP_SIZE];
 groupshared float4 SharedTridiagResult[THREADGROUP_SIZE];
 
-/* -----------------------------------------------------------------
- * Shared memory for datas that will be accessed within the constraints
- * -----------------------------------------------------------------
- */
- #define Test_QUATERNION_IDENTITY float4(0,0,0,1)
-#define Test_SMALL_NUMBER 1e-8
-
-float4 Test_NormalizeQuat(in float4 Quat)
-{
-	float SquaredNorm = dot(Quat,Quat);
-	return (SquaredNorm >= Test_SMALL_NUMBER) ? Quat / sqrt(SquaredNorm) : Test_QUATERNION_IDENTITY;
-}
-
-float4 Test_InverseQuat(in float4 Quat)
-{
-	return float4(-Quat.x,-Quat.y,-Quat.z,Quat.w);
-}
-
-float4 Test_MultiplyQuat(in float4 QuatA, in float4 QuatB)
-{
-	return float4(
-        QuatB.xyz * QuatA.w + QuatA.xyz * QuatB.w + cross(QuatA.xyz, QuatB.xyz),
-        QuatA.w * QuatB.w - dot(QuatA.xyz, QuatB.xyz));
-}
-
-float3 Test_RotateVectorByQuat(in float3 Vector, in float4 Quat)
-{
-	float3 T = 2.0 * cross(Quat.xyz,Vector);
-	return Vector + Quat.w * T + cross(Quat.xyz,T);
-}
-
-float4 Test_FindQuatBetweenHelper(in float3 A, in float3 B, in float NormAB)
-{
-	float W = NormAB + dot(A,B);
-	float4 Quat = (W>1e-6*NormAB) ? float4(A.y*B.z-A.z*B.y,A.z*B.x-A.x*B.z,A.x*B.y-A.y*B.x,W) :
-			(abs(A.x) > abs(A.y)) ? float4(-A.z,0.0,A.x,0.0) : float4(0.0,-A.z,A.y,0.0);
-
-	return Test_NormalizeQuat(Quat);
-}
-
-float4 Test_FindQuatBetweenInternal(in float3 An, in float3 Bn)
-{
-	float3 Cn = normalize(An+Bn);
-	return float4(cross(An,Cn),dot(An,Cn));
-}
-
-float4 Test_FindQuatBetweenNormals(in float3 NormalA, in float3 NormalB)
-{
-	float NormAB = 1.0;
-	return Test_FindQuatBetweenHelper(NormalA,NormalB,NormAB);
-}
-
-void Test_ComputeMaterialFrame(in int StrandsSize)
-{
-	const int LocalIndex = (GGroupThreadId.x % StrandsSize);
-	if( LocalIndex == 0 )
-	{
-		float4 NodeQuaternion = SharedNodeOrientation[GGroupThreadId.x];
-		float3 TangentPrev = normalize(SharedPreviousPosition[GGroupThreadId.x+1] - SharedPreviousPosition[GGroupThreadId.x]);
-		float3 TangentNext = TangentPrev;
-
-		for (int EdgeIndex = 1, EdgeEnd = StrandsSize-1; EdgeIndex < EdgeEnd; ++EdgeIndex)
-		{
-			TangentPrev = TangentNext;
-			TangentNext = normalize(SharedPreviousPosition[GGroupThreadId.x+EdgeIndex+1] - SharedPreviousPosition[GGroupThreadId.x+EdgeIndex]);
-
-			NodeQuaternion = Test_NormalizeQuat( Test_MultiplyQuat( Test_FindQuatBetweenNormals(TangentPrev,TangentNext), NodeQuaternion ) );
-			SharedNodeOrientation[GGroupThreadId.x+EdgeIndex] = NodeQuaternion;
-			GroupMemoryBarrier();
-		}
-		GroupMemoryBarrier();
-	}
-	GroupMemoryBarrier();
-}
-
-
-void InverseMatrix(in float3x3 Input, out float3x3 Output)
-{
-	float det = Input[0][0] * (Input[1][1] * Input[2][2] - Input[2][1] * Input[1][2]) -
-				 Input[0][1] * (Input[1][0] * Input[2][2] - Input[1][2] * Input[2][0]) +
-				 Input[0][2] * (Input[1][0] * Input[2][1] - Input[1][1] * Input[2][0]);
-
-	float invdet = 1 / det;
-
-	Output[0][0] = (Input[1][1] * Input[2][2] - Input[2][1] * Input[1][2]) * invdet;
-	Output[0][1] = (Input[0][2] * Input[2][1] - Input[0][1] * Input[2][2]) * invdet;
-	Output[0][2] = (Input[0][1] * Input[1][2] - Input[0][2] * Input[1][1]) * invdet;
-	Output[1][0] = (Input[1][2] * Input[2][0] - Input[1][0] * Input[2][2]) * invdet;
-	Output[1][1] = (Input[0][0] * Input[2][2] - Input[0][2] * Input[2][0]) * invdet;
-	Output[1][2] = (Input[1][0] * Input[0][2] - Input[0][0] * Input[1][2]) * invdet;
-	Output[2][0] = (Input[1][0] * Input[2][1] - Input[2][0] * Input[1][1]) * invdet;
-	Output[2][1] = (Input[2][0] * Input[0][1] - Input[0][0] * Input[2][1]) * invdet;
-	Output[2][2] = (Input[0][0] * Input[1][1] - Input[1][0] * Input[0][1]) * invdet;
-}
-
 void SolveTridiagSerial2()
 {
 	if(GGroupThreadId.x == 0)
@@ -180,7 +75,6 @@ void SolveTridiagSerial()
 	GroupMemoryBarrier();
 }
 
-
 void SolveConjugateResidual()
 {
 	float X = 0.0;
@@ -403,629 +297,4 @@ float SolveTridiag()
 			SharedTridiagValues[LeftIdx].w * SharedTridiagValues[RightIdx].x) / DataValue;
 }
 
-/* -----------------------------------------------------------------
- * Get the strand parent position
- * -----------------------------------------------------------------
- */
-
-void GetParentPosition(const bool IsMobile, out float3 ParentPosition)
-{	
-	ParentPosition = SharedNodePosition[IsMobile ? GGroupThreadId.x-1 : GGroupThreadId.x];
-}
-
-/* -----------------------------------------------------------------
- * Fill the shared datas from the particles
- * -----------------------------------------------------------------
- */
-
-void ExtractPositionSharedDatas(const bool IsMobile, const float3 InParticlePosition, const float3 InParticleVelocity, const float InParticleMass)
-{	
-	SharedNodePosition[GGroupThreadId.x] = InParticlePosition;
-	//SharedLocalVelocity[GGroupThreadId.x] = InParticleVelocity;
-	SharedInverseMass[GGroupThreadId.x] = (InParticleMass != 0.0 && IsMobile) ? 1.0/InParticleMass : 0.0;
-
-	GroupMemoryBarrier(); 
-}
-
-void ExtractOrientationSharedDatas(const bool IsMobile, const float4 InParticleQuaternion, const float3 InParticleOmega, const float InParticleInertia)
-{
-	SharedNodeOrientation[GGroupThreadId.x] = InParticleQuaternion;
-	//SharedLocalOmega[GGroupThreadId.x] = InParticleOmega;
-	SharedInverseInertia[GGroupThreadId.x] = (InParticleInertia != 0.0 && IsMobile) ? 1.0/InParticleInertia : 0.0;
-
-	GroupMemoryBarrier(); 
-}
-
-/* -----------------------------------------------------------------
- * Report the shared datas onto the particles
- * -----------------------------------------------------------------
- */
-
-void ReportPositionSharedDatas(const float InDeltaTime, inout float3 OutParticleVelocity, out float3 OutParticlePosition)
-{
-	const float3 CurrentPosition = SharedNodePosition[GGroupThreadId.x];
-	//OutParticleVelocity += (CurrentPosition-OutParticlePosition) / InDeltaTime;
-	OutParticlePosition = CurrentPosition;
-}
-
-void ReportOrientationSharedDatas(const float InDeltaTime, inout float3 OutParticleOmega, out float4 OutParticleQuaternion)
-{
-	const float4 CurrentQuaternion = SharedNodeOrientation[GGroupThreadId.x];
-	//const float4 DeltaQuat = float4(
-	//		-OutParticleQuaternion.xyz * CurrentQuaternion.w + CurrentQuaternion.xyz * OutParticleQuaternion.w + cross(CurrentQuaternion.xyz, -OutParticleQuaternion.xyz),
-	//		CurrentQuaternion.w * OutParticleQuaternion.w - dot(CurrentQuaternion.xyz, -OutParticleQuaternion.xyz));
-
-	OutParticleQuaternion = CurrentQuaternion;
-
-	//const float AxisLength = length( DeltaQuat.xyz );
-	//const float QuatAngle = 2.0 * atan2(AxisLength,DeltaQuat.w );
-	//OutParticleOmega += DeltaQuat.xyz * QuatAngle / (AxisLength*InDeltaTime);
-}
-
-/* -----------------------------------------------------------------
- * Advect the particles shared datas with forces
- * -----------------------------------------------------------------
- */
-
-void AdvectPositionSharedDatas(const bool IsMobile, const float3 InParticleAcceleration, const float InDeltaTime, inout float3 OutParticleVelocity,  inout float3 OutParticlePosition)
-{	
-	const float3 ParticleAcceleration = IsMobile ? InParticleAcceleration : float3(0,0,0);
-	OutParticleVelocity += ParticleAcceleration * InDeltaTime;
-	OutParticlePosition += OutParticleVelocity * InDeltaTime;
-}
-
-void AdvectOrientationSharedDatas(const float3 InParticleOmega, const float InDeltaTime, inout float4 OutParticleQuaternion)
-{
-	OutParticleQuaternion = OutParticleQuaternion + 0.5 * InDeltaTime * float4(
-        InParticleOmega.xyz * OutParticleQuaternion.w + cross(InParticleOmega.xyz, OutParticleQuaternion.xyz),
-        - dot(InParticleOmega.xyz, OutParticleQuaternion.xyz) );
-
-	OutParticleQuaternion = normalize( OutParticleQuaternion );
-}
-
-/* -----------------------------------------------------------------
- * Hooke spring material
- * -----------------------------------------------------------------
- */
-
-void StoreHookeSpringConfig(const bool IsMobile, const float InMaterialCompliance, out float OutRestLength)
-{
-	if (!IsMobile){
-		const float3 EdgeDirection = SharedNodePosition[GGroupThreadId.x] - SharedNodePosition[GGroupThreadId.x-1];
-		OutRestLength = length(EdgeDirection);
-	}
-	else
-	{
-		OutRestLength = 0.0;
-	}
-}
-
-void ComputeHookeSpringWeight(const float InMaterialCompliance, out float OutMaterialWeight)
-{
-	const float SchurDiagonal = ( SharedInverseMass[GGroupThreadId.x] + SharedInverseMass[GGroupThreadId.x-1] + InMaterialCompliance );
-	OutMaterialWeight = ( SchurDiagonal != 0.0 ) ? 1.0 / SchurDiagonal : 0.0;
-}
-
-void EvalHookeSpringCompliance(const float InDeltaTime, const float InYoungModulus, const float InRodThickness, const float InRestLength, out float OutMaterialCompliance)
-{
-	// Compliance = 1.0 / (k * dt * dt)
-	// with k = Y * A / L 
-	// A is the cross section area (Pi * R * R), L is the rest length and Y is the young modulus
-	OutMaterialCompliance = InRestLength*4.0/(InYoungModulus*PI*InRodThickness*InRodThickness*InDeltaTime*InDeltaTime);
-}
-
-void ResetHookeSpringMultiplier(out float OutMaterialMultiplier)
-{
-	OutMaterialMultiplier = 0;
-}
-
-void UpdateHookeSpringMultiplier(const float InMaterialCompliance, const float InMaterialWeight, const float InRestLength, inout float OutMaterialMultiplier)
-{
-	float3 EdgeDirection = SharedNodePosition[GGroupThreadId.x] - SharedNodePosition[GGroupThreadId.x-1];
-	const float EdgeLength = length(EdgeDirection);
-	EdgeDirection /= EdgeLength;
-
-	// XPBD lagrange multiplier update : dL = -(C+compliance*L) / (dC * invM * dCt + alpha)
-	const float DeltaLambda = -((EdgeLength-InRestLength) + OutMaterialMultiplier * InMaterialCompliance) * InMaterialWeight;
-
-	// L += dL
-	OutMaterialMultiplier += DeltaLambda;
-
-	// XPBD position update : dX = dL * dCt * invM
-	const float3 PositionDelta = EdgeDirection * DeltaLambda;
-
-	// dX += dX
-	SharedNodePosition[GGroupThreadId.x] += PositionDelta * SharedInverseMass[GGroupThreadId.x];
-	SharedNodePosition[GGroupThreadId.x-1] -= PositionDelta * SharedInverseMass[GGroupThreadId.x-1];
-}
-
-void SolveHookeSpringMaterial(const float InMaterialCompliance, const float InMaterialWeight, const float InRestLength, const bool InIsRed, inout float OutMaterialMultiplier)
-{
-	// Process all the red rods
-	if (InIsRed)
-	{	
-		UpdateHookeSpringMultiplier(InMaterialCompliance,InMaterialWeight,InRestLength,OutMaterialMultiplier);
-	}
-	// Process all the black rods
-	GroupMemoryBarrier(); 
-	if (!InIsRed)
-	{
-		UpdateHookeSpringMultiplier(InMaterialCompliance,InMaterialWeight,InRestLength,OutMaterialMultiplier);
-	}
-	GroupMemoryBarrier(); 
-} 
-
-/* -----------------------------------------------------------------
- * Stretch/shear rod material
- * -----------------------------------------------------------------
- */
-
- void ComputeStretchShearWeight(const float InMaterialCompliance, const float InRestLength, out float OutMaterialWeight)
-{
-	//const float SchurDiagonal = ( SharedInverseMass[GGroupThreadId.x] + SharedInverseMass[GGroupThreadId.x-1] ) / (InRestLength*InRestLength) + 
-	//				SharedInverseInertia[GGroupThreadId.x-1]*4.0 + InMaterialCompliance;
-	const float SchurDiagonal = ( SharedInverseMass[GGroupThreadId.x] + SharedInverseMass[GGroupThreadId.x-1] ) / (InRestLength*InRestLength) + 
-					 + InMaterialCompliance;
-	OutMaterialWeight = ( SchurDiagonal != 0.0 ) ? 1.0 / SchurDiagonal : 0.0;
-}
-
-void EvalStretchShearCompliance(const float InDeltaTime, const float InYoungModulus, const float InRodThickness, const float InRestLength, out float OutMaterialCompliance)
-{
-	// Compliance = 1.0 / (k * dt * dt)
-	// with k = L * L * (Y * A / L) (the L*L term before is coming from the fact that our constraint is dL/L and not dL)
-	// A is the cross section area (Pi * R * R), L is the rest length and Y is the young modulus
-	OutMaterialCompliance = 4.0/(InYoungModulus*PI*InRodThickness*InRestLength*InRodThickness*InDeltaTime*InDeltaTime);
-}
-
-void ResetStretchShearMultiplier(out float3 OutMaterialMultiplier)
-{
-	OutMaterialMultiplier = float3(0,0,0);
-}
-
-void UpdateStretchShearMultiplier(const float InMaterialCompliance, const float InMaterialWeight, const float InRestLength,  inout float3 OutMaterialMultiplier)
-{
-	float3 EdgeDirection;
-	float4 q0 = SharedNodeOrientation[GGroupThreadId.x-1];
-	EdgeDirection[0] = 2.0 * (q0.x * q0.z + q0.w * q0.y);
-	EdgeDirection[1] = 2.0 * (q0.y * q0.z - q0.w * q0.x);
-	EdgeDirection[2] = q0.w * q0.w - q0.x * q0.x - q0.y * q0.y + q0.z * q0.z;
-
-	// XPBD lagrange multiplier update : dL = -(C+compliance*L) / (dC * invM * dCt + alpha)
-	const float3 DeltaLambda = -((SharedNodePosition[GGroupThreadId.x] - SharedNodePosition[GGroupThreadId.x-1]) / InRestLength - EdgeDirection + 
-					OutMaterialMultiplier * InMaterialCompliance) * InMaterialWeight;
-
-	// L += dL
-	OutMaterialMultiplier += DeltaLambda;
-
-	// XPBD position update : dX = dL * dCt * invM
-	const float3 PositionDelta  = DeltaLambda/InRestLength;
-
-	// dX += dX
-	SharedNodePosition[GGroupThreadId.x] += PositionDelta * SharedInverseMass[GGroupThreadId.x];
-	SharedNodePosition[GGroupThreadId.x-1] -= PositionDelta * SharedInverseMass[GGroupThreadId.x-1];
-
-	//const float4 qebar = float4(-q0.y, q0.x, -q0.w, q0.z);	
-	//const float4 QuaternionDelta = -2.0 * float4( DeltaLambda.xyz * qebar.w + cross(DeltaLambda.xyz, qebar.xyz),
-	//				- dot(DeltaLambda.xyz, qebar.xyz) );
-
-	//q0 += QuaternionDelta * SharedInverseInertia[GGroupThreadId.x-1];
-	//SharedNodeOrientation[GGroupThreadId.x-1] = normalize(q0);
-}
-
-void SolveStretchShearMaterial(const float InMaterialCompliance, const float InMaterialWeight, const float InRestLength, const bool InIsRed, inout float3 OutMaterialMultiplier)
-{
-	// Process all the red rods
-	if (InIsRed)
-	{	
-		UpdateStretchShearMultiplier(InMaterialCompliance,InMaterialWeight,InRestLength,OutMaterialMultiplier);
-	}
-	// Process all the black rods
-	GroupMemoryBarrier(); 
-	if (!InIsRed)
-	{
-		UpdateStretchShearMultiplier(InMaterialCompliance,InMaterialWeight,InRestLength,OutMaterialMultiplier);
-	}
-	GroupMemoryBarrier(); 
-} 
-
-/* -----------------------------------------------------------------
- * Angular Spring material
- * -----------------------------------------------------------------
- */
-
- void ComputeAngularSpringWeight(const float InMaterialCompliance, const float InRestLength, out float OutMaterialWeight)
-{
-	const float SchurDiagonal = SharedInverseInertia[GGroupThreadId.x-1]*4.0 + InMaterialCompliance;
-	OutMaterialWeight = ( SchurDiagonal != 0.0 ) ? 1.0 / SchurDiagonal : 0.0;
-}
-
-void EvalAngularSpringCompliance(const float InDeltaTime, const float InYoungModulus, const float InRodThickness, const float InRestLength, out float OutMaterialCompliance)
-{
-	// Compliance = 1.0 / (k * dt * dt)
-	// with k = L * L * (Y * A / L) (the L*L term before is coming from the fact that our constraint is dL/L and not dL)
-	// A is the cross section area (Pi * R * R), L is the rest length and Y is the young modulus
-	OutMaterialCompliance = 4.0/(InYoungModulus*PI*InRodThickness*InRestLength*InRodThickness*InDeltaTime*InDeltaTime);
-}
-
-void ResetAngularSpringMultiplier(out float3 OutMaterialMultiplier)
-{
-	OutMaterialMultiplier = float3(0,0,0);
-}
-
-void UpdateAngularSpringMultiplier(const float InMaterialCompliance, const float InMaterialWeight, const float InRestLength,  inout float3 OutMaterialMultiplier)
-{
-	float3 EdgeDirection;
-	float4 q0 = SharedNodeOrientation[GGroupThreadId.x-1];
-	EdgeDirection[0] = 2.0 * (q0.x * q0.z + q0.w * q0.y);
-	EdgeDirection[1] = 2.0 * (q0.y * q0.z - q0.w * q0.x);
-	EdgeDirection[2] = q0.w * q0.w - q0.x * q0.x - q0.y * q0.y + q0.z * q0.z;
-
-	// XPBD lagrange multiplier update : dL = -(C+compliance*L) / (dC * invM * dCt + alpha)
-	const float3 DeltaLambda = -(normalize(SharedNodePosition[GGroupThreadId.x] - SharedNodePosition[GGroupThreadId.x-1]) - EdgeDirection + 
-					OutMaterialMultiplier * InMaterialCompliance) * InMaterialWeight;
-
-	// L += dL
-	OutMaterialMultiplier += DeltaLambda;
-
-	const float4 qebar = float4(-q0.y, q0.x, -q0.w, q0.z);	
-	const float4 QuaternionDelta = -2.0 * float4( DeltaLambda.xyz * qebar.w + cross(DeltaLambda.xyz, qebar.xyz),
-					- dot(DeltaLambda.xyz, qebar.xyz) );
-
-	q0 += QuaternionDelta * SharedInverseInertia[GGroupThreadId.x-1];
-	SharedNodeOrientation[GGroupThreadId.x-1] = normalize(q0);
-}
-
-void SolveAngularSpringMaterial(const float InMaterialCompliance, const float InMaterialWeight, const float InRestLength, const bool InIsRed, inout float3 OutMaterialMultiplier)
-{
-	// Process all the red rods
-	if (InIsRed)
-	{	
-		UpdateAngularSpringMultiplier(InMaterialCompliance,InMaterialWeight,InRestLength,OutMaterialMultiplier);
-	}
-	// Process all the black rods
-	GroupMemoryBarrier(); 
-	if (!InIsRed)
-	{
-		UpdateAngularSpringMultiplier(InMaterialCompliance,InMaterialWeight,InRestLength,OutMaterialMultiplier);
-	}
-	GroupMemoryBarrier(); 
-} 
-
-/* -----------------------------------------------------------------
- * Bend spring rod material
- * -----------------------------------------------------------------
- */
-
-void UpdateBendSpringMultiplier(const float InMaterialCompliance, const float InMaterialWeight, const float InRestLength, const float3 InRestDirection, inout float3 OutMaterialMultiplier, const bool IsFirst)
-{
-	float4 ParentFrame = (IsFirst) ? SharedNodeOrientation[GGroupThreadId.x-1] : SharedNodeOrientation[GGroupThreadId.x-2];
-	float3 CrossProduct = 2.0 * cross(ParentFrame.xyz,InRestDirection);
-	float3 EdgeDirection = InRestDirection + ParentFrame.w * CrossProduct + cross(ParentFrame.xyz,CrossProduct);
-
-	// XPBD lagrange multiplier update : dL = -(C+compliance*L) / (dC * invM * dCt + alpha)
-	const float3 DeltaLambda = -((SharedNodePosition[GGroupThreadId.x] - SharedNodePosition[GGroupThreadId.x-1]) / InRestLength - EdgeDirection / InRestLength + 
-					OutMaterialMultiplier * InMaterialCompliance) * InMaterialWeight;
-
-	// L += dL
-	OutMaterialMultiplier += DeltaLambda;
-
-	// XPBD position update : dX = dL * dCt * invM
-	const float3 PositionDelta  = DeltaLambda/InRestLength;
-
-	// dX += dX
-	SharedNodePosition[GGroupThreadId.x] += PositionDelta * SharedInverseMass[GGroupThreadId.x];
-	SharedNodePosition[GGroupThreadId.x-1] -= PositionDelta * SharedInverseMass[GGroupThreadId.x-1];
-}
-
-void SolveBendSpringMaterial(const float InMaterialCompliance, const float InMaterialWeight, const float InRestLength, const float3 InRestDirection, const bool InIsRed, inout float3 OutMaterialMultiplier, const int InStrandSize)
-{
-	const bool IsFirst = (GGroupThreadId.x % InStrandSize) == 1;
-	// Process all the red rods
-	if (InIsRed)
-	{	
-		UpdateBendSpringMultiplier(InMaterialCompliance,InMaterialWeight,InRestLength,InRestDirection,OutMaterialMultiplier,IsFirst);
-	}
-	// Process all the black rods
-	GroupMemoryBarrier(); 
-	if (!InIsRed)
-	{
-		UpdateBendSpringMultiplier(InMaterialCompliance,InMaterialWeight,InRestLength,InRestDirection,OutMaterialMultiplier,IsFirst);
-	}
-	GroupMemoryBarrier(); 
-} 
-
-
-/* -----------------------------------------------------------------
- * Parallel transport
- * -----------------------------------------------------------------
- */
-
-float4 MultiplyQuat_tmp(in float4 QuatA, in float4 QuatB)
-{
-	return float4(
-        QuatB.xyz * QuatA.w + QuatA.xyz * QuatB.w + cross(QuatA.xyz, QuatB.xyz),
-        QuatA.w * QuatB.w - dot(QuatA.xyz, QuatB.xyz));
-}
-
-float4 FindQuatBetweenHelper_tmp(in float3 A, in float3 B, in float NormAB)
-{
-	float W = NormAB + dot(A,B);
-	float4 Quat = (W>1e-6*NormAB) ? float4(A.y*B.z-A.z*B.y,A.z*B.x-A.x*B.z,A.x*B.y-A.y*B.x,W) :
-			(abs(A.x) > abs(A.y)) ? float4(-A.z,0.0,A.x,0.0) : float4(0.0,-A.z,A.y,0.0);
-	return normalize(Quat);
-}
-
- void ComputeBishopFrame(const int StrandSize)
-{
-	const bool IsMobile = (GGroupThreadId.x % StrandSize) != 0;
-	if (!IsMobile)
-	{
-		float4 NodeQuaternion = float4(0,0,0,1);//SharedNodeOrientation[GGroupThreadId.x];
-		float3 TangentNext = float3(0,0,1);
-		float3 TangentPrev = TangentNext;
-
-		float3 PositionNext = SharedNodePosition[GGroupThreadId.x];
-		float3 PositionPrev = PositionNext;
-
-		for (int EdgeIndex = 0, EdgeEnd = StrandSize-1; EdgeIndex < EdgeEnd; ++EdgeIndex)
-		{
-			PositionPrev = PositionNext;
-			PositionNext = SharedNodePosition[GGroupThreadId.x+EdgeIndex+1];
-
-			TangentPrev = TangentNext;
-			TangentNext = normalize(PositionNext - PositionPrev);
-			NodeQuaternion = normalize( MultiplyQuat_tmp( FindQuatBetweenHelper_tmp(TangentPrev,TangentNext,1.0), NodeQuaternion ) );
-			SharedNodeOrientation[GGroupThreadId.x+EdgeIndex] = NodeQuaternion;
-		}
-	}
-	GroupMemoryBarrier(); 
-}
-/* -----------------------------------------------------------------
- * Bend/twist rod material
- * -----------------------------------------------------------------
- */
-
- void ComputeBendTwistWeight(const float InMaterialCompliance, const float InRestLength, out float OutMaterialWeight)
-{
-	const float SchurDiagonal = ( SharedInverseInertia[GGroupThreadId.x] + SharedInverseInertia[GGroupThreadId.x-1] ) * 4.0 / (InRestLength*InRestLength) + InMaterialCompliance;
-	OutMaterialWeight = ( SchurDiagonal != 0.0 ) ? 1.0 / SchurDiagonal : 0.0;
-}
-
-void EvalBendTwistCompliance(const float InDeltaTime, const float InYoungModulus, const float InRodThickness, const float InRestLength, out float OutMaterialCompliance)
-{
-	// Compliance = 1.0 / (k * dt * dt)
-	// with k = (Y * I * L) 
-	// A is the polar moment of inertia (Pi * R * R * R * R / 4), L is the rest length and Y is the young modulus
-	OutMaterialCompliance = 64.0/(InYoungModulus*PI*InRestLength*InRodThickness*InRodThickness*InRodThickness*InRodThickness*InDeltaTime*InDeltaTime);
-}
-
-void ResetBendTwistMultiplier(out float3 OutMaterialMultiplier)
-{
-	OutMaterialMultiplier = float3(0,0,0);
-}
-
-void UpdateBendTwistMultiplier(const float InMaterialCompliance, const float InMaterialWeight, const float InRestLength, const float4 InRestDarboux, inout float3 OutMaterialMultiplier)
-{
-	float4 q0 = SharedNodeOrientation[GGroupThreadId.x-1];
-	float4 q1 = SharedNodeOrientation[GGroupThreadId.x];
-	const float4 CurrentDarboux = float4(
-			q1.xyz * q0.w - q0.xyz * q1.w + cross(-q0.xyz, q1.xyz),
-			q0.w * q1.w - dot(-q0.xyz, q1.xyz));
-	const float InverseRestLength = 2.0 / InRestLength;
-
-	const float4 OmegaPlus = (CurrentDarboux + InRestDarboux) * InverseRestLength;
-	const float4 OmegaMinus = (CurrentDarboux - InRestDarboux) * InverseRestLength;
-
-	float4 DeltaOmega = (dot(OmegaPlus,OmegaPlus) > dot(OmegaMinus,OmegaMinus) ) ? OmegaMinus : OmegaPlus;
-	DeltaOmega.xyz += OutMaterialMultiplier * InMaterialCompliance;  
-	DeltaOmega.w = 0;
-	DeltaOmega *= InMaterialWeight;
-
-	OutMaterialMultiplier -= DeltaOmega.xyz;
-
-	SharedNodeOrientation[GGroupThreadId.x-1] += SharedInverseInertia[GGroupThreadId.x-1] *  float4(
-			DeltaOmega.xyz * q1.w + cross(q1.xyz, DeltaOmega.xyz), - dot(q1.xyz, DeltaOmega.xyz)) * InverseRestLength;
-	SharedNodeOrientation[GGroupThreadId.x] -= SharedInverseInertia[GGroupThreadId.x] *  float4(
-			DeltaOmega.xyz * q0.w + cross(q0.xyz, DeltaOmega.xyz), - dot(q0.xyz, DeltaOmega.xyz)) * InverseRestLength;
-
-	SharedNodeOrientation[GGroupThreadId.x-1] = normalize(SharedNodeOrientation[GGroupThreadId.x-1]);
-	SharedNodeOrientation[GGroupThreadId.x] = normalize(SharedNodeOrientation[GGroupThreadId.x]);
-}
-
-/*void UpdateBendTwistMultiplier(const float InMaterialCompliance, const float InMaterialWeight, const float InRestLength, const float4 InRestDarboux, inout float3 OutMaterialMultiplier)
-{
-	float4 q0 = SharedNodeOrientation[GGroupThreadId.x-1];
-	float4 q1 = SharedNodeOrientation[GGroupThreadId.x];
-	const float4 CurrentDarboux = float4(
-			q1.xyz * q0.w - q0.xyz * q1.w + cross(-q0.xyz, q1.xyz),
-			q0.w * q1.w - dot(-q0.xyz, q1.xyz));
-	const float InverseRestLength = 2.0 / InRestLength;
-
-	float3 DeltaOmega = (CurrentDarboux.xyz / CurrentDarboux.w - InRestDarboux.xyz / InRestDarboux.w) * InverseRestLength;
-	DeltaOmega.xyz += OutMaterialMultiplier * InMaterialCompliance;  
-	DeltaOmega *= InMaterialWeight;
-
-	OutMaterialMultiplier -= DeltaOmega.xyz;
-
-	DeltaOmega = DeltaOmega - CurrentDarboux.xyz * dot(CurrentDarboux.xyz,DeltaOmega); 
-
-	const float qtu = dot(q0,q1);
-	const float dto = dot(CurrentDarboux.xyz, DeltaOmega.xyz);
-
-	const float4 q1tmp = qtu * q1; 
-	const float4 q0tmp = qtu * q0; 
-
-	const float afac = dot(CurrentDarboux.xyz,DeltaOmega);
-
-	SharedNodeOrientation[GGroupThreadId.x-1] += SharedInverseInertia[GGroupThreadId.x-1] * ( float4(
-			DeltaOmega.xyz * q1tmp.w + cross(q1tmp.xyz, DeltaOmega.xyz), - dot(q1tmp.xyz, DeltaOmega.xyz)) + dto * q1) * InverseRestLength;
-	SharedNodeOrientation[GGroupThreadId.x] -= SharedInverseInertia[GGroupThreadId.x] *  ( float4(
-			DeltaOmega.xyz * q0tmp.w + cross(q0tmp.xyz, DeltaOmega.xyz), - dot(q0tmp.xyz, DeltaOmega.xyz)) - dto * q0) * InverseRestLength;
-
-	SharedNodeOrientation[GGroupThreadId.x-1] = normalize(SharedNodeOrientation[GGroupThreadId.x-1]);
-	SharedNodeOrientation[GGroupThreadId.x] = normalize(SharedNodeOrientation[GGroupThreadId.x]);
-}*/
-
-void SolveBendTwistMaterial(const float InMaterialCompliance, const float InMaterialWeight, const float InRestLength, const float4 InRestDarboux, const bool InIsRed, inout float3 OutMaterialMultiplier)
-{
-	// Process all the red rods
-	if (InIsRed)
-	{	
-		UpdateBendTwistMultiplier(InMaterialCompliance,InMaterialWeight,InRestLength,InRestDarboux,OutMaterialMultiplier);
-	}
-	// Process all the black rods
-	GroupMemoryBarrier(); 
-	if (!InIsRed)
-	{
-		UpdateBendTwistMultiplier(InMaterialCompliance,InMaterialWeight,InRestLength,InRestDarboux,OutMaterialMultiplier);
-	}
-	GroupMemoryBarrier(); 
-} 
-
-/* -----------------------------------------------------------------
- * Follow the leader constraint
- * -----------------------------------------------------------------
- */
-
- void SolveFollowTheLeaderConstraint(const int InBlockSize, const float InRestLength, const float InDampening, const float InDeltaTime, const float InBending)
-{
-	const bool IsMobile = (GGroupThreadId.x % InBlockSize) != 0;
-	if (!IsMobile)
-	{
-		const float BendingLimit = 2*InBending*InRestLength;
-		for (uint i = GGroupThreadId.x+1, end = GGroupThreadId.x+InBlockSize; i < end; ++i)
-		{
-			if (i > GGroupThreadId.x+1)
-			{
-				float3 EdgeDirection = SharedNodePosition[i] - SharedNodePosition[i-2];
-				const float EdgeLength = length(EdgeDirection);
-				EdgeDirection /= EdgeLength;
-				if (EdgeLength<=BendingLimit)
-				{
-					// Comutes dL
-					const float DeltaLambda = -(EdgeLength-BendingLimit);
-					const float3 DeltaPosition = EdgeDirection * DeltaLambda;
-
-					// dX += dX
-					SharedNodePosition[i] += DeltaPosition;
-
-					const float3 DeltaVelocity = DeltaPosition / InDeltaTime;
-
-					// dV += dV
-					//SharedLocalVelocity[i] += SharedInverseMass[i] * DeltaVelocity / ( SharedInverseMass[i] + InDampening * SharedInverseMass[i-2]);
-					//SharedLocalVelocity[i-2] -= InDampening * SharedInverseMass[i-2] * DeltaVelocity / ( SharedInverseMass[i] + InDampening * SharedInverseMass[i-2]);
-				}
-				GroupMemoryBarrier(); 
-			}
-
-
-			float3 EdgeDirection = SharedNodePosition[i] - SharedNodePosition[i-1];
-			const float EdgeLength = length(EdgeDirection);
-			EdgeDirection /= EdgeLength;
-
-			// Comutes dL
-			const float DeltaLambda = -(EdgeLength-InRestLength);
-			const float3 DeltaPosition = EdgeDirection * DeltaLambda;
-
-			// dX += dX
-			SharedNodePosition[i] += DeltaPosition;
-
-			//SharedNodePosition[i] += SharedInverseMass[i] * DeltaPosition / ( SharedInverseMass[i] + InDampening * SharedInverseMass[i-1]);
-			//SharedNodePosition[i-1] -= InDampening * SharedInverseMass[i-1] * DeltaPosition / ( SharedInverseMass[i] + InDampening * SharedInverseMass[i-1]);
-
-			const float3 DeltaVelocity = DeltaPosition / InDeltaTime;
-
-			// dV += dV
-			//SharedLocalVelocity[i] += SharedInverseMass[i] * DeltaVelocity / ( SharedInverseMass[i] + InDampening * SharedInverseMass[i-1]);
-			//SharedLocalVelocity[i-1] -= InDampening * SharedInverseMass[i-1] * DeltaVelocity / ( SharedInverseMass[i] + InDampening * SharedInverseMass[i-1]);
-
-			GroupMemoryBarrier(); 
-		}
-	}
-}
-
-/* -----------------------------------------------------------------
- * Static volume collisions against sphere, box...
- * -----------------------------------------------------------------
- */
-
- float3 ProcessCollisionData( const float3 InDeltaPosition, const float3 InDeltaVelocity, const float3 InCollisionNormal, const float InThicknessSum,
-							const float InStaticFriction, const float InKineticFriction, const bool IsProjection )
- {
-	const float DeltaVNormal = dot(InDeltaVelocity,InCollisionNormal);
-	const float DeltaPNormal = dot(InDeltaPosition,InCollisionNormal)-InThicknessSum;
-
-	if (DeltaPNormal < 0.0)
-	{
-		const float DeltaFNormal = IsProjection ? DeltaPNormal : DeltaVNormal;
-		const float NormDeltaFN = abs(DeltaFNormal);
-
-		const float3 CollisionTangent = InDeltaVelocity - DeltaVNormal * InCollisionNormal;
-		const float TangentLength = length(CollisionTangent);
-
-		float3 CollisionTangentA = (TangentLength > 0.0) ? CollisionTangent/TangentLength : (abs(InCollisionNormal.x-1.0f) > 0.0) ? float3(1,0,0) : 
-																					        (abs(InCollisionNormal.y-1.0f) > 0.0) ? float3(0,1,0) :
-																						     float3(0,0,1);
-		float3 CollisionTangentB = cross(InCollisionNormal,CollisionTangentA);
-		CollisionTangentA = cross(CollisionTangentB,InCollisionNormal);
-
-		const float TangentLengthA = length(CollisionTangentA);
-		const float TangentLengthB = length(CollisionTangentB);
-
-		CollisionTangentA = (TangentLengthA > 0.0) ? CollisionTangentA/TangentLengthA : float3(0,0,0);
-		CollisionTangentB = (TangentLengthB > 0.0) ? CollisionTangentB/TangentLengthB : float3(0,0,0);
-
-		const float DeltaVTangentA = dot( InDeltaVelocity, CollisionTangentA);
-		const float DeltaVTangentB = dot( InDeltaVelocity, CollisionTangentB);
-
-		const float NormDeltaVTA = abs(DeltaVTangentA);
-		const float NormDeltaVTB = abs(DeltaVTangentB);
-
-		const float AlphaTangentA = ( NormDeltaVTA < InStaticFriction * NormDeltaFN ) ? 1.0 : (InKineticFriction > 0.0) ? min(InKineticFriction*NormDeltaFN/NormDeltaVTA, 1.0) : 0.0;
-		const float AlphaTangentB = ( NormDeltaVTB < InStaticFriction * NormDeltaFN ) ? 1.0 : (InKineticFriction > 0.0) ? min(InKineticFriction*NormDeltaFN/NormDeltaVTB, 1.0) : 0.0;
-
-		return DeltaFNormal * InCollisionNormal + AlphaTangentA * CollisionTangentA * DeltaVTangentA + AlphaTangentB * CollisionTangentB * DeltaVTangentB;
-	}
-	return float3(0,0,0);
- }
-
-/* -----------------------------------------------------------------
- * Static volume collisions against sphere, box...
- * -----------------------------------------------------------------
- */
-
-void DetectSphereVolumeCollision(const float3 InParticlePosition, const float3 InSphereCenter, const float InSphereRadius, 
-								 out float3 OutCollisionPosition, out float3 OutCollisionVelocity, inout float3 OutCollisionNormal)   
-{
-	OutCollisionNormal = InParticlePosition - InSphereCenter;
-	const float ParticleDistance = length(OutCollisionNormal);
-	const float CollisionDepth = ParticleDistance-InSphereRadius;
-	
-	OutCollisionNormal = (ParticleDistance != 0.0) ? OutCollisionNormal/ParticleDistance : float3(0,0,0);
-	OutCollisionPosition = InParticlePosition - OutCollisionNormal * CollisionDepth;
-
-	// Compute the sdf velocity based on the transform deriv
-	OutCollisionVelocity = float3(0,0,0);
-}  
-
-/* -----------------------------------------------------------------
- * Static collision constraint
- * -----------------------------------------------------------------
- */
-
-void SolveStaticCollisionConstraint(const float InParticleRadius, const bool InIsMobile, inout float3 OutParticlePosition, inout float3 OutParticleVelocity, const float3 InCollisionPosition, const float3 InCollisionVelocity, const float3 InCollisionNormal, 
-				const float InStaticFriction, const float InKineticFriction, const float InDeltaTime, const bool IsProjection )
-{
-	if(InIsMobile)
-	{
-		const float3 DeltaPosition = OutParticlePosition - InCollisionPosition;
-		const float3 DeltaVelocity = (OutParticleVelocity - InCollisionVelocity) * InDeltaTime;
-		const float ThicknessSum = InParticleRadius;
-
-		const float3 PositionDisplacement = ProcessCollisionData(DeltaPosition,DeltaVelocity,InCollisionNormal,ThicknessSum,InStaticFriction,InKineticFriction,IsProjection);
-		OutParticlePosition -= PositionDisplacement;
-		OutParticleVelocity -= PositionDisplacement / InDeltaTime;
-	}
-}
-
 #endif //GPU_SIMULATION
\ No newline at end of file
diff --git a/Engine/Plugins/Runtime/HairStrands/Source/HairStrandsCore/Private/GroomComponent.cpp b/Engine/Plugins/Runtime/HairStrands/Source/HairStrandsCore/Private/GroomComponent.cpp
--- a/Engine/Plugins/Runtime/HairStrands/Source/HairStrandsCore/Private/GroomComponent.cpp
+++ b/Engine/Plugins/Runtime/HairStrands/Source/HairStrandsCore/Private/GroomComponent.cpp
@@ -1357,7 +1357,7 @@ void UGroomComponent::UpdateHairSimulation()
 			if (!NiagaraComponent)
 			{
 				NiagaraComponent = NewObject<UNiagaraComponent>(this, NAME_None, RF_Transient);
-				NiagaraComponent->SetVisibleFlag(false);
+				NiagaraComponent->SetVisibleFlag(SimulationSettings.SimulationSetup.bDebugSimulation);
 			}
 			if (GetWorld() && GetWorld()->bIsWorldInitialized)
 			{
@@ -1393,7 +1393,7 @@ void UGroomComponent::SetGroomAsset(UGroomAsset* Asset)
 	SetGroomAsset(Asset, BindingAsset);
 }
 
-void UGroomComponent::SetGroomAsset(UGroomAsset* Asset, UGroomBindingAsset* InBinding)
+void UGroomComponent::SetGroomAsset(UGroomAsset* Asset, UGroomBindingAsset* InBinding, const bool bUpdateSimulation)
 {
 	ReleaseResources();
 	if (Asset && Asset->IsValid())
@@ -1439,7 +1439,7 @@ void UGroomComponent::SetGroomAsset(UGroomAsset* Asset, UGroomBindingAsset* InBi
 	}
 
 	UpdateHairGroupsDesc();
-	UpdateHairSimulation();
+	if(bUpdateSimulation) UpdateHairSimulation();
 	if (!GroomAsset || !GroomAsset->IsValid())
 	{
 		return;
@@ -2713,7 +2713,7 @@ void UGroomComponent::PostLoad()
 	}
 
 	// This call will handle the GroomAsset properly if it's still being loaded
-	SetGroomAsset(GroomAsset, BindingAsset);
+	SetGroomAsset(GroomAsset, BindingAsset, false);
 
 #if WITH_EDITOR
 	if (GroomAsset && !bIsGroomAssetCallbackRegistered)
@@ -2931,6 +2931,23 @@ void UGroomComponent::TickAtThisTime(const float Time, bool bInIsRunning, bool b
 	}
 }
 
+void UGroomComponent::BuildSimulationTransform(FTransform& SimulationTransform) const
+{
+	SimulationTransform = FTransform::Identity;
+	if (SimulationSettings.SimulationSetup.bLocalSimulation)
+	{
+		const USkeletalMeshComponent* SkeletelMeshComponent = Cast<const USkeletalMeshComponent>(RegisteredMeshComponent);
+		if (SkeletelMeshComponent && SkeletelMeshComponent->SkeletalMesh && !SimulationSettings.SimulationSetup.LocalBone.IsEmpty())
+		{
+			const FReferenceSkeleton& RefSkeleton = SkeletelMeshComponent->SkeletalMesh->GetRefSkeleton();
+			const FName BoneName(SimulationSettings.SimulationSetup.LocalBone);
+			const int32 BoneIndex = RefSkeleton.FindBoneIndex(BoneName);
+
+			SimulationTransform = SkeletelMeshComponent->GetBoneTransform(BoneIndex);
+		}
+	}
+}
+
 void UGroomComponent::TickComponent(float DeltaTime, enum ELevelTick TickType, FActorComponentTickFunction *ThisTickFunction)
 {
 	Super::TickComponent(DeltaTime, TickType, ThisTickFunction);	
@@ -2965,23 +2982,26 @@ void UGroomComponent::TickComponent(float DeltaTime, enum ELevelTick TickType, F
 	bResetSimulation = bInitSimulation;
 	if (!bInitSimulation)
 	{
-		if (USkeletalMeshComponent* ParentComp = Cast<USkeletalMeshComponent>(GetAttachParent()))
+		if (!SimulationSettings.SimulationSetup.bLocalSimulation)
 		{
-			if (ParentComp->GetNumBones() > 0)
+			if (USkeletalMeshComponent* ParentComp = Cast<USkeletalMeshComponent>(GetAttachParent()))
 			{
-				const int32 BoneIndex = FMath::Min(1, ParentComp->GetNumBones() - 1);
-				const FMatrix NextBoneMatrix = ParentComp->GetBoneMatrix(BoneIndex);
-
-				const float BoneDistance = FVector::DistSquared(PrevBoneMatrix.GetOrigin(), NextBoneMatrix.GetOrigin());
-				if (ParentComp->GetTeleportDistanceThreshold() > 0.0 && BoneDistance >
-					ParentComp->GetTeleportDistanceThreshold() * ParentComp->GetTeleportDistanceThreshold())
+				if (ParentComp->GetNumBones() > 0)
 				{
-					bResetSimulation = true;
+					const int32 BoneIndex = FMath::Min(1, ParentComp->GetNumBones() - 1);
+					const FMatrix NextBoneMatrix = ParentComp->GetBoneMatrix(BoneIndex);
+
+					const float BoneDistance = FVector::DistSquared(PrevBoneMatrix.GetOrigin(), NextBoneMatrix.GetOrigin());
+					if (SimulationSettings.SimulationSetup.TeleportDistance > 0.0 && BoneDistance >
+						SimulationSettings.SimulationSetup.TeleportDistance * SimulationSettings.SimulationSetup.TeleportDistance)
+					{
+						bResetSimulation = true;
+					}
+					PrevBoneMatrix = NextBoneMatrix;
 				}
-				PrevBoneMatrix = NextBoneMatrix;
 			}
 		}
-		bResetSimulation |= SimulationSettings.bResetSimulation;
+		bResetSimulation |= SimulationSettings.SimulationSetup.bResetSimulation;
 	}
 	bInitSimulation = false;
 
diff --git a/Engine/Plugins/Runtime/HairStrands/Source/HairStrandsCore/Private/Niagara/NiagaraDataInterfaceHairStrands.cpp b/Engine/Plugins/Runtime/HairStrands/Source/HairStrandsCore/Private/Niagara/NiagaraDataInterfaceHairStrands.cpp
--- a/Engine/Plugins/Runtime/HairStrands/Source/HairStrandsCore/Private/Niagara/NiagaraDataInterfaceHairStrands.cpp
+++ b/Engine/Plugins/Runtime/HairStrands/Source/HairStrandsCore/Private/Niagara/NiagaraDataInterfaceHairStrands.cpp
@@ -92,6 +92,9 @@ static const FName UpdateAngularVelocityName(TEXT("UpdateAngularVelocity"));
 
 //------------------------------------------------------------------------------------------------------------
 
+static const FName GetLocalVectorName(TEXT("GetLocalVector"));
+static const FName GetWorldVectorName(TEXT("GetWorldVector"));
+
 static const FName AttachNodePositionName(TEXT("AttachNodePosition"));
 static const FName AttachNodeOrientationName(TEXT("AttachNodeOrientation"));
 
@@ -166,10 +169,15 @@ static const FName GetSampleStateName(TEXT("GetSampleState"));
 
 const FString UNiagaraDataInterfaceHairStrands::NumStrandsName(TEXT("NumStrands_"));
 const FString UNiagaraDataInterfaceHairStrands::StrandSizeName(TEXT("StrandSize_"));
+
 const FString UNiagaraDataInterfaceHairStrands::WorldTransformName(TEXT("WorldTransform_"));
 const FString UNiagaraDataInterfaceHairStrands::WorldInverseName(TEXT("WorldInverse_"));
 const FString UNiagaraDataInterfaceHairStrands::WorldRotationName(TEXT("WorldRotation_"));
 
+const FString UNiagaraDataInterfaceHairStrands::BoneTransformName(TEXT("BoneTransform_"));
+const FString UNiagaraDataInterfaceHairStrands::BoneInverseName(TEXT("BoneInverse_"));
+const FString UNiagaraDataInterfaceHairStrands::BoneRotationName(TEXT("BoneRotation_"));
+
 const FString UNiagaraDataInterfaceHairStrands::DeformedPositionBufferName(TEXT("DeformedPositionBuffer_"));
 const FString UNiagaraDataInterfaceHairStrands::CurvesOffsetsBufferName(TEXT("CurvesOffsetsBuffer_"));
 const FString UNiagaraDataInterfaceHairStrands::RestPositionBufferName(TEXT("RestPositionBuffer_"));
@@ -177,6 +185,7 @@ const FString UNiagaraDataInterfaceHairStrands::RestPositionBufferName(TEXT("Res
 const FString UNiagaraDataInterfaceHairStrands::ResetSimulationName(TEXT("ResetSimulation_"));
 const FString UNiagaraDataInterfaceHairStrands::InterpolationModeName(TEXT("InterpolationMode_"));
 const FString UNiagaraDataInterfaceHairStrands::RestUpdateName(TEXT("RestUpdate_"));
+const FString UNiagaraDataInterfaceHairStrands::LocalSimulationName(TEXT("LocalSimulation_"));
 const FString UNiagaraDataInterfaceHairStrands::RootBarycentricCoordinatesName(TEXT("RootBarycentricCoordinatesBuffer_"));
 
 const FString UNiagaraDataInterfaceHairStrands::RestRootOffsetName(TEXT("RestRootOffset_"));
@@ -220,6 +229,10 @@ struct FNDIHairStrandsParametersName
 		WorldInverseName = UNiagaraDataInterfaceHairStrands::WorldInverseName + Suffix;
 		WorldRotationName = UNiagaraDataInterfaceHairStrands::WorldRotationName + Suffix;
 
+		BoneTransformName = UNiagaraDataInterfaceHairStrands::BoneTransformName + Suffix;
+		BoneInverseName = UNiagaraDataInterfaceHairStrands::BoneInverseName + Suffix;
+		BoneRotationName = UNiagaraDataInterfaceHairStrands::BoneRotationName + Suffix;
+
 		DeformedPositionBufferName = UNiagaraDataInterfaceHairStrands::DeformedPositionBufferName + Suffix;
 		CurvesOffsetsBufferName = UNiagaraDataInterfaceHairStrands::CurvesOffsetsBufferName + Suffix;
 		RestPositionBufferName = UNiagaraDataInterfaceHairStrands::RestPositionBufferName + Suffix;
@@ -227,6 +240,7 @@ struct FNDIHairStrandsParametersName
 		InterpolationModeName = UNiagaraDataInterfaceHairStrands::InterpolationModeName + Suffix;
 		ResetSimulationName = UNiagaraDataInterfaceHairStrands::ResetSimulationName + Suffix;
 		RestUpdateName = UNiagaraDataInterfaceHairStrands::RestUpdateName + Suffix;
+		LocalSimulationName = UNiagaraDataInterfaceHairStrands::LocalSimulationName + Suffix;
 		RootBarycentricCoordinatesName = UNiagaraDataInterfaceHairStrands::RootBarycentricCoordinatesName + Suffix;
 
 		RestRootOffsetName = UNiagaraDataInterfaceHairStrands::RestRootOffsetName + Suffix;
@@ -257,6 +271,10 @@ struct FNDIHairStrandsParametersName
 	FString WorldInverseName;
 	FString WorldRotationName;
 
+	FString BoneTransformName;
+	FString BoneInverseName;
+	FString BoneRotationName;
+
 	FString DeformedPositionBufferName;
 	FString CurvesOffsetsBufferName;
 	FString RestPositionBufferName;
@@ -264,6 +282,7 @@ struct FNDIHairStrandsParametersName
 	FString ResetSimulationName;
 	FString InterpolationModeName;
 	FString RestUpdateName;
+	FString LocalSimulationName;
 	FString RootBarycentricCoordinatesName;
 
 	FString RestRootOffsetName;
@@ -427,8 +446,7 @@ void FNDIHairStrandsData::Update(UNiagaraDataInterfaceHairStrands* Interface, FN
 
 		GlobalInterpolation = (Interface->IsComponentValid() && Interface->SourceComponent->BindingAsset && Interface->SourceComponent->GroomAsset) ?
 			Interface->SourceComponent->GroomAsset->EnableGlobalInterpolation : false;
-		HairGroupInstance = Interface->IsComponentValid() ? Interface->SourceComponent->GetGroupInstance(GroupIndex) : nullptr;
-
+		
 		TickingGroup = Interface->IsComponentValid() ? ComputeTickingGroup(Interface->SourceComponent) : NiagaraFirstTickGroup;
 
 		const bool bIsSimulationEnable = Interface->IsComponentValid() ? Interface->SourceComponent->IsSimulationEnable(GroupIndex, LODIndex) : 
@@ -439,6 +457,8 @@ void FNDIHairStrandsData::Update(UNiagaraDataInterfaceHairStrands* Interface, FN
 			FHairGroupsPhysics& HairPhysics = GroomAsset->HairGroupsPhysics[GroupIndex];
 			StrandsSize = static_cast<uint8>(HairPhysics.StrandsParameters.StrandsSize);
 
+			HairGroupInstance = Interface->IsComponentValid() ? Interface->SourceComponent->GetGroupInstance(GroupIndex) : nullptr;
+
 			SubSteps = HairPhysics.SolverSettings.SubSteps;
 			IterationCount = HairPhysics.SolverSettings.IterationCount;
 
@@ -480,10 +500,15 @@ void FNDIHairStrandsData::Update(UNiagaraDataInterfaceHairStrands* Interface, FN
 			const FBox& StrandsBox = StrandsDatas->BoundingBox;
 
 			NumStrands = StrandsDatas->GetNumCurves();
+			LocalSimulation = false;
+			BoneTransform = FTransform::Identity;
 
 			if (Interface->IsComponentValid())
 			{
 				const FHairSimulationSettings& SimulationSettings = Interface->SourceComponent->SimulationSettings;
+				LocalSimulation = SimulationSettings.SimulationSetup.bLocalSimulation;
+				Interface->SourceComponent->BuildSimulationTransform(BoneTransform);
+
 				if (SimulationSettings.bOverrideSettings)
 				{
 					GravityVector = SimulationSettings.ExternalForces.GravityVector;
@@ -564,6 +589,10 @@ struct FNDIHairStrandsParametersCS : public FNiagaraDataInterfaceParametersCS
 		NumStrands.Bind(ParameterMap, *ParamNames.NumStrandsName);
 		StrandSize.Bind(ParameterMap, *ParamNames.StrandSizeName);
 
+		BoneTransform.Bind(ParameterMap, *ParamNames.BoneTransformName);
+		BoneInverse.Bind(ParameterMap, *ParamNames.BoneInverseName);
+		BoneRotation.Bind(ParameterMap, *ParamNames.BoneRotationName);
+
 		DeformedPositionBuffer.Bind(ParameterMap, *ParamNames.DeformedPositionBufferName);
 		CurvesOffsetsBuffer.Bind(ParameterMap, *ParamNames.CurvesOffsetsBufferName);
 		RestPositionBuffer.Bind(ParameterMap, *ParamNames.RestPositionBufferName);
@@ -571,6 +600,7 @@ struct FNDIHairStrandsParametersCS : public FNiagaraDataInterfaceParametersCS
 		ResetSimulation.Bind(ParameterMap, *ParamNames.ResetSimulationName);
 		InterpolationMode.Bind(ParameterMap,*ParamNames.InterpolationModeName);
 		RestUpdate.Bind(ParameterMap, *ParamNames.RestUpdateName);
+		LocalSimulation.Bind(ParameterMap, *ParamNames.LocalSimulationName);
 		RestRootOffset.Bind(ParameterMap, *ParamNames.RestRootOffsetName);
 		DeformedRootOffset.Bind(ParameterMap, *ParamNames.DeformedRootOffsetName);
 
@@ -594,8 +624,6 @@ struct FNDIHairStrandsParametersCS : public FNiagaraDataInterfaceParametersCS
 		BoundingBoxOffsets.Bind(ParameterMap, *ParamNames.BoundingBoxOffsetsName);
 		BoundingBoxBuffer.Bind(ParameterMap, *ParamNames.BoundingBoxBufferName);
 		ParamsScaleBuffer.Bind(ParameterMap, *ParamNames.ParamsScaleBufferName);
-
-
 	}
 
 	void Set(FRHICommandList& RHICmdList, const FNiagaraDataInterfaceSetArgs& Context) const
@@ -604,165 +632,160 @@ struct FNDIHairStrandsParametersCS : public FNiagaraDataInterfaceParametersCS
 
 		FRHIComputeShader* ComputeShaderRHI = RHICmdList.GetBoundComputeShader();
 
-		FNDIHairStrandsProxy* InterfaceProxy =
-			static_cast<FNDIHairStrandsProxy*>(Context.DataInterface);
-		FNDIHairStrandsData* ProxyData =
-			InterfaceProxy->SystemInstancesToProxyData.Find(Context.SystemInstanceID);
-
-		const bool IsHairValid = ProxyData != nullptr && ProxyData->HairStrandsBuffer != nullptr && ProxyData->HairStrandsBuffer->IsInitialized();
-		const bool bHasSkinningBinding = IsHairValid && ProxyData->HairGroupInstance && ProxyData->HairGroupInstance->BindingType == EHairBindingType::Skinning;
-		const bool IsRootValid = IsHairValid && ProxyData->HairStrandsBuffer->SourceDeformedRootResources != nullptr && bHasSkinningBinding;//&& ProxyData->HairStrandsBuffer->SourceRootResources->IsInitialized();
-		const bool IsRestValid = IsHairValid && ProxyData->HairStrandsBuffer->SourceRestResources != nullptr;//&& ProxyData->HairStrandsBuffer->SourceRestResources->IsInitialized();
-		const bool IsDeformedValid = IsHairValid && ProxyData->HairStrandsBuffer->SourceDeformedResources != nullptr && ProxyData->HairStrandsBuffer->SourceDeformedResources->IsInitialized();
+		FNDIHairStrandsProxy* InterfaceProxy = static_cast<FNDIHairStrandsProxy*>(Context.DataInterface);
+		FNDIHairStrandsData* ProxyData = InterfaceProxy->SystemInstancesToProxyData.Find(Context.SystemInstanceID);
 
-		const bool bHasLODSwitched = IsHairValid && ProxyData->HairGroupInstance && ProxyData->HairGroupInstance->HairGroupPublicData && ProxyData->HairGroupInstance->HairGroupPublicData->VFInput.bHasLODSwitch;
+		const bool bIsHairValid = ProxyData != nullptr && ProxyData->HairStrandsBuffer && ProxyData->HairStrandsBuffer->IsInitialized();
+		const bool bHasSkinningBinding = bIsHairValid && ProxyData->HairGroupInstance && ProxyData->HairGroupInstance->BindingType == EHairBindingType::Skinning;
+		const bool bIsRootValid = bIsHairValid && ProxyData->HairStrandsBuffer->SourceDeformedRootResources && bHasSkinningBinding;
+		const bool bIsRestValid = bIsHairValid && ProxyData->HairStrandsBuffer->SourceRestResources;
+		const bool bIsDeformedValid = bIsHairValid && ProxyData->HairStrandsBuffer->SourceDeformedResources && ProxyData->HairStrandsBuffer->SourceDeformedResources->IsInitialized();
+		const bool bHasLODSwitched = bIsHairValid && ProxyData->HairGroupInstance && ProxyData->HairGroupInstance->HairGroupPublicData && ProxyData->HairGroupInstance->HairGroupPublicData->VFInput.bHasLODSwitch;
 
-		if (IsHairValid && IsRestValid)
+		if (bIsHairValid && bIsRestValid)
 		{
 			check(ProxyData);
 
 			FNDIHairStrandsBuffer* HairStrandsBuffer = ProxyData->HairStrandsBuffer;
 
-			FUnorderedAccessViewRHIRef PointPositionsUAV = IsDeformedValid ?
+			FUnorderedAccessViewRHIRef DeformedPositionBufferUAV = bIsDeformedValid ?
 				HairStrandsBuffer->SourceDeformedResources->DeformedPositionBuffer[HairStrandsBuffer->SourceDeformedResources->CurrentIndex].UAV : HairStrandsBuffer->DeformedPositionBuffer.UAV;
-
-			FRHIShaderResourceView* DeformedPositionOffsetSRV = IsDeformedValid ? 
+			FRHIShaderResourceView* DeformedPositionOffsetSRV = bIsDeformedValid ? 
 				HairStrandsBuffer->SourceDeformedResources->DeformedOffsetBuffer[HairStrandsBuffer->SourceDeformedResources->CurrentIndex].SRV.GetReference() : FNiagaraRenderer::GetDummyFloatBuffer();
 
-			const int32 MeshLODIndex = IsRootValid ? HairStrandsBuffer->SourceDeformedRootResources->MeshLODIndex : -1;
-			EHairSimulationInterpolationMode InterpolationModeValue = EHairSimulationInterpolationMode::Rigid;
-			if (IsRootValid && HairStrandsBuffer->SourceDeformedRootResources->IsValid(MeshLODIndex))
-			{
-				InterpolationModeValue = EHairSimulationInterpolationMode::Skinned;
-			}
-			const int32 RestUpdateValue = GHairSimulationRestUpdate;
+			const int32 MeshLODIndex = bIsRootValid ? HairStrandsBuffer->SourceDeformedRootResources->MeshLODIndex : -1;
+
+			// Projection Buffers
+			const bool bHasSkinnedInterpolation = (bIsRootValid && HairStrandsBuffer->SourceDeformedRootResources->IsValid(MeshLODIndex));
+			const EHairSimulationInterpolationMode InterpolationModeValue = bHasSkinnedInterpolation ? (ProxyData->GlobalInterpolation ? 
+				EHairSimulationInterpolationMode::RBF : EHairSimulationInterpolationMode::Skinned ): EHairSimulationInterpolationMode::Rigid;
 
-			const FHairStrandsRestRootResource::FLOD* RestMeshProjection = (InterpolationModeValue == EHairSimulationInterpolationMode::Skinned) ? &(HairStrandsBuffer->SourceRestRootResources->LODs[MeshLODIndex]) : nullptr;
-			const FHairStrandsDeformedRootResource::FLOD* DeformedMeshProjection = (InterpolationModeValue == EHairSimulationInterpolationMode::Skinned) ? &(HairStrandsBuffer->SourceDeformedRootResources->LODs[MeshLODIndex]) : nullptr;
+			const FHairStrandsRestRootResource::FLOD* RestMeshProjection = bHasSkinnedInterpolation ? &(HairStrandsBuffer->SourceRestRootResources->LODs[MeshLODIndex]) : nullptr;
+			const FHairStrandsDeformedRootResource::FLOD* DeformedMeshProjection = bHasSkinnedInterpolation ? &(HairStrandsBuffer->SourceDeformedRootResources->LODs[MeshLODIndex]) : nullptr;
 
-			FRHIShaderResourceView* RestTrianglePositionASRV = (InterpolationModeValue == EHairSimulationInterpolationMode::Skinned && RestMeshProjection != nullptr) ?
+			FRHIShaderResourceView* RestTrianglePositionASRV = (bHasSkinnedInterpolation && RestMeshProjection) ?
 				RestMeshProjection->RestRootTrianglePosition0Buffer.SRV.GetReference() : FNiagaraRenderer::GetDummyFloatBuffer();
-			FRHIShaderResourceView* RestTrianglePositionBSRV = (InterpolationModeValue == EHairSimulationInterpolationMode::Skinned && RestMeshProjection != nullptr) ?
+			FRHIShaderResourceView* RestTrianglePositionBSRV = (bHasSkinnedInterpolation && RestMeshProjection) ?
 				RestMeshProjection->RestRootTrianglePosition1Buffer.SRV.GetReference() : FNiagaraRenderer::GetDummyFloatBuffer();
-			FRHIShaderResourceView* RestTrianglePositionCSRV = (InterpolationModeValue == EHairSimulationInterpolationMode::Skinned && RestMeshProjection != nullptr) ?
+			FRHIShaderResourceView* RestTrianglePositionCSRV = (bHasSkinnedInterpolation && RestMeshProjection) ?
 				RestMeshProjection->RestRootTrianglePosition2Buffer.SRV.GetReference() : FNiagaraRenderer::GetDummyFloatBuffer();
 
-			FRHIShaderResourceView* DeformedTrianglePositionASRV = (InterpolationModeValue == EHairSimulationInterpolationMode::Skinned && DeformedMeshProjection != nullptr) ?
+			FRHIShaderResourceView* DeformedTrianglePositionASRV = (bHasSkinnedInterpolation && DeformedMeshProjection) ?
 				DeformedMeshProjection->DeformedRootTrianglePosition0Buffer.SRV.GetReference() : FNiagaraRenderer::GetDummyFloatBuffer();
-			FRHIShaderResourceView* DeformedTrianglePositionBSRV = (InterpolationModeValue == EHairSimulationInterpolationMode::Skinned && DeformedMeshProjection != nullptr) ?
+			FRHIShaderResourceView* DeformedTrianglePositionBSRV = (bHasSkinnedInterpolation && DeformedMeshProjection) ?
 				DeformedMeshProjection->DeformedRootTrianglePosition1Buffer.SRV.GetReference() : FNiagaraRenderer::GetDummyFloatBuffer();
-			FRHIShaderResourceView* DeformedTrianglePositionCSRV = (InterpolationModeValue == EHairSimulationInterpolationMode::Skinned && DeformedMeshProjection != nullptr) ?
+			FRHIShaderResourceView* DeformedTrianglePositionCSRV = (bHasSkinnedInterpolation && DeformedMeshProjection) ?
 				DeformedMeshProjection->DeformedRootTrianglePosition2Buffer.SRV.GetReference() : FNiagaraRenderer::GetDummyFloatBuffer();
+			FRHIShaderResourceView* RootBarycentricCoordinatesSRV = (bHasSkinnedInterpolation && RestMeshProjection) ?
+				RestMeshProjection->RootTriangleBarycentricBuffer.SRV.GetReference() : FNiagaraRenderer::GetDummyFloatBuffer();
 
-			const bool bHasSamples = (RestMeshProjection != nullptr && RestMeshProjection->SampleCount > 0);
+			// RBF buffers
+			const bool bHasSamples = (RestMeshProjection && RestMeshProjection->SampleCount > 0);
 			const int32 SampleCountValue = bHasSamples ? RestMeshProjection->SampleCount : 0;
-			FShaderResourceViewRHIRef RestSamplePositionsBufferSRV = bHasSamples ?
+
+			FShaderResourceViewRHIRef RestSamplePositionsBufferSRV = (bHasSamples && RestMeshProjection) ?
 				RestMeshProjection->RestSamplePositionsBuffer.SRV.GetReference() : FNiagaraRenderer::GetDummyFloatBuffer();
-			FShaderResourceViewRHIRef MeshSampleWeightsBufferSRV = bHasSamples && DeformedMeshProjection != nullptr ?
+			FShaderResourceViewRHIRef MeshSampleWeightsBufferSRV = (bHasSamples && DeformedMeshProjection) ?
 				DeformedMeshProjection->MeshSampleWeightsBuffer.SRV.GetReference() : FNiagaraRenderer::GetDummyFloatBuffer();
 
-			FRHIShaderResourceView* RootBarycentricCoordinatesSRV = (InterpolationModeValue == EHairSimulationInterpolationMode::Skinned && RestMeshProjection != nullptr) ?
-				RestMeshProjection->RootTriangleBarycentricBuffer.SRV.GetReference() : FNiagaraRenderer::GetDummyFloatBuffer();
-
-			int32 bNeedSimReset = bHasLODSwitched || (ProxyData->TickCount <= GHairSimulationMaxDelay ? 1 : 0);
-
-			if (InterpolationModeValue == EHairSimulationInterpolationMode::Skinned && ProxyData->GlobalInterpolation)
-			{
-				InterpolationModeValue = EHairSimulationInterpolationMode::RBF;
-			}
-
-			//FUnorderedAccessViewRHIRef DummyFloatUAV(Context.Batcher->GetEmptyUAVFromPool(RHICmdList, PF_R32_FLOAT, ENiagaraEmptyUAVType::Buffer), false);
-			//FUnorderedAccessViewRHIRef MeshSampleWeightsBufferUAV = (DeformedMeshProjection != nullptr && RestMeshProjection->SampleCount > 0) ?
-			//	DeformedMeshProjection->MeshSampleWeightsBuffer.UAV : Context.Batcher->GetEmptyUAVFromPool(RHICmdList, PF_R32_FLOAT, ENiagaraEmptyUAVType::Buffer);
-
-			//UE_LOG(LogHairStrands, Log, TEXT("Shader Reset : %d %d %d"), bNeedSimReset, IsRootValid, InterpolationModeValue);
-			FVector3f RestRootOffsetValue = FVector3f::ZeroVector;
-			FVector3f DeformedRootOffsetValue = FVector3f::ZeroVector;
+			// Simulation setup
+			const int32 NeedResetValue = bHasLODSwitched || (ProxyData->TickCount <= GHairSimulationMaxDelay);
+			const int32 RestUpdateValue = GHairSimulationRestUpdate;
+			const int32 LocalSimulationValue = ProxyData->LocalSimulation;
 			
+			// Offsets / Transforms
 			FVector3f RestPositionOffsetValue = ProxyData->HairStrandsBuffer->SourceRestResources->GetPositionOffset();
 
+			FMatrix44f WorldTransformFloat = ProxyData->HairGroupInstance ? (ProxyData->HairGroupInstance->BindingType == EHairBindingType::Skinning) ?
+				ProxyData->HairGroupInstance->Debug.SkeletalLocalToWorld.ToMatrixWithScale() :
+				ProxyData->HairGroupInstance->LocalToWorld.ToMatrixWithScale() :
+				ProxyData->WorldTransform.ToMatrixWithScale();
+
+			FMatrix44f BoneTransformFloat = ProxyData->BoneTransform.ToMatrixWithScale();
+
 			FRHITransitionInfo Transitions[] = {
-				FRHITransitionInfo(PointPositionsUAV, ERHIAccess::Unknown, ERHIAccess::UAVCompute),
+				FRHITransitionInfo(DeformedPositionBufferUAV, ERHIAccess::Unknown, ERHIAccess::UAVCompute),
 				FRHITransitionInfo(HairStrandsBuffer->BoundingBoxBuffer.UAV, ERHIAccess::Unknown, ERHIAccess::UAVCompute),
+				FRHITransitionInfo(HairStrandsBuffer->ParamsScaleBuffer.UAV, ERHIAccess::Unknown, ERHIAccess::SRVCompute),
+				FRHITransitionInfo(HairStrandsBuffer->CurvesOffsetsBuffer.UAV, ERHIAccess::Unknown, ERHIAccess::SRVCompute)
 			};
 			RHICmdList.Transition(MakeArrayView(Transitions, UE_ARRAY_COUNT(Transitions)));
 
-			SetUAVParameter(RHICmdList, ComputeShaderRHI, DeformedPositionBuffer, PointPositionsUAV);
-
-			SetUAVParameter(RHICmdList, ComputeShaderRHI, BoundingBoxBuffer, HairStrandsBuffer->BoundingBoxBuffer.UAV);
-
-			SetSRVParameter(RHICmdList, ComputeShaderRHI, CurvesOffsetsBuffer, HairStrandsBuffer->CurvesOffsetsBuffer.SRV);
-			SetSRVParameter(RHICmdList, ComputeShaderRHI, RestPositionBuffer, HairStrandsBuffer->SourceRestResources->PositionBuffer.SRV);
-
+			// Set shader constants
 			SetShaderValue(RHICmdList, ComputeShaderRHI, BoundingBoxOffsets, HairStrandsBuffer->BoundingBoxOffsets);
-			FMatrix44f WorldTransformFloat = ProxyData->WorldTransform.ToMatrixWithScale();
 			SetShaderValue(RHICmdList, ComputeShaderRHI, WorldTransform, WorldTransformFloat);
 			SetShaderValue(RHICmdList, ComputeShaderRHI, WorldInverse, WorldTransformFloat.Inverse());
 			SetShaderValue(RHICmdList, ComputeShaderRHI, WorldRotation, WorldTransformFloat.ToQuat());
 			SetShaderValue(RHICmdList, ComputeShaderRHI, NumStrands, ProxyData->NumStrands);
 			SetShaderValue(RHICmdList, ComputeShaderRHI, StrandSize, ProxyData->StrandsSize);
-
-			SetShaderValue(RHICmdList, ComputeShaderRHI, ResetSimulation, bNeedSimReset);
+			SetShaderValue(RHICmdList, ComputeShaderRHI, BoneTransform, BoneTransformFloat);
+			SetShaderValue(RHICmdList, ComputeShaderRHI, BoneInverse, BoneTransformFloat.Inverse());
+			SetShaderValue(RHICmdList, ComputeShaderRHI, BoneRotation, BoneTransformFloat.ToQuat());
+			SetShaderValue(RHICmdList, ComputeShaderRHI, ResetSimulation, NeedResetValue);
 			SetShaderValue(RHICmdList, ComputeShaderRHI, InterpolationMode, int32(InterpolationModeValue));
 			SetShaderValue(RHICmdList, ComputeShaderRHI, RestUpdate, RestUpdateValue);
-			SetShaderValue(RHICmdList, ComputeShaderRHI, RestRootOffset, RestRootOffsetValue);
-			SetShaderValue(RHICmdList, ComputeShaderRHI, DeformedRootOffset, DeformedRootOffsetValue);
-
+			SetShaderValue(RHICmdList, ComputeShaderRHI, LocalSimulation, LocalSimulationValue);
+			SetShaderValue(RHICmdList, ComputeShaderRHI, RestRootOffset, FVector3f::ZeroVector);
+			SetShaderValue(RHICmdList, ComputeShaderRHI, DeformedRootOffset, FVector3f::ZeroVector);
 			SetShaderValue(RHICmdList, ComputeShaderRHI, RestPositionOffset, RestPositionOffsetValue);
-			SetSRVParameter(RHICmdList, ComputeShaderRHI, DeformedPositionOffset, DeformedPositionOffsetSRV);
+			SetShaderValue(RHICmdList, ComputeShaderRHI, SampleCount, SampleCountValue);
 
+			// Set Shader UAV
+			SetUAVParameter(RHICmdList, ComputeShaderRHI, DeformedPositionBuffer, DeformedPositionBufferUAV);
+			SetUAVParameter(RHICmdList, ComputeShaderRHI, BoundingBoxBuffer, HairStrandsBuffer->BoundingBoxBuffer.UAV);
+
+			// Set Shader SRV
+			SetSRVParameter(RHICmdList, ComputeShaderRHI, CurvesOffsetsBuffer, HairStrandsBuffer->CurvesOffsetsBuffer.SRV);
+			SetSRVParameter(RHICmdList, ComputeShaderRHI, ParamsScaleBuffer, HairStrandsBuffer->ParamsScaleBuffer.SRV);
+			SetSRVParameter(RHICmdList, ComputeShaderRHI, RestPositionBuffer, HairStrandsBuffer->SourceRestResources->PositionBuffer.SRV);
+			SetSRVParameter(RHICmdList, ComputeShaderRHI, DeformedPositionOffset, DeformedPositionOffsetSRV);
 			SetSRVParameter(RHICmdList, ComputeShaderRHI, RestTrianglePositionABuffer, RestTrianglePositionASRV);
 			SetSRVParameter(RHICmdList, ComputeShaderRHI, RestTrianglePositionBBuffer, RestTrianglePositionBSRV);
 			SetSRVParameter(RHICmdList, ComputeShaderRHI, RestTrianglePositionCBuffer, RestTrianglePositionCSRV);
-
 			SetSRVParameter(RHICmdList, ComputeShaderRHI, DeformedTrianglePositionABuffer, DeformedTrianglePositionASRV);
 			SetSRVParameter(RHICmdList, ComputeShaderRHI, DeformedTrianglePositionBBuffer, DeformedTrianglePositionBSRV);
 			SetSRVParameter(RHICmdList, ComputeShaderRHI, DeformedTrianglePositionCBuffer, DeformedTrianglePositionCSRV);
-
-			SetShaderValue(RHICmdList, ComputeShaderRHI, SampleCount, SampleCountValue);
 			SetSRVParameter(RHICmdList, ComputeShaderRHI, RestSamplePositionsBuffer, RestSamplePositionsBufferSRV);
-
 			SetSRVParameter(RHICmdList, ComputeShaderRHI, MeshSampleWeightsBuffer, MeshSampleWeightsBufferSRV);
-
 			SetSRVParameter(RHICmdList, ComputeShaderRHI, RootBarycentricCoordinatesBuffer, RootBarycentricCoordinatesSRV);
-
-			SetSRVParameter(RHICmdList, ComputeShaderRHI, ParamsScaleBuffer, HairStrandsBuffer->ParamsScaleBuffer.SRV);
 		}
 		else
 		{
-			SetUAVParameter(RHICmdList, ComputeShaderRHI, DeformedPositionBuffer, Context.Batcher->GetEmptyUAVFromPool(RHICmdList, PF_R32_FLOAT, ENiagaraEmptyUAVType::Buffer));
-			SetUAVParameter(RHICmdList, ComputeShaderRHI, BoundingBoxBuffer, Context.Batcher->GetEmptyUAVFromPool(RHICmdList, PF_R32_UINT, ENiagaraEmptyUAVType::Buffer));
-			SetSRVParameter(RHICmdList, ComputeShaderRHI, CurvesOffsetsBuffer, FNiagaraRenderer::GetDummyUIntBuffer());
-			SetSRVParameter(RHICmdList, ComputeShaderRHI, RestPositionBuffer, FNiagaraRenderer::GetDummyFloatBuffer());
-
+			// Set shader constants
 			SetShaderValue(RHICmdList, ComputeShaderRHI, BoundingBoxOffsets, FIntVector4(0,1,2,3));
 			SetShaderValue(RHICmdList, ComputeShaderRHI, WorldTransform, FMatrix44f::Identity);
 			SetShaderValue(RHICmdList, ComputeShaderRHI, WorldInverse, FMatrix44f::Identity);
 			SetShaderValue(RHICmdList, ComputeShaderRHI, WorldRotation, FQuat::Identity);
 			SetShaderValue(RHICmdList, ComputeShaderRHI, NumStrands, 1);
 			SetShaderValue(RHICmdList, ComputeShaderRHI, StrandSize, 1);
-
+			SetShaderValue(RHICmdList, ComputeShaderRHI, BoneTransform, FMatrix44f::Identity);
+			SetShaderValue(RHICmdList, ComputeShaderRHI, BoneInverse, FMatrix44f::Identity);
+			SetShaderValue(RHICmdList, ComputeShaderRHI, BoneRotation, FQuat::Identity);
 			SetShaderValue(RHICmdList, ComputeShaderRHI, ResetSimulation, 0);
 			SetShaderValue(RHICmdList, ComputeShaderRHI, InterpolationMode, 0);
 			SetShaderValue(RHICmdList, ComputeShaderRHI, RestUpdate, 0);
-			SetShaderValue(RHICmdList, ComputeShaderRHI, RestRootOffset, FVector4(0, 0, 0, 0));
-			SetShaderValue(RHICmdList, ComputeShaderRHI, DeformedRootOffset, FVector4(0, 0, 0, 0));
-
+			SetShaderValue(RHICmdList, ComputeShaderRHI, LocalSimulation, 0);
+			SetShaderValue(RHICmdList, ComputeShaderRHI, RestRootOffset, FVector3f::ZeroVector);
+			SetShaderValue(RHICmdList, ComputeShaderRHI, DeformedRootOffset, FVector3f::ZeroVector);
 			SetShaderValue(RHICmdList, ComputeShaderRHI, RestPositionOffset, FVector3f::ZeroVector);
-			SetSRVParameter(RHICmdList, ComputeShaderRHI, DeformedPositionOffset, FNiagaraRenderer::GetDummyFloatBuffer());
+			SetShaderValue(RHICmdList, ComputeShaderRHI, SampleCount, 0);
+
+			// Set Shader UAV
+			SetUAVParameter(RHICmdList, ComputeShaderRHI, DeformedPositionBuffer, Context.Batcher->GetEmptyUAVFromPool(RHICmdList, PF_R32_FLOAT, ENiagaraEmptyUAVType::Buffer));
+			SetUAVParameter(RHICmdList, ComputeShaderRHI, BoundingBoxBuffer, Context.Batcher->GetEmptyUAVFromPool(RHICmdList, PF_R32_UINT, ENiagaraEmptyUAVType::Buffer));
 
+			// Set Shader SRV
+			SetSRVParameter(RHICmdList, ComputeShaderRHI, CurvesOffsetsBuffer, FNiagaraRenderer::GetDummyUIntBuffer());
+			SetSRVParameter(RHICmdList, ComputeShaderRHI, RestPositionBuffer, FNiagaraRenderer::GetDummyFloatBuffer());
+			SetSRVParameter(RHICmdList, ComputeShaderRHI, DeformedPositionOffset, FNiagaraRenderer::GetDummyFloatBuffer());
 			SetSRVParameter(RHICmdList, ComputeShaderRHI, RestTrianglePositionABuffer, FNiagaraRenderer::GetDummyFloatBuffer());
 			SetSRVParameter(RHICmdList, ComputeShaderRHI, RestTrianglePositionBBuffer, FNiagaraRenderer::GetDummyFloatBuffer());
 			SetSRVParameter(RHICmdList, ComputeShaderRHI, RestTrianglePositionCBuffer, FNiagaraRenderer::GetDummyFloatBuffer());
-
 			SetSRVParameter(RHICmdList, ComputeShaderRHI, DeformedTrianglePositionABuffer, FNiagaraRenderer::GetDummyFloatBuffer());
 			SetSRVParameter(RHICmdList, ComputeShaderRHI, DeformedTrianglePositionBBuffer, FNiagaraRenderer::GetDummyFloatBuffer());
 			SetSRVParameter(RHICmdList, ComputeShaderRHI, DeformedTrianglePositionCBuffer, FNiagaraRenderer::GetDummyFloatBuffer());
-
-			SetShaderValue(RHICmdList, ComputeShaderRHI, SampleCount, 0);
 			SetSRVParameter(RHICmdList, ComputeShaderRHI, RestSamplePositionsBuffer, FNiagaraRenderer::GetDummyFloatBuffer());
 			SetSRVParameter(RHICmdList, ComputeShaderRHI, MeshSampleWeightsBuffer, FNiagaraRenderer::GetDummyFloatBuffer());
-
 			SetSRVParameter(RHICmdList, ComputeShaderRHI, RootBarycentricCoordinatesBuffer, FNiagaraRenderer::GetDummyFloatBuffer());
 			SetSRVParameter(RHICmdList, ComputeShaderRHI, ParamsScaleBuffer, FNiagaraRenderer::GetDummyFloatBuffer());
 		}
@@ -782,37 +805,35 @@ private:
 	LAYOUT_FIELD(FShaderParameter, WorldRotation);
 	LAYOUT_FIELD(FShaderParameter, NumStrands);
 	LAYOUT_FIELD(FShaderParameter, StrandSize);
-
-	LAYOUT_FIELD(FShaderResourceParameter, DeformedPositionBuffer);
-	LAYOUT_FIELD(FShaderResourceParameter, CurvesOffsetsBuffer);
-	LAYOUT_FIELD(FShaderResourceParameter, RestPositionBuffer);
-
+	LAYOUT_FIELD(FShaderParameter, BoneTransform);
+	LAYOUT_FIELD(FShaderParameter, BoneInverse);
+	LAYOUT_FIELD(FShaderParameter, BoneRotation);
 	LAYOUT_FIELD(FShaderParameter, ResetSimulation);
 	LAYOUT_FIELD(FShaderParameter, InterpolationMode);
 	LAYOUT_FIELD(FShaderParameter, RestUpdate);
+	LAYOUT_FIELD(FShaderParameter, LocalSimulation);
 	LAYOUT_FIELD(FShaderParameter, RestRootOffset);
 	LAYOUT_FIELD(FShaderParameter, DeformedRootOffset);
+	LAYOUT_FIELD(FShaderParameter, SampleCount);
+	LAYOUT_FIELD(FShaderParameter, RestPositionOffset);
+	LAYOUT_FIELD(FShaderParameter, BoundingBoxOffsets);
 
-	LAYOUT_FIELD(FShaderResourceParameter, RootBarycentricCoordinatesBuffer);
+	LAYOUT_FIELD(FShaderResourceParameter, DeformedPositionBuffer);
+	LAYOUT_FIELD(FShaderResourceParameter, BoundingBoxBuffer);
 
+	LAYOUT_FIELD(FShaderResourceParameter, CurvesOffsetsBuffer);
+	LAYOUT_FIELD(FShaderResourceParameter, RestPositionBuffer);
+	LAYOUT_FIELD(FShaderResourceParameter, RootBarycentricCoordinatesBuffer);
 	LAYOUT_FIELD(FShaderResourceParameter, RestTrianglePositionABuffer);
 	LAYOUT_FIELD(FShaderResourceParameter, RestTrianglePositionBBuffer);
 	LAYOUT_FIELD(FShaderResourceParameter, RestTrianglePositionCBuffer);
-
 	LAYOUT_FIELD(FShaderResourceParameter, DeformedTrianglePositionABuffer);
 	LAYOUT_FIELD(FShaderResourceParameter, DeformedTrianglePositionBBuffer);
 	LAYOUT_FIELD(FShaderResourceParameter, DeformedTrianglePositionCBuffer);
-
-	LAYOUT_FIELD(FShaderParameter, SampleCount);
 	LAYOUT_FIELD(FShaderResourceParameter, RestSamplePositionsBuffer);
 	LAYOUT_FIELD(FShaderResourceParameter, MeshSampleWeightsBuffer);
-
-	LAYOUT_FIELD(FShaderParameter, RestPositionOffset);
 	LAYOUT_FIELD(FShaderResourceParameter, DeformedPositionOffset);
-
 	LAYOUT_FIELD(FShaderResourceParameter, ParamsScaleBuffer);
-	LAYOUT_FIELD(FShaderResourceParameter, BoundingBoxBuffer);
-	LAYOUT_FIELD(FShaderParameter, BoundingBoxOffsets);
 };
 
 IMPLEMENT_TYPE_LAYOUT(FNDIHairStrandsParametersCS);
@@ -993,7 +1014,7 @@ void UNiagaraDataInterfaceHairStrands::ExtractDatasAndResources(
 			OutStrandsDeformedRootResource = SourceComponent->GetGuideStrandsDeformedRootResource(OutGroupIndex);
 			OutGroomAsset = SourceComponent->GroomAsset;
 			OutLODIndex = SourceComponent->GetForcedLOD();
-			OutLocalToWorld = SourceComponent->GetGuideStrandsLocalToWorld(OutGroupIndex);
+			OutLocalToWorld = SourceComponent->GetComponentTransform();
 		}
 	}
 	else if (DefaultSource != nullptr)
@@ -1003,7 +1024,7 @@ void UNiagaraDataInterfaceHairStrands::ExtractDatasAndResources(
 		OutLocalToWorld = SystemInstance ? SystemInstance->GetWorldTransform() : FTransform::Identity;
 		if (OutGroupIndex < DefaultSource->GetNumHairGroups())
 		{
-			OutStrandsRestResource = DefaultSource->HairGroupsData[OutGroupIndex].Strands.RestResource;
+			OutStrandsRestResource = DefaultSource->HairGroupsData[OutGroupIndex].Guides.RestResource;
 			OutGroomAsset = DefaultSource;
 		}
 	}
@@ -1064,12 +1085,11 @@ bool UNiagaraDataInterfaceHairStrands::PerInstanceTick(void* PerInstanceData, FN
 
 	if (SourceComponent != nullptr)
 	{
-		const bool bRequiresReset = SourceComponent->bResetSimulation || RequiresSimulationReset(SystemInstance, InstanceData->SkeletalMeshes);
-		if (bRequiresReset)
+		if (SourceComponent->bResetSimulation || RequiresSimulationReset(SystemInstance, InstanceData->SkeletalMeshes))
 		{
 			InstanceData->TickCount = 0;
 		}
-		InstanceData->ForceReset = bRequiresReset;
+		InstanceData->ForceReset = SourceComponent->bResetSimulation;
 	}
 	InstanceData->Update(this, SystemInstance, StrandsRestResource ? &StrandsRestResource->BulkData : nullptr, GroomAsset, GroupIndex, LODIndex, LocalToWorld);
 	return false;
@@ -1606,6 +1626,30 @@ void UNiagaraDataInterfaceHairStrands::GetFunctions(TArray<FNiagaraFunctionSigna
 
 		OutFunctions.Add(Sig);
 	}
+	{
+		FNiagaraFunctionSignature Sig;
+		Sig.Name = GetLocalVectorName;
+		Sig.bMemberFunction = true;
+		Sig.bRequiresContext = false;
+		Sig.Inputs.Add(FNiagaraVariable(FNiagaraTypeDefinition(GetClass()), TEXT("Hair Strands")));
+		Sig.Inputs.Add(FNiagaraVariable(FNiagaraTypeDefinition::GetVec3Def(), TEXT("World Vector")));
+		Sig.Inputs.Add(FNiagaraVariable(FNiagaraTypeDefinition::GetBoolDef(), TEXT("Is Position")));
+		Sig.Outputs.Add(FNiagaraVariable(FNiagaraTypeDefinition::GetVec3Def(), TEXT("Local Vector")));
+
+		OutFunctions.Add(Sig);
+	}
+	{
+		FNiagaraFunctionSignature Sig;
+		Sig.Name = GetWorldVectorName;
+		Sig.bMemberFunction = true;
+		Sig.bRequiresContext = false;
+		Sig.Inputs.Add(FNiagaraVariable(FNiagaraTypeDefinition(GetClass()), TEXT("Hair Strands")));
+		Sig.Inputs.Add(FNiagaraVariable(FNiagaraTypeDefinition::GetVec3Def(), TEXT("Local Vector")));
+		Sig.Inputs.Add(FNiagaraVariable(FNiagaraTypeDefinition::GetBoolDef(), TEXT("Is Position")));
+		Sig.Outputs.Add(FNiagaraVariable(FNiagaraTypeDefinition::GetVec3Def(), TEXT("World Vector")));
+
+		OutFunctions.Add(Sig);
+	}
 	{
 		FNiagaraFunctionSignature Sig;
 		Sig.Name = AttachNodePositionName;
@@ -3458,17 +3502,39 @@ bool UNiagaraDataInterfaceHairStrands::GetFunctionHLSL(const FNiagaraDataInterfa
 		OutHLSL += FString::Format(FormatSample, ArgsSample);
 		return true;
 	}
-	else if (FunctionInfo.DefinitionName == AttachNodePositionName)
+	else if (FunctionInfo.DefinitionName == GetLocalVectorName)
 	{
 		static const TCHAR *FormatSample = TEXT(R"(
-			void {InstanceFunctionName} (in float3 RestPosition, out float3 NodePosition)
+			void {InstanceFunctionName} (in float3 WorldVector, in bool IsPosition, out float3 LocalVector)
 			{
-				{HairStrandsContextName} DIHairStrands_AttachNodePosition(DIContext,RestPosition,NodePosition);
+				{HairStrandsContextName} DIHairStrands_GetLocalVector(DIContext,WorldVector,IsPosition,LocalVector);
 			}
 			)");
 		OutHLSL += FString::Format(FormatSample, ArgsSample);
 		return true;
 	}
+	else if (FunctionInfo.DefinitionName == GetWorldVectorName)
+	{
+		static const TCHAR* FormatSample = TEXT(R"(
+				void {InstanceFunctionName} (in float3 LocalVector, in bool IsPosition, out float3 WorldVector)
+				{
+					{HairStrandsContextName} DIHairStrands_GetWorldVector(DIContext,LocalVector,IsPosition,WorldVector);
+				}
+				)");
+		OutHLSL += FString::Format(FormatSample, ArgsSample);
+		return true;
+	}
+	else if (FunctionInfo.DefinitionName == AttachNodePositionName)
+	{
+		static const TCHAR* FormatSample = TEXT(R"(
+				void {InstanceFunctionName} (in float3 RestPosition, out float3 NodePosition)
+				{
+					{HairStrandsContextName} DIHairStrands_AttachNodePosition(DIContext,RestPosition,NodePosition);
+				}
+				)");
+		OutHLSL += FString::Format(FormatSample, ArgsSample);
+		return true;
+	}
 	else if (FunctionInfo.DefinitionName == AttachNodeOrientationName)
 	{
 		static const TCHAR *FormatSample = TEXT(R"(
diff --git a/Engine/Plugins/Runtime/HairStrands/Source/HairStrandsCore/Private/Niagara/NiagaraDataInterfacePhysicsAsset.cpp b/Engine/Plugins/Runtime/HairStrands/Source/HairStrandsCore/Private/Niagara/NiagaraDataInterfacePhysicsAsset.cpp
--- a/Engine/Plugins/Runtime/HairStrands/Source/HairStrandsCore/Private/Niagara/NiagaraDataInterfacePhysicsAsset.cpp
+++ b/Engine/Plugins/Runtime/HairStrands/Source/HairStrandsCore/Private/Niagara/NiagaraDataInterfacePhysicsAsset.cpp
@@ -30,6 +30,7 @@ static const FName GetElementPointName(TEXT("GetElementPoint"));
 static const FName GetElementDistanceName(TEXT("GetElementDistance"));
 static const FName GetClosestPointName(TEXT("GetClosestPoint"));
 static const FName GetClosestDistanceName(TEXT("GetClosestDistance"));
+static const FName GetRestDistanceName(TEXT("GetRestDistance"));
 static const FName GetTexturePointName(TEXT("GetTexturePoint"));
 static const FName GetProjectionPointName(TEXT("GetProjectionPoint"));
 
@@ -183,7 +184,7 @@ void CompactInternalArrays(FNDIPhysicsAssetArrays* OutAssetArrays)
 }
 
 void CreateInternalArrays(const TArray<TWeakObjectPtr<UPhysicsAsset>>& PhysicsAssets, const TArray<TWeakObjectPtr<USkeletalMeshComponent>>& SkeletalMeshs,
-	FNDIPhysicsAssetArrays* OutAssetArrays, const FTransform& InWorldTransform)
+	FNDIPhysicsAssetArrays* OutAssetArrays, const FTransform& InWorldTransform, const FTransform& LocalTransform)
 {
 	if (OutAssetArrays != nullptr)
 	{
@@ -239,7 +240,7 @@ void CreateInternalArrays(const TArray<TWeakObjectPtr<UPhysicsAsset>>& PhysicsAs
 							if (BoneIndex != INDEX_NONE && BoneIndex < RestTransforms.Num())
 							{
 								const FTransform RestTransform = RestTransforms[BoneIndex];
-								const FTransform BoneTransform = IsSkelMeshValid ? SkeletalMesh->GetBoneTransform(BoneIndex) : RestTransform * WorldTransform;
+								const FTransform BoneTransform = IsSkelMeshValid ? SkeletalMesh->GetBoneTransform(BoneIndex) * LocalTransform.Inverse() : RestTransform * WorldTransform * LocalTransform.Inverse();
 
 								for (const FKBoxElem& BoxElem : BodySetup->AggGeom.BoxElems)
 								{
@@ -306,7 +307,7 @@ void CreateInternalArrays(const TArray<TWeakObjectPtr<UPhysicsAsset>>& PhysicsAs
 }
 
 void UpdateInternalArrays(const TArray<TWeakObjectPtr<UPhysicsAsset>>& PhysicsAssets, const TArray<TWeakObjectPtr<USkeletalMeshComponent>>& SkeletalMeshs,
-	FNDIPhysicsAssetArrays* OutAssetArrays, const FTransform& InWorldTransform)
+	FNDIPhysicsAssetArrays* OutAssetArrays, const FTransform& InWorldTransform, const FTransform& LocalTransform)
 {
 	if (OutAssetArrays != nullptr && OutAssetArrays->ElementOffsets.NumElements < PHYSICS_ASSET_MAX_PRIMITIVES)
 	{
@@ -320,7 +321,7 @@ void UpdateInternalArrays(const TArray<TWeakObjectPtr<UPhysicsAsset>>& PhysicsAs
 			((OutAssetArrays->ElementOffsets.CapsuleOffset - OutAssetArrays->ElementOffsets.SphereOffset) != NumSpheres) ||
 			((OutAssetArrays->ElementOffsets.NumElements - OutAssetArrays->ElementOffsets.CapsuleOffset) != NumCapsules))
 		{
-			CreateInternalArrays(PhysicsAssets, SkeletalMeshs, OutAssetArrays, InWorldTransform);
+			CreateInternalArrays(PhysicsAssets, SkeletalMeshs, OutAssetArrays, InWorldTransform, LocalTransform);
 		}
 
 		OutAssetArrays->PreviousTransform = OutAssetArrays->CurrentTransform;
@@ -356,7 +357,7 @@ void UpdateInternalArrays(const TArray<TWeakObjectPtr<UPhysicsAsset>>& PhysicsAs
 						const int32 BoneIndex = RefSkeleton.FindBoneIndex(BoneName);
 						if (BoneIndex != INDEX_NONE && BoneIndex < RestTransforms.Num())
 						{
-							const FTransform BoneTransform = IsSkelMeshValid ? SkeletalMesh->GetBoneTransform(BoneIndex) : RestTransforms[BoneIndex] * WorldTransform;
+							const FTransform BoneTransform = IsSkelMeshValid ? SkeletalMesh->GetBoneTransform(BoneIndex) * LocalTransform.Inverse() : RestTransforms[BoneIndex] * WorldTransform * LocalTransform.Inverse();
 
 							for (const FKBoxElem& BoxElem : BodySetup->AggGeom.BoxElems)
 							{
@@ -453,13 +454,14 @@ void FNDIPhysicsAssetData::Init(UNiagaraDataInterfacePhysicsAsset* Interface, FN
 
 	if (Interface != nullptr && SystemInstance != nullptr)
 	{
-		Interface->ExtractSourceComponent(SystemInstance);
+		FTransform BoneTransform = FTransform::Identity;
+		Interface->ExtractSourceComponent(SystemInstance, BoneTransform);
 		TickingGroup = ComputeTickingGroup(Interface->SourceComponents);
 
 		if(0 < Interface->PhysicsAssets.Num() && Interface->PhysicsAssets[0].IsValid() && Interface->PhysicsAssets[0].Get() != nullptr && 
 			Interface->PhysicsAssets.Num() == Interface->SourceComponents.Num() )
 		{
-			CreateInternalArrays(Interface->PhysicsAssets, Interface->SourceComponents, &AssetArrays, SystemInstance->GetWorldTransform());
+			CreateInternalArrays(Interface->PhysicsAssets, Interface->SourceComponents, &AssetArrays, SystemInstance->GetWorldTransform(), BoneTransform);
 		}
 
 		AssetBuffer = new FNDIPhysicsAssetBuffer();
@@ -491,13 +493,14 @@ void FNDIPhysicsAssetData::Update(UNiagaraDataInterfacePhysicsAsset* Interface,
 {
 	if (Interface != nullptr && SystemInstance != nullptr)
 	{
-		Interface->ExtractSourceComponent(SystemInstance);
+		FTransform BoneTransform = FTransform::Identity;
+		Interface->ExtractSourceComponent(SystemInstance, BoneTransform);
 		TickingGroup = ComputeTickingGroup(Interface->SourceComponents);
 
 		if (0 < Interface->PhysicsAssets.Num() && Interface->PhysicsAssets[0].IsValid() && Interface->PhysicsAssets[0].Get() != nullptr &&
 			Interface->PhysicsAssets.Num() == Interface->SourceComponents.Num())
 		{
-			UpdateInternalArrays(Interface->PhysicsAssets, Interface->SourceComponents, &AssetArrays, SystemInstance->GetWorldTransform());
+			UpdateInternalArrays(Interface->PhysicsAssets, Interface->SourceComponents, &AssetArrays, SystemInstance->GetWorldTransform(), BoneTransform);
 		}
 	}
 }
@@ -539,6 +542,15 @@ public:
 		if (ProxyData != nullptr && ProxyData->AssetBuffer && ProxyData->AssetBuffer->IsInitialized())
 		{
 			FNDIPhysicsAssetBuffer* AssetBuffer = ProxyData->AssetBuffer;
+
+			FRHITransitionInfo Transitions[] = {
+				FRHITransitionInfo(AssetBuffer->WorldTransformBuffer.UAV, ERHIAccess::Unknown, ERHIAccess::SRVCompute),
+				FRHITransitionInfo(AssetBuffer->InverseTransformBuffer.UAV, ERHIAccess::Unknown, ERHIAccess::SRVCompute),
+				FRHITransitionInfo(AssetBuffer->ElementExtentBuffer.UAV, ERHIAccess::Unknown, ERHIAccess::SRVCompute),
+				FRHITransitionInfo(AssetBuffer->PhysicsTypeBuffer.UAV, ERHIAccess::Unknown, ERHIAccess::SRVCompute)
+			};
+			RHICmdList.Transition(MakeArrayView(Transitions, UE_ARRAY_COUNT(Transitions)));
+
 			SetSRVParameter(RHICmdList, ComputeShaderRHI, WorldTransformBuffer, AssetBuffer->WorldTransformBuffer.SRV);
 			SetSRVParameter(RHICmdList, ComputeShaderRHI, InverseTransformBuffer, AssetBuffer->InverseTransformBuffer.SRV);
 			SetSRVParameter(RHICmdList, ComputeShaderRHI, ElementExtentBuffer, AssetBuffer->ElementExtentBuffer.SRV);
@@ -629,9 +641,16 @@ void FNDIPhysicsAssetProxy::PreStage(FRHICommandList& RHICmdList, const FNiagara
 	{
 		if (Context.SimStageData->bFirstStage)
 		{
+			FRHITransitionInfo Transitions[] = {
+				FRHITransitionInfo(ProxyData->AssetBuffer->WorldTransformBuffer.UAV, ERHIAccess::Unknown, ERHIAccess::UAVCompute),
+				FRHITransitionInfo(ProxyData->AssetBuffer->InverseTransformBuffer.UAV, ERHIAccess::Unknown, ERHIAccess::UAVCompute),
+				FRHITransitionInfo(ProxyData->AssetBuffer->PhysicsTypeBuffer.UAV, ERHIAccess::Unknown, ERHIAccess::UAVCompute),
+				FRHITransitionInfo(ProxyData->AssetBuffer->ElementExtentBuffer.UAV, ERHIAccess::Unknown, ERHIAccess::UAVCompute)
+			};
+			RHICmdList.Transition(MakeArrayView(Transitions, UE_ARRAY_COUNT(Transitions)));
+
 			UpdateInternalBuffer<FVector4, EPixelFormat::PF_A32B32G32R32F, PHYSICS_ASSET_MAX_TRANSFORMS, 3>(ProxyData->AssetArrays.WorldTransform, ProxyData->AssetBuffer->WorldTransformBuffer);
 			UpdateInternalBuffer<FVector4, EPixelFormat::PF_A32B32G32R32F, PHYSICS_ASSET_MAX_TRANSFORMS, 3>(ProxyData->AssetArrays.InverseTransform, ProxyData->AssetBuffer->InverseTransformBuffer);
-
 			UpdateInternalBuffer<FVector4, EPixelFormat::PF_A32B32G32R32F, PHYSICS_ASSET_MAX_PRIMITIVES>(ProxyData->AssetArrays.ElementExtent, ProxyData->AssetBuffer->ElementExtentBuffer);
 			UpdateInternalBuffer<uint32, EPixelFormat::PF_R32_UINT, PHYSICS_ASSET_MAX_PRIMITIVES>(ProxyData->AssetArrays.PhysicsType, ProxyData->AssetBuffer->PhysicsTypeBuffer);
 		}
@@ -653,8 +672,10 @@ UNiagaraDataInterfacePhysicsAsset::UNiagaraDataInterfacePhysicsAsset(FObjectInit
 	Proxy.Reset(new FNDIPhysicsAssetProxy());
 }
 
-void UNiagaraDataInterfacePhysicsAsset::ExtractSourceComponent(FNiagaraSystemInstance* SystemInstance)
+void UNiagaraDataInterfacePhysicsAsset::ExtractSourceComponent(FNiagaraSystemInstance* SystemInstance, FTransform& BoneTransform)
 {
+	BoneTransform = FTransform::Identity;
+
 	// Track down the source component
 	TWeakObjectPtr<USkeletalMeshComponent> SourceComponent;
 	if (SourceActor)
@@ -697,10 +718,15 @@ void UNiagaraDataInterfacePhysicsAsset::ExtractSourceComponent(FNiagaraSystemIns
 	for (USceneComponent* Curr = SystemInstance->GetAttachComponent(); Curr; Curr = Curr->GetAttachParent())
 	{
 		UGroomComponent* GroomComponent = Cast<UGroomComponent>(Curr);
-		if (GroomComponent && GroomComponent->PhysicsAsset)
+		if (GroomComponent)
 		{
-			GroomPhysicsAsset = GroomComponent->PhysicsAsset;
-			break;
+			GroomComponent->BuildSimulationTransform(BoneTransform);
+
+			if (GroomComponent->PhysicsAsset)
+			{
+				GroomPhysicsAsset = GroomComponent->PhysicsAsset;
+				break;
+			}
 		}
 	}
 
@@ -948,6 +974,19 @@ void UNiagaraDataInterfacePhysicsAsset::GetFunctions(TArray<FNiagaraFunctionSign
 
 		OutFunctions.Add(Sig);
 	}
+	{
+		FNiagaraFunctionSignature Sig;
+		Sig.Name = GetRestDistanceName;
+		Sig.bSupportsGPU = true;
+		Sig.bSupportsCPU = false;
+		Sig.bMemberFunction = true;
+		Sig.bRequiresContext = false;
+		Sig.Inputs.Add(FNiagaraVariable(FNiagaraTypeDefinition(GetClass()), TEXT("Physics Asset")));
+		Sig.Inputs.Add(FNiagaraVariable(FNiagaraTypeDefinition::GetVec3Def(), TEXT("Rest Position")));
+		Sig.Outputs.Add(FNiagaraVariable(FNiagaraTypeDefinition::GetFloatDef(), TEXT("Rest Distance")));
+
+		OutFunctions.Add(Sig);
+	}
 	{
 		FNiagaraFunctionSignature Sig;
 		Sig.Name = GetTexturePointName;
@@ -1200,6 +1239,17 @@ bool UNiagaraDataInterfacePhysicsAsset::GetFunctionHLSL(const FNiagaraDataInterf
 		OutHLSL += FString::Format(FormatSample, ArgsSample);
 		return true;
 	}
+	else if (FunctionInfo.DefinitionName == GetRestDistanceName)
+	{
+		static const TCHAR* FormatSample = TEXT(R"(
+		void {InstanceFunctionName}(in float3 RestPosition, out float OutRestDistance)
+		{
+			{PhysicsAssetContextName} DIPhysicsAsset_GetRestDistance(DIContext,RestPosition,OutRestDistance);
+		}
+		)");
+		OutHLSL += FString::Format(FormatSample, ArgsSample);
+		return true;
+	}
 	else if (FunctionInfo.DefinitionName == GetTexturePointName)
 	{
 		static const TCHAR* FormatSample = TEXT(R"(
diff --git a/Engine/Plugins/Runtime/HairStrands/Source/HairStrandsCore/Private/Niagara/NiagaraDataInterfaceVelocityGrid.cpp b/Engine/Plugins/Runtime/HairStrands/Source/HairStrandsCore/Private/Niagara/NiagaraDataInterfaceVelocityGrid.cpp
--- a/Engine/Plugins/Runtime/HairStrands/Source/HairStrandsCore/Private/Niagara/NiagaraDataInterfaceVelocityGrid.cpp
+++ b/Engine/Plugins/Runtime/HairStrands/Source/HairStrandsCore/Private/Niagara/NiagaraDataInterfaceVelocityGrid.cpp
@@ -174,6 +174,12 @@ void FNDIVelocityGridParametersCS::Set(FRHICommandList& RHICmdList, const FNiaga
 		FNDIVelocityGridBuffer* CurrentGridBuffer = ProxyData->CurrentGridBuffer;
 		FNDIVelocityGridBuffer* DestinationGridBuffer = ProxyData->DestinationGridBuffer;
 
+		FRHITransitionInfo Transitions[] = {
+			FRHITransitionInfo(CurrentGridBuffer->GridDataBuffer.UAV, ERHIAccess::Unknown, ERHIAccess::SRVCompute),
+			FRHITransitionInfo(DestinationGridBuffer->GridDataBuffer.UAV, ERHIAccess::Unknown, ERHIAccess::UAVCompute)
+		};
+		RHICmdList.Transition(MakeArrayView(Transitions, UE_ARRAY_COUNT(Transitions)));
+
 		SetUAVParameter(RHICmdList, ComputeShaderRHI, GridDestinationBuffer, DestinationGridBuffer->GridDataBuffer.UAV);
 		SetSRVParameter(RHICmdList, ComputeShaderRHI, GridCurrentBuffer, CurrentGridBuffer->GridDataBuffer.SRV);
 
@@ -601,6 +607,11 @@ inline void ClearBuffer(FRHICommandList& RHICmdList, FNDIVelocityGridBuffer* Des
 
 	if (DestinationGridBufferUAV != nullptr )
 	{
+		FRHITransitionInfo Transitions[] = {
+			FRHITransitionInfo(DestinationGridBufferUAV, ERHIAccess::Unknown, ERHIAccess::UAVCompute),
+		};
+		RHICmdList.Transition(MakeArrayView(Transitions, UE_ARRAY_COUNT(Transitions)));
+
 		RHICmdList.ClearUAVUint(DestinationGridBufferUAV, FUintVector4(0, 0, 0, 0));
 	}
 }
diff --git a/Engine/Plugins/Runtime/HairStrands/Source/HairStrandsCore/Public/GroomAssetPhysics.h b/Engine/Plugins/Runtime/HairStrands/Source/HairStrandsCore/Public/GroomAssetPhysics.h
--- a/Engine/Plugins/Runtime/HairStrands/Source/HairStrandsCore/Public/GroomAssetPhysics.h
+++ b/Engine/Plugins/Runtime/HairStrands/Source/HairStrandsCore/Public/GroomAssetPhysics.h
@@ -342,18 +342,44 @@ struct HAIRSTRANDSCORE_API FHairSimulationConstraints
 };
 
 USTRUCT(BlueprintType)
-struct HAIRSTRANDSCORE_API FHairSimulationSettings
+struct HAIRSTRANDSCORE_API FHairSimulationSetup
 {
 	GENERATED_USTRUCT_BODY()
 
 	/** Reset the simulation in time*/
-	UPROPERTY(EditAnywhere, BlueprintReadWrite, Interp, Category = "GroupsSimulation", meta = (ToolTip = "Boolean to control if we want to reset trhe simulation at some point in time"))
+	UPROPERTY(EditAnywhere, BlueprintReadWrite, Interp, Category = "SimulationSetup", meta = (ToolTip = "Boolean to control if we want to reset trhe simulation at some point in time"))
 	bool bResetSimulation = false;
 
+	/** Boolean to make the simulation strands visible */
+	UPROPERTY(EditAnywhere, BlueprintReadWrite, Interp, Category = "SimulationSetup", meta = (ToolTip = "Boolean to make the simulation strands visible"))
+	bool bDebugSimulation = false;
+
+	/** Strands simulation is done in local space */
+	UPROPERTY(EditAnywhere, BlueprintReadWrite, Interp, Category = "SimulationSetup", meta = (ToolTip = "Strands simulation is done in local space"))
+	bool bLocalSimulation = false;
+
+	/** Bone used for the simulation local space */
+	UPROPERTY(EditAnywhere, BlueprintReadWrite, Interp, Category = "SimulationSetup", meta = (ToolTip = "Bone used for the simulation local space", EditCondition = "bLocalSimulation"))
+	FString LocalBone;
+
+	/** Teleport distance threshold to reset the simulation */
+	UPROPERTY(EditAnywhere, BlueprintReadWrite, Interp, Category = "SimulationSetup", meta = (ToolTip = "Teleport distance threshold to reset the simulation"))
+	float TeleportDistance = 50.0;
+};
+
+USTRUCT(BlueprintType)
+struct HAIRSTRANDSCORE_API FHairSimulationSettings
+{
+	GENERATED_USTRUCT_BODY()
+
 	/** Override the asset simulation settings */
-	UPROPERTY(EditAnywhere, BlueprintReadWrite, Interp, Category = "GroupsSimulation", meta = (ToolTip = "Boolean to control if we are going to override the groom asset physics settings"))
+	UPROPERTY(EditAnywhere, BlueprintReadWrite, Interp, Category = "GroupsPhysics", meta = (ToolTip = "Boolean to control if we are going to override the groom asset physics settings"))
 	bool bOverrideSettings = false;
 
+	/** Solver settings */
+	UPROPERTY(EditAnywhere, BlueprintReadWrite, Interp, Category = "GroupsPhysics", meta = (ToolTip = "Solver Settings for the hair physics"))
+	FHairSimulationSetup SimulationSetup;
+
 	/** Solver settings */
 	UPROPERTY(EditAnywhere, BlueprintReadWrite, Interp, Category = "GroupsPhysics", meta = (ToolTip = "Solver Settings for the hair physics", EditCondition = "bOverrideSettings"))
 	FHairSimulationSolver SolverSettings;
diff --git a/Engine/Plugins/Runtime/HairStrands/Source/HairStrandsCore/Public/GroomComponent.h b/Engine/Plugins/Runtime/HairStrands/Source/HairStrandsCore/Public/GroomComponent.h
--- a/Engine/Plugins/Runtime/HairStrands/Source/HairStrandsCore/Public/GroomComponent.h
+++ b/Engine/Plugins/Runtime/HairStrands/Source/HairStrandsCore/Public/GroomComponent.h
@@ -160,7 +160,7 @@ public:
 	void SetEnableSimulation(bool bInEnableSimulation);
 
 	void SetStableRasterization(bool bEnable);
-	void SetGroomAsset(UGroomAsset* Asset, UGroomBindingAsset* InBinding);
+	void SetGroomAsset(UGroomAsset* Asset, UGroomBindingAsset* InBinding, const bool bUpdateSimulation = true);
 	void SetHairRootScale(float Scale);
 	void SetHairWidth(float HairWidth);
 	void SetScatterSceneLighting(bool Enable);
@@ -222,6 +222,9 @@ public:
 	float GetAnimationTime() const;
 	bool IsLooping() const { return bLooping; }
 
+	/** Build the local simulation transform that could be used in strands simulation */
+	void BuildSimulationTransform(FTransform& SimulationTransform) const;
+
 private:
 	void UpdateGroomCache(float Time);
 
diff --git a/Engine/Plugins/Runtime/HairStrands/Source/HairStrandsCore/Public/Niagara/NiagaraDataInterfaceHairStrands.h b/Engine/Plugins/Runtime/HairStrands/Source/HairStrandsCore/Public/Niagara/NiagaraDataInterfaceHairStrands.h
--- a/Engine/Plugins/Runtime/HairStrands/Source/HairStrandsCore/Public/Niagara/NiagaraDataInterfaceHairStrands.h
+++ b/Engine/Plugins/Runtime/HairStrands/Source/HairStrandsCore/Public/Niagara/NiagaraDataInterfaceHairStrands.h
@@ -28,7 +28,7 @@ struct FNDIHairStrandsBuffer : public FRenderResource
 		const FHairStrandsDeformedResource*  HairStrandsDeformedResource, 
 		const FHairStrandsRestRootResource* HairStrandsRestRootResource, 
 		const FHairStrandsDeformedRootResource* HairStrandsDeformedRootResource,
-		const TStaticArray<float, 32 * NumScales>& InParamsScale);  
+		const TStaticArray<float, 32 * NumScales>& InParamsScale);
 
 	/** Set the asset that will be used to affect the buffer */
 	void Update(
@@ -99,6 +99,7 @@ struct FNDIHairStrandsData
 	inline void ResetDatas()
 	{
 		WorldTransform.SetIdentity();
+		BoneTransform.SetIdentity();
 		GlobalInterpolation = false;
 		HairGroupInstance = nullptr;
 
@@ -144,6 +145,7 @@ struct FNDIHairStrandsData
 			ParamsScale[i] = 1.0;
 		}
 		SkeletalMeshes = 0;
+		LocalSimulation = false;
 	}
 
 	inline void CopyDatas(const FNDIHairStrandsData* OtherDatas)
@@ -153,6 +155,7 @@ struct FNDIHairStrandsData
 			HairStrandsBuffer = OtherDatas->HairStrandsBuffer;
 
 			WorldTransform = OtherDatas->WorldTransform;
+			BoneTransform = OtherDatas->BoneTransform;
 
 			GlobalInterpolation = OtherDatas->GlobalInterpolation;
 			BindingType = OtherDatas->BindingType;
@@ -198,12 +201,16 @@ struct FNDIHairStrandsData
 			SkeletalMeshes = OtherDatas->SkeletalMeshes;
 
 			TickingGroup = OtherDatas->TickingGroup;
+			LocalSimulation = OtherDatas->LocalSimulation;
 		}
 	}
 
 	/** Cached World transform. */
 	FTransform WorldTransform;
 
+	/** Bone transform that will be used for local strands simulation */
+	FTransform BoneTransform;
+
 	/** Global Interpolation */
 	bool GlobalInterpolation;
 
@@ -305,6 +312,9 @@ struct FNDIHairStrandsData
 
 	/** The instance ticking group */
 	ETickingGroup TickingGroup;
+
+	/** Check if the simulation is running in local coordinate */
+	bool LocalSimulation;
 };
 
 /** Data Interface for the strand base */
@@ -593,12 +603,21 @@ public:
 	/** Name of the bounding box offsets*/
 	static const FString BoundingBoxOffsetsName;
 
-	/** Name of the world transform */
+	/** Name of the world inverse transform */
 	static const FString WorldInverseName;
 
 	/** Name of the world rotation */
 	static const FString WorldRotationName;
 
+	/** Name of the bone transform */
+	static const FString BoneTransformName;
+
+	/** Name of the bone inverse transform */
+	static const FString BoneInverseName;
+
+	/** Name of the world rotation */
+	static const FString BoneRotationName;
+
 	/** Name of the number of strands */
 	static const FString NumStrandsName;
 
@@ -623,6 +642,9 @@ public:
 	/** Param to check if we need to update the rest pose */
 	static const FString RestUpdateName;
 
+	/** Param to check if the simulation is going to be run in local space */
+	static const FString LocalSimulationName;
+
 	/** boolean to check if we need to rest the simulation*/
 	static const FString ResetSimulationName;
 
diff --git a/Engine/Plugins/Runtime/HairStrands/Source/HairStrandsCore/Public/Niagara/NiagaraDataInterfacePhysicsAsset.h b/Engine/Plugins/Runtime/HairStrands/Source/HairStrandsCore/Public/Niagara/NiagaraDataInterfacePhysicsAsset.h
--- a/Engine/Plugins/Runtime/HairStrands/Source/HairStrandsCore/Public/Niagara/NiagaraDataInterfacePhysicsAsset.h
+++ b/Engine/Plugins/Runtime/HairStrands/Source/HairStrandsCore/Public/Niagara/NiagaraDataInterfacePhysicsAsset.h
@@ -145,7 +145,7 @@ public:
 	virtual void ProvidePerInstanceDataForRenderThread(void* DataForRenderThread, void* PerInstanceData, const FNiagaraSystemInstanceID& SystemInstance) override;
 
 	/** Extract the source component */
-	void ExtractSourceComponent(FNiagaraSystemInstance* SystemInstance);
+	void ExtractSourceComponent(FNiagaraSystemInstance* SystemInstance, FTransform& LocalTransform);
 
 	/** Get the number of boxes*/
 	void GetNumBoxes(FVectorVMContext& Context);
