Fix crash on platform which does not support ShaderDrawDebug.

diff --git a/Engine/Plugins/Runtime/HairStrands/Source/HairStrandsCore/Private/HairStrandsRendering.cpp b/Engine/Plugins/Runtime/HairStrands/Source/HairStrandsCore/Private/HairStrandsRendering.cpp
--- a/Engine/Plugins/Runtime/HairStrands/Source/HairStrandsCore/Private/HairStrandsRendering.cpp
+++ b/Engine/Plugins/Runtime/HairStrands/Source/HairStrandsCore/Private/HairStrandsRendering.cpp
@@ -702,7 +702,7 @@ static void AddHairStrandsInterpolationPass(
 	if (bSupportDynamicMesh && (Instance->Guides.bIsSimulationEnable || Instance->Guides.bHasGlobalInterpolation))
 	{
 		const FHairStrandsRestRootResource::FLOD& Sim_RestLODDatas = SimRestRootResources->LODs[MeshLODIndex];
-		const FHairStrandsDeformedRootResource::FLOD& Sim_DeformedLODDatas = SimDeformedRootResources->LODs[MeshLODIndex];\
+		const FHairStrandsDeformedRootResource::FLOD& Sim_DeformedLODDatas = SimDeformedRootResources->LODs[MeshLODIndex];
 		bSupportGlobalInterpolation = Instance->Guides.bHasGlobalInterpolation && (Sim_RestLODDatas.SampleCount > 0);
 		{
 			Parameters->SimRestPosition0Buffer = RegisterAsSRV(GraphBuilder, Sim_RestLODDatas.RestRootTrianglePosition0Buffer);
@@ -736,7 +736,7 @@ static void AddHairStrandsInterpolationPass(
 		}
 	}
 
-	if (ShaderDrawDebug::IsEnabled() && ShaderDrawData)
+	if (ShaderDrawData)
 	{
 		ShaderDrawDebug::SetParameters(GraphBuilder, *ShaderDrawData, Parameters->ShaderDrawParameters);
 	}
@@ -889,7 +889,7 @@ static void AddHairClusterAABBPass(
 	Parameters->RenderDeformedPositionBuffer = RenderPositionBufferSRV;
 	Parameters->RenderDeformedOffsetBuffer = RenderDeformedOffsetBuffer;
 	Parameters->TotalClusterCount = 1;
-	if (ShaderDrawDebug::IsEnabled() && ShaderDrawData)
+	if (ShaderDrawData)
 	{
 		ShaderDrawDebug::SetParameters(GraphBuilder, *ShaderDrawData, Parameters->ShaderDrawParameters);
 		
@@ -1059,7 +1059,7 @@ static void AddHairCardsDeformationPass(
 		Parameters->TriangleDeformedPosition2Buffer = RegisterAsSRV(GraphBuilder, DeformedLODDatas.DeformedRootTrianglePosition2Buffer);
 	}
 
-	if (ShaderDrawDebug::IsEnabled() && ShaderDrawData)
+	if (ShaderDrawData)
 	{
 		ShaderDrawDebug::SetParameters(GraphBuilder, *ShaderDrawData, Parameters->ShaderDrawParameters);
 	}
diff --git a/Engine/Source/Runtime/Renderer/Private/HairStrands/HairStrandsInterface.cpp b/Engine/Source/Runtime/Renderer/Private/HairStrands/HairStrandsInterface.cpp
--- a/Engine/Source/Runtime/Renderer/Private/HairStrands/HairStrandsInterface.cpp
+++ b/Engine/Source/Runtime/Renderer/Private/HairStrands/HairStrandsInterface.cpp
@@ -447,7 +447,7 @@ FHairStrandsBookmarkParameters CreateHairStrandsBookmarkParameters(FScene* Scene
 			}
 		}
 	}
-	Out.DebugShaderData			= &View.ShaderDrawData;
+	Out.DebugShaderData			= ShaderDrawDebug::IsEnabled(View) ? &View.ShaderDrawData : nullptr;
 	Out.SkinCache				= View.Family->Scene->GetGPUSkinCache();
 	Out.ShaderMap				= View.ShaderMap;
 	Out.Instances				= &Scene->HairStrandsSceneData.RegisteredProxies;
