Fix crash when generating hair strands texture with very dense groom.

diff --git a/Engine/Plugins/Runtime/HairStrands/Source/HairStrandsCore/Private/GroomAsset.cpp b/Engine/Plugins/Runtime/HairStrands/Source/HairStrandsCore/Private/GroomAsset.cpp
--- a/Engine/Plugins/Runtime/HairStrands/Source/HairStrandsCore/Private/GroomAsset.cpp
+++ b/Engine/Plugins/Runtime/HairStrands/Source/HairStrandsCore/Private/GroomAsset.cpp
@@ -3282,7 +3282,7 @@ void UGroomAsset::CreateDebugData()
 		FGroomBuilder::BuildData(LocalHairDescriptionGroups.HairGroups[GroupIndex], HairGroupsInterpolation[GroupIndex], DummyInfo, StrandsData, GuidesData);
 
 		FHairGroupData& GroupData = HairGroupsData[GroupIndex];
-		CreateHairStrandsDebugDatas(StrandsData, 1.f, GroupData.Debug.Data);
+		CreateHairStrandsDebugDatas(StrandsData, GroupData.Debug.Data);
 
 		if (GroupData.Debug.Data.IsValid())
 		{
diff --git a/Engine/Plugins/Runtime/HairStrands/Source/HairStrandsCore/Private/GroomResources.cpp b/Engine/Plugins/Runtime/HairStrands/Source/HairStrandsCore/Private/GroomResources.cpp
--- a/Engine/Plugins/Runtime/HairStrands/Source/HairStrandsCore/Private/GroomResources.cpp
+++ b/Engine/Plugins/Runtime/HairStrands/Source/HairStrandsCore/Private/GroomResources.cpp
@@ -33,6 +33,11 @@ static FAutoConsoleVariableRef CVarHairStrandsBulkData_ReleaseAfterUse(TEXT("r.H
 static int32 GHairStrandsBulkData_Validation = 1;
 static FAutoConsoleVariableRef CVarHairStrandsBulkData_Validation(TEXT("r.HairStrands.Strands.BulkData.Validation"), GHairStrandsBulkData_Validation, TEXT("Validate some hair strands data at serialization/loading time."));
 
+static float GHairStrandsDebugVoxel_WorldSize = 1.0f;
+static int32 GHairStrandsDebugVoxel_MaxSegmentPerVoxel = 512;
+static FAutoConsoleVariableRef CVarHairStrandsDebugVoxel_WorldSize(TEXT("r.HairStrands.DebugData.VoxelSize"), GHairStrandsDebugVoxel_WorldSize, TEXT("Voxel size use for creating debug data."));
+static FAutoConsoleVariableRef CVarHairStrandsDebugVoxel_MaxSegmentPerVoxel(TEXT("r.HairStrands.DebugData.MaxSegmentPerVoxel"), GHairStrandsDebugVoxel_MaxSegmentPerVoxel, TEXT("Max number of segments per Voxel size when creating debug data."));
+
 ////////////////////////////////////////////////////////////////////////////////////
 
 namespace HairTransition
@@ -1422,15 +1427,15 @@ static FIntVector ToCoord(const FVector& T, const FIntVector& Resolution, const
 
 void CreateHairStrandsDebugDatas(
 	const FHairStrandsDatas& InData,
-	float WorldVoxelSize,
 	FHairStrandsDebugDatas& Out)
 {
 	const FVector BoundSize = InData.BoundingBox.Max - InData.BoundingBox.Min;
-	Out.VoxelDescription.VoxelSize = WorldVoxelSize;
-	Out.VoxelDescription.VoxelResolution = FIntVector(FMath::CeilToFloat(BoundSize.X / Out.VoxelDescription.VoxelSize), FMath::CeilToFloat(BoundSize.Y / Out.VoxelDescription.VoxelSize), FMath::CeilToFloat(BoundSize.Z / Out.VoxelDescription.VoxelSize));
+	Out.VoxelDescription.VoxelSize = FMath::Clamp(GHairStrandsDebugVoxel_WorldSize, 0.1f, 10.f);
+	Out.VoxelDescription.VoxelResolution = FIntVector(FMath::CeilToInt(BoundSize.X / Out.VoxelDescription.VoxelSize), FMath::CeilToInt(BoundSize.Y / Out.VoxelDescription.VoxelSize), FMath::CeilToInt(BoundSize.Z / Out.VoxelDescription.VoxelSize));
 	Out.VoxelDescription.VoxelMinBound = InData.BoundingBox.Min;
 	Out.VoxelDescription.VoxelMaxBound = FVector(Out.VoxelDescription.VoxelResolution) * Out.VoxelDescription.VoxelSize + InData.BoundingBox.Min;
 	Out.VoxelOffsetAndCount.Init(FHairStrandsDebugDatas::FOffsetAndCount(), Out.VoxelDescription.VoxelResolution.X * Out.VoxelDescription.VoxelResolution.Y * Out.VoxelDescription.VoxelResolution.Z);
+	Out.VoxelDescription.MaxSegmentPerVoxel = 0;
 
 	uint32 AllocationCount = 0;
 	TArray<TArray<FHairStrandsDebugDatas::FVoxel>> TempVoxelData;
@@ -1467,19 +1472,19 @@ void CreateHairStrandsDebugDatas(
 					}
 
 					const uint32 Offset = Out.VoxelOffsetAndCount[LinearCoord].Offset;
-					const uint32 LocalOffset = Out.VoxelOffsetAndCount[LinearCoord].Count;
 					Out.VoxelOffsetAndCount[LinearCoord].Count += 1;
 					TempVoxelData[Offset].Add({Index0, Index1});
 
+					Out.VoxelDescription.MaxSegmentPerVoxel = FMath::Max(Out.VoxelDescription.MaxSegmentPerVoxel, Out.VoxelOffsetAndCount[LinearCoord].Count);
+
 					PrevLinearCoord = LinearCoord;
 
 					++AllocationCount;
 				}
 			}
 		}
-
 	}
-
+	Out.VoxelDescription.MaxSegmentPerVoxel = FMath::Min(Out.VoxelDescription.MaxSegmentPerVoxel, uint32(FMath::Clamp(GHairStrandsDebugVoxel_MaxSegmentPerVoxel, 16, 64000)));
 	Out.VoxelData.Reserve(AllocationCount);
 
 	for (int32 Index = 0, Count = Out.VoxelOffsetAndCount.Num(); Index < Count; ++Index)
diff --git a/Engine/Plugins/Runtime/HairStrands/Source/HairStrandsCore/Private/GroomTextureBuilder.cpp b/Engine/Plugins/Runtime/HairStrands/Source/HairStrandsCore/Private/GroomTextureBuilder.cpp
--- a/Engine/Plugins/Runtime/HairStrands/Source/HairStrandsCore/Private/GroomTextureBuilder.cpp
+++ b/Engine/Plugins/Runtime/HairStrands/Source/HairStrandsCore/Private/GroomTextureBuilder.cpp
@@ -583,6 +583,7 @@ class FHairStrandsTexturePS : public FGlobalShader
 		SHADER_PARAMETER(FVector3f, Voxel_MaxBound)
 		SHADER_PARAMETER(FIntVector, Voxel_Resolution)
 		SHADER_PARAMETER(float, Voxel_Size)
+		SHADER_PARAMETER(uint32, Voxel_MaxSegmentPerVoxel)
 		SHADER_PARAMETER_RDG_BUFFER_SRV(StructuredBuffer, Voxel_OffsetAndCount)
 		SHADER_PARAMETER_RDG_BUFFER_SRV(StructuredBuffer, Voxel_Data)
 
@@ -629,6 +630,7 @@ static void InternalGenerateHairStrandsTextures(
 	const FVector& VoxelMaxBound,
 	const FIntVector& VoxelResolution,
 	float VoxelSize,
+	uint32 VoxelMaxSegmentPerVoxel,
 	FRDGBufferRef VoxelOffsetAndCount,
 	FRDGBufferRef VoxelData,
 	
@@ -671,6 +673,7 @@ static void InternalGenerateHairStrandsTextures(
 	ParametersPS->Voxel_MaxBound = VoxelMaxBound;
 	ParametersPS->Voxel_Resolution = VoxelResolution;
 	ParametersPS->Voxel_Size = VoxelSize;
+	ParametersPS->Voxel_MaxSegmentPerVoxel = VoxelMaxSegmentPerVoxel;
 	ParametersPS->Voxel_OffsetAndCount = GraphBuilder.CreateSRV(VoxelOffsetAndCount);
 	ParametersPS->Voxel_Data = GraphBuilder.CreateSRV(VoxelData);
 
@@ -738,8 +741,10 @@ static void InternalGenerateHairStrandsTextures(
 				RHICmdList.DrawIndexedPrimitive(InMeshIndexBuffer, VertexBaseIndex, 0, VertexCount, IndexBaseIndex, PrimitiveCount, 1);
 
 				// Flush, to ensure that all texture generation is done (TDR)
+				#if 0
 				GDynamicRHI->RHISubmitCommandsAndFlushGPU();
 				GDynamicRHI->RHIBlockUntilGPUIdle();
+				#endif
 			}
 			else
 			{
@@ -754,8 +759,10 @@ static void InternalGenerateHairStrandsTextures(
 					RHICmdList.DrawIndexedPrimitive(InMeshIndexBuffer, VertexBaseIndex, 0, VertexCount, IndexBaseIndex, PrimitiveCount, 1);
 
 					// Flush, to ensure that all texture generation is done (TDR)
+					#if 0
 					GDynamicRHI->RHISubmitCommandsAndFlushGPU();
 					GDynamicRHI->RHIBlockUntilGPUIdle();
+					#endif
 				}
 			}
 		});
@@ -949,6 +956,9 @@ static void InternalBuildStrandsTextures_GPU(
 				VertexBaseIndex = 0;
 			}
 
+			// Ensure the rest resources are loaded when rendering the strands textures
+			if (GroupData.Strands.RestResource) { GroupData.Strands.RestResource->Allocate(GraphBuilder); }
+
 			InternalGenerateHairStrandsTextures(
 				GraphBuilder,
 				ShaderMap,
@@ -974,6 +984,7 @@ static void InternalBuildStrandsTextures_GPU(
 				GroupData.Debug.Resource->VoxelDescription.VoxelMaxBound,
 				GroupData.Debug.Resource->VoxelDescription.VoxelResolution,
 				GroupData.Debug.Resource->VoxelDescription.VoxelSize,
+				GroupData.Debug.Resource->VoxelDescription.MaxSegmentPerVoxel,
 				VoxelOffsetAndCount,
 				VoxelData,
 
diff --git a/Engine/Plugins/Runtime/HairStrands/Source/HairStrandsCore/Public/GroomResources.h b/Engine/Plugins/Runtime/HairStrands/Source/HairStrandsCore/Public/GroomResources.h
--- a/Engine/Plugins/Runtime/HairStrands/Source/HairStrandsCore/Public/GroomResources.h
+++ b/Engine/Plugins/Runtime/HairStrands/Source/HairStrandsCore/Public/GroomResources.h
@@ -731,5 +731,5 @@ struct FHairMeshesDeformedResource : public FHairCommonResource
 
 ////////////////////////////////////////////////////////////////////////////////////////////////////////////////////////
 // Debug data (used for debug visalization but also for texture generation)
-void CreateHairStrandsDebugDatas(const FHairStrandsDatas& InData, float WorldVoxelSize, FHairStrandsDebugDatas& Out);
+void CreateHairStrandsDebugDatas(const FHairStrandsDatas& InData, FHairStrandsDebugDatas& Out);
 void CreateHairStrandsDebugResources(class FRDGBuilder& GraphBuilder, const FHairStrandsDebugDatas* In, FHairStrandsDebugDatas::FResources* Out);
\ No newline at end of file
diff --git a/Engine/Plugins/Runtime/HairStrands/Source/HairStrandsCore/Public/HairStrandsDatas.h b/Engine/Plugins/Runtime/HairStrands/Source/HairStrandsCore/Public/HairStrandsDatas.h
--- a/Engine/Plugins/Runtime/HairStrands/Source/HairStrandsCore/Public/HairStrandsDatas.h
+++ b/Engine/Plugins/Runtime/HairStrands/Source/HairStrandsCore/Public/HairStrandsDatas.h
@@ -478,10 +478,11 @@ struct HAIRSTRANDSCORE_API FHairStrandsDebugDatas
 
 	struct FDesc
 	{
-		FVector VoxelMinBound;
-		FVector VoxelMaxBound;
-		FIntVector VoxelResolution;
-		float VoxelSize;
+		FVector VoxelMinBound = FVector::ZeroVector;
+		FVector VoxelMaxBound = FVector::ZeroVector;
+		FIntVector VoxelResolution = FIntVector::ZeroValue;
+		float VoxelSize = 0;
+		uint32 MaxSegmentPerVoxel = 0;
 	};
 
 	FDesc VoxelDescription;
diff --git a/Engine/Plugins/Runtime/HairStrands/Source/HairStrandsEditor/Private/GroomActions.cpp b/Engine/Plugins/Runtime/HairStrands/Source/HairStrandsEditor/Private/GroomActions.cpp
--- a/Engine/Plugins/Runtime/HairStrands/Source/HairStrandsEditor/Private/GroomActions.cpp
+++ b/Engine/Plugins/Runtime/HairStrands/Source/HairStrandsEditor/Private/GroomActions.cpp
@@ -436,9 +436,11 @@ void FGroomActions::ExecuteCreateStrandsTextures(TArray<TWeakObjectPtr<UGroomAss
 				case EStrandsTexturesMeshType::Static: StaticMesh = CurrentOptions->StaticMesh; break;
 				case EStrandsTexturesMeshType::Skeletal: SkeletalMesh = CurrentOptions->SkeletalMesh; break;
 				}
+				if (SkeletalMesh == nullptr && StaticMesh == nullptr)
+				{
+					return;
+				}
 				
-				EStrandsTexturesMeshType MeshType = EStrandsTexturesMeshType::Skeletal;
-
 				FStrandsTexturesInfo Info;
 				Info.GroomAsset   = GroomAsset.Get();
 				Info.TracingDirection = SignDirection;
@@ -460,7 +462,6 @@ void FGroomActions::ExecuteCreateStrandsTextures(TArray<TWeakObjectPtr<UGroomAss
 						Info.GroupIndices.Add(GroupIndex);
 					}
 				}
-
 				FStrandsTexturesOutput Output = FGroomTextureBuilder::CreateGroomStrandsTexturesTexture(GroomAsset.Get(), Info.Resolution);
 				if (Output.IsValid())
 				{
diff --git a/Engine/Shaders/Private/HairStrands/HairStrandsTexturesGeneration.usf b/Engine/Shaders/Private/HairStrands/HairStrandsTexturesGeneration.usf
--- a/Engine/Shaders/Private/HairStrands/HairStrandsTexturesGeneration.usf
+++ b/Engine/Shaders/Private/HairStrands/HairStrandsTexturesGeneration.usf
@@ -100,6 +100,7 @@ float3 Voxel_MinBound;
 float3 Voxel_MaxBound;
 int3 Voxel_Resolution;
 float Voxel_Size;
+uint Voxel_MaxSegmentPerVoxel;
 StructuredBuffer<uint2> Voxel_OffsetAndCount;
 StructuredBuffer<uint2> Voxel_Data;
 
@@ -202,9 +203,6 @@ bool Trace(
 			PageIndexCoord = clamp(uint3(Coords), uint3(0, 0, 0), Voxel_Resolution - 1);
 		}
 
-		// Page stepping is the walking quantity (i.e. number of voxel) for fine ray-marching within a valid page
-		float PageStepping = 1;
-
 		bool bLastStep = false;
 		const uint LoopCount = 256u;
 		for (uint LoopIt = 0; LoopIt < LoopCount; ++LoopIt)
@@ -251,15 +249,18 @@ bool Trace(
 				{
 					const uint LinearIndexCoord = CoordToIndex(VolumeCoord, Voxel_Resolution, 0);
 					OffsetAndCount = Voxel_OffsetAndCount[LinearIndexCoord];
+					OffsetAndCount.y = min(OffsetAndCount.y, Voxel_MaxSegmentPerVoxel);
 				}
 
 				if (OffsetAndCount.y > 0)
 				{
+					const uint MaxVoxelCount = Voxel_Resolution.x * Voxel_Resolution.y * Voxel_Resolution.z;
 					for (uint PointIt = 0; PointIt < OffsetAndCount.y; ++PointIt)
 					{
+						const uint VoxelIndex = clamp(OffsetAndCount.x + PointIt, 0, MaxVoxelCount);
 						const uint2 P0_P1 = Voxel_Data[OffsetAndCount.x + PointIt];
-						const uint P0 = P0_P1.x;
-						const uint P1 = P0_P1.y;
+						const uint P0 = clamp(P0_P1.x, 0, InVF_ControlPointCount);
+						const uint P1 = clamp(P0_P1.y, 0, InVF_ControlPointCount);
 
 						const uint4 PackedCP_P0 = InVF_PositionBuffer[P0];
 						const uint4 PackedCP_P1 = InVF_PositionBuffer[P1];
