FIx CIS static analysis issue.

diff --git a/Engine/Plugins/Runtime/HairStrands/Source/HairStrandsCore/Private/GroomTextureBuilder.cpp b/Engine/Plugins/Runtime/HairStrands/Source/HairStrandsCore/Private/GroomTextureBuilder.cpp
--- a/Engine/Plugins/Runtime/HairStrands/Source/HairStrandsCore/Private/GroomTextureBuilder.cpp
+++ b/Engine/Plugins/Runtime/HairStrands/Source/HairStrandsCore/Private/GroomTextureBuilder.cpp
@@ -957,7 +957,7 @@ static void InternalBuildStrandsTextures_GPU(
 			}
 
 			// Ensure the rest resources are loaded when rendering the strands textures
-			if (GroupData.Strands.RestResource) { GroupData.Strands.RestResource->Allocate(GraphBuilder); }
+			GroupData.Strands.RestResource->Allocate(GraphBuilder);
 
 			InternalGenerateHairStrandsTextures(
 				GraphBuilder,
