Wrapped ObjectPtr upgrade refresh for Engine + ShooterGame

diff --git a/Engine/Plugins/Experimental/MeshModelingToolset/Source/MeshModelingTools/Public/SplitMeshesTool.h b/Engine/Plugins/Experimental/MeshModelingToolset/Source/MeshModelingTools/Public/SplitMeshesTool.h
--- a/Engine/Plugins/Experimental/MeshModelingToolset/Source/MeshModelingTools/Public/SplitMeshesTool.h
+++ b/Engine/Plugins/Experimental/MeshModelingToolset/Source/MeshModelingTools/Public/SplitMeshesTool.h
@@ -52,10 +52,10 @@ public:
 	virtual bool CanAccept() const override;
 
 	UPROPERTY()
-	USplitMeshesToolProperties* BasicProperties;
+	TObjectPtr<USplitMeshesToolProperties> BasicProperties;
 
 	UPROPERTY()
-	UCreateMeshObjectTypeProperties* OutputTypeProperties;
+	TObjectPtr<UCreateMeshObjectTypeProperties> OutputTypeProperties;
 
 protected:
 	UWorld* TargetWorld;
diff --git a/Engine/Plugins/Experimental/MeshModelingToolset/Source/MeshModelingTools/Public/TransferMeshTool.h b/Engine/Plugins/Experimental/MeshModelingToolset/Source/MeshModelingTools/Public/TransferMeshTool.h
--- a/Engine/Plugins/Experimental/MeshModelingToolset/Source/MeshModelingTools/Public/TransferMeshTool.h
+++ b/Engine/Plugins/Experimental/MeshModelingToolset/Source/MeshModelingTools/Public/TransferMeshTool.h
@@ -53,7 +53,7 @@ public:
 	virtual bool CanAccept() const override;
 
 	UPROPERTY()
-	UTransferMeshToolProperties* BasicProperties;
+	TObjectPtr<UTransferMeshToolProperties> BasicProperties;
 
 protected:
 	UWorld* TargetWorld;
diff --git a/Engine/Plugins/Experimental/MeshModelingToolset/Source/MeshModelingTools/Public/TransformMeshesTool.h b/Engine/Plugins/Experimental/MeshModelingToolset/Source/MeshModelingTools/Public/TransformMeshesTool.h
--- a/Engine/Plugins/Experimental/MeshModelingToolset/Source/MeshModelingTools/Public/TransformMeshesTool.h
+++ b/Engine/Plugins/Experimental/MeshModelingToolset/Source/MeshModelingTools/Public/TransformMeshesTool.h
@@ -147,10 +147,10 @@ struct FTransformMeshesTarget
 	GENERATED_BODY()
 
 	UPROPERTY()
-	UTransformProxy* TransformProxy = nullptr;
+	TObjectPtr<UTransformProxy> TransformProxy = nullptr;
 
 	UPROPERTY()
-	UTransformGizmo* TransformGizmo = nullptr;
+	TObjectPtr<UTransformGizmo> TransformGizmo = nullptr;
 };
 
 
@@ -188,7 +188,7 @@ public:
 
 protected:
 	UPROPERTY()
-	UTransformMeshesToolProperties* TransformProps;
+	TObjectPtr<UTransformMeshesToolProperties> TransformProps;
 
 protected:
 	UWorld* TargetWorld;
@@ -198,7 +198,7 @@ protected:
 	TArray<FTransformMeshesTarget> ActiveGizmos;
 
 	UPROPERTY()
-	UDragAlignmentMechanic* DragAlignmentMechanic = nullptr;
+	TObjectPtr<UDragAlignmentMechanic> DragAlignmentMechanic = nullptr;
 
 	ETransformMeshesTransformMode CurTransformMode;
 	void UpdateTransformMode(ETransformMeshesTransformMode NewMode);
diff --git a/Engine/Plugins/Experimental/MeshModelingToolset/Source/MeshModelingTools/Public/UVProjectionTool.h b/Engine/Plugins/Experimental/MeshModelingToolset/Source/MeshModelingTools/Public/UVProjectionTool.h
--- a/Engine/Plugins/Experimental/MeshModelingToolset/Source/MeshModelingTools/Public/UVProjectionTool.h
+++ b/Engine/Plugins/Experimental/MeshModelingToolset/Source/MeshModelingTools/Public/UVProjectionTool.h
@@ -217,7 +217,7 @@ public:
 	virtual TUniquePtr<UE::Geometry::FDynamicMeshOperator> MakeNewOperator() override;
 
 	UPROPERTY()
-	UUVProjectionTool *Tool;
+	TObjectPtr<UUVProjectionTool> Tool;
 };
 
 
@@ -248,41 +248,41 @@ public:
 protected:
 
 	UPROPERTY()
-	UMeshUVChannelProperties* UVChannelProperties = nullptr;
+	TObjectPtr<UMeshUVChannelProperties> UVChannelProperties = nullptr;
 
 	UPROPERTY()
-	UUVProjectionToolProperties* BasicProperties = nullptr;
+	TObjectPtr<UUVProjectionToolProperties> BasicProperties = nullptr;
 
 	UPROPERTY()
-	UUVProjectionToolExpMapProperties* ExpMapProperties = nullptr;
+	TObjectPtr<UUVProjectionToolExpMapProperties> ExpMapProperties = nullptr;
 
 	UPROPERTY()
-	UUVProjectionToolCylinderProperties* CylinderProperties = nullptr;
+	TObjectPtr<UUVProjectionToolCylinderProperties> CylinderProperties = nullptr;
 
 	UPROPERTY()
-	UUVProjectionToolEditActions* EditActions = nullptr;
+	TObjectPtr<UUVProjectionToolEditActions> EditActions = nullptr;
 
 	UPROPERTY()
-	UExistingMeshMaterialProperties* MaterialSettings = nullptr;
+	TObjectPtr<UExistingMeshMaterialProperties> MaterialSettings = nullptr;
 
 	UPROPERTY()
-	UMeshOpPreviewWithBackgroundCompute* Preview = nullptr;
+	TObjectPtr<UMeshOpPreviewWithBackgroundCompute> Preview = nullptr;
 
 
 	UPROPERTY()
-	UMaterialInstanceDynamic* CheckerMaterial = nullptr;
+	TObjectPtr<UMaterialInstanceDynamic> CheckerMaterial = nullptr;
 
 	UPROPERTY()
-	UTransformGizmo* TransformGizmo = nullptr;
+	TObjectPtr<UTransformGizmo> TransformGizmo = nullptr;
 	
 	UPROPERTY()
-	UTransformProxy* TransformProxy = nullptr;
+	TObjectPtr<UTransformProxy> TransformProxy = nullptr;
 
 	UPROPERTY()
-	UUVProjectionOperatorFactory* OperatorFactory = nullptr;
+	TObjectPtr<UUVProjectionOperatorFactory> OperatorFactory = nullptr;
 
 	UPROPERTY()
-	UPreviewGeometry* EdgeRenderer = nullptr;
+	TObjectPtr<UPreviewGeometry> EdgeRenderer = nullptr;
 
 protected:
 	TSharedPtr<UE::Geometry::FDynamicMesh3, ESPMode::ThreadSafe> InputMesh;
@@ -324,7 +324,7 @@ protected:
 	TUniquePtr<FSelectClickedAction> SetPlaneCtrlClickBehaviorTarget;
 
 	UPROPERTY()
-	USingleClickInputBehavior* ClickToSetPlaneBehavior;
+	TObjectPtr<USingleClickInputBehavior> ClickToSetPlaneBehavior;
 
 	void UpdatePlaneFromClick(const FVector3d& Position, const FVector3d& Normal, bool bTransationOnly);
 
diff --git a/Engine/Plugins/Experimental/MeshModelingToolset/Source/MeshModelingTools/Public/VolumeToMeshTool.h b/Engine/Plugins/Experimental/MeshModelingToolset/Source/MeshModelingTools/Public/VolumeToMeshTool.h
--- a/Engine/Plugins/Experimental/MeshModelingToolset/Source/MeshModelingTools/Public/VolumeToMeshTool.h
+++ b/Engine/Plugins/Experimental/MeshModelingToolset/Source/MeshModelingTools/Public/VolumeToMeshTool.h
@@ -85,16 +85,16 @@ public:
 
 protected:
 	UPROPERTY()
-	UVolumeToMeshToolProperties* Settings;
+	TObjectPtr<UVolumeToMeshToolProperties> Settings;
 
 	UPROPERTY()
-	UPreviewMesh* PreviewMesh;
+	TObjectPtr<UPreviewMesh> PreviewMesh;
 
 	UPROPERTY()
 	TLazyObjectPtr<AVolume> TargetVolume;
 
 	UPROPERTY()
-	ULineSetComponent* VolumeEdgesSet;
+	TObjectPtr<ULineSetComponent> VolumeEdgesSet;
 
 protected:
 	UWorld* TargetWorld = nullptr;
diff --git a/Engine/Plugins/Experimental/MeshModelingToolset/Source/MeshModelingTools/Public/VoxelBlendMeshesTool.h b/Engine/Plugins/Experimental/MeshModelingToolset/Source/MeshModelingTools/Public/VoxelBlendMeshesTool.h
--- a/Engine/Plugins/Experimental/MeshModelingToolset/Source/MeshModelingTools/Public/VoxelBlendMeshesTool.h
+++ b/Engine/Plugins/Experimental/MeshModelingToolset/Source/MeshModelingTools/Public/VoxelBlendMeshesTool.h
@@ -83,7 +83,7 @@ protected:
 	virtual TUniquePtr<UE::Geometry::FDynamicMeshOperator> MakeNewOperator() override;
 
 	UPROPERTY()
-	UVoxelBlendMeshesToolProperties* BlendProperties;
+	TObjectPtr<UVoxelBlendMeshesToolProperties> BlendProperties;
 };
 
 
diff --git a/Engine/Plugins/Experimental/MeshModelingToolset/Source/MeshModelingTools/Public/VoxelMorphologyMeshesTool.h b/Engine/Plugins/Experimental/MeshModelingToolset/Source/MeshModelingTools/Public/VoxelMorphologyMeshesTool.h
--- a/Engine/Plugins/Experimental/MeshModelingToolset/Source/MeshModelingTools/Public/VoxelMorphologyMeshesTool.h
+++ b/Engine/Plugins/Experimental/MeshModelingToolset/Source/MeshModelingTools/Public/VoxelMorphologyMeshesTool.h
@@ -66,7 +66,7 @@ protected:
 	virtual TUniquePtr<UE::Geometry::FDynamicMeshOperator> MakeNewOperator() override;
 
 	UPROPERTY()
-	UVoxelMorphologyMeshesToolProperties* MorphologyProperties;
+	TObjectPtr<UVoxelMorphologyMeshesToolProperties> MorphologyProperties;
 
 };
 
diff --git a/Engine/Plugins/Experimental/MeshModelingToolset/Source/MeshModelingTools/Public/VoxelSolidifyMeshesTool.h b/Engine/Plugins/Experimental/MeshModelingToolset/Source/MeshModelingTools/Public/VoxelSolidifyMeshesTool.h
--- a/Engine/Plugins/Experimental/MeshModelingToolset/Source/MeshModelingTools/Public/VoxelSolidifyMeshesTool.h
+++ b/Engine/Plugins/Experimental/MeshModelingToolset/Source/MeshModelingTools/Public/VoxelSolidifyMeshesTool.h
@@ -72,7 +72,7 @@ public:
 protected:
 
 	UPROPERTY()
-	UVoxelSolidifyMeshesToolProperties* SolidifyProperties;
+	TObjectPtr<UVoxelSolidifyMeshesToolProperties> SolidifyProperties;
 
 };
 
diff --git a/Engine/Plugins/Experimental/MeshModelingToolset/Source/MeshModelingTools/Public/WeldMeshEdgesTool.h b/Engine/Plugins/Experimental/MeshModelingToolset/Source/MeshModelingTools/Public/WeldMeshEdgesTool.h
--- a/Engine/Plugins/Experimental/MeshModelingToolset/Source/MeshModelingTools/Public/WeldMeshEdgesTool.h
+++ b/Engine/Plugins/Experimental/MeshModelingToolset/Source/MeshModelingTools/Public/WeldMeshEdgesTool.h
@@ -58,7 +58,7 @@ public:
 	virtual TUniquePtr<UE::Geometry::FDynamicMeshOperator> MakeNewOperator() override;
 
 	UPROPERTY()
-	UWeldMeshEdgesTool* WeldMeshEdgesTool;
+	TObjectPtr<UWeldMeshEdgesTool> WeldMeshEdgesTool;
 };
 
 
diff --git a/Engine/Plugins/Experimental/MeshModelingToolset/Source/MeshModelingToolsEditorOnly/Public/AttributeEditorTool.h b/Engine/Plugins/Experimental/MeshModelingToolset/Source/MeshModelingToolsEditorOnly/Public/AttributeEditorTool.h
--- a/Engine/Plugins/Experimental/MeshModelingToolset/Source/MeshModelingToolsEditorOnly/Public/AttributeEditorTool.h
+++ b/Engine/Plugins/Experimental/MeshModelingToolset/Source/MeshModelingToolsEditorOnly/Public/AttributeEditorTool.h
@@ -366,28 +366,28 @@ public:
 protected:
 
 	UPROPERTY()
-	UAttributeEditorOptimizeActions* OptimizeActions;
+	TObjectPtr<UAttributeEditorOptimizeActions> OptimizeActions;
 
 	UPROPERTY()
-	UAttributeEditorNormalsActions* NormalsActions;
+	TObjectPtr<UAttributeEditorNormalsActions> NormalsActions;
 
 	UPROPERTY()
-	UAttributeEditorUVActions* UVActions;
+	TObjectPtr<UAttributeEditorUVActions> UVActions;
 
 	UPROPERTY()
-	UAttributeEditorLightmapUVActions* LightmapUVActions;
+	TObjectPtr<UAttributeEditorLightmapUVActions> LightmapUVActions;
 
 	UPROPERTY()
-	UAttributeEditorAttribProperties* AttributeProps;
+	TObjectPtr<UAttributeEditorAttribProperties> AttributeProps;
 
 	UPROPERTY()
-	UAttributeEditorNewAttributeActions* NewAttributeProps;
+	TObjectPtr<UAttributeEditorNewAttributeActions> NewAttributeProps;
 
 	UPROPERTY()
-	UAttributeEditorModifyAttributeActions* ModifyAttributeProps;
+	TObjectPtr<UAttributeEditorModifyAttributeActions> ModifyAttributeProps;
 
 	UPROPERTY()
-	UAttributeEditorCopyAttributeActions* CopyAttributeProps;
+	TObjectPtr<UAttributeEditorCopyAttributeActions> CopyAttributeProps;
 
 
 protected:
diff --git a/Engine/Plugins/Experimental/MeshModelingToolset/Source/MeshModelingToolsEditorOnly/Public/BspConversionTool.h b/Engine/Plugins/Experimental/MeshModelingToolset/Source/MeshModelingToolsEditorOnly/Public/BspConversionTool.h
--- a/Engine/Plugins/Experimental/MeshModelingToolset/Source/MeshModelingToolsEditorOnly/Public/BspConversionTool.h
+++ b/Engine/Plugins/Experimental/MeshModelingToolset/Source/MeshModelingToolsEditorOnly/Public/BspConversionTool.h
@@ -150,10 +150,10 @@ public:
 	virtual void RequestAction(EBspConversionToolAction ActionType);
 
 	UPROPERTY()
-	UBspConversionToolProperties* Settings = nullptr;
+	TObjectPtr<UBspConversionToolProperties> Settings = nullptr;
 
 	UPROPERTY()
-	UBspConversionToolActionPropertySet* ToolActions = nullptr;
+	TObjectPtr<UBspConversionToolActionPropertySet> ToolActions = nullptr;
 
 protected:
 	bool bCanAccept = false;
@@ -162,7 +162,7 @@ protected:
 	EBspConversionToolAction PendingAction;
 
 	UPROPERTY()
-	UPreviewMesh *PreviewMesh;
+	TObjectPtr<UPreviewMesh> PreviewMesh;
 
 	TArray<ABrush*> BrushesToConvert;
 	TArray<ABrush*> BrushesToDelete;
diff --git a/Engine/Plugins/Experimental/MeshModelingToolset/Source/MeshModelingToolsEditorOnly/Public/EditMeshMaterialsTool.h b/Engine/Plugins/Experimental/MeshModelingToolset/Source/MeshModelingToolsEditorOnly/Public/EditMeshMaterialsTool.h
--- a/Engine/Plugins/Experimental/MeshModelingToolset/Source/MeshModelingToolsEditorOnly/Public/EditMeshMaterialsTool.h
+++ b/Engine/Plugins/Experimental/MeshModelingToolset/Source/MeshModelingToolsEditorOnly/Public/EditMeshMaterialsTool.h
@@ -42,7 +42,7 @@ public:
 	int32 GetSelectedMaterialIndex() const;
 
 	UPROPERTY(EditAnywhere, Category=Materials)
-	TArray<UMaterialInterface*> Materials;
+	TArray<TObjectPtr<UMaterialInterface>> Materials;
 };
 
 
@@ -96,7 +96,7 @@ public:
 
 protected:
 	UPROPERTY()
-	UEditMeshMaterialsToolProperties* MaterialProps;
+	TObjectPtr<UEditMeshMaterialsToolProperties> MaterialProps;
 
 	virtual UMeshSelectionToolActionPropertySet* CreateEditActions() override;
 	virtual void AddSubclassPropertySets() override;
diff --git a/Engine/Plugins/Experimental/MeshModelingToolset/Source/MeshModelingToolsEditorOnly/Public/GenerateLODMeshesTool.h b/Engine/Plugins/Experimental/MeshModelingToolset/Source/MeshModelingToolsEditorOnly/Public/GenerateLODMeshesTool.h
--- a/Engine/Plugins/Experimental/MeshModelingToolset/Source/MeshModelingToolsEditorOnly/Public/GenerateLODMeshesTool.h
+++ b/Engine/Plugins/Experimental/MeshModelingToolset/Source/MeshModelingToolsEditorOnly/Public/GenerateLODMeshesTool.h
@@ -165,10 +165,10 @@ public:
 
 private:
 	UPROPERTY()
-	UGenerateLODMeshesToolProperties* SimplifyProperties;
+	TObjectPtr<UGenerateLODMeshesToolProperties> SimplifyProperties;
 
 	UPROPERTY()
-	TArray<UMeshOpPreviewWithBackgroundCompute*> Previews;
+	TArray<TObjectPtr<UMeshOpPreviewWithBackgroundCompute>> Previews;
 
 	TArray<TUniquePtr<FGenerateLODOperatorFactory>> PreviewFactories;
 	friend class FGenerateLODOperatorFactory;
diff --git a/Engine/Plugins/Experimental/MeshModelingToolset/Source/MeshModelingToolsEditorOnly/Public/Hair/GroomCardsEditorTool.h b/Engine/Plugins/Experimental/MeshModelingToolset/Source/MeshModelingToolsEditorOnly/Public/Hair/GroomCardsEditorTool.h
--- a/Engine/Plugins/Experimental/MeshModelingToolset/Source/MeshModelingToolsEditorOnly/Public/Hair/GroomCardsEditorTool.h
+++ b/Engine/Plugins/Experimental/MeshModelingToolset/Source/MeshModelingToolsEditorOnly/Public/Hair/GroomCardsEditorTool.h
@@ -182,38 +182,38 @@ public:
 
 protected:
 	UPROPERTY()
-	USelectGroomCardsToolActions* SelectActions;
+	TObjectPtr<USelectGroomCardsToolActions> SelectActions;
 
 	UPROPERTY()
-	UEditGroomCardsToolActions* EditActions;
+	TObjectPtr<UEditGroomCardsToolActions> EditActions;
 
 	UPROPERTY()
-	UGroomCardsInfoToolProperties* InfoProperties;
+	TObjectPtr<UGroomCardsInfoToolProperties> InfoProperties;
 
 	UPROPERTY()
-	UPreviewMesh* PreviewMesh;
+	TObjectPtr<UPreviewMesh> PreviewMesh;
 
 	UPROPERTY()
 	TLazyObjectPtr<AGroomActor> TargetGroom;
 
 	UPROPERTY()
-	UPreviewGeometry* PreviewGeom;
+	TObjectPtr<UPreviewGeometry> PreviewGeom;
 
 protected:		// materials
 
 	UPROPERTY()
-	UMaterialInterface* MeshMaterial;
+	TObjectPtr<UMaterialInterface> MeshMaterial;
 
 	UPROPERTY()
-	UMaterialInterface* UVMaterial;
+	TObjectPtr<UMaterialInterface> UVMaterial;
 
 protected:		// mechanics
 
 	UPROPERTY()
-	UPolygonSelectionMechanic* CardMeshSelectionMechanic;
+	TObjectPtr<UPolygonSelectionMechanic> CardMeshSelectionMechanic;
 
 	UPROPERTY()
-	USpaceCurveDeformationMechanic* ControlPointsMechanic;
+	TObjectPtr<USpaceCurveDeformationMechanic> ControlPointsMechanic;
 
 
 	bool bSelectionStateDirty = false;
diff --git a/Engine/Plugins/Experimental/MeshModelingToolset/Source/MeshModelingToolsEditorOnly/Public/Hair/GroomToMeshTool.h b/Engine/Plugins/Experimental/MeshModelingToolset/Source/MeshModelingToolsEditorOnly/Public/Hair/GroomToMeshTool.h
--- a/Engine/Plugins/Experimental/MeshModelingToolset/Source/MeshModelingToolsEditorOnly/Public/Hair/GroomToMeshTool.h
+++ b/Engine/Plugins/Experimental/MeshModelingToolset/Source/MeshModelingToolsEditorOnly/Public/Hair/GroomToMeshTool.h
@@ -137,23 +137,23 @@ public:
 
 protected:
 	UPROPERTY()
-	UGroomToMeshToolProperties* Settings;
+	TObjectPtr<UGroomToMeshToolProperties> Settings;
 
 	UPROPERTY()
-	UPreviewMesh* PreviewMesh;
+	TObjectPtr<UPreviewMesh> PreviewMesh;
 
 	UPROPERTY()
 	TLazyObjectPtr<AGroomActor> TargetGroom;
 
 	UPROPERTY()
-	UPreviewGeometry* PreviewGeom;
+	TObjectPtr<UPreviewGeometry> PreviewGeom;
 
 protected:
 	UPROPERTY()
-	UMaterialInterface* MeshMaterial;
+	TObjectPtr<UMaterialInterface> MeshMaterial;
 
 	UPROPERTY()
-	UMaterialInterface* UVMaterial;
+	TObjectPtr<UMaterialInterface> UVMaterial;
 
 protected:
 	UWorld* TargetWorld = nullptr;
diff --git a/Engine/Plugins/Experimental/MeshModelingToolset/Source/MeshModelingToolsEditorOnly/Public/MergeMeshesTool.h b/Engine/Plugins/Experimental/MeshModelingToolset/Source/MeshModelingToolsEditorOnly/Public/MergeMeshesTool.h
--- a/Engine/Plugins/Experimental/MeshModelingToolset/Source/MeshModelingToolsEditorOnly/Public/MergeMeshesTool.h
+++ b/Engine/Plugins/Experimental/MeshModelingToolset/Source/MeshModelingToolsEditorOnly/Public/MergeMeshesTool.h
@@ -84,17 +84,17 @@ public:
 
 protected:
 	UPROPERTY()
-	UMergeMeshesToolProperties* MergeProps;
+	TObjectPtr<UMergeMeshesToolProperties> MergeProps;
 
 	UPROPERTY()
-	UMeshStatisticsProperties* MeshStatisticsProperties;
+	TObjectPtr<UMeshStatisticsProperties> MeshStatisticsProperties;
 
 	UPROPERTY()
-	UOnAcceptHandleSourcesProperties* HandleSourcesProperties;
+	TObjectPtr<UOnAcceptHandleSourcesProperties> HandleSourcesProperties;
 
 
 	UPROPERTY()
-	UMeshOpPreviewWithBackgroundCompute* Preview;
+	TObjectPtr<UMeshOpPreviewWithBackgroundCompute> Preview;
 
 protected:
 	UWorld* TargetWorld;
diff --git a/Engine/Plugins/Experimental/MeshModelingToolset/Source/MeshModelingToolsEditorOnly/Public/MeshTangentsTool.h b/Engine/Plugins/Experimental/MeshModelingToolset/Source/MeshModelingToolsEditorOnly/Public/MeshTangentsTool.h
--- a/Engine/Plugins/Experimental/MeshModelingToolset/Source/MeshModelingToolsEditorOnly/Public/MeshTangentsTool.h
+++ b/Engine/Plugins/Experimental/MeshModelingToolset/Source/MeshModelingToolsEditorOnly/Public/MeshTangentsTool.h
@@ -102,17 +102,17 @@ public:
 
 protected:
 	UPROPERTY()
-	UMeshTangentsToolProperties* Settings = nullptr;
+	TObjectPtr<UMeshTangentsToolProperties> Settings = nullptr;
 
 protected:
 	UPROPERTY()
-	UMaterialInterface* DefaultMaterial = nullptr;
+	TObjectPtr<UMaterialInterface> DefaultMaterial = nullptr;
 
 	UPROPERTY()
-	UPreviewMesh* PreviewMesh = nullptr;
+	TObjectPtr<UPreviewMesh> PreviewMesh = nullptr;
 
 	UPROPERTY()
-	UPreviewGeometry* PreviewGeometry = nullptr;
+	TObjectPtr<UPreviewGeometry> PreviewGeometry = nullptr;
 
 	TUniquePtr<TGenericDataBackgroundCompute<UE::Geometry::FMeshTangentsd>> Compute = nullptr;
 
diff --git a/Engine/Plugins/Experimental/MeshModelingToolset/Source/MeshModelingToolsEditorOnly/Public/MeshToVolumeTool.h b/Engine/Plugins/Experimental/MeshModelingToolset/Source/MeshModelingToolsEditorOnly/Public/MeshToVolumeTool.h
--- a/Engine/Plugins/Experimental/MeshModelingToolset/Source/MeshModelingToolsEditorOnly/Public/MeshToVolumeTool.h
+++ b/Engine/Plugins/Experimental/MeshModelingToolset/Source/MeshModelingToolsEditorOnly/Public/MeshToVolumeTool.h
@@ -87,16 +87,16 @@ public:
 
 protected:
 	UPROPERTY()
-	UMeshToVolumeToolProperties* Settings;
+	TObjectPtr<UMeshToVolumeToolProperties> Settings;
 
 	UPROPERTY()
-	UOnAcceptHandleSourcesProperties* HandleSourcesProperties;
+	TObjectPtr<UOnAcceptHandleSourcesProperties> HandleSourcesProperties;
 
 	UPROPERTY()
-	UPreviewMesh* PreviewMesh;
+	TObjectPtr<UPreviewMesh> PreviewMesh;
 
 	UPROPERTY()
-	ULineSetComponent* VolumeEdgesSet;
+	TObjectPtr<ULineSetComponent> VolumeEdgesSet;
 
 protected:
 	UE::Geometry::FDynamicMesh3 InputMesh;
diff --git a/Engine/Plugins/Experimental/MeshModelingToolset/Source/MeshModelingToolsEditorOnly/Public/ParameterizeMeshTool.h b/Engine/Plugins/Experimental/MeshModelingToolset/Source/MeshModelingToolsEditorOnly/Public/ParameterizeMeshTool.h
--- a/Engine/Plugins/Experimental/MeshModelingToolset/Source/MeshModelingToolsEditorOnly/Public/ParameterizeMeshTool.h
+++ b/Engine/Plugins/Experimental/MeshModelingToolset/Source/MeshModelingToolsEditorOnly/Public/ParameterizeMeshTool.h
@@ -106,24 +106,24 @@ public:
 
 protected:
 	UPROPERTY()
-	UMeshUVChannelProperties* UVChannelProperties = nullptr;
+	TObjectPtr<UMeshUVChannelProperties> UVChannelProperties = nullptr;
 
 	UPROPERTY()
-	UParameterizeMeshToolProperties* Settings = nullptr;
+	TObjectPtr<UParameterizeMeshToolProperties> Settings = nullptr;
 
 	UPROPERTY()
-	UParameterizeMeshToolUVAtlasProperties* UVAtlasProperties = nullptr;
+	TObjectPtr<UParameterizeMeshToolUVAtlasProperties> UVAtlasProperties = nullptr;
 
 	UPROPERTY()
-	UParameterizeMeshToolXAtlasProperties* XAtlasProperties = nullptr;
+	TObjectPtr<UParameterizeMeshToolXAtlasProperties> XAtlasProperties = nullptr;
 
 
 	UPROPERTY()
-	UExistingMeshMaterialProperties* MaterialSettings = nullptr;
+	TObjectPtr<UExistingMeshMaterialProperties> MaterialSettings = nullptr;
 
 protected:
 	UPROPERTY()
-	UMeshOpPreviewWithBackgroundCompute* Preview = nullptr;
+	TObjectPtr<UMeshOpPreviewWithBackgroundCompute> Preview = nullptr;
 
 protected:
 	TSharedPtr<UE::Geometry::FDynamicMesh3, ESPMode::ThreadSafe> InputMesh;
diff --git a/Engine/Plugins/Experimental/MeshModelingToolset/Source/MeshModelingToolsEditorOnly/Public/PolygonOnMeshTool.h b/Engine/Plugins/Experimental/MeshModelingToolset/Source/MeshModelingToolsEditorOnly/Public/PolygonOnMeshTool.h
--- a/Engine/Plugins/Experimental/MeshModelingToolset/Source/MeshModelingToolsEditorOnly/Public/PolygonOnMeshTool.h
+++ b/Engine/Plugins/Experimental/MeshModelingToolset/Source/MeshModelingToolsEditorOnly/Public/PolygonOnMeshTool.h
@@ -194,16 +194,16 @@ public:
 protected:
 
 	UPROPERTY()
-	UPolygonOnMeshToolProperties* BasicProperties;
+	TObjectPtr<UPolygonOnMeshToolProperties> BasicProperties;
 
 	UPROPERTY()
-	UPolygonOnMeshToolActionPropertySet* ActionProperties;
+	TObjectPtr<UPolygonOnMeshToolActionPropertySet> ActionProperties;
 
 	UPROPERTY()
-	UMeshOpPreviewWithBackgroundCompute* Preview;
+	TObjectPtr<UMeshOpPreviewWithBackgroundCompute> Preview;
 
 	UPROPERTY()
-	ULineSetComponent* DrawnLineSet;
+	TObjectPtr<ULineSetComponent> DrawnLineSet;
 
 	TArray<int> EdgesOnFailure;
 	TArray<int> EmbeddedEdges;
@@ -218,10 +218,10 @@ protected:
 	TSharedPtr<UE::Geometry::FDynamicMesh3, ESPMode::ThreadSafe> OriginalDynamicMesh;
 
 	UPROPERTY()
-	UConstructionPlaneMechanic* PlaneMechanic = nullptr;
+	TObjectPtr<UConstructionPlaneMechanic> PlaneMechanic = nullptr;
 
 	UPROPERTY()
-	UCollectSurfacePathMechanic* DrawPolygonMechanic = nullptr;
+	TObjectPtr<UCollectSurfacePathMechanic> DrawPolygonMechanic = nullptr;
 
 	EPolygonOnMeshToolActions PendingAction = EPolygonOnMeshToolActions::NoAction;
 
diff --git a/Engine/Plugins/Experimental/MeshModelingToolset/Source/MeshModelingToolsEditorOnly/Public/ShapeSprayTool.h b/Engine/Plugins/Experimental/MeshModelingToolset/Source/MeshModelingToolsEditorOnly/Public/ShapeSprayTool.h
--- a/Engine/Plugins/Experimental/MeshModelingToolset/Source/MeshModelingToolsEditorOnly/Public/ShapeSprayTool.h
+++ b/Engine/Plugins/Experimental/MeshModelingToolset/Source/MeshModelingToolsEditorOnly/Public/ShapeSprayTool.h
@@ -46,7 +46,7 @@ public:
 	int NumSplats;
 
 	UPROPERTY(EditAnywhere, Category = Options)
-	UMaterialInterface* Material = nullptr;
+	TObjectPtr<UMaterialInterface> Material = nullptr;
 };
 
 
@@ -85,11 +85,11 @@ public:
 
 protected:
 	UPROPERTY()
-	UShapeSprayToolProperties* Settings;
+	TObjectPtr<UShapeSprayToolProperties> Settings;
 
 	// small meshes are accumulated here
 	UPROPERTY()
-	UDynamicMeshComponent* AccumMeshComponent;
+	TObjectPtr<UDynamicMeshComponent> AccumMeshComponent;
 
 
 protected:
diff --git a/Engine/Plugins/Experimental/MeshModelingToolset/Source/MeshModelingToolsEditorOnly/Public/SimplifyMeshTool.h b/Engine/Plugins/Experimental/MeshModelingToolset/Source/MeshModelingToolsEditorOnly/Public/SimplifyMeshTool.h
--- a/Engine/Plugins/Experimental/MeshModelingToolset/Source/MeshModelingToolsEditorOnly/Public/SimplifyMeshTool.h
+++ b/Engine/Plugins/Experimental/MeshModelingToolset/Source/MeshModelingToolsEditorOnly/Public/SimplifyMeshTool.h
@@ -108,16 +108,16 @@ public:
 
 private:
 	UPROPERTY()
-	USimplifyMeshToolProperties* SimplifyProperties;
+	TObjectPtr<USimplifyMeshToolProperties> SimplifyProperties;
 
 	UPROPERTY()
-	UMeshStatisticsProperties* MeshStatisticsProperties;
+	TObjectPtr<UMeshStatisticsProperties> MeshStatisticsProperties;
 
 	UPROPERTY()
-	UMeshOpPreviewWithBackgroundCompute* Preview;
+	TObjectPtr<UMeshOpPreviewWithBackgroundCompute> Preview;
 
 	UPROPERTY()
-	UMeshElementsVisualizer* MeshElementsDisplay;
+	TObjectPtr<UMeshElementsVisualizer> MeshElementsDisplay;
 
 	TSharedPtr<FMeshDescription, ESPMode::ThreadSafe> OriginalMeshDescription;
 	// Dynamic Mesh versions precomputed in Setup (rather than recomputed for every simplify op)
diff --git a/Engine/Plugins/Experimental/MeshModelingToolset/Source/MeshModelingToolsEditorOnly/Public/SubdividePolyTool.h b/Engine/Plugins/Experimental/MeshModelingToolset/Source/MeshModelingToolsEditorOnly/Public/SubdividePolyTool.h
--- a/Engine/Plugins/Experimental/MeshModelingToolset/Source/MeshModelingToolsEditorOnly/Public/SubdividePolyTool.h
+++ b/Engine/Plugins/Experimental/MeshModelingToolset/Source/MeshModelingToolsEditorOnly/Public/SubdividePolyTool.h
@@ -87,16 +87,16 @@ protected:
 	friend class USubdividePolyOperatorFactory;
 
 	UPROPERTY()
-	UPreviewMesh* PreviewMesh = nullptr;
+	TObjectPtr<UPreviewMesh> PreviewMesh = nullptr;
 
 	UPROPERTY()
-	USubdividePolyToolProperties* Properties = nullptr;
+	TObjectPtr<USubdividePolyToolProperties> Properties = nullptr;
 
 	// Input mesh
 	TSharedPtr<UE::Geometry::FDynamicMesh3, ESPMode::ThreadSafe> OriginalMesh;
 
 	UPROPERTY()
-	UPreviewGeometry* PreviewGeometry = nullptr;
+	TObjectPtr<UPreviewGeometry> PreviewGeometry = nullptr;
 
 	bool bPreviewGeometryNeedsUpdate;
 	void CreateOrUpdatePreviewGeometry();
diff --git a/Engine/Plugins/Experimental/MeshModelingToolset/Source/MeshModelingToolsEditorOnly/Public/UVLayoutTool.h b/Engine/Plugins/Experimental/MeshModelingToolset/Source/MeshModelingToolsEditorOnly/Public/UVLayoutTool.h
--- a/Engine/Plugins/Experimental/MeshModelingToolset/Source/MeshModelingToolsEditorOnly/Public/UVLayoutTool.h
+++ b/Engine/Plugins/Experimental/MeshModelingToolset/Source/MeshModelingToolsEditorOnly/Public/UVLayoutTool.h
@@ -101,7 +101,7 @@ public:
 	virtual TUniquePtr<UE::Geometry::FDynamicMeshOperator> MakeNewOperator() override;
 
 	UPROPERTY()
-	UUVLayoutTool *Tool;
+	TObjectPtr<UUVLayoutTool> Tool;
 
 	int ComponentIndex;
 
@@ -143,16 +143,16 @@ public:
 protected:
 
 	UPROPERTY()
-	UMeshUVChannelProperties* UVChannelProperties = nullptr;
+	TObjectPtr<UMeshUVChannelProperties> UVChannelProperties = nullptr;
 
 	UPROPERTY()
-	UUVLayoutToolProperties* BasicProperties = nullptr;
+	TObjectPtr<UUVLayoutToolProperties> BasicProperties = nullptr;
 
 	UPROPERTY()
-	UExistingMeshMaterialProperties* MaterialSettings = nullptr;
+	TObjectPtr<UExistingMeshMaterialProperties> MaterialSettings = nullptr;
 
 	UPROPERTY()
-	TArray<UMeshOpPreviewWithBackgroundCompute*> Previews;
+	TArray<TObjectPtr<UMeshOpPreviewWithBackgroundCompute>> Previews;
 
 protected:
 	TArray<TSharedPtr<UE::Geometry::FDynamicMesh3, ESPMode::ThreadSafe>> OriginalDynamicMeshes;
@@ -172,5 +172,5 @@ protected:
 
 protected:
 	UPROPERTY()
-	UUVLayoutPreview* UVLayoutView = nullptr;
+	TObjectPtr<UUVLayoutPreview> UVLayoutView = nullptr;
 };
diff --git a/Engine/Plugins/Experimental/MeshModelingToolset/Source/MeshModelingToolsEditorOnly/Public/VoxelCSGMeshesTool.h b/Engine/Plugins/Experimental/MeshModelingToolset/Source/MeshModelingToolsEditorOnly/Public/VoxelCSGMeshesTool.h
--- a/Engine/Plugins/Experimental/MeshModelingToolset/Source/MeshModelingToolsEditorOnly/Public/VoxelCSGMeshesTool.h
+++ b/Engine/Plugins/Experimental/MeshModelingToolset/Source/MeshModelingToolsEditorOnly/Public/VoxelCSGMeshesTool.h
@@ -111,16 +111,16 @@ public:
 
 protected:
 	UPROPERTY()
-	UVoxelCSGMeshesToolProperties* CSGProps;
+	TObjectPtr<UVoxelCSGMeshesToolProperties> CSGProps;
 
 	UPROPERTY()
-	UMeshStatisticsProperties* MeshStatisticsProperties;
+	TObjectPtr<UMeshStatisticsProperties> MeshStatisticsProperties;
 
 	UPROPERTY()
-	UOnAcceptHandleSourcesProperties* HandleSourcesProperties;
+	TObjectPtr<UOnAcceptHandleSourcesProperties> HandleSourcesProperties;
 
 	UPROPERTY()
-	UMeshOpPreviewWithBackgroundCompute* Preview;
+	TObjectPtr<UMeshOpPreviewWithBackgroundCompute> Preview;
 
 protected:
 	UWorld* TargetWorld;
diff --git a/Engine/Plugins/Experimental/MeshModelingToolset/Source/ModelingComponents/Public/BaseTools/BaseCreateFromSelectedTool.h b/Engine/Plugins/Experimental/MeshModelingToolset/Source/ModelingComponents/Public/BaseTools/BaseCreateFromSelectedTool.h
--- a/Engine/Plugins/Experimental/MeshModelingToolset/Source/ModelingComponents/Public/BaseTools/BaseCreateFromSelectedTool.h
+++ b/Engine/Plugins/Experimental/MeshModelingToolset/Source/ModelingComponents/Public/BaseTools/BaseCreateFromSelectedTool.h
@@ -190,22 +190,22 @@ protected:
 protected:
 
 	UPROPERTY()
-	UTransformInputsToolProperties* TransformProperties;
+	TObjectPtr<UTransformInputsToolProperties> TransformProperties;
 
 	UPROPERTY()
-	UCreateMeshObjectTypeProperties* OutputTypeProperties;
+	TObjectPtr<UCreateMeshObjectTypeProperties> OutputTypeProperties;
 
 	UPROPERTY()
-	UBaseCreateFromSelectedHandleSourceProperties* HandleSourcesProperties;
+	TObjectPtr<UBaseCreateFromSelectedHandleSourceProperties> HandleSourcesProperties;
 
 	UPROPERTY()
-	UMeshOpPreviewWithBackgroundCompute* Preview;
+	TObjectPtr<UMeshOpPreviewWithBackgroundCompute> Preview;
 
 	UPROPERTY()
-	TArray<UTransformProxy*> TransformProxies;
+	TArray<TObjectPtr<UTransformProxy>> TransformProxies;
 
 	UPROPERTY()
-	TArray<UTransformGizmo*> TransformGizmos;
+	TArray<TObjectPtr<UTransformGizmo>> TransformGizmos;
 
 	UWorld* TargetWorld = nullptr;
 };
diff --git a/Engine/Plugins/Experimental/MeshModelingToolset/Source/ModelingComponents/Public/BaseTools/BaseMeshProcessingTool.h b/Engine/Plugins/Experimental/MeshModelingToolset/Source/ModelingComponents/Public/BaseTools/BaseMeshProcessingTool.h
--- a/Engine/Plugins/Experimental/MeshModelingToolset/Source/ModelingComponents/Public/BaseTools/BaseMeshProcessingTool.h
+++ b/Engine/Plugins/Experimental/MeshModelingToolset/Source/ModelingComponents/Public/BaseTools/BaseMeshProcessingTool.h
@@ -197,7 +197,7 @@ protected:
 
 	// Preview object holds temporary Actor with preview mesh component
 	UPROPERTY()
-	UMeshOpPreviewWithBackgroundCompute* Preview = nullptr;
+	TObjectPtr<UMeshOpPreviewWithBackgroundCompute> Preview = nullptr;
 
 	UPreviewMesh* GetUPreviewMesh() const { return Preview->PreviewMesh; }
 	const FTransform& GetPreviewTransform() const { return OverrideTransform; }
diff --git a/Engine/Plugins/Experimental/MeshModelingToolset/Source/ModelingComponents/Public/BaseTools/BaseVoxelTool.h b/Engine/Plugins/Experimental/MeshModelingToolset/Source/ModelingComponents/Public/BaseTools/BaseVoxelTool.h
--- a/Engine/Plugins/Experimental/MeshModelingToolset/Source/ModelingComponents/Public/BaseTools/BaseVoxelTool.h
+++ b/Engine/Plugins/Experimental/MeshModelingToolset/Source/ModelingComponents/Public/BaseTools/BaseVoxelTool.h
@@ -35,7 +35,7 @@ protected:
 
 
 	UPROPERTY()
-	UVoxelProperties* VoxProperties;
+	TObjectPtr<UVoxelProperties> VoxProperties;
 
 	TArray<TSharedPtr<FDynamicMesh3, ESPMode::ThreadSafe>> OriginalDynamicMeshes;
 };
diff --git a/Engine/Plugins/Experimental/MeshModelingToolset/Source/ModelingComponents/Public/BaseTools/SingleSelectionMeshEditingTool.h b/Engine/Plugins/Experimental/MeshModelingToolset/Source/ModelingComponents/Public/BaseTools/SingleSelectionMeshEditingTool.h
--- a/Engine/Plugins/Experimental/MeshModelingToolset/Source/ModelingComponents/Public/BaseTools/SingleSelectionMeshEditingTool.h
+++ b/Engine/Plugins/Experimental/MeshModelingToolset/Source/ModelingComponents/Public/BaseTools/SingleSelectionMeshEditingTool.h
@@ -87,6 +87,6 @@ protected:
 	 * (optional) Stored Selection provided on Tool Input. This should never be modified after the Tool's Setup() is called
 	 */
 	UPROPERTY()
-	const UInteractiveToolStorableSelection* InputSelection = nullptr;
+	TObjectPtr<const UInteractiveToolStorableSelection> InputSelection = nullptr;
 
 };
diff --git a/Engine/Plugins/Experimental/MeshModelingToolset/Source/ModelingComponents/Public/Components/OctreeDynamicMeshComponent.h b/Engine/Plugins/Experimental/MeshModelingToolset/Source/ModelingComponents/Public/Components/OctreeDynamicMeshComponent.h
--- a/Engine/Plugins/Experimental/MeshModelingToolset/Source/ModelingComponents/Public/Components/OctreeDynamicMeshComponent.h
+++ b/Engine/Plugins/Experimental/MeshModelingToolset/Source/ModelingComponents/Public/Components/OctreeDynamicMeshComponent.h
@@ -128,7 +128,7 @@ private:
 	//~ Begin USceneComponent Interface.
 
 	UPROPERTY()
-	UDynamicMesh* MeshObject;
+	TObjectPtr<UDynamicMesh> MeshObject;
 
 	FDelegateHandle MeshObjectChangedHandle;
 	void OnMeshObjectChanged(UDynamicMesh* ChangedMeshObject, FDynamicMeshChangeInfo ChangeInfo);
diff --git a/Engine/Plugins/Experimental/MeshModelingToolset/Source/ModelingComponents/Public/Drawing/LineSetComponent.h b/Engine/Plugins/Experimental/MeshModelingToolset/Source/ModelingComponents/Public/Drawing/LineSetComponent.h
--- a/Engine/Plugins/Experimental/MeshModelingToolset/Source/ModelingComponents/Public/Drawing/LineSetComponent.h
+++ b/Engine/Plugins/Experimental/MeshModelingToolset/Source/ModelingComponents/Public/Drawing/LineSetComponent.h
@@ -127,7 +127,7 @@ private:
 	//~ End USceneComponent Interface.
 
 	UPROPERTY()
-	const UMaterialInterface* LineMaterial;
+	TObjectPtr<const UMaterialInterface> LineMaterial;
 
 	UPROPERTY()
 	mutable FBoxSphereBounds Bounds;
diff --git a/Engine/Plugins/Experimental/MeshModelingToolset/Source/ModelingComponents/Public/Drawing/MeshElementsVisualizer.h b/Engine/Plugins/Experimental/MeshModelingToolset/Source/ModelingComponents/Public/Drawing/MeshElementsVisualizer.h
--- a/Engine/Plugins/Experimental/MeshModelingToolset/Source/ModelingComponents/Public/Drawing/MeshElementsVisualizer.h
+++ b/Engine/Plugins/Experimental/MeshModelingToolset/Source/ModelingComponents/Public/Drawing/MeshElementsVisualizer.h
@@ -136,11 +136,11 @@ public:
 public:
 	/** Visualization settings */
 	UPROPERTY()
-	UMeshElementsVisualizerProperties* Settings;
+	TObjectPtr<UMeshElementsVisualizerProperties> Settings;
 
 	/** Mesh Wireframe component, draws wireframe, boundaries, UV seams, normal seams, color seams */
 	UPROPERTY()
-	UMeshWireframeComponent* WireframeComponent;
+	TObjectPtr<UMeshWireframeComponent> WireframeComponent;
 
 
 protected:
diff --git a/Engine/Plugins/Experimental/MeshModelingToolset/Source/ModelingComponents/Public/Drawing/MeshWireframeComponent.h b/Engine/Plugins/Experimental/MeshModelingToolset/Source/ModelingComponents/Public/Drawing/MeshWireframeComponent.h
--- a/Engine/Plugins/Experimental/MeshModelingToolset/Source/ModelingComponents/Public/Drawing/MeshWireframeComponent.h
+++ b/Engine/Plugins/Experimental/MeshModelingToolset/Source/ModelingComponents/Public/Drawing/MeshWireframeComponent.h
@@ -199,7 +199,7 @@ private:
 	//~ End USceneComponent Interface.
 
 	UPROPERTY()
-	const UMaterialInterface* LineMaterial;
+	TObjectPtr<const UMaterialInterface> LineMaterial;
 
 	UPROPERTY()
 	mutable FBoxSphereBounds LocalBounds;
diff --git a/Engine/Plugins/Experimental/MeshModelingToolset/Source/ModelingComponents/Public/Drawing/PointSetComponent.h b/Engine/Plugins/Experimental/MeshModelingToolset/Source/ModelingComponents/Public/Drawing/PointSetComponent.h
--- a/Engine/Plugins/Experimental/MeshModelingToolset/Source/ModelingComponents/Public/Drawing/PointSetComponent.h
+++ b/Engine/Plugins/Experimental/MeshModelingToolset/Source/ModelingComponents/Public/Drawing/PointSetComponent.h
@@ -103,7 +103,7 @@ private:
 	//~ Begin USceneComponent Interface.
 
 	UPROPERTY()
-	const UMaterialInterface* PointMaterial;
+	TObjectPtr<const UMaterialInterface> PointMaterial;
 
 	UPROPERTY()
 	mutable FBoxSphereBounds Bounds;
diff --git a/Engine/Plugins/Experimental/MeshModelingToolset/Source/ModelingComponents/Public/Drawing/PreviewGeometryActor.h b/Engine/Plugins/Experimental/MeshModelingToolset/Source/ModelingComponents/Public/Drawing/PreviewGeometryActor.h
--- a/Engine/Plugins/Experimental/MeshModelingToolset/Source/ModelingComponents/Public/Drawing/PreviewGeometryActor.h
+++ b/Engine/Plugins/Experimental/MeshModelingToolset/Source/ModelingComponents/Public/Drawing/PreviewGeometryActor.h
@@ -170,13 +170,13 @@ public:
 	 * Actor created and managed by the UPreviewGeometry
 	 */
 	UPROPERTY()
-	APreviewGeometryActor* ParentActor = nullptr;
+	TObjectPtr<APreviewGeometryActor> ParentActor = nullptr;
 
 	/**
 	 * LineSetComponents created and owned by the UPreviewGeometry, and added as child components of the ParentActor
 	 */
 	UPROPERTY()
-	TMap<FString, ULineSetComponent*> LineSets;
+	TMap<FString, TObjectPtr<ULineSetComponent>> LineSets;
 
 
 
diff --git a/Engine/Plugins/Experimental/MeshModelingToolset/Source/ModelingComponents/Public/Drawing/TriangleSetComponent.h b/Engine/Plugins/Experimental/MeshModelingToolset/Source/ModelingComponents/Public/Drawing/TriangleSetComponent.h
--- a/Engine/Plugins/Experimental/MeshModelingToolset/Source/ModelingComponents/Public/Drawing/TriangleSetComponent.h
+++ b/Engine/Plugins/Experimental/MeshModelingToolset/Source/ModelingComponents/Public/Drawing/TriangleSetComponent.h
@@ -66,7 +66,7 @@ struct FRenderableTriangle
 	}
 
 	UPROPERTY()
-	UMaterialInterface* Material;
+	TObjectPtr<UMaterialInterface> Material;
 
 	UPROPERTY()
 	FRenderableTriangleVertex Vertex0;
diff --git a/Engine/Plugins/Experimental/MeshModelingToolset/Source/ModelingComponents/Public/Drawing/UVLayoutPreview.h b/Engine/Plugins/Experimental/MeshModelingToolset/Source/ModelingComponents/Public/Drawing/UVLayoutPreview.h
--- a/Engine/Plugins/Experimental/MeshModelingToolset/Source/ModelingComponents/Public/Drawing/UVLayoutPreview.h
+++ b/Engine/Plugins/Experimental/MeshModelingToolset/Source/ModelingComponents/Public/Drawing/UVLayoutPreview.h
@@ -123,15 +123,15 @@ public:
 public:
 	/** Visualization settings */
 	UPROPERTY()
-	UUVLayoutPreviewProperties* Settings;
+	TObjectPtr<UUVLayoutPreviewProperties> Settings;
 
 	/** PreviewMesh is initialized with a copy of an input mesh with UVs mapped to position, ie such that (X,Y,Z) = (U,V,0) */
 	UPROPERTY()
-	UPreviewMesh* PreviewMesh;
+	TObjectPtr<UPreviewMesh> PreviewMesh;
 
 	/** Set of additional triangles to draw, eg for backing rectangle, etc */
 	UPROPERTY()
-	UTriangleSetComponent* TriangleComponent;
+	TObjectPtr<UTriangleSetComponent> TriangleComponent;
 
 	/** Configure whether the backing rectangle should be shown */
 	UPROPERTY()
@@ -139,7 +139,7 @@ public:
 
 	/** Configure the backing rectangle material */
 	UPROPERTY()
-	UMaterialInterface* BackingRectangleMaterial = nullptr;
+	TObjectPtr<UMaterialInterface> BackingRectangleMaterial = nullptr;
 
 
 protected:
diff --git a/Engine/Plugins/Experimental/MeshModelingToolset/Source/ModelingComponents/Public/Mechanics/ConstructionPlaneMechanic.h b/Engine/Plugins/Experimental/MeshModelingToolset/Source/ModelingComponents/Public/Mechanics/ConstructionPlaneMechanic.h
--- a/Engine/Plugins/Experimental/MeshModelingToolset/Source/ModelingComponents/Public/Mechanics/ConstructionPlaneMechanic.h
+++ b/Engine/Plugins/Experimental/MeshModelingToolset/Source/ModelingComponents/Public/Mechanics/ConstructionPlaneMechanic.h
@@ -64,10 +64,10 @@ public:
 
 public:
 	UPROPERTY()
-	UTransformGizmo* PlaneTransformGizmo;
+	TObjectPtr<UTransformGizmo> PlaneTransformGizmo;
 	
 	UPROPERTY()
-	UTransformProxy* PlaneTransformProxy;
+	TObjectPtr<UTransformProxy> PlaneTransformProxy;
 
 	/** 
 	 * This is the behavior target used for the Ctrl+click behavior that sets the plane
@@ -81,7 +81,7 @@ protected:
 	bool bEnableSnapToWorldRotGrid = false;
 
 	UPROPERTY()
-	USingleClickInputBehavior* ClickToSetPlaneBehavior;
+	TObjectPtr<USingleClickInputBehavior> ClickToSetPlaneBehavior;
 
 	void TransformChanged(UTransformProxy* Proxy, FTransform Transform);
 };
\ No newline at end of file
diff --git a/Engine/Plugins/Experimental/MeshModelingToolset/Source/ModelingComponents/Public/Mechanics/CurveControlPointsMechanic.h b/Engine/Plugins/Experimental/MeshModelingToolset/Source/ModelingComponents/Public/Mechanics/CurveControlPointsMechanic.h
--- a/Engine/Plugins/Experimental/MeshModelingToolset/Source/ModelingComponents/Public/Mechanics/CurveControlPointsMechanic.h
+++ b/Engine/Plugins/Experimental/MeshModelingToolset/Source/ModelingComponents/Public/Mechanics/CurveControlPointsMechanic.h
@@ -195,9 +195,9 @@ public:
 
 	// Behaviors used for moving points around and hovering them
 	UPROPERTY()
-	USingleClickInputBehavior* ClickBehavior = nullptr;
+	TObjectPtr<USingleClickInputBehavior> ClickBehavior = nullptr;
 	UPROPERTY()
-	UMouseHoverBehavior* HoverBehavior = nullptr;
+	TObjectPtr<UMouseHoverBehavior> HoverBehavior = nullptr;
 
 	// This delegate is called every time the control point sequence is altered.
 	DECLARE_MULTICAST_DELEGATE(OnPointsChangedEvent);
@@ -338,18 +338,18 @@ protected:
 
 	/** Used for displaying points/segments */
 	UPROPERTY()
-	APreviewGeometryActor* PreviewGeometryActor;
+	TObjectPtr<APreviewGeometryActor> PreviewGeometryActor;
 	UPROPERTY()
-	UPointSetComponent* DrawnControlPoints;
+	TObjectPtr<UPointSetComponent> DrawnControlPoints;
 	UPROPERTY()
-	ULineSetComponent* DrawnControlSegments;
+	TObjectPtr<ULineSetComponent> DrawnControlSegments;
 
 	// These get drawn separately because the other components have to be 1:1 with the control
 	// points structure, which would make it complicated to keep track of special id's.
 	UPROPERTY()
-	UPointSetComponent* PreviewPoint;
+	TObjectPtr<UPointSetComponent> PreviewPoint;
 	UPROPERTY()
-	ULineSetComponent* PreviewSegment;
+	TObjectPtr<ULineSetComponent> PreviewSegment;
 
 	// Variables for drawing
 	FColor InitializationCurveColor;
@@ -377,9 +377,9 @@ protected:
 	// Support for gizmo. Since the points aren't individual components, we don't actually use UTransformProxy
 	// for the transform forwarding- we just use it for the callbacks.
 	UPROPERTY()
-	UTransformProxy* PointTransformProxy;
+	TObjectPtr<UTransformProxy> PointTransformProxy;
 	UPROPERTY()
-	UTransformGizmo* PointTransformGizmo;
+	TObjectPtr<UTransformGizmo> PointTransformGizmo;
 
 	// Used to make it easy to tell whether the gizmo was moved by the user or by undo/redo or
 	// some other change that we shoulnd't respond to. Basing our movement undo/redo on the
diff --git a/Engine/Plugins/Experimental/MeshModelingToolset/Source/ModelingComponents/Public/Mechanics/LatticeControlPointsMechanic.h b/Engine/Plugins/Experimental/MeshModelingToolset/Source/ModelingComponents/Public/Mechanics/LatticeControlPointsMechanic.h
--- a/Engine/Plugins/Experimental/MeshModelingToolset/Source/ModelingComponents/Public/Mechanics/LatticeControlPointsMechanic.h
+++ b/Engine/Plugins/Experimental/MeshModelingToolset/Source/ModelingComponents/Public/Mechanics/LatticeControlPointsMechanic.h
@@ -119,11 +119,11 @@ protected:
 
 	/** Used for displaying points/segments */
 	UPROPERTY()
-	APreviewGeometryActor* PreviewGeometryActor;
+	TObjectPtr<APreviewGeometryActor> PreviewGeometryActor;
 	UPROPERTY()
-	UPointSetComponent* DrawnControlPoints;
+	TObjectPtr<UPointSetComponent> DrawnControlPoints;
 	UPROPERTY()
-	ULineSetComponent* DrawnLatticeEdges;
+	TObjectPtr<ULineSetComponent> DrawnLatticeEdges;
 
 	// Variables for drawing
 	FColor NormalSegmentColor;
@@ -149,9 +149,9 @@ protected:
 	// Support for gizmo. Since the points aren't individual components, we don't actually use UTransformProxy
 	// for the transform forwarding- we just use it for the callbacks.
 	UPROPERTY()
-	UTransformProxy* PointTransformProxy;
+	TObjectPtr<UTransformProxy> PointTransformProxy;
 	UPROPERTY()
-	UTransformGizmo* PointTransformGizmo;
+	TObjectPtr<UTransformGizmo> PointTransformGizmo;
 
 	// Used to make it easy to tell whether the gizmo was moved by the user or by undo/redo or
 	// some other change that we shouldn't respond to. Basing our movement undo/redo on the
diff --git a/Engine/Plugins/Experimental/MeshModelingToolset/Source/ModelingComponents/Public/Mechanics/PolyLassoMarqueeMechanic.h b/Engine/Plugins/Experimental/MeshModelingToolset/Source/ModelingComponents/Public/Mechanics/PolyLassoMarqueeMechanic.h
--- a/Engine/Plugins/Experimental/MeshModelingToolset/Source/ModelingComponents/Public/Mechanics/PolyLassoMarqueeMechanic.h
+++ b/Engine/Plugins/Experimental/MeshModelingToolset/Source/ModelingComponents/Public/Mechanics/PolyLassoMarqueeMechanic.h
@@ -148,10 +148,10 @@ public:
 
 protected:
 	UPROPERTY()
-	UClickDragInputBehavior* ClickDragBehavior = nullptr;
+	TObjectPtr<UClickDragInputBehavior> ClickDragBehavior = nullptr;
 
 	UPROPERTY()
-	UMouseHoverBehavior* HoverBehavior = nullptr;
+	TObjectPtr<UMouseHoverBehavior> HoverBehavior = nullptr;
 
 	// camera state for currently active PolyLasso
 	FViewCameraState CachedCameraState;
diff --git a/Engine/Plugins/Experimental/MeshModelingToolset/Source/ModelingComponents/Public/Mechanics/RectangleMarqueeMechanic.h b/Engine/Plugins/Experimental/MeshModelingToolset/Source/ModelingComponents/Public/Mechanics/RectangleMarqueeMechanic.h
--- a/Engine/Plugins/Experimental/MeshModelingToolset/Source/ModelingComponents/Public/Mechanics/RectangleMarqueeMechanic.h
+++ b/Engine/Plugins/Experimental/MeshModelingToolset/Source/ModelingComponents/Public/Mechanics/RectangleMarqueeMechanic.h
@@ -115,7 +115,7 @@ public:
 	
 protected:
 	UPROPERTY()
-	UClickDragInputBehavior* ClickDragBehavior = nullptr;
+	TObjectPtr<UClickDragInputBehavior> ClickDragBehavior = nullptr;
 
 	FViewCameraState CachedCameraState;
 
diff --git a/Engine/Plugins/Experimental/MeshModelingToolset/Source/ModelingComponents/Public/Mechanics/SpaceCurveDeformationMechanic.h b/Engine/Plugins/Experimental/MeshModelingToolset/Source/ModelingComponents/Public/Mechanics/SpaceCurveDeformationMechanic.h
--- a/Engine/Plugins/Experimental/MeshModelingToolset/Source/ModelingComponents/Public/Mechanics/SpaceCurveDeformationMechanic.h
+++ b/Engine/Plugins/Experimental/MeshModelingToolset/Source/ModelingComponents/Public/Mechanics/SpaceCurveDeformationMechanic.h
@@ -98,9 +98,9 @@ public:
 
 	// Behaviors used for moving points around and hovering them
 	UPROPERTY()
-	USingleClickInputBehavior* ClickBehavior = nullptr;
+	TObjectPtr<USingleClickInputBehavior> ClickBehavior = nullptr;
 	UPROPERTY()
-	UMouseHoverBehavior* HoverBehavior = nullptr;
+	TObjectPtr<UMouseHoverBehavior> HoverBehavior = nullptr;
 
 	// This delegate is called every time the control point sequence is altered.
 	DECLARE_MULTICAST_DELEGATE(OnPointsChangedEvent);
@@ -144,7 +144,7 @@ public:
 
 public:
 	UPROPERTY()
-	USpaceCurveDeformationMechanicPropertySet* TransformProperties;
+	TObjectPtr<USpaceCurveDeformationMechanicPropertySet> TransformProperties;
 
 protected:
 
@@ -161,11 +161,11 @@ protected:
 
 	/** Used for displaying points/segments */
 	UPROPERTY()
-	APreviewGeometryActor* PreviewGeometryActor;
+	TObjectPtr<APreviewGeometryActor> PreviewGeometryActor;
 	UPROPERTY()
-	UPointSetComponent* RenderPoints;
+	TObjectPtr<UPointSetComponent> RenderPoints;
 	UPROPERTY()
-	ULineSetComponent* RenderSegments;
+	TObjectPtr<ULineSetComponent> RenderSegments;
 
 	bool bRenderGeometryValid = false;
 	void UpdateRenderGeometry();
@@ -189,9 +189,9 @@ protected:
 	// Support for gizmo. Since the points aren't individual components, we don't actually use UTransformProxy
 	// for the transform forwarding- we just use it for the callbacks.
 	UPROPERTY()
-	UTransformProxy* PointTransformProxy;
+	TObjectPtr<UTransformProxy> PointTransformProxy;
 	UPROPERTY()
-	UTransformGizmo* PointTransformGizmo;
+	TObjectPtr<UTransformGizmo> PointTransformGizmo;
 
 	// Used to make it easy to tell whether the gizmo was moved by the user or by undo/redo or
 	// some other change that we shoulnd't respond to. Basing our movement undo/redo on the
diff --git a/Engine/Plugins/Experimental/MeshModelingToolset/Source/ModelingComponents/Public/MeshOpPreviewHelpers.h b/Engine/Plugins/Experimental/MeshModelingToolset/Source/ModelingComponents/Public/MeshOpPreviewHelpers.h
--- a/Engine/Plugins/Experimental/MeshModelingToolset/Source/ModelingComponents/Public/MeshOpPreviewHelpers.h
+++ b/Engine/Plugins/Experimental/MeshModelingToolset/Source/ModelingComponents/Public/MeshOpPreviewHelpers.h
@@ -197,23 +197,23 @@ public:
 public:
 	// preview of MeshOperator result
 	UPROPERTY()
-	UPreviewMesh* PreviewMesh = nullptr;
+	TObjectPtr<UPreviewMesh> PreviewMesh = nullptr;
 
 	// input set of materials to assign to PreviewMesh
 	UPROPERTY()
-	TArray<UMaterialInterface*> StandardMaterials;
+	TArray<TObjectPtr<UMaterialInterface>> StandardMaterials;
 
 	// override material to forward to PreviewMesh if set
 	UPROPERTY()
-	UMaterialInterface* OverrideMaterial = nullptr;
+	TObjectPtr<UMaterialInterface> OverrideMaterial = nullptr;
 
 	// if non-null, this material is swapped in when a background compute is active
 	UPROPERTY()
-	UMaterialInterface* WorkingMaterial = nullptr;
+	TObjectPtr<UMaterialInterface> WorkingMaterial = nullptr;
 	
 	// secondary render material to forward to PreviewMesh if set
 	UPROPERTY()
-	UMaterialInterface* SecondaryMaterial = nullptr;
+	TObjectPtr<UMaterialInterface> SecondaryMaterial = nullptr;
 
 	UWorld* PreviewWorld;
 protected:
diff --git a/Engine/Plugins/Experimental/MeshModelingToolset/Source/ModelingComponents/Public/PreviewMesh.h b/Engine/Plugins/Experimental/MeshModelingToolset/Source/ModelingComponents/Public/PreviewMesh.h
--- a/Engine/Plugins/Experimental/MeshModelingToolset/Source/ModelingComponents/Public/PreviewMesh.h
+++ b/Engine/Plugins/Experimental/MeshModelingToolset/Source/ModelingComponents/Public/PreviewMesh.h
@@ -427,7 +427,7 @@ protected:
 
 	/** This component is set as the root component of TemporaryParentActor */
 	UPROPERTY()
-	UDynamicMeshComponent* DynamicMeshComponent = nullptr;
+	TObjectPtr<UDynamicMeshComponent> DynamicMeshComponent = nullptr;
 
 	/** Spatial data structure that is initialized if bBuildSpatialDataStructure = true when UpdatePreview() is called */
 	UE::Geometry::FDynamicMeshAABBTree3 MeshAABBTree;
diff --git a/Engine/Plugins/Experimental/MeshModelingToolset/Source/ModelingComponents/Public/Selection/MeshSelectionMechanic.h b/Engine/Plugins/Experimental/MeshModelingToolset/Source/ModelingComponents/Public/Selection/MeshSelectionMechanic.h
--- a/Engine/Plugins/Experimental/MeshModelingToolset/Source/ModelingComponents/Public/Selection/MeshSelectionMechanic.h
+++ b/Engine/Plugins/Experimental/MeshModelingToolset/Source/ModelingComponents/Public/Selection/MeshSelectionMechanic.h
@@ -90,10 +90,10 @@ protected:
 	bool InMultiSelectMode() const { return bShiftToggle; }
 
 	UPROPERTY()
-	APreviewGeometryActor* PreviewGeometryActor = nullptr;
+	TObjectPtr<APreviewGeometryActor> PreviewGeometryActor = nullptr;
 
 	UPROPERTY()
-	ULineSetComponent* LineSet = nullptr;
+	TObjectPtr<ULineSetComponent> LineSet = nullptr;
 
 
 	TArray<TSharedPtr<FDynamicMeshAABBTree3>> MeshSpatials;
diff --git a/Engine/Plugins/Experimental/MeshModelingToolset/Source/ModelingComponents/Public/Selection/PolygonSelectionMechanic.h b/Engine/Plugins/Experimental/MeshModelingToolset/Source/ModelingComponents/Public/Selection/PolygonSelectionMechanic.h
--- a/Engine/Plugins/Experimental/MeshModelingToolset/Source/ModelingComponents/Public/Selection/PolygonSelectionMechanic.h
+++ b/Engine/Plugins/Experimental/MeshModelingToolset/Source/ModelingComponents/Public/Selection/PolygonSelectionMechanic.h
@@ -317,7 +317,7 @@ public:
 	// allow the user to bundle in some FChanges into the normal callback?
 
 	UPROPERTY()
-	UPolygonSelectionMechanicProperties* Properties;
+	TObjectPtr<UPolygonSelectionMechanicProperties> Properties;
 
 protected:
 	const FDynamicMesh3* Mesh;
@@ -325,13 +325,13 @@ protected:
 	TFunction<FDynamicMeshAABBTree3*()> GetSpatialFunc;
 
 	UPROPERTY()
-	UMouseHoverBehavior* HoverBehavior;
+	TObjectPtr<UMouseHoverBehavior> HoverBehavior;
 
 	UPROPERTY()
-	USingleClickInputBehavior* ClickBehavior;
+	TObjectPtr<USingleClickInputBehavior> ClickBehavior;
 
 	UPROPERTY()
-	URectangleMarqueeMechanic* MarqueeMechanic;
+	TObjectPtr<URectangleMarqueeMechanic> MarqueeMechanic;
 
 	FInputCapturePriority BasePriority = FInputCapturePriority(FInputCapturePriority::DEFAULT_TOOL_PRIORITY);
 
@@ -371,15 +371,15 @@ protected:
 
 	/** The actor we create internally to own the DrawnTriangleSetComponent */
 	UPROPERTY()
-	APreviewGeometryActor* PreviewGeometryActor;
+	TObjectPtr<APreviewGeometryActor> PreviewGeometryActor;
 
 	UPROPERTY()
-	UTriangleSetComponent* DrawnTriangleSetComponent;
+	TObjectPtr<UTriangleSetComponent> DrawnTriangleSetComponent;
 
 	TSet<int> CurrentlyHighlightedGroups;
 
 	UPROPERTY()
-	UMaterialInterface* HighlightedFaceMaterial;
+	TObjectPtr<UMaterialInterface> HighlightedFaceMaterial;
 
 	FViewCameraState CameraState;
 
diff --git a/Engine/Plugins/Experimental/MeshModelingToolset/Source/ModelingComponents/Public/Transforms/MultiTransformer.h b/Engine/Plugins/Experimental/MeshModelingToolset/Source/ModelingComponents/Public/Transforms/MultiTransformer.h
--- a/Engine/Plugins/Experimental/MeshModelingToolset/Source/ModelingComponents/Public/Transforms/MultiTransformer.h
+++ b/Engine/Plugins/Experimental/MeshModelingToolset/Source/ModelingComponents/Public/Transforms/MultiTransformer.h
@@ -92,7 +92,7 @@ public:
 
 protected:
 	UPROPERTY()
-	UInteractiveGizmoManager* GizmoManager;
+	TObjectPtr<UInteractiveGizmoManager> GizmoManager;
 
 	IToolContextTransactionProvider* TransactionProvider;
 
@@ -112,15 +112,15 @@ protected:
 	bool bDisallowNegativeScaling = false;
 
 	UPROPERTY()
-	UTransformGizmo* TransformGizmo;
+	TObjectPtr<UTransformGizmo> TransformGizmo;
 
 	UPROPERTY()
-	UTransformProxy* TransformProxy;
+	TObjectPtr<UTransformProxy> TransformProxy;
 
 	// We have to hold on to the mechanic only because the MultiTransformer has the capacity to delete and
 	// recreate its gizmo, in which case we'll need to attach the alignment mechanic again.
 	UPROPERTY()
-	UDragAlignmentMechanic* DragAlignmentMechanic = nullptr;
+	TObjectPtr<UDragAlignmentMechanic> DragAlignmentMechanic = nullptr;
 
 	TUniqueFunction<bool()> EnableSnapToWorldGridFunc;
 	TFunction<bool()> IsNonUniformScaleAllowed;
diff --git a/Engine/Plugins/FX/Niagara/Source/Niagara/Classes/NiagaraBakerSettings.h b/Engine/Plugins/FX/Niagara/Source/Niagara/Classes/NiagaraBakerSettings.h
--- a/Engine/Plugins/FX/Niagara/Source/Niagara/Classes/NiagaraBakerSettings.h
+++ b/Engine/Plugins/FX/Niagara/Source/Niagara/Classes/NiagaraBakerSettings.h
@@ -62,7 +62,7 @@ struct FNiagaraBakerTextureSettings
 
 	/** Final texture generated, an existing entry will be updated with new capture data. */
 	UPROPERTY(EditAnywhere, Category = "Texture")
-	UTexture2D* GeneratedTexture = nullptr;
+	TObjectPtr<UTexture2D> GeneratedTexture = nullptr;
 
 	bool Equals(const FNiagaraBakerTextureSettings& Other) const;
 
diff --git a/Engine/Plugins/FX/Niagara/Source/Niagara/Classes/NiagaraDataInterfaceCubeTexture.h b/Engine/Plugins/FX/Niagara/Source/Niagara/Classes/NiagaraDataInterfaceCubeTexture.h
--- a/Engine/Plugins/FX/Niagara/Source/Niagara/Classes/NiagaraDataInterfaceCubeTexture.h
+++ b/Engine/Plugins/FX/Niagara/Source/Niagara/Classes/NiagaraDataInterfaceCubeTexture.h
@@ -19,7 +19,7 @@ public:
 	DECLARE_NIAGARA_DI_PARAMETER();
 
 	UPROPERTY(EditAnywhere, Category = "Texture")
-	UTextureCube* Texture;
+	TObjectPtr<UTextureCube> Texture;
 
 	//UObject Interface
 	virtual void PostInitProperties()override;
diff --git a/Engine/Plugins/FX/Niagara/Source/Niagara/Classes/NiagaraDataInterfaceIntRenderTarget2D.h b/Engine/Plugins/FX/Niagara/Source/Niagara/Classes/NiagaraDataInterfaceIntRenderTarget2D.h
--- a/Engine/Plugins/FX/Niagara/Source/Niagara/Classes/NiagaraDataInterfaceIntRenderTarget2D.h
+++ b/Engine/Plugins/FX/Niagara/Source/Niagara/Classes/NiagaraDataInterfaceIntRenderTarget2D.h
@@ -67,5 +67,5 @@ protected:
 	static FNiagaraVariableBase ExposedRTVar;
 
 	UPROPERTY(Transient, DuplicateTransient)
-	TMap<uint64, UTextureRenderTarget2D*> ManagedRenderTargets;
+	TMap<uint64, TObjectPtr<UTextureRenderTarget2D>> ManagedRenderTargets;
 };
diff --git a/Engine/Plugins/FX/Niagara/Source/Niagara/Classes/NiagaraDataInterfaceLandscape.h b/Engine/Plugins/FX/Niagara/Source/Niagara/Classes/NiagaraDataInterfaceLandscape.h
--- a/Engine/Plugins/FX/Niagara/Source/Niagara/Classes/NiagaraDataInterfaceLandscape.h
+++ b/Engine/Plugins/FX/Niagara/Source/Niagara/Classes/NiagaraDataInterfaceLandscape.h
@@ -49,7 +49,7 @@ public:
 	ENDILandscape_SourceMode SourceMode = ENDILandscape_SourceMode::Default;
 
 	UPROPERTY(EditAnywhere, Category = "Landscape")
-	TArray<UPhysicalMaterial*> PhysicalMaterials;
+	TArray<TObjectPtr<UPhysicalMaterial>> PhysicalMaterials;
 
 	//UObject Interface
 	virtual void PostInitProperties() override;	
diff --git a/Engine/Plugins/FX/Niagara/Source/Niagara/Classes/NiagaraDataInterfaceRenderTargetCube.h b/Engine/Plugins/FX/Niagara/Source/Niagara/Classes/NiagaraDataInterfaceRenderTargetCube.h
--- a/Engine/Plugins/FX/Niagara/Source/Niagara/Classes/NiagaraDataInterfaceRenderTargetCube.h
+++ b/Engine/Plugins/FX/Niagara/Source/Niagara/Classes/NiagaraDataInterfaceRenderTargetCube.h
@@ -153,5 +153,5 @@ protected:
 	static FNiagaraVariableBase ExposedRTVar;
 	
 	UPROPERTY(Transient, DuplicateTransient)
-	TMap<uint64, UTextureRenderTargetCube*> ManagedRenderTargets;
+	TMap<uint64, TObjectPtr<UTextureRenderTargetCube>> ManagedRenderTargets;
 };
diff --git a/Engine/Plugins/FX/Niagara/Source/Niagara/Classes/NiagaraEmitter.h b/Engine/Plugins/FX/Niagara/Source/Niagara/Classes/NiagaraEmitter.h
--- a/Engine/Plugins/FX/Niagara/Source/Niagara/Classes/NiagaraEmitter.h
+++ b/Engine/Plugins/FX/Niagara/Source/Niagara/Classes/NiagaraEmitter.h
@@ -614,7 +614,7 @@ private:
 
 	/** Wrapper for editor only parameters. */
 	UPROPERTY()
-	UNiagaraEditorParametersAdapterBase* EditorParameters;
+	TObjectPtr<UNiagaraEditorParametersAdapterBase> EditorParameters;
 
 	/** A multicast delegate which is called whenever all the scripts for this emitter have been compiled (successfully or not). */
 	FOnEmitterCompiled OnVMScriptCompiledDelegate;
diff --git a/Engine/Plugins/FX/Niagara/Source/Niagara/Classes/NiagaraLensEffectBase.h b/Engine/Plugins/FX/Niagara/Source/Niagara/Classes/NiagaraLensEffectBase.h
--- a/Engine/Plugins/FX/Niagara/Source/Niagara/Classes/NiagaraLensEffectBase.h
+++ b/Engine/Plugins/FX/Niagara/Source/Niagara/Classes/NiagaraLensEffectBase.h
@@ -59,7 +59,7 @@ protected:
 	TArray<TSubclassOf<AActor>> EmittersToTreatAsSame;
 
 	UPROPERTY()
-	APlayerCameraManager* OwningCameraManager;
+	TObjectPtr<APlayerCameraManager> OwningCameraManager;
 
 protected:
 
diff --git a/Engine/Plugins/FX/Niagara/Source/Niagara/Classes/NiagaraParameterDefinitionsSubscriber.h b/Engine/Plugins/FX/Niagara/Source/Niagara/Classes/NiagaraParameterDefinitionsSubscriber.h
--- a/Engine/Plugins/FX/Niagara/Source/Niagara/Classes/NiagaraParameterDefinitionsSubscriber.h
+++ b/Engine/Plugins/FX/Niagara/Source/Niagara/Classes/NiagaraParameterDefinitionsSubscriber.h
@@ -28,7 +28,7 @@ public:
 	{};
 
 	UPROPERTY(meta = (DeprecatedProperty))
-	UNiagaraParameterDefinitionsBase* ParameterDefinitions_DEPRECATED;
+	TObjectPtr<UNiagaraParameterDefinitionsBase> ParameterDefinitions_DEPRECATED;
 
 	UPROPERTY()
 	FGuid DefinitionsId;
diff --git a/Engine/Plugins/FX/Niagara/Source/Niagara/Classes/NiagaraScript.h b/Engine/Plugins/FX/Niagara/Source/Niagara/Classes/NiagaraScript.h
--- a/Engine/Plugins/FX/Niagara/Source/Niagara/Classes/NiagaraScript.h
+++ b/Engine/Plugins/FX/Niagara/Source/Niagara/Classes/NiagaraScript.h
@@ -665,7 +665,7 @@ public:
 
 	/** Use property in struct returned from GetScriptData() instead */
 	UPROPERTY(meta = (DeprecatedProperty))
-	UNiagaraScript* DeprecationRecommendation_DEPRECATED;
+	TObjectPtr<UNiagaraScript> DeprecationRecommendation_DEPRECATED;
 
 	/** Use property in struct returned from GetScriptData() instead */
 	UPROPERTY(meta = (DeprecatedProperty))
@@ -717,7 +717,7 @@ public:
 
 	/** 'Source' data/graphs for this script */
 	UPROPERTY(meta = (DeprecatedProperty))
-	class UNiagaraScriptSourceBase*	Source_DEPRECATED;
+	TObjectPtr<class UNiagaraScriptSourceBase>	Source_DEPRECATED;
 
 	NIAGARA_API static const FName NiagaraCustomVersionTagName;
 #endif
diff --git a/Engine/Plugins/FX/Niagara/Source/Niagara/Classes/NiagaraSystem.h b/Engine/Plugins/FX/Niagara/Source/Niagara/Classes/NiagaraSystem.h
--- a/Engine/Plugins/FX/Niagara/Source/Niagara/Classes/NiagaraSystem.h
+++ b/Engine/Plugins/FX/Niagara/Source/Niagara/Classes/NiagaraSystem.h
@@ -742,7 +742,7 @@ protected:
 
 	/** Wrapper for editor only parameters. */
 	UPROPERTY()
-	UNiagaraEditorParametersAdapterBase* EditorParameters;
+	TObjectPtr<UNiagaraEditorParametersAdapterBase> EditorParameters;
 
 	bool bIsolateEnabled;
 
@@ -775,11 +775,11 @@ protected:
 #if WITH_EDITORONLY_DATA
 	/** Settings used inside the baker */
 	UPROPERTY(Export)
-	UNiagaraBakerSettings* BakerSettings;
+	TObjectPtr<UNiagaraBakerSettings> BakerSettings;
 
 	/** Generated data baker settings, will be null until we have generated at least once. */
 	UPROPERTY(Export)
-	UNiagaraBakerSettings* BakerGeneratedSettings;
+	TObjectPtr<UNiagaraBakerSettings> BakerGeneratedSettings;
 #endif
 
 	UPROPERTY()
diff --git a/Engine/Plugins/FX/Niagara/Source/Niagara/Public/NiagaraComponent.h b/Engine/Plugins/FX/Niagara/Source/Niagara/Public/NiagaraComponent.h
--- a/Engine/Plugins/FX/Niagara/Source/Niagara/Public/NiagaraComponent.h
+++ b/Engine/Plugins/FX/Niagara/Source/Niagara/Public/NiagaraComponent.h
@@ -690,7 +690,7 @@ private:
 	float CustomTimeDilation = 1.0f;
 
 	UPROPERTY()
-	class UNiagaraCullProxyComponent* CullProxy;
+	TObjectPtr<class UNiagaraCullProxyComponent> CullProxy;
 
 	void CreateCullProxy(bool bForce = false);
 	void DestroyCullProxy();
diff --git a/Engine/Plugins/FX/Niagara/Source/NiagaraEditor/Public/NiagaraActions.h b/Engine/Plugins/FX/Niagara/Source/NiagaraEditor/Public/NiagaraActions.h
--- a/Engine/Plugins/FX/Niagara/Source/NiagaraEditor/Public/NiagaraActions.h
+++ b/Engine/Plugins/FX/Niagara/Source/NiagaraEditor/Public/NiagaraActions.h
@@ -208,7 +208,7 @@ struct NIAGARAEDITOR_API FNiagaraAction_NewNode : public FNiagaraMenuAction_Gene
 	class UEdGraphNode* CreateNode(UEdGraph* Graph, TArray<UEdGraphPin*>& FromPins, FVector2D NodePosition, bool bSelectNewNode = true) const;
 
 	UPROPERTY()
-	class UEdGraphNode* NodeTemplate = nullptr;
+	TObjectPtr<class UEdGraphNode> NodeTemplate = nullptr;
 };
 
 struct NIAGARAEDITOR_API FNiagaraParameterAction : public FEdGraphSchemaAction
diff --git a/Engine/Plugins/FX/Niagara/Source/NiagaraEditor/Public/NiagaraEditorData.h b/Engine/Plugins/FX/Niagara/Source/NiagaraEditor/Public/NiagaraEditorData.h
--- a/Engine/Plugins/FX/Niagara/Source/NiagaraEditor/Public/NiagaraEditorData.h
+++ b/Engine/Plugins/FX/Niagara/Source/NiagaraEditor/Public/NiagaraEditorData.h
@@ -50,6 +50,6 @@ public:
 
 private:
 	UPROPERTY()
-	TArray<UNiagaraScriptVariable*> EditorOnlyScriptVars;
+	TArray<TObjectPtr<UNiagaraScriptVariable>> EditorOnlyScriptVars;
 #endif
 };
diff --git a/Engine/Plugins/FX/Niagara/Source/NiagaraEditor/Public/NiagaraEditorModule.h b/Engine/Plugins/FX/Niagara/Source/NiagaraEditor/Public/NiagaraEditorModule.h
--- a/Engine/Plugins/FX/Niagara/Source/NiagaraEditor/Public/NiagaraEditorModule.h
+++ b/Engine/Plugins/FX/Niagara/Source/NiagaraEditor/Public/NiagaraEditorModule.h
@@ -84,7 +84,7 @@ private:
 	FNiagaraVariableBase Parameter;
 		 
 	UPROPERTY(transient)
-	UNiagaraParameterDefinitions* ReservingDefinitionsAsset;
+	TObjectPtr<UNiagaraParameterDefinitions> ReservingDefinitionsAsset;
 };
 
 FORCEINLINE uint32 GetTypeHash(const FReservedParameter& ReservedParameter) { return GetTypeHash(ReservedParameter.GetParameter().GetName()); };
diff --git a/Engine/Plugins/FX/Niagara/Source/NiagaraEditor/Public/NiagaraParameterDefinitions.h b/Engine/Plugins/FX/Niagara/Source/NiagaraEditor/Public/NiagaraParameterDefinitions.h
--- a/Engine/Plugins/FX/Niagara/Source/NiagaraEditor/Public/NiagaraParameterDefinitions.h
+++ b/Engine/Plugins/FX/Niagara/Source/NiagaraEditor/Public/NiagaraParameterDefinitions.h
@@ -37,7 +37,7 @@ public:
 	{};
 
 	UPROPERTY()
-	UNiagaraParameterDefinitions* SubscribedParameterDefinitions;
+	TObjectPtr<UNiagaraParameterDefinitions> SubscribedParameterDefinitions;
 
 	UPROPERTY()
 	TArray<FScriptVarBindingNameSubscription> BindingNameSubscriptions;
@@ -101,7 +101,7 @@ private:
 	int32 MenuSortOrder;
 
 	UPROPERTY()
-	TArray<UNiagaraScriptVariable*> ScriptVariables;
+	TArray<TObjectPtr<UNiagaraScriptVariable>> ScriptVariables;
 
 	UPROPERTY()
 	TArray<FParameterDefinitionsBindingNameSubscription> ExternalParameterDefinitionsSubscriptions;
diff --git a/Engine/Plugins/FX/Niagara/Source/NiagaraEditor/Public/UpgradeNiagaraScriptResults.h b/Engine/Plugins/FX/Niagara/Source/NiagaraEditor/Public/UpgradeNiagaraScriptResults.h
--- a/Engine/Plugins/FX/Niagara/Source/NiagaraEditor/Public/UpgradeNiagaraScriptResults.h
+++ b/Engine/Plugins/FX/Niagara/Source/NiagaraEditor/Public/UpgradeNiagaraScriptResults.h
@@ -40,7 +40,7 @@ public:
 	UNiagaraPythonScriptModuleInput() {};
 
 	UPROPERTY()
-	const UNiagaraClipboardFunctionInput* Input = nullptr;
+	TObjectPtr<const UNiagaraClipboardFunctionInput> Input = nullptr;
 
 	UFUNCTION(BlueprintCallable, Category = "Scripting")
     bool IsSet() const;
@@ -95,10 +95,10 @@ public:
 	bool bCancelledByPythonError;
 
 	UPROPERTY(BlueprintReadWrite, Category = "Scripting")
-	TArray<UNiagaraPythonScriptModuleInput*> OldInputs;
+	TArray<TObjectPtr<UNiagaraPythonScriptModuleInput>> OldInputs;
 
 	UPROPERTY(BlueprintReadWrite, Category = "Scripting")
-	TArray<UNiagaraPythonScriptModuleInput*> NewInputs;
+	TArray<TObjectPtr<UNiagaraPythonScriptModuleInput>> NewInputs;
 
 	UFUNCTION(BlueprintCallable, Category = "Scripting")
     UNiagaraPythonScriptModuleInput* GetOldInput(const FString& InputName);
@@ -135,7 +135,7 @@ private:
 
 	// This is used as a placeholder object for python to interact with when a module input could not be found. Returning a nullptr instead would crash the script.
 	UPROPERTY(Transient)
-	UNiagaraPythonScriptModuleInput* DummyInput;
+	TObjectPtr<UNiagaraPythonScriptModuleInput> DummyInput;
 };
 
 struct NIAGARAEDITOR_API FNiagaraScriptVersionUpgradeContext {
diff --git a/Engine/Plugins/FX/Niagara/Source/NiagaraEditor/Public/ViewModels/Stack/NiagaraStackEntry.h b/Engine/Plugins/FX/Niagara/Source/NiagaraEditor/Public/ViewModels/Stack/NiagaraStackEntry.h
--- a/Engine/Plugins/FX/Niagara/Source/NiagaraEditor/Public/ViewModels/Stack/NiagaraStackEntry.h
+++ b/Engine/Plugins/FX/Niagara/Source/NiagaraEditor/Public/ViewModels/Stack/NiagaraStackEntry.h
@@ -494,7 +494,7 @@ private:
 	TWeakPtr<FNiagaraEmitterViewModel> EmitterViewModel;
 
 	UPROPERTY()
-	UNiagaraStackEditorData* StackEditorData;
+	TObjectPtr<UNiagaraStackEditorData> StackEditorData;
 
 	FString StackEditorDataKey;
 
diff --git a/Engine/Plugins/MovieScene/SequencerScripting/Source/SequencerScripting/Public/SequencerBindingProxy.h b/Engine/Plugins/MovieScene/SequencerScripting/Source/SequencerScripting/Public/SequencerBindingProxy.h
--- a/Engine/Plugins/MovieScene/SequencerScripting/Source/SequencerScripting/Public/SequencerBindingProxy.h
+++ b/Engine/Plugins/MovieScene/SequencerScripting/Source/SequencerScripting/Public/SequencerBindingProxy.h
@@ -31,5 +31,5 @@ struct FSequencerBindingProxy
 	FGuid BindingID;
 
 	UPROPERTY(BlueprintReadOnly, Category=Binding)
-	UMovieSceneSequence* Sequence;
+	TObjectPtr<UMovieSceneSequence> Sequence;
 };
diff --git a/Engine/Plugins/MovieScene/SequencerScripting/Source/SequencerScriptingEditor/Public/SequencerTools.h b/Engine/Plugins/MovieScene/SequencerScripting/Source/SequencerScriptingEditor/Public/SequencerTools.h
--- a/Engine/Plugins/MovieScene/SequencerScripting/Source/SequencerScriptingEditor/Public/SequencerTools.h
+++ b/Engine/Plugins/MovieScene/SequencerScripting/Source/SequencerScriptingEditor/Public/SequencerTools.h
@@ -42,7 +42,7 @@ struct SEQUENCERSCRIPTINGEDITOR_API FSequencerBoundObjects
 	FSequencerBindingProxy BindingProxy;
 
 	UPROPERTY(BlueprintReadWrite, Category=Binding)
-	TArray<UObject*> BoundObjects;
+	TArray<TObjectPtr<UObject>> BoundObjects;
 };
 
 /** Wrapper around result of quick binding for event track in sequencer. */
@@ -55,7 +55,7 @@ struct FSequencerQuickBindingResult
 
 	/** Actual endpoint wrapped by this structure.  */
 	UPROPERTY()
-	UK2Node_CustomEvent* EventEndpoint;
+	TObjectPtr<UK2Node_CustomEvent> EventEndpoint;
 
 	/** Names of the payload variables of the event. */
 	UPROPERTY(BlueprintReadOnly, Category = Data)
diff --git a/Engine/Plugins/Runtime/HairStrands/Source/HairStrandsCore/Public/GroomActor.h b/Engine/Plugins/Runtime/HairStrands/Source/HairStrandsCore/Public/GroomActor.h
--- a/Engine/Plugins/Runtime/HairStrands/Source/HairStrandsCore/Public/GroomActor.h
+++ b/Engine/Plugins/Runtime/HairStrands/Source/HairStrandsCore/Public/GroomActor.h
@@ -16,7 +16,7 @@ class HAIRSTRANDSCORE_API AGroomActor : public AActor
 
 	/** Strand hair component that performs simulation and rendering */
 	UPROPERTY(Category = StrandHair, VisibleAnywhere, BlueprintReadOnly)
-	class UGroomComponent* GroomComponent;
+	TObjectPtr<class UGroomComponent> GroomComponent;
 
 	//~ Begin AActor Interface
 #if WITH_EDITOR
@@ -28,7 +28,7 @@ class HAIRSTRANDSCORE_API AGroomActor : public AActor
 protected:
 	/** Billboard used to see the scene in the editor */
 	UPROPERTY()
-		UBillboardComponent* SpriteComponent;
+		TObjectPtr<UBillboardComponent> SpriteComponent;
 #endif
 
 public:
diff --git a/Engine/Plugins/Runtime/HairStrands/Source/HairStrandsCore/Public/GroomAsset.h b/Engine/Plugins/Runtime/HairStrands/Source/HairStrandsCore/Public/GroomAsset.h
--- a/Engine/Plugins/Runtime/HairStrands/Source/HairStrandsCore/Public/GroomAsset.h
+++ b/Engine/Plugins/Runtime/HairStrands/Source/HairStrandsCore/Public/GroomAsset.h
@@ -62,7 +62,7 @@ struct HAIRSTRANDSCORE_API FHairGroupsMaterial
 	GENERATED_BODY()
 
 	UPROPERTY(EditAnywhere, Category = "Material")
-	UMaterialInterface* Material = nullptr;
+	TObjectPtr<UMaterialInterface> Material = nullptr;
 
 	UPROPERTY(EditAnywhere, Category = "Material")
 	FName SlotName = NAME_None;
@@ -480,7 +480,7 @@ public:
 
 	/** Asset data to be used when re-importing */
 	UPROPERTY(VisibleAnywhere, Instanced, Category = ImportSettings)
-	class UAssetImportData* AssetImportData;
+	TObjectPtr<class UAssetImportData> AssetImportData;
 
 	/** Retrievde the asset tags*/
 	virtual void GetAssetRegistryTags(TArray<FAssetRegistryTag>& OutTags) const override;
@@ -562,7 +562,7 @@ public:
 
 	/** Array of user data stored with the asset */
 	UPROPERTY(EditAnywhere, AdvancedDisplay, Instanced, Category = Hidden)
-	TArray<UAssetUserData*> AssetUserData;
+	TArray<TObjectPtr<UAssetUserData>> AssetUserData;
 
 	/* Return the material slot index corresponding to the material name */
 	int32 GetMaterialIndex(FName MaterialSlotName) const;
diff --git a/Engine/Plugins/Runtime/HairStrands/Source/HairStrandsCore/Public/GroomAssetCards.h b/Engine/Plugins/Runtime/HairStrands/Source/HairStrandsCore/Public/GroomAssetCards.h
--- a/Engine/Plugins/Runtime/HairStrands/Source/HairStrandsCore/Public/GroomAssetCards.h
+++ b/Engine/Plugins/Runtime/HairStrands/Source/HairStrandsCore/Public/GroomAssetCards.h
@@ -167,19 +167,19 @@ struct HAIRSTRANDSCORE_API FHairGroupCardsTextures
 	GENERATED_BODY()
 
 	UPROPERTY(EditAnywhere, Category = "CardsTextures")
-	UTexture2D* DepthTexture = nullptr;
+	TObjectPtr<UTexture2D> DepthTexture = nullptr;
 
 	UPROPERTY(EditAnywhere, Category = "CardsTextures")
-	UTexture2D* CoverageTexture = nullptr;
+	TObjectPtr<UTexture2D> CoverageTexture = nullptr;
 
 	UPROPERTY(EditAnywhere, Category = "CardsTextures")
-	UTexture2D* TangentTexture = nullptr;
+	TObjectPtr<UTexture2D> TangentTexture = nullptr;
 
 	UPROPERTY(EditAnywhere, Category = "CardsAttributes")
-	UTexture2D* AttributeTexture = nullptr;
+	TObjectPtr<UTexture2D> AttributeTexture = nullptr;
 
 	UPROPERTY(EditAnywhere, Category = "CardsAuxilaryData")
-	UTexture2D* AuxilaryDataTexture = nullptr;
+	TObjectPtr<UTexture2D> AuxilaryDataTexture = nullptr;
 
 	bool bNeedToBeSaved = false;
 };
@@ -193,7 +193,7 @@ struct HAIRSTRANDSCORE_API FHairGroupsCardsSourceDescription
 
 	/* Deprecated */
 	UPROPERTY()
-	UMaterialInterface* Material = nullptr;
+	TObjectPtr<UMaterialInterface> Material = nullptr;
 
 	UPROPERTY()
 	FName MaterialSlotName;
@@ -202,13 +202,13 @@ struct HAIRSTRANDSCORE_API FHairGroupsCardsSourceDescription
 	EHairCardsSourceType SourceType = EHairCardsSourceType::Imported;
 
 	UPROPERTY(EditAnywhere, Category = "CardsSource")
-	UStaticMesh* ProceduralMesh = nullptr;
+	TObjectPtr<UStaticMesh> ProceduralMesh = nullptr;
 
 	UPROPERTY()
 	FString ProceduralMeshKey;
 
 	UPROPERTY(EditAnywhere, Category = "CardsSource")
-	UStaticMesh* ImportedMesh = nullptr;
+	TObjectPtr<UStaticMesh> ImportedMesh = nullptr;
 
 	UPROPERTY(EditAnywhere, Category = "CardsSource")
 	FHairGroupsProceduralCards ProceduralSettings;
diff --git a/Engine/Plugins/Runtime/HairStrands/Source/HairStrandsCore/Public/GroomAssetImportData.h b/Engine/Plugins/Runtime/HairStrands/Source/HairStrandsCore/Public/GroomAssetImportData.h
--- a/Engine/Plugins/Runtime/HairStrands/Source/HairStrandsCore/Public/GroomAssetImportData.h
+++ b/Engine/Plugins/Runtime/HairStrands/Source/HairStrandsCore/Public/GroomAssetImportData.h
@@ -15,5 +15,5 @@ class HAIRSTRANDSCORE_API UGroomAssetImportData : public UAssetImportData
 public:
 
 	UPROPERTY()
-	class UGroomImportOptions* ImportOptions;
+	TObjectPtr<class UGroomImportOptions> ImportOptions;
 };
diff --git a/Engine/Plugins/Runtime/HairStrands/Source/HairStrandsCore/Public/GroomAssetMeshes.h b/Engine/Plugins/Runtime/HairStrands/Source/HairStrandsCore/Public/GroomAssetMeshes.h
--- a/Engine/Plugins/Runtime/HairStrands/Source/HairStrandsCore/Public/GroomAssetMeshes.h
+++ b/Engine/Plugins/Runtime/HairStrands/Source/HairStrandsCore/Public/GroomAssetMeshes.h
@@ -22,13 +22,13 @@ struct HAIRSTRANDSCORE_API FHairGroupsMeshesSourceDescription
 
 	/* Deprecated */
 	UPROPERTY()
-	UMaterialInterface* Material = nullptr;
+	TObjectPtr<UMaterialInterface> Material = nullptr;
 
 	UPROPERTY()
 	FName MaterialSlotName;
 
 	UPROPERTY(EditAnywhere, Category = "MeshSettings", meta = (ToolTip = "Mesh settings"))
-	class UStaticMesh* ImportedMesh;
+	TObjectPtr<class UStaticMesh> ImportedMesh;
 
 	UPROPERTY(EditAnywhere, Category = "MeshesSource")
 	FHairGroupCardsTextures Textures;
diff --git a/Engine/Plugins/Runtime/HairStrands/Source/HairStrandsCore/Public/GroomAssetRendering.h b/Engine/Plugins/Runtime/HairStrands/Source/HairStrandsCore/Public/GroomAssetRendering.h
--- a/Engine/Plugins/Runtime/HairStrands/Source/HairStrandsCore/Public/GroomAssetRendering.h
+++ b/Engine/Plugins/Runtime/HairStrands/Source/HairStrandsCore/Public/GroomAssetRendering.h
@@ -89,7 +89,7 @@ struct HAIRSTRANDSCORE_API FHairGroupsRendering
 
 	/* Deprecated */
 	UPROPERTY()
-	UMaterialInterface* Material = nullptr;
+	TObjectPtr<UMaterialInterface> Material = nullptr;
 
 	UPROPERTY(EditAnywhere, Category = "GeometrySettings", meta = (ToolTip = "Geometry settings"))
 	FHairGeometrySettings GeometrySettings;
diff --git a/Engine/Plugins/Runtime/HairStrands/Source/HairStrandsCore/Public/GroomBindingAsset.h b/Engine/Plugins/Runtime/HairStrands/Source/HairStrandsCore/Public/GroomBindingAsset.h
--- a/Engine/Plugins/Runtime/HairStrands/Source/HairStrandsCore/Public/GroomBindingAsset.h
+++ b/Engine/Plugins/Runtime/HairStrands/Source/HairStrandsCore/Public/GroomBindingAsset.h
@@ -83,22 +83,22 @@ public:
 
 	/** Groom to bind. */
 	UPROPERTY(EditAnywhere, BlueprintReadWrite, Category = "BuildSettings")
-	UGroomAsset* Groom;
+	TObjectPtr<UGroomAsset> Groom;
 
 	/** Skeletal mesh on which the groom has been authored. This is optional, and used only if the hair
 		binding is done a different mesh than the one which it has been authored */
 	UPROPERTY(EditAnywhere, BlueprintReadWrite, Category = "BuildSettings")
-	USkeletalMesh* SourceSkeletalMesh;
+	TObjectPtr<USkeletalMesh> SourceSkeletalMesh;
 
 	/** Skeletal mesh on which the groom is attached to. */
 	UPROPERTY(EditAnywhere, BlueprintReadWrite, Category = "BuildSettings")
-	USkeletalMesh* TargetSkeletalMesh;
+	TObjectPtr<USkeletalMesh> TargetSkeletalMesh;
 
 	UPROPERTY(VisibleAnywhere, BlueprintReadWrite, Category = "BuildSettings")
-	UGeometryCache* SourceGeometryCache;
+	TObjectPtr<UGeometryCache> SourceGeometryCache;
 
 	UPROPERTY(VisibleAnywhere, BlueprintReadWrite, Category = "BuildSettings")
-	UGeometryCache* TargetGeometryCache;
+	TObjectPtr<UGeometryCache> TargetGeometryCache;
 
 	/** Number of points used for the rbf interpolation */
 	UPROPERTY(VisibleAnywhere, BlueprintReadWrite, Category = "BuildSettings")
diff --git a/Engine/Plugins/Runtime/HairStrands/Source/HairStrandsCore/Public/GroomComponent.h b/Engine/Plugins/Runtime/HairStrands/Source/HairStrandsCore/Public/GroomComponent.h
--- a/Engine/Plugins/Runtime/HairStrands/Source/HairStrandsCore/Public/GroomComponent.h
+++ b/Engine/Plugins/Runtime/HairStrands/Source/HairStrandsCore/Public/GroomComponent.h
@@ -26,26 +26,26 @@ public:
 
 	/** Groom asset . */
 	UPROPERTY(EditAnywhere, BlueprintReadWrite, interp, Category = "Groom")
-	UGroomAsset* GroomAsset;
+	TObjectPtr<UGroomAsset> GroomAsset;
 
 	UPROPERTY(EditAnywhere, BlueprintReadWrite, interp, Category = "GroomCache", meta = (EditCondition = "BindingAsset == nullptr"))
-	UGroomCache* GroomCache;
+	TObjectPtr<UGroomCache> GroomCache;
 
 	/** Niagara components that will be attached to the system*/
 	UPROPERTY(Transient)
-	TArray<class UNiagaraComponent*> NiagaraComponents;
+	TArray<TObjectPtr<class UNiagaraComponent>> NiagaraComponents;
 
 	// Kept for debugging mesh transfer
 	UPROPERTY()
-	class USkeletalMesh* SourceSkeletalMesh;
+	TObjectPtr<class USkeletalMesh> SourceSkeletalMesh;
 
 	/** Optional binding asset for binding a groom onto a skeletal mesh. If the binding asset is not provided the projection is done at runtime, which implies a large GPU cost at startup time. */
 	UPROPERTY(EditAnywhere, BlueprintReadOnly, interp, Category = "Groom", meta = (EditCondition = "GroomCache == nullptr"))
-	class UGroomBindingAsset* BindingAsset;
+	TObjectPtr<class UGroomBindingAsset> BindingAsset;
 
 	/** Physics asset to be used for hair simulation */
 	UPROPERTY(EditAnywhere, BlueprintReadWrite, Category = "Simulation")
-	class UPhysicsAsset* PhysicsAsset;
+	TObjectPtr<class UPhysicsAsset> PhysicsAsset;
 
 	/** Groom's simulation settings */
 	UPROPERTY(EditAnywhere, BlueprintReadWrite, interp, Category = "Simulation")
@@ -53,19 +53,19 @@ public:
 
 	/* Reference of the default/debug materials for each geometric representation */
 	UPROPERTY()
-	UMaterialInterface* Strands_DebugMaterial;
+	TObjectPtr<UMaterialInterface> Strands_DebugMaterial;
 	UPROPERTY()
-	UMaterialInterface* Strands_DefaultMaterial;
+	TObjectPtr<UMaterialInterface> Strands_DefaultMaterial;
 	UPROPERTY()
-	UMaterialInterface* Cards_DefaultMaterial;
+	TObjectPtr<UMaterialInterface> Cards_DefaultMaterial;
 	UPROPERTY()
-	UMaterialInterface* Meshes_DefaultMaterial;
+	TObjectPtr<UMaterialInterface> Meshes_DefaultMaterial;
 
 	UPROPERTY()
-	class UNiagaraSystem* AngularSpringsSystem;
+	TObjectPtr<class UNiagaraSystem> AngularSpringsSystem;
 
 	UPROPERTY()
-	class UNiagaraSystem* CosseratRodsSystem;
+	TObjectPtr<class UNiagaraSystem> CosseratRodsSystem;
 
 	/** Optional socket name, where the groom component should be attached at, when parented with a skeletal mesh */
 	UPROPERTY(EditAnywhere, BlueprintReadWrite, AdvancedDisplay, interp, Category = "Groom")
@@ -248,10 +248,10 @@ private:
 
 #if WITH_EDITORONLY_DATA
 	UPROPERTY(Transient)
-	UGroomAsset* GroomAssetBeingLoaded;
+	TObjectPtr<UGroomAsset> GroomAssetBeingLoaded;
 
 	UPROPERTY(Transient)
-	UGroomBindingAsset* BindingAssetBeingLoaded;
+	TObjectPtr<UGroomBindingAsset> BindingAssetBeingLoaded;
 #endif
 
 protected:
diff --git a/Engine/Plugins/Runtime/HairStrands/Source/HairStrandsCore/Public/GroomCreateBindingOptions.h b/Engine/Plugins/Runtime/HairStrands/Source/HairStrandsCore/Public/GroomCreateBindingOptions.h
--- a/Engine/Plugins/Runtime/HairStrands/Source/HairStrandsCore/Public/GroomCreateBindingOptions.h
+++ b/Engine/Plugins/Runtime/HairStrands/Source/HairStrandsCore/Public/GroomCreateBindingOptions.h
@@ -20,19 +20,19 @@ public:
 		binding is done a different mesh than the one which it has been authored, i.e., only if the curves 
 		roots and the surface geometry don't aligned and need to be wrapped/transformed. */
 	UPROPERTY(EditAnywhere, BlueprintReadWrite, AdvancedDisplay, Category = Conversion)
-	USkeletalMesh* SourceSkeletalMesh;
+	TObjectPtr<USkeletalMesh> SourceSkeletalMesh;
 
 	/** Skeletal mesh on which the groom is attached to. */
 	UPROPERTY(EditAnywhere, BlueprintReadWrite, Category = BuildSettings)
-	USkeletalMesh* TargetSkeletalMesh;
+	TObjectPtr<USkeletalMesh> TargetSkeletalMesh;
 
 	/** GeometryCache on which the groom has been authored */
 	UPROPERTY(EditAnywhere, BlueprintReadWrite, AdvancedDisplay, Category = Conversion)
-	UGeometryCache* SourceGeometryCache;
+	TObjectPtr<UGeometryCache> SourceGeometryCache;
 
 	/** GeometryCache on which the groom is attached to. */
 	UPROPERTY(EditAnywhere, BlueprintReadWrite, Category = BuildSettings)
-	UGeometryCache* TargetGeometryCache;
+	TObjectPtr<UGeometryCache> TargetGeometryCache;
 
 	/** Number of points used for the rbf interpolation */
 	UPROPERTY(EditAnywhere, BlueprintReadWrite, meta = (ShowOnlyInnerProperties, ClampMin = "0", ClampMax = "100"), Category = HairInterpolation)
diff --git a/Engine/Plugins/Runtime/HairStrands/Source/HairStrandsCore/Public/GroomCreateFollicleMaskOptions.h b/Engine/Plugins/Runtime/HairStrands/Source/HairStrandsCore/Public/GroomCreateFollicleMaskOptions.h
--- a/Engine/Plugins/Runtime/HairStrands/Source/HairStrandsCore/Public/GroomCreateFollicleMaskOptions.h
+++ b/Engine/Plugins/Runtime/HairStrands/Source/HairStrandsCore/Public/GroomCreateFollicleMaskOptions.h
@@ -26,7 +26,7 @@ struct HAIRSTRANDSCORE_API FFollicleMaskOptions
 
 	/** Groom asset */
 	UPROPERTY(VisibleAnywhere, BlueprintReadWrite, meta = (ShowOnlyInnerProperties), Category = Groom)
-	UGroomAsset* Groom = nullptr;
+	TObjectPtr<UGroomAsset> Groom = nullptr;
 
 	/** Texture channel in which the groom's roots mask will be writtent to. */
 	UPROPERTY(EditAnywhere, BlueprintReadWrite, meta = (ShowOnlyInnerProperties), Category = Groom)
diff --git a/Engine/Plugins/Runtime/HairStrands/Source/HairStrandsCore/Public/GroomCreateStrandsTexturesOptions.h b/Engine/Plugins/Runtime/HairStrands/Source/HairStrandsCore/Public/GroomCreateStrandsTexturesOptions.h
--- a/Engine/Plugins/Runtime/HairStrands/Source/HairStrandsCore/Public/GroomCreateStrandsTexturesOptions.h
+++ b/Engine/Plugins/Runtime/HairStrands/Source/HairStrandsCore/Public/GroomCreateStrandsTexturesOptions.h
@@ -51,11 +51,11 @@ public:
 
 	/** Mesh on which the groom strands will be projected on. If non empty and if the skeletal mesh entry is empty, the static mesh will be used for generating the textures. */
 	UPROPERTY(EditAnywhere, BlueprintReadWrite, meta = (ShowOnlyInnerProperties), Category = Options)
-	UStaticMesh* StaticMesh;
+	TObjectPtr<UStaticMesh> StaticMesh;
 
 	/** Mesh on which the groom strands will be projected on. If non empty, the skeletal mesh will be used for generating the textures.*/
 	UPROPERTY(EditAnywhere, BlueprintReadWrite, meta = (ShowOnlyInnerProperties), Category = Options)
-	USkeletalMesh* SkeletalMesh;
+	TObjectPtr<USkeletalMesh> SkeletalMesh;
 
 	/** LOD of the mesh, on which the texture projection is done */
 	UPROPERTY(EditAnywhere, BlueprintReadWrite, meta = (ShowOnlyInnerProperties), Category = Options)
diff --git a/Engine/Plugins/Runtime/HairStrands/Source/HairStrandsCore/Public/MovieSceneGroomCacheSection.h b/Engine/Plugins/Runtime/HairStrands/Source/HairStrandsCore/Public/MovieSceneGroomCacheSection.h
--- a/Engine/Plugins/Runtime/HairStrands/Source/HairStrandsCore/Public/MovieSceneGroomCacheSection.h
+++ b/Engine/Plugins/Runtime/HairStrands/Source/HairStrandsCore/Public/MovieSceneGroomCacheSection.h
@@ -16,7 +16,7 @@ USTRUCT()
 	 float GetSequenceLength() const;
 	/** The animation this section plays */
 	UPROPERTY(EditAnywhere, Category = GroomCache)
-	class UGroomCache* GroomCache;
+	TObjectPtr<class UGroomCache> GroomCache;
 
 	/** The offset for the first loop of the animation clip */
 	UPROPERTY(EditAnywhere, Category = GroomCache)
diff --git a/Engine/Plugins/Runtime/HairStrands/Source/HairStrandsCore/Public/Niagara/NiagaraDataInterfaceHairStrands.h b/Engine/Plugins/Runtime/HairStrands/Source/HairStrandsCore/Public/Niagara/NiagaraDataInterfaceHairStrands.h
--- a/Engine/Plugins/Runtime/HairStrands/Source/HairStrandsCore/Public/Niagara/NiagaraDataInterfaceHairStrands.h
+++ b/Engine/Plugins/Runtime/HairStrands/Source/HairStrandsCore/Public/Niagara/NiagaraDataInterfaceHairStrands.h
@@ -329,11 +329,11 @@ public:
 
 	/** Hair Strands Asset used to sample from when not overridden by a source actor from the scene. Also useful for previewing in the editor. */
 	UPROPERTY(EditAnywhere, Category = "Source")
-	UGroomAsset* DefaultSource;
+	TObjectPtr<UGroomAsset> DefaultSource;
 
 	/** The source actor from which to sample */
 	UPROPERTY(EditAnywhere, Category = "Source")
-	AActor* SourceActor;
+	TObjectPtr<AActor> SourceActor;
 
 	/** The source component from which to sample */
 	TWeakObjectPtr<class UGroomComponent> SourceComponent;
diff --git a/Engine/Plugins/Runtime/HairStrands/Source/HairStrandsCore/Public/Niagara/NiagaraDataInterfacePhysicsAsset.h b/Engine/Plugins/Runtime/HairStrands/Source/HairStrandsCore/Public/Niagara/NiagaraDataInterfacePhysicsAsset.h
--- a/Engine/Plugins/Runtime/HairStrands/Source/HairStrandsCore/Public/Niagara/NiagaraDataInterfacePhysicsAsset.h
+++ b/Engine/Plugins/Runtime/HairStrands/Source/HairStrandsCore/Public/Niagara/NiagaraDataInterfacePhysicsAsset.h
@@ -108,11 +108,11 @@ public:
 
 	/** Skeletal Mesh from which the Physics Asset will be found. */
 	UPROPERTY(EditAnywhere, Category = "Source")
-		UPhysicsAsset* DefaultSource;
+		TObjectPtr<UPhysicsAsset> DefaultSource;
 
 	/** The source actor from which to sample */
 	UPROPERTY(EditAnywhere, Category = "Source")
-		AActor* SourceActor;
+		TObjectPtr<AActor> SourceActor;
 
 	/** The source component from which to sample */
 	TArray<TWeakObjectPtr<class USkeletalMeshComponent>> SourceComponents;
diff --git a/Engine/Plugins/Runtime/Interchange/Source/Pipelines/Public/InterchangePipelineMeshesUtilities.h b/Engine/Plugins/Runtime/Interchange/Source/Pipelines/Public/InterchangePipelineMeshesUtilities.h
--- a/Engine/Plugins/Runtime/Interchange/Source/Pipelines/Public/InterchangePipelineMeshesUtilities.h
+++ b/Engine/Plugins/Runtime/Interchange/Source/Pipelines/Public/InterchangePipelineMeshesUtilities.h
@@ -24,7 +24,7 @@ struct FInterchangeLodSceneNodeContainer
 	* Each scene node here represent a mesh scene node. Only if we represent a lod group we can have more then 1 mesh scene node for a specific lod index.
 	*/
 	UPROPERTY(EditAnywhere, Category = "Interchange | Pipeline | MeshesInstance")
-	TArray<UInterchangeSceneNode*> SceneNodes;
+	TArray<TObjectPtr<UInterchangeSceneNode>> SceneNodes;
 };
 
 /*
@@ -53,7 +53,7 @@ struct FInterchangeMeshInstance
 	* If this mesh instance represent a LodGroup this member will not be null, but will be null if the mesh instance do not represent a lod group
 	*/
 	UPROPERTY(EditAnywhere, Category = "Interchange | Pipeline | MeshesInstance")
-	UInterchangeSceneNode* LodGroupNode;
+	TObjectPtr<UInterchangeSceneNode> LodGroupNode;
 
 	UPROPERTY(EditAnywhere, Category = "Interchange | Pipeline | MeshesInstance")
 	bool bReferenceSkinnedMesh;
@@ -97,7 +97,7 @@ struct FInterchangeMeshGeometry
 	* The UInterchangeMeshNode pointer represent by this structure.
 	*/
 	UPROPERTY(EditAnywhere, Category = "Interchange | Pipeline | MeshesGeometry")
-	UInterchangeMeshNode* MeshNode = nullptr;
+	TObjectPtr<UInterchangeMeshNode> MeshNode = nullptr;
 
 	/**
 	* All mesh instance referencing this UInterchangeMeshNode pointer.
diff --git a/Engine/Plugins/VirtualProduction/Takes/Source/TakeRecorder/Public/Recorder/TakeRecorder.h b/Engine/Plugins/VirtualProduction/Takes/Source/TakeRecorder/Public/Recorder/TakeRecorder.h
--- a/Engine/Plugins/VirtualProduction/Takes/Source/TakeRecorder/Public/Recorder/TakeRecorder.h
+++ b/Engine/Plugins/VirtualProduction/Takes/Source/TakeRecorder/Public/Recorder/TakeRecorder.h
@@ -225,11 +225,11 @@ private:
 
 	/** The asset that we should output recorded data into */
 	UPROPERTY(transient)
-	ULevelSequence* SequenceAsset;
+	TObjectPtr<ULevelSequence> SequenceAsset;
 
 	/** The overlay widget for this recording */
 	UPROPERTY(transient)
-	UTakeRecorderOverlayWidget* OverlayWidget;
+	TObjectPtr<UTakeRecorderOverlayWidget> OverlayWidget;
 
 	/** The world that we are recording within */
 	UPROPERTY(transient)
diff --git a/Engine/Plugins/VirtualProduction/Takes/Source/TakeRecorder/Public/TakeRecorderOverlayWidget.h b/Engine/Plugins/VirtualProduction/Takes/Source/TakeRecorder/Public/TakeRecorderOverlayWidget.h
--- a/Engine/Plugins/VirtualProduction/Takes/Source/TakeRecorder/Public/TakeRecorderOverlayWidget.h
+++ b/Engine/Plugins/VirtualProduction/Takes/Source/TakeRecorder/Public/TakeRecorderOverlayWidget.h
@@ -28,5 +28,5 @@ protected:
 
 	/** The recorder that this overlay is reflecting */
 	UPROPERTY(BlueprintReadOnly, Category="Take Recorder")
-	UTakeRecorder* Recorder;
+	TObjectPtr<UTakeRecorder> Recorder;
 };
\ No newline at end of file
diff --git a/Engine/Plugins/VirtualProduction/Takes/Source/TakeRecorderSources/Private/TakeRecorderActorSource.h b/Engine/Plugins/VirtualProduction/Takes/Source/TakeRecorderSources/Private/TakeRecorderActorSource.h
--- a/Engine/Plugins/VirtualProduction/Takes/Source/TakeRecorderSources/Private/TakeRecorderActorSource.h
+++ b/Engine/Plugins/VirtualProduction/Takes/Source/TakeRecorderSources/Private/TakeRecorderActorSource.h
@@ -66,15 +66,15 @@ public:
 	 * recorded into a track in the resulting Level Sequence. 
 	 */
 	UPROPERTY(EditAnywhere, Instanced, BlueprintReadWrite, meta=(ShowInnerProperties), Category = "Actor Source")
-	UActorRecorderPropertyMap* RecordedProperties;
+	TObjectPtr<UActorRecorderPropertyMap> RecordedProperties;
 
 	/** The level sequence that this source is being recorded into. Set during PreRecording, null after PostRecording. */
 	UPROPERTY()
-	ULevelSequence* TargetLevelSequence;
+	TObjectPtr<ULevelSequence> TargetLevelSequence;
 
 	/** The master or uppermost level sequence that this source is being recorded into. Set during PreRecording, null after PostRecording. */
 	UPROPERTY()
-	ULevelSequence* MasterLevelSequence;
+	TObjectPtr<ULevelSequence> MasterLevelSequence;
 
 	/**
 	* Dynamically created list of settings objects for the different factories that are recording something 
@@ -83,14 +83,14 @@ public:
 	* instance.
 	*/
 	UPROPERTY()
-	TArray<UObject*> FactorySettings;
+	TArray<TObjectPtr<UObject>> FactorySettings;
 	
 	/**
 	* An array of section recorders created during the recording process that are capturing data about the actor/components.
 	* Will be an empty list when a recording is not in progress.
 	*/
 	UPROPERTY()
-	TArray<class UMovieSceneTrackRecorder*> TrackRecorders;
+	TArray<TObjectPtr<class UMovieSceneTrackRecorder>> TrackRecorders;
 
 
 public:
diff --git a/Engine/Plugins/VirtualProduction/Takes/Source/TakeRecorderSources/Private/TakeRecorderCameraCutSource.h b/Engine/Plugins/VirtualProduction/Takes/Source/TakeRecorderSources/Private/TakeRecorderCameraCutSource.h
--- a/Engine/Plugins/VirtualProduction/Takes/Source/TakeRecorderSources/Private/TakeRecorderCameraCutSource.h
+++ b/Engine/Plugins/VirtualProduction/Takes/Source/TakeRecorderSources/Private/TakeRecorderCameraCutSource.h
@@ -36,11 +36,11 @@ private:
 private:
 
 	UPROPERTY()
-	UWorld* World;
+	TObjectPtr<UWorld> World;
 
 	/** The master or uppermost level sequence that this source is being recorded into. Set during PreRecording, null after PostRecording. */
 	UPROPERTY()
-	ULevelSequence* MasterLevelSequence;
+	TObjectPtr<ULevelSequence> MasterLevelSequence;
 
 
 	struct FCameraCutData
diff --git a/Engine/Plugins/VirtualProduction/Takes/Source/TakeRecorderSources/Private/TakeRecorderLevelSequenceSource.h b/Engine/Plugins/VirtualProduction/Takes/Source/TakeRecorderSources/Private/TakeRecorderLevelSequenceSource.h
--- a/Engine/Plugins/VirtualProduction/Takes/Source/TakeRecorderSources/Private/TakeRecorderLevelSequenceSource.h
+++ b/Engine/Plugins/VirtualProduction/Takes/Source/TakeRecorderSources/Private/TakeRecorderLevelSequenceSource.h
@@ -20,7 +20,7 @@ public:
 	UTakeRecorderLevelSequenceSource(const FObjectInitializer& ObjInit);
 
 	UPROPERTY(EditAnywhere, BlueprintReadWrite, Category = "Source")
-	TArray<ULevelSequence*> LevelSequencesToTrigger;
+	TArray<TObjectPtr<ULevelSequence>> LevelSequencesToTrigger;
 
 private:
 
diff --git a/Engine/Plugins/VirtualProduction/Takes/Source/TakesCore/Public/TakePreset.h b/Engine/Plugins/VirtualProduction/Takes/Source/TakesCore/Public/TakePreset.h
--- a/Engine/Plugins/VirtualProduction/Takes/Source/TakesCore/Public/TakePreset.h
+++ b/Engine/Plugins/VirtualProduction/Takes/Source/TakesCore/Public/TakePreset.h
@@ -67,7 +67,7 @@ private:
 
 	/** Instanced level sequence template that is used to define a starting point for a new take recording */
 	UPROPERTY(Instanced)
-	ULevelSequence* LevelSequence;
+	TObjectPtr<ULevelSequence> LevelSequence;
 
 	FSimpleMulticastDelegate OnLevelSequenceChangedEvent;
 };
diff --git a/Engine/Plugins/VirtualProduction/Takes/Source/TakesCore/Public/TakeRecorderSourceProperty.h b/Engine/Plugins/VirtualProduction/Takes/Source/TakesCore/Public/TakeRecorderSourceProperty.h
--- a/Engine/Plugins/VirtualProduction/Takes/Source/TakesCore/Public/TakeRecorderSourceProperty.h
+++ b/Engine/Plugins/VirtualProduction/Takes/Source/TakesCore/Public/TakeRecorderSourceProperty.h
@@ -58,7 +58,7 @@ public:
 	TArray<FActorRecordedProperty> Properties;
 
 	UPROPERTY(EditAnywhere, Instanced, BlueprintReadWrite, meta=(ShowInnerProperties, EditFixedOrder), Category = "Property")
-	TArray<UActorRecorderPropertyMap*> Children;
+	TArray<TObjectPtr<UActorRecorderPropertyMap>> Children;
 
 public:
 	struct Cache
diff --git a/Engine/Plugins/VirtualProduction/Takes/Source/TakesCore/Public/TakeRecorderSources.h b/Engine/Plugins/VirtualProduction/Takes/Source/TakesCore/Public/TakeRecorderSources.h
--- a/Engine/Plugins/VirtualProduction/Takes/Source/TakesCore/Public/TakeRecorderSources.h
+++ b/Engine/Plugins/VirtualProduction/Takes/Source/TakesCore/Public/TakeRecorderSources.h
@@ -208,15 +208,15 @@ private:
 
 	/** The array of all sources contained within this list */
 	UPROPERTY(Instanced)
-	TArray<UTakeRecorderSource*> Sources;
+	TArray<TObjectPtr<UTakeRecorderSource>> Sources;
 
 	/** Maps each source to the level sequence that was created for that source, or to the master source if a subsequence was not created. */
 	UPROPERTY(Transient)
-	TMap<UTakeRecorderSource*, ULevelSequence*> SourceSubSequenceMap;
+	TMap<TObjectPtr<UTakeRecorderSource>, TObjectPtr<ULevelSequence>> SourceSubSequenceMap;
 
 	/** List of sub-sections that we're recording into. Needed to ensure they're all the right size at the end without re-adjusting every sub-section in a sequence. */
 	UPROPERTY(Transient)
-	TArray<class UMovieSceneSubSection*> ActiveSubSections;
+	TArray<TObjectPtr<class UMovieSceneSubSection>> ActiveSubSections;
 
 	/** Are we currently in a recording pass and should be ticking our Sources? */
 	bool bIsRecording;
diff --git a/Engine/Source/Developer/RigVMDeveloper/Public/RigVMCompiler/RigVMAST.h b/Engine/Source/Developer/RigVMDeveloper/Public/RigVMCompiler/RigVMAST.h
--- a/Engine/Source/Developer/RigVMDeveloper/Public/RigVMCompiler/RigVMAST.h
+++ b/Engine/Source/Developer/RigVMDeveloper/Public/RigVMCompiler/RigVMAST.h
@@ -1115,7 +1115,7 @@ struct RIGVMDEVELOPER_API FRigVMParserASTSettings
 
 	// links to be ignored during the parse
 	UPROPERTY()
-	TArray<URigVMLink*> LinksToSkip;
+	TArray<TObjectPtr<URigVMLink>> LinksToSkip;
 
 	FRigVMReportDelegate ReportDelegate;
 
diff --git a/Engine/Source/Developer/RigVMDeveloper/Public/RigVMModel/Nodes/RigVMCollapseNode.h b/Engine/Source/Developer/RigVMDeveloper/Public/RigVMModel/Nodes/RigVMCollapseNode.h
--- a/Engine/Source/Developer/RigVMDeveloper/Public/RigVMModel/Nodes/RigVMCollapseNode.h
+++ b/Engine/Source/Developer/RigVMDeveloper/Public/RigVMModel/Nodes/RigVMCollapseNode.h
@@ -34,7 +34,7 @@ public:
 private:
 
 	UPROPERTY()
-	URigVMGraph* ContainedGraph;
+	TObjectPtr<URigVMGraph> ContainedGraph;
 
 	UPROPERTY()
 	FString NodeCategory;
diff --git a/Engine/Source/Developer/RigVMDeveloper/Public/RigVMModel/Nodes/RigVMParameterNode.h b/Engine/Source/Developer/RigVMDeveloper/Public/RigVMModel/Nodes/RigVMParameterNode.h
--- a/Engine/Source/Developer/RigVMDeveloper/Public/RigVMModel/Nodes/RigVMParameterNode.h
+++ b/Engine/Source/Developer/RigVMDeveloper/Public/RigVMModel/Nodes/RigVMParameterNode.h
@@ -39,7 +39,7 @@ public:
 
 	// The Struct of the C++ data type of the parameter (or nullptr)
 	UPROPERTY(BlueprintReadOnly, Category = RigVMGraphParameterDescription)
-	UObject* CPPTypeObject = nullptr;
+	TObjectPtr<UObject> CPPTypeObject = nullptr;
 
 	// The default value of the parameter
 	UPROPERTY(BlueprintReadOnly, Category = RigVMGraphParameterDescription)
diff --git a/Engine/Source/Developer/RigVMDeveloper/Public/RigVMModel/Nodes/RigVMUnitNode.h b/Engine/Source/Developer/RigVMDeveloper/Public/RigVMModel/Nodes/RigVMUnitNode.h
--- a/Engine/Source/Developer/RigVMDeveloper/Public/RigVMModel/Nodes/RigVMUnitNode.h
+++ b/Engine/Source/Developer/RigVMDeveloper/Public/RigVMModel/Nodes/RigVMUnitNode.h
@@ -56,7 +56,7 @@ protected:
 private:
 
 	UPROPERTY()
-	UScriptStruct* ScriptStruct;
+	TObjectPtr<UScriptStruct> ScriptStruct;
 
 	UPROPERTY()
 	FName MethodName;
diff --git a/Engine/Source/Developer/RigVMDeveloper/Public/RigVMModel/RigVMController.h b/Engine/Source/Developer/RigVMDeveloper/Public/RigVMModel/RigVMController.h
--- a/Engine/Source/Developer/RigVMDeveloper/Public/RigVMModel/RigVMController.h
+++ b/Engine/Source/Developer/RigVMDeveloper/Public/RigVMModel/RigVMController.h
@@ -828,13 +828,13 @@ private:
 	static void SanitizeName(FString& InOutName);
 	
 	UPROPERTY(transient)
-	TArray<URigVMGraph*> Graphs;
+	TArray<TObjectPtr<URigVMGraph>> Graphs;
 
 	UPROPERTY()
-	UStruct* ExecuteContextStruct;
+	TObjectPtr<UStruct> ExecuteContextStruct;
 
 	UPROPERTY(transient)
-	URigVMActionStack* ActionStack;
+	TObjectPtr<URigVMActionStack> ActionStack;
 
 	bool bSuspendNotifications;
 	bool bReportWarningsAndErrors;
diff --git a/Engine/Source/Developer/RigVMDeveloper/Public/RigVMModel/RigVMControllerActions.h b/Engine/Source/Developer/RigVMDeveloper/Public/RigVMModel/RigVMControllerActions.h
--- a/Engine/Source/Developer/RigVMDeveloper/Public/RigVMModel/RigVMControllerActions.h
+++ b/Engine/Source/Developer/RigVMDeveloper/Public/RigVMModel/RigVMControllerActions.h
@@ -1583,5 +1583,5 @@ public:
 	FString CPPType;
 
 	UPROPERTY()
-	UObject* CPPTypeObject;
+	TObjectPtr<UObject> CPPTypeObject;
 };
\ No newline at end of file
