Wrapped FVectorVMContext in NiagaraDataInterfaceFunctionContext for external functions to facilitate running an experimental VM in parallel.

diff --git a/Engine/Plugins/FX/Niagara/Source/Niagara/Private/NiagaraDataInterfaceVectorCurve.cpp b/Engine/Plugins/FX/Niagara/Source/Niagara/Private/NiagaraDataInterfaceVectorCurve.cpp
--- a/Engine/Plugins/FX/Niagara/Source/Niagara/Private/NiagaraDataInterfaceVectorCurve.cpp
+++ b/Engine/Plugins/FX/Niagara/Source/Niagara/Private/NiagaraDataInterfaceVectorCurve.cpp
@@ -226,7 +226,7 @@ FORCEINLINE_DEBUGGABLE FVector UNiagaraDataInterfaceVectorCurve::SampleCurveInte
 }
 
 template<typename UseLUT>
-void UNiagaraDataInterfaceVectorCurve::SampleCurve(FVectorVMContext& Context)
+void UNiagaraDataInterfaceVectorCurve::SampleCurve(FVectorVMExternalFunctionContext& Context)
 {
 	//TODO: Create some SIMDable optimized representation of the curve to do this faster.
 	VectorVM::FExternalFuncInputHandler<float> XParam(Context);
@@ -234,7 +234,7 @@ void UNiagaraDataInterfaceVectorCurve::SampleCurve(FVectorVMContext& Context)
 	VectorVM::FExternalFuncRegisterHandler<float> OutSampleY(Context);
 	VectorVM::FExternalFuncRegisterHandler<float> OutSampleZ(Context);
 
-	for (int32 i = 0; i < Context.NumInstances; ++i)
+	for (int32 i = 0; i < Context.GetNumInstances(); ++i)
 	{
 		float X = XParam.GetAndAdvance();
 		FVector V = SampleCurveInternal<UseLUT>(X);
diff --git a/Engine/Plugins/FX/Niagara/Source/Niagara/Private/NiagaraDataInterfaceVectorField.cpp b/Engine/Plugins/FX/Niagara/Source/Niagara/Private/NiagaraDataInterfaceVectorField.cpp
--- a/Engine/Plugins/FX/Niagara/Source/Niagara/Private/NiagaraDataInterfaceVectorField.cpp
+++ b/Engine/Plugins/FX/Niagara/Source/Niagara/Private/NiagaraDataInterfaceVectorField.cpp
@@ -492,14 +492,14 @@ IMPLEMENT_NIAGARA_DI_PARAMETER(UNiagaraDataInterfaceVectorField, FNiagaraDataInt
 
 /*--------------------------------------------------------------------------------------------------------------------------*/
 
-void UNiagaraDataInterfaceVectorField::GetFieldTilingAxes(FVectorVMContext& Context)
+void UNiagaraDataInterfaceVectorField::GetFieldTilingAxes(FVectorVMExternalFunctionContext& Context)
 {
 	VectorVM::FExternalFuncRegisterHandler<float> OutSizeX(Context);
 	VectorVM::FExternalFuncRegisterHandler<float> OutSizeY(Context);
 	VectorVM::FExternalFuncRegisterHandler<float> OutSizeZ(Context);
 
 	FVector Tilings = GetTilingAxes();
-	for (int32 i = 0; i < Context.NumInstances; ++i)
+	for (int32 i = 0; i < Context.GetNumInstances(); ++i)
 	{
 		*OutSizeX.GetDest() = Tilings.X;
 		*OutSizeY.GetDest() = Tilings.Y;
@@ -511,14 +511,14 @@ void UNiagaraDataInterfaceVectorField::GetFieldTilingAxes(FVectorVMContext& Cont
 	}
 }
 
-void UNiagaraDataInterfaceVectorField::GetFieldDimensions(FVectorVMContext& Context)
+void UNiagaraDataInterfaceVectorField::GetFieldDimensions(FVectorVMExternalFunctionContext& Context)
 {
 	VectorVM::FExternalFuncRegisterHandler<float> OutSizeX(Context);
 	VectorVM::FExternalFuncRegisterHandler<float> OutSizeY(Context);
 	VectorVM::FExternalFuncRegisterHandler<float> OutSizeZ(Context);
 
 	FVector Dim = GetDimensions();
-	for (int32 i = 0; i < Context.NumInstances; ++i)
+	for (int32 i = 0; i < Context.GetNumInstances(); ++i)
 	{
 		*OutSizeX.GetDest() = Dim.X;
 		*OutSizeY.GetDest() = Dim.Y;
@@ -530,7 +530,7 @@ void UNiagaraDataInterfaceVectorField::GetFieldDimensions(FVectorVMContext& Cont
 	}
 }
 
-void UNiagaraDataInterfaceVectorField::GetFieldBounds(FVectorVMContext& Context)
+void UNiagaraDataInterfaceVectorField::GetFieldBounds(FVectorVMExternalFunctionContext& Context)
 {
 	VectorVM::FExternalFuncRegisterHandler<float> OutMinX(Context);
 	VectorVM::FExternalFuncRegisterHandler<float> OutMinY(Context);
@@ -541,7 +541,7 @@ void UNiagaraDataInterfaceVectorField::GetFieldBounds(FVectorVMContext& Context)
 
 	FVector MinBounds = GetMinBounds();
 	FVector MaxBounds = GetMaxBounds();
-	for (int32 i = 0; i < Context.NumInstances; ++i)
+	for (int32 i = 0; i < Context.GetNumInstances(); ++i)
 	{
 		*OutMinX.GetDest() = MinBounds.X;
 		*OutMinY.GetDest() = MinBounds.Y;
@@ -559,7 +559,7 @@ void UNiagaraDataInterfaceVectorField::GetFieldBounds(FVectorVMContext& Context)
 	}
 }
 
-void UNiagaraDataInterfaceVectorField::SampleVectorField(FVectorVMContext& Context)
+void UNiagaraDataInterfaceVectorField::SampleVectorField(FVectorVMExternalFunctionContext& Context)
 {
 	// Input arguments...
 	VectorVM::FExternalFuncInputHandler<float> XParam(Context);
@@ -601,12 +601,12 @@ void UNiagaraDataInterfaceVectorField::SampleVectorField(FVectorVMContext& Conte
 					XParam.IsConstant(), YParam.IsConstant(), ZParam.IsConstant(),
 					OutSampleX.GetDest(), OutSampleY.GetDest(), OutSampleZ.GetDest(),
 					(ispc::FHalfVector*) FieldSamples.GetData(), FieldSamples.Num() - sizeof(ispc::FHalfVector), (ispc::FVector&)MinBounds, (ispc::FVector&)OneOverBoundSize,
-					(ispc::FVector&)Size, (ispc::FVector&)TilingAxes, Context.NumInstances);
+					(ispc::FVector&)Size, (ispc::FVector&)TilingAxes, Context.GetNumInstances());
 			}
 			else
 #endif
 			{
-				for (int32 InstanceIdx = 0; InstanceIdx < Context.NumInstances; ++InstanceIdx)
+				for (int32 InstanceIdx = 0; InstanceIdx < Context.GetNumInstances(); ++InstanceIdx)
 				{
 					// Position in Volume Space
 					FVector Pos(XParam.Get(), YParam.Get(), ZParam.Get());
@@ -658,7 +658,7 @@ void UNiagaraDataInterfaceVectorField::SampleVectorField(FVectorVMContext& Conte
 
 		// Set the default vector to positive X axis corresponding to a velocity of 100 cm/s
 		// Rationale: Setting to the zero vector can be visually confusing and likely to cause problems elsewhere
-		for (int32 InstanceIdx = 0; InstanceIdx < Context.NumInstances; ++InstanceIdx)
+		for (int32 InstanceIdx = 0; InstanceIdx < Context.GetNumInstances(); ++InstanceIdx)
 		{
 			*OutSampleX.GetDest() = 0.0f;
 			*OutSampleY.GetDest() = 0.0f;
diff --git a/Engine/Plugins/FX/Niagara/Source/Niagara/Private/NiagaraDataInterfaceVolumeTexture.cpp b/Engine/Plugins/FX/Niagara/Source/Niagara/Private/NiagaraDataInterfaceVolumeTexture.cpp
--- a/Engine/Plugins/FX/Niagara/Source/Niagara/Private/NiagaraDataInterfaceVolumeTexture.cpp
+++ b/Engine/Plugins/FX/Niagara/Source/Niagara/Private/NiagaraDataInterfaceVolumeTexture.cpp
@@ -137,14 +137,14 @@ bool UNiagaraDataInterfaceVolumeTexture::PerInstanceTick(void* PerInstanceData,
 	return false;
 }
 
-void UNiagaraDataInterfaceVolumeTexture::GetTextureDimensions(FVectorVMContext& Context)
+void UNiagaraDataInterfaceVolumeTexture::GetTextureDimensions(FVectorVMExternalFunctionContext& Context)
 {
 	FNDIOutputParam<float> OutWidth(Context);
 	FNDIOutputParam<float> OutHeight(Context);
 	FNDIOutputParam<float> OutDepth(Context);
 
 	FVector FloatTextureSize(TextureSize.X, TextureSize.Y, TextureSize.Z);
-	for (int32 i = 0; i < Context.NumInstances; ++i)
+	for (int32 i = 0; i < Context.GetNumInstances(); ++i)
 	{
 		OutWidth.SetAndAdvance(FloatTextureSize.X);
 		OutHeight.SetAndAdvance(FloatTextureSize.Y);
@@ -152,7 +152,7 @@ void UNiagaraDataInterfaceVolumeTexture::GetTextureDimensions(FVectorVMContext&
 	}
 }
 
-void UNiagaraDataInterfaceVolumeTexture::SampleVolumeTexture(FVectorVMContext& Context)
+void UNiagaraDataInterfaceVolumeTexture::SampleVolumeTexture(FVectorVMExternalFunctionContext& Context)
 {
 	VectorVM::FExternalFuncInputHandler<float> XParam(Context);
 	VectorVM::FExternalFuncInputHandler<float> YParam(Context);
@@ -163,7 +163,7 @@ void UNiagaraDataInterfaceVolumeTexture::SampleVolumeTexture(FVectorVMContext& C
 	VectorVM::FExternalFuncRegisterHandler<float> OutSampleB(Context);
 	VectorVM::FExternalFuncRegisterHandler<float> OutSampleA(Context);
 
-	for (int32 i = 0; i < Context.NumInstances; ++i)
+	for (int32 i = 0; i < Context.GetNumInstances(); ++i)
 	{
 		float X = XParam.GetAndAdvance();
 		float Y = YParam.GetAndAdvance();
diff --git a/Engine/Plugins/FX/Niagara/Source/Niagara/Private/NiagaraFunctionLibrary.cpp b/Engine/Plugins/FX/Niagara/Source/Niagara/Private/NiagaraFunctionLibrary.cpp
--- a/Engine/Plugins/FX/Niagara/Source/Niagara/Private/NiagaraFunctionLibrary.cpp
+++ b/Engine/Plugins/FX/Niagara/Source/Niagara/Private/NiagaraFunctionLibrary.cpp
@@ -660,7 +660,7 @@ struct FVectorKernelFastDot4
 		return FString();
 	}
 
-	static void Exec(FVectorVMContext& Context)
+	static void Exec(FVectorVMExternalFunctionContext& Context)
 	{
 		//SCOPE_CYCLE_COUNTER(STAT_NiagaraFastDot4);
 
@@ -753,7 +753,7 @@ struct FVectorKernelFastTransformPosition
 		return FString();
 	}
 
-	static void Exec(FVectorVMContext& Context)
+	static void Exec(FVectorVMExternalFunctionContext& Context)
 	{
 #if 0
 		TArray<VectorVM::FExternalFuncInputHandler<float>, TInlineAllocator<16>> InMatrix;
@@ -933,7 +933,7 @@ struct FVectorKernelFastMatrixToQuaternion
 		return FunctionHLSL;
 	}
 
-	static void Exec(FVectorVMContext& Context)
+	static void Exec(FVectorVMExternalFunctionContext& Context)
 	{
 		TArray<VectorVM::FExternalFuncInputHandler<float>, TInlineAllocator<16>> InMatrix;
 		for (int i=0; i < 16; i++)
@@ -1064,7 +1064,7 @@ struct FVectorKernel_EmitterLifeCycle
 		return FString();
 	}
 
-	static void Exec(FVectorVMContext& Context)
+	static void Exec(FVectorVMExternalFunctionContext& Context)
 	{
 		VectorVM::FExternalFuncInputHandler<float> InEngineDeltaTime(Context);
 		VectorVM::FExternalFuncInputHandler<int>   InEngineNumParticles(Context);
@@ -1230,7 +1230,7 @@ struct FVectorKernel_SpawnRate
 		return FString();
 	}
 
-	static void Exec(FVectorVMContext& Context)
+	static void Exec(FVectorVMExternalFunctionContext& Context)
 	{
 		VectorVM::FExternalFuncInputHandler<float> InEngineDeltaTime(Context);
 		VectorVM::FExternalFuncInputHandler<float> InModuleSpawnRate(Context);
@@ -1318,7 +1318,7 @@ struct FVectorKernel_SpawnBurstInstantaneous
 		return FString();
 	}
 
-	static void Exec(FVectorVMContext& Context)
+	static void Exec(FVectorVMExternalFunctionContext& Context)
 	{
 		VectorVM::FExternalFuncInputHandler<float> InEngineDeltaTime(Context);
 		VectorVM::FExternalFuncInputHandler<float> InScalabilityEmitterSpawnCountScale(Context);
@@ -1392,7 +1392,7 @@ struct FVectorKernel_SolveVelocitiesAndForces
 		return FString();
 	}
 
-	static void Exec(FVectorVMContext& Context)
+	static void Exec(FVectorVMExternalFunctionContext& Context)
 	{
 		VectorVM::FExternalFuncInputHandler<float> InEngineDeltaTime(Context);
 		VectorVM::FExternalFuncInputHandler<float> InPhysicsForceX(Context);
@@ -1523,7 +1523,7 @@ struct FVectorKernel_SolveVelocitiesAndForces
 	}
 
 	template<bool bForceConstant, bool bDragConstant, bool bMassConstant>
-	FORCEINLINE static void ExecOptimized(FVectorVMContext& Context)
+	FORCEINLINE static void ExecOptimized(FVectorVMExternalFunctionContext& Context)
 	{
 #if 0
 		const VectorRegister4Float MassMin = VectorSetFloat1(0.0001f);
diff --git a/Engine/Plugins/FX/Niagara/Source/Niagara/Private/NiagaraScriptExecutionContext.cpp b/Engine/Plugins/FX/Niagara/Source/Niagara/Private/NiagaraScriptExecutionContext.cpp
--- a/Engine/Plugins/FX/Niagara/Source/Niagara/Private/NiagaraScriptExecutionContext.cpp
+++ b/Engine/Plugins/FX/Niagara/Source/Niagara/Private/NiagaraScriptExecutionContext.cpp
@@ -366,28 +366,28 @@ void FNiagaraScriptExecutionContextBase::PostTick()
 
 //////////////////////////////////////////////////////////////////////////
 
-void FNiagaraSystemScriptExecutionContext::PerInstanceFunctionHook(FVectorVMContext& Context, int32 PerInstFunctionIndex, int32 UserPtrIndex)
+void FNiagaraSystemScriptExecutionContext::PerInstanceFunctionHook(FVectorVMExternalFunctionContext& Context, int32 PerInstFunctionIndex, int32 UserPtrIndex)
 {
 	check(SystemInstances);
 	
 	//This is a bit of a hack. We grab the base offset into the instance data from the primary dataset.
 	//TODO: Find a cleaner way to do this.
-	int32 InstanceOffset = Context.GetDataSetMeta(0).InstanceOffset;
+	int32 InstanceOffset = Context.VectorVMContext->GetDataSetMeta(0).InstanceOffset;
 
 	//Cache context state.
-	int32 CachedContextStartInstance = Context.StartInstance;
-	int32 CachedContextNumInstances = Context.NumInstances;
-	uint8 const* CachedCodeLocation = Context.Code;
+	int32 CachedContextStartInstance = Context.VectorVMContext->GetStartInstance();
+	int32 CachedContextNumInstances = Context.VectorVMContext->GetNumInstances();
+	uint8 const* CachedCodeLocation = Context.VectorVMContext->Code;
 
 	//Hack context so we can run the DI calls one by one.
-	Context.NumInstances = 1;
+	Context.VectorVMContext->NumInstances = 1;
 
 	for (int32 i = 0; i < CachedContextNumInstances; ++i)
 	{
 		//Reset the code each iteration.
-		Context.Code = CachedCodeLocation;
+		Context.VectorVMContext->Code = CachedCodeLocation;
 		//Offset buffer I/O to the correct instance's data.
-		Context.ExternalFunctionInstanceOffset = i;
+		Context.VectorVMContext->ExternalFunctionInstanceOffset = i;
 
 		int32 InstanceIndex = InstanceOffset + CachedContextStartInstance + i;
 		FNiagaraSystemInstance* Instance = (*SystemInstances)[InstanceIndex];
@@ -397,10 +397,10 @@ void FNiagaraSystemScriptExecutionContext::PerInstanceFunctionHook(FVectorVMCont
 		//Do this way for now to reduce overall complexity of the initial change. Doing this needs extensive boiler plate changes to most DI classes and a script recompile.
 		if(UserPtrIndex != INDEX_NONE)
 		{
-			Context.UserPtrTable[UserPtrIndex] = FuncInfo.InstData;
+			Context.VectorVMContext->UserPtrTable[UserPtrIndex] = FuncInfo.InstData;
 		}
 
-		Context.StartInstance = InstanceIndex;
+		Context.VectorVMContext->StartInstance = InstanceIndex;
 
 		//TODO: In future for DIs where more perf is needed here we could split the DI func into an args gen and a execution.
 		//The this path could gen args from the bytecode once and just run the execution func per instance.
@@ -409,9 +409,9 @@ void FNiagaraSystemScriptExecutionContext::PerInstanceFunctionHook(FVectorVMCont
 	}
 
 	//Restore the context state.
-	Context.ExternalFunctionInstanceOffset = 0;
-	Context.StartInstance = CachedContextStartInstance;
-	Context.NumInstances = CachedContextNumInstances;
+	Context.VectorVMContext->ExternalFunctionInstanceOffset = 0;
+	Context.VectorVMContext->StartInstance = CachedContextStartInstance;
+	Context.VectorVMContext->NumInstances = CachedContextNumInstances;
 }
 
 bool FNiagaraSystemScriptExecutionContext::Init(UNiagaraScript* InScript, ENiagaraSimTarget InTarget)
@@ -471,7 +471,7 @@ bool FNiagaraSystemScriptExecutionContext::Tick(class FNiagaraSystemInstance* In
 						if (ScriptDIInfo.NeedsPerInstanceBinding())
 						{
 							//This DI needs a binding per instance so we just bind to the external function hook which will call the correct binding for each instance.
-							auto PerInstFunctionHookLambda = [ExecContext = this, NumPerInstanceFunctions, UserPtrIndex = ScriptDIInfo.UserPtrIdx](FVectorVMContext& Context)
+							auto PerInstFunctionHookLambda = [ExecContext = this, NumPerInstanceFunctions, UserPtrIndex = ScriptDIInfo.UserPtrIdx](FVectorVMExternalFunctionContext& Context)
 							{
 								ExecContext->PerInstanceFunctionHook(Context, NumPerInstanceFunctions, UserPtrIndex);
 							};
diff --git a/Engine/Plugins/Runtime/HairStrands/Source/HairStrandsCore/Private/Niagara/NiagaraDataInterfaceHairStrands.cpp b/Engine/Plugins/Runtime/HairStrands/Source/HairStrandsCore/Private/Niagara/NiagaraDataInterfaceHairStrands.cpp
--- a/Engine/Plugins/Runtime/HairStrands/Source/HairStrandsCore/Private/Niagara/NiagaraDataInterfaceHairStrands.cpp
+++ b/Engine/Plugins/Runtime/HairStrands/Source/HairStrandsCore/Private/Niagara/NiagaraDataInterfaceHairStrands.cpp
@@ -2663,7 +2663,7 @@ void UNiagaraDataInterfaceHairStrands::GetVMExternalFunction(const FVMExternalFu
 	}
 }
 
-void WriteTransform(const FMatrix& ToWrite, FVectorVMContext& Context)
+void WriteTransform(const FMatrix& ToWrite, FVectorVMExternalFunctionContext& Context)
 {
 	VectorVM::FExternalFuncRegisterHandler<float> Out00(Context);
 	VectorVM::FExternalFuncRegisterHandler<float> Out01(Context);
@@ -2682,7 +2682,7 @@ void WriteTransform(const FMatrix& ToWrite, FVectorVMContext& Context)
 	VectorVM::FExternalFuncRegisterHandler<float> Out14(Context);
 	VectorVM::FExternalFuncRegisterHandler<float> Out15(Context);
 
-	for (int32 i = 0; i < Context.NumInstances; ++i)
+	for (int32 i = 0; i < Context.GetNumInstances(); ++i)
 	{
 		*Out00.GetDest() = ToWrite.M[0][0]; Out00.Advance();
 		*Out01.GetDest() = ToWrite.M[0][1]; Out01.Advance();
@@ -2703,58 +2703,58 @@ void WriteTransform(const FMatrix& ToWrite, FVectorVMContext& Context)
 	}
 }
 
-void UNiagaraDataInterfaceHairStrands::GetNumStrands(FVectorVMContext& Context)
+void UNiagaraDataInterfaceHairStrands::GetNumStrands(FVectorVMExternalFunctionContext& Context)
 {
 	VectorVM::FUserPtrHandler<FNDIHairStrandsData> InstData(Context);
 	VectorVM::FExternalFuncRegisterHandler<int32> OutNumStrands(Context);
 
-	for (int32 InstanceIdx = 0; InstanceIdx < Context.NumInstances; ++InstanceIdx)
+	for (int32 InstanceIdx = 0; InstanceIdx < Context.GetNumInstances(); ++InstanceIdx)
 	{
 		*OutNumStrands.GetDestAndAdvance() = InstData->NumStrands;
 	}
 }
 
-void UNiagaraDataInterfaceHairStrands::GetStrandSize(FVectorVMContext& Context)
+void UNiagaraDataInterfaceHairStrands::GetStrandSize(FVectorVMExternalFunctionContext& Context)
 {
 	VectorVM::FUserPtrHandler<FNDIHairStrandsData> InstData(Context);
 	VectorVM::FExternalFuncRegisterHandler<int32> OutStrandSize(Context);
 
-	for (int32 InstanceIdx = 0; InstanceIdx < Context.NumInstances; ++InstanceIdx)
+	for (int32 InstanceIdx = 0; InstanceIdx < Context.GetNumInstances(); ++InstanceIdx)
 	{
 		*OutStrandSize.GetDestAndAdvance() = InstData->StrandsSize;
 	}
 }
 
-void UNiagaraDataInterfaceHairStrands::GetSubSteps(FVectorVMContext& Context)
+void UNiagaraDataInterfaceHairStrands::GetSubSteps(FVectorVMExternalFunctionContext& Context)
 {
 	VectorVM::FUserPtrHandler<FNDIHairStrandsData> InstData(Context);
 	VectorVM::FExternalFuncRegisterHandler<int32> OutSubSteps(Context);
 
-	for (int32 InstanceIdx = 0; InstanceIdx < Context.NumInstances; ++InstanceIdx)
+	for (int32 InstanceIdx = 0; InstanceIdx < Context.GetNumInstances(); ++InstanceIdx)
 	{
 		*OutSubSteps.GetDestAndAdvance() = InstData->SubSteps;
 	}
 }
 
-void UNiagaraDataInterfaceHairStrands::GetIterationCount(FVectorVMContext& Context)
+void UNiagaraDataInterfaceHairStrands::GetIterationCount(FVectorVMExternalFunctionContext& Context)
 {
 	VectorVM::FUserPtrHandler<FNDIHairStrandsData> InstData(Context);
 	VectorVM::FExternalFuncRegisterHandler<int32> OutIterationCount(Context);
 
-	for (int32 InstanceIdx = 0; InstanceIdx < Context.NumInstances; ++InstanceIdx)
+	for (int32 InstanceIdx = 0; InstanceIdx < Context.GetNumInstances(); ++InstanceIdx)
 	{
 		*OutIterationCount.GetDestAndAdvance() = InstData->IterationCount;
 	}
 }
 
-void UNiagaraDataInterfaceHairStrands::GetGravityVector(FVectorVMContext& Context)
+void UNiagaraDataInterfaceHairStrands::GetGravityVector(FVectorVMExternalFunctionContext& Context)
 {
 	VectorVM::FUserPtrHandler<FNDIHairStrandsData> InstData(Context);
 	VectorVM::FExternalFuncRegisterHandler<float> OutGravityVectorX(Context);
 	VectorVM::FExternalFuncRegisterHandler<float> OutGravityVectorY(Context);
 	VectorVM::FExternalFuncRegisterHandler<float> OutGravityVectorZ(Context);
 
-	for (int32 InstanceIdx = 0; InstanceIdx < Context.NumInstances; ++InstanceIdx)
+	for (int32 InstanceIdx = 0; InstanceIdx < Context.GetNumInstances(); ++InstanceIdx)
 	{
 		*OutGravityVectorX.GetDestAndAdvance() = InstData->GravityVector.X;
 		*OutGravityVectorY.GetDestAndAdvance() = InstData->GravityVector.Y;
@@ -2762,25 +2762,25 @@ void UNiagaraDataInterfaceHairStrands::GetGravityVector(FVectorVMContext& Contex
 	}
 }
 
-void UNiagaraDataInterfaceHairStrands::GetAirDrag(FVectorVMContext& Context)
+void UNiagaraDataInterfaceHairStrands::GetAirDrag(FVectorVMExternalFunctionContext& Context)
 {
 	VectorVM::FUserPtrHandler<FNDIHairStrandsData> InstData(Context);
 	VectorVM::FExternalFuncRegisterHandler<float> OutAirDrag(Context);
 
-	for (int32 InstanceIdx = 0; InstanceIdx < Context.NumInstances; ++InstanceIdx)
+	for (int32 InstanceIdx = 0; InstanceIdx < Context.GetNumInstances(); ++InstanceIdx)
 	{
 		*OutAirDrag.GetDestAndAdvance() = InstData->AirDrag;
 	}
 }
 
-void UNiagaraDataInterfaceHairStrands::GetAirVelocity(FVectorVMContext& Context)
+void UNiagaraDataInterfaceHairStrands::GetAirVelocity(FVectorVMExternalFunctionContext& Context)
 {
 	VectorVM::FUserPtrHandler<FNDIHairStrandsData> InstData(Context);
 	VectorVM::FExternalFuncRegisterHandler<float> OutAirVelocityX(Context);
 	VectorVM::FExternalFuncRegisterHandler<float> OutAirVelocityY(Context);
 	VectorVM::FExternalFuncRegisterHandler<float> OutAirVelocityZ(Context);
 
-	for (int32 InstanceIdx = 0; InstanceIdx < Context.NumInstances; ++InstanceIdx)
+	for (int32 InstanceIdx = 0; InstanceIdx < Context.GetNumInstances(); ++InstanceIdx)
 	{
 		*OutAirVelocityX.GetDestAndAdvance() = InstData->AirVelocity.X;
 		*OutAirVelocityY.GetDestAndAdvance() = InstData->AirVelocity.Y;
@@ -2788,224 +2788,224 @@ void UNiagaraDataInterfaceHairStrands::GetAirVelocity(FVectorVMContext& Context)
 	}
 }
 
-void UNiagaraDataInterfaceHairStrands::GetSolveBend(FVectorVMContext& Context)
+void UNiagaraDataInterfaceHairStrands::GetSolveBend(FVectorVMExternalFunctionContext& Context)
 {
 	VectorVM::FUserPtrHandler<FNDIHairStrandsData> InstData(Context);
 	VectorVM::FExternalFuncRegisterHandler<int32> OutSolveBend(Context);
 
-	for (int32 InstanceIdx = 0; InstanceIdx < Context.NumInstances; ++InstanceIdx)
+	for (int32 InstanceIdx = 0; InstanceIdx < Context.GetNumInstances(); ++InstanceIdx)
 	{
 		*OutSolveBend.GetDestAndAdvance() = InstData->SolveBend;
 	}
 }
 
-void UNiagaraDataInterfaceHairStrands::GetProjectBend(FVectorVMContext& Context)
+void UNiagaraDataInterfaceHairStrands::GetProjectBend(FVectorVMExternalFunctionContext& Context)
 {
 	VectorVM::FUserPtrHandler<FNDIHairStrandsData> InstData(Context);
 	VectorVM::FExternalFuncRegisterHandler<int32> OutProjectBend(Context);
 
-	for (int32 InstanceIdx = 0; InstanceIdx < Context.NumInstances; ++InstanceIdx)
+	for (int32 InstanceIdx = 0; InstanceIdx < Context.GetNumInstances(); ++InstanceIdx)
 	{
 		*OutProjectBend.GetDestAndAdvance() = InstData->ProjectBend;
 	}
 }
 
-void UNiagaraDataInterfaceHairStrands::GetBendDamping(FVectorVMContext& Context)
+void UNiagaraDataInterfaceHairStrands::GetBendDamping(FVectorVMExternalFunctionContext& Context)
 {
 	VectorVM::FUserPtrHandler<FNDIHairStrandsData> InstData(Context);
 	VectorVM::FExternalFuncRegisterHandler<float> OutBendDamping(Context);
 
-	for (int32 InstanceIdx = 0; InstanceIdx < Context.NumInstances; ++InstanceIdx)
+	for (int32 InstanceIdx = 0; InstanceIdx < Context.GetNumInstances(); ++InstanceIdx)
 	{
 		*OutBendDamping.GetDestAndAdvance() = InstData->BendDamping;
 	}
 }
 
-void UNiagaraDataInterfaceHairStrands::GetBendStiffness(FVectorVMContext& Context)
+void UNiagaraDataInterfaceHairStrands::GetBendStiffness(FVectorVMExternalFunctionContext& Context)
 {
 	VectorVM::FUserPtrHandler<FNDIHairStrandsData> InstData(Context);
 	VectorVM::FExternalFuncRegisterHandler<float> OutBendStiffness(Context);
 
-	for (int32 InstanceIdx = 0; InstanceIdx < Context.NumInstances; ++InstanceIdx)
+	for (int32 InstanceIdx = 0; InstanceIdx < Context.GetNumInstances(); ++InstanceIdx)
 	{
 		*OutBendStiffness.GetDestAndAdvance() = InstData->BendStiffness;
 	}
 }
 
-void UNiagaraDataInterfaceHairStrands::GetBendScale(FVectorVMContext& Context)
+void UNiagaraDataInterfaceHairStrands::GetBendScale(FVectorVMExternalFunctionContext& Context)
 {}
 
-void UNiagaraDataInterfaceHairStrands::GetSolveStretch(FVectorVMContext& Context)
+void UNiagaraDataInterfaceHairStrands::GetSolveStretch(FVectorVMExternalFunctionContext& Context)
 {
 	VectorVM::FUserPtrHandler<FNDIHairStrandsData> InstData(Context);
 	VectorVM::FExternalFuncRegisterHandler<int32> OutSolveStretch(Context);
 
-	for (int32 InstanceIdx = 0; InstanceIdx < Context.NumInstances; ++InstanceIdx)
+	for (int32 InstanceIdx = 0; InstanceIdx < Context.GetNumInstances(); ++InstanceIdx)
 	{
 		*OutSolveStretch.GetDestAndAdvance() = InstData->SolveStretch;
 	}
 }
 
-void UNiagaraDataInterfaceHairStrands::GetProjectStretch(FVectorVMContext& Context)
+void UNiagaraDataInterfaceHairStrands::GetProjectStretch(FVectorVMExternalFunctionContext& Context)
 {
 	VectorVM::FUserPtrHandler<FNDIHairStrandsData> InstData(Context);
 	VectorVM::FExternalFuncRegisterHandler<int32> OutProjectStretch(Context);
 
-	for (int32 InstanceIdx = 0; InstanceIdx < Context.NumInstances; ++InstanceIdx)
+	for (int32 InstanceIdx = 0; InstanceIdx < Context.GetNumInstances(); ++InstanceIdx)
 	{
 		*OutProjectStretch.GetDestAndAdvance() = InstData->ProjectStretch;
 	}
 }
 
-void UNiagaraDataInterfaceHairStrands::GetStretchDamping(FVectorVMContext& Context)
+void UNiagaraDataInterfaceHairStrands::GetStretchDamping(FVectorVMExternalFunctionContext& Context)
 {
 	VectorVM::FUserPtrHandler<FNDIHairStrandsData> InstData(Context);
 	VectorVM::FExternalFuncRegisterHandler<float> OutStretchDamping(Context);
 
-	for (int32 InstanceIdx = 0; InstanceIdx < Context.NumInstances; ++InstanceIdx)
+	for (int32 InstanceIdx = 0; InstanceIdx < Context.GetNumInstances(); ++InstanceIdx)
 	{
 		*OutStretchDamping.GetDestAndAdvance() = InstData->StretchDamping;
 	}
 }
 
-void UNiagaraDataInterfaceHairStrands::GetStretchStiffness(FVectorVMContext& Context)
+void UNiagaraDataInterfaceHairStrands::GetStretchStiffness(FVectorVMExternalFunctionContext& Context)
 {
 	VectorVM::FUserPtrHandler<FNDIHairStrandsData> InstData(Context);
 	VectorVM::FExternalFuncRegisterHandler<float> OutStretchStiffness(Context);
 
-	for (int32 InstanceIdx = 0; InstanceIdx < Context.NumInstances; ++InstanceIdx)
+	for (int32 InstanceIdx = 0; InstanceIdx < Context.GetNumInstances(); ++InstanceIdx)
 	{
 		*OutStretchStiffness.GetDestAndAdvance() = InstData->StretchStiffness;
 	}
 }
 
-void UNiagaraDataInterfaceHairStrands::GetStretchScale(FVectorVMContext& Context)
+void UNiagaraDataInterfaceHairStrands::GetStretchScale(FVectorVMExternalFunctionContext& Context)
 {
 }
 
 
-void UNiagaraDataInterfaceHairStrands::GetSolveCollision(FVectorVMContext& Context)
+void UNiagaraDataInterfaceHairStrands::GetSolveCollision(FVectorVMExternalFunctionContext& Context)
 {
 	VectorVM::FUserPtrHandler<FNDIHairStrandsData> InstData(Context);
 	VectorVM::FExternalFuncRegisterHandler<int32> OutSolveCollision(Context);
 
-	for (int32 InstanceIdx = 0; InstanceIdx < Context.NumInstances; ++InstanceIdx)
+	for (int32 InstanceIdx = 0; InstanceIdx < Context.GetNumInstances(); ++InstanceIdx)
 	{
 		*OutSolveCollision.GetDestAndAdvance() = InstData->SolveCollision;
 	}
 }
 
-void UNiagaraDataInterfaceHairStrands::GetProjectCollision(FVectorVMContext& Context)
+void UNiagaraDataInterfaceHairStrands::GetProjectCollision(FVectorVMExternalFunctionContext& Context)
 {
 	VectorVM::FUserPtrHandler<FNDIHairStrandsData> InstData(Context);
 	VectorVM::FExternalFuncRegisterHandler<int32> OutProjectCollision(Context);
 
-	for (int32 InstanceIdx = 0; InstanceIdx < Context.NumInstances; ++InstanceIdx)
+	for (int32 InstanceIdx = 0; InstanceIdx < Context.GetNumInstances(); ++InstanceIdx)
 	{
 		*OutProjectCollision.GetDestAndAdvance() = InstData->ProjectCollision;
 	}
 }
 
-void UNiagaraDataInterfaceHairStrands::GetStaticFriction(FVectorVMContext& Context)
+void UNiagaraDataInterfaceHairStrands::GetStaticFriction(FVectorVMExternalFunctionContext& Context)
 {
 	VectorVM::FUserPtrHandler<FNDIHairStrandsData> InstData(Context);
 	VectorVM::FExternalFuncRegisterHandler<float> OutStaticFriction(Context);
 
-	for (int32 InstanceIdx = 0; InstanceIdx < Context.NumInstances; ++InstanceIdx)
+	for (int32 InstanceIdx = 0; InstanceIdx < Context.GetNumInstances(); ++InstanceIdx)
 	{
 		*OutStaticFriction.GetDestAndAdvance() = InstData->StaticFriction;
 	}
 }
 
-void UNiagaraDataInterfaceHairStrands::GetKineticFriction(FVectorVMContext& Context)
+void UNiagaraDataInterfaceHairStrands::GetKineticFriction(FVectorVMExternalFunctionContext& Context)
 {
 	VectorVM::FUserPtrHandler<FNDIHairStrandsData> InstData(Context);
 	VectorVM::FExternalFuncRegisterHandler<float> OutKineticFriction(Context);
 
-	for (int32 InstanceIdx = 0; InstanceIdx < Context.NumInstances; ++InstanceIdx)
+	for (int32 InstanceIdx = 0; InstanceIdx < Context.GetNumInstances(); ++InstanceIdx)
 	{
 		*OutKineticFriction.GetDestAndAdvance() = InstData->KineticFriction;
 	}
 }
 
-void UNiagaraDataInterfaceHairStrands::GetStrandsViscosity(FVectorVMContext& Context)
+void UNiagaraDataInterfaceHairStrands::GetStrandsViscosity(FVectorVMExternalFunctionContext& Context)
 {
 	VectorVM::FUserPtrHandler<FNDIHairStrandsData> InstData(Context);
 	VectorVM::FExternalFuncRegisterHandler<float> OutStrandsViscosity(Context);
 
-	for (int32 InstanceIdx = 0; InstanceIdx < Context.NumInstances; ++InstanceIdx)
+	for (int32 InstanceIdx = 0; InstanceIdx < Context.GetNumInstances(); ++InstanceIdx)
 	{
 		*OutStrandsViscosity.GetDestAndAdvance() = InstData->StrandsViscosity;
 	}
 }
 
-void UNiagaraDataInterfaceHairStrands::GetGridDimension(FVectorVMContext& Context)
+void UNiagaraDataInterfaceHairStrands::GetGridDimension(FVectorVMExternalFunctionContext& Context)
 {
 	VectorVM::FUserPtrHandler<FNDIHairStrandsData> InstData(Context);
 	VectorVM::FExternalFuncRegisterHandler<float> OutGridDimensionX(Context);
 	VectorVM::FExternalFuncRegisterHandler<float> OutGridDimensionY(Context);
 	VectorVM::FExternalFuncRegisterHandler<float> OutGridDimensionZ(Context);
 
-	for (int32 InstanceIdx = 0; InstanceIdx < Context.NumInstances; ++InstanceIdx)
+	for (int32 InstanceIdx = 0; InstanceIdx < Context.GetNumInstances(); ++InstanceIdx)
 	{
 		*OutGridDimensionX.GetDestAndAdvance() = InstData->GridDimension.X;
 		*OutGridDimensionY.GetDestAndAdvance() = InstData->GridDimension.Y;
 		*OutGridDimensionZ.GetDestAndAdvance() = InstData->GridDimension.Z;
 	}
 }
-void UNiagaraDataInterfaceHairStrands::GetCollisionRadius(FVectorVMContext& Context)
+void UNiagaraDataInterfaceHairStrands::GetCollisionRadius(FVectorVMExternalFunctionContext& Context)
 {
 	VectorVM::FUserPtrHandler<FNDIHairStrandsData> InstData(Context);
 	VectorVM::FExternalFuncRegisterHandler<float> OutCollisionRadius(Context);
 
-	for (int32 InstanceIdx = 0; InstanceIdx < Context.NumInstances; ++InstanceIdx)
+	for (int32 InstanceIdx = 0; InstanceIdx < Context.GetNumInstances(); ++InstanceIdx)
 	{
 		*OutCollisionRadius.GetDestAndAdvance() = InstData->CollisionRadius;
 	}
 }
 
-void UNiagaraDataInterfaceHairStrands::GetRadiusScale(FVectorVMContext& Context)
+void UNiagaraDataInterfaceHairStrands::GetRadiusScale(FVectorVMExternalFunctionContext& Context)
 {
 }
 
-void UNiagaraDataInterfaceHairStrands::GetStrandsSmoothing(FVectorVMContext& Context)
+void UNiagaraDataInterfaceHairStrands::GetStrandsSmoothing(FVectorVMExternalFunctionContext& Context)
 {
 	VectorVM::FUserPtrHandler<FNDIHairStrandsData> InstData(Context);
 	VectorVM::FExternalFuncRegisterHandler<float> OutStrandsSmoothing(Context);
 
-	for (int32 InstanceIdx = 0; InstanceIdx < Context.NumInstances; ++InstanceIdx)
+	for (int32 InstanceIdx = 0; InstanceIdx < Context.GetNumInstances(); ++InstanceIdx)
 	{
 		*OutStrandsSmoothing.GetDestAndAdvance() = InstData->StrandsSmoothing;
 	}
 }
 
-void UNiagaraDataInterfaceHairStrands::GetStrandsDensity(FVectorVMContext& Context)
+void UNiagaraDataInterfaceHairStrands::GetStrandsDensity(FVectorVMExternalFunctionContext& Context)
 {
 	VectorVM::FUserPtrHandler<FNDIHairStrandsData> InstData(Context);
 	VectorVM::FExternalFuncRegisterHandler<float> OutStrandsDensity(Context);
 
-	for (int32 InstanceIdx = 0; InstanceIdx < Context.NumInstances; ++InstanceIdx)
+	for (int32 InstanceIdx = 0; InstanceIdx < Context.GetNumInstances(); ++InstanceIdx)
 	{
 		*OutStrandsDensity.GetDestAndAdvance() = InstData->StrandsDensity;
 	}
 }
 
-void UNiagaraDataInterfaceHairStrands::GetStrandsThickness(FVectorVMContext& Context)
+void UNiagaraDataInterfaceHairStrands::GetStrandsThickness(FVectorVMExternalFunctionContext& Context)
 {
 	VectorVM::FUserPtrHandler<FNDIHairStrandsData> InstData(Context);
 	VectorVM::FExternalFuncRegisterHandler<float> OutStrandsThickness(Context);
 
-	for (int32 InstanceIdx = 0; InstanceIdx < Context.NumInstances; ++InstanceIdx)
+	for (int32 InstanceIdx = 0; InstanceIdx < Context.GetNumInstances(); ++InstanceIdx)
 	{
 		*OutStrandsThickness.GetDestAndAdvance() = InstData->StrandsThickness;
 	}
 }
 
-void UNiagaraDataInterfaceHairStrands::GetThicknessScale(FVectorVMContext& Context)
+void UNiagaraDataInterfaceHairStrands::GetThicknessScale(FVectorVMExternalFunctionContext& Context)
 {
 }
 
-void UNiagaraDataInterfaceHairStrands::GetWorldTransform(FVectorVMContext& Context)
+void UNiagaraDataInterfaceHairStrands::GetWorldTransform(FVectorVMExternalFunctionContext& Context)
 {
 	VectorVM::FUserPtrHandler<FNDIHairStrandsData> InstData(Context);
 	FMatrix WorldTransform = InstData->WorldTransform.ToMatrixWithScale();
@@ -3013,7 +3013,7 @@ void UNiagaraDataInterfaceHairStrands::GetWorldTransform(FVectorVMContext& Conte
 	WriteTransform(WorldTransform, Context);
 }
 
-void UNiagaraDataInterfaceHairStrands::GetWorldInverse(FVectorVMContext& Context)
+void UNiagaraDataInterfaceHairStrands::GetWorldInverse(FVectorVMExternalFunctionContext& Context)
 {
 	VectorVM::FUserPtrHandler<FNDIHairStrandsData> InstData(Context);
 	FMatrix WorldInverse = InstData->WorldTransform.ToMatrixWithScale().Inverse();
@@ -3021,258 +3021,258 @@ void UNiagaraDataInterfaceHairStrands::GetWorldInverse(FVectorVMContext& Context
 	WriteTransform(WorldInverse, Context);
 }
 
-void UNiagaraDataInterfaceHairStrands::GetBoundingBox(FVectorVMContext& Context)
+void UNiagaraDataInterfaceHairStrands::GetBoundingBox(FVectorVMExternalFunctionContext& Context)
 {
 	// @todo : implement function for cpu
 }
 
-void UNiagaraDataInterfaceHairStrands::ResetBoundingBox(FVectorVMContext& Context)
+void UNiagaraDataInterfaceHairStrands::ResetBoundingBox(FVectorVMExternalFunctionContext& Context)
 {
 	// @todo : implement function for cpu
 }
 
-void UNiagaraDataInterfaceHairStrands::BuildBoundingBox(FVectorVMContext& Context)
+void UNiagaraDataInterfaceHairStrands::BuildBoundingBox(FVectorVMExternalFunctionContext& Context)
 {
 	// @todo : implement function for cpu
 }
 
-void UNiagaraDataInterfaceHairStrands::GetPointPosition(FVectorVMContext& Context)
+void UNiagaraDataInterfaceHairStrands::GetPointPosition(FVectorVMExternalFunctionContext& Context)
 {
 	// @todo : implement function for cpu
 }
 
-void UNiagaraDataInterfaceHairStrands::ComputeNodePosition(FVectorVMContext& Context)
+void UNiagaraDataInterfaceHairStrands::ComputeNodePosition(FVectorVMExternalFunctionContext& Context)
 {
 	// @todo : implement function for cpu
 }
 
-void UNiagaraDataInterfaceHairStrands::ComputeNodeOrientation(FVectorVMContext& Context)
+void UNiagaraDataInterfaceHairStrands::ComputeNodeOrientation(FVectorVMExternalFunctionContext& Context)
 {
 	// @todo : implement function for cpu
 }
 
-void UNiagaraDataInterfaceHairStrands::ComputeNodeMass(FVectorVMContext& Context)
+void UNiagaraDataInterfaceHairStrands::ComputeNodeMass(FVectorVMExternalFunctionContext& Context)
 {
 	// @todo : implement function for cpu
 }
 
-void UNiagaraDataInterfaceHairStrands::ComputeNodeInertia(FVectorVMContext& Context)
+void UNiagaraDataInterfaceHairStrands::ComputeNodeInertia(FVectorVMExternalFunctionContext& Context)
 {
 	// @todo : implement function for cpu
 }
 
-void UNiagaraDataInterfaceHairStrands::ComputeEdgeLength(FVectorVMContext& Context)
+void UNiagaraDataInterfaceHairStrands::ComputeEdgeLength(FVectorVMExternalFunctionContext& Context)
 {
 	// @todo : implement function for cpu
 }
 
-void UNiagaraDataInterfaceHairStrands::ComputeEdgeRotation(FVectorVMContext& Context)
+void UNiagaraDataInterfaceHairStrands::ComputeEdgeRotation(FVectorVMExternalFunctionContext& Context)
 {
 	// @todo : implement function for cpu
 }
 
-void UNiagaraDataInterfaceHairStrands::ComputeRestPosition(FVectorVMContext& Context)
+void UNiagaraDataInterfaceHairStrands::ComputeRestPosition(FVectorVMExternalFunctionContext& Context)
 {
 	// @todo : implement function for cpu
 }
 
-void UNiagaraDataInterfaceHairStrands::ComputeRestOrientation(FVectorVMContext& Context)
+void UNiagaraDataInterfaceHairStrands::ComputeRestOrientation(FVectorVMExternalFunctionContext& Context)
 {
 	// @todo : implement function for cpu
 }
 
-void UNiagaraDataInterfaceHairStrands::ComputeLocalState(FVectorVMContext& Context)
+void UNiagaraDataInterfaceHairStrands::ComputeLocalState(FVectorVMExternalFunctionContext& Context)
 {
 	// @todo : implement function for cpu
 }
 
-void UNiagaraDataInterfaceHairStrands::UpdatePointPosition(FVectorVMContext& Context)
+void UNiagaraDataInterfaceHairStrands::UpdatePointPosition(FVectorVMExternalFunctionContext& Context)
 {
 	// @todo : implement function for cpu
 }
 
-void UNiagaraDataInterfaceHairStrands::ResetPointPosition(FVectorVMContext& Context)
+void UNiagaraDataInterfaceHairStrands::ResetPointPosition(FVectorVMExternalFunctionContext& Context)
 {
 	// @todo : implement function for cpu
 }
 
-void UNiagaraDataInterfaceHairStrands::AttachNodePosition(FVectorVMContext& Context)
+void UNiagaraDataInterfaceHairStrands::AttachNodePosition(FVectorVMExternalFunctionContext& Context)
 {
 	// @todo : implement function for cpu
 }
 
-void UNiagaraDataInterfaceHairStrands::EvalSkinnedPosition(FVectorVMContext& Context)
+void UNiagaraDataInterfaceHairStrands::EvalSkinnedPosition(FVectorVMExternalFunctionContext& Context)
 {
 	// @todo : implement function for cpu
 }
 
-void UNiagaraDataInterfaceHairStrands::AttachNodeOrientation(FVectorVMContext& Context)
+void UNiagaraDataInterfaceHairStrands::AttachNodeOrientation(FVectorVMExternalFunctionContext& Context)
 {
 	// @todo : implement function for cpu
 }
 
-void UNiagaraDataInterfaceHairStrands::AttachNodeState(FVectorVMContext& Context)
+void UNiagaraDataInterfaceHairStrands::AttachNodeState(FVectorVMExternalFunctionContext& Context)
 {
 	// @todo : implement function for cpu
 }
 
-void UNiagaraDataInterfaceHairStrands::UpdateNodeState(FVectorVMContext& Context)
+void UNiagaraDataInterfaceHairStrands::UpdateNodeState(FVectorVMExternalFunctionContext& Context)
 {
 	// @todo : implement function for cpu
 }
 
-void UNiagaraDataInterfaceHairStrands::AdvectNodePosition(FVectorVMContext& Context)
+void UNiagaraDataInterfaceHairStrands::AdvectNodePosition(FVectorVMExternalFunctionContext& Context)
 {
 	// @todo : implement function for cpu
 }
 
-void UNiagaraDataInterfaceHairStrands::AdvectNodeOrientation(FVectorVMContext& Context)
+void UNiagaraDataInterfaceHairStrands::AdvectNodeOrientation(FVectorVMExternalFunctionContext& Context)
 {
 	// @todo : implement function for cpu
 }
 
-void UNiagaraDataInterfaceHairStrands::UpdateLinearVelocity(FVectorVMContext& Context)
+void UNiagaraDataInterfaceHairStrands::UpdateLinearVelocity(FVectorVMExternalFunctionContext& Context)
 {
 	// @todo : implement function for cpu
 }
 
-void UNiagaraDataInterfaceHairStrands::UpdateAngularVelocity(FVectorVMContext& Context)
+void UNiagaraDataInterfaceHairStrands::UpdateAngularVelocity(FVectorVMExternalFunctionContext& Context)
 {
 	// @todo : implement function for cpu
 }
 
-void UNiagaraDataInterfaceHairStrands::SetupDistanceSpringMaterial(FVectorVMContext& Context)
+void UNiagaraDataInterfaceHairStrands::SetupDistanceSpringMaterial(FVectorVMExternalFunctionContext& Context)
 {
 	// @todo : implement function for cpu
 }
 
-void UNiagaraDataInterfaceHairStrands::SolveDistanceSpringMaterial(FVectorVMContext& Context)
+void UNiagaraDataInterfaceHairStrands::SolveDistanceSpringMaterial(FVectorVMExternalFunctionContext& Context)
 {
 	// @todo : implement function for cpu
 }
 
-void UNiagaraDataInterfaceHairStrands::ProjectDistanceSpringMaterial(FVectorVMContext& Context)
+void UNiagaraDataInterfaceHairStrands::ProjectDistanceSpringMaterial(FVectorVMExternalFunctionContext& Context)
 {
 	// @todo : implement function for cpu
 }
 
-void UNiagaraDataInterfaceHairStrands::SetupAngularSpringMaterial(FVectorVMContext& Context)
+void UNiagaraDataInterfaceHairStrands::SetupAngularSpringMaterial(FVectorVMExternalFunctionContext& Context)
 {
 	// @todo : implement function for cpu
 }
 
-void UNiagaraDataInterfaceHairStrands::SolveAngularSpringMaterial(FVectorVMContext& Context)
+void UNiagaraDataInterfaceHairStrands::SolveAngularSpringMaterial(FVectorVMExternalFunctionContext& Context)
 {
 	// @todo : implement function for cpu
 }
 
-void UNiagaraDataInterfaceHairStrands::ProjectAngularSpringMaterial(FVectorVMContext& Context)
+void UNiagaraDataInterfaceHairStrands::ProjectAngularSpringMaterial(FVectorVMExternalFunctionContext& Context)
 {
 	// @todo : implement function for cpu
 }
 
 
-void UNiagaraDataInterfaceHairStrands::SetupStretchRodMaterial(FVectorVMContext& Context)
+void UNiagaraDataInterfaceHairStrands::SetupStretchRodMaterial(FVectorVMExternalFunctionContext& Context)
 {
 	// @todo : implement function for cpu
 }
 
-void UNiagaraDataInterfaceHairStrands::SolveStretchRodMaterial(FVectorVMContext& Context)
+void UNiagaraDataInterfaceHairStrands::SolveStretchRodMaterial(FVectorVMExternalFunctionContext& Context)
 {
 	// @todo : implement function for cpu
 }
 
-void UNiagaraDataInterfaceHairStrands::ProjectStretchRodMaterial(FVectorVMContext& Context)
+void UNiagaraDataInterfaceHairStrands::ProjectStretchRodMaterial(FVectorVMExternalFunctionContext& Context)
 {
 	// @todo : implement function for cpu
 }
 
-void UNiagaraDataInterfaceHairStrands::SetupBendRodMaterial(FVectorVMContext& Context)
+void UNiagaraDataInterfaceHairStrands::SetupBendRodMaterial(FVectorVMExternalFunctionContext& Context)
 {
 	// @todo : implement function for cpu
 }
 
-void UNiagaraDataInterfaceHairStrands::SolveBendRodMaterial(FVectorVMContext& Context)
+void UNiagaraDataInterfaceHairStrands::SolveBendRodMaterial(FVectorVMExternalFunctionContext& Context)
 {
 	// @todo : implement function for cpu
 }
 
-void UNiagaraDataInterfaceHairStrands::ProjectBendRodMaterial(FVectorVMContext& Context)
+void UNiagaraDataInterfaceHairStrands::ProjectBendRodMaterial(FVectorVMExternalFunctionContext& Context)
 {
 	// @todo : implement function for cpu
 }
 
-void UNiagaraDataInterfaceHairStrands::ComputeEdgeDirection(FVectorVMContext& Context)
+void UNiagaraDataInterfaceHairStrands::ComputeEdgeDirection(FVectorVMExternalFunctionContext& Context)
 {
 	// @todo : implement function for cpu
 }
 
-void UNiagaraDataInterfaceHairStrands::UpdateMaterialFrame(FVectorVMContext& Context)
+void UNiagaraDataInterfaceHairStrands::UpdateMaterialFrame(FVectorVMExternalFunctionContext& Context)
 {
 	// @todo : implement function for cpu
 }
 
-void UNiagaraDataInterfaceHairStrands::ComputeMaterialFrame(FVectorVMContext& Context)
+void UNiagaraDataInterfaceHairStrands::ComputeMaterialFrame(FVectorVMExternalFunctionContext& Context)
 {
 	// @todo : implement function for cpu
 }
 
-void UNiagaraDataInterfaceHairStrands::SolveHardCollisionConstraint(FVectorVMContext& Context)
+void UNiagaraDataInterfaceHairStrands::SolveHardCollisionConstraint(FVectorVMExternalFunctionContext& Context)
 {
 	// @todo : implement function for cpu
 }
 
-void UNiagaraDataInterfaceHairStrands::ProjectHardCollisionConstraint(FVectorVMContext& Context)
+void UNiagaraDataInterfaceHairStrands::ProjectHardCollisionConstraint(FVectorVMExternalFunctionContext& Context)
 {
 	// @todo : implement function for cpu
 }
 
-void UNiagaraDataInterfaceHairStrands::SolveSoftCollisionConstraint(FVectorVMContext& Context)
+void UNiagaraDataInterfaceHairStrands::SolveSoftCollisionConstraint(FVectorVMExternalFunctionContext& Context)
 {
 	// @todo : implement function for cpu
 }
 
-void UNiagaraDataInterfaceHairStrands::ProjectSoftCollisionConstraint(FVectorVMContext& Context)
+void UNiagaraDataInterfaceHairStrands::ProjectSoftCollisionConstraint(FVectorVMExternalFunctionContext& Context)
 {
 	// @todo : implement function for cpu
 }
-void UNiagaraDataInterfaceHairStrands::SetupSoftCollisionConstraint(FVectorVMContext& Context)
+void UNiagaraDataInterfaceHairStrands::SetupSoftCollisionConstraint(FVectorVMExternalFunctionContext& Context)
 {
 	// @todo : implement function for cpu
 }
 
-void UNiagaraDataInterfaceHairStrands::ComputeAirDragForce(FVectorVMContext& Context)
+void UNiagaraDataInterfaceHairStrands::ComputeAirDragForce(FVectorVMExternalFunctionContext& Context)
 {
 	// @todo : implement function for cpu
 }
 
-void UNiagaraDataInterfaceHairStrands::NeedSimulationReset(FVectorVMContext& Context)
+void UNiagaraDataInterfaceHairStrands::NeedSimulationReset(FVectorVMExternalFunctionContext& Context)
 {
 	// @todo : implement function for cpu
 }
 
-void UNiagaraDataInterfaceHairStrands::InitGridSamples(FVectorVMContext& Context)
+void UNiagaraDataInterfaceHairStrands::InitGridSamples(FVectorVMExternalFunctionContext& Context)
 {
 	// @todo : implement function for cpu
 }
 
-void UNiagaraDataInterfaceHairStrands::GetSampleState(FVectorVMContext& Context)
+void UNiagaraDataInterfaceHairStrands::GetSampleState(FVectorVMExternalFunctionContext& Context)
 {
 	// @todo : implement function for cpu
 }
 
-void UNiagaraDataInterfaceHairStrands::HasGlobalInterpolation(FVectorVMContext& Context)
+void UNiagaraDataInterfaceHairStrands::HasGlobalInterpolation(FVectorVMExternalFunctionContext& Context)
 {
 	VectorVM::FUserPtrHandler<FNDIHairStrandsData> InstData(Context);
 	VectorVM::FExternalFuncRegisterHandler<int32> OutGlobalInterpolation(Context);
 
-	for (int32 InstanceIdx = 0; InstanceIdx < Context.NumInstances; ++InstanceIdx)
+	for (int32 InstanceIdx = 0; InstanceIdx < Context.GetNumInstances(); ++InstanceIdx)
 	{
 		*OutGlobalInterpolation.GetDestAndAdvance() = InstData->GlobalInterpolation;
 	}
 }
 
-void UNiagaraDataInterfaceHairStrands::NeedRestUpdate(FVectorVMContext& Context)
+void UNiagaraDataInterfaceHairStrands::NeedRestUpdate(FVectorVMExternalFunctionContext& Context)
 {
 	// @todo : implement function for cpu 
 }
diff --git a/Engine/Plugins/Runtime/HairStrands/Source/HairStrandsCore/Private/Niagara/NiagaraDataInterfacePhysicsAsset.cpp b/Engine/Plugins/Runtime/HairStrands/Source/HairStrandsCore/Private/Niagara/NiagaraDataInterfacePhysicsAsset.cpp
--- a/Engine/Plugins/Runtime/HairStrands/Source/HairStrandsCore/Private/Niagara/NiagaraDataInterfacePhysicsAsset.cpp
+++ b/Engine/Plugins/Runtime/HairStrands/Source/HairStrandsCore/Private/Niagara/NiagaraDataInterfacePhysicsAsset.cpp
@@ -1088,43 +1088,43 @@ void UNiagaraDataInterfacePhysicsAsset::GetVMExternalFunction(const FVMExternalF
 	}
 }
 
-void UNiagaraDataInterfacePhysicsAsset::GetNumBoxes(FVectorVMContext& Context)
+void UNiagaraDataInterfacePhysicsAsset::GetNumBoxes(FVectorVMExternalFunctionContext& Context)
 {
 }
 
-void UNiagaraDataInterfacePhysicsAsset::GetNumSpheres(FVectorVMContext& Context)
+void UNiagaraDataInterfacePhysicsAsset::GetNumSpheres(FVectorVMExternalFunctionContext& Context)
 {
 }
 
-void UNiagaraDataInterfacePhysicsAsset::GetNumCapsules(FVectorVMContext& Context)
+void UNiagaraDataInterfacePhysicsAsset::GetNumCapsules(FVectorVMExternalFunctionContext& Context)
 {
 }
 
-void UNiagaraDataInterfacePhysicsAsset::GetClosestPoint(FVectorVMContext& Context)
+void UNiagaraDataInterfacePhysicsAsset::GetClosestPoint(FVectorVMExternalFunctionContext& Context)
 {
 }
 
-void UNiagaraDataInterfacePhysicsAsset::GetClosestElement(FVectorVMContext& Context)
+void UNiagaraDataInterfacePhysicsAsset::GetClosestElement(FVectorVMExternalFunctionContext& Context)
 {
 }
 
-void UNiagaraDataInterfacePhysicsAsset::GetElementPoint(FVectorVMContext& Context)
+void UNiagaraDataInterfacePhysicsAsset::GetElementPoint(FVectorVMExternalFunctionContext& Context)
 {
 }
 
-void UNiagaraDataInterfacePhysicsAsset::GetElementDistance(FVectorVMContext& Context)
+void UNiagaraDataInterfacePhysicsAsset::GetElementDistance(FVectorVMExternalFunctionContext& Context)
 {
 }
 
-void UNiagaraDataInterfacePhysicsAsset::GetClosestDistance(FVectorVMContext& Context)
+void UNiagaraDataInterfacePhysicsAsset::GetClosestDistance(FVectorVMExternalFunctionContext& Context)
 {
 }
 
-void UNiagaraDataInterfacePhysicsAsset::GetTexturePoint(FVectorVMContext& Context)
+void UNiagaraDataInterfacePhysicsAsset::GetTexturePoint(FVectorVMExternalFunctionContext& Context)
 {
 }
 
-void UNiagaraDataInterfacePhysicsAsset::GetProjectionPoint(FVectorVMContext& Context)
+void UNiagaraDataInterfacePhysicsAsset::GetProjectionPoint(FVectorVMExternalFunctionContext& Context)
 {
 }
 
diff --git a/Engine/Plugins/Runtime/HairStrands/Source/HairStrandsCore/Private/Niagara/NiagaraDataInterfacePressureGrid.cpp b/Engine/Plugins/Runtime/HairStrands/Source/HairStrandsCore/Private/Niagara/NiagaraDataInterfacePressureGrid.cpp
--- a/Engine/Plugins/Runtime/HairStrands/Source/HairStrandsCore/Private/Niagara/NiagaraDataInterfacePressureGrid.cpp
+++ b/Engine/Plugins/Runtime/HairStrands/Source/HairStrandsCore/Private/Niagara/NiagaraDataInterfacePressureGrid.cpp
@@ -249,46 +249,46 @@ void UNiagaraDataInterfacePressureGrid::GetVMExternalFunction(const FVMExternalF
 	}
 }
 
-void UNiagaraDataInterfacePressureGrid::BuildDistanceField(FVectorVMContext& Context)
+void UNiagaraDataInterfacePressureGrid::BuildDistanceField(FVectorVMExternalFunctionContext& Context)
 {
 	// @todo : implement function for cpu 
 }
 
-void UNiagaraDataInterfacePressureGrid::BuildDensityField(FVectorVMContext& Context)
+void UNiagaraDataInterfacePressureGrid::BuildDensityField(FVectorVMExternalFunctionContext& Context)
 {
 	// @todo : implement function for cpu 
 }
 
-void UNiagaraDataInterfacePressureGrid::SolveGridPressure(FVectorVMContext& Context)
+void UNiagaraDataInterfacePressureGrid::SolveGridPressure(FVectorVMExternalFunctionContext& Context)
 {
 	// @todo : implement function for cpu 
 }
 
-void UNiagaraDataInterfacePressureGrid::ComputeBoundaryWeights(FVectorVMContext& Context)
+void UNiagaraDataInterfacePressureGrid::ComputeBoundaryWeights(FVectorVMExternalFunctionContext& Context)
 {
 	// @todo : implement function for cpu 
 }
 
-void UNiagaraDataInterfacePressureGrid::SetSolidBoundary(FVectorVMContext& Context)
+void UNiagaraDataInterfacePressureGrid::SetSolidBoundary(FVectorVMExternalFunctionContext& Context)
 {
 	// @todo : implement function for cpu 
 }
 
-void UNiagaraDataInterfacePressureGrid::ScaleCellFields(FVectorVMContext& Context)
+void UNiagaraDataInterfacePressureGrid::ScaleCellFields(FVectorVMExternalFunctionContext& Context)
 {
 	// @todo : implement function for cpu 
 }
 
-void UNiagaraDataInterfacePressureGrid::GetNodePosition(FVectorVMContext& Context)
+void UNiagaraDataInterfacePressureGrid::GetNodePosition(FVectorVMExternalFunctionContext& Context)
 {
 	// @todo : implement function for cpu 
 }
 
-void UNiagaraDataInterfacePressureGrid::GetDensityField(FVectorVMContext& Context)
+void UNiagaraDataInterfacePressureGrid::GetDensityField(FVectorVMExternalFunctionContext& Context)
 {
 	// @todo : implement function for cpu 
 }
-void UNiagaraDataInterfacePressureGrid::UpdateDeformationGradient(FVectorVMContext& Context)
+void UNiagaraDataInterfacePressureGrid::UpdateDeformationGradient(FVectorVMExternalFunctionContext& Context)
 {
 	// @todo : implement function for cpu 
 }
diff --git a/Engine/Plugins/Runtime/HairStrands/Source/HairStrandsCore/Private/Niagara/NiagaraDataInterfaceVelocityGrid.cpp b/Engine/Plugins/Runtime/HairStrands/Source/HairStrandsCore/Private/Niagara/NiagaraDataInterfaceVelocityGrid.cpp
--- a/Engine/Plugins/Runtime/HairStrands/Source/HairStrandsCore/Private/Niagara/NiagaraDataInterfaceVelocityGrid.cpp
+++ b/Engine/Plugins/Runtime/HairStrands/Source/HairStrandsCore/Private/Niagara/NiagaraDataInterfaceVelocityGrid.cpp
@@ -435,22 +435,22 @@ void UNiagaraDataInterfaceVelocityGrid::GetVMExternalFunction(const FVMExternalF
 	}
 }
 
-void UNiagaraDataInterfaceVelocityGrid::BuildVelocityField(FVectorVMContext& Context)
+void UNiagaraDataInterfaceVelocityGrid::BuildVelocityField(FVectorVMExternalFunctionContext& Context)
 {
 	// @todo : implement function for cpu 
 }
 
-void UNiagaraDataInterfaceVelocityGrid::SampleVelocityField(FVectorVMContext& Context)
+void UNiagaraDataInterfaceVelocityGrid::SampleVelocityField(FVectorVMExternalFunctionContext& Context)
 {
 	// @todo : implement function for cpu 
 }
 
-void UNiagaraDataInterfaceVelocityGrid::ComputeGridSize(FVectorVMContext& Context)
+void UNiagaraDataInterfaceVelocityGrid::ComputeGridSize(FVectorVMExternalFunctionContext& Context)
 {
 	// @todo : implement function for cpu
 }
 
-void UNiagaraDataInterfaceVelocityGrid::SetGridDimension(FVectorVMContext& Context)
+void UNiagaraDataInterfaceVelocityGrid::SetGridDimension(FVectorVMExternalFunctionContext& Context)
 {
 	VectorVM::FUserPtrHandler<FNDIVelocityGridData> InstData(Context);
 	VectorVM::FExternalFuncInputHandler<float> GridDimensionX(Context);
@@ -459,7 +459,7 @@ void UNiagaraDataInterfaceVelocityGrid::SetGridDimension(FVectorVMContext& Conte
 
 	VectorVM::FExternalFuncRegisterHandler<bool> OutFunctionStatus(Context);
 
-	for (int32 InstanceIdx = 0; InstanceIdx < Context.NumInstances; ++InstanceIdx)
+	for (int32 InstanceIdx = 0; InstanceIdx < Context.GetNumInstances(); ++InstanceIdx)
 	{
 		FIntVector GridDimension;
 		GridDimension.X = *GridDimensionX.GetDestAndAdvance();
@@ -473,7 +473,7 @@ void UNiagaraDataInterfaceVelocityGrid::SetGridDimension(FVectorVMContext& Conte
 	}
 }
 
-void UNiagaraDataInterfaceVelocityGrid::UpdateGridTransform(FVectorVMContext& Context)
+void UNiagaraDataInterfaceVelocityGrid::UpdateGridTransform(FVectorVMExternalFunctionContext& Context)
 {
 	VectorVM::FUserPtrHandler<FNDIVelocityGridData> InstData(Context);
 
@@ -499,7 +499,7 @@ void UNiagaraDataInterfaceVelocityGrid::UpdateGridTransform(FVectorVMContext& Co
 
 	VectorVM::FExternalFuncRegisterHandler<bool> OutTransformStatus(Context);
 
-	for (int32 InstanceIdx = 0; InstanceIdx < Context.NumInstances; ++InstanceIdx)
+	for (int32 InstanceIdx = 0; InstanceIdx < Context.GetNumInstances(); ++InstanceIdx)
 	{
 		FMatrix Transform;
 		Transform.M[0][0] = *Out00.GetDestAndAdvance();
diff --git a/Engine/Plugins/Runtime/HairStrands/Source/HairStrandsCore/Public/Niagara/NiagaraDataInterfaceHairStrands.h b/Engine/Plugins/Runtime/HairStrands/Source/HairStrandsCore/Public/Niagara/NiagaraDataInterfaceHairStrands.h
--- a/Engine/Plugins/Runtime/HairStrands/Source/HairStrandsCore/Public/Niagara/NiagaraDataInterfaceHairStrands.h
+++ b/Engine/Plugins/Runtime/HairStrands/Source/HairStrandsCore/Public/Niagara/NiagaraDataInterfaceHairStrands.h
@@ -382,220 +382,220 @@ public:
 		FTransform& OutLocalToWorld);
 
 	/** Get the number of strands */
-	void GetNumStrands(FVectorVMContext& Context);
+	void GetNumStrands(FVectorVMExternalFunctionContext& Context);
 
 	/** Get the groom asset datas  */
-	void GetStrandSize(FVectorVMContext& Context);
+	void GetStrandSize(FVectorVMExternalFunctionContext& Context);
 
-	void GetSubSteps(FVectorVMContext& Context);
+	void GetSubSteps(FVectorVMExternalFunctionContext& Context);
 
-	void GetIterationCount(FVectorVMContext& Context);
+	void GetIterationCount(FVectorVMExternalFunctionContext& Context);
 
-	void GetGravityVector(FVectorVMContext& Context);
+	void GetGravityVector(FVectorVMExternalFunctionContext& Context);
 
-	void GetAirDrag(FVectorVMContext& Context);
+	void GetAirDrag(FVectorVMExternalFunctionContext& Context);
 
-	void GetAirVelocity(FVectorVMContext& Context);
+	void GetAirVelocity(FVectorVMExternalFunctionContext& Context);
 
-	void GetSolveBend(FVectorVMContext& Context);
+	void GetSolveBend(FVectorVMExternalFunctionContext& Context);
 
-	void GetProjectBend(FVectorVMContext& Context);
+	void GetProjectBend(FVectorVMExternalFunctionContext& Context);
 
-	void GetBendDamping(FVectorVMContext& Context);
+	void GetBendDamping(FVectorVMExternalFunctionContext& Context);
 
-	void GetBendStiffness(FVectorVMContext& Context);
+	void GetBendStiffness(FVectorVMExternalFunctionContext& Context);
 
-	void GetBendScale(FVectorVMContext& Context);
+	void GetBendScale(FVectorVMExternalFunctionContext& Context);
 
-	void GetSolveStretch(FVectorVMContext& Context);
+	void GetSolveStretch(FVectorVMExternalFunctionContext& Context);
 
-	void GetProjectStretch(FVectorVMContext& Context);
+	void GetProjectStretch(FVectorVMExternalFunctionContext& Context);
 
-	void GetStretchDamping(FVectorVMContext& Context);
+	void GetStretchDamping(FVectorVMExternalFunctionContext& Context);
 
-	void GetStretchStiffness(FVectorVMContext& Context);
+	void GetStretchStiffness(FVectorVMExternalFunctionContext& Context);
 
-	void GetStretchScale(FVectorVMContext& Context);
+	void GetStretchScale(FVectorVMExternalFunctionContext& Context);
 
-	void GetSolveCollision(FVectorVMContext& Context);
+	void GetSolveCollision(FVectorVMExternalFunctionContext& Context);
 
-	void GetProjectCollision(FVectorVMContext& Context);
+	void GetProjectCollision(FVectorVMExternalFunctionContext& Context);
 
-	void GetStaticFriction(FVectorVMContext& Context);
+	void GetStaticFriction(FVectorVMExternalFunctionContext& Context);
 
-	void GetKineticFriction(FVectorVMContext& Context);
+	void GetKineticFriction(FVectorVMExternalFunctionContext& Context);
 
-	void GetStrandsViscosity(FVectorVMContext& Context);
+	void GetStrandsViscosity(FVectorVMExternalFunctionContext& Context);
 
-	void GetGridDimension(FVectorVMContext& Context);
+	void GetGridDimension(FVectorVMExternalFunctionContext& Context);
 
-	void GetCollisionRadius(FVectorVMContext& Context);
+	void GetCollisionRadius(FVectorVMExternalFunctionContext& Context);
 
-	void GetRadiusScale(FVectorVMContext& Context);
+	void GetRadiusScale(FVectorVMExternalFunctionContext& Context);
 
-	void GetStrandsSmoothing(FVectorVMContext& Context);
+	void GetStrandsSmoothing(FVectorVMExternalFunctionContext& Context);
 
-	void GetStrandsDensity(FVectorVMContext& Context);
+	void GetStrandsDensity(FVectorVMExternalFunctionContext& Context);
 
-	void GetStrandsThickness(FVectorVMContext& Context);
+	void GetStrandsThickness(FVectorVMExternalFunctionContext& Context);
 
-	void GetThicknessScale(FVectorVMContext& Context);
+	void GetThicknessScale(FVectorVMExternalFunctionContext& Context);
 
 	/** Get the world transform */
-	void GetWorldTransform(FVectorVMContext& Context);
+	void GetWorldTransform(FVectorVMExternalFunctionContext& Context);
 
 	/** Get the world inverse */
-	void GetWorldInverse(FVectorVMContext& Context);
+	void GetWorldInverse(FVectorVMExternalFunctionContext& Context);
 
 	/** Get the strand vertex position in world space*/
-	void GetPointPosition(FVectorVMContext& Context);
+	void GetPointPosition(FVectorVMExternalFunctionContext& Context);
 
 	/** Get the strand node position in world space*/
-	void ComputeNodePosition(FVectorVMContext& Context);
+	void ComputeNodePosition(FVectorVMExternalFunctionContext& Context);
 
 	/** Get the strand node orientation in world space*/
-	void ComputeNodeOrientation(FVectorVMContext& Context);
+	void ComputeNodeOrientation(FVectorVMExternalFunctionContext& Context);
 
 	/** Get the strand node mass */
-	void ComputeNodeMass(FVectorVMContext& Context);
+	void ComputeNodeMass(FVectorVMExternalFunctionContext& Context);
 
 	/** Get the strand node inertia */
-	void ComputeNodeInertia(FVectorVMContext& Context);
+	void ComputeNodeInertia(FVectorVMExternalFunctionContext& Context);
 
 	/** Compute the edge length (diff between 2 nodes positions)*/
-	void ComputeEdgeLength(FVectorVMContext& Context);
+	void ComputeEdgeLength(FVectorVMExternalFunctionContext& Context);
 
 	/** Compute the edge orientation (diff between 2 nodes orientations) */
-	void ComputeEdgeRotation(FVectorVMContext& Context);
+	void ComputeEdgeRotation(FVectorVMExternalFunctionContext& Context);
 
 	/** Compute the rest local position */
-	void ComputeRestPosition(FVectorVMContext& Context);
+	void ComputeRestPosition(FVectorVMExternalFunctionContext& Context);
 
 	/** Compute the rest local orientation */
-	void ComputeRestOrientation(FVectorVMContext& Context);
+	void ComputeRestOrientation(FVectorVMExternalFunctionContext& Context);
 
 	/** Update the root node orientation based on the current transform */
-	void AttachNodePosition(FVectorVMContext& Context);
+	void AttachNodePosition(FVectorVMExternalFunctionContext& Context);
 
 	/** Update the root node position based on the current transform */
-	void AttachNodeOrientation(FVectorVMContext& Context);
+	void AttachNodeOrientation(FVectorVMExternalFunctionContext& Context);
 
 	/** Report the node displacement onto the points position*/
-	void UpdatePointPosition(FVectorVMContext& Context);
+	void UpdatePointPosition(FVectorVMExternalFunctionContext& Context);
 
 	/** Reset the point position to be the rest one */
-	void ResetPointPosition(FVectorVMContext& Context);
+	void ResetPointPosition(FVectorVMExternalFunctionContext& Context);
 
 	/** Add external force to the linear velocity and advect node position */
-	void AdvectNodePosition(FVectorVMContext& Context);
+	void AdvectNodePosition(FVectorVMExternalFunctionContext& Context);
 
 	/** Add external torque to the angular velocity and advect node orientation*/
-	void AdvectNodeOrientation(FVectorVMContext& Context);
+	void AdvectNodeOrientation(FVectorVMExternalFunctionContext& Context);
 
 	/** Update the node linear velocity based on the node position difference */
-	void UpdateLinearVelocity(FVectorVMContext& Context);
+	void UpdateLinearVelocity(FVectorVMExternalFunctionContext& Context);
 
 	/** Update the node angular velocity based on the node orientation difference */
-	void UpdateAngularVelocity(FVectorVMContext& Context);
+	void UpdateAngularVelocity(FVectorVMExternalFunctionContext& Context);
 
 	/** Get the bounding box center */
-	void GetBoundingBox(FVectorVMContext& Context);
+	void GetBoundingBox(FVectorVMExternalFunctionContext& Context);
 
 	/** Reset the bounding box extent */
-	void ResetBoundingBox(FVectorVMContext& Context);
+	void ResetBoundingBox(FVectorVMExternalFunctionContext& Context);
 
 	/** Build the groom bounding box */
-	void BuildBoundingBox(FVectorVMContext& Context);
+	void BuildBoundingBox(FVectorVMExternalFunctionContext& Context);
 
 	/** Setup the distance spring material */
-	void SetupDistanceSpringMaterial(FVectorVMContext& Context);
+	void SetupDistanceSpringMaterial(FVectorVMExternalFunctionContext& Context);
 
 	/** Solve the distance spring material */
-	void SolveDistanceSpringMaterial(FVectorVMContext& Context);
+	void SolveDistanceSpringMaterial(FVectorVMExternalFunctionContext& Context);
 
 	/** Project the distance spring material */
-	void ProjectDistanceSpringMaterial(FVectorVMContext& Context);
+	void ProjectDistanceSpringMaterial(FVectorVMExternalFunctionContext& Context);
 
 	/** Setup the angular spring material */
-	void SetupAngularSpringMaterial(FVectorVMContext& Context);
+	void SetupAngularSpringMaterial(FVectorVMExternalFunctionContext& Context);
 
 	/** Solve the angular spring material */
-	void SolveAngularSpringMaterial(FVectorVMContext& Context);
+	void SolveAngularSpringMaterial(FVectorVMExternalFunctionContext& Context);
 
 	/** Project the angular spring material */
-	void ProjectAngularSpringMaterial(FVectorVMContext& Context);
+	void ProjectAngularSpringMaterial(FVectorVMExternalFunctionContext& Context);
 
 	/** Setup the stretch rod material */
-	void SetupStretchRodMaterial(FVectorVMContext& Context);
+	void SetupStretchRodMaterial(FVectorVMExternalFunctionContext& Context);
 
 	/** Solve the stretch rod material */
-	void SolveStretchRodMaterial(FVectorVMContext& Context);
+	void SolveStretchRodMaterial(FVectorVMExternalFunctionContext& Context);
 
 	/** Project the stretch rod material */
-	void ProjectStretchRodMaterial(FVectorVMContext& Context);
+	void ProjectStretchRodMaterial(FVectorVMExternalFunctionContext& Context);
 
 	/** Setup the bend rod material */
-	void SetupBendRodMaterial(FVectorVMContext& Context);
+	void SetupBendRodMaterial(FVectorVMExternalFunctionContext& Context);
 
 	/** Solve the bend rod material */
-	void SolveBendRodMaterial(FVectorVMContext& Context);
+	void SolveBendRodMaterial(FVectorVMExternalFunctionContext& Context);
 
 	/** Project the bend rod material */
-	void ProjectBendRodMaterial(FVectorVMContext& Context);
+	void ProjectBendRodMaterial(FVectorVMExternalFunctionContext& Context);
 
 	/** Solve the static collision constraint */
-	void SolveHardCollisionConstraint(FVectorVMContext& Context);
+	void SolveHardCollisionConstraint(FVectorVMExternalFunctionContext& Context);
 
 	/** Project the static collision constraint */
-	void ProjectHardCollisionConstraint(FVectorVMContext& Context);
+	void ProjectHardCollisionConstraint(FVectorVMExternalFunctionContext& Context);
 
 	/** Solve the soft collision constraint */
-	void SolveSoftCollisionConstraint(FVectorVMContext& Context);
+	void SolveSoftCollisionConstraint(FVectorVMExternalFunctionContext& Context);
 
 	/** Project the soft collision constraint */
-	void ProjectSoftCollisionConstraint(FVectorVMContext& Context);
+	void ProjectSoftCollisionConstraint(FVectorVMExternalFunctionContext& Context);
 
 	/** Setup the soft collision constraint */
-	void SetupSoftCollisionConstraint(FVectorVMContext& Context);
+	void SetupSoftCollisionConstraint(FVectorVMExternalFunctionContext& Context);
 
 	/** Compute the rest direction*/
-	void ComputeEdgeDirection(FVectorVMContext& Context);
+	void ComputeEdgeDirection(FVectorVMExternalFunctionContext& Context);
 
 	/** Update the strands material frame */
-	void UpdateMaterialFrame(FVectorVMContext& Context);
+	void UpdateMaterialFrame(FVectorVMExternalFunctionContext& Context);
 
 	/** Compute the strands material frame */
-	void ComputeMaterialFrame(FVectorVMContext& Context);
+	void ComputeMaterialFrame(FVectorVMExternalFunctionContext& Context);
 
 	/** Compute the air drag force */
-	void ComputeAirDragForce(FVectorVMContext& Context);
+	void ComputeAirDragForce(FVectorVMExternalFunctionContext& Context);
 
 	/** Get the rest position and orientation relative to the transform or to the skin cache */
-	void ComputeLocalState(FVectorVMContext& Context);
+	void ComputeLocalState(FVectorVMExternalFunctionContext& Context);
 
 	/** Attach the node position and orientation to the transform or to the skin cache */
-	void AttachNodeState(FVectorVMContext& Context);
+	void AttachNodeState(FVectorVMExternalFunctionContext& Context);
 
 	/** Update the node position and orientation based on rbf transfer */
-	void UpdateNodeState(FVectorVMContext& Context);
+	void UpdateNodeState(FVectorVMExternalFunctionContext& Context);
 
 	/** Check if we need or not a simulation reset*/
-	void NeedSimulationReset(FVectorVMContext& Context);
+	void NeedSimulationReset(FVectorVMExternalFunctionContext& Context);
 
 	/** Check if we have a global interpolation */
-	void HasGlobalInterpolation(FVectorVMContext& Context);
+	void HasGlobalInterpolation(FVectorVMExternalFunctionContext& Context);
 
 	/** Check if we need a rest pose update */
-	void NeedRestUpdate(FVectorVMContext& Context);
+	void NeedRestUpdate(FVectorVMExternalFunctionContext& Context);
 
 	/** Eval the skinned position given a rest position*/
-	void EvalSkinnedPosition(FVectorVMContext& Context);
+	void EvalSkinnedPosition(FVectorVMExternalFunctionContext& Context);
 
 	/** Init the samples along the strands that will be used to transfer informations to the grid */
-	void InitGridSamples(FVectorVMContext& Context);
+	void InitGridSamples(FVectorVMExternalFunctionContext& Context);
 
 	/** Get the sample state given an index */
-	void GetSampleState(FVectorVMContext& Context);
+	void GetSampleState(FVectorVMExternalFunctionContext& Context);
 
 	/** Name of the world transform */
 	static const FString WorldTransformName;
diff --git a/Engine/Plugins/Runtime/HairStrands/Source/HairStrandsCore/Public/Niagara/NiagaraDataInterfacePhysicsAsset.h b/Engine/Plugins/Runtime/HairStrands/Source/HairStrandsCore/Public/Niagara/NiagaraDataInterfacePhysicsAsset.h
--- a/Engine/Plugins/Runtime/HairStrands/Source/HairStrandsCore/Public/Niagara/NiagaraDataInterfacePhysicsAsset.h
+++ b/Engine/Plugins/Runtime/HairStrands/Source/HairStrandsCore/Public/Niagara/NiagaraDataInterfacePhysicsAsset.h
@@ -148,34 +148,34 @@ public:
 	void ExtractSourceComponent(FNiagaraSystemInstance* SystemInstance, FTransform& LocalTransform);
 
 	/** Get the number of boxes*/
-	void GetNumBoxes(FVectorVMContext& Context);
+	void GetNumBoxes(FVectorVMExternalFunctionContext& Context);
 
 	/** Get the number of spheres  */
-	void GetNumSpheres(FVectorVMContext& Context);
+	void GetNumSpheres(FVectorVMExternalFunctionContext& Context);
 
 	/** Get the number of capsules */
-	void GetNumCapsules(FVectorVMContext& Context);
+	void GetNumCapsules(FVectorVMExternalFunctionContext& Context);
 
 	/** Get the element point */
-	void GetElementPoint(FVectorVMContext& Context);
+	void GetElementPoint(FVectorVMExternalFunctionContext& Context);
 
 	/** Get the element distance */
-	void GetElementDistance(FVectorVMContext& Context);
+	void GetElementDistance(FVectorVMExternalFunctionContext& Context);
 
 	/** Get the closest element */
-	void GetClosestElement(FVectorVMContext& Context);
+	void GetClosestElement(FVectorVMExternalFunctionContext& Context);
 
 	/** Get the closest point */
-	void GetClosestPoint(FVectorVMContext& Context);
+	void GetClosestPoint(FVectorVMExternalFunctionContext& Context);
 
 	/** Get the closest distance */
-	void GetClosestDistance(FVectorVMContext& Context);
+	void GetClosestDistance(FVectorVMExternalFunctionContext& Context);
 
 	/** Get the closest texture point */
-	void GetTexturePoint(FVectorVMContext& Context);
+	void GetTexturePoint(FVectorVMExternalFunctionContext& Context);
 
 	/** Get the projection point */
-	void GetProjectionPoint(FVectorVMContext& Context);
+	void GetProjectionPoint(FVectorVMExternalFunctionContext& Context);
 
 	/** Name of element offsets */
 	static const FString ElementOffsetsName;
diff --git a/Engine/Plugins/Runtime/HairStrands/Source/HairStrandsCore/Public/Niagara/NiagaraDataInterfacePressureGrid.h b/Engine/Plugins/Runtime/HairStrands/Source/HairStrandsCore/Public/Niagara/NiagaraDataInterfacePressureGrid.h
--- a/Engine/Plugins/Runtime/HairStrands/Source/HairStrandsCore/Public/Niagara/NiagaraDataInterfacePressureGrid.h
+++ b/Engine/Plugins/Runtime/HairStrands/Source/HairStrandsCore/Public/Niagara/NiagaraDataInterfacePressureGrid.h
@@ -33,31 +33,31 @@ public:
 #endif
 
 	/** Build the velocity field */
-	void BuildDistanceField(FVectorVMContext& Context);
+	void BuildDistanceField(FVectorVMExternalFunctionContext& Context);
 
 	/** Project the velocity field to be divergence free */
-	void SolveGridPressure(FVectorVMContext& Context);
+	void SolveGridPressure(FVectorVMExternalFunctionContext& Context);
 
 	/** Scale Cell Fields */
-	void ScaleCellFields(FVectorVMContext& Context);
+	void ScaleCellFields(FVectorVMExternalFunctionContext& Context);
 
 	/** Set the solid boundary */
-	void SetSolidBoundary(FVectorVMContext& Context);
+	void SetSolidBoundary(FVectorVMExternalFunctionContext& Context);
 
 	/** Compute the solid weights */
-	void ComputeBoundaryWeights(FVectorVMContext& Context);
+	void ComputeBoundaryWeights(FVectorVMExternalFunctionContext& Context);
 
 	/** Get Node Position */
-	void GetNodePosition(FVectorVMContext& Context);
+	void GetNodePosition(FVectorVMExternalFunctionContext& Context);
 
 	/** Get Density Field */
-	void GetDensityField(FVectorVMContext& Context);
+	void GetDensityField(FVectorVMExternalFunctionContext& Context);
 
 	/** Build the Density Field */
-	void BuildDensityField(FVectorVMContext& Context);
+	void BuildDensityField(FVectorVMExternalFunctionContext& Context);
 
 	/** Update the deformation gradient */
-	void UpdateDeformationGradient(FVectorVMContext& Context);
+	void UpdateDeformationGradient(FVectorVMExternalFunctionContext& Context);
 
 };
 
diff --git a/Engine/Plugins/Runtime/HairStrands/Source/HairStrandsCore/Public/Niagara/NiagaraDataInterfaceVelocityGrid.h b/Engine/Plugins/Runtime/HairStrands/Source/HairStrandsCore/Public/Niagara/NiagaraDataInterfaceVelocityGrid.h
--- a/Engine/Plugins/Runtime/HairStrands/Source/HairStrandsCore/Public/Niagara/NiagaraDataInterfaceVelocityGrid.h
+++ b/Engine/Plugins/Runtime/HairStrands/Source/HairStrandsCore/Public/Niagara/NiagaraDataInterfaceVelocityGrid.h
@@ -144,19 +144,19 @@ public:
 	virtual void ProvidePerInstanceDataForRenderThread(void* DataForRenderThread, void* PerInstanceData, const FNiagaraSystemInstanceID& SystemInstance) override;
 
 	/** Build the velocity field */
-	void BuildVelocityField(FVectorVMContext& Context);
+	void BuildVelocityField(FVectorVMExternalFunctionContext& Context);
 
 	/** Sample the grid */
-	void SampleVelocityField(FVectorVMContext& Context);
+	void SampleVelocityField(FVectorVMExternalFunctionContext& Context);
 
 	/** Compute the grid Size (Origin and length) */
-	void ComputeGridSize(FVectorVMContext& Context);
+	void ComputeGridSize(FVectorVMExternalFunctionContext& Context);
 
 	/** Update the grid transform */
-	void UpdateGridTransform(FVectorVMContext& Context);
+	void UpdateGridTransform(FVectorVMExternalFunctionContext& Context);
 
 	/** Set the grid dimension */
-	void SetGridDimension(FVectorVMContext& Context);
+	void SetGridDimension(FVectorVMExternalFunctionContext& Context);
 
 	/** Name of the grid current buffer */
 	static const FString GridCurrentBufferName;
diff --git a/Engine/Source/Runtime/VectorVM/Private/VectorVM.cpp b/Engine/Source/Runtime/VectorVM/Private/VectorVM.cpp
--- a/Engine/Source/Runtime/VectorVM/Private/VectorVM.cpp
+++ b/Engine/Source/Runtime/VectorVM/Private/VectorVM.cpp
@@ -2046,14 +2046,19 @@ struct FKernelExternalFunctionCall
 
 	static void Exec(FVectorVMContext& Context)
 	{
+#ifdef NIAGARA_EXP_VM
+		check(false);
+#else
 		const uint32 ExternalFuncIdx = Context.DecodeU8();
 		const FVMExternalFunction* ExternalFunction = Context.ExternalFunctionTable[ExternalFuncIdx];
 		check(ExternalFunction);
 
 		if (ExternalFunction)
 		{
-			ExternalFunction->Execute(Context);
+			FVectorVMExternalFunctionContext DataInterfaceFunctionContext(&Context);
+			ExternalFunction->Execute(DataInterfaceFunctionContext);
 		}
+#endif
 	}
 };
 
diff --git a/Engine/Source/Runtime/VectorVM/Public/VectorVM.h b/Engine/Source/Runtime/VectorVM/Public/VectorVM.h
--- a/Engine/Source/Runtime/VectorVM/Public/VectorVM.h
+++ b/Engine/Source/Runtime/VectorVM/Public/VectorVM.h
@@ -16,7 +16,7 @@
 #define VECTOR_WIDTH_BYTES (16)
 #define VECTOR_WIDTH_FLOATS (4)
 
-DECLARE_DELEGATE_OneParam(FVMExternalFunction, struct FVectorVMContext& /*Context*/);
+DECLARE_DELEGATE_OneParam(FVMExternalFunction, class FVectorVMExternalFunctionContext& /*Context*/);
 
 UENUM()
 enum class EVectorVMBaseTypes : uint8
@@ -478,6 +478,36 @@ public:
 	}
 };
 
+class FVectorVMExternalFunctionContext {
+	friend struct FNiagaraSystemScriptExecutionContext; //@NOTE(smcgrath): required for the PerInstanceFunctionHook() in the non-experimental version of VectorVM
+public:
+	FVectorVMExternalFunctionContext(FVectorVMContext *InVectorVMContext) : VectorVMContext(InVectorVMContext) { }
+	FORCEINLINE int32 GetStartInstance() const { return VectorVMContext->GetStartInstance(); }
+	FORCEINLINE int32 GetNumInstances() const { return VectorVMContext->GetNumInstances(); }
+	FORCEINLINE TArray<int32> &GetRandCounters() { return VectorVMContext->RandCounters; }
+	FORCEINLINE FRandomStream &GetRandStream() { return VectorVMContext->RandStream; }
+
+#ifndef NIAGARA_EXP_VM
+	FORCEINLINE uint8 *RESTRICT GetTempRegister(int32 RegisterIndex) { return VectorVMContext->GetTempRegister(RegisterIndex); }
+	FORCEINLINE uint8 DecodeU8() { return VectorVMContext->DecodeU8(); }
+	FORCEINLINE uint16 DecodeU16() { return VectorVMContext->DecodeU16(); }
+	FORCEINLINE uint32 DecodeU32() { return VectorVMContext->DecodeU32(); }
+	FORCEINLINE uint64 DecodeU64() { return VectorVMContext->DecodeU64(); }
+	template<typename T = uint8>
+	FORCEINLINE const T* GetConstant(int32 TableIndex, int32 TableOffset) const { return VectorVMContext->GetConstant<T>(TableIndex, TableOffset); }
+	template<typename T = uint8>
+	FORCEINLINE const T* GetConstant(int32 Offset) const { return VectorVMContext->GetConstant<T>(Offset); }
+	FORCEINLINE void *GetUserPtrTable(int32 UserPtrIdx) { return VectorVMContext->UserPtrTable[UserPtrIdx]; }
+	FORCEINLINE int32 GetExternalFunctionInstanceOffset() const { return VectorVMContext->ExternalFunctionInstanceOffset; }
+	FORCEINLINE int32 GetNumTempRegisters() const { return VectorVMContext->NumTempRegisters; }
+#endif
+
+	template<uint32 InstancesPerOp>
+	FORCEINLINE int32 GetNumLoops() const { return VectorVMContext->GetNumLoops<InstancesPerOp>(); }
+private:
+	FVectorVMContext *VectorVMContext;
+};
+
 namespace VectorVM
 {
 	/** Get total number of op-codes */
@@ -527,8 +557,11 @@ namespace VectorVM
 	{
 		int32 UserPtrIdx;
 		T* Ptr;
-		FUserPtrHandler(FVectorVMContext& Context)
+		FUserPtrHandler(FVectorVMExternalFunctionContext& Context)
 		{
+#ifdef NIAGARA_EXP_VM
+			check(false);
+#else
 			const uint16 VariableOffset = Context.DecodeU16();
 			check(!(VariableOffset & VVM_EXT_FUNC_INPUT_LOC_BIT));
 
@@ -536,7 +569,8 @@ namespace VectorVM
 			UserPtrIdx = *Context.GetConstant<int32>(ConstantTableOffset);
 			check(UserPtrIdx != INDEX_NONE);
 
-			Ptr = reinterpret_cast<T*>(Context.UserPtrTable[UserPtrIdx]);
+			Ptr = reinterpret_cast<T*>(Context.GetUserPtrTable(UserPtrIdx));
+#endif
 		}
 
 		FORCEINLINE T* Get() { return Ptr; }
@@ -561,13 +595,16 @@ namespace VectorVM
 			, AdvanceOffset(0)
 		{}
 
-		FORCEINLINE FExternalFuncInputHandler(FVectorVMContext& Context)
+		FORCEINLINE FExternalFuncInputHandler(FVectorVMExternalFunctionContext& Context)
 		{
 			Init(Context);
 		}
 
-		void Init(FVectorVMContext& Context)
+		void Init(FVectorVMExternalFunctionContext& Context)
 		{
+#ifdef NIAGARA_EXP_VM
+			check(false);
+#else
 			InputOffset = Context.DecodeU16();
 
 			const int32 Offset = GetOffset();
@@ -575,7 +612,8 @@ namespace VectorVM
 			AdvanceOffset = IsConstant() ? 0 : 1;
 
 			//Hack: Offset into the buffer by the instance offset.
-			InputPtr += Context.ExternalFunctionInstanceOffset * AdvanceOffset;
+			InputPtr += Context.GetExternalFunctionInstanceOffset() * AdvanceOffset;
+#endif
 		}
 
 		FORCEINLINE bool IsConstant()const { return !IsRegister(); }
@@ -608,24 +646,29 @@ namespace VectorVM
 		T Dummy;
 		T* RESTRICT Register;
 	public:
-		FORCEINLINE FExternalFuncRegisterHandler(FVectorVMContext& Context)
+#ifdef NIAGARA_EXP_VM
+		FORCEINLINE FExternalFuncRegisterHandler(FVectorVMExternalFunctionContext& Context) {
+			check(false);
+		}
+#else
+		FORCEINLINE FExternalFuncRegisterHandler(FVectorVMExternalFunctionContext& Context)
 			: RegisterIndex(Context.DecodeU16() & VVM_EXT_FUNC_INPUT_LOC_MASK)
 			, AdvanceOffset(IsValid() ? 1 : 0)
 		{
 			if (IsValid())
 			{
-				checkSlow(RegisterIndex < Context.NumTempRegisters);
+				checkSlow(RegisterIndex < Context.GetNumTempRegisters());
 				Register = (T*)Context.GetTempRegister(RegisterIndex);
 
 				//Hack: Offset into the buffer by the instance offset.
-				Register += Context.ExternalFunctionInstanceOffset * AdvanceOffset;
+				Register += Context.GetExternalFunctionInstanceOffset() * AdvanceOffset;
 			}
 			else
 			{
 				Register = &Dummy;
 			}
 		}
-
+#endif
 		FORCEINLINE bool IsValid() const { return RegisterIndex != (uint16)VVM_EXT_FUNC_INPUT_LOC_MASK; }
 
 		FORCEINLINE const T Get() { return *Register; }
@@ -650,10 +693,16 @@ namespace VectorVM
 	{
 		uint16 ConstantIndex;
 		T Constant;
-		FExternalFuncConstHandler(FVectorVMContext& Context)
+#ifdef NIAGARA_EXP_VM
+		FExternalFuncConstHandler(FVectorVMExternalFunctionContext& Context) {
+			check(false);
+		}
+#else
+		FExternalFuncConstHandler(FVectorVMExternalFunctionContext& Context)
 			: ConstantIndex(Context.DecodeU16() & VVM_EXT_FUNC_INPUT_LOC_MASK)
 			, Constant(*Context.GetConstant<T>(ConstantIndex))
 		{}
+#endif
 		FORCEINLINE const T& Get() { return Constant; }
 		FORCEINLINE const T& GetAndAdvance() { return Constant; }
 		FORCEINLINE void Advance() { }
