Temp fix for crash in InternalAddHairRBFInterpolationPass accessing LOD index -1 (array out of bounds).

diff --git a/Engine/Plugins/Runtime/HairStrands/Source/HairStrandsCore/Private/HairStrandsMeshProjection.cpp b/Engine/Plugins/Runtime/HairStrands/Source/HairStrandsCore/Private/HairStrandsMeshProjection.cpp
--- a/Engine/Plugins/Runtime/HairStrands/Source/HairStrandsCore/Private/HairStrandsMeshProjection.cpp
+++ b/Engine/Plugins/Runtime/HairStrands/Source/HairStrandsCore/Private/HairStrandsMeshProjection.cpp
@@ -1013,7 +1013,7 @@ void InternalAddHairRBFInterpolationPass(
 	FHairStrandsDeformedRootResource* DeformedRootResources)
 {
 	const uint32 VertexCount = RestResources ? RestResources->GetVertexCount() : 0;
-	if (!RestResources || !DeformedResources || !RestRootResources || !DeformedRootResources || VertexCount == 0 || MeshLODIndex >= RestRootResources->LODs.Num())
+	if (!RestResources || !DeformedResources || !RestRootResources || !DeformedRootResources || VertexCount == 0 || MeshLODIndex >= RestRootResources->LODs.Num() || MeshLODIndex < 0)
 	{
 		return;
 	}
