Ensure cards deformation does not run if no skel. mesh data are available during interpolation.

diff --git a/Engine/Plugins/Runtime/HairStrands/Source/HairStrandsCore/Private/HairStrandsRendering.cpp b/Engine/Plugins/Runtime/HairStrands/Source/HairStrandsCore/Private/HairStrandsRendering.cpp
--- a/Engine/Plugins/Runtime/HairStrands/Source/HairStrandsCore/Private/HairStrandsRendering.cpp
+++ b/Engine/Plugins/Runtime/HairStrands/Source/HairStrandsCore/Private/HairStrandsRendering.cpp
@@ -1864,7 +1864,7 @@ void ComputeHairStrandsInterpolation(
 		if (bIsCardsValid)
 		{
 			const bool bValidGuide = Instance->Guides.bIsSimulationEnable || Instance->Guides.bHasGlobalInterpolation;
-			const bool bHasSkinning = Instance->BindingType == EHairBindingType::Skinning;
+			const bool bHasSkinning = Instance->BindingType == EHairBindingType::Skinning && MeshLODIndex >= 0;
 			const bool bNeedDeformation = bValidGuide || bHasSkinning;
 
 			FHairGroupInstance::FCards::FLOD& LOD = Instance->Cards.LODs[HairLODIndex];
