Fix crash during hair interpolation when platform don't support shader draw debug.

This CLs remove unecessary debug mode, and insure the shader draw debug data are used/bound only when necessary. It also fix guide influence debug mode.

diff --git a/Engine/Plugins/Runtime/HairStrands/Source/HairStrandsCore/Private/GroomDebug.cpp b/Engine/Plugins/Runtime/HairStrands/Source/HairStrandsCore/Private/GroomDebug.cpp
--- a/Engine/Plugins/Runtime/HairStrands/Source/HairStrandsCore/Private/GroomDebug.cpp
+++ b/Engine/Plugins/Runtime/HairStrands/Source/HairStrandsCore/Private/GroomDebug.cpp
@@ -671,6 +671,11 @@ static void AddDrawDebugCardsGuidesPass(
 	const bool bDeformed, 
 	const bool bRen)
 {
+	if (!ShaderDrawDebug::IsSupported(View.GetShaderPlatform()))
+	{
+		return;
+	}
+
 	// Force shader debug to be enabled
 	if (!ShaderDrawDebug::IsEnabled())
 	{
diff --git a/Engine/Plugins/Runtime/HairStrands/Source/HairStrandsCore/Private/HairStrandsRendering.cpp b/Engine/Plugins/Runtime/HairStrands/Source/HairStrandsCore/Private/HairStrandsRendering.cpp
--- a/Engine/Plugins/Runtime/HairStrands/Source/HairStrandsCore/Private/HairStrandsRendering.cpp
+++ b/Engine/Plugins/Runtime/HairStrands/Source/HairStrandsCore/Private/HairStrandsRendering.cpp
@@ -61,7 +61,7 @@ namespace HairTransition
 		Flags |= ERDGPassFlags::NeverCull;
 
 		GraphBuilder.AddPass(
-			RDG_EVENT_NAME("HairTransitionToSRVPass(%s)", Flags == ERDGPassFlags::Raster ? TEXT("Raster") : TEXT("Compute")),
+			RDG_EVENT_NAME("HairStrands::TransitionToSRVPass(%s)", Flags == ERDGPassFlags::Raster ? TEXT("Raster") : TEXT("Compute")),
 			Parameters,
 			Flags,
 			[Parameters](FRHICommandList& RHICmdList)
@@ -176,7 +176,7 @@ static void AddTransferPositionPass(
 	TShaderMapRef<FTransferVelocityPassCS> ComputeShader(ShaderMap);
 	FComputeShaderUtils::AddPass(
 		GraphBuilder,
-		RDG_EVENT_NAME("TransferPositionForVelocity"),
+		RDG_EVENT_NAME("HairStrands::TransferPositionForVelocity"),
 		ComputeShader,
 		Parameters,
 		FIntVector(DispatchCountX, DispatchCountY, 1));
@@ -284,7 +284,7 @@ static void AddGroomCacheUpdatePass(
 	TShaderMapRef<FGroomCacheUpdatePassCS> ComputeShader(ShaderMap);
 	FComputeShaderUtils::AddPass(
 		GraphBuilder,
-		RDG_EVENT_NAME("GroomCacheUpdate"),
+		RDG_EVENT_NAME("HairStrands::GroomCacheUpdate"),
 		ComputeShader,
 		Parameters,
 		DispatchCount);
@@ -432,7 +432,7 @@ static void AddDeformSimHairStrandsPass(
 	TShaderMapRef<FDeformGuideCS> ComputeShader(ShaderMap, PermutationVector);
 	FComputeShaderUtils::AddPass(
 		GraphBuilder,
-		RDG_EVENT_NAME("DeformSimHairStrands"),
+		RDG_EVENT_NAME("HairStrands::DeformSimHairStrands"),
 		ComputeShader,
 		Parameters,
 		FIntVector(DispatchCountX, DispatchCountY, 1));
@@ -479,11 +479,11 @@ class FHairInterpolationCS : public FGlobalShader
 	SHADER_USE_PARAMETER_STRUCT(FHairInterpolationCS, FGlobalShader);
 
 	class FGroupSize : SHADER_PERMUTATION_SPARSE_INT("PERMUTATION_GROUP_SIZE", 32, 64);
-	class FDebug : SHADER_PERMUTATION_INT("PERMUTATION_DEBUG", 2);
+	class FDebug : SHADER_PERMUTATION_BOOL("PERMUTATION_DEBUG");
 	class FDynamicGeometry : SHADER_PERMUTATION_INT("PERMUTATION_DYNAMIC_GEOMETRY", 5);
-	class FSimulation : SHADER_PERMUTATION_INT("PERMUTATION_SIMULATION", 2);
-	class FSingleGuide : SHADER_PERMUTATION_INT("PERMUTATION_USE_SINGLE_GUIDE", 2);
-	class FCulling : SHADER_PERMUTATION_INT("PERMUTATION_CULLING", 2);
+	class FSimulation : SHADER_PERMUTATION_BOOL("PERMUTATION_SIMULATION");
+	class FSingleGuide : SHADER_PERMUTATION_BOOL("PERMUTATION_USE_SINGLE_GUIDE");
+	class FCulling : SHADER_PERMUTATION_BOOL("PERMUTATION_CULLING");
 	using FPermutationDomain = TShaderPermutationDomain<FGroupSize, FDebug, FDynamicGeometry, FSimulation, FSingleGuide, FCulling>;
 
 	BEGIN_SHADER_PARAMETER_STRUCT(FParameters, )
@@ -493,7 +493,6 @@ class FHairInterpolationCS : public FGlobalShader
 		SHADER_PARAMETER(uint32, HairDebugMode)
 		SHADER_PARAMETER(FVector3f, InRenderHairPositionOffset)
 		SHADER_PARAMETER(FVector3f, InSimHairPositionOffset)
-		SHADER_PARAMETER(FIntPoint, HairStrandsCullIndex)
 		SHADER_PARAMETER(uint32,  HairStrandsVF_bIsCullingEnable)
 		SHADER_PARAMETER(FMatrix44f, LocalToWorldMatrix)
 
@@ -509,9 +508,9 @@ class FHairInterpolationCS : public FGlobalShader
 		SHADER_PARAMETER_RDG_BUFFER_SRV(Buffer, Interpolation0Buffer)
 		SHADER_PARAMETER_RDG_BUFFER_SRV(Buffer, Interpolation1Buffer)
 
-		SHADER_PARAMETER_RDG_BUFFER_SRV(Buffer, AttributeBuffer)
-		SHADER_PARAMETER_RDG_BUFFER_SRV(Buffer, SimAttributeBuffer)
-		SHADER_PARAMETER_RDG_BUFFER_UAV(RWBuffer, OutRenderAttributeBuffer)
+		SHADER_PARAMETER_RDG_BUFFER_SRV(Buffer<float4>, AttributeBuffer)
+		SHADER_PARAMETER_RDG_BUFFER_SRV(Buffer<float4>, SimAttributeBuffer)
+		SHADER_PARAMETER_RDG_BUFFER_UAV(RWBuffer<float4>, OutRenderAttributeBuffer)
 		SHADER_PARAMETER_RDG_BUFFER_SRV(Buffer<float4>, RestPosition0Buffer)
 		SHADER_PARAMETER_RDG_BUFFER_SRV(Buffer<float4>, RestPosition1Buffer)
 		SHADER_PARAMETER_RDG_BUFFER_SRV(Buffer<float4>, RestPosition2Buffer)
@@ -617,7 +616,6 @@ static void AddHairStrandsInterpolationPass(
 		Parameters->Interpolation1Buffer = Interpolation1Buffer;
 	}
 	Parameters->OutRenderDeformedPositionBuffer = OutRenderPositionBuffer.UAV;
-	Parameters->HairStrandsCullIndex = FIntPoint(-1, -1);
 	Parameters->VertexCount = VertexCount;
 	Parameters->InRenderHairPositionOffset = InRenderHairWorldOffset;
 	Parameters->InSimHairPositionOffset =  InSimHairWorldOffset;
@@ -651,33 +649,20 @@ static void AddHairStrandsInterpolationPass(
 	Parameters->LocalToWorldMatrix = Instance->LocalToWorld.ToMatrixWithScale();
 
 	// Debug rendering
+	// 1: Patch attributes for displaying guides influences onto the render strands
+	// 2: Show the visible cluster
 	Parameters->HairDebugMode = 0;
 	{
-		const FHairCullInfo Info = GetHairStrandsCullInfo();
-		const bool bCullingEnable = Info.CullMode != EHairCullMode::None && bIsVertexToCurveBuffersValid_Ren;
-
+		// Patch Attribute0 (RooUV/NormalizedLength/Seed)
 		if (bPatchedAttributeBuffer && OutRenderAttributeBuffer)
 		{
 			Parameters->HairDebugMode = 1;
 			check(SimAttributeBuffer != nullptr);
 			check(OutRenderAttributeBuffer != nullptr);
 		}
-		else if (GStrandHairInterpolationDebug > 0)
-		{
-			Parameters->HairDebugMode = 4;
-		}
 		else if (GetHairStrandsDebugStrandsMode() == EHairStrandsDebugMode::RenderVisCluster)
-		{
-			Parameters->HairDebugMode = 3;
-		}
-		else if (bCullingEnable)
 		{
 			Parameters->HairDebugMode = 2;
-			
-			if (Info.CullMode == EHairCullMode::Sim)
-				Parameters->HairStrandsCullIndex.Y = Info.ExplicitIndex >= 0 ? Info.ExplicitIndex : FMath::Clamp(uint32(Info.NormalizedIndex * SimRestRootResources->BulkData.RootCount), 0u, SimRestRootResources->BulkData.RootCount - 1);
-			if (Info.CullMode == EHairCullMode::Render)
-				Parameters->HairStrandsCullIndex.X = Info.ExplicitIndex >= 0 ? Info.ExplicitIndex : FMath::Clamp(uint32(Info.NormalizedIndex * RenRestRootResources->BulkData.RootCount), 0u, RenRestRootResources->BulkData.RootCount - 1);
 		}
 
 		if (Parameters->HairDebugMode > 0 && OutRenderAttributeBuffer)
@@ -687,6 +672,14 @@ static void AddHairStrandsInterpolationPass(
 		}
 	}
 
+	// Only needed if DEBUG_ENABLE is manually enabledin HairStrandsInterpolation.usf. This is for manual debug purpose only
+	#if 0
+	if (ShaderDrawData && ShaderDrawData->IsValid())
+	{
+		ShaderDrawDebug::SetParameters(GraphBuilder, *ShaderDrawData, Parameters->ShaderDrawParameters);
+	}
+	#endif
+
 	const bool bSupportDynamicMesh = 
 		bIsVertexToCurveBuffersValid_Ren &&
 		RenRestRootResources &&
@@ -736,11 +729,6 @@ static void AddHairStrandsInterpolationPass(
 		}
 	}
 
-	if (ShaderDrawData)
-	{
-		ShaderDrawDebug::SetParameters(GraphBuilder, *ShaderDrawData, Parameters->ShaderDrawParameters);
-	}
-
 	const bool bHasLocalDeformation = Instance->Guides.bIsSimulationEnable || bSupportGlobalInterpolation;
 	const bool bCullingEnable		= InstanceGeometryType == EHairGeometryType::Strands && CullingData.bCullingResultAvailable;
 	Parameters->HairStrandsVF_bIsCullingEnable = bCullingEnable ? 1 : 0;
@@ -764,11 +752,11 @@ static void AddHairStrandsInterpolationPass(
 
 	FHairInterpolationCS::FPermutationDomain PermutationVector;
 	PermutationVector.Set<FHairInterpolationCS::FGroupSize>(GroupSize);
-	PermutationVector.Set<FHairInterpolationCS::FDebug>(Parameters->HairDebugMode > 0 ? 1 : 0);
+	PermutationVector.Set<FHairInterpolationCS::FDebug>(Parameters->HairDebugMode > 0);
 	PermutationVector.Set<FHairInterpolationCS::FDynamicGeometry>(DynamicGeometryType);
-	PermutationVector.Set<FHairInterpolationCS::FSimulation>(bHasLocalDeformation ? 1 : 0);
-	PermutationVector.Set<FHairInterpolationCS::FSingleGuide>(bUseSingleGuide ? 1 : 0);
-	PermutationVector.Set<FHairInterpolationCS::FCulling>(bCullingEnable ? 1 : 0);
+	PermutationVector.Set<FHairInterpolationCS::FSimulation>(bHasLocalDeformation);
+	PermutationVector.Set<FHairInterpolationCS::FSingleGuide>(bUseSingleGuide);
+	PermutationVector.Set<FHairInterpolationCS::FCulling>(bCullingEnable);
 
 	TShaderMapRef<FHairInterpolationCS> ComputeShader(ShaderMap, PermutationVector);
 
@@ -781,7 +769,7 @@ static void AddHairStrandsInterpolationPass(
 
 		FComputeShaderUtils::AddPass(
 			GraphBuilder,
-			RDG_EVENT_NAME("HairStrandsInterpolation(culling=on)"),
+			RDG_EVENT_NAME("HairStrands::Interpolation(culling=on)"),
 			ComputeShader, 
 			Parameters,
 			CullingData.HairStrandsVF_CullingIndirectBuffer.Buffer,
@@ -791,7 +779,7 @@ static void AddHairStrandsInterpolationPass(
 	{
 		FComputeShaderUtils::AddPass(
 			GraphBuilder,
-			RDG_EVENT_NAME("HairStrandsInterpolation(culling=off)"),
+			RDG_EVENT_NAME("HairStrands::Interpolation(culling=off)"),
 			ComputeShader,
 			Parameters,
 			DispatchCount);
@@ -936,7 +924,7 @@ static void AddHairClusterAABBPass(
 	{
 		FComputeShaderUtils::AddPass(
 			GraphBuilder,
-			RDG_EVENT_NAME("HairStrandsClusterAABB(%s)", TEXT("Cluster,Group"), TEXT("Group Only")),
+			RDG_EVENT_NAME("HairStrands::ClusterAABB(%s)", TEXT("Cluster,Group"), TEXT("Group Only")),
 			ComputeShader,
 			Parameters,
 			Parameters->IndirectArgsBuffer, 0);
@@ -945,7 +933,7 @@ static void AddHairClusterAABBPass(
 	{
 		FComputeShaderUtils::AddPass(
 			GraphBuilder,
-			RDG_EVENT_NAME("HairStrandsClusterAABB(%s)", TEXT("Cluster,Group"), TEXT("Group Only")),
+			RDG_EVENT_NAME("HairStrands::ClusterAABB(%s)", TEXT("Cluster,Group"), TEXT("Group Only")),
 			ComputeShader,
 			Parameters,
 			FIntVector(1,1,1));
@@ -1074,7 +1062,7 @@ static void AddHairCardsDeformationPass(
 	const int32 DispatchCountX = FMath::DivideAndRoundUp(Parameters->CardsVertexCount, GroupSize);
 	FComputeShaderUtils::AddPass(
 		GraphBuilder,
-		RDG_EVENT_NAME("HairCardsDeformation"),
+		RDG_EVENT_NAME("HairStrands::CardsDeformation"),
 		ComputeShader,
 		Parameters,
 		FIntVector(DispatchCountX,1,1));
@@ -1144,7 +1132,7 @@ void AddHairTangentPass(
 		
 		FComputeShaderUtils::AddPass(
 			GraphBuilder,
-			RDG_EVENT_NAME("HairStrandsTangent(culling=on)"),
+			RDG_EVENT_NAME("HairStrands::Tangent(culling=on)"),
 			ComputeShader,
 			Parameters,
 			DrawIndirectRasterComputeBuffer.Buffer, 0);
@@ -1153,7 +1141,7 @@ void AddHairTangentPass(
 	{
 		FComputeShaderUtils::AddPass(
 			GraphBuilder,
-			RDG_EVENT_NAME("HairStrandsTangent(culling=off)"),
+			RDG_EVENT_NAME("HairStrands::Tangent(culling=off)"),
 			ComputeShader,
 			Parameters,
 			DispatchCount);
@@ -1237,7 +1225,7 @@ static void AddGenerateRaytracingGeometryPass(
 	TShaderMapRef<FHairRaytracingGeometryCS> ComputeShader(ShaderMap, PermutationVector);
 	FComputeShaderUtils::AddPass(
 		GraphBuilder,
-		RDG_EVENT_NAME("HairStrandsRaytracingGeometry"),
+		RDG_EVENT_NAME("HairStrands::GenerateRaytracingGeometry"),
 		ComputeShader,
 		Parameters,
 		DispatchCount);
@@ -1293,7 +1281,7 @@ static void AddClearClusterAABBPass(
 
 	FComputeShaderUtils::AddPass(
 		GraphBuilder,
-		RDG_EVENT_NAME("HairStrandsClearClusterAABB(%s)", UpdateType == EHairAABBUpdateType::UpdateClusterAABB ? TEXT("Cluster,Group") : TEXT("Group Only")),
+		RDG_EVENT_NAME("HairStrands::ClearClusterAABB(%s)", UpdateType == EHairAABBUpdateType::UpdateClusterAABB ? TEXT("Cluster,Group") : TEXT("Group Only")),
 		ComputeShader,
 		Parameters,
 		DispatchCount);
@@ -1509,7 +1497,7 @@ static FHairGroupPublicData::FVertexFactoryInput InternalComputeHairStrandsVerte
 		CONVERT_HAIRSSTRANDS_VF_PARAMETERS(OutVFInput.Strands.PositionBuffer,			Instance->Strands.DeformedResource->GetBuffer(FHairStrandsDeformedResource::EFrameType::Current));
 		CONVERT_HAIRSSTRANDS_VF_PARAMETERS(OutVFInput.Strands.PrevPositionBuffer,		Instance->Strands.DeformedResource->GetBuffer(FHairStrandsDeformedResource::EFrameType::Previous));
 		CONVERT_HAIRSSTRANDS_VF_PARAMETERS(OutVFInput.Strands.TangentBuffer,			Instance->Strands.DeformedResource->TangentBuffer);
-		CONVERT_HAIRSSTRANDS_VF_PARAMETERS(OutVFInput.Strands.Attribute0Buffer,			/*bDebugModePatchedAttributeBuffer ? Instance->Strands.DebugAttributeBuffer : */Instance->Strands.RestResource->Attribute0Buffer); // TODO
+		CONVERT_HAIRSSTRANDS_VF_PARAMETERS(OutVFInput.Strands.Attribute0Buffer,			bDebugModePatchedAttributeBuffer ? Instance->Strands.DebugAttributeBuffer : Instance->Strands.RestResource->Attribute0Buffer); // TODO
 		CONVERT_HAIRSSTRANDS_VF_PARAMETERS(OutVFInput.Strands.Attribute1Buffer,			Instance->Strands.RestResource->Attribute1Buffer);
 		CONVERT_HAIRSSTRANDS_VF_PARAMETERS(OutVFInput.Strands.MaterialBuffer,			Instance->Strands.RestResource->MaterialBuffer);
 		CONVERT_HAIRSSTRANDS_VF_PARAMETERS(OutVFInput.Strands.PositionOffsetBuffer,		Instance->Strands.DeformedResource->GetPositionOffsetBuffer(FHairStrandsDeformedResource::EFrameType::Current));
@@ -1829,7 +1817,7 @@ void ComputeHairStrandsInterpolation(
 						Raytracing_PositionBuffer.UAV);
 
 					GraphBuilder.AddPass(
-						RDG_EVENT_NAME("HairUpdateBLAS(Strands)"),
+						RDG_EVENT_NAME("HairStrands::UpdateBLAS(Strands)"),
 						ERDGPassFlags::NeverCull,
 					[Instance, bNeedUpdate](FRHICommandList& RHICmdList)
 					{
@@ -1955,7 +1943,7 @@ void ComputeHairStrandsInterpolation(
 				if (bNeedBuild || bNeedUpdate)
 				{
 					GraphBuilder.AddPass(
-						RDG_EVENT_NAME("HairUpdateBLAS(Cards)"),
+						RDG_EVENT_NAME("HairStrands::UpdateBLAS(Cards)"),
 						ERDGPassFlags::NeverCull,
 						[Instance, HairLODIndex, bNeedUpdate](FRHICommandListImmediate& RHICmdList)
 					{
@@ -2018,7 +2006,7 @@ void ComputeHairStrandsInterpolation(
 				if (bNeedBuid || bNeedUpdate)
 				{
 					GraphBuilder.AddPass(
-						RDG_EVENT_NAME("HairUpdateBLAS(Meshes)"),
+						RDG_EVENT_NAME("HairStrands::UpdateBLAS(Meshes)"),
 						ERDGPassFlags::NeverCull,
 						[Instance, HairLODIndex, bNeedUpdate](FRHICommandListImmediate& RHICmdList)
 					{
diff --git a/Engine/Shaders/Private/HairStrands/HairStrandsInterpolation.usf b/Engine/Shaders/Private/HairStrands/HairStrandsInterpolation.usf
--- a/Engine/Shaders/Private/HairStrands/HairStrandsInterpolation.usf
+++ b/Engine/Shaders/Private/HairStrands/HairStrandsInterpolation.usf
@@ -7,9 +7,15 @@
 #include "HairStrandsClusterCommon.ush"
 #include "HairStrandsVertexFactoryCommon.ush"
 
+#if COMPILER_METAL || COMPILER_VULKAN
+#define ALLOW_DEBUG_RENDERING 0
+#else
+#define ALLOW_DEBUG_RENDERING 1
+#endif
+
 #define DEBUG_ENABLE 0
 
-#if DEBUG_ENABLE
+#if DEBUG_ENABLE && ALLOW_DEBUG_RENDERING
 #include "../ShaderDrawDebug.ush"
 #include "../ColorMap.ush"
 #endif
@@ -28,22 +34,11 @@
 
 #if SHADER_HAIRINTERPOLATION
 
-#if COMPILER_METAL || COMPILER_VULKAN
-#define ALLOW_DEBUG_RENDERING 0
-#else
-#define ALLOW_DEBUG_RENDERING 1
-#endif
-
-#if PERMUTATION_DEBUG == 1 && ALLOW_DEBUG_RENDERING
-#include "../ShaderDrawDebug.ush" 
-#endif
-
 uint VertexCount;
 uint DispatchCountX;
 uint HairDebugMode;
 float3 InRenderHairPositionOffset;
 float3 InSimHairPositionOffset;
-int2 HairStrandsCullIndex;
 float4x4 LocalToWorldMatrix;
 
 Buffer<uint4>  RenderRestPosePositionBuffer;
@@ -209,7 +204,7 @@ void MainCS(uint2 DispatchThreadId : SV_DispatchThreadID)
 		uint3  GuideVertexIndices = 0;
 		float3 GuideVertexLerps = 0;
 		float3 GuideVertexWeights = 0;
-		#if PERMUTATION_SIMULATION == 1 || PERMUTATION_DEBUG == 1 || PERMUTATION_DEBUG == 2 
+		#if PERMUTATION_SIMULATION == 1 || PERMUTATION_DEBUG == 1
 		{
 
 		  #if PERMUTATION_USE_SINGLE_GUIDE == 1
@@ -320,24 +315,23 @@ void MainCS(uint2 DispatchThreadId : SV_DispatchThreadID)
 
 			WorldVertexPosition = ControlPoint;
 
-			#if PERMUTATION_DEBUG == 1 && ALLOW_DEBUG_RENDERING
+			#if DEBUG_ENABLE
+			// Debug drawing which is manually enabled, then 
 			{
-				const bool bDebugEnabled = HairDebugMode == 4;
 				const float4 RestColor = float4(0, 1, 0, 1);
 				const float4 DeformedColor = float4(1, 0, 0, 1);
-				if (bDebugEnabled && VertexIndex==0)
+				if (VertexIndex==0)
 				{	
 					AddCross(InRenderHairPositionOffset, 1.2, RestColor);
 					AddCross(OutRenHairPositionOffset, 1, DeformedColor);
 				}
 			
-				if (bDebugEnabled && VertexIndex >= 0 && VertexIndex < 64)
+				if (VertexIndex >= 0 && VertexIndex < 64)
 				{
 					AddReferential(RestControlPoint, float3(0, 0, 1));
 					AddReferential(WorldVertexPosition, float3(0, 0, 1));
 				}
 
-				if (bDebugEnabled)
 				{
 					const uint ControlPointType = asuint(PackedRestPosition.w) & 0x3;
 					const bool bIsRoot = ControlPointType == 1;
@@ -503,39 +497,7 @@ void MainCS(uint2 DispatchThreadId : SV_DispatchThreadID)
 			OutRenderAttributeBuffer[VertexIndex] = GuideAttribute;
 		}
 
-		// Cull hair strands or clusters which are not matching the provided index
 		if (HairDebugMode == 2)
-		{
-		#if 0
-			bool bIsValid = false;
-			if (HairStrandsCullIndex.y >= 0)
-			{
-
-				for (uint K = 0; K < 3; ++K)
-				{
-					const uint GuideIndex = GuideVertexIndices[K];
-					const uint RootIndex = SimVertexToRootIndexBuffer[GuideIndex];
-					if (uint(HairStrandsCullIndex.y) == RootIndex)
-					{
-						bIsValid = true;
-						break;
-					}
-				}
-			}
-			else if (HairStrandsCullIndex.x >= 0)
-			{
-				const uint RootIndex = RenVertexToRootIndexBuffer[VertexIndex];
-				bIsValid = uint(HairStrandsCullIndex.x) == RootIndex;
-			}
-
-			if (!bIsValid)
-			{
-				OutRenderDeformedPositionBuffer[VertexIndex] = float4(float3(POSITIVE_INFINITY, POSITIVE_INFINITY, POSITIVE_INFINITY), 2);
-			}
-		#endif		
-		}
-
-		if (HairDebugMode == 3)
 		{
 			const uint ClusterId = VertexToClusterIdBuffer[VertexIndex];
 			const uint ClusterIdLo8bits = ClusterId & 0xFF;
diff --git a/Engine/Source/Runtime/Renderer/Private/HairStrands/HairStrandsDebug.cpp b/Engine/Source/Runtime/Renderer/Private/HairStrands/HairStrandsDebug.cpp
--- a/Engine/Source/Runtime/Renderer/Private/HairStrands/HairStrandsDebug.cpp
+++ b/Engine/Source/Runtime/Renderer/Private/HairStrands/HairStrandsDebug.cpp
@@ -65,14 +65,6 @@ static FAutoConsoleVariableRef CVarHairVirtualVoxel_ForceMipLevel(TEXT("r.HairSt
 static int32 GHairVirtualVoxel_DebugTraversalType = 0;
 static FAutoConsoleVariableRef CVarHairVirtualVoxel_DebugTraversalType(TEXT("r.HairStrands.Voxelization.Virtual.DebugTraversalType"), GHairVirtualVoxel_DebugTraversalType, TEXT("Traversal mode (0:linear, 1:mip) for debug voxel visualization."));
 
-static int32 GHairStrandsCull = 0;
-static int32 GHairStrandsCullIndex = -1;
-static int32 GHairStrandsUpdateCullIndex = 0;
-static float GHairStrandsCullNormalizedIndex = -1;
-static FAutoConsoleVariableRef CVarHairStrandsCull			(TEXT("r.HairStrands.Cull"), GHairStrandsCull, TEXT("Cull hair strands (0:disabled, 1: render cull, 2: sim cull)."));
-static FAutoConsoleVariableRef CVarHairStrandsCullIndex		(TEXT("r.HairStrands.Cull.Index"), GHairStrandsCullIndex, TEXT("Hair strands index to be kept. Other will be culled."));
-static FAutoConsoleVariableRef CVarHairStrandsUpdateCullIndex(TEXT("r.HairStrands.Cull.Update"), GHairStrandsUpdateCullIndex, TEXT("Update the guide index to be kept using mouse position for fast selection."));
-
 static bool IsDebugDrawAndDebugPrintEnabled(const FViewInfo& View)
 {
 	return ShaderDrawDebug::IsEnabled(View) && ShaderPrint::IsEnabled(View);
@@ -98,15 +90,6 @@ void SetHairScreenLODInfo(bool bEnable)
 	GHairStrandsClusterDebug = bEnable ? 4 : 1;
 }
 
-FHairCullInfo GetHairStrandsCullInfo()
-{
-	FHairCullInfo Out;
-	Out.CullMode		= GHairStrandsCull == 1 ? EHairCullMode::Render : (GHairStrandsCull == 2 ? EHairCullMode::Sim : EHairCullMode::None);
-	Out.ExplicitIndex	= GHairStrandsCullIndex >= 0 ? GHairStrandsCullIndex : -1;
-	Out.NormalizedIndex = GHairStrandsCullNormalizedIndex;
-	return Out;
-}
-
 FHairStrandsDebugData::Data FHairStrandsDebugData::CreateData(FRDGBuilder& GraphBuilder)
 {
 	FHairStrandsDebugData::Data Out;
@@ -1044,13 +1027,6 @@ static void InternalRenderHairStrandsDebugInfo(
 
 	RDG_EVENT_SCOPE(GraphBuilder, "HairStrandsDebug");
 
-	if (GHairStrandsUpdateCullIndex)
-	{
-		const float TotalPixelCount = View.ViewRect.Width() * View.ViewRect.Height();
-		const float Index = View.CursorPos.X + View.CursorPos.Y * View.ViewRect.Width();
-		GHairStrandsCullNormalizedIndex = Index / TotalPixelCount;
-	}
-
 	// Only render debug information for the main view
 	const FSceneTextures& SceneTextures = FSceneTextures::Get(GraphBuilder);
 
diff --git a/Engine/Source/Runtime/Renderer/Public/HairStrandsInterface.h b/Engine/Source/Runtime/Renderer/Public/HairStrandsInterface.h
--- a/Engine/Source/Runtime/Renderer/Public/HairStrandsInterface.h
+++ b/Engine/Source/Runtime/Renderer/Public/HairStrandsInterface.h
@@ -382,21 +382,6 @@ RENDERER_API bool IsHairStrandsSupported(EHairStrandsShaderType Type, EShaderPla
 RENDERER_API bool IsHairStrandsEnabled(EHairStrandsShaderType Type, EShaderPlatform Platform = EShaderPlatform::SP_NumPlatforms);
 RENDERER_API void SetHairStrandsEnabled(bool In);
 
-// Return strands & guide indices to be preserved, while all others strands/guides should be culled
-enum class EHairCullMode : uint8
-{
-	None,
-	Render,
-	Sim
-};
-struct FHairCullInfo
-{
-	int32 ExplicitIndex = -1; 
-	float NormalizedIndex = 0; // [0,1]
-	EHairCullMode CullMode = EHairCullMode::None;
-};
-RENDERER_API FHairCullInfo GetHairStrandsCullInfo();
-
 RENDERER_API bool IsHairRayTracingEnabled();
 
 // Return true if hair simulation is enabled.
diff --git a/Engine/Source/Runtime/Renderer/Public/ShaderDebug.h b/Engine/Source/Runtime/Renderer/Public/ShaderDebug.h
--- a/Engine/Source/Runtime/Renderer/Public/ShaderDebug.h
+++ b/Engine/Source/Runtime/Renderer/Public/ShaderDebug.h
@@ -13,6 +13,7 @@ struct FShaderDrawDebugData
 	FIntPoint CursorPosition;
 	FRDGBufferRef Buffer = nullptr;
 	FRDGBufferRef IndirectBuffer = nullptr;
+	bool IsValid() const { return Buffer != nullptr && IndirectBuffer != nullptr; }
 };
 
 namespace ShaderDrawDebug 
