Fix groom guide visualization.

diff --git a/Engine/Plugins/Runtime/HairStrands/Source/HairStrandsCore/Private/GroomBuilder.cpp b/Engine/Plugins/Runtime/HairStrands/Source/HairStrandsCore/Private/GroomBuilder.cpp
--- a/Engine/Plugins/Runtime/HairStrands/Source/HairStrandsCore/Private/GroomBuilder.cpp
+++ b/Engine/Plugins/Runtime/HairStrands/Source/HairStrandsCore/Private/GroomBuilder.cpp
@@ -33,7 +33,7 @@ static FAutoConsoleVariableRef CVarHairClusterBuilder_MaxVoxelResolution(TEXT("r
 FString FGroomBuilder::GetVersion()
 {
 	// Important to update the version when groom building changes
-	return TEXT("2n");
+	return TEXT("2o");
 }
 
 namespace FHairStrandsDecimation
@@ -1610,6 +1610,12 @@ bool FGroomBuilder::BuildHairDescriptionGroups(const FHairDescription& HairDescr
 				VertexWidth = StrandWidth;
 			}
 
+			// If the curve is a guide, and its width is 0, force to a constant width to insure the guides can be correctly rendered/displayed for debug purpose
+			if (bIsGuide && VertexWidth == 0.f)
+			{
+				VertexWidth = 0.01f;
+			}
+
 			CurrentHairStrandsDatas->StrandsPoints.PointsRadius.Add(VertexWidth * 0.5f);
 		}
 	}
