Fix crash when creating/updating BLAS data for hair strands/cards/meshes, as the update doesn't run on the rendering thread.

This CL forces the update to happen onto the rendering thread.

diff --git a/Engine/Plugins/Runtime/HairStrands/Source/HairStrandsCore/Private/HairStrandsRendering.cpp b/Engine/Plugins/Runtime/HairStrands/Source/HairStrandsCore/Private/HairStrandsRendering.cpp
--- a/Engine/Plugins/Runtime/HairStrands/Source/HairStrandsCore/Private/HairStrandsRendering.cpp
+++ b/Engine/Plugins/Runtime/HairStrands/Source/HairStrandsCore/Private/HairStrandsRendering.cpp
@@ -1819,7 +1819,7 @@ void ComputeHairStrandsInterpolation(
 					GraphBuilder.AddPass(
 						RDG_EVENT_NAME("HairStrands::UpdateBLAS(Strands)"),
 						ERDGPassFlags::NeverCull,
-					[Instance, bNeedUpdate](FRHICommandList& RHICmdList)
+					[Instance, bNeedUpdate](FRHICommandListImmediate& RHICmdList)
 					{
 						const bool bLocalNeedBuild = !Instance->Strands.RenRaytracingResource->bIsRTGeometryInitialized;
 						if (bLocalNeedBuild)
