Allocate deforme resources without data initialization to insure.

This allows to unload the CPU data once the resources are created.

diff --git a/Engine/Plugins/Runtime/HairStrands/Source/HairStrandsCore/Private/GroomComponent.cpp b/Engine/Plugins/Runtime/HairStrands/Source/HairStrandsCore/Private/GroomComponent.cpp
--- a/Engine/Plugins/Runtime/HairStrands/Source/HairStrandsCore/Private/GroomComponent.cpp
+++ b/Engine/Plugins/Runtime/HairStrands/Source/HairStrandsCore/Private/GroomComponent.cpp
@@ -2311,7 +2311,7 @@ void UGroomComponent::InitResources(bool bIsBindingReloading)
 			check(GroupData.Guides.RestResource);
 
 			// If guides are allocated, deformed resources are always needs since they are either used with simulation, or RBF deformation. Both are dynamics, and require deformed positions
-			HairGroupInstance->Guides.DeformedResource = new FHairStrandsDeformedResource(GroupData.Guides.BulkData, true, EHairStrandsResourcesType::Guides);
+			HairGroupInstance->Guides.DeformedResource = new FHairStrandsDeformedResource(GroupData.Guides.BulkData, EHairStrandsResourcesType::Guides);
 
 			// Initialize the simulation and the global deformation to its default behavior by setting it with LODIndex = -1
 			const int32 LODIndex = -1;
@@ -2408,7 +2408,7 @@ void UGroomComponent::InitResources(bool bIsBindingReloading)
 			HairGroupInstance->Strands.RestResource = GroupData.Strands.RestResource;
 			if (bNeedDynamicResources)
 			{
-				HairGroupInstance->Strands.DeformedResource = new FHairStrandsDeformedResource(GroupData.Strands.BulkData, false, EHairStrandsResourcesType::Strands);
+				HairGroupInstance->Strands.DeformedResource = new FHairStrandsDeformedResource(GroupData.Strands.BulkData, EHairStrandsResourcesType::Strands);
 			} 
 
 			// An empty groom doesn't have a ClusterCullingResource
@@ -2511,7 +2511,7 @@ void UGroomComponent::InitResources(bool bIsBindingReloading)
 
 					InstanceLOD.Guides.RestResource = LOD.Guides.RestResource;
 					{
-						InstanceLOD.Guides.DeformedResource = new FHairStrandsDeformedResource(LOD.Guides.BulkData, false, EHairStrandsResourcesType::Cards);
+						InstanceLOD.Guides.DeformedResource = new FHairStrandsDeformedResource(LOD.Guides.BulkData, EHairStrandsResourcesType::Cards);
 					}
 
 					InstanceLOD.Guides.HairInterpolationType = HairInterpolationType;
diff --git a/Engine/Plugins/Runtime/HairStrands/Source/HairStrandsCore/Private/GroomManager.cpp b/Engine/Plugins/Runtime/HairStrands/Source/HairStrandsCore/Private/GroomManager.cpp
--- a/Engine/Plugins/Runtime/HairStrands/Source/HairStrandsCore/Private/GroomManager.cpp
+++ b/Engine/Plugins/Runtime/HairStrands/Source/HairStrandsCore/Private/GroomManager.cpp
@@ -576,6 +576,18 @@ static void RunHairBufferSwap(const FHairStrandsInstances& Instances, const TArr
 	}
 }
 
+static void AddCopyHairStrandsPositionPass(
+	FRDGBuilder& GraphBuilder,
+	FHairStrandsRestResource& RestResources,
+	FHairStrandsDeformedResource& DeformedResources)
+{
+	FRDGBufferRef RestBuffer = Register(GraphBuilder, RestResources.PositionBuffer, ERDGImportedBufferFlags::None).Buffer;
+	FRDGBufferRef Deformed0Buffer = Register(GraphBuilder, DeformedResources.DeformedPositionBuffer[0], ERDGImportedBufferFlags::None).Buffer;
+	FRDGBufferRef Deformed1Buffer = Register(GraphBuilder, DeformedResources.DeformedPositionBuffer[1], ERDGImportedBufferFlags::None).Buffer;
+	AddCopyBufferPass(GraphBuilder, Deformed0Buffer, RestBuffer);
+	AddCopyBufferPass(GraphBuilder, Deformed1Buffer, RestBuffer);
+}
+
 static void RunHairLODSelection(
 	FRDGBuilder& GraphBuilder, 
 	const FHairStrandsInstances& Instances, 
@@ -686,13 +698,13 @@ static void RunHairLODSelection(
 		}
 
 		// Lazy allocation of resources
-		// Note: Allocation will only be done if the resources it not yet initialized.
+		// Note: Allocation will only be done if the resources is not initialized yet. Guides deformed position are also initialized from the Rest position at creation time.
 		if (Instance->Guides.Data)
 		{
 			if (Instance->Guides.RestRootResource)			{ Instance->Guides.RestRootResource->Allocate(GraphBuilder); Instance->Guides.RestRootResource->AllocateLOD(GraphBuilder, MeshLODIndex); }
 			if (Instance->Guides.RestResource)				{ Instance->Guides.RestResource->Allocate(GraphBuilder); }
 			if (Instance->Guides.DeformedRootResource)		{ Instance->Guides.DeformedRootResource->Allocate(GraphBuilder); Instance->Guides.DeformedRootResource->AllocateLOD(GraphBuilder, MeshLODIndex); }
-			if (Instance->Guides.DeformedResource)			{ Instance->Guides.DeformedResource->Allocate(GraphBuilder); }
+			if (Instance->Guides.DeformedResource)			{ const bool bNeedCopy = !Instance->Guides.DeformedResource->bIsInitialized; Instance->Guides.DeformedResource->Allocate(GraphBuilder); if (bNeedCopy) { AddCopyHairStrandsPositionPass(GraphBuilder, *Instance->Guides.RestResource, *Instance->Guides.DeformedResource); }}
 		}
 
 		if (GeometryType == EHairGeometryType::Meshes)
diff --git a/Engine/Plugins/Runtime/HairStrands/Source/HairStrandsCore/Private/GroomResources.cpp b/Engine/Plugins/Runtime/HairStrands/Source/HairStrandsCore/Private/GroomResources.cpp
--- a/Engine/Plugins/Runtime/HairStrands/Source/HairStrandsCore/Private/GroomResources.cpp
+++ b/Engine/Plugins/Runtime/HairStrands/Source/HairStrandsCore/Private/GroomResources.cpp
@@ -768,9 +768,9 @@ void FHairStrandsRestResource::InternalRelease()
 
 /////////////////////////////////////////////////////////////////////////////////////////
 
-FHairStrandsDeformedResource::FHairStrandsDeformedResource(FHairStrandsBulkData& InBulkData, bool bInInitializedData, EHairStrandsResourcesType InCurveType) :
+FHairStrandsDeformedResource::FHairStrandsDeformedResource(FHairStrandsBulkData& InBulkData, EHairStrandsResourcesType InCurveType) :
 	FHairCommonResource(EHairStrandsAllocationType::Deferred),
-	BulkData(InBulkData), bInitializedData(bInInitializedData), CurveType(InCurveType)
+	BulkData(InBulkData), CurveType(InCurveType)
 {
 	GetPositionOffset(FHairStrandsDeformedResource::EFrameType::Current)  = BulkData.GetPositionOffset();
 	GetPositionOffset(FHairStrandsDeformedResource::EFrameType::Previous) = BulkData.GetPositionOffset();
@@ -780,18 +780,8 @@ void FHairStrandsDeformedResource::InternalAllocate(FRDGBuilder& GraphBuilder)
 {
 	const uint32 PointCount = BulkData.PointCount;
 
-	if (bInitializedData)
-	{
-		// 1. Lock data, which force the loading data from files (on non-editor build/cooked data). These data are then uploaded to the GPU
-		// 2. A local copy is done by the buffer uploader. This copy is discarded once the uploading is done.
-		InternalCreateVertexBufferRDG_FromBulkData<FHairStrandsPositionFormat>(GraphBuilder, BulkData.Positions, PointCount, DeformedPositionBuffer[0], HAIRSTRANDS_RESOUCE_NAME(CurveType, Hair.StrandsDeformed_DeformedPositionBuffer0), EHairResourceUsageType::Dynamic);
-		InternalCreateVertexBufferRDG_FromBulkData<FHairStrandsPositionFormat>(GraphBuilder, BulkData.Positions, PointCount, DeformedPositionBuffer[1], HAIRSTRANDS_RESOUCE_NAME(CurveType, Hair.StrandsDeformed_DeformedPositionBuffer1), EHairResourceUsageType::Dynamic);
-	}
-	else
-	{
-		InternalCreateVertexBufferRDG<FHairStrandsPositionFormat>(GraphBuilder, BulkData.PointCount, DeformedPositionBuffer[0], HAIRSTRANDS_RESOUCE_NAME(CurveType, Hair.StrandsDeformed_DeformedPositionBuffer0), EHairResourceUsageType::Dynamic);
-		InternalCreateVertexBufferRDG<FHairStrandsPositionFormat>(GraphBuilder, BulkData.PointCount, DeformedPositionBuffer[1], HAIRSTRANDS_RESOUCE_NAME(CurveType, Hair.StrandsDeformed_DeformedPositionBuffer1), EHairResourceUsageType::Dynamic);
-	}
+	InternalCreateVertexBufferRDG<FHairStrandsPositionFormat>(GraphBuilder, BulkData.PointCount, DeformedPositionBuffer[0], HAIRSTRANDS_RESOUCE_NAME(CurveType, Hair.StrandsDeformed_DeformedPositionBuffer0), EHairResourceUsageType::Dynamic);
+	InternalCreateVertexBufferRDG<FHairStrandsPositionFormat>(GraphBuilder, BulkData.PointCount, DeformedPositionBuffer[1], HAIRSTRANDS_RESOUCE_NAME(CurveType, Hair.StrandsDeformed_DeformedPositionBuffer1), EHairResourceUsageType::Dynamic);
 	InternalCreateVertexBufferRDG<FHairStrandsTangentFormat>(GraphBuilder, BulkData.PointCount * FHairStrandsTangentFormat::ComponentCount, TangentBuffer, HAIRSTRANDS_RESOUCE_NAME(CurveType, Hair.StrandsDeformed_TangentBuffer), EHairResourceUsageType::Dynamic);
 
 	TArray<FVector4> DefaultOffsets;
diff --git a/Engine/Plugins/Runtime/HairStrands/Source/HairStrandsCore/Public/GroomResources.h b/Engine/Plugins/Runtime/HairStrands/Source/HairStrandsCore/Public/GroomResources.h
--- a/Engine/Plugins/Runtime/HairStrands/Source/HairStrandsCore/Public/GroomResources.h
+++ b/Engine/Plugins/Runtime/HairStrands/Source/HairStrandsCore/Public/GroomResources.h
@@ -252,7 +252,7 @@ struct FHairStrandsRestResource : public FHairCommonResource
 struct FHairStrandsDeformedResource : public FHairCommonResource
 {
 	/** Build the hair strands resource */
-	FHairStrandsDeformedResource(FHairStrandsBulkData& BulkData, bool bInitializeData, EHairStrandsResourcesType CurveType);
+	FHairStrandsDeformedResource(FHairStrandsBulkData& BulkData, EHairStrandsResourcesType CurveType);
 
 	/* Init/Release buffers */
 	virtual void InternalAllocate(FRDGBuilder& GraphBuilder) override;
@@ -286,9 +286,6 @@ struct FHairStrandsDeformedResource : public FHairCommonResource
 	/* Reference to the hair strands render data */
 	FHairStrandsBulkData& BulkData;
 
-	/* Whether the GPU data should be initialized with the asset data or not */
-	const bool bInitializedData = false;
-
 	/* Whether the GPU data should be initialized with the asset data or not */
 	uint32 CurrentIndex = 0;
 
