Fix groom component to respect RenderInMainPass option.

diff --git a/Engine/Plugins/Runtime/HairStrands/Source/HairStrandsCore/Private/GroomComponent.cpp b/Engine/Plugins/Runtime/HairStrands/Source/HairStrandsCore/Private/GroomComponent.cpp
--- a/Engine/Plugins/Runtime/HairStrands/Source/HairStrandsCore/Private/GroomComponent.cpp
+++ b/Engine/Plugins/Runtime/HairStrands/Source/HairStrandsCore/Private/GroomComponent.cpp
@@ -885,12 +885,12 @@ public:
 		Result.bHairStrands = IsShown(View);
 
 		// Special pass for hair strands geometry (not part of the base pass, and shadowing is handlded in a custom fashion). When cards rendering is enabled we reusethe base pass
-		Result.bDrawRelevance		= true;
-		Result.bRenderInMainPass	= bUseCardsOrMesh;
-		Result.bShadowRelevance		= true;
+		Result.bDrawRelevance		= IsShown(View);
+		Result.bRenderInMainPass	= bUseCardsOrMesh && ShouldRenderInMainPass();
+		Result.bShadowRelevance		= IsShadowCast(View);
 		Result.bDynamicRelevance	= true;
 		Result.bRenderCustomDepth	= ShouldRenderCustomDepth();
-		Result.bVelocityRelevance	= bUseCardsOrMesh;
+		Result.bVelocityRelevance	= Result.bRenderInMainPass && bUseCardsOrMesh;
 		Result.bUsesLightingChannels= GetLightingChannelMask() != GetDefaultLightingChannelMask();
 
 		// Selection only
