Add pass for intializing hair strands position.

This is used in place of the generic copy buffer, as currently the copy region buffer RHI function is not implemented on certain platform.

diff --git a/Engine/Plugins/Runtime/HairStrands/Source/HairStrandsCore/Private/GroomManager.cpp b/Engine/Plugins/Runtime/HairStrands/Source/HairStrandsCore/Private/GroomManager.cpp
--- a/Engine/Plugins/Runtime/HairStrands/Source/HairStrandsCore/Private/GroomManager.cpp
+++ b/Engine/Plugins/Runtime/HairStrands/Source/HairStrandsCore/Private/GroomManager.cpp
@@ -576,16 +576,28 @@ static void RunHairBufferSwap(const FHairStrandsInstances& Instances, const TArr
 	}
 }
 
+void AddHairStrandsCopyPositionPass(FRDGBuilder& GraphBuilder, FGlobalShaderMap* ShaderMap, FRDGBufferUAVRef& Out, FRDGBufferSRVRef& In);
+
 static void AddCopyHairStrandsPositionPass(
 	FRDGBuilder& GraphBuilder,
+	FGlobalShaderMap* ShaderMap,
 	FHairStrandsRestResource& RestResources,
 	FHairStrandsDeformedResource& DeformedResources)
 {
+	// Copy using an explicit copy, as AddCopyBufferPass is not implemented on PS5
+#if 0
 	FRDGBufferRef RestBuffer = Register(GraphBuilder, RestResources.PositionBuffer, ERDGImportedBufferFlags::None).Buffer;
 	FRDGBufferRef Deformed0Buffer = Register(GraphBuilder, DeformedResources.DeformedPositionBuffer[0], ERDGImportedBufferFlags::None).Buffer;
 	FRDGBufferRef Deformed1Buffer = Register(GraphBuilder, DeformedResources.DeformedPositionBuffer[1], ERDGImportedBufferFlags::None).Buffer;
 	AddCopyBufferPass(GraphBuilder, Deformed0Buffer, RestBuffer);
 	AddCopyBufferPass(GraphBuilder, Deformed1Buffer, RestBuffer);
+#else
+	FRDGBufferSRVRef RestBuffer = Register(GraphBuilder, RestResources.PositionBuffer, ERDGImportedBufferFlags::CreateSRV).SRV;
+	FRDGBufferUAVRef Deformed0Buffer = Register(GraphBuilder, DeformedResources.DeformedPositionBuffer[0], ERDGImportedBufferFlags::CreateUAV).UAV;
+	FRDGBufferUAVRef Deformed1Buffer = Register(GraphBuilder, DeformedResources.DeformedPositionBuffer[1], ERDGImportedBufferFlags::CreateUAV).UAV;
+	AddHairStrandsCopyPositionPass(GraphBuilder, ShaderMap, Deformed0Buffer, RestBuffer);
+	AddHairStrandsCopyPositionPass(GraphBuilder, ShaderMap, Deformed1Buffer, RestBuffer);
+#endif
 }
 
 static void RunHairLODSelection(
@@ -704,7 +716,7 @@ static void RunHairLODSelection(
 			if (Instance->Guides.RestRootResource)			{ Instance->Guides.RestRootResource->Allocate(GraphBuilder); Instance->Guides.RestRootResource->AllocateLOD(GraphBuilder, MeshLODIndex); }
 			if (Instance->Guides.RestResource)				{ Instance->Guides.RestResource->Allocate(GraphBuilder); }
 			if (Instance->Guides.DeformedRootResource)		{ Instance->Guides.DeformedRootResource->Allocate(GraphBuilder); Instance->Guides.DeformedRootResource->AllocateLOD(GraphBuilder, MeshLODIndex); }
-			if (Instance->Guides.DeformedResource)			{ const bool bNeedCopy = !Instance->Guides.DeformedResource->bIsInitialized; Instance->Guides.DeformedResource->Allocate(GraphBuilder); if (bNeedCopy) { AddCopyHairStrandsPositionPass(GraphBuilder, *Instance->Guides.RestResource, *Instance->Guides.DeformedResource); }}
+			if (Instance->Guides.DeformedResource)			{ const bool bNeedCopy = !Instance->Guides.DeformedResource->bIsInitialized; Instance->Guides.DeformedResource->Allocate(GraphBuilder); if (bNeedCopy) { AddCopyHairStrandsPositionPass(GraphBuilder, ShaderMap, *Instance->Guides.RestResource, *Instance->Guides.DeformedResource); }}
 		}
 
 		if (GeometryType == EHairGeometryType::Meshes)
diff --git a/Engine/Plugins/Runtime/HairStrands/Source/HairStrandsCore/Private/HairStrandsRendering.cpp b/Engine/Plugins/Runtime/HairStrands/Source/HairStrandsCore/Private/HairStrandsRendering.cpp
--- a/Engine/Plugins/Runtime/HairStrands/Source/HairStrandsCore/Private/HairStrandsRendering.cpp
+++ b/Engine/Plugins/Runtime/HairStrands/Source/HairStrandsCore/Private/HairStrandsRendering.cpp
@@ -130,6 +130,58 @@ enum class EDeformationType : uint8
 	OffsetGuide		// Offset the guides
 };
 
+
+///////////////////////////////////////////////////////////////////////////////////////////////////
+class FHairStrandsCopyPositionCS : public FGlobalShader
+{
+	DECLARE_GLOBAL_SHADER(FHairStrandsCopyPositionCS);
+	SHADER_USE_PARAMETER_STRUCT(FHairStrandsCopyPositionCS, FGlobalShader);
+
+	BEGIN_SHADER_PARAMETER_STRUCT(FParameters, )
+		SHADER_PARAMETER(uint32, ElementCount)
+		SHADER_PARAMETER(uint32, DispatchCountX)
+		SHADER_PARAMETER_RDG_BUFFER_SRV(Buffer, InBuffer)
+		SHADER_PARAMETER_RDG_BUFFER_UAV(RWBuffer, OutBuffer)
+	END_SHADER_PARAMETER_STRUCT()
+
+public:
+
+	static bool ShouldCompilePermutation(const FGlobalShaderPermutationParameters& Parameters) { return IsHairStrandsSupported(EHairStrandsShaderType::All, Parameters.Platform); }
+	static void ModifyCompilationEnvironment(const FGlobalShaderPermutationParameters& Parameters, FShaderCompilerEnvironment& OutEnvironment)
+	{
+		FGlobalShader::ModifyCompilationEnvironment(Parameters, OutEnvironment);
+		OutEnvironment.SetDefine(TEXT("SHADER_COPY_POSITION"), 1);
+	}
+};
+
+IMPLEMENT_GLOBAL_SHADER(FHairStrandsCopyPositionCS, "/Engine/Private/HairStrands/HairStrandsInterpolation.usf", "MainCS", SF_Compute);
+
+void AddHairStrandsCopyPositionPass(
+	FRDGBuilder& GraphBuilder,
+	FGlobalShaderMap* ShaderMap,
+	FRDGBufferUAVRef& Out,
+	FRDGBufferSRVRef& In)
+{
+	FHairStrandsCopyPositionCS::FParameters* PassParameters = GraphBuilder.AllocParameters<FHairStrandsCopyPositionCS::FParameters>();
+	PassParameters->InBuffer = In;
+	PassParameters->OutBuffer = Out;
+	FHairStrandsCopyPositionCS::FPermutationDomain PermutationVector;
+	TShaderMapRef<FHairStrandsCopyPositionCS> ComputeShader(ShaderMap, PermutationVector);
+
+	const uint32 GroupSize = 256;
+	const uint32 ElementCount = Out->GetParent()->Desc.NumElements;
+	const FIntVector GroupCount = FComputeShaderUtils::GetGroupCount(ElementCount, GroupSize);
+	check(GroupCount.X < GRHIMaxDispatchThreadGroupsPerDimension.X);
+
+	FComputeShaderUtils::AddPass(
+		GraphBuilder,
+		RDG_EVENT_NAME("HairStrands::CopyPosition"),
+		ComputeShader,
+		PassParameters,
+		GroupCount);
+
+}
+
 ////////////////////////////////////////////////////////////////////////////////////////////////////////////////////////
 class FTransferVelocityPassCS : public FGlobalShader
 {
diff --git a/Engine/Shaders/Private/HairStrands/HairStrandsInterpolation.usf b/Engine/Shaders/Private/HairStrands/HairStrandsInterpolation.usf
--- a/Engine/Shaders/Private/HairStrands/HairStrandsInterpolation.usf
+++ b/Engine/Shaders/Private/HairStrands/HairStrandsInterpolation.usf
@@ -713,6 +713,25 @@ void MainCS(uint2 DispatchThreadId : SV_DispatchThreadID)
 
 #endif // SHADER_HAIRTRANSFER_PREV_POSITION
 
+#if SHADER_COPY_POSITION
+
+uint ElementCount;
+Buffer<uint4> InBuffer;
+RWBuffer<uint4> OutBuffer;
+
+[numthreads(256, 1, 1)]
+void MainCS(uint3 DispatchThreadId : SV_DispatchThreadID)
+{
+	const uint Index = DispatchThreadId.x;
+	if (Index >= ElementCount)
+	{
+		return;
+	}
+	OutBuffer[Index] = InBuffer[Index];
+}
+
+#endif // SHADER_COPY_POSITION
+
 #if SHADER_GROOMCACHE_UPDATE
 
 uint ElementCount;
