Fix groom when enabling/disabling simulation/RFB deformation within the groom editor, and groom assets is used by some running instance.

diff --git a/Engine/Plugins/Runtime/HairStrands/Source/HairStrandsCore/Private/GroomAsset.cpp b/Engine/Plugins/Runtime/HairStrands/Source/HairStrandsCore/Private/GroomAsset.cpp
--- a/Engine/Plugins/Runtime/HairStrands/Source/HairStrandsCore/Private/GroomAsset.cpp
+++ b/Engine/Plugins/Runtime/HairStrands/Source/HairStrandsCore/Private/GroomAsset.cpp
@@ -1326,6 +1326,7 @@ static bool IsStrandsInterpolationAttributes(const FName PropertyName)
 		|| PropertyName == GET_MEMBER_NAME_CHECKED(FHairDecimationSettings, VertexDecimation)
 
 		// Add dependency on simulation and per LOD-simulation/global-interoplation to strip-out interoplation data if there are not needed
+		|| PropertyName == GET_MEMBER_NAME_CHECKED(UGroomAsset, EnableGlobalInterpolation)
 		|| PropertyName == GET_MEMBER_NAME_CHECKED(FHairSolverSettings, EnableSimulation)
 		|| PropertyName == GET_MEMBER_NAME_CHECKED(FHairLODSettings, Simulation)
 		|| PropertyName == GET_MEMBER_NAME_CHECKED(FHairLODSettings, GlobalInterpolation)
@@ -3166,6 +3167,22 @@ bool UGroomAsset::IsGlobalInterpolationEnable(int32 GroupIndex, int32 LODIndex)
 		(HairGroupsLOD[GroupIndex].LODs[LODIndex].GlobalInterpolation == EGroomOverrideType::Auto && EnableGlobalInterpolation);
 }
 
+bool UGroomAsset::IsSimulationEnable() const
+{
+	for (int32 GroupIndex = 0; GroupIndex < HairGroupsData.Num(); ++GroupIndex)
+	{
+		for (int32 LODIt = 0; LODIt < HairGroupsLOD.Num(); ++LODIt)
+		{
+			if (IsSimulationEnable(GroupIndex, LODIt))
+			{
+				return true;
+			}
+		}
+	}
+
+	return false;
+}
+
 bool UGroomAsset::NeedsInterpolationData(int32 GroupIndex) const
 {
 	for (int32 LODIt = 0; LODIt < HairGroupsLOD.Num(); ++LODIt)
@@ -3179,6 +3196,19 @@ bool UGroomAsset::NeedsInterpolationData(int32 GroupIndex) const
 	return false;
 }
 
+bool UGroomAsset::NeedsInterpolationData() const
+{
+	for (int32 GroupIndex = 0; GroupIndex < HairGroupsData.Num(); ++GroupIndex)
+	{
+		if (NeedsInterpolationData(GroupIndex))
+		{
+			return true;
+		}
+	}
+
+	return false;
+}
+
 bool FHairDescriptionGroups::IsValid() const
 {
 	for (const FHairDescriptionGroup& HairGroup : HairGroups)
diff --git a/Engine/Plugins/Runtime/HairStrands/Source/HairStrandsCore/Private/GroomComponent.cpp b/Engine/Plugins/Runtime/HairStrands/Source/HairStrandsCore/Private/GroomComponent.cpp
--- a/Engine/Plugins/Runtime/HairStrands/Source/HairStrandsCore/Private/GroomComponent.cpp
+++ b/Engine/Plugins/Runtime/HairStrands/Source/HairStrandsCore/Private/GroomComponent.cpp
@@ -2287,7 +2287,7 @@ void UGroomComponent::InitResources(bool bIsBindingReloading)
 		// * Physics simulation
 		// * RBF deformation.
 		// Therefore, even if simulation is disabled, we need to run partially the update if the binding system is enabled (skin deformation + RBF correction)
-		const bool bNeedGuides = (GroupData.Guides.HasValidData() && (bHasNeedSimulation || bHasNeedGlobalDeformation || bPreviewMode)) || (HairGroupInstance->Debug.GroomCacheType == EGroomCacheType::Guides);
+		const bool bNeedGuides = (GroupData.Guides.HasValidData() && (bHasNeedSimulation || bHasNeedGlobalDeformation)) || (HairGroupInstance->Debug.GroomCacheType == EGroomCacheType::Guides);
 		if (bNeedGuides)
 		{
 			HairGroupInstance->Guides.Data = &GroupData.Guides.BulkData;
diff --git a/Engine/Plugins/Runtime/HairStrands/Source/HairStrandsCore/Public/GroomAsset.h b/Engine/Plugins/Runtime/HairStrands/Source/HairStrandsCore/Public/GroomAsset.h
--- a/Engine/Plugins/Runtime/HairStrands/Source/HairStrandsCore/Public/GroomAsset.h
+++ b/Engine/Plugins/Runtime/HairStrands/Source/HairStrandsCore/Public/GroomAsset.h
@@ -538,8 +538,10 @@ public:
 	EGroomGeometryType GetGeometryType(int32 GroupIndex, int32 LODIndex) const;
 	EGroomBindingType GetBindingType(int32 GroupIndex, int32 LODIndex) const;
 	bool IsSimulationEnable(int32 GroupIndex, int32 LODIndex) const;
+	bool IsSimulationEnable() const;
 	bool IsGlobalInterpolationEnable(int32 GroupIndex, int32 LODIndex) const;
 	bool NeedsInterpolationData(int32 GroupIndex) const;
+	bool NeedsInterpolationData() const;
 
 #if WITH_EDITORONLY_DATA
 	void SetDebugMode(EHairStrandsDebugMode InMode) { DebugMode = InMode; OnGroomAssetChanged.Broadcast(); }
diff --git a/Engine/Plugins/Runtime/HairStrands/Source/HairStrandsEditor/Private/Widgets/SGroomEditorViewport.cpp b/Engine/Plugins/Runtime/HairStrands/Source/HairStrandsEditor/Private/Widgets/SGroomEditorViewport.cpp
--- a/Engine/Plugins/Runtime/HairStrands/Source/HairStrandsEditor/Private/Widgets/SGroomEditorViewport.cpp
+++ b/Engine/Plugins/Runtime/HairStrands/Source/HairStrandsEditor/Private/Widgets/SGroomEditorViewport.cpp
@@ -188,19 +188,19 @@ void SGroomEditorViewport::BindCommands()
 		FIsActionChecked::CreateSP(this, &SGroomEditorViewport::IsLODModelSelected, 0));
 	// all other LODs will be added dynamically
 	
-	CommandList->MapAction(FGroomEditorCommands::Get().ViewMode_Lit,				FExecuteAction::CreateSP(this, &SGroomEditorViewport::OnViewMode, EHairStrandsDebugMode::NoneDebug),				FCanExecuteAction::CreateSP(this, &SGroomEditorViewport::CanViewMode));
-	CommandList->MapAction(FGroomEditorCommands::Get().ViewMode_Guide,				FExecuteAction::CreateSP(this, &SGroomEditorViewport::OnViewMode, EHairStrandsDebugMode::SimHairStrands),			FCanExecuteAction::CreateSP(this, &SGroomEditorViewport::CanViewMode));
-	CommandList->MapAction(FGroomEditorCommands::Get().ViewMode_GuideInfluence,		FExecuteAction::CreateSP(this, &SGroomEditorViewport::OnViewMode, EHairStrandsDebugMode::RenderHairStrands),		FCanExecuteAction::CreateSP(this, &SGroomEditorViewport::CanViewMode));
-	CommandList->MapAction(FGroomEditorCommands::Get().ViewMode_UV,					FExecuteAction::CreateSP(this, &SGroomEditorViewport::OnViewMode, EHairStrandsDebugMode::RenderHairUV),				FCanExecuteAction::CreateSP(this, &SGroomEditorViewport::CanViewMode));
-	CommandList->MapAction(FGroomEditorCommands::Get().ViewMode_RootUV,				FExecuteAction::CreateSP(this, &SGroomEditorViewport::OnViewMode, EHairStrandsDebugMode::RenderHairRootUV),			FCanExecuteAction::CreateSP(this, &SGroomEditorViewport::CanViewMode));
-	CommandList->MapAction(FGroomEditorCommands::Get().ViewMode_RootUDIM,			FExecuteAction::CreateSP(this, &SGroomEditorViewport::OnViewMode, EHairStrandsDebugMode::RenderHairRootUDIM),		FCanExecuteAction::CreateSP(this, &SGroomEditorViewport::CanViewMode));
-	CommandList->MapAction(FGroomEditorCommands::Get().ViewMode_Seed,				FExecuteAction::CreateSP(this, &SGroomEditorViewport::OnViewMode, EHairStrandsDebugMode::RenderHairSeed),			FCanExecuteAction::CreateSP(this, &SGroomEditorViewport::CanViewMode));
-	CommandList->MapAction(FGroomEditorCommands::Get().ViewMode_Dimension,			FExecuteAction::CreateSP(this, &SGroomEditorViewport::OnViewMode, EHairStrandsDebugMode::RenderHairDimension),		FCanExecuteAction::CreateSP(this, &SGroomEditorViewport::CanViewMode));
-	CommandList->MapAction(FGroomEditorCommands::Get().ViewMode_RadiusVariation,	FExecuteAction::CreateSP(this, &SGroomEditorViewport::OnViewMode, EHairStrandsDebugMode::RenderHairRadiusVariation),FCanExecuteAction::CreateSP(this, &SGroomEditorViewport::CanViewMode));
-	CommandList->MapAction(FGroomEditorCommands::Get().ViewMode_BaseColor,			FExecuteAction::CreateSP(this, &SGroomEditorViewport::OnViewMode, EHairStrandsDebugMode::RenderHairBaseColor),		FCanExecuteAction::CreateSP(this, &SGroomEditorViewport::CanViewMode));
-	CommandList->MapAction(FGroomEditorCommands::Get().ViewMode_Roughness,			FExecuteAction::CreateSP(this, &SGroomEditorViewport::OnViewMode, EHairStrandsDebugMode::RenderHairRoughness),		FCanExecuteAction::CreateSP(this, &SGroomEditorViewport::CanViewMode));
-	CommandList->MapAction(FGroomEditorCommands::Get().ViewMode_VisCluster,			FExecuteAction::CreateSP(this, &SGroomEditorViewport::OnViewMode, EHairStrandsDebugMode::RenderVisCluster),			FCanExecuteAction::CreateSP(this, &SGroomEditorViewport::CanViewMode));
-	CommandList->MapAction(FGroomEditorCommands::Get().ViewMode_Group,				FExecuteAction::CreateSP(this, &SGroomEditorViewport::OnViewMode, EHairStrandsDebugMode::RenderHairGroup),			FCanExecuteAction::CreateSP(this, &SGroomEditorViewport::CanViewMode));
+	CommandList->MapAction(FGroomEditorCommands::Get().ViewMode_Lit,				FExecuteAction::CreateSP(this, &SGroomEditorViewport::OnViewMode, EHairStrandsDebugMode::NoneDebug),				FCanExecuteAction::CreateSP(this, &SGroomEditorViewport::CanViewMode, false));
+	CommandList->MapAction(FGroomEditorCommands::Get().ViewMode_Guide,				FExecuteAction::CreateSP(this, &SGroomEditorViewport::OnViewMode, EHairStrandsDebugMode::SimHairStrands),			FCanExecuteAction::CreateSP(this, &SGroomEditorViewport::CanViewMode, true));
+	CommandList->MapAction(FGroomEditorCommands::Get().ViewMode_GuideInfluence,		FExecuteAction::CreateSP(this, &SGroomEditorViewport::OnViewMode, EHairStrandsDebugMode::RenderHairStrands),		FCanExecuteAction::CreateSP(this, &SGroomEditorViewport::CanViewMode, true));
+	CommandList->MapAction(FGroomEditorCommands::Get().ViewMode_UV,					FExecuteAction::CreateSP(this, &SGroomEditorViewport::OnViewMode, EHairStrandsDebugMode::RenderHairUV),				FCanExecuteAction::CreateSP(this, &SGroomEditorViewport::CanViewMode, false));
+	CommandList->MapAction(FGroomEditorCommands::Get().ViewMode_RootUV,				FExecuteAction::CreateSP(this, &SGroomEditorViewport::OnViewMode, EHairStrandsDebugMode::RenderHairRootUV),			FCanExecuteAction::CreateSP(this, &SGroomEditorViewport::CanViewMode, false));
+	CommandList->MapAction(FGroomEditorCommands::Get().ViewMode_RootUDIM,			FExecuteAction::CreateSP(this, &SGroomEditorViewport::OnViewMode, EHairStrandsDebugMode::RenderHairRootUDIM),		FCanExecuteAction::CreateSP(this, &SGroomEditorViewport::CanViewMode, false));
+	CommandList->MapAction(FGroomEditorCommands::Get().ViewMode_Seed,				FExecuteAction::CreateSP(this, &SGroomEditorViewport::OnViewMode, EHairStrandsDebugMode::RenderHairSeed),			FCanExecuteAction::CreateSP(this, &SGroomEditorViewport::CanViewMode, false));
+	CommandList->MapAction(FGroomEditorCommands::Get().ViewMode_Dimension,			FExecuteAction::CreateSP(this, &SGroomEditorViewport::OnViewMode, EHairStrandsDebugMode::RenderHairDimension),		FCanExecuteAction::CreateSP(this, &SGroomEditorViewport::CanViewMode, false));
+	CommandList->MapAction(FGroomEditorCommands::Get().ViewMode_RadiusVariation,	FExecuteAction::CreateSP(this, &SGroomEditorViewport::OnViewMode, EHairStrandsDebugMode::RenderHairRadiusVariation),FCanExecuteAction::CreateSP(this, &SGroomEditorViewport::CanViewMode, false));
+	CommandList->MapAction(FGroomEditorCommands::Get().ViewMode_BaseColor,			FExecuteAction::CreateSP(this, &SGroomEditorViewport::OnViewMode, EHairStrandsDebugMode::RenderHairBaseColor),		FCanExecuteAction::CreateSP(this, &SGroomEditorViewport::CanViewMode, false));
+	CommandList->MapAction(FGroomEditorCommands::Get().ViewMode_Roughness,			FExecuteAction::CreateSP(this, &SGroomEditorViewport::OnViewMode, EHairStrandsDebugMode::RenderHairRoughness),		FCanExecuteAction::CreateSP(this, &SGroomEditorViewport::CanViewMode, false));
+	CommandList->MapAction(FGroomEditorCommands::Get().ViewMode_VisCluster,			FExecuteAction::CreateSP(this, &SGroomEditorViewport::OnViewMode, EHairStrandsDebugMode::RenderVisCluster),			FCanExecuteAction::CreateSP(this, &SGroomEditorViewport::CanViewMode, false));
+	CommandList->MapAction(FGroomEditorCommands::Get().ViewMode_Group,				FExecuteAction::CreateSP(this, &SGroomEditorViewport::OnViewMode, EHairStrandsDebugMode::RenderHairGroup),			FCanExecuteAction::CreateSP(this, &SGroomEditorViewport::CanViewMode, false));
 	CommandList->MapAction(FGroomEditorCommands::Get().ViewMode_CardsGuides,		FExecuteAction::CreateSP(this, &SGroomEditorViewport::OnCardsGuides), FCanExecuteAction::CreateSP(this, &SGroomEditorViewport::CanCardsGuides));
 }
 
@@ -220,9 +220,16 @@ void SGroomEditorViewport::OnViewMode(EHairStrandsDebugMode Mode)
 	}
 }
 
-bool SGroomEditorViewport::CanViewMode() const
+bool SGroomEditorViewport::CanViewMode(bool bGuideMode) const
 {
-	return GroomComponent && GroomComponent->GroomAsset != nullptr;
+	if (bGuideMode)
+	{
+		return GroomComponent && GroomComponent->GroomAsset != nullptr && GroomComponent->GroomAsset->IsSimulationEnable();
+	}
+	else
+	{
+		return GroomComponent && GroomComponent->GroomAsset != nullptr;
+	}
 }
 
 void SGroomEditorViewport::OnCardsGuides()
diff --git a/Engine/Plugins/Runtime/HairStrands/Source/HairStrandsEditor/Public/Widgets/SGroomEditorViewport.h b/Engine/Plugins/Runtime/HairStrands/Source/HairStrandsEditor/Public/Widgets/SGroomEditorViewport.h
--- a/Engine/Plugins/Runtime/HairStrands/Source/HairStrandsEditor/Public/Widgets/SGroomEditorViewport.h
+++ b/Engine/Plugins/Runtime/HairStrands/Source/HairStrandsEditor/Public/Widgets/SGroomEditorViewport.h
@@ -62,7 +62,7 @@ public:
 	void OnSetLODModel(int32 InLODSelection);
 
 	void OnViewMode(EHairStrandsDebugMode Mode);
-	bool CanViewMode() const;
+	bool CanViewMode(bool bGuideMode) const;
 
 	void OnCardsGuides();
 	bool CanCardsGuides() const;
