fix callers of SerializeCompressed that didn't specify size
no longer allowed
fixes Horde issue 50245

diff --git a/Engine/Plugins/Runtime/HairStrands/Source/HairStrandsCore/Private/GroomAsset.cpp b/Engine/Plugins/Runtime/HairStrands/Source/HairStrandsCore/Private/GroomAsset.cpp
--- a/Engine/Plugins/Runtime/HairStrands/Source/HairStrandsCore/Private/GroomAsset.cpp
+++ b/Engine/Plugins/Runtime/HairStrands/Source/HairStrandsCore/Private/GroomAsset.cpp
@@ -2120,7 +2120,7 @@ bool UGroomAsset::CacheStrandsData(uint32 GroupIndex, FString& OutDerivedDataKey
 		Ar << UncompressedSize;
 
 		uint8* DecompressionBuffer = reinterpret_cast<uint8*>(FMemory::Malloc(UncompressedSize));
-		Ar.SerializeCompressed(DecompressionBuffer, 0, NAME_Zlib);
+		Ar.SerializeCompressed(DecompressionBuffer, UncompressedSize, NAME_Zlib);
 
 		FHairGroupData& HairGroupData = HairGroupsData[GroupIndex];
 		FLargeMemoryReader LargeMemReader(DecompressionBuffer, UncompressedSize, ELargeMemoryReaderFlags::Persistent | ELargeMemoryReaderFlags::TakeOwnership);
diff --git a/Engine/Plugins/Runtime/HairStrands/Source/HairStrandsCore/Private/GroomBindingAsset.cpp b/Engine/Plugins/Runtime/HairStrands/Source/HairStrandsCore/Private/GroomBindingAsset.cpp
--- a/Engine/Plugins/Runtime/HairStrands/Source/HairStrandsCore/Private/GroomBindingAsset.cpp
+++ b/Engine/Plugins/Runtime/HairStrands/Source/HairStrandsCore/Private/GroomBindingAsset.cpp
@@ -727,7 +727,7 @@ void UGroomBindingAsset::CacheDerivedDatas()
 			Ar << UncompressedSize;
 
 			uint8* DecompressionBuffer = reinterpret_cast<uint8*>(FMemory::Malloc(UncompressedSize));
-			Ar.SerializeCompressed(DecompressionBuffer, 0, NAME_Zlib);
+			Ar.SerializeCompressed(DecompressionBuffer, UncompressedSize, NAME_Zlib);
 
 			FLargeMemoryReader LargeMemReader(DecompressionBuffer, UncompressedSize, ELargeMemoryReaderFlags::Persistent | ELargeMemoryReaderFlags::TakeOwnership);
 			InternalSerialize(LargeMemReader, this, HairGroupBulkDatas, 0);
diff --git a/Engine/Source/Developer/SlateReflector/Private/Widgets/SWidgetSnapshotVisualizer.cpp b/Engine/Source/Developer/SlateReflector/Private/Widgets/SWidgetSnapshotVisualizer.cpp
--- a/Engine/Source/Developer/SlateReflector/Private/Widgets/SWidgetSnapshotVisualizer.cpp
+++ b/Engine/Source/Developer/SlateReflector/Private/Widgets/SWidgetSnapshotVisualizer.cpp
@@ -579,7 +579,7 @@ void FWidgetSnapshotData::LoadSnapshotFromBuffer(const TArray<uint8>& InData)
 		BufferReader << UncompressedDataSize;
 
 		UncompressedData.AddZeroed(UncompressedDataSize);
-		BufferReader.SerializeCompressed(UncompressedData.GetData(), 0, NAME_Zlib);
+		BufferReader.SerializeCompressed(UncompressedData.GetData(), UncompressedDataSize, NAME_Zlib);
 	}
 
 	bool bJsonLoaded = false;
