Add view tracking to hair strands buffer update to detect when positions get deformed through different view/viewports.

Previously it was causing aliasing/jittering as deformed positions would be updated through different viewport, causing incorrect motion vector. We now detect such a case by tracking view ID, and insure if the deformed position are not issued by the same view update, motion vector are discards.

diff --git a/Engine/Plugins/Runtime/HairStrands/Source/HairStrandsCore/Private/GroomManager.cpp b/Engine/Plugins/Runtime/HairStrands/Source/HairStrandsCore/Private/GroomManager.cpp
--- a/Engine/Plugins/Runtime/HairStrands/Source/HairStrandsCore/Private/GroomManager.cpp
+++ b/Engine/Plugins/Runtime/HairStrands/Source/HairStrandsCore/Private/GroomManager.cpp
@@ -125,6 +125,7 @@ static int32 GetCacheGeometryLODIndex(const FCachedGeometry& In)
 static void RunInternalHairStrandsInterpolation(
 	FRDGBuilder& GraphBuilder,
 	const FSceneView* View,
+	const uint32 ViewUniqueID,
 	const FHairStrandsInstances& Instances,
 	const FGPUSkinCache* SkinCache,
 	const FShaderDrawDebugData* ShaderDrawData,
@@ -359,6 +360,7 @@ static void RunInternalHairStrandsInterpolation(
  			ComputeHairStrandsInterpolation(
 				GraphBuilder, 
 				ShaderMap,
+				ViewUniqueID,
 				ShaderDrawData, 
 				Instance,
 				Instance->Debug.MeshLODIndex,
@@ -370,6 +372,7 @@ static void RunInternalHairStrandsInterpolation(
 static void RunHairStrandsInterpolation_Guide(
 	FRDGBuilder& GraphBuilder,
 	const FSceneView* View,
+	const uint32 ViewUniqueID,
 	const FHairStrandsInstances& Instances,
 	const FGPUSkinCache* SkinCache,
 	const FShaderDrawDebugData* ShaderDrawData,
@@ -385,6 +388,7 @@ static void RunHairStrandsInterpolation_Guide(
 	RunInternalHairStrandsInterpolation(
 		GraphBuilder,
 		View,
+		ViewUniqueID,
 		Instances,
 		SkinCache,
 		ShaderDrawData,
@@ -396,6 +400,7 @@ static void RunHairStrandsInterpolation_Guide(
 static void RunHairStrandsInterpolation_Strands(
 	FRDGBuilder& GraphBuilder,
 	const FSceneView* View,
+	const uint32 ViewUniqueID,
 	const FHairStrandsInstances& Instances,
 	const FGPUSkinCache* SkinCache,
 	const FShaderDrawDebugData* ShaderDrawData,
@@ -411,6 +416,7 @@ static void RunHairStrandsInterpolation_Strands(
 	RunInternalHairStrandsInterpolation(
 		GraphBuilder,
 		View,
+		ViewUniqueID,
 		Instances,
 		SkinCache,
 		ShaderDrawData,
@@ -892,6 +898,7 @@ void ProcessHairStrandsBookmark(
 		RunHairStrandsInterpolation_Guide(
 			*GraphBuilder,
 			Parameters.View,
+			Parameters.ViewUniqueID,
 			Instances,
 			Parameters.SkinCache,
 			Parameters.DebugShaderData,
@@ -910,6 +917,7 @@ void ProcessHairStrandsBookmark(
 		RunHairStrandsInterpolation_Strands(
 			*GraphBuilder,
 			Parameters.View,
+			Parameters.ViewUniqueID,
 			Instances,
 			Parameters.SkinCache,
 			Parameters.DebugShaderData,
diff --git a/Engine/Plugins/Runtime/HairStrands/Source/HairStrandsCore/Private/HairStrandsRendering.cpp b/Engine/Plugins/Runtime/HairStrands/Source/HairStrandsCore/Private/HairStrandsRendering.cpp
--- a/Engine/Plugins/Runtime/HairStrands/Source/HairStrandsCore/Private/HairStrandsRendering.cpp
+++ b/Engine/Plugins/Runtime/HairStrands/Source/HairStrandsCore/Private/HairStrandsRendering.cpp
@@ -1546,14 +1546,27 @@ static FHairGroupPublicData::FVertexFactoryInput InternalComputeHairStrandsVerte
 	// 2. Render Strands with deformation
 	else if (bSupportDeformation)
 	{
-		CONVERT_HAIRSSTRANDS_VF_PARAMETERS(OutVFInput.Strands.PositionBuffer,			Instance->Strands.DeformedResource->GetBuffer(FHairStrandsDeformedResource::EFrameType::Current));
-		CONVERT_HAIRSSTRANDS_VF_PARAMETERS(OutVFInput.Strands.PrevPositionBuffer,		Instance->Strands.DeformedResource->GetBuffer(FHairStrandsDeformedResource::EFrameType::Previous));
+		const bool bHasValidMotionVector =
+			Instance->Strands.DeformedResource->GetUniqueViewID(FHairStrandsDeformedResource::Current) ==
+			Instance->Strands.DeformedResource->GetUniqueViewID(FHairStrandsDeformedResource::Previous);
+		if (bHasValidMotionVector)
+		{
+			CONVERT_HAIRSSTRANDS_VF_PARAMETERS(OutVFInput.Strands.PositionBuffer,			Instance->Strands.DeformedResource->GetBuffer(FHairStrandsDeformedResource::EFrameType::Current));
+			CONVERT_HAIRSSTRANDS_VF_PARAMETERS(OutVFInput.Strands.PrevPositionBuffer,		Instance->Strands.DeformedResource->GetBuffer(FHairStrandsDeformedResource::EFrameType::Previous));
+			CONVERT_HAIRSSTRANDS_VF_PARAMETERS(OutVFInput.Strands.PositionOffsetBuffer,		Instance->Strands.DeformedResource->GetPositionOffsetBuffer(FHairStrandsDeformedResource::EFrameType::Current));
+			CONVERT_HAIRSSTRANDS_VF_PARAMETERS(OutVFInput.Strands.PrevPositionOffsetBuffer, Instance->Strands.DeformedResource->GetPositionOffsetBuffer(FHairStrandsDeformedResource::EFrameType::Previous));
+		}
+		else
+		{
+			CONVERT_HAIRSSTRANDS_VF_PARAMETERS(OutVFInput.Strands.PositionBuffer,			Instance->Strands.DeformedResource->GetBuffer(FHairStrandsDeformedResource::EFrameType::Current));
+			CONVERT_HAIRSSTRANDS_VF_PARAMETERS(OutVFInput.Strands.PrevPositionBuffer,		Instance->Strands.DeformedResource->GetBuffer(FHairStrandsDeformedResource::EFrameType::Current));
+			CONVERT_HAIRSSTRANDS_VF_PARAMETERS(OutVFInput.Strands.PositionOffsetBuffer,		Instance->Strands.DeformedResource->GetPositionOffsetBuffer(FHairStrandsDeformedResource::EFrameType::Current));
+			CONVERT_HAIRSSTRANDS_VF_PARAMETERS(OutVFInput.Strands.PrevPositionOffsetBuffer,	Instance->Strands.DeformedResource->GetPositionOffsetBuffer(FHairStrandsDeformedResource::EFrameType::Current));
+		}
 		CONVERT_HAIRSSTRANDS_VF_PARAMETERS(OutVFInput.Strands.TangentBuffer,			Instance->Strands.DeformedResource->TangentBuffer);
 		CONVERT_HAIRSSTRANDS_VF_PARAMETERS(OutVFInput.Strands.Attribute0Buffer,			bDebugModePatchedAttributeBuffer ? Instance->Strands.DebugAttributeBuffer : Instance->Strands.RestResource->Attribute0Buffer); // TODO
 		CONVERT_HAIRSSTRANDS_VF_PARAMETERS(OutVFInput.Strands.Attribute1Buffer,			Instance->Strands.RestResource->Attribute1Buffer);
 		CONVERT_HAIRSSTRANDS_VF_PARAMETERS(OutVFInput.Strands.MaterialBuffer,			Instance->Strands.RestResource->MaterialBuffer);
-		CONVERT_HAIRSSTRANDS_VF_PARAMETERS(OutVFInput.Strands.PositionOffsetBuffer,		Instance->Strands.DeformedResource->GetPositionOffsetBuffer(FHairStrandsDeformedResource::EFrameType::Current));
-		CONVERT_HAIRSSTRANDS_VF_PARAMETERS(OutVFInput.Strands.PrevPositionOffsetBuffer,	Instance->Strands.DeformedResource->GetPositionOffsetBuffer(FHairStrandsDeformedResource::EFrameType::Previous));
 
 		OutVFInput.Strands.PositionOffset = Instance->Strands.DeformedResource->GetPositionOffset(FHairStrandsDeformedResource::EFrameType::Current);
 		OutVFInput.Strands.PrevPositionOffset = Instance->Strands.DeformedResource->GetPositionOffset(FHairStrandsDeformedResource::EFrameType::Previous);
@@ -1597,6 +1610,7 @@ void CreateHairStrandsDebugAttributeBuffer(FRDGBuilder& GraphBuilder, FRDGExtern
 void ComputeHairStrandsInterpolation(
 	FRDGBuilder& GraphBuilder,
 	FGlobalShaderMap* ShaderMap,
+	const uint32 ViewUniqueID,
 	const FShaderDrawDebugData* ShaderDrawData,
 	FHairGroupInstance* Instance,
 	int32 MeshLODIndex,
@@ -1670,6 +1684,10 @@ void ComputeHairStrandsInterpolation(
 				FRDGImportedBuffer Strands_DeformedPrevPosition = Register(GraphBuilder, Instance->Strands.DeformedResource->GetBuffer(FHairStrandsDeformedResource::Previous), ERDGImportedBufferFlags::CreateViews);
 				FRDGImportedBuffer Strands_DeformedTangent = Register(GraphBuilder, Instance->Strands.DeformedResource->TangentBuffer, ERDGImportedBufferFlags::CreateViews);
 				FRDGImportedBuffer* Strands_DebugAttribute = nullptr;
+
+				// Trach on which view the position has been update, so that we can ensure motion vector are coherent
+				Instance->Strands.DeformedResource->GetUniqueViewID(FHairStrandsDeformedResource::Current) = ViewUniqueID;				
+
 			#if WITH_EDITOR
 				FRDGImportedBuffer Strands_DebugAttributeReg;
 				if (bDebugModePatchedAttributeBuffer)
diff --git a/Engine/Plugins/Runtime/HairStrands/Source/HairStrandsCore/Public/GroomResources.h b/Engine/Plugins/Runtime/HairStrands/Source/HairStrandsCore/Public/GroomResources.h
--- a/Engine/Plugins/Runtime/HairStrands/Source/HairStrandsCore/Public/GroomResources.h
+++ b/Engine/Plugins/Runtime/HairStrands/Source/HairStrandsCore/Public/GroomResources.h
@@ -289,6 +289,12 @@ struct FHairStrandsDeformedResource : public FHairCommonResource
 	/* Whether the GPU data should be initialized with the asset data or not */
 	uint32 CurrentIndex = 0;
 
+	/* Track the view which has update the formed position. This is used when rendering the same 
+	   instance accross several editor viewport (not views of the same viewport), to prevent 
+	   incorrect motion vector, as editor viewport are refresh not on the same tick, but with a 
+	   throttling mechanism which make the updates not coherent. */
+	uint32 UniqueViewIDs[2] = { 0, 0 };
+
 	/* Type of curves */
 	const EHairStrandsResourcesType CurveType;
 
@@ -306,6 +312,8 @@ struct FHairStrandsDeformedResource : public FHairCommonResource
 	inline const FVector& GetPositionOffset(EFrameType T) const { return PositionOffset[GetIndex(T)]; }
 	inline void SetPositionOffset(EFrameType T, const FVector& Offset)  { PositionOffset[GetIndex(T)] = Offset; }
 	inline void SwapBuffer()									{ CurrentIndex = 1u - CurrentIndex; }
+	inline uint32& GetUniqueViewID(EFrameType T)				{ return UniqueViewIDs[GetIndex(T)]; }
+	inline uint32 GetUniqueViewID(EFrameType T) const			{ return UniqueViewIDs[GetIndex(T)]; }
 	//bool NeedsToUpdateTangent();
 };
 
diff --git a/Engine/Plugins/Runtime/HairStrands/Source/HairStrandsCore/Public/HairStrandsRendering.h b/Engine/Plugins/Runtime/HairStrands/Source/HairStrandsCore/Public/HairStrandsRendering.h
--- a/Engine/Plugins/Runtime/HairStrands/Source/HairStrandsCore/Public/HairStrandsRendering.h
+++ b/Engine/Plugins/Runtime/HairStrands/Source/HairStrandsCore/Public/HairStrandsRendering.h
@@ -26,6 +26,7 @@ void ResetHairStrandsInterpolation(
 void ComputeHairStrandsInterpolation(
 	FRDGBuilder& GraphBuilder,
 	FGlobalShaderMap* ShaderMap,
+	const uint32 ViewUniqueID,
 	const struct FShaderDrawDebugData* DebugShaderData,
 	struct FHairGroupInstance* Instance,
 	int32 LODIndex,
diff --git a/Engine/Source/Runtime/Renderer/Private/HairStrands/HairStrandsInterface.cpp b/Engine/Source/Runtime/Renderer/Private/HairStrands/HairStrandsInterface.cpp
--- a/Engine/Source/Runtime/Renderer/Private/HairStrands/HairStrandsInterface.cpp
+++ b/Engine/Source/Runtime/Renderer/Private/HairStrands/HairStrandsInterface.cpp
@@ -432,13 +432,14 @@ FHairStrandsBookmarkParameters CreateHairStrandsBookmarkParameters(FScene* Scene
 				Out.VisibleInstances.Add(HairGroupPublicData->Instance);
 			}
 		}
-	}
+	}	
 	Out.DebugShaderData			= ShaderDrawDebug::IsEnabled(View) ? &View.ShaderDrawData : nullptr;
 	Out.SkinCache				= View.Family->Scene->GetGPUSkinCache();
 	Out.ShaderMap				= View.ShaderMap;
 	Out.Instances				= &Scene->HairStrandsSceneData.RegisteredProxies;
 	Out.View					= &View;
 	Out.ViewRect				= View.ViewRect;
+	Out.ViewUniqueID			= View.ViewState ? View.ViewState->UniqueID : ~0;
 	Out.SceneColorTexture		= nullptr;
 	Out.bStrandsGeometryEnabled = IsHairStrandsEnabled(EHairStrandsShaderType::Strands, View.GetShaderPlatform());
 	if (GHairStrandsParameterFunction)
diff --git a/Engine/Source/Runtime/Renderer/Public/HairStrandsInterface.h b/Engine/Source/Runtime/Renderer/Public/HairStrandsInterface.h
--- a/Engine/Source/Runtime/Renderer/Public/HairStrandsInterface.h
+++ b/Engine/Source/Runtime/Renderer/Public/HairStrandsInterface.h
@@ -424,6 +424,7 @@ struct FHairStrandsBookmarkParameters
 	FShaderDrawDebugData* DebugShaderData = nullptr;
 	class FGlobalShaderMap* ShaderMap = nullptr;
 
+	uint32 ViewUniqueID = ~0; // View 0
 	FIntRect ViewRect; // View 0
 	FHairStrandsInstances VisibleInstances;
 	FHairStrandsInstances* Instances = nullptr;
