Add LOD coloration support for groom component.

diff --git a/Engine/Plugins/Runtime/HairStrands/Source/HairStrandsCore/Private/GroomComponent.cpp b/Engine/Plugins/Runtime/HairStrands/Source/HairStrandsCore/Private/GroomComponent.cpp
--- a/Engine/Plugins/Runtime/HairStrands/Source/HairStrandsCore/Private/GroomComponent.cpp
+++ b/Engine/Plugins/Runtime/HairStrands/Source/HairStrandsCore/Private/GroomComponent.cpp
@@ -674,8 +674,13 @@ public:
 				{
 					check(Instances[GroupIt]->GetRefCount() > 0);
 
-					FMaterialRenderProxy* Strands_MaterialProxy = nullptr;
-					const EHairStrandsDebugMode DebugMode = GetHairStrandsGeometryDebugMode(Instances[GroupIt]);
+					FMaterialRenderProxy* Debug_MaterialProxy = nullptr;
+					EHairStrandsDebugMode DebugMode = GetHairStrandsGeometryDebugMode(Instances[GroupIt]);
+					if (AllowDebugViewmodes() && View->Family->EngineShowFlags.LODColoration)
+					{
+						DebugMode = EHairStrandsDebugMode::RenderLODColoration;
+					}
+					
 					if (DebugMode != EHairStrandsDebugMode::NoneDebug)
 					{
 						float DebugModeScalar = 0;
@@ -694,6 +699,7 @@ public:
 						case EHairStrandsDebugMode::RenderHairRoughness			: DebugModeScalar = 8.f; break;
 						case EHairStrandsDebugMode::RenderVisCluster			: DebugModeScalar = 0.f; break;
 						case EHairStrandsDebugMode::RenderHairGroup				: DebugModeScalar = 9.f; break;
+						case EHairStrandsDebugMode::RenderLODColoration			: DebugModeScalar = 10.f; break;
 						};
 
 						// TODO: fix this as the radius is incorrect. This code run before the interpolation code, which is where HairRadius is updated.
@@ -702,14 +708,31 @@ public:
 						{
 							HairMaxRadius = FMath::Max(HairMaxRadius, Instance->Strands.Modifier.HairWidth * 0.5f);
 						}
+						
+						// Reuse the HairMaxRadius field to send the LOD index instead of adding yet another variable
+						if (DebugMode == EHairStrandsDebugMode::RenderLODColoration)
+						{
+							HairMaxRadius = Instances[GroupIt]->HairGroupPublicData ? Instances[GroupIt]->HairGroupPublicData->LODIndex : 0;
+						}
 
-						const FVector HairGroupColor = FVector(GetHairGroupDebugColor(GroupIt));
-						auto DebugMaterial = new FHairDebugModeMaterialRenderProxy(Strands_DebugMaterial ? Strands_DebugMaterial->GetRenderProxy() : nullptr, DebugModeScalar, 0, HairMaxRadius, HairGroupColor);
+						FVector HairColor = FVector::ZeroVector;
+						if (DebugMode == EHairStrandsDebugMode::RenderHairGroup)
+						{
+							HairColor = FVector(GetHairGroupDebugColor(GroupIt));
+						}
+						else if (DebugMode == EHairStrandsDebugMode::RenderLODColoration)
+						{
+							int32 LODIndex = Instances[GroupIt]->HairGroupPublicData ? Instances[GroupIt]->HairGroupPublicData->LODIndex : 0;
+							LODIndex = FMath::Clamp(LODIndex, 0, GEngine->LODColorationColors.Num() - 1);
+							const FLinearColor LODColor = GEngine->LODColorationColors[LODIndex];
+							HairColor = FVector(LODColor.R, LODColor.G, LODColor.B);
+						}
+						auto DebugMaterial = new FHairDebugModeMaterialRenderProxy(Strands_DebugMaterial ? Strands_DebugMaterial->GetRenderProxy() : nullptr, DebugModeScalar, 0, HairMaxRadius, HairColor);
 						Collector.RegisterOneFrameMaterialProxy(DebugMaterial);
-						Strands_MaterialProxy = DebugMaterial;
+						Debug_MaterialProxy = DebugMaterial;
 					}
 
-					FMeshBatch* MeshBatch = CreateMeshBatch(View, ViewFamily, Collector, HairGroups[GroupIt], Instances[GroupIt], GroupIt, Strands_MaterialProxy);
+					FMeshBatch* MeshBatch = CreateMeshBatch(View, ViewFamily, Collector, HairGroups[GroupIt], Instances[GroupIt], GroupIt, Debug_MaterialProxy);
 					if (MeshBatch == nullptr)
 					{
 						continue;
@@ -732,7 +755,7 @@ public:
 		const FHairGroup& GroupData,
 		const FHairGroupInstance* Instance,
 		uint32 GroupIndex,
-		FMaterialRenderProxy* Strands_MaterialProxy) const
+		FMaterialRenderProxy* Debug_MaterialProxy) const
 	{
 		const EHairGeometryType GeometryType = Instance->GeometryType;
 		if (GeometryType == EHairGeometryType::NoneGeometry)
@@ -765,7 +788,7 @@ public:
 			NumPrimitive = HairVertexCount / 3;
 			IndexBuffer = &Instance->Meshes.LODs[IntLODIndex].RestResource->IndexBuffer;
 			bUseCulling = false;
-			MaterialRenderProxy = Instance->Meshes.LODs[IntLODIndex].Material->GetRenderProxy();
+			MaterialRenderProxy = Debug_MaterialProxy == nullptr ? Instance->Meshes.LODs[IntLODIndex].Material->GetRenderProxy() : Debug_MaterialProxy;
 		}
 		else if (GeometryType == EHairGeometryType::Cards)
 		{
@@ -781,7 +804,7 @@ public:
 			NumPrimitive = HairVertexCount / 3;
 			IndexBuffer = &Instance->Cards.LODs[IntLODIndex].RestResource->RestIndexBuffer;
 			bUseCulling = false;
-			MaterialRenderProxy = Instance->Cards.LODs[IntLODIndex].Material->GetRenderProxy();
+			MaterialRenderProxy = Debug_MaterialProxy == nullptr ? Instance->Cards.LODs[IntLODIndex].Material->GetRenderProxy() : Debug_MaterialProxy;
 			bWireframe = AllowDebugViewmodes() && ViewFamily.EngineShowFlags.Wireframe;
 		}
 		else // if (GeometryType == EHairGeometryType::Strands)
@@ -791,7 +814,7 @@ public:
 			MaxVertexIndex = HairVertexCount * 6;
 			bUseCulling = Instance->Strands.bIsCullingEnabled;
 			NumPrimitive = bUseCulling ? 0 : HairVertexCount * 2;
-			MaterialRenderProxy = Strands_MaterialProxy == nullptr ? Instance->Strands.Material->GetRenderProxy() : Strands_MaterialProxy;
+			MaterialRenderProxy = Debug_MaterialProxy == nullptr ? Instance->Strands.Material->GetRenderProxy() : Debug_MaterialProxy;
 			bWireframe = AllowDebugViewmodes() && ViewFamily.EngineShowFlags.Wireframe;
 		}
 
@@ -2199,6 +2222,32 @@ void UGroomComponent::InitResources(bool bIsBindingReloading)
 	{
 		if (USkeletalMeshComponent* ParentSkelMeshComponent = Cast<USkeletalMeshComponent>(ParentMeshComponent))
 		{
+			// Validate that if we are bound to a skel. mesh, and we have binding, and the binding type is set to skinning, that the skin cache is enabled for these
+			if (BindingAsset)
+			{
+				const uint32 SkelLODCount = ParentSkelMeshComponent->GetNumLODs();
+				for (int32 GroupIt = 0, GroupCount = GroomAsset->HairGroupsData.Num(); GroupIt < GroupCount; ++GroupIt)
+				{
+					for (uint32 LODIt = 0, LODCount = GroomAsset->GetLODCount(); LODIt < LODCount; ++LODIt)
+					{
+						const uint32 EffectiveLODIt = FMath::Clamp(LODIt, 0, SkelLODCount-1);
+						const bool bSupportSkinCache = ParentSkelMeshComponent->IsSkinCacheAllowed(EffectiveLODIt);
+						const EGroomBindingType BindingType = GroomAsset->GetBindingType(GroupIt, LODIt);
+
+						if (BindingType == EGroomBindingType::Skinning && !bSupportSkinCache)
+						{
+							UE_LOG(LogHairStrands, Warning, TEXT("[Groom] Groom asset (Group:%d/%d) is set to use Skinning at LOD %d/%d while the parent skel. mesh does not support skin. cache for this LOD - Groom:%s - Skel.Mesh:%s"),
+								GroupIt,
+								GroupCount,
+								LODIt,
+								LODCount,
+								ComponentIt->GroomAsset ? *GroomAsset->GetPathName() : TEXT("."),
+								ParentMeshComponent->GetPathName());
+						}
+					}
+				}
+			}
+
 			ValidatedMeshComponent = ParentSkelMeshComponent->SkeletalMesh ? ParentMeshComponent : nullptr;
 		}
 	}
diff --git a/Engine/Shaders/Private/HairStrands/HairStrandsVisibilityComposeSubPixelPS.usf b/Engine/Shaders/Private/HairStrands/HairStrandsVisibilityComposeSubPixelPS.usf
--- a/Engine/Shaders/Private/HairStrands/HairStrandsVisibilityComposeSubPixelPS.usf
+++ b/Engine/Shaders/Private/HairStrands/HairStrandsVisibilityComposeSubPixelPS.usf
@@ -82,6 +82,7 @@ void ComposeSamplePS(
 	in FScreenVertexOutput Input,
 	out float4 OutColor : SV_Target0,
 	out float  OutDepth : SV_DEPTH)
+#if PERMUTATION_DEBUG == 0
 {
 	OutColor = 0.0f;
 	const uint3 PixelCoord = uint3(floor(Input.Position.xy), 0);
@@ -123,6 +124,39 @@ void ComposeSamplePS(
 		OutDepth = HairDOFDepthTexture.Load(PixelCoord);
 	}
 }
+#else // PERMUTATION_DEBUG == 1
+{
+	OutColor = 0.0f;
+	const uint3 PixelCoord = uint3(floor(Input.Position.xy), 0);
+
+	const FNodeDesc NodeDesc = DecodeNodeDesc(HairStrands.HairSampleOffset.Load(PixelCoord));
+	if (NodeDesc.Count == 0)
+	{
+		discard;
+	}
+
+	const float PixelCoverage = min(HairStrands.HairCoverageTexture.Load(PixelCoord), 1);
+	if (PixelCoverage == 0)
+		discard;
+
+	const float HairDeviceZ = HairStrands.HairOnlyDepthTexture.Load(PixelCoord);
+
+	float3 LocalAccColor = 0;
+	LOOP
+	for (uint SampleIt = 0; SampleIt < NodeDesc.Count; SampleIt++)
+	{
+		const FPackedHairSample PackedSample = HairStrands.HairSampleData[NodeDesc.Offset + SampleIt];
+		const FHairSample Sample = UnpackHairSample(PackedSample);
+		const float LocalCoverage = From8bitCoverage(Sample.Coverage8bit);
+		LocalAccColor += LocalCoverage * Sample.BaseColor;
+	}
+	OutColor.rgb = LocalAccColor;
+
+	OutColor.rgb *= PixelCoverage;
+	OutColor.a = PixelCoverage;
+	OutDepth = HairDeviceZ;
+}
+#endif
 
 #endif // SHADER_COMPOSE_SAMPLE
 
diff --git a/Engine/Source/Runtime/Renderer/Private/HairStrands/HairStrandsComposition.cpp b/Engine/Source/Runtime/Renderer/Private/HairStrands/HairStrandsComposition.cpp
--- a/Engine/Source/Runtime/Renderer/Private/HairStrands/HairStrandsComposition.cpp
+++ b/Engine/Source/Runtime/Renderer/Private/HairStrands/HairStrandsComposition.cpp
@@ -165,6 +165,9 @@ class FHairVisibilityComposeSamplePS : public FGlobalShader
 	DECLARE_GLOBAL_SHADER(FHairVisibilityComposeSamplePS);
 	SHADER_USE_PARAMETER_STRUCT(FHairVisibilityComposeSamplePS, FGlobalShader);
 
+	class FDebug : SHADER_PERMUTATION_BOOL("PERMUTATION_DEBUG");
+	using FPermutationDomain = TShaderPermutationDomain<FDebug>;
+	
 	BEGIN_SHADER_PARAMETER_STRUCT(FParameters, )
 		SHADER_PARAMETER_STRUCT_REF(FViewUniformShaderParameters, ViewUniformBuffer)
 		SHADER_PARAMETER_STRUCT_INCLUDE(FHairStrandsTilePassVS::FParameters, TileData)
@@ -213,10 +216,13 @@ static void AddHairVisibilityComposeSamplePass(
 	Parameters->RenderTargets[0] = FRenderTargetBinding(OutColorTexture, ERenderTargetLoadAction::ELoad);
 	Parameters->RenderTargets.DepthStencil = FDepthStencilBinding(OutDepthTexture, ERenderTargetLoadAction::ELoad, ERenderTargetLoadAction::ENoAction, FExclusiveDepthStencil::DepthWrite_StencilNop);
 
-	TShaderMapRef<FHairVisibilityComposeSamplePS> PixelShader(View.ShaderMap);
+	const bool bDebugComposition = View.Family->EngineShowFlags.LODColoration > 0;
+	FHairVisibilityComposeSamplePS::FPermutationDomain PermutationVector;
+	PermutationVector.Set<FHairVisibilityComposeSamplePS::FDebug>(bDebugComposition);
+	TShaderMapRef<FHairVisibilityComposeSamplePS> PixelShader(View.ShaderMap, PermutationVector);
 	InternalCommonDrawPass(
 		GraphBuilder,
-		RDG_EVENT_NAME("HairStrandsComposeSample"),
+		RDG_EVENT_NAME("HairStrands::ComposeSample"),
 		View,
 		OutColorTexture->Desc.Extent,
 		EHairStrandsCommonPassType::Composition,
@@ -284,7 +290,7 @@ static FRDGTextureRef AddHairDOFDepthPass(
 	TShaderMapRef<FHairDOFDepthPS> PixelShader(View.ShaderMap);
 	InternalCommonDrawPass(
 		GraphBuilder,
-		RDG_EVENT_NAME("HairStrandsDOFDepth"),
+		RDG_EVENT_NAME("HairStrands::DOFDepth"),
 		View,
 		OutputResolution,
 		EHairStrandsCommonPassType::DOF,
@@ -347,7 +353,7 @@ static void AddHairVisibilityFastResolveMaskPass(
 	TShaderMapRef<FHairVisibilityFastResolveMaskPS> PixelShader(View.ShaderMap);
 	InternalCommonDrawPass(
 		GraphBuilder,
-		RDG_EVENT_NAME("HairStrandsVisibilityMarkTAAFastResolve"),
+		RDG_EVENT_NAME("HairStrands::MarkTAAFastResolve"),
 		View,
 		Resolution,
 		EHairStrandsCommonPassType::TAAFastResolve,
@@ -365,7 +371,6 @@ class FHairVisibilityGBufferWritePS : public FGlobalShader
 	DECLARE_GLOBAL_SHADER(FHairVisibilityGBufferWritePS);
 	SHADER_USE_PARAMETER_STRUCT(FHairVisibilityGBufferWritePS, FGlobalShader);
 
-
 	class FOutputType : SHADER_PERMUTATION_INT("PERMUTATION_OUTPUT_TYPE", 2);
 	using FPermutationDomain = TShaderPermutationDomain<FOutputType>;
 
@@ -453,7 +458,7 @@ static void AddHairVisibilityGBufferWritePass(
 	TShaderMapRef<FHairVisibilityGBufferWritePS> PixelShader(View.ShaderMap, PermutationVector);
 	InternalCommonDrawPass(
 		GraphBuilder,
-		RDG_EVENT_NAME("HairStrandsVisibilityMarkTAAFastResolve"),
+		RDG_EVENT_NAME("HairStrands::GBufferOverride"),
 		View,
 		OutGBufferATexture->Desc.Extent,
 		EHairStrandsCommonPassType::GBuffer,
diff --git a/Engine/Source/Runtime/Renderer/Public/HairStrandsInterface.h b/Engine/Source/Runtime/Renderer/Public/HairStrandsInterface.h
--- a/Engine/Source/Runtime/Renderer/Public/HairStrandsInterface.h
+++ b/Engine/Source/Runtime/Renderer/Public/HairStrandsInterface.h
@@ -67,6 +67,7 @@ enum class EHairStrandsDebugMode : uint8
 	RenderHairRoughness,
 	RenderVisCluster,
 	RenderHairGroup,
+	RenderLODColoration,
 	Count
 };
 
