Fix/disable in-progress validation code checked in by mistake in previous CL.

diff --git a/Engine/Plugins/Runtime/HairStrands/Source/HairStrandsCore/Private/GroomComponent.cpp b/Engine/Plugins/Runtime/HairStrands/Source/HairStrandsCore/Private/GroomComponent.cpp
--- a/Engine/Plugins/Runtime/HairStrands/Source/HairStrandsCore/Private/GroomComponent.cpp
+++ b/Engine/Plugins/Runtime/HairStrands/Source/HairStrandsCore/Private/GroomComponent.cpp
@@ -2223,6 +2223,7 @@ void UGroomComponent::InitResources(bool bIsBindingReloading)
 		if (USkeletalMeshComponent* ParentSkelMeshComponent = Cast<USkeletalMeshComponent>(ParentMeshComponent))
 		{
 			// Validate that if we are bound to a skel. mesh, and we have binding, and the binding type is set to skinning, that the skin cache is enabled for these
+			#if 0
 			if (BindingAsset)
 			{
 				const uint32 SkelLODCount = ParentSkelMeshComponent->GetNumLODs();
@@ -2247,7 +2248,7 @@ void UGroomComponent::InitResources(bool bIsBindingReloading)
 					}
 				}
 			}
-
+			#endif
 			ValidatedMeshComponent = ParentSkelMeshComponent->SkeletalMesh ? ParentMeshComponent : nullptr;
 		}
 	}
