Add validation/debug when a groom request skinning binding, but the underlying skel. mesh does not have skin. cache support.

diff --git a/Engine/Plugins/Runtime/HairStrands/Source/HairStrandsCore/Private/GroomComponent.cpp b/Engine/Plugins/Runtime/HairStrands/Source/HairStrandsCore/Private/GroomComponent.cpp
--- a/Engine/Plugins/Runtime/HairStrands/Source/HairStrandsCore/Private/GroomComponent.cpp
+++ b/Engine/Plugins/Runtime/HairStrands/Source/HairStrandsCore/Private/GroomComponent.cpp
@@ -2222,8 +2222,15 @@ void UGroomComponent::InitResources(bool bIsBindingReloading)
 	{
 		if (USkeletalMeshComponent* ParentSkelMeshComponent = Cast<USkeletalMeshComponent>(ParentMeshComponent))
 		{
-			// Validate that if we are bound to a skel. mesh, and we have binding, and the binding type is set to skinning, that the skin cache is enabled for these
-			#if 0
+			ValidatedMeshComponent = ParentSkelMeshComponent->SkeletalMesh ? ParentMeshComponent : nullptr;
+		}
+	}
+
+	// Validate that if we are bound to a skel. mesh, and we have binding, and the binding type is set to skinning, that the skin cache is enabled for these
+	if (bHasNeedSkinningBinding && ParentMeshComponent)
+	{
+		if (USkeletalMeshComponent* ParentSkelMeshComponent = Cast<USkeletalMeshComponent>(ParentMeshComponent))
+		{
 			if (BindingAsset)
 			{
 				const uint32 SkelLODCount = ParentSkelMeshComponent->GetNumLODs();
@@ -2231,25 +2238,23 @@ void UGroomComponent::InitResources(bool bIsBindingReloading)
 				{
 					for (uint32 LODIt = 0, LODCount = GroomAsset->GetLODCount(); LODIt < LODCount; ++LODIt)
 					{
-						const uint32 EffectiveLODIt = FMath::Clamp(LODIt, 0, SkelLODCount-1);
+						const uint32 EffectiveLODIt = FMath::Clamp(LODIt, 0, SkelLODCount - 1);
 						const bool bSupportSkinCache = ParentSkelMeshComponent->IsSkinCacheAllowed(EffectiveLODIt);
 						const EGroomBindingType BindingType = GroomAsset->GetBindingType(GroupIt, LODIt);
 
 						if (BindingType == EGroomBindingType::Skinning && !bSupportSkinCache)
 						{
-							UE_LOG(LogHairStrands, Warning, TEXT("[Groom] Groom asset (Group:%d/%d) is set to use Skinning at LOD %d/%d while the parent skel. mesh does not support skin. cache for this LOD - Groom:%s - Skel.Mesh:%s"),
+							UE_LOG(LogHairStrands, Warning, TEXT("[Groom] Groom asset (Group:%d/%d) is set to use Skinning at LOD %d/%d while the parent skel. mesh does not support skin. cache at this LOD - Groom:%s - Skel.Mesh:%s"),
 								GroupIt,
 								GroupCount,
 								LODIt,
 								LODCount,
-								ComponentIt->GroomAsset ? *GroomAsset->GetPathName() : TEXT("."),
-								ParentMeshComponent->GetPathName());
+								*GroomAsset->GetPathName(),
+								*ParentMeshComponent->GetPathName());
 						}
 					}
 				}
 			}
-			#endif
-			ValidatedMeshComponent = ParentSkelMeshComponent->SkeletalMesh ? ParentMeshComponent : nullptr;
 		}
 	}
 
diff --git a/Engine/Plugins/Runtime/HairStrands/Source/HairStrandsCore/Private/GroomDebug.cpp b/Engine/Plugins/Runtime/HairStrands/Source/HairStrandsCore/Private/GroomDebug.cpp
--- a/Engine/Plugins/Runtime/HairStrands/Source/HairStrandsCore/Private/GroomDebug.cpp
+++ b/Engine/Plugins/Runtime/HairStrands/Source/HairStrandsCore/Private/GroomDebug.cpp
@@ -767,6 +767,12 @@ static const TCHAR* ToString(EHairBindingType In)
 	return TEXT("None");
 }
 
+FCachedGeometry GetCacheGeometryForHair(
+	FRDGBuilder& GraphBuilder,
+	FHairGroupInstance* Instance,
+	const FGPUSkinCache* SkinCache,
+	FGlobalShaderMap* ShaderMap);
+
 void RunHairStrandsDebug(
 	FRDGBuilder& GraphBuilder,
 	FGlobalShaderMap* ShaderMap,
@@ -797,10 +803,11 @@ void RunHairStrandsDebug(
 			float Y = ClusterY;
 			const FLinearColor InactiveColor(0.5, 0.5, 0.5);
 			const FLinearColor DebugColor(1, 1, 0);
-			const FLinearColor DebugGroupColor(0.5f, 0, 0);
+			const FLinearColor DebugGroupColor(1, 0.5f, 0);
+			const FLinearColor DebugWarningColor(0.75f, 0, 0);
 			FString Line;
 			// Active groom
-			// Name | Group x / x | LOD x / x | GeometryType | BindingType | Sim | RBF | VertexCount
+			// Name | Group x / x | LOD x / x | GeometryType | BindingType | Sim | RBF | VertexCount | Warning
 
 			Line = FString::Printf(TEXT("----------------------------------------------------------------"));
 			Canvas.DrawShadowedString(X, Y += YStep, *Line, GetStatsFont(), DebugColor);
@@ -812,6 +819,25 @@ void RunHairStrandsDebug(
 			{
 				FHairGroupInstance* Instance = static_cast<FHairGroupInstance*>(AbstractInstance);
 
+				// Check if a groom request a Skinning binding for the current LOD, that the parent has its skin cache enabled
+				bool bInvaliSkinBinding = false;
+				{
+					const int32 GroomLODIndex = Instance->HairGroupPublicData->LODIndex;
+					EHairBindingType OriginalBindingType = Instance->HairGroupPublicData->GetBindingType(GroomLODIndex);
+
+					if (OriginalBindingType == EHairBindingType::Skinning)
+					{
+						bool bHasValidSkinCacheData = false;
+						FHairStrandsProjectionMeshData::LOD MeshDataLOD;
+						const FCachedGeometry CachedGeometry = GetCacheGeometryForHair(GraphBuilder, Instance, SkinCache, ShaderMap);
+						for (const FCachedGeometry::Section& Section : CachedGeometry.Sections)
+						{
+							if (Section.LODIndex >= 0) { bHasValidSkinCacheData = true; break; }
+						}
+						bInvaliSkinBinding = !bHasValidSkinCacheData;
+					}
+				}
+
 				Line = FString::Printf(TEXT(" * Group:%d/%d | LOD:%1.2f/%d | GeometryType:%s | BindingType:%s | Sim:%d | RBF:%d | VertexCount:%d | Name: %s"),
 					Instance->Debug.GroupIndex,
 					Instance->Debug.GroupCount,
@@ -825,7 +851,15 @@ void RunHairStrandsDebug(
 					Instance->Guides.bHasGlobalInterpolation,
 					Instance->HairGroupPublicData->VertexCount,
 					*Instance->Debug.GroomAssetName);
-				Canvas.DrawShadowedString(X, Y += YStep, *Line, GetStatsFont(), DebugGroupColor);
+				int32 XOffset = Canvas.DrawShadowedString(X, Y += YStep, *Line, GetStatsFont(), DebugGroupColor);
+
+				// Warning/error
+				if (bInvaliSkinBinding)
+				{
+					Line = FString::Printf(TEXT(" - Warning: Request SKINNING binding but skel. mesh has no Skin Cache (%s)"),						
+						*Instance->Debug.MeshComponentName);
+					Canvas.DrawShadowedString(X + XOffset, Y += YStep, *Line, GetStatsFont(), DebugWarningColor);
+				}
 			}
 
 			const bool bFlush = false;
diff --git a/Engine/Plugins/Runtime/HairStrands/Source/HairStrandsCore/Private/GroomManager.cpp b/Engine/Plugins/Runtime/HairStrands/Source/HairStrandsCore/Private/GroomManager.cpp
--- a/Engine/Plugins/Runtime/HairStrands/Source/HairStrandsCore/Private/GroomManager.cpp
+++ b/Engine/Plugins/Runtime/HairStrands/Source/HairStrandsCore/Private/GroomManager.cpp
@@ -76,7 +76,7 @@ static bool IsInstanceFrustumCullingEnable()
 bool NeedsUpdateCardsMeshTriangles();
 
 // Returns the cached geometry of the underlying geometry on which a hair instance is attached to
-static FCachedGeometry GetCacheGeometryForHair(
+FCachedGeometry GetCacheGeometryForHair(
 	FRDGBuilder& GraphBuilder, 
 	FHairGroupInstance* Instance, 
 	const FGPUSkinCache* SkinCache, 
