Fix for static analysis warnings

diff --git a/Engine/Plugins/Runtime/HairStrands/Source/HairStrandsCore/Private/HairCardsVertexFactory.cpp b/Engine/Plugins/Runtime/HairStrands/Source/HairStrandsCore/Private/HairCardsVertexFactory.cpp
--- a/Engine/Plugins/Runtime/HairStrands/Source/HairStrandsCore/Private/HairCardsVertexFactory.cpp
+++ b/Engine/Plugins/Runtime/HairStrands/Source/HairStrandsCore/Private/HairCardsVertexFactory.cpp
@@ -217,12 +217,14 @@ void FHairCardsVertexFactory::ModifyCompilationEnvironment(const FVertexFactoryS
 {
 	FVertexFactory::ModifyCompilationEnvironment(Parameters, OutEnvironment);
 
-	const bool bUseGPUSceneAndPrimitiveIdStream = 
-		VF_CARDS_SUPPORT_GPU_SCENE && 
+	bool bUseGPUSceneAndPrimitiveIdStream = false;
+#if VF_CARDS_SUPPORT_GPU_SCENE
+	bUseGPUSceneAndPrimitiveIdStream = 
 		Parameters.VertexFactoryType->SupportsPrimitiveIdStream() && 
 		UseGPUScene(Parameters.Platform, GetMaxSupportedFeatureLevel(Parameters.Platform))
 		// TODO: support GPUScene on mobile
 		&& !IsMobilePlatform(Parameters.Platform);
+#endif
 
 	OutEnvironment.SetDefine(TEXT("VF_SUPPORTS_PRIMITIVE_SCENE_DATA"), bUseGPUSceneAndPrimitiveIdStream);
 	OutEnvironment.SetDefine(TEXT("HAIR_CARD_MESH_FACTORY"), TEXT("1"));
diff --git a/Engine/Plugins/Runtime/HairStrands/Source/HairStrandsCore/Private/HairStrandsVertexFactory.cpp b/Engine/Plugins/Runtime/HairStrands/Source/HairStrandsCore/Private/HairStrandsVertexFactory.cpp
--- a/Engine/Plugins/Runtime/HairStrands/Source/HairStrandsCore/Private/HairStrandsVertexFactory.cpp
+++ b/Engine/Plugins/Runtime/HairStrands/Source/HairStrandsCore/Private/HairStrandsVertexFactory.cpp
@@ -209,12 +209,14 @@ bool FHairStrandsVertexFactory::ShouldCompilePermutation(const FVertexFactorySha
 
 void FHairStrandsVertexFactory::ModifyCompilationEnvironment(const FVertexFactoryShaderPermutationParameters& Parameters, FShaderCompilerEnvironment& OutEnvironment)
 {
-	const bool bUseGPUSceneAndPrimitiveIdStream = 
-		VF_STRANDS_SUPPORT_GPU_SCENE
+	bool bUseGPUSceneAndPrimitiveIdStream = false;
+#if VF_STRANDS_SUPPORT_GPU_SCENE
+	bUseGPUSceneAndPrimitiveIdStream = 
 		&& Parameters.VertexFactoryType->SupportsPrimitiveIdStream() 
 		&& UseGPUScene(Parameters.Platform, GetMaxSupportedFeatureLevel(Parameters.Platform))
 		// TODO: support GPUScene on mobile
 		&& !IsMobilePlatform(Parameters.Platform);
+#endif
 	OutEnvironment.SetDefine(TEXT("VF_SUPPORTS_PRIMITIVE_SCENE_DATA"), bUseGPUSceneAndPrimitiveIdStream);
 	OutEnvironment.SetDefine(TEXT("HAIR_STRAND_MESH_FACTORY"), TEXT("1"));
 }
