Fix hair strands rendering being rendered during the base material pass.

diff --git a/Engine/Plugins/Runtime/HairStrands/Source/HairStrandsCore/Private/GroomComponent.cpp b/Engine/Plugins/Runtime/HairStrands/Source/HairStrandsCore/Private/GroomComponent.cpp
--- a/Engine/Plugins/Runtime/HairStrands/Source/HairStrandsCore/Private/GroomComponent.cpp
+++ b/Engine/Plugins/Runtime/HairStrands/Source/HairStrandsCore/Private/GroomComponent.cpp
@@ -840,7 +840,7 @@ public:
 
 		const bool bUseCardsOrMeshes = GeometryType == EHairGeometryType::Cards || GeometryType == EHairGeometryType::Meshes;
 		Mesh.CastShadow = bUseCardsOrMeshes;
-		Mesh.bUseForMaterial  = true;
+		Mesh.bUseForMaterial  = bUseCardsOrMeshes;
 		Mesh.bUseForDepthPass = bUseCardsOrMeshes;
 		Mesh.SegmentIndex = 0;
 
@@ -911,7 +911,7 @@ public:
 		Result.bDrawRelevance		= IsShown(View);
 		Result.bRenderInMainPass	= bUseCardsOrMesh && ShouldRenderInMainPass();
 		Result.bShadowRelevance		= IsShadowCast(View);
-		Result.bDynamicRelevance	= true;
+		Result.bDynamicRelevance	= bUseCardsOrMesh;
 		Result.bRenderCustomDepth	= ShouldRenderCustomDepth();
 		Result.bVelocityRelevance	= Result.bRenderInMainPass && bUseCardsOrMesh;
 		Result.bUsesLightingChannels= GetLightingChannelMask() != GetDefaultLightingChannelMask();
