Fix LOD colorization for hair cards/meshes

diff --git a/Engine/Plugins/Runtime/HairStrands/Source/HairStrandsCore/Private/GroomComponent.cpp b/Engine/Plugins/Runtime/HairStrands/Source/HairStrandsCore/Private/GroomComponent.cpp
--- a/Engine/Plugins/Runtime/HairStrands/Source/HairStrandsCore/Private/GroomComponent.cpp
+++ b/Engine/Plugins/Runtime/HairStrands/Source/HairStrandsCore/Private/GroomComponent.cpp
@@ -843,6 +843,7 @@ public:
 		Mesh.bUseForMaterial  = bUseCardsOrMeshes;
 		Mesh.bUseForDepthPass = bUseCardsOrMeshes;
 		Mesh.SegmentIndex = 0;
+		Mesh.VisualizeLODIndex = IntLODIndex;
 
 		FMeshBatchElement& BatchElement = Mesh.Elements[0];
 		BatchElement.IndexBuffer = IndexBuffer;
