Fix CIS issue with shipping build.

diff --git a/Engine/Plugins/Runtime/HairStrands/Source/HairStrandsCore/Private/GroomComponent.cpp b/Engine/Plugins/Runtime/HairStrands/Source/HairStrandsCore/Private/GroomComponent.cpp
--- a/Engine/Plugins/Runtime/HairStrands/Source/HairStrandsCore/Private/GroomComponent.cpp
+++ b/Engine/Plugins/Runtime/HairStrands/Source/HairStrandsCore/Private/GroomComponent.cpp
@@ -843,7 +843,9 @@ public:
 		Mesh.bUseForMaterial  = bUseCardsOrMeshes;
 		Mesh.bUseForDepthPass = bUseCardsOrMeshes;
 		Mesh.SegmentIndex = 0;
+		#if !(UE_BUILD_SHIPPING || UE_BUILD_TEST)
 		Mesh.VisualizeLODIndex = IntLODIndex;
+		#endif
 
 		FMeshBatchElement& BatchElement = Mesh.Elements[0];
 		BatchElement.IndexBuffer = IndexBuffer;
