Fixed issue with GroomCache ticking when it wasn't actually playing (notably in sequencer).
Fixed GroomCache not being refreshed when resetting the elapsed time.

diff --git a/Engine/Plugins/Runtime/HairStrands/Source/HairStrandsCore/Private/GroomComponent.cpp b/Engine/Plugins/Runtime/HairStrands/Source/HairStrandsCore/Private/GroomComponent.cpp
--- a/Engine/Plugins/Runtime/HairStrands/Source/HairStrandsCore/Private/GroomComponent.cpp
+++ b/Engine/Plugins/Runtime/HairStrands/Source/HairStrandsCore/Private/GroomComponent.cpp
@@ -1075,32 +1075,32 @@ public:
 				// In case the topology is varying, make sure the interpolated group data can hold the required number of vertices
 				InterpolatedGroupData.VertexData.PointsPosition.SetNum(NumVertices);
 
-			// Parallel batched interpolation
-			const int32 BatchSize = 1024;
-			const int32 BatchCount = (NumVertices + BatchSize - 1) / BatchSize;
+				// Parallel batched interpolation
+				const int32 BatchSize = 1024;
+				const int32 BatchCount = (NumVertices + BatchSize - 1) / BatchSize;
 
-			ParallelFor(BatchCount, [&](int32 BatchIndex)
-			{
-				const int32 Start = BatchIndex * BatchSize;
-				const int32 End = FMath::Min(Start + BatchSize, NumVertices); // one-past end index
+				ParallelFor(BatchCount, [&](int32 BatchIndex)
+					{
+						const int32 Start = BatchIndex * BatchSize;
+						const int32 End = FMath::Min(Start + BatchSize, NumVertices); // one-past end index
 
-				for (int32 VertexIndex = Start; VertexIndex < End; ++VertexIndex)
-				{
-					const FVector& CurrentPosition = CurrentGroupData.VertexData.PointsPosition[VertexIndex];
-					const FVector& NextPosition = NextGroupData.VertexData.PointsPosition[VertexIndex];
+						for (int32 VertexIndex = Start; VertexIndex < End; ++VertexIndex)
+						{
+							const FVector& CurrentPosition = CurrentGroupData.VertexData.PointsPosition[VertexIndex];
+							const FVector& NextPosition = NextGroupData.VertexData.PointsPosition[VertexIndex];
 
-					InterpolatedGroupData.VertexData.PointsPosition[VertexIndex] = FMath::Lerp(CurrentPosition, NextPosition, InterpolationFactor);
+							InterpolatedGroupData.VertexData.PointsPosition[VertexIndex] = FMath::Lerp(CurrentPosition, NextPosition, InterpolationFactor);
 
-					if (bHasRadiusData)
-					{
-						const float CurrentRadius = CurrentGroupData.VertexData.PointsRadius[VertexIndex];
-						const float NextRadius = NextGroupData.VertexData.PointsRadius[VertexIndex];
+							if (bHasRadiusData)
+							{
+								const float CurrentRadius = CurrentGroupData.VertexData.PointsRadius[VertexIndex];
+								const float NextRadius = NextGroupData.VertexData.PointsRadius[VertexIndex];
 
-						InterpolatedGroupData.VertexData.PointsRadius[VertexIndex] = FMath::Lerp(CurrentRadius, NextRadius, InterpolationFactor);
-					}
-				}
-			});
-		}
+								InterpolatedGroupData.VertexData.PointsRadius[VertexIndex] = FMath::Lerp(CurrentRadius, NextRadius, InterpolationFactor);
+							}
+						}
+					});
+			}
 			else
 			{
 				// Cannot interpolate, use the closest frame
@@ -1110,7 +1110,7 @@ public:
 				{
 					InterpolatedGroupData.VertexData.PointsRadius = InterpolationFactor < 0.5f ? CurrentGroupData.VertexData.PointsRadius : NextGroupData.VertexData.PointsRadius;
 				}
-	}
+			}
 		}
 	}
 
@@ -2980,6 +2980,11 @@ bool UGroomComponent::GetManualTick() const
 void UGroomComponent::ResetAnimationTime()
 {
 	ElapsedTime = 0.0f;
+	if (GroomCache && bRunning && GUseGroomCacheStreaming)
+	{
+		IGroomCacheStreamingManager::Get().PrefetchData(this);
+	}
+	UpdateGroomCache(ElapsedTime);
 }
 
 float UGroomComponent::GetAnimationTime() const
@@ -3092,7 +3097,8 @@ void UGroomComponent::TickComponent(float DeltaTime, enum ELevelTick TickType, F
 			MarkRenderTransformDirty();
 		}
 
-	if (GroomCache && bRunning && !bManualTick)
+	// Tick GroomCache only when playing
+	if (GroomCache && GetWorld()->AreActorsInitialized() && bRunning && !bManualTick)
 	{
 		ElapsedTime += DeltaTime;
 		UpdateGroomCache(ElapsedTime);
@@ -3238,6 +3244,11 @@ void UGroomComponent::PostEditChangeProperty(FPropertyChangedEvent& PropertyChan
 	}
 
 #if WITH_EDITOR
+	if (bGroomCacheChanged)
+	{
+		ResetAnimationTime();
+	}
+
 	if (bAssetChanged)
 	{
 		if (GroomAsset)
