Groom log error and fail when MeshLODData.GetNumSections == 0

diff --git a/Engine/Plugins/Runtime/HairStrands/Source/HairStrandsCore/Private/GroomBindingBuilder.cpp b/Engine/Plugins/Runtime/HairStrands/Source/HairStrandsCore/Private/GroomBindingBuilder.cpp
--- a/Engine/Plugins/Runtime/HairStrands/Source/HairStrandsCore/Private/GroomBindingBuilder.cpp
+++ b/Engine/Plugins/Runtime/HairStrands/Source/HairStrandsCore/Private/GroomBindingBuilder.cpp
@@ -1107,6 +1107,13 @@ namespace GroomBinding_RootProjection
 			FBox MeshBound;
 			MeshBound.Init();
 			const uint32 SectionCount = MeshLODData.GetNumSections();
+			
+			if ( SectionCount == 0 )
+			{
+				UE_LOG(LogHairStrands, Error, TEXT("[Groom] Binding asset could not be built. MeshLODData has 0 sections."));
+				return false;
+			}
+			
 			check(SectionCount > 0);
 			for (uint32 SectionIt = 0; SectionIt < SectionCount; ++SectionIt)
 			{
