Groom check for GetNumSections() == 0 and return error

diff --git a/Engine/Plugins/Runtime/HairStrands/Source/HairStrandsCore/Private/GroomBindingBuilder.cpp b/Engine/Plugins/Runtime/HairStrands/Source/HairStrandsCore/Private/GroomBindingBuilder.cpp
--- a/Engine/Plugins/Runtime/HairStrands/Source/HairStrandsCore/Private/GroomBindingBuilder.cpp
+++ b/Engine/Plugins/Runtime/HairStrands/Source/HairStrandsCore/Private/GroomBindingBuilder.cpp
@@ -1671,8 +1671,15 @@ namespace GroomBinding_Transfer
 			int32 LocalSourceSectionId = SourceSectionId;
 			int32 LocalTargetSectionId = TargetSectionId;
 			const IMeshLODData& TargetMeshLODData = InTargetMeshData->GetMeshLODData(TargetLODIndex);
+
 			if (LocalTargetSectionId >= TargetMeshLODData.GetNumSections())
 			{
+				if ( TargetMeshLODData.GetNumSections() == 0 )
+				{
+					UE_LOG(LogHairStrands, Error, TEXT("[Groom] Binding asset could not be built for LOD %d. TargetMeshLODData.GetNumSections() == 0."));
+					return false;
+				}
+
 				if (GHairStrandsBindingBuilderWarningEnable > 0)
 				{
 					UE_LOG(LogHairStrands, Warning, TEXT("[Groom] Binding asset will not respect the requested 'Matching section' %d for LOD %d. The target skeletal mesh does not have such a section for this LOD. Instead section 0 will be used for this given LOD."), TargetSectionId, TargetLODIndex);
