Add hair strands length scaling feature for morphing/transitioning application.

This 'new' version is an opt-in feature, enabled at component level. When enabled, it forced the allocation of deformed resources, which increase the memory footprint of a groom instance.

diff --git a/Engine/Plugins/Runtime/HairStrands/Source/HairStrandsCore/Private/GroomComponent.cpp b/Engine/Plugins/Runtime/HairStrands/Source/HairStrandsCore/Private/GroomComponent.cpp
--- a/Engine/Plugins/Runtime/HairStrands/Source/HairStrandsCore/Private/GroomComponent.cpp
+++ b/Engine/Plugins/Runtime/HairStrands/Source/HairStrandsCore/Private/GroomComponent.cpp
@@ -1537,6 +1537,15 @@ void UGroomComponent::SetUseCards(bool InbUseCards)
 	UpdateHairGroupsDescAndInvalidateRenderState();
 }
 
+void UGroomComponent::SetHairLengthScale(float InLengthScale)
+{
+	for (FHairGroupDesc& HairDesc : GroomGroupsDesc)
+	{
+		HairDesc.HairWidth = InLengthScale;
+	}
+	//UpdateHairGroupsDescAndInvalidateRenderState();
+}
+
 void UGroomComponent::SetForcedLOD(int32 LODIndex)
 {
 	if (GroomAsset)
@@ -2442,10 +2451,20 @@ void UGroomComponent::InitResources(bool bIsBindingReloading)
 						}
 					}
 				}
-			}
 
-			// The strands GroomCache needs the Strands.DeformedResource, but the guides GroomCache also needs it for the HairTangentPass
-			bNeedDynamicResources = bNeedDynamicResources || HairGroupInstance->Debug.GroomCacheType != EGroomCacheType::None;
+				// If length scale override is enabled, the group will potentially be scaled, which will modified the strands position. 
+				// For this reason we need to allocate deformed resources
+				if (GroomGroupsDesc[GroupIt].HairLengthScale_Override)
+				{
+					bNeedDynamicResources = true;
+				}
+
+				// The strands GroomCache needs the Strands.DeformedResource, but the guides GroomCache also needs it for the HairTangentPass
+				if (HairGroupInstance->Debug.GroomCacheType != EGroomCacheType::None)
+				{
+					bNeedDynamicResources = true;
+				}
+			}
 
 			#if RHI_RAYTRACING
 			if (IsHairRayTracingEnabled() && HairGroupInstance->Strands.Modifier.bUseHairRaytracingGeometry)
@@ -3199,6 +3218,7 @@ void UGroomComponent::PostEditChangeProperty(FPropertyChangedEvent& PropertyChan
 	const bool bIsBindingCompatible = UGroomBindingAsset::IsCompatible(GroomAsset, BindingAsset, bValidationEnable);
 	const bool bEnableSolverChanged = PropertyName == GET_MEMBER_NAME_CHECKED(FHairSimulationSolver, bEnableSimulation);
 	const bool bGroomCacheChanged = PropertyName == GET_MEMBER_NAME_CHECKED(UGroomComponent, GroomCache);
+	const bool bEnableLengthScaleChanged = PropertyName == GET_MEMBER_NAME_CHECKED(FHairGroupDesc, HairLengthScale);
 	if (!bIsBindingCompatible || !UGroomBindingAsset::IsBindingAssetValid(BindingAsset, false, bValidationEnable))
 	{
 		BindingAsset = nullptr;
@@ -3209,6 +3229,23 @@ void UGroomComponent::PostEditChangeProperty(FPropertyChangedEvent& PropertyChan
 		GroomAsset = nullptr;
 	}
 
+	// HairLengthScale_Override does not generate an edit event (only HairLengthScale does), so we manually go through 
+	// all the groups to check if there is a need for reallocating the resources
+	bool bEnableLengthScaleOverrideChanged = false;
+	if (bEnableLengthScaleChanged && GroomAsset && !bAssetChanged)
+	{
+		for (const FHairGroupInstance* Instance : HairGroupInstances)
+		{
+			const FHairGroupDesc GroupDesc = GetGroomGroupsDesc(GroomAsset, this, Instance->Debug.GroupIndex);
+			const bool bRecreate = Instance->Strands.DeformedResource == nullptr && GroupDesc.HairLengthScale_Override;
+			if (bRecreate)
+			{
+				bEnableLengthScaleOverrideChanged = true;
+				break;
+			}
+		}
+	}
+
 	bool bIsEventProcess = false;
 
 	// If the raytracing flag change, then force to recreate resources
@@ -3234,7 +3271,7 @@ void UGroomComponent::PostEditChangeProperty(FPropertyChangedEvent& PropertyChan
 	}
 	#endif
 
-	const bool bRecreateResources = bAssetChanged || bBindingAssetChanged || bGroomCacheChanged || PropertyThatChanged == nullptr || bSourceSkeletalMeshChanged || bRayTracingGeometryChanged || bEnableSolverChanged;
+	const bool bRecreateResources = bAssetChanged || bBindingAssetChanged || bGroomCacheChanged || bEnableLengthScaleOverrideChanged || PropertyThatChanged == nullptr || bSourceSkeletalMeshChanged || bRayTracingGeometryChanged || bEnableSolverChanged;
 	if (bRecreateResources)
 	{
 		// Release the resources before Super::PostEditChangeProperty so that they get
@@ -3278,7 +3315,8 @@ void UGroomComponent::PostEditChangeProperty(FPropertyChangedEvent& PropertyChan
 		PropertyName == GET_MEMBER_NAME_CHECKED(FHairGroupDesc, HairRaytracingRadiusScale) ||
 		PropertyName == GET_MEMBER_NAME_CHECKED(FHairGroupDesc, LODBias) ||
 		PropertyName == GET_MEMBER_NAME_CHECKED(FHairGroupDesc, bUseStableRasterization) ||
-		PropertyName == GET_MEMBER_NAME_CHECKED(FHairGroupDesc, bScatterSceneLighting))
+		PropertyName == GET_MEMBER_NAME_CHECKED(FHairGroupDesc, bScatterSceneLighting) ||
+		PropertyName == GET_MEMBER_NAME_CHECKED(FHairGroupDesc, HairLengthScale))
 	{
 		UpdateHairGroupsDescAndInvalidateRenderState();
 		bIsEventProcess = true;
diff --git a/Engine/Plugins/Runtime/HairStrands/Source/HairStrandsCore/Private/GroomDebug.cpp b/Engine/Plugins/Runtime/HairStrands/Source/HairStrandsCore/Private/GroomDebug.cpp
--- a/Engine/Plugins/Runtime/HairStrands/Source/HairStrandsCore/Private/GroomDebug.cpp
+++ b/Engine/Plugins/Runtime/HairStrands/Source/HairStrandsCore/Private/GroomDebug.cpp
@@ -838,7 +838,7 @@ void RunHairStrandsDebug(
 					}
 				}
 
-				Line = FString::Printf(TEXT(" * Group:%d/%d | LOD:%1.2f/%d | GeometryType:%s | BindingType:%s | Sim:%d | RBF:%d | VertexCount:%d | Name: %s"),
+				Line = FString::Printf(TEXT(" * Group:%d/%d | LOD:%1.2f/%d | GeometryType:%s | BindingType:%s | Sim:%d | RBF:%d | Clip:%d | VertexCount:%d | Name: %s"),
 					Instance->Debug.GroupIndex,
 					Instance->Debug.GroupCount,
 
@@ -849,6 +849,7 @@ void RunHairStrandsDebug(
 					ToString(Instance->BindingType),
 					Instance->Guides.bIsSimulationEnable,
 					Instance->Guides.bHasGlobalInterpolation,
+					Instance->Strands.Modifier.HairLengthScale_Override,
 					Instance->HairGroupPublicData->VertexCount,
 					*Instance->Debug.GroomAssetName);
 				int32 XOffset = Canvas.DrawShadowedString(X, Y += YStep, *Line, GetStatsFont(), DebugGroupColor);
diff --git a/Engine/Plugins/Runtime/HairStrands/Source/HairStrandsCore/Private/HairStrandsRendering.cpp b/Engine/Plugins/Runtime/HairStrands/Source/HairStrandsCore/Private/HairStrandsRendering.cpp
--- a/Engine/Plugins/Runtime/HairStrands/Source/HairStrandsCore/Private/HairStrandsRendering.cpp
+++ b/Engine/Plugins/Runtime/HairStrands/Source/HairStrandsCore/Private/HairStrandsRendering.cpp
@@ -543,6 +543,7 @@ class FHairInterpolationCS : public FGlobalShader
 		SHADER_PARAMETER(uint32, VertexCount)
 		SHADER_PARAMETER(uint32, DispatchCountX)
 		SHADER_PARAMETER(uint32, HairDebugMode)
+		SHADER_PARAMETER(float, HairLengthScale)
 		SHADER_PARAMETER(FVector3f, InRenderHairPositionOffset)
 		SHADER_PARAMETER(FVector3f, InSimHairPositionOffset)
 		SHADER_PARAMETER(uint32,  HairStrandsVF_bIsCullingEnable)
@@ -623,6 +624,7 @@ static void AddHairStrandsInterpolationPass(
 	const FHairGroupInstance* Instance,
 	const uint32 VertexCount,
 	const int32 MeshLODIndex,
+	const float HairLengthScale,
 	const bool bPatchedAttributeBuffer,
 	const EHairInterpolationType HairInterpolationType,
 	const EHairGeometryType InstanceGeometryType,
@@ -699,6 +701,7 @@ static void AddHairStrandsInterpolationPass(
 	Parameters->VertexToClusterIdBuffer = VertexToClusterIdBuffer;
 	
 	Parameters->LocalToWorldMatrix = Instance->LocalToWorld.ToMatrixWithScale();
+	Parameters->HairLengthScale = HairLengthScale;
 
 	// Debug rendering
 	// 1: Patch attributes for displaying guides influences onto the render strands
@@ -1747,6 +1750,7 @@ void ComputeHairStrandsInterpolation(
 						Instance,
 						Instance->Strands.RestResource->GetVertexCount(),
 						MeshLODIndex,
+						Instance->Strands.Modifier.HairLengthScale,
 						bDebugModePatchedAttributeBuffer,
 						Instance->Strands.HairInterpolationType,
 						InstanceGeometryType,
@@ -1953,6 +1957,7 @@ void ComputeHairStrandsInterpolation(
 						Instance,
 						LOD.Guides.RestResource->GetVertexCount(),
 						MeshLODIndex,
+						1.0f,
 						false,
 						LOD.Guides.HairInterpolationType,
 						InstanceGeometryType,
diff --git a/Engine/Plugins/Runtime/HairStrands/Source/HairStrandsCore/Private/HairStrandsVertexFactory.cpp b/Engine/Plugins/Runtime/HairStrands/Source/HairStrandsCore/Private/HairStrandsVertexFactory.cpp
--- a/Engine/Plugins/Runtime/HairStrands/Source/HairStrandsCore/Private/HairStrandsVertexFactory.cpp
+++ b/Engine/Plugins/Runtime/HairStrands/Source/HairStrandsCore/Private/HairStrandsVertexFactory.cpp
@@ -67,6 +67,7 @@ inline bool  UseScatterSceneLighting(const FHairGroupPublicData::FVertexFactoryI
 inline float GetMaxStrandRadius(const FHairGroupPublicData::FVertexFactoryInput& VFInput)						{ return VFInput.Strands.HairRadius; };
 inline float GetStrandRootScale(const FHairGroupPublicData::FVertexFactoryInput& VFInput)						{ return VFInput.Strands.HairRootScale; };
 inline float GetStrandTipScale(const FHairGroupPublicData::FVertexFactoryInput& VFInput)						{ return VFInput.Strands.HairTipScale; };
+inline float GetStrandLengthScale(const FHairGroupPublicData::FVertexFactoryInput& VFInput)						{ return VFInput.Strands.HairLengthScale; };
 inline float GetMaxStrandLength(const FHairGroupPublicData::FVertexFactoryInput& VFInput)						{ return VFInput.Strands.HairLength; };
 inline float GetHairDensity(const FHairGroupPublicData::FVertexFactoryInput& VFInput)							{ return VFInput.Strands.HairDensity; };
 //inline float GetHasAttribute1(const FHairGroupPublicData::FVertexFactoryInput& VFInput)							{ return VFInput.Strands.HasAttribute1; };
diff --git a/Engine/Plugins/Runtime/HairStrands/Source/HairStrandsCore/Public/GroomComponent.h b/Engine/Plugins/Runtime/HairStrands/Source/HairStrandsCore/Public/GroomComponent.h
--- a/Engine/Plugins/Runtime/HairStrands/Source/HairStrandsCore/Public/GroomComponent.h
+++ b/Engine/Plugins/Runtime/HairStrands/Source/HairStrandsCore/Public/GroomComponent.h
@@ -159,6 +159,10 @@ public:
 	UFUNCTION(BlueprintCallable, Category = "Simulation")
 	void SetEnableSimulation(bool bInEnableSimulation);
 
+	/* Accessor function for changing hair length scale from blueprint/sequencer */
+	UFUNCTION(BlueprintCallable, Category = "Groom")
+	void SetHairLengthScale(float Scale);
+
 	void SetStableRasterization(bool bEnable);
 	void SetGroomAsset(UGroomAsset* Asset, UGroomBindingAsset* InBinding, const bool bUpdateSimulation = true);
 	void SetHairRootScale(float Scale);
diff --git a/Engine/Plugins/Runtime/HairStrands/Source/HairStrandsCore/Public/GroomDesc.h b/Engine/Plugins/Runtime/HairStrands/Source/HairStrandsCore/Public/GroomDesc.h
--- a/Engine/Plugins/Runtime/HairStrands/Source/HairStrandsCore/Public/GroomDesc.h
+++ b/Engine/Plugins/Runtime/HairStrands/Source/HairStrandsCore/Public/GroomDesc.h
@@ -75,6 +75,12 @@ struct FHairGroupDesc
 	UPROPERTY()
 	bool bSupportVoxelization_Override = 0;
 
+	/** When enabled, Length Scale allow to scale the length of the hair. */
+	UPROPERTY(EditAnywhere, BlueprintReadWrite, interp, Category = "Groom", meta = (ClampMin = "0.0", ClampMax = "1.0", UIMin = "0.0", UIMax = "1.0", SliderExponent = 1, editcondition = "HairLengthScale_Override"))
+	float HairLengthScale = 1.f;
+	UPROPERTY(EditAnywhere, BlueprintReadWrite, Category = Overrides, meta = (PinHiddenByDefault, InlineEditConditionToggle))
+	bool HairLengthScale_Override = false;
+
 	/** Force a specific LOD index */
 	UPROPERTY()
 	int32 LODForcedIndex = -1;
diff --git a/Engine/Plugins/Runtime/HairStrands/Source/HairStrandsEditor/Private/GroomComponentDetailsCustomization.cpp b/Engine/Plugins/Runtime/HairStrands/Source/HairStrandsEditor/Private/GroomComponentDetailsCustomization.cpp
--- a/Engine/Plugins/Runtime/HairStrands/Source/HairStrandsEditor/Private/GroomComponentDetailsCustomization.cpp
+++ b/Engine/Plugins/Runtime/HairStrands/Source/HairStrandsEditor/Private/GroomComponentDetailsCustomization.cpp
@@ -112,6 +112,7 @@ bool FGroomComponentDetailsCustomization::CommonResetToDefault(TSharedPtr<IPrope
 	HAIR_RESET_COMPONENT(HairWidth);
 	HAIR_RESET_COMPONENT(HairRootScale);
 	HAIR_RESET_COMPONENT(HairTipScale);
+	HAIR_RESET_COMPONENT(HairLengthScale);
 	HAIR_RESET_COMPONENT(HairShadowDensity);
 	HAIR_RESET_COMPONENT(HairRaytracingRadiusScale);
 	HAIR_RESET_COMPONENT(bUseHairRaytracingGeometry);
diff --git a/Engine/Shaders/Private/HairStrands/HairStrandsInterpolation.usf b/Engine/Shaders/Private/HairStrands/HairStrandsInterpolation.usf
--- a/Engine/Shaders/Private/HairStrands/HairStrandsInterpolation.usf
+++ b/Engine/Shaders/Private/HairStrands/HairStrandsInterpolation.usf
@@ -7,6 +7,13 @@
 #include "HairStrandsClusterCommon.ush"
 #include "HairStrandsVertexFactoryCommon.ush"
 
+// Should be moved into HairStrandsVertexFactoryCommon
+float UnpackHairControlPointCoordU(uint InPackedDataW)
+{
+	const uint PackedAlpha = asuint(InPackedDataW);
+	return ((PackedAlpha & 0xFF00) >> 8) / 255.f;
+}
+
 #if COMPILER_METAL || COMPILER_VULKAN
 #define ALLOW_DEBUG_RENDERING 0
 #else
@@ -37,6 +44,7 @@
 uint VertexCount;
 uint DispatchCountX;
 uint HairDebugMode;
+float HairLengthScale;
 float3 InRenderHairPositionOffset;
 float3 InSimHairPositionOffset;
 float4x4 LocalToWorldMatrix;
@@ -178,6 +186,21 @@ float3 ComputeRotateGeometryOffset(uint GuideIndex, float GuideVertexWeight,
 	return (LocalDeformedGuidePoint - LocalRestGuidePoint) * GuideVertexWeight;
 }
 
+uint4 WriteOutput(float3 WorldPosition, float3 OutPositionOffset, uint PackedDataW)
+{
+	uint3 EncodedRelativePosition = f32tof16(WorldPosition - OutPositionOffset);
+	if (HairLengthScale < 1.0f)
+	{
+		const float CoordU = UnpackHairControlPointCoordU(PackedDataW);
+		if (HairLengthScale < CoordU)
+		{
+			EncodedRelativePosition = f32tof16(INFINITE_FLOAT).xxx;
+		}
+	}
+
+	return uint4(EncodedRelativePosition, PackedDataW);
+}
+
 [numthreads(GROUP_SIZE, 1, 1)]
 void MainCS(uint2 DispatchThreadId : SV_DispatchThreadID)
 {
@@ -198,7 +221,6 @@ void MainCS(uint2 DispatchThreadId : SV_DispatchThreadID)
 	}
 #endif
 
-	uint4 OutPosition = 0;
 	if (VertexIndex < VertexCount)
 	{
 		uint3  GuideVertexIndices = 0;
@@ -264,8 +286,7 @@ void MainCS(uint2 DispatchThreadId : SV_DispatchThreadID)
 
 			WorldVertexPosition = ControlPoint + CurrOffset;
 
-			OutPosition = uint4(f32tof16(WorldVertexPosition - OutRenHairPositionOffset), PackedRestPosition.w);
-			OutRenderDeformedPositionBuffer[VertexIndex] = OutPosition;
+			OutRenderDeformedPositionBuffer[VertexIndex] = WriteOutput(WorldVertexPosition, OutRenHairPositionOffset, PackedRestPosition.w);
 		}
 		#endif	
 
@@ -344,8 +365,7 @@ void MainCS(uint2 DispatchThreadId : SV_DispatchThreadID)
 			}
 			#endif
 
-			OutPosition = uint4(f32tof16(WorldVertexPosition - OutRenHairPositionOffset), PackedRestPosition.w);
-			OutRenderDeformedPositionBuffer[VertexIndex] = OutPosition;
+			OutRenderDeformedPositionBuffer[VertexIndex] = WriteOutput(WorldVertexPosition, OutRenHairPositionOffset, PackedRestPosition.w);
 		}
 		#endif	
 
@@ -387,8 +407,7 @@ void MainCS(uint2 DispatchThreadId : SV_DispatchThreadID)
 
 			WorldVertexPosition = ControlPoint + CurrOffset;
 
-			OutPosition = uint4(f32tof16(WorldVertexPosition - OutRenHairPositionOffset), PackedRestPosition.w);
-			OutRenderDeformedPositionBuffer[VertexIndex] = OutPosition;
+			OutRenderDeformedPositionBuffer[VertexIndex] = WriteOutput(WorldVertexPosition, OutRenHairPositionOffset, PackedRestPosition.w);
 		}
 		#endif	
 
@@ -437,8 +456,7 @@ void MainCS(uint2 DispatchThreadId : SV_DispatchThreadID)
 
 			WorldVertexPosition = ControlPoint + CurrOffset;
 
-			OutPosition = uint4(f32tof16(WorldVertexPosition - OutRenHairPositionOffset), PackedRestPosition.w);
-			OutRenderDeformedPositionBuffer[VertexIndex] = OutPosition;
+			OutRenderDeformedPositionBuffer[VertexIndex] = WriteOutput(WorldVertexPosition, OutRenHairPositionOffset, PackedRestPosition.w);
 		}
 		#endif
 
@@ -484,8 +502,7 @@ void MainCS(uint2 DispatchThreadId : SV_DispatchThreadID)
 
 			WorldVertexPosition = CurrOffset;
 
-			OutPosition = uint4(f32tof16(WorldVertexPosition - OutRenHairPositionOffset), PackedRestPosition.w);
-			OutRenderDeformedPositionBuffer[VertexIndex] = OutPosition;
+			OutRenderDeformedPositionBuffer[VertexIndex] = WriteOutput(WorldVertexPosition, OutRenHairPositionOffset, PackedRestPosition.w);
 		}
 		#endif
 
diff --git a/Engine/Source/Runtime/Renderer/Public/HairStrandsInterface.h b/Engine/Source/Runtime/Renderer/Public/HairStrandsInterface.h
--- a/Engine/Source/Runtime/Renderer/Public/HairStrandsInterface.h
+++ b/Engine/Source/Runtime/Renderer/Public/HairStrandsInterface.h
@@ -256,6 +256,7 @@ public:
 			float HairRadius = 0;
 			float HairRootScale = 0;
 			float HairTipScale = 0;
+			float HairLengthScale = 0;
 			float HairLength = 0;
 			float HairDensity = 0;
 			bool bUseStableRasterization = false;
