Ensure Hair simulation use initialized resources.

This is a short term fix to unblock. Proper investigation will follow.

diff --git a/Engine/Plugins/Runtime/HairStrands/Source/HairStrandsCore/Private/Niagara/NiagaraDataInterfaceHairStrands.cpp b/Engine/Plugins/Runtime/HairStrands/Source/HairStrandsCore/Private/Niagara/NiagaraDataInterfaceHairStrands.cpp
--- a/Engine/Plugins/Runtime/HairStrands/Source/HairStrandsCore/Private/Niagara/NiagaraDataInterfaceHairStrands.cpp
+++ b/Engine/Plugins/Runtime/HairStrands/Source/HairStrandsCore/Private/Niagara/NiagaraDataInterfaceHairStrands.cpp
@@ -638,7 +638,7 @@ struct FNDIHairStrandsParametersCS : public FNiagaraDataInterfaceParametersCS
 		const bool bIsHairValid = ProxyData != nullptr && ProxyData->HairStrandsBuffer && ProxyData->HairStrandsBuffer->IsInitialized();
 		const bool bHasSkinningBinding = bIsHairValid && ProxyData->HairGroupInstance && ProxyData->HairGroupInstance->BindingType == EHairBindingType::Skinning;
 		const bool bIsRootValid = bIsHairValid && ProxyData->HairStrandsBuffer->SourceDeformedRootResources && bHasSkinningBinding;
-		const bool bIsRestValid = bIsHairValid && ProxyData->HairStrandsBuffer->SourceRestResources;
+		const bool bIsRestValid = bIsHairValid && ProxyData->HairStrandsBuffer->SourceRestResources && ProxyData->HairStrandsBuffer->SourceRestResources->PositionBuffer.SRV;
 		const bool bIsDeformedValid = bIsHairValid && ProxyData->HairStrandsBuffer->SourceDeformedResources && ProxyData->HairStrandsBuffer->SourceDeformedResources->IsInitialized();
 		const bool bHasLODSwitched = bIsHairValid && ProxyData->HairGroupInstance && ProxyData->HairGroupInstance->HairGroupPublicData && ProxyData->HairGroupInstance->HairGroupPublicData->VFInput.bHasLODSwitch;
 
