Remove compression flag from Groom asset.

Compression does not bring gain and are a potential source of issue due to its compressed version in DDC but not in cooked data.

diff --git a/Engine/Plugins/Runtime/HairStrands/Source/HairStrandsCore/Private/GroomBindingBuilder.cpp b/Engine/Plugins/Runtime/HairStrands/Source/HairStrandsCore/Private/GroomBindingBuilder.cpp
--- a/Engine/Plugins/Runtime/HairStrands/Source/HairStrandsCore/Private/GroomBindingBuilder.cpp
+++ b/Engine/Plugins/Runtime/HairStrands/Source/HairStrandsCore/Private/GroomBindingBuilder.cpp
@@ -57,7 +57,7 @@ static FAutoConsoleVariableRef CVarHairStrandsBindingBuilderWarningEnable(TEXT("
 FString FGroomBindingBuilder::GetVersion()
 {
 	// Important to update the version when groom building changes
-	return TEXT("2d");
+	return TEXT("2e");
 }
 
 ///////////////////////////////////////////////////////////////////////////////////////////////////
diff --git a/Engine/Plugins/Runtime/HairStrands/Source/HairStrandsCore/Private/GroomBuilder.cpp b/Engine/Plugins/Runtime/HairStrands/Source/HairStrandsCore/Private/GroomBuilder.cpp
--- a/Engine/Plugins/Runtime/HairStrands/Source/HairStrandsCore/Private/GroomBuilder.cpp
+++ b/Engine/Plugins/Runtime/HairStrands/Source/HairStrandsCore/Private/GroomBuilder.cpp
@@ -33,7 +33,7 @@ static FAutoConsoleVariableRef CVarHairClusterBuilder_MaxVoxelResolution(TEXT("r
 FString FGroomBuilder::GetVersion()
 {
 	// Important to update the version when groom building changes
-	return TEXT("2o");
+	return TEXT("2p");
 }
 
 namespace FHairStrandsDecimation
diff --git a/Engine/Plugins/Runtime/HairStrands/Source/HairStrandsCore/Private/GroomResources.cpp b/Engine/Plugins/Runtime/HairStrands/Source/HairStrandsCore/Private/GroomResources.cpp
--- a/Engine/Plugins/Runtime/HairStrands/Source/HairStrandsCore/Private/GroomResources.cpp
+++ b/Engine/Plugins/Runtime/HairStrands/Source/HairStrandsCore/Private/GroomResources.cpp
@@ -23,10 +23,6 @@
 #include "RenderTargetPool.h"
 #include "GroomBindingBuilder.h"
 
-#ifndef CLEAR_COMPRESSION
-#define CLEAR_COMPRESSION 1
-#endif
-
 static int32 GHairStrandsBulkData_ReleaseAfterUse = 0;
 static FAutoConsoleVariableRef CVarHairStrandsBulkData_ReleaseAfterUse(TEXT("r.HairStrands.Strands.BulkData.ReleaseAfterUse"), GHairStrandsBulkData_ReleaseAfterUse, TEXT("Release CPU bulk data once hair groom/groom binding asset GPU resources are created. This saves memory"));
 
@@ -851,21 +847,11 @@ void FHairStrandsClusterCullingBulkData::Serialize(FArchive& Ar, UObject* Owner)
 
 	if (Ar.IsSaving())
 	{
-		const uint32 BulkFlags = BULKDATA_Force_NOT_InlinePayload | BULKDATA_SerializeCompressed;
+		const uint32 BulkFlags = BULKDATA_Force_NOT_InlinePayload;
 		ClusterLODInfos.SetBulkDataFlags(BulkFlags);
 		VertexToClusterIds.SetBulkDataFlags(BulkFlags);
 		ClusterVertexIds.SetBulkDataFlags(BulkFlags);
 		PackedClusterInfos.SetBulkDataFlags(BulkFlags);
-
-		#if CLEAR_COMPRESSION
-		if (Ar.IsCooking())
-		{
-			ClusterLODInfos.ClearBulkDataFlags(BULKDATA_SerializeCompressed);
-			VertexToClusterIds.ClearBulkDataFlags(BULKDATA_SerializeCompressed);
-			ClusterVertexIds.ClearBulkDataFlags(BULKDATA_SerializeCompressed);
-			PackedClusterInfos.ClearBulkDataFlags(BULKDATA_SerializeCompressed);
-		}
-		#endif
 	}
 
 	if (ClusterLODCount)
@@ -1161,7 +1147,7 @@ static void InternalSerialize(FArchive& Ar, UObject* Owner, FHairStrandsRootBulk
 
 	if (Ar.IsSaving())
 	{
-		const uint32 BulkFlags = BULKDATA_Force_NOT_InlinePayload | BULKDATA_SerializeCompressed;
+		const uint32 BulkFlags = BULKDATA_Force_NOT_InlinePayload;
 		LOD.RootTriangleIndexBuffer.SetBulkDataFlags(BulkFlags);
 		LOD.RootTriangleBarycentricBuffer.SetBulkDataFlags(BulkFlags);
 		LOD.RestRootTrianglePosition0Buffer.SetBulkDataFlags(BulkFlags);
@@ -1171,21 +1157,6 @@ static void InternalSerialize(FArchive& Ar, UObject* Owner, FHairStrandsRootBulk
 		LOD.MeshInterpolationWeightsBuffer.SetBulkDataFlags(BulkFlags);
 		LOD.MeshSampleIndicesBuffer.SetBulkDataFlags(BulkFlags);
 		LOD.RestSamplePositionsBuffer.SetBulkDataFlags(BulkFlags);
-
-		#if CLEAR_COMPRESSION
-		if (Ar.IsCooking())
-		{
-			LOD.RootTriangleIndexBuffer.ClearBulkDataFlags(BULKDATA_SerializeCompressed);
-			LOD.RootTriangleBarycentricBuffer.ClearBulkDataFlags(BULKDATA_SerializeCompressed);
-			LOD.RestRootTrianglePosition0Buffer.ClearBulkDataFlags(BULKDATA_SerializeCompressed);
-			LOD.RestRootTrianglePosition1Buffer.ClearBulkDataFlags(BULKDATA_SerializeCompressed);
-			LOD.RestRootTrianglePosition2Buffer.ClearBulkDataFlags(BULKDATA_SerializeCompressed);
-
-			LOD.MeshInterpolationWeightsBuffer.ClearBulkDataFlags(BULKDATA_SerializeCompressed);
-			LOD.MeshSampleIndicesBuffer.ClearBulkDataFlags(BULKDATA_SerializeCompressed);
-			LOD.RestSamplePositionsBuffer.ClearBulkDataFlags(BULKDATA_SerializeCompressed);
-		}
-		#endif
 	}
 
 	Ar << LOD.LODIndex;
@@ -1225,14 +1196,8 @@ void FHairStrandsRootBulkData::Serialize(FArchive& Ar, UObject* Owner)
 	bool bAttemptFileMapping = false;
 
 	{
-		const uint32 BulkFlags = BULKDATA_Force_NOT_InlinePayload | BULKDATA_SerializeCompressed;
+		const uint32 BulkFlags = BULKDATA_Force_NOT_InlinePayload;
 		VertexToCurveIndexBuffer.SetBulkDataFlags(BulkFlags);
-		#if CLEAR_COMPRESSION
-		if (Ar.IsCooking())
-		{
-			VertexToCurveIndexBuffer.ClearBulkDataFlags(BULKDATA_SerializeCompressed);
-		}
-		#endif
 	}
 
 	if (!Ar.IsObjectReferenceCollector())
diff --git a/Engine/Plugins/Runtime/HairStrands/Source/HairStrandsCore/Private/HairStrandsDatas.cpp b/Engine/Plugins/Runtime/HairStrands/Source/HairStrandsCore/Private/HairStrandsDatas.cpp
--- a/Engine/Plugins/Runtime/HairStrands/Source/HairStrandsCore/Private/HairStrandsDatas.cpp
+++ b/Engine/Plugins/Runtime/HairStrands/Source/HairStrandsCore/Private/HairStrandsDatas.cpp
@@ -4,10 +4,6 @@
 #include "UObject/ReleaseObjectVersion.h"
 #include "UObject/UE5ReleaseStreamObjectVersion.h"
 
-#ifndef CLEAR_COMPRESSION
-#define CLEAR_COMPRESSION 1
-#endif
-
 void FHairStrandsInterpolationDatas::SetNum(const uint32 NumCurves)
 {
 	PointsSimCurvesVertexWeights.SetNum(NumCurves);
@@ -187,21 +183,11 @@ void FHairStrandsInterpolationBulkData::Serialize(FArchive& Ar, UObject* Owner)
 
 	if (Ar.IsSaving())
 	{
-		const uint32 BulkFlags = BULKDATA_Force_NOT_InlinePayload | BULKDATA_SerializeCompressed;
+		const uint32 BulkFlags = BULKDATA_Force_NOT_InlinePayload;
 		Interpolation.SetBulkDataFlags(BulkFlags);
 		Interpolation0.SetBulkDataFlags(BulkFlags);
 		Interpolation1.SetBulkDataFlags(BulkFlags);
 		SimRootPointIndex.SetBulkDataFlags(BulkFlags);
-
-		#if CLEAR_COMPRESSION
-		if (Ar.IsCooking())
-		{
-			Interpolation.ClearBulkDataFlags(BULKDATA_SerializeCompressed);
-			Interpolation0.ClearBulkDataFlags(BULKDATA_SerializeCompressed);
-			Interpolation1.ClearBulkDataFlags(BULKDATA_SerializeCompressed);
-			SimRootPointIndex.ClearBulkDataFlags(BULKDATA_SerializeCompressed);
-		}
-		#endif
 	}
 
 	Ar << Flags;
@@ -248,23 +234,12 @@ void FHairStrandsBulkData::Serialize(FArchive& Ar, UObject* Owner)
 	// but only when we explicitly take action to load it
 	if (Ar.IsSaving())
 	{
-		const uint32 BulkFlags = BULKDATA_Force_NOT_InlinePayload | BULKDATA_SerializeCompressed;
+		const uint32 BulkFlags = BULKDATA_Force_NOT_InlinePayload;
 		Positions.SetBulkDataFlags(BulkFlags);
 		Attributes0.SetBulkDataFlags(BulkFlags);
 		Attributes1.SetBulkDataFlags(BulkFlags);
 		Materials.SetBulkDataFlags(BulkFlags);
 		CurveOffsets.SetBulkDataFlags(BulkFlags);
-
-		#if CLEAR_COMPRESSION
-		if (Ar.IsCooking())
-		{
-			Positions.ClearBulkDataFlags(BULKDATA_SerializeCompressed);
-			Attributes0.ClearBulkDataFlags(BULKDATA_SerializeCompressed);
-			Attributes1.ClearBulkDataFlags(BULKDATA_SerializeCompressed);
-			Materials.ClearBulkDataFlags(BULKDATA_SerializeCompressed);
-			CurveOffsets.ClearBulkDataFlags(BULKDATA_SerializeCompressed);
-		}
-		#endif
 	}
 
 	if (!!(Flags & DataFlags_HasData))
