Fix false-positive hair strands binding check.

diff --git a/Engine/Plugins/Runtime/HairStrands/Source/HairStrandsCore/Private/GroomAsset.cpp b/Engine/Plugins/Runtime/HairStrands/Source/HairStrandsCore/Private/GroomAsset.cpp
--- a/Engine/Plugins/Runtime/HairStrands/Source/HairStrandsCore/Private/GroomAsset.cpp
+++ b/Engine/Plugins/Runtime/HairStrands/Source/HairStrandsCore/Private/GroomAsset.cpp
@@ -3114,6 +3114,21 @@ EGroomBindingType UGroomAsset::GetBindingType(int32 GroupIndex, int32 LODIndex)
 	return HairGroupsLOD[GroupIndex].LODs[LODIndex].BindingType;
 }
 
+bool UGroomAsset::IsVisible(int32 GroupIndex, int32 LODIndex) const
+{
+	if (GroupIndex < 0 || GroupIndex >= HairGroupsLOD.Num())
+	{
+		return false;
+	}
+
+	if (LODIndex < 0 || LODIndex >= HairGroupsLOD[GroupIndex].LODs.Num())
+	{
+		return false;
+	}
+
+	return HairGroupsLOD[GroupIndex].LODs[LODIndex].bVisible;
+}
+
 bool UGroomAsset::IsSimulationEnable(int32 GroupIndex, int32 LODIndex) const 
 {
 	if (GroupIndex < 0 || GroupIndex >= HairGroupsLOD.Num() || !IsHairStrandsSimulationEnable())
diff --git a/Engine/Plugins/Runtime/HairStrands/Source/HairStrandsCore/Private/GroomComponent.cpp b/Engine/Plugins/Runtime/HairStrands/Source/HairStrandsCore/Private/GroomComponent.cpp
--- a/Engine/Plugins/Runtime/HairStrands/Source/HairStrandsCore/Private/GroomComponent.cpp
+++ b/Engine/Plugins/Runtime/HairStrands/Source/HairStrandsCore/Private/GroomComponent.cpp
@@ -2253,8 +2253,9 @@ void UGroomComponent::InitResources(bool bIsBindingReloading)
 						const uint32 EffectiveLODIt = FMath::Clamp(LODIt, 0, SkelLODCount - 1);
 						const bool bSupportSkinCache = ParentSkelMeshComponent->IsSkinCacheAllowed(EffectiveLODIt);
 						const EGroomBindingType BindingType = GroomAsset->GetBindingType(GroupIt, LODIt);
+						const bool bIsVisible = GroomAsset->IsVisible(GroupIt, LODIt);
 
-						if (BindingType == EGroomBindingType::Skinning && !bSupportSkinCache)
+						if (BindingType == EGroomBindingType::Skinning && !bSupportSkinCache && bIsVisible)
 						{
 							UE_LOG(LogHairStrands, Warning, TEXT("[Groom] Groom asset (Group:%d/%d) is set to use Skinning at LOD %d/%d while the parent skel. mesh does not support skin. cache at this LOD - Groom:%s - Skel.Mesh:%s"),
 								GroupIt,
diff --git a/Engine/Plugins/Runtime/HairStrands/Source/HairStrandsCore/Public/GroomAsset.h b/Engine/Plugins/Runtime/HairStrands/Source/HairStrandsCore/Public/GroomAsset.h
--- a/Engine/Plugins/Runtime/HairStrands/Source/HairStrandsCore/Public/GroomAsset.h
+++ b/Engine/Plugins/Runtime/HairStrands/Source/HairStrandsCore/Public/GroomAsset.h
@@ -537,6 +537,7 @@ public:
 
 	EGroomGeometryType GetGeometryType(int32 GroupIndex, int32 LODIndex) const;
 	EGroomBindingType GetBindingType(int32 GroupIndex, int32 LODIndex) const;
+	bool IsVisible(int32 GroupIndex, int32 LODIndex) const;
 	bool IsSimulationEnable(int32 GroupIndex, int32 LODIndex) const;
 	bool IsSimulationEnable() const;
 	bool IsGlobalInterpolationEnable(int32 GroupIndex, int32 LODIndex) const;
