Fix crash when toggling r.hairstrands.enable.

The result might not be correct (e.g. missing hair) when toggling this CVar, while some groom component are spawn and the CVar is off, then on again.

diff --git a/Engine/Plugins/Runtime/HairStrands/Source/HairStrandsCore/Private/GroomComponent.cpp b/Engine/Plugins/Runtime/HairStrands/Source/HairStrandsCore/Private/GroomComponent.cpp
--- a/Engine/Plugins/Runtime/HairStrands/Source/HairStrandsCore/Private/GroomComponent.cpp
+++ b/Engine/Plugins/Runtime/HairStrands/Source/HairStrandsCore/Private/GroomComponent.cpp
@@ -569,6 +569,14 @@ public:
 		if (!IsHairRayTracingEnabled() || HairGroups.Num() == 0)
 			return;
 
+		const EShaderPlatform Platform = Context.ReferenceView->GetShaderPlatform();
+		if (!IsHairStrandsEnabled(EHairStrandsShaderType::Strands, Platform) &&
+			!IsHairStrandsEnabled(EHairStrandsShaderType::Cards, Platform) &&
+			!IsHairStrandsEnabled(EHairStrandsShaderType::Meshes, Platform))
+		{
+			return;
+		}
+
 		const bool bWireframe = AllowDebugViewmodes() && Context.ReferenceViewFamily.EngineShowFlags.Wireframe;
 		if (bWireframe)
 			return;
@@ -768,7 +776,7 @@ public:
 		const int32 IntLODIndex = Instance->HairGroupPublicData->GetIntLODIndex();
 		const FVertexFactory* VertexFactory = nullptr;
 		FIndexBuffer* IndexBuffer = nullptr;
-		FMaterialRenderProxy* MaterialRenderProxy = nullptr;
+		FMaterialRenderProxy* MaterialRenderProxy = Debug_MaterialProxy;
 
 		uint32 NumPrimitive = 0;
 		uint32 HairVertexCount = 0;
@@ -788,7 +796,11 @@ public:
 			NumPrimitive = HairVertexCount / 3;
 			IndexBuffer = &Instance->Meshes.LODs[IntLODIndex].RestResource->IndexBuffer;
 			bUseCulling = false;
-			MaterialRenderProxy = Debug_MaterialProxy == nullptr ? Instance->Meshes.LODs[IntLODIndex].Material->GetRenderProxy() : Debug_MaterialProxy;
+			if (MaterialRenderProxy == nullptr && Instance->Meshes.LODs[IntLODIndex].Material)
+			{
+				MaterialRenderProxy = Instance->Meshes.LODs[IntLODIndex].Material->GetRenderProxy();
+			}
+			bWireframe = AllowDebugViewmodes() && ViewFamily.EngineShowFlags.Wireframe;
 		}
 		else if (GeometryType == EHairGeometryType::Cards)
 		{
@@ -804,7 +816,10 @@ public:
 			NumPrimitive = HairVertexCount / 3;
 			IndexBuffer = &Instance->Cards.LODs[IntLODIndex].RestResource->RestIndexBuffer;
 			bUseCulling = false;
-			MaterialRenderProxy = Debug_MaterialProxy == nullptr ? Instance->Cards.LODs[IntLODIndex].Material->GetRenderProxy() : Debug_MaterialProxy;
+			if (MaterialRenderProxy == nullptr && Instance->Cards.LODs[IntLODIndex].Material)
+			{
+				MaterialRenderProxy = Instance->Cards.LODs[IntLODIndex].Material->GetRenderProxy();
+			}
 			bWireframe = AllowDebugViewmodes() && ViewFamily.EngineShowFlags.Wireframe;
 		}
 		else // if (GeometryType == EHairGeometryType::Strands)
@@ -814,7 +829,10 @@ public:
 			MaxVertexIndex = HairVertexCount * 6;
 			bUseCulling = Instance->Strands.bIsCullingEnabled;
 			NumPrimitive = bUseCulling ? 0 : HairVertexCount * 2;
-			MaterialRenderProxy = Debug_MaterialProxy == nullptr ? Instance->Strands.Material->GetRenderProxy() : Debug_MaterialProxy;
+			if (MaterialRenderProxy == nullptr && Instance->Strands.Material)
+			{
+				MaterialRenderProxy = Instance->Strands.Material->GetRenderProxy();
+			}
 			bWireframe = AllowDebugViewmodes() && ViewFamily.EngineShowFlags.Wireframe;
 		}
 
