Fix CIS static analysis error.

diff --git a/Engine/Plugins/Runtime/HairStrands/Source/HairStrandsCore/Private/GroomComponent.cpp b/Engine/Plugins/Runtime/HairStrands/Source/HairStrandsCore/Private/GroomComponent.cpp
--- a/Engine/Plugins/Runtime/HairStrands/Source/HairStrandsCore/Private/GroomComponent.cpp
+++ b/Engine/Plugins/Runtime/HairStrands/Source/HairStrandsCore/Private/GroomComponent.cpp
@@ -2296,7 +2296,10 @@ void UGroomComponent::InitResources(bool bIsBindingReloading)
 		RegisteredMeshComponent = ValidatedMeshComponent;
 		AddTickPrerequisiteComponent(ValidatedMeshComponent);
 		bUseAttachParentBound = true;
-		LocalBindingAsset = bHasNeedBindingData ? BindingAsset : nullptr;
+		if (bHasNeedBindingData)
+		{
+			LocalBindingAsset = BindingAsset;
+		}
 	}
 
 	// Update requirement based on binding asset availability
