Fix typo into hair length scale update function.

diff --git a/Engine/Plugins/Runtime/HairStrands/Source/HairStrandsCore/Private/GroomComponent.cpp b/Engine/Plugins/Runtime/HairStrands/Source/HairStrandsCore/Private/GroomComponent.cpp
--- a/Engine/Plugins/Runtime/HairStrands/Source/HairStrandsCore/Private/GroomComponent.cpp
+++ b/Engine/Plugins/Runtime/HairStrands/Source/HairStrandsCore/Private/GroomComponent.cpp
@@ -1559,9 +1559,9 @@ void UGroomComponent::SetHairLengthScale(float InLengthScale)
 {
 	for (FHairGroupDesc& HairDesc : GroomGroupsDesc)
 	{
-		HairDesc.HairWidth = InLengthScale;
+		HairDesc.HairLengthScale = InLengthScale;
 	}
-	//UpdateHairGroupsDescAndInvalidateRenderState();
+	UpdateHairGroupsDescAndInvalidateRenderState(false);
 }
 
 void UGroomComponent::SetForcedLOD(int32 LODIndex)
@@ -1671,7 +1671,7 @@ void UGroomComponent::SetEnableSimulation(bool bInEnableSimulation)
 	}
 }
 
-void UGroomComponent::UpdateHairGroupsDescAndInvalidateRenderState()
+void UGroomComponent::UpdateHairGroupsDescAndInvalidateRenderState(bool bInvalid)
 {
 	UpdateHairGroupsDesc();
 
@@ -1684,7 +1684,10 @@ void UGroomComponent::UpdateHairGroupsDescAndInvalidateRenderState()
 #endif// #if WITH_EDITORONLY_DATA
 		++GroupIndex;
 	}
-	MarkRenderStateDirty();
+	if (bInvalid)
+	{
+		MarkRenderStateDirty();
+	}
 }
 
 FPrimitiveSceneProxy* UGroomComponent::CreateSceneProxy()
diff --git a/Engine/Plugins/Runtime/HairStrands/Source/HairStrandsCore/Public/GroomComponent.h b/Engine/Plugins/Runtime/HairStrands/Source/HairStrandsCore/Public/GroomComponent.h
--- a/Engine/Plugins/Runtime/HairStrands/Source/HairStrandsCore/Public/GroomComponent.h
+++ b/Engine/Plugins/Runtime/HairStrands/Source/HairStrandsCore/Public/GroomComponent.h
@@ -91,7 +91,7 @@ public:
 
 	/** Update Group Description */
 	void UpdateHairGroupsDesc();
-	void UpdateHairGroupsDescAndInvalidateRenderState();
+	void UpdateHairGroupsDescAndInvalidateRenderState(bool bInvalidate=true);
 
 	/** Update simulated groups */
 	void UpdateSimulatedGroups();
