Add console command for resetting all grooms at once.

diff --git a/Engine/Plugins/Runtime/HairStrands/Source/HairStrandsCore/Private/GroomComponent.cpp b/Engine/Plugins/Runtime/HairStrands/Source/HairStrandsCore/Private/GroomComponent.cpp
--- a/Engine/Plugins/Runtime/HairStrands/Source/HairStrandsCore/Private/GroomComponent.cpp
+++ b/Engine/Plugins/Runtime/HairStrands/Source/HairStrandsCore/Private/GroomComponent.cpp
@@ -3135,11 +3135,11 @@ void UGroomComponent::TickComponent(float DeltaTime, enum ELevelTick TickType, F
 	}
 
 	if (RegisteredMeshComponent)
-		{
-			// When a skeletal object with projection is enabled, activate the refresh of the bounding box to
-			// insure the component/proxy bounding box always lies onto the actual skinned mesh
-			MarkRenderTransformDirty();
-		}
+	{
+		// When a skeletal object with projection is enabled, activate the refresh of the bounding box to
+		// insure the component/proxy bounding box always lies onto the actual skinned mesh
+		MarkRenderTransformDirty();
+	}
 
 	// Tick GroomCache only when playing
 	if (GroomCache && GetWorld()->AreActorsInitialized() && bRunning && !bManualTick)
@@ -3672,6 +3672,26 @@ void DumpLoadedGroomComponent(IConsoleVariable* InCVarPakTesterEnabled)
 
 /////////////////////////////////////////////////////////////////////////////////////////
 
+int32 GHairStrandsResetSimulationOnAllGroomComponents = 0;
+void ResetSimulationOnAllGroomComponents(IConsoleVariable* InCVarPakTesterEnabled)
+{
+	if (GHairStrandsResetSimulationOnAllGroomComponents > 0)
+	{
+		for (TObjectIterator<UGroomComponent> HairStrandsComponentIt; HairStrandsComponentIt; ++HairStrandsComponentIt)
+		{
+			HairStrandsComponentIt->bResetSimulation = true;
+			HairStrandsComponentIt->bInitSimulation  = true;
+		}
+	}
+	GHairStrandsResetSimulationOnAllGroomComponents = 0;
+}
+
+static FAutoConsoleVariableRef CVarHairStrandsResetSimulationOnAllGroomComponents(
+	TEXT("r.HairStrands.Simulation.ResetAll"),
+	GHairStrandsResetSimulationOnAllGroomComponents,
+	TEXT("Reset hair strands simulation on all groom components."),
+	FConsoleVariableDelegate::CreateStatic(ResetSimulationOnAllGroomComponents));
+
 #if WITH_EDITORONLY_DATA
 FGroomComponentRecreateRenderStateContext::FGroomComponentRecreateRenderStateContext(UGroomAsset* GroomAsset)
 {
