* Remove unused RT buffer allocation for hair cards/hair meshes
* Remove duplicated hair instance tracking within FHairStrandsProxy
* Ensure GroomComponent respect bIsVisibleInRayTracing.

diff --git a/Engine/Plugins/Runtime/HairStrands/Source/HairStrandsCore/Private/GroomAsset.cpp b/Engine/Plugins/Runtime/HairStrands/Source/HairStrandsCore/Private/GroomAsset.cpp
--- a/Engine/Plugins/Runtime/HairStrands/Source/HairStrandsCore/Private/GroomAsset.cpp
+++ b/Engine/Plugins/Runtime/HairStrands/Source/HairStrandsCore/Private/GroomAsset.cpp
@@ -2849,7 +2849,6 @@ FHairStrandsRaytracingResource* UGroomAsset::AllocateCardsRaytracingResources(ui
 		if (LOD.RaytracingResource == nullptr)
 		{
 			LOD.RaytracingResource = new FHairStrandsRaytracingResource(LOD.BulkData);
-			BeginInitResource(LOD.RaytracingResource);
 		}
 		return LOD.RaytracingResource;
 	}
@@ -2867,7 +2866,6 @@ FHairStrandsRaytracingResource* UGroomAsset::AllocateMeshesRaytracingResources(u
 		if (LOD.RaytracingResource == nullptr)
 		{
 			LOD.RaytracingResource = new FHairStrandsRaytracingResource(LOD.BulkData);
-			BeginInitResource(LOD.RaytracingResource);
 		}
 		return LOD.RaytracingResource;
 	}
@@ -2884,7 +2882,6 @@ FHairStrandsRaytracingResource* UGroomAsset::AllocateStrandsRaytracingResources(
 		if (GroupData.Strands.RaytracingResource == nullptr)
 		{
 			GroupData.Strands.RaytracingResource = new FHairStrandsRaytracingResource(GroupData.Strands.BulkData);
-			BeginInitResource(GroupData.Strands.RaytracingResource);
 		}
 		return GroupData.Strands.RaytracingResource;
 	}
diff --git a/Engine/Plugins/Runtime/HairStrands/Source/HairStrandsCore/Private/GroomComponent.cpp b/Engine/Plugins/Runtime/HairStrands/Source/HairStrandsCore/Private/GroomComponent.cpp
--- a/Engine/Plugins/Runtime/HairStrands/Source/HairStrandsCore/Private/GroomComponent.cpp
+++ b/Engine/Plugins/Runtime/HairStrands/Source/HairStrandsCore/Private/GroomComponent.cpp
@@ -277,19 +277,6 @@ static EHairGeometryType ToHairGeometryType(EGroomGeometryType Type)
 
 class FHairStrandsSceneProxy final : public FPrimitiveSceneProxy
 {
-private:
-	struct FHairGroup
-	{
-#if RHI_RAYTRACING
-		FRayTracingGeometry* RayTracingGeometry_Strands = nullptr;
-		TArray<FRayTracingGeometry*> RayTracingGeometries_Cards;  // Indexed by LOD
-		TArray<FRayTracingGeometry*> RayTracingGeometries_Meshes;  // Indexed by LOD
-#endif
-		FHairGroupPublicData* PublicData = nullptr;
-		TArray<FHairLODSettings> LODSettings;
-		bool bIsVsibible = true;
-	};
-
 public:
 	SIZE_T GetTypeHash() const override
 	{
@@ -373,50 +360,11 @@ public:
 			}
 
 			{
-				FHairGroup& OutGroupData = HairGroups.AddDefaulted_GetRef();
-				OutGroupData.bIsVsibible = bIsVisible;
-				OutGroupData.PublicData = HairInstance->HairGroupPublicData;
-				OutGroupData.LODSettings = Component->GroomAsset->HairGroupsLOD[GroupIt].LODs;
-
 				// If one of the group has simulation enable, then we enable velocity rendering for meshes/cards
 				if (IsHairStrandsSimulationEnable() && HairInstance->Guides.IsValid() && HairInstance->Guides.bIsSimulationEnable)
 				{
 					bAlwaysHasVelocity = true;
 				}
-
-				#if RHI_RAYTRACING
-				OutGroupData.RayTracingGeometry_Strands = nullptr;
-
-				// Strands
-				if (HairInstance->Strands.RenRaytracingResource && HairInstance->Strands.IsValid())
-				{
-					OutGroupData.RayTracingGeometry_Strands = &HairInstance->Strands.RenRaytracingResource->RayTracingGeometry;
-				}
-
-				// Cards
-				OutGroupData.RayTracingGeometries_Cards.Reserve(HairInstance->Cards.LODs.Num());
-				for (FHairGroupInstance::FCards::FLOD& LOD : HairInstance->Cards.LODs)
-				{
-					FRayTracingGeometry*& Geometry = OutGroupData.RayTracingGeometries_Cards.AddDefaulted_GetRef();
-					Geometry = nullptr;
-					if (LOD.IsValid() && LOD.RaytracingResource)
-					{
-						Geometry = &LOD.RaytracingResource->RayTracingGeometry;
-					}
-				}
-
-				// Meshes
-				OutGroupData.RayTracingGeometries_Meshes.Reserve(HairInstance->Meshes.LODs.Num());
-				for (FHairGroupInstance::FMeshes::FLOD& LOD : HairInstance->Meshes.LODs)
-				{
-					FRayTracingGeometry*& Geometry = OutGroupData.RayTracingGeometries_Meshes.AddDefaulted_GetRef();
-					Geometry = nullptr;
-					if (LOD.IsValid() && LOD.RaytracingResource)
-					{
-						Geometry = &LOD.RaytracingResource->RayTracingGeometry;
-					}
-				}
-				#endif
 			}
 
 			// Material - Strands
@@ -566,7 +514,7 @@ public:
 
 	virtual void GetDynamicRayTracingInstances(FRayTracingMaterialGatheringContext & Context, TArray<struct FRayTracingInstance>& OutRayTracingInstances) override
 	{
-		if (!IsHairRayTracingEnabled() || HairGroups.Num() == 0)
+		if (!IsHairRayTracingEnabled() || HairGroupInstances.Num() == 0)
 			return;
 
 		const EShaderPlatform Platform = Context.ReferenceView->GetShaderPlatform();
@@ -581,66 +529,59 @@ public:
 		if (bWireframe)
 			return;
 
-		for (uint32 GroupIt = 0, GroupCount = HairGroups.Num(); GroupIt < GroupCount; ++GroupIt)
+		for (uint32 GroupIt = 0, GroupCount = HairGroupInstances.Num(); GroupIt < GroupCount; ++GroupIt)
 		{
-			const FHairGroup& GroupData = HairGroups[GroupIt];
-
 			FHairGroupInstance* Instance = HairGroupInstances[GroupIt];
 			check(Instance->GetRefCount() > 0);
 			FMatrix OverrideLocalToWorld = GetLocalToWorld();
 			GetOverrideLocalToTransform(Instance, OverrideLocalToWorld);
 
-			if (GroupData.PublicData->VFInput.GeometryType == EHairGeometryType::Strands)
+			const EHairGeometryType GeometryType = Instance->HairGroupPublicData->VFInput.GeometryType;
+			const uint32 LODIndex = Instance->HairGroupPublicData->GetIntLODIndex();
+
+			if (GeometryType == EHairGeometryType::Strands)
 			{
-				if (GroupData.RayTracingGeometry_Strands && GroupData.RayTracingGeometry_Strands->RayTracingGeometryRHI.IsValid())
+				FHairStrandsRaytracingResource* RTGeometry = Instance->Strands.RenRaytracingResource ? Instance->Strands.RenRaytracingResource : nullptr;
+				if (RTGeometry && RTGeometry->RayTracingGeometry.RayTracingGeometryRHI.IsValid())
 				{
-					for (const FRayTracingGeometrySegment& Segment : GroupData.RayTracingGeometry_Strands->Initializer.Segments)
+					for (const FRayTracingGeometrySegment& Segment : RTGeometry->RayTracingGeometry.Initializer.Segments)
 					{
 						check(Segment.VertexBuffer.IsValid());
 					}
-					AddOpaqueRaytracingInstance(OverrideLocalToWorld, GroupData.RayTracingGeometry_Strands, RaytracingInstanceMask_ThinShadow, OutRayTracingInstances);
+					AddOpaqueRaytracingInstance(OverrideLocalToWorld, &RTGeometry->RayTracingGeometry, RaytracingInstanceMask_ThinShadow, OutRayTracingInstances);
 				}
 			}
-			else if (GroupData.PublicData->VFInput.GeometryType == EHairGeometryType::Cards)
+			else if (GeometryType == EHairGeometryType::Cards)
 			{
-				const uint32 LODIndex = GroupData.PublicData->GetIntLODIndex();
-				if (GroupData.RayTracingGeometries_Cards[LODIndex] && GroupData.RayTracingGeometries_Cards[LODIndex]->RayTracingGeometryRHI.IsValid())
+				FHairStrandsRaytracingResource* RTGeometry = Instance->Cards.LODs[LODIndex].RaytracingResource ? Instance->Cards.LODs[LODIndex].RaytracingResource : nullptr;
+				if (RTGeometry && RTGeometry->RayTracingGeometry.RayTracingGeometryRHI.IsValid())
 				{
-					for (const FRayTracingGeometrySegment& Segment : GroupData.RayTracingGeometries_Cards[LODIndex]->Initializer.Segments)
+					for (const FRayTracingGeometrySegment& Segment : RTGeometry->RayTracingGeometry.Initializer.Segments)
 					{
 						check(Segment.VertexBuffer.IsValid());
 					}
-				// Either use shadow only (no material evaluation) or full material evaluation during tracing
-				#if 0
-					AddOpaqueRaytracingInstance(OverrideLocalToWorld, GroupData.RayTracingGeometries_Cards[LODIndex], RaytracingInstanceMask_ThinShadow, OutRayTracingInstances);
-				#else
-
-					FMeshBatch* MeshBatch = CreateMeshBatch(Context.ReferenceView, Context.ReferenceViewFamily, Context.RayTracingMeshResourceCollector, GroupData, Instance, GroupIt, nullptr);
+				
+					FMeshBatch* MeshBatch = CreateMeshBatch(Context.ReferenceView, Context.ReferenceViewFamily, Context.RayTracingMeshResourceCollector, Instance, GroupIt, nullptr);
 					TArray<FMeshBatch> MeshBatches;
 					MeshBatches.Add(*MeshBatch);
-					AddOpaqueRaytracingInstance(OverrideLocalToWorld, GroupData.RayTracingGeometries_Cards[LODIndex], RaytracingInstanceMask_Opaque, MeshBatches, GetScene().GetFeatureLevel(), OutRayTracingInstances);
-				#endif
+					AddOpaqueRaytracingInstance(OverrideLocalToWorld, &RTGeometry->RayTracingGeometry, RaytracingInstanceMask_Opaque, MeshBatches, GetScene().GetFeatureLevel(), OutRayTracingInstances);
+
 				}
 			}
-			else if (GroupData.PublicData->VFInput.GeometryType == EHairGeometryType::Meshes)
+			else if (GeometryType == EHairGeometryType::Meshes)
 			{
-				const uint32 LODIndex = GroupData.PublicData->LODIndex;
-				if (GroupData.RayTracingGeometries_Meshes[LODIndex] && GroupData.RayTracingGeometries_Meshes[LODIndex]->RayTracingGeometryRHI.IsValid())
+				FHairStrandsRaytracingResource* RTGeometry = Instance->Meshes.LODs[LODIndex].RaytracingResource ? Instance->Meshes.LODs[LODIndex].RaytracingResource : nullptr;
+				if (RTGeometry && RTGeometry->RayTracingGeometry.RayTracingGeometryRHI.IsValid())
 				{
-					for (const FRayTracingGeometrySegment& Segment : GroupData.RayTracingGeometries_Meshes[LODIndex]->Initializer.Segments)
+					for (const FRayTracingGeometrySegment& Segment : RTGeometry->RayTracingGeometry.Initializer.Segments)
 					{
 						check(Segment.VertexBuffer.IsValid());
 					}
-				// Either use shadow only (no material evaluation) or full material evaluation during tracing
-				#if 0
-					AddOpaqueRaytracingInstance(OverrideLocalToWorld, GroupData.RayTracingGeometries_Meshes[LODIndex], RaytracingInstanceMask_ThinShadow, OutRayTracingInstances);
-				#else
 
-					FMeshBatch* MeshBatch = CreateMeshBatch(Context.ReferenceView, Context.ReferenceViewFamily, Context.RayTracingMeshResourceCollector, GroupData, Instance, GroupIt, nullptr);
+					FMeshBatch* MeshBatch = CreateMeshBatch(Context.ReferenceView, Context.ReferenceViewFamily, Context.RayTracingMeshResourceCollector, Instance, GroupIt, nullptr);
 					TArray<FMeshBatch> MeshBatches;
 					MeshBatches.Add(*MeshBatch);
-					AddOpaqueRaytracingInstance(OverrideLocalToWorld, GroupData.RayTracingGeometries_Meshes[LODIndex], RaytracingInstanceMask_Opaque, MeshBatches, GetScene().GetFeatureLevel(), OutRayTracingInstances);
-				#endif
+					AddOpaqueRaytracingInstance(OverrideLocalToWorld, &RTGeometry->RayTracingGeometry, RaytracingInstanceMask_Opaque, MeshBatches, GetScene().GetFeatureLevel(), OutRayTracingInstances);
 				}
 			}
 		}
@@ -740,7 +681,7 @@ public:
 						Debug_MaterialProxy = DebugMaterial;
 					}
 
-					FMeshBatch* MeshBatch = CreateMeshBatch(View, ViewFamily, Collector, HairGroups[GroupIt], Instances[GroupIt], GroupIt, Debug_MaterialProxy);
+					FMeshBatch* MeshBatch = CreateMeshBatch(View, ViewFamily, Collector, Instances[GroupIt], GroupIt, Debug_MaterialProxy);
 					if (MeshBatch == nullptr)
 					{
 						continue;
@@ -760,7 +701,6 @@ public:
 		const FSceneView* View,
 		const FSceneViewFamily& ViewFamily,
 		FMeshElementCollector& Collector,
-		const FHairGroup& GroupData,
 		const FHairGroupInstance* Instance,
 		uint32 GroupIndex,
 		FMaterialRenderProxy* Debug_MaterialProxy) const
@@ -774,6 +714,8 @@ public:
 		check(Instance->GetRefCount());
 
 		const int32 IntLODIndex = Instance->HairGroupPublicData->GetIntLODIndex();
+		const bool bIsVisible = Instance->HairGroupPublicData->GetLODVisibility();
+
 		const FVertexFactory* VertexFactory = nullptr;
 		FIndexBuffer* IndexBuffer = nullptr;
 		FMaterialRenderProxy* MaterialRenderProxy = Debug_MaterialProxy;
@@ -836,7 +778,7 @@ public:
 			bWireframe = AllowDebugViewmodes() && ViewFamily.EngineShowFlags.Wireframe;
 		}
 
-		if (MaterialRenderProxy == nullptr || !GroupData.bIsVsibible)
+		if (MaterialRenderProxy == nullptr || !bIsVisible)
 		{
 			return nullptr;
 		}
@@ -958,8 +900,6 @@ private:
 	FMaterialRelevance MaterialRelevance;
 	UMaterialInterface* Strands_DebugMaterial = nullptr;
 	int32* PredictedLODIndex = nullptr;
-
-	TArray<FHairGroup> HairGroups;
 };
 
 /** GroomCacheBuffers implementation that hold copies of the GroomCacheAnimationData needed for playback */
@@ -2492,7 +2432,7 @@ void UGroomComponent::InitResources(bool bIsBindingReloading)
 			}
 
 			#if RHI_RAYTRACING
-			if (IsHairRayTracingEnabled() && HairGroupInstance->Strands.Modifier.bUseHairRaytracingGeometry)
+			if (IsHairRayTracingEnabled() && HairGroupInstance->Strands.Modifier.bUseHairRaytracingGeometry && bVisibleInRayTracing)
 			{
 				if (bNeedDynamicResources)
 				{
@@ -2588,7 +2528,7 @@ void UGroomComponent::InitResources(bool bIsBindingReloading)
 				}
 
 				#if RHI_RAYTRACING
-				if (IsRayTracingEnabled())
+				if (IsRayTracingEnabled() && bVisibleInRayTracing)
 				{
 					if (bNeedDeformedPositions)
 					{
@@ -2657,7 +2597,7 @@ void UGroomComponent::InitResources(bool bIsBindingReloading)
 				}
 
 				#if RHI_RAYTRACING
-				if (IsRayTracingEnabled() && bNeedDeformedPositions)
+				if (IsRayTracingEnabled() && bVisibleInRayTracing)
 				{
 					if (bNeedDeformedPositions)
 					{
diff --git a/Engine/Plugins/Runtime/HairStrands/Source/HairStrandsCore/Private/GroomResources.cpp b/Engine/Plugins/Runtime/HairStrands/Source/HairStrandsCore/Private/GroomResources.cpp
--- a/Engine/Plugins/Runtime/HairStrands/Source/HairStrandsCore/Private/GroomResources.cpp
+++ b/Engine/Plugins/Runtime/HairStrands/Source/HairStrandsCore/Private/GroomResources.cpp
@@ -1337,22 +1337,25 @@ void FHairCardsInterpolationResource::InternalRelease()
 // 4 triangles per hair vertex => 12 vertices per hair vertex
 FHairStrandsRaytracingResource::FHairStrandsRaytracingResource(const FHairStrandsBulkData& InData) :
 	FHairCommonResource(EHairStrandsAllocationType::Deferred),
-	PositionBuffer(), VertexCount(InData.GetNumPoints()*12)  
+	PositionBuffer(), VertexCount(InData.GetNumPoints()*12), bOwnPositionBuffer(true)
 {}
 
 FHairStrandsRaytracingResource::FHairStrandsRaytracingResource(const FHairCardsBulkData& InData) :
 	FHairCommonResource(EHairStrandsAllocationType::Deferred),
-	PositionBuffer(), VertexCount(InData.GetNumVertices())
+	PositionBuffer(), VertexCount(InData.GetNumVertices()), bOwnPositionBuffer(false)
 {}
 
 FHairStrandsRaytracingResource::FHairStrandsRaytracingResource(const FHairMeshesBulkData& InData) :
 	FHairCommonResource(EHairStrandsAllocationType::Deferred),
-	PositionBuffer(), VertexCount(InData.GetNumVertices())
+	PositionBuffer(), VertexCount(InData.GetNumVertices()), bOwnPositionBuffer(false)
 {}
 
 void FHairStrandsRaytracingResource::InternalAllocate(FRDGBuilder& GraphBuilder)
 {
-	InternalCreateVertexBufferRDG<FHairStrandsRaytracingFormat>(GraphBuilder, VertexCount, PositionBuffer, TEXT("Hair.StrandsRaytracing_PositionBuffer"), EHairResourceUsageType::Dynamic);
+	if (bOwnPositionBuffer)
+	{
+		InternalCreateVertexBufferRDG<FHairStrandsRaytracingFormat>(GraphBuilder, VertexCount, PositionBuffer, TEXT("Hair.StrandsRaytracing_PositionBuffer"), EHairResourceUsageType::Dynamic);
+	}
 }
 
 void FHairStrandsRaytracingResource::InternalRelease()
diff --git a/Engine/Plugins/Runtime/HairStrands/Source/HairStrandsCore/Private/HairStrandsRendering.cpp b/Engine/Plugins/Runtime/HairStrands/Source/HairStrandsCore/Private/HairStrandsRendering.cpp
--- a/Engine/Plugins/Runtime/HairStrands/Source/HairStrandsCore/Private/HairStrandsRendering.cpp
+++ b/Engine/Plugins/Runtime/HairStrands/Source/HairStrandsCore/Private/HairStrandsRendering.cpp
@@ -1425,7 +1425,7 @@ static void BuildHairAccelerationStructure_Meshes(FRHICommandList& RHICmdList,
 	FRayTracingGeometry* OutRayTracingGeometry)
 {
 	FRayTracingGeometryInitializer Initializer;
-	static const FName DebugName("HairCards");
+	static const FName DebugName("HairMeshes");
 	static int32 DebugNumber = 0;
 	Initializer.DebugName = FName(DebugName, DebugNumber++);
 	Initializer.IndexBuffer = RestResource->IndexBuffer.IndexBufferRHI;
diff --git a/Engine/Plugins/Runtime/HairStrands/Source/HairStrandsCore/Public/GroomResources.h b/Engine/Plugins/Runtime/HairStrands/Source/HairStrandsCore/Public/GroomResources.h
--- a/Engine/Plugins/Runtime/HairStrands/Source/HairStrandsCore/Public/GroomResources.h
+++ b/Engine/Plugins/Runtime/HairStrands/Source/HairStrandsCore/Public/GroomResources.h
@@ -412,6 +412,7 @@ struct FHairStrandsRaytracingResource : public FHairCommonResource
 	FRDGExternalBuffer PositionBuffer;
 	FRayTracingGeometry RayTracingGeometry;
 	uint32 VertexCount = 0;
+	bool bOwnPositionBuffer = false;
 	bool bIsRTGeometryInitialized = false;
 };
 #endif
