Add support for (precis) hair resource name for debug purpose.

diff --git a/Engine/Plugins/Runtime/HairStrands/Source/HairStrandsCore/Private/GroomAsset.cpp b/Engine/Plugins/Runtime/HairStrands/Source/HairStrandsCore/Private/GroomAsset.cpp
--- a/Engine/Plugins/Runtime/HairStrands/Source/HairStrandsCore/Private/GroomAsset.cpp
+++ b/Engine/Plugins/Runtime/HairStrands/Source/HairStrandsCore/Private/GroomAsset.cpp
@@ -867,7 +867,7 @@ void UGroomAsset::UpdateResource()
 					HairGroupsInfo,
 					HairGroupsData);
 
-				GroupData.Strands.ClusterCullingResource = new FHairStrandsClusterCullingResource(GroupData.Strands.ClusterCullingBulkData);
+				GroupData.Strands.ClusterCullingResource = new FHairStrandsClusterCullingResource(GroupData.Strands.ClusterCullingBulkData, FHairResourceName(GetFName(), GroupIndex));
 				BeginInitResource(GroupData.Strands.ClusterCullingResource);
 			}
 			else
@@ -2540,7 +2540,9 @@ bool UGroomAsset::BuildCardsGeometry(uint32 GroupIndex)
 			// 2. Init geometry & texture resources, and generate interpolation data/resources
 			if (bInitResources)
 			{
-				LOD.RestResource = new FHairCardsRestResource(LOD.BulkData);
+				FHairResourceName ResourceName(GetFName(), GroupIndex, LODIt);
+
+				LOD.RestResource = new FHairCardsRestResource(LOD.BulkData, ResourceName);
 				BeginInitResource(LOD.RestResource);
 
 				// 2.1 Load atlas textures
@@ -2552,7 +2554,7 @@ bool UGroomAsset::BuildCardsGeometry(uint32 GroupIndex)
 				LOD.RestResource->bInvertUV = Desc->SourceType == EHairCardsSourceType::Procedural;
 				
 				// 2.2 Load interoplatino resources
-				LOD.InterpolationResource = new FHairCardsInterpolationResource(LOD.InterpolationBulkData);
+				LOD.InterpolationResource = new FHairCardsInterpolationResource(LOD.InterpolationBulkData, ResourceName);
 				BeginInitResource(LOD.InterpolationResource);
 
 				// Create own interpolation settings for cards.
@@ -2597,10 +2599,10 @@ bool UGroomAsset::BuildCardsGeometry(uint32 GroupIndex)
 					GroupData.Guides.BulkData  = LOD.Guides.BulkData;
 				}
 
-				LOD.Guides.RestResource = new FHairStrandsRestResource(LOD.Guides.BulkData, EHairStrandsResourcesType::Cards);
+				LOD.Guides.RestResource = new FHairStrandsRestResource(LOD.Guides.BulkData, EHairStrandsResourcesType::Cards, ResourceName);
 				BeginInitResource(LOD.Guides.RestResource);
 
-				LOD.Guides.InterpolationResource = new FHairStrandsInterpolationResource(LOD.Guides.InterpolationBulkData);
+				LOD.Guides.InterpolationResource = new FHairStrandsInterpolationResource(LOD.Guides.InterpolationBulkData, ResourceName);
 				BeginInitResource(LOD.Guides.InterpolationResource);
 
 				// Update card stats to display
@@ -2812,7 +2814,7 @@ FHairStrandsRestResource* UGroomAsset::AllocateGuidesResources(uint32 GroupIndex
 		{
 			if (GroupData.Guides.RestResource == nullptr)
 			{
-				GroupData.Guides.RestResource = new FHairStrandsRestResource(GroupData.Guides.BulkData, EHairStrandsResourcesType::Guides);
+				GroupData.Guides.RestResource = new FHairStrandsRestResource(GroupData.Guides.BulkData, EHairStrandsResourcesType::Guides, FHairResourceName(GetFName(), GroupIndex));
 				BeginInitResource(GroupData.Guides.RestResource);
 			}
 			return GroupData.Guides.RestResource;
@@ -2829,7 +2831,7 @@ FHairStrandsInterpolationResource* UGroomAsset::AllocateInterpolationResources(u
 		check(GroupData.Guides.IsValid());
 		if (GroupData.Strands.InterpolationResource == nullptr)
 		{
-			GroupData.Strands.InterpolationResource = new FHairStrandsInterpolationResource(GroupData.Strands.InterpolationBulkData);
+			GroupData.Strands.InterpolationResource = new FHairStrandsInterpolationResource(GroupData.Strands.InterpolationBulkData, FHairResourceName(GetFName(), GroupIndex));
 			BeginInitResource(GroupData.Strands.InterpolationResource);
 		}
 		return GroupData.Strands.InterpolationResource;
@@ -2848,7 +2850,7 @@ FHairStrandsRaytracingResource* UGroomAsset::AllocateCardsRaytracingResources(ui
 
 		if (LOD.RaytracingResource == nullptr)
 		{
-			LOD.RaytracingResource = new FHairStrandsRaytracingResource(LOD.BulkData);
+			LOD.RaytracingResource = new FHairStrandsRaytracingResource(LOD.BulkData, FHairResourceName(GetFName(), GroupIndex, LODIndex));
 		}
 		return LOD.RaytracingResource;
 	}
@@ -2865,7 +2867,7 @@ FHairStrandsRaytracingResource* UGroomAsset::AllocateMeshesRaytracingResources(u
 
 		if (LOD.RaytracingResource == nullptr)
 		{
-			LOD.RaytracingResource = new FHairStrandsRaytracingResource(LOD.BulkData);
+			LOD.RaytracingResource = new FHairStrandsRaytracingResource(LOD.BulkData, FHairResourceName(GetFName(), GroupIndex, LODIndex));
 		}
 		return LOD.RaytracingResource;
 	}
@@ -2881,7 +2883,7 @@ FHairStrandsRaytracingResource* UGroomAsset::AllocateStrandsRaytracingResources(
 
 		if (GroupData.Strands.RaytracingResource == nullptr)
 		{
-			GroupData.Strands.RaytracingResource = new FHairStrandsRaytracingResource(GroupData.Strands.BulkData);
+			GroupData.Strands.RaytracingResource = new FHairStrandsRaytracingResource(GroupData.Strands.BulkData, FHairResourceName(GetFName(), GroupIndex));
 		}
 		return GroupData.Strands.RaytracingResource;
 	}
@@ -2909,7 +2911,9 @@ void UGroomAsset::InitStrandsResources()
 
 		if (GroupData.Strands.HasValidData())
 		{
-			GroupData.Strands.RestResource = new FHairStrandsRestResource(GroupData.Strands.BulkData, EHairStrandsResourcesType::Strands);
+			FHairResourceName ResourceName(GetFName(), GroupIndex);
+
+			GroupData.Strands.RestResource = new FHairStrandsRestResource(GroupData.Strands.BulkData, EHairStrandsResourcesType::Strands, ResourceName);
 			BeginInitResource(GroupData.Strands.RestResource);
 
 			if (GroupData.Strands.ClusterCullingBulkData.IsValid())
@@ -2921,7 +2925,7 @@ void UGroomAsset::InitStrandsResources()
 					GroupData.Strands.ClusterCullingBulkData.LODVisibility[LODIt] = HairGroupsLOD[GroupIndex].LODs[LODIt].bVisible;
 				}
 
-				GroupData.Strands.ClusterCullingResource = new FHairStrandsClusterCullingResource(GroupData.Strands.ClusterCullingBulkData);
+				GroupData.Strands.ClusterCullingResource = new FHairStrandsClusterCullingResource(GroupData.Strands.ClusterCullingBulkData, ResourceName);
 				BeginInitResource(GroupData.Strands.ClusterCullingResource);
 			}
 
@@ -3005,16 +3009,18 @@ void UGroomAsset::InitCardsResources()
 			if (LOD.RestResource == nullptr &&			// don't initialize again if they were previously initialized during the BuildCardsGeometry
 				LOD.HasValidData())
 			{
-				LOD.RestResource = new FHairCardsRestResource(LOD.BulkData);
+				FHairResourceName ResourceName(GetFName(), GroupIndex, LODIt);
+
+				LOD.RestResource = new FHairCardsRestResource(LOD.BulkData, ResourceName);
 				BeginInitResource(LOD.RestResource);
 
-				LOD.InterpolationResource = new FHairCardsInterpolationResource(LOD.InterpolationBulkData);
+				LOD.InterpolationResource = new FHairCardsInterpolationResource(LOD.InterpolationBulkData, ResourceName);
 				BeginInitResource(LOD.InterpolationResource);
 
-				LOD.Guides.RestResource = new FHairStrandsRestResource(LOD.Guides.BulkData, EHairStrandsResourcesType::Cards);
+				LOD.Guides.RestResource = new FHairStrandsRestResource(LOD.Guides.BulkData, EHairStrandsResourcesType::Cards, ResourceName);
 				BeginInitResource(LOD.Guides.RestResource);
 
-				LOD.Guides.InterpolationResource = new FHairStrandsInterpolationResource(LOD.Guides.InterpolationBulkData);
+				LOD.Guides.InterpolationResource = new FHairStrandsInterpolationResource(LOD.Guides.InterpolationBulkData, ResourceName);
 				BeginInitResource(LOD.Guides.InterpolationResource);
 
 				if (Desc)
@@ -3065,7 +3071,7 @@ void UGroomAsset::InitMeshesResources()
 
 			if (LOD.HasValidData())
 			{
-				LOD.RestResource = new FHairMeshesRestResource(LOD.BulkData);
+				LOD.RestResource = new FHairMeshesRestResource(LOD.BulkData, FHairResourceName(GetFName(), GroupIndex, LODIt));
 				BeginInitResource(LOD.RestResource);
 
 				if (Desc)
@@ -3488,7 +3494,7 @@ void UGroomAsset::SaveProceduralCards(uint32 DescIndex)
 	}
 
 	// 3. Create resources and enqueue texture generation (GPU, kicked by the render thread) 
-	Q.Resources = new FHairCardsRestResource(Q.BulkData);
+	Q.Resources = new FHairCardsRestResource(Q.BulkData, FHairResourceName(GetFName(), GroupIndex, LODIndex));
 	BeginInitResource(Q.Resources);
 	Q.ProceduralResources = new FHairCardsProceduralResource(Q.ProceduralData.RenderData, Q.ProceduralData.Atlas.Resolution, Q.ProceduralData.Voxels);
 	BeginInitResource(Q.ProceduralResources);
diff --git a/Engine/Plugins/Runtime/HairStrands/Source/HairStrandsCore/Private/GroomBindingAsset.cpp b/Engine/Plugins/Runtime/HairStrands/Source/HairStrandsCore/Private/GroomBindingAsset.cpp
--- a/Engine/Plugins/Runtime/HairStrands/Source/HairStrandsCore/Private/GroomBindingAsset.cpp
+++ b/Engine/Plugins/Runtime/HairStrands/Source/HairStrandsCore/Private/GroomBindingAsset.cpp
@@ -97,11 +97,13 @@ void UGroomBindingAsset::InitResource()
 		const int32 GroupIndex = HairGroupResources.Num();
 		FHairGroupResource& Resource = HairGroupResources.AddDefaulted_GetRef();
 
+		FHairResourceName ResourceName(GetFName(), GroupIndex);
+
 		// Guides
 		Resource.SimRootResources = nullptr;
 		if (BulkData.SimRootBulkData.IsValid())
 		{
-			Resource.SimRootResources = new FHairStrandsRestRootResource(BulkData.SimRootBulkData, EHairStrandsResourcesType::Guides);
+			Resource.SimRootResources = new FHairStrandsRestRootResource(BulkData.SimRootBulkData, EHairStrandsResourcesType::Guides, ResourceName);
 			BeginInitResource(Resource.SimRootResources);
 		}
 
@@ -109,7 +111,7 @@ void UGroomBindingAsset::InitResource()
 		Resource.RenRootResources = nullptr;
 		if (IsHairStrandsEnabled(EHairStrandsShaderType::Strands) && BulkData.RenRootBulkData.IsValid())
 		{
-			Resource.RenRootResources = new FHairStrandsRestRootResource(BulkData.RenRootBulkData, EHairStrandsResourcesType::Strands);
+			Resource.RenRootResources = new FHairStrandsRestRootResource(BulkData.RenRootBulkData, EHairStrandsResourcesType::Strands, ResourceName);
 			BeginInitResource(Resource.RenRootResources);
 		}
 
@@ -123,7 +125,7 @@ void UGroomBindingAsset::InitResource()
 				Resource.CardsRootResources[CardsLODIt] = nullptr;
 				if (BulkData.CardsRootBulkData[CardsLODIt].IsValid())
 				{
-					Resource.CardsRootResources[CardsLODIt] = new FHairStrandsRestRootResource(BulkData.CardsRootBulkData[CardsLODIt], EHairStrandsResourcesType::Cards);
+					Resource.CardsRootResources[CardsLODIt] = new FHairStrandsRestRootResource(BulkData.CardsRootBulkData[CardsLODIt], EHairStrandsResourcesType::Cards, FHairResourceName(GetFName(), GroupIndex, CardsLODIt));
 					BeginInitResource(Resource.CardsRootResources[CardsLODIt]);
 				}
 			}
diff --git a/Engine/Plugins/Runtime/HairStrands/Source/HairStrandsCore/Private/GroomComponent.cpp b/Engine/Plugins/Runtime/HairStrands/Source/HairStrandsCore/Private/GroomComponent.cpp
--- a/Engine/Plugins/Runtime/HairStrands/Source/HairStrandsCore/Private/GroomComponent.cpp
+++ b/Engine/Plugins/Runtime/HairStrands/Source/HairStrandsCore/Private/GroomComponent.cpp
@@ -2328,6 +2328,8 @@ void UGroomComponent::InitResources(bool bIsBindingReloading)
 			HairGroupInstance->HairGroupPublicData->SetLODGeometryTypes(LODGeometryTypes);
 		}
 
+		FHairResourceName ResourceName(GetFName(), GroupIt);
+
 		// Not needed anymore
 		//const bool bDynamicResources = IsHairStrandsSimulationEnable() || IsHairStrandsBindingEnable();
 
@@ -2352,7 +2354,7 @@ void UGroomComponent::InitResources(bool bIsBindingReloading)
 				}
 
 				HairGroupInstance->Guides.RestRootResource = LocalBindingAsset->HairGroupResources[GroupIt].SimRootResources;
-				HairGroupInstance->Guides.DeformedRootResource = new FHairStrandsDeformedRootResource(HairGroupInstance->Guides.RestRootResource, EHairStrandsResourcesType::Guides);
+				HairGroupInstance->Guides.DeformedRootResource = new FHairStrandsDeformedRootResource(HairGroupInstance->Guides.RestRootResource, EHairStrandsResourcesType::Guides, ResourceName);
 			}
 
 			// Lazy allocation of the guide resources
@@ -2360,7 +2362,7 @@ void UGroomComponent::InitResources(bool bIsBindingReloading)
 			check(GroupData.Guides.RestResource);
 
 			// If guides are allocated, deformed resources are always needs since they are either used with simulation, or RBF deformation. Both are dynamics, and require deformed positions
-			HairGroupInstance->Guides.DeformedResource = new FHairStrandsDeformedResource(GroupData.Guides.BulkData, EHairStrandsResourcesType::Guides);
+			HairGroupInstance->Guides.DeformedResource = new FHairStrandsDeformedResource(GroupData.Guides.BulkData, EHairStrandsResourcesType::Guides, ResourceName);
 
 			// Initialize the simulation and the global deformation to its default behavior by setting it with LODIndex = -1
 			const int32 LODIndex = -1;
@@ -2437,7 +2439,7 @@ void UGroomComponent::InitResources(bool bIsBindingReloading)
 				if (bNeedDynamicResources)
 				{
 					// Allocate dynamic raytracing resources (owned by the groom component/instance)
-					HairGroupInstance->Strands.RenRaytracingResource = new FHairStrandsRaytracingResource(GroupData.Strands.BulkData);
+					HairGroupInstance->Strands.RenRaytracingResource = new FHairStrandsRaytracingResource(GroupData.Strands.BulkData, ResourceName);
 					HairGroupInstance->Strands.RenRaytracingResourceOwned = true;
 				}
 				else
@@ -2461,13 +2463,13 @@ void UGroomComponent::InitResources(bool bIsBindingReloading)
 				}
 
 				HairGroupInstance->Strands.RestRootResource = LocalBindingAsset->HairGroupResources[GroupIt].RenRootResources;
-				HairGroupInstance->Strands.DeformedRootResource = new FHairStrandsDeformedRootResource(HairGroupInstance->Strands.RestRootResource, EHairStrandsResourcesType::Strands);
+				HairGroupInstance->Strands.DeformedRootResource = new FHairStrandsDeformedRootResource(HairGroupInstance->Strands.RestRootResource, EHairStrandsResourcesType::Strands, ResourceName);
 			}
 
 			HairGroupInstance->Strands.RestResource = GroupData.Strands.RestResource;
 			if (bNeedDynamicResources)
 			{
-				HairGroupInstance->Strands.DeformedResource = new FHairStrandsDeformedResource(GroupData.Strands.BulkData, EHairStrandsResourcesType::Strands);
+				HairGroupInstance->Strands.DeformedResource = new FHairStrandsDeformedResource(GroupData.Strands.BulkData, EHairStrandsResourcesType::Strands, ResourceName);
 			} 
 
 			// An empty groom doesn't have a ClusterCullingResource
@@ -2506,6 +2508,8 @@ void UGroomComponent::InitResources(bool bIsBindingReloading)
 			FHairGroupInstance::FCards::FLOD& InstanceLOD = HairGroupInstance->Cards.LODs.AddDefaulted_GetRef();
 			if (LOD.IsValid())
 			{
+				FHairResourceName LODResourceName(GetFName(), GroupIt, CardsLODIndex);
+
 				const EHairBindingType BindingType	= HairGroupInstance->HairGroupPublicData->GetBindingType(CardsLODIndex);
 				const bool bHasSimulation			= HairGroupInstance->HairGroupPublicData->IsSimulationEnable(CardsLODIndex);
 				const bool bHasGlobalDeformation	= HairGroupInstance->HairGroupPublicData->IsGlobalInterpolationEnable(CardsLODIndex);
@@ -2524,7 +2528,7 @@ void UGroomComponent::InitResources(bool bIsBindingReloading)
 
 				if (bNeedDeformedPositions)
 				{
-					InstanceLOD.DeformedResource = new FHairCardsDeformedResource(LOD.BulkData, false);
+					InstanceLOD.DeformedResource = new FHairCardsDeformedResource(LOD.BulkData, false, LODResourceName);
 				}
 
 				#if RHI_RAYTRACING
@@ -2533,7 +2537,7 @@ void UGroomComponent::InitResources(bool bIsBindingReloading)
 					if (bNeedDeformedPositions)
 					{
 						// Allocate dynamic raytracing resources (owned by the groom component/instance)
-						InstanceLOD.RaytracingResource = new FHairStrandsRaytracingResource(*InstanceLOD.Data);
+						InstanceLOD.RaytracingResource = new FHairStrandsRaytracingResource(*InstanceLOD.Data, LODResourceName);
 						InstanceLOD.RaytracingResourceOwned = true;
 					}
 					else
@@ -2565,12 +2569,12 @@ void UGroomComponent::InitResources(bool bIsBindingReloading)
 						}
 
 						InstanceLOD.Guides.RestRootResource = LocalBindingAsset->HairGroupResources[GroupIt].CardsRootResources[CardsLODIndex];
-						InstanceLOD.Guides.DeformedRootResource = new FHairStrandsDeformedRootResource(InstanceLOD.Guides.RestRootResource, EHairStrandsResourcesType::Cards);
+						InstanceLOD.Guides.DeformedRootResource = new FHairStrandsDeformedRootResource(InstanceLOD.Guides.RestRootResource, EHairStrandsResourcesType::Cards, LODResourceName);
 					}
 
 					InstanceLOD.Guides.RestResource = LOD.Guides.RestResource;
 					{
-						InstanceLOD.Guides.DeformedResource = new FHairStrandsDeformedResource(LOD.Guides.BulkData, EHairStrandsResourcesType::Cards);
+						InstanceLOD.Guides.DeformedResource = new FHairStrandsDeformedResource(LOD.Guides.BulkData, EHairStrandsResourcesType::Cards, LODResourceName);
 					}
 
 					InstanceLOD.Guides.HairInterpolationType = HairInterpolationType;
@@ -2585,6 +2589,8 @@ void UGroomComponent::InitResources(bool bIsBindingReloading)
 			FHairGroupInstance::FMeshes::FLOD& InstanceLOD = HairGroupInstance->Meshes.LODs.AddDefaulted_GetRef();
 			if (LOD.IsValid())
 			{
+				FHairResourceName LODResourceName(GetFName(), GroupIt, MeshLODIndex);
+
 				const EHairBindingType BindingType = HairGroupInstance->HairGroupPublicData->GetBindingType(MeshLODIndex);
 				const bool bHasGlobalDeformation   = HairGroupInstance->HairGroupPublicData->IsGlobalInterpolationEnable(MeshLODIndex);
 				const bool bNeedDeformedPositions  = bHasGlobalDeformation && BindingType == EHairBindingType::Skinning;
@@ -2593,7 +2599,7 @@ void UGroomComponent::InitResources(bool bIsBindingReloading)
 				InstanceLOD.RestResource = LOD.RestResource;
 				if (bNeedDeformedPositions)
 				{
-					InstanceLOD.DeformedResource = new FHairMeshesDeformedResource(LOD.BulkData, true);
+					InstanceLOD.DeformedResource = new FHairMeshesDeformedResource(LOD.BulkData, true, LODResourceName);
 				}
 
 				#if RHI_RAYTRACING
@@ -2602,7 +2608,7 @@ void UGroomComponent::InitResources(bool bIsBindingReloading)
 					if (bNeedDeformedPositions)
 					{
 						// Allocate dynamic raytracing resources (owned by the groom component/instance)
-						InstanceLOD.RaytracingResource = new FHairStrandsRaytracingResource(*InstanceLOD.Data);
+						InstanceLOD.RaytracingResource = new FHairStrandsRaytracingResource(*InstanceLOD.Data, LODResourceName);
 						InstanceLOD.RaytracingResourceOwned = true;
 					}
 					else
diff --git a/Engine/Plugins/Runtime/HairStrands/Source/HairStrandsCore/Private/GroomResources.cpp b/Engine/Plugins/Runtime/HairStrands/Source/HairStrandsCore/Private/GroomResources.cpp
--- a/Engine/Plugins/Runtime/HairStrands/Source/HairStrandsCore/Private/GroomResources.cpp
+++ b/Engine/Plugins/Runtime/HairStrands/Source/HairStrandsCore/Private/GroomResources.cpp
@@ -49,6 +49,26 @@ enum class EHairResourceUsageType : uint8
 
 #define HAIRSTRANDS_RESOUCE_NAME(Type, Name) (Type == EHairStrandsResourcesType::Guides  ? TEXT(#Name "(Guides)") : (Type == EHairStrandsResourcesType::Strands ? TEXT(#Name "(Strands)") : TEXT(#Name "(Cards)")))
 
+const TCHAR* ToHairResourceDebugName(const TCHAR* In, FHairResourceName& InDebugNames)
+{
+#if HAIR_RESOURCE_DEBUG_NAME
+	FString& TempDebugName = InDebugNames.Names.AddDefaulted_GetRef();
+	TempDebugName = In;
+	if (InDebugNames.GroupIndex >= 0)
+	{
+		TempDebugName += TEXT("_GROUP") + FString::FromInt(InDebugNames.GroupIndex);
+	}
+	if (InDebugNames.LODIndex >= 0)
+	{
+		TempDebugName += TEXT("_LOD") + FString::FromInt(InDebugNames.LODIndex);
+	}
+	TempDebugName += TEXT("_") + InDebugNames.AssetName.ToString();
+	return *TempDebugName;
+#else
+	return In;
+#endif
+}
+
 /////////////////////////////////////////////////////////////////////////////////////////
 // FRWBuffer utils 
 void UploadDataToBuffer(FReadBuffer& OutBuffer, uint32 DataSizeInBytes, const void* InCpuData)
@@ -427,10 +447,11 @@ void CreateHairStrandsDebugAttributeBuffer(FRDGBuilder& GraphBuilder, FRDGExtern
 
 /////////////////////////////////////////////////////////////////////////////////////////
 
-FHairCommonResource::FHairCommonResource(EHairStrandsAllocationType InAllocationType, bool bInUseRenderGraph):
+FHairCommonResource::FHairCommonResource(EHairStrandsAllocationType InAllocationType, const FHairResourceName& InResourceName, bool bInUseRenderGraph):
 bUseRenderGraph(bInUseRenderGraph),
 bIsInitialized(false),
-AllocationType(InAllocationType)
+AllocationType(InAllocationType),
+ResourceName(InResourceName)
 {
 }
 
@@ -491,8 +512,8 @@ void FHairCardIndexBuffer::InitRHI()
 	RHIUnlockBuffer(IndexBufferRHI);
 }
 
-FHairCardsRestResource::FHairCardsRestResource(const FHairCardsBulkData& InBulkData) :
-	FHairCommonResource(EHairStrandsAllocationType::Immediate, false),
+FHairCardsRestResource::FHairCardsRestResource(const FHairCardsBulkData& InBulkData, const FHairResourceName& InResourceName) :
+	FHairCommonResource(EHairStrandsAllocationType::Immediate, InResourceName, false),
 	RestPositionBuffer(),
 	RestIndexBuffer(InBulkData.Indices),
 	NormalsBuffer(),
@@ -505,9 +526,9 @@ FHairCardsRestResource::FHairCardsRestResource(const FHairCardsBulkData& InBulkD
 void FHairCardsRestResource::InternalAllocate()
 {
 	// These resources are kept as regular (i.e., non-RDG resources) as they need to be bound at the input assembly stage by the Vertex declaraction which requires FVertexBuffer type
-	CreateBuffer<FHairCardsPositionFormat>(BulkData.Positions, RestPositionBuffer, TEXT("Hair.CardsRest_PositionBuffer"));
-	CreateBuffer<FHairCardsNormalFormat>(BulkData.Normals, NormalsBuffer, TEXT("Hair.CardsRest_NormalBuffer"));
-	CreateBuffer<FHairCardsUVFormat>(BulkData.UVs, UVsBuffer, TEXT("Hair.CardsRest_UVBuffer"));
+	CreateBuffer<FHairCardsPositionFormat>(BulkData.Positions, RestPositionBuffer, ToHairResourceDebugName(TEXT("Hair.CardsRest_PositionBuffer"), ResourceName));
+	CreateBuffer<FHairCardsNormalFormat>(BulkData.Normals, NormalsBuffer, ToHairResourceDebugName(TEXT("Hair.CardsRest_NormalBuffer"), ResourceName));
+	CreateBuffer<FHairCardsUVFormat>(BulkData.UVs, UVsBuffer, ToHairResourceDebugName(TEXT("Hair.CardsRest_UVBuffer"), ResourceName));
 
 	FSamplerStateRHIRef DefaultSampler = TStaticSamplerState<SF_Bilinear, AM_Clamp, AM_Clamp, AM_Clamp>::GetRHI();
 	DepthSampler = DefaultSampler;
@@ -540,7 +561,7 @@ void FHairCardsRestResource::ReleaseResource()
 
 /////////////////////////////////////////////////////////////////////////////////////////
 FHairCardsProceduralResource::FHairCardsProceduralResource(const FHairCardsProceduralDatas::FRenderData& InRenderData, const FIntPoint& InAtlasResolution, const FHairCardsVoxel& InVoxel):
-	FHairCommonResource(EHairStrandsAllocationType::Immediate),
+	FHairCommonResource(EHairStrandsAllocationType::Immediate, FHairResourceName()),
 	CardBoundCount(InRenderData.ClusterBounds.Num()),
 	AtlasResolution(InAtlasResolution),
 	AtlasRectBuffer(),
@@ -557,19 +578,19 @@ FHairCardsProceduralResource::FHairCardsProceduralResource(const FHairCardsProce
 
 void FHairCardsProceduralResource::InternalAllocate(FRDGBuilder& GraphBuilder)
 {
-	InternalCreateVertexBufferRDG<FHairCardsAtlasRectFormat>(GraphBuilder, RenderData.CardsRect, AtlasRectBuffer, TEXT("Hair.CardsProcedural_AtlasRectBuffer"), EHairResourceUsageType::Static);
-	InternalCreateVertexBufferRDG<FHairCardsDimensionFormat>(GraphBuilder, RenderData.CardsLengths, LengthBuffer, TEXT("Hair.CardsProcedural_LengthBuffer"), EHairResourceUsageType::Static);
+	InternalCreateVertexBufferRDG<FHairCardsAtlasRectFormat>(GraphBuilder, RenderData.CardsRect, AtlasRectBuffer, ToHairResourceDebugName(TEXT("Hair.CardsProcedural_AtlasRectBuffer"), ResourceName), EHairResourceUsageType::Static);
+	InternalCreateVertexBufferRDG<FHairCardsDimensionFormat>(GraphBuilder, RenderData.CardsLengths, LengthBuffer, ToHairResourceDebugName(TEXT("Hair.CardsProcedural_LengthBuffer"), ResourceName), EHairResourceUsageType::Static);
 
-	InternalCreateVertexBufferRDG<FHairCardsOffsetAndCount>(GraphBuilder, RenderData.CardItToCluster, CardItToClusterBuffer, TEXT("Hair.CardsProcedural_CardItToClusterBuffer"), EHairResourceUsageType::Static);
-	InternalCreateVertexBufferRDG<FHairCardsOffsetAndCount>(GraphBuilder, RenderData.ClusterIdToVertices, ClusterIdToVerticesBuffer, TEXT("Hair.CardsProcedural_ClusterIdToVerticesBuffer"), EHairResourceUsageType::Static);
-	InternalCreateVertexBufferRDG<FHairCardsBoundsFormat>(GraphBuilder, RenderData.ClusterBounds, ClusterBoundBuffer, TEXT("Hair.CardsProcedural_ClusterBoundBuffer"), EHairResourceUsageType::Static);
+	InternalCreateVertexBufferRDG<FHairCardsOffsetAndCount>(GraphBuilder, RenderData.CardItToCluster, CardItToClusterBuffer, ToHairResourceDebugName(TEXT("Hair.CardsProcedural_CardItToClusterBuffer"), ResourceName), EHairResourceUsageType::Static);
+	InternalCreateVertexBufferRDG<FHairCardsOffsetAndCount>(GraphBuilder, RenderData.ClusterIdToVertices, ClusterIdToVerticesBuffer, ToHairResourceDebugName(TEXT("Hair.CardsProcedural_ClusterIdToVerticesBuffer"), ResourceName), EHairResourceUsageType::Static);
+	InternalCreateVertexBufferRDG<FHairCardsBoundsFormat>(GraphBuilder, RenderData.ClusterBounds, ClusterBoundBuffer, ToHairResourceDebugName(TEXT("Hair.CardsProcedural_ClusterBoundBuffer"), ResourceName), EHairResourceUsageType::Static);
 
-	InternalCreateVertexBufferRDG<FHairCardsVoxelDensityFormat>(GraphBuilder, RenderData.VoxelDensity, CardVoxel.DensityBuffer, TEXT("Hair.CardsProcedural_VoxelDensityBuffer"), EHairResourceUsageType::Static);
-	InternalCreateVertexBufferRDG<FHairCardsVoxelTangentFormat>(GraphBuilder, RenderData.VoxelTangent, CardVoxel.TangentBuffer, TEXT("Hair.CardsProcedural_VoxelTangentBuffer"), EHairResourceUsageType::Static);
-	InternalCreateVertexBufferRDG<FHairCardsVoxelTangentFormat>(GraphBuilder, RenderData.VoxelNormal, CardVoxel.NormalBuffer, TEXT("Hair.CardsProcedural_VoxelNormalBuffer"), EHairResourceUsageType::Static);
+	InternalCreateVertexBufferRDG<FHairCardsVoxelDensityFormat>(GraphBuilder, RenderData.VoxelDensity, CardVoxel.DensityBuffer, ToHairResourceDebugName(TEXT("Hair.CardsProcedural_VoxelDensityBuffer"), ResourceName), EHairResourceUsageType::Static);
+	InternalCreateVertexBufferRDG<FHairCardsVoxelTangentFormat>(GraphBuilder, RenderData.VoxelTangent, CardVoxel.TangentBuffer, ToHairResourceDebugName(TEXT("Hair.CardsProcedural_VoxelTangentBuffer"), ResourceName), EHairResourceUsageType::Static);
+	InternalCreateVertexBufferRDG<FHairCardsVoxelTangentFormat>(GraphBuilder, RenderData.VoxelNormal, CardVoxel.NormalBuffer, ToHairResourceDebugName(TEXT("Hair.CardsProcedural_VoxelNormalBuffer"), ResourceName), EHairResourceUsageType::Static);
 
-	InternalCreateVertexBufferRDG<FHairCardsStrandsPositionFormat>(GraphBuilder, RenderData.CardsStrandsPositions, CardsStrandsPositions, TEXT("Hair.CardsProcedural_CardsStrandsPositions"), EHairResourceUsageType::Static);
-	InternalCreateVertexBufferRDG<FHairCardsStrandsAttributeFormat>(GraphBuilder, RenderData.CardsStrandsAttributes, CardsStrandsAttributes, TEXT("Hair.CardsProcedural_CardsStrandsAttributes"), EHairResourceUsageType::Static);
+	InternalCreateVertexBufferRDG<FHairCardsStrandsPositionFormat>(GraphBuilder, RenderData.CardsStrandsPositions, CardsStrandsPositions, ToHairResourceDebugName(TEXT("Hair.CardsProcedural_CardsStrandsPositions"), ResourceName), EHairResourceUsageType::Static);
+	InternalCreateVertexBufferRDG<FHairCardsStrandsAttributeFormat>(GraphBuilder, RenderData.CardsStrandsAttributes, CardsStrandsAttributes, ToHairResourceDebugName(TEXT("Hair.CardsProcedural_CardsStrandsAttributes"), ResourceName), EHairResourceUsageType::Static);
 }
 
 void FHairCardsProceduralResource::InternalRelease()
@@ -590,8 +611,8 @@ void FHairCardsProceduralResource::InternalRelease()
 
 /////////////////////////////////////////////////////////////////////////////////////////
 
-FHairCardsDeformedResource::FHairCardsDeformedResource(const FHairCardsBulkData& InBulkData, bool bInInitializedData) :
-	FHairCommonResource(EHairStrandsAllocationType::Deferred),
+FHairCardsDeformedResource::FHairCardsDeformedResource(const FHairCardsBulkData& InBulkData, bool bInInitializedData, const FHairResourceName& InResourceName) :
+	FHairCommonResource(EHairStrandsAllocationType::Deferred, InResourceName),
 	BulkData(InBulkData), bInitializedData(bInInitializedData)
 {}
 
@@ -599,13 +620,13 @@ void FHairCardsDeformedResource::InternalAllocate(FRDGBuilder& GraphBuilder)
 {
 	if (bInitializedData)
 	{
-		InternalCreateVertexBufferRDG<FHairCardsPositionFormat>(GraphBuilder, BulkData.Positions, DeformedPositionBuffer[0], TEXT("Hair.CardsDeformed_Position0"), EHairResourceUsageType::Dynamic);
-		InternalCreateVertexBufferRDG<FHairCardsPositionFormat>(GraphBuilder, BulkData.Positions, DeformedPositionBuffer[1], TEXT("Hair.CardsDeformed_Position1"), EHairResourceUsageType::Dynamic);
+		InternalCreateVertexBufferRDG<FHairCardsPositionFormat>(GraphBuilder, BulkData.Positions, DeformedPositionBuffer[0], ToHairResourceDebugName(TEXT("Hair.CardsDeformed_Position0"), ResourceName), EHairResourceUsageType::Dynamic);
+		InternalCreateVertexBufferRDG<FHairCardsPositionFormat>(GraphBuilder, BulkData.Positions, DeformedPositionBuffer[1], ToHairResourceDebugName(TEXT("Hair.CardsDeformed_Position1"), ResourceName), EHairResourceUsageType::Dynamic);
 	}
 	else
 	{
-		InternalCreateVertexBufferRDG<FHairCardsPositionFormat>(GraphBuilder, BulkData.GetNumVertices(), DeformedPositionBuffer[0], TEXT("Hair.CardsDeformed_Position0"), EHairResourceUsageType::Dynamic);
-		InternalCreateVertexBufferRDG<FHairCardsPositionFormat>(GraphBuilder, BulkData.GetNumVertices(), DeformedPositionBuffer[1], TEXT("Hair.CardsDeformed_Position1"), EHairResourceUsageType::Dynamic);
+		InternalCreateVertexBufferRDG<FHairCardsPositionFormat>(GraphBuilder, BulkData.GetNumVertices(), DeformedPositionBuffer[0], ToHairResourceDebugName(TEXT("Hair.CardsDeformed_Position0"), ResourceName), EHairResourceUsageType::Dynamic);
+		InternalCreateVertexBufferRDG<FHairCardsPositionFormat>(GraphBuilder, BulkData.GetNumVertices(), DeformedPositionBuffer[1], ToHairResourceDebugName(TEXT("Hair.CardsDeformed_Position1"), ResourceName), EHairResourceUsageType::Dynamic);
 
 		// Manually transit to SRVs, in case of the cards are rendered not visible, but still rendered (in shadows for instance). In such a case, the cards deformation pass is not called, and thus the 
 		// buffers are never transit from UAV (clear) to SRV for rasterization. 
@@ -622,8 +643,8 @@ void FHairCardsDeformedResource::InternalRelease()
 
 /////////////////////////////////////////////////////////////////////////////////////////
 
-FHairMeshesRestResource::FHairMeshesRestResource(const FHairMeshesBulkData& InBulkData) :
-	FHairCommonResource(EHairStrandsAllocationType::Immediate, false),
+FHairMeshesRestResource::FHairMeshesRestResource(const FHairMeshesBulkData& InBulkData, const FHairResourceName& InResourceName) :
+	FHairCommonResource(EHairStrandsAllocationType::Immediate, InResourceName, false),
 	RestPositionBuffer(),
 	IndexBuffer(InBulkData.Indices),
 	NormalsBuffer(),
@@ -637,9 +658,9 @@ FHairMeshesRestResource::FHairMeshesRestResource(const FHairMeshesBulkData& InBu
 void FHairMeshesRestResource::InternalAllocate()
 {
 	// These resources are kept as regular (i.e., non-RDG resources) as they need to be bound at the input assembly stage by the Vertex declaraction which requires FVertexBuffer type
-	CreateBuffer<FHairCardsPositionFormat>(BulkData.Positions, RestPositionBuffer, TEXT("Hair.MeshesRest_Positions"));
-	CreateBuffer<FHairCardsNormalFormat>(BulkData.Normals, NormalsBuffer, TEXT("Hair.MeshesRest_Normals"));
-	CreateBuffer<FHairCardsUVFormat>(BulkData.UVs, UVsBuffer, TEXT("Hair.MeshesRest_UVs"));
+	CreateBuffer<FHairCardsPositionFormat>(BulkData.Positions, RestPositionBuffer, ToHairResourceDebugName(TEXT("Hair.MeshesRest_Positions"), ResourceName));
+	CreateBuffer<FHairCardsNormalFormat>(BulkData.Normals, NormalsBuffer, ToHairResourceDebugName(TEXT("Hair.MeshesRest_Normals"), ResourceName));
+	CreateBuffer<FHairCardsUVFormat>(BulkData.UVs, UVsBuffer, ToHairResourceDebugName(TEXT("Hair.MeshesRest_UVs"), ResourceName));
 }
 
 void FHairMeshesRestResource::InternalRelease()
@@ -668,8 +689,8 @@ void FHairMeshesRestResource::ReleaseResource()
 
 /////////////////////////////////////////////////////////////////////////////////////////
 
-FHairMeshesDeformedResource::FHairMeshesDeformedResource(const FHairMeshesBulkData& InBulkData, bool bInInitializedData) :
-	FHairCommonResource(EHairStrandsAllocationType::Deferred),
+FHairMeshesDeformedResource::FHairMeshesDeformedResource(const FHairMeshesBulkData& InBulkData, bool bInInitializedData, const FHairResourceName& InResourceName) :
+	FHairCommonResource(EHairStrandsAllocationType::Deferred, InResourceName),
 	BulkData(InBulkData), bInitializedData(bInInitializedData)
 {}
 
@@ -677,13 +698,13 @@ void FHairMeshesDeformedResource::InternalAllocate(FRDGBuilder& GraphBuilder)
 {
 	if (bInitializedData)
 	{
-		InternalCreateVertexBufferRDG<FHairCardsPositionFormat>(GraphBuilder, BulkData.Positions, DeformedPositionBuffer[0], TEXT("Hair.MeshesDeformed_Positions0"), EHairResourceUsageType::Dynamic);
-		InternalCreateVertexBufferRDG<FHairCardsPositionFormat>(GraphBuilder, BulkData.Positions, DeformedPositionBuffer[1], TEXT("Hair.MeshesDeformed_Positions1"), EHairResourceUsageType::Dynamic);
+		InternalCreateVertexBufferRDG<FHairCardsPositionFormat>(GraphBuilder, BulkData.Positions, DeformedPositionBuffer[0], ToHairResourceDebugName(TEXT("Hair.MeshesDeformed_Positions0"), ResourceName), EHairResourceUsageType::Dynamic);
+		InternalCreateVertexBufferRDG<FHairCardsPositionFormat>(GraphBuilder, BulkData.Positions, DeformedPositionBuffer[1], ToHairResourceDebugName(TEXT("Hair.MeshesDeformed_Positions1"), ResourceName), EHairResourceUsageType::Dynamic);
 	}
 	else
 	{
-		InternalCreateVertexBufferRDG<FHairCardsPositionFormat>(GraphBuilder, BulkData.GetNumVertices(), DeformedPositionBuffer[0], TEXT("Hair.MeshesDeformed_Positions0"), EHairResourceUsageType::Dynamic);
-		InternalCreateVertexBufferRDG<FHairCardsPositionFormat>(GraphBuilder, BulkData.GetNumVertices(), DeformedPositionBuffer[1], TEXT("Hair.MeshesDeformed_Positions1"), EHairResourceUsageType::Dynamic);
+		InternalCreateVertexBufferRDG<FHairCardsPositionFormat>(GraphBuilder, BulkData.GetNumVertices(), DeformedPositionBuffer[0], ToHairResourceDebugName(TEXT("Hair.MeshesDeformed_Positions0"), ResourceName), EHairResourceUsageType::Dynamic);
+		InternalCreateVertexBufferRDG<FHairCardsPositionFormat>(GraphBuilder, BulkData.GetNumVertices(), DeformedPositionBuffer[1], ToHairResourceDebugName(TEXT("Hair.MeshesDeformed_Positions1"), ResourceName), EHairResourceUsageType::Dynamic);
 	}
 }
 
@@ -695,8 +716,8 @@ void FHairMeshesDeformedResource::InternalRelease()
 
 /////////////////////////////////////////////////////////////////////////////////////////
 
-FHairStrandsRestResource::FHairStrandsRestResource(FHairStrandsBulkData& InBulkData, EHairStrandsResourcesType InCurveType) :
-	FHairCommonResource(EHairStrandsAllocationType::Deferred),
+FHairStrandsRestResource::FHairStrandsRestResource(FHairStrandsBulkData& InBulkData, EHairStrandsResourcesType InCurveType, const FHairResourceName& InResourceName) :
+	FHairCommonResource(EHairStrandsAllocationType::Deferred, InResourceName),
 	PositionBuffer(), Attribute0Buffer(), Attribute1Buffer(), MaterialBuffer(), BulkData(InBulkData), CurveType(InCurveType)
 {}
 
@@ -706,22 +727,22 @@ void FHairStrandsRestResource::InternalAllocate(FRDGBuilder& GraphBuilder)
 
 	// 1. Lock data, which force the loading data from files (on non-editor build/cooked data). These data are then uploaded to the GPU
 	// 2. A local copy is done by the buffer uploader. This copy is discarded once the uploading is done.
-	InternalCreateVertexBufferRDG_FromBulkData<FHairStrandsPositionFormat>(GraphBuilder, BulkData.Positions, PointCount, PositionBuffer, HAIRSTRANDS_RESOUCE_NAME(CurveType, Hair.StrandsRest_PositionBuffer), EHairResourceUsageType::Static);
-	InternalCreateVertexBufferRDG_FromBulkData<FHairStrandsAttribute0Format>(GraphBuilder, BulkData.Attributes0, PointCount, Attribute0Buffer, HAIRSTRANDS_RESOUCE_NAME(CurveType, Hair.StrandsRest_Attribute0Buffer), EHairResourceUsageType::Static);
+	InternalCreateVertexBufferRDG_FromBulkData<FHairStrandsPositionFormat>(GraphBuilder, BulkData.Positions, PointCount, PositionBuffer, ToHairResourceDebugName(HAIRSTRANDS_RESOUCE_NAME(CurveType, Hair.StrandsRest_PositionBuffer), ResourceName), EHairResourceUsageType::Static);
+	InternalCreateVertexBufferRDG_FromBulkData<FHairStrandsAttribute0Format>(GraphBuilder, BulkData.Attributes0, PointCount, Attribute0Buffer, ToHairResourceDebugName(HAIRSTRANDS_RESOUCE_NAME(CurveType, Hair.StrandsRest_Attribute0Buffer), ResourceName), EHairResourceUsageType::Static);
 
 	if (!!(BulkData.Flags & FHairStrandsBulkData::DataFlags_HasUDIMData))
 	{
-		InternalCreateVertexBufferRDG_FromBulkData<FHairStrandsAttribute1Format>(GraphBuilder, BulkData.Attributes1, PointCount, Attribute1Buffer, HAIRSTRANDS_RESOUCE_NAME(CurveType, Hair.StrandsRest_Attribute1Buffer), EHairResourceUsageType::Static);
+		InternalCreateVertexBufferRDG_FromBulkData<FHairStrandsAttribute1Format>(GraphBuilder, BulkData.Attributes1, PointCount, Attribute1Buffer, ToHairResourceDebugName(HAIRSTRANDS_RESOUCE_NAME(CurveType, Hair.StrandsRest_Attribute1Buffer), ResourceName), EHairResourceUsageType::Static);
 	}
 
 	if (!!(BulkData.Flags & FHairStrandsBulkData::DataFlags_HasMaterialData))
 	{
-		InternalCreateVertexBufferRDG_FromBulkData<FHairStrandsMaterialFormat>(GraphBuilder, BulkData.Materials, PointCount, MaterialBuffer, HAIRSTRANDS_RESOUCE_NAME(CurveType, Hair.StrandsRest_MaterialBuffer), EHairResourceUsageType::Static);
+		InternalCreateVertexBufferRDG_FromBulkData<FHairStrandsMaterialFormat>(GraphBuilder, BulkData.Materials, PointCount, MaterialBuffer, ToHairResourceDebugName(HAIRSTRANDS_RESOUCE_NAME(CurveType, Hair.StrandsRest_MaterialBuffer), ResourceName), EHairResourceUsageType::Static);
 	}
 
 	TArray<FVector4> RestOffset;
 	RestOffset.Add(BulkData.GetPositionOffset());
-	InternalCreateVertexBufferRDG<FHairStrandsPositionOffsetFormat>(GraphBuilder, RestOffset, PositionOffsetBuffer, HAIRSTRANDS_RESOUCE_NAME(CurveType, Hair.StrandsRest_PositionOffsetBuffer), EHairResourceUsageType::Static, ERDGInitialDataFlags::None);
+	InternalCreateVertexBufferRDG<FHairStrandsPositionOffsetFormat>(GraphBuilder, RestOffset, PositionOffsetBuffer, ToHairResourceDebugName(HAIRSTRANDS_RESOUCE_NAME(CurveType, Hair.StrandsRest_PositionOffsetBuffer), ResourceName), EHairResourceUsageType::Static, ERDGInitialDataFlags::None);
 	HairTransition::TransitToSRV(GraphBuilder, Register(GraphBuilder, PositionOffsetBuffer, ERDGImportedBufferFlags::CreateSRV).SRV, ERDGPassFlags::Compute);
 }
 
@@ -738,7 +759,7 @@ FRDGExternalBuffer FHairStrandsRestResource::GetTangentBuffer(FRDGBuilder& Graph
 	// Lazy allocation and update
 	if (TangentBuffer.Buffer == nullptr)
 	{
-		InternalCreateVertexBufferRDG<FHairStrandsTangentFormat>(GraphBuilder, BulkData.PointCount * FHairStrandsTangentFormat::ComponentCount, TangentBuffer, TEXT("Hair.StrandsRest_TangentBuffer"), EHairResourceUsageType::Dynamic);
+		InternalCreateVertexBufferRDG<FHairStrandsTangentFormat>(GraphBuilder, BulkData.PointCount * FHairStrandsTangentFormat::ComponentCount, TangentBuffer, ToHairResourceDebugName(TEXT("Hair.StrandsRest_TangentBuffer"), ResourceName), EHairResourceUsageType::Dynamic);
 
 		AddHairTangentPass(
 			GraphBuilder,
@@ -764,8 +785,8 @@ void FHairStrandsRestResource::InternalRelease()
 
 /////////////////////////////////////////////////////////////////////////////////////////
 
-FHairStrandsDeformedResource::FHairStrandsDeformedResource(FHairStrandsBulkData& InBulkData, EHairStrandsResourcesType InCurveType) :
-	FHairCommonResource(EHairStrandsAllocationType::Deferred),
+FHairStrandsDeformedResource::FHairStrandsDeformedResource(FHairStrandsBulkData& InBulkData, EHairStrandsResourcesType InCurveType, const FHairResourceName& InResourceName) :
+	FHairCommonResource(EHairStrandsAllocationType::Deferred, InResourceName),
 	BulkData(InBulkData), CurveType(InCurveType)
 {
 	GetPositionOffset(FHairStrandsDeformedResource::EFrameType::Current)  = BulkData.GetPositionOffset();
@@ -776,14 +797,14 @@ void FHairStrandsDeformedResource::InternalAllocate(FRDGBuilder& GraphBuilder)
 {
 	const uint32 PointCount = BulkData.PointCount;
 
-	InternalCreateVertexBufferRDG<FHairStrandsPositionFormat>(GraphBuilder, BulkData.PointCount, DeformedPositionBuffer[0], HAIRSTRANDS_RESOUCE_NAME(CurveType, Hair.StrandsDeformed_DeformedPositionBuffer0), EHairResourceUsageType::Dynamic);
-	InternalCreateVertexBufferRDG<FHairStrandsPositionFormat>(GraphBuilder, BulkData.PointCount, DeformedPositionBuffer[1], HAIRSTRANDS_RESOUCE_NAME(CurveType, Hair.StrandsDeformed_DeformedPositionBuffer1), EHairResourceUsageType::Dynamic);
-	InternalCreateVertexBufferRDG<FHairStrandsTangentFormat>(GraphBuilder, BulkData.PointCount * FHairStrandsTangentFormat::ComponentCount, TangentBuffer, HAIRSTRANDS_RESOUCE_NAME(CurveType, Hair.StrandsDeformed_TangentBuffer), EHairResourceUsageType::Dynamic);
+	InternalCreateVertexBufferRDG<FHairStrandsPositionFormat>(GraphBuilder, BulkData.PointCount, DeformedPositionBuffer[0], ToHairResourceDebugName(HAIRSTRANDS_RESOUCE_NAME(CurveType, Hair.StrandsDeformed_DeformedPositionBuffer0), ResourceName), EHairResourceUsageType::Dynamic);
+	InternalCreateVertexBufferRDG<FHairStrandsPositionFormat>(GraphBuilder, BulkData.PointCount, DeformedPositionBuffer[1], ToHairResourceDebugName(HAIRSTRANDS_RESOUCE_NAME(CurveType, Hair.StrandsDeformed_DeformedPositionBuffer1), ResourceName), EHairResourceUsageType::Dynamic);
+	InternalCreateVertexBufferRDG<FHairStrandsTangentFormat>(GraphBuilder, BulkData.PointCount * FHairStrandsTangentFormat::ComponentCount, TangentBuffer, ToHairResourceDebugName(HAIRSTRANDS_RESOUCE_NAME(CurveType, Hair.StrandsDeformed_TangentBuffer), ResourceName), EHairResourceUsageType::Dynamic);
 
 	TArray<FVector4> DefaultOffsets;
 	DefaultOffsets.Add(BulkData.GetPositionOffset());
-	InternalCreateVertexBufferRDG<FHairStrandsPositionOffsetFormat>(GraphBuilder, DefaultOffsets, DeformedOffsetBuffer[0], HAIRSTRANDS_RESOUCE_NAME(CurveType, Hair.StrandsDeformed_DeformedOffsetBuffer0), EHairResourceUsageType::Dynamic, ERDGInitialDataFlags::None);
-	InternalCreateVertexBufferRDG<FHairStrandsPositionOffsetFormat>(GraphBuilder, DefaultOffsets, DeformedOffsetBuffer[1], HAIRSTRANDS_RESOUCE_NAME(CurveType, Hair.StrandsDeformed_DeformedOffsetBuffer1), EHairResourceUsageType::Dynamic, ERDGInitialDataFlags::None);
+	InternalCreateVertexBufferRDG<FHairStrandsPositionOffsetFormat>(GraphBuilder, DefaultOffsets, DeformedOffsetBuffer[0], ToHairResourceDebugName(HAIRSTRANDS_RESOUCE_NAME(CurveType, Hair.StrandsDeformed_DeformedOffsetBuffer0), ResourceName), EHairResourceUsageType::Dynamic, ERDGInitialDataFlags::None);
+	InternalCreateVertexBufferRDG<FHairStrandsPositionOffsetFormat>(GraphBuilder, DefaultOffsets, DeformedOffsetBuffer[1], ToHairResourceDebugName(HAIRSTRANDS_RESOUCE_NAME(CurveType, Hair.StrandsDeformed_DeformedOffsetBuffer1), ResourceName), EHairResourceUsageType::Dynamic, ERDGInitialDataFlags::None);
 
 }
 
@@ -906,8 +927,8 @@ void FHairStrandsClusterCullingBulkData::Validate(bool bIsSaving)
 
 /////////////////////////////////////////////////////////////////////////////////////////
 // Cluster culling resources
-FHairStrandsClusterCullingResource::FHairStrandsClusterCullingResource(FHairStrandsClusterCullingBulkData& InBulkData): 
-	FHairCommonResource(EHairStrandsAllocationType::Deferred),
+FHairStrandsClusterCullingResource::FHairStrandsClusterCullingResource(FHairStrandsClusterCullingBulkData& InBulkData, const FHairResourceName& InResourceName):
+	FHairCommonResource(EHairStrandsAllocationType::Deferred, InResourceName),
 	BulkData(InBulkData)
 {
 
@@ -920,11 +941,11 @@ void FHairStrandsClusterCullingResource::InternalAllocate(FRDGBuilder& GraphBuil
 		BulkData.Validate(false);
 	}
 
-	InternalCreateStructuredBufferRDG_FromBulkData<FHairClusterInfoFormat>(GraphBuilder, BulkData.PackedClusterInfos, BulkData.ClusterCount, ClusterInfoBuffer, TEXT("Hair.StrandsClusterCulling_ClusterInfoBuffer"), EHairResourceUsageType::Static);
-	InternalCreateStructuredBufferRDG_FromBulkData<FHairClusterLODInfoFormat>(GraphBuilder, BulkData.ClusterLODInfos, BulkData.ClusterLODCount, ClusterLODInfoBuffer, TEXT("Hair.StrandsClusterCulling_ClusterLODInfoBuffer"), EHairResourceUsageType::Static);
+	InternalCreateStructuredBufferRDG_FromBulkData<FHairClusterInfoFormat>(GraphBuilder, BulkData.PackedClusterInfos, BulkData.ClusterCount, ClusterInfoBuffer, ToHairResourceDebugName(TEXT("Hair.StrandsClusterCulling_ClusterInfoBuffer"), ResourceName), EHairResourceUsageType::Static);
+	InternalCreateStructuredBufferRDG_FromBulkData<FHairClusterLODInfoFormat>(GraphBuilder, BulkData.ClusterLODInfos, BulkData.ClusterLODCount, ClusterLODInfoBuffer, ToHairResourceDebugName(TEXT("Hair.StrandsClusterCulling_ClusterLODInfoBuffer"), ResourceName), EHairResourceUsageType::Static);
 
-	InternalCreateVertexBufferRDG_FromBulkData<FHairClusterIndexFormat>(GraphBuilder, BulkData.VertexToClusterIds, BulkData.VertexCount, VertexToClusterIdBuffer, TEXT("Hair.StrandsClusterCulling_VertexToClusterIds"), EHairResourceUsageType::Static);
-	InternalCreateVertexBufferRDG_FromBulkData<FHairClusterIndexFormat>(GraphBuilder, BulkData.ClusterVertexIds, BulkData.VertexLODCount, ClusterVertexIdBuffer, TEXT("Hair.StrandsClusterCulling_ClusterVertexIds"), EHairResourceUsageType::Static);
+	InternalCreateVertexBufferRDG_FromBulkData<FHairClusterIndexFormat>(GraphBuilder, BulkData.VertexToClusterIds, BulkData.VertexCount, VertexToClusterIdBuffer, ToHairResourceDebugName(TEXT("Hair.StrandsClusterCulling_VertexToClusterIds"), ResourceName), EHairResourceUsageType::Static);
+	InternalCreateVertexBufferRDG_FromBulkData<FHairClusterIndexFormat>(GraphBuilder, BulkData.ClusterVertexIds, BulkData.VertexLODCount, ClusterVertexIdBuffer, ToHairResourceDebugName(TEXT("Hair.StrandsClusterCulling_ClusterVertexIds"), ResourceName), EHairResourceUsageType::Static);
 }
 
 void FHairStrandsClusterCullingResource::InternalRelease()
@@ -937,8 +958,8 @@ void FHairStrandsClusterCullingResource::InternalRelease()
 
 /////////////////////////////////////////////////////////////////////////////////////////
 
-FHairStrandsRestRootResource::FHairStrandsRestRootResource(FHairStrandsRootBulkData& InBulkData, EHairStrandsResourcesType InCurveType) :
-	FHairCommonResource(EHairStrandsAllocationType::Deferred),
+FHairStrandsRestRootResource::FHairStrandsRestRootResource(FHairStrandsRootBulkData& InBulkData, EHairStrandsResourcesType InCurveType, const FHairResourceName& InResourceName) :
+	FHairCommonResource(EHairStrandsAllocationType::Deferred, InResourceName),
 	BulkData(InBulkData), CurveType(InCurveType)
 {
 	PopulateFromRootData();
@@ -994,24 +1015,24 @@ void FHairStrandsRestRootResource::InternalAllocateLOD(FRDGBuilder& GraphBuilder
 		{
 			GPUData.Status = FLOD::EStatus::Completed;
 
-			InternalCreateVertexBufferRDG_FromBulkData<FHairStrandsCurveTriangleBarycentricFormat>(GraphBuilder, CPUData.RootTriangleBarycentricBuffer, BulkData.RootCount, GPUData.RootTriangleBarycentricBuffer, HAIRSTRANDS_RESOUCE_NAME(CurveType, Hair.StrandsRestRoot_RootTriangleBarycentricBuffer), EHairResourceUsageType::Static);
-			InternalCreateVertexBufferRDG_FromBulkData<FHairStrandsCurveTriangleIndexFormat>(GraphBuilder, CPUData.RootTriangleIndexBuffer, BulkData.RootCount, GPUData.RootTriangleIndexBuffer, HAIRSTRANDS_RESOUCE_NAME(CurveType, Hair.StrandsRestRoot_RootTriangleIndexBuffer), EHairResourceUsageType::Static);
+			InternalCreateVertexBufferRDG_FromBulkData<FHairStrandsCurveTriangleBarycentricFormat>(GraphBuilder, CPUData.RootTriangleBarycentricBuffer, BulkData.RootCount, GPUData.RootTriangleBarycentricBuffer, ToHairResourceDebugName(HAIRSTRANDS_RESOUCE_NAME(CurveType, Hair.StrandsRestRoot_RootTriangleBarycentricBuffer), ResourceName), EHairResourceUsageType::Static);
+			InternalCreateVertexBufferRDG_FromBulkData<FHairStrandsCurveTriangleIndexFormat>(GraphBuilder, CPUData.RootTriangleIndexBuffer, BulkData.RootCount, GPUData.RootTriangleIndexBuffer, ToHairResourceDebugName(HAIRSTRANDS_RESOUCE_NAME(CurveType, Hair.StrandsRestRoot_RootTriangleIndexBuffer), ResourceName), EHairResourceUsageType::Static);
 
-			InternalCreateVertexBufferRDG_FromBulkData<FHairStrandsMeshTrianglePositionFormat>(GraphBuilder, CPUData.RestRootTrianglePosition0Buffer, BulkData.RootCount, GPUData.RestRootTrianglePosition0Buffer, HAIRSTRANDS_RESOUCE_NAME(CurveType, Hair.StrandsRestRoot_RestRootTrianglePosition0Buffer), EHairResourceUsageType::Static);
-			InternalCreateVertexBufferRDG_FromBulkData<FHairStrandsMeshTrianglePositionFormat>(GraphBuilder, CPUData.RestRootTrianglePosition1Buffer, BulkData.RootCount, GPUData.RestRootTrianglePosition1Buffer, HAIRSTRANDS_RESOUCE_NAME(CurveType, Hair.StrandsRestRoot_RestRootTrianglePosition1Buffer), EHairResourceUsageType::Static);
-			InternalCreateVertexBufferRDG_FromBulkData<FHairStrandsMeshTrianglePositionFormat>(GraphBuilder, CPUData.RestRootTrianglePosition2Buffer, BulkData.RootCount, GPUData.RestRootTrianglePosition2Buffer, HAIRSTRANDS_RESOUCE_NAME(CurveType, Hair.StrandsRestRoot_RestRootTrianglePosition2Buffer), EHairResourceUsageType::Static);
+			InternalCreateVertexBufferRDG_FromBulkData<FHairStrandsMeshTrianglePositionFormat>(GraphBuilder, CPUData.RestRootTrianglePosition0Buffer, BulkData.RootCount, GPUData.RestRootTrianglePosition0Buffer, ToHairResourceDebugName(HAIRSTRANDS_RESOUCE_NAME(CurveType, Hair.StrandsRestRoot_RestRootTrianglePosition0Buffer), ResourceName), EHairResourceUsageType::Static);
+			InternalCreateVertexBufferRDG_FromBulkData<FHairStrandsMeshTrianglePositionFormat>(GraphBuilder, CPUData.RestRootTrianglePosition1Buffer, BulkData.RootCount, GPUData.RestRootTrianglePosition1Buffer, ToHairResourceDebugName(HAIRSTRANDS_RESOUCE_NAME(CurveType, Hair.StrandsRestRoot_RestRootTrianglePosition1Buffer), ResourceName), EHairResourceUsageType::Static);
+			InternalCreateVertexBufferRDG_FromBulkData<FHairStrandsMeshTrianglePositionFormat>(GraphBuilder, CPUData.RestRootTrianglePosition2Buffer, BulkData.RootCount, GPUData.RestRootTrianglePosition2Buffer, ToHairResourceDebugName(HAIRSTRANDS_RESOUCE_NAME(CurveType, Hair.StrandsRestRoot_RestRootTrianglePosition2Buffer), ResourceName), EHairResourceUsageType::Static);
 		}
 		else
 		{
 			GPUData.Status = FLOD::EStatus::Initialized;
 			
-			InternalCreateVertexBufferRDG<FHairStrandsCurveTriangleBarycentricFormat>(GraphBuilder, BulkData.RootCount, GPUData.RootTriangleBarycentricBuffer, HAIRSTRANDS_RESOUCE_NAME(CurveType, Hair.StrandsRestRoot_RootTriangleBarycentricBuffer), EHairResourceUsageType::Dynamic);
-			InternalCreateVertexBufferRDG<FHairStrandsCurveTriangleIndexFormat>(GraphBuilder, BulkData.RootCount, GPUData.RootTriangleIndexBuffer, HAIRSTRANDS_RESOUCE_NAME(CurveType, Hair.StrandsRestRoot_RootTriangleIndexBuffer), EHairResourceUsageType::Dynamic);
+			InternalCreateVertexBufferRDG<FHairStrandsCurveTriangleBarycentricFormat>(GraphBuilder, BulkData.RootCount, GPUData.RootTriangleBarycentricBuffer, ToHairResourceDebugName(HAIRSTRANDS_RESOUCE_NAME(CurveType, Hair.StrandsRestRoot_RootTriangleBarycentricBuffer), ResourceName), EHairResourceUsageType::Dynamic);
+			InternalCreateVertexBufferRDG<FHairStrandsCurveTriangleIndexFormat>(GraphBuilder, BulkData.RootCount, GPUData.RootTriangleIndexBuffer, ToHairResourceDebugName(HAIRSTRANDS_RESOUCE_NAME(CurveType, Hair.StrandsRestRoot_RootTriangleIndexBuffer), ResourceName), EHairResourceUsageType::Dynamic);
 			
 			// Create buffers. Initialization will be done by render passes
-			InternalCreateVertexBufferRDG<FHairStrandsMeshTrianglePositionFormat>(GraphBuilder, BulkData.RootCount, GPUData.RestRootTrianglePosition0Buffer, HAIRSTRANDS_RESOUCE_NAME(CurveType, Hair.StrandsRestRoot_RestRootTrianglePosition0Buffer), EHairResourceUsageType::Dynamic);
-			InternalCreateVertexBufferRDG<FHairStrandsMeshTrianglePositionFormat>(GraphBuilder, BulkData.RootCount, GPUData.RestRootTrianglePosition1Buffer, HAIRSTRANDS_RESOUCE_NAME(CurveType, Hair.StrandsRestRoot_RestRootTrianglePosition1Buffer), EHairResourceUsageType::Dynamic);
-			InternalCreateVertexBufferRDG<FHairStrandsMeshTrianglePositionFormat>(GraphBuilder, BulkData.RootCount, GPUData.RestRootTrianglePosition2Buffer, HAIRSTRANDS_RESOUCE_NAME(CurveType, Hair.StrandsRestRoot_RestRootTrianglePosition2Buffer), EHairResourceUsageType::Dynamic);
+			InternalCreateVertexBufferRDG<FHairStrandsMeshTrianglePositionFormat>(GraphBuilder, BulkData.RootCount, GPUData.RestRootTrianglePosition0Buffer, ToHairResourceDebugName(HAIRSTRANDS_RESOUCE_NAME(CurveType, Hair.StrandsRestRoot_RestRootTrianglePosition0Buffer), ResourceName), EHairResourceUsageType::Dynamic);
+			InternalCreateVertexBufferRDG<FHairStrandsMeshTrianglePositionFormat>(GraphBuilder, BulkData.RootCount, GPUData.RestRootTrianglePosition1Buffer, ToHairResourceDebugName(HAIRSTRANDS_RESOUCE_NAME(CurveType, Hair.StrandsRestRoot_RestRootTrianglePosition1Buffer), ResourceName), EHairResourceUsageType::Dynamic);
+			InternalCreateVertexBufferRDG<FHairStrandsMeshTrianglePositionFormat>(GraphBuilder, BulkData.RootCount, GPUData.RestRootTrianglePosition2Buffer, ToHairResourceDebugName(HAIRSTRANDS_RESOUCE_NAME(CurveType, Hair.StrandsRestRoot_RestRootTrianglePosition2Buffer), ResourceName), EHairResourceUsageType::Dynamic);
 		}
 
 		GPUData.SampleCount = CPUData.SampleCount;
@@ -1019,16 +1040,16 @@ void FHairStrandsRestRootResource::InternalAllocateLOD(FRDGBuilder& GraphBuilder
 		if (bHasValidCPUWeights)
 		{
 			const uint32 InteroplationWeightCount = CPUData.MeshInterpolationWeightsBuffer.GetBulkDataSize() / sizeof(FHairStrandsWeightFormat::Type);
-			InternalCreateVertexBufferRDG_FromBulkData<FHairStrandsWeightFormat>(GraphBuilder, CPUData.MeshInterpolationWeightsBuffer, InteroplationWeightCount, GPUData.MeshInterpolationWeightsBuffer, HAIRSTRANDS_RESOUCE_NAME(CurveType, Hair.StrandsRestRoot_MeshInterpolationWeightsBuffer), EHairResourceUsageType::Static);
-			InternalCreateVertexBufferRDG_FromBulkData<FHairStrandsIndexFormat>(GraphBuilder, CPUData.MeshSampleIndicesBuffer, CPUData.SampleCount, GPUData.MeshSampleIndicesBuffer, HAIRSTRANDS_RESOUCE_NAME(CurveType, Hair.StrandsRestRoot_MeshSampleIndicesBuffer), EHairResourceUsageType::Static);
-			InternalCreateVertexBufferRDG_FromBulkData<FHairStrandsMeshTrianglePositionFormat>(GraphBuilder, CPUData.RestSamplePositionsBuffer, CPUData.SampleCount, GPUData.RestSamplePositionsBuffer, HAIRSTRANDS_RESOUCE_NAME(CurveType, Hair.StrandsRestRoot_RestSamplePositionsBuffer), EHairResourceUsageType::Static);
+			InternalCreateVertexBufferRDG_FromBulkData<FHairStrandsWeightFormat>(GraphBuilder, CPUData.MeshInterpolationWeightsBuffer, InteroplationWeightCount, GPUData.MeshInterpolationWeightsBuffer, ToHairResourceDebugName(HAIRSTRANDS_RESOUCE_NAME(CurveType, Hair.StrandsRestRoot_MeshInterpolationWeightsBuffer), ResourceName), EHairResourceUsageType::Static);
+			InternalCreateVertexBufferRDG_FromBulkData<FHairStrandsIndexFormat>(GraphBuilder, CPUData.MeshSampleIndicesBuffer, CPUData.SampleCount, GPUData.MeshSampleIndicesBuffer, ToHairResourceDebugName(HAIRSTRANDS_RESOUCE_NAME(CurveType, Hair.StrandsRestRoot_MeshSampleIndicesBuffer), ResourceName), EHairResourceUsageType::Static);
+			InternalCreateVertexBufferRDG_FromBulkData<FHairStrandsMeshTrianglePositionFormat>(GraphBuilder, CPUData.RestSamplePositionsBuffer, CPUData.SampleCount, GPUData.RestSamplePositionsBuffer, ToHairResourceDebugName(HAIRSTRANDS_RESOUCE_NAME(CurveType, Hair.StrandsRestRoot_RestSamplePositionsBuffer), ResourceName), EHairResourceUsageType::Static);
 		}
 		else
 		{
 			// TODO: do not allocate these resources, since they won't be used
-			InternalCreateVertexBufferRDG<FHairStrandsWeightFormat>(GraphBuilder, (CPUData.SampleCount+4) * (CPUData.SampleCount+4), GPUData.MeshInterpolationWeightsBuffer, HAIRSTRANDS_RESOUCE_NAME(CurveType, Hair.StrandsRestRoot_MeshInterpolationWeightsBuffer), EHairResourceUsageType::Dynamic);
-			InternalCreateVertexBufferRDG<FHairStrandsIndexFormat>(GraphBuilder, CPUData.SampleCount, GPUData.MeshSampleIndicesBuffer, HAIRSTRANDS_RESOUCE_NAME(CurveType, Hair.StrandsRestRoot_MeshSampleIndicesBuffer), EHairResourceUsageType::Dynamic);
-			InternalCreateVertexBufferRDG<FHairStrandsMeshTrianglePositionFormat>(GraphBuilder, CPUData.SampleCount, GPUData.RestSamplePositionsBuffer, HAIRSTRANDS_RESOUCE_NAME(CurveType, Hair.StrandsRestRoot_RestSamplePositionsBuffer), EHairResourceUsageType::Dynamic);
+			InternalCreateVertexBufferRDG<FHairStrandsWeightFormat>(GraphBuilder, (CPUData.SampleCount+4) * (CPUData.SampleCount+4), GPUData.MeshInterpolationWeightsBuffer, ToHairResourceDebugName(HAIRSTRANDS_RESOUCE_NAME(CurveType, Hair.StrandsRestRoot_MeshInterpolationWeightsBuffer), ResourceName), EHairResourceUsageType::Dynamic);
+			InternalCreateVertexBufferRDG<FHairStrandsIndexFormat>(GraphBuilder, CPUData.SampleCount, GPUData.MeshSampleIndicesBuffer, ToHairResourceDebugName(HAIRSTRANDS_RESOUCE_NAME(CurveType, Hair.StrandsRestRoot_MeshSampleIndicesBuffer), ResourceName), EHairResourceUsageType::Dynamic);
+			InternalCreateVertexBufferRDG<FHairStrandsMeshTrianglePositionFormat>(GraphBuilder, CPUData.SampleCount, GPUData.RestSamplePositionsBuffer, ToHairResourceDebugName(HAIRSTRANDS_RESOUCE_NAME(CurveType, Hair.StrandsRestRoot_RestSamplePositionsBuffer), ResourceName), EHairResourceUsageType::Dynamic);
 		}
 	}
 }
@@ -1055,15 +1076,15 @@ void FHairStrandsRestRootResource::InternalRelease()
 
 /////////////////////////////////////////////////////////////////////////////////////////
 
-FHairStrandsDeformedRootResource::FHairStrandsDeformedRootResource(EHairStrandsResourcesType InCurveType):
-	FHairCommonResource(EHairStrandsAllocationType::Deferred),
+FHairStrandsDeformedRootResource::FHairStrandsDeformedRootResource(EHairStrandsResourcesType InCurveType, const FHairResourceName& InResourceName):
+	FHairCommonResource(EHairStrandsAllocationType::Deferred, InResourceName),
 	CurveType(InCurveType)
 {
 
 }
 
-FHairStrandsDeformedRootResource::FHairStrandsDeformedRootResource(const FHairStrandsRestRootResource* InRestResources, EHairStrandsResourcesType InCurveType):
-	FHairCommonResource(EHairStrandsAllocationType::Deferred),
+FHairStrandsDeformedRootResource::FHairStrandsDeformedRootResource(const FHairStrandsRestRootResource* InRestResources, EHairStrandsResourcesType InCurveType, const FHairResourceName& InResourceName):
+	FHairCommonResource(EHairStrandsAllocationType::Deferred, InResourceName),
 	CurveType(InCurveType)
 {
 	check(InRestResources);
@@ -1090,13 +1111,13 @@ void FHairStrandsDeformedRootResource::InternalAllocateLOD(FRDGBuilder& GraphBui
 			LOD.Status = FLOD::EStatus::Initialized;
 			if (LOD.SampleCount > 0)
 			{
-				InternalCreateVertexBufferRDG<FHairStrandsMeshTrianglePositionFormat>(GraphBuilder, LOD.SampleCount, LOD.DeformedSamplePositionsBuffer, HAIRSTRANDS_RESOUCE_NAME(CurveType, Hair.StrandsRootDeformed_DeformedSamplePositionsBuffer), EHairResourceUsageType::Dynamic);
-				InternalCreateVertexBufferRDG<FHairStrandsMeshTrianglePositionFormat>(GraphBuilder, LOD.SampleCount + 4, LOD.MeshSampleWeightsBuffer, HAIRSTRANDS_RESOUCE_NAME(CurveType, Hair.StrandsRootDeformed_MeshSampleWeightsBuffer), EHairResourceUsageType::Dynamic);
+				InternalCreateVertexBufferRDG<FHairStrandsMeshTrianglePositionFormat>(GraphBuilder, LOD.SampleCount, LOD.DeformedSamplePositionsBuffer, ToHairResourceDebugName(HAIRSTRANDS_RESOUCE_NAME(CurveType, Hair.StrandsRootDeformed_DeformedSamplePositionsBuffer), ResourceName), EHairResourceUsageType::Dynamic);
+				InternalCreateVertexBufferRDG<FHairStrandsMeshTrianglePositionFormat>(GraphBuilder, LOD.SampleCount + 4, LOD.MeshSampleWeightsBuffer, ToHairResourceDebugName(HAIRSTRANDS_RESOUCE_NAME(CurveType, Hair.StrandsRootDeformed_MeshSampleWeightsBuffer), ResourceName), EHairResourceUsageType::Dynamic);
 			}
 
-			InternalCreateVertexBufferRDG<FHairStrandsMeshTrianglePositionFormat>(GraphBuilder, RootCount, LOD.DeformedRootTrianglePosition0Buffer, HAIRSTRANDS_RESOUCE_NAME(CurveType, Hair.StrandsRootDeformed_DeformedRootTrianglePosition0Buffer), EHairResourceUsageType::Dynamic);
-			InternalCreateVertexBufferRDG<FHairStrandsMeshTrianglePositionFormat>(GraphBuilder, RootCount, LOD.DeformedRootTrianglePosition1Buffer, HAIRSTRANDS_RESOUCE_NAME(CurveType, Hair.StrandsRootDeformed_DeformedRootTrianglePosition1Buffer), EHairResourceUsageType::Dynamic);
-			InternalCreateVertexBufferRDG<FHairStrandsMeshTrianglePositionFormat>(GraphBuilder, RootCount, LOD.DeformedRootTrianglePosition2Buffer, HAIRSTRANDS_RESOUCE_NAME(CurveType, Hair.StrandsRootDeformed_DeformedRootTrianglePosition2Buffer), EHairResourceUsageType::Dynamic);
+			InternalCreateVertexBufferRDG<FHairStrandsMeshTrianglePositionFormat>(GraphBuilder, RootCount, LOD.DeformedRootTrianglePosition0Buffer, ToHairResourceDebugName(HAIRSTRANDS_RESOUCE_NAME(CurveType, Hair.StrandsRootDeformed_DeformedRootTrianglePosition0Buffer), ResourceName), EHairResourceUsageType::Dynamic);
+			InternalCreateVertexBufferRDG<FHairStrandsMeshTrianglePositionFormat>(GraphBuilder, RootCount, LOD.DeformedRootTrianglePosition1Buffer, ToHairResourceDebugName(HAIRSTRANDS_RESOUCE_NAME(CurveType, Hair.StrandsRootDeformed_DeformedRootTrianglePosition1Buffer), ResourceName), EHairResourceUsageType::Dynamic);
+			InternalCreateVertexBufferRDG<FHairStrandsMeshTrianglePositionFormat>(GraphBuilder, RootCount, LOD.DeformedRootTrianglePosition2Buffer, ToHairResourceDebugName(HAIRSTRANDS_RESOUCE_NAME(CurveType, Hair.StrandsRootDeformed_DeformedRootTrianglePosition2Buffer), ResourceName), EHairResourceUsageType::Dynamic);
 		}
 	}
 }
@@ -1265,8 +1286,8 @@ void FHairStrandsRootData::Reset()
 
 /////////////////////////////////////////////////////////////////////////////////////////
 
-FHairStrandsInterpolationResource::FHairStrandsInterpolationResource(FHairStrandsInterpolationBulkData& InBulkData) :
-	FHairCommonResource(EHairStrandsAllocationType::Deferred),
+FHairStrandsInterpolationResource::FHairStrandsInterpolationResource(FHairStrandsInterpolationBulkData& InBulkData, const FHairResourceName& InResourceName) :
+	FHairCommonResource(EHairStrandsAllocationType::Deferred, InResourceName),
 	InterpolationBuffer(), Interpolation0Buffer(), Interpolation1Buffer(), BulkData(InBulkData)
 {
 }
@@ -1276,14 +1297,14 @@ void FHairStrandsInterpolationResource::InternalAllocate(FRDGBuilder& GraphBuild
 	const bool bUseSingleGuide = !!(BulkData.Flags & FHairStrandsInterpolationBulkData::DataFlags_HasSingleGuideData);
 	if (bUseSingleGuide)
 	{
-		InternalCreateVertexBufferRDG_FromBulkData<FHairStrandsInterpolationFormat>(GraphBuilder, BulkData.Interpolation, BulkData.PointCount, InterpolationBuffer, TEXT("Hair.StrandsInterpolation_InterpolationBuffer"), EHairResourceUsageType::Static);
+		InternalCreateVertexBufferRDG_FromBulkData<FHairStrandsInterpolationFormat>(GraphBuilder, BulkData.Interpolation, BulkData.PointCount, InterpolationBuffer, ToHairResourceDebugName(TEXT("Hair.StrandsInterpolation_InterpolationBuffer"), ResourceName), EHairResourceUsageType::Static);
 	}
 	else
 	{
-		InternalCreateVertexBufferRDG_FromBulkData<FHairStrandsInterpolation0Format>(GraphBuilder, BulkData.Interpolation0, BulkData.PointCount, Interpolation0Buffer, TEXT("Hair.StrandsInterpolation_Interpolation0Buffer"), EHairResourceUsageType::Static);
-		InternalCreateVertexBufferRDG_FromBulkData<FHairStrandsInterpolation1Format>(GraphBuilder, BulkData.Interpolation1, BulkData.PointCount, Interpolation1Buffer, TEXT("Hair.StrandsInterpolation_Interpolation1Buffer"), EHairResourceUsageType::Static);
+		InternalCreateVertexBufferRDG_FromBulkData<FHairStrandsInterpolation0Format>(GraphBuilder, BulkData.Interpolation0, BulkData.PointCount, Interpolation0Buffer, ToHairResourceDebugName(TEXT("Hair.StrandsInterpolation_Interpolation0Buffer"), ResourceName), EHairResourceUsageType::Static);
+		InternalCreateVertexBufferRDG_FromBulkData<FHairStrandsInterpolation1Format>(GraphBuilder, BulkData.Interpolation1, BulkData.PointCount, Interpolation1Buffer, ToHairResourceDebugName(TEXT("Hair.StrandsInterpolation_Interpolation1Buffer"), ResourceName), EHairResourceUsageType::Static);
 	}
-	InternalCreateVertexBufferRDG_FromBulkData<FHairStrandsRootIndexFormat>(GraphBuilder, BulkData.SimRootPointIndex, BulkData.SimPointCount, SimRootPointIndexBuffer, TEXT("Hair.StrandsInterpolation_SimRootPointIndex"), EHairResourceUsageType::Static);
+	InternalCreateVertexBufferRDG_FromBulkData<FHairStrandsRootIndexFormat>(GraphBuilder, BulkData.SimRootPointIndex, BulkData.SimPointCount, SimRootPointIndexBuffer, ToHairResourceDebugName(TEXT("Hair.StrandsInterpolation_SimRootPointIndex"), ResourceName), EHairResourceUsageType::Static);
 }
 
 void FHairStrandsInterpolationResource::InternalRelease()
@@ -1314,15 +1335,15 @@ void FHairCardsInterpolationBulkData::Serialize(FArchive& Ar)
 	Ar << Interpolation;
 }
 
-FHairCardsInterpolationResource::FHairCardsInterpolationResource(FHairCardsInterpolationBulkData& InBulkData) :
-	FHairCommonResource(EHairStrandsAllocationType::Immediate),
+FHairCardsInterpolationResource::FHairCardsInterpolationResource(FHairCardsInterpolationBulkData& InBulkData, const FHairResourceName& InResourceName) :
+	FHairCommonResource(EHairStrandsAllocationType::Immediate, InResourceName),
 	InterpolationBuffer(), BulkData(InBulkData)
 {
 }
 
 void FHairCardsInterpolationResource::InternalAllocate(FRDGBuilder& GraphBuilder)
 {
-	InternalCreateVertexBufferRDG<FHairCardsInterpolationFormat>(GraphBuilder, BulkData.Interpolation, InterpolationBuffer, TEXT("Hair.CardsInterpolation_InterpolationBuffer"), EHairResourceUsageType::Static);
+	InternalCreateVertexBufferRDG<FHairCardsInterpolationFormat>(GraphBuilder, BulkData.Interpolation, InterpolationBuffer, ToHairResourceDebugName(TEXT("Hair.CardsInterpolation_InterpolationBuffer"), ResourceName), EHairResourceUsageType::Static);
 }
 
 void FHairCardsInterpolationResource::InternalRelease()
@@ -1335,18 +1356,18 @@ void FHairCardsInterpolationResource::InternalRelease()
 #if RHI_RAYTRACING
 // RT geometry is built to for a cross around the fiber.
 // 4 triangles per hair vertex => 12 vertices per hair vertex
-FHairStrandsRaytracingResource::FHairStrandsRaytracingResource(const FHairStrandsBulkData& InData) :
-	FHairCommonResource(EHairStrandsAllocationType::Deferred),
+FHairStrandsRaytracingResource::FHairStrandsRaytracingResource(const FHairStrandsBulkData& InData, const FHairResourceName& ResourceName) :
+	FHairCommonResource(EHairStrandsAllocationType::Deferred, ResourceName),
 	PositionBuffer(), VertexCount(InData.GetNumPoints()*12), bOwnPositionBuffer(true)
 {}
 
-FHairStrandsRaytracingResource::FHairStrandsRaytracingResource(const FHairCardsBulkData& InData) :
-	FHairCommonResource(EHairStrandsAllocationType::Deferred),
+FHairStrandsRaytracingResource::FHairStrandsRaytracingResource(const FHairCardsBulkData& InData, const FHairResourceName& ResourceName) :
+	FHairCommonResource(EHairStrandsAllocationType::Deferred, ResourceName),
 	PositionBuffer(), VertexCount(InData.GetNumVertices()), bOwnPositionBuffer(false)
 {}
 
-FHairStrandsRaytracingResource::FHairStrandsRaytracingResource(const FHairMeshesBulkData& InData) :
-	FHairCommonResource(EHairStrandsAllocationType::Deferred),
+FHairStrandsRaytracingResource::FHairStrandsRaytracingResource(const FHairMeshesBulkData& InData, const FHairResourceName& ResourceName) :
+	FHairCommonResource(EHairStrandsAllocationType::Deferred, ResourceName),
 	PositionBuffer(), VertexCount(InData.GetNumVertices()), bOwnPositionBuffer(false)
 {}
 
@@ -1354,7 +1375,7 @@ void FHairStrandsRaytracingResource::InternalAllocate(FRDGBuilder& GraphBuilder)
 {
 	if (bOwnPositionBuffer)
 	{
-		InternalCreateVertexBufferRDG<FHairStrandsRaytracingFormat>(GraphBuilder, VertexCount, PositionBuffer, TEXT("Hair.StrandsRaytracing_PositionBuffer"), EHairResourceUsageType::Dynamic);
+		InternalCreateVertexBufferRDG<FHairStrandsRaytracingFormat>(GraphBuilder, VertexCount, PositionBuffer, ToHairResourceDebugName(TEXT("Hair.StrandsRaytracing_PositionBuffer"), ResourceName), EHairResourceUsageType::Dynamic);
 	}
 }
 
diff --git a/Engine/Plugins/Runtime/HairStrands/Source/HairStrandsCore/Private/HairStrandsRendering.cpp b/Engine/Plugins/Runtime/HairStrands/Source/HairStrandsCore/Private/HairStrandsRendering.cpp
--- a/Engine/Plugins/Runtime/HairStrands/Source/HairStrandsCore/Private/HairStrandsRendering.cpp
+++ b/Engine/Plugins/Runtime/HairStrands/Source/HairStrandsCore/Private/HairStrandsRendering.cpp
@@ -1360,12 +1360,23 @@ static void UpdateHairAccelerationStructure(FRHICommandList& RHICmdList, FRayTra
 	RHICmdList.BuildAccelerationStructures(MakeArrayView(&Params, 1));
 }
 
-static void BuildHairAccelerationStructure_Strands(FRHICommandList& RHICmdList, uint32 RaytracingVertexCount, FBufferRHIRef& PositionBuffer, FRayTracingGeometry* OutRayTracingGeometry)
+static void BuildHairAccelerationStructure_Strands(
+	FRHICommandList& RHICmdList, 
+	uint32 RaytracingVertexCount, 
+	FBufferRHIRef& PositionBuffer, 
+	FRayTracingGeometry* OutRayTracingGeometry,
+	const FString& AssetDebugName,
+	uint32 LODIndex)
 {
 	FRayTracingGeometryInitializer Initializer;
-	static const FName DebugName("HairStrands");
+#if HAIR_RESOURCE_DEBUG_NAME
+	FString DebugName(TEXT("AS_HairStrands_LOD") + FString::FromInt(LODIndex) + TEXT("_") + AssetDebugName);
+	Initializer.DebugName = FName(DebugName);
+#else
+	static const FName DebugName("AS_HairStrands");
 	static int32 DebugNumber = 0;
 	Initializer.DebugName = FName(DebugName, DebugNumber++);
+#endif
 	Initializer.IndexBuffer = nullptr;
 	Initializer.IndexBufferOffset = 0;
 	Initializer.GeometryType = RTGT_Triangles;
@@ -1388,12 +1399,19 @@ static void BuildHairAccelerationStructure_Strands(FRHICommandList& RHICmdList,
 static void BuildHairAccelerationStructure_Cards(FRHICommandList& RHICmdList, 	
 	FHairCardsRestResource* RestResource,
 	FHairCardsDeformedResource* DeformedResource,
-	FRayTracingGeometry* OutRayTracingGeometry)
+	FRayTracingGeometry* OutRayTracingGeometry,
+	const FString& AssetDebugName,
+	uint32 LODIndex)
 {
 	FRayTracingGeometryInitializer Initializer;
-	static const FName DebugName("HairCards");
+#if HAIR_RESOURCE_DEBUG_NAME
+	FString DebugName(TEXT("AS_HairCards_LOD") + FString::FromInt(LODIndex) + TEXT("_") + AssetDebugName);
+	Initializer.DebugName = FName(DebugName);
+#else
+	static const FName DebugName("AS_HairCards");
 	static int32 DebugNumber = 0;
 	Initializer.DebugName = FName(DebugName, DebugNumber++);
+#endif
 	Initializer.IndexBuffer = RestResource->RestIndexBuffer.IndexBufferRHI;
 	Initializer.IndexBufferOffset = 0;
 	Initializer.GeometryType = RTGT_Triangles;
@@ -1422,12 +1440,19 @@ static void BuildHairAccelerationStructure_Cards(FRHICommandList& RHICmdList,
 static void BuildHairAccelerationStructure_Meshes(FRHICommandList& RHICmdList,
 	FHairMeshesRestResource* RestResource,
 	FHairMeshesDeformedResource* DeformedResource,
-	FRayTracingGeometry* OutRayTracingGeometry)
+	FRayTracingGeometry* OutRayTracingGeometry,
+	const FString& AssetDebugName,
+	uint32 LODIndex)
 {
 	FRayTracingGeometryInitializer Initializer;
-	static const FName DebugName("HairMeshes");
+#if HAIR_RESOURCE_DEBUG_NAME
+	FString DebugName(TEXT("AS_HairMeshes_LOD") + FString::FromInt(LODIndex) + TEXT("_") + AssetDebugName);
+	Initializer.DebugName = FName(DebugName);
+#else
+	static const FName DebugName("AS_HairMeshes");
 	static int32 DebugNumber = 0;
 	Initializer.DebugName = FName(DebugName, DebugNumber++);
+#endif
 	Initializer.IndexBuffer = RestResource->IndexBuffer.IndexBufferRHI;
 	Initializer.IndexBufferOffset = 0;
 	Initializer.GeometryType = RTGT_Triangles;
@@ -1877,6 +1902,8 @@ void ComputeHairStrandsInterpolation(
 				const bool bNeedBuild = !Instance->Strands.RenRaytracingResource->bIsRTGeometryInitialized;
 				if (bNeedBuild || bNeedUpdate)
 				{
+					const uint32 HairLODIndex = Instance->HairGroupPublicData->GetIntLODIndex();
+
 					FRDGImportedBuffer Raytracing_PositionBuffer = Register(GraphBuilder, Instance->Strands.RenRaytracingResource->PositionBuffer, ERDGImportedBufferFlags::CreateViews);
 					AddGenerateRaytracingGeometryPass(
 						GraphBuilder,
@@ -1893,7 +1920,7 @@ void ComputeHairStrandsInterpolation(
 					GraphBuilder.AddPass(
 						RDG_EVENT_NAME("HairStrands::UpdateBLAS(Strands)"),
 						ERDGPassFlags::NeverCull,
-					[Instance, bNeedUpdate](FRHICommandListImmediate& RHICmdList)
+					[Instance, HairLODIndex, bNeedUpdate](FRHICommandListImmediate& RHICmdList)
 					{
 						SCOPED_GPU_MASK(RHICmdList, FRHIGPUMask::All());
 
@@ -1901,7 +1928,7 @@ void ComputeHairStrandsInterpolation(
 						if (bLocalNeedBuild)
 						{
 							FBufferRHIRef PositionBuffer(Instance->Strands.RenRaytracingResource->PositionBuffer.Buffer->GetRHI());
-							BuildHairAccelerationStructure_Strands(RHICmdList, Instance->Strands.RenRaytracingResource->VertexCount, PositionBuffer, &Instance->Strands.RenRaytracingResource->RayTracingGeometry);
+							BuildHairAccelerationStructure_Strands(RHICmdList, Instance->Strands.RenRaytracingResource->VertexCount, PositionBuffer, &Instance->Strands.RenRaytracingResource->RayTracingGeometry, Instance->Debug.GroomAssetName, HairLODIndex);
 						}
 						else if (bNeedUpdate)
 						{
@@ -2031,7 +2058,7 @@ void ComputeHairStrandsInterpolation(
 						const bool bLocalNeedBuild = !LocalLOD.RaytracingResource->bIsRTGeometryInitialized;
 						if (bLocalNeedBuild)
 						{
-							BuildHairAccelerationStructure_Cards(RHICmdList, LocalLOD.RestResource, LocalLOD.DeformedResource, &LocalLOD.RaytracingResource->RayTracingGeometry);
+							BuildHairAccelerationStructure_Cards(RHICmdList, LocalLOD.RestResource, LocalLOD.DeformedResource, &LocalLOD.RaytracingResource->RayTracingGeometry, Instance->Debug.GroomAssetName, HairLODIndex);
 							LocalLOD.RaytracingResource->bIsRTGeometryInitialized = true;
 						}
 						else if (bNeedUpdate)
@@ -2095,7 +2122,7 @@ void ComputeHairStrandsInterpolation(
 						const bool bLocalNeedBuild = !LocalLOD.RaytracingResource->bIsRTGeometryInitialized;
 						if (bLocalNeedBuild)
 						{
-							BuildHairAccelerationStructure_Meshes(RHICmdList, LocalLOD.RestResource, LocalLOD.DeformedResource,  &LocalLOD.RaytracingResource->RayTracingGeometry);
+							BuildHairAccelerationStructure_Meshes(RHICmdList, LocalLOD.RestResource, LocalLOD.DeformedResource,  &LocalLOD.RaytracingResource->RayTracingGeometry, Instance->Debug.GroomAssetName, HairLODIndex);
 						}
 						else if (bNeedUpdate)
 						{
diff --git a/Engine/Plugins/Runtime/HairStrands/Source/HairStrandsCore/Public/GroomResources.h b/Engine/Plugins/Runtime/HairStrands/Source/HairStrandsCore/Public/GroomResources.h
--- a/Engine/Plugins/Runtime/HairStrands/Source/HairStrandsCore/Public/GroomResources.h
+++ b/Engine/Plugins/Runtime/HairStrands/Source/HairStrandsCore/Public/GroomResources.h
@@ -29,11 +29,32 @@ enum class EHairStrandsAllocationType : uint8
 	Deferred	// Resources allocation is deferred to first usage
 };
 
+/* Hair resource name - Allow more precise resource name for debug purpose. Enabled with HAIR_RESOURCE_DEBUG_NAME */
+#define HAIR_RESOURCE_DEBUG_NAME 0
+struct FHairResourceName
+{
+	FHairResourceName() = default;
+#if HAIR_RESOURCE_DEBUG_NAME
+	FHairResourceName(const FName& In) : AssetName(In) {}
+	FHairResourceName(const FName& In, int32 InGroupIndex) : AssetName(In), GroupIndex(InGroupIndex) {}
+	FHairResourceName(const FName& In, int32 InGroupIndex, int32 InLODIndex) : AssetName(In), GroupIndex(InGroupIndex), LODIndex(InLODIndex) {}
+
+	FName AssetName;
+	int32 GroupIndex = -1;
+	int32 LODIndex = -1;
+	TArray<FString> Names;
+#else
+	FHairResourceName(const FName& In) {}
+	FHairResourceName(const FName& In, int32 InGroupIndex) {}
+	FHairResourceName(const FName& In, int32 InGroupIndex, int32 InLODIndex) {}
+#endif
+};
+
 /* Hair resouces which whom allocation can be deferred */
 struct FHairCommonResource : public FRenderResource
 {
 	/** Build the hair strands resource */
-	FHairCommonResource(EHairStrandsAllocationType AllocationType, bool bUseRenderGraph=true);
+	FHairCommonResource(EHairStrandsAllocationType AllocationType, const FHairResourceName& InResourceName, bool bUseRenderGraph=true);
 
 	/* Init/Release buffers (FRenderResource) */
 	virtual void InitRHI() override;
@@ -50,13 +71,16 @@ struct FHairCommonResource : public FRenderResource
 	bool bUseRenderGraph = true;
 	bool bIsInitialized = false;
 	EHairStrandsAllocationType AllocationType = EHairStrandsAllocationType::Deferred;
+
+	/* Store (precis) debug name */
+	FHairResourceName ResourceName;
 };
 
 /* Render buffers for root deformation for dynamic meshes */
 struct FHairStrandsRestRootResource : public FHairCommonResource
 {
 	/** Build the hair strands resource */
-	FHairStrandsRestRootResource(FHairStrandsRootBulkData& BulkData, EHairStrandsResourcesType CurveType);
+	FHairStrandsRestRootResource(FHairStrandsRootBulkData& BulkData, EHairStrandsResourcesType CurveType, const FHairResourceName& ResourceName);
 
 	/* Init/Release buffers */
 	virtual void InternalAllocate(FRDGBuilder& GraphBuilder) override;
@@ -130,8 +154,8 @@ struct FHairStrandsRestRootResource : public FHairCommonResource
 struct FHairStrandsDeformedRootResource : public FHairCommonResource
 {
 	/** Build the hair strands resource */
-	FHairStrandsDeformedRootResource(EHairStrandsResourcesType CurveType);
-	FHairStrandsDeformedRootResource(const FHairStrandsRestRootResource* InRestResources, EHairStrandsResourcesType CurveType);
+	FHairStrandsDeformedRootResource(EHairStrandsResourcesType CurveType, const FHairResourceName& ResourceName);
+	FHairStrandsDeformedRootResource(const FHairStrandsRestRootResource* InRestResources, EHairStrandsResourcesType CurveType, const FHairResourceName& ResourceName);
 
 	/* Init/Release buffers */
 	virtual void InternalAllocateLOD(FRDGBuilder& GraphBuilder, int32 LODIndex) override;
@@ -196,7 +220,7 @@ struct FHairStrandsDeformedRootResource : public FHairCommonResource
 struct FHairStrandsRestResource : public FHairCommonResource
 {
 	/** Build the hair strands resource */
-	FHairStrandsRestResource(FHairStrandsBulkData& InBulkData, EHairStrandsResourcesType CurveType);
+	FHairStrandsRestResource(FHairStrandsBulkData& InBulkData, EHairStrandsResourcesType CurveType, const FHairResourceName& ResourceName);
 
 	/* Init/Release buffers */
 	virtual void InternalAllocate(FRDGBuilder& GraphBuilder) override;
@@ -252,7 +276,7 @@ struct FHairStrandsRestResource : public FHairCommonResource
 struct FHairStrandsDeformedResource : public FHairCommonResource
 {
 	/** Build the hair strands resource */
-	FHairStrandsDeformedResource(FHairStrandsBulkData& BulkData, EHairStrandsResourcesType CurveType);
+	FHairStrandsDeformedResource(FHairStrandsBulkData& BulkData, EHairStrandsResourcesType CurveType, const FHairResourceName& ResourceName);
 
 	/* Init/Release buffers */
 	virtual void InternalAllocate(FRDGBuilder& GraphBuilder) override;
@@ -319,7 +343,7 @@ struct FHairStrandsDeformedResource : public FHairCommonResource
 
 struct FHairStrandsClusterCullingResource : public FHairCommonResource
 {
-	FHairStrandsClusterCullingResource(FHairStrandsClusterCullingBulkData& Data);
+	FHairStrandsClusterCullingResource(FHairStrandsClusterCullingBulkData& Data, const FHairResourceName& ResourceName);
 
 	/* Init/Release buffers */
 	virtual void InternalAllocate(FRDGBuilder& GraphBuilder) override;
@@ -355,7 +379,7 @@ struct FHairStrandsClusterCullingResource : public FHairCommonResource
 struct FHairStrandsInterpolationResource : public FHairCommonResource
 {
 	/** Build the hair strands resource */
-	FHairStrandsInterpolationResource(FHairStrandsInterpolationBulkData& InBulkData);
+	FHairStrandsInterpolationResource(FHairStrandsInterpolationBulkData& InBulkData, const FHairResourceName& ResourceName);
 	
 	/* Init/Release buffers */
 	virtual void InternalAllocate(FRDGBuilder& GraphBuilder) override;
@@ -390,9 +414,9 @@ struct FHairStrandsInterpolationResource : public FHairCommonResource
 struct FHairStrandsRaytracingResource : public FHairCommonResource
 {
 	/** Build the hair strands resource */
-	FHairStrandsRaytracingResource(const FHairStrandsBulkData& InData);
-	FHairStrandsRaytracingResource(const FHairCardsBulkData& InData);
-	FHairStrandsRaytracingResource(const FHairMeshesBulkData& InData);
+	FHairStrandsRaytracingResource(const FHairStrandsBulkData& InData, const FHairResourceName& ResourceName);
+	FHairStrandsRaytracingResource(const FHairCardsBulkData& InData, const FHairResourceName& ResourceName);
+	FHairStrandsRaytracingResource(const FHairMeshesBulkData& InData, const FHairResourceName& ResourceName);
 
 	/* Init/Release buffers */
 	virtual void InternalAllocate(FRDGBuilder& GraphBuilder) override;
@@ -440,7 +464,7 @@ struct FHairCardsBulkData;
 struct FHairCardsRestResource : public FHairCommonResource
 {
 	/** Build the hair strands resource */
-	FHairCardsRestResource(const FHairCardsBulkData& InBulkData);
+	FHairCardsRestResource(const FHairCardsBulkData& InBulkData, const FHairResourceName& ResourceName);
 
 	virtual void InitResource() override;
 	virtual void ReleaseResource() override;
@@ -533,7 +557,7 @@ struct FHairCardsProceduralResource : public FHairCommonResource
 struct FHairCardsDeformedResource : public FHairCommonResource
 {
 	/** Build the hair strands resource */
-	FHairCardsDeformedResource(const FHairCardsBulkData& BulkData, bool bInitializeData);
+	FHairCardsDeformedResource(const FHairCardsBulkData& BulkData, bool bInitializeData, const FHairResourceName& ResourceName);
 
 	/* Init/release buffers */
 	virtual void InternalAllocate(FRDGBuilder& GraphBuilder) override;
@@ -609,7 +633,7 @@ struct HAIRSTRANDSCORE_API FHairCardsInterpolationBulkData
 struct FHairCardsInterpolationResource : public FHairCommonResource
 {
 	/** Build the hair strands resource */
-	FHairCardsInterpolationResource(FHairCardsInterpolationBulkData& InBulkData);
+	FHairCardsInterpolationResource(FHairCardsInterpolationBulkData& InBulkData, const FHairResourceName& ResourceName);
 
 	/* Init/release buffers */
 	virtual void InternalAllocate(FRDGBuilder& GraphBuilder) override;
@@ -639,7 +663,7 @@ struct FHairCardsInterpolationResource : public FHairCommonResource
 struct FHairMeshesRestResource : public FHairCommonResource
 {
 	/** Build the hair strands resource */
-	FHairMeshesRestResource(const FHairMeshesBulkData& BulkData);
+	FHairMeshesRestResource(const FHairMeshesBulkData& BulkData, const FHairResourceName& ResourceName);
 
 	virtual void InitResource() override;
 	virtual void ReleaseResource() override;
@@ -693,7 +717,7 @@ struct FHairMeshesRestResource : public FHairCommonResource
 struct FHairMeshesDeformedResource : public FHairCommonResource
 {
 	/** Build the hair strands resource */
-	FHairMeshesDeformedResource(const FHairMeshesBulkData& InBulkData, bool bInInitializedData);
+	FHairMeshesDeformedResource(const FHairMeshesBulkData& InBulkData, bool bInInitializedData, const FHairResourceName& ResourceName);
 
 	/* Init/release buffers */
 	virtual void InternalAllocate(FRDGBuilder& GraphBuilder) override;
