Second round of temp. fix for GPU crash.

diff --git a/Engine/Plugins/Runtime/HairStrands/Source/HairStrandsCore/Private/Niagara/NiagaraDataInterfaceHairStrands.cpp b/Engine/Plugins/Runtime/HairStrands/Source/HairStrandsCore/Private/Niagara/NiagaraDataInterfaceHairStrands.cpp
--- a/Engine/Plugins/Runtime/HairStrands/Source/HairStrandsCore/Private/Niagara/NiagaraDataInterfaceHairStrands.cpp
+++ b/Engine/Plugins/Runtime/HairStrands/Source/HairStrandsCore/Private/Niagara/NiagaraDataInterfaceHairStrands.cpp
@@ -353,7 +353,7 @@ void FNDIHairStrandsBuffer::InitRHI()
 {
 	if (SourceRestResources != nullptr)
 	{
-		FHairStrandsBulkData* SourceDatas = &SourceRestResources->BulkData;
+		FHairStrandsBulkData* SourceDatas = &SourceRestResources->BulkData; // This could be released by that time depending on how the initialization order is
 		{
 			const uint32 OffsetCount = SourceDatas->GetNumCurves() + 1;
 			const uint32 OffsetBytes = sizeof(uint32)*OffsetCount;
@@ -637,8 +637,15 @@ struct FNDIHairStrandsParametersCS : public FNiagaraDataInterfaceParametersCS
 
 		const bool bIsHairValid = ProxyData != nullptr && ProxyData->HairStrandsBuffer && ProxyData->HairStrandsBuffer->IsInitialized();
 		const bool bHasSkinningBinding = bIsHairValid && ProxyData->HairGroupInstance && ProxyData->HairGroupInstance->BindingType == EHairBindingType::Skinning;
-		const bool bIsRootValid = bIsHairValid && ProxyData->HairStrandsBuffer->SourceDeformedRootResources && bHasSkinningBinding;
-		const bool bIsRestValid = bIsHairValid && ProxyData->HairStrandsBuffer->SourceRestResources && ProxyData->HairStrandsBuffer->SourceRestResources->PositionBuffer.SRV && ProxyData->HairStrandsBuffer->CurvesOffsetsBuffer.SRV;
+		const bool bIsRootValid = bIsHairValid && ProxyData->HairStrandsBuffer->SourceDeformedRootResources && ProxyData->HairStrandsBuffer->SourceDeformedRootResources->IsInitialized() && bHasSkinningBinding;
+		const bool bIsRestValid = bIsHairValid && ProxyData->HairStrandsBuffer->SourceRestResources && 
+
+			// SourceRestResources->PositionBuffer is lazily allocated, when the instance LOD is scheduled (this happens after this called). So this is why this check is here. 
+			// This code should be refactor so that it reflects the lazy allocation scheme
+			ProxyData->HairStrandsBuffer->SourceRestResources->PositionBuffer.SRV && 
+			// TEMP: These check are only temporary for avoiding crashes while we find the bottom of the issue.
+			ProxyData->HairStrandsBuffer->CurvesOffsetsBuffer.SRV && ProxyData->HairStrandsBuffer->ParamsScaleBuffer.SRV && ProxyData->HairStrandsBuffer->BoundingBoxBuffer.UAV;
+
 		const bool bIsDeformedValid = bIsHairValid && ProxyData->HairStrandsBuffer->SourceDeformedResources && ProxyData->HairStrandsBuffer->SourceDeformedResources->IsInitialized();
 		const bool bHasLODSwitched = bIsHairValid && ProxyData->HairGroupInstance && ProxyData->HairGroupInstance->HairGroupPublicData && ProxyData->HairGroupInstance->HairGroupPublicData->VFInput.bHasLODSwitch;
 
