Add enable/disable blueprint function for hair length scale.

diff --git a/Engine/Plugins/Runtime/HairStrands/Source/HairStrandsCore/Private/GroomComponent.cpp b/Engine/Plugins/Runtime/HairStrands/Source/HairStrandsCore/Private/GroomComponent.cpp
--- a/Engine/Plugins/Runtime/HairStrands/Source/HairStrandsCore/Private/GroomComponent.cpp
+++ b/Engine/Plugins/Runtime/HairStrands/Source/HairStrandsCore/Private/GroomComponent.cpp
@@ -1504,6 +1504,16 @@ void UGroomComponent::SetHairLengthScale(float InLengthScale)
 	UpdateHairGroupsDescAndInvalidateRenderState(false);
 }
 
+void UGroomComponent::SetHairLengthScaleEnable(bool bEnable)
+{
+	for (FHairGroupDesc& HairDesc : GroomGroupsDesc)
+	{
+		HairDesc.HairLengthScale_Override = bEnable;
+	}
+	UpdateHairGroupsDescAndInvalidateRenderState(true);
+	InitResources();
+}
+
 void UGroomComponent::SetForcedLOD(int32 LODIndex)
 {
 	if (GroomAsset)
diff --git a/Engine/Plugins/Runtime/HairStrands/Source/HairStrandsCore/Public/GroomComponent.h b/Engine/Plugins/Runtime/HairStrands/Source/HairStrandsCore/Public/GroomComponent.h
--- a/Engine/Plugins/Runtime/HairStrands/Source/HairStrandsCore/Public/GroomComponent.h
+++ b/Engine/Plugins/Runtime/HairStrands/Source/HairStrandsCore/Public/GroomComponent.h
@@ -163,6 +163,9 @@ public:
 	UFUNCTION(BlueprintCallable, Category = "Groom")
 	void SetHairLengthScale(float Scale);
 
+	UFUNCTION(BlueprintCallable, Category = "Groom")
+	void SetHairLengthScaleEnable(bool bEnable);
+
 	void SetStableRasterization(bool bEnable);
 	void SetGroomAsset(UGroomAsset* Asset, UGroomBindingAsset* InBinding, const bool bUpdateSimulation = true);
 	void SetHairRootScale(float Scale);
