Fix crash when removing hair LODs

diff --git a/Engine/Plugins/Runtime/HairStrands/Source/HairStrandsCore/Private/GroomAsset.cpp b/Engine/Plugins/Runtime/HairStrands/Source/HairStrandsCore/Private/GroomAsset.cpp
--- a/Engine/Plugins/Runtime/HairStrands/Source/HairStrandsCore/Private/GroomAsset.cpp
+++ b/Engine/Plugins/Runtime/HairStrands/Source/HairStrandsCore/Private/GroomAsset.cpp
@@ -2918,11 +2918,21 @@ void UGroomAsset::InitStrandsResources()
 
 			if (GroupData.Strands.ClusterCullingBulkData.IsValid())
 			{
+				const uint32 HairGroupLODCount = HairGroupsLOD[GroupIndex].LODs.Num();
+				bool bLastVisibility = true;
 				// LOD visibility is not serialized into FHairStrandsClusterCullingBulkData as it does not affect the actual generated data. For consistency we 
 				// patch the LOD visibility with the groom asset value, which might be different from what has been serialized
 				for (uint32 LODIt = 0, LODCount = GroupData.Strands.ClusterCullingBulkData.LODVisibility.Num(); LODIt < LODCount; ++LODIt)
 				{
-					GroupData.Strands.ClusterCullingBulkData.LODVisibility[LODIt] = HairGroupsLOD[GroupIndex].LODs[LODIt].bVisible;
+					if (LODIt < HairGroupLODCount)
+					{
+						bLastVisibility = HairGroupsLOD[GroupIndex].LODs[LODIt].bVisible;
+						GroupData.Strands.ClusterCullingBulkData.LODVisibility[LODIt] = bLastVisibility;
+					}
+					else
+					{
+						GroupData.Strands.ClusterCullingBulkData.LODVisibility[LODIt] = bLastVisibility;
+					}
 				}
 
 				GroupData.Strands.ClusterCullingResource = new FHairStrandsClusterCullingResource(GroupData.Strands.ClusterCullingBulkData, ResourceName);
