Add function to check if hair length scale is enabled

diff --git a/Engine/Plugins/Runtime/HairStrands/Source/HairStrandsCore/Private/GroomComponent.cpp b/Engine/Plugins/Runtime/HairStrands/Source/HairStrandsCore/Private/GroomComponent.cpp
--- a/Engine/Plugins/Runtime/HairStrands/Source/HairStrandsCore/Private/GroomComponent.cpp
+++ b/Engine/Plugins/Runtime/HairStrands/Source/HairStrandsCore/Private/GroomComponent.cpp
@@ -1514,6 +1514,20 @@ void UGroomComponent::SetHairLengthScaleEnable(bool bEnable)
 	InitResources();
 }
 
+bool UGroomComponent::GetIsHairLengthScaleEnabled()
+{
+	bool IsEnabled = true;
+	for (FHairGroupDesc& HairDesc : GroomGroupsDesc)
+	{
+		if (!HairDesc.HairLengthScale_Override)
+		{
+			IsEnabled = false;
+		}
+	}
+	return IsEnabled;
+}
+
+
 void UGroomComponent::SetForcedLOD(int32 LODIndex)
 {
 	if (GroomAsset)
diff --git a/Engine/Plugins/Runtime/HairStrands/Source/HairStrandsCore/Public/GroomComponent.h b/Engine/Plugins/Runtime/HairStrands/Source/HairStrandsCore/Public/GroomComponent.h
--- a/Engine/Plugins/Runtime/HairStrands/Source/HairStrandsCore/Public/GroomComponent.h
+++ b/Engine/Plugins/Runtime/HairStrands/Source/HairStrandsCore/Public/GroomComponent.h
@@ -166,6 +166,9 @@ public:
 	UFUNCTION(BlueprintCallable, Category = "Groom")
 	void SetHairLengthScaleEnable(bool bEnable);
 
+	UFUNCTION(BlueprintCallable, Category = "Groom")
+	bool GetIsHairLengthScaleEnabled();
+
 	void SetStableRasterization(bool bEnable);
 	void SetGroomAsset(UGroomAsset* Asset, UGroomBindingAsset* InBinding, const bool bUpdateSimulation = true);
 	void SetHairRootScale(float Scale);
