Added GroomCache import settings to specify the animation range and wether it should skip empty frames at the start. Work like the corresponding Alembic sampling settings.

diff --git a/Engine/Plugins/Runtime/HairStrands/Source/HairStrandsCore/Private/GroomCache.cpp b/Engine/Plugins/Runtime/HairStrands/Source/HairStrandsCore/Private/GroomCache.cpp
--- a/Engine/Plugins/Runtime/HairStrands/Source/HairStrandsCore/Private/GroomCache.cpp
+++ b/Engine/Plugins/Runtime/HairStrands/Source/HairStrandsCore/Private/GroomCache.cpp
@@ -55,7 +55,10 @@ int32 UGroomCache::GetFrameIndexAtTime(const float Time, bool bLooping) const
 	}
 
 	const int32 NumFrames = GroomCacheInfo.AnimationInfo.NumFrames;
-	const float Duration = GroomCacheInfo.AnimationInfo.Duration;
+
+	// Include a small fudge factor to Duration to account for possible computation discrepancies in Time.
+	// For example when sequencer computes the time for the end of a section, it might not exactly match the duration
+	const float Duration = GroomCacheInfo.AnimationInfo.Duration - KINDA_SMALL_NUMBER;
 	float AdjustedTime = Time;
 	if (bLooping)
 	{
@@ -73,10 +76,10 @@ int32 UGroomCache::GetFrameIndexAtTime(const float Time, bool bLooping) const
 void UGroomCache::GetFrameIndicesAtTime(float Time, bool bLooping, bool bIsPlayingBackwards, int32 &OutFrameIndex, int32 &OutNextFrameIndex, float &InterpolationFactor)
 {
 	const int32 NumFrames = GroomCacheInfo.AnimationInfo.NumFrames;
-	const float Duration = GroomCacheInfo.AnimationInfo.Duration;
+	const float Duration = GroomCacheInfo.AnimationInfo.Duration - KINDA_SMALL_NUMBER;
 
 	// No index possible
-	if (NumFrames == 0 || NumFrames == 1 || Duration == 0.0f)
+	if (NumFrames == 0 || Duration == 0.0f)
 	{
 		OutFrameIndex = 0;
 		OutNextFrameIndex = 0;
@@ -85,7 +88,7 @@ void UGroomCache::GetFrameIndicesAtTime(float Time, bool bLooping, bool bIsPlayi
 	}
 
 	OutFrameIndex = GetFrameIndexAtTime(Time, bLooping);
-	OutNextFrameIndex = FMath::Min(OutFrameIndex + 1, NumFrames - 1);
+	OutNextFrameIndex = FMath::Min(OutFrameIndex + 1, NumFrames);
 
 	const float FrameDuration = GroomCacheInfo.AnimationInfo.SecondsPerFrame;
 	if (FMath::IsNearlyZero(FrameDuration))
diff --git a/Engine/Plugins/Runtime/HairStrands/Source/HairStrandsCore/Public/GroomCacheImportOptions.h b/Engine/Plugins/Runtime/HairStrands/Source/HairStrandsCore/Public/GroomCacheImportOptions.h
--- a/Engine/Plugins/Runtime/HairStrands/Source/HairStrandsCore/Public/GroomCacheImportOptions.h
+++ b/Engine/Plugins/Runtime/HairStrands/Source/HairStrandsCore/Public/GroomCacheImportOptions.h
@@ -15,6 +15,18 @@ struct HAIRSTRANDSCORE_API FGroomCacheImportSettings
 	UPROPERTY(EditAnywhere, BlueprintReadWrite, Category = GroomCache)
 	bool bImportGroomCache = true;
 
+	/** Starting index to start sampling the animation from */
+	UPROPERTY(EditAnywhere, BlueprintReadWrite, Category = Sampling)
+	int32 FrameStart = 0;
+
+	/** Ending index to stop sampling the animation at */
+	UPROPERTY(EditAnywhere, BlueprintReadWrite, Category = Sampling)
+	int32 FrameEnd = 0;
+
+	/** Skip empty (pre-roll) frames and start importing at the frame which actually contains data */
+	UPROPERTY(EditAnywhere, BlueprintReadWrite, Category = Sampling, meta = (DisplayName = "Skip Empty Frames at Start of Groom Animation"))
+	bool bSkipEmptyFrames = false;
+
 	/** Import or re-import the groom asset from this file */
 	UPROPERTY(EditAnywhere, BlueprintReadWrite, Category = GroomCache)
 	bool bImportGroomAsset = true;
diff --git a/Engine/Plugins/Runtime/HairStrands/Source/HairStrandsEditor/Private/GroomCacheImportSettingsCustomization.cpp b/Engine/Plugins/Runtime/HairStrands/Source/HairStrandsEditor/Private/GroomCacheImportSettingsCustomization.cpp
--- a/Engine/Plugins/Runtime/HairStrands/Source/HairStrandsEditor/Private/GroomCacheImportSettingsCustomization.cpp
+++ b/Engine/Plugins/Runtime/HairStrands/Source/HairStrandsEditor/Private/GroomCacheImportSettingsCustomization.cpp
@@ -31,6 +31,13 @@ void FGroomCacheImportSettingsCustomization::CustomizeChildren(TSharedRef<IPrope
 			// Visibility depends on other settings
 			VisibleType = 1;
 		}
+		else if	(ChildHandle->GetProperty()->GetFName() == GET_MEMBER_NAME_CHECKED(FGroomCacheImportSettings, FrameStart) ||
+			ChildHandle->GetProperty()->GetFName() == GET_MEMBER_NAME_CHECKED(FGroomCacheImportSettings, FrameEnd) ||
+			ChildHandle->GetProperty()->GetFName() == GET_MEMBER_NAME_CHECKED(FGroomCacheImportSettings, bSkipEmptyFrames))
+		{
+			// Visibility depends on Import Groom Cache
+			VisibleType = 2;
+		}
 
 		IDetailPropertyRow& Property = StructBuilder.AddProperty(ChildHandle);
 		Property.Visibility(TAttribute<EVisibility>::Create(TAttribute<EVisibility>::FGetter::CreateSP(this, &FGroomCacheImportSettingsCustomization::ArePropertiesVisible, VisibleType)));
@@ -45,5 +52,9 @@ EVisibility FGroomCacheImportSettingsCustomization::ArePropertiesVisible(const i
 		// If GroomCache is not imported, then this property is not relevant
 		return Settings->bImportGroomAsset || !Settings->bImportGroomCache ? EVisibility::Collapsed : EVisibility::Visible;
 	}
+	else if (VisibleType == 2 && Settings)
+	{
+		return !Settings->bImportGroomCache ? EVisibility::Collapsed : EVisibility::Visible;
+	}
 	return EVisibility::Visible;
 }
diff --git a/Engine/Plugins/Runtime/HairStrands/Source/HairStrandsEditor/Private/GroomCacheImporter.cpp b/Engine/Plugins/Runtime/HairStrands/Source/HairStrandsEditor/Private/GroomCacheImporter.cpp
--- a/Engine/Plugins/Runtime/HairStrands/Source/HairStrandsEditor/Private/GroomCacheImporter.cpp
+++ b/Engine/Plugins/Runtime/HairStrands/Source/HairStrandsEditor/Private/GroomCacheImporter.cpp
@@ -4,6 +4,7 @@
 #include "GroomAsset.h"
 #include "GroomBuilder.h"
 #include "GroomCache.h"
+#include "GroomCacheImportOptions.h"
 #include "GroomImportOptions.h"
 #include "HairStrandsImporter.h"
 #include "HairStrandsTranslator.h"
@@ -101,7 +102,8 @@ TArray<UGroomCache*> FGroomCacheImporter::ImportGroomCache(const FString& Source
 	FGroomCacheProcessor GuidesProcessor(EGroomCacheType::Guides, AnimInfo.Attributes);
 	if (Translator->BeginTranslation(SourceFilename))
 	{
-		const int32 NumFrames = AnimInfo.NumFrames;
+		// Sample one extra frame so that we can interpolate between EndFrame - 1 and EndFrame
+		const int32 NumFrames = AnimInfo.NumFrames + 1;
 		FScopedSlowTask SlowTask(NumFrames, LOCTEXT("ImportGroomCache", "Importing GroomCache frames"));
 		SlowTask.MakeDialog();
 
@@ -195,4 +197,69 @@ TArray<UGroomCache*> FGroomCacheImporter::ImportGroomCache(const FString& Source
 	return GroomCaches;
 }
 
+void FGroomCacheImporter::SetupImportSettings(FGroomCacheImportSettings& ImportSettings, const FGroomAnimationInfo& AnimInfo)
+{
+	// Prepare the import settings for display
+	// GroomCache options are only shown if there's a valid groom animation
+	ImportSettings.bImportGroomCache = ImportSettings.bImportGroomCache && AnimInfo.IsValid();
+
+	if (ImportSettings.bImportGroomCache)
+	{
+		if (ImportSettings.FrameEnd == 0)
+		{
+			ImportSettings.FrameEnd = AnimInfo.EndFrame;
+		}
+	}
+}
+
+void FGroomCacheImporter::ApplyImportSettings(FGroomCacheImportSettings& ImportSettings, FGroomAnimationInfo& AnimInfo)
+{
+	// Harmonize the values between what's in the settings (set by the user) and the animation info (extracted from the Alembic) used for importing
+	// The user settings usually take precedence over the animation info
+	if (ImportSettings.bImportGroomCache)
+	{
+		if (ImportSettings.bSkipEmptyFrames)
+		{
+			// Skipping empty frames will start from the beginning of the animation range or beyond if specified by the user
+			if (ImportSettings.FrameStart > AnimInfo.StartFrame)
+			{
+				AnimInfo.StartFrame = ImportSettings.FrameStart;
+			}
+			else
+			{
+				ImportSettings.FrameStart = AnimInfo.StartFrame;
+			}
+		}
+		else
+		{
+			// Otherwise, just take the value set by the user
+			AnimInfo.StartFrame = ImportSettings.FrameStart;
+		}
+
+		if (ImportSettings.FrameEnd == 0)
+		{
+			// If the user manually set the end to 0, use the actual end of the animation range
+			ImportSettings.FrameEnd = AnimInfo.EndFrame;
+		}
+		else
+		{
+			// Otherwise, just take the value set by the user
+			AnimInfo.EndFrame = ImportSettings.FrameEnd;
+		}
+
+		// Sanity check
+		if (ImportSettings.FrameEnd <= ImportSettings.FrameStart)
+		{
+			ImportSettings.FrameEnd = ImportSettings.FrameStart + 1;
+			AnimInfo.EndFrame = ImportSettings.FrameStart + 1;
+		}
+
+		// EndFrame is not included and must have at least 1 frame
+		AnimInfo.NumFrames = FMath::Max(AnimInfo.EndFrame - AnimInfo.StartFrame, 1);
+
+		// Compute the duration as it is not known yet
+		AnimInfo.Duration = AnimInfo.NumFrames * AnimInfo.SecondsPerFrame;
+	}
+}
+
 #undef LOCTEXT_NAMESPACE
\ No newline at end of file
diff --git a/Engine/Plugins/Runtime/HairStrands/Source/HairStrandsEditor/Private/HairStrandsFactory.cpp b/Engine/Plugins/Runtime/HairStrands/Source/HairStrandsEditor/Private/HairStrandsFactory.cpp
--- a/Engine/Plugins/Runtime/HairStrands/Source/HairStrandsEditor/Private/HairStrandsFactory.cpp
+++ b/Engine/Plugins/Runtime/HairStrands/Source/HairStrandsEditor/Private/HairStrandsFactory.cpp
@@ -120,8 +120,7 @@ UObject* UHairStrandsFactory::FactoryCreateFile(UClass* InClass, UObject* InPare
 			}
 		}
 
-		// GroomCache options are only shown if there's a valid groom animation
-		GroomCacheImportOptions->ImportSettings.bImportGroomCache = GroomCacheImportOptions->ImportSettings.bImportGroomCache && AnimInfo.IsValid();
+		FGroomCacheImporter::SetupImportSettings(GroomCacheImportOptions->ImportSettings, AnimInfo);
 
 		if (!GIsRunningUnattendedScript && !IsAutomatedImport())
 		{
@@ -145,6 +144,8 @@ UObject* UHairStrandsFactory::FactoryCreateFile(UClass* InClass, UObject* InPare
 		ImportOptions->SaveConfig();
 	}
 
+	FGroomCacheImporter::ApplyImportSettings(GroomCacheImportOptions->ImportSettings, AnimInfo);
+
 	FScopedSlowTask Progress((float) 1, LOCTEXT("ImportHairAsset", "Importing hair asset..."), true);
 	Progress.MakeDialog(true);
 
@@ -191,9 +192,6 @@ UObject* UHairStrandsFactory::FactoryCreateFile(UClass* InClass, UObject* InPare
 
 	if (GroomCacheImportOptions->ImportSettings.bImportGroomCache && GroomAssetForCache)
 	{
-		// Compute the duration as it is not known yet
-		AnimInfo.Duration = AnimInfo.NumFrames * AnimInfo.SecondsPerFrame;
-
 		TArray<UGroomCache*> GroomCaches = FGroomCacheImporter::ImportGroomCache(Filename, SelectedTranslator, AnimInfo, HairImportContext, GroomAssetForCache);
 
 		// Setup asset import data
diff --git a/Engine/Plugins/Runtime/HairStrands/Source/HairStrandsEditor/Private/ReimportHairStrandsFactory.cpp b/Engine/Plugins/Runtime/HairStrands/Source/HairStrandsEditor/Private/ReimportHairStrandsFactory.cpp
--- a/Engine/Plugins/Runtime/HairStrands/Source/HairStrandsEditor/Private/ReimportHairStrandsFactory.cpp
+++ b/Engine/Plugins/Runtime/HairStrands/Source/HairStrandsEditor/Private/ReimportHairStrandsFactory.cpp
@@ -218,9 +218,6 @@ namespace ReimportGroomCacheHelpers
 	{
 		if (Settings.bImportGroomCache && GroomAsset)
 		{
-			// Compute the duration as it is not known yet
-			AnimInfo.Duration = AnimInfo.NumFrames * AnimInfo.SecondsPerFrame;
-
 			// If the reimport was from a GroomCache, set it as the parent to preserve the package name
 			// Otherwise, use the previously set GroomAsset package as the base for the package name
 			if (SourceGroomCache)
@@ -364,8 +361,7 @@ EReimportResult::Type UReimportHairStrandsFactory::Reimport(UObject* Obj)
 		}
 	}
 
-	// GroomCache options are only shown if there's a valid groom animation
-	GroomCacheReimportOptions->ImportSettings.bImportGroomCache = GroomCacheReimportOptions->ImportSettings.bImportGroomCache && AnimInfo.IsValid();
+	FGroomCacheImporter::SetupImportSettings(GroomCacheReimportOptions->ImportSettings, AnimInfo);
 
 	if (!GIsRunningUnattendedScript && !IsAutomatedImport())
 	{
@@ -403,6 +399,8 @@ EReimportResult::Type UReimportHairStrandsFactory::Reimport(UObject* Obj)
 		}
 	}
 
+	FGroomCacheImporter::ApplyImportSettings(GroomCacheReimportOptions->ImportSettings, AnimInfo);
+
 	FHairDescription HairDescription;
 	if (!SelectedTranslator->Translate(CurrentFilename, HairDescription, GroomReimportOptions->ConversionSettings))
 	{
diff --git a/Engine/Plugins/Runtime/HairStrands/Source/HairStrandsEditor/Public/GroomCacheImporter.h b/Engine/Plugins/Runtime/HairStrands/Source/HairStrandsEditor/Public/GroomCacheImporter.h
--- a/Engine/Plugins/Runtime/HairStrands/Source/HairStrandsEditor/Public/GroomCacheImporter.h
+++ b/Engine/Plugins/Runtime/HairStrands/Source/HairStrandsEditor/Public/GroomCacheImporter.h
@@ -13,4 +13,6 @@ class UGroomCache;
 struct HAIRSTRANDSEDITOR_API FGroomCacheImporter
 {
 	static TArray<UGroomCache*> ImportGroomCache(const FString& SourceFilename, TSharedPtr<IGroomTranslator> Translator, const FGroomAnimationInfo& InAnimInfo, FHairImportContext& HairImportContext, UGroomAsset* GroomAssetForCache);
+	static void SetupImportSettings(struct FGroomCacheImportSettings& ImportSettings, const FGroomAnimationInfo& AnimInfo);
+	static void ApplyImportSettings(struct FGroomCacheImportSettings& ImportSettings, FGroomAnimationInfo& AnimInfo);
 };
