Change hair groom vis. guide tooltips to make it clear when it is enabled or not.

diff --git a/Engine/Plugins/Runtime/HairStrands/Source/HairStrandsEditor/Private/GroomEditorCommands.cpp b/Engine/Plugins/Runtime/HairStrands/Source/HairStrandsEditor/Private/GroomEditorCommands.cpp
--- a/Engine/Plugins/Runtime/HairStrands/Source/HairStrandsEditor/Private/GroomEditorCommands.cpp
+++ b/Engine/Plugins/Runtime/HairStrands/Source/HairStrandsEditor/Private/GroomEditorCommands.cpp
@@ -27,8 +27,8 @@ void FGroomEditorCommands::RegisterCommands()
 
 
 	UI_COMMAND(ViewMode_Lit,			"Lit",				"Default",						EUserInterfaceActionType::Button, FInputChord());
-	UI_COMMAND(ViewMode_Guide,			"Guide",			"View guide",					EUserInterfaceActionType::Button, FInputChord());
-	UI_COMMAND(ViewMode_GuideInfluence, "GuideInfluence",	"View guide influence",			EUserInterfaceActionType::Button, FInputChord());
+	UI_COMMAND(ViewMode_Guide,			"Guide",			"View guide (only available if the groom has simulation enabled, or RBF Interpolation enabled)",			EUserInterfaceActionType::Button, FInputChord());
+	UI_COMMAND(ViewMode_GuideInfluence, "GuideInfluence",	"View guide influence (only available if the groom has simulation enabled, or RBF Interpolation enabled)",	EUserInterfaceActionType::Button, FInputChord());
 	UI_COMMAND(ViewMode_UV,				"UV",				"View strands UVs ",			EUserInterfaceActionType::Button, FInputChord());
 	UI_COMMAND(ViewMode_RootUV,			"RootUV",			"View roots UVs ",				EUserInterfaceActionType::Button, FInputChord());
 	UI_COMMAND(ViewMode_RootUDIM,		"RootUDIM",			"View root UDIM UVs",			EUserInterfaceActionType::Button, FInputChord());
@@ -40,7 +40,7 @@ void FGroomEditorCommands::RegisterCommands()
 	UI_COMMAND(ViewMode_VisCluster,		"VisCluster",		"View clusters",				EUserInterfaceActionType::Button, FInputChord());
 	UI_COMMAND(ViewMode_Group,			"Hair Groups",		"View hair groups",				EUserInterfaceActionType::Button, FInputChord());
 
-	UI_COMMAND(ViewMode_CardsGuides,	"CardsGuide",		"Cards Guides",					EUserInterfaceActionType::Button, FInputChord());
+	UI_COMMAND(ViewMode_CardsGuides,	"CardsGuide",		"Cards Guides (only available if the groom has cards geometry)",											EUserInterfaceActionType::Button, FInputChord());
 }
 
 #undef LOCTEXT_NAMESPACE
