Async loading option for hair resources.

diff --git a/Engine/Plugins/Runtime/HairStrands/Source/HairStrandsCore/Private/GroomComponent.cpp b/Engine/Plugins/Runtime/HairStrands/Source/HairStrandsCore/Private/GroomComponent.cpp
--- a/Engine/Plugins/Runtime/HairStrands/Source/HairStrandsCore/Private/GroomComponent.cpp
+++ b/Engine/Plugins/Runtime/HairStrands/Source/HairStrandsCore/Private/GroomComponent.cpp
@@ -561,10 +561,12 @@ public:
 						check(Segment.VertexBuffer.IsValid());
 					}
 				
-					FMeshBatch* MeshBatch = CreateMeshBatch(Context.ReferenceView, Context.ReferenceViewFamily, Context.RayTracingMeshResourceCollector, Instance, GroupIt, nullptr);
-					TArray<FMeshBatch> MeshBatches;
-					MeshBatches.Add(*MeshBatch);
-					AddOpaqueRaytracingInstance(OverrideLocalToWorld, &RTGeometry->RayTracingGeometry, RaytracingInstanceMask_Opaque, MeshBatches, GetScene().GetFeatureLevel(), OutRayTracingInstances);
+					if (FMeshBatch* MeshBatch = CreateMeshBatch(Context.ReferenceView, Context.ReferenceViewFamily, Context.RayTracingMeshResourceCollector, Instance, GroupIt, nullptr))
+					{
+						TArray<FMeshBatch> MeshBatches;
+						MeshBatches.Add(*MeshBatch);
+						AddOpaqueRaytracingInstance(OverrideLocalToWorld, &RTGeometry->RayTracingGeometry, RaytracingInstanceMask_Opaque, MeshBatches, GetScene().GetFeatureLevel(), OutRayTracingInstances);
+					}
 
 				}
 			}
@@ -578,10 +580,12 @@ public:
 						check(Segment.VertexBuffer.IsValid());
 					}
 
-					FMeshBatch* MeshBatch = CreateMeshBatch(Context.ReferenceView, Context.ReferenceViewFamily, Context.RayTracingMeshResourceCollector, Instance, GroupIt, nullptr);
-					TArray<FMeshBatch> MeshBatches;
-					MeshBatches.Add(*MeshBatch);
-					AddOpaqueRaytracingInstance(OverrideLocalToWorld, &RTGeometry->RayTracingGeometry, RaytracingInstanceMask_Opaque, MeshBatches, GetScene().GetFeatureLevel(), OutRayTracingInstances);
+					if (FMeshBatch* MeshBatch = CreateMeshBatch(Context.ReferenceView, Context.ReferenceViewFamily, Context.RayTracingMeshResourceCollector, Instance, GroupIt, nullptr))
+					{
+						TArray<FMeshBatch> MeshBatches;
+						MeshBatches.Add(*MeshBatch);
+						AddOpaqueRaytracingInstance(OverrideLocalToWorld, &RTGeometry->RayTracingGeometry, RaytracingInstanceMask_Opaque, MeshBatches, GetScene().GetFeatureLevel(), OutRayTracingInstances);
+					}
 				}
 			}
 		}
@@ -681,12 +685,14 @@ public:
 						Debug_MaterialProxy = DebugMaterial;
 					}
 
-					FMeshBatch* MeshBatch = CreateMeshBatch(View, ViewFamily, Collector, Instances[GroupIt], GroupIt, Debug_MaterialProxy);
-					if (MeshBatch == nullptr)
+					if (FMeshBatch* MeshBatch = CreateMeshBatch(View, ViewFamily, Collector, Instances[GroupIt], GroupIt, Debug_MaterialProxy))
+					{
+						Collector.AddMesh(ViewIndex, *MeshBatch);
+					}
+					else
 					{
 						continue;
 					}
-					Collector.AddMesh(ViewIndex, *MeshBatch);
 
 				#if !(UE_BUILD_SHIPPING || UE_BUILD_TEST)
 					// Render bounds
diff --git a/Engine/Plugins/Runtime/HairStrands/Source/HairStrandsCore/Private/GroomManager.cpp b/Engine/Plugins/Runtime/HairStrands/Source/HairStrandsCore/Private/GroomManager.cpp
--- a/Engine/Plugins/Runtime/HairStrands/Source/HairStrandsCore/Private/GroomManager.cpp
+++ b/Engine/Plugins/Runtime/HairStrands/Source/HairStrandsCore/Private/GroomManager.cpp
@@ -715,23 +715,38 @@ static void RunHairLODSelection(
 			GeometryType = EHairGeometryType::NoneGeometry;
 		}
 
+		const EHairResourceLoadingType LoadingType = GetHairResourceLoadingType();
+		EHairResourceStatus ResourceStatus = EHairResourceStatus::None;
+
 		// Lazy allocation of resources
 		// Note: Allocation will only be done if the resources is not initialized yet. Guides deformed position are also initialized from the Rest position at creation time.
 		if (Instance->Guides.Data)
 		{
-			if (Instance->Guides.RestRootResource)			{ Instance->Guides.RestRootResource->Allocate(GraphBuilder); Instance->Guides.RestRootResource->AllocateLOD(GraphBuilder, MeshLODIndex); }
-			if (Instance->Guides.RestResource)				{ Instance->Guides.RestResource->Allocate(GraphBuilder); }
-			if (Instance->Guides.DeformedRootResource)		{ Instance->Guides.DeformedRootResource->Allocate(GraphBuilder); Instance->Guides.DeformedRootResource->AllocateLOD(GraphBuilder, MeshLODIndex); }
-			if (Instance->Guides.DeformedResource)			{ const bool bNeedCopy = !Instance->Guides.DeformedResource->bIsInitialized; Instance->Guides.DeformedResource->Allocate(GraphBuilder); if (bNeedCopy) { AddCopyHairStrandsPositionPass(GraphBuilder, ShaderMap, *Instance->Guides.RestResource, *Instance->Guides.DeformedResource); }}
+			if (Instance->Guides.RestRootResource)			{ Instance->Guides.RestRootResource->Allocate(GraphBuilder, LoadingType, ResourceStatus); Instance->Guides.RestRootResource->AllocateLOD(GraphBuilder, MeshLODIndex, LoadingType, ResourceStatus); }
+			if (Instance->Guides.RestResource)				{ Instance->Guides.RestResource->Allocate(GraphBuilder, LoadingType, ResourceStatus); }
+			if (Instance->Guides.DeformedRootResource)		{ Instance->Guides.DeformedRootResource->Allocate(GraphBuilder, LoadingType, ResourceStatus); Instance->Guides.DeformedRootResource->AllocateLOD(GraphBuilder, MeshLODIndex, LoadingType, ResourceStatus); }
+			if (Instance->Guides.DeformedResource)			
+			{ 
+				// Ensure the rest resources are correctly loaded prior to initialized and copy the rest position into the deformed positions
+				if (Instance->Guides.RestResource->bIsInitialized)
+				{
+					const bool bNeedCopy = !Instance->Guides.DeformedResource->bIsInitialized; 
+					Instance->Guides.DeformedResource->Allocate(GraphBuilder, EHairResourceLoadingType::Sync); 
+					if (bNeedCopy) 
+					{ 
+						AddCopyHairStrandsPositionPass(GraphBuilder, ShaderMap, *Instance->Guides.RestResource, *Instance->Guides.DeformedResource); 
+					}
+				}
+			}
 		}
 
 		if (GeometryType == EHairGeometryType::Meshes)
 		{
 			FHairGroupInstance::FMeshes::FLOD& InstanceLOD = Instance->Meshes.LODs[IntLODIndex];
 
-			if (InstanceLOD.DeformedResource)				{ InstanceLOD.DeformedResource->Allocate(GraphBuilder); }
+			if (InstanceLOD.DeformedResource)				{ InstanceLOD.DeformedResource->Allocate(GraphBuilder, LoadingType, ResourceStatus); }
 			#if RHI_RAYTRACING
-			if (InstanceLOD.RaytracingResource)				{ InstanceLOD.RaytracingResource->Allocate(GraphBuilder);}
+			if (InstanceLOD.RaytracingResource)				{ InstanceLOD.RaytracingResource->Allocate(GraphBuilder, LoadingType, ResourceStatus);}
 			#endif
 
 			InstanceLOD.InitVertexFactory();
@@ -740,14 +755,14 @@ static void RunHairLODSelection(
 		{
 			FHairGroupInstance::FCards::FLOD& InstanceLOD = Instance->Cards.LODs[IntLODIndex];
 
-			if (InstanceLOD.DeformedResource)				{ InstanceLOD.DeformedResource->Allocate(GraphBuilder); }
-			if (InstanceLOD.Guides.RestRootResource)		{ InstanceLOD.Guides.RestRootResource->Allocate(GraphBuilder); InstanceLOD.Guides.RestRootResource->AllocateLOD(GraphBuilder, MeshLODIndex); }
-			if (InstanceLOD.Guides.RestResource)			{ InstanceLOD.Guides.RestResource->Allocate(GraphBuilder); }
-			if (InstanceLOD.Guides.DeformedRootResource)	{ InstanceLOD.Guides.DeformedRootResource->Allocate(GraphBuilder); InstanceLOD.Guides.DeformedRootResource->AllocateLOD(GraphBuilder, MeshLODIndex); }
-			if (InstanceLOD.Guides.DeformedResource)		{ InstanceLOD.Guides.DeformedResource->Allocate(GraphBuilder); }
-			if (InstanceLOD.Guides.InterpolationResource)	{ InstanceLOD.Guides.InterpolationResource->Allocate(GraphBuilder); }
+			if (InstanceLOD.DeformedResource)				{ InstanceLOD.DeformedResource->Allocate(GraphBuilder, LoadingType, ResourceStatus); }
+			if (InstanceLOD.Guides.RestRootResource)		{ InstanceLOD.Guides.RestRootResource->Allocate(GraphBuilder, LoadingType, ResourceStatus); InstanceLOD.Guides.RestRootResource->AllocateLOD(GraphBuilder, MeshLODIndex, LoadingType, ResourceStatus); }
+			if (InstanceLOD.Guides.RestResource)			{ InstanceLOD.Guides.RestResource->Allocate(GraphBuilder, LoadingType, ResourceStatus); }
+			if (InstanceLOD.Guides.DeformedRootResource)	{ InstanceLOD.Guides.DeformedRootResource->Allocate(GraphBuilder, LoadingType, ResourceStatus); InstanceLOD.Guides.DeformedRootResource->AllocateLOD(GraphBuilder, MeshLODIndex, LoadingType, ResourceStatus); }
+			if (InstanceLOD.Guides.DeformedResource)		{ InstanceLOD.Guides.DeformedResource->Allocate(GraphBuilder, LoadingType, ResourceStatus); }
+			if (InstanceLOD.Guides.InterpolationResource)	{ InstanceLOD.Guides.InterpolationResource->Allocate(GraphBuilder, LoadingType, ResourceStatus); }
 			#if RHI_RAYTRACING
-			if (InstanceLOD.RaytracingResource)				{ InstanceLOD.RaytracingResource->Allocate(GraphBuilder);}
+			if (InstanceLOD.RaytracingResource)				{ InstanceLOD.RaytracingResource->Allocate(GraphBuilder, LoadingType, ResourceStatus);}
 			#endif
 			InstanceLOD.InitVertexFactory();
 		}
@@ -756,15 +771,15 @@ static void RunHairLODSelection(
 			check(Instance->HairGroupPublicData);
 			Instance->HairGroupPublicData->Allocate(GraphBuilder);
 
-			if (Instance->Strands.RestRootResource)			{ Instance->Strands.RestRootResource->Allocate(GraphBuilder); Instance->Strands.RestRootResource->AllocateLOD(GraphBuilder, MeshLODIndex); }
-			if (Instance->Strands.RestResource)				{ Instance->Strands.RestResource->Allocate(GraphBuilder); }
-			if (Instance->Strands.ClusterCullingResource)	{ Instance->Strands.ClusterCullingResource->Allocate(GraphBuilder); }
-			if (Instance->Strands.InterpolationResource)	{ Instance->Strands.InterpolationResource->Allocate(GraphBuilder); }
+			if (Instance->Strands.RestRootResource)			{ Instance->Strands.RestRootResource->Allocate(GraphBuilder, LoadingType, ResourceStatus); Instance->Strands.RestRootResource->AllocateLOD(GraphBuilder, MeshLODIndex, LoadingType, ResourceStatus); }
+			if (Instance->Strands.RestResource)				{ Instance->Strands.RestResource->Allocate(GraphBuilder, LoadingType, ResourceStatus); }
+			if (Instance->Strands.ClusterCullingResource)	{ Instance->Strands.ClusterCullingResource->Allocate(GraphBuilder, LoadingType, ResourceStatus); }
+			if (Instance->Strands.InterpolationResource)	{ Instance->Strands.InterpolationResource->Allocate(GraphBuilder, LoadingType, ResourceStatus); }
 
-			if (Instance->Strands.DeformedRootResource)		{ Instance->Strands.DeformedRootResource->Allocate(GraphBuilder); Instance->Strands.DeformedRootResource->AllocateLOD(GraphBuilder, MeshLODIndex); }
-			if (Instance->Strands.DeformedResource)			{ Instance->Strands.DeformedResource->Allocate(GraphBuilder); }
+			if (Instance->Strands.DeformedRootResource)		{ Instance->Strands.DeformedRootResource->Allocate(GraphBuilder, LoadingType, ResourceStatus); Instance->Strands.DeformedRootResource->AllocateLOD(GraphBuilder, MeshLODIndex, LoadingType, ResourceStatus); }
+			if (Instance->Strands.DeformedResource)			{ Instance->Strands.DeformedResource->Allocate(GraphBuilder, LoadingType, ResourceStatus); }
 			#if RHI_RAYTRACING
-			if (Instance->Strands.RenRaytracingResource)	{ Instance->Strands.RenRaytracingResource->Allocate(GraphBuilder); }
+			if (Instance->Strands.RenRaytracingResource)	{ Instance->Strands.RenRaytracingResource->Allocate(GraphBuilder, LoadingType, ResourceStatus); }
 			#endif
 			Instance->Strands.VertexFactory->InitResources();
 		}
@@ -782,6 +797,14 @@ static void RunHairLODSelection(
 
 		// Update the local-to-world transform based on the binding type 
 		Instance->LocalToWorld = Instance->BindingType == EHairBindingType::Skinning ? Instance->Debug.SkeletalLocalToWorld : Instance->LocalToWorld;
+
+		// If all resources are not ready yet, change the geomtry to be invalid, so that it don't get processed this frame.
+		if (!(ResourceStatus & EHairResourceStatus::Loading))
+		{
+			Instance->GeometryType = EHairGeometryType::NoneGeometry;
+			Instance->BindingType = EHairBindingType::NoneBinding;
+		}
+
 	}
 }
 
diff --git a/Engine/Plugins/Runtime/HairStrands/Source/HairStrandsCore/Private/GroomResources.cpp b/Engine/Plugins/Runtime/HairStrands/Source/HairStrandsCore/Private/GroomResources.cpp
--- a/Engine/Plugins/Runtime/HairStrands/Source/HairStrandsCore/Private/GroomResources.cpp
+++ b/Engine/Plugins/Runtime/HairStrands/Source/HairStrandsCore/Private/GroomResources.cpp
@@ -26,6 +26,9 @@
 static int32 GHairStrandsBulkData_ReleaseAfterUse = 0;
 static FAutoConsoleVariableRef CVarHairStrandsBulkData_ReleaseAfterUse(TEXT("r.HairStrands.Strands.BulkData.ReleaseAfterUse"), GHairStrandsBulkData_ReleaseAfterUse, TEXT("Release CPU bulk data once hair groom/groom binding asset GPU resources are created. This saves memory"));
 
+static int32 GHairStrandsBulkData_AsyncLoading = 0;
+static FAutoConsoleVariableRef CVarHairStrandsBulkData_AsyncLoading(TEXT("r.HairStrands.Strands.BulkData.AsyncLoading"), GHairStrandsBulkData_AsyncLoading, TEXT("Load hair data with async loading so that it is not blocking the rendering thread."));
+
 static int32 GHairStrandsBulkData_Validation = 1;
 static FAutoConsoleVariableRef CVarHairStrandsBulkData_Validation(TEXT("r.HairStrands.Strands.BulkData.Validation"), GHairStrandsBulkData_Validation, TEXT("Validate some hair strands data at serialization/loading time."));
 
@@ -36,6 +39,11 @@ static FAutoConsoleVariableRef CVarHairStrandsDebugVoxel_MaxSegmentPerVoxel(TEXT
 
 ////////////////////////////////////////////////////////////////////////////////////
 
+EHairResourceLoadingType GetHairResourceLoadingType()
+{
+	return GHairStrandsBulkData_AsyncLoading > 0 ? EHairResourceLoadingType::Async : EHairResourceLoadingType::Sync;
+}
+
 namespace HairTransition
 {
 	void TransitToSRV(FRDGBuilder& GraphBuilder, FRDGBufferSRVRef InBuffer, ERDGPassFlags Flags);	
@@ -459,6 +467,8 @@ void FHairCommonResource::InitRHI()
 {
 	if (bIsInitialized || AllocationType == EHairStrandsAllocationType::Deferred || GUsingNullRHI) { return; }
 
+	check(InternalIsDataLoaded());
+
 	if (bUseRenderGraph)
 	{
 		FMemMark Mark(FMemStack::Get());
@@ -480,23 +490,40 @@ void FHairCommonResource::ReleaseRHI()
 	bIsInitialized = false;
 }
 
-void FHairCommonResource::Allocate(FRDGBuilder& GraphBuilder)
+void FHairCommonResource::Allocate(FRDGBuilder& GraphBuilder, EHairResourceLoadingType LoadingType)
 {
-	if (bIsInitialized) { return; }
+	EHairResourceStatus Status = EHairResourceStatus::None;
+	Allocate(GraphBuilder, LoadingType, Status);
+}
+
+void FHairCommonResource::Allocate(FRDGBuilder& GraphBuilder, EHairResourceLoadingType LoadingType, EHairResourceStatus& Status)
+{
+	if (bIsInitialized) { Status |= EHairResourceStatus::Valid; return; }
+
+	if (LoadingType == EHairResourceLoadingType::Async && !InternalIsDataLoaded()) { Status |= EHairResourceStatus::Loading; return; }
 
 	check(AllocationType == EHairStrandsAllocationType::Deferred);
 	FRenderResource::InitResource(); // Call RenderResource InitResource() so that the resources is marked as initialized
 	InternalAllocate(GraphBuilder);
 	bIsInitialized = true;
+	Status |= EHairResourceStatus::Valid;
 }
 
-void FHairCommonResource::AllocateLOD(FRDGBuilder& GraphBuilder, int32 LODIndex)
+void FHairCommonResource::AllocateLOD(FRDGBuilder& GraphBuilder, int32 LODIndex, EHairResourceLoadingType LoadingType, EHairResourceStatus& Status)
 {
+	if (LoadingType == EHairResourceLoadingType::Async && !InternalIsLODDataLoaded(LODIndex)) { Status |= EHairResourceStatus::Loading; return; }
+
 	// Sanity check. For allocation sub-resources (like LOD resources) the common/main resource needs to be already initialized.
 	check(bIsInitialized);
 	check(AllocationType == EHairStrandsAllocationType::Deferred);
 
 	InternalAllocateLOD(GraphBuilder, LODIndex);
+	Status |= EHairResourceStatus::Valid;
+}
+
+void AsyncLoadHairBulkData(FByteBulkData& InData, bool& bIsLoading)
+{
+	bIsLoading = InData.StartAsyncLoading() || bIsLoading;
 }
 
 /////////////////////////////////////////////////////////////////////////////////////////
@@ -721,6 +748,26 @@ FHairStrandsRestResource::FHairStrandsRestResource(FHairStrandsBulkData& InBulkD
 	PositionBuffer(), Attribute0Buffer(), Attribute1Buffer(), MaterialBuffer(), BulkData(InBulkData), CurveType(InCurveType)
 {}
 
+bool FHairStrandsRestResource::InternalIsDataLoaded()
+{
+	bool bIsLoading = false;
+
+	AsyncLoadHairBulkData(BulkData.Positions, bIsLoading);
+	AsyncLoadHairBulkData(BulkData.Attributes0, bIsLoading);
+
+	if (!!(BulkData.Flags & FHairStrandsBulkData::DataFlags_HasUDIMData))
+	{
+		AsyncLoadHairBulkData(BulkData.Attributes1, bIsLoading);
+	}
+
+	if (!!(BulkData.Flags & FHairStrandsBulkData::DataFlags_HasMaterialData))
+	{
+		AsyncLoadHairBulkData(BulkData.Materials, bIsLoading);
+	}
+
+	return !bIsLoading;
+}
+
 void FHairStrandsRestResource::InternalAllocate(FRDGBuilder& GraphBuilder)
 {
 	const uint32 PointCount = BulkData.PointCount;
@@ -934,6 +981,16 @@ FHairStrandsClusterCullingResource::FHairStrandsClusterCullingResource(FHairStra
 
 }
 
+bool FHairStrandsClusterCullingResource::InternalIsDataLoaded()
+{
+	bool bIsLoading = false;
+	AsyncLoadHairBulkData(BulkData.PackedClusterInfos, bIsLoading);
+	AsyncLoadHairBulkData(BulkData.ClusterLODInfos, bIsLoading);
+	AsyncLoadHairBulkData(BulkData.VertexToClusterIds, bIsLoading);
+	AsyncLoadHairBulkData(BulkData.ClusterVertexIds, bIsLoading);
+	return !bIsLoading;
+}
+
 void FHairStrandsClusterCullingResource::InternalAllocate(FRDGBuilder& GraphBuilder)
 {
 	if (GHairStrandsBulkData_Validation > 0)
@@ -979,6 +1036,47 @@ void FHairStrandsRestRootResource::PopulateFromRootData()
 	}
 }
 
+bool FHairStrandsRestRootResource::InternalIsDataLoaded()
+{
+	bool bIsLoading = false;
+	if (BulkData.PointCount > 0)
+	{
+		AsyncLoadHairBulkData(BulkData.VertexToCurveIndexBuffer, bIsLoading);
+	}
+	return !bIsLoading;
+}
+
+bool FHairStrandsRestRootResource::InternalIsLODDataLoaded(int32 LODIndex)
+{
+	bool bIsLoading = false;
+	
+	check(LODs.Num() == BulkData.MeshProjectionLODs.Num());
+	if (LODIndex >= 0 && LODIndex < LODs.Num())
+	{
+		FHairStrandsRootBulkData::FMeshProjectionLOD& CPUData = BulkData.MeshProjectionLODs[LODIndex];
+		const bool bHasValidCPUData = CPUData.RootTriangleBarycentricBuffer.GetBulkDataSize() > 0;
+		if (bHasValidCPUData)
+		{
+			AsyncLoadHairBulkData(CPUData.RootTriangleBarycentricBuffer, bIsLoading);
+			AsyncLoadHairBulkData(CPUData.RootTriangleIndexBuffer, bIsLoading);
+
+			AsyncLoadHairBulkData(CPUData.RestRootTrianglePosition0Buffer, bIsLoading);
+			AsyncLoadHairBulkData(CPUData.RestRootTrianglePosition1Buffer, bIsLoading);
+			AsyncLoadHairBulkData(CPUData.RestRootTrianglePosition2Buffer, bIsLoading);
+		}
+
+		const bool bHasValidCPUWeights = CPUData.MeshSampleIndicesBuffer.GetBulkDataSize() > 0;
+		if (bHasValidCPUWeights)
+		{
+			AsyncLoadHairBulkData(CPUData.MeshInterpolationWeightsBuffer, bIsLoading);
+			AsyncLoadHairBulkData(CPUData.MeshSampleIndicesBuffer, bIsLoading);
+			AsyncLoadHairBulkData(CPUData.RestSamplePositionsBuffer, bIsLoading);
+		}		
+	}
+
+	return !bIsLoading;
+}
+
 void FHairStrandsRestRootResource::InternalAllocate(FRDGBuilder& GraphBuilder)
 {
 	// Once empty, the MeshProjectionLODsneeds to be repopulate as it might be re-initialized. 
@@ -1292,6 +1390,24 @@ FHairStrandsInterpolationResource::FHairStrandsInterpolationResource(FHairStrand
 {
 }
 
+bool FHairStrandsInterpolationResource::InternalIsDataLoaded()
+{
+	bool bIsLoading = false;
+	const bool bUseSingleGuide = !!(BulkData.Flags & FHairStrandsInterpolationBulkData::DataFlags_HasSingleGuideData);
+	if (bUseSingleGuide)
+	{
+		AsyncLoadHairBulkData(BulkData.Interpolation, bIsLoading);
+	}
+	else
+	{
+		AsyncLoadHairBulkData(BulkData.Interpolation0, bIsLoading);
+		AsyncLoadHairBulkData(BulkData.Interpolation1, bIsLoading);
+	}
+	AsyncLoadHairBulkData(BulkData.SimRootPointIndex, bIsLoading);
+
+	return !bIsLoading;
+}
+
 void FHairStrandsInterpolationResource::InternalAllocate(FRDGBuilder& GraphBuilder)
 {
 	const bool bUseSingleGuide = !!(BulkData.Flags & FHairStrandsInterpolationBulkData::DataFlags_HasSingleGuideData);
diff --git a/Engine/Plugins/Runtime/HairStrands/Source/HairStrandsCore/Private/GroomTextureBuilder.cpp b/Engine/Plugins/Runtime/HairStrands/Source/HairStrandsCore/Private/GroomTextureBuilder.cpp
--- a/Engine/Plugins/Runtime/HairStrands/Source/HairStrandsCore/Private/GroomTextureBuilder.cpp
+++ b/Engine/Plugins/Runtime/HairStrands/Source/HairStrandsCore/Private/GroomTextureBuilder.cpp
@@ -957,7 +957,7 @@ static void InternalBuildStrandsTextures_GPU(
 			}
 
 			// Ensure the rest resources are loaded when rendering the strands textures
-			GroupData.Strands.RestResource->Allocate(GraphBuilder);
+			GroupData.Strands.RestResource->Allocate(GraphBuilder, EHairResourceLoadingType::Sync);
 
 			InternalGenerateHairStrandsTextures(
 				GraphBuilder,
diff --git a/Engine/Plugins/Runtime/HairStrands/Source/HairStrandsCore/Public/GroomResources.h b/Engine/Plugins/Runtime/HairStrands/Source/HairStrandsCore/Public/GroomResources.h
--- a/Engine/Plugins/Runtime/HairStrands/Source/HairStrandsCore/Public/GroomResources.h
+++ b/Engine/Plugins/Runtime/HairStrands/Source/HairStrandsCore/Public/GroomResources.h
@@ -50,6 +50,26 @@ struct FHairResourceName
 #endif
 };
 
+enum class EHairResourceLoadingType : uint8
+{
+	Async,
+	Sync
+};
+
+enum class EHairResourceStatus : uint8
+{
+	None = 0,
+	Loading = 1,
+	Valid = 2
+};
+
+FORCEINLINE EHairResourceStatus  operator& (EHairResourceStatus A, EHairResourceStatus B) { return static_cast<EHairResourceStatus>(static_cast<uint8>(A) & static_cast<uint8>(B)); }
+FORCEINLINE EHairResourceStatus  operator| (EHairResourceStatus A, EHairResourceStatus B) { return static_cast<EHairResourceStatus>(static_cast<uint8>(A) | static_cast<uint8>(B)); }
+FORCEINLINE EHairResourceStatus& operator|=(EHairResourceStatus&A, EHairResourceStatus B) { return A=A|B; }
+FORCEINLINE bool operator! (EHairResourceStatus A) { return static_cast<uint8>(A) != 0; }
+
+EHairResourceLoadingType GetHairResourceLoadingType();
+
 /* Hair resouces which whom allocation can be deferred */
 struct FHairCommonResource : public FRenderResource
 {
@@ -61,12 +81,15 @@ struct FHairCommonResource : public FRenderResource
 	virtual void ReleaseRHI() override;
 
 	/* Init/Release buffers (FHairCommonResource) */
-	void Allocate(FRDGBuilder& GraphBuilder);
-	void AllocateLOD(FRDGBuilder& GraphBuilder, int32 LODIndex);
-	virtual void InternalAllocate() {};
-	virtual void InternalAllocate(FRDGBuilder& GraphBuilder) {};
-	virtual void InternalAllocateLOD(FRDGBuilder& GraphBuilder, int32 LODIndex) {};
-	virtual void InternalRelease() {};
+	void Allocate(FRDGBuilder& GraphBuilder, EHairResourceLoadingType LoadingType);
+	void Allocate(FRDGBuilder& GraphBuilder, EHairResourceLoadingType LoadingType, EHairResourceStatus& Status);
+	void AllocateLOD(FRDGBuilder& GraphBuilder, int32 LODIndex, EHairResourceLoadingType LoadingType, EHairResourceStatus& Status);
+	virtual void InternalAllocate() {}
+	virtual void InternalAllocate(FRDGBuilder& GraphBuilder) {}
+	virtual void InternalAllocateLOD(FRDGBuilder& GraphBuilder, int32 LODIndex) {}
+	virtual void InternalRelease() {}
+	virtual bool InternalIsDataLoaded() { return true; }
+	virtual bool InternalIsLODDataLoaded(int32 LODIndex) { return true; }
 
 	bool bUseRenderGraph = true;
 	bool bIsInitialized = false;
@@ -86,6 +109,8 @@ struct FHairStrandsRestRootResource : public FHairCommonResource
 	virtual void InternalAllocate(FRDGBuilder& GraphBuilder) override;
 	virtual void InternalAllocateLOD(FRDGBuilder& GraphBuilder, int32 LODIndex) override;
 	virtual void InternalRelease() override;
+	virtual bool InternalIsDataLoaded() override;
+	virtual bool InternalIsLODDataLoaded(int32 LODIndex) override;
 
 	/* Get the resource name */
 	virtual FString GetFriendlyName() const override { return TEXT("FHairStrandsRestRootResource"); }
@@ -225,6 +250,7 @@ struct FHairStrandsRestResource : public FHairCommonResource
 	/* Init/Release buffers */
 	virtual void InternalAllocate(FRDGBuilder& GraphBuilder) override;
 	virtual void InternalRelease() override;
+	virtual bool InternalIsDataLoaded() override;
 
 	/* Get the resource name */
 	virtual FString GetFriendlyName() const override { return TEXT("FHairStrandsResource"); }
@@ -348,6 +374,7 @@ struct FHairStrandsClusterCullingResource : public FHairCommonResource
 	/* Init/Release buffers */
 	virtual void InternalAllocate(FRDGBuilder& GraphBuilder) override;
 	virtual void InternalRelease() override;
+	virtual bool InternalIsDataLoaded() override;
 
 	/* Get the resource name */
 	virtual FString GetFriendlyName() const override { return TEXT("FHairStrandsClusterResource"); }
@@ -384,6 +411,7 @@ struct FHairStrandsInterpolationResource : public FHairCommonResource
 	/* Init/Release buffers */
 	virtual void InternalAllocate(FRDGBuilder& GraphBuilder) override;
 	virtual void InternalRelease() override;
+	virtual bool InternalIsDataLoaded() override;
 
 	/* Get the resource name */
 	virtual FString GetFriendlyName() const override { return TEXT("FHairStrandsInterplationResource"); }
