Fix cards guide visualization.

diff --git a/Engine/Plugins/Runtime/HairStrands/Source/HairStrandsCore/Private/GroomDebug.cpp b/Engine/Plugins/Runtime/HairStrands/Source/HairStrandsCore/Private/GroomDebug.cpp
--- a/Engine/Plugins/Runtime/HairStrands/Source/HairStrandsCore/Private/GroomDebug.cpp
+++ b/Engine/Plugins/Runtime/HairStrands/Source/HairStrandsCore/Private/GroomDebug.cpp
@@ -10,6 +10,7 @@
 #include "GroomGeometryCache.h"
 #include "Components/SkeletalMeshComponent.h"
 #include "UnrealEngine.h"
+#include "SystemTextures.h"
 
 ///////////////////////////////////////////////////////////////////////////////////////////////////
 
@@ -706,23 +707,56 @@ static void AddDrawDebugCardsGuidesPass(
 		return;
 	}
 	TShaderMapRef<FDrawDebugCardGuidesCS> ComputeShader(ShaderMap);
+	const bool bGuideValid			= Instance->Guides.RestResource != nullptr;
+	const bool bGuideDeformValid	= Instance->Guides.DeformedResource != nullptr;
+	const bool bRenderValid			= LOD.Guides.RestResource != nullptr;
+	const bool bRenderDeformValid	= LOD.Guides.DeformedResource != nullptr;
+	if (bRen  && !bRenderValid)						{ return; }
+	if (bRen  && bDeformed && !bRenderDeformValid)	{ return; }
+	if (!bRen && !bGuideValid)						{ return; }
+	if (!bRen && bDeformed && !bGuideDeformValid)	{ return; }
+
+	FRDGBufferSRVRef DefaultBuffer = GraphBuilder.CreateSRV(GSystemTextures.GetDefaultBuffer(GraphBuilder, 8, 0u), PF_R16G16B16A16_UINT);
 
 	FDrawDebugCardGuidesCS::FParameters* Parameters = GraphBuilder.AllocParameters<FDrawDebugCardGuidesCS::FParameters>();
 	Parameters->ViewUniformBuffer = View.ViewUniformBuffer;
-	Parameters->RenVertexCount = LOD.Guides.RestResource->GetVertexCount();
-	Parameters->SimVertexCount = Instance->Guides.RestResource->GetVertexCount();
 
-	Parameters->RenRestOffset = LOD.Guides.RestResource->GetPositionOffset();
-	Parameters->RenRestPosition = RegisterAsSRV(GraphBuilder, LOD.Guides.RestResource->PositionBuffer);
+	Parameters->RenVertexCount = 0;
+	Parameters->RenRestOffset = FVector3f::ZeroVector;
+	Parameters->RenRestPosition = DefaultBuffer;
+	Parameters->RenDeformedOffset = FVector3f::ZeroVector;
+	Parameters->RenDeformedPosition = DefaultBuffer;
+
+	Parameters->SimVertexCount = 0;
+	Parameters->SimRestOffset = FVector3f::ZeroVector;
+	Parameters->SimRestPosition = DefaultBuffer;
+	Parameters->SimDeformedOffset = FVector3f::ZeroVector;
+	Parameters->SimDeformedPosition = DefaultBuffer;
 
-	Parameters->RenDeformedOffset = LOD.Guides.DeformedResource->GetPositionOffset(FHairStrandsDeformedResource::Current);
-	Parameters->RenDeformedPosition = RegisterAsSRV(GraphBuilder, LOD.Guides.DeformedResource->GetBuffer(FHairStrandsDeformedResource::Current));
+	if (bRen)
+	{
+		Parameters->RenVertexCount = LOD.Guides.RestResource->GetVertexCount();
+		Parameters->RenRestOffset = LOD.Guides.RestResource->GetPositionOffset();
+		Parameters->RenRestPosition = RegisterAsSRV(GraphBuilder, LOD.Guides.RestResource->PositionBuffer);
+		if (bDeformed)
+		{
+			Parameters->RenDeformedOffset = LOD.Guides.DeformedResource->GetPositionOffset(FHairStrandsDeformedResource::Current);
+			Parameters->RenDeformedPosition = RegisterAsSRV(GraphBuilder, LOD.Guides.DeformedResource->GetBuffer(FHairStrandsDeformedResource::Current));
+		}
+	}
 
-	Parameters->SimRestOffset = Instance->Guides.RestResource->GetPositionOffset();
-	Parameters->SimRestPosition = RegisterAsSRV(GraphBuilder, Instance->Guides.RestResource->PositionBuffer);
+	if (!bRen)
+	{
+		Parameters->SimVertexCount = Instance->Guides.RestResource->GetVertexCount();
+		Parameters->SimRestOffset = Instance->Guides.RestResource->GetPositionOffset();
+		Parameters->SimRestPosition = RegisterAsSRV(GraphBuilder, Instance->Guides.RestResource->PositionBuffer);
+		if (bDeformed)
+		{
+			Parameters->SimDeformedOffset = Instance->Guides.DeformedResource->GetPositionOffset(FHairStrandsDeformedResource::Current);
+			Parameters->SimDeformedPosition = RegisterAsSRV(GraphBuilder, Instance->Guides.DeformedResource->GetBuffer(FHairStrandsDeformedResource::Current));
+		}
+	}
 
-	Parameters->SimDeformedOffset = Instance->Guides.DeformedResource->GetPositionOffset(FHairStrandsDeformedResource::Current);
-	Parameters->SimDeformedPosition = RegisterAsSRV(GraphBuilder, Instance->Guides.DeformedResource->GetBuffer(FHairStrandsDeformedResource::Current));
 	Parameters->LocalToWorld = Instance->LocalToWorld.ToMatrixWithScale();
 
 	if (!bDeformed &&  bRen) Parameters->DebugMode = 1;
