Fix crash when enabling/disabling simulation.

The crash arise due to the data serialized into the DDC contains invalid bulk data size, as depending on the groom having simulation enabled before, it would serialize stale data. When reading back data from DDC, the size wouldn't match causing validation code to fire.

This CL ensures the data are properly reset when recomputed, so that only valid data are stored into the DDC.

diff --git a/Engine/Plugins/Runtime/HairStrands/Source/HairStrandsCore/Private/GroomBuilder.cpp b/Engine/Plugins/Runtime/HairStrands/Source/HairStrandsCore/Private/GroomBuilder.cpp
--- a/Engine/Plugins/Runtime/HairStrands/Source/HairStrandsCore/Private/GroomBuilder.cpp
+++ b/Engine/Plugins/Runtime/HairStrands/Source/HairStrandsCore/Private/GroomBuilder.cpp
@@ -33,7 +33,7 @@ static FAutoConsoleVariableRef CVarHairClusterBuilder_MaxVoxelResolution(TEXT("r
 FString FGroomBuilder::GetVersion()
 {
 	// Important to update the version when groom building changes
-	return TEXT("2p");
+	return TEXT("2q");
 }
 
 namespace FHairStrandsDecimation
@@ -1205,6 +1205,8 @@ namespace HairInterpolationBuilder
 		const FHairStrandsInterpolationDatas& HairInterpolation, 
 		FHairStrandsInterpolationBulkData& OutBulkData)
 	{
+		OutBulkData.Reset();
+
 		TRACE_CPUPROFILER_EVENT_SCOPE(HairInterpolationBuilder::BuildRenderData);
 
 		const uint32 PointCount = HairInterpolation.Num();
@@ -1759,6 +1761,8 @@ void FGroomBuilder::BuildBulkData(
 	const FHairStrandsDatas& InData,
 	FHairStrandsBulkData& OutBulkData)
 {
+	OutBulkData.Reset();
+
 	const int32 NumCurves = InData.GetNumCurves();
 	FRandomStream Random;
 	Random.Initialize(InInfo.GroupID);
@@ -2553,6 +2557,8 @@ static void BuildClusterBulkData(
 	const FHairStrandsClusterCullingData& In,
 	FHairStrandsClusterCullingBulkData& Out)
 {
+	Out.Reset();
+
 	Out.ClusterCount	= In.ClusterCount;
 	Out.VertexCount		= In.VertexCount;
 	Out.VertexLODCount	= In.ClusterVertexIds.Num();
diff --git a/Engine/Plugins/Runtime/HairStrands/Source/HairStrandsCore/Private/GroomComponent.cpp b/Engine/Plugins/Runtime/HairStrands/Source/HairStrandsCore/Private/GroomComponent.cpp
--- a/Engine/Plugins/Runtime/HairStrands/Source/HairStrandsCore/Private/GroomComponent.cpp
+++ b/Engine/Plugins/Runtime/HairStrands/Source/HairStrandsCore/Private/GroomComponent.cpp
@@ -2194,9 +2194,14 @@ void UGroomComponent::InitResources(bool bIsBindingReloading)
 	//              RBF | Yes		    | Yes		   | Yes
 	//
 	bool bHasNeedSkinningBinding = false;			   
-	bool bHasNeedGlobalDeformation = false;
-	bool bHasNeedSimulation = false;
 	bool bHasNeedSkeletalMesh = false;
+	TArray<bool> bHasNeedGlobalDeformation;
+	TArray<bool> bHasNeedSimulation;
+	bHasNeedGlobalDeformation.Init(false, GroomAsset->HairGroupsData.Num());
+	bHasNeedSimulation.Init(false, GroomAsset->HairGroupsData.Num());
+
+	bool bHasAnyNeedSimulation = false;
+	bool bHasAnyNeedGlobalDeformation = false;
 	for (int32 GroupIt = 0, GroupCount = GroomAsset->HairGroupsData.Num(); GroupIt < GroupCount; ++GroupIt)
 	{
 		for (uint32 LODIt = 0, LODCount = GroomAsset->GetLODCount(); LODIt < LODCount; ++LODIt)
@@ -2207,13 +2212,16 @@ void UGroomComponent::InitResources(bool bIsBindingReloading)
 			// Note on Global Deformation:
 			// * Global deformation require to have skinning binding
 			// * Force global interpolation to be enable for meshes with skinning binding as we use RBF defomation for 'sticking' meshes onto skel. mesh surface			
-			bHasNeedSkeletalMesh	  = bHasNeedSkeletalMesh || BindingType == EGroomBindingType::Rigid;
-			bHasNeedSkinningBinding   = bHasNeedSkinningBinding || BindingType == EGroomBindingType::Skinning;
-			bHasNeedSimulation		  = bHasNeedSimulation || IsSimulationEnable(GroupIt, LODIt);
-			bHasNeedGlobalDeformation = bHasNeedGlobalDeformation || (BindingType == EGroomBindingType::Skinning && GroomAsset->IsGlobalInterpolationEnable(GroupIt, LODIt));
+			bHasNeedSkeletalMesh				= bHasNeedSkeletalMesh || BindingType == EGroomBindingType::Rigid;
+			bHasNeedSkinningBinding				= bHasNeedSkinningBinding || BindingType == EGroomBindingType::Skinning;
+			bHasNeedSimulation[GroupIt]			= bHasNeedSimulation[GroupIt] || IsSimulationEnable(GroupIt, LODIt);
+			bHasNeedGlobalDeformation[GroupIt]	= bHasNeedGlobalDeformation[GroupIt] || (BindingType == EGroomBindingType::Skinning && GroomAsset->IsGlobalInterpolationEnable(GroupIt, LODIt));
+
+			bHasAnyNeedSimulation				= bHasAnyNeedSimulation			|| bHasNeedSimulation[GroupIt];
+			bHasAnyNeedGlobalDeformation		= bHasAnyNeedGlobalDeformation	|| bHasNeedGlobalDeformation[GroupIt];
 		}
 	}
-	const bool bHasNeedBindingData = bHasNeedSkinningBinding || bHasNeedGlobalDeformation;
+	const bool bHasNeedBindingData = bHasNeedSkinningBinding || bHasAnyNeedGlobalDeformation;
 
 	// 2. Insure that the binding asset is compatible, otherwise no binding
 	UMeshComponent* ParentMeshComponent = GetAttachParent() ? Cast<UMeshComponent>(GetAttachParent()) : nullptr;
@@ -2289,7 +2297,8 @@ void UGroomComponent::InitResources(bool bIsBindingReloading)
 	if (LocalBindingAsset == nullptr)
 	{
 		bHasNeedSkinningBinding = false;
-		bHasNeedGlobalDeformation = false;
+		bHasAnyNeedGlobalDeformation = false;
+		bHasNeedGlobalDeformation.Init(false, GroomAsset->HairGroupsData.Num());
 	}
 
 	if (GroomCache)
@@ -2378,7 +2387,7 @@ void UGroomComponent::InitResources(bool bIsBindingReloading)
 		// * Physics simulation
 		// * RBF deformation.
 		// Therefore, even if simulation is disabled, we need to run partially the update if the binding system is enabled (skin deformation + RBF correction)
-		const bool bNeedGuides = (GroupData.Guides.HasValidData() && (bHasNeedSimulation || bHasNeedGlobalDeformation)) || (HairGroupInstance->Debug.GroomCacheType == EGroomCacheType::Guides);
+		const bool bNeedGuides = (GroupData.Guides.HasValidData() && (bHasNeedSimulation[GroupIt] || bHasNeedGlobalDeformation[GroupIt])) || (HairGroupInstance->Debug.GroomCacheType == EGroomCacheType::Guides);
 		if (bNeedGuides)
 		{
 			HairGroupInstance->Guides.Data = &GroupData.Guides.BulkData;
diff --git a/Engine/Plugins/Runtime/HairStrands/Source/HairStrandsCore/Private/GroomResources.cpp b/Engine/Plugins/Runtime/HairStrands/Source/HairStrandsCore/Private/GroomResources.cpp
--- a/Engine/Plugins/Runtime/HairStrands/Source/HairStrandsCore/Private/GroomResources.cpp
+++ b/Engine/Plugins/Runtime/HairStrands/Source/HairStrandsCore/Private/GroomResources.cpp
@@ -746,7 +746,10 @@ void FHairMeshesDeformedResource::InternalRelease()
 FHairStrandsRestResource::FHairStrandsRestResource(FHairStrandsBulkData& InBulkData, EHairStrandsResourcesType InCurveType, const FHairResourceName& InResourceName) :
 	FHairCommonResource(EHairStrandsAllocationType::Deferred, InResourceName),
 	PositionBuffer(), Attribute0Buffer(), Attribute1Buffer(), MaterialBuffer(), BulkData(InBulkData), CurveType(InCurveType)
-{}
+{
+	// Sanity check
+	check(!!(BulkData.Flags & FHairStrandsBulkData::DataFlags_HasData));
+}
 
 bool FHairStrandsRestResource::InternalIsDataLoaded()
 {
@@ -899,6 +902,12 @@ void FHairStrandsClusterCullingBulkData::Reset()
 	VertexToClusterIds.RemoveBulkData();
 	ClusterVertexIds.RemoveBulkData();
 	PackedClusterInfos.RemoveBulkData();
+
+	// Reset the bulk byte buffer to ensure the (serialize) data size is reset to 0
+	ClusterLODInfos 	= FByteBulkData();
+	VertexToClusterIds 	= FByteBulkData();
+	ClusterVertexIds 	= FByteBulkData();
+	PackedClusterInfos 	= FByteBulkData();
 }
 
 void FHairStrandsClusterCullingBulkData::Serialize(FArchive& Ar, UObject* Owner)
@@ -1332,6 +1341,7 @@ void FHairStrandsRootBulkData::Reset()
 {
 	RootCount = 0;
 	VertexToCurveIndexBuffer.RemoveBulkData();
+	VertexToCurveIndexBuffer = FByteBulkData();
 	for (FMeshProjectionLOD& LOD : MeshProjectionLODs)
 	{
 		LOD.RootTriangleIndexBuffer.RemoveBulkData();
@@ -1342,6 +1352,17 @@ void FHairStrandsRootBulkData::Reset()
 		LOD.MeshInterpolationWeightsBuffer.RemoveBulkData();
 		LOD.MeshSampleIndicesBuffer.RemoveBulkData();
 		LOD.RestSamplePositionsBuffer.RemoveBulkData();
+
+		// Reset the bulk byte buffer to ensure the (serialize) data size is reset to 0
+		LOD.RootTriangleIndexBuffer 		= FByteBulkData();
+		LOD.RootTriangleBarycentricBuffer 	= FByteBulkData();
+		LOD.RestRootTrianglePosition0Buffer = FByteBulkData();
+		LOD.RestRootTrianglePosition1Buffer = FByteBulkData();
+		LOD.RestRootTrianglePosition2Buffer = FByteBulkData();
+		LOD.MeshInterpolationWeightsBuffer 	= FByteBulkData();
+		LOD.MeshSampleIndicesBuffer 		= FByteBulkData();
+		LOD.RestSamplePositionsBuffer 		= FByteBulkData();
+
 		LOD.ValidSectionIndices.Empty();
 	}
 	MeshProjectionLODs.Empty();
@@ -1388,6 +1409,8 @@ FHairStrandsInterpolationResource::FHairStrandsInterpolationResource(FHairStrand
 	FHairCommonResource(EHairStrandsAllocationType::Deferred, InResourceName),
 	InterpolationBuffer(), Interpolation0Buffer(), Interpolation1Buffer(), BulkData(InBulkData)
 {
+	// Sanity check
+	check(!!(BulkData.Flags & FHairStrandsInterpolationBulkData::DataFlags_HasData));
 }
 
 bool FHairStrandsInterpolationResource::InternalIsDataLoaded()
diff --git a/Engine/Plugins/Runtime/HairStrands/Source/HairStrandsCore/Private/HairStrandsDatas.cpp b/Engine/Plugins/Runtime/HairStrands/Source/HairStrandsCore/Private/HairStrandsDatas.cpp
--- a/Engine/Plugins/Runtime/HairStrands/Source/HairStrandsCore/Private/HairStrandsDatas.cpp
+++ b/Engine/Plugins/Runtime/HairStrands/Source/HairStrandsCore/Private/HairStrandsDatas.cpp
@@ -168,10 +168,18 @@ void FHairStrandsInterpolationBulkData::Reset()
 	Flags = 0;
 	PointCount = 0;
 	SimPointCount = 0;
+	
+	// Deallocate memory if needed
 	Interpolation.RemoveBulkData();
 	Interpolation0.RemoveBulkData();
 	Interpolation1.RemoveBulkData();
 	SimRootPointIndex.RemoveBulkData();
+
+	// Reset the bulk byte buffer to ensure the (serialize) data size is reset to 0
+	Interpolation		= FByteBulkData();
+	Interpolation0		= FByteBulkData();
+	Interpolation1		= FByteBulkData();
+	SimRootPointIndex	= FByteBulkData();
 }
 
 void FHairStrandsInterpolationBulkData::Serialize(FArchive& Ar, UObject* Owner)
@@ -261,6 +269,30 @@ void FHairStrandsBulkData::Serialize(FArchive& Ar, UObject* Owner)
 	}
 }
 
+void FHairStrandsBulkData::Reset()
+{
+	CurveCount = 0;
+	PointCount = 0;
+	MaxLength = 0;
+	MaxRadius = 0;
+	BoundingBox;
+	Flags = 0;
+
+	// Deallocate memory if needed
+	Positions.RemoveBulkData();
+	Attributes0.RemoveBulkData();
+	Attributes1.RemoveBulkData();
+	Materials.RemoveBulkData();
+	CurveOffsets.RemoveBulkData();
+
+	// Reset the bulk byte buffer to ensure the (serialize) data size is reset to 0
+	Positions 	 = FByteBulkData();
+	Attributes0  = FByteBulkData();
+	Attributes1  = FByteBulkData();
+	Materials 	 = FByteBulkData();
+	CurveOffsets = FByteBulkData();
+}
+
 void FHairStrandsDatas::Reset()
 {
 	StrandsCurves.Reset();
diff --git a/Engine/Plugins/Runtime/HairStrands/Source/HairStrandsCore/Public/HairStrandsDatas.h b/Engine/Plugins/Runtime/HairStrands/Source/HairStrandsCore/Public/HairStrandsDatas.h
--- a/Engine/Plugins/Runtime/HairStrands/Source/HairStrandsCore/Public/HairStrandsDatas.h
+++ b/Engine/Plugins/Runtime/HairStrands/Source/HairStrandsCore/Public/HairStrandsDatas.h
@@ -435,8 +435,9 @@ struct HAIRSTRANDSCORE_API FHairStrandsBulkData
 
 	void Serialize(FArchive& Ar, UObject* Owner);
 
-	bool IsValid() const { return CurveCount > 0 && PointCount > 0;	}
-	
+	bool IsValid() const { return CurveCount > 0 && PointCount > 0; }
+	void Reset();
+
 	uint32 GetNumCurves() const { return CurveCount;  };
 	uint32 GetNumPoints() const { return PointCount; };
 	float  GetMaxLength() const	{ return MaxLength; };
