Fix missing reset of hair strands bulk data.

diff --git a/Engine/Plugins/Runtime/HairStrands/Source/HairStrandsCore/Private/HairStrandsDatas.cpp b/Engine/Plugins/Runtime/HairStrands/Source/HairStrandsCore/Private/HairStrandsDatas.cpp
--- a/Engine/Plugins/Runtime/HairStrands/Source/HairStrandsCore/Private/HairStrandsDatas.cpp
+++ b/Engine/Plugins/Runtime/HairStrands/Source/HairStrandsCore/Private/HairStrandsDatas.cpp
@@ -275,7 +275,7 @@ void FHairStrandsBulkData::Reset()
 	PointCount = 0;
 	MaxLength = 0;
 	MaxRadius = 0;
-	BoundingBox;
+	BoundingBox = FBox(EForceInit::ForceInit);
 	Flags = 0;
 
 	// Deallocate memory if needed
