Change hair resources bulkdata to use 'SingleUse' flag for unloading/deallocating CPU memory.

diff --git a/Engine/Plugins/Runtime/HairStrands/Source/HairStrandsCore/Private/GroomResources.cpp b/Engine/Plugins/Runtime/HairStrands/Source/HairStrandsCore/Private/GroomResources.cpp
--- a/Engine/Plugins/Runtime/HairStrands/Source/HairStrandsCore/Private/GroomResources.cpp
+++ b/Engine/Plugins/Runtime/HairStrands/Source/HairStrandsCore/Private/GroomResources.cpp
@@ -180,6 +180,15 @@ void InternalCreateVertexBufferRDG_FromBulkData(FRDGBuilder& GraphBuilder, FByte
 		Desc.Usage = Desc.Usage & (~BUF_UnorderedAccess);
 	}
 
+	// Unloading of the bulk data is only supported on cooked build, as we can reload the data from the file/archieve
+	#if !WITH_EDITORONLY_DATA
+	const bool bReleaseCPUData = GHairStrandsBulkData_ReleaseAfterUse > 0;
+	if (bReleaseCPUData)
+	{
+		InBulkData.SetBulkDataFlags(BULKDATA_SingleUse);
+	}
+	#endif
+
 	const typename FormatType::BulkType* BulkData = (const typename FormatType::BulkType*)InBulkData.Lock(LOCK_READ_ONLY);
 	FRDGBufferRef Buffer = InternalCreateVertexBuffer(
 		GraphBuilder,
@@ -190,15 +199,6 @@ void InternalCreateVertexBufferRDG_FromBulkData(FRDGBuilder& GraphBuilder, FByte
 		ERDGInitialDataFlags::None); // Copy data internally
 	InBulkData.Unlock();
 
-	// Unloading of the bulk data is only supported on cooked build, as we can reload the data from the file/archieve
-	#if !WITH_EDITORONLY_DATA
-	const bool bReleaseCPUData = GHairStrandsBulkData_ReleaseAfterUse > 0;
-	if (bReleaseCPUData)
-	{
-		InBulkData.RemoveBulkData();
-	}
-	#endif
-
 	ConvertToExternalBufferWithViews(GraphBuilder, Buffer, Out, FormatType::Format);
 }
 
@@ -385,6 +385,15 @@ void InternalCreateStructuredBufferRDG_FromBulkData(FRDGBuilder& GraphBuilder, F
 		Desc.Usage = Desc.Usage & (~BUF_UnorderedAccess);
 	}
 
+	// Unloading of the bulk data is only supported on cooked build, as we can reload the data from the file/archieve
+	#if !WITH_EDITORONLY_DATA
+	const bool bReleaseCPUData = GHairStrandsBulkData_ReleaseAfterUse > 0;
+	if (bReleaseCPUData)
+	{
+		InBulkData.SetBulkDataFlags(BULKDATA_SingleUse);
+	}
+	#endif
+
 	const typename FormatType::Type* Data = (const typename FormatType::Type*)InBulkData.Lock(LOCK_READ_ONLY);
 	FRDGBufferRef Buffer = GraphBuilder.CreateBuffer(Desc, DebugName, ERDGBufferFlags::MultiFrame);
 	if (Data && DataSizeInBytes)
@@ -393,13 +402,6 @@ void InternalCreateStructuredBufferRDG_FromBulkData(FRDGBuilder& GraphBuilder, F
 	}
 	InBulkData.Unlock();
 
-	#if !WITH_EDITORONLY_DATA
-	const bool bReleaseCPUData = GHairStrandsBulkData_ReleaseAfterUse > 0;
-	if (bReleaseCPUData)
-	{
-		InBulkData.RemoveBulkData();
-	}
-	#endif
 
 	ConvertToExternalBufferWithViews(GraphBuilder, Buffer, Out);
 }
