Change groom LOD selection to change LOD only when the LOD resources are available (instead of drawing no hair).

If a previous LOD was loaded/used prior to the newly requested LOD, used it during the loading time of the requested LOD data.

diff --git a/Engine/Plugins/Runtime/HairStrands/Source/HairStrandsCore/Private/GroomManager.cpp b/Engine/Plugins/Runtime/HairStrands/Source/HairStrandsCore/Private/GroomManager.cpp
--- a/Engine/Plugins/Runtime/HairStrands/Source/HairStrandsCore/Private/GroomManager.cpp
+++ b/Engine/Plugins/Runtime/HairStrands/Source/HairStrandsCore/Private/GroomManager.cpp
@@ -641,6 +641,7 @@ static void RunHairLODSelection(
 		// Use the forced LOD index if set, otherwise compute the LOD based on the maximal screensize accross all views
 		// Compute the view where the screen size is maximale
 		const float PrevLODIndex = Instance->HairGroupPublicData->GetLODIndex();
+		const float PrevMeshLODIndex = Instance->HairGroupPublicData->GetMeshLODIndex();
 		float LODIndex = Instance->Strands.Modifier.LODForcedIndex; // check where this is updated 
 
 		// Insure that MinLOD is necessary taken into account if a force LOD is request (i.e., LODIndex>=0). If a Force LOD 
@@ -784,27 +785,35 @@ static void RunHairLODSelection(
 			Instance->Strands.VertexFactory->InitResources();
 		}
 
-		Instance->HairGroupPublicData->SetLODVisibility(bIsVisible);
-		Instance->HairGroupPublicData->SetLODIndex(LODIndex);
-		Instance->HairGroupPublicData->SetLODBias(0);
-		Instance->HairGroupPublicData->VFInput.GeometryType = GeometryType;
-		Instance->HairGroupPublicData->VFInput.bHasLODSwitch = (FMath::FloorToInt(PrevLODIndex) != FMath::FloorToInt(LODIndex));
-		Instance->GeometryType = GeometryType;
-		Instance->BindingType = Instance->HairGroupPublicData->GetBindingType(IntLODIndex);
-		Instance->Guides.bIsSimulationEnable = Instance->HairGroupPublicData->IsSimulationEnable(IntLODIndex);
-		Instance->Guides.bHasGlobalInterpolation = Instance->HairGroupPublicData->IsGlobalInterpolationEnable(IntLODIndex);
-		Instance->Strands.bIsCullingEnabled = bCullingEnable;
-
-		// Update the local-to-world transform based on the binding type 
-		Instance->LocalToWorld = Instance->BindingType == EHairBindingType::Skinning ? Instance->Debug.SkeletalLocalToWorld : Instance->LocalToWorld;
+		// Only switch LOD if the data are ready to be used
+		const bool bIsLODDataReady = !!(ResourceStatus & EHairResourceStatus::Loading);
+		if (bIsLODDataReady)
+		{
+			Instance->HairGroupPublicData->SetLODVisibility(bIsVisible);
+			Instance->HairGroupPublicData->SetLODIndex(LODIndex);
+			Instance->HairGroupPublicData->SetLODBias(0);
+			Instance->HairGroupPublicData->SetMeshLODIndex(MeshLODIndex);
+			Instance->HairGroupPublicData->VFInput.GeometryType = GeometryType;
+			Instance->HairGroupPublicData->VFInput.bHasLODSwitch = (FMath::FloorToInt(PrevLODIndex) != FMath::FloorToInt(LODIndex));
+			Instance->GeometryType = GeometryType;
+			Instance->BindingType = Instance->HairGroupPublicData->GetBindingType(IntLODIndex);
+			Instance->Guides.bIsSimulationEnable = Instance->HairGroupPublicData->IsSimulationEnable(IntLODIndex);
+			Instance->Guides.bHasGlobalInterpolation = Instance->HairGroupPublicData->IsGlobalInterpolationEnable(IntLODIndex);
+			Instance->Strands.bIsCullingEnabled = bCullingEnable;
+		}
 
-		// If all resources are not ready yet, change the geomtry to be invalid, so that it don't get processed this frame.
-		if (!(ResourceStatus & EHairResourceStatus::Loading))
+		// If requested LOD's resources are not ready yet, and the previous LOD was invalid or if the MeshLODIndex has changed 
+		// (which would required extra data loading) change the geometry to be invalid, so that it don't get processed this frame.
+		const bool bIsLODValid = PrevLODIndex >= 0;
+		const bool bHasMeshLODChanged = PrevMeshLODIndex != MeshLODIndex && (Instance->Guides.bHasGlobalInterpolation || Instance->BindingType == EHairBindingType::Skinning);
+		if (!bIsLODDataReady && (!bIsLODValid || bHasMeshLODChanged))
 		{
 			Instance->GeometryType = EHairGeometryType::NoneGeometry;
 			Instance->BindingType = EHairBindingType::NoneBinding;
 		}
 
+		// Update the local-to-world transform based on the binding type 
+		Instance->LocalToWorld = Instance->BindingType == EHairBindingType::Skinning ? Instance->Debug.SkeletalLocalToWorld : Instance->LocalToWorld;
 	}
 }
 
diff --git a/Engine/Source/Runtime/Renderer/Public/HairStrandsInterface.h b/Engine/Source/Runtime/Renderer/Public/HairStrandsInterface.h
--- a/Engine/Source/Runtime/Renderer/Public/HairStrandsInterface.h
+++ b/Engine/Source/Runtime/Renderer/Public/HairStrandsInterface.h
@@ -223,6 +223,9 @@ public:
 	float GetLODIndex() const { return LODIndex; }
 	int32 GetIntLODIndex() const { return FMath::Max(0, FMath::FloorToInt(LODIndex)); }
 
+	void SetMeshLODIndex(float InMeshLODIndex) { MeshLODIndex = InMeshLODIndex; }
+	float GetMeshLODIndex() const { return MeshLODIndex; }
+
 	void SetLODVisibility(bool bVisible) { bLODVisibility = bVisible; }
 	bool GetLODVisibility() const { return bLODVisibility; }
 
@@ -320,7 +323,8 @@ public:
 	TArray<EHairBindingType> BindingTypes;
 
 	// Data change every frame by the groom proxy based on views data
-	float LODIndex = 0;			// Current LOD used for all views
+	float MeshLODIndex = -1;
+	float LODIndex = -1;		// Current LOD used for all views
 	float LODBias = 0;			// Current LOD bias
 	bool bLODVisibility = true; // Enable/disable hair rendering for this component
 };
