Add simulation reset when the groom is not visible

diff --git a/Engine/Plugins/Runtime/HairStrands/Source/HairStrandsCore/Private/GroomComponent.cpp b/Engine/Plugins/Runtime/HairStrands/Source/HairStrandsCore/Private/GroomComponent.cpp
--- a/Engine/Plugins/Runtime/HairStrands/Source/HairStrandsCore/Private/GroomComponent.cpp
+++ b/Engine/Plugins/Runtime/HairStrands/Source/HairStrandsCore/Private/GroomComponent.cpp
@@ -3121,6 +3121,10 @@ void UGroomComponent::TickComponent(float DeltaTime, enum ELevelTick TickType, F
 			}
 		}
 		bResetSimulation |= SimulationSettings.SimulationSetup.bResetSimulation;
+		if (!IsVisible() || (GetOwner() && GetOwner()->IsHidden()))
+		{
+			bResetSimulation = true;
+		}
 	}
 	bInitSimulation = false;
 
