Add hair length clip debug info.

diff --git a/Engine/Plugins/Runtime/HairStrands/Source/HairStrandsCore/Private/GroomDebug.cpp b/Engine/Plugins/Runtime/HairStrands/Source/HairStrandsCore/Private/GroomDebug.cpp
--- a/Engine/Plugins/Runtime/HairStrands/Source/HairStrandsCore/Private/GroomDebug.cpp
+++ b/Engine/Plugins/Runtime/HairStrands/Source/HairStrandsCore/Private/GroomDebug.cpp
@@ -869,7 +869,7 @@ void RunHairStrandsDebug(
 					}
 				}
 
-				Line = FString::Printf(TEXT(" * Group:%d/%d | LOD:%1.2f/%d | GeometryType:%s | BindingType:%s | Sim:%d | RBF:%d | Clip:%d | VertexCount:%d | Name: %s"),
+				Line = FString::Printf(TEXT(" * Group:%d/%d | LOD:%1.2f/%d | GeometryType:%s | BindingType:%s | Sim:%d | RBF:%d | Clip:%d(%1.2f) | VertexCount:%d | Name: %s"),
 					Instance->Debug.GroupIndex,
 					Instance->Debug.GroupCount,
 
@@ -880,7 +880,10 @@ void RunHairStrandsDebug(
 					ToString(Instance->BindingType),
 					Instance->Guides.bIsSimulationEnable,
 					Instance->Guides.bHasGlobalInterpolation,
+
 					Instance->Strands.Modifier.HairLengthScale_Override,
+					Instance->Strands.Modifier.HairLengthScale,
+
 					Instance->HairGroupPublicData->VertexCount,
 					*Instance->Debug.GroomAssetName);
 				int32 XOffset = Canvas.DrawShadowedString(X, Y += YStep, *Line, GetStatsFont(), DebugGroupColor);
