Ensure no backfaceculling run on hair strands HW path.

diff --git a/Engine/Plugins/Runtime/HairStrands/Source/HairStrandsCore/Private/GroomComponent.cpp b/Engine/Plugins/Runtime/HairStrands/Source/HairStrandsCore/Private/GroomComponent.cpp
--- a/Engine/Plugins/Runtime/HairStrands/Source/HairStrandsCore/Private/GroomComponent.cpp
+++ b/Engine/Plugins/Runtime/HairStrands/Source/HairStrandsCore/Private/GroomComponent.cpp
@@ -862,6 +862,7 @@ public:
 		BatchElement.MaxVertexIndex = MaxVertexIndex;
 		BatchElement.UserData = reinterpret_cast<void*>(uint64(ComponentId));
 		Mesh.ReverseCulling = bUseCardsOrMeshes ? IsLocalToWorldDeterminantNegative() : false;
+		Mesh.bDisableBackfaceCulling = GeometryType == EHairGeometryType::Strands;
 		#if USE_HAIR_TRIANGLE_STRIP
 		Mesh.Type = GeometryType == EHairGeometryType::Strands ? PT_TriangleStrip : PT_TriangleList;
 		#else
