Add groom_roughness import attribute.

diff --git a/Engine/Plugins/Runtime/HairStrands/Source/HairStrandsCore/Private/GroomBuilder.cpp b/Engine/Plugins/Runtime/HairStrands/Source/HairStrandsCore/Private/GroomBuilder.cpp
--- a/Engine/Plugins/Runtime/HairStrands/Source/HairStrandsCore/Private/GroomBuilder.cpp
+++ b/Engine/Plugins/Runtime/HairStrands/Source/HairStrandsCore/Private/GroomBuilder.cpp
@@ -1470,8 +1470,16 @@ bool FGroomBuilder::BuildHairDescriptionGroups(const FHairDescription& HairDescr
 		GroomHairColor = GroomHairColorAttribute[GroomID];
 	}
 
+	TGroomAttributesConstRef<float> GroomHairRoughnessAttribute = HairDescription.GroomAttributes().GetAttributesRef<float>(HairAttribute::Groom::Roughness);
+	TOptional<float> GroomHairRoughness;
+	if (GroomHairRoughnessAttribute.IsValid())
+	{
+		GroomHairRoughness = GroomHairRoughnessAttribute[GroomID];
+	}
+
 	TVertexAttributesConstRef<FVector3f> VertexPositions	= HairDescription.VertexAttributes().GetAttributesRef<FVector3f>(HairAttribute::Vertex::Position);
 	TVertexAttributesConstRef<FVector3f> VertexBaseColor	= HairDescription.VertexAttributes().GetAttributesRef<FVector3f>(HairAttribute::Vertex::Color);
+	TVertexAttributesConstRef<float> VertexRoughness	= HairDescription.VertexAttributes().GetAttributesRef<float>(HairAttribute::Vertex::Roughness);
 	TStrandAttributesConstRef<int> StrandNumVertices	= HairDescription.StrandAttributes().GetAttributesRef<int>(HairAttribute::Strand::VertexCount);
 
 	if (!VertexPositions.IsValid() || !StrandNumVertices.IsValid())
@@ -1481,6 +1489,7 @@ bool FGroomBuilder::BuildHairDescriptionGroups(const FHairDescription& HairDescr
 	}
 
 	const bool bHasBaseColorAttribute = VertexBaseColor.IsValid();
+	const bool bHasRoughnessAttribute = VertexRoughness.IsValid();
 
 	TVertexAttributesConstRef<float> VertexWidths = HairDescription.VertexAttributes().GetAttributesRef<float>(HairAttribute::Vertex::Width);
 	TStrandAttributesConstRef<float> StrandWidths = HairDescription.StrandAttributes().GetAttributesRef<float>(HairAttribute::Strand::Width);
@@ -1598,8 +1607,8 @@ bool FGroomBuilder::BuildHairDescriptionGroups(const FHairDescription& HairDescr
 			FVertexID VertexID(GlobalVertexIndex);
 
 			CurrentHairStrandsDatas->StrandsPoints.PointsPosition.Add(VertexPositions[VertexID]);
-			CurrentHairStrandsDatas->StrandsPoints.PointsBaseColor.Add(bHasBaseColorAttribute ? VertexBaseColor[VertexID] : FLinearColor::Black);
-			CurrentHairStrandsDatas->StrandsPoints.PointsRoughness.Add(0); // @hair_todo: add attribute read on the alembic for reading roughness per groom/strands/vertex
+			CurrentHairStrandsDatas->StrandsPoints.PointsBaseColor.Add(bHasBaseColorAttribute ? VertexBaseColor[VertexID] : (GroomHairColor     ? GroomHairColor.GetValue() : FLinearColor::Black));
+			CurrentHairStrandsDatas->StrandsPoints.PointsRoughness.Add(bHasRoughnessAttribute ? VertexRoughness[VertexID] : (GroomHairRoughness ? GroomHairRoughness.GetValue() : 0.f));
 
 			float VertexWidth = 0.f;
 			if (VertexWidths.IsValid())
@@ -1931,8 +1940,8 @@ void Decimate(
 			OutData.StrandsPoints.PointsPosition [OutPointIndex + OutPointOffset] = InData.StrandsPoints.PointsPosition	[InPointIndex + InPointOffset];
 			OutData.StrandsPoints.PointsCoordU	 [OutPointIndex + OutPointOffset] = InData.StrandsPoints.PointsCoordU	[InPointIndex + InPointOffset];
 			OutData.StrandsPoints.PointsRadius	 [OutPointIndex + OutPointOffset] = InData.StrandsPoints.PointsRadius	[InPointIndex + InPointOffset];
-			OutData.StrandsPoints.PointsBaseColor[OutPointIndex + OutPointOffset] = FLinearColor::Black;
-			OutData.StrandsPoints.PointsRoughness[OutPointIndex + OutPointOffset] = 0;
+			OutData.StrandsPoints.PointsBaseColor[OutPointIndex + OutPointOffset] = InData.StrandsPoints.PointsBaseColor[InPointIndex + InPointOffset];
+			OutData.StrandsPoints.PointsRoughness[OutPointIndex + OutPointOffset] = InData.StrandsPoints.PointsRoughness[InPointIndex + InPointOffset];
 		}
 	#else
 		// Decimation using uniform/interleaved removal
diff --git a/Engine/Plugins/Runtime/HairStrands/Source/HairStrandsCore/Private/HairAttributes.cpp b/Engine/Plugins/Runtime/HairStrands/Source/HairStrandsCore/Private/HairAttributes.cpp
--- a/Engine/Plugins/Runtime/HairStrands/Source/HairStrandsCore/Private/HairAttributes.cpp
+++ b/Engine/Plugins/Runtime/HairStrands/Source/HairStrandsCore/Private/HairAttributes.cpp
@@ -5,10 +5,12 @@
 namespace HairAttribute
 {
 	const FName Vertex::Color("groom_color");
+	const FName Vertex::Roughness("groom_roughness");
 	const FName Vertex::Position("position");
 	const FName Vertex::Width("groom_width");
 
 	const FName Strand::Color("groom_color");
+	const FName Strand::Roughness("groom_roughness");
 	const FName Strand::GroupID("groom_group_id");
 	const FName Strand::Guide("groom_guide");
 	const FName Strand::ID("groom_id");
@@ -19,6 +21,7 @@ namespace HairAttribute
 	const FName Strand::GuideWeights("groom_guide_weights");
 
 	const FName Groom::Color("groom_color");
+	const FName Groom::Roughness("groom_roughness");
 	const FName Groom::Width("groom_width");
 	const FName Groom::MajorVersion("groom_version_major");
 	const FName Groom::MinorVersion("groom_version_minor");
diff --git a/Engine/Plugins/Runtime/HairStrands/Source/HairStrandsCore/Public/HairAttributes.h b/Engine/Plugins/Runtime/HairStrands/Source/HairStrandsCore/Public/HairAttributes.h
--- a/Engine/Plugins/Runtime/HairStrands/Source/HairStrandsCore/Public/HairAttributes.h
+++ b/Engine/Plugins/Runtime/HairStrands/Source/HairStrandsCore/Public/HairAttributes.h
@@ -9,6 +9,7 @@ namespace HairAttribute
 	namespace Vertex
 	{
 		extern HAIRSTRANDSCORE_API const FName Color;			// FVector
+		extern HAIRSTRANDSCORE_API const FName Roughness;		// float
 		extern HAIRSTRANDSCORE_API const FName Position;		// FVector
 		extern HAIRSTRANDSCORE_API const FName Width;			// float
 	}
@@ -16,6 +17,7 @@ namespace HairAttribute
 	namespace Strand
 	{
 		extern HAIRSTRANDSCORE_API const FName Color;			// FVector
+		extern HAIRSTRANDSCORE_API const FName Roughness;		// float
 		extern HAIRSTRANDSCORE_API const FName GroupID;			// int
 		extern HAIRSTRANDSCORE_API const FName Guide;			// int
 		extern HAIRSTRANDSCORE_API const FName ID;				// int
@@ -29,6 +31,7 @@ namespace HairAttribute
 	namespace Groom
 	{
 		extern HAIRSTRANDSCORE_API const FName Color;			// FVector
+		extern HAIRSTRANDSCORE_API const FName Roughness;		// float
 		extern HAIRSTRANDSCORE_API const FName MajorVersion;	// int
 		extern HAIRSTRANDSCORE_API const FName MinorVersion;	// int
 		extern HAIRSTRANDSCORE_API const FName Tool;			// FName
