Add ResetSimulation event to UGroomComponent.

diff --git a/Engine/Plugins/Runtime/HairStrands/Source/HairStrandsCore/Private/GroomComponent.cpp b/Engine/Plugins/Runtime/HairStrands/Source/HairStrandsCore/Private/GroomComponent.cpp
--- a/Engine/Plugins/Runtime/HairStrands/Source/HairStrandsCore/Private/GroomComponent.cpp
+++ b/Engine/Plugins/Runtime/HairStrands/Source/HairStrandsCore/Private/GroomComponent.cpp
@@ -1652,6 +1652,15 @@ void UGroomComponent::SetEnableSimulation(bool bInEnableSimulation)
 	}
 }
 
+void UGroomComponent::ResetSimulation()
+{
+	if (SimulationSettings.SolverSettings.bEnableSimulation)
+	{
+		bResetSimulation = true;
+		bInitSimulation = true;
+	}
+}
+
 void UGroomComponent::UpdateHairGroupsDescAndInvalidateRenderState(bool bInvalid)
 {
 	UpdateHairGroupsDesc();
@@ -3679,8 +3688,7 @@ void ResetSimulationOnAllGroomComponents(IConsoleVariable* InCVarPakTesterEnable
 	{
 		for (TObjectIterator<UGroomComponent> HairStrandsComponentIt; HairStrandsComponentIt; ++HairStrandsComponentIt)
 		{
-			HairStrandsComponentIt->bResetSimulation = true;
-			HairStrandsComponentIt->bInitSimulation  = true;
+			HairStrandsComponentIt->ResetSimulation();
 		}
 	}
 	GHairStrandsResetSimulationOnAllGroomComponents = 0;
diff --git a/Engine/Plugins/Runtime/HairStrands/Source/HairStrandsCore/Public/GroomComponent.h b/Engine/Plugins/Runtime/HairStrands/Source/HairStrandsCore/Public/GroomComponent.h
--- a/Engine/Plugins/Runtime/HairStrands/Source/HairStrandsCore/Public/GroomComponent.h
+++ b/Engine/Plugins/Runtime/HairStrands/Source/HairStrandsCore/Public/GroomComponent.h
@@ -159,6 +159,10 @@ public:
 	UFUNCTION(BlueprintCallable, Category = "Simulation")
 	void SetEnableSimulation(bool bInEnableSimulation);
 
+	/* Reset the simulation, if enabled */
+	UFUNCTION(BlueprintCallable, Category = "Simulation")
+	void ResetSimulation();
+
 	/* Accessor function for changing hair length scale from blueprint/sequencer */
 	UFUNCTION(BlueprintCallable, Category = "Groom")
 	void SetHairLengthScale(float Scale);
