Fix hair cards flickering.

diff --git a/Engine/Plugins/Runtime/HairStrands/Source/HairStrandsCore/Private/GroomComponent.cpp b/Engine/Plugins/Runtime/HairStrands/Source/HairStrandsCore/Private/GroomComponent.cpp
--- a/Engine/Plugins/Runtime/HairStrands/Source/HairStrandsCore/Private/GroomComponent.cpp
+++ b/Engine/Plugins/Runtime/HairStrands/Source/HairStrandsCore/Private/GroomComponent.cpp
@@ -340,8 +340,7 @@ public:
 				{
 					if (HairInstance->Cards.IsValid(LODIt))
 					{
-						if (HairInstance->Cards.LODs[LODIt].VertexFactory[0] == nullptr) HairInstance->Cards.LODs[LODIt].VertexFactory[0] = new FHairCardsVertexFactory(HairInstance, GroupIt, LODIt, 0u, EHairGeometryType::Cards, ShaderPlatform, FeatureLevel, "HairCardsVertexFactory");
-						if (HairInstance->Cards.LODs[LODIt].VertexFactory[1] == nullptr) HairInstance->Cards.LODs[LODIt].VertexFactory[1] = new FHairCardsVertexFactory(HairInstance, GroupIt, LODIt, 1u, EHairGeometryType::Cards, ShaderPlatform, FeatureLevel, "HairCardsVertexFactory");
+						if (HairInstance->Cards.LODs[LODIt].VertexFactory == nullptr) HairInstance->Cards.LODs[LODIt].VertexFactory = new FHairCardsVertexFactory(HairInstance, GroupIt, LODIt, EHairGeometryType::Cards, ShaderPlatform, FeatureLevel, "HairCardsVertexFactory");
 					}
 				}
 			}
@@ -354,8 +353,7 @@ public:
 				{
 					if (HairInstance->Meshes.IsValid(LODIt))
 					{
-						if (HairInstance->Meshes.LODs[LODIt].VertexFactory[0] == nullptr) HairInstance->Meshes.LODs[LODIt].VertexFactory[0] = new FHairCardsVertexFactory(HairInstance, GroupIt, LODIt, 0u, EHairGeometryType::Meshes, ShaderPlatform, FeatureLevel, "HairMeshesVertexFactory");
-						if (HairInstance->Meshes.LODs[LODIt].VertexFactory[1] == nullptr) HairInstance->Meshes.LODs[LODIt].VertexFactory[1] = new FHairCardsVertexFactory(HairInstance, GroupIt, LODIt, 1u, EHairGeometryType::Meshes, ShaderPlatform, FeatureLevel, "HairMeshesVertexFactory");
+						if (HairInstance->Meshes.LODs[LODIt].VertexFactory == nullptr) HairInstance->Meshes.LODs[LODIt].VertexFactory = new FHairCardsVertexFactory(HairInstance, GroupIt, LODIt, EHairGeometryType::Meshes, ShaderPlatform, FeatureLevel, "HairMeshesVertexFactory");
 
 					}
 				}
@@ -2783,8 +2781,7 @@ void UGroomComponent::DeleteDeferredHairGroupInstances()
 						}
 						#endif
 
-						InternalResourceRelease(LocalInstance->Cards.LODs[CardLODIt].VertexFactory[0]);
-						InternalResourceRelease(LocalInstance->Cards.LODs[CardLODIt].VertexFactory[1]);
+						InternalResourceRelease(LocalInstance->Cards.LODs[CardLODIt].VertexFactory);
 					}
 				}
 			}
@@ -2804,8 +2801,7 @@ void UGroomComponent::DeleteDeferredHairGroupInstances()
 						}
 						#endif
 
-						InternalResourceRelease(LocalInstance->Meshes.LODs[MeshesLODIt].VertexFactory[0]);
-						InternalResourceRelease(LocalInstance->Meshes.LODs[MeshesLODIt].VertexFactory[1]);
+						InternalResourceRelease(LocalInstance->Meshes.LODs[MeshesLODIt].VertexFactory);
 					}
 				}
 			}
diff --git a/Engine/Plugins/Runtime/HairStrands/Source/HairStrandsCore/Private/GroomManager.cpp b/Engine/Plugins/Runtime/HairStrands/Source/HairStrandsCore/Private/GroomManager.cpp
--- a/Engine/Plugins/Runtime/HairStrands/Source/HairStrandsCore/Private/GroomManager.cpp
+++ b/Engine/Plugins/Runtime/HairStrands/Source/HairStrandsCore/Private/GroomManager.cpp
@@ -52,20 +52,12 @@ namespace HairTransition
 
 void FHairGroupInstance::FCards::FLOD::InitVertexFactory()
 {
-	VertexFactory[0]->InitResources();
-	if (DeformedResource)
-	{
-		VertexFactory[1]->InitResources();
-	}
+	VertexFactory->InitResources();
 }
 
 void FHairGroupInstance::FMeshes::FLOD::InitVertexFactory()
 {
-	VertexFactory[0]->InitResources();
-	if (DeformedResource)
-	{
-		VertexFactory[1]->InitResources();
-	}
+	VertexFactory->InitResources();
 }
 
 static bool IsInstanceFrustumCullingEnable()
@@ -525,33 +517,6 @@ static void RunHairBufferSwap(const FHairStrandsInstances& Instances, const TArr
 		{
 			if (Instance->Guides.DeformedResource) { Instance->Guides.DeformedResource->SwapBuffer(); }
 			if (Instance->Strands.DeformedResource) { Instance->Strands.DeformedResource->SwapBuffer(); }
-
-			for (uint32 LODIt = 0, LODCount = Instance->Cards.LODs.Num(); LODIt < LODCount; ++LODIt)
-			{
-				if (Instance->Cards.IsValid(LODIt))
-				{
-					if (Instance->Cards.LODs[LODIt].DeformedResource)
-					{
-						Instance->Cards.LODs[LODIt].DeformedResource->SwapBuffer();
-					}
-					if (Instance->Cards.LODs[LODIt].Guides.IsValid())
-					{
-						Instance->Cards.LODs[LODIt].Guides.DeformedResource->SwapBuffer();
-					}
-				}
-			}
-
-			for (uint32 LODIt = 0, LODCount = Instance->Meshes.LODs.Num(); LODIt < LODCount; ++LODIt)
-			{
-				if (Instance->Meshes.IsValid(LODIt))
-				{
-					if (Instance->Meshes.LODs[LODIt].DeformedResource)
-					{
-						Instance->Meshes.LODs[LODIt].DeformedResource->SwapBuffer();
-					}
-				}
-			}
-
 			Instance->Debug.LastFrameIndex = Views[0]->Family->FrameNumber;
 		}
 
diff --git a/Engine/Plugins/Runtime/HairStrands/Source/HairStrandsCore/Private/GroomResources.cpp b/Engine/Plugins/Runtime/HairStrands/Source/HairStrandsCore/Private/GroomResources.cpp
--- a/Engine/Plugins/Runtime/HairStrands/Source/HairStrandsCore/Private/GroomResources.cpp
+++ b/Engine/Plugins/Runtime/HairStrands/Source/HairStrandsCore/Private/GroomResources.cpp
@@ -649,13 +649,13 @@ void FHairCardsDeformedResource::InternalAllocate(FRDGBuilder& GraphBuilder)
 {
 	if (bInitializedData)
 	{
-		InternalCreateVertexBufferRDG<FHairCardsPositionFormat>(GraphBuilder, BulkData.Positions, DeformedPositionBuffer[0], ToHairResourceDebugName(TEXT("Hair.CardsDeformed_Position0"), ResourceName), EHairResourceUsageType::Dynamic);
-		InternalCreateVertexBufferRDG<FHairCardsPositionFormat>(GraphBuilder, BulkData.Positions, DeformedPositionBuffer[1], ToHairResourceDebugName(TEXT("Hair.CardsDeformed_Position1"), ResourceName), EHairResourceUsageType::Dynamic);
+		InternalCreateVertexBufferRDG<FHairCardsPositionFormat>(GraphBuilder, BulkData.Positions, DeformedPositionBuffer[0], ToHairResourceDebugName(TEXT("Hair.CardsDeformed(Current)"), ResourceName), EHairResourceUsageType::Dynamic);
+		InternalCreateVertexBufferRDG<FHairCardsPositionFormat>(GraphBuilder, BulkData.Positions, DeformedPositionBuffer[1], ToHairResourceDebugName(TEXT("Hair.CardsDeformed(Previous)"), ResourceName), EHairResourceUsageType::Dynamic);
 	}
 	else
 	{
-		InternalCreateVertexBufferRDG<FHairCardsPositionFormat>(GraphBuilder, BulkData.GetNumVertices(), DeformedPositionBuffer[0], ToHairResourceDebugName(TEXT("Hair.CardsDeformed_Position0"), ResourceName), EHairResourceUsageType::Dynamic);
-		InternalCreateVertexBufferRDG<FHairCardsPositionFormat>(GraphBuilder, BulkData.GetNumVertices(), DeformedPositionBuffer[1], ToHairResourceDebugName(TEXT("Hair.CardsDeformed_Position1"), ResourceName), EHairResourceUsageType::Dynamic);
+		InternalCreateVertexBufferRDG<FHairCardsPositionFormat>(GraphBuilder, BulkData.GetNumVertices(), DeformedPositionBuffer[0], ToHairResourceDebugName(TEXT("Hair.CardsDeformed(Current)"), ResourceName), EHairResourceUsageType::Dynamic);
+		InternalCreateVertexBufferRDG<FHairCardsPositionFormat>(GraphBuilder, BulkData.GetNumVertices(), DeformedPositionBuffer[1], ToHairResourceDebugName(TEXT("Hair.CardsDeformed(Previous)"), ResourceName), EHairResourceUsageType::Dynamic);
 
 		// Manually transit to SRVs, in case of the cards are rendered not visible, but still rendered (in shadows for instance). In such a case, the cards deformation pass is not called, and thus the 
 		// buffers are never transit from UAV (clear) to SRV for rasterization. 
@@ -727,13 +727,13 @@ void FHairMeshesDeformedResource::InternalAllocate(FRDGBuilder& GraphBuilder)
 {
 	if (bInitializedData)
 	{
-		InternalCreateVertexBufferRDG<FHairCardsPositionFormat>(GraphBuilder, BulkData.Positions, DeformedPositionBuffer[0], ToHairResourceDebugName(TEXT("Hair.MeshesDeformed_Positions0"), ResourceName), EHairResourceUsageType::Dynamic);
-		InternalCreateVertexBufferRDG<FHairCardsPositionFormat>(GraphBuilder, BulkData.Positions, DeformedPositionBuffer[1], ToHairResourceDebugName(TEXT("Hair.MeshesDeformed_Positions1"), ResourceName), EHairResourceUsageType::Dynamic);
+		InternalCreateVertexBufferRDG<FHairCardsPositionFormat>(GraphBuilder, BulkData.Positions, DeformedPositionBuffer[0], ToHairResourceDebugName(TEXT("Hair.MeshesDeformed(Current)"), ResourceName), EHairResourceUsageType::Dynamic);
+		InternalCreateVertexBufferRDG<FHairCardsPositionFormat>(GraphBuilder, BulkData.Positions, DeformedPositionBuffer[1], ToHairResourceDebugName(TEXT("Hair.MeshesDeformed(Previous)"), ResourceName), EHairResourceUsageType::Dynamic);
 	}
 	else
 	{
-		InternalCreateVertexBufferRDG<FHairCardsPositionFormat>(GraphBuilder, BulkData.GetNumVertices(), DeformedPositionBuffer[0], ToHairResourceDebugName(TEXT("Hair.MeshesDeformed_Positions0"), ResourceName), EHairResourceUsageType::Dynamic);
-		InternalCreateVertexBufferRDG<FHairCardsPositionFormat>(GraphBuilder, BulkData.GetNumVertices(), DeformedPositionBuffer[1], ToHairResourceDebugName(TEXT("Hair.MeshesDeformed_Positions1"), ResourceName), EHairResourceUsageType::Dynamic);
+		InternalCreateVertexBufferRDG<FHairCardsPositionFormat>(GraphBuilder, BulkData.GetNumVertices(), DeformedPositionBuffer[0], ToHairResourceDebugName(TEXT("Hair.MeshesDeformed(Current)"), ResourceName), EHairResourceUsageType::Dynamic);
+		InternalCreateVertexBufferRDG<FHairCardsPositionFormat>(GraphBuilder, BulkData.GetNumVertices(), DeformedPositionBuffer[1], ToHairResourceDebugName(TEXT("Hair.MeshesDeformed(Previous)"), ResourceName), EHairResourceUsageType::Dynamic);
 	}
 }
 
diff --git a/Engine/Plugins/Runtime/HairStrands/Source/HairStrandsCore/Private/HairCardsVertexFactory.cpp b/Engine/Plugins/Runtime/HairStrands/Source/HairStrandsCore/Private/HairCardsVertexFactory.cpp
--- a/Engine/Plugins/Runtime/HairStrands/Source/HairStrandsCore/Private/HairCardsVertexFactory.cpp
+++ b/Engine/Plugins/Runtime/HairStrands/Source/HairStrandsCore/Private/HairCardsVertexFactory.cpp
@@ -25,7 +25,6 @@ template<typename T> inline void VFC_BindParam(FMeshDrawSingleShaderBindings& Sh
 IMPLEMENT_GLOBAL_SHADER_PARAMETER_STRUCT(FHairCardsVertexFactoryUniformShaderParameters, "HairCardsVF");
 
 FHairCardsUniformBuffer CreateHairCardsVFUniformBuffer(
-	uint32 Current,
 	const FHairGroupInstance* Instance,
 	const uint32 LODIndex,
 	EHairGeometryType GeometryType, 
@@ -46,8 +45,8 @@ FHairCardsUniformBuffer CreateHairCardsVFUniformBuffer(
 		const bool bIsDynamic = LOD.DeformedResource != nullptr;
 		if (bIsDynamic)
 		{
-			UniformParameters.PositionBuffer = LOD.DeformedResource->DeformedPositionBuffer[bIsDynamic ? Current : 0].SRV.GetReference();
-			UniformParameters.PreviousPositionBuffer = LOD.DeformedResource->DeformedPositionBuffer[bIsDynamic ? 1u-Current : 0].SRV.GetReference();
+			UniformParameters.PositionBuffer = LOD.DeformedResource->GetBuffer(FHairCardsDeformedResource::EFrameType::Current).SRV.GetReference();
+			UniformParameters.PreviousPositionBuffer = LOD.DeformedResource->GetBuffer(FHairCardsDeformedResource::EFrameType::Previous).SRV.GetReference();
 		}
 		else
 		{
@@ -78,8 +77,8 @@ FHairCardsUniformBuffer CreateHairCardsVFUniformBuffer(
 		const bool bIsDynamic = LOD.DeformedResource != nullptr;
 		if (bIsDynamic)
 		{
-			UniformParameters.PositionBuffer = LOD.DeformedResource->DeformedPositionBuffer[bIsDynamic ? Current : 0].SRV.GetReference();
-			UniformParameters.PreviousPositionBuffer = LOD.DeformedResource->DeformedPositionBuffer[bIsDynamic ? 1u - Current : 0].SRV.GetReference();
+			UniformParameters.PositionBuffer = LOD.DeformedResource->GetBuffer(FHairMeshesDeformedResource::EFrameType::Current).SRV.GetReference();
+			UniformParameters.PreviousPositionBuffer = LOD.DeformedResource->GetBuffer(FHairMeshesDeformedResource::EFrameType::Previous).SRV.GetReference();
 		}
 		else
 		{
@@ -172,18 +171,14 @@ public:
 		if (InstanceGeometryType == EHairGeometryType::Cards)
 		{
 			const FHairGroupInstance::FCards::FLOD& LOD = Instance->Cards.LODs[LODIndex];
-			const bool bIsDynamic = LOD.DeformedResource != nullptr;
-			const uint32 UniformIndex = bIsDynamic ? LOD.DeformedResource->GetIndex(FHairCardsDeformedResource::Current) : 0;
-			check(LOD.UniformBuffer[UniformIndex]);
-			VertexFactoryUniformBuffer = LOD.UniformBuffer[UniformIndex];
+			check(LOD.UniformBuffer);
+			VertexFactoryUniformBuffer = LOD.UniformBuffer;
 		}
 		else if (InstanceGeometryType == EHairGeometryType::Meshes)
 		{
 			const FHairGroupInstance::FMeshes::FLOD& LOD = Instance->Meshes.LODs[LODIndex];
-			const bool bIsDynamic = LOD.DeformedResource != nullptr;
-			const uint32 UniformIndex = bIsDynamic ? LOD.DeformedResource->GetIndex(FHairMeshesDeformedResource::Current) : 0;
-			check(LOD.UniformBuffer[UniformIndex]);
-			VertexFactoryUniformBuffer = LOD.UniformBuffer[UniformIndex];
+			check(LOD.UniformBuffer);
+			VertexFactoryUniformBuffer = LOD.UniformBuffer;
 		}
 		
 		check(VertexFactoryUniformBuffer);
@@ -193,7 +188,7 @@ public:
 
 IMPLEMENT_TYPE_LAYOUT(FHairCardsVertexFactoryShaderParameters);
 
-FHairCardsVertexFactory::FHairCardsVertexFactory(FHairGroupInstance* Instance, uint32 GroupIndex, uint32 LODIndex, uint32 BufferIndex, EHairGeometryType GeometryType, EShaderPlatform InShaderPlatform, ERHIFeatureLevel::Type InFeatureLevel, const char* InDebugName)
+FHairCardsVertexFactory::FHairCardsVertexFactory(FHairGroupInstance* Instance, uint32 GroupIndex, uint32 LODIndex, EHairGeometryType GeometryType, EShaderPlatform InShaderPlatform, ERHIFeatureLevel::Type InFeatureLevel, const char* InDebugName)
 	: FVertexFactory(InFeatureLevel)
 	, DebugName(InDebugName)
 {
@@ -202,7 +197,6 @@ FHairCardsVertexFactory::FHairCardsVertexFactory(FHairGroupInstance* Instance, u
 	Data.Instance = Instance;
 	Data.GroupIndex = GroupIndex;
 	Data.LODIndex = LODIndex;
-	Data.BufferIndex = BufferIndex;
 	Data.GeometryType = GeometryType;
 }
 /**
@@ -323,7 +317,7 @@ void FHairCardsVertexFactory::InitResources()
 		if (HairInstance->Cards.LODs.IsValidIndex(Data.LODIndex))
 		{
 			const FHairGroupInstance::FCards::FLOD& LOD = HairInstance->Cards.LODs[Data.LODIndex];
-			HairInstance->Cards.LODs[Data.LODIndex].UniformBuffer[Data.BufferIndex] = CreateHairCardsVFUniformBuffer(Data.BufferIndex, HairInstance, Data.LODIndex, EHairGeometryType::Cards, bManualFetch);
+			HairInstance->Cards.LODs[Data.LODIndex].UniformBuffer = CreateHairCardsVFUniformBuffer(HairInstance, Data.LODIndex, EHairGeometryType::Cards, bManualFetch);
 		}
 	}
 	else if (Data.GeometryType == EHairGeometryType::Meshes && IsHairStrandsEnabled(EHairStrandsShaderType::Meshes, GMaxRHIShaderPlatform))
@@ -331,7 +325,7 @@ void FHairCardsVertexFactory::InitResources()
 		if (HairInstance->Meshes.LODs.IsValidIndex(Data.LODIndex))
 		{
 			const FHairGroupInstance::FMeshes::FLOD& LOD = HairInstance->Meshes.LODs[Data.LODIndex];
-			HairInstance->Meshes.LODs[Data.LODIndex].UniformBuffer[Data.BufferIndex] = CreateHairCardsVFUniformBuffer(Data.BufferIndex, HairInstance, Data.LODIndex, EHairGeometryType::Meshes, bManualFetch);
+			HairInstance->Meshes.LODs[Data.LODIndex].UniformBuffer = CreateHairCardsVFUniformBuffer(HairInstance, Data.LODIndex, EHairGeometryType::Meshes, bManualFetch);
 		}
 	}
 }
diff --git a/Engine/Plugins/Runtime/HairStrands/Source/HairStrandsCore/Private/HairStrandsMeshProjection.cpp b/Engine/Plugins/Runtime/HairStrands/Source/HairStrandsCore/Private/HairStrandsMeshProjection.cpp
--- a/Engine/Plugins/Runtime/HairStrands/Source/HairStrandsCore/Private/HairStrandsMeshProjection.cpp
+++ b/Engine/Plugins/Runtime/HairStrands/Source/HairStrandsCore/Private/HairStrandsMeshProjection.cpp
@@ -1006,6 +1006,7 @@ template<typename TRestResource, typename TDeformedResource>
 void InternalAddHairRBFInterpolationPass(
 	FRDGBuilder& GraphBuilder,
 	FGlobalShaderMap* ShaderMap,
+	const bool bCards,
 	const int32 MeshLODIndex,
 	TRestResource* RestResources,
 	TDeformedResource* DeformedResources,
@@ -1018,15 +1019,19 @@ void InternalAddHairRBFInterpolationPass(
 		return;
 	}
 
+	// Copy current to previous position before update the current position. This allows to have correct motion vectors.
+	FRDGImportedBuffer DeformedPositionBuffer_Curr = Register(GraphBuilder, DeformedResources->GetBuffer(TDeformedResource::EFrameType::Current), ERDGImportedBufferFlags::CreateUAV);
+	FRDGImportedBuffer DeformedPositionBuffer_Prev = Register(GraphBuilder, DeformedResources->GetBuffer(TDeformedResource::EFrameType::Previous), ERDGImportedBufferFlags::CreateUAV);
+	AddCopyBufferPass(GraphBuilder, DeformedPositionBuffer_Prev.Buffer, DeformedPositionBuffer_Curr.Buffer);
+
 	FHairStrandsRestRootResource::FLOD& RestLODData = RestRootResources->LODs[MeshLODIndex];
 	FHairStrandsDeformedRootResource::FLOD& DeformedLODData = DeformedRootResources->LODs[MeshLODIndex];
 	FHairMeshesInterpolateCS::FParameters* Parameters = GraphBuilder.AllocParameters<FHairMeshesInterpolateCS::FParameters>();
 	Parameters->VertexCount = VertexCount;
 	Parameters->MaxSampleCount = RestLODData.SampleCount;
 
-	FRDGImportedBuffer OutDeformedPositionBuffer = Register(GraphBuilder, DeformedResources->GetBuffer(TDeformedResource::EFrameType::Current), ERDGImportedBufferFlags::CreateUAV);
 	Parameters->RestPositionBuffer			= RestResources->RestPositionBuffer.ShaderResourceViewRHI;
-	Parameters->OutDeformedPositionBuffer	= OutDeformedPositionBuffer.UAV;
+	Parameters->OutDeformedPositionBuffer	= DeformedPositionBuffer_Curr.UAV;
 
 	Parameters->RestSamplePositionsBuffer	= RegisterAsSRV(GraphBuilder, RestLODData.RestSamplePositionsBuffer);
 	Parameters->MeshSampleWeightsBuffer		= RegisterAsSRV(GraphBuilder, DeformedLODData.MeshSampleWeightsBuffer);
@@ -1036,12 +1041,14 @@ void InternalAddHairRBFInterpolationPass(
 	TShaderMapRef<FHairMeshesInterpolateCS> ComputeShader(ShaderMap);
 	FComputeShaderUtils::AddPass(
 		GraphBuilder,
-		RDG_EVENT_NAME("HairInterpolationRBF"),
+		RDG_EVENT_NAME("HairInterpolationRBF(%s)", bCards ? TEXT("Cards") : TEXT("Meshes")),
 		ComputeShader,
 		Parameters,
 		DispatchGroupCount);
 
-	GraphBuilder.SetBufferAccessFinal(OutDeformedPositionBuffer.Buffer, ERHIAccess::SRVMask);
+	GraphBuilder.SetBufferAccessFinal(DeformedPositionBuffer_Curr.Buffer, ERHIAccess::SRVMask);
+	GraphBuilder.SetBufferAccessFinal(DeformedPositionBuffer_Prev.Buffer, ERHIAccess::SRVMask);
+	
 }
 
 void AddHairMeshesRBFInterpolationPass(
@@ -1056,6 +1063,7 @@ void AddHairMeshesRBFInterpolationPass(
 	InternalAddHairRBFInterpolationPass(
 		GraphBuilder,
 		ShaderMap,
+		false,
 		MeshLODIndex,
 		RestResources,
 		DeformedResources,
@@ -1075,6 +1083,7 @@ void AddHairCardsRBFInterpolationPass(
 	InternalAddHairRBFInterpolationPass(
 		GraphBuilder,
 		ShaderMap,
+		true,
 		MeshLODIndex,
 		RestResources,
 		DeformedResources,
diff --git a/Engine/Plugins/Runtime/HairStrands/Source/HairStrandsCore/Private/HairStrandsRendering.cpp b/Engine/Plugins/Runtime/HairStrands/Source/HairStrandsCore/Private/HairStrandsRendering.cpp
--- a/Engine/Plugins/Runtime/HairStrands/Source/HairStrandsCore/Private/HairStrandsRendering.cpp
+++ b/Engine/Plugins/Runtime/HairStrands/Source/HairStrandsCore/Private/HairStrandsRendering.cpp
@@ -1056,7 +1056,9 @@ static void AddHairCardsDeformationPass(
 
 	FHairGroupInstance::FCards::FLOD& LOD = Instance->Cards.LODs[HairLODIndex];
 	
-	FRDGImportedBuffer CardsDeformedPositionBuffer = Register(GraphBuilder, LOD.DeformedResource->GetBuffer(FHairCardsDeformedResource::Current), ERDGImportedBufferFlags::CreateUAV);
+	FRDGImportedBuffer CardsDeformedPositionBuffer_Curr = Register(GraphBuilder, LOD.DeformedResource->GetBuffer(FHairCardsDeformedResource::Current), ERDGImportedBufferFlags::CreateUAV);
+	FRDGImportedBuffer CardsDeformedPositionBuffer_Prev = Register(GraphBuilder, LOD.DeformedResource->GetBuffer(FHairCardsDeformedResource::Previous), ERDGImportedBufferFlags::CreateUAV);
+	AddCopyBufferPass(GraphBuilder, CardsDeformedPositionBuffer_Prev.Buffer, CardsDeformedPositionBuffer_Curr.Buffer);
 
 	FHairCardsDeformationCS::FParameters* Parameters = GraphBuilder.AllocParameters<FHairCardsDeformationCS::FParameters>();
 	Parameters->GuideVertexCount			= LOD.Guides.RestResource->GetVertexCount();
@@ -1067,7 +1069,7 @@ static void AddHairCardsDeformationPass(
 	
 	Parameters->CardsVertexCount			= LOD.RestResource->GetVertexCount();
 	Parameters->CardsRestPositionBuffer		= LOD.RestResource->RestPositionBuffer.ShaderResourceViewRHI;
-	Parameters->CardsDeformedPositionBuffer = CardsDeformedPositionBuffer.UAV;
+	Parameters->CardsDeformedPositionBuffer = CardsDeformedPositionBuffer_Curr.UAV;
 
 	Parameters->CardsInterpolationBuffer	= RegisterAsSRV(GraphBuilder, LOD.InterpolationResource->InterpolationBuffer);
 
@@ -1122,7 +1124,8 @@ static void AddHairCardsDeformationPass(
 		Parameters,
 		FIntVector(DispatchCountX,1,1));
 
-	GraphBuilder.SetBufferAccessFinal(CardsDeformedPositionBuffer.Buffer, ERHIAccess::SRVMask);
+	GraphBuilder.SetBufferAccessFinal(CardsDeformedPositionBuffer_Curr.Buffer, ERHIAccess::SRVMask);
+	GraphBuilder.SetBufferAccessFinal(CardsDeformedPositionBuffer_Prev.Buffer, ERHIAccess::SRVMask);
 }
 
 ///////////////////////////////////////////////////////////////////////////////////////////////////
diff --git a/Engine/Plugins/Runtime/HairStrands/Source/HairStrandsCore/Public/GroomInstance.h b/Engine/Plugins/Runtime/HairStrands/Source/HairStrandsCore/Public/GroomInstance.h
--- a/Engine/Plugins/Runtime/HairStrands/Source/HairStrandsCore/Public/GroomInstance.h
+++ b/Engine/Plugins/Runtime/HairStrands/Source/HairStrandsCore/Public/GroomInstance.h
@@ -134,12 +134,11 @@ struct HAIRSTRANDSCORE_API FHairGroupInstance : public FHairStrandsInstance
 			FStrandsBaseWithInterpolation Guides;
 
 			UMaterialInterface* Material = nullptr;
-			FHairCardsUniformBuffer UniformBuffer[2];
-			FHairCardsVertexFactory* VertexFactory[2] = { nullptr, nullptr };
+			FHairCardsUniformBuffer UniformBuffer;
+			FHairCardsVertexFactory* VertexFactory = nullptr;
 			FHairCardsVertexFactory* GetVertexFactory() const
 			{
-				const uint32 Index = DeformedResource ? DeformedResource->CurrentIndex : 0;
-				return VertexFactory[Index];
+				return VertexFactory;
 			}
 			void InitVertexFactory();
 
@@ -169,12 +168,11 @@ struct HAIRSTRANDSCORE_API FHairGroupInstance : public FHairStrandsInstance
 			#endif
 
 			UMaterialInterface* Material = nullptr;
-			FHairCardsUniformBuffer UniformBuffer[2];
-			FHairCardsVertexFactory* VertexFactory[2] = { nullptr, nullptr };
+			FHairCardsUniformBuffer UniformBuffer;
+			FHairCardsVertexFactory* VertexFactory = nullptr;
 			FHairCardsVertexFactory* GetVertexFactory() const 
 			{
-				const uint32 Index = DeformedResource ? DeformedResource->CurrentIndex : 0;
-				return VertexFactory[Index];
+				return VertexFactory;
 			}
 			void InitVertexFactory();
 		};
diff --git a/Engine/Plugins/Runtime/HairStrands/Source/HairStrandsCore/Public/GroomResources.h b/Engine/Plugins/Runtime/HairStrands/Source/HairStrandsCore/Public/GroomResources.h
--- a/Engine/Plugins/Runtime/HairStrands/Source/HairStrandsCore/Public/GroomResources.h
+++ b/Engine/Plugins/Runtime/HairStrands/Source/HairStrandsCore/Public/GroomResources.h
@@ -612,9 +612,6 @@ struct FHairCardsDeformedResource : public FHairCommonResource
 	/* Whether the GPU data should be initialized with the asset data or not */
 	const bool bInitializedData = false;
 
-	/* Whether the GPU data should be initialized with the asset data or not */
-	uint32 CurrentIndex = 0;
-
 	enum EFrameType
 	{
 		Previous,
@@ -622,9 +619,8 @@ struct FHairCardsDeformedResource : public FHairCommonResource
 	};
 
 	// Helper accessors
-	inline uint32 GetIndex(EFrameType T)				{ return T == EFrameType::Current ? CurrentIndex : 1u - CurrentIndex; }
+	inline uint32 GetIndex(EFrameType T)				{ return T == EFrameType::Current ? 0u : 1u; }
 	inline FRDGExternalBuffer& GetBuffer(EFrameType T)	{ return DeformedPositionBuffer[GetIndex(T)];  }
-	inline void SwapBuffer()							{ CurrentIndex = 1u - CurrentIndex; }
 };
 
 struct FHairCardsInterpolationBulkData;
diff --git a/Engine/Plugins/Runtime/HairStrands/Source/HairStrandsCore/Public/HairCardsVertexFactory.h b/Engine/Plugins/Runtime/HairStrands/Source/HairStrandsCore/Public/HairCardsVertexFactory.h
--- a/Engine/Plugins/Runtime/HairStrands/Source/HairStrandsCore/Public/HairCardsVertexFactory.h
+++ b/Engine/Plugins/Runtime/HairStrands/Source/HairStrandsCore/Public/HairCardsVertexFactory.h
@@ -28,11 +28,10 @@ public:
 		FHairGroupInstance* Instance = nullptr;
 		uint32 GroupIndex = 0;
 		uint32 LODIndex = 0;
-		uint32 BufferIndex = 0;// If the RHI does not support manual fetch, we create two vertex factory to have two declarations: 1) for position0:current & position1:previous 2) for position0:previous & position1:current
 		EHairGeometryType GeometryType = EHairGeometryType::NoneGeometry;
 	};
 
-	FHairCardsVertexFactory(FHairGroupInstance* Instance, uint32 GroupIndex, uint32 LODIndex, uint32 BufferIndex, EHairGeometryType GeometryType, EShaderPlatform InShaderPlatform, ERHIFeatureLevel::Type InFeatureLevel, const char* InDebugName);
+	FHairCardsVertexFactory(FHairGroupInstance* Instance, uint32 GroupIndex, uint32 LODIndex, EHairGeometryType GeometryType, EShaderPlatform InShaderPlatform, ERHIFeatureLevel::Type InFeatureLevel, const char* InDebugName);
 	
 	/**
 	 * Should we cache the material's shadertype on this platform with this vertex factory? 
