Enable async loading for hair data.

diff --git a/Engine/Plugins/Runtime/HairStrands/Source/HairStrandsCore/Private/GroomResources.cpp b/Engine/Plugins/Runtime/HairStrands/Source/HairStrandsCore/Private/GroomResources.cpp
--- a/Engine/Plugins/Runtime/HairStrands/Source/HairStrandsCore/Private/GroomResources.cpp
+++ b/Engine/Plugins/Runtime/HairStrands/Source/HairStrandsCore/Private/GroomResources.cpp
@@ -41,7 +41,11 @@ static FAutoConsoleVariableRef CVarHairStrandsDebugVoxel_MaxSegmentPerVoxel(TEXT
 
 EHairResourceLoadingType GetHairResourceLoadingType()
 {
+#if !WITH_EDITORONLY_DATA
 	return GHairStrandsBulkData_AsyncLoading > 0 ? EHairResourceLoadingType::Async : EHairResourceLoadingType::Sync;
+#else
+	return EHairResourceLoadingType::Sync;
+#endif
 }
 
 namespace HairTransition
