Fix material proxy validation error.

This happens due to a race condition, when a component needs to recreate its proxy, but the current one is still used/in-flight for the current mesh batch generation.

diff --git a/Engine/Plugins/Runtime/HairStrands/Source/HairStrandsCore/Private/GroomComponent.cpp b/Engine/Plugins/Runtime/HairStrands/Source/HairStrandsCore/Private/GroomComponent.cpp
--- a/Engine/Plugins/Runtime/HairStrands/Source/HairStrandsCore/Private/GroomComponent.cpp
+++ b/Engine/Plugins/Runtime/HairStrands/Source/HairStrandsCore/Private/GroomComponent.cpp
@@ -294,6 +294,7 @@ public:
 		bVFRequiresPrimitiveUniformBuffer = true;
 		bCastDeepShadow = true;
 
+		HairGroupMaterialProxies.SetNum(Component->HairGroupInstances.Num());
 		HairGroupInstances = Component->HairGroupInstances;
 		check(Component);
 		check(Component->GroomAsset);
@@ -348,13 +349,11 @@ public:
 			if (HairInstance->Meshes.IsValid())
 			{
 				const uint32 LODCount = HairInstance->Meshes.LODs.Num();
-
 				for (uint32 LODIt = 0; LODIt < LODCount; ++LODIt)
 				{
 					if (HairInstance->Meshes.IsValid(LODIt))
 					{
 						if (HairInstance->Meshes.LODs[LODIt].VertexFactory == nullptr) HairInstance->Meshes.LODs[LODIt].VertexFactory = new FHairCardsVertexFactory(HairInstance, GroupIt, LODIt, EHairGeometryType::Meshes, ShaderPlatform, FeatureLevel, "HairMeshesVertexFactory");
-
 					}
 				}
 			}
@@ -371,10 +370,12 @@ public:
 			if (IsHairStrandsEnabled(EHairStrandsShaderType::Strands, ShaderPlatform))
 			{
 				const int32 SlotIndex = Component->GroomAsset->GetMaterialIndex(Component->GroomAsset->HairGroupsRendering[GroupIt].MaterialSlotName);
-				HairInstance->Strands.Material = Component->GetMaterial(GetMaterialIndexWithFallback(SlotIndex), EHairGeometryType::Strands, true);
+				const UMaterialInterface* Material = Component->GetMaterial(GetMaterialIndexWithFallback(SlotIndex), EHairGeometryType::Strands, true);
+				HairGroupMaterialProxies[GroupIt].Strands = Material ? Material->GetRenderProxy() : nullptr;
 			}
 
 			// Material - Cards
+			HairGroupMaterialProxies[GroupIt].Cards.Init(nullptr, InGroupData.Cards.LODs.Num());
 			if (IsHairStrandsEnabled(EHairStrandsShaderType::Cards, ShaderPlatform))
 			{
 				uint32 CardsLODIndex = 0;
@@ -392,13 +393,15 @@ public:
 								break;
 							}
 						}
-						HairInstance->Cards.LODs[CardsLODIndex].Material = Component->GetMaterial(GetMaterialIndexWithFallback(SlotIndex), EHairGeometryType::Cards, true);
+						const UMaterialInterface* Material = Component->GetMaterial(GetMaterialIndexWithFallback(SlotIndex), EHairGeometryType::Cards, true);
+						HairGroupMaterialProxies[GroupIt].Cards[CardsLODIndex] = Material ? Material->GetRenderProxy() : nullptr;
 					}
 					++CardsLODIndex;
 				}
 			}
 
 			// Material - Meshes
+			HairGroupMaterialProxies[GroupIt].Meshes.Init(nullptr, InGroupData.Meshes.LODs.Num());
 			if (IsHairStrandsEnabled(EHairStrandsShaderType::Meshes, ShaderPlatform))
 			{
 				uint32 MeshesLODIndex = 0;
@@ -416,7 +419,8 @@ public:
 								break;
 							}
 						}
-						HairInstance->Meshes.LODs[MeshesLODIndex].Material = Component->GetMaterial(GetMaterialIndexWithFallback(SlotIndex), EHairGeometryType::Meshes, true);
+						const UMaterialInterface* Material = Component->GetMaterial(GetMaterialIndexWithFallback(SlotIndex), EHairGeometryType::Meshes, true);
+						HairGroupMaterialProxies[GroupIt].Meshes[MeshesLODIndex] = Material ? Material->GetRenderProxy() : nullptr;
 					}
 					++MeshesLODIndex;
 				}
@@ -724,7 +728,7 @@ public:
 
 		const FVertexFactory* VertexFactory = nullptr;
 		FIndexBuffer* IndexBuffer = nullptr;
-		FMaterialRenderProxy* MaterialRenderProxy = Debug_MaterialProxy;
+		const FMaterialRenderProxy* MaterialRenderProxy = Debug_MaterialProxy;
 
 		uint32 NumPrimitive = 0;
 		uint32 HairVertexCount = 0;
@@ -744,9 +748,9 @@ public:
 			NumPrimitive = HairVertexCount / 3;
 			IndexBuffer = &Instance->Meshes.LODs[IntLODIndex].RestResource->IndexBuffer;
 			bUseCulling = false;
-			if (MaterialRenderProxy == nullptr && Instance->Meshes.LODs[IntLODIndex].Material)
+			if (MaterialRenderProxy == nullptr)
 			{
-				MaterialRenderProxy = Instance->Meshes.LODs[IntLODIndex].Material->GetRenderProxy();
+				MaterialRenderProxy = HairGroupMaterialProxies[GroupIndex].Meshes[IntLODIndex];
 			}
 			bWireframe = AllowDebugViewmodes() && ViewFamily.EngineShowFlags.Wireframe;
 		}
@@ -764,9 +768,9 @@ public:
 			NumPrimitive = HairVertexCount / 3;
 			IndexBuffer = &Instance->Cards.LODs[IntLODIndex].RestResource->RestIndexBuffer;
 			bUseCulling = false;
-			if (MaterialRenderProxy == nullptr && Instance->Cards.LODs[IntLODIndex].Material)
+			if (MaterialRenderProxy == nullptr)
 			{
-				MaterialRenderProxy = Instance->Cards.LODs[IntLODIndex].Material->GetRenderProxy();
+				MaterialRenderProxy = HairGroupMaterialProxies[GroupIndex].Cards[IntLODIndex];
 			}
 			bWireframe = AllowDebugViewmodes() && ViewFamily.EngineShowFlags.Wireframe;
 		}
@@ -781,9 +785,9 @@ public:
 			#endif
 			bUseCulling = Instance->Strands.bIsCullingEnabled;
 			NumPrimitive = bUseCulling ? 0 : HairVertexCount * 2;
-			if (MaterialRenderProxy == nullptr && Instance->Strands.Material)
+			if (MaterialRenderProxy == nullptr)
 			{
-				MaterialRenderProxy = Instance->Strands.Material->GetRenderProxy();
+				MaterialRenderProxy = HairGroupMaterialProxies[GroupIndex].Strands;
 			}
 			bWireframe = AllowDebugViewmodes() && ViewFamily.EngineShowFlags.Wireframe;
 		}
@@ -801,8 +805,9 @@ public:
 
 		if (bWireframe)
 		{
-			MaterialRenderProxy = new FColoredMaterialRenderProxy( GEngine->WireframeMaterial ? GEngine->WireframeMaterial->GetRenderProxy() : NULL, FLinearColor(1.f, 0.5f, 0.f));
-			Collector.RegisterOneFrameMaterialProxy(MaterialRenderProxy);
+			FMaterialRenderProxy* ColoredMaterialRenderProxy = new FColoredMaterialRenderProxy( GEngine->WireframeMaterial ? GEngine->WireframeMaterial->GetRenderProxy() : NULL, FLinearColor(1.f, 0.5f, 0.f));
+			Collector.RegisterOneFrameMaterialProxy(ColoredMaterialRenderProxy);
+			MaterialRenderProxy = ColoredMaterialRenderProxy;
 		}
 
 		// Draw the mesh.
@@ -910,6 +915,16 @@ public:
 
 	TArray<FHairGroupInstance*> HairGroupInstances;
 
+	// Cache the material proxy to avoid race condition, when groom component's proxy is recreated, 
+	// while another one is currently in flight for drawing.
+	struct FHairGroupMaterialProxy
+	{
+		const FMaterialRenderProxy* Strands = nullptr;
+		TArray<const FMaterialRenderProxy*> Cards;
+		TArray<const FMaterialRenderProxy*> Meshes;
+	};
+	TArray<FHairGroupMaterialProxy> HairGroupMaterialProxies;
+	
 private:
 	uint32 ComponentId = 0;
 	FMaterialRelevance MaterialRelevance;
@@ -2448,10 +2463,6 @@ void UGroomComponent::InitResources(bool bIsBindingReloading)
 				check(HairGroupInstance->Strands.InterpolationResource);
 			}
 
-			// Material
-			HairGroupInstance->Strands.Material = nullptr;
-
-
 			// Check if strands needs to load/allocate root resources
 			bool bNeedRootResources = false;
 			bool bNeedDynamicResources = false;
@@ -2607,9 +2618,6 @@ void UGroomComponent::InitResources(bool bIsBindingReloading)
 				}
 				#endif
 
-				// Material
-				InstanceLOD.Material = nullptr;
-
 				// Strands data/resources
 				if (bNeedDeformedPositions)
 				{
@@ -2677,9 +2685,6 @@ void UGroomComponent::InitResources(bool bIsBindingReloading)
 					}
 				}
 				#endif
-
-				// Material
-				InstanceLOD.Material = nullptr;
 			}
 		}
 	}
diff --git a/Engine/Plugins/Runtime/HairStrands/Source/HairStrandsCore/Public/GroomInstance.h b/Engine/Plugins/Runtime/HairStrands/Source/HairStrandsCore/Public/GroomInstance.h
--- a/Engine/Plugins/Runtime/HairStrands/Source/HairStrandsCore/Public/GroomInstance.h
+++ b/Engine/Plugins/Runtime/HairStrands/Source/HairStrandsCore/Public/GroomInstance.h
@@ -103,7 +103,6 @@ struct HAIRSTRANDSCORE_API FHairGroupInstance : public FHairStrandsInstance
 		FRDGExternalBuffer DebugAttributeBuffer;
 		FHairGroupInstanceModifer Modifier;
 
-		UMaterialInterface* Material = nullptr;
 		FHairStrandsVertexFactory* VertexFactory = nullptr;
 	} Strands;
 
@@ -133,7 +132,6 @@ struct HAIRSTRANDSCORE_API FHairGroupInstance : public FHairStrandsInstance
 
 			FStrandsBaseWithInterpolation Guides;
 
-			UMaterialInterface* Material = nullptr;
 			FHairCardsUniformBuffer UniformBuffer;
 			FHairCardsVertexFactory* VertexFactory = nullptr;
 			FHairCardsVertexFactory* GetVertexFactory() const
@@ -167,7 +165,6 @@ struct HAIRSTRANDSCORE_API FHairGroupInstance : public FHairStrandsInstance
 			bool RaytracingResourceOwned = false;
 			#endif
 
-			UMaterialInterface* Material = nullptr;
 			FHairCardsUniformBuffer UniformBuffer;
 			FHairCardsVertexFactory* VertexFactory = nullptr;
 			FHairCardsVertexFactory* GetVertexFactory() const 
