Change some hair cards resources to be deferred initialized.

diff --git a/Engine/Plugins/Runtime/HairStrands/Source/HairStrandsCore/Private/GroomAsset.cpp b/Engine/Plugins/Runtime/HairStrands/Source/HairStrandsCore/Private/GroomAsset.cpp
--- a/Engine/Plugins/Runtime/HairStrands/Source/HairStrandsCore/Private/GroomAsset.cpp
+++ b/Engine/Plugins/Runtime/HairStrands/Source/HairStrandsCore/Private/GroomAsset.cpp
@@ -868,7 +868,6 @@ void UGroomAsset::UpdateResource()
 					HairGroupsData);
 
 				GroupData.Strands.ClusterCullingResource = new FHairStrandsClusterCullingResource(GroupData.Strands.ClusterCullingBulkData, FHairResourceName(GetFName(), GroupIndex));
-				BeginInitResource(GroupData.Strands.ClusterCullingResource);
 			}
 			else
 			{
@@ -2543,7 +2542,7 @@ bool UGroomAsset::BuildCardsGeometry(uint32 GroupIndex)
 				FHairResourceName ResourceName(GetFName(), GroupIndex, LODIt);
 
 				LOD.RestResource = new FHairCardsRestResource(LOD.BulkData, ResourceName);
-				BeginInitResource(LOD.RestResource);
+				BeginInitResource(LOD.RestResource); // Immediate allocation, as needed for the vertex factory, input stream building
 
 				// 2.1 Load atlas textures
 				InitAtlasTexture(LOD.RestResource, Desc->Textures.DepthTexture, EHairAtlasTextureType::Depth);
@@ -2555,7 +2554,6 @@ bool UGroomAsset::BuildCardsGeometry(uint32 GroupIndex)
 				
 				// 2.2 Load interoplatino resources
 				LOD.InterpolationResource = new FHairCardsInterpolationResource(LOD.InterpolationBulkData, ResourceName);
-				BeginInitResource(LOD.InterpolationResource);
 
 				// Create own interpolation settings for cards.
 				// Force closest guides as this is the most relevant matching metric for cards, due to their coarse geometry
@@ -2600,10 +2598,8 @@ bool UGroomAsset::BuildCardsGeometry(uint32 GroupIndex)
 				}
 
 				LOD.Guides.RestResource = new FHairStrandsRestResource(LOD.Guides.BulkData, EHairStrandsResourcesType::Cards, ResourceName);
-				BeginInitResource(LOD.Guides.RestResource);
 
 				LOD.Guides.InterpolationResource = new FHairStrandsInterpolationResource(LOD.Guides.InterpolationBulkData, ResourceName);
-				BeginInitResource(LOD.Guides.InterpolationResource);
 
 				// Update card stats to display
 				Desc->CardsInfo.NumCardVertices = LOD.BulkData.GetNumVertices();
@@ -2815,7 +2811,6 @@ FHairStrandsRestResource* UGroomAsset::AllocateGuidesResources(uint32 GroupIndex
 			if (GroupData.Guides.RestResource == nullptr)
 			{
 				GroupData.Guides.RestResource = new FHairStrandsRestResource(GroupData.Guides.BulkData, EHairStrandsResourcesType::Guides, FHairResourceName(GetFName(), GroupIndex));
-				BeginInitResource(GroupData.Guides.RestResource);
 			}
 			return GroupData.Guides.RestResource;
 		}
@@ -2832,7 +2827,6 @@ FHairStrandsInterpolationResource* UGroomAsset::AllocateInterpolationResources(u
 		if (GroupData.Strands.InterpolationResource == nullptr)
 		{
 			GroupData.Strands.InterpolationResource = new FHairStrandsInterpolationResource(GroupData.Strands.InterpolationBulkData, FHairResourceName(GetFName(), GroupIndex));
-			BeginInitResource(GroupData.Strands.InterpolationResource);
 		}
 		return GroupData.Strands.InterpolationResource;
 	}
@@ -2914,7 +2908,6 @@ void UGroomAsset::InitStrandsResources()
 			FHairResourceName ResourceName(GetFName(), GroupIndex);
 
 			GroupData.Strands.RestResource = new FHairStrandsRestResource(GroupData.Strands.BulkData, EHairStrandsResourcesType::Strands, ResourceName);
-			BeginInitResource(GroupData.Strands.RestResource);
 
 			if (GroupData.Strands.ClusterCullingBulkData.IsValid())
 			{
@@ -2936,7 +2929,6 @@ void UGroomAsset::InitStrandsResources()
 				}
 
 				GroupData.Strands.ClusterCullingResource = new FHairStrandsClusterCullingResource(GroupData.Strands.ClusterCullingBulkData, ResourceName);
-				BeginInitResource(GroupData.Strands.ClusterCullingResource);
 			}
 
 			// Interpolation are lazy allocated, as these resources are only used when RBF/simulation is enabled by a groom component
@@ -3022,16 +3014,13 @@ void UGroomAsset::InitCardsResources()
 				FHairResourceName ResourceName(GetFName(), GroupIndex, LODIt);
 
 				LOD.RestResource = new FHairCardsRestResource(LOD.BulkData, ResourceName);
-				BeginInitResource(LOD.RestResource);
+				BeginInitResource(LOD.RestResource); // Immediate allocation, as needed for the vertex factory, input stream building
 
 				LOD.InterpolationResource = new FHairCardsInterpolationResource(LOD.InterpolationBulkData, ResourceName);
-				BeginInitResource(LOD.InterpolationResource);
 
 				LOD.Guides.RestResource = new FHairStrandsRestResource(LOD.Guides.BulkData, EHairStrandsResourcesType::Cards, ResourceName);
-				BeginInitResource(LOD.Guides.RestResource);
 
 				LOD.Guides.InterpolationResource = new FHairStrandsInterpolationResource(LOD.Guides.InterpolationBulkData, ResourceName);
-				BeginInitResource(LOD.Guides.InterpolationResource);
 
 				if (Desc)
 				{
@@ -3082,7 +3071,8 @@ void UGroomAsset::InitMeshesResources()
 			if (LOD.HasValidData())
 			{
 				LOD.RestResource = new FHairMeshesRestResource(LOD.BulkData, FHairResourceName(GetFName(), GroupIndex, LODIt));
-				BeginInitResource(LOD.RestResource);
+				BeginInitResource(LOD.RestResource); // Immediate allocation, as needed for the vertex factory, input stream building
+
 
 				if (Desc)
 				{
@@ -3505,9 +3495,9 @@ void UGroomAsset::SaveProceduralCards(uint32 DescIndex)
 
 	// 3. Create resources and enqueue texture generation (GPU, kicked by the render thread) 
 	Q.Resources = new FHairCardsRestResource(Q.BulkData, FHairResourceName(GetFName(), GroupIndex, LODIndex));
-	BeginInitResource(Q.Resources);
+	BeginInitResource(Q.Resources); // Immediate allocation, as needed for the vertex factory, input stream building
 	Q.ProceduralResources = new FHairCardsProceduralResource(Q.ProceduralData.RenderData, Q.ProceduralData.Atlas.Resolution, Q.ProceduralData.Voxels);
-	BeginInitResource(Q.ProceduralResources);
+	BeginInitResource(Q.ProceduralResources); // Immediate allocation, as needed for the vertex factory, input stream building
 
 	FHairCardsBuilder::BuildTextureAtlas(&Q.ProceduralData, Q.Resources, Q.ProceduralResources, Q.Textures);
 
diff --git a/Engine/Plugins/Runtime/HairStrands/Source/HairStrandsCore/Private/GroomResources.cpp b/Engine/Plugins/Runtime/HairStrands/Source/HairStrandsCore/Private/GroomResources.cpp
--- a/Engine/Plugins/Runtime/HairStrands/Source/HairStrandsCore/Private/GroomResources.cpp
+++ b/Engine/Plugins/Runtime/HairStrands/Source/HairStrandsCore/Private/GroomResources.cpp
@@ -1481,7 +1481,7 @@ void FHairCardsInterpolationBulkData::Serialize(FArchive& Ar)
 }
 
 FHairCardsInterpolationResource::FHairCardsInterpolationResource(FHairCardsInterpolationBulkData& InBulkData, const FHairResourceName& InResourceName) :
-	FHairCommonResource(EHairStrandsAllocationType::Immediate, InResourceName),
+	FHairCommonResource(EHairStrandsAllocationType::Deferred, InResourceName),
 	InterpolationBuffer(), BulkData(InBulkData)
 {
 }
@@ -1528,6 +1528,7 @@ void FHairStrandsRaytracingResource::InternalRelease()
 {
 	PositionBuffer.Release();
 	RayTracingGeometry.ReleaseResource();
+	bIsRTGeometryInitialized = false;
 }
 #endif
 
