Change hair resources to use RDG upload buffer pass, which doesn't require a copy of the bulk data.

diff --git a/Engine/Plugins/Runtime/HairStrands/Source/HairStrandsCore/Private/GroomResources.cpp b/Engine/Plugins/Runtime/HairStrands/Source/HairStrandsCore/Private/GroomResources.cpp
--- a/Engine/Plugins/Runtime/HairStrands/Source/HairStrandsCore/Private/GroomResources.cpp
+++ b/Engine/Plugins/Runtime/HairStrands/Source/HairStrandsCore/Private/GroomResources.cpp
@@ -194,14 +194,29 @@ void InternalCreateVertexBufferRDG_FromBulkData(FRDGBuilder& GraphBuilder, FByte
 	#endif
 
 	const typename FormatType::BulkType* BulkData = (const typename FormatType::BulkType*)InBulkData.Lock(LOCK_READ_ONLY);
-	FRDGBufferRef Buffer = InternalCreateVertexBuffer(
-		GraphBuilder,
-		DebugName,
-		Desc,
-		BulkData,
-		DataSizeInBytes,
-		ERDGInitialDataFlags::None); // Copy data internally
-	InBulkData.Unlock();
+	FRDGBufferRef Buffer = nullptr;
+	#if !WITH_EDITORONLY_DATA
+	if (bReleaseCPUData)
+	{
+		checkf(Desc.UnderlyingType == FRDGBufferDesc::EUnderlyingType::VertexBuffer, TEXT("CreateVertexBuffer called with an FRDGBufferDesc underlying type that is not 'VertexBuffer'. Buffer: %s"), DebugName);
+		Buffer = GraphBuilder.CreateBuffer(Desc, DebugName, ERDGBufferFlags::MultiFrame);
+		if (BulkData && DataSizeInBytes)
+		{
+			GraphBuilder.QueueBufferUpload(Buffer, BulkData, DataSizeInBytes, [&InBulkData](const void* Ptr) { InBulkData.Unlock(); });
+		}
+	}
+	else
+	#endif
+	{
+		Buffer = InternalCreateVertexBuffer(
+			GraphBuilder,
+			DebugName,
+			Desc,
+			BulkData,
+			DataSizeInBytes,
+			ERDGInitialDataFlags::None); // Copy data internally
+		InBulkData.Unlock();
+	}
 
 	ConvertToExternalBufferWithViews(GraphBuilder, Buffer, Out, FormatType::Format);
 }
@@ -345,31 +360,6 @@ void InternalCreateVertexBufferRDG(FRDGBuilder& GraphBuilder, uint32 InVertexCou
 	ConvertToExternalBufferWithViews(GraphBuilder, Buffer, Out, FormatType::Format);
 }
 
-template<typename FormatType>
-void InternalCreateStructuredBufferRDG(FRDGBuilder& GraphBuilder, const TArray<FormatType>& InData, uint32 InSizeInByte, FRDGExternalBuffer& Out, const TCHAR* DebugName, EHairResourceUsageType UsageType)
-{
-	const uint32 DataCount = InData.Num();
-	const uint32 DataSizeInBytes = InSizeInByte * DataCount;
-	if (DataSizeInBytes == 0)
-	{
-		Out.Buffer = nullptr;
-		return;
-	}
-	
-	FRDGBufferDesc Desc = FRDGBufferDesc::CreateStructuredDesc(InSizeInByte, DataCount);
-	if (UsageType != EHairResourceUsageType::Dynamic)
-	{
-		Desc.Usage = Desc.Usage & (~BUF_UnorderedAccess);
-	}
-
-	FRDGBufferRef Buffer = GraphBuilder.CreateBuffer(Desc, DebugName, ERDGBufferFlags::MultiFrame);
-	if (InData.GetData() && DataSizeInBytes)
-	{
-		GraphBuilder.QueueBufferUpload(Buffer, InData.GetData(), DataSizeInBytes, ERDGInitialDataFlags::NoCopy);
-	}
-	ConvertToExternalBufferWithViews(GraphBuilder, Buffer, Out);
-}
-
 template<typename FormatType>
 void InternalCreateStructuredBufferRDG_FromBulkData(FRDGBuilder& GraphBuilder, FByteBulkData& InBulkData, uint32 InDataCount, FRDGExternalBuffer& Out, const TCHAR* DebugName, EHairResourceUsageType UsageType)
 {
@@ -402,10 +392,18 @@ void InternalCreateStructuredBufferRDG_FromBulkData(FRDGBuilder& GraphBuilder, F
 	FRDGBufferRef Buffer = GraphBuilder.CreateBuffer(Desc, DebugName, ERDGBufferFlags::MultiFrame);
 	if (Data && DataSizeInBytes)
 	{
-		GraphBuilder.QueueBufferUpload(Buffer, Data, DataSizeInBytes, ERDGInitialDataFlags::None);  // Copy data internally
+		#if !WITH_EDITORONLY_DATA
+		if (bReleaseCPUData)
+		{
+			GraphBuilder.QueueBufferUpload(Buffer, Data, DataSizeInBytes, [&InBulkData](const void* Ptr) { InBulkData.Unlock(); });
+		}
+		else
+		#endif
+		{
+			GraphBuilder.QueueBufferUpload(Buffer, Data, DataSizeInBytes, ERDGInitialDataFlags::None);  // Copy data internally
+			InBulkData.Unlock();
+		}
 	}
-	InBulkData.Unlock();
-
 
 	ConvertToExternalBufferWithViews(GraphBuilder, Buffer, Out);
 }
