Manual migration from 4.27.1 stream of CL 17367349: Fix for inaccurate RBF projection due to floating point error caused by large skeletal mesh translations from the root in animations.
NB also added requested useful code comments to explain why offset was added.

diff --git a/Engine/Plugins/Runtime/HairStrands/Shaders/Private/NiagaraDataInterfaceHairStrands.ush b/Engine/Plugins/Runtime/HairStrands/Shaders/Private/NiagaraDataInterfaceHairStrands.ush
--- a/Engine/Plugins/Runtime/HairStrands/Shaders/Private/NiagaraDataInterfaceHairStrands.ush
+++ b/Engine/Plugins/Runtime/HairStrands/Shaders/Private/NiagaraDataInterfaceHairStrands.ush
@@ -878,7 +878,7 @@ float3 DIHairStrands_ApplyRBF(in FDIHairStrandsContext DIContext, in float3 Rest
 	ControlPoint += DIContext.MeshSampleWeightsBuffer[DIContext.SampleCount + 1].xyz * RestControlPoint.x;
 	ControlPoint += DIContext.MeshSampleWeightsBuffer[DIContext.SampleCount + 2].xyz * RestControlPoint.y;
 	ControlPoint += DIContext.MeshSampleWeightsBuffer[DIContext.SampleCount + 3].xyz * RestControlPoint.z;
-	return ControlPoint;
+	return ControlPoint + DIContext.MeshSampleWeightsBuffer[DIContext.SampleCount + 4].xyz; // offset added to workaround numerical issues when deformed and rest positions are far apart
 }
 
  // Compute the global position
diff --git a/Engine/Plugins/Runtime/HairStrands/Source/HairStrandsCore/Private/GroomRBFDeformer.cpp b/Engine/Plugins/Runtime/HairStrands/Source/HairStrandsCore/Private/GroomRBFDeformer.cpp
--- a/Engine/Plugins/Runtime/HairStrands/Source/HairStrandsCore/Private/GroomRBFDeformer.cpp
+++ b/Engine/Plugins/Runtime/HairStrands/Source/HairStrandsCore/Private/GroomRBFDeformer.cpp
@@ -39,7 +39,9 @@ void UpdateMeshSamples(
 	TArray<FVector3f>& OutSampleDeformationsBuffer
 )
 {
-	OutSampleDeformationsBuffer.SetNum(MaxSampleCount + 4);
+	OutSampleDeformationsBuffer.SetNum(MaxSampleCount + 5);
+	OutSampleDeformationsBuffer[MaxSampleCount + 4] = SampleDeformedPositionsBuffer[0]; // offset used to workaround numerical issues when deformed and rest positions are far apart
+
 	for (uint32 SampleIndex = 0; SampleIndex < MaxSampleCount + 4; ++SampleIndex)
 	{
 		uint32 WeightsOffset = SampleIndex * (MaxSampleCount + 4);
@@ -47,7 +49,7 @@ void UpdateMeshSamples(
 		for (uint32 i = 0; i < MaxSampleCount; ++i, ++WeightsOffset)
 		{
 			SampleDeformation += InterpolationWeightsBuffer[WeightsOffset] *
-				(SampleDeformedPositionsBuffer[i] - SampleRestPositionsBuffer[i]);
+				(SampleDeformedPositionsBuffer[i] - (SampleRestPositionsBuffer[i] + SampleDeformedPositionsBuffer[0]) ); // offset added to workaround numerical issues when deformed and rest positions are far apart
 		}
 
 		OutSampleDeformationsBuffer[SampleIndex] = SampleDeformation;
@@ -75,7 +77,7 @@ FVector3f DisplacePosition(
 	ControlPoint += MeshSampleWeightsBuffer[SampleCount + 1] * RestControlPoint.X;
 	ControlPoint += MeshSampleWeightsBuffer[SampleCount + 2] * RestControlPoint.Y;
 	ControlPoint += MeshSampleWeightsBuffer[SampleCount + 3] * RestControlPoint.Z;
-	return ControlPoint;
+	return ControlPoint + MeshSampleWeightsBuffer[SampleCount + 4]; // offset added to workaround numerical issues when deformed and rest positions are far apart
 }
 
 void DeformStrands(	
diff --git a/Engine/Plugins/Runtime/HairStrands/Source/HairStrandsCore/Private/GroomResources.cpp b/Engine/Plugins/Runtime/HairStrands/Source/HairStrandsCore/Private/GroomResources.cpp
--- a/Engine/Plugins/Runtime/HairStrands/Source/HairStrandsCore/Private/GroomResources.cpp
+++ b/Engine/Plugins/Runtime/HairStrands/Source/HairStrandsCore/Private/GroomResources.cpp
@@ -1223,7 +1223,7 @@ void FHairStrandsDeformedRootResource::InternalAllocateLOD(FRDGBuilder& GraphBui
 			if (LOD.SampleCount > 0)
 			{
 				InternalCreateVertexBufferRDG<FHairStrandsMeshTrianglePositionFormat>(GraphBuilder, LOD.SampleCount, LOD.DeformedSamplePositionsBuffer, ToHairResourceDebugName(HAIRSTRANDS_RESOUCE_NAME(CurveType, Hair.StrandsRootDeformed_DeformedSamplePositionsBuffer), ResourceName), EHairResourceUsageType::Dynamic);
-				InternalCreateVertexBufferRDG<FHairStrandsMeshTrianglePositionFormat>(GraphBuilder, LOD.SampleCount + 4, LOD.MeshSampleWeightsBuffer, ToHairResourceDebugName(HAIRSTRANDS_RESOUCE_NAME(CurveType, Hair.StrandsRootDeformed_MeshSampleWeightsBuffer), ResourceName), EHairResourceUsageType::Dynamic);
+				InternalCreateVertexBufferRDG<FHairStrandsMeshTrianglePositionFormat>(GraphBuilder, LOD.SampleCount + 5, LOD.MeshSampleWeightsBuffer, ToHairResourceDebugName(HAIRSTRANDS_RESOUCE_NAME(CurveType, Hair.StrandsRootDeformed_MeshSampleWeightsBuffer), ResourceName), EHairResourceUsageType::Dynamic); // offset added to workaround numerical issues when deformed and rest positions are far apart
 			}
 
 			InternalCreateVertexBufferRDG<FHairStrandsMeshTrianglePositionFormat>(GraphBuilder, RootCount, LOD.DeformedRootTrianglePosition0Buffer, ToHairResourceDebugName(HAIRSTRANDS_RESOUCE_NAME(CurveType, Hair.StrandsRootDeformed_DeformedRootTrianglePosition0Buffer), ResourceName), EHairResourceUsageType::Dynamic);
diff --git a/Engine/Shaders/Private/HairStrands/HairStrandsMeshRBF.ush b/Engine/Shaders/Private/HairStrands/HairStrandsMeshRBF.ush
--- a/Engine/Shaders/Private/HairStrands/HairStrandsMeshRBF.ush
+++ b/Engine/Shaders/Private/HairStrands/HairStrandsMeshRBF.ush
@@ -20,5 +20,5 @@ float3 ApplyRBF(
 	ControlPoint += InMeshSampleWeightsBuffer[InMaxSampleCount + 1].xyz * RestControlPoint.x;
 	ControlPoint += InMeshSampleWeightsBuffer[InMaxSampleCount + 2].xyz * RestControlPoint.y;
 	ControlPoint += InMeshSampleWeightsBuffer[InMaxSampleCount + 3].xyz * RestControlPoint.z;
-	return ControlPoint;
+	return ControlPoint + InMeshSampleWeightsBuffer[InMaxSampleCount + 4].xyz; // offset added to workaround numerical issues when deformed and rest positions are far apart
 }
\ No newline at end of file
diff --git a/Engine/Shaders/Private/HairStrands/HairStrandsSamplesUpdate.usf b/Engine/Shaders/Private/HairStrands/HairStrandsSamplesUpdate.usf
--- a/Engine/Shaders/Private/HairStrands/HairStrandsSamplesUpdate.usf
+++ b/Engine/Shaders/Private/HairStrands/HairStrandsSamplesUpdate.usf
@@ -15,15 +15,21 @@ RWBuffer<float4> OutSampleDeformationsBuffer;
 void MainCS(uint3 DispatchThreadId : SV_DispatchThreadID)
 { 
 	const uint SampleIndex = DispatchThreadId.x;
-	if (SampleIndex >= (MaxSampleCount+4)) 
+	if (SampleIndex >= (MaxSampleCount + 4))
+	{
+		if (SampleIndex == (MaxSampleCount + 4))
+		{
+			OutSampleDeformationsBuffer[SampleIndex] = float4(SampleDeformedPositionsBuffer[0].xyz, 1); // offset used to workaround numerical issues when deformed and rest positions are far apart
+		}
 		return;
-	
+	}
+
 	uint WeightsOffset = SampleIndex * (MaxSampleCount+4);
 	float3 SampleDeformation = float3(0,0,0);
 	for(uint i = 0; i < MaxSampleCount; ++i, ++WeightsOffset)
 	{
 		SampleDeformation += InterpolationWeightsBuffer[WeightsOffset] * 
-				(SampleDeformedPositionsBuffer[i] - SampleRestPositionsBuffer[i]).xyz;
+				(SampleDeformedPositionsBuffer[i] - (SampleRestPositionsBuffer[i] + SampleDeformedPositionsBuffer[0])).xyz; // offset added to workaround numerical issues when deformed and rest positions are far apart
 	}
 
 	OutSampleDeformationsBuffer[SampleIndex] = float4(SampleDeformation,1);
