Sync of the Skinning/Rigid transforms + rest update while reset + fix reset events

diff --git a/Engine/Plugins/Runtime/HairStrands/Source/HairStrandsCore/Private/GroomComponent.cpp b/Engine/Plugins/Runtime/HairStrands/Source/HairStrandsCore/Private/GroomComponent.cpp
--- a/Engine/Plugins/Runtime/HairStrands/Source/HairStrandsCore/Private/GroomComponent.cpp
+++ b/Engine/Plugins/Runtime/HairStrands/Source/HairStrandsCore/Private/GroomComponent.cpp
@@ -43,6 +43,10 @@ static FAutoConsoleVariableRef CVarHairBindingValidationEnable(TEXT("r.HairStran
 static bool GUseGroomCacheStreaming = true;
 static FAutoConsoleVariableRef CVarGroomCacheStreamingEnable(TEXT("GroomCache.EnableStreaming"), GUseGroomCacheStreaming, TEXT("Enable groom cache streaming and prebuffering. Do not switch while groom caches are in use."));
 
+static bool GUseProxyLocalToWorld = false;
+static FAutoConsoleVariableRef CVarUseProxyLocalToWorld(TEXT("GroomCache.UseProxyLocalToWorld"), GUseProxyLocalToWorld, TEXT("Enable the use of the groom proxy local to world instead of extracting it from the game thread."));
+
+
 #define LOCTEXT_NAMESPACE "GroomComponent"
 
 #define USE_HAIR_TRIANGLE_STRIP 1
@@ -421,9 +425,13 @@ public:
 			}
 
 			// Initialization of the default skel. mesh transform (refresh during ticking).
-			HairInstance->Debug.SkeletalLocalToWorld = HairInstance->Debug.MeshComponent ? HairInstance->Debug.MeshComponent->GetComponentTransform() : FTransform();
-			HairInstance->Debug.SkeletalPreviousLocalToWorld = HairInstance->Debug.SkeletalLocalToWorld;
-			HairInstance->LocalToWorld = HairInstance->Debug.SkeletalLocalToWorld;
+			HairInstance->Debug.SkinningCurrentLocalToWorld = HairInstance->Debug.MeshComponent ? HairInstance->Debug.MeshComponent->GetComponentTransform() : FTransform();
+			HairInstance->Debug.SkinningPreviousLocalToWorld = HairInstance->Debug.SkinningCurrentLocalToWorld;
+
+			HairInstance->Debug.RigidCurrentLocalToWorld = Component->GetComponentTransform();
+			HairInstance->Debug.RigidPreviousLocalToWorld = HairInstance->Debug.RigidCurrentLocalToWorld;
+
+			HairInstance->LocalToWorld = HairInstance->GetCurrentLocalToWorld();
 		}
 	}
 
@@ -479,31 +487,20 @@ public:
 
 	virtual void OnTransformChanged() override
 	{
-		const FTransform HairLocalToWorld = FTransform(GetLocalToWorld());
-		for (FHairGroupInstance* Instance : HairGroupInstances)
+		if (GUseProxyLocalToWorld)
 		{
-			Instance->LocalToWorld = HairLocalToWorld;
-			if (Instance->BindingType == EHairBindingType::Skinning)
+			const FTransform RigidLocalToWorld = FTransform(GetLocalToWorld());
+			for (FHairGroupInstance* Instance : HairGroupInstances)
 			{
-				Instance->LocalToWorld = Instance->Debug.SkeletalLocalToWorld;
+				Instance->Debug.RigidPreviousLocalToWorld = Instance->Debug.RigidCurrentLocalToWorld;
+				Instance->Debug.RigidCurrentLocalToWorld = RigidLocalToWorld;
 			}
 		}
 	}
 
-	void GetOverrideLocalToTransform(const FHairGroupInstance* Instance, FMatrix& OutLocalToWorld) const
-	{
-		if (Instance->BindingType == EHairBindingType::Skinning)
-		{
-			OutLocalToWorld = Instance->Debug.SkeletalLocalToWorld.ToMatrixWithScale();
-		}
-	}
-
-	void GetOverridePreviousLocalToTransform(const FHairGroupInstance* Instance, FMatrix& OutPreviousLocalToWorld) const
+	FORCEINLINE bool UseProxyLocalToWorld(const FHairGroupInstance* Instance) const
 	{
-		if (Instance->BindingType == EHairBindingType::Skinning)
-		{
-			OutPreviousLocalToWorld = Instance->Debug.SkeletalPreviousLocalToWorld.ToMatrixWithScale();
-		}
+		return (GUseProxyLocalToWorld && (Instance->BindingType != EHairBindingType::Skinning));
 	}
 
 #if RHI_RAYTRACING
@@ -531,8 +528,8 @@ public:
 		{
 			FHairGroupInstance* Instance = HairGroupInstances[GroupIt];
 			check(Instance->GetRefCount() > 0);
-			FMatrix OverrideLocalToWorld = GetLocalToWorld();
-			GetOverrideLocalToTransform(Instance, OverrideLocalToWorld);
+
+			FMatrix OverrideLocalToWorld = UseProxyLocalToWorld(Instance) ? GetLocalToWorld() : Instance->GetCurrentLocalToWorld().ToMatrixWithScale();
 
 			const EHairGeometryType GeometryType = Instance->HairGroupPublicData->VFInput.GeometryType;
 			const uint32 LODIndex = Instance->HairGroupPublicData->GetIntLODIndex();
@@ -828,12 +825,13 @@ public:
 		bool bDrawVelocity = bOutputVelocity; // Velocity vector is done in a custom fashion
 		GetScene().GetPrimitiveUniformShaderParameters_RenderThread(GetPrimitiveSceneInfo(), bHasPrecomputedVolumetricLightmap, PreviousLocalToWorld, SingleCaptureIndex, bOutputVelocity);
 
-		FMatrix OverrideLocalToWorld = GetLocalToWorld();
-		GetOverrideLocalToTransform(Instance, OverrideLocalToWorld);
-		GetOverridePreviousLocalToTransform(Instance, PreviousLocalToWorld);
+		const bool bUseProxy = UseProxyLocalToWorld(Instance);
+
+		FMatrix CurrentLocalToWorld = bUseProxy ? GetLocalToWorld() : Instance->GetCurrentLocalToWorld().ToMatrixWithScale();
+		PreviousLocalToWorld = bUseProxy ? PreviousLocalToWorld : Instance->GetPreviousLocalToWorld().ToMatrixWithScale();
 
 		FDynamicPrimitiveUniformBuffer& DynamicPrimitiveUniformBuffer = Collector.AllocateOneFrameResource<FDynamicPrimitiveUniformBuffer>();
-		DynamicPrimitiveUniformBuffer.Set(OverrideLocalToWorld, PreviousLocalToWorld, GetBounds(), GetLocalBounds(), true, false, bDrawVelocity, bOutputVelocity);
+		DynamicPrimitiveUniformBuffer.Set(CurrentLocalToWorld, PreviousLocalToWorld, GetBounds(), GetLocalBounds(), true, false, bDrawVelocity, bOutputVelocity);
 		BatchElement.PrimitiveUniformBufferResource = &DynamicPrimitiveUniformBuffer.UniformBuffer;
 
 		BatchElement.FirstIndex = 0;
@@ -1648,24 +1646,28 @@ void UGroomComponent::SetBindingAsset(UGroomBindingAsset* InBinding)
 
 void UGroomComponent::SetEnableSimulation(bool bInEnableSimulation)
 {
-	if (SimulationSettings.SolverSettings.bEnableSimulation != bInEnableSimulation)
+	if (SimulationSettings.bOverrideSettings)
 	{
-		SimulationSettings.SolverSettings.bEnableSimulation = bInEnableSimulation;
+		if (SimulationSettings.SolverSettings.bEnableSimulation != bInEnableSimulation)
+		{
+			SimulationSettings.SolverSettings.bEnableSimulation = bInEnableSimulation;
 
-		ReleaseResources();
-		UpdateHairGroupsDesc();
-		UpdateHairSimulation();
-		InitResources();
+			ReleaseResources();
+			UpdateHairGroupsDesc();
+			UpdateHairSimulation();
+			InitResources();
+		}
+	}
+	else
+	{
+		UE_LOG(LogHairStrands, Warning, TEXT("[Groom] %s - To be able to enable/disable the simulation from component you must enable the overide settings boolean in the simulatin settings."), *GetPathName());
 	}
 }
 
 void UGroomComponent::ResetSimulation()
 {
-	if (SimulationSettings.SolverSettings.bEnableSimulation)
-	{
-		bResetSimulation = true;
-		bInitSimulation = true;
-	}
+	bResetSimulation = true;
+	bInitSimulation = true;
 }
 
 void UGroomComponent::UpdateHairGroupsDescAndInvalidateRenderState(bool bInvalid)
@@ -3152,18 +3154,25 @@ void UGroomComponent::TickComponent(float DeltaTime, enum ELevelTick TickType, F
 		UpdateGroomCache(ElapsedTime);
 	}
 
-	const FTransform SkinLocalToWorld = RegisteredMeshComponent ? RegisteredMeshComponent->GetComponentTransform() : FTransform();
+	const FTransform SkinningLocalToWorld = RegisteredMeshComponent ? RegisteredMeshComponent->GetComponentTransform() : FTransform();
+	const FTransform RigidLocalToWorld = GetComponentTransform();
 	TArray<FHairGroupInstance*> LocalHairGroupInstances = HairGroupInstances;
 	ENQUEUE_RENDER_COMMAND(FHairStrandsTick_TransformUpdate)(
-		[Id, SkinLocalToWorld, FeatureLevel, LocalHairGroupInstances](FRHICommandListImmediate& RHICmdList)
+		[Id, SkinningLocalToWorld, RigidLocalToWorld, FeatureLevel, LocalHairGroupInstances](FRHICommandListImmediate& RHICmdList)
 	{
 		if (ERHIFeatureLevel::Num == FeatureLevel)
 			return;
 
 		for (FHairGroupInstance* Instance : LocalHairGroupInstances)
 		{
-			Instance->Debug.SkeletalPreviousLocalToWorld = Instance->Debug.SkeletalLocalToWorld;
-			Instance->Debug.SkeletalLocalToWorld = SkinLocalToWorld;
+			Instance->Debug.SkinningPreviousLocalToWorld = Instance->Debug.SkinningCurrentLocalToWorld;
+			Instance->Debug.SkinningCurrentLocalToWorld = SkinningLocalToWorld;
+
+			if (!GUseProxyLocalToWorld)
+			{
+				Instance->Debug.RigidPreviousLocalToWorld = Instance->Debug.RigidCurrentLocalToWorld;
+				Instance->Debug.RigidCurrentLocalToWorld = RigidLocalToWorld;
+			}
 		}
 	});
 }
diff --git a/Engine/Plugins/Runtime/HairStrands/Source/HairStrandsCore/Private/GroomManager.cpp b/Engine/Plugins/Runtime/HairStrands/Source/HairStrandsCore/Private/GroomManager.cpp
--- a/Engine/Plugins/Runtime/HairStrands/Source/HairStrandsCore/Private/GroomManager.cpp
+++ b/Engine/Plugins/Runtime/HairStrands/Source/HairStrandsCore/Private/GroomManager.cpp
@@ -779,7 +779,7 @@ static void RunHairLODSelection(
 		}
 
 		// Update the local-to-world transform based on the binding type 
-		Instance->LocalToWorld = Instance->BindingType == EHairBindingType::Skinning ? Instance->Debug.SkeletalLocalToWorld : Instance->LocalToWorld;
+		Instance->LocalToWorld = Instance->GetCurrentLocalToWorld();
 	}
 }
 
diff --git a/Engine/Plugins/Runtime/HairStrands/Source/HairStrandsCore/Private/HairStrandsRendering.cpp b/Engine/Plugins/Runtime/HairStrands/Source/HairStrandsCore/Private/HairStrandsRendering.cpp
--- a/Engine/Plugins/Runtime/HairStrands/Source/HairStrandsCore/Private/HairStrandsRendering.cpp
+++ b/Engine/Plugins/Runtime/HairStrands/Source/HairStrandsCore/Private/HairStrandsRendering.cpp
@@ -2140,10 +2140,7 @@ void ComputeHairStrandsInterpolation(
 	}
 
 	Instance->HairGroupPublicData->VFInput.GeometryType = InstanceGeometryType;
-	if (Instance->BindingType == EHairBindingType::Skinning)
-		Instance->HairGroupPublicData->VFInput.LocalToWorldTransform = Instance->Debug.SkeletalLocalToWorld;
-	else
-		Instance->HairGroupPublicData->VFInput.LocalToWorldTransform = Instance->LocalToWorld;
+	Instance->HairGroupPublicData->VFInput.LocalToWorldTransform = Instance->GetCurrentLocalToWorld();
 	Instance->HairGroupPublicData->bSupportVoxelization = Instance->Strands.Modifier.bSupportVoxelization && Instance->bCastShadow;
 
 	ResourceAccessFinalizer.Finalize(GraphBuilder);
diff --git a/Engine/Plugins/Runtime/HairStrands/Source/HairStrandsCore/Private/Niagara/NiagaraDataInterfaceHairStrands.cpp b/Engine/Plugins/Runtime/HairStrands/Source/HairStrandsCore/Private/Niagara/NiagaraDataInterfaceHairStrands.cpp
--- a/Engine/Plugins/Runtime/HairStrands/Source/HairStrandsCore/Private/Niagara/NiagaraDataInterfaceHairStrands.cpp
+++ b/Engine/Plugins/Runtime/HairStrands/Source/HairStrandsCore/Private/Niagara/NiagaraDataInterfaceHairStrands.cpp
@@ -638,7 +638,7 @@ struct FNDIHairStrandsParametersCS : public FNiagaraDataInterfaceParametersCS
 		const bool bIsHairValid = ProxyData != nullptr && ProxyData->HairStrandsBuffer && ProxyData->HairStrandsBuffer->IsInitialized();
 		const bool bHasSkinningBinding = bIsHairValid && ProxyData->HairGroupInstance && ProxyData->HairGroupInstance->BindingType == EHairBindingType::Skinning;
 		const bool bIsRootValid = bIsHairValid && ProxyData->HairStrandsBuffer->SourceDeformedRootResources && ProxyData->HairStrandsBuffer->SourceDeformedRootResources->IsInitialized() && bHasSkinningBinding;
-		const bool bIsRestValid = bIsHairValid && ProxyData->HairStrandsBuffer->SourceRestResources && 
+		const bool bIsRestValid = bIsHairValid && ProxyData->HairStrandsBuffer->SourceRestResources && ProxyData->HairStrandsBuffer->SourceRestResources->IsInitialized() &&
 
 			// SourceRestResources->PositionBuffer is lazily allocated, when the instance LOD is scheduled (this happens after this called). So this is why this check is here. 
 			// This code should be refactor so that it reflects the lazy allocation scheme
@@ -647,7 +647,7 @@ struct FNDIHairStrandsParametersCS : public FNiagaraDataInterfaceParametersCS
 			ProxyData->HairStrandsBuffer->CurvesOffsetsBuffer.SRV && ProxyData->HairStrandsBuffer->ParamsScaleBuffer.SRV && ProxyData->HairStrandsBuffer->BoundingBoxBuffer.UAV;
 
 		const bool bIsDeformedValid = bIsHairValid && ProxyData->HairStrandsBuffer->SourceDeformedResources && ProxyData->HairStrandsBuffer->SourceDeformedResources->IsInitialized();
-		const bool bHasLODSwitched = bIsHairValid && ProxyData->HairGroupInstance && ProxyData->HairGroupInstance->HairGroupPublicData && ProxyData->HairGroupInstance->HairGroupPublicData->VFInput.bHasLODSwitch;
+		const bool bHasLODSwitched = bIsHairValid && ProxyData->HairGroupInstance&& ProxyData->HairGroupInstance->HairGroupPublicData&& ProxyData->HairGroupInstance->HairGroupPublicData->VFInput.bHasLODSwitch;
 
 		if (bIsHairValid && bIsRestValid)
 		{
@@ -697,19 +697,22 @@ struct FNDIHairStrandsParametersCS : public FNiagaraDataInterfaceParametersCS
 
 			// Simulation setup
 			const int32 NeedResetValue = bHasLODSwitched || (ProxyData->TickCount <= GHairSimulationMaxDelay);
-			const int32 RestUpdateValue = GHairSimulationRestUpdate;
+			const int32 RestUpdateValue = GHairSimulationRestUpdate || NeedResetValue;
 			const int32 LocalSimulationValue = ProxyData->LocalSimulation;
 			
 			// Offsets / Transforms
 			FVector3f RestPositionOffsetValue = ProxyData->HairStrandsBuffer->SourceRestResources->GetPositionOffset();
 
-			FMatrix44f WorldTransformFloat = ProxyData->HairGroupInstance ? (ProxyData->HairGroupInstance->BindingType == EHairBindingType::Skinning) ?
-				ProxyData->HairGroupInstance->Debug.SkeletalLocalToWorld.ToMatrixWithScale() :
-				ProxyData->HairGroupInstance->LocalToWorld.ToMatrixWithScale() :
-				ProxyData->WorldTransform.ToMatrixWithScale();
-
+			FMatrix44f WorldTransformFloat = ProxyData->HairGroupInstance ? ProxyData->HairGroupInstance->GetCurrentLocalToWorld().ToMatrixWithScale() :
+																			ProxyData->WorldTransform.ToMatrixWithScale();
 			FMatrix44f BoneTransformFloat = ProxyData->BoneTransform.ToMatrixWithScale();
 
+			if (!bIsRootValid && bHasSkinningBinding)
+			{
+				UE_LOG(LogHairStrands, Log, TEXT("FNDIHairStrandsParametersCS() Groom Asset %s from component %s is set to use skinning interpolation but the skin resources are not valid"),
+					 *ProxyData->HairGroupInstance->Debug.GroomAssetName, *ProxyData->HairGroupInstance->Debug.MeshComponentName);
+			}
+
 			FRHITransitionInfo Transitions[] = {
 				FRHITransitionInfo(DeformedPositionBufferUAV, ERHIAccess::Unknown, ERHIAccess::UAVCompute),
 				FRHITransitionInfo(HairStrandsBuffer->BoundingBoxBuffer.UAV, ERHIAccess::Unknown, ERHIAccess::UAVCompute),
diff --git a/Engine/Plugins/Runtime/HairStrands/Source/HairStrandsCore/Public/GroomInstance.h b/Engine/Plugins/Runtime/HairStrands/Source/HairStrandsCore/Public/GroomInstance.h
--- a/Engine/Plugins/Runtime/HairStrands/Source/HairStrandsCore/Public/GroomInstance.h
+++ b/Engine/Plugins/Runtime/HairStrands/Source/HairStrandsCore/Public/GroomInstance.h
@@ -194,8 +194,10 @@ struct HAIRSTRANDSCORE_API FHairGroupInstance : public FHairStrandsInstance
 		UMeshComponent*			MeshComponent = nullptr;
 		FString					MeshComponentName;
 		const UGroomComponent*	GroomComponentForDebug = nullptr; // For debug only, shouldn't be deferred on the rendering thread
-		FTransform				SkeletalLocalToWorld = FTransform::Identity;
-		FTransform				SkeletalPreviousLocalToWorld = FTransform::Identity;
+		FTransform				RigidCurrentLocalToWorld = FTransform::Identity;
+		FTransform				SkinningCurrentLocalToWorld = FTransform::Identity;
+		FTransform				RigidPreviousLocalToWorld = FTransform::Identity;
+		FTransform				SkinningPreviousLocalToWorld = FTransform::Identity;
 		bool					bDrawCardsGuides = false;
 
 		TSharedPtr<class IGroomCacheBuffers, ESPMode::ThreadSafe> GroomCacheBuffers;
@@ -224,4 +226,18 @@ struct HAIRSTRANDSCORE_API FHairGroupInstance : public FHairStrandsInstance
 	{
 		return Meshes.IsValid() || Cards.IsValid() || Strands.IsValid();
 	}
+
+	/** Get the current local to world transform according to the internal binding type */
+	FORCEINLINE const FTransform& GetCurrentLocalToWorld() const
+	{
+		return (BindingType == EHairBindingType::Skinning) ? Debug.SkinningCurrentLocalToWorld: 
+															 Debug.RigidCurrentLocalToWorld;
+	}
+
+	/** Get the previous local to world transform according to the internal binding type */
+	FORCEINLINE const FTransform& GetPreviousLocalToWorld() const
+	{
+		return (BindingType == EHairBindingType::Skinning) ? Debug.SkinningPreviousLocalToWorld :
+															 Debug.RigidPreviousLocalToWorld;
+	}
 };
