Create a view that matches the UINT type of HairStrands' Attribute1 param.  Fixes Vulkan validation error "requires UINT component type, but bound descriptor format is VK_FORMAT_R8G8B8A8_UNORM".

diff --git a/Engine/Plugins/Runtime/HairStrands/Source/HairStrandsCore/Private/HairStrandsVertexFactory.cpp b/Engine/Plugins/Runtime/HairStrands/Source/HairStrandsCore/Private/HairStrandsVertexFactory.cpp
--- a/Engine/Plugins/Runtime/HairStrands/Source/HairStrandsCore/Private/HairStrandsVertexFactory.cpp
+++ b/Engine/Plugins/Runtime/HairStrands/Source/HairStrandsCore/Private/HairStrandsVertexFactory.cpp
@@ -26,6 +26,7 @@ public:
 	FShaderResourceViewRHIRef SRVUint;
 	FShaderResourceViewRHIRef SRVFloat;
 	FShaderResourceViewRHIRef SRVRGBA;
+	FShaderResourceViewRHIRef SRVRGBA_Uint;
 
 	virtual void InitRHI() override
 	{
@@ -39,6 +40,7 @@ public:
 		SRVUint = RHICreateShaderResourceView(VertexBufferRHI, sizeof(uint32), PF_R32_UINT);
 		SRVFloat = RHICreateShaderResourceView(VertexBufferRHI, sizeof(uint32), PF_R32_FLOAT);
 		SRVRGBA = RHICreateShaderResourceView(VertexBufferRHI, sizeof(uint32), PF_R8G8B8A8);
+		SRVRGBA_Uint = RHICreateShaderResourceView(VertexBufferRHI, sizeof(uint32), PF_R8G8B8A8_UINT);
 	}
 
 	virtual void ReleaseRHI() override
@@ -47,6 +49,7 @@ public:
 		SRVUint.SafeRelease();
 		SRVFloat.SafeRelease();
 		SRVRGBA.SafeRelease();
+		SRVRGBA_Uint.SafeRelease();
 	}
 };
 TGlobalResource<FDummyCulledDispatchVertexIdsBuffer> GDummyCulledDispatchVertexIdsBuffer;
@@ -153,7 +156,7 @@ public:
 		const bool bHasAttribute1 = Attribute1SRV != nullptr;
 		if (!bHasAttribute1)
 		{
-			Attribute1SRV = GDummyCulledDispatchVertexIdsBuffer.SRVRGBA;
+			Attribute1SRV = GDummyCulledDispatchVertexIdsBuffer.SRVRGBA_Uint;
 		}
 
 		FShaderResourceViewRHIRef MaterialSRV = GetMaterialSRV(VFInput);
