Rename hair cvar + temporary solution to fix missing hairs

diff --git a/Engine/Plugins/Runtime/HairStrands/Source/HairStrandsCore/Private/GroomComponent.cpp b/Engine/Plugins/Runtime/HairStrands/Source/HairStrandsCore/Private/GroomComponent.cpp
--- a/Engine/Plugins/Runtime/HairStrands/Source/HairStrandsCore/Private/GroomComponent.cpp
+++ b/Engine/Plugins/Runtime/HairStrands/Source/HairStrandsCore/Private/GroomComponent.cpp
@@ -44,7 +44,7 @@ static bool GUseGroomCacheStreaming = true;
 static FAutoConsoleVariableRef CVarGroomCacheStreamingEnable(TEXT("GroomCache.EnableStreaming"), GUseGroomCacheStreaming, TEXT("Enable groom cache streaming and prebuffering. Do not switch while groom caches are in use."));
 
 static bool GUseProxyLocalToWorld = false;
-static FAutoConsoleVariableRef CVarUseProxyLocalToWorld(TEXT("GroomCache.UseProxyLocalToWorld"), GUseProxyLocalToWorld, TEXT("Enable the use of the groom proxy local to world instead of extracting it from the game thread."));
+static FAutoConsoleVariableRef CVarUseProxyLocalToWorld(TEXT("r.HairStrands.UseProxyLocalToWorld"), GUseProxyLocalToWorld, TEXT("Enable the use of the groom proxy local to world instead of extracting it from the game thread."));
 
 
 #define LOCTEXT_NAMESPACE "GroomComponent"
