Hair Simulation could recorver from bad init + remove rest update while resetting + reset niagara max queued in frosty to default

diff --git a/Engine/Plugins/Runtime/HairStrands/Source/HairStrandsCore/Private/Niagara/NiagaraDataInterfaceHairStrands.cpp b/Engine/Plugins/Runtime/HairStrands/Source/HairStrandsCore/Private/Niagara/NiagaraDataInterfaceHairStrands.cpp
--- a/Engine/Plugins/Runtime/HairStrands/Source/HairStrandsCore/Private/Niagara/NiagaraDataInterfaceHairStrands.cpp
+++ b/Engine/Plugins/Runtime/HairStrands/Source/HairStrandsCore/Private/Niagara/NiagaraDataInterfaceHairStrands.cpp
@@ -697,7 +697,7 @@ struct FNDIHairStrandsParametersCS : public FNiagaraDataInterfaceParametersCS
 
 			// Simulation setup
 			const int32 NeedResetValue = bHasLODSwitched || (ProxyData->TickCount <= GHairSimulationMaxDelay);
-			const int32 RestUpdateValue = GHairSimulationRestUpdate || NeedResetValue;
+			const int32 RestUpdateValue = GHairSimulationRestUpdate;
 			const int32 LocalSimulationValue = ProxyData->LocalSimulation;
 			
 			// Offsets / Transforms
