Change groom LOD to have one of the following mode: immediated/predicted/forced
* Immediate: LOD is selected on the rendering thread prior to rendering
* Predictied: LOD is estimated on the rendering-thread and feedback to the game-thread. The game-thread selects which LOD is appropriate based on the feedback.
* Forced: LOD is set by high-level logic on the game-thread, forcing ar particular LOD.

When a groom as one of its LOD/Group with simulation enabled, LOD selection is forced to be predicted. This is mandatory as the simulation needs to know on the game-thread, when to initialized/enabled Niagara components.

This allows to fix case where the simulation is not used/initialized due to starting at a lower LOD (with no sim) and then moving to a higher LOD (with sim), and having the sim. not being kicked-in/initialized properly.

diff --git a/Engine/Plugins/Runtime/HairStrands/Source/HairStrandsCore/Private/GroomComponent.cpp b/Engine/Plugins/Runtime/HairStrands/Source/HairStrandsCore/Private/GroomComponent.cpp
--- a/Engine/Plugins/Runtime/HairStrands/Source/HairStrandsCore/Private/GroomComponent.cpp
+++ b/Engine/Plugins/Runtime/HairStrands/Source/HairStrandsCore/Private/GroomComponent.cpp
@@ -299,7 +299,6 @@ public:
 		check(Component->GroomAsset->GetNumHairGroups() > 0);
 		ComponentId = Component->ComponentId.PrimIDValue;
 		Strands_DebugMaterial = Component->Strands_DebugMaterial;
-		PredictedLODIndex = &Component->PredictedLODIndex;
 		bAlwaysHasVelocity = false;
 		if (IsHairStrandsBindingEnable() && Component->RegisteredMeshComponent)
 		{
@@ -322,7 +321,6 @@ public:
 			check(HairInstance->HairGroupPublicData);
 			HairInstance->ProxyBounds = &GetBounds();
 			HairInstance->ProxyLocalBounds = &GetLocalBounds();
-			HairInstance->bUseCPULODSelection = Component->GroomAsset->LODSelectionType == EHairLODSelectionType::Cpu;
 			HairInstance->bForceCards = Component->bUseCards;
 			HairInstance->bUpdatePositionOffset = Component->RegisteredMeshComponent != nullptr;
 			HairInstance->bCastShadow = Component->CastShadow;
@@ -832,7 +830,8 @@ public:
 
 		FDynamicPrimitiveUniformBuffer& DynamicPrimitiveUniformBuffer = Collector.AllocateOneFrameResource<FDynamicPrimitiveUniformBuffer>();
 		DynamicPrimitiveUniformBuffer.Set(CurrentLocalToWorld, PreviousLocalToWorld, GetBounds(), GetLocalBounds(), true, false, bDrawVelocity, bOutputVelocity);
-		BatchElement.PrimitiveUniformBufferResource = &DynamicPrimitiveUniformBuffer.UniformBuffer;
+		BatchElement.PrimitiveUniformBufferResource = &DynamicPrimitiveUniformBuffer.UniformBuffer; // automatic copy to the gpu scene buffer
+		//primtiveid is set to 0
 
 		BatchElement.FirstIndex = 0;
 		BatchElement.NumInstances = 1;
@@ -921,7 +920,6 @@ private:
 	uint32 ComponentId = 0;
 	FMaterialRelevance MaterialRelevance;
 	UMaterialInterface* Strands_DebugMaterial = nullptr;
-	int32* PredictedLODIndex = nullptr;
 };
 
 /** GroomCacheBuffers implementation that hold copies of the GroomCacheAnimationData needed for playback */
@@ -1550,40 +1548,84 @@ bool UGroomComponent::GetIsHairLengthScaleEnabled()
 }
 
 
-void UGroomComponent::SetForcedLOD(int32 LODIndex)
+void UGroomComponent::SetForcedLOD(int32 CurrLODIndex)
 {
+	// Previous LOD state
+	EHairLODSelectionType PrevLODSelectionType = LODSelectionType;
+	float PrevLODIndex = -1;
+	if (PrevLODSelectionType == EHairLODSelectionType::Forced)			{ PrevLODIndex = LODForcedIndex; }
+	else if (PrevLODSelectionType == EHairLODSelectionType::Predicted)	{ PrevLODIndex = LODPredictedIndex; }
+
+	// Current LOD state
+	EHairLODSelectionType CurrLODSelectionType = EHairLODSelectionType::Immediate;
 	if (GroomAsset)
 	{
-		LODIndex = FMath::Clamp(LODIndex, -1, GroomAsset->GetLODCount() - 1);
+		CurrLODIndex = FMath::Clamp(CurrLODIndex, -1, GroomAsset->GetLODCount() - 1);
+		CurrLODSelectionType = EHairLODSelectionType::Forced;
 	}
 	else
 	{
-		LODIndex = -1;
+		CurrLODIndex = -1;
 	}
-	const bool bHasLODChanged = (GroomGroupsDesc.Num() > 0) ? 
-		(IsSimulationEnable(0, GetForcedLOD()) != IsSimulationEnable(0, LODIndex)) : false;
 
-	for (FHairGroupDesc& HairDesc : GroomGroupsDesc)
+	// Reset to non-forced LOD, and change LOD selection type to Predicted/Immediate
+	if (CurrLODIndex < 0)
 	{
-		HairDesc.LODForcedIndex = LODIndex;
+		// If one of the Group/LOD needs simulation, then we switch to predicted LOD to ensure LOD switch happen 
+		// on the game-thread so that the simulation component can do work at transition-time
+		bool bNeedSimulationNeed = false;
+		if (GroomAsset)
+		{
+			for (int32 GroupIt = 0, GroupCount = GroomAsset->HairGroupsData.Num(); GroupIt < GroupCount; ++GroupIt)
+			{
+				for (uint32 LODIt = 0, LODCount = GroomAsset->GetLODCount(); LODIt < LODCount; ++LODIt)
+				{
+					if (IsSimulationEnable(GroupIt, LODIt))
+					{
+						bNeedSimulationNeed = true;
+						break;
+					}
+				}
+			}
+		}
+		CurrLODSelectionType = bNeedSimulationNeed ? EHairLODSelectionType::Predicted : EHairLODSelectionType::Immediate;
 	}
-	UpdateHairGroupsDesc();
 	
+	const bool bHasLODChanged = CurrLODIndex != PrevLODIndex;
+
+	// Inform simulation about LOD switch.
 	if (bHasLODChanged)
 	{
-		UpdateHairSimulation();
+		if (GroomAsset)
+		{
+			for (int32 GroupIt = 0, GroupCount = GroomAsset->HairGroupsData.Num(); GroupIt < GroupCount; ++GroupIt)
+			{
+				if ((IsSimulationEnable(GroupIt, PrevLODIndex) != IsSimulationEnable(GroupIt, CurrLODIndex)))
+				{
+					// Sanity check. Simulation does not support immediate mode.
+					check(LODSelectionType != EHairLODSelectionType::Immediate);
+
+					UpdateHairSimulation();
+				}
+			}
+		}
 	}
 
-	// Do not invalidate completly the proxy, but just update LOD index on the rendering thread
-	FHairStrandsSceneProxy* GroomSceneProxy = (FHairStrandsSceneProxy*)SceneProxy;
-	if (GroomSceneProxy)
+	// Finally, update the forced LOD value and LOD selection type
+	LODForcedIndex	 = CurrLODIndex;
+	LODSelectionType = CurrLODSelectionType;
+
+	// Do not invalidate completly the proxy, but just update LOD index on the rendering thread	
+	if (FHairStrandsSceneProxy* GroomSceneProxy = (FHairStrandsSceneProxy*)SceneProxy)
 	{
+		const EHairLODSelectionType LocalLODSelectionType = LODSelectionType;
 		ENQUEUE_RENDER_COMMAND(FHairComponentSendLODIndex)(
-		[GroomSceneProxy, LODIndex](FRHICommandListImmediate& RHICmdList)
+		[GroomSceneProxy, CurrLODIndex, LocalLODSelectionType](FRHICommandListImmediate& RHICmdList)
 		{
 			for (FHairGroupInstance* Instance : GroomSceneProxy->HairGroupInstances)
 			{
-				Instance->Strands.Modifier.LODForcedIndex = LODIndex;
+				Instance->Debug.LODForcedIndex = CurrLODIndex;
+				Instance->Debug.LODSelectionTypeForDebug = LocalLODSelectionType;
 			}
 		});
 	}
@@ -1596,19 +1638,12 @@ int32 UGroomComponent::GetNumLODs() const
 
 int32 UGroomComponent::GetForcedLOD() const
 {
-	if (GroomGroupsDesc.Num() > 0)
-	{
-		return GroomGroupsDesc[0].LODForcedIndex;
-	}
-	else
-	{
-		return -1;
-	}
+	return int32(LODForcedIndex);
 }
 
 int32 UGroomComponent::GetDesiredSyncLOD() const
 {
-	return PredictedLODIndex;
+	return LODPredictedIndex;
 }
 
 void  UGroomComponent::SetSyncLOD(int32 InLODIndex)
@@ -2242,6 +2277,11 @@ void UGroomComponent::InitResources(bool bIsBindingReloading)
 	}
 	const bool bHasNeedBindingData = bHasNeedSkinningBinding || bHasAnyNeedGlobalDeformation;
 
+	// LOD Selection type
+	// By default set to immediate, meaning the LOD selection will be done on the rendering thread based on the screen coverage of the groom
+	// If simulation is enabled, then the LOD selection will be done on the game thread based on feedback from the rendering thread.
+	LODSelectionType = bHasAnyNeedSimulation ? EHairLODSelectionType::Predicted : EHairLODSelectionType::Immediate;
+
 	// 2. Insure that the binding asset is compatible, otherwise no binding
 	UMeshComponent* ParentMeshComponent = GetAttachParent() ? Cast<UMeshComponent>(GetAttachParent()) : nullptr;
 	UMeshComponent* ValidatedMeshComponent = nullptr;
@@ -2344,6 +2384,10 @@ void UGroomComponent::InitResources(bool bIsBindingReloading)
 		HairGroupInstance->Debug.GroomBindingType = BindingAsset ? BindingAsset->GroomBindingType : EGroomBindingMeshType::SkeletalMesh;
 		HairGroupInstance->Debug.GroomCacheType = GroomCache ? GroomCache->GetType() : EGroomCacheType::None;
 		HairGroupInstance->Debug.GroomCacheBuffers = GroomCacheBuffers;
+		HairGroupInstance->Debug.LODForcedIndex = LODForcedIndex;
+		HairGroupInstance->Debug.LODPredictedIndex = LODPredictedIndex;
+		HairGroupInstance->Debug.LODSelectionTypeForDebug = LODSelectionType;
+
 		if (RegisteredMeshComponent)
 		{
 			HairGroupInstance->Debug.MeshComponentName = RegisteredMeshComponent->GetPathName();
@@ -2450,7 +2494,6 @@ void UGroomComponent::InitResources(bool bIsBindingReloading)
 			check(GroupIt < GroomGroupsDesc.Num());
 
 			HairGroupInstance->Strands.Data = &GroupData.Strands.BulkData;
-//			HairGroupInstance->Strands.InterpolationData = &GroupData.Strands.InterpolationBulkData;
 
 			// (Lazy) Allocate interpolation resources, only if guides are required
 			if (bNeedGuides)
@@ -3154,17 +3197,57 @@ void UGroomComponent::TickComponent(float DeltaTime, enum ELevelTick TickType, F
 		UpdateGroomCache(ElapsedTime);
 	}
 
+	// Based on LOD selection type pick the effect LOD index. 
+	// When LOD prediction is used, we read from the rendering thread and use the smallest prediction accross all groups
+	float EffectiveForceLOD = -1;
+	if (LODSelectionType == EHairLODSelectionType::Predicted)
+	{
+		// 1. Compute the effective LOD (taking the min. LOD across all the groups)
+		for (const FHairGroupInstance* Instance : HairGroupInstances)
+		{
+			/* /!\ Access rendering thread data at that point /!\ */
+			const float LODPredictedIndex_Instance = Instance->Debug.LODPredictedIndex;
+			if (LODPredictedIndex_Instance > 0)
+			{
+				EffectiveForceLOD = EffectiveForceLOD < 0 ? LODPredictedIndex_Instance : FMath::Min(LODPredictedIndex_Instance, EffectiveForceLOD);
+			}
+		}
+
+		// 2. If there is a LOD change, update the simulation adequately 
+		if (FMath::Abs(LODPredictedIndex - EffectiveForceLOD) > KINDA_SMALL_NUMBER)
+		{
+			for (int32 GroupIt = 0, GroupCount = GroomAsset->HairGroupsData.Num(); GroupIt < GroupCount; ++GroupIt)
+			{
+				if ((IsSimulationEnable(GroupIt, LODPredictedIndex) != IsSimulationEnable(GroupIt, EffectiveForceLOD)))
+				{
+					UpdateHairSimulation();
+				}
+			}
+		}
+
+		// 3. Update global predicted index (used for SyncLOD API)
+		LODPredictedIndex = EffectiveForceLOD;
+	}
+	else if (LODSelectionType == EHairLODSelectionType::Forced)
+	{
+		EffectiveForceLOD = LODForcedIndex;
+	}
+
 	const FTransform SkinningLocalToWorld = RegisteredMeshComponent ? RegisteredMeshComponent->GetComponentTransform() : FTransform();
 	const FTransform RigidLocalToWorld = GetComponentTransform();
 	TArray<FHairGroupInstance*> LocalHairGroupInstances = HairGroupInstances;
+	const EHairLODSelectionType LocalLODSelectionType = LODSelectionType;
 	ENQUEUE_RENDER_COMMAND(FHairStrandsTick_TransformUpdate)(
-		[Id, SkinningLocalToWorld, RigidLocalToWorld, FeatureLevel, LocalHairGroupInstances](FRHICommandListImmediate& RHICmdList)
+		[Id, SkinningLocalToWorld, RigidLocalToWorld, FeatureLevel, LocalHairGroupInstances, EffectiveForceLOD, LocalLODSelectionType](FRHICommandListImmediate& RHICmdList)
 	{
 		if (ERHIFeatureLevel::Num == FeatureLevel)
 			return;
 
 		for (FHairGroupInstance* Instance : LocalHairGroupInstances)
 		{
+			Instance->Debug.LODForcedIndex = EffectiveForceLOD;
+			Instance->Debug.LODSelectionTypeForDebug = LocalLODSelectionType;
+
 			Instance->Debug.SkinningPreviousLocalToWorld = Instance->Debug.SkinningCurrentLocalToWorld;
 			Instance->Debug.SkinningCurrentLocalToWorld = SkinningLocalToWorld;
 
diff --git a/Engine/Plugins/Runtime/HairStrands/Source/HairStrandsCore/Private/GroomDebug.cpp b/Engine/Plugins/Runtime/HairStrands/Source/HairStrandsCore/Private/GroomDebug.cpp
--- a/Engine/Plugins/Runtime/HairStrands/Source/HairStrandsCore/Private/GroomDebug.cpp
+++ b/Engine/Plugins/Runtime/HairStrands/Source/HairStrandsCore/Private/GroomDebug.cpp
@@ -798,6 +798,17 @@ static const TCHAR* ToString(EHairBindingType In)
 	return TEXT("None");
 }
 
+static const TCHAR* ToString(EHairLODSelectionType In)
+{
+	switch (In)
+	{
+	case EHairLODSelectionType::Immediate:	return TEXT("Immed");
+	case EHairLODSelectionType::Predicted:	return TEXT("Predic");
+	case EHairLODSelectionType::Forced:		return TEXT("Forced");
+	}
+	return TEXT("None");
+}
+
 FCachedGeometry GetCacheGeometryForHair(
 	FRDGBuilder& GraphBuilder,
 	FHairGroupInstance* Instance,
@@ -869,12 +880,13 @@ void RunHairStrandsDebug(
 					}
 				}
 
-				Line = FString::Printf(TEXT(" * Group:%d/%d | LOD:%1.2f/%d | GeometryType:%s | BindingType:%s | Sim:%d | RBF:%d | Clip:%d(%1.2f) | VertexCount:%d | Name: %s"),
+				Line = FString::Printf(TEXT(" * Group:%d/%d | LOD:%1.2f/%d (%s) | GeometryType:%s | BindingType:%s | Sim:%d | RBF:%d | Clip:%d(%1.2f) | VertexCount:%d | Name: %s"),
 					Instance->Debug.GroupIndex,
 					Instance->Debug.GroupCount,
 
 					Instance->HairGroupPublicData->LODIndex,
 					Instance->HairGroupPublicData->GetLODScreenSizes().Num(),
+					ToString(Instance->Debug.LODSelectionTypeForDebug),
 
 					ToString(Instance->GeometryType),
 					ToString(Instance->BindingType),
diff --git a/Engine/Plugins/Runtime/HairStrands/Source/HairStrandsCore/Private/GroomManager.cpp b/Engine/Plugins/Runtime/HairStrands/Source/HairStrandsCore/Private/GroomManager.cpp
--- a/Engine/Plugins/Runtime/HairStrands/Source/HairStrandsCore/Private/GroomManager.cpp
+++ b/Engine/Plugins/Runtime/HairStrands/Source/HairStrandsCore/Private/GroomManager.cpp
@@ -607,40 +607,38 @@ static void RunHairLODSelection(
 		// Compute the view where the screen size is maximale
 		const float PrevLODIndex = Instance->HairGroupPublicData->GetLODIndex();
 		const float PrevMeshLODIndex = Instance->HairGroupPublicData->GetMeshLODIndex();
-		float LODIndex = Instance->Strands.Modifier.LODForcedIndex; // check where this is updated 
 
 		// Insure that MinLOD is necessary taken into account if a force LOD is request (i.e., LODIndex>=0). If a Force LOD 
 		// is not resquested (i.e., LODIndex<0), the MinLOD is applied after ViewLODIndex has been determined in the codeblock below
 		const float MinLOD = FMath::Max(0, GHairStrandsMinLOD);
-		LODIndex = Instance->bUseCPULODSelection && LODIndex >= 0 ? FMath::Max(LODIndex, MinLOD) : LODIndex;
-		if (LODIndex < 0 && Instance->bUseCPULODSelection)
+		float LODIndex = Instance->Debug.LODForcedIndex >= 0 ? FMath::Max(Instance->Debug.LODForcedIndex, MinLOD) : -1.0f;
+		float LODViewIndex = -1;
 		{
 			const FSphere SphereBound = Instance->ProxyBounds ? Instance->ProxyBounds->GetSphere() : FSphere(0);
 			for (const FSceneView* View : Views)
 			{
 				const float ScreenSize = ComputeBoundsScreenSize(FVector4(SphereBound.Center, 1), SphereBound.W, *View);
 				const float LODBias = Instance->Strands.Modifier.LODBias;
-				const float ViewLODIndex = FMath::Max(MinLOD, GetHairInstanceLODIndex(Instance->HairGroupPublicData->GetLODScreenSizes(), ScreenSize, LODBias));
+				const float CurrLODViewIndex = FMath::Max(MinLOD, GetHairInstanceLODIndex(Instance->HairGroupPublicData->GetLODScreenSizes(), ScreenSize, LODBias));
 
 				// Select highest LOD accross all views
-				LODIndex = LODIndex == -1 ? ViewLODIndex : FMath::Min(LODIndex, ViewLODIndex);
+				LODViewIndex = LODViewIndex < 0 ? CurrLODViewIndex : FMath::Min(LODViewIndex, CurrLODViewIndex);
 			}
 		}
+		if (LODIndex < 0)
+		{
+			LODIndex = LODViewIndex;
+		}
 
 		const TArray<EHairGeometryType>& LODGeometryTypes = Instance->HairGroupPublicData->GetLODGeometryTypes();
 		const TArray<bool>& LODVisibilities = Instance->HairGroupPublicData->GetLODVisibilities();
 		const int32 LODCount = LODVisibilities.Num();
 
-		// CPU selection: insure the LOD index is in valid range 
-		// GPU selection: -1 means auto-selection based on GPU data, >=0 means forced LOD
-		if (Instance->bUseCPULODSelection)
-		{
-			LODIndex = FMath::Clamp(LODIndex, 0.f, float(LODCount - 1));
-		}
-		else
-		{
-			LODIndex = FMath::Clamp(LODIndex, -1.f, float(LODCount - 1));
-		}
+		LODIndex = FMath::Clamp(LODIndex, 0.f, float(LODCount - 1));
+
+		// Feedback game thread with LOD selection 
+		Instance->Debug.LODPredictedIndex = LODViewIndex;
+
 		const int32 IntLODIndex = FMath::Clamp(FMath::FloorToInt(LODIndex), 0, LODCount - 1);
 		const bool bIsVisible = LODVisibilities[IntLODIndex];
 		const bool bForceCards = GHairStrands_UseCards > 0 || Instance->bForceCards; // todo
diff --git a/Engine/Plugins/Runtime/HairStrands/Source/HairStrandsCore/Private/HairStrandsRendering.cpp b/Engine/Plugins/Runtime/HairStrands/Source/HairStrandsCore/Private/HairStrandsRendering.cpp
--- a/Engine/Plugins/Runtime/HairStrands/Source/HairStrandsCore/Private/HairStrandsRendering.cpp
+++ b/Engine/Plugins/Runtime/HairStrands/Source/HairStrandsCore/Private/HairStrandsRendering.cpp
@@ -1859,8 +1859,8 @@ void ComputeHairStrandsInterpolation(
 
 			// 3. Compute cluster AABBs (used for LODing and voxelization)
 			{
-				// Optim: If an instance is using CPU selection and does not voxelize it's data, then there is no need for having valid AABB
-				bool bNeedAABB = !Instance->bUseCPULODSelection || (Instance->Strands.Modifier.bSupportVoxelization && Instance->bCastShadow);
+				// Optim: If an instance does not voxelize it's data, then there is no need for having valid AABB
+				bool bNeedAABB = Instance->Strands.Modifier.bSupportVoxelization && Instance->bCastShadow;
 
 				FHairStrandClusterData::FHairGroup* HairGroupCluster = nullptr;
 				if (CullingData.bCullingResultAvailable)
diff --git a/Engine/Plugins/Runtime/HairStrands/Source/HairStrandsCore/Public/GroomAsset.h b/Engine/Plugins/Runtime/HairStrands/Source/HairStrandsCore/Public/GroomAsset.h
--- a/Engine/Plugins/Runtime/HairStrands/Source/HairStrandsCore/Public/GroomAsset.h
+++ b/Engine/Plugins/Runtime/HairStrands/Source/HairStrandsCore/Public/GroomAsset.h
@@ -438,10 +438,6 @@ public:
 	UPROPERTY(EditAnywhere, EditFixedSize, BlueprintReadWrite, Category = "HairInterpolation", meta = (ToolTip = "Type of interpolation used (WIP)"))
 	EGroomInterpolationType HairInterpolationType = EGroomInterpolationType::RigidTransform;
 
-	/** LOD selection type  */
-	UPROPERTY(EditAnywhere, Category = "LOD", meta = (ToolTip = "LOD selection type (CPU/GPU)"))
-	EHairLODSelectionType LODSelectionType;
-
 	/** Minimum LOD to cook */
 	UPROPERTY(EditAnywhere, Category = "LOD", meta = (DisplayName = "Minimum LOD"))
 	FPerPlatformInt MinLOD;
diff --git a/Engine/Plugins/Runtime/HairStrands/Source/HairStrandsCore/Public/GroomAssetInterpolation.h b/Engine/Plugins/Runtime/HairStrands/Source/HairStrandsCore/Public/GroomAssetInterpolation.h
--- a/Engine/Plugins/Runtime/HairStrands/Source/HairStrandsCore/Public/GroomAssetInterpolation.h
+++ b/Engine/Plugins/Runtime/HairStrands/Source/HairStrandsCore/Public/GroomAssetInterpolation.h
@@ -25,13 +25,6 @@ enum class EHairInterpolationWeight : uint8
 	Unknown		UMETA(Hidden),
 };
 
-UENUM()
-enum class EHairLODSelectionType : uint8
-{
-	Cpu		UMETA(DisplayName = "CPU", ToolTip = "Use CPU information for computing the LOD selection. This used the CPU object bounds, which is not always accurate, especially when the groom is bound to a skeletal mesh."),
-	Gpu		UMETA(DisplayName = "GPU", ToolTip = "Use GPU information for computing the LOD selection. This use GPU hair clusters information which are more precise, but also more expensive to compute"),
-};
-
 UENUM(BlueprintType)
 enum class EGroomGeometryType : uint8
 {
diff --git a/Engine/Plugins/Runtime/HairStrands/Source/HairStrandsCore/Public/GroomComponent.h b/Engine/Plugins/Runtime/HairStrands/Source/HairStrandsCore/Public/GroomComponent.h
--- a/Engine/Plugins/Runtime/HairStrands/Source/HairStrandsCore/Public/GroomComponent.h
+++ b/Engine/Plugins/Runtime/HairStrands/Source/HairStrandsCore/Public/GroomComponent.h
@@ -280,10 +280,14 @@ private:
 	FVector SkeletalPreviousPositionOffset;
 	bool bIsGroomAssetCallbackRegistered;
 	bool bIsGroomBindingAssetCallbackRegistered;
-	int32 PredictedLODIndex = -1;
 	bool bValidationEnable = true;
 	bool bPreviewMode = false;
 
+	// LOD selection
+	EHairLODSelectionType LODSelectionType = EHairLODSelectionType::Immediate;
+	float LODPredictedIndex = -1.f;
+	float LODForcedIndex = -1.f;
+
 	void InitResources(bool bIsBindingReloading=false);
 	void ReleaseResources();
 
diff --git a/Engine/Plugins/Runtime/HairStrands/Source/HairStrandsCore/Public/GroomDesc.h b/Engine/Plugins/Runtime/HairStrands/Source/HairStrandsCore/Public/GroomDesc.h
--- a/Engine/Plugins/Runtime/HairStrands/Source/HairStrandsCore/Public/GroomDesc.h
+++ b/Engine/Plugins/Runtime/HairStrands/Source/HairStrandsCore/Public/GroomDesc.h
@@ -80,9 +80,5 @@ struct FHairGroupDesc
 	float HairLengthScale = 1.f;
 	UPROPERTY(EditAnywhere, BlueprintReadWrite, Category = Overrides, meta = (PinHiddenByDefault, InlineEditConditionToggle))
 	bool HairLengthScale_Override = false;
-
-	/** Force a specific LOD index */
-	UPROPERTY()
-	int32 LODForcedIndex = -1;
 };
 typedef FHairGroupDesc FHairGroupInstanceModifer;
\ No newline at end of file
diff --git a/Engine/Plugins/Runtime/HairStrands/Source/HairStrandsCore/Public/GroomInstance.h b/Engine/Plugins/Runtime/HairStrands/Source/HairStrandsCore/Public/GroomInstance.h
--- a/Engine/Plugins/Runtime/HairStrands/Source/HairStrandsCore/Public/GroomInstance.h
+++ b/Engine/Plugins/Runtime/HairStrands/Source/HairStrandsCore/Public/GroomInstance.h
@@ -40,6 +40,13 @@ END_GLOBAL_SHADER_PARAMETER_STRUCT()
 
 typedef TUniformBufferRef<FHairCardsVertexFactoryUniformShaderParameters> FHairCardsUniformBuffer;
 
+enum class EHairLODSelectionType
+{
+	Immediate,	// Done on the rendering thread, prior to render
+	Predicted,	// Predicted based on rendering->game thread feedback
+	Forced		// Forced LOD value
+};
+
 // Represent/Describe data & resources of a hair group belonging to a groom
 struct HAIRSTRANDSCORE_API FHairGroupInstance : public FHairStrandsInstance
 {
@@ -188,6 +195,10 @@ struct HAIRSTRANDSCORE_API FHairGroupInstance : public FHairStrandsInstance
 		FString					GroomAssetName;
 		uint32					LastFrameIndex = ~0;
 
+		float					LODPredictedIndex = -1.f;	// Computed on the rendering-thread, readback on the game-thread during tick()
+		float					LODForcedIndex = -1.f;		// Set by the game-thread, read on the rendering-thread to force a particular LOD
+		EHairLODSelectionType	LODSelectionTypeForDebug = EHairLODSelectionType::Immediate; // For debug only, not used
+
 		int32					MeshLODIndex = ~0;
 		EGroomBindingMeshType	GroomBindingType;
 		EGroomCacheType			GroomCacheType;
@@ -217,7 +228,6 @@ struct HAIRSTRANDSCORE_API FHairGroupInstance : public FHairStrandsInstance
 	EHairBindingType		BindingType = EHairBindingType::NoneBinding;
 	const FBoxSphereBounds*	ProxyBounds = nullptr;
 	const FBoxSphereBounds* ProxyLocalBounds = nullptr;
-	bool					bUseCPULODSelection = true;
 	bool					bForceCards = false;
 	bool					bUpdatePositionOffset = false;
 	bool					bCastShadow = true;
diff --git a/Engine/Plugins/Runtime/HairStrands/Source/HairStrandsEditor/Private/GroomAssetDetails.cpp b/Engine/Plugins/Runtime/HairStrands/Source/HairStrandsEditor/Private/GroomAssetDetails.cpp
--- a/Engine/Plugins/Runtime/HairStrands/Source/HairStrandsEditor/Private/GroomAssetDetails.cpp
+++ b/Engine/Plugins/Runtime/HairStrands/Source/HairStrandsEditor/Private/GroomAssetDetails.cpp
@@ -339,10 +339,8 @@ void FGroomRenderingDetails::CustomizeStrandsGroupProperties(IDetailLayoutBuilde
 		DetailLayout.HideProperty(DetailLayout.GetProperty(GET_MEMBER_NAME_CHECKED(UGroomAsset, HairGroupsMaterials), UGroomAsset::StaticClass()));
 		DetailLayout.HideProperty(DetailLayout.GetProperty(GET_MEMBER_NAME_CHECKED(UGroomAsset, HairGroupsLOD), UGroomAsset::StaticClass()));
 		DetailLayout.HideProperty(DetailLayout.GetProperty(GET_MEMBER_NAME_CHECKED(UGroomAsset, HairGroupsInfo), UGroomAsset::StaticClass()));
-		DetailLayout.HideProperty(DetailLayout.GetProperty(GET_MEMBER_NAME_CHECKED(UGroomAsset, LODSelectionType), UGroomAsset::StaticClass()));
 		DetailLayout.HideProperty(DetailLayout.GetProperty(GET_MEMBER_NAME_CHECKED(UGroomAsset, EnableGlobalInterpolation), UGroomAsset::StaticClass()));
 		DetailLayout.HideProperty(DetailLayout.GetProperty(GET_MEMBER_NAME_CHECKED(UGroomAsset, HairInterpolationType), UGroomAsset::StaticClass()));
-		DetailLayout.HideProperty(DetailLayout.GetProperty(GET_MEMBER_NAME_CHECKED(UGroomAsset, LODSelectionType), UGroomAsset::StaticClass()));
 		DetailLayout.HideProperty(DetailLayout.GetProperty(GET_MEMBER_NAME_CHECKED(UGroomAsset, MinLOD), UGroomAsset::StaticClass()));
 		DetailLayout.HideProperty(DetailLayout.GetProperty(GET_MEMBER_NAME_CHECKED(UGroomAsset, DisableBelowMinLodStripping), UGroomAsset::StaticClass()));
 	}
@@ -357,10 +355,8 @@ void FGroomRenderingDetails::CustomizeStrandsGroupProperties(IDetailLayoutBuilde
 		DetailLayout.HideProperty(DetailLayout.GetProperty(GET_MEMBER_NAME_CHECKED(UGroomAsset, HairGroupsMaterials), UGroomAsset::StaticClass()));
 		DetailLayout.HideProperty(DetailLayout.GetProperty(GET_MEMBER_NAME_CHECKED(UGroomAsset, HairGroupsLOD), UGroomAsset::StaticClass()));
 		DetailLayout.HideProperty(DetailLayout.GetProperty(GET_MEMBER_NAME_CHECKED(UGroomAsset, HairGroupsInfo), UGroomAsset::StaticClass()));
-		DetailLayout.HideProperty(DetailLayout.GetProperty(GET_MEMBER_NAME_CHECKED(UGroomAsset, LODSelectionType), UGroomAsset::StaticClass()));
 		DetailLayout.HideProperty(DetailLayout.GetProperty(GET_MEMBER_NAME_CHECKED(UGroomAsset, EnableGlobalInterpolation), UGroomAsset::StaticClass()));
 		DetailLayout.HideProperty(DetailLayout.GetProperty(GET_MEMBER_NAME_CHECKED(UGroomAsset, HairInterpolationType), UGroomAsset::StaticClass()));
-		DetailLayout.HideProperty(DetailLayout.GetProperty(GET_MEMBER_NAME_CHECKED(UGroomAsset, LODSelectionType), UGroomAsset::StaticClass()));
 		DetailLayout.HideProperty(DetailLayout.GetProperty(GET_MEMBER_NAME_CHECKED(UGroomAsset, MinLOD), UGroomAsset::StaticClass()));
 		DetailLayout.HideProperty(DetailLayout.GetProperty(GET_MEMBER_NAME_CHECKED(UGroomAsset, DisableBelowMinLodStripping), UGroomAsset::StaticClass()));
 	}
@@ -375,10 +371,8 @@ void FGroomRenderingDetails::CustomizeStrandsGroupProperties(IDetailLayoutBuilde
 		DetailLayout.HideProperty(DetailLayout.GetProperty(GET_MEMBER_NAME_CHECKED(UGroomAsset, HairGroupsMaterials), UGroomAsset::StaticClass()));
 		DetailLayout.HideProperty(DetailLayout.GetProperty(GET_MEMBER_NAME_CHECKED(UGroomAsset, HairGroupsLOD), UGroomAsset::StaticClass()));
 		DetailLayout.HideProperty(DetailLayout.GetProperty(GET_MEMBER_NAME_CHECKED(UGroomAsset, HairGroupsInfo), UGroomAsset::StaticClass()));
-		DetailLayout.HideProperty(DetailLayout.GetProperty(GET_MEMBER_NAME_CHECKED(UGroomAsset, LODSelectionType), UGroomAsset::StaticClass()));
 		DetailLayout.HideProperty(DetailLayout.GetProperty(GET_MEMBER_NAME_CHECKED(UGroomAsset, EnableGlobalInterpolation), UGroomAsset::StaticClass()));
 		DetailLayout.HideProperty(DetailLayout.GetProperty(GET_MEMBER_NAME_CHECKED(UGroomAsset, HairInterpolationType), UGroomAsset::StaticClass()));
-		DetailLayout.HideProperty(DetailLayout.GetProperty(GET_MEMBER_NAME_CHECKED(UGroomAsset, LODSelectionType), UGroomAsset::StaticClass()));
 		DetailLayout.HideProperty(DetailLayout.GetProperty(GET_MEMBER_NAME_CHECKED(UGroomAsset, MinLOD), UGroomAsset::StaticClass()));
 		DetailLayout.HideProperty(DetailLayout.GetProperty(GET_MEMBER_NAME_CHECKED(UGroomAsset, DisableBelowMinLodStripping), UGroomAsset::StaticClass()));
 	}
@@ -393,10 +387,8 @@ void FGroomRenderingDetails::CustomizeStrandsGroupProperties(IDetailLayoutBuilde
 		DetailLayout.HideProperty(DetailLayout.GetProperty(GET_MEMBER_NAME_CHECKED(UGroomAsset, HairGroupsMaterials), UGroomAsset::StaticClass()));
 		DetailLayout.HideProperty(DetailLayout.GetProperty(GET_MEMBER_NAME_CHECKED(UGroomAsset, HairGroupsLOD), UGroomAsset::StaticClass()));
 		DetailLayout.HideProperty(DetailLayout.GetProperty(GET_MEMBER_NAME_CHECKED(UGroomAsset, HairGroupsInfo), UGroomAsset::StaticClass()));
-		DetailLayout.HideProperty(DetailLayout.GetProperty(GET_MEMBER_NAME_CHECKED(UGroomAsset, LODSelectionType), UGroomAsset::StaticClass()));
 		DetailLayout.HideProperty(DetailLayout.GetProperty(GET_MEMBER_NAME_CHECKED(UGroomAsset, EnableGlobalInterpolation), UGroomAsset::StaticClass()));
 		DetailLayout.HideProperty(DetailLayout.GetProperty(GET_MEMBER_NAME_CHECKED(UGroomAsset, HairInterpolationType), UGroomAsset::StaticClass()));
-		DetailLayout.HideProperty(DetailLayout.GetProperty(GET_MEMBER_NAME_CHECKED(UGroomAsset, LODSelectionType), UGroomAsset::StaticClass()));
 		DetailLayout.HideProperty(DetailLayout.GetProperty(GET_MEMBER_NAME_CHECKED(UGroomAsset, MinLOD), UGroomAsset::StaticClass()));
 		DetailLayout.HideProperty(DetailLayout.GetProperty(GET_MEMBER_NAME_CHECKED(UGroomAsset, DisableBelowMinLodStripping), UGroomAsset::StaticClass()));
 	}
@@ -411,10 +403,8 @@ void FGroomRenderingDetails::CustomizeStrandsGroupProperties(IDetailLayoutBuilde
 		DetailLayout.HideProperty(DetailLayout.GetProperty(GET_MEMBER_NAME_CHECKED(UGroomAsset, HairGroupsMaterials), UGroomAsset::StaticClass()));
 		DetailLayout.HideProperty(DetailLayout.GetProperty(GET_MEMBER_NAME_CHECKED(UGroomAsset, HairGroupsLOD), UGroomAsset::StaticClass()));
 		DetailLayout.HideProperty(DetailLayout.GetProperty(GET_MEMBER_NAME_CHECKED(UGroomAsset, HairGroupsInfo), UGroomAsset::StaticClass()));
-		DetailLayout.HideProperty(DetailLayout.GetProperty(GET_MEMBER_NAME_CHECKED(UGroomAsset, LODSelectionType), UGroomAsset::StaticClass()));
 //		DetailLayout.HideProperty(DetailLayout.GetProperty(GET_MEMBER_NAME_CHECKED(UGroomAsset, EnableGlobalInterpolation), UGroomAsset::StaticClass()));
 //		DetailLayout.HideProperty(DetailLayout.GetProperty(GET_MEMBER_NAME_CHECKED(UGroomAsset, HairInterpolationType), UGroomAsset::StaticClass()));
-		DetailLayout.HideProperty(DetailLayout.GetProperty(GET_MEMBER_NAME_CHECKED(UGroomAsset, LODSelectionType), UGroomAsset::StaticClass()));
 		DetailLayout.HideProperty(DetailLayout.GetProperty(GET_MEMBER_NAME_CHECKED(UGroomAsset, MinLOD), UGroomAsset::StaticClass()));
 		DetailLayout.HideProperty(DetailLayout.GetProperty(GET_MEMBER_NAME_CHECKED(UGroomAsset, DisableBelowMinLodStripping), UGroomAsset::StaticClass()));
 	}
@@ -429,10 +419,8 @@ void FGroomRenderingDetails::CustomizeStrandsGroupProperties(IDetailLayoutBuilde
 		DetailLayout.HideProperty(DetailLayout.GetProperty(GET_MEMBER_NAME_CHECKED(UGroomAsset, HairGroupsMaterials), UGroomAsset::StaticClass()));
 //		DetailLayout.HideProperty(DetailLayout.GetProperty(GET_MEMBER_NAME_CHECKED(UGroomAsset, HairGroupsLOD), UGroomAsset::StaticClass()));
 		DetailLayout.HideProperty(DetailLayout.GetProperty(GET_MEMBER_NAME_CHECKED(UGroomAsset, HairGroupsInfo), UGroomAsset::StaticClass()));
-//		DetailLayout.HideProperty(DetailLayout.GetProperty(GET_MEMBER_NAME_CHECKED(UGroomAsset, LODSelectionType), UGroomAsset::StaticClass()));
 		DetailLayout.HideProperty(DetailLayout.GetProperty(GET_MEMBER_NAME_CHECKED(UGroomAsset, EnableGlobalInterpolation), UGroomAsset::StaticClass()));
 		DetailLayout.HideProperty(DetailLayout.GetProperty(GET_MEMBER_NAME_CHECKED(UGroomAsset, HairInterpolationType), UGroomAsset::StaticClass()));
-//		DetailLayout.HideProperty(DetailLayout.GetProperty(GET_MEMBER_NAME_CHECKED(UGroomAsset, LODSelectionType), UGroomAsset::StaticClass()));
 //		DetailLayout.HideProperty(DetailLayout.GetProperty(GET_MEMBER_NAME_CHECKED(UGroomAsset, MinLOD), UGroomAsset::StaticClass()));
 //		DetailLayout.HideProperty(DetailLayout.GetProperty(GET_MEMBER_NAME_CHECKED(UGroomAsset, DisableBelowMinLodStripping), UGroomAsset::StaticClass()));
 	}
diff --git a/Engine/Plugins/Runtime/HairStrands/Source/HairStrandsEditor/Private/GroomCustomAssetEditorToolkit.cpp b/Engine/Plugins/Runtime/HairStrands/Source/HairStrandsEditor/Private/GroomCustomAssetEditorToolkit.cpp
--- a/Engine/Plugins/Runtime/HairStrands/Source/HairStrandsEditor/Private/GroomCustomAssetEditorToolkit.cpp
+++ b/Engine/Plugins/Runtime/HairStrands/Source/HairStrandsEditor/Private/GroomCustomAssetEditorToolkit.cpp
@@ -379,7 +379,6 @@ void FGroomCustomAssetEditorToolkit::InitCustomAssetEditor(const EToolkitMode::T
 	FProperty* P5 = FindFProperty<FProperty>(GroomAsset->GetClass(), GET_MEMBER_NAME_CHECKED(UGroomAsset, HairGroupsMeshes));
 	FProperty* P6 = FindFProperty<FProperty>(GroomAsset->GetClass(), GET_MEMBER_NAME_CHECKED(UGroomAsset, HairGroupsMaterials));
 	FProperty* P7 = FindFProperty<FProperty>(GroomAsset->GetClass(), GET_MEMBER_NAME_CHECKED(UGroomAsset, HairGroupsInfo));
-	FProperty* P8 = FindFProperty<FProperty>(GroomAsset->GetClass(), GET_MEMBER_NAME_CHECKED(UGroomAsset, LODSelectionType));
 	FProperty* P9 = FindFProperty<FProperty>(GroomAsset->GetClass(), GET_MEMBER_NAME_CHECKED(UGroomAsset, AssetUserData));
 
 	P7->SetMetaData(TEXT("Category"), TEXT("Hidden"));
@@ -391,7 +390,6 @@ void FGroomCustomAssetEditorToolkit::InitCustomAssetEditor(const EToolkitMode::T
 	P5->RemoveMetaData(TEXT("ShowOnlyInnerProperties"));
 	P6->RemoveMetaData(TEXT("ShowOnlyInnerProperties"));
 	P7->RemoveMetaData(TEXT("ShowOnlyInnerProperties"));
-	P8->RemoveMetaData(TEXT("ShowOnlyInnerProperties"));
 	P9->RemoveMetaData(TEXT("ShowOnlyInnerProperties"));
 
 	// Override the display name so that it show correct/readable name when the details view is displayed as grid view
diff --git a/Engine/Plugins/Runtime/HairStrands/Source/HairStrandsEditor/Private/GroomMaterialDetails.cpp b/Engine/Plugins/Runtime/HairStrands/Source/HairStrandsEditor/Private/GroomMaterialDetails.cpp
--- a/Engine/Plugins/Runtime/HairStrands/Source/HairStrandsEditor/Private/GroomMaterialDetails.cpp
+++ b/Engine/Plugins/Runtime/HairStrands/Source/HairStrandsEditor/Private/GroomMaterialDetails.cpp
@@ -208,10 +208,8 @@ void FGroomMaterialDetails::CustomizeDetails( IDetailLayoutBuilder& DetailLayout
 		DetailLayout.HideProperty(DetailLayout.GetProperty(GET_MEMBER_NAME_CHECKED(UGroomAsset, HairGroupsMaterials), UGroomAsset::StaticClass()));
 		DetailLayout.HideProperty(DetailLayout.GetProperty(GET_MEMBER_NAME_CHECKED(UGroomAsset, HairGroupsLOD), UGroomAsset::StaticClass()));
 		DetailLayout.HideProperty(DetailLayout.GetProperty(GET_MEMBER_NAME_CHECKED(UGroomAsset, HairGroupsInfo), UGroomAsset::StaticClass()));
-		DetailLayout.HideProperty(DetailLayout.GetProperty(GET_MEMBER_NAME_CHECKED(UGroomAsset, LODSelectionType), UGroomAsset::StaticClass()));
 		DetailLayout.HideProperty(DetailLayout.GetProperty(GET_MEMBER_NAME_CHECKED(UGroomAsset, EnableGlobalInterpolation), UGroomAsset::StaticClass()));
 		DetailLayout.HideProperty(DetailLayout.GetProperty(GET_MEMBER_NAME_CHECKED(UGroomAsset, HairInterpolationType), UGroomAsset::StaticClass()));
-		DetailLayout.HideProperty(DetailLayout.GetProperty(GET_MEMBER_NAME_CHECKED(UGroomAsset, LODSelectionType), UGroomAsset::StaticClass()));
 		DetailLayout.HideProperty(DetailLayout.GetProperty(GET_MEMBER_NAME_CHECKED(UGroomAsset, MinLOD), UGroomAsset::StaticClass()));
 		DetailLayout.HideProperty(DetailLayout.GetProperty(GET_MEMBER_NAME_CHECKED(UGroomAsset, DisableBelowMinLodStripping), UGroomAsset::StaticClass()));
 	}
