Simulation support to the new predicted/forced/immediate rework

diff --git a/Engine/Plugins/Runtime/HairStrands/Source/HairStrandsCore/Private/GroomComponent.cpp b/Engine/Plugins/Runtime/HairStrands/Source/HairStrandsCore/Private/GroomComponent.cpp
--- a/Engine/Plugins/Runtime/HairStrands/Source/HairStrandsCore/Private/GroomComponent.cpp
+++ b/Engine/Plugins/Runtime/HairStrands/Source/HairStrandsCore/Private/GroomComponent.cpp
@@ -1349,31 +1349,12 @@ void UGroomComponent::ReleaseHairSimulation()
 	NiagaraComponents.Empty();
 }
 
-void UGroomComponent::UpdateHairSimulation()
+void UGroomComponent::CreateHairSimulation(const int32 GroupIndex, const int32 LODIndex)
 {
-	const int32 NumGroups = GroomAsset ? GroomAsset->HairGroupsPhysics.Num() : 0;
-	const int32 NumComponents = FMath::Max(NumGroups, NiagaraComponents.Num());
-
-	TArray<bool> ValidComponents;
-	ValidComponents.Init(false, NumComponents);
-
-	if (GroomAsset)
-	{
-		for (int32 GroupIndex = 0; GroupIndex < NumGroups; ++GroupIndex)
-		{
-			const int32 NumLODs = GroomAsset->HairGroupsLOD[GroupIndex].LODs.Num();
-			for (int32 LODIndex = 0; LODIndex < NumLODs; ++LODIndex)
-			{
-				ValidComponents[GroupIndex] = IsSimulationEnable(GroupIndex, LODIndex);
-			}
-		}
-	}
-
-	NiagaraComponents.SetNumZeroed(NumComponents);
-	for (int32 CompIndex = 0; CompIndex < NumComponents; ++CompIndex)
+	if (GroupIndex < NiagaraComponents.Num())
 	{
-		UE_TRANSITIONAL_OBJECT_PTR(UNiagaraComponent)& NiagaraComponent = NiagaraComponents[CompIndex];
-		if (ValidComponents[CompIndex])
+		UE_TRANSITIONAL_OBJECT_PTR(UNiagaraComponent)& NiagaraComponent = NiagaraComponents[GroupIndex];
+		if (GroomAsset && (GroupIndex < GroomAsset->HairGroupsPhysics.Num()) && (LODIndex < GroomAsset->HairGroupsLOD[GroupIndex].LODs.Num()) && IsSimulationEnable(GroupIndex, LODIndex))
 		{
 			if (!NiagaraComponent)
 			{
@@ -1388,13 +1369,13 @@ void UGroomComponent::UpdateHairSimulation()
 					NiagaraComponent->RegisterComponentWithWorld(GetWorld());
 				}
 
-				if(!AngularSpringsSystem) AngularSpringsSystem = LoadObject<UNiagaraSystem>(nullptr, TEXT("/HairStrands/Emitters/StableSpringsSystem.StableSpringsSystem"));
-				if(!CosseratRodsSystem) CosseratRodsSystem = LoadObject<UNiagaraSystem>(nullptr, TEXT("/HairStrands/Emitters/StableRodsSystem.StableRodsSystem"));
+				if (!AngularSpringsSystem) AngularSpringsSystem = LoadObject<UNiagaraSystem>(nullptr, TEXT("/HairStrands/Emitters/StableSpringsSystem.StableSpringsSystem"));
+				if (!CosseratRodsSystem) CosseratRodsSystem = LoadObject<UNiagaraSystem>(nullptr, TEXT("/HairStrands/Emitters/StableRodsSystem.StableRodsSystem"));
 
 				UNiagaraSystem* NiagaraAsset =
-					(GroomAsset->HairGroupsPhysics[CompIndex].SolverSettings.NiagaraSolver == EGroomNiagaraSolvers::AngularSprings) ? ToRawPtr(AngularSpringsSystem) :
-					(GroomAsset->HairGroupsPhysics[CompIndex].SolverSettings.NiagaraSolver == EGroomNiagaraSolvers::CosseratRods) ? ToRawPtr(CosseratRodsSystem) :
-					GroomAsset->HairGroupsPhysics[CompIndex].SolverSettings.CustomSystem.LoadSynchronous();
+					(GroomAsset->HairGroupsPhysics[GroupIndex].SolverSettings.NiagaraSolver == EGroomNiagaraSolvers::AngularSprings) ? ToRawPtr(AngularSpringsSystem) :
+					(GroomAsset->HairGroupsPhysics[GroupIndex].SolverSettings.NiagaraSolver == EGroomNiagaraSolvers::CosseratRods) ? ToRawPtr(CosseratRodsSystem) :
+					GroomAsset->HairGroupsPhysics[GroupIndex].SolverSettings.CustomSystem.LoadSynchronous();
 
 				NiagaraComponent->SetAsset(NiagaraAsset);
 				NiagaraComponent->ReinitializeSystem();
@@ -1406,9 +1387,45 @@ void UGroomComponent::UpdateHairSimulation()
 			NiagaraComponent->DeactivateImmediate();
 		}
 	}
+}
+
+void UGroomComponent::UpdateHairSimulation()
+{
+	const int32 NumGroups = GroomAsset ? GroomAsset->HairGroupsPhysics.Num() : 0;
+	const int32 NumComponents = FMath::Max(NumGroups, NiagaraComponents.Num());
+
+	NiagaraComponents.SetNumZeroed(NumComponents);
+	const int32 LODInit = (LODForcedIndex != -1.f) ? LODForcedIndex : (LODPredictedIndex != -1.f) ? LODPredictedIndex : 0;
+	for (int32 CompIndex = 0; CompIndex < NumComponents; ++CompIndex)
+	{
+		CreateHairSimulation(CompIndex, LODInit);
+	}
 	UpdateSimulatedGroups();
 }
 
+void UGroomComponent::SwitchSimulationLOD(const int32 PreviousLOD, const int32 CurrentLOD)
+{
+	if (GroomAsset && PreviousLOD != CurrentLOD)
+	{
+		bool bRequiresSimulationUpdate = false;
+		for (int32 GroupIt = 0, GroupCount = GroomAsset->HairGroupsData.Num(); GroupIt < GroupCount; ++GroupIt)
+		{
+			if ((IsSimulationEnable(GroupIt, PreviousLOD) != IsSimulationEnable(GroupIt, CurrentLOD)))
+			{
+				// Sanity check. Simulation does not support immediate mode.
+				check(LODSelectionType != EHairLODSelectionType::Immediate);
+
+				CreateHairSimulation(GroupIt, CurrentLOD);
+				bRequiresSimulationUpdate = true;
+			}
+		}
+		if (bRequiresSimulationUpdate)
+		{
+			UpdateSimulatedGroups();
+		}
+	}
+}
+
 void UGroomComponent::SetGroomAsset(UGroomAsset* Asset)
 {
 	SetGroomAsset(Asset, BindingAsset);
@@ -1590,26 +1607,9 @@ void UGroomComponent::SetForcedLOD(int32 CurrLODIndex)
 		}
 		CurrLODSelectionType = bNeedSimulationNeed ? EHairLODSelectionType::Predicted : EHairLODSelectionType::Immediate;
 	}
-	
-	const bool bHasLODChanged = CurrLODIndex != PrevLODIndex;
 
 	// Inform simulation about LOD switch.
-	if (bHasLODChanged)
-	{
-		if (GroomAsset)
-		{
-			for (int32 GroupIt = 0, GroupCount = GroomAsset->HairGroupsData.Num(); GroupIt < GroupCount; ++GroupIt)
-			{
-				if ((IsSimulationEnable(GroupIt, PrevLODIndex) != IsSimulationEnable(GroupIt, CurrLODIndex)))
-				{
-					// Sanity check. Simulation does not support immediate mode.
-					check(LODSelectionType != EHairLODSelectionType::Immediate);
-
-					UpdateHairSimulation();
-				}
-			}
-		}
-	}
+	SwitchSimulationLOD(PrevLODIndex, CurrLODIndex);
 
 	// Finally, update the forced LOD value and LOD selection type
 	LODForcedIndex	 = CurrLODIndex;
@@ -3214,16 +3214,7 @@ void UGroomComponent::TickComponent(float DeltaTime, enum ELevelTick TickType, F
 		}
 
 		// 2. If there is a LOD change, update the simulation adequately 
-		if (FMath::Abs(LODPredictedIndex - EffectiveForceLOD) > KINDA_SMALL_NUMBER)
-		{
-			for (int32 GroupIt = 0, GroupCount = GroomAsset->HairGroupsData.Num(); GroupIt < GroupCount; ++GroupIt)
-			{
-				if ((IsSimulationEnable(GroupIt, LODPredictedIndex) != IsSimulationEnable(GroupIt, EffectiveForceLOD)))
-				{
-					UpdateHairSimulation();
-				}
-			}
-		}
+		SwitchSimulationLOD(LODPredictedIndex, EffectiveForceLOD);
 
 		// 3. Update global predicted index (used for SyncLOD API)
 		LODPredictedIndex = EffectiveForceLOD;
diff --git a/Engine/Plugins/Runtime/HairStrands/Source/HairStrandsCore/Public/GroomComponent.h b/Engine/Plugins/Runtime/HairStrands/Source/HairStrandsCore/Public/GroomComponent.h
--- a/Engine/Plugins/Runtime/HairStrands/Source/HairStrandsCore/Public/GroomComponent.h
+++ b/Engine/Plugins/Runtime/HairStrands/Source/HairStrandsCore/Public/GroomComponent.h
@@ -86,6 +86,12 @@ public:
 	/** Release Niagara components */
 	void ReleaseHairSimulation();
 
+	/** Create per Group/LOD the Niagara component */
+	void CreateHairSimulation(const int32 GroupIndex, const int32 LODIndex);
+
+	/** Enable/Disable hair simulation while transitioning from one LOD to another one */
+	void SwitchSimulationLOD(const int32 PreviousLOD, const int32 CurrentLOD);
+
 	/** Check if the simulation is enabled or not */
 	bool IsSimulationEnable(int32 GroupIndex, int32 LODIndex) const;
 
