Fix missing tangent frame deformation on cards when deformed by skinned surface.

This is noticeable under strong skinning deformation.

diff --git a/Engine/Plugins/Runtime/HairStrands/Source/HairStrandsCore/Private/GroomResources.cpp b/Engine/Plugins/Runtime/HairStrands/Source/HairStrandsCore/Private/GroomResources.cpp
--- a/Engine/Plugins/Runtime/HairStrands/Source/HairStrandsCore/Private/GroomResources.cpp
+++ b/Engine/Plugins/Runtime/HairStrands/Source/HairStrandsCore/Private/GroomResources.cpp
@@ -651,18 +651,24 @@ void FHairCardsDeformedResource::InternalAllocate(FRDGBuilder& GraphBuilder)
 {
 	if (bInitializedData)
 	{
-		InternalCreateVertexBufferRDG<FHairCardsPositionFormat>(GraphBuilder, BulkData.Positions, DeformedPositionBuffer[0], ToHairResourceDebugName(TEXT("Hair.CardsDeformed(Current)"), ResourceName), EHairResourceUsageType::Dynamic);
-		InternalCreateVertexBufferRDG<FHairCardsPositionFormat>(GraphBuilder, BulkData.Positions, DeformedPositionBuffer[1], ToHairResourceDebugName(TEXT("Hair.CardsDeformed(Previous)"), ResourceName), EHairResourceUsageType::Dynamic);
+		InternalCreateVertexBufferRDG<FHairCardsPositionFormat>(GraphBuilder, BulkData.Positions, DeformedPositionBuffer[0], ToHairResourceDebugName(TEXT("Hair.CardsDeformedPosition(Current)"), ResourceName), EHairResourceUsageType::Dynamic);
+		InternalCreateVertexBufferRDG<FHairCardsPositionFormat>(GraphBuilder, BulkData.Positions, DeformedPositionBuffer[1], ToHairResourceDebugName(TEXT("Hair.CardsDeformedPosition(Previous)"), ResourceName), EHairResourceUsageType::Dynamic);
+
+		InternalCreateVertexBufferRDG<FHairCardsNormalFormat>(GraphBuilder, BulkData.Normals, DeformedNormalBuffer, ToHairResourceDebugName(TEXT("Hair.CardsDeformedNormal"), ResourceName), EHairResourceUsageType::Dynamic);
 	}
 	else
 	{
-		InternalCreateVertexBufferRDG<FHairCardsPositionFormat>(GraphBuilder, BulkData.GetNumVertices(), DeformedPositionBuffer[0], ToHairResourceDebugName(TEXT("Hair.CardsDeformed(Current)"), ResourceName), EHairResourceUsageType::Dynamic);
-		InternalCreateVertexBufferRDG<FHairCardsPositionFormat>(GraphBuilder, BulkData.GetNumVertices(), DeformedPositionBuffer[1], ToHairResourceDebugName(TEXT("Hair.CardsDeformed(Previous)"), ResourceName), EHairResourceUsageType::Dynamic);
+		InternalCreateVertexBufferRDG<FHairCardsPositionFormat>(GraphBuilder, BulkData.GetNumVertices(), DeformedPositionBuffer[0], ToHairResourceDebugName(TEXT("Hair.CardsDeformedPosition(Current)"), ResourceName), EHairResourceUsageType::Dynamic);
+		InternalCreateVertexBufferRDG<FHairCardsPositionFormat>(GraphBuilder, BulkData.GetNumVertices(), DeformedPositionBuffer[1], ToHairResourceDebugName(TEXT("Hair.CardsDeformedPosition(Previous)"), ResourceName), EHairResourceUsageType::Dynamic);
+
+		InternalCreateVertexBufferRDG<FHairCardsNormalFormat>(GraphBuilder, BulkData.GetNumVertices() * FHairCardsNormalFormat::ComponentCount, DeformedNormalBuffer, ToHairResourceDebugName(TEXT("Hair.CardsDeformedNormal"), ResourceName), EHairResourceUsageType::Dynamic);
 
 		// Manually transit to SRVs, in case of the cards are rendered not visible, but still rendered (in shadows for instance). In such a case, the cards deformation pass is not called, and thus the 
 		// buffers are never transit from UAV (clear) to SRV for rasterization. 
 		HairTransition::TransitToSRV(GraphBuilder, Register(GraphBuilder, DeformedPositionBuffer[0], ERDGImportedBufferFlags::CreateSRV).SRV, ERDGPassFlags::Raster);
 		HairTransition::TransitToSRV(GraphBuilder, Register(GraphBuilder, DeformedPositionBuffer[1], ERDGImportedBufferFlags::CreateSRV).SRV, ERDGPassFlags::Raster);
+
+		HairTransition::TransitToSRV(GraphBuilder, Register(GraphBuilder, DeformedNormalBuffer, ERDGImportedBufferFlags::CreateSRV).SRV, ERDGPassFlags::Raster);
 	}
 }
 
@@ -670,6 +676,7 @@ void FHairCardsDeformedResource::InternalRelease()
 {
 	DeformedPositionBuffer[0].Release();
 	DeformedPositionBuffer[1].Release();
+	DeformedNormalBuffer.Release();
 }
 
 /////////////////////////////////////////////////////////////////////////////////////////
diff --git a/Engine/Plugins/Runtime/HairStrands/Source/HairStrandsCore/Private/HairCardsVertexFactory.cpp b/Engine/Plugins/Runtime/HairStrands/Source/HairStrandsCore/Private/HairCardsVertexFactory.cpp
--- a/Engine/Plugins/Runtime/HairStrands/Source/HairStrandsCore/Private/HairCardsVertexFactory.cpp
+++ b/Engine/Plugins/Runtime/HairStrands/Source/HairStrandsCore/Private/HairCardsVertexFactory.cpp
@@ -47,14 +47,15 @@ FHairCardsUniformBuffer CreateHairCardsVFUniformBuffer(
 		{
 			UniformParameters.PositionBuffer = LOD.DeformedResource->GetBuffer(FHairCardsDeformedResource::EFrameType::Current).SRV.GetReference();
 			UniformParameters.PreviousPositionBuffer = LOD.DeformedResource->GetBuffer(FHairCardsDeformedResource::EFrameType::Previous).SRV.GetReference();
+			UniformParameters.NormalsBuffer = LOD.DeformedResource->DeformedNormalBuffer.SRV.GetReference();
 		}
 		else
 		{
 			UniformParameters.PositionBuffer = LOD.RestResource->RestPositionBuffer.ShaderResourceViewRHI;
 			UniformParameters.PreviousPositionBuffer = LOD.RestResource->RestPositionBuffer.ShaderResourceViewRHI;
+			UniformParameters.NormalsBuffer = LOD.RestResource->NormalsBuffer.ShaderResourceViewRHI.GetReference();
 		}
 
-		UniformParameters.NormalsBuffer = LOD.RestResource->NormalsBuffer.ShaderResourceViewRHI.GetReference();
 		UniformParameters.UVsBuffer = LOD.RestResource->UVsBuffer.ShaderResourceViewRHI.GetReference();
 
 		UniformParameters.DepthTexture = LOD.RestResource->DepthTexture;
diff --git a/Engine/Plugins/Runtime/HairStrands/Source/HairStrandsCore/Private/HairStrandsRendering.cpp b/Engine/Plugins/Runtime/HairStrands/Source/HairStrandsCore/Private/HairStrandsRendering.cpp
--- a/Engine/Plugins/Runtime/HairStrands/Source/HairStrandsCore/Private/HairStrandsRendering.cpp
+++ b/Engine/Plugins/Runtime/HairStrands/Source/HairStrandsCore/Private/HairStrandsRendering.cpp
@@ -1017,6 +1017,7 @@ class FHairCardsDeformationCS : public FGlobalShader
 		SHADER_PARAMETER_RDG_BUFFER_SRV(Buffer, GuideDeformedPositionBuffer)
 		SHADER_PARAMETER_RDG_BUFFER_SRV(Buffer, GuideDeformedPositionOffsetBuffer)
 		SHADER_PARAMETER_SRV(Buffer, CardsRestPositionBuffer)
+		SHADER_PARAMETER_SRV(Buffer, CardsRestTangentBuffer)
 		SHADER_PARAMETER_RDG_BUFFER_SRV(Buffer, CardsInterpolationBuffer)
 
 		SHADER_PARAMETER_RDG_BUFFER_SRV(Buffer<float4>, TriangleRestPosition0Buffer)
@@ -1030,9 +1031,8 @@ class FHairCardsDeformationCS : public FGlobalShader
 		SHADER_PARAMETER_RDG_BUFFER_SRV(Buffer<uint>, GuideRootBarycentricBuffer)
 		SHADER_PARAMETER_RDG_BUFFER_SRV(Buffer<uint>, GuideVertexToRootIndexBuffer)
 
-
 		SHADER_PARAMETER_RDG_BUFFER_UAV(RWBuffer, CardsDeformedPositionBuffer)
-
+		SHADER_PARAMETER_RDG_BUFFER_UAV(RWBuffer, CardsDeformedTangentBuffer)
 
 		SHADER_PARAMETER_STRUCT_INCLUDE(ShaderDrawDebug::FShaderParameters, ShaderDrawParameters)
 	END_SHADER_PARAMETER_STRUCT()
@@ -1059,6 +1059,8 @@ static void AddHairCardsDeformationPass(
 	FRDGImportedBuffer CardsDeformedPositionBuffer_Curr = Register(GraphBuilder, LOD.DeformedResource->GetBuffer(FHairCardsDeformedResource::Current), ERDGImportedBufferFlags::CreateUAV);
 	FRDGImportedBuffer CardsDeformedPositionBuffer_Prev = Register(GraphBuilder, LOD.DeformedResource->GetBuffer(FHairCardsDeformedResource::Previous), ERDGImportedBufferFlags::CreateUAV);
 	AddCopyBufferPass(GraphBuilder, CardsDeformedPositionBuffer_Prev.Buffer, CardsDeformedPositionBuffer_Curr.Buffer);
+	
+	FRDGImportedBuffer CardsDeformedNormalBuffer = Register(GraphBuilder, LOD.DeformedResource->DeformedNormalBuffer, ERDGImportedBufferFlags::CreateUAV);
 
 	FHairCardsDeformationCS::FParameters* Parameters = GraphBuilder.AllocParameters<FHairCardsDeformationCS::FParameters>();
 	Parameters->GuideVertexCount			= LOD.Guides.RestResource->GetVertexCount();
@@ -1069,7 +1071,9 @@ static void AddHairCardsDeformationPass(
 	
 	Parameters->CardsVertexCount			= LOD.RestResource->GetVertexCount();
 	Parameters->CardsRestPositionBuffer		= LOD.RestResource->RestPositionBuffer.ShaderResourceViewRHI;
+	Parameters->CardsRestTangentBuffer		= LOD.RestResource->NormalsBuffer.ShaderResourceViewRHI;
 	Parameters->CardsDeformedPositionBuffer = CardsDeformedPositionBuffer_Curr.UAV;
+	Parameters->CardsDeformedTangentBuffer	= CardsDeformedNormalBuffer.UAV;
 
 	Parameters->CardsInterpolationBuffer	= RegisterAsSRV(GraphBuilder, LOD.InterpolationResource->InterpolationBuffer);
 
@@ -1119,13 +1123,14 @@ static void AddHairCardsDeformationPass(
 	const int32 DispatchCountX = FMath::DivideAndRoundUp(Parameters->CardsVertexCount, GroupSize);
 	FComputeShaderUtils::AddPass(
 		GraphBuilder,
-		RDG_EVENT_NAME("HairStrands::CardsDeformation"),
+		RDG_EVENT_NAME("HairStrands::CardsDeformation(%s)", bSupportDynamicMesh ? TEXT("Dynamic") : TEXT("Static")),
 		ComputeShader,
 		Parameters,
 		FIntVector(DispatchCountX,1,1));
 
 	GraphBuilder.SetBufferAccessFinal(CardsDeformedPositionBuffer_Curr.Buffer, ERHIAccess::SRVMask);
 	GraphBuilder.SetBufferAccessFinal(CardsDeformedPositionBuffer_Prev.Buffer, ERHIAccess::SRVMask);
+	GraphBuilder.SetBufferAccessFinal(CardsDeformedNormalBuffer.Buffer, ERHIAccess::SRVMask);
 }
 
 ///////////////////////////////////////////////////////////////////////////////////////////////////
diff --git a/Engine/Plugins/Runtime/HairStrands/Source/HairStrandsCore/Public/GroomResources.h b/Engine/Plugins/Runtime/HairStrands/Source/HairStrandsCore/Public/GroomResources.h
--- a/Engine/Plugins/Runtime/HairStrands/Source/HairStrandsCore/Public/GroomResources.h
+++ b/Engine/Plugins/Runtime/HairStrands/Source/HairStrandsCore/Public/GroomResources.h
@@ -605,6 +605,7 @@ struct FHairCardsDeformedResource : public FHairCommonResource
 
 	/* Strand hair deformed position buffer (previous and current) */
 	FRDGExternalBuffer DeformedPositionBuffer[2];
+	FRDGExternalBuffer DeformedNormalBuffer;
 
 	/* Reference to the hair strands render data */
 	const FHairCardsBulkData& BulkData;
diff --git a/Engine/Shaders/Private/HairStrands/HairCardsDeformation.usf b/Engine/Shaders/Private/HairStrands/HairCardsDeformation.usf
--- a/Engine/Shaders/Private/HairStrands/HairCardsDeformation.usf
+++ b/Engine/Shaders/Private/HairStrands/HairCardsDeformation.usf
@@ -16,6 +16,9 @@ Buffer<uint>   CardsInterpolationBuffer;
 
 RWBuffer<float4> CardsDeformedPositionBuffer; 
 
+Buffer<float4> CardsRestTangentBuffer;
+RWBuffer<float4> CardsDeformedTangentBuffer;
+
 // Guides == Strands going through the middle of the card geometry
 Buffer<float4>  TriangleRestPosition0Buffer;
 Buffer<float4>  TriangleRestPosition1Buffer;
@@ -55,8 +58,25 @@ float3 ComputeDynamicGeometryOffset(
 
 	return (LocalDeformedGuidePoint - LocalRestGuidePoint) * GuideVertexWeight;
 }
+
+void TransformTangent(
+	uint InVertexId, 
+	uint InOffset,
+	FHairMeshTriangle InRestTriangle,
+	FHairMeshTriangle InDeformedTriangle)
+{
+	float4 Rest		= CardsRestTangentBuffer[InVertexId * 2 + InOffset];
+	float3 Local	= VectorToTriangle(Rest.xyz, InRestTriangle);
+	float3 Deformed = VectorToWorld(Local, InDeformedTriangle);
+	CardsDeformedTangentBuffer[InVertexId * 2 + InOffset] = float4(Deformed, Rest.w);
+}
 #endif
 
+void CopyTangent(uint InVertexId, uint InOffset)
+{
+	CardsDeformedTangentBuffer[InVertexId * 2 + InOffset] = CardsRestTangentBuffer[InVertexId * 2 + InOffset];
+}
+
 [numthreads(GROUP_SIZE, 1, 1)]
 void MainCS(uint2 DispatchThreadId : SV_DispatchThreadID)
 {
@@ -114,6 +134,9 @@ void MainCS(uint2 DispatchThreadId : SV_DispatchThreadID)
 	const float3 OutPosition = RestPosition.xyz + DeformationOffset;
 	CardsDeformedPositionBuffer[VertexId] = float4(OutPosition, RestPosition.w);
 
+	CopyTangent(VertexId, 0);
+	CopyTangent(VertexId, 1);
+
 #elif PERMUTATION_DYNAMIC_GEOMETRY == 1
 	// Compute the deformation offset in triangle local space
 
@@ -140,5 +163,8 @@ void MainCS(uint2 DispatchThreadId : SV_DispatchThreadID)
 
 	CardsDeformedPositionBuffer[VertexId] = float4(ControlPoint, RestControlPoint.w);
 
+	// Transform tangent
+	TransformTangent(VertexId, 0, RestTriangle, DeformedTriangle);
+	TransformTangent(VertexId, 1, RestTriangle, DeformedTriangle);
 #endif
 }
\ No newline at end of file
