Add Groom Alembic supports for GroupName

diff --git a/Engine/Plugins/Runtime/HairStrands/Source/HairStrandsCore/Private/GroomBuilder.cpp b/Engine/Plugins/Runtime/HairStrands/Source/HairStrandsCore/Private/GroomBuilder.cpp
--- a/Engine/Plugins/Runtime/HairStrands/Source/HairStrandsCore/Private/GroomBuilder.cpp
+++ b/Engine/Plugins/Runtime/HairStrands/Source/HairStrandsCore/Private/GroomBuilder.cpp
@@ -1500,6 +1500,7 @@ bool FGroomBuilder::BuildHairDescriptionGroups(const FHairDescription& HairDescr
 	TStrandAttributesConstRef<int> StrandGuides = HairDescription.StrandAttributes().GetAttributesRef<int>(HairAttribute::Strand::Guide);
 	TStrandAttributesConstRef<int> GroupIDs = HairDescription.StrandAttributes().GetAttributesRef<int>(HairAttribute::Strand::GroupID);
 	TStrandAttributesConstRef<int> StrandIDs = HairDescription.StrandAttributes().GetAttributesRef<int>(HairAttribute::Strand::ID);
+	TStrandAttributesConstRef<FName> GroupNames = HairDescription.StrandAttributes().GetAttributesRef<FName>(HairAttribute::Strand::GroupName);
 
 	bool bImportGuides = true;
 
@@ -1510,7 +1511,7 @@ bool FGroomBuilder::BuildHairDescriptionGroups(const FHairDescription& HairDescr
 	// must include StrandID attribute since ClosestGuides references those IDs
 	const bool bCanUseClosestGuidesAndWeights = bImportGuides && StrandIDs.IsValid() && ClosestGuides.IsValid() && GuideWeights.IsValid();
 
-	auto FindOrAdd = [&Out](int32 GroupID) -> FHairDescriptionGroup&
+	auto FindOrAdd = [&Out](int32 GroupID, FName GroupName) -> FHairDescriptionGroup&
 	{
 		for (FHairDescriptionGroup& Group : Out.HairGroups)
 		{
@@ -1521,6 +1522,7 @@ bool FGroomBuilder::BuildHairDescriptionGroups(const FHairDescription& HairDescr
 		}
 		FHairDescriptionGroup& Group = Out.HairGroups.AddDefaulted_GetRef();
 		Group.Info.GroupID = GroupID;
+		Group.Info.GroupName = GroupName;
 		return Group;
 	};
 
@@ -1546,13 +1548,15 @@ bool FGroomBuilder::BuildHairDescriptionGroups(const FHairDescription& HairDescr
 		}
 
 		int32 GroupID = 0;
+		FName GroupName = NAME_None;
 		if (GroupIDs.IsValid())
 		{
 			GroupID = GroupIDs[StrandID];
+			GroupName = GroupNames[StrandID];
 		}
 
 		FHairStrandsDatas* CurrentHairStrandsDatas = nullptr;
-		FHairDescriptionGroup& Group = FindOrAdd(GroupID);
+		FHairDescriptionGroup& Group = FindOrAdd(GroupID, GroupName);
 		check(Group.Info.GroupID == GroupID);
 		if (!bIsGuide)
 		{
diff --git a/Engine/Plugins/Runtime/HairStrands/Source/HairStrandsCore/Private/HairAttributes.cpp b/Engine/Plugins/Runtime/HairStrands/Source/HairStrandsCore/Private/HairAttributes.cpp
--- a/Engine/Plugins/Runtime/HairStrands/Source/HairStrandsCore/Private/HairAttributes.cpp
+++ b/Engine/Plugins/Runtime/HairStrands/Source/HairStrandsCore/Private/HairAttributes.cpp
@@ -22,6 +22,7 @@ namespace HairAttribute
 	const FName Strand::BasisType("groom_basis_type");
 	const FName Strand::CurveType("groom_curve_type");
 	const FName Strand::Knots("groom_knots");
+	const FName Strand::GroupName("groom_group_name");
 
 	const FName Groom::Color("groom_color");
 	const FName Groom::Roughness("groom_roughness");
diff --git a/Engine/Plugins/Runtime/HairStrands/Source/HairStrandsCore/Public/GroomAsset.h b/Engine/Plugins/Runtime/HairStrands/Source/HairStrandsCore/Public/GroomAsset.h
--- a/Engine/Plugins/Runtime/HairStrands/Source/HairStrandsCore/Public/GroomAsset.h
+++ b/Engine/Plugins/Runtime/HairStrands/Source/HairStrandsCore/Public/GroomAsset.h
@@ -40,6 +40,9 @@ struct HAIRSTRANDSCORE_API FHairGroupInfo
 	UPROPERTY(VisibleAnywhere, Category = "Info")
 	int32 GroupID = 0;
 
+	UPROPERTY(VisibleAnywhere, Category = "Info")
+	FName GroupName = NAME_None;
+
 	UPROPERTY(VisibleAnywhere, Category = "Info", meta = (DisplayName = "Curve Count"))
 	int32 NumCurves = 0;
 
diff --git a/Engine/Plugins/Runtime/HairStrands/Source/HairStrandsCore/Public/HairAttributes.h b/Engine/Plugins/Runtime/HairStrands/Source/HairStrandsCore/Public/HairAttributes.h
--- a/Engine/Plugins/Runtime/HairStrands/Source/HairStrandsCore/Public/HairAttributes.h
+++ b/Engine/Plugins/Runtime/HairStrands/Source/HairStrandsCore/Public/HairAttributes.h
@@ -29,6 +29,7 @@ namespace HairAttribute
 		extern HAIRSTRANDSCORE_API const FName BasisType;		// FName (EGroomBasisType)
 		extern HAIRSTRANDSCORE_API const FName CurveType;		// FName (EGroomCurveType)
 		extern HAIRSTRANDSCORE_API const FName Knots;			// float[]
+		extern HAIRSTRANDSCORE_API const FName GroupName;		// FName
 	}
 
 	namespace Groom
diff --git a/Engine/Plugins/Runtime/HairStrands/Source/HairStrandsEditor/Private/GroomAssetDetails.cpp b/Engine/Plugins/Runtime/HairStrands/Source/HairStrandsEditor/Private/GroomAssetDetails.cpp
--- a/Engine/Plugins/Runtime/HairStrands/Source/HairStrandsEditor/Private/GroomAssetDetails.cpp
+++ b/Engine/Plugins/Runtime/HairStrands/Source/HairStrandsEditor/Private/GroomAssetDetails.cpp
@@ -1145,6 +1145,15 @@ TSharedRef<SWidget> FGroomRenderingDetails::MakeGroupNameButtonCustomization(int
 	return SNullWidget::NullWidget;
 }
 
+FName FGroomRenderingDetails::GetGroupName(int32 GroupIndex) const
+{
+	if (GroomAsset && GroupIndex >= 0 && GroupIndex < GroomAsset->HairGroupsInfo.Num())
+	{
+		return GroomAsset->HairGroupsInfo[GroupIndex].GroupName;
+	}
+	return NAME_None;
+}
+
 TSharedRef<SWidget> FGroomRenderingDetails::MakeGroupNameCustomization(int32 GroupIndex)
 {
 	switch (PanelType)
@@ -1167,10 +1176,21 @@ TSharedRef<SWidget> FGroomRenderingDetails::MakeGroupNameCustomization(int32 Gro
 	break;
 	default:
 	{
-		return SNew(STextBlock)
-			.Font(IDetailLayoutBuilder::GetDetailFont())
-			.ColorAndOpacity(HairGroupColor)
-			.Text(FText::Format(LOCTEXT("Group", "Group ID {0} "), FText::AsNumber(GroupIndex)));
+		FName GroupName = GetGroupName(GroupIndex);
+		if (GroupName != NAME_None)
+		{
+			return SNew(STextBlock)
+				.Font(IDetailLayoutBuilder::GetDetailFont())
+				.ColorAndOpacity(HairGroupColor)
+				.Text(FText::Format(LOCTEXT("Group", "Group ID {0} - {1}"), FText::AsNumber(GroupIndex), FText::FromName(GroupName)));
+		}
+		else
+		{
+			return SNew(STextBlock)
+				.Font(IDetailLayoutBuilder::GetDetailFont())
+				.ColorAndOpacity(HairGroupColor)
+				.Text(FText::Format(LOCTEXT("Group", "Group ID {0}"), FText::AsNumber(GroupIndex)));
+		}
 	}
 	break;
 	}
diff --git a/Engine/Plugins/Runtime/HairStrands/Source/HairStrandsEditor/Private/GroomAssetDetails.h b/Engine/Plugins/Runtime/HairStrands/Source/HairStrandsEditor/Private/GroomAssetDetails.h
--- a/Engine/Plugins/Runtime/HairStrands/Source/HairStrandsEditor/Private/GroomAssetDetails.h
+++ b/Engine/Plugins/Runtime/HairStrands/Source/HairStrandsEditor/Private/GroomAssetDetails.h
@@ -49,7 +49,8 @@ private:
 	FName& GetMaterialSlotName(int32 GroupIndex);
 	const FName& GetMaterialSlotName(int32 GroupIndex) const;
 	int32 GetGroupCount() const;
-	
+	FName GetGroupName(int32 GroupIndex) const;
+
 	void AddNewGroupButton(IDetailCategoryBuilder& StrandsGroupFilesCategory, FProperty* Property);
 
 	void ApplyChanges();
