Fix crash when rebuilding groom asset whose have no GroupName attribute.

diff --git a/Engine/Plugins/Runtime/HairStrands/Source/HairStrandsCore/Private/GroomBuilder.cpp b/Engine/Plugins/Runtime/HairStrands/Source/HairStrandsCore/Private/GroomBuilder.cpp
--- a/Engine/Plugins/Runtime/HairStrands/Source/HairStrandsCore/Private/GroomBuilder.cpp
+++ b/Engine/Plugins/Runtime/HairStrands/Source/HairStrandsCore/Private/GroomBuilder.cpp
@@ -1522,7 +1522,7 @@ bool FGroomBuilder::BuildHairDescriptionGroups(const FHairDescription& HairDescr
 		}
 		FHairDescriptionGroup& Group = Out.HairGroups.AddDefaulted_GetRef();
 		Group.Info.GroupID = GroupID;
-		Group.Info.GroupName = GroupName;
+		Group.Info.GroupName = GroupName != NAME_None ? GroupName : FName(FString::Printf(TEXT("Group_%d"), GroupID));
 		return Group;
 	};
 
@@ -1552,7 +1552,7 @@ bool FGroomBuilder::BuildHairDescriptionGroups(const FHairDescription& HairDescr
 		if (GroupIDs.IsValid())
 		{
 			GroupID = GroupIDs[StrandID];
-			GroupName = GroupNames[StrandID];
+			GroupName = GroupNames.IsValid() && StrandID < GroupNames.GetNumElements() ? GroupNames[StrandID] : NAME_None;
 		}
 
 		FHairStrandsDatas* CurrentHairStrandsDatas = nullptr;
@@ -1713,6 +1713,7 @@ void FGroomBuilder::BuildData(
 
 	{
 		OutGroupInfo.GroupID = InHairDescriptionGroup.Info.GroupID;
+		OutGroupInfo.GroupName = InHairDescriptionGroup.Info.GroupName;
 
 		// Rendering data
 		{
