Fix crash when assigning material with bUseStrands relevance to a non-groom component.
Fix material not-being assigned when dragging&droppring to a groom component in viewport.

diff --git a/Engine/Plugins/Runtime/HairStrands/Source/HairStrandsCore/Private/GroomComponent.cpp b/Engine/Plugins/Runtime/HairStrands/Source/HairStrandsCore/Private/GroomComponent.cpp
--- a/Engine/Plugins/Runtime/HairStrands/Source/HairStrandsCore/Private/GroomComponent.cpp
+++ b/Engine/Plugins/Runtime/HairStrands/Source/HairStrandsCore/Private/GroomComponent.cpp
@@ -3383,7 +3383,10 @@ void UGroomComponent::PostEditChangeProperty(FPropertyChangedEvent& PropertyChan
 	}
 	#endif
 
-	const bool bRecreateResources = bAssetChanged || bBindingAssetChanged || bGroomCacheChanged || bEnableLengthScaleOverrideChanged || PropertyThatChanged == nullptr || bSourceSkeletalMeshChanged || bRayTracingGeometryChanged || bEnableSolverChanged;
+	// If material is assigned to the groom from the viewport (i.e., drag&drop a material from the content brown onto the groom geometry, it results into a unknown property). There is other case 
+	const bool bIsUnknown = PropertyThatChanged == nullptr;
+
+	const bool bRecreateResources = bAssetChanged || bBindingAssetChanged || bGroomCacheChanged || bEnableLengthScaleOverrideChanged || bIsUnknown || bSourceSkeletalMeshChanged || bRayTracingGeometryChanged || bEnableSolverChanged;
 	if (bRecreateResources)
 	{
 		// Release the resources before Super::PostEditChangeProperty so that they get
@@ -3434,6 +3437,11 @@ void UGroomComponent::PostEditChangeProperty(FPropertyChangedEvent& PropertyChan
 		bIsEventProcess = true;
 	}
 
+	if (bIsUnknown)
+	{
+		InitResources(false);
+	}
+
 	// Always call parent PostEditChangeProperty as parent expect such a call (would crash in certain case otherwise)
 	//if (!bIsEventProcess)
 	{
diff --git a/Engine/Source/Runtime/Renderer/Private/HairStrands/HairStrandsInterface.cpp b/Engine/Source/Runtime/Renderer/Private/HairStrands/HairStrandsInterface.cpp
--- a/Engine/Source/Runtime/Renderer/Private/HairStrands/HairStrandsInterface.cpp
+++ b/Engine/Source/Runtime/Renderer/Private/HairStrands/HairStrandsInterface.cpp
@@ -442,6 +442,10 @@ FHairStrandsBookmarkParameters CreateHairStrandsBookmarkParameters(FScene* Scene
 	Out.ViewUniqueID			= View.ViewState ? View.ViewState->UniqueID : ~0;
 	Out.SceneColorTexture		= nullptr;
 	Out.bStrandsGeometryEnabled = IsHairStrandsEnabled(EHairStrandsShaderType::Strands, View.GetShaderPlatform());
+
+	// Sanity check
+	check(Out.Instances->Num() >= Out.VisibleInstances.Num());
+
 	if (GHairStrandsParameterFunction)
 	{
 		GHairStrandsParameterFunction(Out);
@@ -463,3 +467,16 @@ FHairStrandsBookmarkParameters CreateHairStrandsBookmarkParameters(FScene* Scene
 
 	return Out;
 }
+
+bool IsHairStrandsCompatible(const FMeshBatch* Mesh)
+{
+	if (Mesh)
+	{
+		static const FHashedName& VFType0 = FVertexFactoryType::GetVFByName(TEXT("FHairStrandsVertexFactory"))->GetHashedName();
+		static const FHashedName& VFType1 = FVertexFactoryType::GetVFByName(TEXT("FHairCardsVertexFactory"))->GetHashedName();
+		const FHashedName& VFType = Mesh->VertexFactory->GetType()->GetHashedName();
+		return VFType == VFType0 || VFType == VFType1;
+	}
+	return false;
+}
+
diff --git a/Engine/Source/Runtime/Renderer/Private/SceneVisibility.cpp b/Engine/Source/Runtime/Renderer/Private/SceneVisibility.cpp
--- a/Engine/Source/Runtime/Renderer/Private/SceneVisibility.cpp
+++ b/Engine/Source/Runtime/Renderer/Private/SceneVisibility.cpp
@@ -2988,7 +2988,7 @@ void ComputeDynamicMeshRelevance(EShadingPath ShadingPath, bool bAddLightmapDens
 		BatchAndProxy.SortKey = MeshBatch.PrimitiveSceneProxy->GetTranslucencySortPriority();
 	}
 	
-	const bool bIsHairStrandsCompatible = ViewRelevance.bHairStrands && IsHairStrandsEnabled(EHairStrandsShaderType::All, View.GetShaderPlatform());
+	const bool bIsHairStrandsCompatible = ViewRelevance.bHairStrands && IsHairStrandsEnabled(EHairStrandsShaderType::All, View.GetShaderPlatform()) && IsHairStrandsCompatible(MeshBatch.Mesh);
 	if (bIsHairStrandsCompatible)
 	{
 		View.HairStrandsMeshElements.AddUninitialized(1);
diff --git a/Engine/Source/Runtime/Renderer/Public/HairStrandsInterface.h b/Engine/Source/Runtime/Renderer/Public/HairStrandsInterface.h
--- a/Engine/Source/Runtime/Renderer/Public/HairStrandsInterface.h
+++ b/Engine/Source/Runtime/Renderer/Public/HairStrandsInterface.h
@@ -388,6 +388,8 @@ RENDERER_API bool IsHairStrandsSupported(EHairStrandsShaderType Type, EShaderPla
 RENDERER_API bool IsHairStrandsEnabled(EHairStrandsShaderType Type, EShaderPlatform Platform = EShaderPlatform::SP_NumPlatforms);
 RENDERER_API void SetHairStrandsEnabled(bool In);
 
+RENDERER_API bool IsHairStrandsCompatible(const FMeshBatch* Mesh);
+
 RENDERER_API bool IsHairRayTracingEnabled();
 
 // Return true if hair simulation is enabled.
