Improve groom decimation.

Fix invalid hair decimation with decimation factor larger than 0.5.

This CL improves overall decimation accuracy overall.

diff --git a/Engine/Plugins/Runtime/HairStrands/Source/HairStrandsCore/Private/GroomBuilder.cpp b/Engine/Plugins/Runtime/HairStrands/Source/HairStrandsCore/Private/GroomBuilder.cpp
--- a/Engine/Plugins/Runtime/HairStrands/Source/HairStrandsCore/Private/GroomBuilder.cpp
+++ b/Engine/Plugins/Runtime/HairStrands/Source/HairStrandsCore/Private/GroomBuilder.cpp
@@ -1888,21 +1888,30 @@ void Decimate(
 	// Divide strands in buckets and pick randomly one stand per bucket
 	const uint32 InCurveCount = InData.StrandsCurves.Num();
 	const uint32 OutCurveCount = FMath::Clamp(uint32(InCurveCount * CurveDecimationPercentage), 1u, InCurveCount);
-	const uint32 CurveBucketSize = InCurveCount / OutCurveCount;
 
 	TArray<uint32> CurveIndices;
 	CurveIndices.SetNum(OutCurveCount);
 
 	uint32 OutTotalPointCount = 0;
 	FRandomStream Random;
+
+	const float CurveBucketSize = float(InCurveCount) / float(OutCurveCount);
+	int32 LastCurveIndex = -1;
 	for (uint32 BucketIndex = 0; BucketIndex < OutCurveCount; BucketIndex++)
 	{
-		const uint32 CurveIndex = BucketIndex * CurveBucketSize;// +BucketSize * Random.FRand();
-		CurveIndices[BucketIndex] = CurveIndex;
+		const float MinBucket = FMath::Max(BucketIndex  * CurveBucketSize, float(LastCurveIndex+1));
+		const float MaxBucket = (BucketIndex+1) * CurveBucketSize;
+		const float AdjustedBucketSize = MaxBucket - MinBucket;
+		if (AdjustedBucketSize > 0)
+		{
+			const uint32 CurveIndex = FMath::FloorToInt(MinBucket + Random.FRand() * AdjustedBucketSize);
+			CurveIndices[BucketIndex] = CurveIndex;
+			LastCurveIndex = CurveIndex;
 
-		const uint32 InPointCount = InData.StrandsCurves.CurvesCount[CurveIndex];
-		const uint32 OutPointCount = DecimatePointCount(InPointCount, VertexDecimationPercentage);
-		OutTotalPointCount += OutPointCount;
+			const uint32 InPointCount = InData.StrandsCurves.CurvesCount[CurveIndex];
+			const uint32 OutPointCount = DecimatePointCount(InPointCount, VertexDecimationPercentage);
+			OutTotalPointCount += OutPointCount;
+		}
 	}
 
 	OutData.StrandsCurves.SetNum(OutCurveCount);
