Plumbing a way for separate plugins to provide a HairCard generation solution (falling back to the pre-existing internal solution if an external one does not exist).

diff --git a/Engine/Plugins/Runtime/HairStrands/HairStrands.uplugin b/Engine/Plugins/Runtime/HairStrands/HairStrands.uplugin
--- a/Engine/Plugins/Runtime/HairStrands/HairStrands.uplugin
+++ b/Engine/Plugins/Runtime/HairStrands/HairStrands.uplugin
@@ -25,6 +25,11 @@
 			"Type": "Editor",
 			"LoadingPhase": "Default"
 		},
+		{
+			"Name": "HairCardGeneratorFramework",
+			"Type": "Editor",
+			"LoadingPhase": "Default"
+		}
 	],
 	"Plugins": [
 		{
diff --git a/Engine/Plugins/Runtime/HairStrands/Source/HairCardGeneratorFramework/HairCardGeneratorFramework.Build.cs b/Engine/Plugins/Runtime/HairStrands/Source/HairCardGeneratorFramework/HairCardGeneratorFramework.Build.cs
new file mode 100644
--- /dev/null
+++ b/Engine/Plugins/Runtime/HairStrands/Source/HairCardGeneratorFramework/HairCardGeneratorFramework.Build.cs
@@ -0,0 +1,49 @@
+// Copyright Epic Games, Inc. All Rights Reserved.
+
+using UnrealBuildTool;
+using System.IO; // for Path
+public class HairCardGeneratorFramework : ModuleRules
+{
+	public HairCardGeneratorFramework(ReadOnlyTargetRules Target) : base(Target)
+	{
+		PCHUsage = ModuleRules.PCHUsageMode.UseExplicitOrSharedPCHs;
+
+		PublicIncludePaths.AddRange(
+			new string[] {
+				// ... add public include paths required here ...
+			}
+		);
+
+
+		PrivateIncludePaths.AddRange(
+			new string[] {
+				// ... add other private include paths required here ...
+			}
+		);
+
+
+		PublicDependencyModuleNames.AddRange(
+			new string[]
+			{
+				"Core",
+				"CoreUObject",
+				"Engine",
+				// ... add other public dependencies that you statically link with here ...
+			}
+		);
+
+		PrivateDependencyModuleNames.AddRange(
+			new string[]
+			{
+				// ... add private dependencies that you statically link with here ...	
+			}
+		);
+
+		DynamicallyLoadedModuleNames.AddRange(
+			new string[]
+			{
+				// ... add any modules that your module loads dynamically here ...
+			}
+		);
+	}
+}
diff --git a/Engine/Plugins/Runtime/HairStrands/Source/HairCardGeneratorFramework/Private/HairCardGenerator.cpp b/Engine/Plugins/Runtime/HairStrands/Source/HairCardGeneratorFramework/Private/HairCardGenerator.cpp
new file mode 100644
--- /dev/null
+++ b/Engine/Plugins/Runtime/HairStrands/Source/HairCardGeneratorFramework/Private/HairCardGenerator.cpp
@@ -0,0 +1,7 @@
+// Copyright Epic Games, Inc. All Rights Reserved.
+
+#pragma once
+
+#include "IHairCardGenerator.h"
+
+const FName IHairCardGenerator::ModularFeatureName(TEXT("HairCardGenerator"));
diff --git a/Engine/Plugins/Runtime/HairStrands/Source/HairCardGeneratorFramework/Private/HairCardGeneratorFramework.cpp b/Engine/Plugins/Runtime/HairStrands/Source/HairCardGeneratorFramework/Private/HairCardGeneratorFramework.cpp
new file mode 100644
--- /dev/null
+++ b/Engine/Plugins/Runtime/HairStrands/Source/HairCardGeneratorFramework/Private/HairCardGeneratorFramework.cpp
@@ -0,0 +1,5 @@
+// Copyright Epic Games, Inc. All Rights Reserved.
+
+#include "Modules/ModuleManager.h" // for IMPLEMENT_MODULE(), FDefaultModuleImpl
+
+IMPLEMENT_MODULE(FDefaultModuleImpl, HairCardGeneratorFramework);
\ No newline at end of file
diff --git a/Engine/Plugins/Runtime/HairStrands/Source/HairCardGeneratorFramework/Public/IHairCardGenerator.h b/Engine/Plugins/Runtime/HairStrands/Source/HairCardGeneratorFramework/Public/IHairCardGenerator.h
new file mode 100644
--- /dev/null
+++ b/Engine/Plugins/Runtime/HairStrands/Source/HairCardGeneratorFramework/Public/IHairCardGenerator.h
@@ -0,0 +1,30 @@
+// Copyright Epic Games, Inc. All Rights Reserved.
+
+#pragma once
+
+#include "Features/IModularFeature.h"
+#include "Features/IModularFeatures.h" // for RegisterModularFeature()/UnregisterModularFeature()
+
+class  UGroomAsset;
+struct FHairGroupsCardsSourceDescription;
+
+class HAIRCARDGENERATORFRAMEWORK_API IHairCardGenerator : public IModularFeature
+{
+public:
+	static const FName ModularFeatureName; // "HairCardGenerator"
+
+	virtual bool GenerateHairCardsForLOD(const UGroomAsset* Groom, FHairGroupsCardsSourceDescription& CardsDesc) = 0;
+};
+
+namespace HairCardGenerator_Utils
+{
+	void RegisterModularHairCardGenerator(IHairCardGenerator* Generator)
+	{
+		IModularFeatures::Get().RegisterModularFeature(IHairCardGenerator::ModularFeatureName, Generator);
+	}
+
+	void UnregisterModularHairCardGenerator(IHairCardGenerator* Generator)
+	{
+		IModularFeatures::Get().UnregisterModularFeature(IHairCardGenerator::ModularFeatureName, Generator);
+	}
+}
\ No newline at end of file
diff --git a/Engine/Plugins/Runtime/HairStrands/Source/HairStrandsEditor/HairStrandsEditor.Build.cs b/Engine/Plugins/Runtime/HairStrands/Source/HairStrandsEditor/HairStrandsEditor.Build.cs
--- a/Engine/Plugins/Runtime/HairStrands/Source/HairStrandsEditor/HairStrandsEditor.Build.cs
+++ b/Engine/Plugins/Runtime/HairStrands/Source/HairStrandsEditor/HairStrandsEditor.Build.cs
@@ -36,7 +36,8 @@ namespace UnrealBuildTool.Rules
 					"LevelSequence",
 					"MovieScene",
 					"MovieSceneTools",
-					"Sequencer"
+					"Sequencer",
+					"HairCardGeneratorFramework"
 				});
 			AddEngineThirdPartyPrivateStaticDependencies(Target,
 			 "FBX"
diff --git a/Engine/Plugins/Runtime/HairStrands/Source/HairStrandsEditor/Private/GroomAssetDetails.cpp b/Engine/Plugins/Runtime/HairStrands/Source/HairStrandsEditor/Private/GroomAssetDetails.cpp
--- a/Engine/Plugins/Runtime/HairStrands/Source/HairStrandsEditor/Private/GroomAssetDetails.cpp
+++ b/Engine/Plugins/Runtime/HairStrands/Source/HairStrandsEditor/Private/GroomAssetDetails.cpp
@@ -42,6 +42,7 @@
 #include "ContentBrowserModule.h"
 #include "EditorFramework/AssetImportData.h"
 #include "Logging/LogMacros.h"
+#include "IHairCardGenerator.h"
 
 #include "MeshDescription.h"
 #include "MeshAttributes.h"
@@ -66,6 +67,8 @@ static FLinearColor HairLODColor(1.0f, 0.5f, 0.0f);
 
 #define LOCTEXT_NAMESPACE "GroomRenderingDetails"
 
+DEFINE_LOG_CATEGORY_STATIC(LogGroomAssetDetails, Log, All);
+
 static int32 GHairCardsProcerudalResolution = 4096;
 static int32 GHairCardsProcerudalResolution_LOD0 = -1;
 static int32 GHairCardsProcerudalResolution_LOD1 = -1;
@@ -924,12 +927,12 @@ void FGroomRenderingDetails::ResetToDefault(TSharedPtr<IPropertyHandle> ChildHan
 	CommonResetToDefault(ChildHandle, GroupIndex, LODIndex, true);
 }
 
-void FGroomRenderingDetails::AddPropertyWithCustomReset(TSharedPtr<IPropertyHandle>& PropertyHandle, IDetailChildrenBuilder& Builder, int32 GroupIndex, int32 LODIndex)
+IDetailPropertyRow& FGroomRenderingDetails::AddPropertyWithCustomReset(TSharedPtr<IPropertyHandle>& PropertyHandle, IDetailChildrenBuilder& Builder, int32 GroupIndex, int32 LODIndex)
 {	
 	FIsResetToDefaultVisible IsResetVisible = FIsResetToDefaultVisible::CreateSP(this, &FGroomRenderingDetails::ShouldResetToDefault, GroupIndex, LODIndex);
 	FResetToDefaultHandler ResetHandler = FResetToDefaultHandler::CreateSP(this, &FGroomRenderingDetails::ResetToDefault, GroupIndex, LODIndex);
 	FResetToDefaultOverride ResetOverride = FResetToDefaultOverride::Create(IsResetVisible, ResetHandler);
-	Builder.AddProperty(PropertyHandle.ToSharedRef()).OverrideResetToDefault(ResetOverride);
+	return Builder.AddProperty(PropertyHandle.ToSharedRef()).OverrideResetToDefault(ResetOverride);
 }
 
 void FGroomRenderingDetails::ExpandStruct(TSharedPtr<IPropertyHandle>& PropertyHandle, IDetailChildrenBuilder& ChildrenBuilder, int32 GroupIndex, int32 LODIndex, bool bOverrideReset)
@@ -1041,6 +1044,20 @@ FReply FGroomRenderingDetails::OnSaveCards(int32 DescIndex, FProperty* Property)
 	return FReply::Handled();
 }
 
+FReply FGroomRenderingDetails::OnGenerateCardDataUsingPlugin(int32 GroupIndex)
+{
+	if (GroomAsset->HairGroupsCards.IsValidIndex(GroupIndex))
+	{
+		TArray<IHairCardGenerator*> HairCardGeneratorPlugins = IModularFeatures::Get().GetModularFeatureImplementations<IHairCardGenerator>(IHairCardGenerator::ModularFeatureName);
+		if (HairCardGeneratorPlugins.Num() > 0)
+		{
+			UE_CLOG(HairCardGeneratorPlugins.Num() > 1, LogGroomAssetDetails, Warning, TEXT("There are more than one available hair-card generator options. Defaulting to the first one found."));
+			HairCardGeneratorPlugins[0]->GenerateHairCardsForLOD(GroomAsset, GroomAsset->HairGroupsCards[GroupIndex]);
+		}
+	}
+	return FReply::Handled();
+}
+
 void FGroomRenderingDetails::AddLODSlot(TSharedRef<IPropertyHandle>& LODHandle, IDetailChildrenBuilder& ChildrenBuilder, int32 GroupIndex, int32 LODIndex)
 {	
 	ExpandStruct(LODHandle, ChildrenBuilder, GroupIndex, LODIndex, true);
@@ -1326,60 +1343,156 @@ void FGroomRenderingDetails::OnGenerateElementForHairGroup(TSharedRef<IPropertyH
 		break;
 		case EMaterialPanelType::Cards:
 		{
-			if (PropertyName == GET_MEMBER_NAME_CHECKED(FHairGroupsCardsSourceDescription, SourceType))
+			if (PropertyName == GET_MEMBER_NAME_CHECKED(FHairGroupsCardsSourceDescription, CardsInfo))
 			{
-				// Add the Source type selection and just below add a save button
-				AddPropertyWithCustomReset(ChildHandle, ChildrenBuilder, GroupIndex, -1);
-				if (GroomAsset != nullptr && GroupIndex >= 0 && GroupIndex < GroomAsset->HairGroupsCards.Num() && (PanelType == EMaterialPanelType::Cards))
+				// Not node display
+			}
+			else
+			{
+				IDetailPropertyRow& PropertyRow = AddPropertyWithCustomReset(ChildHandle, ChildrenBuilder, GroupIndex, -1);
+
+				auto CustomizeMeshPropertyRow = [](IDetailPropertyRow& PropertyRow, bool bEnable)->FDetailWidgetRow&
 				{
-					FText ToolTipTextForGeneration(FText::FromString(TEXT("Generate procedural cards data (meshes and textures) based on current procedural settings. Cards generation needs to run prior to the (re)loading of the cards data.")));
-					FText ToolTipTextForReloading(FText::FromString(TEXT("(Re)Load generated cards data (meshes and textures) into the groom asset. The data need to be generated with the save/generated button prior to reloading.")));
+					TSharedPtr<SWidget> NameWidget;
+					TSharedPtr<SWidget> ValueWidget;
+					FDetailWidgetRow Row;
+					PropertyRow.GetDefaultWidgets(NameWidget, ValueWidget, Row);
 
-					ChildrenBuilder.AddCustomRow(LOCTEXT("HairCardsButtons", "HairCardsButtons"))
-					.ValueContent()
-					.HAlign(HAlign_Fill)
-					[
+					ValueWidget->SetEnabled(bEnable);
 
-						SNew(SHorizontalBox)
-						+ SHorizontalBox::Slot()
-						.AutoWidth()
+					return PropertyRow.CustomWidget()
+						.NameContent()
 						[
-							SNew(SButton)
-							.VAlign(VAlign_Center)
-							.HAlign(HAlign_Center)
-							.ButtonStyle(FEditorStyle::Get(), "HoverHintOnly")
-							.ToolTipText(ToolTipTextForGeneration)
-							.OnClicked(this, &FGroomRenderingDetails::OnSaveCards, GroupIndex, Property)
+							PropertyRow.GetPropertyHandle()->CreatePropertyNameWidget(
+								LOCTEXT("HairCardsMeshProperty", "Mesh"),
+								LOCTEXT("HairCardsMeshTooltop",  "")
+							)
+						]
+						.ValueContent()
+						[
+							ValueWidget.ToSharedRef()
+						];
+				};
+
+				TWeakObjectPtr<UGroomAsset> GroomAssetPtr = GroomAsset;
+				auto ShouldShowProceduralProperties = [GroomAssetPtr, GroupIndex]()
+				{
+					return GroomAssetPtr.IsValid() && GroomAssetPtr->HairGroupsCards.IsValidIndex(GroupIndex) &&
+						GroomAssetPtr->HairGroupsCards[GroupIndex].SourceType == EHairCardsSourceType::Procedural;
+				};
+				TAttribute<EVisibility> ProceduralPropertyVisibility = TAttribute<EVisibility>::CreateLambda([ShouldShowProceduralProperties]()
+					{
+						return ShouldShowProceduralProperties() ? EVisibility::Visible : EVisibility::Collapsed;
+					}
+				);
+				
+				if (PropertyName == GET_MEMBER_NAME_CHECKED(FHairGroupsCardsSourceDescription, ProceduralMesh))
+				{
+					PropertyRow.Visibility(ProceduralPropertyVisibility);
+					CustomizeMeshPropertyRow(PropertyRow, /*bEnable =*/false);
+				}
+				else if (PropertyName == GET_MEMBER_NAME_CHECKED(FHairGroupsCardsSourceDescription, ProceduralSettings))
+				{
+					PropertyRow.Visibility(ProceduralPropertyVisibility);
+				}
+				else if (PropertyName == GET_MEMBER_NAME_CHECKED(FHairGroupsCardsSourceDescription, ImportedMesh))
+				{
+					TAttribute<EVisibility> NonProceduralPropertyVisibility = TAttribute<EVisibility>::CreateLambda([ShouldShowProceduralProperties]()
+						{
+							return ShouldShowProceduralProperties() ? EVisibility::Collapsed : EVisibility::Visible;
+						}
+					);
+					PropertyRow.Visibility(NonProceduralPropertyVisibility);
+
+					CustomizeMeshPropertyRow(PropertyRow, /*bEnable =*/true);
+				}
+				else if (PropertyName == GET_MEMBER_NAME_CHECKED(FHairGroupsCardsSourceDescription, SourceType))
+				{
+					TSharedPtr<SWidget> NameWidget;
+					TSharedPtr<SWidget> ValueWidget;
+					FDetailWidgetRow Row;
+					PropertyRow.GetDefaultWidgets(NameWidget, ValueWidget, Row);
+
+					TSharedRef<SHorizontalBox> ProceduralGenButtons = SNew(SHorizontalBox).Visibility(ProceduralPropertyVisibility);
+					TArray<IHairCardGenerator*> HairCardGeneratorPlugins = IModularFeatures::Get().GetModularFeatureImplementations<IHairCardGenerator>(IHairCardGenerator::ModularFeatureName);
+					if (HairCardGeneratorPlugins.Num() > 0)
+					{
+						FText RegenToolTipText = LOCTEXT("GenerateNewCardsTooltip", "Generate new card assets (meshes and textures) using the current procedural settings. NOTE: This will overwrite preexisting card assets for this LOD.");
+						ProceduralGenButtons->AddSlot()
+							.AutoWidth()
 							[
-								SNew(SImage)
-								.Image(FEditorStyle::GetBrush("AssetEditor.SaveAsset"))
-							]
+								SNew(SButton)
+									.VAlign(VAlign_Center)
+									.HAlign(HAlign_Center)
+									.ButtonStyle(FEditorStyle::Get(), "HoverHintOnly")
+									.ToolTipText(RegenToolTipText)
+									.OnClicked(this, &FGroomRenderingDetails::OnGenerateCardDataUsingPlugin, GroupIndex)
+									[
+										SNew(SImage)
+											// @TODO: Need a specialized icon for this?
+											.Image(FEditorStyle::GetBrush("ContentBrowser.AssetActions.ReimportAsset"))
+											.RenderTransform(FSlateRenderTransform(FQuat2D(FMath::DegreesToRadians(90.0f))))
+											.RenderTransformPivot(FVector2D(0.5f, 0.5f))
+									]
+							];
+					}
+					else
+					{
+						FText ToolTipTextForGeneration(FText::FromString(TEXT("Generate procedural cards data (meshes and textures) based on current procedural settings. Cards generation needs to run prior to the (re)loading of the cards data.")));
+						FText ToolTipTextForReloading(FText::FromString(TEXT("(Re)Load generated cards data (meshes and textures) into the groom asset. The data need to be generated with the save/generated button prior to reloading.")));
+
+						ProceduralGenButtons->AddSlot()
+							.AutoWidth()
+							[
+								SNew(SButton)
+								.VAlign(VAlign_Center)
+								.HAlign(HAlign_Center)
+								.ButtonStyle(FEditorStyle::Get(), "HoverHintOnly")
+								.ToolTipText(ToolTipTextForGeneration)
+								.OnClicked(this, &FGroomRenderingDetails::OnSaveCards, GroupIndex, Property)
+								[
+									SNew(SImage)
+									.Image(FEditorStyle::GetBrush("AssetEditor.SaveAsset"))
+								]
+							];
+
+						ProceduralGenButtons->AddSlot()
+							.AutoWidth()
+							[
+								SNew(SButton)
+								.VAlign(VAlign_Center)
+								.HAlign(HAlign_Center)
+								.ButtonStyle(FEditorStyle::Get(), "HoverHintOnly")
+								.ToolTipText(ToolTipTextForReloading)
+								.OnClicked(this, &FGroomRenderingDetails::OnRefreshCards, GroupIndex, Property)
+								[
+									SNew(SImage)
+									.Image(FEditorStyle::GetBrush("Icons.Refresh"))
+								]
+							];
+					}
+
+					PropertyRow.CustomWidget()
+						.NameContent()
+						[
+							NameWidget.ToSharedRef()
 						]
-						+ SHorizontalBox::Slot()
-						.AutoWidth()
+						.ValueContent()
 						[
-							SNew(SButton)
-							.VAlign(VAlign_Center)
-							.HAlign(HAlign_Center)
-							.ButtonStyle(FEditorStyle::Get(), "HoverHintOnly")
-							.ToolTipText(ToolTipTextForReloading)
-							.OnClicked(this, &FGroomRenderingDetails::OnRefreshCards, GroupIndex, Property)
+							SNew(SVerticalBox)
+							+ SVerticalBox::Slot()
+							.AutoHeight()
 							[
-								SNew(SImage)
-								.Image(FEditorStyle::GetBrush("Icons.Refresh"))
+								ValueWidget.ToSharedRef()
 							]
-						]
-					];
+							+ SVerticalBox::Slot()
+							.AutoHeight()
+							[
+								ProceduralGenButtons
+							]
+						];
 				}
 			}
-			else if (PropertyName == GET_MEMBER_NAME_CHECKED(FHairGroupsCardsSourceDescription, CardsInfo))
-			{
-				// Not node display
-			}
-			else
-			{
-				AddPropertyWithCustomReset(ChildHandle, ChildrenBuilder, GroupIndex, -1);
-			}
 		}
 		break;
 		case EMaterialPanelType::Meshes:
diff --git a/Engine/Plugins/Runtime/HairStrands/Source/HairStrandsEditor/Private/GroomAssetDetails.h b/Engine/Plugins/Runtime/HairStrands/Source/HairStrandsEditor/Private/GroomAssetDetails.h
--- a/Engine/Plugins/Runtime/HairStrands/Source/HairStrandsEditor/Private/GroomAssetDetails.h
+++ b/Engine/Plugins/Runtime/HairStrands/Source/HairStrandsEditor/Private/GroomAssetDetails.h
@@ -18,6 +18,7 @@ class IGroomCustomAssetEditorToolkit;
 class IDetailLayoutBuilder;
 class IDetailCategoryBuilder;
 class FDetailWidgetRow;
+class IDetailPropertyRow;
 
 enum class EMaterialPanelType
 {
@@ -82,10 +83,11 @@ private:
 	FReply OnAddGroup(FProperty* Property);
 	FReply OnRefreshCards(int32 GroupIndex, FProperty* Property);
 	FReply OnSaveCards(int32 GroupIndex, FProperty* Property);
+	FReply OnGenerateCardDataUsingPlugin(int32 GroupIndex);
 
 	void ExpandStruct(TSharedPtr<IPropertyHandle>& PropertyHandle, IDetailChildrenBuilder& ChildrenBuilder, int32 GroupIndex, int32 LODIndex, bool bOverrideReset);
 	void ExpandStruct(TSharedRef<IPropertyHandle>& PropertyHandle, IDetailChildrenBuilder& ChildrenBuilder, int32 GroupIndex, int32 LODIndex, bool bOverrideReset);
-	void AddPropertyWithCustomReset(TSharedPtr<IPropertyHandle>& PropertyHandle, IDetailChildrenBuilder& Builder, int32 GroupIndex, int32 LODIndex);
+	IDetailPropertyRow& AddPropertyWithCustomReset(TSharedPtr<IPropertyHandle>& PropertyHandle, IDetailChildrenBuilder& Builder, int32 GroupIndex, int32 LODIndex);
 
 	bool CommonResetToDefault(TSharedPtr<IPropertyHandle> ChildHandle, int32 GroupIndex, int32 LODIndex, bool bSetValue);
 	bool ShouldResetToDefault(TSharedPtr<IPropertyHandle> ChildHandle, int32 GroupIndex, int32 LODInex);
