Fixing CIS error from misplaced #pragma once.

diff --git a/Engine/Plugins/Runtime/HairStrands/Source/HairCardGeneratorFramework/Private/HairCardGenerator.cpp b/Engine/Plugins/Runtime/HairStrands/Source/HairCardGeneratorFramework/Private/HairCardGenerator.cpp
--- a/Engine/Plugins/Runtime/HairStrands/Source/HairCardGeneratorFramework/Private/HairCardGenerator.cpp
+++ b/Engine/Plugins/Runtime/HairStrands/Source/HairCardGeneratorFramework/Private/HairCardGenerator.cpp
@@ -1,7 +1,5 @@
 // Copyright Epic Games, Inc. All Rights Reserved.
 
-#pragma once
-
 #include "IHairCardGenerator.h"
 
 const FName IHairCardGenerator::ModularFeatureName(TEXT("HairCardGenerator"));
