Fix crash when selecting hair cards panel.

diff --git a/Engine/Plugins/Runtime/HairStrands/Source/HairStrandsEditor/Private/GroomAssetDetails.cpp b/Engine/Plugins/Runtime/HairStrands/Source/HairStrandsEditor/Private/GroomAssetDetails.cpp
--- a/Engine/Plugins/Runtime/HairStrands/Source/HairStrandsEditor/Private/GroomAssetDetails.cpp
+++ b/Engine/Plugins/Runtime/HairStrands/Source/HairStrandsEditor/Private/GroomAssetDetails.cpp
@@ -796,7 +796,7 @@ bool FGroomRenderingDetails::CommonResetToDefault(TSharedPtr<IPropertyHandle> Ch
 	}
 
 	// LODs
-	if (bIsGroupIndexValid && LODIndex >= 0)
+	if (bIsGroupIndexValid)
 	{
 		{
 			FHairGroupsLOD Default;
@@ -804,6 +804,7 @@ bool FGroomRenderingDetails::CommonResetToDefault(TSharedPtr<IPropertyHandle> Ch
 			HAIR_RESET0(HairGroupsLOD, FHairGroupsLOD, ClusterScreenSizeScale);
 		}
 
+		if (LODIndex>=0 && LODIndex < GroomAsset->HairGroupsLOD[GroupIndex].LODs.Num())
 		{
 			FHairLODSettings Default;
 			HAIR_RESET1(HairGroupsLOD, FHairLODSettings, LODs[LODIndex], CurveDecimation);
