* Adding a way for the external hair card generator plugin to pack it's own generation options with the groom asset.
* Modifying the hair card options to reflect the options in the standalone python tool from research.
* Marshalling the new options into the python runtime ([REDACTED_HANDLE]: utilize the passed settings).
* [WIP] Cleaning up some of the C++ code.

diff --git a/Engine/Plugins/Runtime/HairStrands/Source/HairCardGeneratorFramework/Public/IHairCardGenerator.h b/Engine/Plugins/Runtime/HairStrands/Source/HairCardGeneratorFramework/Public/IHairCardGenerator.h
--- a/Engine/Plugins/Runtime/HairStrands/Source/HairCardGeneratorFramework/Public/IHairCardGenerator.h
+++ b/Engine/Plugins/Runtime/HairStrands/Source/HairCardGeneratorFramework/Public/IHairCardGenerator.h
@@ -7,13 +7,16 @@
 
 class  UGroomAsset;
 struct FHairGroupsCardsSourceDescription;
+class  UHairCardGenerationSettings;
 
 class HAIRCARDGENERATORFRAMEWORK_API IHairCardGenerator : public IModularFeature
 {
 public:
 	static const FName ModularFeatureName; // "HairCardGenerator"
 
-	virtual bool GenerateHairCardsForLOD(const UGroomAsset* Groom, FHairGroupsCardsSourceDescription& CardsDesc) = 0;
+	virtual bool GenerateHairCardsForLOD(UGroomAsset* Groom, FHairGroupsCardsSourceDescription& CardsDesc) = 0;
+
+	virtual bool IsCompatibleSettings(UHairCardGenerationSettings* OldSettings) = 0;
 };
 
 namespace HairCardGenerator_Utils
diff --git a/Engine/Plugins/Runtime/HairStrands/Source/HairStrandsCore/Public/GroomAssetCards.h b/Engine/Plugins/Runtime/HairStrands/Source/HairStrandsCore/Public/GroomAssetCards.h
--- a/Engine/Plugins/Runtime/HairStrands/Source/HairStrandsCore/Public/GroomAssetCards.h
+++ b/Engine/Plugins/Runtime/HairStrands/Source/HairStrandsCore/Public/GroomAssetCards.h
@@ -184,6 +184,17 @@ struct HAIRSTRANDSCORE_API FHairGroupCardsTextures
 	bool bNeedToBeSaved = false;
 };
 
+/** 
+ * Since hair-card generation can be controlled external to this plugin, this  
+ * provides a way for those external generators a way to store their own 
+ * generation data along with the groom/cards-entry.
+ */
+UCLASS(Abstract, EditInlineNew)
+class HAIRSTRANDSCORE_API UHairCardGenerationSettings : public UObject
+{
+	GENERATED_BODY()
+};
+
 USTRUCT(BlueprintType)
 struct HAIRSTRANDSCORE_API FHairGroupsCardsSourceDescription
 {
@@ -224,6 +235,12 @@ struct HAIRSTRANDSCORE_API FHairGroupsCardsSourceDescription
 	UPROPERTY(EditAnywhere, Category = "CardsSource")
 	int32 LODIndex = -1; 
 
+#if WITH_EDITORONLY_DATA
+	/* Card generation data saved from the last procedural run. Dependent on the generator responsible for running the generation. */
+	UPROPERTY(meta=(EditInline))
+	TObjectPtr<UHairCardGenerationSettings> GenerationSettings = nullptr;
+#endif // WITH_EDITORONLY_DATA
+
 	UPROPERTY(VisibleAnywhere, Transient, Category = "CardsSource")
 	mutable FHairGroupCardsInfo CardsInfo;
 
diff --git a/Engine/Plugins/Runtime/HairStrands/Source/HairStrandsEditor/Private/GroomAssetDetails.cpp b/Engine/Plugins/Runtime/HairStrands/Source/HairStrandsEditor/Private/GroomAssetDetails.cpp
--- a/Engine/Plugins/Runtime/HairStrands/Source/HairStrandsEditor/Private/GroomAssetDetails.cpp
+++ b/Engine/Plugins/Runtime/HairStrands/Source/HairStrandsEditor/Private/GroomAssetDetails.cpp
@@ -1046,13 +1046,16 @@ FReply FGroomRenderingDetails::OnSaveCards(int32 DescIndex, FProperty* Property)
 
 FReply FGroomRenderingDetails::OnGenerateCardDataUsingPlugin(int32 GroupIndex)
 {
-	if (GroomAsset->HairGroupsCards.IsValidIndex(GroupIndex))
+	if (GroomAsset && GroomAsset->HairGroupsCards.IsValidIndex(GroupIndex))
 	{
 		TArray<IHairCardGenerator*> HairCardGeneratorPlugins = IModularFeatures::Get().GetModularFeatureImplementations<IHairCardGenerator>(IHairCardGenerator::ModularFeatureName);
 		if (HairCardGeneratorPlugins.Num() > 0)
 		{
 			UE_CLOG(HairCardGeneratorPlugins.Num() > 1, LogGroomAssetDetails, Warning, TEXT("There are more than one available hair-card generator options. Defaulting to the first one found."));
-			HairCardGeneratorPlugins[0]->GenerateHairCardsForLOD(GroomAsset, GroomAsset->HairGroupsCards[GroupIndex]);
+
+			const FScopedTransaction Transaction(LOCTEXT("GenerateHairCardsTransaction", "Generate hair cards."));
+			FHairGroupsCardsSourceDescription& HairCardsEntry = GroomAsset->HairGroupsCards[GroupIndex];
+			HairCardGeneratorPlugins[0]->GenerateHairCardsForLOD(GroomAsset, HairCardsEntry);
 		}
 	}
 	return FReply::Handled();
@@ -1393,7 +1396,12 @@ void FGroomRenderingDetails::OnGenerateElementForHairGroup(TSharedRef<IPropertyH
 				}
 				else if (PropertyName == GET_MEMBER_NAME_CHECKED(FHairGroupsCardsSourceDescription, ProceduralSettings))
 				{
-					PropertyRow.Visibility(ProceduralPropertyVisibility);
+					PropertyRow.Visibility(TAttribute<EVisibility>::CreateLambda([ShouldShowProceduralProperties]()
+						{
+							TArray<IHairCardGenerator*> HairCardGeneratorPlugins = IModularFeatures::Get().GetModularFeatureImplementations<IHairCardGenerator>(IHairCardGenerator::ModularFeatureName);
+							return HairCardGeneratorPlugins.IsEmpty() && ShouldShowProceduralProperties() ? EVisibility::Visible : EVisibility::Collapsed;
+						}
+					));
 				}
 				else if (PropertyName == GET_MEMBER_NAME_CHECKED(FHairGroupsCardsSourceDescription, ImportedMesh))
 				{
