Populating the groom with the generated mesh and textures.
  * Making the default asset names configurable
  * Making it so the operation can be properly undone (with the exception of overwriting assets)
  * Making the new generation options part of the DDC key generation

diff --git a/Engine/Plugins/Runtime/HairStrands/Source/HairStrandsCore/Private/GroomAsset.cpp b/Engine/Plugins/Runtime/HairStrands/Source/HairStrandsCore/Private/GroomAsset.cpp
--- a/Engine/Plugins/Runtime/HairStrands/Source/HairStrandsCore/Private/GroomAsset.cpp
+++ b/Engine/Plugins/Runtime/HairStrands/Source/HairStrandsCore/Private/GroomAsset.cpp
@@ -377,15 +377,6 @@ void UpdateHairStrandsLogVerbosity()
 
 /////////////////////////////////////////////////////////////////////////////////////////
 
-enum class EHairAtlasTextureType
-{
-	Depth,
-	Tangent,
-	Attribute,
-	Coverage,
-	AuxilaryData
-};
-
 template<typename ResourceType>
 static void InitAtlasTexture(ResourceType* InResource, UTexture2D* InTexture, EHairAtlasTextureType InType)
 {
@@ -1795,6 +1786,11 @@ namespace GroomDerivedDataCacheUtils
 				Ar << Key;
 			}
 			Desc.ProceduralSettings.BuildDDCKey(Ar);
+
+			if (Desc.GenerationSettings)
+			{
+				Desc.GenerationSettings->BuildDDCKey(Ar);
+			}
 		}
 
 		FSHAHash Hash;
diff --git a/Engine/Plugins/Runtime/HairStrands/Source/HairStrandsCore/Private/GroomAssetCards.cpp b/Engine/Plugins/Runtime/HairStrands/Source/HairStrandsCore/Private/GroomAssetCards.cpp
--- a/Engine/Plugins/Runtime/HairStrands/Source/HairStrandsCore/Private/GroomAssetCards.cpp
+++ b/Engine/Plugins/Runtime/HairStrands/Source/HairStrandsCore/Private/GroomAssetCards.cpp
@@ -3,6 +3,7 @@
 #include "GroomAssetCards.h"
 #include "Engine/StaticMesh.h"
 #include "MeshDescription.h"
+#include "GroomAsset.h" // for EHairAtlasTextureType
 
 FHairCardsClusterSettings::FHairCardsClusterSettings()
 {
@@ -86,6 +87,32 @@ bool FHairGroupsProceduralCards::operator==(const FHairGroupsProceduralCards& A)
 		Version == A.Version;
 }
 
+void FHairGroupCardsTextures::SetTexture(EHairAtlasTextureType SlotID, UTexture2D* Texture)
+{
+	switch (SlotID)
+	{
+	case EHairAtlasTextureType::Depth:
+		DepthTexture = Texture;
+		break;
+
+	case EHairAtlasTextureType::Coverage:
+		CoverageTexture = Texture;
+		break;
+
+	case EHairAtlasTextureType::Tangent:
+		TangentTexture = Texture;
+		break;
+
+	case EHairAtlasTextureType::Attribute:
+		AttributeTexture = Texture;
+		break;
+
+	case EHairAtlasTextureType::AuxilaryData:
+		AuxilaryDataTexture = Texture;
+		break;
+	};
+}
+
 void FHairGroupsProceduralCards::BuildDDCKey(FArchive& Ar)
 {
 	Ar << GeometrySettings.GenerationType;
@@ -176,4 +203,5 @@ UStaticMesh* FHairGroupsCardsSourceDescription::GetMesh() const
 	}
 #endif
 	return nullptr;
-}
\ No newline at end of file
+}
+
diff --git a/Engine/Plugins/Runtime/HairStrands/Source/HairStrandsCore/Public/GroomAsset.h b/Engine/Plugins/Runtime/HairStrands/Source/HairStrandsCore/Public/GroomAsset.h
--- a/Engine/Plugins/Runtime/HairStrands/Source/HairStrandsCore/Public/GroomAsset.h
+++ b/Engine/Plugins/Runtime/HairStrands/Source/HairStrandsCore/Public/GroomAsset.h
@@ -386,6 +386,16 @@ struct HAIRSTRANDSCORE_API FHairGroupInfoWithVisibility : public FHairGroupInfo
 	bool bIsVisible = true;
 };
 
+UENUM(BlueprintType)
+enum class EHairAtlasTextureType : uint8
+{
+	Depth,
+	Tangent,
+	Attribute,
+	Coverage,
+	AuxilaryData
+};
+
 /**
  * Implements an asset that can be used to store hair strands
  */
diff --git a/Engine/Plugins/Runtime/HairStrands/Source/HairStrandsCore/Public/GroomAssetCards.h b/Engine/Plugins/Runtime/HairStrands/Source/HairStrandsCore/Public/GroomAssetCards.h
--- a/Engine/Plugins/Runtime/HairStrands/Source/HairStrandsCore/Public/GroomAssetCards.h
+++ b/Engine/Plugins/Runtime/HairStrands/Source/HairStrandsCore/Public/GroomAssetCards.h
@@ -11,6 +11,8 @@
 class UMaterialInterface;
 class UStaticMesh;
 class UTexture2D;
+enum class EHairAtlasTextureType : uint8;
+class UGroomAsset;
 
 UENUM(BlueprintType)
 enum class EHairCardsClusterType : uint8
@@ -181,6 +183,8 @@ struct HAIRSTRANDSCORE_API FHairGroupCardsTextures
 	UPROPERTY(EditAnywhere, Category = "CardsAuxilaryData")
 	TObjectPtr<UTexture2D> AuxilaryDataTexture = nullptr;
 
+	void SetTexture(EHairAtlasTextureType SlotID, UTexture2D* Texture);
+
 	bool bNeedToBeSaved = false;
 };
 
@@ -193,6 +197,8 @@ UCLASS(Abstract, EditInlineNew)
 class HAIRSTRANDSCORE_API UHairCardGenerationSettings : public UObject
 {
 	GENERATED_BODY()
+public:
+	virtual void BuildDDCKey(FArchive& Ar) PURE_VIRTUAL(BuildDDCKey,);
 };
 
 USTRUCT(BlueprintType)
diff --git a/Engine/Plugins/Runtime/HairStrands/Source/HairStrandsEditor/Private/GroomAssetDetails.cpp b/Engine/Plugins/Runtime/HairStrands/Source/HairStrandsEditor/Private/GroomAssetDetails.cpp
--- a/Engine/Plugins/Runtime/HairStrands/Source/HairStrandsEditor/Private/GroomAssetDetails.cpp
+++ b/Engine/Plugins/Runtime/HairStrands/Source/HairStrandsEditor/Private/GroomAssetDetails.cpp
@@ -1054,8 +1054,18 @@ FReply FGroomRenderingDetails::OnGenerateCardDataUsingPlugin(int32 GroupIndex)
 			UE_CLOG(HairCardGeneratorPlugins.Num() > 1, LogGroomAssetDetails, Warning, TEXT("There are more than one available hair-card generator options. Defaulting to the first one found."));
 
 			const FScopedTransaction Transaction(LOCTEXT("GenerateHairCardsTransaction", "Generate hair cards."));
-			FHairGroupsCardsSourceDescription& HairCardsEntry = GroomAsset->HairGroupsCards[GroupIndex];
-			HairCardGeneratorPlugins[0]->GenerateHairCardsForLOD(GroomAsset, HairCardsEntry);
+
+			// Use a copy so we can only apply changes on success
+			FHairGroupsCardsSourceDescription HairCardsCopy = GroomAsset->HairGroupsCards[GroupIndex];
+			// Clear fields that are supposed to be set by the generation (in case it leaves any unset, and we don't cary over old settings)
+			HairCardsCopy.Textures = FHairGroupCardsTextures();
+
+			const bool bSuccess = HairCardGeneratorPlugins[0]->GenerateHairCardsForLOD(GroomAsset, HairCardsCopy);
+			if (bSuccess)
+			{
+				GroomAsset->Modify();
+				GroomAsset->HairGroupsCards[GroupIndex] = HairCardsCopy;
+			}
 		}
 	}
 	return FReply::Handled();
