Fix missing include.

diff --git a/Engine/Plugins/Runtime/HairStrands/Source/HairStrandsCore/Private/GroomDebug.cpp b/Engine/Plugins/Runtime/HairStrands/Source/HairStrandsCore/Private/GroomDebug.cpp
--- a/Engine/Plugins/Runtime/HairStrands/Source/HairStrandsCore/Private/GroomDebug.cpp
+++ b/Engine/Plugins/Runtime/HairStrands/Source/HairStrandsCore/Private/GroomDebug.cpp
@@ -11,6 +11,7 @@
 #include "Components/SkeletalMeshComponent.h"
 #include "UnrealEngine.h"
 #include "SystemTextures.h"
+#include "CanvasTypes.h"
 
 ///////////////////////////////////////////////////////////////////////////////////////////////////
 
