Fix missing update of hair RT geometry when hair radius/tip/root/RT scale are updated on a static groom.

diff --git a/Engine/Plugins/Runtime/HairStrands/Source/HairStrandsCore/Private/HairStrandsRendering.cpp b/Engine/Plugins/Runtime/HairStrands/Source/HairStrandsCore/Private/HairStrandsRendering.cpp
--- a/Engine/Plugins/Runtime/HairStrands/Source/HairStrandsCore/Private/HairStrandsRendering.cpp
+++ b/Engine/Plugins/Runtime/HairStrands/Source/HairStrandsCore/Private/HairStrandsRendering.cpp
@@ -1910,7 +1910,11 @@ void ComputeHairStrandsInterpolation(
 			if (Instance->Strands.RenRaytracingResource)
 			{
 				const float HairRadiusScaleRT = (GHairRaytracingRadiusScale > 0 ? GHairRaytracingRadiusScale : Instance->Strands.Modifier.HairRaytracingRadiusScale);
-				const bool bNeedUpdate = Instance->Strands.DeformedResource != nullptr;
+				const float HairRadiusRT	= Instance->HairGroupPublicData->VFInput.Strands.HairRadius * HairRadiusScaleRT;
+				const float HairRootScaleRT = Instance->HairGroupPublicData->VFInput.Strands.HairRootScale;
+				const float HairTipScaleRT	= Instance->HairGroupPublicData->VFInput.Strands.HairTipScale;
+
+				const bool bNeedUpdate = Instance->Strands.DeformedResource != nullptr ||  Instance->Strands.CachedHairScaledRadius != HairRadiusRT || Instance->Strands.CachedHairRootScale != HairRootScaleRT || Instance->Strands.CachedHairTipScale != HairTipScaleRT;
 				const bool bNeedBuild = !Instance->Strands.RenRaytracingResource->bIsRTGeometryInitialized;
 				if (bNeedBuild || bNeedUpdate)
 				{
@@ -1921,14 +1925,18 @@ void ComputeHairStrandsInterpolation(
 						GraphBuilder,
 						ShaderMap,
 						Instance->Strands.RestResource->GetVertexCount(),
-						Instance->HairGroupPublicData->VFInput.Strands.HairRadius * HairRadiusScaleRT,
-						Instance->HairGroupPublicData->VFInput.Strands.HairRootScale,
-						Instance->HairGroupPublicData->VFInput.Strands.HairTipScale,
+						HairRadiusRT,
+						HairRootScaleRT,
+						HairTipScaleRT,
 						Strands_PositionOffsetSRV,
 						CullingData,
 						Strands_PositionSRV,
 						Raytracing_PositionBuffer.UAV);
 
+					Instance->Strands.CachedHairScaledRadius= HairRadiusRT;
+					Instance->Strands.CachedHairRootScale	= HairRootScaleRT;
+					Instance->Strands.CachedHairTipScale	= HairTipScaleRT;
+
 					GraphBuilder.AddPass(
 						RDG_EVENT_NAME("HairStrands::UpdateBLAS(Strands)"),
 						ERDGPassFlags::NeverCull,
diff --git a/Engine/Plugins/Runtime/HairStrands/Source/HairStrandsCore/Public/GroomInstance.h b/Engine/Plugins/Runtime/HairStrands/Source/HairStrandsCore/Public/GroomInstance.h
--- a/Engine/Plugins/Runtime/HairStrands/Source/HairStrandsCore/Public/GroomInstance.h
+++ b/Engine/Plugins/Runtime/HairStrands/Source/HairStrandsCore/Public/GroomInstance.h
@@ -105,6 +105,9 @@ struct HAIRSTRANDSCORE_API FHairGroupInstance : public FHairStrandsInstance
 		#if RHI_RAYTRACING
 		FHairStrandsRaytracingResource* RenRaytracingResource = nullptr;
 		bool RenRaytracingResourceOwned = false;
+		float CachedHairScaledRadius = 0;
+		float CachedHairRootScale = 0;
+		float CachedHairTipScale = 0;
 		#endif
 
 		FRDGExternalBuffer DebugAttributeBuffer;
