Disable strip geometry by default and fix incorrect primitive count when drawing groom with culling disabled.

diff --git a/Engine/Plugins/Runtime/HairStrands/Source/HairStrandsCore/Private/GroomComponent.cpp b/Engine/Plugins/Runtime/HairStrands/Source/HairStrandsCore/Private/GroomComponent.cpp
--- a/Engine/Plugins/Runtime/HairStrands/Source/HairStrandsCore/Private/GroomComponent.cpp
+++ b/Engine/Plugins/Runtime/HairStrands/Source/HairStrandsCore/Private/GroomComponent.cpp
@@ -49,7 +49,7 @@ static FAutoConsoleVariableRef CVarUseProxyLocalToWorld(TEXT("r.HairStrands.UseP
 
 #define LOCTEXT_NAMESPACE "GroomComponent"
 
-#define USE_HAIR_TRIANGLE_STRIP 1
+#define USE_HAIR_TRIANGLE_STRIP 0
 
 ////////////////////////////////////////////////////////////////////////////////////////////////////////////////////////
 
@@ -773,7 +773,7 @@ public:
 			MaxVertexIndex = HairVertexCount * 6;
 			#endif
 			bUseCulling = Instance->Strands.bIsCullingEnabled;
-			NumPrimitive = bUseCulling ? 0 : HairVertexCount * 2;
+			NumPrimitive = bUseCulling ? 0 : HairVertexCount * 2 * 3;
 			if (MaterialRenderProxy == nullptr)
 			{
 				MaterialRenderProxy = HairGroupMaterialProxies[GroupIndex].Strands;
diff --git a/Engine/Shaders/Private/HairStrands/HairStrandsVertexFactoryCommon.ush b/Engine/Shaders/Private/HairStrands/HairStrandsVertexFactoryCommon.ush
--- a/Engine/Shaders/Private/HairStrands/HairStrandsVertexFactoryCommon.ush
+++ b/Engine/Shaders/Private/HairStrands/HairStrandsVertexFactoryCommon.ush
@@ -6,7 +6,7 @@
 #define HAIR_CONTROLPOINT_START	1
 #define HAIR_CONTROLPOINT_END	2
 
-#define USE_HAIR_TRIANGLE_STRIP 1
+#define USE_HAIR_TRIANGLE_STRIP 0
 
 ////////////////////////////////////////////////////////////////////////////////////////////////////////////////////////////////
 // Hair control points
