Fix typo causing hair groom to still use triangle strip geometry.

diff --git a/Engine/Plugins/Runtime/HairStrands/Source/HairStrandsCore/Private/GroomComponent.cpp b/Engine/Plugins/Runtime/HairStrands/Source/HairStrandsCore/Private/GroomComponent.cpp
--- a/Engine/Plugins/Runtime/HairStrands/Source/HairStrandsCore/Private/GroomComponent.cpp
+++ b/Engine/Plugins/Runtime/HairStrands/Source/HairStrandsCore/Private/GroomComponent.cpp
@@ -773,7 +773,7 @@ public:
 			MaxVertexIndex = HairVertexCount * 6;
 			#endif
 			bUseCulling = Instance->Strands.bIsCullingEnabled;
-			NumPrimitive = bUseCulling ? 0 : HairVertexCount * 2 * 3;
+			NumPrimitive = bUseCulling ? 0 : HairVertexCount * 2;
 			if (MaterialRenderProxy == nullptr)
 			{
 				MaterialRenderProxy = HairGroupMaterialProxies[GroupIndex].Strands;
@@ -860,7 +860,7 @@ public:
 		#if USE_HAIR_TRIANGLE_STRIP
 		Mesh.Type = GeometryType == EHairGeometryType::Strands ? PT_TriangleStrip : PT_TriangleList;
 		#else
-		Mesh.Type = GeometryType == PT_TriangleList;
+		Mesh.Type = PT_TriangleList;
 		#endif
 		Mesh.DepthPriorityGroup = SDPG_World;
 		Mesh.bCanApplyViewModeOverrides = false;
