Fix incorrect FVector/FVector3f changes in hair code.

diff --git a/Engine/Plugins/Runtime/HairStrands/Source/HairStrandsCore/Private/HairCardsBuilder.cpp b/Engine/Plugins/Runtime/HairStrands/Source/HairStrandsCore/Private/HairCardsBuilder.cpp
--- a/Engine/Plugins/Runtime/HairStrands/Source/HairStrandsCore/Private/HairCardsBuilder.cpp
+++ b/Engine/Plugins/Runtime/HairStrands/Source/HairStrandsCore/Private/HairCardsBuilder.cpp
@@ -335,14 +335,14 @@ namespace HairCards
 	{
 		float VoxelSize = 0;
 		FIntVector Resolution = FIntVector::ZeroValue;
-		FVector MinBound = FVector::ZeroVector;
-		FVector MaxBound = FVector::ZeroVector;
+		FVector3f MinBound = FVector3f::ZeroVector;
+		FVector3f MaxBound = FVector3f::ZeroVector;
 		TArray<uint32> Density;
-		TArray<FVector> Normal;
-		TArray<FVector> Tangent;
+		TArray<FVector3f> Normal;
+		TArray<FVector3f> Tangent;
 	};
 
-	inline uint32 To8bits(const FVector& T)
+	inline uint32 To8bits(const FVector3f& T)
 	{
 		uint32 X = FMath::Clamp((T.X + 1) * 0.5f, 0.f, 1.f) * 255.f;
 		uint32 Y = FMath::Clamp((T.Y + 1) * 0.5f, 0.f, 1.f) * 255.f;
@@ -350,17 +350,17 @@ namespace HairCards
 		return (0xFF & X) | ((0xFF & Y)<<8) | ((0xFF & Z)<<16);
 	}
 
-	inline FVector From8bits(const uint32& T)
+	inline FVector3f From8bits(const uint32& T)
 	{
-		return FVector(
+		return FVector3f(
 			((0xFF & (T))     / 255.f) * 2.f - 1.f,
 			((0xFF & (T>>8))  / 255.f) * 2.f - 1.f,
 			((0xFF & (T>>16)) / 255.f) * 2.f - 1.f);
 	}
 
-	inline FIntVector ToCoord(const FVector& T, const FVoxelVolume& V)
+	inline FIntVector ToCoord(const FVector3f& T, const FVoxelVolume& V)
 	{
-		const FVector C = (T  - V.MinBound) / V.VoxelSize;
+		const FVector3f C = (T  - V.MinBound) / V.VoxelSize;
 		return FIntVector(
 			FMath::Clamp(FMath::FloorToInt(C.X), 0, V.Resolution.X-1), 
 			FMath::Clamp(FMath::FloorToInt(C.Y), 0, V.Resolution.Y-1), 
@@ -429,7 +429,7 @@ namespace HairCards
 		NonVisibleCoord.Reserve(Out.Density.Num() * 0.5f); // Guesstimate
 		VisibleCoord.Reserve(Out.Density.Num() * 0.5f); // Guesstimate
 
-		const FVector CenterBoundCoord = FVector(Out.Resolution) * 0.5f;
+		const FVector3f CenterBoundCoord = FVector3f(Out.Resolution) * 0.5f;
 
 		const TArray<FIntVector> Kernel1 = CreateKernelCoord(1, Out);
 		const TArray<FIntVector> Kernel2 = CreateKernelCoord(2, Out);
@@ -442,7 +442,7 @@ namespace HairCards
 			{
 				const FIntVector Coord(X,Y,Z);
 				const uint32 LinearCoord = ToLinearCoord(Coord, Out);
-				Out.Normal[LinearCoord] = FVector::ZeroVector;
+				Out.Normal[LinearCoord] = FVector3f::ZeroVector;
 
 				const uint32 CenterDensity = Out.Density[LinearCoord];
 				if (CenterDensity == 0)
@@ -459,13 +459,13 @@ namespace HairCards
 				}
 
 				uint32 ValidCount = 0;
-				FVector OutDirection = FVector::ZeroVector;
+				FVector3f OutDirection = FVector3f::ZeroVector;
 				for (const FIntVector& Offset : Kernel1)
 				{
 					const FIntVector C = Coord + Offset;
 					if (Offset != FIntVector::ZeroValue && IsEmpty(C, Out))
 					{
-						OutDirection += FVector(Offset).GetSafeNormal();
+						OutDirection += FVector3f(Offset).GetSafeNormal();
 						ValidCount++;
 					}
 				}
@@ -474,16 +474,16 @@ namespace HairCards
 				Out.Normal[LinearCoord] = OutDirection.GetSafeNormal();
 
 				// If nothing has been found in a 1x1 kernel, try to look for further
-				if (Out.Normal[LinearCoord] == FVector::ZeroVector)
+				if (Out.Normal[LinearCoord] == FVector3f::ZeroVector)
 				{
 					ValidCount = 0;
-					OutDirection = FVector::ZeroVector;
+					OutDirection = FVector3f::ZeroVector;
 					for (const FIntVector& Offset : Kernel2)
 					{
 						const FIntVector C = Coord + Offset;
 						if (Offset != FIntVector::ZeroValue && IsEmpty(C, Out))
 						{
-							OutDirection += FVector(Offset).GetSafeNormal();
+							OutDirection += FVector3f(Offset).GetSafeNormal();
 							ValidCount++;
 						}
 					}
@@ -491,9 +491,9 @@ namespace HairCards
 					Out.Normal[LinearCoord] = OutDirection.GetSafeNormal();
 
 					// Nothing has been found, so use the center of the bound
-					if (Out.Normal[LinearCoord] != FVector::ZeroVector)
+					if (Out.Normal[LinearCoord] != FVector3f::ZeroVector)
 					{
-						Out.Normal[LinearCoord] = (FVector(Coord) - CenterBoundCoord).GetSafeNormal();
+						Out.Normal[LinearCoord] = (FVector3f(Coord) - CenterBoundCoord).GetSafeNormal();
 					}
 				}
 			}
@@ -501,8 +501,8 @@ namespace HairCards
 
 		// Filter normal on the iso surface
 		{	
-			TArray<FVector> OutNormal;
-			OutNormal.Init(FVector::ZeroVector, Out.Normal.Num());
+			TArray<FVector3f> OutNormal;
+			OutNormal.Init(FVector3f::ZeroVector, Out.Normal.Num());
 
 			// Filter only voxel which have valid normals
 			for (const FIntVector& Coord : VisibleCoord)
@@ -510,7 +510,7 @@ namespace HairCards
 				const uint32 LinearCoord = ToLinearCoord(Coord, Out);
 
 				uint32 ValidCount = 0;
-				FVector OutDirection = FVector::ZeroVector;
+				FVector3f OutDirection = FVector3f::ZeroVector;
 				for (const FIntVector& Offset : Kernel2)
 				{
 					const FIntVector C = Coord + Offset;
@@ -519,7 +519,7 @@ namespace HairCards
 					{
 						const uint32 LinearC = ToLinearCoord(C, Out);
 
-						const bool bIsSampleValid = Out.Normal[LinearC] != FVector::ZeroVector;
+						const bool bIsSampleValid = Out.Normal[LinearC] != FVector3f::ZeroVector;
 						if (bIsSampleValid)
 						{
 							OutNormal[LinearCoord]  += Out.Normal[LinearC].GetSafeNormal() * w;
@@ -535,7 +535,7 @@ namespace HairCards
 				else
 				{
 					// Generate normal oriented as a surounding sphere
-					OutNormal[LinearCoord] = (FVector(Coord) - FVector(CenterBoundCoord)).GetSafeNormal();
+					OutNormal[LinearCoord] = (FVector3f(Coord) - FVector3f(CenterBoundCoord)).GetSafeNormal();
 				}
 			}
 
@@ -552,19 +552,19 @@ namespace HairCards
 			for (const FIntVector& Coord : NonVisibleCoord)
 			{
 				const uint32 LinearCoord = ToLinearCoord(Coord, Out);
-				check(Out.Normal[LinearCoord] == FVector::ZeroVector);
-				Out.Normal[LinearCoord] = FVector::ZeroVector;
+				check(Out.Normal[LinearCoord] == FVector3f::ZeroVector);
+				Out.Normal[LinearCoord] = FVector3f::ZeroVector;
 
 				uint32 ValidCount = 0;
-				FVector OutDirection = FVector::ZeroVector;
+				FVector3f OutDirection = FVector3f::ZeroVector;
 				for (const FIntVector& Offset : Kernel1)
 				{
 					const FIntVector C = Coord + Offset;
 					if (IsInside(C, Out) && Offset != FIntVector::ZeroValue)
 					{
 						const uint32 LinearC = ToLinearCoord(C, Out);
-						const FVector SNormal = Out.Normal[LinearC];
-						if (SNormal != FVector::ZeroVector)
+						const FVector3f SNormal = Out.Normal[LinearC];
+						if (SNormal != FVector3f::ZeroVector)
 						{
 							OutDirection += SNormal;
 							ValidCount++;
@@ -592,7 +592,7 @@ namespace HairCards
 	{	
 		const TArray<FIntVector> Kernel2 = CreateKernelCoord(2, Out);
 
-		TArray<FVector> OutTangent;
+		TArray<FVector3f> OutTangent;
 		OutTangent.SetNum(Out.Tangent.Num());
 
 		for (int32 Z = 0; Z < Out.Resolution.Z; ++Z)
@@ -601,14 +601,14 @@ namespace HairCards
 		{
 			const FIntVector Coord(X,Y,Z);
 			const uint32 LinearCoord = ToLinearCoord(Coord, Out);
-			OutTangent[LinearCoord] = FVector::ZeroVector;
+			OutTangent[LinearCoord] = FVector3f::ZeroVector;
 
 			const uint32 CenterDensity = Out.Density[LinearCoord];
 			if (CenterDensity == 0)
 				continue;
 
 			uint32 ValidCount = 0;
-			FVector OutDirection = FVector::ZeroVector;
+			FVector3f OutDirection = FVector3f::ZeroVector;
 			for (const FIntVector& Offset : Kernel2)
 			{
 				const FIntVector C = Coord + Offset;
@@ -634,16 +634,16 @@ namespace HairCards
 
 	void VoxelizeGroom(const FHairStrandsDatas& InData, const uint32 InResolution, FVoxelVolume& Out)
 	{
-		const FVector BoundSize = InData.BoundingBox.Max - InData.BoundingBox.Min;
+		const FVector3f BoundSize = InData.BoundingBox.Max - InData.BoundingBox.Min;
 		Out.VoxelSize = FMath::Max3(BoundSize.X, BoundSize.Y, BoundSize.Z) / InResolution;
 
 		Out.Resolution = FIntVector(FMath::CeilToFloat(BoundSize.X / Out.VoxelSize), FMath::CeilToFloat(BoundSize.Y / Out.VoxelSize), FMath::CeilToFloat(BoundSize.Z / Out.VoxelSize));
 		Out.MinBound   = InData.BoundingBox.Min;
-		Out.MaxBound   = FVector(Out.Resolution) * Out.VoxelSize + InData.BoundingBox.Min;
+		Out.MaxBound   = FVector3f(Out.Resolution) * Out.VoxelSize + InData.BoundingBox.Min;
 
 		const uint32 TotalVoxelCount = Out.Resolution.X * Out.Resolution.Y * Out.Resolution.Z;
 		Out.Density.Init(0, TotalVoxelCount);
-		Out.Tangent.Init(FVector::ZeroVector, TotalVoxelCount);
+		Out.Tangent.Init(FVector3f::ZeroVector, TotalVoxelCount);
 
 		// Fill in voxel (TODO: make it parallel)
 		const uint32 CurveCount = InData.StrandsCurves.Num();
@@ -654,10 +654,10 @@ namespace HairCards
 			
 			for (uint32 PointIndex = 0; PointIndex < PointCount-1; ++PointIndex)
 			{
-				const FVector& P0 = InData.StrandsPoints.PointsPosition[PointOffset + PointIndex];
-				const FVector& P1 = InData.StrandsPoints.PointsPosition[PointOffset + PointIndex + 1];
-				const FVector Segment = P1 - P0;
-				const FVector T = Segment.GetSafeNormal();
+				const FVector3f& P0 = InData.StrandsPoints.PointsPosition[PointOffset + PointIndex];
+				const FVector3f& P1 = InData.StrandsPoints.PointsPosition[PointOffset + PointIndex + 1];
+				const FVector3f Segment = P1 - P0;
+				const FVector3f T = Segment.GetSafeNormal();
 				//const uint32 T8bits = To8bits(T);
 
 				const float Length = Segment.Size();
@@ -665,7 +665,7 @@ namespace HairCards
 				uint32 PrevLinearCoord = ~0;
 				for (uint32 StepIt = 0; StepIt < StepCount+1; ++StepIt)
 				{
-					const FVector P = P0 + Segment * StepIt / float(StepCount);
+					const FVector3f P = P0 + Segment * StepIt / float(StepCount);
 					const FIntVector Coord = ToCoord(P, Out);
 					const uint32 LinearCoord = ToLinearCoord(Coord, Out);
 					if (LinearCoord != PrevLinearCoord)
@@ -685,7 +685,7 @@ namespace HairCards
 		FilterTangent(Out);
 	}
 
-	FVector GetTangent(const FVector InP, const FVoxelVolume& InV)
+	FVector3f GetTangent(const FVector3f InP, const FVoxelVolume& InV)
 	{
 		const FIntVector Coord = ToCoord(InP, InV);
 		const uint32 LinearCoord = ToLinearCoord(Coord, InV);
@@ -693,7 +693,7 @@ namespace HairCards
 		return InV.Tangent[LinearCoord].GetSafeNormal();
 	}
 
-	FVector GetNormal(const FVector InP, const FVoxelVolume& InV)
+	FVector3f GetNormal(const FVector3f InP, const FVoxelVolume& InV)
 	{
 		const FIntVector Coord = ToCoord(InP, InV);
 		const uint32 LinearCoord = ToLinearCoord(Coord, InV);
@@ -701,7 +701,7 @@ namespace HairCards
 		return InV.Normal[LinearCoord].GetSafeNormal();
 	}
 
-	void GetFrame(const FVector InP, const FVoxelVolume& InV, FVector& OutT, FVector& OutB, FVector& OutN)
+	void GetFrame(const FVector3f InP, const FVoxelVolume& InV, FVector3f& OutT, FVector3f& OutB, FVector3f& OutN)
 	{
 		const FIntVector Coord = ToCoord(InP, InV);
 		const uint32 LinearCoord = ToLinearCoord(Coord, InV);
@@ -710,13 +710,13 @@ namespace HairCards
 		// Build ortho-normal frame
 		OutT = InV.Tangent[LinearCoord].GetSafeNormal();
 		OutN = InV.Normal[LinearCoord].GetSafeNormal();
-		OutB = FVector::CrossProduct(OutN, OutT).GetSafeNormal();
-		OutN = FVector::CrossProduct(OutT, OutN).GetSafeNormal();
+		OutB = FVector3f::CrossProduct(OutN, OutT).GetSafeNormal();
+		OutN = FVector3f::CrossProduct(OutT, OutN).GetSafeNormal();
 #else
 		// Return a non-ortho-normal frame
 		OutT = InV.Tangent[LinearCoord].GetSafeNormal();
 		OutN = InV.Normal[LinearCoord].GetSafeNormal();
-		OutB = FVector::CrossProduct(OutN, OutT).GetSafeNormal();
+		OutB = FVector3f::CrossProduct(OutN, OutT).GetSafeNormal();
 #endif
 	}
 }
@@ -729,15 +729,15 @@ namespace HairCards
 {
 	struct FHairRoot
 	{
-		FVector Position;
+		FVector3f Position;
 		uint32  VertexCount;
-		FVector Normal;
+		FVector3f Normal;
 		uint32  Index;
 		float   Length;
 	};
 
 	template<uint32 NumSamples>
-	inline FVector GetCurvePosition(const FHairStrandsDatas& CurvesDatas, const uint32 CurveIndex, const uint32 SampleIndex)
+	inline FVector3f GetCurvePosition(const FHairStrandsDatas& CurvesDatas, const uint32 CurveIndex, const uint32 SampleIndex)
 	{
 		const float PointCount = CurvesDatas.StrandsCurves.CurvesCount[CurveIndex]-1.0;
 		const uint32 PointOffset = CurvesDatas.StrandsCurves.CurvesOffset[CurveIndex];
@@ -767,15 +767,15 @@ namespace HairCards
 
 		static const float DeltaCoord = 1.0f / static_cast<float>(NumSamples-1);
 
-		const FVector& RenderRoot = RenderCurvesDatas.StrandsPoints.PointsPosition[RenderCurvesDatas.StrandsCurves.CurvesOffset[RenderCurveIndex]];
-		const FVector& GuideRoot = GuideCurvesDatas.StrandsPoints.PointsPosition[GuideCurvesDatas.StrandsCurves.CurvesOffset[GuideCurveIndex]];
+		const FVector3f& RenderRoot = RenderCurvesDatas.StrandsPoints.PointsPosition[RenderCurvesDatas.StrandsCurves.CurvesOffset[RenderCurveIndex]];
+		const FVector3f& GuideRoot = GuideCurvesDatas.StrandsPoints.PointsPosition[GuideCurvesDatas.StrandsCurves.CurvesOffset[GuideCurveIndex]];
 
 		float CurveProximityMetric = 0.0;
 		float CurveShapeMetric = 0.0;
 		for (uint32 SampleIndex = 0; SampleIndex < NumSamples; ++SampleIndex)
 		{
-			const FVector GuidePosition = GetCurvePosition<NumSamples>(GuideCurvesDatas, GuideCurveIndex, SampleIndex);
-			const FVector RenderPosition = GetCurvePosition<NumSamples>(RenderCurvesDatas, RenderCurveIndex, SampleIndex);
+			const FVector3f GuidePosition = GetCurvePosition<NumSamples>(GuideCurvesDatas, GuideCurveIndex, SampleIndex);
+			const FVector3f RenderPosition = GetCurvePosition<NumSamples>(RenderCurvesDatas, RenderCurveIndex, SampleIndex);
 			const float RootWeight = FMath::Exp(-RootImportance*SampleIndex*DeltaCoord);
 
 			CurveProximityMetric += (GuidePosition - RenderPosition).Size() * RootWeight;
@@ -832,7 +832,7 @@ namespace HairCards
 		const int32 RenCurveIndex,
 		const int32 SimCurveIndex)
 	{
-		const float Metric = FVector::Dist(GuideRoot.Position, RenRoot.Position);
+		const float Metric = FVector3f::Dist(GuideRoot.Position, RenRoot.Position);
 		if (Metric < Metrics1.KMinMetrics[FMetrics::Count - 1])
 		{
 			int32 LastGuideIndex = SimCurveIndex;
@@ -874,8 +874,8 @@ namespace HairCards
 
 	struct FRootsGrid
 	{
-		FVector MinBound = FVector::ZeroVector;
-		FVector MaxBound = FVector::ZeroVector;
+		FVector3f MinBound = FVector3f::ZeroVector;
+		FVector3f MaxBound = FVector3f::ZeroVector;
 		
 		const uint32 MaxLookupDistance = 31;
 		const FIntVector GridResolution = FIntVector(32, 32, 32);
@@ -899,10 +899,10 @@ namespace HairCards
 				FMath::Clamp(CellCoord.Z, 0, GridResolution.Z - 1));
 		}
 
-		FORCEINLINE FIntVector ToCellCoord(const FVector& P) const
+		FORCEINLINE FIntVector ToCellCoord(const FVector3f& P) const
 		{
 			bool bIsValid = false;
-			const FVector F = ((P - MinBound) / (MaxBound - MinBound));
+			const FVector3f F = ((P - MinBound) / (MaxBound - MinBound));
 			const FIntVector CellCoord = FIntVector(FMath::FloorToInt(F.X * GridResolution.X), FMath::FloorToInt(F.Y * GridResolution.Y), FMath::FloorToInt(F.Z * GridResolution.Z));			
 			return ClampToVolume(CellCoord, bIsValid);
 		}
@@ -914,7 +914,7 @@ namespace HairCards
 			return CellIndex;
 		}
 
-		void InsertRoots(TArray<FHairRoot>& Roots, const FVector& InMinBound, const FVector& InMaxBound)
+		void InsertRoots(TArray<FHairRoot>& Roots, const FVector3f& InMinBound, const FVector3f& InMaxBound)
 		{
 			MinBound = InMinBound;
 			MaxBound = InMaxBound;
@@ -1185,10 +1185,10 @@ namespace HairCards
 	}
 
 	// Extract strand roots
-	static void ExtractRoots(const FHairStrandsDatas& InData, TArray<FHairRoot>& OutRoots, FVector& MinBound, FVector& MaxBound)
+	static void ExtractRoots(const FHairStrandsDatas& InData, TArray<FHairRoot>& OutRoots, FVector3f& MinBound, FVector3f& MaxBound)
 	{
-		MinBound = FVector(FLT_MAX, FLT_MAX, FLT_MAX);
-		MaxBound = FVector(-FLT_MAX, -FLT_MAX, -FLT_MAX);
+		MinBound = FVector3f(FLT_MAX, FLT_MAX, FLT_MAX);
+		MaxBound = FVector3f(-FLT_MAX, -FLT_MAX, -FLT_MAX);
 		const uint32 CurveCount = InData.StrandsCurves.Num();
 		OutRoots.Reserve(CurveCount);
 		for (uint32 CurveIndex = 0; CurveIndex < CurveCount; ++CurveIndex)
@@ -1197,14 +1197,14 @@ namespace HairCards
 			const uint32 PointCount = InData.StrandsCurves.CurvesCount[CurveIndex];
 			const float  CurveLength = InData.StrandsCurves.CurvesLength[CurveIndex] * InData.StrandsCurves.MaxLength;
 			check(PointCount > 1);
-			const FVector& P0 = InData.StrandsPoints.PointsPosition[PointOffset];
-			const FVector& P1 = InData.StrandsPoints.PointsPosition[PointOffset + 1];
-			FVector N = (P1 - P0).GetSafeNormal();
+			const FVector3f& P0 = InData.StrandsPoints.PointsPosition[PointOffset];
+			const FVector3f& P1 = InData.StrandsPoints.PointsPosition[PointOffset + 1];
+			FVector3f N = (P1 - P0).GetSafeNormal();
 
 			// Fallback in case the initial points are too close (this happens on certain assets)
-			if (FVector::DotProduct(N, N) == 0)
+			if (FVector3f::DotProduct(N, N) == 0)
 			{
-				N = FVector(0, 0, 1);
+				N = FVector3f(0, 0, 1);
 			}
 			OutRoots.Add({ P0, PointCount, N, PointOffset, CurveLength });
 
@@ -1285,7 +1285,7 @@ namespace HairCards
 				*OffsetIterator = StrandOffset;
 
 				float StrandLength = 0.0;
-				FVector PreviousPosition(0.0, 0.0, 0.0);
+				FVector3f PreviousPosition(0.0, 0.0, 0.0);
 				for (uint32 PointIndex = 0; PointIndex < StrandCount; ++PointIndex, ++PositionIterator, ++RadiusIterator, ++CoordUIterator)
 				{
 					HairStrands.BoundingBox += *PositionIterator;
@@ -1386,16 +1386,16 @@ namespace HairCards
 
 		TArray<uint32>  GuidePointCount;
 		TArray<uint32>  GuidePointOffset;
-		TArray<FVector> GuidePoints;
+		TArray<FVector3f> GuidePoints;
 		TArray<FVector2D> GuideRootUVs;
 
 		TArray<FHairOrientedBound> Bounds;
 
 		// Cluster bound/envelop
-		TArray<FVector> BoundPositions;
-		TArray<FVector> BoundAxis_T;
-		TArray<FVector> BoundAxis_B;
-		TArray<FVector> BoundAxis_N;
+		TArray<FVector3f> BoundPositions;
+		TArray<FVector3f> BoundAxis_T;
+		TArray<FVector3f> BoundAxis_B;
+		TArray<FVector3f> BoundAxis_N;
 
 		TArray<uint32> BoundOffset;
 		TArray<uint32> BoundCount;
@@ -1460,8 +1460,8 @@ namespace HairCards
 		FRoots RenRoots, SimRoots;
 		FRootsGrid RootsGrid;
 		{
-			FVector RenMinBound, RenMaxBound;
-			FVector SimMinBound, SimMaxBound;
+			FVector3f RenMinBound, RenMaxBound;
+			FVector3f SimMinBound, SimMaxBound;
 			ExtractRoots(RenStrandsData, RenRoots, RenMinBound, RenMaxBound);
 			ExtractRoots(SimStrandsData, SimRoots, SimMinBound, SimMaxBound);
 
@@ -1469,8 +1469,8 @@ namespace HairCards
 			{
 				// Build a conservative bound, to insure all queries will fall 
 				// into the grid volume.
-				const FVector MinBound = RenMinBound.ComponentMin(SimMinBound);
-				const FVector MaxBound = RenMaxBound.ComponentMax(SimMaxBound);
+				const FVector3f MinBound = RenMinBound.ComponentMin(SimMinBound);
+				const FVector3f MaxBound = RenMaxBound.ComponentMax(SimMaxBound);
 				RootsGrid.InsertRoots(SimRoots, MinBound, MaxBound);
 			}
 		}
@@ -1647,12 +1647,12 @@ namespace HairCards
 		}
 	}
 
-	static Eigen::Matrix3f ComputeCovarianceMatrix(TArray<FVector> InPoints, FVector& Mean)
+	static Eigen::Matrix3f ComputeCovarianceMatrix(TArray<FVector3f> InPoints, FVector3f& Mean)
 	{
 		// Mean
 		const uint32 PointCount = InPoints.Num();
-		Mean = FVector::ZeroVector;
-		for (FVector& P : InPoints)
+		Mean = FVector3f::ZeroVector;
+		for (FVector3f& P : InPoints)
 		{
 			Mean += P;
 		}
@@ -1661,7 +1661,7 @@ namespace HairCards
 		// Covariance
 		Eigen::Matrix3f Out;
 		Out << 0, 0, 0, 0, 0, 0, 0, 0, 0;
-		for (FVector& P : InPoints)
+		for (FVector3f& P : InPoints)
 		{
 			Out(0, 0) += (P.X - Mean.X) * (P.X - Mean.X);
 			Out(1, 0) += (P.X - Mean.X) * (P.Y - Mean.Y);
@@ -1689,31 +1689,31 @@ namespace HairCards
 		return Out;
 	}
 
-	//static void ComputePCA(const TArray<FVector>& InPoints, FVector& Mean, FVector& B0, FVector& B1, FVector& B2)
-	static void ComputePCA(const TArray<FVector>& InPoints, const uint32 InPointOffset, const uint32 InPointCount, FVector& Mean, FVector& B0, FVector& B1, FVector& B2)
+	//static void ComputePCA(const TArray<FVector3f>& InPoints, FVector3f& Mean, FVector3f& B0, FVector3f& B1, FVector3f& B2)
+	static void ComputePCA(const TArray<FVector3f>& InPoints, const uint32 InPointOffset, const uint32 InPointCount, FVector3f& Mean, FVector3f& B0, FVector3f& B1, FVector3f& B2)
 	{
 		Eigen::Matrix3f covariance = ComputeCovarianceMatrix(InPoints, Mean);
 		Eigen::SelfAdjointEigenSolver<Eigen::Matrix3f> eigen_solver(covariance, Eigen::ComputeEigenvectors);
 		Eigen::Matrix3f eigenVectorsPCA = eigen_solver.eigenvectors();
 
-		B0 = FVector(eigenVectorsPCA.col(0)(0), eigenVectorsPCA.col(0)(1), eigenVectorsPCA.col(0)(2));
-		B1 = FVector(eigenVectorsPCA.col(1)(0), eigenVectorsPCA.col(1)(1), eigenVectorsPCA.col(1)(2));
-		B2 = FVector(eigenVectorsPCA.col(2)(0), eigenVectorsPCA.col(2)(1), eigenVectorsPCA.col(2)(2));
+		B0 = FVector3f(eigenVectorsPCA.col(0)(0), eigenVectorsPCA.col(0)(1), eigenVectorsPCA.col(0)(2));
+		B1 = FVector3f(eigenVectorsPCA.col(1)(0), eigenVectorsPCA.col(1)(1), eigenVectorsPCA.col(1)(2));
+		B2 = FVector3f(eigenVectorsPCA.col(2)(0), eigenVectorsPCA.col(2)(1), eigenVectorsPCA.col(2)(2));
 
 		B0 = B0.GetSafeNormal();
 		B1 = B1.GetSafeNormal();
 		B2 = B2.GetSafeNormal();
 
-		FVector LocalMinBound( FLT_MAX);
-		FVector LocalMaxBound(-FLT_MAX);
-		//for (const FVector& P : InPoints)
+		FVector3f LocalMinBound( FLT_MAX);
+		FVector3f LocalMaxBound(-FLT_MAX);
+		//for (const FVector3f& P : InPoints)
 		for (uint32 PointIt=0;PointIt<InPointCount;++PointIt)
 		{
-			FVector PLocal = InPoints[InPointOffset+PointIt] - Mean;
-			PLocal = FVector(
-				FVector::DotProduct(PLocal, B0),
-				FVector::DotProduct(PLocal, B1),
-				FVector::DotProduct(PLocal, B2));
+			FVector3f PLocal = InPoints[InPointOffset+PointIt] - Mean;
+			PLocal = FVector3f(
+				FVector3f::DotProduct(PLocal, B0),
+				FVector3f::DotProduct(PLocal, B1),
+				FVector3f::DotProduct(PLocal, B2));
 
 			LocalMinBound.X = FMath::Min(LocalMinBound.X, PLocal.X);
 			LocalMinBound.Y = FMath::Min(LocalMinBound.Y, PLocal.Y);
@@ -1725,17 +1725,17 @@ namespace HairCards
 		}
 
 		// Recompte the new mean in local space, and then compute its world position
-		const FVector LocalMean(
+		const FVector3f LocalMean(
 			(LocalMaxBound.X + LocalMinBound.X) * 0.5f,
 			(LocalMaxBound.Y + LocalMinBound.Y) * 0.5f,
 			(LocalMaxBound.Z + LocalMinBound.Z) * 0.5f);
 
-		const FVector LocalExtent(
+		const FVector3f LocalExtent(
 			(LocalMaxBound.X - LocalMinBound.X) * 0.5f,
 			(LocalMaxBound.Y - LocalMinBound.Y) * 0.5f,
 			(LocalMaxBound.Z - LocalMinBound.Z) * 0.5f);
 
-		const FVector MeanPrime = Mean + B0 * LocalMean.X + B1 * LocalMean.Y + B2 * LocalMean.Z;
+		const FVector3f MeanPrime = Mean + B0 * LocalMean.X + B1 * LocalMean.Y + B2 * LocalMean.Z;
 		Mean = MeanPrime;
 
 		B0 *= LocalExtent.X;
@@ -1743,14 +1743,14 @@ namespace HairCards
 		B2 *= LocalExtent.Z;
 	}
 
-	static FHairOrientedBound ComputeOrientedBoundingBox(const TArray<FVector>& InPoints)
+	static FHairOrientedBound ComputeOrientedBoundingBox(const TArray<FVector3f>& InPoints)
 	{
 		FHairOrientedBound Out;
 		ComputePCA(InPoints, 0, InPoints.Num(), Out.Center, Out.ExtentX, Out.ExtentY, Out.ExtentZ);
 		return Out;
 	}
 
-	static FHairOrientedBound ComputeOrientedBoundingBox(const TArray<FVector>& InPoints, const uint32 InPointOffset, const uint32 InPointCount)
+	static FHairOrientedBound ComputeOrientedBoundingBox(const TArray<FVector3f>& InPoints, const uint32 InPointOffset, const uint32 InPointCount)
 	{
 		FHairOrientedBound Out;
 		ComputePCA(InPoints, InPointOffset, InPointCount, Out.Center, Out.ExtentX, Out.ExtentY, Out.ExtentZ);
@@ -1781,26 +1781,26 @@ namespace HairCards
 	}
 
 
-	inline FVector ToLocal(
-		const FVector& InCenter, 
-		const FVector& InAxisX, 
-		const FVector& InAxisY, 
-		const FVector& InAxisZ, 
-		const FVector& InP)
+	inline FVector3f ToLocal(
+		const FVector3f& InCenter, 
+		const FVector3f& InAxisX, 
+		const FVector3f& InAxisY, 
+		const FVector3f& InAxisZ, 
+		const FVector3f& InP)
 	{
-		const FVector Local = InP - InCenter;
-		return FVector(
-			FVector::DotProduct(Local, InAxisX),
-			FVector::DotProduct(Local, InAxisY),
-			FVector::DotProduct(Local, InAxisZ));
+		const FVector3f Local = InP - InCenter;
+		return FVector3f(
+			FVector3f::DotProduct(Local, InAxisX),
+			FVector3f::DotProduct(Local, InAxisY),
+			FVector3f::DotProduct(Local, InAxisZ));
 	}
 
-	inline FVector ToWorld(
-		const FVector& InCenter,
-		const FVector& InAxisX,
-		const FVector& InAxisY,
-		const FVector& InAxisZ,
-		const FVector& InP)
+	inline FVector3f ToWorld(
+		const FVector3f& InCenter,
+		const FVector3f& InAxisX,
+		const FVector3f& InAxisY,
+		const FVector3f& InAxisZ,
+		const FVector3f& InP)
 	{
 		return 
 			InCenter + 
@@ -1809,13 +1809,13 @@ namespace HairCards
 			InP.Z * InAxisZ;
 	}
 
-	inline FVector ToLocal(const FHairOrientedBound& InB, const FVector& InP)
+	inline FVector3f ToLocal(const FHairOrientedBound& InB, const FVector3f& InP)
 	{
 		return ToLocal(InB.Center, InB.ExtentX, InB.ExtentY, InB.ExtentZ, InP);
 	}
 
 
-	static void GetOBBPlane(const FHairOrientedBound& InB, FVector& OutT, FVector& OutB, FVector& OutN)
+	static void GetOBBPlane(const FHairOrientedBound& InB, FVector3f& OutT, FVector3f& OutB, FVector3f& OutN)
 	{
 		uint8 Axis0;
 		uint8 Axis1; 
@@ -1838,13 +1838,13 @@ namespace HairCards
 		Out1.Extent(Axis) = In.Extent(Axis) * 0.5f;
 	}
 
-	static TArray<FVector> Inside(const TArray<FVector>& InPoints, const FHairOrientedBound& InBound)
+	static TArray<FVector3f> Inside(const TArray<FVector3f>& InPoints, const FHairOrientedBound& InBound)
 	{
-		TArray<FVector> Out;
+		TArray<FVector3f> Out;
 		Out.Reserve(InPoints.Num() / 2);
-		for (const FVector& P : InPoints)
+		for (const FVector3f& P : InPoints)
 		{
-			FVector PLocal = ToLocal(InBound, P);
+			FVector3f PLocal = ToLocal(InBound, P);
 			const bool bInside =
 				FMath::Abs(PLocal.X) <= 1 &&
 				FMath::Abs(PLocal.Y) <= 1 &&
@@ -1859,7 +1859,7 @@ namespace HairCards
 		return Out;
 	}
 
-	static void ComputeOrientedBoundingBoxes(const TArray<FVector>& InPoints, float MinSectionLength, TArray<FHairOrientedBound>& Out, TArray<FHairOrientedBound>& OutAll, uint8 Depth=0, uint8 MaxDepth=0)
+	static void ComputeOrientedBoundingBoxes(const TArray<FVector3f>& InPoints, float MinSectionLength, TArray<FHairOrientedBound>& Out, TArray<FHairOrientedBound>& OutAll, uint8 Depth=0, uint8 MaxDepth=0)
 	{
 		const float Scale = 1.2f;
 		FHairOrientedBound B;
@@ -1885,8 +1885,8 @@ namespace HairCards
 			OutAll.Add(B0);
 			OutAll.Add(B1);
 
-			const TArray<FVector> InPoints0 = Inside(InPoints, B0);
-			const TArray<FVector> InPoints1 = Inside(InPoints, B1);
+			const TArray<FVector3f> InPoints0 = Inside(InPoints, B0);
+			const TArray<FVector3f> InPoints1 = Inside(InPoints, B1);
 			ComputeOrientedBoundingBoxes(InPoints0, MinSectionLength, Out, OutAll, Depth + 1, MaxDepth);
 			ComputeOrientedBoundingBoxes(InPoints1, MinSectionLength, Out, OutAll, Depth + 1, MaxDepth);
 		}
@@ -1896,15 +1896,15 @@ namespace HairCards
 		}
 	}
 
-	static void FindTBN(FHairOrientedBound& B0, FHairOrientedBound& B1, FVector& T, FVector& B, FVector& N)
+	static void FindTBN(FHairOrientedBound& B0, FHairOrientedBound& B1, FVector3f& T, FVector3f& B, FVector3f& N)
 	{
-		const FVector& C0 = B0.Center;
-		const FVector& C1 = B1.Center;
-		T = FVector(C1 - C0).GetSafeNormal();
+		const FVector3f& C0 = B0.Center;
+		const FVector3f& C1 = B1.Center;
+		T = FVector3f(C1 - C0).GetSafeNormal();
 
-		const float DX = FVector::DotProduct(T, B0.ExtentX);
-		const float DY = FVector::DotProduct(T, B0.ExtentY);
-		const float DZ = FVector::DotProduct(T, B0.ExtentZ);
+		const float DX = FVector3f::DotProduct(T, B0.ExtentX);
+		const float DY = FVector3f::DotProduct(T, B0.ExtentY);
+		const float DZ = FVector3f::DotProduct(T, B0.ExtentZ);
 
 		uint8 T_Axis = 0;
 		float MaxD = DX;
@@ -1928,7 +1928,7 @@ namespace HairCards
 		B = B0.Extent(B_Axis);
 	}
 
-	static void DecimateCurve(const float InAngularThresholdInDegree, const float OutTargetCount, TArray<FVector>& OutPoints)
+	static void DecimateCurve(const float InAngularThresholdInDegree, const float OutTargetCount, TArray<FVector3f>& OutPoints)
 	{
 		check(OutPoints.Num() > 2);
 		check(OutTargetCount >= 2 && OutTargetCount <= OutPoints.Num());
@@ -1944,11 +1944,11 @@ namespace HairCards
 			const uint32 Count = OutPoints.Num();
 			for (uint32 IndexIt = 1; IndexIt < Count - 1; ++IndexIt)
 			{
-				const FVector& P0 = OutPoints[IndexIt - 1];
-				const FVector& P1 = OutPoints[IndexIt];
-				const FVector& P2 = OutPoints[IndexIt + 1];
+				const FVector3f& P0 = OutPoints[IndexIt - 1];
+				const FVector3f& P1 = OutPoints[IndexIt];
+				const FVector3f& P2 = OutPoints[IndexIt + 1];
 
-				const float Area = FVector::CrossProduct(P0 - P1, P2 - P1).Size() * 0.5f;
+				const float Area = FVector3f::CrossProduct(P0 - P1, P2 - P1).Size() * 0.5f;
 
 				//     P0 .       . P2
 				//         \Inner/
@@ -1956,9 +1956,9 @@ namespace HairCards
 				// Thres(` . \^/ ) Angle
 				//    --------.---------
 				//            P1
-				const FVector V0 = (P0 - P1).GetSafeNormal();
-				const FVector V1 = (P2 - P1).GetSafeNormal();
-				const float InnerAngle = FMath::Abs(FMath::Acos(FVector::DotProduct(V0, V1)));
+				const FVector3f V0 = (P0 - P1).GetSafeNormal();
+				const FVector3f V1 = (P2 - P1).GetSafeNormal();
+				const float InnerAngle = FMath::Abs(FMath::Acos(FVector3f::DotProduct(V0, V1)));
 				const float Angle = (PI - InnerAngle) * 0.5f;
 
 				if (Area < MinError && Angle < AngularThresholdInRad)
@@ -2003,11 +2003,11 @@ namespace HairCards
 			// Select the longest strands
 			// Compute initial size based on the root positions
 			float RootRadius = 0;
-			FVector RootCenter = FVector::ZeroVector;
+			FVector3f RootCenter = FVector3f::ZeroVector;
 			uint32 LongestCurveIndex = 0;
 			float LongestLength = 0;
 			FVector2D RootUV(0, 0);
-			TArray<FVector> Points;
+			TArray<FVector3f> Points;
 			Points.Reserve(CurveCount * 20); // Guess number for avoiding too much reallocation
 			for (uint32 StrandIt = 0; StrandIt < CurveCount; ++StrandIt)
 			{
@@ -2024,7 +2024,7 @@ namespace HairCards
 
 				for (uint32 PointIt = 0; PointIt < PointCount; ++PointIt)
 				{
-					FVector P0 = In.StrandsPoints.PointsPosition[StrandOffset + PointIt];
+					FVector3f P0 = In.StrandsPoints.PointsPosition[StrandOffset + PointIt];
 					Points.Add(P0);
 				}
 
@@ -2046,18 +2046,18 @@ namespace HairCards
 				const uint32 GuidePointCount  = Out.GuidePointCount[ClusterIt];
 				const uint32 GuidePointOffset = Out.GuidePointOffset[ClusterIt];
 				const FVector2D GuideRootUV = Out.GuideRootUVs[ClusterIt];
-				TArray<FVector> BoundPoints;
+				TArray<FVector3f> BoundPoints;
 				TArray<float> BoundLengths;
 				float TotalLength = 0;
 				for (uint32 VertexIt = 0; VertexIt < GuidePointCount; ++VertexIt)
 				{
-					const FVector& P0 = Out.GuidePoints[GuidePointOffset + VertexIt];
+					const FVector3f& P0 = Out.GuidePoints[GuidePointOffset + VertexIt];
 					BoundPoints.Add(P0);
 					BoundLengths.Add(TotalLength);
 
 					if (VertexIt > 0)
 					{
-						const FVector& PPrev = Out.GuidePoints[GuidePointOffset + VertexIt - 1];
+						const FVector3f& PPrev = Out.GuidePoints[GuidePointOffset + VertexIt - 1];
 						TotalLength += (P0 - PPrev).Size();
 					}
 				}
@@ -2073,20 +2073,20 @@ namespace HairCards
 				const float UseCurveOrientation = 1;
 				const float OrientationU = FMath::Clamp(UseCurveOrientation, 0.f, 1.0f);
 
-				FVector PrevN;
+				FVector3f PrevN;
 				for (uint32 VertexIt = 0; VertexIt < VertexCount; ++VertexIt)
 				{
-					const FVector& P0 = BoundPoints[VertexIt];
+					const FVector3f& P0 = BoundPoints[VertexIt];
 
-					FVector T = FVector::ZeroVector;
-					FVector B = FVector::ZeroVector;
-					FVector N = FVector::ZeroVector;
+					FVector3f T = FVector3f::ZeroVector;
+					FVector3f B = FVector3f::ZeroVector;
+					FVector3f N = FVector3f::ZeroVector;
 					GetFrame(P0, InVoxels, T, B, N);
 
 					if (VertexIt == 0)
 					{
 						// Bias the initial normal by assuming the groom is centered aroung a head (spherical hypothesis)
-						const FVector StartN = (P0- In.BoundingBox.GetCenter()).GetSafeNormal();
+						const FVector3f StartN = (P0- In.BoundingBox.GetCenter()).GetSafeNormal();
 						//PrevN = N;
 						PrevN = StartN;
 					}
@@ -2094,30 +2094,30 @@ namespace HairCards
 					const float LengthU = BoundLengths[VertexIt] / TotalLength;
 					PrevN = FMath::Lerp(N, PrevN, FMath::Lerp(LengthU, 1.f, OrientationU));
 
-					FVector PT;
+					FVector3f PT;
 					if (VertexIt > 0 && VertexIt < VertexCount - 1)
 					{
-						const FVector& PPrev = BoundPoints[VertexIt - 1];
-						const FVector& PNext = BoundPoints[VertexIt + 1];
+						const FVector3f& PPrev = BoundPoints[VertexIt - 1];
+						const FVector3f& PNext = BoundPoints[VertexIt + 1];
 
 						PT = ((P0 - PPrev).GetSafeNormal() + (PNext - P0).GetSafeNormal()).GetSafeNormal();
 					}
 					else if (VertexIt == 0)
 					{
-						const FVector& PNext = BoundPoints[VertexIt + 1];
+						const FVector3f& PNext = BoundPoints[VertexIt + 1];
 						PT = (PNext - P0).GetSafeNormal();
 					}
 					else if (VertexIt == VertexCount - 1)
 					{
-						const FVector& PPrev = BoundPoints[VertexIt - 1];
+						const FVector3f& PPrev = BoundPoints[VertexIt - 1];
 						PT = (P0 - PPrev).GetSafeNormal();
 					}
 
 					if (PT.Size() > 0)
 					{
 						T = PT;
-						B = FVector::CrossProduct(PrevN, T).GetSafeNormal();
-						N = FVector::CrossProduct(T, B).GetSafeNormal();
+						B = FVector3f::CrossProduct(PrevN, T).GetSafeNormal();
+						N = FVector3f::CrossProduct(T, B).GetSafeNormal();
 					}
 
 					Out.BoundAxis_T.Add(T);
@@ -2239,19 +2239,19 @@ namespace HairCards
 
 				float CardMaxWidth = 0;
 				float CardLength = 0;
-				FVector PrevP0(0, 0, 0);
-				FVector PrevT(0, 0, 0);
+				FVector3f PrevP0(0, 0, 0);
+				FVector3f PrevT(0, 0, 0);
 				const uint32 StartVertexIt = VertexIt;
 				const uint32 StartIndexIt = IndexIt;
 				for (uint32 PointIt = 0; PointIt < PointCount; ++PointIt)
 				{
 					// Simple geometric normal based on adjacent vertices
-					const FVector P0 = InClusters.BoundPositions[PointOffset + PointIt];
-					FVector T;
+					const FVector3f P0 = InClusters.BoundPositions[PointOffset + PointIt];
+					FVector3f T;
 					float SegmentLength = 0;
 					if (PointIt < PointCount - 1)
 					{
-						const FVector P1 = InClusters.BoundPositions[PointOffset + PointIt + 1];
+						const FVector3f P1 = InClusters.BoundPositions[PointOffset + PointIt + 1];
 						SegmentLength = (P1 - P0).Size();
 						T = (P1 - P0).GetSafeNormal();
 					}
@@ -2260,8 +2260,8 @@ namespace HairCards
 						T = PrevT;
 					}
 
-					const FVector Cluster_N = InClusters.BoundAxis_N[PointOffset + PointIt];
-					const FVector B = InClusters.BoundAxis_B[PointOffset + PointIt].GetSafeNormal();
+					const FVector3f Cluster_N = InClusters.BoundAxis_N[PointOffset + PointIt];
+					const FVector3f B = InClusters.BoundAxis_B[PointOffset + PointIt].GetSafeNormal();
 					const float CardWidth = InClusters.BoundAxis_B[PointOffset + PointIt].Size();
 					CardMaxWidth = FMath::Max(CardWidth, CardMaxWidth);
 
@@ -2270,14 +2270,14 @@ namespace HairCards
 					check(FMath::IsFinite(OutCards.Positions[VertexIt].X)     && FMath::IsFinite(OutCards.Positions[VertexIt].Y)     && FMath::IsFinite(OutCards.Positions[VertexIt].Z));
 					check(FMath::IsFinite(OutCards.Positions[VertexIt + 1].X) && FMath::IsFinite(OutCards.Positions[VertexIt + 1].Y) && FMath::IsFinite(OutCards.Positions[VertexIt + 1].Z));
 
-					OutCards.Normals[VertexIt    ] = FVector::ZeroVector; //MinorAxis0;
-					OutCards.Normals[VertexIt + 1] = FVector::ZeroVector; //MinorAxis0;
+					OutCards.Normals[VertexIt    ] = FVector3f::ZeroVector; //MinorAxis0;
+					OutCards.Normals[VertexIt + 1] = FVector3f::ZeroVector; //MinorAxis0;
 
 					// Compute smooth tangent from the previous & next vertices
 					if (PointIt > 0)
 					{
-						const FVector T00 = OutCards.Positions[VertexIt]     - OutCards.Positions[VertexIt - 2];
-						const FVector T01 = OutCards.Positions[VertexIt + 1] - OutCards.Positions[VertexIt - 2 + 1];
+						const FVector3f T00 = OutCards.Positions[VertexIt]     - OutCards.Positions[VertexIt - 2];
+						const FVector3f T01 = OutCards.Positions[VertexIt + 1] - OutCards.Positions[VertexIt - 2 + 1];
 
 						OutCards.Tangents[VertexIt-2    ] = T00.GetSafeNormal();
 						OutCards.Tangents[VertexIt-2 + 1] = T01.GetSafeNormal();
@@ -2324,19 +2324,19 @@ namespace HairCards
 					const uint32 I1 = OutCards.Indices[InnerIndexIt + 1];
 					const uint32 I2 = OutCards.Indices[InnerIndexIt + 2];
 
-					const FVector P0 = OutCards.Positions[I0];
-					const FVector P1 = OutCards.Positions[I1];
-					const FVector P2 = OutCards.Positions[I2];
+					const FVector3f P0 = OutCards.Positions[I0];
+					const FVector3f P1 = OutCards.Positions[I1];
+					const FVector3f P2 = OutCards.Positions[I2];
 
-					const FVector E0 = (P1 - P0).GetSafeNormal();
-					const FVector E1 = (P2 - P1).GetSafeNormal();
-					const FVector E2 = (P0 - P2).GetSafeNormal();
+					const FVector3f E0 = (P1 - P0).GetSafeNormal();
+					const FVector3f E1 = (P2 - P1).GetSafeNormal();
+					const FVector3f E2 = (P0 - P2).GetSafeNormal();
 
-					const FVector N = -FVector::CrossProduct(E0, -E2).GetSafeNormal();
+					const FVector3f N = -FVector3f::CrossProduct(E0, -E2).GetSafeNormal();
 
-					OutCards.Normals[I0] += N * FMath::Acos(FVector::DotProduct(E0, -E0));
-					OutCards.Normals[I1] += N * FMath::Acos(FVector::DotProduct(E1, -E1));
-					OutCards.Normals[I2] += N * FMath::Acos(FVector::DotProduct(E2, -E2));
+					OutCards.Normals[I0] += N * FMath::Acos(FVector3f::DotProduct(E0, -E0));
+					OutCards.Normals[I1] += N * FMath::Acos(FVector3f::DotProduct(E1, -E1));
+					OutCards.Normals[I2] += N * FMath::Acos(FVector3f::DotProduct(E2, -E2));
 				}
 
 				const uint32 RepresentativeRectIndex = InClusters.AtlasRectIndex[ClusterIt];
@@ -2353,10 +2353,10 @@ namespace HairCards
 					{
 						const uint32 VIndex = StartVertexIt + PointIt*2 + InnerVertexIt;
 
-						const FVector N = OutCards.Normals[VIndex].GetSafeNormal();
-						FVector T		= OutCards.Tangents[VIndex];
-						const FVector B = FVector::CrossProduct(N, T);
-						T = FVector::CrossProduct(B, N).GetSafeNormal();
+						const FVector3f N = OutCards.Normals[VIndex].GetSafeNormal();
+						FVector3f T		= OutCards.Tangents[VIndex];
+						const FVector3f B = FVector3f::CrossProduct(N, T);
+						T = FVector3f::CrossProduct(B, N).GetSafeNormal();
 
 						OutCards.Normals [VIndex] = N;
 						OutCards.Tangents[VIndex] = T;
@@ -2781,12 +2781,12 @@ namespace HairCards
 			for (uint32 PointIt = 0; PointIt < PointCount; ++PointIt)
 			{
 				const uint32 PointIndex = PointOffset + PointIt;
-				const FVector P0 = InClusters.BoundPositions[PointIndex];
+				const FVector3f P0 = InClusters.BoundPositions[PointIndex];
 
 				OutGuides.BoundingBox += P0;
 
 				OutGuides.StrandsPoints.PointsPosition[PointIndex] = P0;
-				OutGuides.StrandsPoints.PointsBaseColor[PointIndex] = FVector::ZeroVector;
+				OutGuides.StrandsPoints.PointsBaseColor[PointIndex] = FVector3f::ZeroVector;
 				OutGuides.StrandsPoints.PointsRoughness[PointIndex] = 0;
 				OutGuides.StrandsPoints.PointsCoordU[PointIndex] = FMath::Clamp(CurrentLength / TotalLength, 0.f, 1.f);
 				OutGuides.StrandsPoints.PointsRadius[PointIndex] = 1;
@@ -2794,7 +2794,7 @@ namespace HairCards
 				// Simple geometric normal based on adjacent vertices
 				if (PointIt < PointCount - 1)
 				{
-					const FVector P1 = InClusters.BoundPositions[PointOffset + PointIt + 1];
+					const FVector3f P1 = InClusters.BoundPositions[PointOffset + PointIt + 1];
 					const float SegmentLength = (P1 - P0).Size();
 					CurrentLength += SegmentLength;
 				}
@@ -2889,7 +2889,7 @@ namespace HairCards
 
 			// Find the perpendicular direction by comparing the number of segment along U and along V
 			const bool bIsMainDirectionU = SimilarVertexU.Num() >= SimilarVertexV.Num();
-			TArray<FVector> CenterPoints;
+			TArray<FVector3f> CenterPoints;
 			{
 				// Sort vertices along the main axis so that, when we iterate through them, we get a correct linear ordering
 				TArray<FSimilarUVVertices>& SimilarVertex = bIsMainDirectionU ? SimilarVertexU : SimilarVertexV;
@@ -2899,12 +2899,12 @@ namespace HairCards
 				});
 
 				CenterPoints.Reserve(SimilarVertex.Num());
-				FVector PrevCenterPoint = FVector::ZeroVector;
+				FVector3f PrevCenterPoint = FVector3f::ZeroVector;
 				float TotalLength = 0;
 				for (const FSimilarUVVertices& Similar : SimilarVertex)
 				{
 					// Compute avg center point of the guide
-					FVector CenterPoint = FVector::ZeroVector;
+					FVector3f CenterPoint = FVector3f::ZeroVector;
 					for (uint32 VertexIndex : Similar.Indices)
 					{
 						CenterPoint += InCards.Positions[VertexIndex];
@@ -2946,9 +2946,9 @@ namespace HairCards
 			check(CenterPoints.Num() > 0);
 			if (CenterPoints.Num() == 1)
 			{
-				const FVector CenterPoint = CenterPoints[0];
+				const FVector3f CenterPoint = CenterPoints[0];
 				const float SegmentSize = 0.5f;
-				const FVector P1 = CenterPoint + SegmentSize * (CenterPoint - InCards.BoundingBox.GetCenter()).GetSafeNormal();
+				const FVector3f P1 = CenterPoint + SegmentSize * (CenterPoint - InCards.BoundingBox.GetCenter()).GetSafeNormal();
 				CenterPoints.Add(P1);
 			}
 
@@ -2973,12 +2973,12 @@ namespace HairCards
 			float CurrentLength = 0;
 			for (uint32 PointIt = 0; PointIt < PointCount; ++PointIt)
 			{
-				const FVector P0 = CenterPoints[PointIt];
+				const FVector3f P0 = CenterPoints[PointIt];
 
 				OutGuides.BoundingBox += P0;
 
 				OutGuides.StrandsPoints.PointsPosition.Add(P0);
-				OutGuides.StrandsPoints.PointsBaseColor.Add(FVector::ZeroVector);
+				OutGuides.StrandsPoints.PointsBaseColor.Add(FVector3f::ZeroVector);
 				OutGuides.StrandsPoints.PointsRoughness.Add(0);
 				OutGuides.StrandsPoints.PointsCoordU.Add(FMath::Clamp(CurrentLength / TotalLength, 0.f, 1.f));
 				OutGuides.StrandsPoints.PointsRadius.Add(1);
@@ -2986,7 +2986,7 @@ namespace HairCards
 				// Simple geometric normal based on adjacent vertices
 				if (PointIt < PointCount - 1)
 				{
-					const FVector P1 = CenterPoints[PointIt + 1];
+					const FVector3f P1 = CenterPoints[PointIt + 1];
 					const float SegmentLength = (P1 - P0).Size();
 					CurrentLength += SegmentLength;
 				}
@@ -3097,40 +3097,40 @@ namespace HairCards
 			const uint32 CurveOffset  = InClusters.CurveOffset[ClusterIndex];
 			const uint32 CurveCount   = InClusters.CurveCount[ClusterIndex];
 
-			OutRect.MinBound = FVector( FLT_MAX);
-			OutRect.MaxBound = FVector(-FLT_MAX);
+			OutRect.MinBound = FVector3f( FLT_MAX);
+			OutRect.MaxBound = FVector3f(-FLT_MAX);
 
 			// Precompute guide rest & deformed position as well as the length of the guides.
-			TArray<FVector> Guide_RestPositions;
-			TArray<FVector> Guide_DeformedPositions;
+			TArray<FVector3f> Guide_RestPositions;
+			TArray<FVector3f> Guide_DeformedPositions;
 			TArray<float>   Guide_Lengths;
-			FVector AlignmentDir = FVector::ZeroVector;
+			FVector3f AlignmentDir = FVector3f::ZeroVector;
 			{
 				const uint32 GuidePointCount  = InClusters.GuidePointCount[ClusterIndex];
 				const uint32 GuidePointOffset = InClusters.GuidePointOffset[ClusterIndex];
 
 				check(GuidePointCount >= 2);
-				const FVector Dir_P0 = InClusters.GuidePoints[GuidePointOffset];
-				const FVector Dir_P1 = InClusters.GuidePoints[GuidePointOffset + 1];
+				const FVector3f Dir_P0 = InClusters.GuidePoints[GuidePointOffset];
+				const FVector3f Dir_P1 = InClusters.GuidePoints[GuidePointOffset + 1];
 				AlignmentDir = (Dir_P1 - Dir_P0).GetSafeNormal();
 
 				// Align the direction onto one of the cardinal axis for easing the projection
 				{
-					FVector MajorAxis = FVector(1,0,0) * FMath::Sign(AlignmentDir.X);
+					FVector3f MajorAxis = FVector3f(1,0,0) * FMath::Sign(AlignmentDir.X);
 					float S = FMath::Abs(AlignmentDir.X);
-					if (FMath::Abs(AlignmentDir.Y) > S) { MajorAxis = FMath::Sign(AlignmentDir.Y) * FVector(0, 1, 0); S = FMath::Abs(AlignmentDir.Y); }
-					if (FMath::Abs(AlignmentDir.Z) > S) { MajorAxis = FMath::Sign(AlignmentDir.Z) * FVector(0, 0, 1); S = FMath::Abs(AlignmentDir.Z); }
+					if (FMath::Abs(AlignmentDir.Y) > S) { MajorAxis = FMath::Sign(AlignmentDir.Y) * FVector3f(0, 1, 0); S = FMath::Abs(AlignmentDir.Y); }
+					if (FMath::Abs(AlignmentDir.Z) > S) { MajorAxis = FMath::Sign(AlignmentDir.Z) * FVector3f(0, 0, 1); S = FMath::Abs(AlignmentDir.Z); }
 					AlignmentDir = MajorAxis;
 				}
-				const FVector Root = Dir_P0;
+				const FVector3f Root = Dir_P0;
 
 				float Length = 0;
 				for (uint32 PointIt = 0; PointIt < GuidePointCount; ++PointIt)
 				{
 					const uint32 GuidePointIt = GuidePointOffset + PointIt;
 
-					const FVector& P0_Rest = InClusters.GuidePoints[GuidePointIt];
-					const FVector  P0_Deformed = Root + AlignmentDir * Length;
+					const FVector3f& P0_Rest = InClusters.GuidePoints[GuidePointIt];
+					const FVector3f  P0_Deformed = Root + AlignmentDir * Length;
 
 					Guide_RestPositions.Add(P0_Rest);
 					Guide_DeformedPositions.Add(P0_Deformed);
@@ -3138,7 +3138,7 @@ namespace HairCards
 
 					if (PointIt < GuidePointCount-1)
 					{
-						const FVector& P1_Rest = InClusters.GuidePoints[GuidePointIt+1];
+						const FVector3f& P1_Rest = InClusters.GuidePoints[GuidePointIt+1];
 						const float L = (P1_Rest - P0_Rest).Size();
 						Length += L > 0 ? L : 0.001f;
 					}
@@ -3165,7 +3165,7 @@ namespace HairCards
 					const FVector2D RootUV = InStrands.StrandsCurves.CurvesRootUV[CurveIndex];
 					const float Seed = 0.f; // TODO float(InStrands.RenderData.Attributes[PointIndex].Seed) / 255.f;
 
-					const FVector& P0_Rest = InStrands.StrandsPoints.PointsPosition[PointIndex];
+					const FVector3f& P0_Rest = InStrands.StrandsPoints.PointsPosition[PointIndex];
 					const float Radius = InStrands.StrandsPoints.PointsRadius[PointIndex] * InStrands.StrandsCurves.MaxRadius;
 
 					// Find the closest guide point using parametric length at a correspondance/matching metric
@@ -3178,38 +3178,38 @@ namespace HairCards
 					}
 
 					// Once the bounded indices are compute the displacement vector used for deforming the curves
-					FVector Displacement = FVector::ZeroVector;
+					FVector3f Displacement = FVector3f::ZeroVector;
 					if (ClosestGuidePointIt1 < GuidePointCount)
 					{
 						const float U = (Length - Guide_Lengths[ClosestGuidePointIt0]) / (Guide_Lengths[ClosestGuidePointIt1] - Guide_Lengths[ClosestGuidePointIt0]);
-						const FVector& GuideP_Rest = FMath::Lerp(Guide_RestPositions[ClosestGuidePointIt0], Guide_RestPositions[ClosestGuidePointIt1], U);
-						const FVector& GuideP_Deformed = FMath::Lerp(Guide_DeformedPositions[ClosestGuidePointIt0], Guide_DeformedPositions[ClosestGuidePointIt1], U);
+						const FVector3f& GuideP_Rest = FMath::Lerp(Guide_RestPositions[ClosestGuidePointIt0], Guide_RestPositions[ClosestGuidePointIt1], U);
+						const FVector3f& GuideP_Deformed = FMath::Lerp(Guide_DeformedPositions[ClosestGuidePointIt0], Guide_DeformedPositions[ClosestGuidePointIt1], U);
 						Displacement = GuideP_Deformed - GuideP_Rest;
 					}
 					else
 					{
 						// Extrapolate the guide rest position, and guide deformed position to compute the displacement
-						const FVector& GuideP_Rest0 = Guide_RestPositions[GuidePointCount - 2];
-						const FVector& GuideP_Rest1 = Guide_RestPositions[GuidePointCount - 1];
+						const FVector3f& GuideP_Rest0 = Guide_RestPositions[GuidePointCount - 2];
+						const FVector3f& GuideP_Rest1 = Guide_RestPositions[GuidePointCount - 1];
 
-						const FVector& GuideP_Deformed0 = Guide_DeformedPositions[GuidePointCount - 2];
-						const FVector& GuideP_Deformed1 = Guide_DeformedPositions[GuidePointCount - 1];
+						const FVector3f& GuideP_Deformed0 = Guide_DeformedPositions[GuidePointCount - 2];
+						const FVector3f& GuideP_Deformed1 = Guide_DeformedPositions[GuidePointCount - 1];
 
-						const FVector GuideDir_Rest		= (GuideP_Rest1 - GuideP_Rest0).GetSafeNormal();
-						const FVector GuideDir_Deformed = (GuideP_Deformed1 - GuideP_Deformed0).GetSafeNormal();
+						const FVector3f GuideDir_Rest		= (GuideP_Rest1 - GuideP_Rest0).GetSafeNormal();
+						const FVector3f GuideDir_Deformed = (GuideP_Deformed1 - GuideP_Deformed0).GetSafeNormal();
 
 						const float ExtraLength = Length - Guide_Lengths[GuidePointCount - 1];
-						const FVector& GuideP_RestExtrapolated		= GuideP_Rest1 + GuideDir_Rest * ExtraLength;
-						const FVector& GuideP_DeformedExtrapolated	= GuideP_Deformed1 + GuideDir_Deformed * ExtraLength;
+						const FVector3f& GuideP_RestExtrapolated		= GuideP_Rest1 + GuideDir_Rest * ExtraLength;
+						const FVector3f& GuideP_DeformedExtrapolated	= GuideP_Deformed1 + GuideDir_Deformed * ExtraLength;
 
 						Displacement = GuideP_DeformedExtrapolated - GuideP_RestExtrapolated;
 					}
 
-					const FVector P0_Deformed = P0_Rest + Displacement;
+					const FVector3f P0_Deformed = P0_Rest + Displacement;
 
 					if (PointIt < PointCount - 1)
 					{
-						const FVector& P1_Rest = InStrands.StrandsPoints.PointsPosition[PointOffset + PointIt + 1];
+						const FVector3f& P1_Rest = InStrands.StrandsPoints.PointsPosition[PointOffset + PointIt + 1];
 						Length += (P1_Rest - P0_Rest).Size();
 					}
 
@@ -3231,28 +3231,28 @@ namespace HairCards
 			// Compute rasterization info (rasterization axis) and cards dimension (width/length)
 			{
 
-				FVector::FReal Size[3] = 
+				FVector3f::FReal Size[3] = 
 				{ 
 					OutRect.MaxBound.X - OutRect.MinBound.X, 
 					OutRect.MaxBound.Y - OutRect.MinBound.Y, 
 					OutRect.MaxBound.Z - OutRect.MinBound.Z 
 				};
-				FVector Axis[3] = 
+				FVector3f Axis[3] = 
 				{ 
-					Size[0] * FVector(1,0,0),
-					Size[1] * FVector(0,1,0), 
-					Size[2] * FVector(0,0,1) 
+					Size[0] * FVector3f(1,0,0),
+					Size[1] * FVector3f(0,1,0), 
+					Size[2] * FVector3f(0,0,1) 
 				};
 
 				if (Size[0] < Size[1]) { Swap(Axis[0], Axis[1]); Swap(Size[0], Size[1]); }
 				if (Size[1] < Size[2]) { Swap(Axis[1], Axis[2]); Swap(Size[1], Size[2]); }
 				if (Size[0] < Size[1]) { Swap(Axis[0], Axis[1]); Swap(Size[0], Size[1]); }
 
-				FVector T = Axis[0];
-				FVector B = Axis[1];
-				FVector N = Axis[2];
+				FVector3f T = Axis[0];
+				FVector3f B = Axis[1];
+				FVector3f N = Axis[2];
 
-				const FVector NormalizedT = T.GetSafeNormal();
+				const FVector3f NormalizedT = T.GetSafeNormal();
 
 				if (FMath::Sign(T.X) != FMath::Sign(AlignmentDir.X) ||
 					FMath::Sign(T.Y) != FMath::Sign(AlignmentDir.Y) ||
@@ -3933,8 +3933,8 @@ void BuildGeometry(
 	const FBox& InBox,
 	FHairMeshesBulkData& OutBulk)
 {
-	const FVector Center = InBox.GetCenter();
-	const FVector Extent = InBox.GetExtent();
+	const FVector3f Center = InBox.GetCenter();
+	const FVector3f Extent = InBox.GetExtent();
 
 	// Simple (incorrect normal/tangent) cube geomtry in place of the hair rendering
 	const uint32 TotalPointCount = 8;
@@ -3948,14 +3948,14 @@ void BuildGeometry(
 	Out.Meshes.UVs.SetNum(TotalPointCount);
 	Out.Meshes.Indices.SetNum(TotalIndexCount);
 
-	Out.Meshes.Positions[0] = Center + FVector(-Extent.X, -Extent.Y, -Extent.Z);
-	Out.Meshes.Positions[1] = Center + FVector(+Extent.X, -Extent.Y, -Extent.Z);
-	Out.Meshes.Positions[2] = Center + FVector(+Extent.X, +Extent.Y, -Extent.Z);
-	Out.Meshes.Positions[3] = Center + FVector(-Extent.X, +Extent.Y, -Extent.Z);
-	Out.Meshes.Positions[4] = Center + FVector(-Extent.X, -Extent.Y, +Extent.Z);
-	Out.Meshes.Positions[5] = Center + FVector(+Extent.X, -Extent.Y, +Extent.Z);
-	Out.Meshes.Positions[6] = Center + FVector(+Extent.X, +Extent.Y, +Extent.Z);
-	Out.Meshes.Positions[7] = Center + FVector(-Extent.X, +Extent.Y, +Extent.Z);
+	Out.Meshes.Positions[0] = Center + FVector3f(-Extent.X, -Extent.Y, -Extent.Z);
+	Out.Meshes.Positions[1] = Center + FVector3f(+Extent.X, -Extent.Y, -Extent.Z);
+	Out.Meshes.Positions[2] = Center + FVector3f(+Extent.X, +Extent.Y, -Extent.Z);
+	Out.Meshes.Positions[3] = Center + FVector3f(-Extent.X, +Extent.Y, -Extent.Z);
+	Out.Meshes.Positions[4] = Center + FVector3f(-Extent.X, -Extent.Y, +Extent.Z);
+	Out.Meshes.Positions[5] = Center + FVector3f(+Extent.X, -Extent.Y, +Extent.Z);
+	Out.Meshes.Positions[6] = Center + FVector3f(+Extent.X, +Extent.Y, +Extent.Z);
+	Out.Meshes.Positions[7] = Center + FVector3f(-Extent.X, +Extent.Y, +Extent.Z);
 
 	Out.Meshes.UVs[0] = FVector2D(0, 0);
 	Out.Meshes.UVs[1] = FVector2D(1, 0);
@@ -3966,23 +3966,23 @@ void BuildGeometry(
 	Out.Meshes.UVs[6] = FVector2D(1, 1);
 	Out.Meshes.UVs[7] = FVector2D(0, 1);
 
-	Out.Meshes.Normals[0] = FVector(0, 0, 1);
-	Out.Meshes.Normals[1] = FVector(0, 0, 1);
-	Out.Meshes.Normals[2] = FVector(0, 0, 1);
-	Out.Meshes.Normals[3] = FVector(0, 0, 1);
-	Out.Meshes.Normals[4] = FVector(0, 0, 1);
-	Out.Meshes.Normals[5] = FVector(0, 0, 1);
-	Out.Meshes.Normals[6] = FVector(0, 0, 1);
-	Out.Meshes.Normals[7] = FVector(0, 0, 1);
-
-	Out.Meshes.Tangents[0] = FVector(1, 0, 0);
-	Out.Meshes.Tangents[1] = FVector(1, 0, 0);
-	Out.Meshes.Tangents[2] = FVector(1, 0, 0);
-	Out.Meshes.Tangents[3] = FVector(1, 0, 0);
-	Out.Meshes.Tangents[4] = FVector(1, 0, 0);
-	Out.Meshes.Tangents[5] = FVector(1, 0, 0);
-	Out.Meshes.Tangents[6] = FVector(1, 0, 0);
-	Out.Meshes.Tangents[7] = FVector(1, 0, 0);
+	Out.Meshes.Normals[0] = FVector3f(0, 0, 1);
+	Out.Meshes.Normals[1] = FVector3f(0, 0, 1);
+	Out.Meshes.Normals[2] = FVector3f(0, 0, 1);
+	Out.Meshes.Normals[3] = FVector3f(0, 0, 1);
+	Out.Meshes.Normals[4] = FVector3f(0, 0, 1);
+	Out.Meshes.Normals[5] = FVector3f(0, 0, 1);
+	Out.Meshes.Normals[6] = FVector3f(0, 0, 1);
+	Out.Meshes.Normals[7] = FVector3f(0, 0, 1);
+
+	Out.Meshes.Tangents[0] = FVector3f(1, 0, 0);
+	Out.Meshes.Tangents[1] = FVector3f(1, 0, 0);
+	Out.Meshes.Tangents[2] = FVector3f(1, 0, 0);
+	Out.Meshes.Tangents[3] = FVector3f(1, 0, 0);
+	Out.Meshes.Tangents[4] = FVector3f(1, 0, 0);
+	Out.Meshes.Tangents[5] = FVector3f(1, 0, 0);
+	Out.Meshes.Tangents[6] = FVector3f(1, 0, 0);
+	Out.Meshes.Tangents[7] = FVector3f(1, 0, 0);
 
 	Out.Meshes.Indices[0] = 0;
 	Out.Meshes.Indices[1] = 1;
@@ -4121,13 +4121,13 @@ public:
 	void SetMeshDescription(FMeshDescription* Description);
 
 	/** Append vertex and return new vertex ID */
-	FVertexID AppendVertex(const FVector& Position);
+	FVertexID AppendVertex(const FVector3f& Position);
 
 	/** Append new vertex instance and return ID */
 	FVertexInstanceID AppendInstance(const FVertexID& VertexID);
 
 	/** Set the Normal of a vertex instance*/
-	void SetInstanceNormal(const FVertexInstanceID& InstanceID, const FVector& Normal);
+	void SetInstanceNormal(const FVertexInstanceID& InstanceID, const FVector3f& Normal);
 
 	/** Set the UV of a vertex instance */
 	void SetInstanceUV(const FVertexInstanceID& InstanceID, const FVector2D& InstanceUV, int32 UVLayerIndex = 0);
@@ -4148,13 +4148,13 @@ public:
 	FPolygonID AppendTriangle(const FVertexID& Vertex0, const FVertexID& Vertex1, const FVertexID& Vertex2, const FPolygonGroupID& PolygonGroup);
 
 	/** Append a triangle to the mesh with the given PolygonGroup ID, and optionally with triangle-vertex UVs and Normals */
-	FPolygonID AppendTriangle(const FVertexID* Triangle, const FPolygonGroupID& PolygonGroup, const FVector2D* VertexUVs = nullptr, const FVector* VertexNormals = nullptr);
+	FPolygonID AppendTriangle(const FVertexID* Triangle, const FPolygonGroupID& PolygonGroup, const FVector2D* VertexUVs = nullptr, const FVector3f* VertexNormals = nullptr);
 
 	/**
 	 * Append an arbitrary polygon to the mesh with the given PolygonGroup ID, and optionally with polygon-vertex UVs and Normals
 	 * Unique Vertex instances will be created for each polygon-vertex.
 	 */
-	FPolygonID AppendPolygon(const TArray<FVertexID>& Vertices, const FPolygonGroupID& PolygonGroup, const TArray<FVector2D>* VertexUVs = nullptr, const TArray<FVector>* VertexNormals = nullptr);
+	FPolygonID AppendPolygon(const TArray<FVertexID>& Vertices, const FPolygonGroupID& PolygonGroup, const TArray<FVector2D>* VertexUVs = nullptr, const TArray<FVector3f>* VertexNormals = nullptr);
 
 	/**
 	 * Append a triangle to the mesh using the given vertex instances and PolygonGroup ID
@@ -4197,7 +4197,7 @@ void FMeshDescriptionBuilder::EnablePolyGroups()
 	}
 }
 
-FVertexID FMeshDescriptionBuilder::AppendVertex(const FVector& Position)
+FVertexID FMeshDescriptionBuilder::AppendVertex(const FVector3f& Position)
 {
 	FVertexID VertexID = MeshDescription->CreateVertex();
 	VertexPositions.Set(VertexID, Position);
@@ -4215,7 +4215,7 @@ FVertexInstanceID FMeshDescriptionBuilder::AppendInstance(const FVertexID& Verte
 	return MeshDescription->CreateVertexInstance(VertexID);
 }
 
-void FMeshDescriptionBuilder::SetInstanceNormal(const FVertexInstanceID& InstanceID, const FVector& Normal)
+void FMeshDescriptionBuilder::SetInstanceNormal(const FVertexInstanceID& InstanceID, const FVector3f& Normal)
 {
 	if (InstanceNormals.IsValid())
 	{
diff --git a/Engine/Plugins/Runtime/HairStrands/Source/HairStrandsCore/Public/HairCardsDatas.h b/Engine/Plugins/Runtime/HairStrands/Source/HairStrandsCore/Public/HairCardsDatas.h
--- a/Engine/Plugins/Runtime/HairStrands/Source/HairStrandsCore/Public/HairCardsDatas.h
+++ b/Engine/Plugins/Runtime/HairStrands/Source/HairStrandsCore/Public/HairCardsDatas.h
@@ -147,12 +147,12 @@ struct FHairCardsInterpolationFormat
 
 struct FHairOrientedBound
 {
-	FVector Center;
-	FVector ExtentX;
-	FVector ExtentY;
-	FVector ExtentZ;
+	FVector3f Center;
+	FVector3f ExtentX;
+	FVector3f ExtentY;
+	FVector3f ExtentZ;
 
-	FVector& Extent(uint8 Axis)
+	FVector3f& Extent(uint8 Axis)
 	{
 		if (Axis == 0) return ExtentX;
 		if (Axis == 1) return ExtentY;
@@ -160,7 +160,7 @@ struct FHairOrientedBound
 		return ExtentX;
 	}
 
-	const FVector& Extent(uint8 Axis) const
+	const FVector3f& Extent(uint8 Axis) const
 	{
 		if (Axis == 0) return ExtentX;
 		if (Axis == 1) return ExtentY;
@@ -177,9 +177,9 @@ struct FHairCardsGeometry
 {
 	// Geometry
 	TArray<FVector4> UVs;
-	TArray<FVector>  Normals;
-	TArray<FVector>  Tangents;
-	TArray<FVector>  Positions;
+	TArray<FVector3f>Normals;
+	TArray<FVector3f>Tangents;
+	TArray<FVector3f>Positions;
 	TArray<uint32>   Indices;
 	TArray<float>    CoordU; // Transient data storing [0..1] parametric value along main axis. This is used for generating guides & interpolation data
 
@@ -293,8 +293,8 @@ struct FHairCardsVoxel
 	FRDGExternalBuffer NormalBuffer;
 	FRDGExternalBuffer DensityBuffer;
 	FIntVector	Resolution;
-	FVector		MinBound;
-	FVector		MaxBound;
+	FVector3f	MinBound;
+	FVector3f	MaxBound;
 	float		VoxelSize;
 };
 
@@ -344,12 +344,12 @@ struct FHairCardsProceduralAtlas
 		FIntPoint Resolution = FIntPoint(0, 0);
 		uint32 VertexOffset = 0;
 		uint32 VertexCount = 0;
-		FVector MinBound = FVector::ZeroVector;
-		FVector MaxBound = FVector::ZeroVector;
+		FVector3f MinBound = FVector3f::ZeroVector;
+		FVector3f MaxBound = FVector3f::ZeroVector;
 
-		FVector RasterAxisX = FVector::ZeroVector;
-		FVector RasterAxisY = FVector::ZeroVector;
-		FVector RasterAxisZ = FVector::ZeroVector;
+		FVector3f RasterAxisX = FVector3f::ZeroVector;
+		FVector3f RasterAxisY = FVector3f::ZeroVector;
+		FVector3f RasterAxisZ = FVector3f::ZeroVector;
 
 		float CardWidth = 0;
 		float CardLength = 0;
@@ -407,9 +407,9 @@ struct FHairMeshes
 {
 	// Geometry
 	TArray<FVector2D> UVs;
-	TArray<FVector>   Normals;
-	TArray<FVector>   Tangents;
-	TArray<FVector>   Positions;
+	TArray<FVector3f> Normals;
+	TArray<FVector3f> Tangents;
+	TArray<FVector3f> Positions;
 	TArray<uint32>    Indices;
 
 	FBox BoundingBox;
diff --git a/Engine/Plugins/Runtime/HairStrands/Source/HairStrandsCore/Public/HairStrandsDatas.h b/Engine/Plugins/Runtime/HairStrands/Source/HairStrandsCore/Public/HairStrandsDatas.h
--- a/Engine/Plugins/Runtime/HairStrands/Source/HairStrandsCore/Public/HairStrandsDatas.h
+++ b/Engine/Plugins/Runtime/HairStrands/Source/HairStrandsCore/Public/HairStrandsDatas.h
@@ -301,10 +301,10 @@ struct HAIRSTRANDSCORE_API FHairStrandsInterpolationDatas
 	TArray<FIntVector> PointsSimCurvesVertexIndex;
 
 	/** Lerp value between the closest vertex indices and the next one, ordered by closest influence */
-	TArray<FVector> PointsSimCurvesVertexLerp;
+	TArray<FVector3f> PointsSimCurvesVertexLerp;
 
 	/** Weight of vertex indices on simulation curve, ordered by closest influence */
-	TArray<FVector>	PointsSimCurvesVertexWeights;
+	TArray<FVector3f>	PointsSimCurvesVertexWeights;
 
 	/** True, if interpolation data are built using a single guide */
 	bool bUseUniqueGuide = false;
diff --git a/Engine/Source/Runtime/Renderer/Private/HairStrands/HairStrandsData.h b/Engine/Source/Runtime/Renderer/Private/HairStrands/HairStrandsData.h
--- a/Engine/Source/Runtime/Renderer/Private/HairStrands/HairStrandsData.h
+++ b/Engine/Source/Runtime/Renderer/Private/HairStrands/HairStrandsData.h
@@ -119,8 +119,8 @@ struct FHairStrandsVisibilityData
 
 struct FHairStrandsVoxelNodeDesc
 {
-	FVector WorldMinAABB = FVector::ZeroVector;
-	FVector WorldMaxAABB = FVector::ZeroVector;
+	FVector3f WorldMinAABB = FVector3f::ZeroVector;
+	FVector3f WorldMaxAABB = FVector3f::ZeroVector;
 	FIntVector PageIndexResolution = FIntVector::ZeroValue;
 	FMatrix WorldToClip;
 };
@@ -132,10 +132,10 @@ struct FPackedVirtualVoxelNodeDesc
 	const static uint32 ComponentCount = 2;
 
 	// Shader View is struct { uint4; uint4; }
-	FVector	MinAABB;
-	uint32	PackedPageIndexResolution;
-	FVector	MaxAABB;
-	uint32	PageIndexOffset;
+	FVector3f	MinAABB;
+	uint32		PackedPageIndexResolution;
+	FVector3f	MaxAABB;
+	uint32		PageIndexOffset;
 };
 
 // PixelRadiusAtDepth1 shouldn't be stored into this structure should be view independent, 
@@ -223,7 +223,7 @@ struct FHairStrandsDeepShadowData
 	FIntPoint ShadowResolution = FIntPoint::ZeroValue;
 	uint32 LightId = ~0;
 	bool bIsLightDirectional = false;
-	FVector  LightDirection;
+	FVector3f LightDirection;
 	FVector4f LightPosition;
 	FLinearColor LightLuminance;
 	float LayerDistribution;
@@ -322,20 +322,20 @@ struct FHairStrandsDebugData
 
 	struct ShadingInfo
 	{
-		FVector BaseColor;
-		float	Roughness;
-		FVector T;
-		uint32	SampleCount;
-		FVector V;
-		float	SampleOffset;
+		FVector3f	BaseColor;
+		float		Roughness;
+		FVector3f	T;
+		uint32		SampleCount;
+		FVector3f	V;
+		float		SampleOffset;
 	};
 
 	struct Sample
 	{
-		FVector Direction;
-		float	Pdf;
-		FVector Weights;
-		float	Pad;
+		FVector3f	Direction;
+		float		Pdf;
+		FVector3f	Weights;
+		float		Pad;
 	};
 	static const uint32 MaxShadingPointCount = 32;
 	static const uint32 MaxSampleCount = 1024 * 32;
diff --git a/Engine/Source/Runtime/Renderer/Private/HairStrands/HairStrandsRasterCommon.cpp b/Engine/Source/Runtime/Renderer/Private/HairStrands/HairStrandsRasterCommon.cpp
--- a/Engine/Source/Runtime/Renderer/Private/HairStrands/HairStrandsRasterCommon.cpp
+++ b/Engine/Source/Runtime/Renderer/Private/HairStrands/HairStrandsRasterCommon.cpp
@@ -261,7 +261,7 @@ void AddHairStrandsRasterPass(
 	const FIntRect& ViewportRect,
 	const FVector4f& HairRenderInfo,
 	const uint32 HairRenderInfoBits,
-	const FVector& RasterDirection,
+	const FVector3f& RasterDirection,
 	TPassParameter* PassParameters,
 	FInstanceCullingManager& InstanceCullingManager)
 {
@@ -280,7 +280,7 @@ void AddHairStrandsRasterPass(
 		ViewInfo->CachedViewUniformShaderParameters->HairRenderInfo = HairRenderInfo;
 		ViewInfo->CachedViewUniformShaderParameters->HairRenderInfoBits = HairRenderInfoBits;
 
-		const FVector SavedViewForward = ViewInfo->CachedViewUniformShaderParameters->ViewForward;
+		const FVector3f SavedViewForward = ViewInfo->CachedViewUniformShaderParameters->ViewForward;
 		ViewInfo->CachedViewUniformShaderParameters->ViewForward = RasterDirection;
 		PassParameters->View = TUniformBufferRef<FViewUniformShaderParameters>::CreateUniformBufferImmediate(*ViewInfo->CachedViewUniformShaderParameters, UniformBuffer_SingleFrame);
 		ViewInfo->CachedViewUniformShaderParameters->ViewForward = SavedViewForward;
@@ -329,7 +329,7 @@ void AddHairDeepShadowRasterPass(
 	const FIntRect& ViewportRect,
 	const FVector4f& HairRenderInfo,
 	const uint32 HairRenderInfoBits,
-	const FVector& LightDirection,
+	const FVector3f& LightDirection,
 	FHairDeepShadowRasterPassParameters* PassParameters,
 	FInstanceCullingManager& InstanceCullingManager)
 {
diff --git a/Engine/Source/Runtime/Renderer/Private/HairStrands/HairStrandsRasterCommon.h b/Engine/Source/Runtime/Renderer/Private/HairStrands/HairStrandsRasterCommon.h
--- a/Engine/Source/Runtime/Renderer/Private/HairStrands/HairStrandsRasterCommon.h
+++ b/Engine/Source/Runtime/Renderer/Private/HairStrands/HairStrandsRasterCommon.h
@@ -48,6 +48,6 @@ void AddHairDeepShadowRasterPass(
 	const FIntRect& ViewportRect,
 	const FVector4f& HairRenderInfo,
 	const uint32 HairRenderInfoBits,
-	const FVector& LightDirection,
+	const FVector3f& LightDirection,
 	FHairDeepShadowRasterPassParameters* PassParameters,
 	FInstanceCullingManager& InstanceCullingManager);
diff --git a/Engine/Source/Runtime/Renderer/Private/HairStrands/HairStrandsTransmittance.cpp b/Engine/Source/Runtime/Renderer/Private/HairStrands/HairStrandsTransmittance.cpp
--- a/Engine/Source/Runtime/Renderer/Private/HairStrands/HairStrandsTransmittance.cpp
+++ b/Engine/Source/Runtime/Renderer/Private/HairStrands/HairStrandsTransmittance.cpp
@@ -91,7 +91,7 @@ FVector4f ComputeDeepShadowLayerDepths(float LayerDistribution)
 
 struct FHairStrandsTransmittanceLightParams
 {
-	FVector LightDirection = FVector::ZeroVector;
+	FVector3f LightDirection = FVector3f::ZeroVector;
 	FVector4f LightPosition = FVector4f(0, 0, 0, 0);
 	uint32 LightChannelMask = 0;
 	float LightRadius = 0;
@@ -454,7 +454,7 @@ IMPLEMENT_GLOBAL_SHADER(FHairStrandsVoxelShadowMaskPS, "/Engine/Private/HairStra
 struct FHairStrandsVoxelShadowParams
 {
 	bool			bIsWholeSceneLight = false;
-	FVector			Voxel_LightDirection = FVector::ZeroVector;
+	FVector3f		Voxel_LightDirection = FVector3f::ZeroVector;
 	FVector4f		Voxel_LightPosition = FVector4f(0, 0, 0, 0);
 	uint32			Voxel_MacroGroupId;
 };
diff --git a/Engine/Source/Runtime/Renderer/Private/HairStrands/HairStrandsUtils.cpp b/Engine/Source/Runtime/Renderer/Private/HairStrands/HairStrandsUtils.cpp
--- a/Engine/Source/Runtime/Renderer/Private/HairStrands/HairStrandsUtils.cpp
+++ b/Engine/Source/Runtime/Renderer/Private/HairStrands/HairStrandsUtils.cpp
@@ -169,7 +169,7 @@ void ComputeWorldToLightClip(
 	WorldToClipTransform = FMatrix::Identity;
 	if (LightType == LightType_Directional)
 	{
-		const FVector LightDirection = LightProxy.GetDirection();
+		const FVector3f LightDirection = LightProxy.GetDirection();
 		FReversedZOrthoMatrix OrthoMatrix(SphereRadius, SphereRadius, 1.f / (2 * SphereRadius), 0);
 		FLookAtMatrix LookAt(SphereBound.Center - LightDirection * SphereRadius, SphereBound.Center, FVector(0, 0, 1));
 		WorldToClipTransform = LookAt * OrthoMatrix;
@@ -181,7 +181,6 @@ void ComputeWorldToLightClip(
 	}
 	else if (LightType == LightType_Spot || LightType == LightType_Point)
 	{
-		const FVector LightDirection = LightPosition - PrimitivesBounds.GetSphere().Center;
 		const float SphereDistance = FVector::Distance(LightPosition, SphereBound.Center);
 		float HalfFov = asin(SphereRadius / SphereDistance);
 		HalfFov = FMath::Min(HalfFov, MaxDeepShadowFrustumHalfAngleInRad);
@@ -193,7 +192,6 @@ void ComputeWorldToLightClip(
 	}
 	else if (LightType == LightType_Rect)
 	{
-		const FVector LightDirection = LightPosition - PrimitivesBounds.GetSphere().Center;
 		const float SphereDistance = FVector::Distance(LightPosition, SphereBound.Center);
 		float HalfFov = asin(SphereRadius / SphereDistance);
 		HalfFov = FMath::Min(HalfFov, MaxDeepShadowFrustumHalfAngleInRad);
