Fix vis. cluster debug visualization by requiring simulation to be enabled.

diff --git a/Engine/Plugins/Runtime/HairStrands/Source/HairStrandsEditor/Private/GroomEditorCommands.cpp b/Engine/Plugins/Runtime/HairStrands/Source/HairStrandsEditor/Private/GroomEditorCommands.cpp
--- a/Engine/Plugins/Runtime/HairStrands/Source/HairStrandsEditor/Private/GroomEditorCommands.cpp
+++ b/Engine/Plugins/Runtime/HairStrands/Source/HairStrandsEditor/Private/GroomEditorCommands.cpp
@@ -37,7 +37,7 @@ void FGroomEditorCommands::RegisterCommands()
 	UI_COMMAND(ViewMode_RadiusVariation,"RadiusVariation",	"View strands radius variation",EUserInterfaceActionType::Button, FInputChord());
 	UI_COMMAND(ViewMode_BaseColor,		"BaseColor",		"View vertex base color",		EUserInterfaceActionType::Button, FInputChord());
 	UI_COMMAND(ViewMode_Roughness,		"Roughness",		"View vertex roughness",		EUserInterfaceActionType::Button, FInputChord());
-	UI_COMMAND(ViewMode_VisCluster,		"VisCluster",		"View clusters",				EUserInterfaceActionType::Button, FInputChord());
+	UI_COMMAND(ViewMode_VisCluster,		"Vis. Clusters",	"View clusters (only available if the groom has simulation enabled)",				EUserInterfaceActionType::Button, FInputChord());
 	UI_COMMAND(ViewMode_Group,			"Hair Groups",		"View hair groups",				EUserInterfaceActionType::Button, FInputChord());
 
 	UI_COMMAND(ViewMode_CardsGuides,	"CardsGuide",		"Cards Guides (only available if the groom has cards geometry)",					EUserInterfaceActionType::Button, FInputChord());
diff --git a/Engine/Plugins/Runtime/HairStrands/Source/HairStrandsEditor/Private/Widgets/SGroomEditorViewport.cpp b/Engine/Plugins/Runtime/HairStrands/Source/HairStrandsEditor/Private/Widgets/SGroomEditorViewport.cpp
--- a/Engine/Plugins/Runtime/HairStrands/Source/HairStrandsEditor/Private/Widgets/SGroomEditorViewport.cpp
+++ b/Engine/Plugins/Runtime/HairStrands/Source/HairStrandsEditor/Private/Widgets/SGroomEditorViewport.cpp
@@ -199,7 +199,7 @@ void SGroomEditorViewport::BindCommands()
 	CommandList->MapAction(FGroomEditorCommands::Get().ViewMode_RadiusVariation,	FExecuteAction::CreateSP(this, &SGroomEditorViewport::OnViewMode, EHairStrandsDebugMode::RenderHairRadiusVariation),FCanExecuteAction::CreateSP(this, &SGroomEditorViewport::CanViewMode, false));
 	CommandList->MapAction(FGroomEditorCommands::Get().ViewMode_BaseColor,			FExecuteAction::CreateSP(this, &SGroomEditorViewport::OnViewMode, EHairStrandsDebugMode::RenderHairBaseColor),		FCanExecuteAction::CreateSP(this, &SGroomEditorViewport::CanViewMode, false));
 	CommandList->MapAction(FGroomEditorCommands::Get().ViewMode_Roughness,			FExecuteAction::CreateSP(this, &SGroomEditorViewport::OnViewMode, EHairStrandsDebugMode::RenderHairRoughness),		FCanExecuteAction::CreateSP(this, &SGroomEditorViewport::CanViewMode, false));
-	CommandList->MapAction(FGroomEditorCommands::Get().ViewMode_VisCluster,			FExecuteAction::CreateSP(this, &SGroomEditorViewport::OnViewMode, EHairStrandsDebugMode::RenderVisCluster),			FCanExecuteAction::CreateSP(this, &SGroomEditorViewport::CanViewMode, false));
+	CommandList->MapAction(FGroomEditorCommands::Get().ViewMode_VisCluster,			FExecuteAction::CreateSP(this, &SGroomEditorViewport::OnViewMode, EHairStrandsDebugMode::RenderVisCluster),			FCanExecuteAction::CreateSP(this, &SGroomEditorViewport::CanViewMode, true)); // While Clusters LOD/Culling doesn't require simulation, cluster visualization requires HairInterpolation to run, which is only available when simulation is enabled in the Groom editor.
 	CommandList->MapAction(FGroomEditorCommands::Get().ViewMode_Group,				FExecuteAction::CreateSP(this, &SGroomEditorViewport::OnViewMode, EHairStrandsDebugMode::RenderHairGroup),			FCanExecuteAction::CreateSP(this, &SGroomEditorViewport::CanViewMode, false));
 	CommandList->MapAction(FGroomEditorCommands::Get().ViewMode_CardsGuides,		FExecuteAction::CreateSP(this, &SGroomEditorViewport::OnCardsGuides), FCanExecuteAction::CreateSP(this, &SGroomEditorViewport::CanCardsGuides));
 }
@@ -220,9 +220,9 @@ void SGroomEditorViewport::OnViewMode(EHairStrandsDebugMode Mode)
 	}
 }
 
-bool SGroomEditorViewport::CanViewMode(bool bGuideMode) const
+bool SGroomEditorViewport::CanViewMode(bool bRequireSimulation) const
 {
-	if (bGuideMode)
+	if (bRequireSimulation)
 	{
 		return GroomComponent && GroomComponent->GroomAsset != nullptr && GroomComponent->GroomAsset->IsSimulationEnable();
 	}
