Fix missing physics asset for hair sim

diff --git a/Engine/Plugins/Runtime/HairStrands/Source/HairStrandsCore/Private/Niagara/NiagaraDataInterfacePhysicsAsset.cpp b/Engine/Plugins/Runtime/HairStrands/Source/HairStrandsCore/Private/Niagara/NiagaraDataInterfacePhysicsAsset.cpp
--- a/Engine/Plugins/Runtime/HairStrands/Source/HairStrandsCore/Private/Niagara/NiagaraDataInterfacePhysicsAsset.cpp
+++ b/Engine/Plugins/Runtime/HairStrands/Source/HairStrandsCore/Private/Niagara/NiagaraDataInterfacePhysicsAsset.cpp
@@ -727,8 +727,11 @@ void UNiagaraDataInterfacePhysicsAsset::ExtractSourceComponent(FNiagaraSystemIns
 
 			if (USkeletalMeshComponent* ParentComp = Cast<USkeletalMeshComponent>(SourceComponent->GetAttachParent()))
 			{
+				SourceComponents.Add(ParentComp);
+				PhysicsAssets.Add(ParentComp->GetPhysicsAsset());
+
 				TArray<USceneComponent*> SceneComponents;
-				ParentComp->GetChildrenComponents(false, SceneComponents);
+				ParentComp->GetChildrenComponents(true, SceneComponents);
 
 				for (USceneComponent* ActorComp : SceneComponents)
 				{
