Fix potential issue when a groom cards resources are not ready.

diff --git a/Engine/Plugins/Runtime/HairStrands/Source/HairStrandsCore/Private/GroomComponent.cpp b/Engine/Plugins/Runtime/HairStrands/Source/HairStrandsCore/Private/GroomComponent.cpp
--- a/Engine/Plugins/Runtime/HairStrands/Source/HairStrandsCore/Private/GroomComponent.cpp
+++ b/Engine/Plugins/Runtime/HairStrands/Source/HairStrandsCore/Private/GroomComponent.cpp
@@ -792,6 +792,12 @@ public:
 			return nullptr;
 		}
 
+		// Invalid primitive setup. This can happens when the (procedural) resources are not ready.
+		if (NumPrimitive == 0 && !bUseCulling)
+		{
+			return nullptr;
+		}
+
 		if (bWireframe)
 		{
 			MaterialRenderProxy = new FColoredMaterialRenderProxy( GEngine->WireframeMaterial ? GEngine->WireframeMaterial->GetRenderProxy() : NULL, FLinearColor(1.f, 0.5f, 0.f));
