Fix crash when adding new cards/meshes LODs.

diff --git a/Engine/Plugins/Runtime/HairStrands/Source/HairStrandsEditor/Private/GroomAssetDetails.cpp b/Engine/Plugins/Runtime/HairStrands/Source/HairStrandsEditor/Private/GroomAssetDetails.cpp
--- a/Engine/Plugins/Runtime/HairStrands/Source/HairStrandsEditor/Private/GroomAssetDetails.cpp
+++ b/Engine/Plugins/Runtime/HairStrands/Source/HairStrandsEditor/Private/GroomAssetDetails.cpp
@@ -739,14 +739,21 @@ bool AssignIfDifferent(T& Dest, const T& Src, bool bSetValue)
 bool FGroomRenderingDetails::CommonResetToDefault(TSharedPtr<IPropertyHandle> ChildHandle, int32 GroupIndex, int32 LODIndex, bool bSetValue)
 {
 	bool bHasChanged = false;
-	if (ChildHandle == nullptr || GroomAsset == nullptr || GroupIndex < 0 || GroupIndex >= GroomAsset->GetNumHairGroups())
+	if (ChildHandle == nullptr || GroomAsset == nullptr || GroupIndex < 0)
 	{
 		return bHasChanged;
 	}
 
 	FName PropertyName = ChildHandle->GetProperty()->GetFName();
 
+	// For cards & meshes the incoming index is actually the cards/mesh description index, not the group index
+	// For the rest, the group index refers to the actual group index.
+	const bool bIsCardDescIndexValid = GroupIndex < GroomAsset->HairGroupsCards.Num();
+	const bool bIsMeshDescIndexValid = GroupIndex < GroomAsset->HairGroupsMeshes.Num();
+	const bool bIsGroupIndexValid = GroupIndex < GroomAsset->GetNumHairGroups();
+		
 	// Hair strands
+	if (bIsGroupIndexValid)
 	{
 		{
 			FHairGeometrySettings Default;
@@ -771,6 +778,7 @@ bool FGroomRenderingDetails::CommonResetToDefault(TSharedPtr<IPropertyHandle> Ch
 	}
 
 	// Interpolation
+	if (bIsGroupIndexValid)
 	{
 		{
 			FHairDecimationSettings Default;
@@ -788,7 +796,7 @@ bool FGroomRenderingDetails::CommonResetToDefault(TSharedPtr<IPropertyHandle> Ch
 	}
 
 	// LODs
-	if (LODIndex>=0)
+	if (bIsGroupIndexValid && LODIndex >= 0)
 	{
 		{
 			FHairGroupsLOD Default;
@@ -808,6 +816,7 @@ bool FGroomRenderingDetails::CommonResetToDefault(TSharedPtr<IPropertyHandle> Ch
 	}
 
 	// Cards
+	if (bIsCardDescIndexValid)
 	{
 		{
 			FHairGroupsCardsSourceDescription Default;
@@ -829,6 +838,7 @@ bool FGroomRenderingDetails::CommonResetToDefault(TSharedPtr<IPropertyHandle> Ch
 	}
 
 	// Meshes
+	if (bIsMeshDescIndexValid)
 	{
 		{
 			FHairGroupsMeshesSourceDescription Default;
@@ -848,6 +858,7 @@ bool FGroomRenderingDetails::CommonResetToDefault(TSharedPtr<IPropertyHandle> Ch
 	}
 
 	// Physics
+	if (bIsGroupIndexValid)
 	{
 		{
 			FHairSolverSettings Default;
