Fix crash when deleting groom asset while still being used.

diff --git a/Engine/Plugins/Runtime/HairStrands/Source/HairStrandsCore/Private/GroomAsset.cpp b/Engine/Plugins/Runtime/HairStrands/Source/HairStrandsCore/Private/GroomAsset.cpp
--- a/Engine/Plugins/Runtime/HairStrands/Source/HairStrandsCore/Private/GroomAsset.cpp
+++ b/Engine/Plugins/Runtime/HairStrands/Source/HairStrandsCore/Private/GroomAsset.cpp
@@ -2477,7 +2477,10 @@ static void InitCardsTextureResources(UGroomAsset* GroomAsset)
 					InitAtlasTexture(LOD.RestResource, Desc->Textures.AttributeTexture, EHairAtlasTextureType::Attribute);
 					InitAtlasTexture(LOD.RestResource, Desc->Textures.CoverageTexture, EHairAtlasTextureType::Coverage);
 					InitAtlasTexture(LOD.RestResource, Desc->Textures.AuxilaryDataTexture, EHairAtlasTextureType::AuxilaryData);
-					LOD.RestResource->bInvertUV = Desc->SourceType == EHairCardsSourceType::Procedural; // Should fix procedural texture so that this does not happen
+					if (LOD.RestResource)
+					{
+						LOD.RestResource->bInvertUV = Desc->SourceType == EHairCardsSourceType::Procedural; // Should fix procedural texture so that this does not happen
+					}
 				}
 			}
 		}
