Added icon for GroomCache asset.

diff --git a/Engine/Plugins/Runtime/HairStrands/Source/HairStrandsEditor/Private/HairStrandsEditor.cpp b/Engine/Plugins/Runtime/HairStrands/Source/HairStrandsEditor/Private/HairStrandsEditor.cpp
--- a/Engine/Plugins/Runtime/HairStrands/Source/HairStrandsEditor/Private/HairStrandsEditor.cpp
+++ b/Engine/Plugins/Runtime/HairStrands/Source/HairStrandsEditor/Private/HairStrandsEditor.cpp
@@ -106,6 +106,9 @@ void FGroomEditor::StartupModule()
 		StyleSet->Set("GroomEditor.SimulationOptions", new FSlateImageBrush(HairStrandsContent + "/Icons/S_SimulationOptions_40x.png", Icon40x40));
 		StyleSet->Set("GroomEditor.SimulationOptions.Small", new FSlateImageBrush(HairStrandsContent + "/Icons/S_SimulationOptions_40x.png", Icon20x20));
 
+		StyleSet->Set("ClassIcon.GroomCache", new FSlateImageBrush(HairStrandsContent + "/Icons/S_GroomCache_64.png", Icon16x16));
+		StyleSet->Set("ClassThumbnail.GroomCache", new FSlateImageBrush(HairStrandsContent + "/Icons/S_GroomCache_64.png", Icon64x64));
+
 		FSlateStyleRegistry::RegisterSlateStyle(*StyleSet.Get());
 
 		// Custom widget for groom component (Group desc override, ...)
