Fix invalid UAV float clear on UNORM/SNORM buffer.

diff --git a/Engine/Plugins/Runtime/HairStrands/Source/HairStrandsCore/Private/GroomResources.cpp b/Engine/Plugins/Runtime/HairStrands/Source/HairStrandsCore/Private/GroomResources.cpp
--- a/Engine/Plugins/Runtime/HairStrands/Source/HairStrandsCore/Private/GroomResources.cpp
+++ b/Engine/Plugins/Runtime/HairStrands/Source/HairStrandsCore/Private/GroomResources.cpp
@@ -206,11 +206,7 @@ void InternalCreateVertexBufferRDG(FRDGBuilder& GraphBuilder, uint32 InVertexCou
 		case PF_G32R32F:
 		case PF_R16F_FILTER:
 		case PF_R16F:
-		case PF_R16G16B16A16_SNORM:
-		case PF_R16G16B16A16_UNORM:
 		case PF_R32_FLOAT:
-		case PF_R5G6B5_UNORM:
-		case PF_R8G8B8A8_SNORM:
 			return true;
 		default:
 			return false;
