Added support for GroomCache frame interpolation for subframe sampling.

diff --git a/Engine/Plugins/Runtime/HairStrands/Source/HairStrandsCore/Private/GroomCache.cpp b/Engine/Plugins/Runtime/HairStrands/Source/HairStrandsCore/Private/GroomCache.cpp
--- a/Engine/Plugins/Runtime/HairStrands/Source/HairStrandsCore/Private/GroomCache.cpp
+++ b/Engine/Plugins/Runtime/HairStrands/Source/HairStrandsCore/Private/GroomCache.cpp
@@ -41,31 +41,92 @@ float UGroomCache::GetDuration() const
 	return GroomCacheInfo.AnimationInfo.Duration;
 }
 
-int32 UGroomCache::GetFrameNumberAtTime(const float Time) const
+int32 UGroomCache::GetFrameNumberAtTime(const float Time, bool bLooping) const
 {
-	return GetStartFrame() + GetFrameIndexAtTime(Time); 
+	return GetStartFrame() + GetFrameIndexAtTime(Time, bLooping); 
 }
 
-int32 UGroomCache::GetFrameIndexAtTime(const float Time) const
+int32 UGroomCache::GetFrameIndexAtTime(const float Time, bool bLooping) const
 {
 	const float FrameTime = GroomCacheInfo.AnimationInfo.SecondsPerFrame;
 	if (FrameTime == 0.0f)
 	{
 		return 0;
 	}
-	const int32 NumberOfFrames = GroomCacheInfo.AnimationInfo.NumFrames;
-	const int32 NormalizedFrame = FMath::Clamp(FMath::RoundToInt(Time / FrameTime), 0, NumberOfFrames - 1);
+
+	const int32 NumFrames = GroomCacheInfo.AnimationInfo.NumFrames;
+	const float Duration = GroomCacheInfo.AnimationInfo.Duration;
+	float AdjustedTime = Time;
+	if (bLooping)
+	{
+		AdjustedTime = Time - Duration * FMath::FloorToFloat(Time / Duration);
+	}
+	else
+	{
+		AdjustedTime = FMath::Clamp(Time, 0.0f, GroomCacheInfo.AnimationInfo.EndTime - GroomCacheInfo.AnimationInfo.StartTime);
+	}
+
+	const int32 NormalizedFrame = FMath::Clamp(FMath::FloorToInt(AdjustedTime / FrameTime), 0, NumFrames - 1);
 	return NormalizedFrame; 
 }
 
-bool UGroomCache::GetGroomDataAtTime(float Time, FGroomCacheAnimationData& AnimData)
+void UGroomCache::FindSampleIndexesFromTime(float Time, bool bLooping, bool bIsPlayingBackwards, int32 &OutFrameIndex, int32 &OutNextFrameIndex, float &InterpolationFactor)
 {
-	const int32 FrameIndex = GetFrameIndexAtTime(Time);
+	const int32 NumFrames = GroomCacheInfo.AnimationInfo.NumFrames;
+	const float Duration = GroomCacheInfo.AnimationInfo.Duration;
+
+	// No index possible
+	if (NumFrames == 0 || NumFrames == 1 || Duration == 0.0f)
+	{
+		OutFrameIndex = 0;
+		OutNextFrameIndex = 0;
+		InterpolationFactor = 0.0f;
+		return;
+	}
+
+	OutFrameIndex = GetFrameIndexAtTime(Time, bLooping);
+	OutNextFrameIndex = FMath::Min(OutFrameIndex + 1, NumFrames - 1);
+
+	const float FrameDuration = GroomCacheInfo.AnimationInfo.SecondsPerFrame;
+	if (FMath::IsNearlyZero(FrameDuration))
+	{
+		InterpolationFactor = 0.0f;
+	}
+	else
+	{
+		float AdjustedTime;
+
+		if (bLooping)
+		{
+			AdjustedTime = Time - Duration * FMath::FloorToFloat(Time / Duration);
+		}
+		else
+		{
+			AdjustedTime = FMath::Clamp(Time, 0.0f, GroomCacheInfo.AnimationInfo.EndTime - GroomCacheInfo.AnimationInfo.StartTime);
+		}
+
+		float Delta = AdjustedTime - FrameDuration * OutFrameIndex;
+		InterpolationFactor = Delta / FrameDuration;
+	}
+
+	// If playing backwards the logical order of previous and next is reversed
+	if (bIsPlayingBackwards)
+	{
+		Swap(OutFrameIndex, OutNextFrameIndex);
+		InterpolationFactor = 1.0f - InterpolationFactor;
+	}
+}
+
+bool UGroomCache::GetGroomDataAtTime(float Time, bool bLooping, FGroomCacheAnimationData& AnimData)
+{
+	const int32 FrameIndex = GetFrameIndexAtTime(Time, bLooping);
 	return GetGroomDataAtFrameIndex(FrameIndex, AnimData);
 }
 
 bool UGroomCache::GetGroomDataAtFrameIndex(int32 FrameIndex, FGroomCacheAnimationData& AnimData)
 {
+	TRACE_CPUPROFILER_EVENT_SCOPE(UGroomCache::GetGroomDataAtFrameIndex);
+
 	if (!Chunks.IsValidIndex(FrameIndex))
 	{
 		return false;
@@ -79,11 +140,17 @@ bool UGroomCache::GetGroomDataAtFrameIndex(int32 FrameIndex, FGroomCacheAnimatio
 
 	// This is where the bulk data is loaded from disk
 	void* Buffer = TempBytes.GetData();
-	Chunk.BulkData.GetCopy(&Buffer, true);
+	{
+		TRACE_CPUPROFILER_EVENT_SCOPE(UGroomCache::GetGroomDataAtFrameIndex_BulkData);
+		Chunk.BulkData.GetCopy(&Buffer, true);
+	}
 
 	// The bulk data buffer is then serialized into GroomCacheAnimationData
-	FMemoryReader Ar(TempBytes, true);
-	AnimData.Serialize(Ar);
+	{
+		TRACE_CPUPROFILER_EVENT_SCOPE(UGroomCache::GetGroomDataAtFrameIndex_Serialize);
+		FMemoryReader Ar(TempBytes, true);
+		AnimData.Serialize(Ar);
+	}
 
 	return true;
 }
@@ -111,7 +178,7 @@ void FGroomCacheChunk::Serialize(FArchive& Ar, UObject* Owner, int32 ChunkIndex)
 
 	// Forced not inline means the bulk data won't automatically be loaded when we deserialize
 	// but only when we explicitly take action to load it
-	BulkData.SetBulkDataFlags(BULKDATA_Force_NOT_InlinePayload | BULKDATA_SerializeCompressed);
+	BulkData.SetBulkDataFlags(BULKDATA_Force_NOT_InlinePayload);
 	BulkData.Serialize(Ar, Owner, ChunkIndex, false);
 }
 
diff --git a/Engine/Plugins/Runtime/HairStrands/Source/HairStrandsCore/Private/GroomComponent.cpp b/Engine/Plugins/Runtime/HairStrands/Source/HairStrandsCore/Private/GroomComponent.cpp
--- a/Engine/Plugins/Runtime/HairStrands/Source/HairStrandsCore/Private/GroomComponent.cpp
+++ b/Engine/Plugins/Runtime/HairStrands/Source/HairStrandsCore/Private/GroomComponent.cpp
@@ -29,6 +29,7 @@
 #include "GroomManager.h"
 #include "GroomInstance.h"
 #include "GroomCache.h"
+#include "Async/ParallelFor.h"
 
 static int32 GHairEnableAdaptiveSubsteps = 0;  
 static FAutoConsoleVariableRef CVarHairEnableAdaptiveSubsteps(TEXT("r.HairStrands.EnableAdaptiveSubsteps"), GHairEnableAdaptiveSubsteps, TEXT("Enable adaptive solver substeps"));
@@ -960,13 +961,87 @@ public:
 		InterpolationFactor = Factor;
 	}
 
+	void UpdateBuffersAtTime(UGroomCache* GroomCache, float Time, bool bIsLooping)
+	{
+		TRACE_CPUPROFILER_EVENT_SCOPE(FGroomCacheBuffers::UpdateBuffersAtTime);
+
+		if (!GroomCache)
+		{
+			return;
+		}
+
+		// Find the frame indices and interpolation factor to interpolate between
+		int32 FrameIndexA = 0;
+		int32 FrameIndexB = 0;
+		float OutInterpolationFactor = 0.0f;
+		GroomCache->FindSampleIndexesFromTime(Time, bIsLooping, false, FrameIndexA, FrameIndexB, OutInterpolationFactor);
+
+		// Update and cache the frame data as needed
+		bool bComputeInterpolation = false;
+		if (FrameIndexA != CurrentFrameIndex)
+		{
+			bComputeInterpolation = true;
+			if (FrameIndexA == NextFrameIndex)
+			{
+				Swap(CurrentFrame, NextFrame);
+				CurrentFrameIndex = NextFrameIndex;
+			}
+			else
+			{
+				GroomCache->GetGroomDataAtFrameIndex(FrameIndexA, CurrentFrame);
+				CurrentFrameIndex = FrameIndexA;
+			}
+		}
+
+		if (FrameIndexB != NextFrameIndex)
+		{
+			bComputeInterpolation = true;
+			GroomCache->GetGroomDataAtFrameIndex(FrameIndexB, NextFrame);
+			NextFrameIndex = FrameIndexB;
+		}
+
+		// Make sure the initial interpolated frame is populated with valid data
+		if (InterpolatedFrame.GroupsData.Num() != CurrentFrame.GroupsData.Num())
+		{
+			InterpolatedFrame = CurrentFrame;
+		}
+
+		// Do interpolation of vertex positions if needed
+		if (bComputeInterpolation || !FMath::IsNearlyEqual(OutInterpolationFactor, InterpolationFactor, KINDA_SMALL_NUMBER))
+		{
+			TRACE_CPUPROFILER_EVENT_SCOPE(FGroomCacheBuffers::UpdateBuffersAtTime_InterpolateCPU);
+
+			InterpolationFactor = OutInterpolationFactor;
+
+			FScopeLock Lock(GetCriticalSection());
+
+			const int32 NumGroups = CurrentFrame.GroupsData.Num();
+			for (int32 GroupIndex = 0; GroupIndex < NumGroups; ++GroupIndex)
+			{
+				const FGroomCacheGroupData& CurrentGroupData = CurrentFrame.GroupsData[GroupIndex];
+				const FGroomCacheGroupData& NextGroupData = NextFrame.GroupsData[GroupIndex];
+				FGroomCacheGroupData& InterpolatedGroupData = InterpolatedFrame.GroupsData[GroupIndex];
+				const int32 NumVertices = CurrentGroupData.VertexData.PointsPosition.Num();
+
+				ParallelFor(NumVertices,
+				[&CurrentGroupData, &NextGroupData, &InterpolatedGroupData, this](int32 VertexIndex)
+				{
+					const FVector& CurrentPosition = CurrentGroupData.VertexData.PointsPosition[VertexIndex];
+					const FVector& NextPosition = NextGroupData.VertexData.PointsPosition[VertexIndex];
+					InterpolatedGroupData.VertexData.PointsPosition[VertexIndex] = FMath::Lerp(CurrentPosition, NextPosition, InterpolationFactor);
+				}
+				);
+			}
+		}
+	}
+
 private:
 	FGroomCacheAnimationData CurrentFrame;
 	FGroomCacheAnimationData NextFrame;
 	FGroomCacheAnimationData InterpolatedFrame;
 
-	int32 CurrentFrameIndex = 0;
-	int32 NextFrameIndex = 0;
+	int32 CurrentFrameIndex = -1;
+	int32 NextFrameIndex = -1;
 	float InterpolationFactor = 0.0f;
 };
 
@@ -2609,14 +2684,8 @@ void UGroomComponent::UpdateGroomCache(float Time)
 {
 	if (GroomCache && GroomCacheBuffers.IsValid() && bRunning)
 	{
-		float AnimationTime = Time;
-		const float Duration = GroomCache->GetDuration();
-		if (bLooping && Duration > 0.0f)
-		{
-			AnimationTime = AnimationTime - Duration * FMath::FloorToFloat(AnimationTime / Duration);
-		}
-		FScopeLock Lock(GroomCacheBuffers->GetCriticalSection());
-		GroomCache->GetGroomDataAtTime(AnimationTime, GroomCacheBuffers->GetInterpolatedFrameBuffer());
+		FGroomCacheBuffers* Buffers = static_cast<FGroomCacheBuffers*>(GroomCacheBuffers.Get());
+		Buffers->UpdateBuffersAtTime(GroomCache, Time, bLooping);
 	}
 }
 
diff --git a/Engine/Plugins/Runtime/HairStrands/Source/HairStrandsCore/Public/GroomCache.h b/Engine/Plugins/Runtime/HairStrands/Source/HairStrandsCore/Public/GroomCache.h
--- a/Engine/Plugins/Runtime/HairStrands/Source/HairStrandsCore/Public/GroomCache.h
+++ b/Engine/Plugins/Runtime/HairStrands/Source/HairStrandsCore/Public/GroomCache.h
@@ -29,12 +29,15 @@ public:
 	float GetDuration() const;
 
 	/** Get the frame number at the specified time within the animation range which might not start at 0 */
-	int32  GetFrameNumberAtTime(const float Time) const;
+	int32  GetFrameNumberAtTime(const float Time, bool bLooping) const;
 
-	/** Get the frame index at the specified time with the index 0 being the start of the animation */
-	int32  GetFrameIndexAtTime(const float Time) const;
+	/** Get the (floored) frame index at the specified time with the index 0 being the start of the animation */
+	int32  GetFrameIndexAtTime(const float Time, bool bLooping) const;
 
-	bool GetGroomDataAtTime(float Time, FGroomCacheAnimationData& AnimData);
+	/** Get the frame indices and interpolation factor between them that correspond to the specified time */
+	void FindSampleIndexesFromTime(float Time, bool bLooping, bool bIsPlayingBackwards, int32 &OutFrameIndex, int32 &OutNextFrameIndex, float &InterpolationFactor);
+
+	bool GetGroomDataAtTime(float Time, bool bLooping, FGroomCacheAnimationData& AnimData);
 	bool GetGroomDataAtFrameIndex(int32 FrameIndex, FGroomCacheAnimationData& AnimData);
 
 	void SetGroomAnimationInfo(const FGroomAnimationInfo& AnimInfo);
diff --git a/Engine/Plugins/Runtime/HairStrands/Source/HairStrandsEditor/Private/GroomCacheTrackEditor.cpp b/Engine/Plugins/Runtime/HairStrands/Source/HairStrandsEditor/Private/GroomCacheTrackEditor.cpp
--- a/Engine/Plugins/Runtime/HairStrands/Source/HairStrandsEditor/Private/GroomCacheTrackEditor.cpp
+++ b/Engine/Plugins/Runtime/HairStrands/Source/HairStrandsEditor/Private/GroomCacheTrackEditor.cpp
@@ -134,7 +134,7 @@ int32 FGroomCacheSection::OnPaintSection(FSequencerSectionPainter& Painter) cons
 
 			// Draw the current time next to the scrub handle
 			const float AnimTime = Section.MapTimeToAnimation(Duration, CurrentTime, TickResolution);
-			int32 FrameTime = GroomCache->GetFrameNumberAtTime(AnimTime);
+			int32 FrameTime = GroomCache->GetFrameNumberAtTime(AnimTime, false);
 			FString FrameString = FString::FromInt(FrameTime);
 
 			const FSlateFontInfo SmallLayoutFont = FCoreStyle::GetDefaultFontStyle("Bold", 10);
