Add SM6 support for Hair Strands

diff --git a/Engine/Plugins/Runtime/HairStrands/Source/HairStrandsCore/Private/Niagara/NiagaraDataInterfaceVelocityGrid.cpp b/Engine/Plugins/Runtime/HairStrands/Source/HairStrandsCore/Private/Niagara/NiagaraDataInterfaceVelocityGrid.cpp
--- a/Engine/Plugins/Runtime/HairStrands/Source/HairStrandsCore/Private/Niagara/NiagaraDataInterfaceVelocityGrid.cpp
+++ b/Engine/Plugins/Runtime/HairStrands/Source/HairStrandsCore/Private/Niagara/NiagaraDataInterfaceVelocityGrid.cpp
@@ -697,7 +697,7 @@ inline void ClearTexture(FRHICommandList& RHICmdList, FNDIVelocityGridBuffer* De
 
 	if (DestinationGridBufferUAV != nullptr)
 	{
-		TShaderMapRef<FClearVelocityGridCS> ComputeShader(GetGlobalShaderMap(ERHIFeatureLevel::SM5));
+		TShaderMapRef<FClearVelocityGridCS> ComputeShader(GetGlobalShaderMap(GMaxRHIFeatureLevel));
 		RHICmdList.SetComputeShader(ComputeShader.GetComputeShader());
 
 		FRHITransitionInfo Transitions[] = {
diff --git a/Engine/Source/Runtime/Engine/Private/Materials/MaterialInterface.cpp b/Engine/Source/Runtime/Engine/Private/Materials/MaterialInterface.cpp
--- a/Engine/Source/Runtime/Engine/Private/Materials/MaterialInterface.cpp
+++ b/Engine/Source/Runtime/Engine/Private/Materials/MaterialInterface.cpp
@@ -50,7 +50,7 @@ bool IsHairStrandsGeometrySupported(const EShaderPlatform Platform)
 bool IsCompatibleWithHairStrands(const FMaterial* Material, const ERHIFeatureLevel::Type FeatureLevel)
 {
 	return
-		ERHIFeatureLevel::SM5 == FeatureLevel &&
+		ERHIFeatureLevel::SM5 <= FeatureLevel &&
 		Material && Material->IsUsedWithHairStrands() && 
 		(Material->GetBlendMode() == BLEND_Opaque || Material->GetBlendMode() == BLEND_Masked);
 }
