Fix invalid groom bound where simulation is enabled but not bound to a skel. mesh.

diff --git a/Engine/Plugins/Runtime/HairStrands/Source/HairStrandsCore/Private/GroomComponent.cpp b/Engine/Plugins/Runtime/HairStrands/Source/HairStrandsCore/Private/GroomComponent.cpp
--- a/Engine/Plugins/Runtime/HairStrands/Source/HairStrandsCore/Private/GroomComponent.cpp
+++ b/Engine/Plugins/Runtime/HairStrands/Source/HairStrandsCore/Private/GroomComponent.cpp
@@ -1361,6 +1361,7 @@ void UGroomComponent::CreateHairSimulation(const int32 GroupIndex, const int32 L
 			{
 				NiagaraComponent = NewObject<UNiagaraComponent>(this, NAME_None, RF_Transient);
 				NiagaraComponent->SetVisibleFlag(SimulationSettings.SimulationSetup.bDebugSimulation);
+				NiagaraComponent->bUseAttachParentBound = true;
 			}
 			if (GetWorld() && GetWorld()->bIsWorldInitialized)
 			{
