Fix crash when setting curve decimation above one.

This CL also fixes a rounding issue when curve decimation factor is set to one, which result in invalid decimation.

diff --git a/Engine/Plugins/Runtime/HairStrands/Source/HairStrandsCore/Private/GroomBuilder.cpp b/Engine/Plugins/Runtime/HairStrands/Source/HairStrandsCore/Private/GroomBuilder.cpp
--- a/Engine/Plugins/Runtime/HairStrands/Source/HairStrandsCore/Private/GroomBuilder.cpp
+++ b/Engine/Plugins/Runtime/HairStrands/Source/HairStrandsCore/Private/GroomBuilder.cpp
@@ -1904,7 +1904,7 @@ void Decimate(
 		const float AdjustedBucketSize = MaxBucket - MinBucket;
 		if (AdjustedBucketSize > 0)
 		{
-			const uint32 CurveIndex = FMath::FloorToInt(MinBucket + Random.FRand() * AdjustedBucketSize);
+			const uint32 CurveIndex = Random.RandRange(MinBucket, FMath::FloorToInt(MinBucket + AdjustedBucketSize)-1);
 			CurveIndices[BucketIndex] = CurveIndex;
 			LastCurveIndex = CurveIndex;
 
diff --git a/Engine/Plugins/Runtime/HairStrands/Source/HairStrandsCore/Public/GroomAssetInterpolation.h b/Engine/Plugins/Runtime/HairStrands/Source/HairStrandsCore/Public/GroomAssetInterpolation.h
--- a/Engine/Plugins/Runtime/HairStrands/Source/HairStrandsCore/Public/GroomAssetInterpolation.h
+++ b/Engine/Plugins/Runtime/HairStrands/Source/HairStrandsCore/Public/GroomAssetInterpolation.h
@@ -124,7 +124,7 @@ struct HAIRSTRANDSCORE_API FHairInterpolationSettings
 	bool bOverrideGuides;
 
 	/** Density factor for converting hair into guide curve if no guides are provided. */
-	UPROPERTY(EditAnywhere, BlueprintReadWrite, Category = "InterpolationSettings", meta = (ClampMin = "0", UIMin = "0", UIMax = "1.0"))
+	UPROPERTY(EditAnywhere, BlueprintReadWrite, Category = "InterpolationSettings", meta = (ClampMin = "0", ClampMax = "1.0", UIMin = "0", UIMax = "1.0"))
 	float HairToGuideDensity;
 
 	/** Interpolation data quality. */
