Fix cards guides visualization only available when simulation is enabled.

diff --git a/Engine/Plugins/Runtime/HairStrands/Source/HairStrandsEditor/Private/Widgets/SGroomEditorViewport.cpp b/Engine/Plugins/Runtime/HairStrands/Source/HairStrandsEditor/Private/Widgets/SGroomEditorViewport.cpp
--- a/Engine/Plugins/Runtime/HairStrands/Source/HairStrandsEditor/Private/Widgets/SGroomEditorViewport.cpp
+++ b/Engine/Plugins/Runtime/HairStrands/Source/HairStrandsEditor/Private/Widgets/SGroomEditorViewport.cpp
@@ -201,7 +201,7 @@ void SGroomEditorViewport::BindCommands()
 	CommandList->MapAction(FGroomEditorCommands::Get().ViewMode_Roughness,			FExecuteAction::CreateSP(this, &SGroomEditorViewport::OnViewMode, EHairStrandsDebugMode::RenderHairRoughness),		FCanExecuteAction::CreateSP(this, &SGroomEditorViewport::CanViewMode, false));
 	CommandList->MapAction(FGroomEditorCommands::Get().ViewMode_VisCluster,			FExecuteAction::CreateSP(this, &SGroomEditorViewport::OnViewMode, EHairStrandsDebugMode::RenderVisCluster),			FCanExecuteAction::CreateSP(this, &SGroomEditorViewport::CanViewMode, true)); // While Clusters LOD/Culling doesn't require simulation, cluster visualization requires HairInterpolation to run, which is only available when simulation is enabled in the Groom editor.
 	CommandList->MapAction(FGroomEditorCommands::Get().ViewMode_Group,				FExecuteAction::CreateSP(this, &SGroomEditorViewport::OnViewMode, EHairStrandsDebugMode::RenderHairGroup),			FCanExecuteAction::CreateSP(this, &SGroomEditorViewport::CanViewMode, false));
-	CommandList->MapAction(FGroomEditorCommands::Get().ViewMode_CardsGuides,		FExecuteAction::CreateSP(this, &SGroomEditorViewport::OnCardsGuides), FCanExecuteAction::CreateSP(this, &SGroomEditorViewport::CanCardsGuides));
+	CommandList->MapAction(FGroomEditorCommands::Get().ViewMode_CardsGuides,		FExecuteAction::CreateSP(this, &SGroomEditorViewport::OnCardsGuides),												FCanExecuteAction::CreateSP(this, &SGroomEditorViewport::CanCardsGuides));
 }
 
 void SGroomEditorViewport::OnViewMode(EHairStrandsDebugMode Mode)
@@ -254,7 +254,7 @@ bool SGroomEditorViewport::CanCardsGuides() const
 			const FHairGroupInstance* Instance = GroomComponent->GetGroupInstance(GroupIt); 
 			if (Instance->GeometryType == EHairGeometryType::Cards)
 			{
-				return true;
+				return GroomComponent->GroomAsset->IsSimulationEnable();
 			}
 		}
 	}
