Add support for root UV for hair cards groom.

diff --git a/Engine/Plugins/Experimental/HairModelingToolset/Source/HairModelingToolset/Private/GroomQueryUtil.cpp b/Engine/Plugins/Experimental/HairModelingToolset/Source/HairModelingToolset/Private/GroomQueryUtil.cpp
--- a/Engine/Plugins/Experimental/HairModelingToolset/Source/HairModelingToolset/Private/GroomQueryUtil.cpp
+++ b/Engine/Plugins/Experimental/HairModelingToolset/Source/HairModelingToolset/Private/GroomQueryUtil.cpp
@@ -71,7 +71,8 @@ void UE::GroomQueries::ExtractAllHairCards(AGroomActor* GroomActor,
 		}
 
 		FHairCardsDatas CardData;
-		if (!FHairCardsBuilder::ExtractCardsData(StaticMesh, CardData))
+		FHairStrandsDatas DummyStrandsData;
+		if (!FHairCardsBuilder::ExtractCardsData(StaticMesh, DummyStrandsData, CardData))
 		{
 			continue;
 		}
diff --git a/Engine/Plugins/Runtime/HairStrands/Source/HairStrandsCore/Private/GroomAsset.cpp b/Engine/Plugins/Runtime/HairStrands/Source/HairStrandsCore/Private/GroomAsset.cpp
--- a/Engine/Plugins/Runtime/HairStrands/Source/HairStrandsCore/Private/GroomAsset.cpp
+++ b/Engine/Plugins/Runtime/HairStrands/Source/HairStrandsCore/Private/GroomAsset.cpp
@@ -2518,7 +2518,17 @@ bool UGroomAsset::BuildCardsGeometry(uint32 GroupIndex)
 			if (CardsMesh != nullptr)
 			{
 				CardsMesh->ConditionalPostLoad();
-				bInitResources = FHairCardsBuilder::ImportGeometry(CardsMesh, LOD.BulkData, LODGuidesData, LOD.InterpolationBulkData);
+
+				// Create a transient FHairStrandsData in order to extract RootUV and transfer them to cards data
+				FHairStrandsDatas TempHairStrandsData;
+				{
+					const FHairDescriptionGroups& LocalHairDescriptionGroups = GetHairDescriptionGroups();
+					FHairGroupInfo DummyInfo;
+					check(LocalHairDescriptionGroups.IsValid());
+					TempHairStrandsData = LocalHairDescriptionGroups.HairGroups[GroupIndex].Strands;
+					FGroomBuilder::BuildData(TempHairStrandsData);
+				}
+				bInitResources = FHairCardsBuilder::ImportGeometry(CardsMesh, TempHairStrandsData, LOD.BulkData, LODGuidesData, LOD.InterpolationBulkData);
 				if (!bInitResources)
 				{
 					UE_LOG(LogHairStrands, Warning, TEXT("Failed to import cards from %s for Group %d LOD %d."), *CardsMesh->GetName(), GroupIndex, LODIt);
@@ -2546,6 +2556,7 @@ bool UGroomAsset::BuildCardsGeometry(uint32 GroupIndex)
 				InitAtlasTexture(LOD.RestResource, Desc->Textures.AttributeTexture, EHairAtlasTextureType::Attribute);
 				InitAtlasTexture(LOD.RestResource, Desc->Textures.CoverageTexture, EHairAtlasTextureType::Coverage);
 				InitAtlasTexture(LOD.RestResource, Desc->Textures.AuxilaryDataTexture, EHairAtlasTextureType::AuxilaryData);
+				InitAtlasTexture(LOD.RestResource, Desc->Textures.AuxilaryDataTexture, EHairAtlasTextureType::AuxilaryData);
 				LOD.RestResource->bInvertUV = Desc->SourceType == EHairCardsSourceType::Procedural;
 				
 				// 2.2 Load interoplatino resources
@@ -3610,7 +3621,8 @@ bool UGroomAsset::GetHairCardsGuidesDatas(
 
 			FHairCardsBulkData					DummyBulkData;
 			FHairCardsInterpolationBulkData		DummyInterpolationBulkData;
-			FHairCardsBuilder::ImportGeometry(CardsMesh, DummyBulkData, OutCardsGuidesData, DummyInterpolationBulkData);
+			FHairStrandsDatas					DummyHairStrandsData;
+			FHairCardsBuilder::ImportGeometry(CardsMesh, DummyHairStrandsData, DummyBulkData, OutCardsGuidesData, DummyInterpolationBulkData);
 			return true;
 		}
 	}
diff --git a/Engine/Plugins/Runtime/HairStrands/Source/HairStrandsCore/Private/HairCardsBuilder.cpp b/Engine/Plugins/Runtime/HairStrands/Source/HairStrandsCore/Private/HairCardsBuilder.cpp
--- a/Engine/Plugins/Runtime/HairStrands/Source/HairStrandsCore/Private/HairCardsBuilder.cpp
+++ b/Engine/Plugins/Runtime/HairStrands/Source/HairStrandsCore/Private/HairCardsBuilder.cpp
@@ -3327,7 +3327,7 @@ namespace FHairCardsBuilder
 FString GetVersion()
 {
 	// Important to update the version when cards building or importing changes
-	return TEXT("9ai");
+	return TEXT("9b5");
 }
 
 void AllocateAtlasTexture(UTexture2D* Out, const FIntPoint& Resolution, uint32 MipCount, EPixelFormat PixelFormat, ETextureSourceFormat SourceFormat)
@@ -3574,6 +3574,7 @@ void SanitizeMeshDescription(FMeshDescription* MeshDescription)
 
 bool InternalImportGeometry(
 	const UStaticMesh* StaticMesh,
+	const FHairStrandsDatas& InStrandsData, // Used for extracting & assigning root UV to cards data
 	FHairCardsDatas& Out,
 	FHairCardsBulkData& OutBulk,
 	FHairStrandsDatas& OutGuides,
@@ -3669,7 +3670,7 @@ bool InternalImportGeometry(
 
 	SanitizeMeshDescription(MeshDescription);
 	const TVertexAttributesRef<const FVector3f> VertexPositions					= MeshDescription->VertexAttributes().GetAttributesRef<FVector3f>(MeshAttribute::Vertex::Position);
-	const TVertexInstanceAttributesRef<const FVector3f> VertexInstanceNormals		= MeshDescription->VertexInstanceAttributes().GetAttributesRef<FVector3f>(MeshAttribute::VertexInstance::Normal);
+	const TVertexInstanceAttributesRef<const FVector3f> VertexInstanceNormals	= MeshDescription->VertexInstanceAttributes().GetAttributesRef<FVector3f>(MeshAttribute::VertexInstance::Normal);
 	const TVertexInstanceAttributesRef<const FVector3f> VertexInstanceTangents	= MeshDescription->VertexInstanceAttributes().GetAttributesRef<FVector3f>(MeshAttribute::VertexInstance::Tangent);
 	const TVertexInstanceAttributesRef<const float> VertexInstanceBinormalSigns	= MeshDescription->VertexInstanceAttributes().GetAttributesRef<float>(MeshAttribute::VertexInstance::BinormalSign);
 	const TVertexInstanceAttributesRef<const FVector2D> VertexInstanceUVs		= MeshDescription->VertexInstanceAttributes().GetAttributesRef<FVector2D>(MeshAttribute::VertexInstance::TextureCoordinate);
@@ -3687,7 +3688,7 @@ bool InternalImportGeometry(
 		const uint32 VertexIndex = VertexId.GetValue();
 		check(VertexIndex < VertexCount);
 		Out.Cards.Positions[VertexIndex]	= VertexPositions[VertexId];
-		Out.Cards.UVs[VertexIndex]			= FVector4(VertexInstanceUVs[VertexInstanceId0].Component(0), VertexInstanceUVs[VertexInstanceId0].Component(1), 0, 0);
+		Out.Cards.UVs[VertexIndex]			= FVector4(VertexInstanceUVs[VertexInstanceId0].Component(0), VertexInstanceUVs[VertexInstanceId0].Component(1), 0, 0); // RootUV are not set here, but will be 'patched' later once guides & interpolation data are built
 		Out.Cards.Tangents[VertexIndex]		= VertexInstanceTangents[VertexInstanceId0];
 		Out.Cards.Normals[VertexIndex]		= VertexInstanceNormals[VertexInstanceId0];
 
@@ -3742,11 +3743,128 @@ bool InternalImportGeometry(
 		}
 	}
 
+	// Patch Cards Position to store cards length into the W component
+	for (uint32 CardIt = 0, CardsCount = Out.Cards.IndexOffsets.Num(); CardIt < CardsCount; ++CardIt)
+	{
+		const uint32 CardsIndexCount = Out.Cards.IndexCounts[CardIt];
+		const uint32 CardsIndexOffset = Out.Cards.IndexOffsets[CardIt];
+
+		const float CardLength = CardLengths[CardIt];
+		for (uint32 IndexIt = 0; IndexIt < CardsIndexCount; ++IndexIt)
+		{
+			const uint32 VertexIndex = Out.Cards.Indices[CardsIndexOffset + IndexIt];
+			const float CoordU = Out.Cards.CoordU[VertexIndex];
+			const float InterpolatedCardLength = FMath::Lerp(0.f, CardLength, CoordU);
+
+			// Instead of storing the interoplated card length, store the actual max length of the card, 
+			// as reconstructing the strands length, based on interpolatd CardLength will be too prone to numerical issue.
+			// This means that the strand length retrieves in shader will be an over estimate of the actual length
+			OutBulk.Positions[VertexIndex].W = CardLength; // InterpolatedCardLength;
+		}
+	}
+
+	// Patch Cards RootUV by transfering the guides root UV onto the cards
+	if (InStrandsData.IsValid())
+	{
+		// 1. Extract all roots
+		struct FStrandsRootData
+		{
+			FVector   Position;
+			uint32	  CurveIndex;
+			FVector2D RootUV;
+		};
+		TArray<FStrandsRootData> StrandsRoots;
+		{
+			const uint32 CurveCount = InStrandsData.StrandsCurves.Num();
+			StrandsRoots.Reserve(CurveCount);
+			for (uint32 CurveIt = 0; CurveIt < CurveCount; ++CurveIt)
+			{
+				const uint32 Offset = InStrandsData.StrandsCurves.CurvesOffset[CurveIt];
+				FStrandsRootData& RootData = StrandsRoots.AddDefaulted_GetRef();
+				RootData.Position = InStrandsData.StrandsPoints.PointsPosition[Offset];
+				RootData.RootUV = InStrandsData.StrandsCurves.CurvesRootUV[CurveIt];
+				RootData.CurveIndex = CurveIt;
+			}
+		}
+
+		// 2. Extract cards root points
+		struct FCardsRootData
+		{
+			FVector Position;
+			uint32  CardsIndex;
+		};
+		TArray<FCardsRootData> CardsRoots;
+		{
+			const uint32 CardsCount = Out.Cards.IndexOffsets.Num();
+			CardsRoots.Reserve(CardsCount);
+			for (uint32 CardIt = 0; CardIt < CardsCount; ++CardIt)
+			{
+				const uint32 CardsIndexCount = Out.Cards.IndexCounts[CardIt];
+				const uint32 CardsIndexOffset = Out.Cards.IndexOffsets[CardIt];
+
+				uint32 AverageCount = 0;
+				FVector AverageRootPosition = FVector::ZeroVector;
+				for (uint32 IndexIt = 0; IndexIt < CardsIndexCount; ++IndexIt)
+				{
+					const uint32 VertexIndex = Out.Cards.Indices[CardsIndexOffset + IndexIt];
+					if (Out.Cards.CoordU[VertexIndex] == 0)
+					{
+						AverageRootPosition += Out.Cards.Positions[VertexIndex];
+						++AverageCount;
+					}
+				}
+
+				FCardsRootData& Cards = CardsRoots.AddDefaulted_GetRef();
+				Cards.CardsIndex = CardIt;
+				Cards.Position = AverageRootPosition / AverageCount;
+			}
+		}
+
+		// 3. Find cards root / curve root
+		{
+			for (const FCardsRootData& CardsRoot : CardsRoots)
+			{
+				// 3.1 Find closet root UV
+				// /!\ N^2 loop: the number of cards should be relatively small
+				uint32 CurveIndex = ~0;
+				float ClosestDistance = FLT_MAX;
+				for (const FStrandsRootData& StrandsRoot : StrandsRoots)
+				{
+					const float Distance = FVector::Distance(StrandsRoot.Position, CardsRoot.Position);
+					if (Distance < ClosestDistance)
+					{
+						ClosestDistance = Distance;
+						CurveIndex = StrandsRoot.CurveIndex;
+					}
+				}
+
+				// 3.2 Apply root UV to all cards vertices
+				check(CurveIndex != ~0);
+				if (CurveIndex != ~0)
+				{
+					const FVector2D RootUV = InStrandsData.StrandsCurves.CurvesRootUV[CurveIndex];
+
+					const uint32 CardsIndexCount = Out.Cards.IndexCounts[CardsRoot.CardsIndex];
+					const uint32 CardsIndexOffset = Out.Cards.IndexOffsets[CardsRoot.CardsIndex];
+					for (uint32 IndexIt = 0; IndexIt < CardsIndexCount; ++IndexIt)
+					{
+						const uint32 VertexIndex = Out.Cards.Indices[CardsIndexOffset + IndexIt];
+						Out.Cards.UVs[VertexIndex].Z = RootUV.X;
+						Out.Cards.UVs[VertexIndex].W = RootUV.Y;
+						OutBulk.UVs[VertexIndex].Z   = RootUV.X;
+						OutBulk.UVs[VertexIndex].W   = RootUV.Y;
+					}
+				}
+			}
+		}
+	}
+
 	return bSuccess;
 }
 
 bool ImportGeometry(
 	const UStaticMesh* StaticMesh,
+	const FHairStrandsDatas& InStrandsData,
 	FHairCardsBulkData& OutBulk,
 	FHairStrandsDatas& OutGuides,
 	FHairCardsInterpolationBulkData& OutInterpolationBulkData)
@@ -3754,19 +3872,21 @@ bool ImportGeometry(
 	FHairCardsDatas CardData;
 	return InternalImportGeometry(
 		StaticMesh,
+		InStrandsData,
 		CardData,
 		OutBulk,
 		OutGuides,
 		OutInterpolationBulkData);
 }
 
-bool ExtractCardsData(const UStaticMesh* StaticMesh, FHairCardsDatas& Out)
+bool ExtractCardsData(const UStaticMesh* StaticMesh, const FHairStrandsDatas& InStrandsData, FHairCardsDatas& Out)
 {
 	FHairCardsBulkData OutBulk;
 	FHairStrandsDatas OutGuides;
 	FHairCardsInterpolationBulkData OutInterpolationBulkData;
 	return InternalImportGeometry(
 		StaticMesh,
+		InStrandsData,
 		Out,
 		OutBulk,
 		OutGuides,
diff --git a/Engine/Plugins/Runtime/HairStrands/Source/HairStrandsCore/Public/GroomResources.h b/Engine/Plugins/Runtime/HairStrands/Source/HairStrandsCore/Public/GroomResources.h
--- a/Engine/Plugins/Runtime/HairStrands/Source/HairStrandsCore/Public/GroomResources.h
+++ b/Engine/Plugins/Runtime/HairStrands/Source/HairStrandsCore/Public/GroomResources.h
@@ -638,13 +638,13 @@ struct HAIRSTRANDSCORE_API FHairCardsInterpolationDatas
 	/** Get the number of interpolated points */
 	uint32 Num() const { return PointsSimCurvesVertexIndex.Num(); }
 
-	/** Simulation curve indices, ordered by closest influence */
+	/** Simulation curve indices */
 	TArray<int32> PointsSimCurvesIndex;
 
-	/** Closest vertex indices on simulation curve, ordered by closest influence */
+	/** Closest vertex indices on simulation curve */
 	TArray<int32> PointsSimCurvesVertexIndex;
 
-	/** Lerp value between the closest vertex indices and the next one, ordered by closest influence */
+	/** Lerp value between the closest vertex indices and the next one */
 	TArray<float> PointsSimCurvesVertexLerp;
 };
 
diff --git a/Engine/Plugins/Runtime/HairStrands/Source/HairStrandsCore/Public/HairCardsBuilder.h b/Engine/Plugins/Runtime/HairStrands/Source/HairStrandsCore/Public/HairCardsBuilder.h
--- a/Engine/Plugins/Runtime/HairStrands/Source/HairStrandsCore/Public/HairCardsBuilder.h
+++ b/Engine/Plugins/Runtime/HairStrands/Source/HairStrandsCore/Public/HairCardsBuilder.h
@@ -19,12 +19,14 @@ namespace FHairCardsBuilder
 {
 	bool ImportGeometry(
 		const UStaticMesh* StaticMesh,
+		const FHairStrandsDatas& InStrandsData,
 		FHairCardsBulkData& OutBulk,
 		FHairStrandsDatas& OutGuides,
 		FHairCardsInterpolationBulkData& OutInterpolationBulkData);
 
 	HAIRSTRANDSCORE_API bool ExtractCardsData(
 		const UStaticMesh* StaticMesh, 
+		const FHairStrandsDatas& InStrandsData,
 		FHairCardsDatas& Out);
 
 	void ExportGeometry(
diff --git a/Engine/Shaders/Private/HairStrands/HairCardsAttributeCommon.ush b/Engine/Shaders/Private/HairStrands/HairCardsAttributeCommon.ush
--- a/Engine/Shaders/Private/HairStrands/HairCardsAttributeCommon.ush
+++ b/Engine/Shaders/Private/HairStrands/HairCardsAttributeCommon.ush
@@ -18,81 +18,88 @@ struct FHairVertexAttributes
 	float  HairRoughness;
 };
 
-float GetHairStrandsDepth(uint HairPrimitiveId, float2 InUV, float InDeviceZ)
+float GetHairStrandsDepth(float2 InAtlasUV, float InDeviceZ)
 {
-	InUV.y = HairCardsVF.bInvertUV ? 1 - InUV.y : InUV.y;
-	const float SceneDepthOffset = Texture2DSample(HairCardsVF.DepthTexture, HairCardsVF.DepthSampler, InUV).x;
+	InAtlasUV.y = HairCardsVF.bInvertUV ? 1 - InAtlasUV.y : InAtlasUV.y;
+	const float SceneDepthOffset = Texture2DSample(HairCardsVF.DepthTexture, HairCardsVF.DepthSampler, InAtlasUV).x;
 	const float SceneDepth = ConvertFromDeviceZ(InDeviceZ);
 
 	return ConvertToDeviceZ(SceneDepth + SceneDepthOffset);
 }
 
-float3 GetHairStrandsTangent(uint HairPrimitiveId, float2 InUV, half3x3 TangentToWorld, bool bUseTangentSpace)
+float3 GetHairStrandsTangent(float2 InAtlasUV, half3x3 TangentToWorld, bool bUseTangentSpace)
 {
-	InUV.y = HairCardsVF.bInvertUV ? 1 - InUV.y : InUV.y;
-	const float3 LocalTangent = Texture2DSample(HairCardsVF.TangentTexture, HairCardsVF.TangentSampler, InUV).xyz * 2 - 1;
+	InAtlasUV.y = HairCardsVF.bInvertUV ? 1 - InAtlasUV.y : InAtlasUV.y;
+	const float3 LocalTangent = Texture2DSample(HairCardsVF.TangentTexture, HairCardsVF.TangentSampler, InAtlasUV).xyz * 2 - 1;
 	const float3 WorldTangent = mul(LocalTangent, TangentToWorld);
 	return bUseTangentSpace ? LocalTangent : WorldTangent;
 }
 
-float GetHairStrandsCoverage(uint HairPrimitiveId, float2 InUV)
+float GetHairStrandsCoverage(float2 InAtlasUV)
 {
-	InUV.y = HairCardsVF.bInvertUV ? 1 - InUV.y : InUV.y;
-	return Texture2DSample(HairCardsVF.CoverageTexture, HairCardsVF.CoverageSampler, InUV).x;
+	InAtlasUV.y = HairCardsVF.bInvertUV ? 1 - InAtlasUV.y : InAtlasUV.y;
+	return Texture2DSample(HairCardsVF.CoverageTexture, HairCardsVF.CoverageSampler, InAtlasUV).x;
 }
 
-FHairVertexAttributes GetHairStrandsAttributes(uint HairPrimitiveId, float2 InUV)
+FHairVertexAttributes GetHairStrandsAttributes(float2 InAtlasUV, float2 InRootUV = float2(0, 0), float InCardLength = 0.f)
 {
-	InUV.y = HairCardsVF.bInvertUV ? 1 - InUV.y : InUV.y;
-	const float4 Value = Texture2DSample(HairCardsVF.AttributeTexture, HairCardsVF.AttributeSampler, InUV);
+	InAtlasUV.y = HairCardsVF.bInvertUV ? 1 - InAtlasUV.y : InAtlasUV.y;
+	const float4 Value = Texture2DSample(HairCardsVF.AttributeTexture, HairCardsVF.AttributeSampler, InAtlasUV);
+	const float CoordU = Value.z;
+	const float Seed = Value.w;
+	const float WorldRadius = 0.01f; // Hardcoded for cards
+	const float WorldLength = InCardLength;
+	
+	// Invert V to compensate image origin flip
+	// Similar to DecodeHairAttribute in HairStrandVertexFactoryCommon.ush
+	InRootUV.y = 1.0 - InRootUV.y;
 
 	FHairVertexAttributes Out;
-	Out.HairUV = float2(Value.z, 0.5f);
-	Out.HairRootUV = Value.xy;
-	Out.HairDimensions = 0;
+	Out.HairUV = float2(CoordU, 0.5f);
+	Out.HairRootUV = InRootUV;
+	Out.HairDimensions = float2(WorldLength, WorldRadius);
 	Out.HairBaseColor = 0;
-	Out.HairSeed = Value.w;
+	Out.HairSeed = Seed;
 	Out.HairRoughness = 0;
 	return Out;
 }
 
-float4 GetHairStrandsAuxilaryData(uint HairPrimitiveId, float2 InUV)
+float4 GetHairStrandsAuxilaryData(float2 HairPrimitiveUV)
 {
-	return Texture2DSample(HairCardsVF.AuxilaryDataTexture, HairCardsVF.AuxilaryDataSampler, InUV);
+	return Texture2DSample(HairCardsVF.AuxilaryDataTexture, HairCardsVF.AuxilaryDataSampler, HairPrimitiveUV);
 }
 
-float2 GetHairStrandsUV(uint HairPrimitiveId, float2 HairPrimitiveUV)
+float2 GetHairStrandsUV(float2 HairPrimitiveUV)
 {
-	return GetHairStrandsAttributes(HairPrimitiveId, HairPrimitiveUV).HairUV;
+	return GetHairStrandsAttributes(HairPrimitiveUV).HairUV;
 }
 
-float2 GetHairStrandsDimensions(uint HairPrimitiveId, float2 HairPrimitiveUV)
+float2 GetHairStrandsDimensions(float2 HairPrimitiveUV, float HairPrimitiveLength)
 {
-	return GetHairStrandsAttributes(HairPrimitiveId, HairPrimitiveUV).HairDimensions;
+	return GetHairStrandsAttributes(HairPrimitiveUV, float2(0, 0), HairPrimitiveLength).HairDimensions;
 }
 
-float2 GetHairStrandsRootUV(uint HairPrimitiveId, float2 HairPrimitiveUV)
+float2 GetHairStrandsRootUV(float2 HairPrimitiveUV, float2 InRootUV)
 {
-#if MANUAL_VERTEX_FETCH
-	return HairCardsVF.UVsBuffer[HairPrimitiveId].zw;
-#else
-	return GetHairStrandsAttributes(HairPrimitiveId, HairPrimitiveUV).HairRootUV; // TODO fix me?
-#endif
+	// Invert V to compensate image origin flip
+	// Similar to DecodeHairAttribute in HairStrandVertexFactoryCommon.ush
+	InRootUV.y = 1.0 - InRootUV.y; 
+	return InRootUV;
 }
 
-float  GetHairStrandsSeed(uint HairPrimitiveId, float2 HairPrimitiveUV)
+float  GetHairStrandsSeed(float2 HairPrimitiveUV)
 {
-	return GetHairStrandsAttributes(HairPrimitiveId, HairPrimitiveUV).HairSeed;
+	return GetHairStrandsAttributes(HairPrimitiveUV).HairSeed;
 }
 
-float3 GetHairStrandsBaseColor(uint HairPrimitiveId, float2 HairPrimitiveUV)
+float3 GetHairStrandsBaseColor(float2 HairPrimitiveUV)
 {
-	return GetHairStrandsAttributes(HairPrimitiveId, HairPrimitiveUV).HairBaseColor;
+	return GetHairStrandsAttributes(HairPrimitiveUV).HairBaseColor;
 }
 
-float GetHairStrandsRoughness(uint HairPrimitiveId, float2 HairPrimitiveUV)
+float GetHairStrandsRoughness(float2 HairPrimitiveUV)
 {
-	return GetHairStrandsAttributes(HairPrimitiveId, HairPrimitiveUV).HairRoughness;
+	return GetHairStrandsAttributes(HairPrimitiveUV).HairRoughness;
 }
 
 ////////////////////////////////////////////////////////////////////////////////
diff --git a/Engine/Shaders/Private/HairStrands/HairCardsVertexFactory.ush b/Engine/Shaders/Private/HairStrands/HairCardsVertexFactory.ush
--- a/Engine/Shaders/Private/HairStrands/HairCardsVertexFactory.ush
+++ b/Engine/Shaders/Private/HairStrands/HairCardsVertexFactory.ush
@@ -83,8 +83,10 @@ struct FVertexFactoryIntermediates
 	half3x3 TangentToLocal;
 	half3x3 TangentToWorld;
 	half TangentToWorldSign;
-	uint HairPrimitiveId;
 	float2 HairPrimitiveUV;
+	float2 HairPrimitiveRootUV;
+	float2 HairPrimitiveLength;
+
 	/** Cached primitive and instance data */
 	FSceneDataIntermediates SceneData;
 };
@@ -103,20 +105,12 @@ uint GetHairVertexFetchId(FVertexFactoryInput Input)
 #endif
 }
 
-uint GetHairPrimitiveId(FVertexFactoryInput Input)
-{
-#if MANUAL_VERTEX_FETCH
-	return GetHairVertexFetchId(Input);
-#else
-	return 0;
-#endif
-}
-
 struct FHairCardsPoint
 {
 	float3 Position;
 	float2 CardsUVs;
 	float2 RootUVs;
+	float  CardLength;
 };
 
 FHairCardsPoint GetVertexPosition(FVertexFactoryInput Input)
@@ -134,6 +128,7 @@ FHairCardsPoint GetVertexPosition(FVertexFactoryInput Input)
 	Out.Position = Data.xyz;
 	Out.CardsUVs = UVs.xy;
 	Out.RootUVs  = UVs.zw;
+	Out.CardLength = Data.w;
 	return Out;
 }
 
@@ -178,8 +173,9 @@ FMaterialPixelParameters GetMaterialPixelParameters(FVertexFactoryInterpolantsVS
 
 	Result.TwoSidedSign = 1;
 	Result.PrimitiveId = ToScalarMemory(GetPrimitiveId(Interpolants));
-	Result.HairPrimitiveId = Interpolants.HairPrimitiveId;
 	Result.HairPrimitiveUV = float2(Interpolants.HairPrimitiveUV);
+	Result.HairPrimitiveRootUV = float2(Interpolants.HairPrimitiveRootUV);
+	Result.HairPrimitiveLength = Interpolants.HairPrimitiveLength;
 #if NUM_TEX_COORD_INTERPOLATORS
 	UNROLL
 	for (int CoordinateIndex = 0; CoordinateIndex < NUM_TEX_COORD_INTERPOLATORS; CoordinateIndex++)
@@ -278,8 +274,9 @@ FVertexFactoryIntermediates GetVertexFactoryIntermediates(FVertexFactoryInput In
 {
 	FVertexFactoryIntermediates Intermediates;
 	Intermediates.SceneData			= VF_GPUSCENE_GET_INTERMEDIATES(Input);
-	Intermediates.HairPrimitiveId	= GetHairPrimitiveId(Input);
 	Intermediates.HairPrimitiveUV	= GetVertexPosition(Input).CardsUVs;
+	Intermediates.HairPrimitiveRootUV= GetVertexPosition(Input).RootUVs;
+	Intermediates.HairPrimitiveLength= GetVertexPosition(Input).CardLength; 
 	//Intermediates.HairDensity		= HairStrandsVF_Density;
 	//Intermediates.HairDimensions	= float2(GetWorldStrandLength(Input), GetWorldStrandRadius(Input));
 
@@ -335,14 +332,11 @@ FVertexFactoryInterpolantsVSToPS VertexFactoryGetInterpolantsVSToPS(FVertexFacto
 	Interpolants.EyeIndex = 0;
 #endif
 
-#if MANUAL_VERTEX_FETCH
-	const uint VertexFetchId = GetHairVertexFetchId(Input);
-	Interpolants.HairPrimitiveId = VertexFetchId;
-	Interpolants.HairPrimitiveUV = GetVertexPosition(Input).CardsUVs;
-#else
-	Interpolants.HairPrimitiveId = 0;
-	Interpolants.HairPrimitiveUV = GetVertexPosition(Input).CardsUVs;
-#endif
+	FHairCardsPoint CardsPoint = GetVertexPosition(Input);
+	Interpolants.HairPrimitiveUV = CardsPoint.CardsUVs;
+	Interpolants.HairPrimitiveRootUV = CardsPoint.RootUVs;
+	Interpolants.HairPrimitiveLength = CardsPoint.CardLength;
+
 	SetPrimitiveId(Interpolants, Intermediates.SceneData.PrimitiveId);
 #if NUM_TEX_COORD_INTERPOLATORS 
 	SetUV(Interpolants, 0, Interpolants.HairPrimitiveUV);
@@ -463,8 +457,9 @@ FVertexFactoryRayTracingInterpolants VertexFactoryInterpolate(FVertexFactoryRayT
 #endif
 
 #if HAIR_CARD_MESH_FACTORY
-	O.InterpolantsVSToPS.HairPrimitiveId = a.InterpolantsVSToPS.HairPrimitiveId;
 	INTERPOLATE_MEMBER(InterpolantsVSToPS.HairPrimitiveUV);
+	INTERPOLATE_MEMBER(InterpolantsVSToPS.HairPrimitiveRootUV);
+	INTERPOLATE_MEMBER(InterpolantsVSToPS.HairPrimitiveLength);
 #endif
 
 	return O;
diff --git a/Engine/Shaders/Private/HairStrands/HairStrandsVertexFactoryCommon.ush b/Engine/Shaders/Private/HairStrands/HairStrandsVertexFactoryCommon.ush
--- a/Engine/Shaders/Private/HairStrands/HairStrandsVertexFactoryCommon.ush
+++ b/Engine/Shaders/Private/HairStrands/HairStrandsVertexFactoryCommon.ush
@@ -69,7 +69,7 @@ FHairAttributes DecodeHairAttribute(float4 Attributes0, uint4 Attributes1, float
 {
 	float2 RootUV, RootIndex;
 	RootUV.x		= Attributes0.x;
-	RootUV.y		= 1.0 - Attributes0.y; // Invert V to compensate image origin flip
+	RootUV.y		= 1.0 - Attributes0.y; // Invert V to compensate image origin flip. Similar to GetHairStrandsRootUV in HairCardsAttributeCommon.ush
 	RootIndex.x		= float(Attributes1.x);
 	RootIndex.y		= float(Attributes1.y);
 
diff --git a/Engine/Shaders/Private/LocalVertexFactoryCommon.ush b/Engine/Shaders/Private/LocalVertexFactoryCommon.ush
--- a/Engine/Shaders/Private/LocalVertexFactoryCommon.ush
+++ b/Engine/Shaders/Private/LocalVertexFactoryCommon.ush
@@ -36,9 +36,14 @@ struct FVertexFactoryInterpolantsVSToPS
 		nointerpolation uint LightmapDataIndex : LIGHTMAP_ID;
 	#endif
 #endif
-#if HAIR_STRAND_MESH_FACTORY || HAIR_CARD_MESH_FACTORY
-	nointerpolation uint HairPrimitiveId	: HAIR_PRIMITIVE_ID; // Control point ID
-	float2 HairPrimitiveUV					: HAIR_PRIMITIVE_UV; // U: parameteric distance between the two surrounding control points. V: parametric distance along the width.
+#if HAIR_STRAND_MESH_FACTORY
+		nointerpolation uint HairPrimitiveId : HAIR_PRIMITIVE_ID;	// Control point ID
+		float2 HairPrimitiveUV				 : HAIR_PRIMITIVE_UV;	// U: parameteric distance between the two surrounding control points. V: parametric distance along the width.
+#endif
+#if HAIR_CARD_MESH_FACTORY
+	float2 HairPrimitiveUV					: HAIR_PRIMITIVE_UV;	// Atlas UV
+	float2 HairPrimitiveRootUV				: HAIR_PRIMITIVE_ROOTUV;// UV at the root point of the cards
+	float  HairPrimitiveLength				: HAIR_PRIMITIVE_LENGTH;
 #endif
 
 #if HAS_INSTANCE_LOCAL_TO_WORLD_PS
diff --git a/Engine/Shaders/Private/MaterialTemplate.ush b/Engine/Shaders/Private/MaterialTemplate.ush
--- a/Engine/Shaders/Private/MaterialTemplate.ush
+++ b/Engine/Shaders/Private/MaterialTemplate.ush
@@ -370,10 +370,15 @@ struct FMaterialPixelParameters
 #endif
 
 	// Actual primitive Id
-#if IS_HAIR_FACTORY
+#if HAIR_STRAND_MESH_FACTORY
 	uint	HairPrimitiveId;	// Control point ID
 	float2	HairPrimitiveUV;	// U: parametric distance between the two surrounding control point. V: parametric distance along hair width
 #endif
+#if HAIR_CARD_MESH_FACTORY
+	float2	HairPrimitiveUV;	// AtlasUV
+	float2  HairPrimitiveRootUV;// RootUV
+	float   HairPrimitiveLength;// Card length
+#endif
 
 #if HAS_INSTANCE_LOCAL_TO_WORLD_PS
 	float4x4 InstanceLocalToWorld;
@@ -1259,8 +1264,10 @@ float3 MaterialExpressionBlackBody( float Temp )
 
 float2 MaterialExpressionGetHairRootUV(FMaterialPixelParameters Parameters)
 {
-#if IS_HAIR_FACTORY
+#if HAIR_STRAND_MESH_FACTORY
 	return GetHairStrandsRootUV(Parameters.HairPrimitiveId, Parameters.HairPrimitiveUV);
+#elif HAIR_CARD_MESH_FACTORY
+	return GetHairStrandsRootUV(Parameters.HairPrimitiveUV, Parameters.HairPrimitiveRootUV);
 #else
 	return float2(0, 0);
 #endif
@@ -1268,8 +1275,10 @@ float2 MaterialExpressionGetHairRootUV(FMaterialPixelParameters Parameters)
 
 float2 MaterialExpressionGetHairUV(FMaterialPixelParameters Parameters)
 {	
-#if IS_HAIR_FACTORY
+#if HAIR_STRAND_MESH_FACTORY
 	return GetHairStrandsUV(Parameters.HairPrimitiveId, Parameters.HairPrimitiveUV);
+#elif HAIR_CARD_MESH_FACTORY
+	return GetHairStrandsUV(Parameters.HairPrimitiveUV);
 #else
 	return float2(0,0);
 #endif
@@ -1277,8 +1286,10 @@ float2 MaterialExpressionGetHairUV(FMaterialPixelParameters Parameters)
 
 float2 MaterialExpressionGetHairDimensions(FMaterialPixelParameters Parameters)
 {
-#if IS_HAIR_FACTORY
+#if HAIR_STRAND_MESH_FACTORY
 	return GetHairStrandsDimensions(Parameters.HairPrimitiveId, Parameters.HairPrimitiveUV);
+#elif HAIR_CARD_MESH_FACTORY
+	return GetHairStrandsDimensions(Parameters.HairPrimitiveUV, Parameters.HairPrimitiveLength);
 #else
 	return float2(0, 0);
 #endif
@@ -1286,8 +1297,10 @@ float2 MaterialExpressionGetHairDimensions(FMaterialPixelParameters Parameters)
 
 float MaterialExpressionGetHairSeed(FMaterialPixelParameters Parameters)
 {
-#if IS_HAIR_FACTORY
+#if HAIR_STRAND_MESH_FACTORY
 	return GetHairStrandsSeed(Parameters.HairPrimitiveId, Parameters.HairPrimitiveUV);
+#elif HAIR_CARD_MESH_FACTORY
+	return GetHairStrandsSeed(Parameters.HairPrimitiveUV);
 #else
 	return 0;
 #endif
@@ -1295,8 +1308,10 @@ float MaterialExpressionGetHairSeed(FMaterialPixelParameters Parameters)
 
 float3 MaterialExpressionGetHairBaseColor(FMaterialPixelParameters Parameters)
 {
-#if IS_HAIR_FACTORY
+#if HAIR_STRAND_MESH_FACTORY
 	return GetHairStrandsBaseColor(Parameters.HairPrimitiveId, Parameters.HairPrimitiveUV);
+#elif HAIR_CARD_MESH_FACTORY
+	return GetHairStrandsBaseColor(Parameters.HairPrimitiveUV);
 #else
 	return float3(0,0,0);
 #endif
@@ -1304,8 +1319,10 @@ float3 MaterialExpressionGetHairBaseColor(FMaterialPixelParameters Parameters)
 
 float MaterialExpressionGetHairRoughness(FMaterialPixelParameters Parameters)
 {
-#if IS_HAIR_FACTORY
+#if HAIR_STRAND_MESH_FACTORY
 	return GetHairStrandsRoughness(Parameters.HairPrimitiveId, Parameters.HairPrimitiveUV);
+#elif HAIR_CARD_MESH_FACTORY
+	return GetHairStrandsRoughness(Parameters.HairPrimitiveUV);
 #else
 	return 0;
 #endif
@@ -1318,8 +1335,10 @@ float MaterialExpressionGetHairDepth(FMaterialVertexParameters Parameters)
 
 float MaterialExpressionGetHairDepth(FMaterialPixelParameters Parameters)
 {
-#if IS_HAIR_FACTORY
+#if HAIR_STRAND_MESH_FACTORY
 	return GetHairStrandsDepth(Parameters.HairPrimitiveId, Parameters.HairPrimitiveUV, Parameters.SvPosition.z);
+#elif HAIR_CARD_MESH_FACTORY
+	return GetHairStrandsDepth(Parameters.HairPrimitiveUV, Parameters.SvPosition.z);
 #else
 	return 0; 
 #endif
@@ -1327,8 +1346,10 @@ float MaterialExpressionGetHairDepth(FMaterialPixelParameters Parameters)
 
 float MaterialExpressionGetHairCoverage(FMaterialPixelParameters Parameters)
 {
-#if IS_HAIR_FACTORY
+#if HAIR_STRAND_MESH_FACTORY
 	return GetHairStrandsCoverage(Parameters.HairPrimitiveId, Parameters.HairPrimitiveUV);
+#elif HAIR_CARD_MESH_FACTORY
+	return GetHairStrandsCoverage(Parameters.HairPrimitiveUV);
 #else
 	return 0;
 #endif
@@ -1339,7 +1360,7 @@ float3 MaterialExpressionGetHairTangent(FMaterialPixelParameters Parameters, boo
 #if HAIR_STRAND_MESH_FACTORY
 	return bUseTangentSpace ? float3(0,1,0) : Parameters.TangentToWorld[2];
 #elif HAIR_CARD_MESH_FACTORY
-	return GetHairStrandsTangent(Parameters.HairPrimitiveId, Parameters.HairPrimitiveUV, Parameters.TangentToWorld, bUseTangentSpace);
+	return GetHairStrandsTangent(Parameters.HairPrimitiveUV, Parameters.TangentToWorld, bUseTangentSpace);
 #else
 	return 0;
 #endif	
@@ -1360,7 +1381,7 @@ float2 MaterialExpressionGetAtlasUVs(FMaterialPixelParameters Parameters)
 float4 MaterialExpressionGetHairAuxilaryData(FMaterialPixelParameters Parameters)
 {
 #if HAIR_CARD_MESH_FACTORY
-	return GetHairStrandsAuxilaryData(Parameters.HairPrimitiveId, Parameters.HairPrimitiveUV);
+	return GetHairStrandsAuxilaryData(Parameters.HairPrimitiveUV);
 #else
 	return 0;
 #endif
