Add tangent debug mode to groom editor.

diff --git a/Engine/Plugins/Runtime/HairStrands/Source/HairStrandsCore/Private/GroomComponent.cpp b/Engine/Plugins/Runtime/HairStrands/Source/HairStrandsCore/Private/GroomComponent.cpp
--- a/Engine/Plugins/Runtime/HairStrands/Source/HairStrandsCore/Private/GroomComponent.cpp
+++ b/Engine/Plugins/Runtime/HairStrands/Source/HairStrandsCore/Private/GroomComponent.cpp
@@ -629,7 +629,11 @@ public:
 						DebugMode = EHairStrandsDebugMode::RenderLODColoration;
 					}
 					
-					if (DebugMode != EHairStrandsDebugMode::NoneDebug)
+					const bool bNeedDebugMaterial = 
+						DebugMode != EHairStrandsDebugMode::NoneDebug &&
+						DebugMode != EHairStrandsDebugMode::RenderHairTangent;
+
+					if (bNeedDebugMaterial)
 					{
 						float DebugModeScalar = 0;
 						switch(DebugMode)
@@ -1557,6 +1561,14 @@ void UGroomComponent::SetHairLengthScaleEnable(bool bEnable)
 	InitResources();
 }
 
+#if WITH_EDITOR
+void UGroomComponent::SetDebugMode(EHairStrandsDebugMode InMode)
+{ 
+	DebugMode = InMode; 
+	UpdateHairGroupsDescAndInvalidateRenderState(true);
+}
+#endif
+
 bool UGroomComponent::GetIsHairLengthScaleEnabled()
 {
 	bool IsEnabled = true;
@@ -1720,7 +1732,8 @@ void UGroomComponent::UpdateHairGroupsDescAndInvalidateRenderState(bool bInvalid
 	{
 		Instance->Strands.Modifier  = GetGroomGroupsDesc(GroomAsset, this, GroupIndex);
 #if WITH_EDITORONLY_DATA
-		Instance->Debug.DebugMode = GroomAsset ? GroomAsset->GetDebugMode() : EHairStrandsDebugMode::NoneDebug;
+		Instance->Debug.DebugMode = DebugMode;
+		Instance->HairGroupPublicData->DebugMode = DebugMode; // Replicated value for accessing it from the engine side. Refactor this.
 #endif// #if WITH_EDITORONLY_DATA
 		++GroupIndex;
 	}
@@ -2491,7 +2504,7 @@ void UGroomComponent::InitResources(bool bIsBindingReloading)
 		// LODBias is in the Modifier which is needed for LOD selection regardless if the strands are there or not
 		HairGroupInstance->Strands.Modifier = GetGroomGroupsDesc(GroomAsset, this, GroupIt);
 		#if WITH_EDITORONLY_DATA
-		HairGroupInstance->Debug.DebugMode = GroomAsset->GetDebugMode();
+		HairGroupInstance->Debug.DebugMode = DebugMode;
 		#endif// #if WITH_EDITORONLY_DATA
 
 		// Strands data/resources
diff --git a/Engine/Plugins/Runtime/HairStrands/Source/HairStrandsCore/Public/GroomAsset.h b/Engine/Plugins/Runtime/HairStrands/Source/HairStrandsCore/Public/GroomAsset.h
--- a/Engine/Plugins/Runtime/HairStrands/Source/HairStrandsCore/Public/GroomAsset.h
+++ b/Engine/Plugins/Runtime/HairStrands/Source/HairStrandsCore/Public/GroomAsset.h
@@ -553,10 +553,6 @@ public:
 	bool NeedsInterpolationData(int32 GroupIndex) const;
 	bool NeedsInterpolationData() const;
 
-#if WITH_EDITORONLY_DATA
-	void SetDebugMode(EHairStrandsDebugMode InMode) { DebugMode = InMode; OnGroomAssetChanged.Broadcast(); }
-	EHairStrandsDebugMode GetDebugMode() const { return DebugMode; }
-#endif // 
 	void UpdateHairGroupsInfo();
 	bool HasGeometryType(EGroomGeometryType Type) const;
 	bool HasGeometryType(uint32 GroupIndex, EGroomGeometryType Type) const;
@@ -648,9 +644,6 @@ private:
 	TUniquePtr<FHairDescriptionBulkData> HairDescriptionBulkData;
 	TUniquePtr<FHairDescriptionGroups> HairDescriptionGroups;
 
-	// Transient property for visualizing the groom in a certain debug mode. This is used by the groom editor
-	EHairStrandsDebugMode DebugMode = EHairStrandsDebugMode::NoneDebug;
-
 	TArray<FString> StrandsDerivedDataKey;
 	TArray<FString> CardsDerivedDataKey;
 	TArray<FString> MeshesDerivedDataKey;
diff --git a/Engine/Plugins/Runtime/HairStrands/Source/HairStrandsCore/Public/GroomComponent.h b/Engine/Plugins/Runtime/HairStrands/Source/HairStrandsCore/Public/GroomComponent.h
--- a/Engine/Plugins/Runtime/HairStrands/Source/HairStrandsCore/Public/GroomComponent.h
+++ b/Engine/Plugins/Runtime/HairStrands/Source/HairStrandsCore/Public/GroomComponent.h
@@ -151,6 +151,9 @@ public:
 	void ValidateMaterials(bool bMapCheck) const;
 	void Invalidate();
 	void InvalidateAndRecreate();
+
+	void SetDebugMode(EHairStrandsDebugMode InMode);
+	EHairStrandsDebugMode GetDebugMode() const { return DebugMode; }
 #endif
 
 	/* Accessor function for changing Groom asset from blueprint/sequencer */
@@ -294,6 +297,9 @@ private:
 	float LODPredictedIndex = -1.f;
 	float LODForcedIndex = -1.f;
 
+	// Transient property for visualizing the groom in a certain debug mode. This is used by the groom editor
+	EHairStrandsDebugMode DebugMode = EHairStrandsDebugMode::NoneDebug;
+
 	void InitResources(bool bIsBindingReloading=false);
 	void ReleaseResources();
 
diff --git a/Engine/Plugins/Runtime/HairStrands/Source/HairStrandsEditor/Private/GroomEditorCommands.cpp b/Engine/Plugins/Runtime/HairStrands/Source/HairStrandsEditor/Private/GroomEditorCommands.cpp
--- a/Engine/Plugins/Runtime/HairStrands/Source/HairStrandsEditor/Private/GroomEditorCommands.cpp
+++ b/Engine/Plugins/Runtime/HairStrands/Source/HairStrandsEditor/Private/GroomEditorCommands.cpp
@@ -35,6 +35,7 @@ void FGroomEditorCommands::RegisterCommands()
 	UI_COMMAND(ViewMode_Seed,			"Seed",				"View root seed",				EUserInterfaceActionType::Button, FInputChord());
 	UI_COMMAND(ViewMode_Dimension,		"Dimension",		"View strands dimension",		EUserInterfaceActionType::Button, FInputChord());
 	UI_COMMAND(ViewMode_RadiusVariation,"RadiusVariation",	"View strands radius variation",EUserInterfaceActionType::Button, FInputChord());
+	UI_COMMAND(ViewMode_Tangent,		"Tangent",			"View tangents direction",		EUserInterfaceActionType::Button, FInputChord());
 	UI_COMMAND(ViewMode_BaseColor,		"BaseColor",		"View vertex base color",		EUserInterfaceActionType::Button, FInputChord());
 	UI_COMMAND(ViewMode_Roughness,		"Roughness",		"View vertex roughness",		EUserInterfaceActionType::Button, FInputChord());
 	UI_COMMAND(ViewMode_VisCluster,		"Vis. Clusters",	"View clusters (only available if the groom has simulation enabled)",				EUserInterfaceActionType::Button, FInputChord());
diff --git a/Engine/Plugins/Runtime/HairStrands/Source/HairStrandsEditor/Private/GroomEditorViewportToolBar.cpp b/Engine/Plugins/Runtime/HairStrands/Source/HairStrandsEditor/Private/GroomEditorViewportToolBar.cpp
--- a/Engine/Plugins/Runtime/HairStrands/Source/HairStrandsEditor/Private/GroomEditorViewportToolBar.cpp
+++ b/Engine/Plugins/Runtime/HairStrands/Source/HairStrandsEditor/Private/GroomEditorViewportToolBar.cpp
@@ -55,6 +55,7 @@ TSharedRef<SWidget> SGroomEditorViewportToolbar::GenerateShowMenu() const
 		ViewMenuBuilder.AddMenuEntry(Commands.ViewMode_Seed);
 		ViewMenuBuilder.AddMenuEntry(Commands.ViewMode_Dimension);
 		ViewMenuBuilder.AddMenuEntry(Commands.ViewMode_RadiusVariation);
+		ViewMenuBuilder.AddMenuEntry(Commands.ViewMode_Tangent);
 		ViewMenuBuilder.AddMenuEntry(Commands.ViewMode_BaseColor);
 		ViewMenuBuilder.AddMenuEntry(Commands.ViewMode_Roughness);
 
diff --git a/Engine/Plugins/Runtime/HairStrands/Source/HairStrandsEditor/Private/Widgets/SGroomEditorViewport.cpp b/Engine/Plugins/Runtime/HairStrands/Source/HairStrandsEditor/Private/Widgets/SGroomEditorViewport.cpp
--- a/Engine/Plugins/Runtime/HairStrands/Source/HairStrandsEditor/Private/Widgets/SGroomEditorViewport.cpp
+++ b/Engine/Plugins/Runtime/HairStrands/Source/HairStrandsEditor/Private/Widgets/SGroomEditorViewport.cpp
@@ -197,6 +197,7 @@ void SGroomEditorViewport::BindCommands()
 	CommandList->MapAction(FGroomEditorCommands::Get().ViewMode_Seed,				FExecuteAction::CreateSP(this, &SGroomEditorViewport::OnViewMode, EHairStrandsDebugMode::RenderHairSeed),			FCanExecuteAction::CreateSP(this, &SGroomEditorViewport::CanViewMode, false));
 	CommandList->MapAction(FGroomEditorCommands::Get().ViewMode_Dimension,			FExecuteAction::CreateSP(this, &SGroomEditorViewport::OnViewMode, EHairStrandsDebugMode::RenderHairDimension),		FCanExecuteAction::CreateSP(this, &SGroomEditorViewport::CanViewMode, false));
 	CommandList->MapAction(FGroomEditorCommands::Get().ViewMode_RadiusVariation,	FExecuteAction::CreateSP(this, &SGroomEditorViewport::OnViewMode, EHairStrandsDebugMode::RenderHairRadiusVariation),FCanExecuteAction::CreateSP(this, &SGroomEditorViewport::CanViewMode, false));
+	CommandList->MapAction(FGroomEditorCommands::Get().ViewMode_Tangent,			FExecuteAction::CreateSP(this, &SGroomEditorViewport::OnViewMode, EHairStrandsDebugMode::RenderHairTangent),		FCanExecuteAction::CreateSP(this, &SGroomEditorViewport::CanViewMode, false));
 	CommandList->MapAction(FGroomEditorCommands::Get().ViewMode_BaseColor,			FExecuteAction::CreateSP(this, &SGroomEditorViewport::OnViewMode, EHairStrandsDebugMode::RenderHairBaseColor),		FCanExecuteAction::CreateSP(this, &SGroomEditorViewport::CanViewMode, false));
 	CommandList->MapAction(FGroomEditorCommands::Get().ViewMode_Roughness,			FExecuteAction::CreateSP(this, &SGroomEditorViewport::OnViewMode, EHairStrandsDebugMode::RenderHairRoughness),		FCanExecuteAction::CreateSP(this, &SGroomEditorViewport::CanViewMode, false));
 	CommandList->MapAction(FGroomEditorCommands::Get().ViewMode_VisCluster,			FExecuteAction::CreateSP(this, &SGroomEditorViewport::OnViewMode, EHairStrandsDebugMode::RenderVisCluster),			FCanExecuteAction::CreateSP(this, &SGroomEditorViewport::CanViewMode, true)); // While Clusters LOD/Culling doesn't require simulation, cluster visualization requires HairInterpolation to run, which is only available when simulation is enabled in the Groom editor.
@@ -206,16 +207,16 @@ void SGroomEditorViewport::BindCommands()
 
 void SGroomEditorViewport::OnViewMode(EHairStrandsDebugMode Mode)
 {
-	if (GroomComponent && GroomComponent->GroomAsset)
+	if (GroomComponent && GroomComponent)
 	{
-		const EHairStrandsDebugMode OldMode = GroomComponent->GroomAsset->GetDebugMode();
+		const EHairStrandsDebugMode OldMode = GroomComponent->GetDebugMode();
 		if (OldMode != Mode)
 		{
-			GroomComponent->GroomAsset->SetDebugMode(Mode);
+			GroomComponent->SetDebugMode(Mode);
 		}
 		else
 		{
-			GroomComponent->GroomAsset->SetDebugMode(EHairStrandsDebugMode::NoneDebug);
+			GroomComponent->SetDebugMode(EHairStrandsDebugMode::NoneDebug);
 		}
 	}
 }
diff --git a/Engine/Plugins/Runtime/HairStrands/Source/HairStrandsEditor/Public/GroomEditorCommands.h b/Engine/Plugins/Runtime/HairStrands/Source/HairStrandsEditor/Public/GroomEditorCommands.h
--- a/Engine/Plugins/Runtime/HairStrands/Source/HairStrandsEditor/Public/GroomEditorCommands.h
+++ b/Engine/Plugins/Runtime/HairStrands/Source/HairStrandsEditor/Public/GroomEditorCommands.h
@@ -29,6 +29,7 @@ public:
 	TSharedPtr<FUICommandInfo> ViewMode_Seed;
 	TSharedPtr<FUICommandInfo> ViewMode_Dimension;
 	TSharedPtr<FUICommandInfo> ViewMode_RadiusVariation;
+	TSharedPtr<FUICommandInfo> ViewMode_Tangent;
 	TSharedPtr<FUICommandInfo> ViewMode_BaseColor;
 	TSharedPtr<FUICommandInfo> ViewMode_Roughness;
 	TSharedPtr<FUICommandInfo> ViewMode_VisCluster;
diff --git a/Engine/Source/Runtime/Renderer/Private/HairStrands/HairStrandsDebug.cpp b/Engine/Source/Runtime/Renderer/Private/HairStrands/HairStrandsDebug.cpp
--- a/Engine/Source/Runtime/Renderer/Private/HairStrands/HairStrandsDebug.cpp
+++ b/Engine/Source/Runtime/Renderer/Private/HairStrands/HairStrandsDebug.cpp
@@ -200,6 +200,7 @@ EHairStrandsDebugMode GetHairStrandsDebugStrandsMode()
 	case  9:  return EHairStrandsDebugMode::RenderHairBaseColor;
 	case 10:  return EHairStrandsDebugMode::RenderHairRoughness;
 	case 11:  return EHairStrandsDebugMode::RenderVisCluster;
+	case 12:  return EHairStrandsDebugMode::RenderHairTangent;
 	default:  return EHairStrandsDebugMode::NoneDebug;
 	};
 }
@@ -217,6 +218,7 @@ static const TCHAR* ToString(EHairStrandsDebugMode DebugMode)
 	case EHairStrandsDebugMode::RenderHairSeed				: return TEXT("Hair seed");
 	case EHairStrandsDebugMode::RenderHairDimension			: return TEXT("Hair dimensions");
 	case EHairStrandsDebugMode::RenderHairRadiusVariation	: return TEXT("Hair radius variation");
+	case EHairStrandsDebugMode::RenderHairTangent			: return TEXT("Hair tangent");
 	case EHairStrandsDebugMode::RenderHairBaseColor			: return TEXT("Hair vertices color");
 	case EHairStrandsDebugMode::RenderHairRoughness			: return TEXT("Hair vertices roughness");
 	case EHairStrandsDebugMode::RenderVisCluster			: return TEXT("Hair visility clusters");
@@ -1098,9 +1100,26 @@ static void InternalRenderHairStrandsDebugInfo(
 	}
 
 	// Display tangent vector for strands/cards/meshes
-	if (GHairTangentDebug > 0)
+	if (GHairTangentDebug > 0 || View.HairStrandsMeshElements.Num() > 0)
 	{
-		AddDebugHairTangentPass(GraphBuilder, View, SceneTextures, SceneColorTexture);
+		// Check among the hair instances, if hair tangent debug mode is requested
+		bool bTangentEnabled = GHairTangentDebug > 0;
+		if (!bTangentEnabled)
+		{
+			for (const FMeshBatchAndRelevance& Mesh : View.HairStrandsMeshElements)
+			{
+				const FHairGroupPublicData* GroupData = reinterpret_cast<const FHairGroupPublicData*>(Mesh.Mesh->Elements[0].VertexFactoryUserData);
+				if (GroupData->DebugMode == EHairStrandsDebugMode::RenderHairTangent)
+				{
+					bTangentEnabled = true;
+					break;
+				}
+			}
+		}
+		if (bTangentEnabled)
+		{
+			AddDebugHairTangentPass(GraphBuilder, View, SceneTextures, SceneColorTexture);
+		}
 	}
 
 	// Pass this point, all debug rendering concern only hair strands data
diff --git a/Engine/Source/Runtime/Renderer/Public/HairStrandsInterface.h b/Engine/Source/Runtime/Renderer/Public/HairStrandsInterface.h
--- a/Engine/Source/Runtime/Renderer/Public/HairStrandsInterface.h
+++ b/Engine/Source/Runtime/Renderer/Public/HairStrandsInterface.h
@@ -63,6 +63,7 @@ enum class EHairStrandsDebugMode : uint8
 	RenderHairSeed,
 	RenderHairDimension,
 	RenderHairRadiusVariation,
+	RenderHairTangent,
 	RenderHairBaseColor,
 	RenderHairRoughness,
 	RenderVisCluster,
@@ -327,6 +328,7 @@ public:
 	float LODIndex = -1;		// Current LOD used for all views
 	float LODBias = 0;			// Current LOD bias
 	bool bLODVisibility = true; // Enable/disable hair rendering for this component
+	EHairStrandsDebugMode DebugMode = EHairStrandsDebugMode::NoneDebug;
 };
 
 ////////////////////////////////////////////////////////////////////////////////////////////////////
