Add shader print access to HairStrands plugin

diff --git a/Engine/Plugins/Runtime/HairStrands/Source/HairStrandsCore/Private/GroomManager.cpp b/Engine/Plugins/Runtime/HairStrands/Source/HairStrandsCore/Private/GroomManager.cpp
--- a/Engine/Plugins/Runtime/HairStrands/Source/HairStrandsCore/Private/GroomManager.cpp
+++ b/Engine/Plugins/Runtime/HairStrands/Source/HairStrandsCore/Private/GroomManager.cpp
@@ -121,6 +121,7 @@ static void RunInternalHairStrandsInterpolation(
 	const FHairStrandsInstances& Instances,
 	const FGPUSkinCache* SkinCache,
 	const FShaderDrawDebugData* ShaderDrawData,
+	const FShaderPrintData* ShaderPrintData,
 	FGlobalShaderMap* ShaderMap, 
 	EHairStrandsInterpolationType Type,
 	FHairStrandClusterData* ClusterData)
@@ -354,6 +355,7 @@ static void RunInternalHairStrandsInterpolation(
 				ShaderMap,
 				ViewUniqueID,
 				ShaderDrawData, 
+				ShaderPrintData,
 				Instance,
 				Instance->Debug.MeshLODIndex,
 				ClusterData);
@@ -368,6 +370,7 @@ static void RunHairStrandsInterpolation_Guide(
 	const FHairStrandsInstances& Instances,
 	const FGPUSkinCache* SkinCache,
 	const FShaderDrawDebugData* ShaderDrawData,
+	const FShaderPrintData* ShaderPrintData,
 	FGlobalShaderMap* ShaderMap,
 	FHairStrandClusterData* ClusterData)
 {
@@ -384,6 +387,7 @@ static void RunHairStrandsInterpolation_Guide(
 		Instances,
 		SkinCache,
 		ShaderDrawData,
+		ShaderPrintData,
 		ShaderMap,
 		EHairStrandsInterpolationType::SimulationStrands,
 		ClusterData);
@@ -396,6 +400,7 @@ static void RunHairStrandsInterpolation_Strands(
 	const FHairStrandsInstances& Instances,
 	const FGPUSkinCache* SkinCache,
 	const FShaderDrawDebugData* ShaderDrawData,
+	const FShaderPrintData* ShaderPrintData,
 	FGlobalShaderMap* ShaderMap,
 	FHairStrandClusterData* ClusterData)
 {
@@ -412,6 +417,7 @@ static void RunHairStrandsInterpolation_Strands(
 		Instances,
 		SkinCache,
 		ShaderDrawData,
+		ShaderPrintData,
 		ShaderMap,
 		EHairStrandsInterpolationType::RenderStrands,
 		ClusterData);
@@ -790,20 +796,20 @@ void RunHairStrandsFolliculeMaskQueries(FRDGBuilder& GraphBuilder, FGlobalShader
 
 #if WITH_EDITOR
 bool HasHairStrandsTexturesQueries();
-void RunHairStrandsTexturesQueries(FRDGBuilder& GraphBuilder, FGlobalShaderMap* ShaderMap, const struct FShaderDrawDebugData* DebugShaderData);
+void RunHairStrandsTexturesQueries(FRDGBuilder& GraphBuilder, FGlobalShaderMap* ShaderMap, const struct FShaderDrawDebugData* ShaderDebugData);
 #endif
 
 #if WITH_EDITOR
 bool HasHairCardsAtlasQueries();
-void RunHairCardsAtlasQueries(FRDGBuilder& GraphBuilder, FGlobalShaderMap* ShaderMap, const struct FShaderDrawDebugData* DebugShaderData);
+void RunHairCardsAtlasQueries(FRDGBuilder& GraphBuilder, FGlobalShaderMap* ShaderMap, const struct FShaderDrawDebugData* ShaderDebugData);
 #endif
 
-static void RunHairStrandsProcess(FRDGBuilder& GraphBuilder, FGlobalShaderMap* ShaderMap, const struct FShaderDrawDebugData* DebugShaderData)
+static void RunHairStrandsProcess(FRDGBuilder& GraphBuilder, FGlobalShaderMap* ShaderMap, const struct FShaderDrawDebugData* ShaderDebugData)
 {
 #if WITH_EDITOR
 	if (HasHairStrandsTexturesQueries())
 	{
-		RunHairStrandsTexturesQueries(GraphBuilder, ShaderMap, DebugShaderData);
+		RunHairStrandsTexturesQueries(GraphBuilder, ShaderMap, ShaderDebugData);
 	}
 #endif
 
@@ -815,7 +821,7 @@ static void RunHairStrandsProcess(FRDGBuilder& GraphBuilder, FGlobalShaderMap* S
 #if WITH_EDITOR
 	if (HasHairCardsAtlasQueries())
 	{
-		RunHairCardsAtlasQueries(GraphBuilder, ShaderMap, DebugShaderData);
+		RunHairCardsAtlasQueries(GraphBuilder, ShaderMap, ShaderDebugData);
 	}
 #endif
 }
@@ -862,7 +868,7 @@ void ProcessHairStrandsBookmark(
 		if (bHasHairStardsnProcess)
 		{
 			check(GraphBuilder);
-			RunHairStrandsProcess(*GraphBuilder, Parameters.ShaderMap, Parameters.DebugShaderData);
+			RunHairStrandsProcess(*GraphBuilder, Parameters.ShaderMap, Parameters.ShaderDebugData);
 		}
 	}
 	else if (Bookmark == EHairStrandsBookmark::ProcessLODSelection)
@@ -899,7 +905,8 @@ void ProcessHairStrandsBookmark(
 			Parameters.ViewUniqueID,
 			Instances,
 			Parameters.SkinCache,
-			Parameters.DebugShaderData,
+			Parameters.ShaderDebugData,
+			Parameters.ShaderPrintData,
 			Parameters.ShaderMap,
 			&Parameters.HairClusterData);
 	}
@@ -918,7 +925,8 @@ void ProcessHairStrandsBookmark(
 			Parameters.ViewUniqueID,
 			Instances,
 			Parameters.SkinCache,
-			Parameters.DebugShaderData,
+			Parameters.ShaderDebugData,
+			Parameters.ShaderPrintData,
 			Parameters.ShaderMap,
 			&Parameters.HairClusterData);
 	}
@@ -931,7 +939,7 @@ void ProcessHairStrandsBookmark(
 			*Parameters.View,
 			Instances,
 			Parameters.SkinCache,
-			Parameters.DebugShaderData,
+			Parameters.ShaderDebugData,
 			Parameters.SceneColorTexture,
 			Parameters.View->UnscaledViewRect,
 			Parameters.View->ViewUniformBuffer);
diff --git a/Engine/Plugins/Runtime/HairStrands/Source/HairStrandsCore/Private/HairStrandsRendering.cpp b/Engine/Plugins/Runtime/HairStrands/Source/HairStrandsCore/Private/HairStrandsRendering.cpp
--- a/Engine/Plugins/Runtime/HairStrands/Source/HairStrandsCore/Private/HairStrandsRendering.cpp
+++ b/Engine/Plugins/Runtime/HairStrands/Source/HairStrandsCore/Private/HairStrandsRendering.cpp
@@ -1657,6 +1657,7 @@ void ComputeHairStrandsInterpolation(
 	FGlobalShaderMap* ShaderMap,
 	const uint32 ViewUniqueID,
 	const FShaderDrawDebugData* ShaderDrawData,
+	const FShaderPrintData* ShaderPrintData,
 	FHairGroupInstance* Instance,
 	int32 MeshLODIndex,
 	FHairStrandClusterData* InClusterData)
diff --git a/Engine/Plugins/Runtime/HairStrands/Source/HairStrandsCore/Public/HairStrandsRendering.h b/Engine/Plugins/Runtime/HairStrands/Source/HairStrandsCore/Public/HairStrandsRendering.h
--- a/Engine/Plugins/Runtime/HairStrands/Source/HairStrandsCore/Public/HairStrandsRendering.h
+++ b/Engine/Plugins/Runtime/HairStrands/Source/HairStrandsCore/Public/HairStrandsRendering.h
@@ -28,6 +28,7 @@ void ComputeHairStrandsInterpolation(
 	FGlobalShaderMap* ShaderMap,
 	const uint32 ViewUniqueID,
 	const struct FShaderDrawDebugData* DebugShaderData,
+	const struct FShaderPrintData* ShaderPrintData,
 	struct FHairGroupInstance* Instance,
 	int32 LODIndex,
 	struct FHairStrandClusterData* ClusterData);
diff --git a/Engine/Source/Runtime/Renderer/Private/HairStrands/HairStrandsInterface.cpp b/Engine/Source/Runtime/Renderer/Private/HairStrands/HairStrandsInterface.cpp
--- a/Engine/Source/Runtime/Renderer/Private/HairStrands/HairStrandsInterface.cpp
+++ b/Engine/Source/Runtime/Renderer/Private/HairStrands/HairStrandsInterface.cpp
@@ -16,6 +16,7 @@
 #include "SkeletalRenderPublic.h"
 #include "SceneRendering.h"
 #include "SystemTextures.h"
+#include "ShaderPrint.h"
 #include "ScenePrivate.h"
 
 DEFINE_LOG_CATEGORY_STATIC(LogHairRendering, Log, All);
@@ -433,7 +434,8 @@ FHairStrandsBookmarkParameters CreateHairStrandsBookmarkParameters(FScene* Scene
 			}
 		}
 	}	
-	Out.DebugShaderData			= ShaderDrawDebug::IsEnabled(View) ? &View.ShaderDrawData : nullptr;
+	Out.ShaderDebugData			= ShaderDrawDebug::IsEnabled(View) ? &View.ShaderDrawData : nullptr;
+	Out.ShaderPrintData			= ShaderPrint::IsEnabled(View) ? &View.ShaderPrintData : nullptr;
 	Out.SkinCache				= View.Family->Scene->GetGPUSkinCache();
 	Out.ShaderMap				= View.ShaderMap;
 	Out.Instances				= &Scene->HairStrandsSceneData.RegisteredProxies;
diff --git a/Engine/Source/Runtime/Renderer/Public/HairStrandsInterface.h b/Engine/Source/Runtime/Renderer/Public/HairStrandsInterface.h
--- a/Engine/Source/Runtime/Renderer/Public/HairStrandsInterface.h
+++ b/Engine/Source/Runtime/Renderer/Public/HairStrandsInterface.h
@@ -14,6 +14,7 @@
 #include "RenderResource.h"
 #include "RenderGraphResources.h"
 #include "ShaderDebug.h"
+#include "ShaderPrintParameters.h"
 
 class UTexture2D;
 
@@ -432,7 +433,8 @@ enum class EHairStrandsBookmark : uint8
 struct FHairStrandsBookmarkParameters
 {
 	class FGPUSkinCache* SkinCache = nullptr;
-	FShaderDrawDebugData* DebugShaderData = nullptr;
+	FShaderDrawDebugData* ShaderDebugData = nullptr;
+	FShaderPrintData* ShaderPrintData = nullptr;
 	class FGlobalShaderMap* ShaderMap = nullptr;
 
 	uint32 ViewUniqueID = ~0; // View 0
