Improved hair strands root UV precision.

* Attribute1 buffer is now always enabled. RootUV changed from 8bits to 11bits. UDIM range changed from 8bits to 5bits.
* Attribute0 buffer footprint changed from 32bits to 16bit.

diff --git a/Engine/Plugins/Runtime/HairStrands/Source/HairStrandsCore/Private/GroomBuilder.cpp b/Engine/Plugins/Runtime/HairStrands/Source/HairStrandsCore/Private/GroomBuilder.cpp
--- a/Engine/Plugins/Runtime/HairStrands/Source/HairStrandsCore/Private/GroomBuilder.cpp
+++ b/Engine/Plugins/Runtime/HairStrands/Source/HairStrandsCore/Private/GroomBuilder.cpp
@@ -33,7 +33,7 @@ static FAutoConsoleVariableRef CVarHairClusterBuilder_MaxVoxelResolution(TEXT("r
 FString FGroomBuilder::GetVersion()
 {
 	// Important to update the version when groom building changes
-	return TEXT("2q");
+	return TEXT("2r1");
 }
 
 namespace FHairStrandsDecimation
@@ -264,7 +264,6 @@ namespace HairStrandsBuilder
 		const FHairStrandsPoints& Points = HairStrands.StrandsPoints;
 
 		// Track what features/data the bulk data contains
-		bool bHasUDIMData = false;
 		bool bHasMaterialData = false;
 
 		struct FPackedRadiusAndType
@@ -317,22 +316,14 @@ namespace HairStrandsBuilder
 				RootUV.Y = RootUV.Y == 1.0f ? RootUV.Y - SmallEpsilon : RootUV.Y;
 				const FVector2D TextureRootUV(FMath::Fractional(RootUV.X), FMath::Fractional(RootUV.Y));
 				const FVector2D TextureIndexUV = RootUV - TextureRootUV;
-				PackedAttributes0.RootU  = uint8(FMath::Clamp(TextureRootUV.X*255.f, 0.f, 255.f));
-				PackedAttributes0.RootV  = uint8(FMath::Clamp(TextureRootUV.Y*255.f, 0.f, 255.f));
 
 				// UDIM
-				{
-					FHairStrandsAttribute1Format::Type& PackedAttributes1 = OutPackedAttributes1[PointIndex + IndexOffset];
-					PackedAttributes1.IndexU = uint8(FMath::Clamp(TextureIndexUV.X, 0.f, 255.f));
-					PackedAttributes1.IndexV = uint8(FMath::Clamp(TextureIndexUV.Y, 0.f, 255.f));
-					PackedAttributes1.Unused0 = 0;
-					PackedAttributes1.Unused1 = 0;
-
-					if (PackedAttributes1.IndexU != 0 || PackedAttributes1.IndexV != 0)
-					{
-						bHasUDIMData = true;
-					}
-				}
+				FHairStrandsAttribute1Format::Type& PackedAttributes1 = OutPackedAttributes1[PointIndex + IndexOffset];
+				PackedAttributes1.Packed = 0;
+				PackedAttributes1.Packed |= (uint32(FMath::Clamp(TextureRootUV.X * 2047.f, 0.f, 2047.f)) & 0x7FFu);
+				PackedAttributes1.Packed |= (uint32(FMath::Clamp(TextureRootUV.Y * 2047.f, 0.f, 2047.f)) & 0x7FFu) << 11u;
+				PackedAttributes1.Packed |= (uint32(FMath::Clamp(TextureIndexUV.X, 0.f, 31.f)) & 0x1F) << 22u;
+				PackedAttributes1.Packed |= (uint32(FMath::Clamp(TextureIndexUV.Y, 0.f, 31.f)) & 0x1F) << 27u;
 
 				// Material
 				{
@@ -360,14 +351,11 @@ namespace HairStrandsBuilder
 		OutBulkData.Flags = FHairStrandsBulkData::DataFlags_HasData;
 
 		if (bHasMaterialData)	{ OutBulkData.Flags |= FHairStrandsBulkData::DataFlags_HasMaterialData; }	else { OutPackedMaterials.Empty(); }
-		if (bHasUDIMData)		{ OutBulkData.Flags |= FHairStrandsBulkData::DataFlags_HasUDIMData; }		else { OutPackedAttributes1.Empty(); }
 
 		CopyToBulkData<FHairStrandsPositionFormat>(OutBulkData.Positions, OutPackedPositions);
 		CopyToBulkData<FHairStrandsAttribute0Format>(OutBulkData.Attributes0, OutPackedAttributes0);
-		if (bHasUDIMData)
-		{
-			CopyToBulkData<FHairStrandsAttribute1Format>(OutBulkData.Attributes1, OutPackedAttributes1);
-		}
+		CopyToBulkData<FHairStrandsAttribute1Format>(OutBulkData.Attributes1, OutPackedAttributes1);
+
 		if (bHasMaterialData)
 		{
 			CopyToBulkData<FHairStrandsMaterialFormat>(OutBulkData.Materials, OutPackedMaterials);
diff --git a/Engine/Plugins/Runtime/HairStrands/Source/HairStrandsCore/Private/GroomResources.cpp b/Engine/Plugins/Runtime/HairStrands/Source/HairStrandsCore/Private/GroomResources.cpp
--- a/Engine/Plugins/Runtime/HairStrands/Source/HairStrandsCore/Private/GroomResources.cpp
+++ b/Engine/Plugins/Runtime/HairStrands/Source/HairStrandsCore/Private/GroomResources.cpp
@@ -768,11 +768,7 @@ bool FHairStrandsRestResource::InternalIsDataLoaded()
 
 	AsyncLoadHairBulkData(BulkData.Positions, bIsLoading);
 	AsyncLoadHairBulkData(BulkData.Attributes0, bIsLoading);
-
-	if (!!(BulkData.Flags & FHairStrandsBulkData::DataFlags_HasUDIMData))
-	{
-		AsyncLoadHairBulkData(BulkData.Attributes1, bIsLoading);
-	}
+	AsyncLoadHairBulkData(BulkData.Attributes1, bIsLoading);
 
 	if (!!(BulkData.Flags & FHairStrandsBulkData::DataFlags_HasMaterialData))
 	{
@@ -790,11 +786,7 @@ void FHairStrandsRestResource::InternalAllocate(FRDGBuilder& GraphBuilder)
 	// 2. A local copy is done by the buffer uploader. This copy is discarded once the uploading is done.
 	InternalCreateVertexBufferRDG_FromBulkData<FHairStrandsPositionFormat>(GraphBuilder, BulkData.Positions, PointCount, PositionBuffer, ToHairResourceDebugName(HAIRSTRANDS_RESOUCE_NAME(CurveType, Hair.StrandsRest_PositionBuffer), ResourceName), EHairResourceUsageType::Static);
 	InternalCreateVertexBufferRDG_FromBulkData<FHairStrandsAttribute0Format>(GraphBuilder, BulkData.Attributes0, PointCount, Attribute0Buffer, ToHairResourceDebugName(HAIRSTRANDS_RESOUCE_NAME(CurveType, Hair.StrandsRest_Attribute0Buffer), ResourceName), EHairResourceUsageType::Static);
-
-	if (!!(BulkData.Flags & FHairStrandsBulkData::DataFlags_HasUDIMData))
-	{
-		InternalCreateVertexBufferRDG_FromBulkData<FHairStrandsAttribute1Format>(GraphBuilder, BulkData.Attributes1, PointCount, Attribute1Buffer, ToHairResourceDebugName(HAIRSTRANDS_RESOUCE_NAME(CurveType, Hair.StrandsRest_Attribute1Buffer), ResourceName), EHairResourceUsageType::Static);
-	}
+	InternalCreateVertexBufferRDG_FromBulkData<FHairStrandsAttribute1Format>(GraphBuilder, BulkData.Attributes1, PointCount, Attribute1Buffer, ToHairResourceDebugName(HAIRSTRANDS_RESOUCE_NAME(CurveType, Hair.StrandsRest_Attribute1Buffer), ResourceName), EHairResourceUsageType::Static);
 
 	if (!!(BulkData.Flags & FHairStrandsBulkData::DataFlags_HasMaterialData))
 	{
diff --git a/Engine/Plugins/Runtime/HairStrands/Source/HairStrandsCore/Private/HairStrandsDatas.cpp b/Engine/Plugins/Runtime/HairStrands/Source/HairStrandsCore/Private/HairStrandsDatas.cpp
--- a/Engine/Plugins/Runtime/HairStrands/Source/HairStrandsCore/Private/HairStrandsDatas.cpp
+++ b/Engine/Plugins/Runtime/HairStrands/Source/HairStrandsCore/Private/HairStrandsDatas.cpp
@@ -72,30 +72,6 @@ FArchive& operator<<(FArchive& Ar, FPackedHairVertex& Vertex)
 	return Ar;
 }
 
-FArchive& operator<<(FArchive& Ar, FPackedHairAttribute0Vertex& Vertex)
-{
-	Ar.UsingCustomVersion(FReleaseObjectVersion::GUID);
-
-	Ar << Vertex.RootU;
-	Ar << Vertex.RootV;
-	Ar << Vertex.NormalizedLength;
-	Ar << Vertex.Seed;
-
-	return Ar;
-}
-
-FArchive& operator<<(FArchive& Ar, FPackedHairAttribute1Vertex& Vertex)
-{
-	Ar.UsingCustomVersion(FReleaseObjectVersion::GUID);
-
-	Ar << Vertex.IndexU;
-	Ar << Vertex.IndexV;
-	Ar << Vertex.Unused0;
-	Ar << Vertex.Unused1;
-
-	return Ar;
-}
-
 FArchive& operator<<(FArchive& Ar, FHairMaterialVertex& Vertex)
 {
 	Ar << Vertex.BaseColorR;
@@ -257,10 +233,8 @@ void FHairStrandsBulkData::Serialize(FArchive& Ar, UObject* Owner)
 
 		Positions.Serialize(Ar, Owner, ChunkIndex, bAttemptFileMapping);
 		Attributes0.Serialize(Ar, Owner, ChunkIndex, bAttemptFileMapping);
-		if (!!(Flags & DataFlags_HasUDIMData))
-		{
-			Attributes1.Serialize(Ar, Owner, ChunkIndex, bAttemptFileMapping);
-		}
+		Attributes1.Serialize(Ar, Owner, ChunkIndex, bAttemptFileMapping);
+
 		if (!!(Flags & DataFlags_HasMaterialData))
 		{
 			Materials.Serialize(Ar, Owner, ChunkIndex, bAttemptFileMapping);
diff --git a/Engine/Plugins/Runtime/HairStrands/Source/HairStrandsCore/Private/HairStrandsVertexFactory.cpp b/Engine/Plugins/Runtime/HairStrands/Source/HairStrandsCore/Private/HairStrandsVertexFactory.cpp
--- a/Engine/Plugins/Runtime/HairStrands/Source/HairStrandsCore/Private/HairStrandsVertexFactory.cpp
+++ b/Engine/Plugins/Runtime/HairStrands/Source/HairStrandsCore/Private/HairStrandsVertexFactory.cpp
@@ -74,7 +74,6 @@ FHairStrandsVertexFactoryUniformShaderParameters FHairGroupInstance::GetHairStan
 	Out.Attribute0Buffer			= VFInput.Strands.Attribute0BufferRHISRV;
 	Out.Attribute1Buffer			= VFInput.Strands.Attribute1BufferRHISRV;
 	Out.MaterialBuffer 				= VFInput.Strands.MaterialBufferRHISRV;
-	Out.HasAttribute1 				= Out.Attribute1Buffer != nullptr;
 	Out.HasMaterial 				= Out.MaterialBuffer != nullptr;
 	Out.TangentBuffer 				= VFInput.Strands.TangentBufferRHISRV;
 	Out.PositionOffsetBuffer 		= VFInput.Strands.PositionOffsetBufferRHISRV;
diff --git a/Engine/Plugins/Runtime/HairStrands/Source/HairStrandsCore/Public/GroomInstance.h b/Engine/Plugins/Runtime/HairStrands/Source/HairStrandsCore/Public/GroomInstance.h
--- a/Engine/Plugins/Runtime/HairStrands/Source/HairStrandsCore/Public/GroomInstance.h
+++ b/Engine/Plugins/Runtime/HairStrands/Source/HairStrandsCore/Public/GroomInstance.h
@@ -48,7 +48,6 @@ BEGIN_GLOBAL_SHADER_PARAMETER_STRUCT(FHairStrandsVertexFactoryUniformShaderParam
 	SHADER_PARAMETER(float, Length)
 	SHADER_PARAMETER(float, Density)
 	SHADER_PARAMETER(uint32, CullingEnable)
-	SHADER_PARAMETER(uint32, HasAttribute1)
 	SHADER_PARAMETER(uint32, HasMaterial)
 	SHADER_PARAMETER(uint32, StableRasterization)
 	SHADER_PARAMETER(uint32, ScatterSceneLighing)
@@ -59,8 +58,8 @@ BEGIN_GLOBAL_SHADER_PARAMETER_STRUCT(FHairStrandsVertexFactoryUniformShaderParam
 	SHADER_PARAMETER_SRV(Buffer<uint4>, PositionBuffer)
 	SHADER_PARAMETER_SRV(Buffer<uint4>, PreviousPositionBuffer)
 
-	SHADER_PARAMETER_SRV(Buffer<float4>, Attribute0Buffer)
-	SHADER_PARAMETER_SRV(Buffer<float4>, Attribute1Buffer)
+	SHADER_PARAMETER_SRV(Buffer<float2>, Attribute0Buffer)
+	SHADER_PARAMETER_SRV(Buffer<uint>, Attribute1Buffer)
 	SHADER_PARAMETER_SRV(Buffer<float4>, MaterialBuffer)
 	SHADER_PARAMETER_SRV(Buffer<float4>, TangentBuffer)
 
diff --git a/Engine/Plugins/Runtime/HairStrands/Source/HairStrandsCore/Public/HairStrandsDatas.h b/Engine/Plugins/Runtime/HairStrands/Source/HairStrandsCore/Public/HairStrandsDatas.h
--- a/Engine/Plugins/Runtime/HairStrands/Source/HairStrandsCore/Public/HairStrandsDatas.h
+++ b/Engine/Plugins/Runtime/HairStrands/Source/HairStrandsCore/Public/HairStrandsDatas.h
@@ -22,10 +22,8 @@ struct FPackedHairVertex
 
 struct FPackedHairAttribute0Vertex
 {
-	typedef uint32 BulkType;
+	typedef uint16 BulkType;
 
-	uint8 RootU;
-	uint8 RootV;
 	uint8 NormalizedLength;
 	uint8 Seed;
 };
@@ -34,10 +32,7 @@ struct FPackedHairAttribute1Vertex
 {
 	typedef uint32 BulkType;
 
-	uint8 IndexU;
-	uint8 IndexV;
-	uint8 Unused0;
-	uint8 Unused1;
+	uint32 Packed;
 };
 
 struct FHairMaterialVertex
@@ -125,7 +120,7 @@ struct FHairStrandsAttribute0Format
 	static const uint32 ComponentCount = 1;
 	static const uint32 SizeInByte = sizeof(Type);
 	static const EVertexElementType VertexElementType = VET_UByte4;
-	static const EPixelFormat Format = PF_R8G8B8A8;
+	static const EPixelFormat Format = PF_R8G8;
 };
 
 struct FHairStrandsAttribute1Format
@@ -134,8 +129,8 @@ struct FHairStrandsAttribute1Format
 	typedef FPackedHairAttribute1Vertex::BulkType BulkType;
 	static const uint32 ComponentCount = 1;
 	static const uint32 SizeInByte = sizeof(Type);
-	static const EVertexElementType VertexElementType = VET_UByte4;
-	static const EPixelFormat Format = PF_R8G8B8A8_UINT;
+	static const EVertexElementType VertexElementType = VET_UInt;
+	static const EPixelFormat Format = PF_R32_UINT;
 };
 
 struct FHairStrandsMaterialFormat
@@ -429,8 +424,7 @@ struct HAIRSTRANDSCORE_API FHairStrandsBulkData
 	enum EDataFlags
 	{
 		DataFlags_HasData = 1,			// Contains valid data. Otherwise: Position, Attributes, ... are all empty
-		DataFlags_HasUDIMData = 2,		// Contains UDIM coords
-		DataFlags_HasMaterialData = 4	// Contains material data (albedo and/or roughness)
+		DataFlags_HasMaterialData = 2	// Contains material data (albedo and/or roughness)
 	};
 
 	void Serialize(FArchive& Ar, UObject* Owner);
diff --git a/Engine/Shaders/Private/HairStrands/HairStrandsAttributeCommon.ush b/Engine/Shaders/Private/HairStrands/HairStrandsAttributeCommon.ush
--- a/Engine/Shaders/Private/HairStrands/HairStrandsAttributeCommon.ush
+++ b/Engine/Shaders/Private/HairStrands/HairStrandsAttributeCommon.ush
@@ -19,16 +19,6 @@ struct FHairVertexAttributes
 	float  HairRoughness;
 };
 
-uint4 GetHairStrandsOptionalAttribute1(uint InVertexId)
-{
-	uint4 Out = 0;
-	if (HairStrandsVF.HasAttribute1 > 0)
-	{
-		Out = HairStrandsVF.Attribute1Buffer[InVertexId];
-	}
-	return Out;
-}
-
 float4 GetHairStrandsOptionalMaterial(uint InVertexId)
 {
 	float4 Out = 0;
@@ -45,8 +35,8 @@ FHairVertexAttributes GetHairStrandsAttributes(uint HairPrimitiveId, float2 Vert
 	const uint VertexId1 = HairPrimitiveId+1;
 
 	// Instead of doing the actual interpolation we could rely on dithering as an optimization
-	const FHairAttributes Attributes0 = DecodeHairAttribute(HairStrandsVF.Attribute0Buffer[VertexId0], GetHairStrandsOptionalAttribute1(VertexId0), HairStrandsVF.Length);
-	const FHairAttributes Attributes1 = DecodeHairAttribute(HairStrandsVF.Attribute0Buffer[VertexId1], GetHairStrandsOptionalAttribute1(VertexId1), HairStrandsVF.Length);
+	const FHairAttributes Attributes0 = DecodeHairAttribute(HairStrandsVF.Attribute0Buffer[VertexId0], HairStrandsVF.Attribute1Buffer[VertexId0], HairStrandsVF.Length);
+	const FHairAttributes Attributes1 = DecodeHairAttribute(HairStrandsVF.Attribute0Buffer[VertexId1], HairStrandsVF.Attribute1Buffer[VertexId1], HairStrandsVF.Length);
 
 	const FHairControlPoint ControlPoint0 = UnpackHairControlPoint(HairStrandsVF.PositionBuffer[VertexId0], float3(0, 0 ,0), HairStrandsVF.Radius, HairStrandsVF.RootScale, HairStrandsVF.TipScale);
 	const FHairControlPoint ControlPoint1 = UnpackHairControlPoint(HairStrandsVF.PositionBuffer[VertexId1], float3(0, 0, 0), HairStrandsVF.Radius, HairStrandsVF.RootScale, HairStrandsVF.TipScale);
diff --git a/Engine/Shaders/Private/HairStrands/HairStrandsVertexFactoryCommon.ush b/Engine/Shaders/Private/HairStrands/HairStrandsVertexFactoryCommon.ush
--- a/Engine/Shaders/Private/HairStrands/HairStrandsVertexFactoryCommon.ush
+++ b/Engine/Shaders/Private/HairStrands/HairStrandsVertexFactoryCommon.ush
@@ -65,17 +65,20 @@ struct FHairAttributes
 	float WorldLength;
 };
 
-FHairAttributes DecodeHairAttribute(float4 Attributes0, uint4 Attributes1, float InVF_Length)
+FHairAttributes DecodeHairAttribute(float2 Attributes0, uint Attributes1, float InVF_Length)
 {
 	float2 RootUV, RootIndex;
-	RootUV.x		= Attributes0.x;
-	RootUV.y		= 1.0 - Attributes0.y; // Invert V to compensate image origin flip. Similar to GetHairStrandsRootUV in HairCardsAttributeCommon.ush
-	RootIndex.x		= float(Attributes1.x);
-	RootIndex.y		= float(Attributes1.y);
+	RootUV.x		= ((Attributes1    ) & 0x7FF) / 2047.f;	// Coord are encoded on 11 bits
+	RootUV.y		= ((Attributes1>>11) & 0x7FF) / 2047.f;	// Coord are encoded on 11 bits 
+	RootIndex.x		= ((Attributes1>>22) & 0x1F);			// UDIM tile are encoded on 5bits
+	RootIndex.y		= ((Attributes1>>27) & 0x1F);			// UDIM tile are encoded on 5bits
+
+	// Invert V to compensate image origin flip. Similar to GetHairStrandsRootUV in HairCardsAttributeCommon.ush
+	RootUV.y = 1.0f - RootUV.y;
 
 	FHairAttributes Out;
 	Out.RootUV		= RootUV + RootIndex;
-	Out.WorldLength = Attributes0.z * InVF_Length;
-	Out.Seed		= Attributes0.w;
+	Out.WorldLength = Attributes0.x * InVF_Length;
+	Out.Seed		= Attributes0.y;
 	return Out;
 }
\ No newline at end of file
