Change the representation of strands from a cross to a quad cylinder to avoid self-shadowing artifacts and prevent discontinuities between segments

Work originally done by Chris.

diff --git a/Engine/Plugins/Runtime/HairStrands/Source/HairStrandsCore/Private/GroomResources.cpp b/Engine/Plugins/Runtime/HairStrands/Source/HairStrandsCore/Private/GroomResources.cpp
--- a/Engine/Plugins/Runtime/HairStrands/Source/HairStrandsCore/Private/GroomResources.cpp
+++ b/Engine/Plugins/Runtime/HairStrands/Source/HairStrandsCore/Private/GroomResources.cpp
@@ -324,31 +324,7 @@ void InternalCreateVertexBufferRDG(FRDGBuilder& GraphBuilder, uint32 InVertexCou
 	FRDGBufferDesc Desc = FRDGBufferDesc::CreateBufferDesc(FormatType::SizeInByte, InVertexCount);
 	FRDGBufferRef Buffer = GraphBuilder.CreateBuffer(Desc, DebugName, ERDGBufferFlags::MultiFrame);
 
-	auto IsFloatFormat = []()
-	{
-		switch (FormatType::Format)
-		{
-		case PF_A32B32G32R32F:
-		case PF_FloatR11G11B10:
-		case PF_FloatRGB:
-		case PF_FloatRGBA:
-		case PF_G16R16F_FILTER:
-		case PF_G16R16F:
-		case PF_G32R32F:
-		case PF_R16F_FILTER:
-		case PF_R16F:
-		case PF_R16G16B16A16_SNORM:
-		case PF_R16G16B16A16_UNORM:
-		case PF_R32_FLOAT:
-		case PF_R5G6B5_UNORM:
-		case PF_R8G8B8A8_SNORM:
-			return true;
-		default:
-			return false;
-		}
-	};
-
-	if (IsFloatFormat())
+	if (IsFloatFormat(FormatType::Format))
 	{
 		AddClearUAVFloatPass(GraphBuilder, GraphBuilder.CreateUAV(Buffer, FormatType::Format), 0.0f);
 	}
@@ -360,6 +336,26 @@ void InternalCreateVertexBufferRDG(FRDGBuilder& GraphBuilder, uint32 InVertexCou
 	ConvertToExternalBufferWithViews(GraphBuilder, Buffer, Out, FormatType::Format);
 }
 
+template<typename FormatType>
+void InternalCreateStructuredBufferRDG(FRDGBuilder& GraphBuilder, uint32 DataCount, FRDGExternalBuffer& Out, const TCHAR* DebugName, EHairResourceUsageType Usage)
+{
+	// Sanity check
+	check(Usage == EHairResourceUsageType::Dynamic);
+
+	const uint32 DataSizeInBytes = FormatType::SizeInByte * DataCount;
+	if (DataSizeInBytes == 0)
+	{
+		Out.Buffer = nullptr;
+		return;
+	}
+
+	FRDGBufferDesc Desc = FRDGBufferDesc::CreateStructuredDesc(FormatType::SizeInByte, DataCount);
+	FRDGBufferRef Buffer = GraphBuilder.CreateBuffer(Desc, DebugName, ERDGBufferFlags::MultiFrame);
+	AddClearUAVPass(GraphBuilder, GraphBuilder.CreateUAV(Buffer, FormatType::Format), 0u);
+	ConvertToExternalBufferWithViews(GraphBuilder, Buffer, Out, FormatType::Format);
+}
+
+
 template<typename FormatType>
 void InternalCreateStructuredBufferRDG_FromBulkData(FRDGBuilder& GraphBuilder, FByteBulkData& InBulkData, uint32 InDataCount, FRDGExternalBuffer& Out, const TCHAR* DebugName, EHairResourceUsageType UsageType)
 {
@@ -1496,34 +1492,47 @@ void FHairCardsInterpolationResource::InternalRelease()
 /////////////////////////////////////////////////////////////////////////////////////////
 
 #if RHI_RAYTRACING
-// RT geometry is built to for a cross around the fiber.
-// 4 triangles per hair vertex => 12 vertices per hair vertex
+// RT geometry for strands is built as a 4 sided cylinder
+//   each vertex of the curve becomes 4 points
+//   each curve segment turns into 2*4=8 triangles (3 indices for each)
+// there is some waste due to the degenerate triangles emitted from the end points of each curve
+// total memory usage is: 4*float4 + 8*uint3 = 40 bytes per vertex
+// The previous implementation used a "cross" layout without an index buffer
+// which used 6*float4 = 48 bytes per vertex
+// NOTE: the vertex buffer is a float4 because it is registered as a UAV for the compute shader to work
+// TODO: use a plain float vertex buffer with 3x the entries instead to save memory? (float3 UAVs are not allowed)
 FHairStrandsRaytracingResource::FHairStrandsRaytracingResource(const FHairStrandsBulkData& InData, const FHairResourceName& ResourceName) :
 	FHairCommonResource(EHairStrandsAllocationType::Deferred, ResourceName),
-	PositionBuffer(), VertexCount(InData.GetNumPoints()*12), bOwnPositionBuffer(true)
+#if STRANDS_CUSTOM_INTERSECTOR
+	VertexCount(InData.GetNumPoints() * 2), bOwnBuffers(true) // only allocate space for primitive AABBs
+#else
+	VertexCount(InData.GetNumPoints() * 4), IndexCount(InData.GetNumPoints() * 8 * 3), bOwnBuffers(true)
+#endif
 {}
 
 FHairStrandsRaytracingResource::FHairStrandsRaytracingResource(const FHairCardsBulkData& InData, const FHairResourceName& ResourceName) :
 	FHairCommonResource(EHairStrandsAllocationType::Deferred, ResourceName),
-	PositionBuffer(), VertexCount(InData.GetNumVertices()), bOwnPositionBuffer(false)
+	VertexCount(InData.GetNumVertices()), bOwnBuffers(false)
 {}
 
 FHairStrandsRaytracingResource::FHairStrandsRaytracingResource(const FHairMeshesBulkData& InData, const FHairResourceName& ResourceName) :
 	FHairCommonResource(EHairStrandsAllocationType::Deferred, ResourceName),
-	PositionBuffer(), VertexCount(InData.GetNumVertices()), bOwnPositionBuffer(false)
+	VertexCount(InData.GetNumVertices()), bOwnBuffers(false)
 {}
 
 void FHairStrandsRaytracingResource::InternalAllocate(FRDGBuilder& GraphBuilder)
 {
-	if (bOwnPositionBuffer)
+	if (bOwnBuffers)
 	{
 		InternalCreateVertexBufferRDG<FHairStrandsRaytracingFormat>(GraphBuilder, VertexCount, PositionBuffer, ToHairResourceDebugName(TEXT("Hair.StrandsRaytracing_PositionBuffer"), ResourceName), EHairResourceUsageType::Dynamic);
+		InternalCreateStructuredBufferRDG<FHairStrandsIndexFormat>(GraphBuilder, IndexCount, IndexBuffer, ToHairResourceDebugName(TEXT("Hair.StrandsRaytracing_IndexBuffer"), ResourceName), EHairResourceUsageType::Dynamic);
 	}
 }
 
 void FHairStrandsRaytracingResource::InternalRelease()
 {
 	PositionBuffer.Release();
+	IndexBuffer.Release();
 	RayTracingGeometry.ReleaseResource();
 	bIsRTGeometryInitialized = false;
 }
diff --git a/Engine/Plugins/Runtime/HairStrands/Source/HairStrandsCore/Private/HairStrandsRendering.cpp b/Engine/Plugins/Runtime/HairStrands/Source/HairStrandsCore/Private/HairStrandsRendering.cpp
--- a/Engine/Plugins/Runtime/HairStrands/Source/HairStrandsCore/Private/HairStrandsRendering.cpp
+++ b/Engine/Plugins/Runtime/HairStrands/Source/HairStrandsCore/Private/HairStrandsRendering.cpp
@@ -24,6 +24,7 @@
 #include "GroomManager.h"
 #include "GroomInstance.h"
 #include "SystemTextures.h"
+#include "ShaderPrintParameters.h"
 
 static float GHairRaytracingRadiusScale = 0;
 static FAutoConsoleVariableRef CVarHairRaytracingRadiusScale(TEXT("r.HairStrands.RaytracingRadiusScale"), GHairRaytracingRadiusScale, TEXT("Override the per instance scale factor for raytracing hair strands geometry (0: disabled, >0:enabled)"));
@@ -1241,7 +1242,13 @@ class FHairRaytracingGeometryCS : public FGlobalShader
 
 		SHADER_PARAMETER_RDG_BUFFER_SRV(Buffer, PositionOffsetBuffer)
 		SHADER_PARAMETER_RDG_BUFFER_SRV(Buffer, PositionBuffer)
+		SHADER_PARAMETER_RDG_BUFFER_SRV(Buffer, TangentBuffer)
 		SHADER_PARAMETER_RDG_BUFFER_UAV(RWBuffer, OutputPositionBuffer)
+		SHADER_PARAMETER_STRUCT_INCLUDE(ShaderPrint::FShaderParameters, ShaderPrintUniformBuffer)
+		SHADER_PARAMETER_STRUCT_INCLUDE(ShaderDrawDebug::FShaderParameters, ShaderDrawParameters)
+	#if !STRANDS_CUSTOM_INTERSECTOR
+		SHADER_PARAMETER_RDG_BUFFER_UAV(RWBuffer, OutputIndexBuffer)
+	#endif
 	END_SHADER_PARAMETER_STRUCT()
 
 public:
@@ -1250,6 +1257,7 @@ public:
 	{
 		FGlobalShader::ModifyCompilationEnvironment(Parameters, OutEnvironment);
 		OutEnvironment.SetDefine(TEXT("SHADER_RT_GEOMETRY"), 1);
+		OutEnvironment.SetDefine(TEXT("ENABLE_CUSTOM_INTERSECTOR"), STRANDS_CUSTOM_INTERSECTOR);
 	}
 };
 
@@ -1258,6 +1266,8 @@ IMPLEMENT_GLOBAL_SHADER(FHairRaytracingGeometryCS, "/Engine/Private/HairStrands/
 static void AddGenerateRaytracingGeometryPass(
 	FRDGBuilder& GraphBuilder,
 	FGlobalShaderMap* ShaderMap,
+	const FShaderPrintData* ShaderPrintData,
+	const FShaderDrawDebugData* ShaderDrawData,
 	uint32 VertexCount,
 	float HairRadius,
 	float RootScale,
@@ -1265,7 +1275,9 @@ static void AddGenerateRaytracingGeometryPass(
 	const FRDGBufferSRVRef& HairWorldOffsetBuffer,
 	FRDGHairStrandsCullingData& CullingData,
 	const FRDGBufferSRVRef& PositionBuffer,
-	const FRDGBufferUAVRef& OutPositionBuffer)
+	const FRDGBufferSRVRef& TangentBuffer,
+	const FRDGBufferUAVRef& OutPositionBuffer,
+	const FRDGBufferUAVRef& OutIndexBuffer)
 {
 	const uint32 GroupSize = ComputeGroupSize();
 	const FIntVector DispatchCount = ComputeDispatchCount(VertexCount, GroupSize);
@@ -1278,7 +1290,20 @@ static void AddGenerateRaytracingGeometryPass(
 	Parameters->HairStrandsVF_HairRootScale = RootScale;
 	Parameters->HairStrandsVF_HairTipScale = TipScale;
 	Parameters->PositionBuffer = PositionBuffer;
+	Parameters->TangentBuffer = TangentBuffer;
 	Parameters->OutputPositionBuffer = OutPositionBuffer;
+	if (ShaderPrintData)
+	{
+		ShaderPrint::SetParameters(GraphBuilder, *ShaderPrintData, Parameters->ShaderPrintUniformBuffer);
+	}
+	if (ShaderDrawData)
+	{
+		ShaderDrawDebug::SetParameters(GraphBuilder, *ShaderDrawData, Parameters->ShaderDrawParameters);
+	}
+
+#if !STRANDS_CUSTOM_INTERSECTOR
+	Parameters->OutputIndexBuffer = OutIndexBuffer;
+#endif
 
 	const bool bCullingEnable = CullingData.bCullingResultAvailable;
 	Parameters->HairStrandsVF_bIsCullingEnable = bCullingEnable ? 1 : 0;
@@ -1365,12 +1390,12 @@ static void AddClearClusterAABBPass(
 ///////////////////////////////////////////////////////////////////////////////////////////////////
 
 #if RHI_RAYTRACING
-static void UpdateHairAccelerationStructure(FRHICommandList& RHICmdList, FRayTracingGeometry* RayTracingGeometry)
+static void UpdateHairAccelerationStructure(FRHICommandList& RHICmdList, FRayTracingGeometry* RayTracingGeometry, EAccelerationStructureBuildMode InMode)
 {
 	SCOPED_DRAW_EVENT(RHICmdList, CommitHairRayTracingGeometryUpdates);
 
 	FRayTracingGeometryBuildParams Params;
-	Params.BuildMode = EAccelerationStructureBuildMode::Update;
+	Params.BuildMode = InMode;
 	Params.Geometry = RayTracingGeometry->RayTracingGeometryRHI;
 	Params.Segments = RayTracingGeometry->Initializer.Segments;
 
@@ -1380,7 +1405,9 @@ static void UpdateHairAccelerationStructure(FRHICommandList& RHICmdList, FRayTra
 static void BuildHairAccelerationStructure_Strands(
 	FRHICommandList& RHICmdList, 
 	uint32 RaytracingVertexCount, 
+	uint32 RayTracingIndexCount,
 	FBufferRHIRef& PositionBuffer, 
+	FBufferRHIRef& IndexBuffer,
 	FRayTracingGeometry* OutRayTracingGeometry,
 	const FString& AssetDebugName,
 	uint32 LODIndex)
@@ -1394,19 +1421,33 @@ static void BuildHairAccelerationStructure_Strands(
 	static int32 DebugNumber = 0;
 	Initializer.DebugName = FName(DebugName, DebugNumber++);
 #endif
-	Initializer.IndexBuffer = nullptr;
+#if STRANDS_CUSTOM_INTERSECTOR
+	Initializer.GeometryType = RTGT_Procedural;
+	Initializer.TotalPrimitiveCount = RaytracingVertexCount / 2;
+#else
+	Initializer.IndexBuffer = IndexBuffer;
 	Initializer.IndexBufferOffset = 0;
 	Initializer.GeometryType = RTGT_Triangles;
-	Initializer.TotalPrimitiveCount = RaytracingVertexCount / 3;
+	Initializer.TotalPrimitiveCount = RayTracingIndexCount / 3;
+#endif
 	Initializer.bFastBuild = true;
-	Initializer.bAllowUpdate = true;
+	Initializer.bAllowUpdate = false; // we won't be using update, so specify this to allow degenerate primitives to be discarded
 
 	FRayTracingGeometrySegment Segment;
 	Segment.VertexBuffer = PositionBuffer;
+#if STRANDS_CUSTOM_INTERSECTOR
+	// stride covers one AABB (which is made up of two vertices)
+	Segment.VertexBufferStride = FHairStrandsRaytracingFormat::SizeInByte * 2;
+#else
 	Segment.VertexBufferStride = FHairStrandsRaytracingFormat::SizeInByte;
+#endif
 	Segment.VertexBufferElementType = FHairStrandsRaytracingFormat::VertexElementType;
-	Segment.MaxVertices = RaytracingVertexCount;
-	Segment.NumPrimitives = RaytracingVertexCount / 3;
+	Segment.MaxVertices = RaytracingVertexCount / 2;
+#if STRANDS_CUSTOM_INTERSECTOR
+	Segment.NumPrimitives = RaytracingVertexCount / 2;
+#else
+	Segment.NumPrimitives = RayTracingIndexCount / 3;
+#endif
 	Initializer.Segments.Add(Segment);
 
 	OutRayTracingGeometry->SetInitializer(Initializer);
@@ -1724,6 +1765,7 @@ void ComputeHairStrandsInterpolation(
 			// 2. Deform hair if needed (e.g.: skinning, simulation, RBF) and recompute tangent
 			FRDGBufferSRVRef Strands_PositionSRV = nullptr;
 			FRDGBufferSRVRef Strands_PositionOffsetSRV = nullptr;
+			FRDGBufferSRVRef Strands_TangentSRV = nullptr;
 			if (bNeedDeformation)
 			{
 				FRDGImportedBuffer Strands_DeformedPosition = Register(GraphBuilder, Instance->Strands.DeformedResource->GetBuffer(FHairStrandsDeformedResource::Current), ERDGImportedBufferFlags::CreateViews);
@@ -1751,6 +1793,7 @@ void ComputeHairStrandsInterpolation(
 			#endif
 				Strands_PositionSRV = Strands_DeformedPosition.SRV;
 				Strands_PositionOffsetSRV = RegisterAsSRV(GraphBuilder, Instance->Strands.DeformedResource->GetPositionOffsetBuffer(FHairStrandsDeformedResource::EFrameType::Current));
+				Strands_TangentSRV = Strands_DeformedTangent.SRV;
 
 				// "WITH_EDITOR && bDebugModePatchedAttributeBuffer" is a special path when visualizing groom within the groom editor, so that we can diplay guides even when there is no simulation or global interpolation enabled
 				const bool bValidGuide = Instance->Guides.bIsSimulationEnable || Instance->Guides.bHasGlobalInterpolation || (WITH_EDITOR && bDebugModePatchedAttributeBuffer);
@@ -1867,6 +1910,7 @@ void ComputeHairStrandsInterpolation(
 				Strands_PositionOffsetSRV = RegisterAsSRV(GraphBuilder, Instance->Strands.RestResource->PositionOffsetBuffer);
 				// Generated the static tangent if they haven;t been generated yet
 				Instance->Strands.RestResource->GetTangentBuffer(GraphBuilder, ShaderMap);
+				Strands_TangentSRV = RegisterAsSRV(GraphBuilder, Instance->Strands.RestResource->TangentBuffer);
 			}
 
 			// 2.1 Update the VF input with the update resources
@@ -1927,9 +1971,12 @@ void ComputeHairStrandsInterpolation(
 					const uint32 HairLODIndex = Instance->HairGroupPublicData->GetIntLODIndex();
 
 					FRDGImportedBuffer Raytracing_PositionBuffer = Register(GraphBuilder, Instance->Strands.RenRaytracingResource->PositionBuffer, ERDGImportedBufferFlags::CreateViews);
+					FRDGImportedBuffer Raytracing_IndexBuffer = Register(GraphBuilder, Instance->Strands.RenRaytracingResource->IndexBuffer, ERDGImportedBufferFlags::CreateViews);
 					AddGenerateRaytracingGeometryPass(
 						GraphBuilder,
 						ShaderMap,
+						ShaderPrintData,
+						ShaderDrawData,
 						Instance->Strands.RestResource->GetVertexCount(),
 						HairRadiusRT,
 						HairRootScaleRT,
@@ -1937,7 +1984,9 @@ void ComputeHairStrandsInterpolation(
 						Strands_PositionOffsetSRV,
 						CullingData,
 						Strands_PositionSRV,
-						Raytracing_PositionBuffer.UAV);
+						Strands_TangentSRV,
+						Raytracing_PositionBuffer.UAV,
+						Raytracing_IndexBuffer.UAV);
 
 					Instance->Strands.CachedHairScaledRadius= HairRadiusRT;
 					Instance->Strands.CachedHairRootScale	= HairRootScaleRT;
@@ -1954,11 +2003,27 @@ void ComputeHairStrandsInterpolation(
 						if (bLocalNeedBuild)
 						{
 							FBufferRHIRef PositionBuffer(Instance->Strands.RenRaytracingResource->PositionBuffer.Buffer->GetRHI());
-							BuildHairAccelerationStructure_Strands(RHICmdList, Instance->Strands.RenRaytracingResource->VertexCount, PositionBuffer, &Instance->Strands.RenRaytracingResource->RayTracingGeometry, Instance->Debug.GroomAssetName, HairLODIndex);
+						#if STRANDS_CUSTOM_INTERSECTOR
+							// no index buffer needed
+							FBufferRHIRef IndexBuffer;
+						#else
+							FBufferRHIRef IndexBuffer(Instance->Strands.RenRaytracingResource->IndexBuffer.Buffer->GetRHI());
+						#endif
+							BuildHairAccelerationStructure_Strands(RHICmdList,
+								Instance->Strands.RenRaytracingResource->VertexCount,
+								Instance->Strands.RenRaytracingResource->IndexCount,
+								PositionBuffer,
+								IndexBuffer,
+								&Instance->Strands.RenRaytracingResource->RayTracingGeometry,
+								Instance->Debug.GroomAssetName,
+								HairLODIndex
+							);
 						}
 						else if (bNeedUpdate)
 						{
-							UpdateHairAccelerationStructure(RHICmdList, &Instance->Strands.RenRaytracingResource->RayTracingGeometry);
+							// hair strands can move chaotically during simulation which can really tank ray tracing performance
+							// even though a rebuild is more expensive, we can still come out ahead overall, even with a single ray cast per pixel
+							UpdateHairAccelerationStructure(RHICmdList, &Instance->Strands.RenRaytracingResource->RayTracingGeometry, EAccelerationStructureBuildMode::Build);
 						}
 						Instance->Strands.RenRaytracingResource->bIsRTGeometryInitialized = true;
 					});
@@ -2090,7 +2155,8 @@ void ComputeHairStrandsInterpolation(
 						}
 						else if (bNeedUpdate)
 						{
-							UpdateHairAccelerationStructure(RHICmdList, &LocalLOD.RaytracingResource->RayTracingGeometry);
+							// TODO: evaluate perf tradeoff of rebuild vs refit
+							UpdateHairAccelerationStructure(RHICmdList, &LocalLOD.RaytracingResource->RayTracingGeometry, EAccelerationStructureBuildMode::Update);
 						}
 					});
 				}
@@ -2153,7 +2219,8 @@ void ComputeHairStrandsInterpolation(
 						}
 						else if (bNeedUpdate)
 						{
-							UpdateHairAccelerationStructure(RHICmdList, &LocalLOD.RaytracingResource->RayTracingGeometry);
+							// TODO: evaluate perf tradeoff of rebuild vs refit
+							UpdateHairAccelerationStructure(RHICmdList, &LocalLOD.RaytracingResource->RayTracingGeometry, EAccelerationStructureBuildMode::Update);
 						}
 						LocalLOD.RaytracingResource->bIsRTGeometryInitialized = true;
 					});
diff --git a/Engine/Plugins/Runtime/HairStrands/Source/HairStrandsCore/Private/HairStrandsVertexFactory.cpp b/Engine/Plugins/Runtime/HairStrands/Source/HairStrandsCore/Private/HairStrandsVertexFactory.cpp
--- a/Engine/Plugins/Runtime/HairStrands/Source/HairStrandsCore/Private/HairStrandsVertexFactory.cpp
+++ b/Engine/Plugins/Runtime/HairStrands/Source/HairStrandsCore/Private/HairStrandsVertexFactory.cpp
@@ -154,6 +154,7 @@ void FHairStrandsVertexFactory::ModifyCompilationEnvironment(const FVertexFactor
 #endif
 	OutEnvironment.SetDefine(TEXT("VF_SUPPORTS_PRIMITIVE_SCENE_DATA"), bUseGPUSceneAndPrimitiveIdStream);
 	OutEnvironment.SetDefine(TEXT("HAIR_STRAND_MESH_FACTORY"), TEXT("1"));
+	OutEnvironment.SetDefine(TEXT("ENABLE_CUSTOM_INTERSECTOR"), STRANDS_CUSTOM_INTERSECTOR);
 }
 
 void FHairStrandsVertexFactory::ValidateCompiledResult(const FVertexFactoryType* Type, EShaderPlatform Platform, const FShaderParameterMap& ParameterMap, TArray<FString>& OutErrors)
diff --git a/Engine/Plugins/Runtime/HairStrands/Source/HairStrandsCore/Public/GroomResources.h b/Engine/Plugins/Runtime/HairStrands/Source/HairStrandsCore/Public/GroomResources.h
--- a/Engine/Plugins/Runtime/HairStrands/Source/HairStrandsCore/Public/GroomResources.h
+++ b/Engine/Plugins/Runtime/HairStrands/Source/HairStrandsCore/Public/GroomResources.h
@@ -11,6 +11,8 @@
 #include "HairStrandsInterface.h"
 #include "HairStrandsMeshProjection.h"
 
+#define STRANDS_CUSTOM_INTERSECTOR 0
+
 inline uint32 GetBufferTotalNumBytes(const FRDGExternalBuffer& In) 
 {
 	return In.Buffer ? In.Buffer->Desc.GetTotalNumBytes() : 0;
@@ -265,10 +267,10 @@ struct FHairStrandsRestResource : public FHairCommonResource
 		uint32 Total = 0;
 		Total += GetBufferTotalNumBytes(PositionBuffer);
 		Total += GetBufferTotalNumBytes(PositionOffsetBuffer);
-		Total += GetBufferTotalNumBytes(Attribute1Buffer);
+		Total += GetBufferTotalNumBytes(TangentBuffer);
 		Total += GetBufferTotalNumBytes(Attribute0Buffer);
+		Total += GetBufferTotalNumBytes(Attribute1Buffer);
 		Total += GetBufferTotalNumBytes(MaterialBuffer);
-		Total += GetBufferTotalNumBytes(TangentBuffer);
 		return Total;
 	}
 
@@ -458,13 +460,16 @@ struct FHairStrandsRaytracingResource : public FHairCommonResource
 	{
 		uint32 Total = 0;
 		Total += GetBufferTotalNumBytes(PositionBuffer);
+		Total += GetBufferTotalNumBytes(IndexBuffer);
 		return Total;
 	}
 
 	FRDGExternalBuffer PositionBuffer;
+	FRDGExternalBuffer IndexBuffer;
 	FRayTracingGeometry RayTracingGeometry;
 	uint32 VertexCount = 0;
-	bool bOwnPositionBuffer = false;
+	uint32 IndexCount = 0;
+	bool bOwnBuffers = false;
 	bool bIsRTGeometryInitialized = false;
 };
 #endif
diff --git a/Engine/Shaders/Private/HairStrands/HairStrandsRaytracingGeometry.usf b/Engine/Shaders/Private/HairStrands/HairStrandsRaytracingGeometry.usf
--- a/Engine/Shaders/Private/HairStrands/HairStrandsRaytracingGeometry.usf
+++ b/Engine/Shaders/Private/HairStrands/HairStrandsRaytracingGeometry.usf
@@ -17,7 +17,11 @@ float HairStrandsVF_HairTipScale;
 
 Buffer<float4>	 PositionOffsetBuffer;
 Buffer<uint4> 	 PositionBuffer;
+Buffer<float4> 	 TangentBuffer;
 RWBuffer<float4> OutputPositionBuffer;
+#if !ENABLE_CUSTOM_INTERSECTOR 
+RWBuffer<uint>   OutputIndexBuffer;
+#endif
 
 uint HairStrandsVF_bIsCullingEnable;
 #if PERMUTATION_CULLING == 1
@@ -52,56 +56,124 @@ FHairControlPoint GetControlPoint(uint ControlPointId)
 	return Out;
 }
 
-float4 PackNormal(float3 V, uint A)
+float3x3 GetTangentBasis(uint ControlPointId)
 {
-	return float4(clamp(V,float(-1).xxx,float(1).xxx), 1.0f);
+	float3 Tx = TangentBuffer[2 * ControlPointId    ].xyz;
+	float3 Tz = TangentBuffer[2 * ControlPointId + 1].xyz;
+	float3 Ty = cross(Tz, Tx);
+	return float3x3(Tx, Ty, Tz);
 }
 
+#if ENABLE_CUSTOM_INTERSECTOR
+void SetAABB(uint InVertexIndex, float3 Lo, float3 Hi)
+{
+	OutputPositionBuffer[2 * InVertexIndex + 0] = float4(Lo.x, Lo.y, Lo.z, Hi.x);
+	OutputPositionBuffer[2 * InVertexIndex + 1] = float4(Hi.y, Hi.z, 0.0, 0.0);
+}
+#endif
+
+// a whole segment of invalid values
 void SetInvalidSegment(uint InVertexIndex)
 {
-	const float NaN = sqrt(-1);
-	const float4 InvalidPoint = float4(NaN, NaN, NaN, NaN);
-
-	const uint BaseOutIndex = InVertexIndex * 12;
-	OutputPositionBuffer[BaseOutIndex]      = InvalidPoint;
-	OutputPositionBuffer[BaseOutIndex + 1]  = InvalidPoint;
-	OutputPositionBuffer[BaseOutIndex + 2]  = InvalidPoint;
-	OutputPositionBuffer[BaseOutIndex + 3]  = InvalidPoint;
-	OutputPositionBuffer[BaseOutIndex + 4]  = InvalidPoint;
-	OutputPositionBuffer[BaseOutIndex + 5]  = InvalidPoint;
-	OutputPositionBuffer[BaseOutIndex + 6]  = InvalidPoint;
-	OutputPositionBuffer[BaseOutIndex + 7]  = InvalidPoint;
-	OutputPositionBuffer[BaseOutIndex + 8]  = InvalidPoint;
-	OutputPositionBuffer[BaseOutIndex + 9]  = InvalidPoint;
-	OutputPositionBuffer[BaseOutIndex + 10] = InvalidPoint;
-	OutputPositionBuffer[BaseOutIndex + 11] = InvalidPoint;
+#if ENABLE_CUSTOM_INTERSECTOR
+	const float NaN = asfloat(0x7FC00000);
+	SetAABB(InVertexIndex, NaN, NaN);
+#else
+	const uint BasePoint = InVertexIndex * 4;
+	const uint TrisPerSegment = 8;
+	const uint IdxsPerSegment = TrisPerSegment * 3;
+	const uint BaseOutIndex = InVertexIndex * IdxsPerSegment;
+	UNROLL
+	for (uint Index = 0; Index < IdxsPerSegment; Index++)
+	{
+		OutputIndexBuffer[BaseOutIndex + Index] = BasePoint;
+	}
+#endif
+}
+
+// TODO: this is not used at the moment
+void SetCapSegment(uint InVertexIndex)
+{
+#if 1
+	SetInvalidSegment(InVertexIndex);
+#else
+	const uint BasePoint = InVertexIndex * 4;
+	const uint TrisPerSegment = 8;
+	const uint IdxsPerSegment = TrisPerSegment * 3;
+	const uint BaseOutIndex = InVertexIndex * IdxsPerSegment;
+	OutputIndexBuffer[BaseOutIndex + 0] = BasePoint;
+	OutputIndexBuffer[BaseOutIndex + 1] = BasePoint + 1;
+	OutputIndexBuffer[BaseOutIndex + 2] = BasePoint + 2;
+
+	OutputIndexBuffer[BaseOutIndex + 3] = BasePoint;
+	OutputIndexBuffer[BaseOutIndex + 4] = BasePoint + 2;
+	OutputIndexBuffer[BaseOutIndex + 5] = BasePoint + 3;
+
+	// fill the reset with degenerate triangles
+	UNROLL
+	for (uint Index = 6; Index < IdxsPerSegment; Index++)
+	{
+		OutputIndexBuffer[BaseOutIndex + Index] = BasePoint;
+	}
+#endif
+}
+
+void SetBodySegment(uint InVertexIndex)
+{
+#if ENABLE_CUSTOM_INTERSECTOR
+	FHairControlPoint P0 = GetControlPoint(InVertexIndex);
+	FHairControlPoint P1 = GetControlPoint(InVertexIndex + 1);
+	float4 A = float4(P0.Position, P0.WorldRadius);
+	float4 B = float4(P1.Position, P1.WorldRadius);
+	float3 Lo = min(A.xyz - A.w, B.xyz - B.w);
+	float3 Hi = max(A.xyz + A.w, B.xyz + B.w);
+	SetAABB(InVertexIndex, Lo, Hi);
+#else
+	const uint BotPoint = InVertexIndex * 4;
+	const uint TopPoint = BotPoint + 4;
+
+	const uint TrisPerSegment = 8;
+	const uint IdxsPerSegment = TrisPerSegment * 3;
+	const uint BaseOutIndex = InVertexIndex * IdxsPerSegment;
+
+	UNROLL
+	for (uint PrevIdx = 3, Idx = 0, OutIndex = BaseOutIndex; Idx < 4; PrevIdx = Idx, Idx++, OutIndex += 6)
+	{
+		// Note that this vertex ordering is carefully chosen to work with the Vertex Factory in RT mode
+		// In particular, we _must_ put one of the bottom point as the last vertex of the triangle
+		// so that the interpolation step will copy the right HairPrimitiveId onwards
+		OutputIndexBuffer[OutIndex + 0] = TopPoint + PrevIdx;
+		OutputIndexBuffer[OutIndex + 1] = TopPoint + Idx;
+		OutputIndexBuffer[OutIndex + 2] = BotPoint + PrevIdx;
+
+		OutputIndexBuffer[OutIndex + 3] = TopPoint + Idx;
+		OutputIndexBuffer[OutIndex + 4] = BotPoint + Idx;
+		OutputIndexBuffer[OutIndex + 5] = BotPoint + PrevIdx; 
+	}
+#endif
 }
 
+
 [numthreads(GROUP_SIZE, 1, 1)]
 void MainCS(uint2 DispatchThreadId : SV_DispatchThreadID)
 {
 	uint VertexIndex0 = GetHairStrandsVertexFetchIndex(DispatchThreadId, GROUP_SIZE, DispatchCountX);
-	uint VertexIndex1 = VertexIndex0 + 1;
 	float LodRadiusScale0 = 1;
 	float LodRadiusScale1 = 1;
 
+	uint SourceIndex0 = VertexIndex0;
 #if PERMUTATION_CULLING == 1
 	if (HairStrandsVF_bIsCullingEnable)
 	{
 		const uint VertexCountAfterCulling = HairStrandsVF_CullingIndirectBuffer[3];
-		uint FetchIndex0 = GetHairStrandsVertexFetchIndex_Culled(DispatchThreadId, GROUP_SIZE);
-		if (FetchIndex0 >= VertexCountAfterCulling)
+		if (VertexIndex0 >= VertexCountAfterCulling)
 		{
 			SetInvalidSegment(VertexIndex0);
 			return;
 		}
-		uint FetchIndex1 = min(FetchIndex0 + 1, VertexCountAfterCulling - 1);
 
-		VertexIndex0 = HairStrandsVF_CullingIndexBuffer[FetchIndex0];
-		LodRadiusScale0 = HairStrandsVF_CullingRadiusScaleBuffer[FetchIndex0];
-
-		VertexIndex1 = HairStrandsVF_CullingIndexBuffer[FetchIndex1];
-		LodRadiusScale1 = HairStrandsVF_CullingRadiusScaleBuffer[FetchIndex1];
+		SourceIndex0 = HairStrandsVF_CullingIndexBuffer[VertexIndex0];
+		LodRadiusScale0 = HairStrandsVF_CullingRadiusScaleBuffer[VertexIndex0];
 
 		if (LodRadiusScale0 <= 0.0f)
 		{
@@ -111,14 +183,39 @@ void MainCS(uint2 DispatchThreadId : SV_DispatchThreadID)
 	}
 #endif
 
-	FHairControlPoint P0 = GetControlPoint(VertexIndex0);
+	FHairControlPoint P0 = GetControlPoint(SourceIndex0);
+#if ENABLE_CUSTOM_INTERSECTOR
+	// aabb position is generated by segment call
+#else
+	// expand current vertex into 4 points around the curve
+	const float3x3 TangentBasis = GetTangentBasis(SourceIndex0);
+	const float3 N = TangentBasis[0];
+	const float3 B = TangentBasis[1];
+
+	const uint BaseOutIndex = VertexIndex0 * 4;
+	{
+		const float3 Px0 = P0.Position - N * P0.WorldRadius * LodRadiusScale0;
+		const float3 Px1 = P0.Position + N * P0.WorldRadius * LodRadiusScale0;
+		const float3 Py0 = P0.Position - B * P0.WorldRadius * LodRadiusScale0;
+		const float3 Py1 = P0.Position + B * P0.WorldRadius * LodRadiusScale0;
+
+		OutputPositionBuffer[BaseOutIndex    ] = float4(Px0, 1);
+		OutputPositionBuffer[BaseOutIndex + 1] = float4(Py0, 1);
+		OutputPositionBuffer[BaseOutIndex + 2] = float4(Px1, 1);
+		OutputPositionBuffer[BaseOutIndex + 3] = float4(Py1, 1);
+	}
+#endif
 
-	// Point at the end of a strand is filled with invalid position. The positions/triangle are 
-	// ignored by DXR during BVH building
-	if (P0.Type == 2)
+	if (P0.Type == HAIR_CONTROLPOINT_END)
 	{
+		// last vertex of a curve, we could create a cap, but this would complicate the logic in the vertex factory
+		// ignore this for now on the assumption that most hairs will have a small tip width
 		SetInvalidSegment(VertexIndex0);
-		return;
+	}
+	else
+	{
+		// regular point, build the body of the cylinder
+		SetBodySegment(VertexIndex0);
 	}
 }
 
diff --git a/Engine/Shaders/Private/HairStrands/HairStrandsVertexFactory.ush b/Engine/Shaders/Private/HairStrands/HairStrandsVertexFactory.ush
--- a/Engine/Shaders/Private/HairStrands/HairStrandsVertexFactory.ush
+++ b/Engine/Shaders/Private/HairStrands/HairStrandsVertexFactory.ush
@@ -98,40 +98,50 @@ struct FVertexFactoryInput
 
 struct FVertexInfo
 {
-	uint QuadIndex;
-	uint BaseIndex;
 	uint VertexIndex;
 	uint HairPrimitiveId;
 	uint IsTip;				// The vertex is on the quad side toward the tip of the strand.
+#if !RAYHITGROUPSHADER
 	uint IsLeft;
 	float LodRadiusScale;
 	bool bForceInvalidQuad;
+#endif
 };
 FVertexInfo GetVertexInfo(FVertexFactoryInput Input)
 {
 	uint VertexId = Input.VertexId;
 
 	FVertexInfo VertexInfo;
-#if USE_HAIR_TRIANGLE_STRIP
-	VertexInfo.QuadIndex	= VertexId % 2;
-	VertexInfo.BaseIndex	= VertexId / 2;
+#if RAYHITGROUPSHADER
+	uint BaseIndex    = VertexId / 4;
+	// TODO: is there any easy way we could handle end-caps?
+	VertexInfo.IsTip        = (Input.TriangleId % 2) == 0 ? Input.TriangleVertexId != 2 : Input.TriangleVertexId == 0;
+#elif USE_HAIR_TRIANGLE_STRIP
+	uint QuadIndex = VertexId % 2;
+	uint BaseIndex	= VertexId / 2;
 	VertexInfo.IsTip		= 0;
-	VertexInfo.IsLeft		= VertexInfo.QuadIndex == 0 ? 1 : 0;
+	VertexInfo.IsLeft		= QuadIndex == 0 ? 1 : 0;
 #else
-	VertexInfo.QuadIndex	= VertexId % 6;
-	VertexInfo.BaseIndex	= VertexId / 6;
-	VertexInfo.IsTip		= VertexInfo.QuadIndex == 0 || VertexInfo.QuadIndex == 2 || VertexInfo.QuadIndex == 4 ? 0 : 1;
-	VertexInfo.IsLeft		= VertexInfo.QuadIndex == 0 || VertexInfo.QuadIndex == 1 || VertexInfo.QuadIndex == 5 ? 1 : 0;
+	uint QuadIndex	= VertexId % 6;
+	uint BaseIndex	= VertexId / 6;
+	VertexInfo.IsTip		= QuadIndex == 0 || QuadIndex == 2 || QuadIndex == 4 ? 0 : 1;
+	VertexInfo.IsLeft		= QuadIndex == 0 || QuadIndex == 1 || QuadIndex == 5 ? 1 : 0;
 #endif
-	VertexInfo.HairPrimitiveId = VertexInfo.BaseIndex;
-	VertexInfo.VertexIndex = VertexInfo.BaseIndex + VertexInfo.IsTip;
+	VertexInfo.HairPrimitiveId = BaseIndex;
+	VertexInfo.VertexIndex = BaseIndex + VertexInfo.IsTip;
+#if !RAYHITGROUPSHADER
 	VertexInfo.LodRadiusScale = 1.0f;
 	VertexInfo.bForceInvalidQuad = false;
+#endif
+
 #if USE_CULLED_CLUSTER == 1
 	if (HairStrandsVF.CullingEnable)
 	{
-		VertexInfo.VertexIndex = HairStrandsVF.CulledVertexIdsBuffer[VertexInfo.BaseIndex + VertexInfo.IsTip];
-		VertexInfo.LodRadiusScale = HairStrandsVF.CulledVertexRadiusScaleBuffer[VertexInfo.BaseIndex + VertexInfo.IsTip];
+		VertexInfo.VertexIndex = HairStrandsVF.CulledVertexIdsBuffer[BaseIndex + VertexInfo.IsTip];
+#if !RAYHITGROUPSHADER
+		VertexInfo.LodRadiusScale = HairStrandsVF.CulledVertexRadiusScaleBuffer[BaseIndex + VertexInfo.IsTip];
+		VertexInfo.bForceInvalidQuad = VertexInfo.LodRadiusScale <= 0.0f;
+#endif
 		VertexInfo.HairPrimitiveId = VertexInfo.VertexIndex - VertexInfo.IsTip;
 	}
 #endif
diff --git a/Engine/Source/Runtime/Renderer/Private/HairStrands/HairStrandsClusters.cpp b/Engine/Source/Runtime/Renderer/Private/HairStrands/HairStrandsClusters.cpp
--- a/Engine/Source/Runtime/Renderer/Private/HairStrands/HairStrandsClusters.cpp
+++ b/Engine/Source/Runtime/Renderer/Private/HairStrands/HairStrandsClusters.cpp
@@ -264,8 +264,8 @@ static void AddClusterCullingPass(
 	FRDGBufferRef PerBlocklTotalIndexCountPreFixSumBuffer = GraphBuilder.CreateBuffer(FRDGBufferDesc::CreateBufferDesc(sizeof(uint32) * 2, ClusterData.ClusterCount), TEXT("Hair.PerBlocklTotalIndexCountPreFixSumBuffer"));
 	FRDGBufferRef PerBlocklIndexCountPreFixSumBuffer = GraphBuilder.CreateBuffer(FRDGBufferDesc::CreateBufferDesc(sizeof(uint32) * 2, ClusterData.ClusterCount), TEXT("Hair.PerBlocklIndexCountPreFixSumBuffer"));
 
-	FRDGImportedBuffer DrawIndirectParametersBuffer = Register(GraphBuilder, ClusterData.HairGroupPublicPtr->GetDrawIndirectBuffer(), ERDGImportedBufferFlags::CreateUAV);
-	FRDGImportedBuffer DrawIndirectParametersRasterComputeBuffer = Register(GraphBuilder, ClusterData.HairGroupPublicPtr->GetDrawIndirectRasterComputeBuffer(), ERDGImportedBufferFlags::CreateUAV);
+	FRDGImportedBuffer DrawIndirectParametersBuffer = Register(GraphBuilder, ClusterData.HairGroupPublicPtr->GetDrawIndirectBuffer(), ERDGImportedBufferFlags::CreateViews);
+	FRDGImportedBuffer DrawIndirectParametersRasterComputeBuffer = Register(GraphBuilder, ClusterData.HairGroupPublicPtr->GetDrawIndirectRasterComputeBuffer(), ERDGImportedBufferFlags::CreateViews);
 	
 	bool bClusterDebugAABBBuffer = false;
 	bool bClusterDebug = false;
@@ -519,7 +519,7 @@ static void AddClusterResetLod0(
 	// Set as culling result not available
 	ClusterData.SetCullingResultAvailable(false);
 
-	FRDGImportedBuffer IndirectBuffer = Register(GraphBuilder, ClusterData.HairGroupPublicPtr->GetDrawIndirectBuffer(), ERDGImportedBufferFlags::CreateUAV);
+	FRDGImportedBuffer IndirectBuffer = Register(GraphBuilder, ClusterData.HairGroupPublicPtr->GetDrawIndirectBuffer(), ERDGImportedBufferFlags::CreateViews);
 
 	// Initialise indirect buffers to entire lod 0 dispatch
 	FHairIndBufferClearCS::FParameters* Parameters = GraphBuilder.AllocParameters<FHairIndBufferClearCS::FParameters>();
diff --git a/Engine/Source/Runtime/Renderer/Private/HairStrands/HairStrandsInterface.cpp b/Engine/Source/Runtime/Renderer/Private/HairStrands/HairStrandsInterface.cpp
--- a/Engine/Source/Runtime/Renderer/Private/HairStrands/HairStrandsInterface.cpp
+++ b/Engine/Source/Runtime/Renderer/Private/HairStrands/HairStrandsInterface.cpp
@@ -238,7 +238,7 @@ void ConvertToExternalBufferWithViews(FRDGBuilder& GraphBuilder, FRDGBufferRef&
 	OutBuffer.Format = Format;
 }
 
-void InternalCreateIndirectBufferRDG(FRDGBuilder& GraphBuilder, FRDGExternalBuffer& Out, const TCHAR* DebugName, const FUintVector4& InitValues)
+void InternalCreateIndirectBufferRDG(FRDGBuilder& GraphBuilder, FRDGExternalBuffer& Out, const TCHAR* DebugName)
 {
 	FRDGBufferDesc Desc = FRDGBufferDesc::CreateBufferDesc(4, 4);
 	Desc.Usage |= BUF_DrawIndirect;
@@ -322,8 +322,8 @@ void FHairGroupPublicData::Allocate(FRDGBuilder& GraphBuilder)
 	
 	if (GUsingNullRHI || !bHasStrands) { return; }
 
-	InternalCreateIndirectBufferRDG(GraphBuilder, DrawIndirectBuffer, TEXT("Hair.Cluster_DrawIndirectBuffer"), FUintVector4(GroupControlTriangleStripVertexCount, 1, 0, 0));
-	InternalCreateIndirectBufferRDG(GraphBuilder, DrawIndirectRasterComputeBuffer, TEXT("Hair.Cluster_DrawIndirectRasterComputeBuffer"), FUintVector4(0, 1, 0, 0));
+	InternalCreateIndirectBufferRDG(GraphBuilder, DrawIndirectBuffer, TEXT("Hair.Cluster_DrawIndirectBuffer"));
+	InternalCreateIndirectBufferRDG(GraphBuilder, DrawIndirectRasterComputeBuffer, TEXT("Hair.Cluster_DrawIndirectRasterComputeBuffer"));
 
 	InternalCreateVertexBufferRDG(GraphBuilder, sizeof(int32), ClusterCount * 6, EPixelFormat::PF_R32_SINT, ClusterAABBBuffer, TEXT("Hair.Cluster_ClusterAABBBuffer"));
 	InternalCreateVertexBufferRDG(GraphBuilder, sizeof(int32), 6, EPixelFormat::PF_R32_SINT, GroupAABBBuffer, TEXT("Hair.Cluster_GroupAABBBuffer"));
