Fix CIS issue.

diff --git a/Engine/Plugins/Runtime/HairStrands/Source/HairStrandsEditor/Private/Widgets/SGroomEditorViewport.cpp b/Engine/Plugins/Runtime/HairStrands/Source/HairStrandsEditor/Private/Widgets/SGroomEditorViewport.cpp
--- a/Engine/Plugins/Runtime/HairStrands/Source/HairStrandsEditor/Private/Widgets/SGroomEditorViewport.cpp
+++ b/Engine/Plugins/Runtime/HairStrands/Source/HairStrandsEditor/Private/Widgets/SGroomEditorViewport.cpp
@@ -207,7 +207,7 @@ void SGroomEditorViewport::BindCommands()
 
 void SGroomEditorViewport::OnViewMode(EHairStrandsDebugMode Mode)
 {
-	if (GroomComponent && GroomComponent)
+	if (GroomComponent)
 	{
 		const EHairStrandsDebugMode OldMode = GroomComponent->GetDebugMode();
 		if (OldMode != Mode)
