Fix platform comment in hair code.

diff --git a/Engine/Plugins/Runtime/HairStrands/Source/HairStrandsCore/Private/GroomManager.cpp b/Engine/Plugins/Runtime/HairStrands/Source/HairStrandsCore/Private/GroomManager.cpp
--- a/Engine/Plugins/Runtime/HairStrands/Source/HairStrandsCore/Private/GroomManager.cpp
+++ b/Engine/Plugins/Runtime/HairStrands/Source/HairStrandsCore/Private/GroomManager.cpp
@@ -561,7 +561,7 @@ static void AddCopyHairStrandsPositionPass(
 	FHairStrandsRestResource& RestResources,
 	FHairStrandsDeformedResource& DeformedResources)
 {
-	// Copy using an explicit copy, as AddCopyBufferPass is not implemented on PS5
+	// Copy using an explicit copy instead of AddCopyBufferPass as some RHI didn't support it. This is now fix & should be fixed (TODO)
 #if 0
 	FRDGBufferRef RestBuffer = Register(GraphBuilder, RestResources.PositionBuffer, ERDGImportedBufferFlags::None).Buffer;
 	FRDGBufferRef Deformed0Buffer = Register(GraphBuilder, DeformedResources.DeformedPositionBuffer[0], ERDGImportedBufferFlags::None).Buffer;
diff --git a/Engine/Plugins/Runtime/HairStrands/Source/HairStrandsCore/Private/HairCardsVertexFactory.cpp b/Engine/Plugins/Runtime/HairStrands/Source/HairStrandsCore/Private/HairCardsVertexFactory.cpp
--- a/Engine/Plugins/Runtime/HairStrands/Source/HairStrandsCore/Private/HairCardsVertexFactory.cpp
+++ b/Engine/Plugins/Runtime/HairStrands/Source/HairStrandsCore/Private/HairCardsVertexFactory.cpp
@@ -14,7 +14,6 @@
 #include "GroomInstance.h"
 #include "SystemTextures.h" 
 
-// Disabling GPU scene data for this vertex factory as there are some Primitive index issue on PS4 which needs to be investigated
 #define VF_CARDS_SUPPORT_GPU_SCENE 0
 
 template<typename T> inline void VFC_BindParam(FMeshDrawSingleShaderBindings& ShaderBindings, const FShaderResourceParameter& Param, T* Value) { if (Param.IsBound() && Value) ShaderBindings.Add(Param, Value); }
