Add automatic raytracing resources allocation when path tracing mode is enabled.

diff --git a/Engine/Plugins/Runtime/HairStrands/Source/HairStrandsCore/Private/GroomComponent.cpp b/Engine/Plugins/Runtime/HairStrands/Source/HairStrandsCore/Private/GroomComponent.cpp
--- a/Engine/Plugins/Runtime/HairStrands/Source/HairStrandsCore/Private/GroomComponent.cpp
+++ b/Engine/Plugins/Runtime/HairStrands/Source/HairStrandsCore/Private/GroomComponent.cpp
@@ -528,6 +528,7 @@ public:
 		}
 
 		const bool bWireframe = AllowDebugViewmodes() && Context.ReferenceViewFamily.EngineShowFlags.Wireframe;
+		const uint32 ViewRayTracingMask = Context.ReferenceViewFamily.EngineShowFlags.PathTracing ? EHairViewRayTracingMask::PathTracing : EHairViewRayTracingMask::RayTracing;
 		if (bWireframe)
 			return;
 
@@ -543,11 +544,13 @@ public:
 
 			FHairStrandsRaytracingResource* RTGeometry = nullptr;
 			uint8 RayTracingMask = RAY_TRACING_MASK_OPAQUE;
+			uint32 InstanceViewRayTracingMask = EHairViewRayTracingMask::PathTracing | EHairViewRayTracingMask::RayTracing;
 			switch (GeometryType)
 			{
 				case EHairGeometryType::Strands:
 				{
 					RTGeometry = Instance->Strands.RenRaytracingResource;
+					InstanceViewRayTracingMask = Instance->Strands.ViewRayTracingMask;
 					RayTracingMask = RAY_TRACING_MASK_THIN_SHADOW;
 					break;
 				}
@@ -563,6 +566,12 @@ public:
 				}
 			}
 
+			// If the view and the instance raytracing mask don't match skip this instance.
+			if ((InstanceViewRayTracingMask & ViewRayTracingMask) == 0)
+			{
+				continue;
+			}
+
 			if (RTGeometry && RTGeometry->RayTracingGeometry.RayTracingGeometryRHI.IsValid())
 			{
 				for (const FRayTracingGeometrySegment& Segment : RTGeometry->RayTracingGeometry.Initializer.Segments)
@@ -2568,6 +2577,7 @@ void UGroomComponent::InitResources(bool bIsBindingReloading)
 			#if RHI_RAYTRACING
 			if (IsHairRayTracingEnabled() && HairGroupInstance->Strands.Modifier.bUseHairRaytracingGeometry && bVisibleInRayTracing)
 			{
+				HairGroupInstance->Strands.ViewRayTracingMask = EHairViewRayTracingMask::RayTracing | EHairViewRayTracingMask::PathTracing;
 				if (bNeedDynamicResources)
 				{
 					// Allocate dynamic raytracing resources (owned by the groom component/instance)
diff --git a/Engine/Plugins/Runtime/HairStrands/Source/HairStrandsCore/Private/GroomManager.cpp b/Engine/Plugins/Runtime/HairStrands/Source/HairStrandsCore/Private/GroomManager.cpp
--- a/Engine/Plugins/Runtime/HairStrands/Source/HairStrandsCore/Private/GroomManager.cpp
+++ b/Engine/Plugins/Runtime/HairStrands/Source/HairStrandsCore/Private/GroomManager.cpp
@@ -131,6 +131,12 @@ static void RunInternalHairStrandsInterpolation(
 {
 	check(IsInRenderingThread());
 
+	#if RHI_RAYTRACING
+	const uint32 ViewRayTracingMask = View->Family->EngineShowFlags.PathTracing ? EHairViewRayTracingMask::PathTracing : EHairViewRayTracingMask::RayTracing;
+	#else
+	const uint32 ViewRayTracingMask = 0u;
+	#endif
+
 	// Update dynamic mesh triangles
 	for (FHairStrandsInstance* AbstractInstance : Instances)
 	{
@@ -357,6 +363,7 @@ static void RunInternalHairStrandsInterpolation(
 				GraphBuilder, 
 				ShaderMap,
 				ViewUniqueID,
+				ViewRayTracingMask,
 				ShaderDrawData, 
 				ShaderPrintData,
 				Instance,
@@ -580,6 +587,22 @@ static void AddCopyHairStrandsPositionPass(
 #endif
 }
 
+#if RHI_RAYTRACING
+static void AllocateRaytracingResources(FHairGroupInstance* Instance)
+{
+	if (IsHairRayTracingEnabled() && !Instance->Strands.RenRaytracingResource)
+	{
+		check(Instance->Strands.Data);
+
+		// Allocate dynamic raytracing resources (owned by the groom component/instance)
+		FHairResourceName ResourceName(FName(Instance->Debug.GroomAssetName), Instance->Debug.GroupIndex);
+		Instance->Strands.RenRaytracingResource		 = new FHairStrandsRaytracingResource(*Instance->Strands.Data, ResourceName);
+		Instance->Strands.RenRaytracingResourceOwned = true;
+	}
+	Instance->Strands.ViewRayTracingMask |= EHairViewRayTracingMask::PathTracing;
+}
+#endif
+
 static void RunHairLODSelection(
 	FRDGBuilder& GraphBuilder, 
 	const FHairStrandsInstances& Instances, 
@@ -593,6 +616,19 @@ static void RunHairLODSelection(
 		ShaderPlatform = Views[0]->GetShaderPlatform();
 	}
 
+	// Detect if one of the current view has path tracing enabled for allocating raytracing resources
+#if RHI_RAYTRACING
+	bool bHasPathTracingView = false;
+	for (const FSceneView* View : Views)
+	{
+		if (View->Family->EngineShowFlags.PathTracing)
+		{
+			bHasPathTracingView = true;
+			break;
+		}
+	}
+#endif
+
 	for (FHairStrandsInstance* AbstractInstance : Instances)
 	{
 		FHairGroupInstance* Instance = static_cast<FHairGroupInstance*>(AbstractInstance);
@@ -753,6 +789,7 @@ static void RunHairLODSelection(
 			if (Instance->Strands.DeformedRootResource)		{ Instance->Strands.DeformedRootResource->Allocate(GraphBuilder, LoadingType, ResourceStatus); Instance->Strands.DeformedRootResource->AllocateLOD(GraphBuilder, MeshLODIndex, LoadingType, ResourceStatus); }
 			if (Instance->Strands.DeformedResource)			{ Instance->Strands.DeformedResource->Allocate(GraphBuilder, LoadingType, ResourceStatus); }
 			#if RHI_RAYTRACING
+			if (bHasPathTracingView)						{ AllocateRaytracingResources(Instance); }
 			if (Instance->Strands.RenRaytracingResource)	{ Instance->Strands.RenRaytracingResource->Allocate(GraphBuilder, LoadingType, ResourceStatus); }
 			#endif
 			Instance->Strands.VertexFactory->InitResources();
diff --git a/Engine/Plugins/Runtime/HairStrands/Source/HairStrandsCore/Private/HairStrandsRendering.cpp b/Engine/Plugins/Runtime/HairStrands/Source/HairStrandsCore/Private/HairStrandsRendering.cpp
--- a/Engine/Plugins/Runtime/HairStrands/Source/HairStrandsCore/Private/HairStrandsRendering.cpp
+++ b/Engine/Plugins/Runtime/HairStrands/Source/HairStrandsCore/Private/HairStrandsRendering.cpp
@@ -1687,10 +1687,12 @@ static FHairGroupPublicData::FVertexFactoryInput InternalComputeHairStrandsVerte
 	OutVFInput.Strands.HairDensity = Instance->Strands.Modifier.HairShadowDensity;
 	OutVFInput.Strands.bScatterSceneLighting = Instance->Strands.Modifier.bScatterSceneLighting;
 	OutVFInput.Strands.bUseStableRasterization = Instance->Strands.Modifier.bUseStableRasterization;
-#if RHI_RAYTRACING
-	OutVFInput.Strands.bUseRaytracingGeometry = Instance->Strands.RenRaytracingResource != nullptr;
-#else
 	OutVFInput.Strands.bUseRaytracingGeometry = false;
+#if RHI_RAYTRACING
+	// Flag bUseRaytracingGeometry only if RT geometry has been allocated for RayTracing view (not for PathTracing view). 
+	// This flag is used later for selecting if voxelization needs to flags voxel has shadow caster or if shadow casting is handled 
+	// by the RT geometry.
+	OutVFInput.Strands.bUseRaytracingGeometry = Instance->Strands.RenRaytracingResource != nullptr && (Instance->Strands.ViewRayTracingMask & EHairViewRayTracingMask::RayTracing) != 0;
 #endif
 
 	return OutVFInput;
@@ -1707,6 +1709,7 @@ void ComputeHairStrandsInterpolation(
 	FRDGBuilder& GraphBuilder,
 	FGlobalShaderMap* ShaderMap,
 	const uint32 ViewUniqueID,
+	const uint32 ViewRayTracingMask,
 	const FShaderDrawDebugData* ShaderDrawData,
 	const FShaderPrintData* ShaderPrintData,
 	FHairGroupInstance* Instance,
@@ -1965,9 +1968,9 @@ void ComputeHairStrandsInterpolation(
 				}
 			}
 
-			// 4. Update raytracing geometry
+			// 4. Update raytracing geometry (update only if the view mask and the RT geometry mask match)
 			#if RHI_RAYTRACING
-			if (Instance->Strands.RenRaytracingResource)
+			if (Instance->Strands.RenRaytracingResource && (Instance->Strands.ViewRayTracingMask & ViewRayTracingMask) != 0)
 			{
 				const float HairRadiusRT	= Instance->HairGroupPublicData->VFInput.Strands.HairRaytracingRadiusScale * Instance->HairGroupPublicData->VFInput.Strands.HairRadius;
 				const float HairRootScaleRT = Instance->HairGroupPublicData->VFInput.Strands.HairRootScale;
diff --git a/Engine/Plugins/Runtime/HairStrands/Source/HairStrandsCore/Public/GroomInstance.h b/Engine/Plugins/Runtime/HairStrands/Source/HairStrandsCore/Public/GroomInstance.h
--- a/Engine/Plugins/Runtime/HairStrands/Source/HairStrandsCore/Public/GroomInstance.h
+++ b/Engine/Plugins/Runtime/HairStrands/Source/HairStrandsCore/Public/GroomInstance.h
@@ -77,6 +77,12 @@ enum class EHairLODSelectionType
 	Forced		// Forced LOD value
 };
 
+enum EHairViewRayTracingMask
+{
+	RayTracing  = 0x1, // Visible for raytracing effects (RT shadow, RT refleciton, Lumen, ...)
+	PathTracing = 0x2, // Visible for pathtracing rendering
+};
+
 // Represent/Describe data & resources of a hair group belonging to a groom
 struct HAIRSTRANDSCORE_API FHairGroupInstance : public FHairStrandsInstance
 {
@@ -135,6 +141,7 @@ struct HAIRSTRANDSCORE_API FHairGroupInstance : public FHairStrandsInstance
 		#if RHI_RAYTRACING
 		FHairStrandsRaytracingResource* RenRaytracingResource = nullptr;
 		bool RenRaytracingResourceOwned = false;
+		uint32 ViewRayTracingMask = 0u;
 		float CachedHairScaledRadius = 0;
 		float CachedHairRootScale = 0;
 		float CachedHairTipScale = 0;
diff --git a/Engine/Plugins/Runtime/HairStrands/Source/HairStrandsCore/Public/HairStrandsRendering.h b/Engine/Plugins/Runtime/HairStrands/Source/HairStrandsCore/Public/HairStrandsRendering.h
--- a/Engine/Plugins/Runtime/HairStrands/Source/HairStrandsCore/Public/HairStrandsRendering.h
+++ b/Engine/Plugins/Runtime/HairStrands/Source/HairStrandsCore/Public/HairStrandsRendering.h
@@ -27,6 +27,7 @@ void ComputeHairStrandsInterpolation(
 	FRDGBuilder& GraphBuilder,
 	FGlobalShaderMap* ShaderMap,
 	const uint32 ViewUniqueID,
+	const uint32 ViewRayTracingMask,
 	const struct FShaderDrawDebugData* DebugShaderData,
 	const struct FShaderPrintData* ShaderPrintData,
 	struct FHairGroupInstance* Instance,
