Change default hair interpolation to smooth as this is the common choice used by artists.

diff --git a/Engine/Plugins/Runtime/HairStrands/Source/HairStrandsCore/Public/GroomAsset.h b/Engine/Plugins/Runtime/HairStrands/Source/HairStrandsCore/Public/GroomAsset.h
--- a/Engine/Plugins/Runtime/HairStrands/Source/HairStrandsCore/Public/GroomAsset.h
+++ b/Engine/Plugins/Runtime/HairStrands/Source/HairStrandsCore/Public/GroomAsset.h
@@ -449,7 +449,7 @@ public:
 
 	/** Type of interpolation used */
 	UPROPERTY(EditAnywhere, EditFixedSize, BlueprintReadWrite, Category = "HairInterpolation", meta = (ToolTip = "Type of interpolation used (WIP)"))
-	EGroomInterpolationType HairInterpolationType = EGroomInterpolationType::RigidTransform;
+	EGroomInterpolationType HairInterpolationType = EGroomInterpolationType::SmoothTransform;
 
 	/** Minimum LOD to cook */
 	UPROPERTY(EditAnywhere, Category = "LOD", meta = (DisplayName = "Minimum LOD"))
