Add imported (DCC) hair width in groom asset.

diff --git a/Engine/Plugins/Runtime/HairStrands/Source/HairStrandsCore/Private/GroomBuilder.cpp b/Engine/Plugins/Runtime/HairStrands/Source/HairStrandsCore/Private/GroomBuilder.cpp
--- a/Engine/Plugins/Runtime/HairStrands/Source/HairStrandsCore/Private/GroomBuilder.cpp
+++ b/Engine/Plugins/Runtime/HairStrands/Source/HairStrandsCore/Private/GroomBuilder.cpp
@@ -33,7 +33,7 @@ static FAutoConsoleVariableRef CVarHairClusterBuilder_MaxVoxelResolution(TEXT("r
 FString FGroomBuilder::GetVersion()
 {
 	// Important to update the version when groom building changes
-	return TEXT("2r1");
+	return TEXT("2r2");
 }
 
 namespace FHairStrandsDecimation
@@ -1511,9 +1511,13 @@ bool FGroomBuilder::BuildHairDescriptionGroups(const FHairDescription& HairDescr
 		FHairDescriptionGroup& Group = Out.HairGroups.AddDefaulted_GetRef();
 		Group.Info.GroupID = GroupID;
 		Group.Info.GroupName = GroupName != NAME_None ? GroupName : FName(FString::Printf(TEXT("Group_%d"), GroupID));
+		Group.Info.MaxImportedWidth = -1.0f; // Invalid width
 		return Group;
 	};
 
+	// Track the imported max hair width among all imported CVs. This information is displayed to artsits to set/tune groom asset later.
+	float DDCMaxHairWidth = -1.f;
+
 	int32 GlobalVertexIndex = 0;
 	int32 NumHairPoints = 0;
 	int32 NumGuidePoints = 0;
@@ -1588,10 +1592,25 @@ bool FGroomBuilder::BuildHairDescriptionGroups(const FHairDescription& HairDescr
 			CurrentHairStrandsDatas->StrandsCurves.CurvesRootUV.Add(StrandRootUV[StrandID]);
 		}
 
-		float StrandWidth = GroomHairWidth ? GroomHairWidth.GetValue() : 0.01f;
+		// Groom 
+		float StrandWidth = 0.01f;
+		if (GroomHairWidth) 
+		{
+			StrandWidth = GroomHairWidth.GetValue();
+			if (!bIsGuide)
+			{
+				Group.Info.MaxImportedWidth = FMath::Max(Group.Info.MaxImportedWidth, StrandWidth);
+			}
+		}
+
+		// Curve
 		if (StrandWidths.IsValid())
 		{
 			StrandWidth = StrandWidths[StrandID];
+			if (!bIsGuide)
+			{
+				Group.Info.MaxImportedWidth = FMath::Max(Group.Info.MaxImportedWidth, StrandWidth);
+			}
 		}
 
 		for (int32 VertexIndex = 0; VertexIndex < CurveNumVertices; ++VertexIndex, ++GlobalVertexIndex)
@@ -1602,10 +1621,15 @@ bool FGroomBuilder::BuildHairDescriptionGroups(const FHairDescription& HairDescr
 			CurrentHairStrandsDatas->StrandsPoints.PointsBaseColor.Add(bHasBaseColorAttribute ? VertexBaseColor[VertexID] : (GroomHairColor     ? GroomHairColor.GetValue() : FLinearColor::Black));
 			CurrentHairStrandsDatas->StrandsPoints.PointsRoughness.Add(bHasRoughnessAttribute ? VertexRoughness[VertexID] : (GroomHairRoughness ? GroomHairRoughness.GetValue() : 0.f));
 
+			// Vertex
 			float VertexWidth = 0.f;
 			if (VertexWidths.IsValid())
 			{
 				VertexWidth = VertexWidths[VertexID];
+				if (!bIsGuide)
+				{
+					Group.Info.MaxImportedWidth = FMath::Max(Group.Info.MaxImportedWidth, VertexWidth);
+				}
 			}
 			else if (StrandWidth != 0.f)
 			{
@@ -1702,6 +1726,7 @@ void FGroomBuilder::BuildData(
 	{
 		OutGroupInfo.GroupID = InHairDescriptionGroup.Info.GroupID;
 		OutGroupInfo.GroupName = InHairDescriptionGroup.Info.GroupName;
+		OutGroupInfo.MaxImportedWidth = InHairDescriptionGroup.Info.MaxImportedWidth;
 
 		// Rendering data
 		{
diff --git a/Engine/Plugins/Runtime/HairStrands/Source/HairStrandsCore/Public/GroomAsset.h b/Engine/Plugins/Runtime/HairStrands/Source/HairStrandsCore/Public/GroomAsset.h
--- a/Engine/Plugins/Runtime/HairStrands/Source/HairStrandsCore/Public/GroomAsset.h
+++ b/Engine/Plugins/Runtime/HairStrands/Source/HairStrandsCore/Public/GroomAsset.h
@@ -57,6 +57,9 @@ struct HAIRSTRANDSCORE_API FHairGroupInfo
 
 	UPROPERTY(VisibleAnywhere, Category = "Info", meta = (DisplayName = "Length of the longest hair strands"))
 	float MaxCurveLength = 0;
+
+	UPROPERTY(VisibleAnywhere, Category = "Info", meta = (DisplayName = "Imported width (DCC) of the larget hair strands"))
+	float MaxImportedWidth = -1.f;
 };
 
 USTRUCT(BlueprintType)
diff --git a/Engine/Plugins/Runtime/HairStrands/Source/HairStrandsEditor/Private/GroomAssetDetails.cpp b/Engine/Plugins/Runtime/HairStrands/Source/HairStrandsEditor/Private/GroomAssetDetails.cpp
--- a/Engine/Plugins/Runtime/HairStrands/Source/HairStrandsEditor/Private/GroomAssetDetails.cpp
+++ b/Engine/Plugins/Runtime/HairStrands/Source/HairStrandsEditor/Private/GroomAssetDetails.cpp
@@ -201,6 +201,29 @@ TSharedRef<SUniformGridPanel> MakeHairStrandsInfoGrid(const FSlateFontInfo& Deta
 		.Font(DetailFontInfo)
 		.Text(FText::AsNumber(MaxRadius*2.0f))
 	];
+
+	// Imported Width (mm)
+	Grid->AddSlot(0, 4) // x, y
+	.HAlign(HAlign_Left)
+	[
+		SNew(STextBlock)
+		.Font(DetailFontInfo)
+	.Text(LOCTEXT("HairInfo_ImportedWidth", "Max. Imported Width"))
+	];
+	Grid->AddSlot(1, 4) // x, y
+	.HAlign(HAlign_Right)
+	[
+		SNew(STextBlock)
+		.Font(DetailFontInfo)
+	.Text(LOCTEXT("HairInfo_GuideImportedWidth", ""))
+	];
+	Grid->AddSlot(2, 4) // x, y
+	.HAlign(HAlign_Right)
+	[
+		SNew(STextBlock)
+		.Font(DetailFontInfo)
+		.Text(CurrentAsset.MaxImportedWidth > 0.f ? FText::AsNumber(CurrentAsset.MaxImportedWidth) : LOCTEXT("HairInfo_ImportedWidthDefault", "Not exported"))
+	];
 	return Grid;
 }
 
