Add CVs visualization

diff --git a/Engine/Plugins/Runtime/HairStrands/Source/HairStrandsCore/Private/GroomComponent.cpp b/Engine/Plugins/Runtime/HairStrands/Source/HairStrandsCore/Private/GroomComponent.cpp
--- a/Engine/Plugins/Runtime/HairStrands/Source/HairStrandsCore/Private/GroomComponent.cpp
+++ b/Engine/Plugins/Runtime/HairStrands/Source/HairStrandsCore/Private/GroomComponent.cpp
@@ -645,6 +645,7 @@ public:
 					
 					const bool bNeedDebugMaterial = 
 						DebugMode != EHairStrandsDebugMode::NoneDebug &&
+						DebugMode != EHairStrandsDebugMode::RenderHairControlPoints &&
 						DebugMode != EHairStrandsDebugMode::RenderHairTangent;
 
 					if (bNeedDebugMaterial)
diff --git a/Engine/Plugins/Runtime/HairStrands/Source/HairStrandsCore/Private/GroomDebug.cpp b/Engine/Plugins/Runtime/HairStrands/Source/HairStrandsCore/Private/GroomDebug.cpp
--- a/Engine/Plugins/Runtime/HairStrands/Source/HairStrandsCore/Private/GroomDebug.cpp
+++ b/Engine/Plugins/Runtime/HairStrands/Source/HairStrandsCore/Private/GroomDebug.cpp
@@ -51,6 +51,9 @@ static int32 GHairCardsGuidesDebug_Sim = 0;
 static FAutoConsoleVariableRef CVarHairCardsGuidesDebug_Ren(TEXT("r.HairStrands.Cards.DebugGuides.Render"), GHairCardsGuidesDebug_Ren, TEXT("Draw debug hair cards guides (1: Rest, 2: Deformed)."));
 static FAutoConsoleVariableRef CVarHairCardsGuidesDebug_Sim(TEXT("r.HairStrands.Cards.DebugGuides.Sim"), GHairCardsGuidesDebug_Sim, TEXT("Draw debug hair sim guides (1: Rest, 2: Deformed)."));
 
+static int32 GHairStrandsControlPointDebug = 0;
+static FAutoConsoleVariableRef CVarHairStrandsControlPointDebug(TEXT("r.HairStrands.Strands.DebugControlPoint"), GHairStrandsControlPointDebug, TEXT("Draw debug hair strands control points)."));
+
 ///////////////////////////////////////////////////////////////////////////////////////////////////
 
 bool IsHairStrandsSkinCacheEnable();
@@ -250,7 +253,7 @@ static void AddDebugProjectionMeshPass(
 	PSParameters.Pass = *Parameters;
 
 	GraphBuilder.AddPass(
-		RDG_EVENT_NAME("HairStrandsMeshProjectionMeshDebug"),
+		RDG_EVENT_NAME("HairStrands::MeshProjectionMeshDebug"),
 		Parameters,
 		ERDGPassFlags::Raster,
 		[Parameters, VSParameters, PSParameters, VertexShader, PixelShader, Viewport, Resolution, PrimitiveCount, PrimitiveType](FRHICommandList& RHICmdList)
@@ -426,7 +429,7 @@ static void AddDebugProjectionHairPass(
 	PSParameters.Pass = *Parameters;
 
 	GraphBuilder.AddPass(
-		RDG_EVENT_NAME("HairStrandsMeshProjectionHairDebug"),
+		RDG_EVENT_NAME("HairStrands::MeshProjectionHairDebug"),
 		Parameters,
 		ERDGPassFlags::Raster,
 		[Parameters, VSParameters, PSParameters, VertexShader, PixelShader, Viewport, Resolution, PrimitiveCount, PrimitiveType](FRHICommandList& RHICmdList)
@@ -537,7 +540,7 @@ static void AddVoxelPlainRaymarchingPass(
 
 		TShaderMapRef<FVoxelPlainRaymarchingCS> ComputeShader(ShaderMap);
 		const FIntVector DispatchCount = DispatchCount.DivideAndRoundUp(FIntVector(OutputTexture->Desc.Extent.X, OutputTexture->Desc.Extent.Y, 1), FIntVector(8, 8, 1));
-		FComputeShaderUtils::AddPass(GraphBuilder, RDG_EVENT_NAME("HairStrandsVoxelPlainRaymarching"), ComputeShader, Parameters, DispatchCount);
+		FComputeShaderUtils::AddPass(GraphBuilder, RDG_EVENT_NAME("HairStrands::VoxelPlainRaymarching"), ComputeShader, Parameters, DispatchCount);
 	}
 #endif
 }
@@ -618,7 +621,7 @@ static void AddDrawDebugCardsAtlasPass(
 
 		ShaderDrawDebug::SetParameters(GraphBuilder, *ShaderDrawData, Parameters->ShaderDrawParameters);
 
-		FComputeShaderUtils::AddPass(GraphBuilder, RDG_EVENT_NAME("DrawDebugCardsAtlas"), ComputeShader, Parameters,
+		FComputeShaderUtils::AddPass(GraphBuilder, RDG_EVENT_NAME("HairStrands::DrawDebugCardsAtlas"), ComputeShader, Parameters,
 		FIntVector::DivideAndRoundUp(FIntVector(Parameters->OutputResolution.X, Parameters->OutputResolution.Y, 1), FIntVector(8, 8, 1)));
 
 	}
@@ -626,6 +629,72 @@ static void AddDrawDebugCardsAtlasPass(
 
 ///////////////////////////////////////////////////////////////////////////////////////////////////
 
+class FDrawDebugStrandsCVsCS : public FGlobalShader
+{
+	DECLARE_GLOBAL_SHADER(FDrawDebugStrandsCVsCS);
+	SHADER_USE_PARAMETER_STRUCT(FDrawDebugStrandsCVsCS, FGlobalShader);
+
+	BEGIN_SHADER_PARAMETER_STRUCT(FParameters, )
+		SHADER_PARAMETER_STRUCT_REF(FViewUniformShaderParameters, ViewUniformBuffer)
+		SHADER_PARAMETER(uint32, MaxVertexCount)
+		SHADER_PARAMETER(FMatrix44f, LocalToWorld)
+		SHADER_PARAMETER_SAMPLER(SamplerState, LinearSampler)
+		SHADER_PARAMETER_RDG_TEXTURE(Texture2D<float>, DepthTexture)
+		SHADER_PARAMETER_RDG_TEXTURE_UAV(RWTexture2D<float4>, ColorTexture)
+		SHADER_PARAMETER_STRUCT_REF(FHairStrandsVertexFactoryUniformShaderParameters, HairStrandsVF)
+	END_SHADER_PARAMETER_STRUCT()
+
+public:
+	static bool ShouldCompilePermutation(const FGlobalShaderPermutationParameters& Parameters) { return IsHairStrandsSupported(EHairStrandsShaderType::Tool, Parameters.Platform); }
+	static void ModifyCompilationEnvironment(const FGlobalShaderPermutationParameters& Parameters, FShaderCompilerEnvironment& OutEnvironment)
+	{
+		FGlobalShader::ModifyCompilationEnvironment(Parameters, OutEnvironment);
+		OutEnvironment.SetDefine(TEXT("SHADER_CVS"), 1);
+	}
+};
+
+IMPLEMENT_GLOBAL_SHADER(FDrawDebugStrandsCVsCS, "/Engine/Private/HairStrands/HairStrandsDebug.usf", "MainCS", SF_Compute);
+
+static void AddDrawDebugStrandsCVsPass(
+	FRDGBuilder& GraphBuilder,
+	const FSceneView& View,
+	FGlobalShaderMap* ShaderMap,
+	const FHairGroupInstance* Instance,
+	const FShaderDrawDebugData* ShaderDrawData,
+	FRDGTextureRef ColorTexture,
+	FRDGTextureRef DepthTexture)
+{
+	if (Instance->HairGroupPublicData->VFInput.GeometryType != EHairGeometryType::Strands || ShaderDrawData == nullptr)
+	{
+		return;
+	}
+
+	if (!Instance->Strands.IsValid())
+	{
+		return;
+	}
+
+	TShaderMapRef<FDrawDebugStrandsCVsCS> ComputeShader(ShaderMap);
+	FDrawDebugStrandsCVsCS::FParameters* Parameters = GraphBuilder.AllocParameters<FDrawDebugStrandsCVsCS::FParameters>();
+	Parameters->ViewUniformBuffer = View.ViewUniformBuffer;
+	Parameters->HairStrandsVF = Instance->Strands.UniformBuffer;
+	Parameters->LocalToWorld = Instance->LocalToWorld.ToMatrixWithScale();
+	Parameters->MaxVertexCount = Instance->Strands.Data->PointCount;
+	Parameters->ColorTexture = GraphBuilder.CreateUAV(ColorTexture);
+	Parameters->DepthTexture = DepthTexture;
+	Parameters->LinearSampler = TStaticSamplerState<SF_Bilinear, AM_Clamp, AM_Clamp, AM_Clamp>::GetRHI();
+
+	const uint32 VertexCount = Instance->HairGroupPublicData->VFInput.Strands.VertexCount;
+	FComputeShaderUtils::AddPass(
+		GraphBuilder, 
+		RDG_EVENT_NAME("HairStrands::DrawCVs"), 
+		ComputeShader, 
+		Parameters,
+		FIntVector::DivideAndRoundUp(FIntVector(VertexCount, 1, 1), FIntVector(256, 1, 1)));
+}
+
+///////////////////////////////////////////////////////////////////////////////////////////////////
+
 class FDrawDebugCardGuidesCS : public FGlobalShader
 {
 	DECLARE_GLOBAL_SHADER(FDrawDebugCardGuidesCS);
@@ -824,6 +893,7 @@ void RunHairStrandsDebug(
 	const FGPUSkinCache* SkinCache,
 	const FShaderDrawDebugData* ShaderDrawData,
 	FRDGTextureRef SceneColorTexture,
+	FRDGTextureRef SceneDepthTexture,
 	FIntRect Viewport,
 	const TUniformBufferRef<FViewUniformShaderParameters>& ViewUniformBuffer)
 {
@@ -1091,8 +1161,18 @@ void RunHairStrandsDebug(
 		FHairGroupInstance* Instance = static_cast<FHairGroupInstance*>(AbstractInstance);
 
 		if (GHairCardsGuidesDebug_Ren > 0 || Instance->Debug.bDrawCardsGuides)
+		{
 			AddDrawDebugCardsGuidesPass(GraphBuilder, View, ShaderMap, Instance, ShaderDrawData, Instance->Debug.bDrawCardsGuides ? false : GHairCardsGuidesDebug_Ren == 1, true);
+		}
+
 		if (GHairCardsGuidesDebug_Sim > 0)
+		{
 			AddDrawDebugCardsGuidesPass(GraphBuilder, View, ShaderMap, Instance, ShaderDrawData, GHairCardsGuidesDebug_Sim == 1, false);
+		}
+
+		if (GHairStrandsControlPointDebug || Instance->HairGroupPublicData->DebugMode == EHairStrandsDebugMode::RenderHairControlPoints)
+		{
+			AddDrawDebugStrandsCVsPass(GraphBuilder, View, ShaderMap, Instance, ShaderDrawData, SceneColorTexture, SceneDepthTexture);
+		}
 	}
 }
\ No newline at end of file
diff --git a/Engine/Plugins/Runtime/HairStrands/Source/HairStrandsCore/Private/GroomManager.cpp b/Engine/Plugins/Runtime/HairStrands/Source/HairStrandsCore/Private/GroomManager.cpp
--- a/Engine/Plugins/Runtime/HairStrands/Source/HairStrandsCore/Private/GroomManager.cpp
+++ b/Engine/Plugins/Runtime/HairStrands/Source/HairStrandsCore/Private/GroomManager.cpp
@@ -876,6 +876,7 @@ void RunHairStrandsDebug(
 	const FGPUSkinCache* SkinCache,
 	const FShaderDrawDebugData* ShaderDrawData,
 	FRDGTextureRef SceneColor,
+	FRDGTextureRef SceneDepth,
 	FIntRect Viewport,
 	const TUniformBufferRef<FViewUniformShaderParameters>& ViewUniformBuffer);
 
@@ -981,6 +982,7 @@ void ProcessHairStrandsBookmark(
 			Parameters.SkinCache,
 			Parameters.ShaderDebugData,
 			Parameters.SceneColorTexture,
+			Parameters.SceneDepthTexture,
 			Parameters.View->UnscaledViewRect,
 			Parameters.View->ViewUniformBuffer);
 	}
diff --git a/Engine/Plugins/Runtime/HairStrands/Source/HairStrandsEditor/Private/GroomEditorCommands.cpp b/Engine/Plugins/Runtime/HairStrands/Source/HairStrandsEditor/Private/GroomEditorCommands.cpp
--- a/Engine/Plugins/Runtime/HairStrands/Source/HairStrandsEditor/Private/GroomEditorCommands.cpp
+++ b/Engine/Plugins/Runtime/HairStrands/Source/HairStrandsEditor/Private/GroomEditorCommands.cpp
@@ -38,6 +38,7 @@ void FGroomEditorCommands::RegisterCommands()
 	UI_COMMAND(ViewMode_Tangent,		"Tangent",			"View tangents direction",		EUserInterfaceActionType::Button, FInputChord());
 	UI_COMMAND(ViewMode_BaseColor,		"BaseColor",		"View vertex base color",		EUserInterfaceActionType::Button, FInputChord());
 	UI_COMMAND(ViewMode_Roughness,		"Roughness",		"View vertex roughness",		EUserInterfaceActionType::Button, FInputChord());
+	UI_COMMAND(ViewMode_ControlPoints,	"Hair CVs",			"View hair strand CVs",			EUserInterfaceActionType::Button, FInputChord());
 	UI_COMMAND(ViewMode_VisCluster,		"Vis. Clusters",	"View clusters (only available if the groom has simulation enabled)",				EUserInterfaceActionType::Button, FInputChord());
 	UI_COMMAND(ViewMode_Group,			"Hair Groups",		"View hair groups",				EUserInterfaceActionType::Button, FInputChord());
 
diff --git a/Engine/Plugins/Runtime/HairStrands/Source/HairStrandsEditor/Private/GroomEditorViewportToolBar.cpp b/Engine/Plugins/Runtime/HairStrands/Source/HairStrandsEditor/Private/GroomEditorViewportToolBar.cpp
--- a/Engine/Plugins/Runtime/HairStrands/Source/HairStrandsEditor/Private/GroomEditorViewportToolBar.cpp
+++ b/Engine/Plugins/Runtime/HairStrands/Source/HairStrandsEditor/Private/GroomEditorViewportToolBar.cpp
@@ -61,6 +61,7 @@ TSharedRef<SWidget> SGroomEditorViewportToolbar::GenerateShowMenu() const
 
 		ViewMenuBuilder.AddMenuSeparator();
 
+		ViewMenuBuilder.AddMenuEntry(Commands.ViewMode_ControlPoints);
 		ViewMenuBuilder.AddMenuEntry(Commands.ViewMode_VisCluster);
 
 		ViewMenuBuilder.AddMenuSeparator();
diff --git a/Engine/Plugins/Runtime/HairStrands/Source/HairStrandsEditor/Private/Widgets/SGroomEditorViewport.cpp b/Engine/Plugins/Runtime/HairStrands/Source/HairStrandsEditor/Private/Widgets/SGroomEditorViewport.cpp
--- a/Engine/Plugins/Runtime/HairStrands/Source/HairStrandsEditor/Private/Widgets/SGroomEditorViewport.cpp
+++ b/Engine/Plugins/Runtime/HairStrands/Source/HairStrandsEditor/Private/Widgets/SGroomEditorViewport.cpp
@@ -200,6 +200,7 @@ void SGroomEditorViewport::BindCommands()
 	CommandList->MapAction(FGroomEditorCommands::Get().ViewMode_Tangent,			FExecuteAction::CreateSP(this, &SGroomEditorViewport::OnViewMode, EHairStrandsDebugMode::RenderHairTangent),		FCanExecuteAction::CreateSP(this, &SGroomEditorViewport::CanViewMode, false));
 	CommandList->MapAction(FGroomEditorCommands::Get().ViewMode_BaseColor,			FExecuteAction::CreateSP(this, &SGroomEditorViewport::OnViewMode, EHairStrandsDebugMode::RenderHairBaseColor),		FCanExecuteAction::CreateSP(this, &SGroomEditorViewport::CanViewMode, false));
 	CommandList->MapAction(FGroomEditorCommands::Get().ViewMode_Roughness,			FExecuteAction::CreateSP(this, &SGroomEditorViewport::OnViewMode, EHairStrandsDebugMode::RenderHairRoughness),		FCanExecuteAction::CreateSP(this, &SGroomEditorViewport::CanViewMode, false));
+	CommandList->MapAction(FGroomEditorCommands::Get().ViewMode_ControlPoints,		FExecuteAction::CreateSP(this, &SGroomEditorViewport::OnViewMode, EHairStrandsDebugMode::RenderHairControlPoints),	FCanExecuteAction::CreateSP(this, &SGroomEditorViewport::CanViewMode, false));
 	CommandList->MapAction(FGroomEditorCommands::Get().ViewMode_VisCluster,			FExecuteAction::CreateSP(this, &SGroomEditorViewport::OnViewMode, EHairStrandsDebugMode::RenderVisCluster),			FCanExecuteAction::CreateSP(this, &SGroomEditorViewport::CanViewMode, true)); // While Clusters LOD/Culling doesn't require simulation, cluster visualization requires HairInterpolation to run, which is only available when simulation is enabled in the Groom editor.
 	CommandList->MapAction(FGroomEditorCommands::Get().ViewMode_Group,				FExecuteAction::CreateSP(this, &SGroomEditorViewport::OnViewMode, EHairStrandsDebugMode::RenderHairGroup),			FCanExecuteAction::CreateSP(this, &SGroomEditorViewport::CanViewMode, false));
 	CommandList->MapAction(FGroomEditorCommands::Get().ViewMode_CardsGuides,		FExecuteAction::CreateSP(this, &SGroomEditorViewport::OnCardsGuides),												FCanExecuteAction::CreateSP(this, &SGroomEditorViewport::CanCardsGuides));
diff --git a/Engine/Plugins/Runtime/HairStrands/Source/HairStrandsEditor/Public/GroomEditorCommands.h b/Engine/Plugins/Runtime/HairStrands/Source/HairStrandsEditor/Public/GroomEditorCommands.h
--- a/Engine/Plugins/Runtime/HairStrands/Source/HairStrandsEditor/Public/GroomEditorCommands.h
+++ b/Engine/Plugins/Runtime/HairStrands/Source/HairStrandsEditor/Public/GroomEditorCommands.h
@@ -32,6 +32,7 @@ public:
 	TSharedPtr<FUICommandInfo> ViewMode_Tangent;
 	TSharedPtr<FUICommandInfo> ViewMode_BaseColor;
 	TSharedPtr<FUICommandInfo> ViewMode_Roughness;
+	TSharedPtr<FUICommandInfo> ViewMode_ControlPoints;
 	TSharedPtr<FUICommandInfo> ViewMode_VisCluster;
 	TSharedPtr<FUICommandInfo> ViewMode_Group;
 
diff --git a/Engine/Shaders/Private/HairStrands/HairStrandsDebug.usf b/Engine/Shaders/Private/HairStrands/HairStrandsDebug.usf
--- a/Engine/Shaders/Private/HairStrands/HairStrandsDebug.usf
+++ b/Engine/Shaders/Private/HairStrands/HairStrandsDebug.usf
@@ -4,6 +4,10 @@
 #include "../ColorMap.ush"
 #include "HairStrandsVisibilityCommon.ush"
 
+/////////////////////////////////////////////////////////////////////////////////////////
+
+#if SHADER_DEBUG_MODE
+
 float2 OutputResolution;
 uint FastResolveMask;
 uint DebugMode;
@@ -141,3 +145,60 @@ void MainPS(
 
 	OutColor = float4(Color, Alpha);
 }
+
+#endif // SHADER_DEBUG_MODE
+
+/////////////////////////////////////////////////////////////////////////////////////////
+
+#if SHADER_CVS
+#include "HairStrandsVertexFactoryCommon.ush"
+
+float4x4 LocalToWorld;
+uint MaxVertexCount;
+
+SamplerState LinearSampler;
+Texture2D<float> DepthTexture;
+RWTexture2D<float4> ColorTexture;
+
+[numthreads(256, 1, 1)]
+void MainCS(uint3 DispatchThreadId : SV_DispatchThreadID)
+{
+	const uint VertexId = DispatchThreadId.x;
+	if (VertexId > MaxVertexCount)
+		return;
+
+	uint VertexIndex = VertexId;
+	float ScaleR = 1;
+	if (HairStrandsVF.CullingEnable)
+	{
+		ScaleR = HairStrandsVF.CulledVertexRadiusScaleBuffer[VertexIndex];
+		VertexIndex = HairStrandsVF.CulledVertexIdsBuffer[VertexIndex];
+	}
+
+	const float3 P_Local = UnpackHairControlPoint(
+		HairStrandsVF.PositionBuffer[VertexIndex],
+		HairStrandsVF.PositionOffsetBuffer[0].xyz,
+		HairStrandsVF.Radius * ScaleR,
+		HairStrandsVF.RootScale,
+		HairStrandsVF.TipScale).Position;
+
+	const float3 P_World = mul(float4(P_Local, 1), LocalToWorld).xyz;
+	float4 SamplePosition = float4(P_World + LWCHackToFloat(PrimaryView.PreViewTranslation), 1);
+	float4 ClipPosition = mul(SamplePosition, View.TranslatedWorldToClip);
+	float2 ScreenPosition = ClipPosition.xy / ClipPosition.w;
+	float2 ScreenUV = ScreenPosition * View.ScreenPositionScaleBias.xy + View.ScreenPositionScaleBias.wz;
+	float SampleSceneDepth = ClipPosition.w;
+
+	const float S = 1;
+	const float4 Color = float4(S, S, 0, 1);
+	const float SceneDepth = ConvertFromDeviceZ(DepthTexture.SampleLevel(LinearSampler, ScreenUV, 0));
+	const bool bIsVisible = SampleSceneDepth <= SceneDepth;
+	if (bIsVisible)
+	{
+		int2 OutCoord = ScreenUV * View.BufferSizeAndInvSize.xy;
+		OutCoord = clamp(OutCoord, int2(0, 0), View.BufferSizeAndInvSize.xy);
+		ColorTexture[OutCoord] = Color;
+	}
+}
+
+#endif // SHADER_CVS
\ No newline at end of file
diff --git a/Engine/Source/Runtime/Renderer/Private/DeferredShadingRenderer.cpp b/Engine/Source/Runtime/Renderer/Private/DeferredShadingRenderer.cpp
--- a/Engine/Source/Runtime/Renderer/Private/DeferredShadingRenderer.cpp
+++ b/Engine/Source/Runtime/Renderer/Private/DeferredShadingRenderer.cpp
@@ -3008,7 +3008,7 @@ void FDeferredShadingSceneRenderer::Render(FRDGBuilder& GraphBuilder)
 	{
 		if (HairStrandsBookmarkParameters.bHasElements)
 		{
-			RenderHairStrandsDebugInfo(GraphBuilder, Scene, Views, HairStrandsBookmarkParameters.HairClusterData, SceneTextures.Color.Target);
+			RenderHairStrandsDebugInfo(GraphBuilder, Scene, Views, HairStrandsBookmarkParameters.HairClusterData, SceneTextures.Color.Target, SceneTextures.Depth.Target);
 		}
 	}
 
diff --git a/Engine/Source/Runtime/Renderer/Private/HairStrands/HairStrandsDebug.cpp b/Engine/Source/Runtime/Renderer/Private/HairStrands/HairStrandsDebug.cpp
--- a/Engine/Source/Runtime/Renderer/Private/HairStrands/HairStrandsDebug.cpp
+++ b/Engine/Source/Runtime/Renderer/Private/HairStrands/HairStrandsDebug.cpp
@@ -201,6 +201,7 @@ EHairStrandsDebugMode GetHairStrandsDebugStrandsMode()
 	case 10:  return EHairStrandsDebugMode::RenderHairRoughness;
 	case 11:  return EHairStrandsDebugMode::RenderVisCluster;
 	case 12:  return EHairStrandsDebugMode::RenderHairTangent;
+	case 13:  return EHairStrandsDebugMode::RenderHairControlPoints;
 	default:  return EHairStrandsDebugMode::NoneDebug;
 	};
 }
@@ -219,6 +220,7 @@ static const TCHAR* ToString(EHairStrandsDebugMode DebugMode)
 	case EHairStrandsDebugMode::RenderHairDimension			: return TEXT("Hair dimensions");
 	case EHairStrandsDebugMode::RenderHairRadiusVariation	: return TEXT("Hair radius variation");
 	case EHairStrandsDebugMode::RenderHairTangent			: return TEXT("Hair tangent");
+	case EHairStrandsDebugMode::RenderHairControlPoints		: return TEXT("Hair control points");
 	case EHairStrandsDebugMode::RenderHairBaseColor			: return TEXT("Hair vertices color");
 	case EHairStrandsDebugMode::RenderHairRoughness			: return TEXT("Hair vertices roughness");
 	case EHairStrandsDebugMode::RenderVisCluster			: return TEXT("Hair visility clusters");
@@ -335,6 +337,11 @@ class FHairDebugPS : public FGlobalShader
 
 public:
 	static bool ShouldCompilePermutation(const FGlobalShaderPermutationParameters& Parameters) { return IsHairStrandsSupported(EHairStrandsShaderType::Tool, Parameters.Platform); }
+	static void ModifyCompilationEnvironment(const FGlobalShaderPermutationParameters& Parameters, FShaderCompilerEnvironment& OutEnvironment)
+	{
+		FGlobalShader::ModifyCompilationEnvironment(Parameters, OutEnvironment);
+		OutEnvironment.SetDefine(TEXT("SHADER_DEBUG_MODE"), 1);
+	}
 };
 
 IMPLEMENT_GLOBAL_SHADER(FHairDebugPS, "/Engine/Private/HairStrands/HairStrandsDebug.usf", "MainPS", SF_Pixel);
@@ -1079,10 +1086,12 @@ static void InternalRenderHairStrandsDebugInfo(
 	FScene* Scene,
 	FViewInfo& View,
 	const struct FHairStrandClusterData& HairClusterData,
-	FRDGTextureRef SceneColorTexture)
+	FRDGTextureRef SceneColorTexture,
+	FRDGTextureRef SceneDepthTexture)
 {
 	FHairStrandsBookmarkParameters Params = CreateHairStrandsBookmarkParameters(Scene, View);
 	Params.SceneColorTexture = SceneColorTexture;
+	Params.SceneDepthTexture = SceneDepthTexture;
 	if (!Params.bHasElements)
 	{
 		return;
@@ -1427,11 +1436,12 @@ void RenderHairStrandsDebugInfo(
 	FScene* Scene,
 	TArrayView<FViewInfo> Views,
 	const struct FHairStrandClusterData& HairClusterData,
-	FRDGTextureRef SceneColorTexture)
+	FRDGTextureRef SceneColorTexture,
+	FRDGTextureRef SceneDepthTexture)
 {
 	bool bHasHairData = false;
 	for (FViewInfo& View : Views)
 	{
-		InternalRenderHairStrandsDebugInfo(GraphBuilder, Scene, View, HairClusterData, SceneColorTexture);
+		InternalRenderHairStrandsDebugInfo(GraphBuilder, Scene, View, HairClusterData, SceneColorTexture, SceneDepthTexture);
 	}
 }
\ No newline at end of file
diff --git a/Engine/Source/Runtime/Renderer/Private/HairStrands/HairStrandsDebug.h b/Engine/Source/Runtime/Renderer/Private/HairStrands/HairStrandsDebug.h
--- a/Engine/Source/Runtime/Renderer/Private/HairStrands/HairStrandsDebug.h
+++ b/Engine/Source/Runtime/Renderer/Private/HairStrands/HairStrandsDebug.h
@@ -14,4 +14,5 @@ void RenderHairStrandsDebugInfo(
 	FScene* Scene, 
 	TArrayView<FViewInfo> Views,
 	const struct FHairStrandClusterData& HairClusterData,
-	FRDGTextureRef SceneColorTexture);
\ No newline at end of file
+	FRDGTextureRef SceneColorTexture, 
+	FRDGTextureRef SceneDepthTexture);
\ No newline at end of file
diff --git a/Engine/Source/Runtime/Renderer/Public/HairStrandsInterface.h b/Engine/Source/Runtime/Renderer/Public/HairStrandsInterface.h
--- a/Engine/Source/Runtime/Renderer/Public/HairStrandsInterface.h
+++ b/Engine/Source/Runtime/Renderer/Public/HairStrandsInterface.h
@@ -70,6 +70,7 @@ enum class EHairStrandsDebugMode : uint8
 	RenderVisCluster,
 	RenderHairGroup,
 	RenderLODColoration,
+	RenderHairControlPoints,
 	Count
 };
 
@@ -445,7 +446,8 @@ struct FHairStrandsBookmarkParameters
 	FHairStrandsInstances* Instances = nullptr;
 	const FSceneView* View = nullptr;// // View 0
 	TArray<const FSceneView*> AllViews;
-	FRDGTextureRef SceneColorTexture{};
+	FRDGTextureRef SceneColorTexture = nullptr;
+	FRDGTextureRef SceneDepthTexture = nullptr; 
 
 	bool bHzbRequest = false;
 	bool bHasElements = false;
