Simplify hair strands LOD UI: Decimation settings are not visible only for hair strands geometry.

diff --git a/Engine/Plugins/Runtime/HairStrands/Source/HairStrandsCore/Public/GroomAssetInterpolation.h b/Engine/Plugins/Runtime/HairStrands/Source/HairStrandsCore/Public/GroomAssetInterpolation.h
--- a/Engine/Plugins/Runtime/HairStrands/Source/HairStrandsCore/Public/GroomAssetInterpolation.h
+++ b/Engine/Plugins/Runtime/HairStrands/Source/HairStrandsCore/Public/GroomAssetInterpolation.h
@@ -82,13 +82,13 @@ struct HAIRSTRANDSCORE_API FHairLODSettings
 	UPROPERTY(EditAnywhere, Category = "DecimationSettings", meta = (ToolTip = "If enable this LOD version will use cards representation"))
 	EGroomGeometryType GeometryType = EGroomGeometryType::Strands;
 
-	UPROPERTY(EditAnywhere, Category = "DecimationSettings", meta = (ToolTip = "If enable this LOD version will use the provided attachment points"))
+	UPROPERTY(EditAnywhere, Category = "DecimationSettings", AdvancedDisplay, meta = (ToolTip = "If enable this LOD version will use the provided attachment points"))
 	EGroomBindingType BindingType = EGroomBindingType::Skinning;
 	
-	UPROPERTY(EditAnywhere, Category = "DecimationSettings", meta = (ToolTip = "Groom simulation"))
+	UPROPERTY(EditAnywhere, Category = "DecimationSettings", AdvancedDisplay, meta = (ToolTip = "Groom simulation"))
 	EGroomOverrideType Simulation = EGroomOverrideType::Auto;
 
-	UPROPERTY(EditAnywhere, Category = "DecimationSettings", meta = (DisplayName = "RBF Interpolation", ToolTip = "Global interpolation"))
+	UPROPERTY(EditAnywhere, Category = "DecimationSettings", AdvancedDisplay, meta = (DisplayName = "RBF Interpolation", ToolTip = "Global interpolation"))
 	EGroomOverrideType GlobalInterpolation = EGroomOverrideType::Auto;
 
 	bool operator==(const FHairLODSettings& A) const;
diff --git a/Engine/Plugins/Runtime/HairStrands/Source/HairStrandsEditor/Private/GroomAssetDetails.cpp b/Engine/Plugins/Runtime/HairStrands/Source/HairStrandsEditor/Private/GroomAssetDetails.cpp
--- a/Engine/Plugins/Runtime/HairStrands/Source/HairStrandsEditor/Private/GroomAssetDetails.cpp
+++ b/Engine/Plugins/Runtime/HairStrands/Source/HairStrandsEditor/Private/GroomAssetDetails.cpp
@@ -958,6 +958,40 @@ IDetailPropertyRow& FGroomRenderingDetails::AddPropertyWithCustomReset(TSharedPt
 	return Builder.AddProperty(PropertyHandle.ToSharedRef()).OverrideResetToDefault(ResetOverride);
 }
 
+void FGroomRenderingDetails::ExpandStructForLOD(TSharedRef<IPropertyHandle>& PropertyHandle, IDetailChildrenBuilder& ChildrenBuilder, int32 GroupIndex, int32 LODIndex, bool bOverrideReset)
+{
+
+	uint32 ChildrenCount = 0;
+	PropertyHandle->GetNumChildren(ChildrenCount);
+	for (uint32 ChildIt = 0; ChildIt < ChildrenCount; ++ChildIt)
+	{
+		TSharedPtr<IPropertyHandle> ChildHandle = PropertyHandle->GetChildHandle(ChildIt);
+		const FName ChildPropertyName = ChildHandle->GetProperty()->GetFName();
+
+		// If the geometry type is not strands, then bypass the display of the strands related property
+		if (GroomAsset->GetGeometryType(GroupIndex, LODIndex) != EGroomGeometryType::Strands && 
+			(ChildPropertyName == GET_MEMBER_NAME_CHECKED(FHairLODSettings, CurveDecimation)  ||
+			 ChildPropertyName == GET_MEMBER_NAME_CHECKED(FHairLODSettings, VertexDecimation) ||
+			 ChildPropertyName == GET_MEMBER_NAME_CHECKED(FHairLODSettings, AngularThreshold) ||
+			 ChildPropertyName == GET_MEMBER_NAME_CHECKED(FHairLODSettings, ThicknessScale)))
+		{
+			continue;
+		}
+
+		if (bOverrideReset)
+		{
+			FIsResetToDefaultVisible IsResetVisible = FIsResetToDefaultVisible::CreateSP(this, &FGroomRenderingDetails::ShouldResetToDefault, GroupIndex, LODIndex);
+			FResetToDefaultHandler ResetHandler = FResetToDefaultHandler::CreateSP(this, &FGroomRenderingDetails::ResetToDefault, GroupIndex, LODIndex);
+			FResetToDefaultOverride ResetOverride = FResetToDefaultOverride::Create(IsResetVisible, ResetHandler);
+			ChildrenBuilder.AddProperty(ChildHandle.ToSharedRef()).OverrideResetToDefault(ResetOverride);
+		}
+		else
+		{
+			ChildrenBuilder.AddProperty(ChildHandle.ToSharedRef());
+		}
+	}
+}
+
 void FGroomRenderingDetails::ExpandStruct(TSharedPtr<IPropertyHandle>& PropertyHandle, IDetailChildrenBuilder& ChildrenBuilder, int32 GroupIndex, int32 LODIndex, bool bOverrideReset)
 {
 
@@ -1147,7 +1181,7 @@ void FGroomRenderingDetails::OnGenerateElementForLODs(TSharedRef<IPropertyHandle
 
 	// Rename the array entry name by its group name and adds all its existing properties
 	StructProperty->SetPropertyDisplayName(LOCTEXT("LODProperties", "LOD Properties"));
-	ExpandStruct(StructProperty, ChildrenBuilder, GroupIndex, LODIndex, true);
+	ExpandStructForLOD(StructProperty, ChildrenBuilder, GroupIndex, LODIndex, true); ///
 }
 
 TSharedRef<SWidget> FGroomRenderingDetails::MakeGroupNameButtonCustomization(int32 GroupIndex, FProperty* Property)
diff --git a/Engine/Plugins/Runtime/HairStrands/Source/HairStrandsEditor/Private/GroomAssetDetails.h b/Engine/Plugins/Runtime/HairStrands/Source/HairStrandsEditor/Private/GroomAssetDetails.h
--- a/Engine/Plugins/Runtime/HairStrands/Source/HairStrandsEditor/Private/GroomAssetDetails.h
+++ b/Engine/Plugins/Runtime/HairStrands/Source/HairStrandsEditor/Private/GroomAssetDetails.h
@@ -85,6 +85,7 @@ private:
 	FReply OnSaveCards(int32 GroupIndex, FProperty* Property);
 	FReply OnGenerateCardDataUsingPlugin(int32 GroupIndex);
 
+	void ExpandStructForLOD(TSharedRef<IPropertyHandle>& PropertyHandle, IDetailChildrenBuilder& ChildrenBuilder, int32 GroupIndex, int32 LODIndex, bool bOverrideReset);
 	void ExpandStruct(TSharedPtr<IPropertyHandle>& PropertyHandle, IDetailChildrenBuilder& ChildrenBuilder, int32 GroupIndex, int32 LODIndex, bool bOverrideReset);
 	void ExpandStruct(TSharedRef<IPropertyHandle>& PropertyHandle, IDetailChildrenBuilder& ChildrenBuilder, int32 GroupIndex, int32 LODIndex, bool bOverrideReset);
 	IDetailPropertyRow& AddPropertyWithCustomReset(TSharedPtr<IPropertyHandle>& PropertyHandle, IDetailChildrenBuilder& Builder, int32 GroupIndex, int32 LODIndex);
