- Fixed monolithic editor link issue
#jria none
#p4v-preflight-copy 18309096

diff --git a/Engine/Plugins/Runtime/HairStrands/Source/HairCardGeneratorFramework/Public/IHairCardGenerator.h b/Engine/Plugins/Runtime/HairStrands/Source/HairCardGeneratorFramework/Public/IHairCardGenerator.h
--- a/Engine/Plugins/Runtime/HairStrands/Source/HairCardGeneratorFramework/Public/IHairCardGenerator.h
+++ b/Engine/Plugins/Runtime/HairStrands/Source/HairCardGeneratorFramework/Public/IHairCardGenerator.h
@@ -21,12 +21,12 @@ public:
 
 namespace HairCardGenerator_Utils
 {
-	void RegisterModularHairCardGenerator(IHairCardGenerator* Generator)
+	FORCEINLINE void RegisterModularHairCardGenerator(IHairCardGenerator* Generator)
 	{
 		IModularFeatures::Get().RegisterModularFeature(IHairCardGenerator::ModularFeatureName, Generator);
 	}
 
-	void UnregisterModularHairCardGenerator(IHairCardGenerator* Generator)
+	FORCEINLINE void UnregisterModularHairCardGenerator(IHairCardGenerator* Generator)
 	{
 		IModularFeatures::Get().UnregisterModularFeature(IHairCardGenerator::ModularFeatureName, Generator);
 	}
