Remove unecessary hair simulation reset when LOD changed.

This change is no longer needed with the new game-thread driven LOD selection when simulation is enabled.

diff --git a/Engine/Plugins/Runtime/HairStrands/Source/HairStrandsCore/Private/Niagara/NiagaraDataInterfaceHairStrands.cpp b/Engine/Plugins/Runtime/HairStrands/Source/HairStrandsCore/Private/Niagara/NiagaraDataInterfaceHairStrands.cpp
--- a/Engine/Plugins/Runtime/HairStrands/Source/HairStrandsCore/Private/Niagara/NiagaraDataInterfaceHairStrands.cpp
+++ b/Engine/Plugins/Runtime/HairStrands/Source/HairStrandsCore/Private/Niagara/NiagaraDataInterfaceHairStrands.cpp
@@ -697,7 +697,7 @@ struct FNDIHairStrandsParametersCS : public FNiagaraDataInterfaceParametersCS
 				DeformedMeshProjection->MeshSampleWeightsBuffer.SRV.GetReference() : FNiagaraRenderer::GetDummyFloatBuffer();
 
 			// Simulation setup
-			const int32 NeedResetValue = bHasLODSwitched || (ProxyData->TickCount <= GHairSimulationMaxDelay);
+			const int32 NeedResetValue = (ProxyData->TickCount <= GHairSimulationMaxDelay);
 			const int32 RestUpdateValue = GHairSimulationRestUpdate;
 			const int32 LocalSimulationValue = ProxyData->LocalSimulation;
 			
