Fix motion vector issue on hair cards when switching LOD.

diff --git a/Engine/Plugins/Runtime/HairStrands/Source/HairStrandsCore/Private/HairStrandsRendering.cpp b/Engine/Plugins/Runtime/HairStrands/Source/HairStrandsCore/Private/HairStrandsRendering.cpp
--- a/Engine/Plugins/Runtime/HairStrands/Source/HairStrandsCore/Private/HairStrandsRendering.cpp
+++ b/Engine/Plugins/Runtime/HairStrands/Source/HairStrandsCore/Private/HairStrandsRendering.cpp
@@ -1075,7 +1075,13 @@ static void AddHairCardsDeformationPass(
 	
 	FRDGImportedBuffer CardsDeformedPositionBuffer_Curr = Register(GraphBuilder, LOD.DeformedResource->GetBuffer(FHairCardsDeformedResource::Current), ERDGImportedBufferFlags::CreateUAV);
 	FRDGImportedBuffer CardsDeformedPositionBuffer_Prev = Register(GraphBuilder, LOD.DeformedResource->GetBuffer(FHairCardsDeformedResource::Previous), ERDGImportedBufferFlags::CreateUAV);
-	AddCopyBufferPass(GraphBuilder, CardsDeformedPositionBuffer_Prev.Buffer, CardsDeformedPositionBuffer_Curr.Buffer);
+
+	// If LOD hasn't switched, copy the current buffer into previous buffer to have correct motion vectors
+	const bool bHasLODSwitch = Instance->HairGroupPublicData->VFInput.bHasLODSwitch;
+	if (!bHasLODSwitch)
+	{
+		AddCopyBufferPass(GraphBuilder, CardsDeformedPositionBuffer_Prev.Buffer, CardsDeformedPositionBuffer_Curr.Buffer);
+	}
 	
 	FRDGImportedBuffer CardsDeformedNormalBuffer = Register(GraphBuilder, LOD.DeformedResource->DeformedNormalBuffer, ERDGImportedBufferFlags::CreateUAV);
 
@@ -1145,6 +1151,12 @@ static void AddHairCardsDeformationPass(
 		Parameters,
 		FIntVector(DispatchCountX,1,1));
 
+	// If LOD has switched, copy the current buffer, so that we don't get incorrect motion vector
+	if (bHasLODSwitch)
+	{
+		AddCopyBufferPass(GraphBuilder, CardsDeformedPositionBuffer_Prev.Buffer, CardsDeformedPositionBuffer_Curr.Buffer);
+	}
+
 	GraphBuilder.SetBufferAccessFinal(CardsDeformedPositionBuffer_Curr.Buffer, ERHIAccess::SRVMask);
 	GraphBuilder.SetBufferAccessFinal(CardsDeformedPositionBuffer_Prev.Buffer, ERHIAccess::SRVMask);
 	GraphBuilder.SetBufferAccessFinal(CardsDeformedNormalBuffer.Buffer, ERHIAccess::SRVMask);
