Fix transition issue.

diff --git a/Engine/Plugins/Runtime/HairStrands/Source/HairStrandsCore/Private/HairStrandsRendering.cpp b/Engine/Plugins/Runtime/HairStrands/Source/HairStrandsCore/Private/HairStrandsRendering.cpp
--- a/Engine/Plugins/Runtime/HairStrands/Source/HairStrandsCore/Private/HairStrandsRendering.cpp
+++ b/Engine/Plugins/Runtime/HairStrands/Source/HairStrandsCore/Private/HairStrandsRendering.cpp
@@ -1073,8 +1073,8 @@ static void AddHairCardsDeformationPass(
 
 	FHairGroupInstance::FCards::FLOD& LOD = Instance->Cards.LODs[HairLODIndex];
 	
-	FRDGImportedBuffer CardsDeformedPositionBuffer_Curr = Register(GraphBuilder, LOD.DeformedResource->GetBuffer(FHairCardsDeformedResource::Current), ERDGImportedBufferFlags::CreateUAV);
-	FRDGImportedBuffer CardsDeformedPositionBuffer_Prev = Register(GraphBuilder, LOD.DeformedResource->GetBuffer(FHairCardsDeformedResource::Previous), ERDGImportedBufferFlags::CreateUAV);
+	FRDGImportedBuffer CardsDeformedPositionBuffer_Curr = Register(GraphBuilder, LOD.DeformedResource->GetBuffer(FHairCardsDeformedResource::Current), ERDGImportedBufferFlags::CreateViews);
+	FRDGImportedBuffer CardsDeformedPositionBuffer_Prev = Register(GraphBuilder, LOD.DeformedResource->GetBuffer(FHairCardsDeformedResource::Previous), ERDGImportedBufferFlags::CreateViews);
 
 	// If LOD hasn't switched, copy the current buffer into previous buffer to have correct motion vectors
 	const bool bHasLODSwitch = Instance->HairGroupPublicData->VFInput.bHasLODSwitch;
@@ -1083,7 +1083,7 @@ static void AddHairCardsDeformationPass(
 		AddCopyBufferPass(GraphBuilder, CardsDeformedPositionBuffer_Prev.Buffer, CardsDeformedPositionBuffer_Curr.Buffer);
 	}
 	
-	FRDGImportedBuffer CardsDeformedNormalBuffer = Register(GraphBuilder, LOD.DeformedResource->DeformedNormalBuffer, ERDGImportedBufferFlags::CreateUAV);
+	FRDGImportedBuffer CardsDeformedNormalBuffer = Register(GraphBuilder, LOD.DeformedResource->DeformedNormalBuffer, ERDGImportedBufferFlags::CreateViews);
 
 	FHairCardsDeformationCS::FParameters* Parameters = GraphBuilder.AllocParameters<FHairCardsDeformationCS::FParameters>();
 	Parameters->GuideVertexCount			= LOD.Guides.RestResource->GetVertexCount();
@@ -1157,9 +1157,9 @@ static void AddHairCardsDeformationPass(
 		AddCopyBufferPass(GraphBuilder, CardsDeformedPositionBuffer_Prev.Buffer, CardsDeformedPositionBuffer_Curr.Buffer);
 	}
 
-	GraphBuilder.SetBufferAccessFinal(CardsDeformedPositionBuffer_Curr.Buffer, ERHIAccess::SRVMask);
-	GraphBuilder.SetBufferAccessFinal(CardsDeformedPositionBuffer_Prev.Buffer, ERHIAccess::SRVMask);
-	GraphBuilder.SetBufferAccessFinal(CardsDeformedNormalBuffer.Buffer, ERHIAccess::SRVMask);
+	HairTransition::TransitToSRV(GraphBuilder, CardsDeformedPositionBuffer_Curr.SRV, ERDGPassFlags::Raster);
+	HairTransition::TransitToSRV(GraphBuilder, CardsDeformedPositionBuffer_Prev.SRV, ERDGPassFlags::Raster);
+	HairTransition::TransitToSRV(GraphBuilder, CardsDeformedNormalBuffer.SRV, ERDGPassFlags::Raster);
 }
 
 ///////////////////////////////////////////////////////////////////////////////////////////////////
