Add hair streaming request support.

This allows to request specific LOD data earlier in the frame when LOD selection type is Predicted or Forced. This prevents case where hair geometry wouldn't be rendered for one frame, as its resources were not loaded/ready in time.

This solves two concreate issues: crowd hair poping when switching to LOD with skin-binding enabled, and cinematic sequence, where the missing resources on time would cause simulation to have invalid data to start with.

This CL also changes the aysnc loading configuration: now one can configure it differently for strands/cards and also states which LOD would be async-loaded

diff --git a/Engine/Plugins/Runtime/HairStrands/Source/HairStrandsCore/Private/GroomComponent.cpp b/Engine/Plugins/Runtime/HairStrands/Source/HairStrandsCore/Private/GroomComponent.cpp
--- a/Engine/Plugins/Runtime/HairStrands/Source/HairStrandsCore/Private/GroomComponent.cpp
+++ b/Engine/Plugins/Runtime/HairStrands/Source/HairStrandsCore/Private/GroomComponent.cpp
@@ -1609,6 +1609,7 @@ bool UGroomComponent::GetIsHairLengthScaleEnabled()
 	return IsEnabled;
 }
 
+void AddHairStreamingRequest(FHairGroupInstance* InInstance, int32 InLODIndex);
 
 void UGroomComponent::SetForcedLOD(int32 CurrLODIndex)
 {
@@ -1657,6 +1658,7 @@ void UGroomComponent::SetForcedLOD(int32 CurrLODIndex)
 	SwitchSimulationLOD(PrevLODIndex, CurrLODIndex);
 
 	// Finally, update the forced LOD value and LOD selection type
+	const bool bHasLODSwitch = CurrLODIndex != PrevLODIndex;
 	LODForcedIndex	 = CurrLODIndex;
 	LODSelectionType = CurrLODSelectionType;
 
@@ -1665,12 +1667,17 @@ void UGroomComponent::SetForcedLOD(int32 CurrLODIndex)
 	{
 		const EHairLODSelectionType LocalLODSelectionType = LODSelectionType;
 		ENQUEUE_RENDER_COMMAND(FHairComponentSendLODIndex)(
-		[GroomSceneProxy, CurrLODIndex, LocalLODSelectionType](FRHICommandListImmediate& RHICmdList)
+		[GroomSceneProxy, CurrLODIndex, LocalLODSelectionType, bHasLODSwitch](FRHICommandListImmediate& RHICmdList)
 		{
 			for (FHairGroupInstance* Instance : GroomSceneProxy->HairGroupInstances)
 			{
 				Instance->Debug.LODForcedIndex = CurrLODIndex;
 				Instance->Debug.LODSelectionTypeForDebug = LocalLODSelectionType;
+
+				if (bHasLODSwitch)
+				{
+					AddHairStreamingRequest(Instance, CurrLODIndex);
+				}
 			}
 		});
 	}
@@ -3301,6 +3308,11 @@ void UGroomComponent::TickComponent(float DeltaTime, enum ELevelTick TickType, F
 				Instance->Debug.RigidPreviousLocalToWorld = Instance->Debug.RigidCurrentLocalToWorld;
 				Instance->Debug.RigidCurrentLocalToWorld = RigidLocalToWorld;
 			}
+
+			if (LocalLODSelectionType == EHairLODSelectionType::Predicted && EffectiveForceLOD >= 0)
+			{
+				AddHairStreamingRequest(Instance, EffectiveForceLOD);
+			}
 		}
 	});
 }
diff --git a/Engine/Plugins/Runtime/HairStrands/Source/HairStrandsCore/Private/GroomManager.cpp b/Engine/Plugins/Runtime/HairStrands/Source/HairStrandsCore/Private/GroomManager.cpp
--- a/Engine/Plugins/Runtime/HairStrands/Source/HairStrandsCore/Private/GroomManager.cpp
+++ b/Engine/Plugins/Runtime/HairStrands/Source/HairStrandsCore/Private/GroomManager.cpp
@@ -587,6 +587,100 @@ static void AddCopyHairStrandsPositionPass(
 #endif
 }
 
+void AddHairStreamingRequest(FHairGroupInstance* Instance, int32 InLODIndex)
+{
+	if (Instance && InLODIndex >= 0)
+	{
+		check(Instance->HairGroupPublicData);
+
+		// Hypothesis: the mesh LOD will be the same as the hair LOD
+		const int32 MeshLODIndex = InLODIndex;
+		
+		// Insure that MinLOD is necessary taken into account if a force LOD is request (i.e., LODIndex>=0). If a Force LOD 
+		// is not resquested (i.e., LODIndex<0), the MinLOD is applied after ViewLODIndex has been determined in the codeblock below
+		const float MinLOD = FMath::Max(0, GHairStrandsMinLOD);
+		float LODIndex = FMath::Max(InLODIndex, MinLOD);
+
+		const TArray<EHairGeometryType>& LODGeometryTypes = Instance->HairGroupPublicData->GetLODGeometryTypes();
+		const TArray<bool>& LODVisibilities = Instance->HairGroupPublicData->GetLODVisibilities();
+		const int32 LODCount = LODVisibilities.Num();
+
+		LODIndex = FMath::Clamp(LODIndex, 0.f, float(LODCount - 1));
+
+		const int32 IntLODIndex = FMath::Clamp(FMath::FloorToInt(LODIndex), 0, LODCount - 1);
+		const bool bIsVisible = LODVisibilities[IntLODIndex];
+		const bool bForceCards = GHairStrands_UseCards > 0 || Instance->bForceCards; // todo
+		EHairGeometryType GeometryType = ConvertLODGeometryType(LODGeometryTypes[IntLODIndex], bForceCards, GMaxRHIShaderPlatform);
+
+		if (GeometryType == EHairGeometryType::Meshes)
+		{
+			if (!Instance->Meshes.IsValid(IntLODIndex))
+			{
+				GeometryType = EHairGeometryType::NoneGeometry;
+			}
+		}
+		else if (GeometryType == EHairGeometryType::Cards)
+		{
+			if (!Instance->Cards.IsValid(IntLODIndex))
+			{
+				GeometryType = EHairGeometryType::NoneGeometry;
+			}
+		}
+		else if (GeometryType == EHairGeometryType::Strands)
+		{
+			if (!Instance->Strands.IsValid())
+			{
+				GeometryType = EHairGeometryType::NoneGeometry;
+			}
+		}
+
+		if (!bIsVisible)
+		{
+			GeometryType = EHairGeometryType::NoneGeometry;
+		}
+
+		const bool bSimulationEnable			= Instance->HairGroupPublicData->IsSimulationEnable(LODIndex);
+		const bool bGlobalInterpolationEnable	= Instance->HairGroupPublicData->IsGlobalInterpolationEnable(LODIndex);
+		const bool bLODNeedsGuides				= bSimulationEnable || bGlobalInterpolationEnable;
+
+		const EHairResourceLoadingType LoadingType = GetHairResourceLoadingType(GeometryType, int32(LODIndex));
+		if (LoadingType != EHairResourceLoadingType::Async || GeometryType == EHairGeometryType::NoneGeometry)
+		{
+			return;
+		}
+
+		// Lazy allocation of resources
+		// Note: Allocation will only be done if the resources is not initialized yet. Guides deformed position are also initialized from the Rest position at creation time.
+		if (Instance->Guides.Data && bLODNeedsGuides)
+		{
+			if (Instance->Guides.RestRootResource)			{ Instance->Guides.RestRootResource->StreamInData(); Instance->Guides.RestRootResource->StreamInLODData(MeshLODIndex); }
+			if (Instance->Guides.RestResource)				{ Instance->Guides.RestResource->StreamInData(); }
+		}
+
+		if (GeometryType == EHairGeometryType::Meshes)
+		{
+			FHairGroupInstance::FMeshes::FLOD& InstanceLOD = Instance->Meshes.LODs[IntLODIndex];
+		}
+		else if (GeometryType == EHairGeometryType::Cards)
+		{
+			FHairGroupInstance::FCards::FLOD& InstanceLOD = Instance->Cards.LODs[IntLODIndex];
+
+			if (InstanceLOD.InterpolationResource)			{ InstanceLOD.InterpolationResource->StreamInData(); }
+			if (InstanceLOD.DeformedResource)				{ InstanceLOD.DeformedResource->StreamInData(); }
+			if (InstanceLOD.Guides.RestRootResource)		{ InstanceLOD.Guides.RestRootResource->StreamInData(); InstanceLOD.Guides.RestRootResource->StreamInLODData(MeshLODIndex); }
+			if (InstanceLOD.Guides.RestResource)			{ InstanceLOD.Guides.RestResource->StreamInData(); }
+			if (InstanceLOD.Guides.InterpolationResource)	{ InstanceLOD.Guides.InterpolationResource->StreamInData(); }
+		}
+		else if (GeometryType == EHairGeometryType::Strands)
+		{
+			if (Instance->Strands.RestRootResource)			{ Instance->Strands.RestRootResource->StreamInData(); Instance->Strands.RestRootResource->StreamInLODData(MeshLODIndex); }
+			if (Instance->Strands.RestResource)				{ Instance->Strands.RestResource->StreamInData(); }
+			if (Instance->Strands.ClusterCullingResource)	{ Instance->Strands.ClusterCullingResource->StreamInData(); }
+			if (Instance->Strands.InterpolationResource)	{ Instance->Strands.InterpolationResource->StreamInData(); }
+		}
+	}
+}
+
 #if RHI_RAYTRACING
 static void AllocateRaytracingResources(FHairGroupInstance* Instance)
 {
@@ -724,12 +818,16 @@ static void RunHairLODSelection(
 			GeometryType = EHairGeometryType::NoneGeometry;
 		}
 
-		const EHairResourceLoadingType LoadingType = GetHairResourceLoadingType();
+		const bool bSimulationEnable = Instance->HairGroupPublicData->IsSimulationEnable(LODIndex);
+		const bool bGlobalInterpolationEnable =  Instance->HairGroupPublicData->IsGlobalInterpolationEnable(LODIndex);
+		const bool bLODNeedsGuides = bSimulationEnable || bGlobalInterpolationEnable;
+
+		const EHairResourceLoadingType LoadingType = GetHairResourceLoadingType(GeometryType, int32(LODIndex));
 		EHairResourceStatus ResourceStatus = EHairResourceStatus::None;
 
 		// Lazy allocation of resources
 		// Note: Allocation will only be done if the resources is not initialized yet. Guides deformed position are also initialized from the Rest position at creation time.
-		if (Instance->Guides.Data)
+		if (Instance->Guides.Data && bLODNeedsGuides)
 		{
 			if (Instance->Guides.RestRootResource)			{ Instance->Guides.RestRootResource->Allocate(GraphBuilder, LoadingType, ResourceStatus); Instance->Guides.RestRootResource->AllocateLOD(GraphBuilder, MeshLODIndex, LoadingType, ResourceStatus); }
 			if (Instance->Guides.RestResource)				{ Instance->Guides.RestResource->Allocate(GraphBuilder, LoadingType, ResourceStatus); }
diff --git a/Engine/Plugins/Runtime/HairStrands/Source/HairStrandsCore/Private/GroomResources.cpp b/Engine/Plugins/Runtime/HairStrands/Source/HairStrandsCore/Private/GroomResources.cpp
--- a/Engine/Plugins/Runtime/HairStrands/Source/HairStrandsCore/Private/GroomResources.cpp
+++ b/Engine/Plugins/Runtime/HairStrands/Source/HairStrandsCore/Private/GroomResources.cpp
@@ -26,8 +26,11 @@
 static int32 GHairStrandsBulkData_ReleaseAfterUse = 0;
 static FAutoConsoleVariableRef CVarHairStrandsBulkData_ReleaseAfterUse(TEXT("r.HairStrands.Strands.BulkData.ReleaseAfterUse"), GHairStrandsBulkData_ReleaseAfterUse, TEXT("Release CPU bulk data once hair groom/groom binding asset GPU resources are created. This saves memory"));
 
-static int32 GHairStrandsBulkData_AsyncLoading = 0;
-static FAutoConsoleVariableRef CVarHairStrandsBulkData_AsyncLoading(TEXT("r.HairStrands.Strands.BulkData.AsyncLoading"), GHairStrandsBulkData_AsyncLoading, TEXT("Load hair data with async loading so that it is not blocking the rendering thread."));
+static int32 GHairStrandsBulkData_AsyncLoading = -1;
+static int32 GHairCardsBulkData_AsyncLoading = -1;
+
+static FAutoConsoleVariableRef CVarHairStrandsBulkData_AsyncLoading(TEXT("r.HairStrands.Strands.BulkData.AsyncLoading"), GHairStrandsBulkData_AsyncLoading, TEXT("Load hair strands data with async loading so that it is not blocking the rendering thread. This value define the MinLOD at which this happen. Default disabled (-1)"));
+static FAutoConsoleVariableRef CVarHairCardsBulkData_AsyncLoading(TEXT("r.HairStrands.Cards.BulkData.AsyncLoading"), GHairCardsBulkData_AsyncLoading, TEXT("Load hair cards/meshes data with async loading so that it is not blocking the rendering thread. This value define the MinLOD at which this happen. Default disabled (-1)"));
 
 static int32 GHairStrandsBulkData_Validation = 1;
 static FAutoConsoleVariableRef CVarHairStrandsBulkData_Validation(TEXT("r.HairStrands.Strands.BulkData.Validation"), GHairStrandsBulkData_Validation, TEXT("Validate some hair strands data at serialization/loading time."));
@@ -39,10 +42,16 @@ static FAutoConsoleVariableRef CVarHairStrandsDebugVoxel_MaxSegmentPerVoxel(TEXT
 
 ////////////////////////////////////////////////////////////////////////////////////
 
-EHairResourceLoadingType GetHairResourceLoadingType()
+EHairResourceLoadingType GetHairResourceLoadingType(EHairGeometryType InGeometryType, int32 InLODIndex)
 {
 #if !WITH_EDITORONLY_DATA
-	return GHairStrandsBulkData_AsyncLoading > 0 ? EHairResourceLoadingType::Async : EHairResourceLoadingType::Sync;
+	switch (InGeometryType)
+	{
+	case EHairGeometryType::Strands: return InLODIndex <= GHairStrandsBulkData_AsyncLoading ? EHairResourceLoadingType::Async : EHairResourceLoadingType::Sync;
+	case EHairGeometryType::Cards:
+	case EHairGeometryType::Meshes: return InLODIndex <= GHairCardsBulkData_AsyncLoading ? EHairResourceLoadingType::Async : EHairResourceLoadingType::Sync;
+	}
+	return EHairResourceLoadingType::Sync;
 #else
 	return EHairResourceLoadingType::Sync;
 #endif
@@ -521,6 +530,22 @@ void FHairCommonResource::AllocateLOD(FRDGBuilder& GraphBuilder, int32 LODIndex,
 	Status |= EHairResourceStatus::Valid;
 }
 
+void FHairCommonResource::StreamInData()
+{
+	if (!bIsInitialized)
+	{
+		InternalIsDataLoaded();
+	}
+}
+
+void FHairCommonResource::StreamInLODData(int32 LODIndex)
+{
+	if (!bIsInitialized)
+	{
+		InternalIsLODDataLoaded(LODIndex);
+	}
+}
+
 void AsyncLoadHairBulkData(FByteBulkData& InData, bool& bIsLoading)
 {
 	bIsLoading = InData.StartAsyncLoading() || bIsLoading;
diff --git a/Engine/Plugins/Runtime/HairStrands/Source/HairStrandsCore/Public/GroomResources.h b/Engine/Plugins/Runtime/HairStrands/Source/HairStrandsCore/Public/GroomResources.h
--- a/Engine/Plugins/Runtime/HairStrands/Source/HairStrandsCore/Public/GroomResources.h
+++ b/Engine/Plugins/Runtime/HairStrands/Source/HairStrandsCore/Public/GroomResources.h
@@ -70,7 +70,7 @@ FORCEINLINE EHairResourceStatus  operator| (EHairResourceStatus A, EHairResource
 FORCEINLINE EHairResourceStatus& operator|=(EHairResourceStatus&A, EHairResourceStatus B) { return A=A|B; }
 FORCEINLINE bool operator! (EHairResourceStatus A) { return static_cast<uint8>(A) != 0; }
 
-EHairResourceLoadingType GetHairResourceLoadingType();
+EHairResourceLoadingType GetHairResourceLoadingType(EHairGeometryType InGeometryType, int32 InLODIndex);
 
 /* Hair resouces which whom allocation can be deferred */
 struct FHairCommonResource : public FRenderResource
@@ -86,6 +86,10 @@ struct FHairCommonResource : public FRenderResource
 	void Allocate(FRDGBuilder& GraphBuilder, EHairResourceLoadingType LoadingType);
 	void Allocate(FRDGBuilder& GraphBuilder, EHairResourceLoadingType LoadingType, EHairResourceStatus& Status);
 	void AllocateLOD(FRDGBuilder& GraphBuilder, int32 LODIndex, EHairResourceLoadingType LoadingType, EHairResourceStatus& Status);
+
+	void StreamInData();
+	void StreamInLODData(int32 LODIndex);
+
 	virtual void InternalAllocate() {}
 	virtual void InternalAllocate(FRDGBuilder& GraphBuilder) {}
 	virtual void InternalAllocateLOD(FRDGBuilder& GraphBuilder, int32 LODIndex) {}
