* Reenable GPUScene supprot for HairCardsVF
* Fix invalid PrevLocalToWorld transform update for dynamic geometry

This allows to fix missing self-shadow on hair cards/hair meshes.

diff --git a/Engine/Plugins/Runtime/HairStrands/Source/HairStrandsCore/Private/HairCardsVertexFactory.cpp b/Engine/Plugins/Runtime/HairStrands/Source/HairStrandsCore/Private/HairCardsVertexFactory.cpp
--- a/Engine/Plugins/Runtime/HairStrands/Source/HairStrandsCore/Private/HairCardsVertexFactory.cpp
+++ b/Engine/Plugins/Runtime/HairStrands/Source/HairStrandsCore/Private/HairCardsVertexFactory.cpp
@@ -14,8 +14,6 @@
 #include "GroomInstance.h"
 #include "SystemTextures.h" 
 
-#define VF_CARDS_SUPPORT_GPU_SCENE 0
-
 template<typename T> inline void VFC_BindParam(FMeshDrawSingleShaderBindings& ShaderBindings, const FShaderResourceParameter& Param, T* Value) { if (Param.IsBound() && Value) ShaderBindings.Add(Param, Value); }
 template<typename T> inline void VFC_BindParam(FMeshDrawSingleShaderBindings& ShaderBindings, const FShaderParameter& Param, const T& Value) { if (Param.IsBound()) ShaderBindings.Add(Param, Value); }
 
@@ -215,15 +213,7 @@ void FHairCardsVertexFactory::ModifyCompilationEnvironment(const FVertexFactoryS
 {
 	FVertexFactory::ModifyCompilationEnvironment(Parameters, OutEnvironment);
 
-	bool bUseGPUSceneAndPrimitiveIdStream = false;
-#if VF_CARDS_SUPPORT_GPU_SCENE
-	bUseGPUSceneAndPrimitiveIdStream = 
-		Parameters.VertexFactoryType->SupportsPrimitiveIdStream() && 
-		UseGPUScene(Parameters.Platform, GetMaxSupportedFeatureLevel(Parameters.Platform))
-		// TODO: support GPUScene on mobile
-		&& !IsMobilePlatform(Parameters.Platform);
-#endif
-
+	const bool bUseGPUSceneAndPrimitiveIdStream = Parameters.VertexFactoryType->SupportsPrimitiveIdStream() && UseGPUScene(Parameters.Platform, GetMaxSupportedFeatureLevel(Parameters.Platform));
 	OutEnvironment.SetDefine(TEXT("VF_SUPPORTS_PRIMITIVE_SCENE_DATA"), bUseGPUSceneAndPrimitiveIdStream);
 	OutEnvironment.SetDefine(TEXT("HAIR_CARD_MESH_FACTORY"), TEXT("1"));
 	OutEnvironment.SetDefine(TEXT("MANUAL_VERTEX_FETCH"), RHISupportsManualVertexFetch(Parameters.Platform) ? TEXT("1") : TEXT("0"));	
@@ -231,14 +221,12 @@ void FHairCardsVertexFactory::ModifyCompilationEnvironment(const FVertexFactoryS
 
 void FHairCardsVertexFactory::ValidateCompiledResult(const FVertexFactoryType* Type, EShaderPlatform Platform, const FShaderParameterMap& ParameterMap, TArray<FString>& OutErrors)
 {
-#if VF_CARDS_SUPPORT_GPU_SCENE
 	if (Type->SupportsPrimitiveIdStream() 
 		&& UseGPUScene(Platform, GetMaxSupportedFeatureLevel(Platform)) 
 		&& ParameterMap.ContainsParameterAllocation(FPrimitiveUniformShaderParameters::StaticStructMetadata.GetShaderVariableName()))
 	{
 		OutErrors.AddUnique(*FString::Printf(TEXT("Shader attempted to bind the Primitive uniform buffer even though Vertex Factory %s computes a PrimitiveId per-instance.  This will break auto-instancing.  Shaders should use GetPrimitiveData(PrimitiveId).Member instead of Primitive.Member."), Type->GetName()));
 	}
-#endif
 }
 
 void FHairCardsVertexFactory::SetData(const FDataType& InData)
@@ -276,13 +264,12 @@ void FHairCardsVertexFactory::InitResources()
 	// We create different streams based on feature level
 	check(HasValidFeatureLevel());
 
+	// If the platform does not support manual vertex fetching we assume it is a low end platform, and so we don't enable deformation.
+	// VertexFactory needs to be able to support max possible shader platform and feature level
+	// in case if we switch feature level at runtime.
 	FVertexDeclarationElementList Elements;
-#if VF_CARDS_SUPPORT_GPU_SCENE
-	if (AddPrimitiveIdStreamElement(EVertexInputStreamType::Default, Elements, 13, 0xff))
-	{
-		bNeedsDeclaration = true;
-	}
-#endif
+	SetPrimitiveIdStreamIndex(EVertexInputStreamType::Default, -1);
+	AddPrimitiveIdStreamElement(EVertexInputStreamType::Default, 13, Elements);
 
 	// Note this is a local version of the VF's bSupportsManualVertexFetch, which take into account the feature level
 	const bool bManualFetch = SupportsManualVertexFetch(GetFeatureLevel());
@@ -357,5 +344,5 @@ IMPLEMENT_VERTEX_FACTORY_TYPE(FHairCardsVertexFactory,"/Engine/Private/HairStran
 	| EVertexFactoryFlags::SupportsPrecisePrevWorldPos
 	| EVertexFactoryFlags::SupportsCachingMeshDrawCommands
 	| EVertexFactoryFlags::SupportsRayTracing
-	| (VF_CARDS_SUPPORT_GPU_SCENE ? EVertexFactoryFlags::SupportsPrimitiveIdStream : EVertexFactoryFlags::None)
+	| EVertexFactoryFlags::SupportsPrimitiveIdStream
 );
diff --git a/Engine/Shaders/Private/HairStrands/HairCardsVertexFactory.ush b/Engine/Shaders/Private/HairStrands/HairCardsVertexFactory.ush
--- a/Engine/Shaders/Private/HairStrands/HairCardsVertexFactory.ush
+++ b/Engine/Shaders/Private/HairStrands/HairCardsVertexFactory.ush
@@ -16,7 +16,7 @@
 #include "../RayTracing/RayTracingHitGroupCommon.ush"
 #endif
 
-#define USE_HAIR_COMPLEX_TRANSMITTANCE 1 
+#define USE_HAIR_COMPLEX_TRANSMITTANCE 1
 
 #if MANUAL_VERTEX_FETCH
 	Buffer<float4> VertexFetch_InstanceOriginBuffer;
@@ -25,8 +25,6 @@
 
 ////////////////////////////////////////////////////////////////////////////////
 // HAIR_CARD_MESH_FACTORY
-// Define used in certain shader like DeepShadow for running specific code (coverage computation for instance) 
-// when vertices are produced by the strand hair vertex factory 
 // This is set by compilation enviromenent of the vertex factory
 ///////////////////////////////////////////////////////////////////////////////
 
@@ -58,7 +56,6 @@ struct FVertexFactoryInput
 #if RAYHITGROUPSHADER
 FVertexFactoryInput LoadVertexFactoryInputForHGS(uint TriangleIndex, int VertexIndex)
 {
-	// #hair_todo
 	FTriangleBaseAttributes TriangleAttributes = LoadTriangleBaseAttributes(TriangleIndex);
 
 	FVertexFactoryInput Input = (FVertexFactoryInput)0;
@@ -253,23 +250,6 @@ half3x3 CalcTangentToWorld(FVertexFactoryIntermediates Intermediates, half3x3 Ta
 	return TangentToWorld;
 }
 
-float3 GetTangent(FVertexFactoryInput Input)
-{
-	FVertexFactoryIntermediates Intermediates;
-	Intermediates.SceneData = VF_GPUSCENE_GET_INTERMEDIATES(Input);
-	
-	float TangentSign;
-	Intermediates.TangentToLocal = CalcTangentToLocal(Input, TangentSign);
-	Intermediates.TangentToWorld = CalcTangentToWorld(Intermediates,Intermediates.TangentToLocal);
-	Intermediates.TangentToWorldSign = TangentSign * GetPrimitive_DeterminantSign_FromFlags(GetPrimitiveData(Intermediates).Flags);
-
-	// Hair shader expec the tangent to be stored in place of the normal. This is what the StrandHairComponent is filling in:
-	// [0]: Real normal		| Should be Tangent
-	// [1]: Bitangent		| Should be BiTangent
-	// [2]: Real Tangent	| Should be Normal
-	return Intermediates.TangentToWorld[2].xyz;
-}
-
 FVertexFactoryIntermediates GetVertexFactoryIntermediates(FVertexFactoryInput Input)
 {
 	FVertexFactoryIntermediates Intermediates;
@@ -287,7 +267,6 @@ FVertexFactoryIntermediates GetVertexFactoryIntermediates(FVertexFactoryInput In
 	return Intermediates;
 }
 
-
 /**
 * Get the 3x3 tangent basis vectors for this vertex factory
 * this vertex factory will calculate the binormal on-the-fly
diff --git a/Engine/Source/Runtime/Renderer/Private/GPUScene.cpp b/Engine/Source/Runtime/Renderer/Private/GPUScene.cpp
--- a/Engine/Source/Runtime/Renderer/Private/GPUScene.cpp
+++ b/Engine/Source/Runtime/Renderer/Private/GPUScene.cpp
@@ -1391,7 +1391,7 @@ struct FUploadDataSourceAdapterDynamicPrimitives
 
 			InstanceUploadInfo.PrimitiveID					= PrimitiveIDStartOffset + ItemIndex;
 			InstanceUploadInfo.PrimitiveToWorld				= PrimitiveData.LocalToRelativeWorld;
-			InstanceUploadInfo.PrevPrimitiveToWorld			= InstanceUploadInfo.PrimitiveToWorld;
+			InstanceUploadInfo.PrevPrimitiveToWorld			= PrimitiveData.PreviousLocalToWorld;
 			InstanceUploadInfo.PrimitiveInstances			= TConstArrayView<FPrimitiveInstance>(&InstanceUploadInfo.DummyInstance, 1);
 			InstanceUploadInfo.InstanceDynamicData			= TConstArrayView<FPrimitiveInstanceDynamicData>();
 			InstanceUploadInfo.InstanceLightShadowUVBias	= TConstArrayView<FVector4f>();
