Fix hair crash during async loading.

The crash was happening while the LOD resource were loading, while the main/parent resource was also loading, causing a sanity check to fire.

diff --git a/Engine/Plugins/Runtime/HairStrands/Source/HairStrandsCore/Private/GroomResources.cpp b/Engine/Plugins/Runtime/HairStrands/Source/HairStrandsCore/Private/GroomResources.cpp
--- a/Engine/Plugins/Runtime/HairStrands/Source/HairStrandsCore/Private/GroomResources.cpp
+++ b/Engine/Plugins/Runtime/HairStrands/Source/HairStrandsCore/Private/GroomResources.cpp
@@ -520,10 +520,10 @@ void FHairCommonResource::Allocate(FRDGBuilder& GraphBuilder, EHairResourceLoadi
 
 void FHairCommonResource::AllocateLOD(FRDGBuilder& GraphBuilder, int32 LODIndex, EHairResourceLoadingType LoadingType, EHairResourceStatus& Status)
 {
-	if (LoadingType == EHairResourceLoadingType::Async && !InternalIsLODDataLoaded(LODIndex)) { Status |= EHairResourceStatus::Loading; return; }
+	// When using a async loading, sub-resources allocation requires the common/main resource to be already initialized.
+	if (LoadingType == EHairResourceLoadingType::Async && (!InternalIsLODDataLoaded(LODIndex) || !bIsInitialized)) { Status |= EHairResourceStatus::Loading; return; }
 
-	// Sanity check. For allocation sub-resources (like LOD resources) the common/main resource needs to be already initialized.
-	check(bIsInitialized);
+	// Sanity check. 
 	check(AllocationType == EHairStrandsAllocationType::Deferred);
 
 	InternalAllocateLOD(GraphBuilder, LODIndex);
