Full reset hair simulaiton when the input buffers are not ready yet to be used

diff --git a/Engine/Plugins/Runtime/HairStrands/Source/HairStrandsCore/Private/Niagara/NiagaraDataInterfaceHairStrands.cpp b/Engine/Plugins/Runtime/HairStrands/Source/HairStrandsCore/Private/Niagara/NiagaraDataInterfaceHairStrands.cpp
--- a/Engine/Plugins/Runtime/HairStrands/Source/HairStrandsCore/Private/Niagara/NiagaraDataInterfaceHairStrands.cpp
+++ b/Engine/Plugins/Runtime/HairStrands/Source/HairStrandsCore/Private/Niagara/NiagaraDataInterfaceHairStrands.cpp
@@ -323,6 +323,7 @@ void FNDIHairStrandsBuffer::Initialize(
 	SourceDeformedRootResources = HairStrandsDeformedRootResource;
 	ParamsScale = InParamsScale;
 	BoundingBoxOffsets = FIntVector4(0,1,2,3);
+	bValidGeometryType = false;
 }
 
 void FNDIHairStrandsBuffer::Update(
@@ -640,17 +641,17 @@ struct FNDIHairStrandsParametersCS : public FNiagaraDataInterfaceParametersCS
 		const bool bHasSkinningBinding = bIsHairValid && ProxyData->HairGroupInstance && ProxyData->HairGroupInstance->BindingType == EHairBindingType::Skinning;
 		const bool bIsRootValid = bIsHairValid && ProxyData->HairStrandsBuffer->SourceDeformedRootResources && ProxyData->HairStrandsBuffer->SourceDeformedRootResources->IsInitialized() && bHasSkinningBinding;
 		const bool bIsRestValid = bIsHairValid && ProxyData->HairStrandsBuffer->SourceRestResources && ProxyData->HairStrandsBuffer->SourceRestResources->IsInitialized() &&
-
+			 
 			// SourceRestResources->PositionBuffer is lazily allocated, when the instance LOD is scheduled (this happens after this called). So this is why this check is here. 
 			// This code should be refactor so that it reflects the lazy allocation scheme
 			ProxyData->HairStrandsBuffer->SourceRestResources->PositionBuffer.SRV && 
 			// TEMP: These check are only temporary for avoiding crashes while we find the bottom of the issue.
 			ProxyData->HairStrandsBuffer->CurvesOffsetsBuffer.SRV && ProxyData->HairStrandsBuffer->ParamsScaleBuffer.SRV && ProxyData->HairStrandsBuffer->BoundingBoxBuffer.UAV;
 
+		const bool bIsGeometryValid = bIsHairValid && ProxyData->HairGroupInstance && (ProxyData->HairGroupInstance->GeometryType != EHairGeometryType::NoneGeometry);
 		const bool bIsDeformedValid = bIsHairValid && ProxyData->HairStrandsBuffer->SourceDeformedResources && ProxyData->HairStrandsBuffer->SourceDeformedResources->IsInitialized();
-		const bool bHasLODSwitched = bIsHairValid && ProxyData->HairGroupInstance&& ProxyData->HairGroupInstance->HairGroupPublicData&& ProxyData->HairGroupInstance->HairGroupPublicData->VFInput.bHasLODSwitch;
 
-		if (bIsHairValid && bIsRestValid)
+		if (bIsHairValid && bIsRestValid && bIsGeometryValid)
 		{
 			check(ProxyData);
 
@@ -697,9 +698,11 @@ struct FNDIHairStrandsParametersCS : public FNiagaraDataInterfaceParametersCS
 				DeformedMeshProjection->MeshSampleWeightsBuffer.SRV.GetReference() : FNiagaraRenderer::GetDummyFloatBuffer();
 
 			// Simulation setup
-			const int32 NeedResetValue = (ProxyData->TickCount <= GHairSimulationMaxDelay);
+			const int32 NeedResetValue = (ProxyData->TickCount <= GHairSimulationMaxDelay) || !HairStrandsBuffer->bValidGeometryType;
 			const int32 RestUpdateValue = GHairSimulationRestUpdate;
 			const int32 LocalSimulationValue = ProxyData->LocalSimulation;
+
+			HairStrandsBuffer->bValidGeometryType = true;
 			
 			// Offsets / Transforms
 			FVector3f RestPositionOffsetValue = ProxyData->HairStrandsBuffer->SourceRestResources->GetPositionOffset();
@@ -762,6 +765,11 @@ struct FNDIHairStrandsParametersCS : public FNiagaraDataInterfaceParametersCS
 		}
 		else
 		{
+			if (bIsHairValid)
+			{
+				ProxyData->HairStrandsBuffer->bValidGeometryType = false;
+			}
+
 			// Set shader constants
 			SetShaderValue(RHICmdList, ComputeShaderRHI, BoundingBoxOffsets, FIntVector4(0,1,2,3));
 			SetShaderValue(RHICmdList, ComputeShaderRHI, WorldTransform, FMatrix44f::Identity);
diff --git a/Engine/Plugins/Runtime/HairStrands/Source/HairStrandsCore/Public/GroomAssetPhysics.h b/Engine/Plugins/Runtime/HairStrands/Source/HairStrandsCore/Public/GroomAssetPhysics.h
--- a/Engine/Plugins/Runtime/HairStrands/Source/HairStrandsCore/Public/GroomAssetPhysics.h
+++ b/Engine/Plugins/Runtime/HairStrands/Source/HairStrandsCore/Public/GroomAssetPhysics.h
@@ -313,7 +313,7 @@ struct HAIRSTRANDSCORE_API FHairSimulationConstraints
 	float BendDamping = 0.001;
 
 	/** Stiffness for the bend constraint in GPa */
-	UPROPERTY(EditAnywhere, BlueprintReadWrite, Interp, Category = "BendConstraint", meta = (ToolTip = "Stiffness for the bend constraint in GPa", EditCondition = "bOverrideSettings"))
+	UPROPERTY(EditAnywhere, BlueprintReadWrite, Interp, Category = "BendConstraint", meta = (ToolTip = "Stiffness for the bend constraint in GPa"))
 	float BendStiffness = 0.01;
 
 	/** Damping for the stretch constraint between 0 and 1 */
diff --git a/Engine/Plugins/Runtime/HairStrands/Source/HairStrandsCore/Public/Niagara/NiagaraDataInterfaceHairStrands.h b/Engine/Plugins/Runtime/HairStrands/Source/HairStrandsCore/Public/Niagara/NiagaraDataInterfaceHairStrands.h
--- a/Engine/Plugins/Runtime/HairStrands/Source/HairStrandsCore/Public/Niagara/NiagaraDataInterfaceHairStrands.h
+++ b/Engine/Plugins/Runtime/HairStrands/Source/HairStrandsCore/Public/Niagara/NiagaraDataInterfaceHairStrands.h
@@ -78,6 +78,9 @@ struct FNDIHairStrandsBuffer : public FRenderResource
 
 	/** Bounding box offsets */
 	FIntVector4 BoundingBoxOffsets;
+
+	/** Valid geometry type for hair (strands, cards, mesh)*/
+	bool bValidGeometryType = false;
 };
 
 /** Data stored per strand base instance*/
