Fix hair simulation jittering due to rigid position offset + enable debug simulation after setup

diff --git a/Engine/Plugins/Runtime/HairStrands/Shaders/Private/NiagaraDataInterfaceHairStrands.ush b/Engine/Plugins/Runtime/HairStrands/Shaders/Private/NiagaraDataInterfaceHairStrands.ush
--- a/Engine/Plugins/Runtime/HairStrands/Shaders/Private/NiagaraDataInterfaceHairStrands.ush
+++ b/Engine/Plugins/Runtime/HairStrands/Shaders/Private/NiagaraDataInterfaceHairStrands.ush
@@ -687,16 +687,44 @@ float4 DIHairStrands_TriangleWorldOrientation(in float3 TrianglePosition, in flo
 void DIHairStrands_UpdateTriangleDisplace(in FDIHairStrandsContext DIContext, in int StrandIndex, in float3 NodePosition, in float3 RestPosition, 
 		out float3 RestTrianglePosition, out float4 RestTriangleOrientation, out float3 DeformedTrianglePosition, out float4 DeformedTriangleOrientation)
 {
-	const float2 ProjectionUV = ExtractBarycentrics(DIContext.RootBarycentricCoordinatesBuffer[StrandIndex]);
+	if( DIContext.InterpolationMode >= 1)
+	{
+		const float2 ProjectionUV = ExtractBarycentrics(DIContext.RootBarycentricCoordinatesBuffer[StrandIndex]);
 	
-	DIHairStrands_BuildRestTriangle(DIContext, ProjectionUV, StrandIndex, RestTrianglePosition, RestTriangleOrientation);
-	DIHairStrands_BuildDeformedTriangle(DIContext, ProjectionUV, StrandIndex, DeformedTrianglePosition, DeformedTriangleOrientation);
+		DIHairStrands_BuildRestTriangle(DIContext, ProjectionUV, StrandIndex, RestTrianglePosition, RestTriangleOrientation);
+		DIHairStrands_BuildDeformedTriangle(DIContext, ProjectionUV, StrandIndex, DeformedTrianglePosition, DeformedTriangleOrientation);
 	
-	const float3 LocalRestPosition = DIHairStrands_TriangleLocalPosition(RestTrianglePosition, RestTriangleOrientation, RestPosition);
-	const float3 LocalDeformedPosition = DIHairStrands_TriangleLocalPosition(DeformedTrianglePosition, DeformedTriangleOrientation, NodePosition);
+		const float3 LocalRestPosition = DIHairStrands_TriangleLocalPosition(RestTrianglePosition, RestTriangleOrientation, RestPosition);
+		const float3 LocalDeformedPosition = DIHairStrands_TriangleLocalPosition(DeformedTrianglePosition, DeformedTriangleOrientation, NodePosition);
 	
-	SharedNodePosition[GGroupThreadId.x] = LocalDeformedPosition - LocalRestPosition;
-	GroupMemoryBarrier();
+		SharedNodePosition[GGroupThreadId.x] = LocalDeformedPosition - LocalRestPosition;
+		GroupMemoryBarrier();
+	}
+	else
+	{
+		SharedNodePosition[GGroupThreadId.x] = NodePosition - RestPosition;
+		GroupMemoryBarrier();
+	}
+}
+
+// Report interpolated nodes displacements onto the points positions 
+void DIHairStrands_UpdateDeformedPositions(in FDIHairStrandsContext DIContext, in float3 LocalDisplace, in int PositionIndex, 
+		in float3 RestTrianglePosition, in float4 RestTriangleOrientation, in float3 DeformedTrianglePosition, in float4 DeformedTriangleOrientation)
+{
+	if( DIContext.InterpolationMode >= 1)
+	{
+		const float3 LocalTrianglePosition = DIHairStrands_TriangleLocalPosition(RestTrianglePosition, RestTriangleOrientation, 
+			f16tof32(DIContext.RestPositionBuffer[PositionIndex].xyz) + DIContext.RestPositionOffset) + LocalDisplace;
+			
+		const float3 LocalComponentPosition = DIHairStrands_TriangleWorldPosition(DeformedTrianglePosition, DeformedTriangleOrientation, LocalTrianglePosition) - DIContext.DeformedPositionOffset[0].xyz;
+		DIContext.DeformedPositionBuffer[PositionIndex] = uint4(f32tof16(LocalComponentPosition), DIContext.RestPositionBuffer[PositionIndex].w);
+		
+	}
+	else
+	{
+		const float3 LocalComponentPosition = f16tof32(DIContext.RestPositionBuffer[PositionIndex].xyz) + LocalDisplace + DIContext.RestPositionOffset - DIContext.DeformedPositionOffset[0].xyz;
+		DIContext.DeformedPositionBuffer[PositionIndex] = uint4(f32tof16(LocalComponentPosition), DIContext.RestPositionBuffer[PositionIndex].w);	
+	}
 }
 
 // Report interpolated nodes displacements onto the points positions 
@@ -707,17 +735,18 @@ void DIHairStrands_UpdatePointPosition(in FDIHairStrandsContext DIContext, in fl
 	
 	int LocalIndex = GGroupThreadId.x % DIContext.StrandSize - 1;
 	int StrandIndex = ExecIndex() / DIContext.StrandSize;
-	
+
 	float3 RestTrianglePosition = float3(0, 0, 0), DeformedTrianglePosition = float3(0, 0, 0);
 	float4 RestTriangleOrientation = QUATERNION_IDENTITY, DeformedTriangleOrientation = QUATERNION_IDENTITY;
-	
+
 	DIHairStrands_UpdateTriangleDisplace(DIContext, StrandIndex, NodePosition, RestPosition, 
-		RestTrianglePosition, RestTriangleOrientation, DeformedTrianglePosition, DeformedTriangleOrientation);
+			RestTrianglePosition, RestTriangleOrientation, DeformedTrianglePosition, DeformedTriangleOrientation);
 	
 	const int PointBegin = DIContext.CurvesOffsetsBuffer[StrandIndex];
 	const int PointEnd = DIContext.CurvesOffsetsBuffer[StrandIndex + 1];
 	const int PointCount = PointEnd - PointBegin;
 
+
 	for (int PointIndex = 0; PointIndex < PointCount; ++PointIndex)
 	{
 		const float PointCoord = (float) (PointIndex) / (PointCount - 1.0);
@@ -731,12 +760,9 @@ void DIHairStrands_UpdatePointPosition(in FDIHairStrandsContext DIContext, in fl
 			OutReportStatus = true;
 			const float PointAlpha = PointNode - (float) NodePrev;
 			const float3 LocalDisplace = SharedNodePosition[GGroupThreadId.x] * (1.0 - PointAlpha) + SharedNodePosition[GGroupThreadId.x + 1] * PointAlpha;
-			
-			const float3 LocalPosition = DIHairStrands_TriangleLocalPosition(RestTrianglePosition, RestTriangleOrientation, 
-				f16tof32(DIContext.RestPositionBuffer[PointBegin + PointIndex].xyz) + DIContext.RestPositionOffset) + LocalDisplace;
-			
-			const float3 WorldPosition = DIHairStrands_TriangleWorldPosition(DeformedTrianglePosition, DeformedTriangleOrientation, LocalPosition) - DIContext.DeformedPositionOffset[0].xyz;
-			DIContext.DeformedPositionBuffer[PointBegin + PointIndex] = uint4(f32tof16(WorldPosition), DIContext.RestPositionBuffer[PointBegin + PointIndex].w);
+
+			DIHairStrands_UpdateDeformedPositions(DIContext, LocalDisplace, PointBegin + PointIndex, 
+				RestTrianglePosition, RestTriangleOrientation, DeformedTrianglePosition, DeformedTriangleOrientation);
 		}
 	}
 }
diff --git a/Engine/Plugins/Runtime/HairStrands/Source/HairStrandsCore/Private/GroomManager.cpp b/Engine/Plugins/Runtime/HairStrands/Source/HairStrandsCore/Private/GroomManager.cpp
--- a/Engine/Plugins/Runtime/HairStrands/Source/HairStrandsCore/Private/GroomManager.cpp
+++ b/Engine/Plugins/Runtime/HairStrands/Source/HairStrandsCore/Private/GroomManager.cpp
@@ -527,7 +527,7 @@ static void RunHairBufferSwap(const FHairStrandsInstances& Instances, const TArr
 
 		check(Instance->HairGroupPublicData);
 
-		// 1. Swap current/previous buffer for all LODs
+		// Swap current/previous buffer for all LODs
 		const bool bIsPaused = Views[0]->Family->bWorldIsPaused;
 		if (!bIsPaused)
 		{
@@ -535,31 +535,6 @@ static void RunHairBufferSwap(const FHairStrandsInstances& Instances, const TArr
 			if (Instance->Strands.DeformedResource) { Instance->Strands.DeformedResource->SwapBuffer(); }
 			Instance->Debug.LastFrameIndex = Views[0]->Family->FrameNumber;
 		}
-
-		// 2. Update the local offset (used for improving precision of strands data, as they are stored in 16bit precision)
-		// If attached to a mesh, ProxyLocalBound is mapped on the bounding box of the parent skeletal mesh. This offset is 
-		// based on the center of the skeletal mesh (which is computed based on the physics capsules/boxes/...)
-		if (Instance->bUpdatePositionOffset && Instance->ProxyLocalBounds && !bIsPaused)
-		{
-			const FVector PositionOffset = Instance->ProxyLocalBounds->GetSphere().Center;
-			if (Instance->Strands.DeformedResource)
-			{
-				Instance->Strands.DeformedResource->GetPositionOffset(FHairStrandsDeformedResource::Current) = PositionOffset;
-			}
-
-			if (Instance->Guides.DeformedResource)
-			{
-				Instance->Guides.DeformedResource->GetPositionOffset(FHairStrandsDeformedResource::Current) = PositionOffset;
-			}
-
-			for (uint32 LODIt = 0, LODCount = Instance->Cards.LODs.Num(); LODIt < LODCount; ++LODIt)
-			{
-				if (Instance->Cards.IsValid(LODIt) && Instance->Cards.LODs[LODIt].Guides.IsValid())
-				{
-					Instance->Cards.LODs[LODIt].Guides.DeformedResource->GetPositionOffset(FHairStrandsDeformedResource::Current) = PositionOffset;
-				}
-			}
-		}
 	}
 }
 
