Fix invalid local to world transform when groom are attached to skel. mesh.

* This changes store the local to world transform of a skel. mesh into its skin cached, ensuring that when when one retrieve skin cache data, it also get its valid local to world transform. This fixes laggy/invalid local to world transform previously used by the groom system.
* This changes also fixes cards sim/render guides display

diff --git a/Engine/Plugins/Runtime/HairStrands/Source/HairStrandsCore/Private/GroomComponent.cpp b/Engine/Plugins/Runtime/HairStrands/Source/HairStrandsCore/Private/GroomComponent.cpp
--- a/Engine/Plugins/Runtime/HairStrands/Source/HairStrandsCore/Private/GroomComponent.cpp
+++ b/Engine/Plugins/Runtime/HairStrands/Source/HairStrandsCore/Private/GroomComponent.cpp
@@ -45,7 +45,7 @@ static FAutoConsoleVariableRef CVarHairBindingValidationEnable(TEXT("r.HairStran
 static bool GUseGroomCacheStreaming = true;
 static FAutoConsoleVariableRef CVarGroomCacheStreamingEnable(TEXT("GroomCache.EnableStreaming"), GUseGroomCacheStreaming, TEXT("Enable groom cache streaming and prebuffering. Do not switch while groom caches are in use."));
 
-static bool GUseProxyLocalToWorld = false;
+static bool GUseProxyLocalToWorld = true;
 static FAutoConsoleVariableRef CVarUseProxyLocalToWorld(TEXT("r.HairStrands.UseProxyLocalToWorld"), GUseProxyLocalToWorld, TEXT("Enable the use of the groom proxy local to world instead of extracting it from the game thread."));
 
 
@@ -430,15 +430,6 @@ public:
 					++MeshesLODIndex;
 				}
 			}
-
-			// Initialization of the default skel. mesh transform (refresh during ticking).
-			HairInstance->Debug.SkinningCurrentLocalToWorld = HairInstance->Debug.MeshComponent ? HairInstance->Debug.MeshComponent->GetComponentTransform() : FTransform();
-			HairInstance->Debug.SkinningPreviousLocalToWorld = HairInstance->Debug.SkinningCurrentLocalToWorld;
-
-			HairInstance->Debug.RigidCurrentLocalToWorld = Component->GetComponentTransform();
-			HairInstance->Debug.RigidPreviousLocalToWorld = HairInstance->Debug.RigidCurrentLocalToWorld;
-
-			HairInstance->LocalToWorld = HairInstance->GetCurrentLocalToWorld();
 		}
 	}
 
@@ -494,14 +485,11 @@ public:
 
 	virtual void OnTransformChanged() override
 	{
-		if (GUseProxyLocalToWorld)
+		const FTransform RigidLocalToWorld = FTransform(GetLocalToWorld());
+		for (FHairGroupInstance* Instance : HairGroupInstances)
 		{
-			const FTransform RigidLocalToWorld = FTransform(GetLocalToWorld());
-			for (FHairGroupInstance* Instance : HairGroupInstances)
-			{
-				Instance->Debug.RigidPreviousLocalToWorld = Instance->Debug.RigidCurrentLocalToWorld;
-				Instance->Debug.RigidCurrentLocalToWorld = RigidLocalToWorld;
-			}
+			Instance->Debug.RigidPreviousLocalToWorld = Instance->Debug.RigidCurrentLocalToWorld;
+			Instance->Debug.RigidCurrentLocalToWorld = RigidLocalToWorld;
 		}
 	}
 
@@ -2063,20 +2051,6 @@ FHairStrandsRestRootResource* UGroomComponent::GetGuideStrandsRestRootResource(u
 	return HairGroupInstances[GroupIndex]->Guides.RestRootResource;
 }
 
-const FTransform& UGroomComponent::GetGuideStrandsLocalToWorld(uint32 GroupIndex) const
-{
-	if (GroupIndex >= uint32(HairGroupInstances.Num()))
-	{
-		return FTransform::Identity;
-	}
-
-	if (!HairGroupInstances[GroupIndex]->Guides.IsValid())
-	{
-		return FTransform::Identity;
-	}
-	return HairGroupInstances[GroupIndex]->LocalToWorld;
-}
-
 FHairStrandsDeformedRootResource* UGroomComponent::GetGuideStrandsDeformedRootResource(uint32 GroupIndex)
 {
 	if (GroupIndex >= uint32(HairGroupInstances.Num()))
@@ -2455,6 +2429,12 @@ void UGroomComponent::InitResources(bool bIsBindingReloading)
 		HairGroupInstance->GeometryType = EHairGeometryType::NoneGeometry;
 		HairGroupInstance->BindingType = EHairBindingType::NoneBinding;
 
+		HairGroupInstance->Debug.SkinningCurrentLocalToWorld	= RegisteredMeshComponent ? RegisteredMeshComponent->GetComponentTransform() : FTransform();
+		HairGroupInstance->Debug.SkinningPreviousLocalToWorld	= HairGroupInstance->Debug.SkinningCurrentLocalToWorld;
+		HairGroupInstance->Debug.RigidCurrentLocalToWorld		= GetComponentTransform();
+		HairGroupInstance->Debug.RigidPreviousLocalToWorld		= HairGroupInstance->Debug.RigidCurrentLocalToWorld;
+		HairGroupInstance->LocalToWorld							= HairGroupInstance->GetCurrentLocalToWorld();
+
 		FHairGroupData& GroupData = GroomAsset->HairGroupsData[GroupIt];
 
 		const EHairInterpolationType HairInterpolationType = ToHairInterpolationType(GroomAsset->HairInterpolationType);
@@ -3285,12 +3265,10 @@ void UGroomComponent::TickComponent(float DeltaTime, enum ELevelTick TickType, F
 		EffectiveForceLOD = LODForcedIndex;
 	}
 
-	const FTransform SkinningLocalToWorld = RegisteredMeshComponent ? RegisteredMeshComponent->GetComponentTransform() : FTransform();
-	const FTransform RigidLocalToWorld = GetComponentTransform();
 	TArray<FHairGroupInstance*> LocalHairGroupInstances = HairGroupInstances;
 	const EHairLODSelectionType LocalLODSelectionType = LODSelectionType;
 	ENQUEUE_RENDER_COMMAND(FHairStrandsTick_TransformUpdate)(
-		[Id, SkinningLocalToWorld, RigidLocalToWorld, FeatureLevel, LocalHairGroupInstances, EffectiveForceLOD, LocalLODSelectionType](FRHICommandListImmediate& RHICmdList)
+		[Id, FeatureLevel, LocalHairGroupInstances, EffectiveForceLOD, LocalLODSelectionType](FRHICommandListImmediate& RHICmdList)
 	{
 		if (ERHIFeatureLevel::Num == FeatureLevel)
 			return;
@@ -3299,16 +3277,6 @@ void UGroomComponent::TickComponent(float DeltaTime, enum ELevelTick TickType, F
 		{
 			Instance->Debug.LODForcedIndex = EffectiveForceLOD;
 			Instance->Debug.LODSelectionTypeForDebug = LocalLODSelectionType;
-
-			Instance->Debug.SkinningPreviousLocalToWorld = Instance->Debug.SkinningCurrentLocalToWorld;
-			Instance->Debug.SkinningCurrentLocalToWorld = SkinningLocalToWorld;
-
-			if (!GUseProxyLocalToWorld)
-			{
-				Instance->Debug.RigidPreviousLocalToWorld = Instance->Debug.RigidCurrentLocalToWorld;
-				Instance->Debug.RigidCurrentLocalToWorld = RigidLocalToWorld;
-			}
-
 			if (LocalLODSelectionType == EHairLODSelectionType::Predicted && EffectiveForceLOD >= 0)
 			{
 				AddHairStreamingRequest(Instance, EffectiveForceLOD);
diff --git a/Engine/Plugins/Runtime/HairStrands/Source/HairStrandsCore/Private/GroomDebug.cpp b/Engine/Plugins/Runtime/HairStrands/Source/HairStrandsCore/Private/GroomDebug.cpp
--- a/Engine/Plugins/Runtime/HairStrands/Source/HairStrandsCore/Private/GroomDebug.cpp
+++ b/Engine/Plugins/Runtime/HairStrands/Source/HairStrandsCore/Private/GroomDebug.cpp
@@ -100,7 +100,7 @@ static void GetGroomInterpolationData(
 		{
 			for (const FCachedGeometry::Section& Section : CachedGeometry.Sections)
 			{
-				FHairStrandsProjectionMeshData::Section OutSection = ConvertMeshSection(Section);
+				FHairStrandsProjectionMeshData::Section OutSection = ConvertMeshSection(Section, CachedGeometry.LocalToWorld);
 				if (MeshType == EHairStrandsProjectionMeshType::RestMesh)
 				{					
 					// If the mesh has some mesh-tranferred data, we display that otherwise we use the rest data
@@ -707,11 +707,11 @@ class FDrawDebugCardGuidesCS : public FGlobalShader
 		
 		SHADER_PARAMETER(uint32,  RenVertexCount)
 		SHADER_PARAMETER(FVector3f, RenRestOffset)
-		SHADER_PARAMETER(FVector3f, RenDeformedOffset)
+		SHADER_PARAMETER_RDG_BUFFER_SRV(Buffer, RenDeformedOffset)
 		
 		SHADER_PARAMETER(uint32,  SimVertexCount)
 		SHADER_PARAMETER(FVector3f, SimRestOffset)
-		SHADER_PARAMETER(FVector3f, SimDeformedOffset)
+		SHADER_PARAMETER_RDG_BUFFER_SRV(Buffer, SimDeformedOffset)
 
 		SHADER_PARAMETER_RDG_BUFFER_SRV(Buffer, RenRestPosition)
 		SHADER_PARAMETER_RDG_BUFFER_SRV(Buffer, RenDeformedPosition)
@@ -791,13 +791,13 @@ static void AddDrawDebugCardsGuidesPass(
 	Parameters->RenVertexCount = 0;
 	Parameters->RenRestOffset = FVector3f::ZeroVector;
 	Parameters->RenRestPosition = DefaultBuffer;
-	Parameters->RenDeformedOffset = FVector3f::ZeroVector;
+	Parameters->RenDeformedOffset = DefaultBuffer;
 	Parameters->RenDeformedPosition = DefaultBuffer;
 
 	Parameters->SimVertexCount = 0;
 	Parameters->SimRestOffset = FVector3f::ZeroVector;
 	Parameters->SimRestPosition = DefaultBuffer;
-	Parameters->SimDeformedOffset = FVector3f::ZeroVector;
+	Parameters->SimDeformedOffset = DefaultBuffer;
 	Parameters->SimDeformedPosition = DefaultBuffer;
 
 	if (bRen)
@@ -807,7 +807,7 @@ static void AddDrawDebugCardsGuidesPass(
 		Parameters->RenRestPosition = RegisterAsSRV(GraphBuilder, LOD.Guides.RestResource->PositionBuffer);
 		if (bDeformed)
 		{
-			Parameters->RenDeformedOffset = LOD.Guides.DeformedResource->GetPositionOffset(FHairStrandsDeformedResource::Current);
+			Parameters->RenDeformedOffset = RegisterAsSRV(GraphBuilder, LOD.Guides.DeformedResource->GetPositionOffsetBuffer(FHairStrandsDeformedResource::Current));
 			Parameters->RenDeformedPosition = RegisterAsSRV(GraphBuilder, LOD.Guides.DeformedResource->GetBuffer(FHairStrandsDeformedResource::Current));
 		}
 	}
@@ -819,22 +819,23 @@ static void AddDrawDebugCardsGuidesPass(
 		Parameters->SimRestPosition = RegisterAsSRV(GraphBuilder, Instance->Guides.RestResource->PositionBuffer);
 		if (bDeformed)
 		{
-			Parameters->SimDeformedOffset = Instance->Guides.DeformedResource->GetPositionOffset(FHairStrandsDeformedResource::Current);
+			Parameters->SimDeformedOffset = RegisterAsSRV(GraphBuilder, Instance->Guides.DeformedResource->GetPositionOffsetBuffer(FHairStrandsDeformedResource::Current));
 			Parameters->SimDeformedPosition = RegisterAsSRV(GraphBuilder, Instance->Guides.DeformedResource->GetBuffer(FHairStrandsDeformedResource::Current));
 		}
 	}
 
 	Parameters->LocalToWorld = Instance->LocalToWorld.ToMatrixWithScale();
 
-	if (!bDeformed &&  bRen) Parameters->DebugMode = 1;
-	if ( bDeformed &&  bRen) Parameters->DebugMode = 2;
-	if (!bDeformed && !bRen) Parameters->DebugMode = 3;
-	if ( bDeformed && !bRen) Parameters->DebugMode = 4;
+	const TCHAR* DebugName = TEXT("Unknown");
+	if (!bDeformed &&  bRen) { Parameters->DebugMode = 1; DebugName = TEXT("Ren, Rest"); } 
+	if ( bDeformed &&  bRen) { Parameters->DebugMode = 2; DebugName = TEXT("Ren, Deformed"); } 
+	if (!bDeformed && !bRen) { Parameters->DebugMode = 3; DebugName = TEXT("Sim, Rest"); } 
+	if ( bDeformed && !bRen) { Parameters->DebugMode = 4; DebugName = TEXT("Sim, Deformed"); } 
 
 	ShaderDrawDebug::SetParameters(GraphBuilder, *ShaderDrawData, Parameters->ShaderDrawParameters);
 
 	const uint32 VertexCount = Parameters->DebugMode <= 2 ? Parameters->RenVertexCount : Parameters->SimVertexCount;
-	FComputeShaderUtils::AddPass(GraphBuilder, RDG_EVENT_NAME("DrawDebugCardsAtlas"), ComputeShader, Parameters,
+	FComputeShaderUtils::AddPass(GraphBuilder, RDG_EVENT_NAME("HairStrands::DebugCards(%s)", DebugName), ComputeShader, Parameters,
 	FIntVector::DivideAndRoundUp(FIntVector(VertexCount, 1, 1), FIntVector(32, 1, 1)));
 }
 
diff --git a/Engine/Plugins/Runtime/HairStrands/Source/HairStrandsCore/Private/GroomGeometryCache.cpp b/Engine/Plugins/Runtime/HairStrands/Source/HairStrandsCore/Private/GroomGeometryCache.cpp
--- a/Engine/Plugins/Runtime/HairStrands/Source/HairStrandsCore/Private/GroomGeometryCache.cpp
+++ b/Engine/Plugins/Runtime/HairStrands/Source/HairStrandsCore/Private/GroomGeometryCache.cpp
@@ -16,7 +16,7 @@
 
 ///////////////////////////////////////////////////////////////////////////////////////////////////
 
-FHairStrandsProjectionMeshData::Section ConvertMeshSection(const FCachedGeometry::Section& In)
+FHairStrandsProjectionMeshData::Section ConvertMeshSection(const FCachedGeometry::Section& In, const FTransform& InLocalToWorld)
 {
 	FHairStrandsProjectionMeshData::Section Out;
 	Out.IndexBuffer = In.IndexBuffer;
@@ -33,6 +33,7 @@ FHairStrandsProjectionMeshData::Section ConvertMeshSection(const FCachedGeometry
 	Out.NumVertices= In.NumVertices;
 	Out.SectionIndex = In.SectionIndex;
 	Out.LODIndex = In.LODIndex;
+	Out.LocalToWorld = InLocalToWorld;
 	return Out;
 }
 
@@ -92,7 +93,7 @@ static void BuildBoneMatrices(USkeletalMeshComponent* SkeletalMeshComponent, con
 
 		CachedGeometry.DeformedPositionBuffer = DeformedPositionsBuffer;
 		FRDGBufferSRVRef DeformedPositionSRV = GraphBuilder.CreateSRV(DeformedPositionsBuffer, PF_R32_FLOAT);
-
+		CachedGeometry.LocalToWorld = SkeletalMeshComponent->SceneProxy ? FTransform(SkeletalMeshComponent->SceneProxy->GetLocalToWorld()) : FTransform();
 		for (int32 SectionIdx = 0; SectionIdx < LODData.RenderSections.Num(); ++SectionIdx)
 		{
 			FCachedGeometry::Section CachedSection;
@@ -141,7 +142,7 @@ static void BuildBoneMatrices(USkeletalMeshComponent* SkeletalMeshComponent, con
 
 			 // PositionBuffer depends on CurrentPositionBufferIndex and on if the cache has motion vectors
 			 const uint32 PositionIndex = (TrackProxy->CurrentPositionBufferIndex == -1 || bHasMotionVectors) ? 0 : TrackProxy->CurrentPositionBufferIndex % 2;
-
+			 CachedGeometry.LocalToWorld = FTransform(SceneProxy->GetLocalToWorld());
 			 for (int32 SectionIdx = 0; SectionIdx < TrackProxy->MeshData->BatchesInfo.Num(); ++SectionIdx)
 			 {
 				const FGeometryCacheMeshBatchInfo& BatchInfo = TrackProxy->MeshData->BatchesInfo[SectionIdx];
diff --git a/Engine/Plugins/Runtime/HairStrands/Source/HairStrandsCore/Private/GroomGeometryCache.h b/Engine/Plugins/Runtime/HairStrands/Source/HairStrandsCore/Private/GroomGeometryCache.h
--- a/Engine/Plugins/Runtime/HairStrands/Source/HairStrandsCore/Private/GroomGeometryCache.h
+++ b/Engine/Plugins/Runtime/HairStrands/Source/HairStrandsCore/Private/GroomGeometryCache.h
@@ -12,7 +12,7 @@
 #include "HairStrandsMeshProjection.h"
 #include "GPUSkinCache.h"
 
-FHairStrandsProjectionMeshData::Section ConvertMeshSection(const FCachedGeometry::Section& In);
+FHairStrandsProjectionMeshData::Section ConvertMeshSection(const FCachedGeometry::Section& In, const FTransform& T);
 
 void BuildCacheGeometry(
 	FRDGBuilder& GraphBuilder,
diff --git a/Engine/Plugins/Runtime/HairStrands/Source/HairStrandsCore/Private/GroomManager.cpp b/Engine/Plugins/Runtime/HairStrands/Source/HairStrandsCore/Private/GroomManager.cpp
--- a/Engine/Plugins/Runtime/HairStrands/Source/HairStrandsCore/Private/GroomManager.cpp
+++ b/Engine/Plugins/Runtime/HairStrands/Source/HairStrandsCore/Private/GroomManager.cpp
@@ -156,7 +156,7 @@ static void RunInternalHairStrandsInterpolation(
 			if (MeshLODIndex < 0) MeshLODIndex = Section.LODIndex;
 			check(MeshLODIndex == Section.LODIndex);
 
-			MeshDataLOD.Sections.Add(ConvertMeshSection(Section));
+			MeshDataLOD.Sections.Add(ConvertMeshSection(Section, CachedGeometry.LocalToWorld));
 		}
 
 		Instance->Debug.MeshLODIndex = MeshLODIndex;
@@ -707,11 +707,6 @@ static void RunHairLODSelection(
 		const FCachedGeometry CachedGeometry = GetCacheGeometryForHair(GraphBuilder, Instance, SkinCache, ShaderMap);
 		const int32 MeshLODIndex = GetCacheGeometryLODIndex(CachedGeometry);
 
-		//if (Instance->ProxyLocalToWorld)
-		//{
-		//	Instance->LocalToWorld = *Instance->ProxyLocalToWorld;
-		//}
-
 		// Perform LOD selection based on all the views	
 		// CPU LOD selection. 
 		// * When enable the CPU LOD selection allow to change the geometry representation. 
@@ -898,6 +893,8 @@ static void RunHairLODSelection(
 		}
 
 		// Update the local-to-world transform based on the binding type 
+		Instance->Debug.SkinningPreviousLocalToWorld = Instance->Debug.SkinningCurrentLocalToWorld;
+		Instance->Debug.SkinningCurrentLocalToWorld = CachedGeometry.LocalToWorld;
 		Instance->LocalToWorld = Instance->GetCurrentLocalToWorld();
 	}
 }
diff --git a/Engine/Plugins/Runtime/HairStrands/Source/HairStrandsCore/Private/HairStrandsRendering.cpp b/Engine/Plugins/Runtime/HairStrands/Source/HairStrandsCore/Private/HairStrandsRendering.cpp
--- a/Engine/Plugins/Runtime/HairStrands/Source/HairStrandsCore/Private/HairStrandsRendering.cpp
+++ b/Engine/Plugins/Runtime/HairStrands/Source/HairStrandsCore/Private/HairStrandsRendering.cpp
@@ -1027,6 +1027,7 @@ class FHairCardsDeformationCS : public FGlobalShader
 	using FPermutationDomain = TShaderPermutationDomain<FGroupSize, FDynamicGeometry>;
 
 	BEGIN_SHADER_PARAMETER_STRUCT(FParameters, )
+		SHADER_PARAMETER(FMatrix44f, LocalToWorld)
 		SHADER_PARAMETER(uint32, CardsVertexCount)
 		SHADER_PARAMETER(uint32, GuideVertexCount)
 		SHADER_PARAMETER(FVector3f, GuideRestPositionOffset)
@@ -1086,6 +1087,7 @@ static void AddHairCardsDeformationPass(
 	FRDGImportedBuffer CardsDeformedNormalBuffer = Register(GraphBuilder, LOD.DeformedResource->DeformedNormalBuffer, ERDGImportedBufferFlags::CreateViews);
 
 	FHairCardsDeformationCS::FParameters* Parameters = GraphBuilder.AllocParameters<FHairCardsDeformationCS::FParameters>();
+	Parameters->LocalToWorld				= Instance->GetCurrentLocalToWorld().ToMatrixWithScale();
 	Parameters->GuideVertexCount			= LOD.Guides.RestResource->GetVertexCount();
 	Parameters->GuideRestPositionOffset		= LOD.Guides.RestResource->GetPositionOffset();
 	Parameters->GuideRestPositionBuffer		= RegisterAsSRV(GraphBuilder, LOD.Guides.RestResource->PositionBuffer);
diff --git a/Engine/Plugins/Runtime/HairStrands/Source/HairStrandsCore/Public/GroomComponent.h b/Engine/Plugins/Runtime/HairStrands/Source/HairStrandsCore/Public/GroomComponent.h
--- a/Engine/Plugins/Runtime/HairStrands/Source/HairStrandsCore/Public/GroomComponent.h
+++ b/Engine/Plugins/Runtime/HairStrands/Source/HairStrandsCore/Public/GroomComponent.h
@@ -140,7 +140,6 @@ public:
 	/** Return the guide hairs root resources*/
 	FHairStrandsRestRootResource* GetGuideStrandsRestRootResource(uint32 GroupIndex);
 	FHairStrandsDeformedRootResource* GetGuideStrandsDeformedRootResource(uint32 GroupIndex);
-	const FTransform& GetGuideStrandsLocalToWorld(uint32 GroupIndex) const;
 
 
 #if WITH_EDITOR
diff --git a/Engine/Shaders/Private/HairStrands/HairCardsDebug.usf b/Engine/Shaders/Private/HairStrands/HairCardsDebug.usf
--- a/Engine/Shaders/Private/HairStrands/HairCardsDebug.usf
+++ b/Engine/Shaders/Private/HairStrands/HairCardsDebug.usf
@@ -94,19 +94,19 @@ void MainCS(uint3 DispatchThreadId : SV_DispatchThreadID)
 
 uint RenVertexCount;
 float3 RenRestOffset;
-float3 RenDeformedOffset;
 
 uint SimVertexCount;
 float3 SimRestOffset;
-float3 SimDeformedOffset;
 
 float4x4 LocalToWorld;
 
 Buffer<uint4> RenRestPosition;
 Buffer<uint4> RenDeformedPosition;
+Buffer<float4> RenDeformedOffset;
 
 Buffer<uint4> SimRestPosition;
 Buffer<uint4> SimDeformedPosition;
+Buffer<float4> SimDeformedOffset;
 
 float3 RestOffset;
 float3 DeformedOffset;
@@ -115,11 +115,16 @@ Buffer<uint4> DeformedPosition;
 
 int DebugMode;
 
+#define DEBUGMODE_REN_REST 1
+#define DEBUGMODE_REN_DEFORMED 2
+#define DEBUGMODE_SIM_REST 3
+#define DEBUGMODE_SIM_DEFORMED 4
+
 [numthreads(32, 1, 1)]
 void MainCS(uint3 DispatchThreadId : SV_DispatchThreadID)
 {
 	const uint VertexId = DispatchThreadId.x;
-	const uint VertexCount = DebugMode == 0 || DebugMode == 1 ? RenVertexCount : SimVertexCount;
+	const uint VertexCount = DebugMode == DEBUGMODE_REN_REST || DebugMode == DEBUGMODE_REN_DEFORMED ? RenVertexCount : SimVertexCount;
 	if (VertexId+1 >= VertexCount)
 		return;
 
@@ -128,32 +133,32 @@ void MainCS(uint3 DispatchThreadId : SV_DispatchThreadID)
 	float3 Offset = 0;
 	float4 Color = float4(1, 1, 0, 1);
 
-	if (DebugMode == 1)
+	if (DebugMode == DEBUGMODE_REN_REST)
 	{
 		PackedP0 = RenRestPosition[VertexId];
 		PackedP1 = RenRestPosition[VertexId + 1];
 		Offset   = RenRestOffset;
 		Color	 = float4(1, 1, 0, 1);
 	}
-	else if (DebugMode == 2)
+	else if (DebugMode == DEBUGMODE_REN_DEFORMED)
 	{
 		PackedP0 = RenDeformedPosition[VertexId];
 		PackedP1 = RenDeformedPosition[VertexId + 1];
-		Offset   = RenDeformedOffset;
+		Offset   = RenDeformedOffset[0].xyz;
 		Color	 = float4(1, 1, 0, 1);
 	}
-	else if (DebugMode == 3)
+	else if (DebugMode == DEBUGMODE_SIM_REST)
 	{
 		PackedP0 = SimRestPosition[VertexId];
 		PackedP1 = SimRestPosition[VertexId + 1];
 		Offset   = SimRestOffset;
 		Color	 = float4(1, 0, 1, 1);
 	}
-	else if (DebugMode == 4)
+	else if (DebugMode == DEBUGMODE_SIM_DEFORMED)
 	{
 		PackedP0 = SimDeformedPosition[VertexId];
 		PackedP1 = SimDeformedPosition[VertexId + 1];
-		Offset   = SimDeformedOffset;
+		Offset   = SimDeformedOffset[0].xyz;
 		Color	 = float4(1, 0, 1, 1);
 	}
 
diff --git a/Engine/Shaders/Private/HairStrands/HairCardsDeformation.usf b/Engine/Shaders/Private/HairStrands/HairCardsDeformation.usf
--- a/Engine/Shaders/Private/HairStrands/HairCardsDeformation.usf
+++ b/Engine/Shaders/Private/HairStrands/HairCardsDeformation.usf
@@ -4,6 +4,14 @@
 #include "HairStrandsVertexFactoryCommon.ush"
 #include "HairStrandsMeshProjectionCommon.ush"
 
+#define DEBUG_ENABLE 0
+
+#if DEBUG_ENABLE
+#include "../ShaderDrawDebug.ush"?
+#endif
+
+float4x4 LocalToWorld;
+
 uint CardsVertexCount;
 uint GuideVertexCount;
 float3 GuideRestPositionOffset;
@@ -161,10 +169,19 @@ void MainCS(uint2 DispatchThreadId : SV_DispatchThreadID)
 	ControlPoint = ToTriangle(ControlPoint, RootBarycentric, RestTriangle);
 	ControlPoint = ToWorld(ControlPoint, RootBarycentric, DeformedTriangle) + CurrOffset;
 
+	#if DEBUG_ENABLE
+	{
+		float3 WP0 = mul(float4(DeformedTriangle.P0, 1), LocalToWorld).xyz;
+		float3 WP1 = mul(float4(DeformedTriangle.P1, 1), LocalToWorld).xyz;
+		float3 WP2 = mul(float4(DeformedTriangle.P2, 1), LocalToWorld).xyz;
+		AddLineTriangle(WP0, WP1, WP2, ColorOrange);
+	}
+	#endif
+
 	CardsDeformedPositionBuffer[VertexId] = float4(ControlPoint, RestControlPoint.w);
 
 	// Transform tangent
 	TransformTangent(VertexId, 0, RestTriangle, DeformedTriangle);
 	TransformTangent(VertexId, 1, RestTriangle, DeformedTriangle);
-#endif
+#endif // PERMUTATION_DYNAMIC_GEOMETRY
 }
\ No newline at end of file
diff --git a/Engine/Source/Runtime/Engine/Private/GPUSkinCache.cpp b/Engine/Source/Runtime/Engine/Private/GPUSkinCache.cpp
--- a/Engine/Source/Runtime/Engine/Private/GPUSkinCache.cpp
+++ b/Engine/Source/Runtime/Engine/Private/GPUSkinCache.cpp
@@ -2317,6 +2317,7 @@ FCachedGeometry FGPUSkinCache::GetCachedGeometry(uint32 ComponentId) const
 					CachedSection.UVsChannelCount	= LODData.StaticVertexBuffers.StaticMeshVertexBuffer.GetNumTexCoords();
 					Out.Sections.Add(CachedSection);
 				}
+				Out.LocalToWorld = FTransform(Entry->GPUSkin->GetTransform());
 				break;
 			}
 		}
diff --git a/Engine/Source/Runtime/Engine/Private/SkeletalMesh.cpp b/Engine/Source/Runtime/Engine/Private/SkeletalMesh.cpp
--- a/Engine/Source/Runtime/Engine/Private/SkeletalMesh.cpp
+++ b/Engine/Source/Runtime/Engine/Private/SkeletalMesh.cpp
@@ -6941,10 +6941,10 @@ bool FSkeletalMeshSceneProxy::GetMeshUVDensities(int32 LODIndex, int32 SectionIn
 				LocalUVDensities[3] * TransformScale);
 			
 			return true;
-	}
 		}
-	return FPrimitiveSceneProxy::GetMeshUVDensities(LODIndex, SectionIndex, WorldUVDensities);
 	}
+	return FPrimitiveSceneProxy::GetMeshUVDensities(LODIndex, SectionIndex, WorldUVDensities);
+}
 
 bool FSkeletalMeshSceneProxy::GetMaterialTextureScales(int32 LODIndex, int32 SectionIndex, const FMaterialRenderProxy* MaterialRenderProxy, FVector4f* OneOverScales, FIntVector4* UVChannelIndices) const
 	{
@@ -6982,6 +6982,7 @@ bool FSkeletalMeshSceneProxy::GetMaterialTextureScales(int32 LODIndex, int32 Sec
 void FSkeletalMeshSceneProxy::OnTransformChanged()
 {
 	// OnTransformChanged is called on the following frame after FSkeletalMeshObject::Update(), thus omit '+ 1' to frame number.
+	MeshObject->SetTransform(GetLocalToWorld(), GetScene().GetFrameNumber());
 	MeshObject->RefreshClothingTransforms(GetLocalToWorld(), GetScene().GetFrameNumber());
 
 	// Update the default-instance
diff --git a/Engine/Source/Runtime/Engine/Private/SkeletalRenderGPUSkin.cpp b/Engine/Source/Runtime/Engine/Private/SkeletalRenderGPUSkin.cpp
--- a/Engine/Source/Runtime/Engine/Private/SkeletalRenderGPUSkin.cpp
+++ b/Engine/Source/Runtime/Engine/Private/SkeletalRenderGPUSkin.cpp
@@ -1279,6 +1279,23 @@ FSkinWeightVertexBuffer* FSkeletalMeshObjectGPUSkin::GetSkinWeightVertexBuffer(i
 	return LODs[LODIndex].MeshObjectWeightBuffer;
 }
 
+FMatrix FSkeletalMeshObjectGPUSkin::GetTransform() const
+{
+	if (DynamicData)
+	{
+		return DynamicData->LocalToWorld;
+	}
+	return FMatrix();
+}
+
+void FSkeletalMeshObjectGPUSkin::SetTransform(const FMatrix& InNewLocalToWorld, uint32 FrameNumber)
+{
+	if (DynamicData)
+	{
+		DynamicData->LocalToWorld = InNewLocalToWorld;
+	}
+}
+
 void FSkeletalMeshObjectGPUSkin::RefreshClothingTransforms(const FMatrix& InNewLocalToWorld, uint32 FrameNumber)
 {
 	if(DynamicData && DynamicData->ClothingSimData.Num() > 0)
@@ -1922,6 +1939,7 @@ void FDynamicSkelMeshObjectDataGPUSkin::Clear()
 #if RHI_RAYTRACING
 	bAnySegmentUsesWorldPositionOffset = false;
 #endif
+	LocalToWorld = FMatrix::Identity;
 }
 
 #define SKELETON_POOL_GPUSKINS 1
@@ -2098,6 +2116,10 @@ void FDynamicSkelMeshObjectDataGPUSkin::InitDynamicSkelMeshObjectDataGPUSkin(
 			ActiveMorphTargets.RemoveAt(MorphIdx, 1, false);
 		}
 	}
+
+	// Update local to world transform
+	LocalToWorld = InMeshComponent ? InMeshComponent->GetComponentTransform().ToMatrixWithScale() : FMatrix::Identity;
+
 	// Update the clothing simulation mesh positions and normals
 	UpdateClothSimulationData(InMeshComponent);
 
diff --git a/Engine/Source/Runtime/Engine/Private/SkeletalRenderGPUSkin.h b/Engine/Source/Runtime/Engine/Private/SkeletalRenderGPUSkin.h
--- a/Engine/Source/Runtime/Engine/Private/SkeletalRenderGPUSkin.h
+++ b/Engine/Source/Runtime/Engine/Private/SkeletalRenderGPUSkin.h
@@ -98,6 +98,9 @@ public:
     /** store transform of the cloth object **/
     FMatrix ClothObjectLocalToWorld;
 
+	/** store transform of the object **/
+	FMatrix LocalToWorld;
+
 	/** a weight factor to blend between simulated positions and skinned positions */	
 	float ClothBlendWeight;
 
@@ -410,6 +413,8 @@ public:
 		uint32 NumVertices = 0;
 	};
 
+	FMatrix GetTransform() const;
+	virtual void SetTransform(const FMatrix& InNewLocalToWorld, uint32 FrameNumber) override;
 	virtual void RefreshClothingTransforms(const FMatrix& InNewLocalToWorld, uint32 FrameNumber) override;
 	virtual void UpdateSkinWeightBuffer(USkinnedMeshComponent* InMeshComponent) override;
 protected:
diff --git a/Engine/Source/Runtime/Engine/Public/GPUSkinCache.h b/Engine/Source/Runtime/Engine/Public/GPUSkinCache.h
--- a/Engine/Source/Runtime/Engine/Public/GPUSkinCache.h
+++ b/Engine/Source/Runtime/Engine/Public/GPUSkinCache.h
@@ -114,6 +114,7 @@ struct FCachedGeometry
 	int32 LODIndex = 0;
 	TArray<Section> Sections;
 	FRDGBufferRef DeformedPositionBuffer = nullptr;
+	FTransform LocalToWorld = FTransform::Identity;
 };
 
 enum class EGPUSkinCacheEntryMode
diff --git a/Engine/Source/Runtime/Engine/Public/SkeletalRenderPublic.h b/Engine/Source/Runtime/Engine/Public/SkeletalRenderPublic.h
--- a/Engine/Source/Runtime/Engine/Public/SkeletalRenderPublic.h
+++ b/Engine/Source/Runtime/Engine/Public/SkeletalRenderPublic.h
@@ -212,6 +212,9 @@ public:
 	virtual FGPUSkinCacheEntry* GetSkinCacheEntryForRayTracing() const { return ShouldUseSeparateSkinCacheEntryForRayTracing() ? SkinCacheEntryForRayTracing : SkinCacheEntry; }
 #endif // RHI_RAYTRACING
 
+	/** Called when that component transform has changed */
+	virtual void SetTransform(const FMatrix& InNewLocalToWorld, uint32 FrameNumber) {};
+
 	/** Called to notify clothing data that component transform has changed */
 	virtual void RefreshClothingTransforms(const FMatrix& InNewLocalToWorld, uint32 FrameNumber) {};
 
