Improve hair stability when using large world coordinate

diff --git a/Engine/Plugins/Runtime/HairStrands/Shaders/Private/NiagaraDataInterfaceHairStrands.ush b/Engine/Plugins/Runtime/HairStrands/Shaders/Private/NiagaraDataInterfaceHairStrands.ush
--- a/Engine/Plugins/Runtime/HairStrands/Shaders/Private/NiagaraDataInterfaceHairStrands.ush
+++ b/Engine/Plugins/Runtime/HairStrands/Shaders/Private/NiagaraDataInterfaceHairStrands.ush
@@ -334,11 +334,6 @@ void DIHairStrands_SetNodePosition(in FDIHairStrandsContext DIContext, in bool R
 						   f16tof32(DIContext.RestPositionBuffer[PointOffset+PointNext].xyz) * NodeAlpha;
 	}
 	OutNodePosition = mul(float4( OutNodePosition+DIContext.RestPositionOffset, 1.0), DIContext.WorldTransform).xyz;
-
-	if(DIContext.LocalSimulation)
-	{
-		OutNodePosition = mul(float4( OutNodePosition, 1.0), DIContext.BoneInverse).xyz;
-	}
 }
 
 // Compute the node position by linear interpolation over the points
@@ -841,35 +836,25 @@ void DIHairStrands_UpdateAngularVelocity(in FDIHairStrandsContext DIContext, in
 // Compute rest position
 void DIHairStrands_ComputeRestPosition(in FDIHairStrandsContext DIContext, in float3 NodePosition, out float3 OutRestPosition)
 {	
-	OutRestPosition = DIContext.LocalSimulation ? mul(float4( NodePosition, 1.0), DIContext.BoneTransform).xyz : NodePosition;
-	OutRestPosition = mul(float4( OutRestPosition, 1.0), DIContext.WorldInverse).xyz;
+	OutRestPosition = mul(float4( NodePosition, 1.0), DIContext.WorldInverse).xyz;
 }
 
 // Compute rest orientation
 void DIHairStrands_ComputeRestOrientation(in FDIHairStrandsContext DIContext, in float4 NodeOrientation, out float4 OutRestOrientation)
 {	
-	OutRestOrientation = DIContext.LocalSimulation ? NormalizeQuat( MultiplyQuat(DIContext.BoneRotation,NodeOrientation) ) : NodeOrientation;
-	OutRestOrientation = NormalizeQuat( MultiplyQuat(InverseQuat(DIContext.WorldRotation),OutRestOrientation) );
+	OutRestOrientation = NormalizeQuat( MultiplyQuat(InverseQuat(DIContext.WorldRotation),NodeOrientation) );
 }
 
 // Update Node Position
 void DIHairStrands_AttachNodePosition(in FDIHairStrandsContext DIContext, in float3 RestPosition, out float3 OutNodePosition)
 {
 	OutNodePosition = mul(float4( RestPosition, 1.0), DIContext.WorldTransform).xyz;
-	if(DIContext.LocalSimulation)
-	{
-		OutNodePosition = mul(float4( OutNodePosition, 1.0), DIContext.BoneInverse).xyz;
-	}
 }
 
 // Update Node Orientation
 void DIHairStrands_AttachNodeOrientation(in FDIHairStrandsContext DIContext, in float4 RestOrientation, out float4 OutNodeOrientation)
 {
 	OutNodeOrientation = NormalizeQuat( MultiplyQuat(DIContext.WorldRotation,RestOrientation) );
-	if(DIContext.LocalSimulation)
-	{
-		OutNodeOrientation = NormalizeQuat( MultiplyQuat(InverseQuat(DIContext.BoneRotation),OutNodeOrientation) );
-	}
 }
 
 // Get World Vector
diff --git a/Engine/Plugins/Runtime/HairStrands/Source/HairStrandsCore/Private/Niagara/NiagaraDataInterfaceHairStrands.cpp b/Engine/Plugins/Runtime/HairStrands/Source/HairStrandsCore/Private/Niagara/NiagaraDataInterfaceHairStrands.cpp
--- a/Engine/Plugins/Runtime/HairStrands/Source/HairStrandsCore/Private/Niagara/NiagaraDataInterfaceHairStrands.cpp
+++ b/Engine/Plugins/Runtime/HairStrands/Source/HairStrandsCore/Private/Niagara/NiagaraDataInterfaceHairStrands.cpp
@@ -648,7 +648,7 @@ struct FNDIHairStrandsParametersCS : public FNiagaraDataInterfaceParametersCS
 			// TEMP: These check are only temporary for avoiding crashes while we find the bottom of the issue.
 			ProxyData->HairStrandsBuffer->CurvesOffsetsBuffer.SRV && ProxyData->HairStrandsBuffer->ParamsScaleBuffer.SRV && ProxyData->HairStrandsBuffer->BoundingBoxBuffer.UAV;
 
-		const bool bIsGeometryValid = bIsHairValid && ProxyData->HairGroupInstance && (ProxyData->HairGroupInstance->GeometryType != EHairGeometryType::NoneGeometry);
+		const bool bIsGeometryValid = bIsHairValid&& ProxyData->HairGroupInstance && (ProxyData->HairGroupInstance->GeometryType != EHairGeometryType::NoneGeometry);
 		const bool bIsDeformedValid = bIsHairValid && ProxyData->HairStrandsBuffer->SourceDeformedResources && ProxyData->HairStrandsBuffer->SourceDeformedResources->IsInitialized();
 
 		if (bIsHairValid && bIsRestValid && bIsGeometryValid)
@@ -711,6 +711,15 @@ struct FNDIHairStrandsParametersCS : public FNiagaraDataInterfaceParametersCS
 																			ProxyData->WorldTransform.ToMatrixWithScale();
 			FMatrix44f BoneTransformFloat = ProxyData->BoneTransform.ToMatrixWithScale();
 
+			if (ProxyData->LocalSimulation)
+			{
+				FMatrix44d WorldTransformDouble = WorldTransformFloat;
+				FMatrix44d BoneTransformDouble = BoneTransformFloat;
+
+				// Due to large world coordinate we store the relative world transform in double precision
+				WorldTransformFloat = WorldTransformDouble * BoneTransformDouble.Inverse();
+			}
+
 			if (!bIsRootValid && bHasSkinningBinding)
 			{
 				UE_LOG(LogHairStrands, Log, TEXT("FNDIHairStrandsParametersCS() Groom Asset %s from component %s is set to use skinning interpolation but the skin resources are not valid"),
