Fix invalid LOD selection mode when SetForceLOD is called, but the component does not have a Groom setup yet.

This changes ensures that the ForcedLOD selection is preserved even though the groom asset change.

diff --git a/Engine/Plugins/Runtime/HairStrands/Source/HairStrandsCore/Private/GroomComponent.cpp b/Engine/Plugins/Runtime/HairStrands/Source/HairStrandsCore/Private/GroomComponent.cpp
--- a/Engine/Plugins/Runtime/HairStrands/Source/HairStrandsCore/Private/GroomComponent.cpp
+++ b/Engine/Plugins/Runtime/HairStrands/Source/HairStrandsCore/Private/GroomComponent.cpp
@@ -1462,6 +1462,7 @@ void UGroomComponent::SetGroomAsset(UGroomAsset* Asset, UGroomBindingAsset* InBi
 	if (Asset && Asset->IsValid())
 	{
 		GroomAsset = Asset;
+		LODForcedIndex = FMath::Clamp(LODForcedIndex, -1, GroomAsset->GetLODCount() - 1);
 
 #if WITH_EDITORONLY_DATA
 		if (InBinding && !InBinding->IsValid())
@@ -1608,15 +1609,10 @@ void UGroomComponent::SetForcedLOD(int32 CurrLODIndex)
 	else if (PrevLODSelectionType == EHairLODSelectionType::Predicted)	{ PrevLODIndex = LODPredictedIndex; }
 
 	// Current LOD state
-	EHairLODSelectionType CurrLODSelectionType = EHairLODSelectionType::Immediate;
+	EHairLODSelectionType CurrLODSelectionType = EHairLODSelectionType::Forced;
 	if (GroomAsset)
 	{
 		CurrLODIndex = FMath::Clamp(CurrLODIndex, -1, GroomAsset->GetLODCount() - 1);
-		CurrLODSelectionType = EHairLODSelectionType::Forced;
-	}
-	else
-	{
-		CurrLODIndex = -1;
 	}
 
 	// Reset to non-forced LOD, and change LOD selection type to Predicted/Immediate
@@ -2314,7 +2310,15 @@ void UGroomComponent::InitResources(bool bIsBindingReloading)
 	// LOD Selection type
 	// By default set to immediate, meaning the LOD selection will be done on the rendering thread based on the screen coverage of the groom
 	// If simulation is enabled, then the LOD selection will be done on the game thread based on feedback from the rendering thread.
-	LODSelectionType = bHasAnyNeedSimulation ? EHairLODSelectionType::Predicted : EHairLODSelectionType::Immediate;
+	// If LODForcedIndex is >= 0 it means it has been set by a previous SetForcedLOD call. LODSelectionType is set to Forced in such a case, to applied this 'cached' value
+	if (LODForcedIndex >= 0)
+	{
+		LODSelectionType = EHairLODSelectionType::Forced;
+	}
+	else
+	{
+		LODSelectionType = bHasAnyNeedSimulation ? EHairLODSelectionType::Predicted : EHairLODSelectionType::Immediate;
+	}
 
 	// 2. Insure that the binding asset is compatible, otherwise no binding
 	UMeshComponent* ParentMeshComponent = GetAttachParent() ? Cast<UMeshComponent>(GetAttachParent()) : nullptr;
