Add extra LOD prediction loading for groom data/resources.

diff --git a/Engine/Plugins/Runtime/HairStrands/Source/HairStrandsCore/Private/GroomComponent.cpp b/Engine/Plugins/Runtime/HairStrands/Source/HairStrandsCore/Private/GroomComponent.cpp
--- a/Engine/Plugins/Runtime/HairStrands/Source/HairStrandsCore/Private/GroomComponent.cpp
+++ b/Engine/Plugins/Runtime/HairStrands/Source/HairStrandsCore/Private/GroomComponent.cpp
@@ -48,6 +48,8 @@ static FAutoConsoleVariableRef CVarGroomCacheStreamingEnable(TEXT("GroomCache.En
 static bool GUseProxyLocalToWorld = true;
 static FAutoConsoleVariableRef CVarUseProxyLocalToWorld(TEXT("r.HairStrands.UseProxyLocalToWorld"), GUseProxyLocalToWorld, TEXT("Enable the use of the groom proxy local to world instead of extracting it from the game thread."));
 
+static bool GHairStrands_Streaming_Prediction = false;
+static FAutoConsoleVariableRef CVarHairStrands_Streaming_Prediction(TEXT("r.HairStrands.Streaming.Prediction"), GHairStrands_Streaming_Prediction, TEXT("Enable LOD streaming prediction."));
 
 #define LOCTEXT_NAMESPACE "GroomComponent"
 
@@ -1646,12 +1648,15 @@ void UGroomComponent::SetForcedLOD(int32 CurrLODIndex)
 	LODForcedIndex	 = CurrLODIndex;
 	LODSelectionType = CurrLODSelectionType;
 
+	// Add streaming request for the subsequent LOD as prediction, only when going to a finer LOD
+	const bool bPredictionLoad = GHairStrands_Streaming_Prediction && bHasLODSwitch && CurrLODIndex > 0 && CurrLODIndex < PrevLODIndex;
+
 	// Do not invalidate completly the proxy, but just update LOD index on the rendering thread	
 	if (FHairStrandsSceneProxy* GroomSceneProxy = (FHairStrandsSceneProxy*)SceneProxy)
 	{
 		const EHairLODSelectionType LocalLODSelectionType = LODSelectionType;
 		ENQUEUE_RENDER_COMMAND(FHairComponentSendLODIndex)(
-		[GroomSceneProxy, CurrLODIndex, LocalLODSelectionType, bHasLODSwitch](FRHICommandListImmediate& RHICmdList)
+		[GroomSceneProxy, CurrLODIndex, LocalLODSelectionType, bHasLODSwitch, bPredictionLoad](FRHICommandListImmediate& RHICmdList)
 		{
 			for (FHairGroupInstance* Instance : GroomSceneProxy->HairGroupInstances)
 			{
@@ -1662,6 +1667,11 @@ void UGroomComponent::SetForcedLOD(int32 CurrLODIndex)
 				{
 					AddHairStreamingRequest(Instance, CurrLODIndex);
 				}
+
+				if (bPredictionLoad)
+				{
+					AddHairStreamingRequest(Instance, CurrLODIndex-1);
+				}
 			}
 		});
 	}
