Add a 'tick' swap hair resources option, which swaps prpreviousecious/current resousrces based on component tick.

This allows to fix motion blur on grooms when rnereded with MRQ aith multiuple spatial samples.

diff --git a/Engine/Plugins/Runtime/HairStrands/Source/HairStrandsCore/Private/GroomComponent.cpp b/Engine/Plugins/Runtime/HairStrands/Source/HairStrandsCore/Private/GroomComponent.cpp
--- a/Engine/Plugins/Runtime/HairStrands/Source/HairStrandsCore/Private/GroomComponent.cpp
+++ b/Engine/Plugins/Runtime/HairStrands/Source/HairStrandsCore/Private/GroomComponent.cpp
@@ -3279,10 +3279,16 @@ void UGroomComponent::TickComponent(float DeltaTime, enum ELevelTick TickType, F
 		EffectiveForceLOD = LODForcedIndex;
 	}
 
+	const bool bSwapBuffer = GetHairSwapBufferType() == EHairBufferSwapType::Tick;
+	FTransform SkelLocalToTransform = FTransform::Identity;	
+	if (RegisteredMeshComponent)
+	{
+		SkelLocalToTransform = RegisteredMeshComponent->GetComponentTransform();
+	}
 	TArray<FHairGroupInstance*> LocalHairGroupInstances = HairGroupInstances;
 	const EHairLODSelectionType LocalLODSelectionType = LODSelectionType;
 	ENQUEUE_RENDER_COMMAND(FHairStrandsTick_TransformUpdate)(
-		[Id, FeatureLevel, LocalHairGroupInstances, EffectiveForceLOD, LocalLODSelectionType](FRHICommandListImmediate& RHICmdList)
+		[Id, FeatureLevel, LocalHairGroupInstances, EffectiveForceLOD, LocalLODSelectionType, bSwapBuffer, SkelLocalToTransform](FRHICommandListImmediate& RHICmdList)
 	{
 		if (ERHIFeatureLevel::Num == FeatureLevel)
 			return;
@@ -3295,6 +3301,15 @@ void UGroomComponent::TickComponent(float DeltaTime, enum ELevelTick TickType, F
 			{
 				AddHairStreamingRequest(Instance, EffectiveForceLOD);
 			}
+
+			if (bSwapBuffer)
+			{
+				Instance->Debug.SkinningPreviousLocalToWorld = Instance->Debug.SkinningCurrentLocalToWorld;
+				Instance->Debug.SkinningCurrentLocalToWorld  = SkelLocalToTransform;
+
+				if (Instance->Guides.DeformedResource)  { Instance->Guides.DeformedResource->SwapBuffer(); }
+				if (Instance->Strands.DeformedResource) { Instance->Strands.DeformedResource->SwapBuffer(); }
+			}
 		}
 	});
 }
diff --git a/Engine/Plugins/Runtime/HairStrands/Source/HairStrandsCore/Private/GroomManager.cpp b/Engine/Plugins/Runtime/HairStrands/Source/HairStrandsCore/Private/GroomManager.cpp
--- a/Engine/Plugins/Runtime/HairStrands/Source/HairStrandsCore/Private/GroomManager.cpp
+++ b/Engine/Plugins/Runtime/HairStrands/Source/HairStrandsCore/Private/GroomManager.cpp
@@ -25,8 +25,8 @@ static FAutoConsoleVariableRef CVarGHairStrandsMinLOD(TEXT("r.HairStrands.MinLOD
 static int32 GHairStrands_UseCards = 0;
 static FAutoConsoleVariableRef CVarHairStrands_UseCards(TEXT("r.HairStrands.UseCardsInsteadOfStrands"), GHairStrands_UseCards, TEXT("Force cards geometry on all groom elements. If no cards data is available, nothing will be displayed"), ECVF_Scalability);
 
-static int32 GHairStrands_SwapBufferEndOfFrame = 1;
-static FAutoConsoleVariableRef CVarGHairStrands_SwapBufferEndOfFrame(TEXT("r.HairStrands.SwapEndOfFrame"), GHairStrands_SwapBufferEndOfFrame, TEXT("Swap rendering buffer at the end of frame. This is an experimental toggle. Default:1"));
+static int32 GHairStrands_SwapBufferType = 1;
+static FAutoConsoleVariableRef CVarGHairStrands_SwapBufferType(TEXT("r.HairStrands.SwapType"), GHairStrands_SwapBufferType, TEXT("Swap rendering buffer at the end of frame. This is an experimental toggle. Default:1"));
 
 static int32 GHairStrands_ManualSkinCache = 0;
 static FAutoConsoleVariableRef CVarGHairStrands_ManualSkinCache(TEXT("r.HairStrands.ManualSkinCache"), GHairStrands_ManualSkinCache, TEXT("If skin cache is not enabled, and grooms use skinning method, this enable a simple skin cache mechanisme for groom. Default:disable"));
@@ -34,6 +34,17 @@ static FAutoConsoleVariableRef CVarGHairStrands_ManualSkinCache(TEXT("r.HairStra
 static int32 GHairStrands_InterpolationFrustumCullingEnable = 1;
 static FAutoConsoleVariableRef CVarHairStrands_InterpolationFrustumCullingEnable(TEXT("r.HairStrands.Interoplation.FrustumCulling"), GHairStrands_InterpolationFrustumCullingEnable, TEXT("Swap rendering buffer at the end of frame. This is an experimental toggle. Default:1"));
 
+EHairBufferSwapType GetHairSwapBufferType()
+{
+	switch (GHairStrands_SwapBufferType)
+	{
+	case 0: return EHairBufferSwapType::BeginOfFrame;
+	case 1: return EHairBufferSwapType::EndOfFrame;
+	case 2: return EHairBufferSwapType::Tick;
+	}
+	return EHairBufferSwapType::EndOfFrame;
+}
+
 bool IsHairStrandsSkinCacheEnable()
 {
 	return GHairStrands_ManualSkinCache > 0;
@@ -897,8 +908,11 @@ static void RunHairLODSelection(
 		}
 
 		// Update the local-to-world transform based on the binding type 
-		Instance->Debug.SkinningPreviousLocalToWorld = Instance->Debug.SkinningCurrentLocalToWorld;
-		Instance->Debug.SkinningCurrentLocalToWorld = CachedGeometry.LocalToWorld;
+		if (GetHairSwapBufferType() != EHairBufferSwapType::Tick)
+		{
+			Instance->Debug.SkinningPreviousLocalToWorld = Instance->Debug.SkinningCurrentLocalToWorld;
+			Instance->Debug.SkinningCurrentLocalToWorld = CachedGeometry.LocalToWorld;
+		}
 		Instance->LocalToWorld = Instance->GetCurrentLocalToWorld();
 	}
 }
@@ -988,7 +1002,7 @@ void ProcessHairStrandsBookmark(
 	}
 	else if (Bookmark == EHairStrandsBookmark::ProcessLODSelection)
 	{
-		if (GHairStrands_SwapBufferEndOfFrame <= 0)
+		if (GetHairSwapBufferType() == EHairBufferSwapType::BeginOfFrame)
 		{
 			RunHairBufferSwap(
 				*Parameters.Instances,
@@ -1004,7 +1018,7 @@ void ProcessHairStrandsBookmark(
 	}
 	else if (Bookmark == EHairStrandsBookmark::ProcessEndOfFrame)
 	{
-		if (GHairStrands_SwapBufferEndOfFrame > 0)
+		if (GetHairSwapBufferType() == EHairBufferSwapType::EndOfFrame)
 		{
 			RunHairBufferSwap(
 				*Parameters.Instances,
diff --git a/Engine/Plugins/Runtime/HairStrands/Source/HairStrandsCore/Private/GroomManager.h b/Engine/Plugins/Runtime/HairStrands/Source/HairStrandsCore/Private/GroomManager.h
--- a/Engine/Plugins/Runtime/HairStrands/Source/HairStrandsCore/Private/GroomManager.h
+++ b/Engine/Plugins/Runtime/HairStrands/Source/HairStrandsCore/Private/GroomManager.h
@@ -48,6 +48,15 @@ enum class EHairStrandsProjectionMeshType
 	TargetMesh
 };
 
+enum EHairBufferSwapType
+{
+	BeginOfFrame,
+	EndOfFrame,
+	Tick
+};
+
+EHairBufferSwapType GetHairSwapBufferType();
+
 ////////////////////////////////////////////////////////////////////////////////////////////////////
 // Registrations & update
 
