Change hair LOD selection so that when a LOD is picked but the resources are not streamed-in in time, a coarser loaded LOD is used in place.

diff --git a/Engine/Plugins/Runtime/HairStrands/Source/HairStrandsCore/Private/GroomManager.cpp b/Engine/Plugins/Runtime/HairStrands/Source/HairStrandsCore/Private/GroomManager.cpp
--- a/Engine/Plugins/Runtime/HairStrands/Source/HairStrandsCore/Private/GroomManager.cpp
+++ b/Engine/Plugins/Runtime/HairStrands/Source/HairStrandsCore/Private/GroomManager.cpp
@@ -728,8 +728,10 @@ static void RunHairLODSelection(
 		const float PrevLODIndex = Instance->HairGroupPublicData->GetLODIndex();
 		const float PrevMeshLODIndex = Instance->HairGroupPublicData->GetMeshLODIndex();
 
+		// 0. Initial LOD index picking
 		// Insure that MinLOD is necessary taken into account if a force LOD is request (i.e., LODIndex>=0). If a Force LOD 
 		// is not resquested (i.e., LODIndex<0), the MinLOD is applied after ViewLODIndex has been determined in the codeblock below
+		const int32 LODCount = Instance->HairGroupPublicData->GetLODVisibilities().Num();
 		const float MinLOD = FMath::Max(0, GHairStrandsMinLOD);
 		float LODIndex = Instance->Debug.LODForcedIndex >= 0 ? FMath::Max(Instance->Debug.LODForcedIndex, MinLOD) : -1.0f;
 		float LODViewIndex = -1;
@@ -744,164 +746,197 @@ static void RunHairLODSelection(
 				// Select highest LOD accross all views
 				LODViewIndex = LODViewIndex < 0 ? CurrLODViewIndex : FMath::Min(LODViewIndex, CurrLODViewIndex);
 			}
-		}
-		if (LODIndex < 0)
-		{
-			LODIndex = LODViewIndex;
-		}
-
-		const TArray<EHairGeometryType>& LODGeometryTypes = Instance->HairGroupPublicData->GetLODGeometryTypes();
-		const TArray<bool>& LODVisibilities = Instance->HairGroupPublicData->GetLODVisibilities();
-		const int32 LODCount = LODVisibilities.Num();
 
-		LODIndex = FMath::Clamp(LODIndex, 0.f, float(LODCount - 1));
+			if (LODIndex < 0)
+			{
+				LODIndex = LODViewIndex;
+			}
 
-		// Feedback game thread with LOD selection 
-		Instance->Debug.LODPredictedIndex = LODViewIndex;
+			LODIndex = FMath::Clamp(LODIndex, 0.f, float(LODCount - 1));
 
-		const int32 IntLODIndex = FMath::Clamp(FMath::FloorToInt(LODIndex), 0, LODCount - 1);
-		const bool bIsVisible = LODVisibilities[IntLODIndex];
-		const bool bForceCards = GHairStrands_UseCards > 0 || Instance->bForceCards; // todo
-		EHairGeometryType GeometryType = ConvertLODGeometryType(LODGeometryTypes[IntLODIndex], bForceCards, ShaderPlatform);
+			// Feedback game thread with LOD selection 
+			Instance->Debug.LODPredictedIndex = LODViewIndex;
+		}
 
-		// Skip cluster culling if strands have a single LOD (i.e.: LOD0), and the instance is static (i.e., no skinning, no simulation, ...)
-		bool bCullingEnable = false;
-		if (GeometryType == EHairGeometryType::Meshes)
+		// Function for selecting, loading, & initializing LOD resources
+		auto SelectValidLOD = [Instance, ShaderPlatform, ShaderMap, PrevLODIndex, LODCount, MeshLODIndex, PrevMeshLODIndex, bHasPathTracingView] (
+			FRDGBuilder& GraphBuilder, float LODIndex) -> bool
 		{
-			if (!Instance->Meshes.IsValid(IntLODIndex))
+			const int32 IntLODIndex = FMath::Clamp(FMath::FloorToInt(LODIndex), 0, LODCount - 1);
+			const TArray<EHairGeometryType>& LODGeometryTypes = Instance->HairGroupPublicData->GetLODGeometryTypes();
+			const TArray<bool>& LODVisibilities = Instance->HairGroupPublicData->GetLODVisibilities();
+			const bool bIsVisible = LODVisibilities[IntLODIndex];
+			const bool bForceCards = GHairStrands_UseCards > 0 || Instance->bForceCards; // todo
+			EHairGeometryType GeometryType = ConvertLODGeometryType(LODGeometryTypes[IntLODIndex], bForceCards, ShaderPlatform);
+
+			// Skip cluster culling if strands have a single LOD (i.e.: LOD0), and the instance is static (i.e., no skinning, no simulation, ...)
+			bool bCullingEnable = false;
+			if (GeometryType == EHairGeometryType::Meshes)
 			{
-				GeometryType = EHairGeometryType::NoneGeometry;
+				if (!Instance->Meshes.IsValid(IntLODIndex))
+				{
+					GeometryType = EHairGeometryType::NoneGeometry;
+				}
 			}
-		}
-		else if (GeometryType == EHairGeometryType::Cards)
-		{
-			if (!Instance->Cards.IsValid(IntLODIndex))
+			else if (GeometryType == EHairGeometryType::Cards)
 			{
-				GeometryType = EHairGeometryType::NoneGeometry;
+				if (!Instance->Cards.IsValid(IntLODIndex))
+				{
+					GeometryType = EHairGeometryType::NoneGeometry;
+				}
 			}
-		}
-		else if (GeometryType == EHairGeometryType::Strands)
-		{
-			if (!Instance->Strands.IsValid())
+			else if (GeometryType == EHairGeometryType::Strands)
 			{
-				GeometryType = EHairGeometryType::NoneGeometry;
+				if (!Instance->Strands.IsValid())
+				{
+					GeometryType = EHairGeometryType::NoneGeometry;
+				}
+				else
+				{
+					const bool bNeedLODing		= LODIndex > 0;
+					const bool bNeedDeformation = Instance->Strands.DeformedResource != nullptr;
+					bCullingEnable = bNeedLODing || bNeedDeformation;
+				}
 			}
-			else
+
+			if (!bIsVisible)
 			{
-				const bool bNeedLODing		= LODIndex > 0;
-				const bool bNeedDeformation = Instance->Strands.DeformedResource != nullptr;
-				bCullingEnable = bNeedLODing || bNeedDeformation;
+				GeometryType = EHairGeometryType::NoneGeometry;
 			}
-		}
 
-		if (!bIsVisible)
-		{
-			GeometryType = EHairGeometryType::NoneGeometry;
-		}
+			const bool bSimulationEnable = Instance->HairGroupPublicData->IsSimulationEnable(IntLODIndex);
+			const bool bGlobalInterpolationEnable =  Instance->HairGroupPublicData->IsGlobalInterpolationEnable(IntLODIndex);
+			const bool bLODNeedsGuides = bSimulationEnable || bGlobalInterpolationEnable;
 
-		const bool bSimulationEnable = Instance->HairGroupPublicData->IsSimulationEnable(LODIndex);
-		const bool bGlobalInterpolationEnable =  Instance->HairGroupPublicData->IsGlobalInterpolationEnable(LODIndex);
-		const bool bLODNeedsGuides = bSimulationEnable || bGlobalInterpolationEnable;
-
-		const EHairResourceLoadingType LoadingType = GetHairResourceLoadingType(GeometryType, int32(LODIndex));
-		EHairResourceStatus ResourceStatus = EHairResourceStatus::None;
+			const EHairResourceLoadingType LoadingType = GetHairResourceLoadingType(GeometryType, IntLODIndex);
+			EHairResourceStatus ResourceStatus = EHairResourceStatus::None;
 
-		// Lazy allocation of resources
-		// Note: Allocation will only be done if the resources is not initialized yet. Guides deformed position are also initialized from the Rest position at creation time.
-		if (Instance->Guides.Data && bLODNeedsGuides)
-		{
-			if (Instance->Guides.RestRootResource)			{ Instance->Guides.RestRootResource->Allocate(GraphBuilder, LoadingType, ResourceStatus); Instance->Guides.RestRootResource->AllocateLOD(GraphBuilder, MeshLODIndex, LoadingType, ResourceStatus); }
-			if (Instance->Guides.RestResource)				{ Instance->Guides.RestResource->Allocate(GraphBuilder, LoadingType, ResourceStatus); }
-			if (Instance->Guides.DeformedRootResource)		{ Instance->Guides.DeformedRootResource->Allocate(GraphBuilder, LoadingType, ResourceStatus); Instance->Guides.DeformedRootResource->AllocateLOD(GraphBuilder, MeshLODIndex, LoadingType, ResourceStatus); }
-			if (Instance->Guides.DeformedResource)			
-			{ 
-				// Ensure the rest resources are correctly loaded prior to initialized and copy the rest position into the deformed positions
-				if (Instance->Guides.RestResource->bIsInitialized)
-				{
-					const bool bNeedCopy = !Instance->Guides.DeformedResource->bIsInitialized; 
-					Instance->Guides.DeformedResource->Allocate(GraphBuilder, EHairResourceLoadingType::Sync); 
-					if (bNeedCopy) 
-					{ 
-						AddCopyHairStrandsPositionPass(GraphBuilder, ShaderMap, *Instance->Guides.RestResource, *Instance->Guides.DeformedResource); 
+			// Lazy allocation of resources
+			// Note: Allocation will only be done if the resources is not initialized yet. Guides deformed position are also initialized from the Rest position at creation time.
+			if (Instance->Guides.Data && bLODNeedsGuides)
+			{
+				if (Instance->Guides.RestRootResource)			{ Instance->Guides.RestRootResource->Allocate(GraphBuilder, LoadingType, ResourceStatus); Instance->Guides.RestRootResource->AllocateLOD(GraphBuilder, MeshLODIndex, LoadingType, ResourceStatus); }
+				if (Instance->Guides.RestResource)				{ Instance->Guides.RestResource->Allocate(GraphBuilder, LoadingType, ResourceStatus); }
+				if (Instance->Guides.DeformedRootResource)		{ Instance->Guides.DeformedRootResource->Allocate(GraphBuilder, LoadingType, ResourceStatus); Instance->Guides.DeformedRootResource->AllocateLOD(GraphBuilder, MeshLODIndex, LoadingType, ResourceStatus); }
+				if (Instance->Guides.DeformedResource)			
+				{ 
+					// Ensure the rest resources are correctly loaded prior to initialized and copy the rest position into the deformed positions
+					if (Instance->Guides.RestResource->bIsInitialized)
+					{
+						const bool bNeedCopy = !Instance->Guides.DeformedResource->bIsInitialized; 
+						Instance->Guides.DeformedResource->Allocate(GraphBuilder, EHairResourceLoadingType::Sync); 
+						if (bNeedCopy) 
+						{ 
+							AddCopyHairStrandsPositionPass(GraphBuilder, ShaderMap, *Instance->Guides.RestResource, *Instance->Guides.DeformedResource); 
+						}
 					}
 				}
 			}
-		}
 
-		if (GeometryType == EHairGeometryType::Meshes)
-		{
-			FHairGroupInstance::FMeshes::FLOD& InstanceLOD = Instance->Meshes.LODs[IntLODIndex];
+			if (GeometryType == EHairGeometryType::Meshes)
+			{
+				FHairGroupInstance::FMeshes::FLOD& InstanceLOD = Instance->Meshes.LODs[IntLODIndex];
 
-			if (InstanceLOD.DeformedResource)				{ InstanceLOD.DeformedResource->Allocate(GraphBuilder, LoadingType, ResourceStatus); }
-			#if RHI_RAYTRACING
-			if (InstanceLOD.RaytracingResource)				{ InstanceLOD.RaytracingResource->Allocate(GraphBuilder, LoadingType, ResourceStatus);}
-			#endif
+				if (InstanceLOD.DeformedResource)				{ InstanceLOD.DeformedResource->Allocate(GraphBuilder, LoadingType, ResourceStatus); }
+				#if RHI_RAYTRACING
+				if (InstanceLOD.RaytracingResource)				{ InstanceLOD.RaytracingResource->Allocate(GraphBuilder, LoadingType, ResourceStatus);}
+				#endif
 
-			InstanceLOD.InitVertexFactory();
-		}
-		else if (GeometryType == EHairGeometryType::Cards)
-		{
-			FHairGroupInstance::FCards::FLOD& InstanceLOD = Instance->Cards.LODs[IntLODIndex];
+				InstanceLOD.InitVertexFactory();
+			}
+			else if (GeometryType == EHairGeometryType::Cards)
+			{
+				FHairGroupInstance::FCards::FLOD& InstanceLOD = Instance->Cards.LODs[IntLODIndex];
+
+				if (InstanceLOD.InterpolationResource)			{ InstanceLOD.InterpolationResource->Allocate(GraphBuilder, LoadingType, ResourceStatus); }
+				if (InstanceLOD.DeformedResource)				{ InstanceLOD.DeformedResource->Allocate(GraphBuilder, LoadingType, ResourceStatus); }
+				if (InstanceLOD.Guides.RestRootResource)		{ InstanceLOD.Guides.RestRootResource->Allocate(GraphBuilder, LoadingType, ResourceStatus); InstanceLOD.Guides.RestRootResource->AllocateLOD(GraphBuilder, MeshLODIndex, LoadingType, ResourceStatus); }
+				if (InstanceLOD.Guides.RestResource)			{ InstanceLOD.Guides.RestResource->Allocate(GraphBuilder, LoadingType, ResourceStatus); }
+				if (InstanceLOD.Guides.DeformedRootResource)	{ InstanceLOD.Guides.DeformedRootResource->Allocate(GraphBuilder, LoadingType, ResourceStatus); InstanceLOD.Guides.DeformedRootResource->AllocateLOD(GraphBuilder, MeshLODIndex, LoadingType, ResourceStatus); }
+				if (InstanceLOD.Guides.DeformedResource)		{ InstanceLOD.Guides.DeformedResource->Allocate(GraphBuilder, LoadingType, ResourceStatus); }
+				if (InstanceLOD.Guides.InterpolationResource)	{ InstanceLOD.Guides.InterpolationResource->Allocate(GraphBuilder, LoadingType, ResourceStatus); }
+				#if RHI_RAYTRACING
+				if (InstanceLOD.RaytracingResource)				{ InstanceLOD.RaytracingResource->Allocate(GraphBuilder, LoadingType, ResourceStatus);}
+				#endif
+				InstanceLOD.InitVertexFactory();
+			}
+			else if (GeometryType == EHairGeometryType::Strands)
+			{
+				check(Instance->HairGroupPublicData);
+				Instance->HairGroupPublicData->Allocate(GraphBuilder);
+
+				if (Instance->Strands.RestRootResource)			{ Instance->Strands.RestRootResource->Allocate(GraphBuilder, LoadingType, ResourceStatus); Instance->Strands.RestRootResource->AllocateLOD(GraphBuilder, MeshLODIndex, LoadingType, ResourceStatus); }
+				if (Instance->Strands.RestResource)				{ Instance->Strands.RestResource->Allocate(GraphBuilder, LoadingType, ResourceStatus); }
+				if (Instance->Strands.ClusterCullingResource)	{ Instance->Strands.ClusterCullingResource->Allocate(GraphBuilder, LoadingType, ResourceStatus); }
+				if (Instance->Strands.InterpolationResource)	{ Instance->Strands.InterpolationResource->Allocate(GraphBuilder, LoadingType, ResourceStatus); }
+
+				if (Instance->Strands.DeformedRootResource)		{ Instance->Strands.DeformedRootResource->Allocate(GraphBuilder, LoadingType, ResourceStatus); Instance->Strands.DeformedRootResource->AllocateLOD(GraphBuilder, MeshLODIndex, LoadingType, ResourceStatus); }
+				if (Instance->Strands.DeformedResource)			{ Instance->Strands.DeformedResource->Allocate(GraphBuilder, LoadingType, ResourceStatus); }
+				#if RHI_RAYTRACING
+				if (bHasPathTracingView)						{ AllocateRaytracingResources(Instance); }
+				if (Instance->Strands.RenRaytracingResource)	{ Instance->Strands.RenRaytracingResource->Allocate(GraphBuilder, LoadingType, ResourceStatus); }
+				#endif
+				Instance->Strands.VertexFactory->InitResources();
+
+				// Early initialization, so that when filtering MeshBatch block in SceneVisiblity, we can use this value to know 
+				// if the hair instance is visible or not (i.e., HairLengthScale > 0)
+				Instance->HairGroupPublicData->VFInput.Strands.HairLengthScale = Instance->Strands.Modifier.HairLengthScale;
+			}
 
-			if (InstanceLOD.InterpolationResource)			{ InstanceLOD.InterpolationResource->Allocate(GraphBuilder, LoadingType, ResourceStatus); }
-			if (InstanceLOD.DeformedResource)				{ InstanceLOD.DeformedResource->Allocate(GraphBuilder, LoadingType, ResourceStatus); }
-			if (InstanceLOD.Guides.RestRootResource)		{ InstanceLOD.Guides.RestRootResource->Allocate(GraphBuilder, LoadingType, ResourceStatus); InstanceLOD.Guides.RestRootResource->AllocateLOD(GraphBuilder, MeshLODIndex, LoadingType, ResourceStatus); }
-			if (InstanceLOD.Guides.RestResource)			{ InstanceLOD.Guides.RestResource->Allocate(GraphBuilder, LoadingType, ResourceStatus); }
-			if (InstanceLOD.Guides.DeformedRootResource)	{ InstanceLOD.Guides.DeformedRootResource->Allocate(GraphBuilder, LoadingType, ResourceStatus); InstanceLOD.Guides.DeformedRootResource->AllocateLOD(GraphBuilder, MeshLODIndex, LoadingType, ResourceStatus); }
-			if (InstanceLOD.Guides.DeformedResource)		{ InstanceLOD.Guides.DeformedResource->Allocate(GraphBuilder, LoadingType, ResourceStatus); }
-			if (InstanceLOD.Guides.InterpolationResource)	{ InstanceLOD.Guides.InterpolationResource->Allocate(GraphBuilder, LoadingType, ResourceStatus); }
-			#if RHI_RAYTRACING
-			if (InstanceLOD.RaytracingResource)				{ InstanceLOD.RaytracingResource->Allocate(GraphBuilder, LoadingType, ResourceStatus);}
-			#endif
-			InstanceLOD.InitVertexFactory();
-		}
-		else if (GeometryType == EHairGeometryType::Strands)
-		{
-			check(Instance->HairGroupPublicData);
-			Instance->HairGroupPublicData->Allocate(GraphBuilder);
-
-			if (Instance->Strands.RestRootResource)			{ Instance->Strands.RestRootResource->Allocate(GraphBuilder, LoadingType, ResourceStatus); Instance->Strands.RestRootResource->AllocateLOD(GraphBuilder, MeshLODIndex, LoadingType, ResourceStatus); }
-			if (Instance->Strands.RestResource)				{ Instance->Strands.RestResource->Allocate(GraphBuilder, LoadingType, ResourceStatus); }
-			if (Instance->Strands.ClusterCullingResource)	{ Instance->Strands.ClusterCullingResource->Allocate(GraphBuilder, LoadingType, ResourceStatus); }
-			if (Instance->Strands.InterpolationResource)	{ Instance->Strands.InterpolationResource->Allocate(GraphBuilder, LoadingType, ResourceStatus); }
-
-			if (Instance->Strands.DeformedRootResource)		{ Instance->Strands.DeformedRootResource->Allocate(GraphBuilder, LoadingType, ResourceStatus); Instance->Strands.DeformedRootResource->AllocateLOD(GraphBuilder, MeshLODIndex, LoadingType, ResourceStatus); }
-			if (Instance->Strands.DeformedResource)			{ Instance->Strands.DeformedResource->Allocate(GraphBuilder, LoadingType, ResourceStatus); }
-			#if RHI_RAYTRACING
-			if (bHasPathTracingView)						{ AllocateRaytracingResources(Instance); }
-			if (Instance->Strands.RenRaytracingResource)	{ Instance->Strands.RenRaytracingResource->Allocate(GraphBuilder, LoadingType, ResourceStatus); }
-			#endif
-			Instance->Strands.VertexFactory->InitResources();
-
-			// Early initialization, so that when filtering MeshBatch block in SceneVisiblity, we can use this value to know 
-			// if the hair instance is visible or not (i.e., HairLengthScale > 0)
-			Instance->HairGroupPublicData->VFInput.Strands.HairLengthScale = Instance->Strands.Modifier.HairLengthScale;
-		}
+			// Only switch LOD if the data are ready to be used
+			const bool bIsLODDataReady = !!(ResourceStatus & EHairResourceStatus::Loading);
+			if (bIsLODDataReady)
+			{
+				EHairBindingType BindingType = Instance->HairGroupPublicData->GetBindingType(IntLODIndex);
+				Instance->HairGroupPublicData->SetLODVisibility(bIsVisible);
+				Instance->HairGroupPublicData->SetLODIndex(LODIndex);
+				Instance->HairGroupPublicData->SetLODBias(0);
+				Instance->HairGroupPublicData->SetMeshLODIndex(MeshLODIndex);
+				Instance->HairGroupPublicData->VFInput.GeometryType = GeometryType;
+				Instance->HairGroupPublicData->VFInput.BindingType = BindingType;
+				Instance->HairGroupPublicData->VFInput.bHasLODSwitch = (FMath::FloorToInt(PrevLODIndex) != FMath::FloorToInt(LODIndex));
+				Instance->GeometryType = GeometryType;
+				Instance->BindingType = BindingType;
+				Instance->Guides.bIsSimulationEnable = Instance->HairGroupPublicData->IsSimulationEnable(IntLODIndex);
+				Instance->Guides.bHasGlobalInterpolation = Instance->HairGroupPublicData->IsGlobalInterpolationEnable(IntLODIndex);
+				Instance->Strands.bIsCullingEnabled = bCullingEnable;
+			}
 
-		// Only switch LOD if the data are ready to be used
-		const bool bIsLODDataReady = !!(ResourceStatus & EHairResourceStatus::Loading);
-		if (bIsLODDataReady)
+			// If requested LOD's resources are not ready yet, and the previous LOD was invalid or if the MeshLODIndex has changed 
+			// (which would required extra data loading) change the geometry to be invalid, so that it don't get processed this frame.
+			const bool bIsLODValid = PrevLODIndex >= 0;
+			const bool bHasMeshLODChanged = PrevMeshLODIndex != MeshLODIndex && (Instance->Guides.bHasGlobalInterpolation || Instance->BindingType == EHairBindingType::Skinning);
+			if (!bIsLODDataReady && (!bIsLODValid || bHasMeshLODChanged))
+			{
+				return false;
+			}
+			return true;
+		};
+
+		// 1. Try to find an available LOD
+		bool bFoundValidLOD = SelectValidLOD(GraphBuilder, LODIndex);
+
+		// 2. If no valid LOD are found, try to find a coarser LOD which is:
+		// * loaded 
+		// * does not require skinning binding 
+		// * does not require global interpolation (as global interpolation requires skel. mesh data)
+		// * does not require simulation (as simulation requires LOD to be selected on the game thread i.e., Predicted/Forced, and so we can override this LOD here)
+		if (!bFoundValidLOD)
 		{
-			EHairBindingType BindingType = Instance->HairGroupPublicData->GetBindingType(IntLODIndex);
-			Instance->HairGroupPublicData->SetLODVisibility(bIsVisible);
-			Instance->HairGroupPublicData->SetLODIndex(LODIndex);
-			Instance->HairGroupPublicData->SetLODBias(0);
-			Instance->HairGroupPublicData->SetMeshLODIndex(MeshLODIndex);
-			Instance->HairGroupPublicData->VFInput.GeometryType = GeometryType;
-			Instance->HairGroupPublicData->VFInput.BindingType = BindingType;
-			Instance->HairGroupPublicData->VFInput.bHasLODSwitch = (FMath::FloorToInt(PrevLODIndex) != FMath::FloorToInt(LODIndex));
-			Instance->GeometryType = GeometryType;
-			Instance->BindingType = BindingType;
-			Instance->Guides.bIsSimulationEnable = Instance->HairGroupPublicData->IsSimulationEnable(IntLODIndex);
-			Instance->Guides.bHasGlobalInterpolation = Instance->HairGroupPublicData->IsGlobalInterpolationEnable(IntLODIndex);
-			Instance->Strands.bIsCullingEnabled = bCullingEnable;
+			for (int32 FallbackLODIndex = FMath::Clamp(FMath::FloorToInt(LODIndex + 1), 0, LODCount - 1); FallbackLODIndex < LODCount; ++FallbackLODIndex)
+			{
+				if (Instance->HairGroupPublicData->GetBindingType(FallbackLODIndex) == EHairBindingType::Rigid &&
+					!Instance->HairGroupPublicData->IsSimulationEnable(FallbackLODIndex) &&
+					!Instance->HairGroupPublicData->IsGlobalInterpolationEnable(FallbackLODIndex))
+				{
+					bFoundValidLOD = SelectValidLOD(GraphBuilder, FallbackLODIndex);
+					break;
+				}
+			}
 		}
 
-		// If requested LOD's resources are not ready yet, and the previous LOD was invalid or if the MeshLODIndex has changed 
-		// (which would required extra data loading) change the geometry to be invalid, so that it don't get processed this frame.
-		const bool bIsLODValid = PrevLODIndex >= 0;
-		const bool bHasMeshLODChanged = PrevMeshLODIndex != MeshLODIndex && (Instance->Guides.bHasGlobalInterpolation || Instance->BindingType == EHairBindingType::Skinning);
-		if (!bIsLODDataReady && (!bIsLODValid || bHasMeshLODChanged))
+		// 3. If no LOD are valid, then mark the instance as invalid for this frame
+		if (!bFoundValidLOD)
 		{
 			Instance->GeometryType = EHairGeometryType::NoneGeometry;
 			Instance->BindingType = EHairBindingType::NoneBinding;
