Fix collision init issue on WU0130

diff --git a/Engine/Plugins/Runtime/HairStrands/Source/HairStrandsCore/Private/Niagara/NiagaraDataInterfaceHairStrands.cpp b/Engine/Plugins/Runtime/HairStrands/Source/HairStrandsCore/Private/Niagara/NiagaraDataInterfaceHairStrands.cpp
--- a/Engine/Plugins/Runtime/HairStrands/Source/HairStrandsCore/Private/Niagara/NiagaraDataInterfaceHairStrands.cpp
+++ b/Engine/Plugins/Runtime/HairStrands/Source/HairStrandsCore/Private/Niagara/NiagaraDataInterfaceHairStrands.cpp
@@ -699,7 +699,7 @@ struct FNDIHairStrandsParametersCS : public FNiagaraDataInterfaceParametersCS
 
 			// Simulation setup
 			const int32 NeedResetValue = (ProxyData->TickCount <= GHairSimulationMaxDelay) || !HairStrandsBuffer->bValidGeometryType;
-			const int32 RestUpdateValue = GHairSimulationRestUpdate;
+			const int32 RestUpdateValue = NeedResetValue || GHairSimulationRestUpdate;
 			const int32 LocalSimulationValue = ProxyData->LocalSimulation;
 
 			HairStrandsBuffer->bValidGeometryType = true;
diff --git a/Engine/Plugins/Runtime/HairStrands/Source/HairStrandsCore/Private/Niagara/NiagaraDataInterfacePhysicsAsset.cpp b/Engine/Plugins/Runtime/HairStrands/Source/HairStrandsCore/Private/Niagara/NiagaraDataInterfacePhysicsAsset.cpp
--- a/Engine/Plugins/Runtime/HairStrands/Source/HairStrandsCore/Private/Niagara/NiagaraDataInterfacePhysicsAsset.cpp
+++ b/Engine/Plugins/Runtime/HairStrands/Source/HairStrandsCore/Private/Niagara/NiagaraDataInterfacePhysicsAsset.cpp
@@ -185,6 +185,16 @@ void CompactInternalArrays(FNDIPhysicsAssetArrays* OutAssetArrays)
 	}
 }
 
+// Due to large world coordinate we compute the relative world transform in double precision
+inline FTransform ComputeRelativeTransform(const FTransform& TransformA, const FTransform& TransformB)
+{
+	const FMatrix44d TransformAMatrix = TransformA.ToMatrixWithScale();
+	const FMatrix44d TransformBMatrix = TransformB.ToMatrixWithScale();
+
+	const FMatrix44d RelativeMatrix  = TransformAMatrix * TransformBMatrix.Inverse();
+	return FTransform(RelativeMatrix);
+}
+
 void CreateInternalArrays(const TArray<TWeakObjectPtr<UPhysicsAsset>>& PhysicsAssets, const TArray<TWeakObjectPtr<USkeletalMeshComponent>>& SkeletalMeshs,
 	FNDIPhysicsAssetArrays* OutAssetArrays, const FTransform& InWorldTransform, const FTransform& LocalTransform)
 {
@@ -242,7 +252,8 @@ void CreateInternalArrays(const TArray<TWeakObjectPtr<UPhysicsAsset>>& PhysicsAs
 							if (BoneIndex != INDEX_NONE && BoneIndex < RestTransforms.Num())
 							{
 								const FTransform RestTransform = RestTransforms[BoneIndex];
-								const FTransform BoneTransform = IsSkelMeshValid ? SkeletalMesh->GetBoneTransform(BoneIndex) * LocalTransform.Inverse() : RestTransform * WorldTransform * LocalTransform.Inverse();
+								const FTransform BoneTransform = IsSkelMeshValid ? ComputeRelativeTransform(SkeletalMesh->GetBoneTransform(BoneIndex), LocalTransform) : 
+									RestTransform * ComputeRelativeTransform(WorldTransform, LocalTransform);
 
 								for (const FKBoxElem& BoxElem : BodySetup->AggGeom.BoxElems)
 								{
@@ -359,7 +370,8 @@ void UpdateInternalArrays(const TArray<TWeakObjectPtr<UPhysicsAsset>>& PhysicsAs
 						const int32 BoneIndex = RefSkeleton.FindBoneIndex(BoneName);
 						if (BoneIndex != INDEX_NONE && BoneIndex < RestTransforms.Num())
 						{
-							const FTransform BoneTransform = IsSkelMeshValid ? SkeletalMesh->GetBoneTransform(BoneIndex) * LocalTransform.Inverse() : RestTransforms[BoneIndex] * WorldTransform * LocalTransform.Inverse();
+							const FTransform BoneTransform = IsSkelMeshValid ? ComputeRelativeTransform(SkeletalMesh->GetBoneTransform(BoneIndex), LocalTransform) : 
+								RestTransforms[BoneIndex] * ComputeRelativeTransform(WorldTransform, LocalTransform);
 
 							for (const FKBoxElem& BoxElem : BodySetup->AggGeom.BoxElems)
 							{
@@ -733,15 +745,12 @@ void UNiagaraDataInterfacePhysicsAsset::ExtractSourceComponent(FNiagaraSystemIns
 		}
 	}
 
-	const bool IsAssetMatching = (SourceComponent != nullptr) && (GroomPhysicsAsset != nullptr) &&
-		(GroomPhysicsAsset->GetPreviewMesh() == SourceComponent->SkeletalMesh);
-
 	SourceComponents.Empty();
 	PhysicsAssets.Empty();
 	if (SourceComponent != nullptr)
 	{
 		SourceComponents.Add(SourceComponent);
-		if (IsAssetMatching)
+		if (GroomPhysicsAsset)
 		{
 			PhysicsAssets.Add(GroomPhysicsAsset);
 		}
