Fix first frame motion blur during camera cut.

This CL add bCameraCut support for hair.

diff --git a/Engine/Plugins/Runtime/HairStrands/Source/HairStrandsCore/Private/GroomManager.cpp b/Engine/Plugins/Runtime/HairStrands/Source/HairStrandsCore/Private/GroomManager.cpp
--- a/Engine/Plugins/Runtime/HairStrands/Source/HairStrandsCore/Private/GroomManager.cpp
+++ b/Engine/Plugins/Runtime/HairStrands/Source/HairStrandsCore/Private/GroomManager.cpp
@@ -363,6 +363,7 @@ static void RunInternalHairStrandsInterpolation(
 	// Hair interpolation
 	if (EHairStrandsInterpolationType::RenderStrands == Type)
 	{
+		const bool bCameraCut = View->bCameraCut;
 		for (FHairStrandsInstance* AbstractInstance : Instances)
 		{
 			FHairGroupInstance* Instance = static_cast<FHairGroupInstance*>(AbstractInstance);
@@ -374,6 +375,7 @@ static void RunInternalHairStrandsInterpolation(
 				GraphBuilder, 
 				ShaderMap,
 				ViewUniqueID,
+				bCameraCut,
 				ViewRayTracingMask,
 				ShaderDrawData, 
 				ShaderPrintData,
@@ -735,6 +737,7 @@ static void RunHairLODSelection(
 		const float MinLOD = FMath::Max(0, GHairStrandsMinLOD);
 		float LODIndex = Instance->Debug.LODForcedIndex >= 0 ? FMath::Max(Instance->Debug.LODForcedIndex, MinLOD) : -1.0f;
 		float LODViewIndex = -1;
+		bool bCameraCut = false;
 		{
 			const FSphere SphereBound = Instance->ProxyBounds ? Instance->ProxyBounds->GetSphere() : FSphere(0);
 			for (const FSceneView* View : Views)
@@ -745,6 +748,8 @@ static void RunHairLODSelection(
 
 				// Select highest LOD accross all views
 				LODViewIndex = LODViewIndex < 0 ? CurrLODViewIndex : FMath::Min(LODViewIndex, CurrLODViewIndex);
+
+				bCameraCut = bCameraCut || View->bCameraCut;
 			}
 
 			if (LODIndex < 0)
@@ -947,6 +952,11 @@ static void RunHairLODSelection(
 		{
 			Instance->Debug.SkinningPreviousLocalToWorld = Instance->Debug.SkinningCurrentLocalToWorld;
 			Instance->Debug.SkinningCurrentLocalToWorld = CachedGeometry.LocalToWorld;
+
+			if (bCameraCut)
+			{
+				Instance->Debug.SkinningPreviousLocalToWorld = Instance->Debug.SkinningCurrentLocalToWorld;
+			}
 		}
 		Instance->LocalToWorld = Instance->GetCurrentLocalToWorld();
 	}
diff --git a/Engine/Plugins/Runtime/HairStrands/Source/HairStrandsCore/Private/HairStrandsRendering.cpp b/Engine/Plugins/Runtime/HairStrands/Source/HairStrandsCore/Private/HairStrandsRendering.cpp
--- a/Engine/Plugins/Runtime/HairStrands/Source/HairStrandsCore/Private/HairStrandsRendering.cpp
+++ b/Engine/Plugins/Runtime/HairStrands/Source/HairStrandsCore/Private/HairStrandsRendering.cpp
@@ -1636,7 +1636,7 @@ static void ConvertHairStrandsVFParameters(
 #define CONVERT_HAIRSSTRANDS_VF_PARAMETERS(OutName, ExternalBuffer) ConvertHairStrandsVFParameters(GraphBuilder, OutName, OutName##RHISRV, ExternalBuffer);
 
 // Compute/Update the hair strands description which will be used for rendering (VF) / voxelization & co.
-static FHairGroupPublicData::FVertexFactoryInput InternalComputeHairStrandsVertexInputData(FRDGBuilder* GraphBuilder, const FHairGroupInstance* Instance)
+static FHairGroupPublicData::FVertexFactoryInput InternalComputeHairStrandsVertexInputData(FRDGBuilder* GraphBuilder, const FHairGroupInstance* Instance, bool bCameraCut)
 {
 	FHairGroupPublicData::FVertexFactoryInput OutVFInput;
 	if (!Instance || Instance->GeometryType != EHairGeometryType::Strands)
@@ -1667,6 +1667,7 @@ static FHairGroupPublicData::FVertexFactoryInput InternalComputeHairStrandsVerte
 	else if (bSupportDeformation)
 	{
 		const bool bHasValidMotionVector =
+			!bCameraCut && 
 			Instance->Strands.DeformedResource->GetUniqueViewID(FHairStrandsDeformedResource::Current) ==
 			Instance->Strands.DeformedResource->GetUniqueViewID(FHairStrandsDeformedResource::Previous);
 		if (bHasValidMotionVector)
@@ -1675,6 +1676,9 @@ static FHairGroupPublicData::FVertexFactoryInput InternalComputeHairStrandsVerte
 			CONVERT_HAIRSSTRANDS_VF_PARAMETERS(OutVFInput.Strands.PrevPositionBuffer,		Instance->Strands.DeformedResource->GetBuffer(FHairStrandsDeformedResource::EFrameType::Previous));
 			CONVERT_HAIRSSTRANDS_VF_PARAMETERS(OutVFInput.Strands.PositionOffsetBuffer,		Instance->Strands.DeformedResource->GetPositionOffsetBuffer(FHairStrandsDeformedResource::EFrameType::Current));
 			CONVERT_HAIRSSTRANDS_VF_PARAMETERS(OutVFInput.Strands.PrevPositionOffsetBuffer, Instance->Strands.DeformedResource->GetPositionOffsetBuffer(FHairStrandsDeformedResource::EFrameType::Previous));
+
+			OutVFInput.Strands.PositionOffset	  = Instance->Strands.DeformedResource->GetPositionOffset(FHairStrandsDeformedResource::EFrameType::Current);
+			OutVFInput.Strands.PrevPositionOffset = Instance->Strands.DeformedResource->GetPositionOffset(FHairStrandsDeformedResource::EFrameType::Previous);
 		}
 		else
 		{
@@ -1682,14 +1686,15 @@ static FHairGroupPublicData::FVertexFactoryInput InternalComputeHairStrandsVerte
 			CONVERT_HAIRSSTRANDS_VF_PARAMETERS(OutVFInput.Strands.PrevPositionBuffer,		Instance->Strands.DeformedResource->GetBuffer(FHairStrandsDeformedResource::EFrameType::Current));
 			CONVERT_HAIRSSTRANDS_VF_PARAMETERS(OutVFInput.Strands.PositionOffsetBuffer,		Instance->Strands.DeformedResource->GetPositionOffsetBuffer(FHairStrandsDeformedResource::EFrameType::Current));
 			CONVERT_HAIRSSTRANDS_VF_PARAMETERS(OutVFInput.Strands.PrevPositionOffsetBuffer,	Instance->Strands.DeformedResource->GetPositionOffsetBuffer(FHairStrandsDeformedResource::EFrameType::Current));
+
+			OutVFInput.Strands.PositionOffset	  = Instance->Strands.DeformedResource->GetPositionOffset(FHairStrandsDeformedResource::EFrameType::Current);
+			OutVFInput.Strands.PrevPositionOffset = Instance->Strands.DeformedResource->GetPositionOffset(FHairStrandsDeformedResource::EFrameType::Current);
 		}
 		CONVERT_HAIRSSTRANDS_VF_PARAMETERS(OutVFInput.Strands.TangentBuffer,			Instance->Strands.DeformedResource->TangentBuffer);
 		CONVERT_HAIRSSTRANDS_VF_PARAMETERS(OutVFInput.Strands.Attribute0Buffer,			bDebugModePatchedAttributeBuffer ? Instance->Strands.DebugAttributeBuffer : Instance->Strands.RestResource->Attribute0Buffer); // TODO
 		CONVERT_HAIRSSTRANDS_VF_PARAMETERS(OutVFInput.Strands.Attribute1Buffer,			Instance->Strands.RestResource->Attribute1Buffer);
 		CONVERT_HAIRSSTRANDS_VF_PARAMETERS(OutVFInput.Strands.MaterialBuffer,			Instance->Strands.RestResource->MaterialBuffer);
 
-		OutVFInput.Strands.PositionOffset = Instance->Strands.DeformedResource->GetPositionOffset(FHairStrandsDeformedResource::EFrameType::Current);
-		OutVFInput.Strands.PrevPositionOffset = Instance->Strands.DeformedResource->GetPositionOffset(FHairStrandsDeformedResource::EFrameType::Previous);
 		OutVFInput.Strands.VertexCount = Instance->Strands.RestResource->GetVertexCount();
 	}
 	// 3. Render Strands in rest position: used when there are no skinning (rigid binding or no binding), no simulation, and no RBF
@@ -1728,9 +1733,9 @@ static FHairGroupPublicData::FVertexFactoryInput InternalComputeHairStrandsVerte
 	return OutVFInput;
 }
 
-FHairGroupPublicData::FVertexFactoryInput ComputeHairStrandsVertexInputData(const FHairGroupInstance* Instance)
+FHairGroupPublicData::FVertexFactoryInput ComputeHairStrandsVertexInputData(const FHairGroupInstance* Instance, bool bCameraCut)
 {
-	return InternalComputeHairStrandsVertexInputData(nullptr, Instance);
+	return InternalComputeHairStrandsVertexInputData(nullptr, Instance, bCameraCut);
 }
 
 void CreateHairStrandsDebugAttributeBuffer(FRDGBuilder& GraphBuilder, FRDGExternalBuffer* DebugAttributeBuffer, uint32 VertexCount);
@@ -1739,6 +1744,7 @@ void ComputeHairStrandsInterpolation(
 	FRDGBuilder& GraphBuilder,
 	FGlobalShaderMap* ShaderMap,
 	const uint32 ViewUniqueID,
+	const bool bCameraCut,
 	const uint32 ViewRayTracingMask,
 	const FShaderDrawDebugData* ShaderDrawData,
 	const FShaderPrintData* ShaderPrintData,
@@ -1786,7 +1792,7 @@ void ComputeHairStrandsInterpolation(
 				RegisterAsSRV(GraphBuilder, Instance->Guides.DeformedResource->GetBuffer(FHairStrandsDeformedResource::Current)),
 				Register(GraphBuilder, Instance->Guides.DeformedResource->TangentBuffer, ERDGImportedBufferFlags::CreateUAV));
 
-			Instance->HairGroupPublicData->VFInput = InternalComputeHairStrandsVertexInputData(&GraphBuilder, Instance);
+			Instance->HairGroupPublicData->VFInput = InternalComputeHairStrandsVertexInputData(&GraphBuilder, Instance, bCameraCut);
 		}
 		else
 		{
@@ -1957,7 +1963,7 @@ void ComputeHairStrandsInterpolation(
 			}
 
 			// 2.1 Update the VF input with the update resources
-			Instance->HairGroupPublicData->VFInput = InternalComputeHairStrandsVertexInputData(&GraphBuilder, Instance);
+			Instance->HairGroupPublicData->VFInput = InternalComputeHairStrandsVertexInputData(&GraphBuilder, Instance, bCameraCut);
 
 			// 3. Compute cluster AABBs (used for LODing and voxelization)
 			{
