Fix compile issue introduced in CL-17945620
Add missing parameters to SetPrimitiveIdStreamIndex / AddPrimitiveIdStreamElement calls.

diff --git a/Engine/Plugins/Runtime/HairStrands/Source/HairStrandsCore/Private/HairCardsVertexFactory.cpp b/Engine/Plugins/Runtime/HairStrands/Source/HairStrandsCore/Private/HairCardsVertexFactory.cpp
--- a/Engine/Plugins/Runtime/HairStrands/Source/HairStrandsCore/Private/HairCardsVertexFactory.cpp
+++ b/Engine/Plugins/Runtime/HairStrands/Source/HairStrandsCore/Private/HairCardsVertexFactory.cpp
@@ -268,8 +268,8 @@ void FHairCardsVertexFactory::InitResources()
 	// VertexFactory needs to be able to support max possible shader platform and feature level
 	// in case if we switch feature level at runtime.
 	FVertexDeclarationElementList Elements;
-	SetPrimitiveIdStreamIndex(EVertexInputStreamType::Default, -1);
-	AddPrimitiveIdStreamElement(EVertexInputStreamType::Default, 13, Elements);
+	SetPrimitiveIdStreamIndex(GetFeatureLevel(), EVertexInputStreamType::Default, -1);
+	AddPrimitiveIdStreamElement(EVertexInputStreamType::Default, Elements, 13, 0xff);
 
 	// Note this is a local version of the VF's bSupportsManualVertexFetch, which take into account the feature level
 	const bool bManualFetch = SupportsManualVertexFetch(GetFeatureLevel());
