Tentative fix for low repro assert on array bound.

This fix ensures that if the binding asset group count doesn't match the current groom asset group count, it doesn't get used.

diff --git a/Engine/Plugins/Runtime/HairStrands/Source/HairStrandsCore/Private/GroomComponent.cpp b/Engine/Plugins/Runtime/HairStrands/Source/HairStrandsCore/Private/GroomComponent.cpp
--- a/Engine/Plugins/Runtime/HairStrands/Source/HairStrandsCore/Private/GroomComponent.cpp
+++ b/Engine/Plugins/Runtime/HairStrands/Source/HairStrandsCore/Private/GroomComponent.cpp
@@ -2397,6 +2397,10 @@ void UGroomComponent::InitResources(bool bIsBindingReloading)
 		if (bHasNeedBindingData)
 		{
 			LocalBindingAsset = BindingAsset;
+			if (LocalBindingAsset && LocalBindingAsset->HairGroupResources.Num() != GroomAsset->HairGroupsData.Num())
+			{
+				LocalBindingAsset = nullptr;
+			}
 		}
 	}
 
