Fix some build issues from previous merge

diff --git a/Engine/Plugins/Runtime/HairStrands/Source/HairStrandsCore/Private/GroomManager.cpp b/Engine/Plugins/Runtime/HairStrands/Source/HairStrandsCore/Private/GroomManager.cpp
--- a/Engine/Plugins/Runtime/HairStrands/Source/HairStrandsCore/Private/GroomManager.cpp
+++ b/Engine/Plugins/Runtime/HairStrands/Source/HairStrandsCore/Private/GroomManager.cpp
@@ -968,11 +968,6 @@ static void RunHairLODSelection(
 		{
 			Instance->Debug.SkinningPreviousLocalToWorld = Instance->Debug.SkinningCurrentLocalToWorld;
 			Instance->Debug.SkinningCurrentLocalToWorld = CachedGeometry.LocalToWorld;
-
-			if (bCameraCut)
-			{
-				Instance->Debug.SkinningPreviousLocalToWorld = Instance->Debug.SkinningCurrentLocalToWorld;
-			}
 		}
 		Instance->LocalToWorld = Instance->GetCurrentLocalToWorld();
 	}
diff --git a/Engine/Plugins/Runtime/HairStrands/Source/HairStrandsCore/Private/HairStrandsRendering.cpp b/Engine/Plugins/Runtime/HairStrands/Source/HairStrandsCore/Private/HairStrandsRendering.cpp
--- a/Engine/Plugins/Runtime/HairStrands/Source/HairStrandsCore/Private/HairStrandsRendering.cpp
+++ b/Engine/Plugins/Runtime/HairStrands/Source/HairStrandsCore/Private/HairStrandsRendering.cpp
@@ -1636,7 +1636,7 @@ static void ConvertHairStrandsVFParameters(
 #define CONVERT_HAIRSSTRANDS_VF_PARAMETERS(OutName, ExternalBuffer) ConvertHairStrandsVFParameters(GraphBuilder, OutName, OutName##RHISRV, ExternalBuffer);
 
 // Compute/Update the hair strands description which will be used for rendering (VF) / voxelization & co.
-static FHairGroupPublicData::FVertexFactoryInput InternalComputeHairStrandsVertexInputData(FRDGBuilder* GraphBuilder, const FHairGroupInstance* Instance, bool bCameraCut)
+static FHairGroupPublicData::FVertexFactoryInput InternalComputeHairStrandsVertexInputData(FRDGBuilder* GraphBuilder, const FHairGroupInstance* Instance)
 {
 	FHairGroupPublicData::FVertexFactoryInput OutVFInput;
 	if (!Instance || Instance->GeometryType != EHairGeometryType::Strands)
@@ -1667,7 +1667,6 @@ static FHairGroupPublicData::FVertexFactoryInput InternalComputeHairStrandsVerte
 	else if (bSupportDeformation)
 	{
 		const bool bHasValidMotionVector =
-			!bCameraCut && 
 			Instance->Strands.DeformedResource->GetUniqueViewID(FHairStrandsDeformedResource::Current) ==
 			Instance->Strands.DeformedResource->GetUniqueViewID(FHairStrandsDeformedResource::Previous);
 		if (bHasValidMotionVector)
@@ -1733,9 +1732,9 @@ static FHairGroupPublicData::FVertexFactoryInput InternalComputeHairStrandsVerte
 	return OutVFInput;
 }
 
-FHairGroupPublicData::FVertexFactoryInput ComputeHairStrandsVertexInputData(const FHairGroupInstance* Instance, bool bCameraCut)
+FHairGroupPublicData::FVertexFactoryInput ComputeHairStrandsVertexInputData(const FHairGroupInstance* Instance)
 {
-	return InternalComputeHairStrandsVertexInputData(nullptr, Instance, bCameraCut);
+	return InternalComputeHairStrandsVertexInputData(nullptr, Instance);
 }
 
 void CreateHairStrandsDebugAttributeBuffer(FRDGBuilder& GraphBuilder, FRDGExternalBuffer* DebugAttributeBuffer, uint32 VertexCount);
@@ -1744,7 +1743,6 @@ void ComputeHairStrandsInterpolation(
 	FRDGBuilder& GraphBuilder,
 	FGlobalShaderMap* ShaderMap,
 	const uint32 ViewUniqueID,
-	const bool bCameraCut,
 	const uint32 ViewRayTracingMask,
 	const FShaderDrawDebugData* ShaderDrawData,
 	const FShaderPrintData* ShaderPrintData,
@@ -1792,7 +1790,7 @@ void ComputeHairStrandsInterpolation(
 				RegisterAsSRV(GraphBuilder, Instance->Guides.DeformedResource->GetBuffer(FHairStrandsDeformedResource::Current)),
 				Register(GraphBuilder, Instance->Guides.DeformedResource->TangentBuffer, ERDGImportedBufferFlags::CreateUAV));
 
-			Instance->HairGroupPublicData->VFInput = InternalComputeHairStrandsVertexInputData(&GraphBuilder, Instance, bCameraCut);
+			Instance->HairGroupPublicData->VFInput = InternalComputeHairStrandsVertexInputData(&GraphBuilder, Instance);
 		}
 		else
 		{
@@ -1963,7 +1961,7 @@ void ComputeHairStrandsInterpolation(
 			}
 
 			// 2.1 Update the VF input with the update resources
-			Instance->HairGroupPublicData->VFInput = InternalComputeHairStrandsVertexInputData(&GraphBuilder, Instance, bCameraCut);
+			Instance->HairGroupPublicData->VFInput = InternalComputeHairStrandsVertexInputData(&GraphBuilder, Instance);
 
 			// 3. Compute cluster AABBs (used for LODing and voxelization)
 			{
diff --git a/Engine/Source/Runtime/MovieSceneTracks/Private/Systems/MovieSceneDataLayerSystem.cpp b/Engine/Source/Runtime/MovieSceneTracks/Private/Systems/MovieSceneDataLayerSystem.cpp
--- a/Engine/Source/Runtime/MovieSceneTracks/Private/Systems/MovieSceneDataLayerSystem.cpp
+++ b/Engine/Source/Runtime/MovieSceneTracks/Private/Systems/MovieSceneDataLayerSystem.cpp
@@ -337,7 +337,7 @@ EDataLayerUpdateFlags FDesiredLayerStates::Apply(FPreAnimatedDataLayerStorage* P
 							
 						CSV_EVENT_GLOBAL(TEXT("SeqDataLayerFlushFull-%s-%s"),
 							*DataLayer->GetDataLayerLabel().ToString(),
-							*StaticEnum<EDataLayerRuntimeState>()->GetDisplayNameTextByValue((int64)CurrentState).ToString());
+							*StaticEnum<EDataLayerRuntimeState>()->GetDisplayNameTextByValue((int64)DesiredStateValue).ToString());
 					}
 					else
 					{
