Fix groom binding resources not always being released before being recreated.

diff --git a/Engine/Plugins/Runtime/HairStrands/Source/HairStrandsCore/Private/GroomBindingAsset.cpp b/Engine/Plugins/Runtime/HairStrands/Source/HairStrandsCore/Private/GroomBindingAsset.cpp
--- a/Engine/Plugins/Runtime/HairStrands/Source/HairStrandsCore/Private/GroomBindingAsset.cpp
+++ b/Engine/Plugins/Runtime/HairStrands/Source/HairStrandsCore/Private/GroomBindingAsset.cpp
@@ -92,6 +92,9 @@ void UGroomBindingAsset::InitResource()
 {
 	LLM_SCOPE(ELLMTag::Meshes) // This should be a Groom LLM tag, but there is no LLM tag bit left
 
+	// Ensure we are releasing binding resources before reallocating them
+	ReleaseResource();
+
 	for (FHairGroupBulkData& BulkData : HairGroupBulkDatas)
 	{
 		const int32 GroupIndex = HairGroupResources.Num();
