Fix bad merge

diff --git a/Engine/Plugins/Runtime/HairStrands/Source/HairStrandsCore/Private/GroomManager.cpp b/Engine/Plugins/Runtime/HairStrands/Source/HairStrandsCore/Private/GroomManager.cpp
--- a/Engine/Plugins/Runtime/HairStrands/Source/HairStrandsCore/Private/GroomManager.cpp
+++ b/Engine/Plugins/Runtime/HairStrands/Source/HairStrandsCore/Private/GroomManager.cpp
@@ -363,7 +363,6 @@ static void RunInternalHairStrandsInterpolation(
 	// Hair interpolation
 	if (EHairStrandsInterpolationType::RenderStrands == Type)
 	{
-		const bool bCameraCut = View->bCameraCut;
 		for (FHairStrandsInstance* AbstractInstance : Instances)
 		{
 			FHairGroupInstance* Instance = static_cast<FHairGroupInstance*>(AbstractInstance);
@@ -375,7 +374,6 @@ static void RunInternalHairStrandsInterpolation(
 				GraphBuilder, 
 				ShaderMap,
 				ViewUniqueID,
-				bCameraCut,
 				ViewRayTracingMask,
 				ShaderDrawData, 
 				ShaderPrintData,
@@ -685,22 +683,6 @@ void AddHairStreamingRequest(FHairGroupInstance* Instance, int32 InLODIndex)
 	}
 }
 
-#if RHI_RAYTRACING
-static void AllocateRaytracingResources(FHairGroupInstance* Instance)
-{
-	if (IsHairRayTracingEnabled() && !Instance->Strands.RenRaytracingResource)
-	{
-		check(Instance->Strands.Data);
-
-		// Allocate dynamic raytracing resources (owned by the groom component/instance)
-		FHairResourceName ResourceName(FName(Instance->Debug.GroomAssetName), Instance->Debug.GroupIndex);
-		Instance->Strands.RenRaytracingResource		 = new FHairStrandsRaytracingResource(*Instance->Strands.Data, ResourceName);
-		Instance->Strands.RenRaytracingResourceOwned = true;
-	}
-	Instance->Strands.ViewRayTracingMask |= EHairViewRayTracingMask::PathTracing;
-}
-#endif
-
 static void RunHairLODSelection(
 	FRDGBuilder& GraphBuilder, 
 	const FHairStrandsInstances& Instances, 
