* Fix imported precomputed weight for groom asset.
* Added preview when imported groom, if precomputed weights are available.
* Imported precomputed weight are now compatible with rendering strands decimation

diff --git a/Engine/Plugins/Runtime/HairStrands/Source/HairStrandsCore/Private/GroomBuilder.cpp b/Engine/Plugins/Runtime/HairStrands/Source/HairStrandsCore/Private/GroomBuilder.cpp
--- a/Engine/Plugins/Runtime/HairStrands/Source/HairStrandsCore/Private/GroomBuilder.cpp
+++ b/Engine/Plugins/Runtime/HairStrands/Source/HairStrandsCore/Private/GroomBuilder.cpp
@@ -15,7 +15,6 @@
 DEFINE_LOG_CATEGORY_STATIC(LogGroomBuilder, Log, All);
 
 #define LOCTEXT_NAMESPACE "GroomBuilder"
-#define ENABLE_EXTERNAL_INTERPOLATION_DATA 0
 
 // For debug purpose
 static float GHairInterpolationMetric_Distance = 1;
@@ -1286,97 +1285,63 @@ namespace HairInterpolationBuilder
 	}
 
 	/** Fill the GroomAsset with the interpolation data that exists in the HairDescription */
-	#if ENABLE_EXTERNAL_INTERPOLATION_DATA
-	void FillInterpolationData(UGroomAsset* GroomAsset, const FHairDescription& HairDescription)
+	void FillInterpolationData(
+		const FHairStrandsDatas& RenData,
+		const FHairStrandsDatas& SimData,
+		FHairStrandsInterpolationDatas& OutInterpolation)
 	{
 		TRACE_CPUPROFILER_EVENT_SCOPE(HairInterpolationBuilder::FillInterpolationData);
 
-		TStrandAttributesConstRef<FVector3f> ClosestGuides = HairDescription.StrandAttributes().GetAttributesRef<FVector3f>(HairAttribute::Strand::ClosestGuides);
-		TStrandAttributesConstRef<FVector3f> GuideWeights = HairDescription.StrandAttributes().GetAttributesRef<FVector3f>(HairAttribute::Strand::GuideWeights);
+		check(RenData.StrandsCurves.CurvesClosestGuideIDs.Num() != 0);
+		check(RenData.StrandsCurves.CurvesClosestGuideWeights.Num() != 0);
+		check(SimData.GetNumCurves() > 0);
+		int32 SimCurveCount = SimData.GetNumCurves();
 
-		if (!GroomAsset || !ClosestGuides.IsValid() || !GuideWeights.IsValid())
+		OutInterpolation.SetNum(RenData.GetNumPoints());
+		for (uint32 CurveIndex = 0; CurveIndex < RenData.GetNumCurves(); ++CurveIndex)
 		{
-			return;
-		}
-
-		const int32 NumGroups = GroomAsset->HairGroupsData.Num();
-		for (int32 GroupID = 0; GroupID < NumGroups; ++GroupID)
-		{
-			FHairGroupData& GroupData = GroomAsset->HairGroupsData[GroupID];
+			const uint32 CurveOffset = RenData.StrandsCurves.CurvesOffset[CurveIndex];
+			const uint16 CurveNumVertices = RenData.StrandsCurves.CurvesCount[CurveIndex];
 
-			const FHairStrandsDatas& RenderData = GroupData.Strands.Data;
-			const FHairStrandsDatas& SimulationData = GroupData.Guides.Data;
-			FHairStrandsInterpolationDatas& InterpolationData = GroupData.Strands.InterpolationData;
-
-			// Group must include imported guides data for imported interpolation data to be meaningful
-			// Otherwise, interpolation data will be built from generated guide data
-			if (SimulationData.GetNumCurves() == 0)
+			const FIntVector StrandClosestGuides = RenData.StrandsCurves.CurvesClosestGuideIDs[CurveIndex];
+			const FVector StrandGuideWeights  = RenData.StrandsCurves.CurvesClosestGuideWeights[CurveIndex];
+			for (uint16 VertexIndex = 0; VertexIndex < CurveNumVertices; ++VertexIndex)
 			{
-				continue;
-			}
+				const uint32 PointGlobalIndex = VertexIndex + CurveOffset;
+				const FVector& RenPointPosition = RenData.StrandsPoints.PointsPosition[PointGlobalIndex];
+				const float RenPointDistance = RenData.StrandsPoints.PointsCoordU[PointGlobalIndex] * RenData.StrandsCurves.CurvesLength[CurveIndex] * RenData.StrandsCurves.MaxLength;
 
-			InterpolationData.SetNum(RenderData.GetNumPoints());
-			for (uint32 CurveIndex = 0; CurveIndex < RenderData.GetNumCurves(); ++CurveIndex)
-			{
-				const uint32 CurveOffset = RenderData.StrandsCurves.CurvesOffset[CurveIndex];
-				const uint16 CurveNumVertices = RenderData.StrandsCurves.CurvesCount[CurveIndex];
-
-				const FStrandID StrandID(RenderData.StrandsCurves.StrandIDs[CurveIndex]);
-				const FVector StrandClosestGuides = ClosestGuides[StrandID];
-				const FVector StrandGuideWeights = GuideWeights[StrandID];
-				for (uint16 VertexIndex = 0; VertexIndex < CurveNumVertices; ++VertexIndex)
+				float TotalWeight = 0;
+				for (uint32 GuideIndex = 0; GuideIndex < FClosestGuides::Count; ++GuideIndex)
 				{
-					const uint32 PointGlobalIndex = VertexIndex + CurveOffset;
-					const FVector& RenPointPosition = RenderData.StrandsPoints.PointsPosition[PointGlobalIndex];
-					const float RenPointDistance = RenderData.StrandsPoints.PointsCoordU[PointGlobalIndex] * RenderData.StrandsCurves.CurvesLength[CurveIndex] * RenderData.StrandsCurves.MaxLength;
+					int32 ImportedGroomID = StrandClosestGuides[GuideIndex];
+					const int* SimCurveIndex = SimData.StrandsCurves.GroomIDToIndex.Find(ImportedGroomID);
 
-					float TotalWeight = 0;
-					for (uint32 GuideIndex = 0; GuideIndex < FClosestGuides::Count; ++GuideIndex)
+					if (SimCurveIndex && *SimCurveIndex >= 0 && *SimCurveIndex < SimCurveCount)
 					{
-						int32 ImportedGroomID = StrandClosestGuides[GuideIndex];
-						const int* SimCurveIndex = SimulationData.StrandsCurves.GroomIDToIndex.Find(ImportedGroomID);
-
-						if (SimCurveIndex && *SimCurveIndex >= 0)
-						{
-							// Fill the interpolation data using the ParametricDistance algorithm with a constant weight for all vertices along the strand
-							const uint32 SimOffset = SimulationData.StrandsCurves.CurvesOffset[*SimCurveIndex];
-							const FVertexInterpolationDesc Desc = FindMatchingVertex(RenPointDistance, SimulationData, *SimCurveIndex);
-
-							InterpolationData.PointsSimCurvesIndex[PointGlobalIndex][GuideIndex] = *SimCurveIndex;
-							InterpolationData.PointsSimCurvesVertexIndex[PointGlobalIndex][GuideIndex] = Desc.Index0 + SimOffset;
-							InterpolationData.PointsSimCurvesVertexLerp[PointGlobalIndex][GuideIndex] = Desc.T;
-							InterpolationData.PointsSimCurvesVertexWeights[PointGlobalIndex][GuideIndex] = StrandGuideWeights[GuideIndex];
-							TotalWeight += InterpolationData.PointsSimCurvesVertexWeights[PointGlobalIndex][GuideIndex];
-						}
+						// Fill the interpolation data using the ParametricDistance algorithm with a constant weight for all vertices along the strand
+						const uint32 SimOffset = SimData.StrandsCurves.CurvesOffset[*SimCurveIndex];
+						const FVertexInterpolationDesc Desc = FindMatchingVertex(RenPointDistance, SimData, *SimCurveIndex);
+
+						OutInterpolation.PointsSimCurvesIndex[PointGlobalIndex][GuideIndex] = *SimCurveIndex;
+						OutInterpolation.PointsSimCurvesVertexIndex[PointGlobalIndex][GuideIndex] = Desc.Index0 + SimOffset;
+						OutInterpolation.PointsSimCurvesVertexLerp[PointGlobalIndex][GuideIndex] = Desc.T;
+						OutInterpolation.PointsSimCurvesVertexWeights[PointGlobalIndex][GuideIndex] = StrandGuideWeights[GuideIndex];
+						TotalWeight += OutInterpolation.PointsSimCurvesVertexWeights[PointGlobalIndex][GuideIndex];
 					}
+				}
 
-					// Normalize the weights
-					if (TotalWeight > 0.f)
+				// Normalize the weights
+				if (TotalWeight > 0.f)
+				{
+					for (int32 GuideIndex = 0; GuideIndex < FClosestGuides::Count; ++GuideIndex)
 					{
-						for (int32 GuideIndex = 0; GuideIndex < FClosestGuides::Count; ++GuideIndex)
-						{
-							InterpolationData.PointsSimCurvesVertexWeights[PointGlobalIndex][GuideIndex] /= TotalWeight;
-						}
+						OutInterpolation.PointsSimCurvesVertexWeights[PointGlobalIndex][GuideIndex] /= TotalWeight;
 					}
 				}
 			}
 		}
-
-		// Deallocate the memory used for indices mapping since it's not needed anymore
-		for (int32 GroupID = 0; GroupID < NumGroups; ++GroupID)
-		{
-			FHairGroupData& GroupData = GroomAsset->HairGroupsData[GroupID];
-
-			FHairStrandsDatas& RenderData = GroupData.Strands.Data;
-			RenderData.StrandsCurves.GroomIDToIndex.Empty();
-			RenderData.StrandsCurves.StrandIDs.Empty();
-
-			FHairStrandsDatas& SimulationData = GroupData.Guides.Data;
-			SimulationData.StrandsCurves.GroomIDToIndex.Empty();
-			SimulationData.StrandsCurves.StrandIDs.Empty();
-		}
 	}
-	#endif // ENABLE_EXTERNAL_INTERPOLATION_DATA
 }
 
 class FGroomDataRandomizer
@@ -1421,6 +1386,11 @@ private:
 	TArray<FIntVector> GuideIndices;
 };
 
+FORCEINLINE FIntVector VectorToIntVector(const FVector& Index)
+{
+	return FIntVector(Index.X, Index.Y, Index.Z);
+}
+
 bool FGroomBuilder::BuildHairDescriptionGroups(const FHairDescription& HairDescription, FHairDescriptionGroups& Out)
 {
 	TRACE_CPUPROFILER_EVENT_SCOPE(FGroomBuilder::BuildHairDescriptionGroups);
@@ -1492,12 +1462,19 @@ bool FGroomBuilder::BuildHairDescriptionGroups(const FHairDescription& HairDescr
 
 	bool bImportGuides = true;
 
-	TStrandAttributesConstRef<FVector3f> ClosestGuides = HairDescription.StrandAttributes().GetAttributesRef<FVector3f>(HairAttribute::Strand::ClosestGuides);
-	TStrandAttributesConstRef<FVector3f> GuideWeights = HairDescription.StrandAttributes().GetAttributesRef<FVector3f>(HairAttribute::Strand::GuideWeights);
+	// For precomputed weights, there is two variants with single or triplet IDs/weights
+	// Single
+	TStrandAttributesConstRef<int> ClosestGuide		  = HairDescription.StrandAttributes().GetAttributesRef<int>(HairAttribute::Strand::ClosestGuides);
+	TStrandAttributesConstRef<float> GuideWeight	  = HairDescription.StrandAttributes().GetAttributesRef<float>(HairAttribute::Strand::GuideWeights);
+	// Triplet
+	TStrandAttributesConstRef<FVector3f> ClosestGuides  = HairDescription.StrandAttributes().GetAttributesRef<FVector3f>(HairAttribute::Strand::ClosestGuides);
+	TStrandAttributesConstRef<FVector3f> GuideWeights	  = HairDescription.StrandAttributes().GetAttributesRef<FVector3f>(HairAttribute::Strand::GuideWeights);
 
 	// To use ClosestGuides and GuideWeights attributes, guides must be imported from HairDescription and
 	// must include StrandID attribute since ClosestGuides references those IDs
-	const bool bCanUseClosestGuidesAndWeights = bImportGuides && StrandIDs.IsValid() && ClosestGuides.IsValid() && GuideWeights.IsValid();
+	const bool bPrecomputedWeight1 = ClosestGuide.IsValid()  && GuideWeight.IsValid();
+	const bool bPrecomputedWeight3 = ClosestGuides.IsValid() && GuideWeights.IsValid();
+	const bool bCanUseClosestGuidesAndWeights = bImportGuides && StrandIDs.IsValid() && (bPrecomputedWeight1 || bPrecomputedWeight3);
 
 	auto FindOrAdd = [&Out](int32 GroupID, FName GroupName) -> FHairDescriptionGroup&
 	{
@@ -1585,6 +1562,20 @@ bool FGroomBuilder::BuildHairDescriptionGroups(const FHairDescription& HairDescr
 			const int ImportedGroomID = StrandIDs[StrandID];
 			const int StrandCurveIndex = CurrentHairStrandsDatas->StrandsCurves.CurvesCount.Num() - 1;
 			CurrentHairStrandsDatas->StrandsCurves.GroomIDToIndex.Add(ImportedGroomID, StrandCurveIndex);
+
+			if (!bIsGuide)
+			{
+				if (bPrecomputedWeight3)
+				{
+					CurrentHairStrandsDatas->StrandsCurves.CurvesClosestGuideIDs.Add(VectorToIntVector(ClosestGuides[StrandID]));
+					CurrentHairStrandsDatas->StrandsCurves.CurvesClosestGuideWeights.Add(GuideWeights[StrandID]);
+				}
+				else
+				{
+					CurrentHairStrandsDatas->StrandsCurves.CurvesClosestGuideIDs.Add(FIntVector(ClosestGuide[StrandID]));
+					CurrentHairStrandsDatas->StrandsCurves.CurvesClosestGuideWeights.Add(FVector(GuideWeight[StrandID]));
+				}
+			}
 		}
 
 		if (bHasUVData)
@@ -1772,15 +1763,6 @@ void FGroomBuilder::BuildData(
 			check(OutGroupInfo.NumGuideVertices > 0);
 		}
 	}
-
-	// If there's usable closest guides and guide weights attributes, fill them into the asset
-	// This step requires the HairSimulationData (guides) to be filled prior to this
-#if ENABLE_EXTERNAL_INTERPOLATION_DATA
-	if (InHairDescriptionGroup.bCanUseClosestGuidesAndWeights)
-	{
-		HairInterpolationBuilder::FillInterpolationData(GroomAsset, HairDescription);
-	}
-#endif
 }
 
 void FGroomBuilder::BuildBulkData(
@@ -1814,22 +1796,35 @@ void FGroomBuilder::BuildInterplationData(
 	// Skip building if interpolation data was provided by the source file
 	if (OutInterpolationData.Num() == 0)
 	{
-		// This randomization makes certain strands being affected by 1, 2, or 3 guides
-		FRandomStream Random;
-		Random.Initialize(InInfo.GroupID);
-		const int32 NumRenderCurves = InRenData.GetNumCurves();
-		TArray<FIntVector> GuideIndices;
-		GuideIndices.SetNumUninitialized(NumRenderCurves);
-		for (int32 Index = 0; Index < NumRenderCurves; ++Index)
+		// If there's usable closest guides and guide weights attributes, fill them into the asset
+		// This step requires the HairSimulationData (guides) to be filled prior to this
+		const bool bUsePrecomputedWeights = !InInterpolationSettings.bOverrideGuides && InRenData.StrandsCurves.CurvesClosestGuideIDs.Num() > 0 && InRenData.StrandsCurves.CurvesClosestGuideWeights.Num() > 0;
+		if (bUsePrecomputedWeights)
 		{
-			FIntVector& RandomIndices = GuideIndices[Index];
-			for (int GuideIndex = 0; GuideIndex < HairInterpolationBuilder::FMetrics::Count; ++GuideIndex)
+			FHairStrandsInterpolationDatas OutInterpolation;
+			HairInterpolationBuilder::FillInterpolationData(InRenData, InSimData, OutInterpolationData);
+		}
+		else
+		{
+			// This randomization makes certain strands being affected by 1, 2, or 3 guides
+			FRandomStream Random;
+			Random.Initialize(InInfo.GroupID);
+			const int32 NumRenderCurves = InRenData.GetNumCurves();
+			TArray<FIntVector> GuideIndices;
+			GuideIndices.SetNumUninitialized(NumRenderCurves);
+			for (int32 Index = 0; Index < NumRenderCurves; ++Index)
 			{
-				RandomIndices[GuideIndex] = Random.RandRange(0, HairInterpolationBuilder::FMetrics::Count - 1);
+				FIntVector& RandomIndices = GuideIndices[Index];
+				for (int GuideIndex = 0; GuideIndex < HairInterpolationBuilder::FMetrics::Count; ++GuideIndex)
+				{
+					RandomIndices[GuideIndex] = Random.RandRange(0, HairInterpolationBuilder::FMetrics::Count - 1);
+				}
 			}
+
+			HairInterpolationBuilder::BuildInterpolationData(OutInterpolationData, InSimData, InRenData, InInterpolationSettings, GuideIndices);
 		}
 
-		HairInterpolationBuilder::BuildInterpolationData(OutInterpolationData, InSimData, InRenData, InInterpolationSettings, GuideIndices);
+
 	}
 }
 
@@ -1901,6 +1896,7 @@ void Decimate(
 	// Divide strands in buckets and pick randomly one stand per bucket
 	const uint32 InCurveCount = InData.StrandsCurves.Num();
 	const uint32 OutCurveCount = FMath::Clamp(uint32(InCurveCount * CurveDecimationPercentage), 1u, InCurveCount);
+	const bool bHasPrecomputedWeights = InData.StrandsCurves.CurvesClosestGuideIDs.Num() > 0 && InData.StrandsCurves.CurvesClosestGuideWeights.Num() > 0;
 
 	TArray<uint32> CurveIndices;
 	CurveIndices.SetNum(OutCurveCount);
@@ -1999,6 +1995,11 @@ void Decimate(
 		OutData.StrandsCurves.CurvesRootUV[OutCurveIndex] = InData.StrandsCurves.CurvesRootUV[InCurveIndex];
 		OutData.StrandsCurves.CurvesOffset[OutCurveIndex] = OutPointOffset;
 		OutData.StrandsCurves.CurvesLength[OutCurveIndex] = InData.StrandsCurves.CurvesLength[InCurveIndex];
+		if (bHasPrecomputedWeights)
+		{
+			OutData.StrandsCurves.CurvesClosestGuideIDs[OutCurveIndex] = InData.StrandsCurves.CurvesClosestGuideIDs[InCurveIndex];
+			OutData.StrandsCurves.CurvesClosestGuideWeights[OutCurveIndex] = InData.StrandsCurves.CurvesClosestGuideWeights[InCurveIndex];
+		}
 		OutData.StrandsCurves.MaxLength = InData.StrandsCurves.MaxLength;
 		OutData.StrandsCurves.MaxRadius = InData.StrandsCurves.MaxRadius;
 		OutPointOffset += OutPointCount;
diff --git a/Engine/Plugins/Runtime/HairStrands/Source/HairStrandsCore/Private/HairStrandsDatas.cpp b/Engine/Plugins/Runtime/HairStrands/Source/HairStrandsCore/Private/HairStrandsDatas.cpp
--- a/Engine/Plugins/Runtime/HairStrands/Source/HairStrandsCore/Private/HairStrandsDatas.cpp
+++ b/Engine/Plugins/Runtime/HairStrands/Source/HairStrandsCore/Private/HairStrandsDatas.cpp
@@ -18,6 +18,11 @@ void FHairStrandsCurves::SetNum(const uint32 NumCurves)
 	CurvesCount.SetNum(NumCurves);
 	CurvesLength.SetNum(NumCurves);
 	CurvesRootUV.SetNum(NumCurves);
+
+	// Curves clossest guide are not initialized to track if the curves 
+	// data contains or not precomputed weights
+	// CurvesClosestGuideIDs.SetNum(0);
+	// CurvesClosestGuideWeights.SetNum(0);
 }
 
 void FHairStrandsPoints::SetNum(const uint32 NumPoints)
@@ -43,6 +48,8 @@ void FHairStrandsCurves::Reset()
 	CurvesCount.Reset();
 	CurvesLength.Reset();
 	CurvesRootUV.Reset();
+	CurvesClosestGuideIDs.Reset();
+	CurvesClosestGuideWeights.Reset();
 }
 
 void FHairStrandsPoints::Reset()
diff --git a/Engine/Plugins/Runtime/HairStrands/Source/HairStrandsCore/Public/GroomImportOptions.h b/Engine/Plugins/Runtime/HairStrands/Source/HairStrandsCore/Public/GroomImportOptions.h
--- a/Engine/Plugins/Runtime/HairStrands/Source/HairStrandsCore/Public/GroomImportOptions.h
+++ b/Engine/Plugins/Runtime/HairStrands/Source/HairStrandsCore/Public/GroomImportOptions.h
@@ -34,6 +34,9 @@ struct HAIRSTRANDSCORE_API FGroomHairGroupPreview
 	, InterpolationSettings()
 	{}
 
+	UPROPERTY(VisibleAnywhere, BlueprintReadWrite, Category = Preview)
+	FName GroupName;
+
 	UPROPERTY(VisibleAnywhere, BlueprintReadWrite, Category = Preview)
 	int32 GroupID;
 
@@ -43,6 +46,9 @@ struct HAIRSTRANDSCORE_API FGroomHairGroupPreview
 	UPROPERTY(VisibleAnywhere, BlueprintReadWrite, Category = Preview)
 	int32 GuideCount;
 
+	UPROPERTY(VisibleAnywhere, BlueprintReadWrite, Category = Preview)
+	bool bHasPrecomputedWeights;
+
 	UPROPERTY(EditAnywhere, BlueprintReadWrite, Category = Preview)
 	FHairGroupsInterpolation InterpolationSettings;
 };
diff --git a/Engine/Plugins/Runtime/HairStrands/Source/HairStrandsCore/Public/HairStrandsDatas.h b/Engine/Plugins/Runtime/HairStrands/Source/HairStrandsCore/Public/HairStrandsDatas.h
--- a/Engine/Plugins/Runtime/HairStrands/Source/HairStrandsCore/Public/HairStrandsDatas.h
+++ b/Engine/Plugins/Runtime/HairStrands/Source/HairStrandsCore/Public/HairStrandsDatas.h
@@ -366,6 +366,8 @@ struct HAIRSTRANDSCORE_API FHairStrandsCurves
 	/** Get the number of Curves */
 	uint32 Num() const { return CurvesCount.Num(); }
 
+	bool HasPrecomputedWeights() const { return CurvesClosestGuideIDs.Num() > 0 && CurvesClosestGuideWeights.Num() > 0; }
+
 	/** Number of points per rod */
 	TArray<uint16> CurvesCount;
 
@@ -384,6 +386,12 @@ struct HAIRSTRANDSCORE_API FHairStrandsCurves
 	/** Mapping of imported Groom ID to index */
 	TMap<int, int> GroomIDToIndex;
 
+	/** Custom guide IDs (indexed with StrandID) */
+	TArray<FIntVector> CurvesClosestGuideIDs;
+
+	/** Custom guid weights (indexed with StrandID) */
+	TArray<FVector> CurvesClosestGuideWeights;
+
 	/** Max strands Curves length */
 	float MaxLength = 0;
 
diff --git a/Engine/Plugins/Runtime/HairStrands/Source/HairStrandsEditor/Private/GroomActions.cpp b/Engine/Plugins/Runtime/HairStrands/Source/HairStrandsEditor/Private/GroomActions.cpp
--- a/Engine/Plugins/Runtime/HairStrands/Source/HairStrandsEditor/Private/GroomActions.cpp
+++ b/Engine/Plugins/Runtime/HairStrands/Source/HairStrandsEditor/Private/GroomActions.cpp
@@ -187,9 +187,11 @@ void FGroomActions::ExecuteRebuild(TArray<TWeakObjectPtr<UGroomAsset>> Objects)
 					{
 						FGroomHairGroupPreview& OutGroup = GroupsPreview->Groups.AddDefaulted_GetRef();
 						OutGroup.GroupID    = GroupIndex;
+						OutGroup.GroupName	= GroomAsset->HairGroupsInfo[GroupIndex].GroupName;
 						OutGroup.CurveCount = GroomAsset->HairGroupsData[GroupIndex].Strands.BulkData.GetNumCurves();
 						OutGroup.GuideCount = GroomAsset->HairGroupsData[GroupIndex].Guides.BulkData.GetNumCurves();
 						OutGroup.InterpolationSettings = GroomAsset->HairGroupsInterpolation[GroupIndex];
+						OutGroup.bHasPrecomputedWeights = false;
 					}
 				}
 				TSharedPtr<SGroomImportOptionsWindow> GroomOptionWindow = SGroomImportOptionsWindow::DisplayRebuildOptions(CurrentOptions, GroupsPreview, Filename);
diff --git a/Engine/Plugins/Runtime/HairStrands/Source/HairStrandsEditor/Private/HairStrandsFactory.cpp b/Engine/Plugins/Runtime/HairStrands/Source/HairStrandsEditor/Private/HairStrandsFactory.cpp
--- a/Engine/Plugins/Runtime/HairStrands/Source/HairStrandsEditor/Private/HairStrandsFactory.cpp
+++ b/Engine/Plugins/Runtime/HairStrands/Source/HairStrandsEditor/Private/HairStrandsFactory.cpp
@@ -109,9 +109,11 @@ UObject* UHairStrandsFactory::FactoryCreateFile(UClass* InClass, UObject* InPare
 			for (const FHairDescriptionGroup& Group : OutDescription.HairGroups)
 			{
 				FGroomHairGroupPreview& OutGroup = GroupsPreview->Groups.AddDefaulted_GetRef();
+				OutGroup.GroupName  = Group.Info.GroupName;
 				OutGroup.GroupID	= Group.Info.GroupID;
 				OutGroup.CurveCount = Group.Info.NumCurves;
 				OutGroup.GuideCount = Group.Info.NumGuides;
+				OutGroup.bHasPrecomputedWeights = Group.Strands.StrandsCurves.HasPrecomputedWeights();
 
 				if (OutGroup.GroupID < OutDescription.HairGroups.Num())
 				{				
diff --git a/Engine/Plugins/Runtime/HairStrands/Source/HairStrandsEditor/Private/ReimportHairStrandsFactory.cpp b/Engine/Plugins/Runtime/HairStrands/Source/HairStrandsEditor/Private/ReimportHairStrandsFactory.cpp
--- a/Engine/Plugins/Runtime/HairStrands/Source/HairStrandsEditor/Private/ReimportHairStrandsFactory.cpp
+++ b/Engine/Plugins/Runtime/HairStrands/Source/HairStrandsEditor/Private/ReimportHairStrandsFactory.cpp
@@ -371,9 +371,11 @@ EReimportResult::Type UReimportHairStrandsFactory::Reimport(UObject* Obj)
 			for (const FHairDescriptionGroup& Group : OutDescription.HairGroups)
 			{
 				FGroomHairGroupPreview& OutGroup = GroupsPreview->Groups.AddDefaulted_GetRef();
+				OutGroup.GroupName  = Group.Info.GroupName;
 				OutGroup.GroupID	= Group.Info.GroupID;
 				OutGroup.CurveCount = Group.Info.NumCurves;
 				OutGroup.GuideCount = Group.Info.NumGuides;
+				OutGroup.bHasPrecomputedWeights = Group.Strands.StrandsCurves.HasPrecomputedWeights();
 
 				if (OutGroup.GroupID < OutDescription.HairGroups.Num())
 				{
