Fix for warnings in incremental cook jobs

diff --git a/Engine/Plugins/Runtime/HairStrands/Source/HairStrandsCore/Public/GroomImportOptions.h b/Engine/Plugins/Runtime/HairStrands/Source/HairStrandsCore/Public/GroomImportOptions.h
--- a/Engine/Plugins/Runtime/HairStrands/Source/HairStrandsCore/Public/GroomImportOptions.h
+++ b/Engine/Plugins/Runtime/HairStrands/Source/HairStrandsCore/Public/GroomImportOptions.h
@@ -31,6 +31,7 @@ struct HAIRSTRANDSCORE_API FGroomHairGroupPreview
 	: GroupID(0)
 	, CurveCount(0)
 	, GuideCount(0)
+	, bHasPrecomputedWeights(false)
 	, InterpolationSettings()
 	{}
 
