Fix invalid error/warning with hair cards factory.

diff --git a/Engine/Plugins/Runtime/HairStrands/Source/HairStrandsCore/Private/HairCardsVertexFactory.cpp b/Engine/Plugins/Runtime/HairStrands/Source/HairStrandsCore/Private/HairCardsVertexFactory.cpp
--- a/Engine/Plugins/Runtime/HairStrands/Source/HairStrandsCore/Private/HairCardsVertexFactory.cpp
+++ b/Engine/Plugins/Runtime/HairStrands/Source/HairStrandsCore/Private/HairCardsVertexFactory.cpp
@@ -221,12 +221,15 @@ void FHairCardsVertexFactory::ModifyCompilationEnvironment(const FVertexFactoryS
 
 void FHairCardsVertexFactory::ValidateCompiledResult(const FVertexFactoryType* Type, EShaderPlatform Platform, const FShaderParameterMap& ParameterMap, TArray<FString>& OutErrors)
 {
+	// This is not relevant to hair cards at the moment, as instancing is not supported anyway (need to bind unique resource per card instance
+	#if 0
 	if (Type->SupportsPrimitiveIdStream() 
 		&& UseGPUScene(Platform, GetMaxSupportedFeatureLevel(Platform)) 
 		&& ParameterMap.ContainsParameterAllocation(FPrimitiveUniformShaderParameters::StaticStructMetadata.GetShaderVariableName()))
 	{
 		OutErrors.AddUnique(*FString::Printf(TEXT("Shader attempted to bind the Primitive uniform buffer even though Vertex Factory %s computes a PrimitiveId per-instance.  This will break auto-instancing.  Shaders should use GetPrimitiveData(PrimitiveId).Member instead of Primitive.Member."), Type->GetName()));
 	}
+	#endif
 }
 
 void FHairCardsVertexFactory::SetData(const FDataType& InData)
