Add hair group index support for hair cards & hair meshes.

This CL adds:
 * Hair strands group index texture for transferring hair strands group ID onto hair meshes.
 * Hair strands group index voxelization for transferring hair strands group ID onto hair cards geometry.
 * Hair group index debug mode support for hair cards & hair meshes.

Hair group index transfer is automatically done, but hair mehes transfer require manual texture generation.

diff --git a/Engine/Plugins/Runtime/HairStrands/Source/HairStrandsCore/Private/GroomAsset.cpp b/Engine/Plugins/Runtime/HairStrands/Source/HairStrandsCore/Private/GroomAsset.cpp
--- a/Engine/Plugins/Runtime/HairStrands/Source/HairStrandsCore/Private/GroomAsset.cpp
+++ b/Engine/Plugins/Runtime/HairStrands/Source/HairStrandsCore/Private/GroomAsset.cpp
@@ -416,6 +416,11 @@ static void InitAtlasTexture(ResourceType* InResource, UTexture2D* InTexture, EH
 			InResource->AuxilaryDataTexture = InTexture->TextureReference.TextureReferenceRHI;
 			InResource->AuxilaryDataSampler = DefaultSampler;
 		} break;
+		case EHairAtlasTextureType::GroupIndex:
+		{
+			InResource->GroupIndexTexture = InTexture->TextureReference.TextureReferenceRHI;
+			InResource->GroupIndexSampler = DefaultSampler;
+		} break;
 		}
 	});
 }
@@ -499,7 +504,7 @@ static bool BuildHairGroup(
 			OutHairGroupsData[GroupIndex].Strands.InterpolationBulkData.Reset();
 		}
 
-		FGroomBuilder::BuildClusterBulkData(StrandsData, HairDescriptionGroups.BoundRadius, HairGroupsLOD[GroupIndex], OutHairGroupsData[GroupIndex].Strands.ClusterCullingBulkData);
+		FGroomBuilder::BuildClusterBulkData(StrandsData, HairDescriptionGroups.Bounds.SphereRadius, HairGroupsLOD[GroupIndex], OutHairGroupsData[GroupIndex].Strands.ClusterCullingBulkData);
 	}
 	return bIsValid;
 }
@@ -526,7 +531,7 @@ static bool BuildHairGroupCluster(
 		FHairStrandsDatas StrandsData;
 		FHairStrandsDatas GuidesData;
 		FGroomBuilder::BuildData(HairGroup, HairGroupsInterpolation[GroupIndex], OutHairGroupsInfo[GroupIndex], StrandsData, GuidesData);
-		FGroomBuilder::BuildClusterBulkData(StrandsData, HairDescriptionGroups.BoundRadius, HairGroupsLOD[GroupIndex], OutHairGroupsData[GroupIndex].Strands.ClusterCullingBulkData);
+		FGroomBuilder::BuildClusterBulkData(StrandsData, HairDescriptionGroups.Bounds.SphereRadius, HairGroupsLOD[GroupIndex], OutHairGroupsData[GroupIndex].Strands.ClusterCullingBulkData);
 	}
 	return bIsValid;
 }
@@ -1284,7 +1289,8 @@ static bool IsCardsTextureResources(const FName PropertyName)
 		|| PropertyName == GET_MEMBER_NAME_CHECKED(FHairGroupCardsTextures, CoverageTexture)
 		|| PropertyName == GET_MEMBER_NAME_CHECKED(FHairGroupCardsTextures, TangentTexture)
 		|| PropertyName == GET_MEMBER_NAME_CHECKED(FHairGroupCardsTextures, AttributeTexture)
-		|| PropertyName == GET_MEMBER_NAME_CHECKED(FHairGroupCardsTextures, AuxilaryDataTexture);
+		|| PropertyName == GET_MEMBER_NAME_CHECKED(FHairGroupCardsTextures, AuxilaryDataTexture)
+		|| PropertyName == GET_MEMBER_NAME_CHECKED(FHairGroupCardsTextures, GroupIndexTexture);
 }
 static void InitCardsTextureResources(UGroomAsset* GroomAsset);
 
@@ -2459,6 +2465,7 @@ bool UGroomAsset::BuildCardsGeometry(uint32 GroupIndex)
 					if (Textures.AttributeTexture != nullptr)	Textures.AttributeTexture->UpdateResource();
 					if (Textures.CoverageTexture != nullptr)	Textures.CoverageTexture->UpdateResource();
 					if (Textures.AuxilaryDataTexture != nullptr)Textures.AuxilaryDataTexture->UpdateResource();
+					if (Textures.GroupIndexTexture != nullptr)	Textures.GroupIndexTexture->UpdateResource();
 				}
 			}
 		}
@@ -2519,16 +2526,20 @@ bool UGroomAsset::BuildCardsGeometry(uint32 GroupIndex)
 			{
 				CardsMesh->ConditionalPostLoad();
 
-				// Create a transient FHairStrandsData in order to extract RootUV and transfer them to cards data
+				// * Create a transient FHairStrandsData in order to extract RootUV and transfer them to cards data
+				// * Voxelize hair strands data (all group), to transfer group index from strands to cards
 				FHairStrandsDatas TempHairStrandsData;
+				FHairStrandsVoxelData TempHairStrandsVoxelData;
 				{
 					const FHairDescriptionGroups& LocalHairDescriptionGroups = GetHairDescriptionGroups();
 					FHairGroupInfo DummyInfo;
 					check(LocalHairDescriptionGroups.IsValid());
 					TempHairStrandsData = LocalHairDescriptionGroups.HairGroups[GroupIndex].Strands;
 					FGroomBuilder::BuildData(TempHairStrandsData);
+					FGroomBuilder::VoxelizeGroupIndex(LocalHairDescriptionGroups, TempHairStrandsVoxelData);
 				}
-				bInitResources = FHairCardsBuilder::ImportGeometry(CardsMesh, TempHairStrandsData, LOD.BulkData, LODGuidesData, LOD.InterpolationBulkData);
+				
+				bInitResources = FHairCardsBuilder::ImportGeometry(CardsMesh, TempHairStrandsData, TempHairStrandsVoxelData, LOD.BulkData, LODGuidesData, LOD.InterpolationBulkData);
 				if (!bInitResources)
 				{
 					UE_LOG(LogHairStrands, Warning, TEXT("Failed to import cards from %s for Group %d LOD %d."), *CardsMesh->GetName(), GroupIndex, LODIt);
@@ -2556,7 +2567,7 @@ bool UGroomAsset::BuildCardsGeometry(uint32 GroupIndex)
 				InitAtlasTexture(LOD.RestResource, Desc->Textures.AttributeTexture, EHairAtlasTextureType::Attribute);
 				InitAtlasTexture(LOD.RestResource, Desc->Textures.CoverageTexture, EHairAtlasTextureType::Coverage);
 				InitAtlasTexture(LOD.RestResource, Desc->Textures.AuxilaryDataTexture, EHairAtlasTextureType::AuxilaryData);
-				InitAtlasTexture(LOD.RestResource, Desc->Textures.AuxilaryDataTexture, EHairAtlasTextureType::AuxilaryData);
+				InitAtlasTexture(LOD.RestResource, Desc->Textures.GroupIndexTexture, EHairAtlasTextureType::GroupIndex);
 				LOD.RestResource->bInvertUV = Desc->SourceType == EHairCardsSourceType::Procedural;
 				
 				// 2.2 Load interoplatino resources
@@ -2718,6 +2729,7 @@ bool UGroomAsset::BuildMeshesGeometry(uint32 GroupIndex)
 					if (Textures.AttributeTexture != nullptr)	Textures.AttributeTexture->UpdateResource();
 					if (Textures.CoverageTexture != nullptr)	Textures.CoverageTexture->UpdateResource();
 					if (Textures.AuxilaryDataTexture != nullptr)Textures.AuxilaryDataTexture->UpdateResource();
+					if (Textures.GroupIndexTexture != nullptr)	Textures.GroupIndexTexture->UpdateResource();
 				}
 			}
 		}
@@ -2978,6 +2990,7 @@ static void InitCardsTextureResources(UGroomAsset* GroomAsset)
 					InitAtlasTexture(LOD.RestResource, Desc->Textures.AttributeTexture, EHairAtlasTextureType::Attribute);
 					InitAtlasTexture(LOD.RestResource, Desc->Textures.CoverageTexture, EHairAtlasTextureType::Coverage);
 					InitAtlasTexture(LOD.RestResource, Desc->Textures.AuxilaryDataTexture, EHairAtlasTextureType::AuxilaryData);
+					InitAtlasTexture(LOD.RestResource, Desc->Textures.GroupIndexTexture, EHairAtlasTextureType::GroupIndex);
 					if (LOD.RestResource)
 					{
 						LOD.RestResource->bInvertUV = Desc->SourceType == EHairCardsSourceType::Procedural; // Should fix procedural texture so that this does not happen
@@ -3035,7 +3048,8 @@ void UGroomAsset::InitCardsResources()
 					InitAtlasTexture(LOD.RestResource, Desc->Textures.TangentTexture, EHairAtlasTextureType::Tangent);
 					InitAtlasTexture(LOD.RestResource, Desc->Textures.AttributeTexture, EHairAtlasTextureType::Attribute);
 					InitAtlasTexture(LOD.RestResource, Desc->Textures.CoverageTexture, EHairAtlasTextureType::Coverage);
-					InitAtlasTexture(LOD.RestResource, Desc->Textures.AuxilaryDataTexture, EHairAtlasTextureType::AuxilaryData);					
+					InitAtlasTexture(LOD.RestResource, Desc->Textures.AuxilaryDataTexture, EHairAtlasTextureType::AuxilaryData);
+					InitAtlasTexture(LOD.RestResource, Desc->Textures.GroupIndexTexture, EHairAtlasTextureType::GroupIndex);
 					LOD.RestResource->bInvertUV = Desc->SourceType == EHairCardsSourceType::Procedural; // Should fix procedural texture so that this does not happen
 				}
 			}
@@ -3088,6 +3102,7 @@ void UGroomAsset::InitMeshesResources()
 					InitAtlasTexture(LOD.RestResource, Desc->Textures.AttributeTexture, EHairAtlasTextureType::Attribute);
 					InitAtlasTexture(LOD.RestResource, Desc->Textures.CoverageTexture, EHairAtlasTextureType::Coverage);
 					InitAtlasTexture(LOD.RestResource, Desc->Textures.AuxilaryDataTexture, EHairAtlasTextureType::AuxilaryData);
+					InitAtlasTexture(LOD.RestResource, Desc->Textures.GroupIndexTexture, EHairAtlasTextureType::GroupIndex);
 				}
 			}
 		}
@@ -3451,7 +3466,7 @@ void UGroomAsset::SaveProceduralCards(uint32 DescIndex)
 
 	const int32 GroupIndex = Desc->GroupIndex;
 	const int32 LODIndex = Desc->LODIndex;
-	if (GroupIndex >= HairGroupsData.Num() || LODIndex >= HairGroupsData[GroupIndex].Cards.LODs.Num())
+	if (GroupIndex >= HairGroupsData.Num())
 		return;
 
 	// 1. Convert old parameters (ClusterDecimation, bUseCards) to new parameters (GenerationType & CardsCount)
@@ -3550,6 +3565,7 @@ void UGroomAsset::SavePendingProceduralAssets()
 					if (Q->Textures->AuxilaryDataTexture)	FHairStrandsCore::SaveAsset(Q->Textures->AuxilaryDataTexture);
 					if (Q->Textures->CoverageTexture)		FHairStrandsCore::SaveAsset(Q->Textures->CoverageTexture);
 					if (Q->Textures->TangentTexture)		FHairStrandsCore::SaveAsset(Q->Textures->TangentTexture);
+					if (Q->Textures->GroupIndexTexture)		FHairStrandsCore::SaveAsset(Q->Textures->GroupIndexTexture);
 					Q->Textures->bNeedToBeSaved = false;
 
 					InternalReleaseResource(Q->Resources);
@@ -3619,10 +3635,11 @@ bool UGroomAsset::GetHairCardsGuidesDatas(
 		{
 			CardsMesh->ConditionalPostLoad();
 
+			FHairStrandsVoxelData				DummyVoxelData;
 			FHairCardsBulkData					DummyBulkData;
 			FHairCardsInterpolationBulkData		DummyInterpolationBulkData;
 			FHairStrandsDatas					DummyHairStrandsData;
-			FHairCardsBuilder::ImportGeometry(CardsMesh, DummyHairStrandsData, DummyBulkData, OutCardsGuidesData, DummyInterpolationBulkData);
+			FHairCardsBuilder::ImportGeometry(CardsMesh, DummyHairStrandsData, DummyVoxelData, DummyBulkData, OutCardsGuidesData, DummyInterpolationBulkData);
 			return true;
 		}
 	}
diff --git a/Engine/Plugins/Runtime/HairStrands/Source/HairStrandsCore/Private/GroomBuilder.cpp b/Engine/Plugins/Runtime/HairStrands/Source/HairStrandsCore/Private/GroomBuilder.cpp
--- a/Engine/Plugins/Runtime/HairStrands/Source/HairStrandsCore/Private/GroomBuilder.cpp
+++ b/Engine/Plugins/Runtime/HairStrands/Source/HairStrandsCore/Private/GroomBuilder.cpp
@@ -29,6 +29,9 @@ static FAutoConsoleVariableRef CVarHairInterpolationMetric_AngleAttenuation(TEXT
 static int32 GHairClusterBuilder_MaxVoxelResolution = 256;
 static FAutoConsoleVariableRef CVarHairClusterBuilder_MaxVoxelResolution(TEXT("r.HairStrands.ClusterBuilder.MaxVoxelResolution"), GHairClusterBuilder_MaxVoxelResolution, TEXT("Max voxel resolution used when building hair strands cluster data to avoid too long building time (default:128).  "));
 
+static int32 GHairGroupIndexBuilder_MaxVoxelResolution = 64;
+static FAutoConsoleVariableRef CVarHairGroupIndexBuilder_MaxVoxelResolution(TEXT("r.HairStrands.HairGroupBuilder.MaxVoxelResolution"), GHairGroupIndexBuilder_MaxVoxelResolution, TEXT("Max voxel resolution used when voxelizing hair strands to transfer group index grom strands to cards. This avoids too long building time (default:64).  "));
+
 FString FGroomBuilder::GetVersion()
 {
 	// Important to update the version when groom building changes
@@ -1668,7 +1671,7 @@ bool FGroomBuilder::BuildHairDescriptionGroups(const FHairDescription& HairDescr
 			}
 		}
 
-		Out.BoundRadius = FVector3f::Distance(GroomBoundMax, GroomBoundMin) * 0.5f;
+		Out.Bounds = FBoxSphereBounds3f(FBox3f(GroomBoundMin, GroomBoundMax));
 	}
 
 	// Update GroupInfo
@@ -2009,6 +2012,273 @@ void Decimate(
 }
 } // namespace FHairStrandsDecimation
 
+///////////////////////////////////////////////////////////////////////////////////////////////////
+// Hair voxelization
+namespace GroomBuilder_Voxelization
+{
+
+FORCEINLINE uint32 ToLinearCoord(const FIntVector& T, const FIntVector& Resolution)
+{
+	// Morton instead for better locality?
+	return T.X + T.Y * Resolution.X + T.Z * Resolution.X * Resolution.Y;
+}
+
+FORCEINLINE FIntVector ToCoord(const FVector& T, const FIntVector& Resolution, const FVector& MinBound, const float VoxelSize)
+{
+	const FVector C = (T - MinBound) / VoxelSize;
+	return FIntVector(
+		FMath::Clamp(FMath::FloorToInt(C.X), 0, Resolution.X - 1),
+		FMath::Clamp(FMath::FloorToInt(C.Y), 0, Resolution.Y - 1),
+		FMath::Clamp(FMath::FloorToInt(C.Z), 0, Resolution.Z - 1));
+}
+
+FORCEINLINE bool IsValid(const FIntVector& P, const FIntVector& Resolution)
+{
+	return	0 <= P.X && P.X < Resolution.X &&
+			0 <= P.Y && P.Y < Resolution.Y &&
+			0 <= P.Z && P.Z < Resolution.Z;
+}
+
+FORCEINLINE FIntVector ClampToVolume(const FIntVector& CellCoord, const FIntVector& Resolution, bool& bIsValid)
+{
+	bIsValid = IsValid(CellCoord, Resolution);
+	return FIntVector(
+		FMath::Clamp(CellCoord.X, 0, Resolution.X - 1),
+		FMath::Clamp(CellCoord.Y, 0, Resolution.Y - 1),
+		FMath::Clamp(CellCoord.Z, 0, Resolution.Z - 1));
+}
+
+FORCEINLINE FIntVector ToCellCoord(const FVector3f& P, const FVector3f& MinBound, const FVector4f& MaxBound, const FIntVector& Resolution)
+{
+	bool bIsValid = false;
+	const FVector F = ((P - MinBound) / (MaxBound - MinBound));
+	const FIntVector CellCoord = FIntVector(FMath::FloorToInt(F.X * Resolution.X), FMath::FloorToInt(F.Y * Resolution.Y), FMath::FloorToInt(F.Z * Resolution.Z));
+	return ClampToVolume(CellCoord, Resolution, bIsValid);
+}
+
+FORCEINLINE uint32 ToIndex(const FIntVector& CellCoord, const FIntVector& Resolution)
+{
+	uint32 CellIndex = CellCoord.X + CellCoord.Y * Resolution.X + CellCoord.Z * Resolution.X * Resolution.Y;
+	return CellIndex;
+}
+
+struct FHairGrid
+{
+	struct FCurve
+	{
+		uint32 CurveIndex = 0;
+		float  Radius = 0;
+		uint8  GroupIndex = 0;
+	};
+
+	struct FVoxel
+	{
+		TArray<FCurve> VoxelCurves;
+	};
+
+	float VoxelSize;
+	FVector3f MinBound;
+	FVector3f MaxBound;
+	FIntVector Resolution;
+	TArray<FVoxel> Voxels;
+};
+
+static void Voxelize(const FHairDescriptionGroups& InGroups, FHairGrid& Out)
+{
+	// 1. Compute the overal bound of for all hair groups
+	Out.MinBound = InGroups.Bounds.GetBox().Min;
+	Out.MaxBound = InGroups.Bounds.GetBox().Max;
+	Out.Resolution = FIntVector::ZeroValue;
+
+	// 2. Based on the bound, determine the voxel grid resolution
+	Out.VoxelSize = 0.5f; // 5mm
+	{
+		const int32 MaxResolution = FMath::Max(GHairGroupIndexBuilder_MaxVoxelResolution, 2);
+
+		bool bIsValid = false;
+		while (!bIsValid)
+		{
+			FVector3f VoxelResolutionF = (Out.MaxBound - Out.MinBound) / Out.VoxelSize;
+			Out.Resolution = FIntVector(FMath::CeilToInt(VoxelResolutionF.X), FMath::CeilToInt(VoxelResolutionF.Y), FMath::CeilToInt(VoxelResolutionF.Z));
+			bIsValid = Out.Resolution.X <= MaxResolution && Out.Resolution.Y <= MaxResolution && Out.Resolution.Z <= MaxResolution;
+			if (!bIsValid)
+			{
+				Out.VoxelSize *= 2;
+			}
+		}
+		Out.MaxBound = Out.MinBound + FVector3f(Out.Resolution) * Out.VoxelSize;
+	}
+	Out.Voxels.SetNum(FMath::Max(Out.Resolution.X * Out.Resolution.Y * Out.Resolution.Z, 1));
+
+	// 3. Voxelization curves
+	uint32 GroupIndex = 0;
+	for (const FHairDescriptionGroup& Group : InGroups.HairGroups)
+	{
+		// Local copy of the hair strands data, and build the derived data so that curve offset are available
+		FHairStrandsDatas In = Group.Strands;
+		HairStrandsBuilder::BuildInternalData(In, false);
+
+		// Fill in voxel (TODO: make it parallel)
+		for (uint32 CurveIt = 0, CurveCount = In.GetNumCurves(); CurveIt < CurveCount; ++CurveIt)
+		{
+			const uint32 PointOffset = In.StrandsCurves.CurvesOffset[CurveIt];
+			const uint32 PointCount = In.StrandsCurves.CurvesCount[CurveIt];
+
+			uint32 PrevLinearCoord = ~0;
+			for (uint32 PointIndex = 0; PointIndex < PointCount - 1; ++PointIndex)
+			{
+				const uint32 Index0 = PointOffset + PointIndex;
+				const uint32 Index1 = PointOffset + PointIndex + 1;
+				const FVector3f& P0 = In.StrandsPoints.PointsPosition[Index0];
+				const FVector3f& P1 = In.StrandsPoints.PointsPosition[Index1];
+				const float R0 = In.StrandsPoints.PointsRadius[Index0];
+				const float R1 = In.StrandsPoints.PointsRadius[Index1];
+				const FVector3f Segment = P1 - P0;
+
+				// This is a coarse/non-conservative voxelization, by ray-marching segment, instead of walking intersected voxels
+				const float Length = Segment.Size();
+				const uint32 StepCount = FMath::CeilToInt(Length / Out.VoxelSize);
+				for (uint32 StepIt = 0; StepIt < StepCount + 1; ++StepIt)
+				{
+					const float T = float(StepIt) / float(StepCount);
+					const FVector3f P = P0 + Segment * T;
+					const float R = FMath::Lerp(R0, R1, T);
+					const FIntVector Coord = ToCoord(P, Out.Resolution, Out.MinBound, Out.VoxelSize);
+					const uint32 LinearCoord = ToLinearCoord(Coord, Out.Resolution);
+					if (LinearCoord != PrevLinearCoord)
+					{
+						FHairGrid::FCurve RCurve;
+						RCurve.CurveIndex = CurveIt;
+						RCurve.GroupIndex = GroupIndex;
+						RCurve.Radius = R;
+						Out.Voxels[LinearCoord].VoxelCurves.Add(RCurve);
+
+						PrevLinearCoord = LinearCoord;
+					}
+				}
+			}
+		}
+
+		++GroupIndex;
+	}
+}
+
+FORCEINLINE void SearchCell(const FHairStrandsVoxelData& In, const FIntVector& C, uint8& Out)
+{
+	const uint32 I = GroomBuilder_Voxelization::ToIndex(C, In.Resolution);
+	if (In.GroupIndices[I] != FHairStrandsVoxelData::InvalidGroupIndex)
+	{
+		Out = In.GroupIndices[I];
+	}
+}
+
+} // namespace GroomBuilder_Voxelization
+
+uint8 FHairStrandsVoxelData::GetGroupIndex(const FVector3f& P) const
+{
+	const uint32 MaxLookupDistance = FMath::Max3(Resolution.X, Resolution.Y, Resolution.Z) * 0.5f;
+	const FIntVector C = GroomBuilder_Voxelization::ToCellCoord(P, MinBound, MaxBound, Resolution);
+
+	uint8 GroupIndex = FHairStrandsVoxelData::InvalidGroupIndex;
+	for (int32 Offset = 1; Offset <= int32(MaxLookupDistance); ++Offset)
+	{
+		// Center
+		{
+			bool bIsValid = false;
+			const FIntVector CellCoord = GroomBuilder_Voxelization::ClampToVolume(C, Resolution, bIsValid);
+			if (bIsValid) GroomBuilder_Voxelization::SearchCell(*this, CellCoord, GroupIndex);
+		}
+
+		// Top & Bottom
+		for (int32 X = -Offset; X <= Offset; ++X)
+		for (int32 Y = -Offset; Y <= Offset; ++Y)
+		{
+			bool bIsValid0 = false, bIsValid1 = false;
+			const FIntVector CellCoord0 = GroomBuilder_Voxelization::ClampToVolume(C + FIntVector(X, Y, Offset), Resolution, bIsValid0);
+			const FIntVector CellCoord1 = GroomBuilder_Voxelization::ClampToVolume(C + FIntVector(X, Y, -Offset), Resolution, bIsValid1);
+			if (bIsValid0) GroomBuilder_Voxelization::SearchCell(*this, CellCoord0, GroupIndex);
+			if (bIsValid1) GroomBuilder_Voxelization::SearchCell(*this, CellCoord1, GroupIndex);
+		}
+
+		const int32 OffsetMinusOne = Offset - 1;
+		// Front & Back
+		for (int32 X = -Offset; X <= Offset; ++X)
+		for (int32 Z = -OffsetMinusOne; Z <= OffsetMinusOne; ++Z)
+		{
+			bool bIsValid0 = false, bIsValid1 = false;
+			const FIntVector CellCoord0 = GroomBuilder_Voxelization::ClampToVolume(C + FIntVector(X, Offset, Z), Resolution, bIsValid0);
+			const FIntVector CellCoord1 = GroomBuilder_Voxelization::ClampToVolume(C + FIntVector(X, -Offset, Z), Resolution, bIsValid1);
+			if (bIsValid0) GroomBuilder_Voxelization::SearchCell(*this, CellCoord0, GroupIndex);
+			if (bIsValid1) GroomBuilder_Voxelization::SearchCell(*this, CellCoord1, GroupIndex);
+		}
+
+		// Left & Right
+		for (int32 Y = -OffsetMinusOne; Y <= OffsetMinusOne; ++Y)
+		for (int32 Z = -OffsetMinusOne; Z <= OffsetMinusOne; ++Z)
+		{
+			bool bIsValid0 = false, bIsValid1 = false;
+			const FIntVector CellCoord0 = GroomBuilder_Voxelization::ClampToVolume(C + FIntVector(Offset, Y, Z), Resolution, bIsValid0);
+			const FIntVector CellCoord1 = GroomBuilder_Voxelization::ClampToVolume(C + FIntVector(-Offset, Y, Z), Resolution, bIsValid1);
+			if (bIsValid0) GroomBuilder_Voxelization::SearchCell(*this, CellCoord0, GroupIndex);
+			if (bIsValid1) GroomBuilder_Voxelization::SearchCell(*this, CellCoord1, GroupIndex);
+		}
+
+		// Early out if we have found a valid group index during a ring/layer step.
+		if (GroupIndex != InvalidGroupIndex)
+		{
+			break;
+		}
+	}
+
+	// No valid group index has been found. In this case force the group index to 0
+	check(GroupIndex != FHairStrandsVoxelData::InvalidGroupIndex);
+	if (GroupIndex == FHairStrandsVoxelData::InvalidGroupIndex) { GroupIndex = 0; } 
+
+	return GroupIndex;
+}
+
+void FGroomBuilder::VoxelizeGroupIndex(const FHairDescriptionGroups& In, FHairStrandsVoxelData& Out)
+{
+	GroomBuilder_Voxelization::FHairGrid Grid;
+	GroomBuilder_Voxelization::Voxelize(In, Grid);
+	Out.MinBound = Grid.MinBound;
+	Out.MaxBound = Grid.MaxBound;
+	Out.Resolution = Grid.Resolution;
+	Out.GroupIndices.SetNum(Grid.Voxels.Num());
+
+	TArray<uint8> GroupBins;
+	for (uint32 I=0,VoxelCount=Grid.Voxels.Num();I<VoxelCount;++I)
+	{
+		Out.GroupIndices[I] = FHairStrandsVoxelData::InvalidGroupIndex;
+
+		// 1. For a given voxel, build an histogram of all the groups falling into this voxel
+		GroupBins.Init(0, 64u);
+		for (const GroomBuilder_Voxelization::FHairGrid::FCurve& Curve : Grid.Voxels[I].VoxelCurves)
+		{
+			check(Curve.GroupIndex != FHairStrandsVoxelData::InvalidGroupIndex);
+			check(Curve.GroupIndex < GroupBins.Num());
+			GroupBins[Curve.GroupIndex]++;
+		}
+
+		// 2. Select the hair group having the largest count
+		if (Grid.Voxels[I].VoxelCurves.Num() > 1)
+		{
+			uint32 MaxBinCount = 0u;
+			for (uint8 GroupIndex = 0, GroupCount = GroupBins.Num(); GroupIndex < GroupCount; ++GroupIndex)
+			{
+				if (GroupBins[GroupIndex] > MaxBinCount)
+				{
+					Out.GroupIndices[I] = GroupIndex;
+					MaxBinCount = GroupBins[GroupIndex];
+				}
+			}
+		}
+		else if (Grid.Voxels[I].VoxelCurves.Num() == 1)
+		{
+			Out.GroupIndices[I] = Grid.Voxels[I].VoxelCurves[0].GroupIndex;
+		}
+	}
+}
 
 ///////////////////////////////////////////////////////////////////////////////////////////////////
 // Cluster culling data
diff --git a/Engine/Plugins/Runtime/HairStrands/Source/HairStrandsCore/Private/GroomComponent.cpp b/Engine/Plugins/Runtime/HairStrands/Source/HairStrandsCore/Private/GroomComponent.cpp
--- a/Engine/Plugins/Runtime/HairStrands/Source/HairStrandsCore/Private/GroomComponent.cpp
+++ b/Engine/Plugins/Runtime/HairStrands/Source/HairStrandsCore/Private/GroomComponent.cpp
@@ -348,7 +348,7 @@ public:
 				{
 					if (HairInstance->Cards.IsValid(LODIt))
 					{
-						if (HairInstance->Cards.LODs[LODIt].VertexFactory == nullptr) HairInstance->Cards.LODs[LODIt].VertexFactory = new FHairCardsVertexFactory(HairInstance, GroupIt, LODIt, EHairGeometryType::Cards, ShaderPlatform, FeatureLevel, "HairCardsVertexFactory");
+						if (HairInstance->Cards.LODs[LODIt].VertexFactory == nullptr) HairInstance->Cards.LODs[LODIt].VertexFactory = new FHairCardsVertexFactory(HairInstance, LODIt, EHairGeometryType::Cards, ShaderPlatform, FeatureLevel, "HairCardsVertexFactory");
 					}
 				}
 			}
@@ -360,7 +360,7 @@ public:
 				{
 					if (HairInstance->Meshes.IsValid(LODIt))
 					{
-						if (HairInstance->Meshes.LODs[LODIt].VertexFactory == nullptr) HairInstance->Meshes.LODs[LODIt].VertexFactory = new FHairCardsVertexFactory(HairInstance, GroupIt, LODIt, EHairGeometryType::Meshes, ShaderPlatform, FeatureLevel, "HairMeshesVertexFactory");
+						if (HairInstance->Meshes.LODs[LODIt].VertexFactory == nullptr) HairInstance->Meshes.LODs[LODIt].VertexFactory = new FHairCardsVertexFactory(HairInstance, LODIt, EHairGeometryType::Meshes, ShaderPlatform, FeatureLevel, "HairMeshesVertexFactory");
 					}
 				}
 			}
diff --git a/Engine/Plugins/Runtime/HairStrands/Source/HairStrandsCore/Private/GroomRBFDeformer.cpp b/Engine/Plugins/Runtime/HairStrands/Source/HairStrandsCore/Private/GroomRBFDeformer.cpp
--- a/Engine/Plugins/Runtime/HairStrands/Source/HairStrandsCore/Private/GroomRBFDeformer.cpp
+++ b/Engine/Plugins/Runtime/HairStrands/Source/HairStrandsCore/Private/GroomRBFDeformer.cpp
@@ -411,7 +411,7 @@ static void ApplyDeformationToGroom(const TArray<FRBFDeformedPositions>& Deforme
 				FGroomBuilder::BuildInterplationData(HairGroup.Info, StrandsData, GuidesData, InterpolationSettings, StrandsInterpolationData);
 				FGroomBuilder::BuildInterplationBulkData(GuidesData, StrandsInterpolationData, HairGroupsData.Strands.InterpolationBulkData);
 
-				FGroomBuilder::BuildClusterBulkData(StrandsData, HairDescriptionGroups.BoundRadius, HairGroupLOD, HairGroupsData.Strands.ClusterCullingBulkData);
+				FGroomBuilder::BuildClusterBulkData(StrandsData, HairDescriptionGroups.Bounds.SphereRadius, HairGroupLOD, HairGroupsData.Strands.ClusterCullingBulkData);
 			}
 		}
 	}
diff --git a/Engine/Plugins/Runtime/HairStrands/Source/HairStrandsCore/Private/GroomTextureBuilder.cpp b/Engine/Plugins/Runtime/HairStrands/Source/HairStrandsCore/Private/GroomTextureBuilder.cpp
--- a/Engine/Plugins/Runtime/HairStrands/Source/HairStrandsCore/Private/GroomTextureBuilder.cpp
+++ b/Engine/Plugins/Runtime/HairStrands/Source/HairStrandsCore/Private/GroomTextureBuilder.cpp
@@ -416,12 +416,14 @@ class FHairStrandTextureDilationCS : public FGlobalShader
 		SHADER_PARAMETER_RDG_TEXTURE(Texture2D, Source_CoverageTexture)
 		SHADER_PARAMETER_RDG_TEXTURE(Texture2D, Source_TangentTexture)
 		SHADER_PARAMETER_RDG_TEXTURE(Texture2D, Source_AttributeTexture)
+		SHADER_PARAMETER_RDG_TEXTURE(Texture2D, Source_GroupIndexTexture)
 
 		SHADER_PARAMETER_RDG_TEXTURE_UAV(RWTexture2D, TriangleMaskTexture)
 		SHADER_PARAMETER_RDG_TEXTURE_UAV(RWTexture2D, Target_DepthTexture)
 		SHADER_PARAMETER_RDG_TEXTURE_UAV(RWTexture2D, Target_CoverageTexture)
 		SHADER_PARAMETER_RDG_TEXTURE_UAV(RWTexture2D, Target_TangentTexture)
 		SHADER_PARAMETER_RDG_TEXTURE_UAV(RWTexture2D, Target_AttributeTexture)
+		SHADER_PARAMETER_RDG_TEXTURE_UAV(RWTexture2D, Target_GroupIndexTexture)
 	END_SHADER_PARAMETER_STRUCT()
 
 public:
@@ -445,11 +447,13 @@ static void AddTextureDilationPass(
 	FRDGTextureRef Source_CoverageTexture,
 	FRDGTextureRef Source_TangentTexture,
 	FRDGTextureRef Source_AttributeTexture,
+	FRDGTextureRef Source_GroupIndexTexture,
 
 	FRDGTextureRef Target_DepthTexture,
 	FRDGTextureRef Target_CoverageTexture,
 	FRDGTextureRef Target_TangentTexture,
-	FRDGTextureRef Target_AttributeTexture)
+	FRDGTextureRef Target_AttributeTexture,
+	FRDGTextureRef Target_GroupIndexTexture)
 {
 	FHairStrandTextureDilationCS::FParameters* Parameters = GraphBuilder.AllocParameters<FHairStrandTextureDilationCS::FParameters>();
 	Parameters->Resolution				= Resolution;
@@ -458,10 +462,12 @@ static void AddTextureDilationPass(
 	Parameters->Source_CoverageTexture	= Source_CoverageTexture;
 	Parameters->Source_TangentTexture	= Source_TangentTexture;
 	Parameters->Source_AttributeTexture = Source_AttributeTexture;
+	Parameters->Source_GroupIndexTexture= Source_GroupIndexTexture;
 	Parameters->Target_DepthTexture		= GraphBuilder.CreateUAV(Target_DepthTexture);
 	Parameters->Target_CoverageTexture	= GraphBuilder.CreateUAV(Target_CoverageTexture);
 	Parameters->Target_TangentTexture	= GraphBuilder.CreateUAV(Target_TangentTexture);
 	Parameters->Target_AttributeTexture = GraphBuilder.CreateUAV(Target_AttributeTexture);
+	Parameters->Target_GroupIndexTexture= GraphBuilder.CreateUAV(Target_GroupIndexTexture);
 
 	TShaderMapRef<FHairStrandTextureDilationCS> ComputeShader(ShaderMap);
 	FIntVector DispatchCount = FComputeShaderUtils::GetGroupCount(Resolution, FIntPoint(8, 4));
@@ -512,6 +518,11 @@ static void InternalAllocateStrandsTexture_Attribute(UTexture2D* Out, const FInt
 	InternalAllocateStrandsTexture(Out, Resolution, 1, PF_B8G8R8A8, ETextureSourceFormat::TSF_BGRA8);
 }
 
+static void InternalAllocateStrandsTexture_GroupIndex(UTexture2D* Out, const FIntPoint& Resolution, uint32 MipCount)
+{
+	InternalAllocateStrandsTexture(Out, Resolution, 1, PF_B8G8R8A8, ETextureSourceFormat::TSF_BGRA8);
+}
+
 #if WITH_EDITOR
 FStrandsTexturesOutput FGroomTextureBuilder::CreateGroomStrandsTexturesTexture(const UGroomAsset* GroomAsset, uint32 InResolution)
 {
@@ -528,6 +539,7 @@ FStrandsTexturesOutput FGroomTextureBuilder::CreateGroomStrandsTexturesTexture(c
 	Output.Coverage = FHairStrandsCore::CreateTexture(PackageName, Resolution, TEXT("_Opacity"), InternalAllocateStrandsTexture_Coverage);
 	Output.Tangent = FHairStrandsCore::CreateTexture(PackageName, Resolution, TEXT("_Tangent"), InternalAllocateStrandsTexture_Tangent);
 	Output.Attribute = FHairStrandsCore::CreateTexture(PackageName, Resolution, TEXT("_Attribute"), InternalAllocateStrandsTexture_Attribute);
+	Output.GroupIndex = FHairStrandsCore::CreateTexture(PackageName, Resolution, TEXT("_GroupIndex"), InternalAllocateStrandsTexture_GroupIndex);
 	return Output;
 }
 #endif
@@ -578,6 +590,7 @@ class FHairStrandsTexturePS : public FGlobalShader
 		SHADER_PARAMETER_SRV(Buffer, InVF_PositionBuffer)
 		SHADER_PARAMETER_SRV(Buffer, InVF_Attribute0Buffer)
 		SHADER_PARAMETER(uint32, InVF_ControlPointCount)
+		SHADER_PARAMETER(uint32, InVF_GroupIndex)
 
 		SHADER_PARAMETER(FVector3f, Voxel_MinBound)
 		SHADER_PARAMETER(FVector3f, Voxel_MaxBound)
@@ -640,12 +653,14 @@ static void InternalGenerateHairStrandsTextures(
 	float InHairStrands_Radius,
 	float InHairStrands_Length,
 	uint32 InHairStrands_ControlPointCount,
+	uint32 InHairStrands_GroupIndex,
 
 	FRDGTextureRef OutDepthTestTexture,
 	FRDGTextureRef OutDepthTexture,
 	FRDGTextureRef OutTangentTexture,
 	FRDGTextureRef OutCoverageTexture,
 	FRDGTextureRef OutRootUVStrandsUSeedTexture,
+	FRDGTextureRef OutGroupIndexTexture,
 	FRDGTextureRef OutTriangleMask)
 {
 	const FIntPoint OutputResolution = OutDepthTexture->Desc.Extent;
@@ -668,6 +683,7 @@ static void InternalGenerateHairStrandsTextures(
 	ParametersPS->InVF_Radius = InHairStrands_Radius;
 	ParametersPS->InVF_Length = InHairStrands_Length;
 	ParametersPS->InVF_ControlPointCount = InHairStrands_ControlPointCount;
+	ParametersPS->InVF_GroupIndex = InHairStrands_GroupIndex;
 
 	ParametersPS->Voxel_MinBound = VoxelMinBound;
 	ParametersPS->Voxel_MaxBound = VoxelMaxBound;
@@ -686,7 +702,8 @@ static void InternalGenerateHairStrandsTextures(
 	ParametersPS->RenderTargets[1] = FRenderTargetBinding(OutTangentTexture, bClear ? ERenderTargetLoadAction::EClear : ERenderTargetLoadAction::ELoad);
 	ParametersPS->RenderTargets[2] = FRenderTargetBinding(OutCoverageTexture, bClear ? ERenderTargetLoadAction::EClear : ERenderTargetLoadAction::ELoad);
 	ParametersPS->RenderTargets[3] = FRenderTargetBinding(OutRootUVStrandsUSeedTexture, bClear ? ERenderTargetLoadAction::EClear : ERenderTargetLoadAction::ELoad);
-	ParametersPS->RenderTargets[4] = FRenderTargetBinding(OutTriangleMask, bClear ? ERenderTargetLoadAction::EClear : ERenderTargetLoadAction::ELoad);
+	ParametersPS->RenderTargets[4] = FRenderTargetBinding(OutGroupIndexTexture, bClear ? ERenderTargetLoadAction::EClear : ERenderTargetLoadAction::ELoad);
+	ParametersPS->RenderTargets[5] = FRenderTargetBinding(OutTriangleMask, bClear ? ERenderTargetLoadAction::EClear : ERenderTargetLoadAction::ELoad);
 	ParametersPS->RenderTargets.DepthStencil = FDepthStencilBinding(
 		OutDepthTestTexture,
 		bClear ? ERenderTargetLoadAction::EClear : ERenderTargetLoadAction::ELoad,
@@ -717,6 +734,7 @@ static void InternalGenerateHairStrandsTextures(
 				CW_RGBA, BO_Add, BF_One, BF_Zero, BO_Add, BF_One, BF_Zero,
 				CW_RGBA, BO_Add, BF_One, BF_Zero, BO_Add, BF_One, BF_Zero,
 				CW_RGBA, BO_Add, BF_One, BF_Zero, BO_Add, BF_One, BF_Zero,
+				CW_RGBA, BO_Add, BF_One, BF_Zero, BO_Add, BF_One, BF_Zero,
 				CW_RGBA, BO_Add, BF_One, BF_Zero, BO_Add, BF_One, BF_Zero> ::GetRHI();
 
 			GraphicsPSOInit.RasterizerState = TStaticRasterizerState<>::GetRHI();
@@ -873,6 +891,12 @@ static void InternalBuildStrandsTextures_GPU(
 	AttributeTexture[0] = GraphBuilder.CreateTexture(Desc, TEXT("Hair.AttributeTexture"));
 	AttributeTexture[1] = GraphBuilder.CreateTexture(Desc, TEXT("Hair.AttributeTexture"));
 
+	Desc.Format = PF_B8G8R8A8;
+	Desc.ClearValue = FClearValueBinding::Transparent;
+	FRDGTextureRef GroupIndexTexture[2];
+	GroupIndexTexture[0] = GraphBuilder.CreateTexture(Desc, TEXT("Hair.GroupIndexTexture"));
+	GroupIndexTexture[1] = GraphBuilder.CreateTexture(Desc, TEXT("Hair.GroupIndexTexture"));
+
 	Desc.Format = PF_R32_UINT;
 	Desc.ClearValue = FClearValueBinding::Black;
 	FRDGTextureRef TriangleMaskTexture = GraphBuilder.CreateTexture(Desc, TEXT("Hair.TriangleMaskTexture"));
@@ -994,12 +1018,14 @@ static void InternalBuildStrandsTextures_GPU(
 				RenderingData.GeometrySettings.HairWidth * 0.5f,
 				GroupData.Strands.RestResource->BulkData.MaxLength,
 				GroupData.Strands.RestResource->BulkData.PointCount,
+				GroupIndex,
 
 				DepthTestTexture,
 				DepthTexture[0],
 				TangentTexture[0],
 				CoverageTexture[0],
 				AttributeTexture[0],
+				GroupIndexTexture[0],
 				TriangleMaskTexture);
 
 			bClear = false;
@@ -1023,11 +1049,13 @@ static void InternalBuildStrandsTextures_GPU(
 			CoverageTexture[SourceIndex],
 			TangentTexture[SourceIndex],
 			AttributeTexture[SourceIndex],
+			GroupIndexTexture[SourceIndex],
 
 			DepthTexture[TargetIndex],
 			CoverageTexture[TargetIndex],
 			TangentTexture[TargetIndex],
-			AttributeTexture[TargetIndex]);
+			AttributeTexture[TargetIndex],
+			GroupIndexTexture[TargetIndex]);
 
 		SourceIndex = TargetIndex;
 	}
@@ -1036,6 +1064,7 @@ static void InternalBuildStrandsTextures_GPU(
 	AddReadbackPass(GraphBuilder, 4, CoverageTexture[TargetIndex]	, Output.Coverage);
 	AddReadbackPass(GraphBuilder, 4, TangentTexture[TargetIndex]	, Output.Tangent);
 	AddReadbackPass(GraphBuilder, 4, AttributeTexture[TargetIndex]	, Output.Attribute);
+	AddReadbackPass(GraphBuilder, 4, GroupIndexTexture[TargetIndex]	, Output.GroupIndex);
 }
 
 ///////////////////////////////////////////////////////////////////////////////////////////////////
diff --git a/Engine/Plugins/Runtime/HairStrands/Source/HairStrandsCore/Private/HairCardsBuilder.cpp b/Engine/Plugins/Runtime/HairStrands/Source/HairStrandsCore/Private/HairCardsBuilder.cpp
--- a/Engine/Plugins/Runtime/HairStrands/Source/HairStrandsCore/Private/HairCardsBuilder.cpp
+++ b/Engine/Plugins/Runtime/HairStrands/Source/HairStrandsCore/Private/HairCardsBuilder.cpp
@@ -4,6 +4,7 @@
 #include "HairStrandsCore.h"
 #include "HairStrandsDatas.h"
 #include "HairCardsDatas.h"
+#include "GroomBuilder.h"
 
 #include "Math/Box.h"
 #include "RenderGraphBuilder.h"
@@ -3327,7 +3328,7 @@ namespace FHairCardsBuilder
 FString GetVersion()
 {
 	// Important to update the version when cards building or importing changes
-	return TEXT("9b5");
+	return TEXT("9c");
 }
 
 void AllocateAtlasTexture(UTexture2D* Out, const FIntPoint& Resolution, uint32 MipCount, EPixelFormat PixelFormat, ETextureSourceFormat SourceFormat)
@@ -3574,7 +3575,8 @@ void SanitizeMeshDescription(FMeshDescription* MeshDescription)
 
 bool InternalImportGeometry(
 	const UStaticMesh* StaticMesh,
-	const FHairStrandsDatas& InStrandsData, // Used for extracting & assigning root UV to cards data
+	const FHairStrandsDatas& InStrandsData,			// Used for extracting & assigning root UV to cards data
+	const FHairStrandsVoxelData& InStrandsVoxelData,// Used for transfering & assigning group index to cards data
 	FHairCardsDatas& Out,
 	FHairCardsBulkData& OutBulk,
 	FHairStrandsDatas& OutGuides,
@@ -3722,8 +3724,10 @@ bool InternalImportGeometry(
 	OutBulk.CoverageTexture = nullptr;
 	OutBulk.AttributeTexture = nullptr;
 
+	const uint32 CardsCount = Out.Cards.IndexOffsets.Num();
+
 	TArray<float> CardLengths;
-	CardLengths.Reserve(Out.Cards.IndexOffsets.Num());
+	CardLengths.Reserve(CardsCount);
 	bool bSuccess = HairCards::CreateCardsGuides(Out.Cards, OutGuides, CardLengths);
 	if (bSuccess)
 	{
@@ -3743,8 +3747,49 @@ bool InternalImportGeometry(
 		}
 	}
 
+	// Used voxelized hair group (from hair strands) to assign hair group index to card vertices
+	TArray<uint32> CardGroupIndices;
+	CardGroupIndices.Init(0u, CardsCount);
+	if (InStrandsVoxelData.IsValid())
+	{
+		for (uint32 CardIt = 0; CardIt < CardsCount; ++CardIt)
+		{
+			// 1. For each cards' vertex, query the closest group index, and build an histogram 
+			//    of the group index covering the card
+			TArray<uint8> GroupBins;
+			GroupBins.Init(0, 64u);
+			const uint32 CardsIndexCount = Out.Cards.IndexCounts[CardIt];
+			const uint32 CardsIndexOffset = Out.Cards.IndexOffsets[CardIt];
+
+			for (uint32 IndexIt = 0; IndexIt < CardsIndexCount; ++IndexIt)
+			{
+				const uint32 VertexIndex = Out.Cards.Indices[CardsIndexOffset + IndexIt];
+				const FVector3f& P = Out.Cards.Positions[VertexIndex];
+
+				const uint8 GroupIndex = InStrandsVoxelData.GetGroupIndex(P);
+				if (GroupIndex != FHairStrandsVoxelData::InvalidGroupIndex)
+				{
+					check(GroupIndex < GroupBins.Num());
+					GroupBins[GroupIndex]++;
+				}
+			}
+
+			// 2. Since a cards can covers several group, select the group index covering the larger 
+			//    number of cards' vertices
+			uint32 MaxBinCount = 0u;
+			for (uint8 GroupIndex=0,GroupCount=GroupBins.Num(); GroupIndex<GroupCount; ++GroupIndex)
+			{
+				if (GroupBins[GroupIndex] > MaxBinCount)
+				{
+					CardGroupIndices[CardIt] = GroupIndex;
+					MaxBinCount = GroupBins[GroupIndex];
+				}
+			}
+		}
+	}
+
 	// Patch Cards Position to store cards length into the W component
-	for (uint32 CardIt = 0, CardsCount = Out.Cards.IndexOffsets.Num(); CardIt < CardsCount; ++CardIt)
+	for (uint32 CardIt = 0; CardIt < CardsCount; ++CardIt)
 	{
 		const uint32 CardsIndexCount = Out.Cards.IndexCounts[CardIt];
 		const uint32 CardsIndexOffset = Out.Cards.IndexOffsets[CardIt];
@@ -3759,7 +3804,11 @@ bool InternalImportGeometry(
 			// Instead of storing the interoplated card length, store the actual max length of the card, 
 			// as reconstructing the strands length, based on interpolatd CardLength will be too prone to numerical issue.
 			// This means that the strand length retrieves in shader will be an over estimate of the actual length
-			OutBulk.Positions[VertexIndex].W = CardLength; // InterpolatedCardLength;
+			const FFloat16 hCardLength = CardLength;  // InterpolatedCardLength;			
+
+			// Encode cards length & group index into the .W component of position
+			const uint32 EncodedW = hCardLength.Encoded | (CardGroupIndices[CardIt] << 16u);
+			OutBulk.Positions[VertexIndex].W = *(float*)&EncodedW;
 		}
 	}
 
@@ -3795,7 +3844,6 @@ bool InternalImportGeometry(
 		};
 		TArray<FCardsRootData> CardsRoots;
 		{
-			const uint32 CardsCount = Out.Cards.IndexOffsets.Num();
 			CardsRoots.Reserve(CardsCount);
 			for (uint32 CardIt = 0; CardIt < CardsCount; ++CardIt)
 			{
@@ -3865,6 +3913,7 @@ bool InternalImportGeometry(
 bool ImportGeometry(
 	const UStaticMesh* StaticMesh,
 	const FHairStrandsDatas& InStrandsData,
+	const FHairStrandsVoxelData& InStrandsVoxelData,
 	FHairCardsBulkData& OutBulk,
 	FHairStrandsDatas& OutGuides,
 	FHairCardsInterpolationBulkData& OutInterpolationBulkData)
@@ -3873,6 +3922,7 @@ bool ImportGeometry(
 	return InternalImportGeometry(
 		StaticMesh,
 		InStrandsData,
+		InStrandsVoxelData,
 		CardData,
 		OutBulk,
 		OutGuides,
@@ -3881,12 +3931,14 @@ bool ImportGeometry(
 
 bool ExtractCardsData(const UStaticMesh* StaticMesh, const FHairStrandsDatas& InStrandsData, FHairCardsDatas& Out)
 {
+	FHairStrandsVoxelData DummyStrandsVoxelData;
 	FHairCardsBulkData OutBulk;
 	FHairStrandsDatas OutGuides;
 	FHairCardsInterpolationBulkData OutInterpolationBulkData;
 	return InternalImportGeometry(
 		StaticMesh,
 		InStrandsData,
+		DummyStrandsVoxelData,
 		Out,
 		OutBulk,
 		OutGuides,
diff --git a/Engine/Plugins/Runtime/HairStrands/Source/HairStrandsCore/Private/HairCardsVertexFactory.cpp b/Engine/Plugins/Runtime/HairStrands/Source/HairStrandsCore/Private/HairCardsVertexFactory.cpp
--- a/Engine/Plugins/Runtime/HairStrands/Source/HairStrandsCore/Private/HairCardsVertexFactory.cpp
+++ b/Engine/Plugins/Runtime/HairStrands/Source/HairStrandsCore/Private/HairCardsVertexFactory.cpp
@@ -34,9 +34,9 @@ FHairCardsUniformBuffer CreateHairCardsVFUniformBuffer(
 		const FHairGroupInstance::FCards::FLOD& LOD = Instance->Cards.LODs[LODIndex];
 
 		// Cards atlas UV are inverted so fetching needs to be inverted on the y-axis
-		UniformParameters.GroupIndex = Instance->Debug.GroupIndex;
 		UniformParameters.bInvertUV = LOD.RestResource->bInvertUV;
 		UniformParameters.bUseTextureRootUV = 0;
+		UniformParameters.bUseTextureGroupIndex = 0;
 		UniformParameters.MaxVertexCount = LOD.RestResource->GetVertexCount();
 
 		// When the geometry is not-dynamic (no binding to skeletal mesh, no simulation), only a single vertex buffer is allocated. 
@@ -67,6 +67,8 @@ FHairCardsUniformBuffer CreateHairCardsVFUniformBuffer(
 		UniformParameters.AttributeSampler = LOD.RestResource->AttributeSampler;
 		UniformParameters.AuxilaryDataTexture = LOD.RestResource->AuxilaryDataTexture;
 		UniformParameters.AuxilaryDataSampler = LOD.RestResource->AuxilaryDataSampler;
+		UniformParameters.GroupIndexTexture = nullptr; // Group index are stored on vertices for cards
+		UniformParameters.GroupIndexSampler = nullptr; // Group index are stored on vertices for cards
 	}
 	else if (GeometryType == EHairGeometryType::Meshes)
 	{
@@ -87,9 +89,9 @@ FHairCardsUniformBuffer CreateHairCardsVFUniformBuffer(
 		}
 
 		// Meshes UV are not inverted so no need to invert the y-axis
-		UniformParameters.GroupIndex = Instance->Debug.GroupIndex;
 		UniformParameters.bInvertUV = 0;
 		UniformParameters.bUseTextureRootUV = 1;
+		UniformParameters.bUseTextureGroupIndex = 1;
 		UniformParameters.MaxVertexCount = LOD.RestResource->GetVertexCount();
 		UniformParameters.NormalsBuffer = LOD.RestResource->NormalsBuffer.ShaderResourceViewRHI.GetReference();
 		UniformParameters.UVsBuffer = LOD.RestResource->UVsBuffer.ShaderResourceViewRHI.GetReference();
@@ -104,6 +106,8 @@ FHairCardsUniformBuffer CreateHairCardsVFUniformBuffer(
 		UniformParameters.AttributeSampler = LOD.RestResource->AttributeSampler;
 		UniformParameters.AuxilaryDataTexture = LOD.RestResource->AuxilaryDataTexture;
 		UniformParameters.AuxilaryDataSampler = LOD.RestResource->AuxilaryDataSampler;
+		UniformParameters.GroupIndexTexture = LOD.RestResource->GroupIndexTexture;
+		UniformParameters.GroupIndexSampler = LOD.RestResource->GroupIndexSampler;
 	}
 
 	if (!bSupportsManualVertexFetch)
@@ -122,12 +126,14 @@ FHairCardsUniformBuffer CreateHairCardsVFUniformBuffer(
 	if (!UniformParameters.CoverageTexture)		{ UniformParameters.CoverageTexture = DefaultTexture; }
 	if (!UniformParameters.AttributeTexture)	{ UniformParameters.AttributeTexture = DefaultTexture;}
 	if (!UniformParameters.AuxilaryDataTexture)	{ UniformParameters.AuxilaryDataTexture = DefaultTexture; }
+	if (!UniformParameters.GroupIndexTexture)	{ UniformParameters.GroupIndexTexture = DefaultTexture; }
 
 	if (!UniformParameters.DepthSampler)		{ UniformParameters.DepthSampler = DefaultSampler;	  }
 	if (!UniformParameters.TangentSampler)		{ UniformParameters.TangentSampler = DefaultSampler;  }
 	if (!UniformParameters.CoverageSampler)		{ UniformParameters.CoverageSampler = DefaultSampler; }
 	if (!UniformParameters.AttributeSampler)	{ UniformParameters.AttributeSampler = DefaultSampler;}
-	if (!UniformParameters.AuxilaryDataSampler)	{ UniformParameters.AuxilaryDataSampler = DefaultSampler;}
+	if (!UniformParameters.AuxilaryDataSampler) { UniformParameters.AuxilaryDataSampler = DefaultSampler; }
+	if (!UniformParameters.GroupIndexSampler)	{ UniformParameters.GroupIndexSampler = DefaultSampler;}
 
 	return TUniformBufferRef<FHairCardsVertexFactoryUniformShaderParameters>::CreateUniformBufferImmediate(UniformParameters, UniformBuffer_MultiFrame);
 }
@@ -158,8 +164,6 @@ public:
 		const FHairCardsVertexFactory* VF = static_cast<const FHairCardsVertexFactory*>(VertexFactory);
 		check(VF);
 
-
-		const int32 GroupIndex = VF->Data.GroupIndex;
 		const int32 LODIndex = VF->Data.LODIndex;
 
 		const FHairGroupInstance* Instance = VF->Data.Instance;
@@ -190,14 +194,13 @@ public:
 
 IMPLEMENT_TYPE_LAYOUT(FHairCardsVertexFactoryShaderParameters);
 
-FHairCardsVertexFactory::FHairCardsVertexFactory(FHairGroupInstance* Instance, uint32 GroupIndex, uint32 LODIndex, EHairGeometryType GeometryType, EShaderPlatform InShaderPlatform, ERHIFeatureLevel::Type InFeatureLevel, const char* InDebugName)
+FHairCardsVertexFactory::FHairCardsVertexFactory(FHairGroupInstance* Instance, uint32 LODIndex, EHairGeometryType GeometryType, EShaderPlatform InShaderPlatform, ERHIFeatureLevel::Type InFeatureLevel, const char* InDebugName)
 	: FVertexFactory(InFeatureLevel)
 	, DebugName(InDebugName)
 {
 	bSupportsManualVertexFetch = RHISupportsManualVertexFetch(InShaderPlatform);
 
 	Data.Instance = Instance;
-	Data.GroupIndex = GroupIndex;
 	Data.LODIndex = LODIndex;
 	Data.GeometryType = GeometryType;
 }
diff --git a/Engine/Plugins/Runtime/HairStrands/Source/HairStrandsCore/Public/GroomAsset.h b/Engine/Plugins/Runtime/HairStrands/Source/HairStrandsCore/Public/GroomAsset.h
--- a/Engine/Plugins/Runtime/HairStrands/Source/HairStrandsCore/Public/GroomAsset.h
+++ b/Engine/Plugins/Runtime/HairStrands/Source/HairStrandsCore/Public/GroomAsset.h
@@ -375,7 +375,7 @@ struct FHairDescriptionGroup
 struct FHairDescriptionGroups
 {
 	TArray<FHairDescriptionGroup> HairGroups;
-	float BoundRadius = 0;
+	FBoxSphereBounds3f Bounds;
 	bool  IsValid() const;
 };
 
@@ -396,7 +396,8 @@ enum class EHairAtlasTextureType : uint8
 	Tangent,
 	Attribute,
 	Coverage,
-	AuxilaryData
+	AuxilaryData,
+	GroupIndex
 };
 
 /**
diff --git a/Engine/Plugins/Runtime/HairStrands/Source/HairStrandsCore/Public/GroomAssetCards.h b/Engine/Plugins/Runtime/HairStrands/Source/HairStrandsCore/Public/GroomAssetCards.h
--- a/Engine/Plugins/Runtime/HairStrands/Source/HairStrandsCore/Public/GroomAssetCards.h
+++ b/Engine/Plugins/Runtime/HairStrands/Source/HairStrandsCore/Public/GroomAssetCards.h
@@ -180,9 +180,13 @@ struct HAIRSTRANDSCORE_API FHairGroupCardsTextures
 	UPROPERTY(EditAnywhere, Category = "CardsAttributes")
 	TObjectPtr<UTexture2D> AttributeTexture = nullptr;
 
-	UPROPERTY(EditAnywhere, Category = "CardsAuxilaryData")
+	UPROPERTY(EditAnywhere, Category = "CardsAttributes")
 	TObjectPtr<UTexture2D> AuxilaryDataTexture = nullptr;
 
+	/* Store hair group indices. This texture is only used by hair meshes and can be generated through the hair-textures option on Groom Asset. */
+	UPROPERTY(EditAnywhere, Category = "CardsAttributes")
+	TObjectPtr<UTexture2D> GroupIndexTexture = nullptr;
+
 	void SetTexture(EHairAtlasTextureType SlotID, UTexture2D* Texture);
 
 	bool bNeedToBeSaved = false;
diff --git a/Engine/Plugins/Runtime/HairStrands/Source/HairStrandsCore/Public/GroomBuilder.h b/Engine/Plugins/Runtime/HairStrands/Source/HairStrandsCore/Public/GroomBuilder.h
--- a/Engine/Plugins/Runtime/HairStrands/Source/HairStrandsCore/Public/GroomBuilder.h
+++ b/Engine/Plugins/Runtime/HairStrands/Source/HairStrandsCore/Public/GroomBuilder.h
@@ -6,6 +6,7 @@
 #include "UObject/Object.h"
 #include "UObject/ObjectMacros.h"
 #include "GroomSettings.h"
+#include "Math/IntVector.h"
 
 struct FHairStrandsDatas;
 struct FHairStrandsClusterCullingData;
@@ -23,6 +24,20 @@ struct FHairGroupInfo;
 class FHairDescription;
 class UGroomAsset;
 
+struct FHairStrandsVoxelData
+{
+	static const uint8 InvalidGroupIndex = 0xFF;
+
+	FVector3f MinBound;
+	FVector3f MaxBound;
+	FIntVector Resolution;
+	TArray<uint8> GroupIndices;
+
+	// Return the (closest) group index corresponding to position P
+	uint8 GetGroupIndex(const FVector3f& P) const;
+	bool IsValid() const { return GroupIndices.Num() > 0; }
+};
+
 // Data flow overview
 // ==================
 // HairDescription -> HairDescriptionGroups -> HairStrandsData -> HairStrandsBulkData*
@@ -79,4 +94,9 @@ struct HAIRSTRANDSCORE_API FGroomBuilder
 		const float InGroomAssetRadius,
 		const FHairGroupsLOD& InSettings,
 		FHairStrandsClusterCullingBulkData& OutClusterCullingData);
+
+	// Optional: Voxelize hair group index
+	static void VoxelizeGroupIndex(
+		const FHairDescriptionGroups& In,
+		FHairStrandsVoxelData& Out);
 };
\ No newline at end of file
diff --git a/Engine/Plugins/Runtime/HairStrands/Source/HairStrandsCore/Public/GroomInstance.h b/Engine/Plugins/Runtime/HairStrands/Source/HairStrandsCore/Public/GroomInstance.h
--- a/Engine/Plugins/Runtime/HairStrands/Source/HairStrandsCore/Public/GroomInstance.h
+++ b/Engine/Plugins/Runtime/HairStrands/Source/HairStrandsCore/Public/GroomInstance.h
@@ -21,8 +21,8 @@ enum class EGroomCacheType : uint8;
 BEGIN_GLOBAL_SHADER_PARAMETER_STRUCT(FHairCardsVertexFactoryUniformShaderParameters, HAIRSTRANDSCORE_API)
 	SHADER_PARAMETER(uint32, bInvertUV)
 	SHADER_PARAMETER(uint32, bUseTextureRootUV)
+	SHADER_PARAMETER(uint32, bUseTextureGroupIndex)
 	SHADER_PARAMETER(uint32, MaxVertexCount)
-	SHADER_PARAMETER(float, GroupIndex)
 	SHADER_PARAMETER_SRV(Buffer<float4>, PositionBuffer)
 	SHADER_PARAMETER_SRV(Buffer<float4>, PreviousPositionBuffer)
 	SHADER_PARAMETER_SRV(Buffer<float4>, NormalsBuffer)
@@ -38,6 +38,8 @@ BEGIN_GLOBAL_SHADER_PARAMETER_STRUCT(FHairCardsVertexFactoryUniformShaderParamet
 	SHADER_PARAMETER_SAMPLER(SamplerState, AttributeSampler)
 	SHADER_PARAMETER_TEXTURE(Texture2D<float4>, AuxilaryDataTexture)
 	SHADER_PARAMETER_SAMPLER(SamplerState, AuxilaryDataSampler)
+	SHADER_PARAMETER_TEXTURE(Texture2D<float4>, GroupIndexTexture)
+	SHADER_PARAMETER_SAMPLER(SamplerState, GroupIndexSampler)
 END_GLOBAL_SHADER_PARAMETER_STRUCT()
 
 typedef TUniformBufferRef<FHairCardsVertexFactoryUniformShaderParameters> FHairCardsUniformBuffer;
diff --git a/Engine/Plugins/Runtime/HairStrands/Source/HairStrandsCore/Public/GroomResources.h b/Engine/Plugins/Runtime/HairStrands/Source/HairStrandsCore/Public/GroomResources.h
--- a/Engine/Plugins/Runtime/HairStrands/Source/HairStrandsCore/Public/GroomResources.h
+++ b/Engine/Plugins/Runtime/HairStrands/Source/HairStrandsCore/Public/GroomResources.h
@@ -536,17 +536,19 @@ struct FHairCardsRestResource : public FHairCommonResource
 	FHairCardsVertexBuffer NormalsBuffer;
 	FHairCardsVertexBuffer UVsBuffer;
 
-	FSamplerStateRHIRef DepthSampler;
-	FSamplerStateRHIRef TangentSampler;
-	FSamplerStateRHIRef CoverageSampler;
-	FSamplerStateRHIRef AttributeSampler;
-	FSamplerStateRHIRef AuxilaryDataSampler;
+	FSamplerStateRHIRef DepthSampler = nullptr;
+	FSamplerStateRHIRef TangentSampler = nullptr;
+	FSamplerStateRHIRef CoverageSampler = nullptr;
+	FSamplerStateRHIRef AttributeSampler = nullptr;
+	FSamplerStateRHIRef AuxilaryDataSampler = nullptr;
+	FSamplerStateRHIRef GroupIndexSampler = nullptr;
 
 	FTextureReferenceRHIRef	DepthTexture = nullptr;
 	FTextureReferenceRHIRef	CoverageTexture = nullptr;
 	FTextureReferenceRHIRef	TangentTexture = nullptr;
 	FTextureReferenceRHIRef	AttributeTexture = nullptr;
 	FTextureReferenceRHIRef	AuxilaryDataTexture = nullptr;
+	FTextureReferenceRHIRef	GroupIndexTexture = nullptr;
 
 	/* Reference to the hair strands render data */
 	const FHairCardsBulkData& BulkData;
@@ -731,17 +733,19 @@ struct FHairMeshesRestResource : public FHairCommonResource
 	FHairCardsVertexBuffer NormalsBuffer;
 	FHairCardsVertexBuffer UVsBuffer;
 
-	FSamplerStateRHIRef DepthSampler;
-	FSamplerStateRHIRef TangentSampler;
-	FSamplerStateRHIRef CoverageSampler;
-	FSamplerStateRHIRef AttributeSampler;
-	FSamplerStateRHIRef AuxilaryDataSampler;
+	FSamplerStateRHIRef DepthSampler = nullptr;
+	FSamplerStateRHIRef TangentSampler = nullptr;
+	FSamplerStateRHIRef CoverageSampler = nullptr;
+	FSamplerStateRHIRef AttributeSampler = nullptr;
+	FSamplerStateRHIRef AuxilaryDataSampler = nullptr;
+	FSamplerStateRHIRef GroupIndexSampler = nullptr;
 
 	FTextureReferenceRHIRef	DepthTexture = nullptr;
 	FTextureReferenceRHIRef	CoverageTexture = nullptr;
 	FTextureReferenceRHIRef	TangentTexture = nullptr;
 	FTextureReferenceRHIRef	AttributeTexture = nullptr;
 	FTextureReferenceRHIRef	AuxilaryDataTexture = nullptr;
+	FTextureReferenceRHIRef	GroupIndexTexture = nullptr;
 
 	/* Reference to the hair strands render data */
 	const FHairMeshesBulkData& BulkData;
diff --git a/Engine/Plugins/Runtime/HairStrands/Source/HairStrandsCore/Public/GroomTextureBuilder.h b/Engine/Plugins/Runtime/HairStrands/Source/HairStrandsCore/Public/GroomTextureBuilder.h
--- a/Engine/Plugins/Runtime/HairStrands/Source/HairStrandsCore/Public/GroomTextureBuilder.h
+++ b/Engine/Plugins/Runtime/HairStrands/Source/HairStrandsCore/Public/GroomTextureBuilder.h
@@ -51,7 +51,8 @@ struct FStrandsTexturesOutput
 	class UTexture2D* Tangent = nullptr;
 	class UTexture2D* Coverage = nullptr;
 	class UTexture2D* Attribute = nullptr;
-	bool IsValid() const { return Depth && Tangent && Coverage && Attribute; }
+	class UTexture2D* GroupIndex = nullptr;
+	bool IsValid() const { return Depth && Tangent && Coverage && Attribute && GroupIndex; }
 };
 
 ////////////////////////////////////////////////////////////////////////////////////////////////////////////////////////
diff --git a/Engine/Plugins/Runtime/HairStrands/Source/HairStrandsCore/Public/HairCardsBuilder.h b/Engine/Plugins/Runtime/HairStrands/Source/HairStrandsCore/Public/HairCardsBuilder.h
--- a/Engine/Plugins/Runtime/HairStrands/Source/HairStrandsCore/Public/HairCardsBuilder.h
+++ b/Engine/Plugins/Runtime/HairStrands/Source/HairStrandsCore/Public/HairCardsBuilder.h
@@ -13,6 +13,7 @@
 #include "GroomResources.h"
 
 struct FHairGroupCardsTextures;
+struct FHairStrandsVoxelData;
 
 #if WITH_EDITOR
 namespace FHairCardsBuilder
@@ -20,6 +21,7 @@ namespace FHairCardsBuilder
 	bool ImportGeometry(
 		const UStaticMesh* StaticMesh,
 		const FHairStrandsDatas& InStrandsData,
+		const FHairStrandsVoxelData& InStrandsVoxelData,
 		FHairCardsBulkData& OutBulk,
 		FHairStrandsDatas& OutGuides,
 		FHairCardsInterpolationBulkData& OutInterpolationBulkData);
diff --git a/Engine/Plugins/Runtime/HairStrands/Source/HairStrandsCore/Public/HairCardsVertexFactory.h b/Engine/Plugins/Runtime/HairStrands/Source/HairStrandsCore/Public/HairCardsVertexFactory.h
--- a/Engine/Plugins/Runtime/HairStrands/Source/HairStrandsCore/Public/HairCardsVertexFactory.h
+++ b/Engine/Plugins/Runtime/HairStrands/Source/HairStrandsCore/Public/HairCardsVertexFactory.h
@@ -26,12 +26,11 @@ public:
 	struct FDataType
 	{
 		FHairGroupInstance* Instance = nullptr;
-		uint32 GroupIndex = 0;
 		uint32 LODIndex = 0;
 		EHairGeometryType GeometryType = EHairGeometryType::NoneGeometry;
 	};
 
-	FHairCardsVertexFactory(FHairGroupInstance* Instance, uint32 GroupIndex, uint32 LODIndex, EHairGeometryType GeometryType, EShaderPlatform InShaderPlatform, ERHIFeatureLevel::Type InFeatureLevel, const char* InDebugName);
+	FHairCardsVertexFactory(FHairGroupInstance* Instance, uint32 LODIndex, EHairGeometryType GeometryType, EShaderPlatform InShaderPlatform, ERHIFeatureLevel::Type InFeatureLevel, const char* InDebugName);
 	
 	/**
 	 * Should we cache the material's shadertype on this platform with this vertex factory? 
diff --git a/Engine/Shaders/Private/HairStrands/HairCardsAttributeCommon.ush b/Engine/Shaders/Private/HairStrands/HairCardsAttributeCommon.ush
--- a/Engine/Shaders/Private/HairStrands/HairCardsAttributeCommon.ush
+++ b/Engine/Shaders/Private/HairStrands/HairCardsAttributeCommon.ush
@@ -16,6 +16,7 @@ struct FHairVertexAttributes
 	float3 HairBaseColor;
 	float  HairSeed;
 	float  HairRoughness;
+	float  HairGroupIndex;
 };
 
 float GetHairStrandsDepth(float2 InAtlasUV, float InDeviceZ)
@@ -41,7 +42,7 @@ float GetHairStrandsCoverage(float2 InAtlasUV)
 	return Texture2DSample(HairCardsVF.CoverageTexture, HairCardsVF.CoverageSampler, InAtlasUV).x;
 }
 
-FHairVertexAttributes GetHairStrandsAttributes(float2 InAtlasUV, float2 InRootUV = float2(0, 0), float InCardLength = 0.f)
+FHairVertexAttributes GetHairStrandsAttributes(float2 InAtlasUV, float2 InRootUV = float2(0, 0), float InCardLength = 0.f, float InCardGroupIndex = 0.f)
 {
 	InAtlasUV.y = HairCardsVF.bInvertUV ? 1 - InAtlasUV.y : InAtlasUV.y;
 	const float4 Value = Texture2DSample(HairCardsVF.AttributeTexture, HairCardsVF.AttributeSampler, InAtlasUV);
@@ -61,6 +62,7 @@ FHairVertexAttributes GetHairStrandsAttributes(float2 InAtlasUV, float2 InRootUV
 	Out.HairBaseColor = 0;
 	Out.HairSeed = Seed;
 	Out.HairRoughness = 0;
+	Out.HairGroupIndex = InCardGroupIndex;
 	return Out;
 }
 
@@ -108,9 +110,15 @@ float GetHairStrandsRoughness(float2 HairPrimitiveUV)
 	return GetHairStrandsAttributes(HairPrimitiveUV).HairRoughness;
 }
 
-float GetHairStrandsGroupIndex(float2 InAtlasUV)
+float GetHairStrandsGroupIndex(float2 InAtlasUV, float InHairPrimitiveGroupIndex)
 {
-	return HairCardsVF.GroupIndex; 
+	// * For hair cards,  group index is stored onto vertices
+	// * For hair meshes, group index is stored within a texture
+	if (HairCardsVF.bUseTextureGroupIndex)
+	{
+		InHairPrimitiveGroupIndex = Texture2DSample(HairCardsVF.GroupIndexTexture, HairCardsVF.GroupIndexSampler, InAtlasUV).x * 255.f;
+	}
+	return InHairPrimitiveGroupIndex;
 }
 ////////////////////////////////////////////////////////////////////////////////
 
diff --git a/Engine/Shaders/Private/HairStrands/HairCardsVertexFactory.ush b/Engine/Shaders/Private/HairStrands/HairCardsVertexFactory.ush
--- a/Engine/Shaders/Private/HairStrands/HairCardsVertexFactory.ush
+++ b/Engine/Shaders/Private/HairStrands/HairCardsVertexFactory.ush
@@ -83,6 +83,7 @@ struct FVertexFactoryIntermediates
 	float2 HairPrimitiveUV;
 	float2 HairPrimitiveRootUV;
 	float2 HairPrimitiveLength;
+	float  HairPrimitiveGroupIndex;
 
 	/** Cached primitive and instance data */
 	FSceneDataIntermediates SceneData;
@@ -108,6 +109,7 @@ struct FHairCardsPoint
 	float2 CardsUVs;
 	float2 RootUVs;
 	float  CardLength;
+	float  GroupIndex;
 };
 
 FHairCardsPoint GetVertexPosition(FVertexFactoryInput Input)
@@ -121,11 +123,14 @@ FHairCardsPoint GetVertexPosition(FVertexFactoryInput Input)
 	const float4 UVs  = Input.UVs;
 #endif
 
+	const uint EncodedW = asuint(Data.w);
+
 	FHairCardsPoint Out;
-	Out.Position = Data.xyz;
-	Out.CardsUVs = UVs.xy;
-	Out.RootUVs  = UVs.zw;
-	Out.CardLength = Data.w;
+	Out.Position	= Data.xyz;
+	Out.CardsUVs	= UVs.xy;
+	Out.RootUVs		= UVs.zw;
+	Out.CardLength	= f16tof32(EncodedW & 0xFFFF);
+	Out.GroupIndex	= (EncodedW>>16u);
 	return Out;
 }
 
@@ -173,6 +178,7 @@ FMaterialPixelParameters GetMaterialPixelParameters(FVertexFactoryInterpolantsVS
 	Result.HairPrimitiveUV = float2(Interpolants.HairPrimitiveUV);
 	Result.HairPrimitiveRootUV = float2(Interpolants.HairPrimitiveRootUV);
 	Result.HairPrimitiveLength = Interpolants.HairPrimitiveLength;
+	Result.HairPrimitiveGroupIndex = Interpolants.HairPrimitiveGroupIndex;
 #if NUM_TEX_COORD_INTERPOLATORS
 	UNROLL
 	for (int CoordinateIndex = 0; CoordinateIndex < NUM_TEX_COORD_INTERPOLATORS; CoordinateIndex++)
@@ -252,13 +258,14 @@ half3x3 CalcTangentToWorld(FVertexFactoryIntermediates Intermediates, half3x3 Ta
 
 FVertexFactoryIntermediates GetVertexFactoryIntermediates(FVertexFactoryInput Input)
 {
+	FHairCardsPoint CardsPoint = GetVertexPosition(Input);
+
 	FVertexFactoryIntermediates Intermediates;
-	Intermediates.SceneData			= VF_GPUSCENE_GET_INTERMEDIATES(Input);
-	Intermediates.HairPrimitiveUV	= GetVertexPosition(Input).CardsUVs;
-	Intermediates.HairPrimitiveRootUV= GetVertexPosition(Input).RootUVs;
-	Intermediates.HairPrimitiveLength= GetVertexPosition(Input).CardLength; 
-	//Intermediates.HairDensity		= HairStrandsVF_Density;
-	//Intermediates.HairDimensions	= float2(GetWorldStrandLength(Input), GetWorldStrandRadius(Input));
+	Intermediates.SceneData					= VF_GPUSCENE_GET_INTERMEDIATES(Input);
+	Intermediates.HairPrimitiveUV			= CardsPoint.CardsUVs;
+	Intermediates.HairPrimitiveRootUV		= CardsPoint.RootUVs;
+	Intermediates.HairPrimitiveLength		= CardsPoint.CardLength; 
+	Intermediates.HairPrimitiveGroupIndex	= CardsPoint.GroupIndex;
 
 	float TangentSign;
 	Intermediates.TangentToLocal = CalcTangentToLocal(Input, TangentSign);
@@ -312,9 +319,10 @@ FVertexFactoryInterpolantsVSToPS VertexFactoryGetInterpolantsVSToPS(FVertexFacto
 #endif
 
 	FHairCardsPoint CardsPoint = GetVertexPosition(Input);
-	Interpolants.HairPrimitiveUV = CardsPoint.CardsUVs;
-	Interpolants.HairPrimitiveRootUV = CardsPoint.RootUVs;
-	Interpolants.HairPrimitiveLength = CardsPoint.CardLength;
+	Interpolants.HairPrimitiveUV			= CardsPoint.CardsUVs;
+	Interpolants.HairPrimitiveRootUV		= CardsPoint.RootUVs;
+	Interpolants.HairPrimitiveLength		= CardsPoint.CardLength;
+	Interpolants.HairPrimitiveGroupIndex	= CardsPoint.GroupIndex;
 
 	SetPrimitiveId(Interpolants, Intermediates.SceneData.PrimitiveId);
 #if NUM_TEX_COORD_INTERPOLATORS 
@@ -439,6 +447,7 @@ FVertexFactoryRayTracingInterpolants VertexFactoryInterpolate(FVertexFactoryRayT
 	INTERPOLATE_MEMBER(InterpolantsVSToPS.HairPrimitiveUV);
 	INTERPOLATE_MEMBER(InterpolantsVSToPS.HairPrimitiveRootUV);
 	INTERPOLATE_MEMBER(InterpolantsVSToPS.HairPrimitiveLength);
+	INTERPOLATE_MEMBER(InterpolantsVSToPS.HairPrimitiveGroupIndex);
 #endif
 
 	return O;
diff --git a/Engine/Shaders/Private/HairStrands/HairStrandsTexturesGeneration.usf b/Engine/Shaders/Private/HairStrands/HairStrandsTexturesGeneration.usf
--- a/Engine/Shaders/Private/HairStrands/HairStrandsTexturesGeneration.usf
+++ b/Engine/Shaders/Private/HairStrands/HairStrandsTexturesGeneration.usf
@@ -88,6 +88,7 @@ float3 InVF_PositionOffset;
 float InVF_Radius;
 float InVF_Length;
 uint InVF_ControlPointCount;
+uint InVF_GroupIndex;
 
 Buffer<float> VertexBuffer;
 Buffer<float4> NormalsBuffer;
@@ -333,7 +334,8 @@ void MainPS(
 	out float4 OutTangent	: SV_Target1,
 	out float  OutCoverage	: SV_Target2,
 	out float4 OutAttribute : SV_Target3,
-	out uint   OutTriangleMask : SV_Target4,
+	out float  OutGroupIndex : SV_Target4,
+	out uint   OutTriangleMask : SV_Target5,
 	out float  OutDepthTest	: SV_Depth)
 {
 	const float3 T = normalize(InTangent);
@@ -347,6 +349,7 @@ void MainPS(
 	OutCoverage = 0;
 	OutTangent = 0;
 	OutAttribute = 0;
+	OutGroupIndex = 0;
 	OutDepthTest = 1;
 
 	float3 Tangent_ObjectSpace = 0;
@@ -421,6 +424,7 @@ void MainPS(
 		OutCoverage /= SampleCount;
 		OutTangent.xyz = (Tangent_FrameSpace + 1) * 0.5f;
 		OutTangent.w = 1;
+		OutGroupIndex = InVF_GroupIndex / 255.f;
 	}
 	else
 	{
@@ -428,6 +432,7 @@ void MainPS(
 		OutCoverage = 0;
 		OutTangent = 0;
 		OutAttribute = 0;
+		OutGroupIndex = 0;
 	}
 
 	OutDepth = saturate(OutDepth / MaxDistance);
@@ -452,11 +457,13 @@ Texture2D<float>  Source_DepthTexture;
 Texture2D<float>  Source_CoverageTexture;
 Texture2D<float4> Source_TangentTexture;
 Texture2D<float4> Source_AttributeTexture;
+Texture2D<float4> Source_GroupIndexTexture;
 
 RWTexture2D<float>  Target_DepthTexture;
 RWTexture2D<float>  Target_CoverageTexture;
 RWTexture2D<float4> Target_TangentTexture;
 RWTexture2D<float4> Target_AttributeTexture;
+RWTexture2D<float4> Target_GroupIndexTexture;
 
 // Dilate the source texture (Depth/Coverage/Tanget/Attribute), using a distance to triangle edge metric. 
 // Original valid pixel are marked into the triangleMasktexture during the texture generation. Then 
@@ -483,6 +490,7 @@ void MainCS(int2 DispatchThreadId : SV_DispatchThreadID)
 		Target_CoverageTexture[PixelCoord]	= Source_CoverageTexture.Load(uint3(PixelCoord, 0));
 		Target_TangentTexture[PixelCoord]	= Source_TangentTexture.Load(uint3(PixelCoord, 0));
 		Target_AttributeTexture[PixelCoord] = Source_AttributeTexture.Load(uint3(PixelCoord, 0));
+		Target_GroupIndexTexture[PixelCoord] = Source_GroupIndexTexture.Load(uint3(PixelCoord, 0));
 		return;
 	}
 	else
@@ -491,6 +499,7 @@ void MainCS(int2 DispatchThreadId : SV_DispatchThreadID)
 		float Target_Coverage = 0;
 		float4 Target_Tangent = 0;
 		float4 Target_Attribute = 0;
+		float Target_GroupIndex = 0;
 
 		// Find the closest neighbor from triangle, amoung valid neighbord
 		bool bIsDilatedPixelValid = false;
@@ -513,6 +522,7 @@ void MainCS(int2 DispatchThreadId : SV_DispatchThreadID)
 				Target_Coverage		= Source_CoverageTexture.Load(uint3(NPixelCoord, 0));
 				Target_Tangent		= Source_TangentTexture.Load(uint3(NPixelCoord, 0));
 				Target_Attribute	= Source_AttributeTexture.Load(uint3(NPixelCoord, 0));
+				Target_GroupIndex	= Source_GroupIndexTexture.Load(uint3(NPixelCoord, 0)).x;
 				MinDistance			= NDistance;
 				bIsDilatedPixelValid= true;
 			}
@@ -525,11 +535,10 @@ void MainCS(int2 DispatchThreadId : SV_DispatchThreadID)
 			Target_CoverageTexture[PixelCoord]	= Target_Coverage;
 			Target_TangentTexture[PixelCoord]	= Target_Tangent;
 			Target_AttributeTexture[PixelCoord] = Target_Attribute;
+			Target_GroupIndexTexture[PixelCoord] = float4(Target_GroupIndex,0,0,0);
 
 			TriangleMaskTexture[PixelCoord] = PackDistance(MinDistance);
 		}
 	}
-
-
 }
 #endif // SHADER_TEXTURE_DILATION
diff --git a/Engine/Shaders/Private/LocalVertexFactoryCommon.ush b/Engine/Shaders/Private/LocalVertexFactoryCommon.ush
--- a/Engine/Shaders/Private/LocalVertexFactoryCommon.ush
+++ b/Engine/Shaders/Private/LocalVertexFactoryCommon.ush
@@ -37,13 +37,14 @@ struct FVertexFactoryInterpolantsVSToPS
 	#endif
 #endif
 #if HAIR_STRAND_MESH_FACTORY
-		nointerpolation uint HairPrimitiveId : HAIR_PRIMITIVE_ID;	// Control point ID
-		float2 HairPrimitiveUV				 : HAIR_PRIMITIVE_UV;	// U: parameteric distance between the two surrounding control points. V: parametric distance along the width.
+		nointerpolation uint HairPrimitiveId		: HAIR_PRIMITIVE_ID;	// Control point ID
+		float2 HairPrimitiveUV						: HAIR_PRIMITIVE_UV;	// U: parameteric distance between the two surrounding control points. V: parametric distance along the width.
 #endif
 #if HAIR_CARD_MESH_FACTORY
-	float2 HairPrimitiveUV					: HAIR_PRIMITIVE_UV;	// Atlas UV
-	float2 HairPrimitiveRootUV				: HAIR_PRIMITIVE_ROOTUV;// UV at the root point of the cards
-	float  HairPrimitiveLength				: HAIR_PRIMITIVE_LENGTH;
+	float2 HairPrimitiveUV							: HAIR_PRIMITIVE_UV;	// Atlas UV
+	float2 HairPrimitiveRootUV						: HAIR_PRIMITIVE_ROOTUV;// UV at the root point of the cards
+	float  HairPrimitiveLength						: HAIR_PRIMITIVE_LENGTH;
+	nointerpolation float HairPrimitiveGroupIndex	: HAIR_PRIMITIVE_GROUPINDEX;
 #endif
 
 #if HAS_INSTANCE_LOCAL_TO_WORLD_PS
diff --git a/Engine/Shaders/Private/MaterialTemplate.ush b/Engine/Shaders/Private/MaterialTemplate.ush
--- a/Engine/Shaders/Private/MaterialTemplate.ush
+++ b/Engine/Shaders/Private/MaterialTemplate.ush
@@ -377,13 +377,14 @@ struct FMaterialPixelParameters
 
 	// Actual primitive Id
 #if HAIR_STRAND_MESH_FACTORY
-	uint	HairPrimitiveId;	// Control point ID
-	float2	HairPrimitiveUV;	// U: parametric distance between the two surrounding control point. V: parametric distance along hair width
+	uint	HairPrimitiveId;		// Control point ID
+	float2	HairPrimitiveUV;		// U: parametric distance between the two surrounding control point. V: parametric distance along hair width
 #endif
 #if HAIR_CARD_MESH_FACTORY
-	float2	HairPrimitiveUV;	// AtlasUV
-	float2  HairPrimitiveRootUV;// RootUV
-	float   HairPrimitiveLength;// Card length
+	float2	HairPrimitiveUV;		// AtlasUV
+	float2  HairPrimitiveRootUV;	// RootUV
+	float   HairPrimitiveLength;	// Card length
+	float	HairPrimitiveGroupIndex;// Card group index
 #endif
 
 #if HAS_INSTANCE_LOCAL_TO_WORLD_PS
@@ -1401,7 +1402,7 @@ float MaterialExpressionGetHairGroupIndex(FMaterialPixelParameters Parameters)
 #if HAIR_STRAND_MESH_FACTORY
 	return GetHairStrandsGroupIndex(Parameters.HairPrimitiveId);
 #elif HAIR_CARD_MESH_FACTORY
-	return GetHairStrandsGroupIndex(Parameters.HairPrimitiveUV);
+	return GetHairStrandsGroupIndex(Parameters.HairPrimitiveUV, Parameters.HairPrimitiveGroupIndex);
 #else
 	return 0;
 #endif
