Fix simulation initialization when toggling enable/disable simulation on a groom asset within the groom editor.

diff --git a/Engine/Plugins/Runtime/HairStrands/Source/HairStrandsCore/Private/GroomComponent.cpp b/Engine/Plugins/Runtime/HairStrands/Source/HairStrandsCore/Private/GroomComponent.cpp
--- a/Engine/Plugins/Runtime/HairStrands/Source/HairStrandsCore/Private/GroomComponent.cpp
+++ b/Engine/Plugins/Runtime/HairStrands/Source/HairStrandsCore/Private/GroomComponent.cpp
@@ -1360,7 +1360,16 @@ void UGroomComponent::ReleaseHairSimulation()
 {
 	for (int32 CompIndex = 0; CompIndex < NiagaraComponents.Num(); ++CompIndex)
 	{
-		TObjectPtr<UNiagaraComponent>& NiagaraComponent = NiagaraComponents[CompIndex];
+		ReleaseHairSimulation(CompIndex);
+	}
+	NiagaraComponents.Empty();
+}
+
+void UGroomComponent::ReleaseHairSimulation(const int32 GroupIndex)
+{
+	if (GroupIndex < NiagaraComponents.Num())
+	{
+		TObjectPtr<UNiagaraComponent>& NiagaraComponent = NiagaraComponents[GroupIndex];
 		if (NiagaraComponent && !NiagaraComponent->IsBeingDestroyed())
 		{
 			if (GetWorld() && NiagaraComponent->IsRegistered())
@@ -1372,7 +1381,6 @@ void UGroomComponent::ReleaseHairSimulation()
 			NiagaraComponent = nullptr;
 		}
 	}
-	NiagaraComponents.Empty();
 }
 
 void UGroomComponent::CreateHairSimulation(const int32 GroupIndex, const int32 LODIndex)
@@ -2972,6 +2980,7 @@ void UGroomComponent::Invalidate()
 void UGroomComponent::InvalidateAndRecreate()
 {
 	InitResources(true);
+	UpdateHairSimulation();
 	MarkRenderStateDirty();
 }
 #endif
diff --git a/Engine/Plugins/Runtime/HairStrands/Source/HairStrandsCore/Public/GroomComponent.h b/Engine/Plugins/Runtime/HairStrands/Source/HairStrandsCore/Public/GroomComponent.h
--- a/Engine/Plugins/Runtime/HairStrands/Source/HairStrandsCore/Public/GroomComponent.h
+++ b/Engine/Plugins/Runtime/HairStrands/Source/HairStrandsCore/Public/GroomComponent.h
@@ -85,6 +85,7 @@ public:
 
 	/** Release Niagara components */
 	void ReleaseHairSimulation();
+	void ReleaseHairSimulation(const int32 GroupIndex);
 
 	/** Create per Group/LOD the Niagara component */
 	void CreateHairSimulation(const int32 GroupIndex, const int32 LODIndex);
