Improve strands to guide density attribute to clarify its meaning/expected value.

diff --git a/Engine/Plugins/Runtime/HairStrands/Source/HairStrandsCore/Public/GroomAssetInterpolation.h b/Engine/Plugins/Runtime/HairStrands/Source/HairStrandsCore/Public/GroomAssetInterpolation.h
--- a/Engine/Plugins/Runtime/HairStrands/Source/HairStrandsCore/Public/GroomAssetInterpolation.h
+++ b/Engine/Plugins/Runtime/HairStrands/Source/HairStrandsCore/Public/GroomAssetInterpolation.h
@@ -123,7 +123,7 @@ struct HAIRSTRANDSCORE_API FHairInterpolationSettings
 	UPROPERTY(EditAnywhere, BlueprintReadWrite, Category = "InterpolationSettings", meta = (ToolTip = "If checked, override imported guides with generated ones."))
 	bool bOverrideGuides;
 
-	/** Density factor for converting hair into guide curve if no guides are provided. */
+	/** Density factor for converting hair into guide curve if no guides are provided. The value should be between 0 and 1, and can be thought as a ratio/percentage of strands used as guides.*/
 	UPROPERTY(EditAnywhere, BlueprintReadWrite, Category = "InterpolationSettings", meta = (ClampMin = "0", ClampMax = "1.0", UIMin = "0", UIMax = "1.0"))
 	float HairToGuideDensity;
 
