Mesh user parameter for physics asset
useful for setting the skel mesh in BP

diff --git a/Engine/Plugins/Runtime/HairStrands/Source/HairStrandsCore/Private/Niagara/NiagaraDataInterfacePhysicsAsset.cpp b/Engine/Plugins/Runtime/HairStrands/Source/HairStrandsCore/Private/Niagara/NiagaraDataInterfacePhysicsAsset.cpp
--- a/Engine/Plugins/Runtime/HairStrands/Source/HairStrandsCore/Private/Niagara/NiagaraDataInterfacePhysicsAsset.cpp
+++ b/Engine/Plugins/Runtime/HairStrands/Source/HairStrandsCore/Private/Niagara/NiagaraDataInterfacePhysicsAsset.cpp
@@ -684,16 +684,90 @@ UNiagaraDataInterfacePhysicsAsset::UNiagaraDataInterfacePhysicsAsset(FObjectInit
 	, SourceComponents()
 	, PhysicsAssets()
 {
+	FNiagaraTypeDefinition Def(UObject::StaticClass());
+	MeshUserParameter.Parameter.SetType(Def);
+
 	Proxy.Reset(new FNDIPhysicsAssetProxy());
 }
 
 void UNiagaraDataInterfacePhysicsAsset::ExtractSourceComponent(FNiagaraSystemInstance* SystemInstance, FTransform& BoneTransform)
 {
+	// Helper to scour an actor (or its parents) for a valid skeletal mesh component
+	auto FindActorSkelMeshComponent = [](AActor* Actor, bool bRecurseParents = false) -> USkeletalMeshComponent*
+	{
+		if (ASkeletalMeshActor* SkelMeshActor = Cast<ASkeletalMeshActor>(Actor))
+		{
+			USkeletalMeshComponent* Comp = SkelMeshActor->GetSkeletalMeshComponent();
+			if (IsValid(Comp))
+			{
+				return Comp;
+			}
+		}
+
+		// Fall back on any valid component on the actor
+		while (Actor)
+		{
+			for (UActorComponent* ActorComp : Actor->GetComponents())
+			{
+				USkeletalMeshComponent* Comp = Cast<USkeletalMeshComponent>(ActorComp);
+				if (IsValid(Comp) && Comp->SkeletalMesh != nullptr)
+				{
+					return Comp;
+				}
+			}
+
+			if (bRecurseParents)
+			{
+				Actor = Actor->GetParentActor();
+			}
+			else
+			{
+				break;
+			}
+		}
+
+		return nullptr;
+	};
+		
 	BoneTransform = FTransform::Identity;
 
 	// Track down the source component
 	TWeakObjectPtr<USkeletalMeshComponent> SourceComponent;
-	if (SourceActor)
+	
+	
+	if (MeshUserParameter.Parameter.IsValid() && SystemInstance != nullptr)
+	{
+		// Initialize the binding and retrieve the object. If a valid object is bound, we'll try and retrieve the SkelMesh component from it.
+		// If it's not valid yet, we'll reset and do this again when/if a valid object is set on the binding
+		FNiagaraParameterDirectBinding<UObject*> UserParamBinding;
+		UObject* UserParamObject = UserParamBinding.Init(SystemInstance->GetInstanceParameters(), MeshUserParameter.Parameter);
+
+		if (UserParamObject)
+		{
+			if (USkeletalMeshComponent* UserSkelMeshComp = Cast<USkeletalMeshComponent>(UserParamObject))
+			{
+				if (IsValid(UserSkelMeshComp))
+				{
+					SourceComponent = UserSkelMeshComp;
+				}
+			}
+			else if (AActor* Actor = Cast<AActor>(UserParamObject))
+			{
+				SourceComponent = FindActorSkelMeshComponent(Actor);
+			}
+			else
+			{
+				//We have a valid, non-null UObject parameter type but it is not a type we can use to get a skeletal mesh from.
+				UE_LOG(LogPhysicsAsset, Warning, TEXT("SkeletalMesh data interface using object parameter with invalid type. Skeletal Mesh Data Interfaces can only get a valid mesh from SkeletalMeshComponents, SkeletalMeshActors or Actors."));
+				UE_LOG(LogPhysicsAsset, Warning, TEXT("Invalid Parameter : %s"), *UserParamObject->GetFullName());
+				UE_LOG(LogPhysicsAsset, Warning, TEXT("System : %s"), *GetFullNameSafe(SystemInstance->GetSystem()));
+			}
+		}
+		else
+		{
+			// The binding exists, but no object is bound. Not warning here in case the user knows what they're doing.
+		}
+	} else if (SourceActor)
 	{
 		ASkeletalMeshActor* SkeletalMeshActor = Cast<ASkeletalMeshActor>(SourceActor);
 		if (SkeletalMeshActor != nullptr)
@@ -844,6 +918,7 @@ bool UNiagaraDataInterfacePhysicsAsset::CopyToInternal(UNiagaraDataInterface* De
 	OtherTyped->SourceActor = SourceActor;
 	OtherTyped->SourceComponents = SourceComponents;
 	OtherTyped->DefaultSource = DefaultSource;
+	OtherTyped->MeshUserParameter = MeshUserParameter;
 
 	return true;
 }
@@ -856,7 +931,11 @@ bool UNiagaraDataInterfacePhysicsAsset::Equals(const UNiagaraDataInterface* Othe
 	}
 	const UNiagaraDataInterfacePhysicsAsset* OtherTyped = CastChecked<const UNiagaraDataInterfacePhysicsAsset>(Other);
 
-	return  (OtherTyped->PhysicsAssets == PhysicsAssets) && (OtherTyped->SourceActor == SourceActor) && (OtherTyped->SourceComponents == SourceComponents) && (OtherTyped->DefaultSource == DefaultSource);
+	return  (OtherTyped->PhysicsAssets == PhysicsAssets) && 
+		(OtherTyped->SourceActor == SourceActor) && 
+		(OtherTyped->SourceComponents == SourceComponents) && 
+		(OtherTyped->DefaultSource == DefaultSource) &&
+		(OtherTyped->MeshUserParameter == MeshUserParameter);
 }
 
 void UNiagaraDataInterfacePhysicsAsset::PostInitProperties()
diff --git a/Engine/Plugins/Runtime/HairStrands/Source/HairStrandsCore/Public/Niagara/NiagaraDataInterfacePhysicsAsset.h b/Engine/Plugins/Runtime/HairStrands/Source/HairStrandsCore/Public/Niagara/NiagaraDataInterfacePhysicsAsset.h
--- a/Engine/Plugins/Runtime/HairStrands/Source/HairStrandsCore/Public/Niagara/NiagaraDataInterfacePhysicsAsset.h
+++ b/Engine/Plugins/Runtime/HairStrands/Source/HairStrandsCore/Public/Niagara/NiagaraDataInterfacePhysicsAsset.h
@@ -112,7 +112,11 @@ public:
 
 	/** The source actor from which to sample */
 	UPROPERTY(EditAnywhere, Category = "Source")
-		TObjectPtr<AActor> SourceActor;
+	TObjectPtr<AActor> SourceActor;
+
+	/** Reference to a user parameter if we're reading one. */
+	UPROPERTY(EditAnywhere, Category = "Source")
+	FNiagaraUserParameterBinding MeshUserParameter;
 
 	/** The source component from which to sample */
 	TArray<TWeakObjectPtr<class USkeletalMeshComponent>> SourceComponents;
