Removed FbxHairTranslator since it was a proof of concept, not a full fledged translator that followed the groom specification.

diff --git a/Engine/Plugins/Runtime/HairStrands/Source/HairStrandsEditor/Private/FbxHairTranslator.cpp b/Engine/Plugins/Runtime/HairStrands/Source/HairStrandsEditor/Private/FbxHairTranslator.cpp
deleted file mode 100644
diff --git a/Engine/Plugins/Runtime/HairStrands/Source/HairStrandsEditor/Private/FbxHairTranslator.h b/Engine/Plugins/Runtime/HairStrands/Source/HairStrandsEditor/Private/FbxHairTranslator.h
deleted file mode 100644
diff --git a/Engine/Plugins/Runtime/HairStrands/Source/HairStrandsEditor/Private/HairStrandsEditor.cpp b/Engine/Plugins/Runtime/HairStrands/Source/HairStrandsEditor/Private/HairStrandsEditor.cpp
--- a/Engine/Plugins/Runtime/HairStrands/Source/HairStrandsEditor/Private/HairStrandsEditor.cpp
+++ b/Engine/Plugins/Runtime/HairStrands/Source/HairStrandsEditor/Private/HairStrandsEditor.cpp
@@ -12,7 +12,6 @@
 #include "EditorStyleSet.h"
 #include "Interfaces/IPluginManager.h"
 #include "SlateOptMacros.h"
-#include "FbxHairTranslator.h"
 
 #include "GroomAsset.h"
 #include "GroomBindingAsset.h"
@@ -83,8 +82,6 @@ void FGroomEditor::StartupModule()
 	RegisteredAssetTypeActions.Add(GroomAssetActions);
 	RegisteredAssetTypeActions.Add(BindingAssetActions);
 
-	RegisterHairTranslator<FFbxHairTranslator>();
-
 	// Only register once
 	if (!StyleSet.IsValid())
 	{
