Fix hair metal compilation + add overlapping skelmesh for collisions + fix groom invalidation + setgroomasset BP function + local simulation to default + fix rigid transform without attachment

diff --git a/Engine/Plugins/FX/Niagara/Source/NiagaraCore/Private/NiagaraCustomVersion.cpp b/Engine/Plugins/FX/Niagara/Source/NiagaraCore/Private/NiagaraCustomVersion.cpp
--- a/Engine/Plugins/FX/Niagara/Source/NiagaraCore/Private/NiagaraCustomVersion.cpp
+++ b/Engine/Plugins/FX/Niagara/Source/NiagaraCore/Private/NiagaraCustomVersion.cpp
@@ -8,4 +8,6 @@ const FGuid FNiagaraCustomVersion::GUID(0xFCF57AFA, 0x50764283, 0xB9A9E658, 0xFF
 // Register the custom version with core
 FCustomVersionRegistration GRegisterNiagaraCustomVersion(FNiagaraCustomVersion::GUID, FNiagaraCustomVersion::LatestVersion, TEXT("NiagaraVer"));
 
-const FGuid FNiagaraCustomVersion::LatestScriptCompileVersion(0xA73D5647, 0xED387740, 0x93DF7180, 0xD1FD586C);
+const FGuid FNiagaraCustomVersion::LatestScriptCompileVersion(0xDB4455C6, 0x96344711, 0xD5A3076, 0x9F35A1B5);
+
+
diff --git a/Engine/Plugins/Runtime/HairStrands/Shaders/Private/NiagaraClearPressureGrid.usf b/Engine/Plugins/Runtime/HairStrands/Shaders/Private/NiagaraClearPressureGrid.usf
--- a/Engine/Plugins/Runtime/HairStrands/Shaders/Private/NiagaraClearPressureGrid.usf
+++ b/Engine/Plugins/Runtime/HairStrands/Shaders/Private/NiagaraClearPressureGrid.usf
@@ -8,36 +8,36 @@ static const int NUM_ATTRIBUTES = 18;
 
 static const float MAX_DISTANCE = 1e+8;
 
-int3 GridSize;
+uint3 GridSize;
 int CopyPressure;
 
-Texture3D<int>		GridCurrentBuffer;
-RWTexture3D<int>	GridDestinationBuffer;
+Buffer<int>		GridCurrentBuffer;
+RWBuffer<int>	GridDestinationBuffer;
 
-[numthreads(THREAD_COUNT, 1, 1)]
-void MainCS(uint2 DispatchThreadId : SV_DispatchThreadID)
+[numthreads(THREAD_COUNT, THREAD_COUNT, THREAD_COUNT)]
+void MainCS(uint3 DispatchThreadId : SV_DispatchThreadID)
 {
-	const uint GridHash = DispatchThreadId.x;
-	if (GridHash < (GridSize.x+1)*(GridSize.y+1)*(GridSize.z+1))
+	if (all(DispatchThreadId < GridSize))
 	{
+		const int LinearIndexYZ = DispatchThreadId.y * (GridSize.z+1) + DispatchThreadId.z;
 		const int GridSizeYZ = (GridSize.y+1) * (GridSize.z+1);
-		const int GridHashYZ = GridHash % GridSizeYZ;
-
-		const int3 GridIndex = int3(GridHash / GridSizeYZ, GridHashYZ / (GridSize.z+1), GridHashYZ % (GridSize.z+1));
-
+		
 		[unroll]
 		for(int i = 0; i < NUM_ATTRIBUTES; ++i)
 		{
-			const int3 AttrIndex = int3(GridIndex.x+i*(GridSize.x+1),GridIndex.yz);
-			GridDestinationBuffer[AttrIndex] = 0;
+			const int AttrIndex = DispatchThreadId.x+i*(GridSize.x+1);
+			const int LinearIndex = AttrIndex.x * GridSizeYZ + LinearIndexYZ;
+			GridDestinationBuffer[LinearIndex] = 0;
 		}
 		if(CopyPressure)
 		{
-			const int3 PressureIndex = int3(GridIndex.x+FLUID_PRESSURE_OFFSET*(GridSize.x+1),GridIndex.yz);
-			GridDestinationBuffer[PressureIndex] = GridCurrentBuffer[PressureIndex];
+			const int PressureIndex = DispatchThreadId.x+FLUID_PRESSURE_OFFSET*(GridSize.x+1);
+			const int LinearIndex = PressureIndex.x * GridSizeYZ + LinearIndexYZ;
+			GridDestinationBuffer[LinearIndex] = GridCurrentBuffer[LinearIndex];
 		}
 
-		const int3 DistanceIndex = int3(GridIndex.x+FLUID_DISTANCE_OFFSET*(GridSize.x+1),GridIndex.yz);
-		GridDestinationBuffer[DistanceIndex] = MAX_DISTANCE;
+		const int DistanceIndex = DispatchThreadId.x+FLUID_DISTANCE_OFFSET*(GridSize.x+1);
+		const int LinearIndex = DistanceIndex.x * GridSizeYZ + LinearIndexYZ;
+		GridDestinationBuffer[LinearIndex] = MAX_DISTANCE;
 	}
 }
\ No newline at end of file
diff --git a/Engine/Plugins/Runtime/HairStrands/Shaders/Private/NiagaraClearVelocityGrid.usf b/Engine/Plugins/Runtime/HairStrands/Shaders/Private/NiagaraClearVelocityGrid.usf
--- a/Engine/Plugins/Runtime/HairStrands/Shaders/Private/NiagaraClearVelocityGrid.usf
+++ b/Engine/Plugins/Runtime/HairStrands/Shaders/Private/NiagaraClearVelocityGrid.usf
@@ -3,13 +3,14 @@
 #include "/Engine/Private/Common.ush"
 
 uint3 GridSize;
-RWTexture3D<int>	GridDestinationBuffer;
+RWBuffer<int>	GridDestinationBuffer;
 
 [numthreads(THREAD_COUNT, THREAD_COUNT, THREAD_COUNT)]
 void MainCS(uint3 DispatchThreadId : SV_DispatchThreadID)
 {
-	if ((DispatchThreadId.x < GridSize.x) && (DispatchThreadId.y < GridSize.y) && (DispatchThreadId.z < GridSize.z))
+	if (all(DispatchThreadId < GridSize))
 	{
-		GridDestinationBuffer[DispatchThreadId] = 0;
+		const int GridIndex = DispatchThreadId.x * (GridSize.y+1) * (GridSize.z+1) + DispatchThreadId.y * (GridSize.z+1) + DispatchThreadId.z;
+		GridDestinationBuffer[GridIndex] = 0;
 	}
 }
\ No newline at end of file
diff --git a/Engine/Plugins/Runtime/HairStrands/Shaders/Private/NiagaraCopyVelocityGrid.usf b/Engine/Plugins/Runtime/HairStrands/Shaders/Private/NiagaraCopyVelocityGrid.usf
--- a/Engine/Plugins/Runtime/HairStrands/Shaders/Private/NiagaraCopyVelocityGrid.usf
+++ b/Engine/Plugins/Runtime/HairStrands/Shaders/Private/NiagaraCopyVelocityGrid.usf
@@ -3,14 +3,15 @@
 #include "/Engine/Private/Common.ush"
 
 uint3 GridSize;
-Texture3D<int>		GridCurrentBuffer;
-RWTexture3D<int>	GridDestinationBuffer;
+Buffer<int>		GridCurrentBuffer;
+RWBuffer<int>	GridDestinationBuffer;
 
 [numthreads(THREAD_COUNT, THREAD_COUNT, THREAD_COUNT)]
 void MainCS(uint3 DispatchThreadId : SV_DispatchThreadID)
 {
-	if ((DispatchThreadId.x < GridSize.x) && (DispatchThreadId.y < GridSize.y) && (DispatchThreadId.z < GridSize.z))
+	if (all(DispatchThreadId < GridSize))
 	{
-		GridDestinationBuffer[DispatchThreadId] = GridCurrentBuffer[DispatchThreadId];
+		const int GridIndex = DispatchThreadId.x * (GridSize.y+1) * (GridSize.z+1) + DispatchThreadId.y * (GridSize.z+1) + DispatchThreadId.z;
+		GridDestinationBuffer[GridIndex] = GridCurrentBuffer[GridIndex];
 	}
 }
\ No newline at end of file
diff --git a/Engine/Plugins/Runtime/HairStrands/Shaders/Private/NiagaraDataInterfacePressureGrid.ush b/Engine/Plugins/Runtime/HairStrands/Shaders/Private/NiagaraDataInterfacePressureGrid.ush
--- a/Engine/Plugins/Runtime/HairStrands/Shaders/Private/NiagaraDataInterfacePressureGrid.ush
+++ b/Engine/Plugins/Runtime/HairStrands/Shaders/Private/NiagaraDataInterfacePressureGrid.ush
@@ -23,7 +23,7 @@ static const int FLUID_DENSITY_OFFSET = 15;
  * -----------------------------------------------------------------
  */
 
-void DIVelocityGrid_ScaleCellDistance(in FDIVelocityGridContext DIContext, in RWTexture3D<int> DIContext_GridDestinationBuffer, in int3 GridIndex, in float GridLength )
+void DIVelocityGrid_ScaleCellDistance(in FDIVelocityGridContext DIContext, in RWBuffer<int> DIContext_GridDestinationBuffer, in int3 GridIndex, in float GridLength )
 {
 	float FluidDistance = DIVelocityGrid_ConvertGridScalar(DIContext,GridIndex,FLUID_DISTANCE_OFFSET);
 
@@ -45,7 +45,7 @@ void DIVelocityGrid_ScaleCellDistance(in FDIVelocityGridContext DIContext, in RW
 	DIVelocityGrid_SetGridScalar(DIContext,DIContext_GridDestinationBuffer,GridIndex,FLUID_DISTANCE_OFFSET,FluidDistance);
 }
 
-void DIVelocityGrid_ScaleFaceVelocity(in FDIVelocityGridContext DIContext, in RWTexture3D<int> DIContext_GridDestinationBuffer, in int3 GridIndex, in int ComponentIndex)
+void DIVelocityGrid_ScaleFaceVelocity(in FDIVelocityGridContext DIContext, in RWBuffer<int> DIContext_GridDestinationBuffer, in int3 GridIndex, in int ComponentIndex)
 {
 	const float GridMomentum = DIVelocityGrid_ConvertGridScalar(DIContext,GridIndex,FLUID_VELOCITY_OFFSET+ComponentIndex);
 	const float GridMass = DIVelocityGrid_ConvertGridScalar(DIContext,GridIndex,FLUID_MASS_OFFSET+ComponentIndex);
@@ -56,7 +56,7 @@ void DIVelocityGrid_ScaleFaceVelocity(in FDIVelocityGridContext DIContext, in RW
 	DIVelocityGrid_SetGridScalar(DIContext,DIContext_GridDestinationBuffer,GridIndex,FLUID_MASS_OFFSET+ComponentIndex,GridMass);
 }
 
-void DIVelocityGrid_ScaleCellFields(in FDIVelocityGridContext DIContext, in RWTexture3D<int> DIContext_GridDestinationBuffer, in int GridHash, in float GridLength, in float DeltaTime, out bool OutTransferStatus)
+void DIVelocityGrid_ScaleCellFields(in FDIVelocityGridContext DIContext, in RWBuffer<int> DIContext_GridDestinationBuffer, in int GridHash, in float GridLength, in float DeltaTime, out bool OutTransferStatus)
 {
 	const int3 GridIndex = DIVelocityGrid_GetGridIndex(DIContext,GridHash);
 
@@ -95,7 +95,7 @@ void DIVelocityGrid_ScaleCellFields(in FDIVelocityGridContext DIContext, in RWTe
  */
  
 // Compute the distance from the particle to the grid
-void DIVelocityGrid_BuildDistanceField(in FDIVelocityGridContext DIContext, in RWTexture3D<int> DIContext_GridDestinationBuffer, in float3 GridOrigin, in float GridLength, in float3 ParticlePosition, out bool OutFunctionStatus )
+void DIVelocityGrid_BuildDistanceField(in FDIVelocityGridContext DIContext, in RWBuffer<int> DIContext_GridDestinationBuffer, in float3 GridOrigin, in float GridLength, in float3 ParticlePosition, out bool OutFunctionStatus )
 {
 	OutFunctionStatus = true;
 	const float3 GridPosition = DIVelocityGrid_GetGridPosition(DIContext,ParticlePosition,GridOrigin);
@@ -122,7 +122,7 @@ void DIVelocityGrid_BuildDistanceField(in FDIVelocityGridContext DIContext, in R
 }
 
 // Transfer the density from the particle to the grid
-void DIVelocityGrid_BuildDensityField(in FDIVelocityGridContext DIContext, in RWTexture3D<int> DIContext_GridDestinationBuffer, in float3 GridOrigin, in float GridLength,
+void DIVelocityGrid_BuildDensityField(in FDIVelocityGridContext DIContext, in RWBuffer<int> DIContext_GridDestinationBuffer, in float3 GridOrigin, in float GridLength,
 				in float3 ParticlePosition, in float ParticleMass, in float ParticleDensity, out bool OutFunctionStatus )
 {
 	OutFunctionStatus = true;
@@ -377,7 +377,7 @@ void DIVelocityGrid_UpdateDeformationGradient(in FDIVelocityGridContext DIContex
  }
 
 // Update the solid weights and velocities
-void DIVelocityGrid_SetSolidBoundary(in FDIVelocityGridContext DIContext, in RWTexture3D<int> DIContext_GridDestinationBuffer,
+void DIVelocityGrid_SetSolidBoundary(in FDIVelocityGridContext DIContext, in RWBuffer<int> DIContext_GridDestinationBuffer,
 									in int GridHash, in float SolidDistance,
 									in float3 SolidVelocity, out bool OutBoundaryStatus)
 {
@@ -393,13 +393,13 @@ void DIVelocityGrid_SetSolidBoundary(in FDIVelocityGridContext DIContext, in RWT
 	}
 }
 
-void DIVelocityGrid_ComputeFaceWeights(in FDIVelocityGridContext DIContext, in RWTexture3D<int> DIContext_GridDestinationBuffer, in int3 GridLocation, in int ComponentIndex, out float SignedDistance, inout float ParticleVelocity )
+void DIVelocityGrid_ComputeFaceWeights(in FDIVelocityGridContext DIContext, in RWBuffer<int> DIContext_GridDestinationBuffer, in int3 GridLocation, in int ComponentIndex, out float SignedDistance, inout float ParticleVelocity )
 {
 	SignedDistance = DIVelocityGrid_GetGridScalar(DIContext,GridLocation,SOLID_DISTANCE_OFFSET);
 	ParticleVelocity += DIVelocityGrid_GetGridScalar(DIContext,GridLocation,SOLID_VELOCITY_OFFSET+ComponentIndex);
 }
 
-void DIVelocityGrid_ComputeComponentWeights(in FDIVelocityGridContext DIContext, in RWTexture3D<int> DIContext_GridDestinationBuffer, in int3 GridIndex, in int ComponentIndex)
+void DIVelocityGrid_ComputeComponentWeights(in FDIVelocityGridContext DIContext, in RWBuffer<int> DIContext_GridDestinationBuffer, in int3 GridIndex, in int ComponentIndex)
 {
 	[branch]
 	if( DIVelocityGrid_IsFaceInside(DIContext,GridIndex,ComponentIndex))
@@ -439,7 +439,7 @@ void DIVelocityGrid_ComputeComponentWeights(in FDIVelocityGridContext DIContext,
 }
 
  // Update the solid weights and velocities
-void DIVelocityGrid_ComputeBoundaryWeights(in FDIVelocityGridContext DIContext, in RWTexture3D<int> DIContext_GridDestinationBuffer, in int GridHash, out bool OutWeightsStatus)
+void DIVelocityGrid_ComputeBoundaryWeights(in FDIVelocityGridContext DIContext, in RWBuffer<int> DIContext_GridDestinationBuffer, in int GridHash, out bool OutWeightsStatus)
 {
 	const int3 GridIndex = DIVelocityGrid_GetGridIndex(DIContext,GridHash);
 	OutWeightsStatus = false;
@@ -500,7 +500,7 @@ float DIVelocityGrid_ComputeGridPressure(in FDIVelocityGridContext DIContext, in
 }
 
 // Project the pressure grid to be divergence free
-void DIVelocityGrid_SolveGridPressure(in FDIVelocityGridContext DIContext, in RWTexture3D<int> DIContext_GridDestinationBuffer, in int GridHash, in int InitStage, out bool OutProjectStatus)
+void DIVelocityGrid_SolveGridPressure(in FDIVelocityGridContext DIContext, in RWBuffer<int> DIContext_GridDestinationBuffer, in int GridHash, in int InitStage, out bool OutProjectStatus)
 {	
 	OutProjectStatus = false;
 
diff --git a/Engine/Plugins/Runtime/HairStrands/Shaders/Private/NiagaraDataInterfaceVelocityGrid.ush b/Engine/Plugins/Runtime/HairStrands/Shaders/Private/NiagaraDataInterfaceVelocityGrid.ush
--- a/Engine/Plugins/Runtime/HairStrands/Shaders/Private/NiagaraDataInterfaceVelocityGrid.ush
+++ b/Engine/Plugins/Runtime/HairStrands/Shaders/Private/NiagaraDataInterfaceVelocityGrid.ush
@@ -23,15 +23,15 @@ struct FDIVelocityGridContext
 	float4x4			WorldTransform;
 	float4x4			WorldInverse;
 	int3				GridSize;
-	Texture3D<int>		GridCurrentBuffer;
+	Buffer<int>		GridCurrentBuffer;
 }; 
 
 #define DIVelocityGrid_DECLARE_CONSTANTS(NAME)\
 float4x4 WorldTransform_##NAME;\
 float4x4 WorldInverse_##NAME;\
 int3 GridSize_##NAME;\
-Texture3D<int> GridCurrentBuffer_##NAME;\
-RWTexture3D<int> GridDestinationBuffer_##NAME;\
+Buffer<int> GridCurrentBuffer_##NAME;\
+RWBuffer<int> GridDestinationBuffer_##NAME;\
 
 //@todo-lh: GridDestinationBuffer is passed around separately as SPIRV-Cross cannot translate texture atomics properly atm.
 #define DIVelocityGrid_MAKE_CONTEXT(NAME)\
@@ -40,7 +40,7 @@ DIContext.WorldTransform = WorldTransform_##NAME;\
 DIContext.WorldInverse = WorldInverse_##NAME;\
 DIContext.GridSize = GridSize_##NAME;\
 DIContext.GridCurrentBuffer = GridCurrentBuffer_##NAME;\
-RWTexture3D<int> DIContext_GridDestinationBuffer = GridDestinationBuffer_##NAME;\
+RWBuffer<int> DIContext_GridDestinationBuffer = GridDestinationBuffer_##NAME;\
 
 /* -----------------------------------------------------------------
  * Grid topology conversion
@@ -231,26 +231,27 @@ void DIVelocityGrid_GetNodePosition(in FDIVelocityGridContext DIContext, in int
  */
 
 // Compute the data index from agrid index and a data offset
-int3 DIVelocityGrid_GetDataIndex( in FDIVelocityGridContext DIContext, in int3 GridIndex, in int DataOffset)
+int DIVelocityGrid_GetDataIndex( in FDIVelocityGridContext DIContext, in int3 GridIndex, in int DataOffset)
 {
-	return int3( GridIndex.x + DataOffset * (DIContext.GridSize.x+1), GridIndex.yz);
+	const int GridIndexX = GridIndex.x + DataOffset * (DIContext.GridSize.x+1);
+	return GridIndexX * (DIContext.GridSize.y+1) * (DIContext.GridSize.z+1) + GridIndex.y * (DIContext.GridSize.z+1) + GridIndex.z;
 }
 
 #if FAST_MODE == 1
 
-void DIVelocityGrid_InterlockedAddInt(in FDIVelocityGridContext DIContext, in RWTexture3D<int> DIContext_GridDestinationBuffer, in int3 GridIndex, in int DataOffset, in int IntValue)
+void DIVelocityGrid_InterlockedAddInt(in FDIVelocityGridContext DIContext, in RWBuffer<int> DIContext_GridDestinationBuffer, in int3 GridIndex, in int DataOffset, in int IntValue)
 {
 	InterlockedAdd(DIContext_GridDestinationBuffer[DIVelocityGrid_GetDataIndex(DIContext,GridIndex,DataOffset)], IntValue);
 }
 
-void DIVelocityGrid_InterlockedMinInt(in FDIVelocityGridContext DIContext, in RWTexture3D<int> DIContext_GridDestinationBuffer, in int3 GridIndex, in int DataOffset, in int IntValue)
+void DIVelocityGrid_InterlockedMinInt(in FDIVelocityGridContext DIContext, in RWBuffer<int> DIContext_GridDestinationBuffer, in int3 GridIndex, in int DataOffset, in int IntValue)
 {
 	InterlockedMin(DIContext_GridDestinationBuffer[DIVelocityGrid_GetDataIndex(DIContext,GridIndex,DataOffset)], IntValue);
 }
 
 #else
 
-void DIVelocityGrid_InterlockedAddFloat(in FDIVelocityGridContext DIContext, in RWTexture3D<int> DIContext_GridDestinationBuffer, in int3 GridIndex, in int DataOffset, in float FloatValue)
+void DIVelocityGrid_InterlockedAddFloat(in FDIVelocityGridContext DIContext, in RWBuffer<int> DIContext_GridDestinationBuffer, in int3 GridIndex, in int DataOffset, in float FloatValue)
 {
 	int IntValue = asint(FloatValue);
 	int CompareValue = 0;
@@ -269,7 +270,7 @@ void DIVelocityGrid_InterlockedAddFloat(in FDIVelocityGridContext DIContext, in
 	}
 }
 
-void DIVelocityGrid_InterlockedMinFloat(in FDIVelocityGridContext DIContext, in RWTexture3D<int> DIContext_GridDestinationBuffer, in int3 GridIndex, in int DataOffset, in float FloatValue)
+void DIVelocityGrid_InterlockedMinFloat(in FDIVelocityGridContext DIContext, in RWBuffer<int> DIContext_GridDestinationBuffer, in int3 GridIndex, in int DataOffset, in float FloatValue)
 {
 	int IntValue = asint(FloatValue);
 	int CompareValue = 0;
@@ -295,7 +296,7 @@ void DIVelocityGrid_InterlockedMinFloat(in FDIVelocityGridContext DIContext, in
  * -----------------------------------------------------------------
  */
 
- void DIVelocityGrid_InterlockedMinScalar(in FDIVelocityGridContext DIContext, in RWTexture3D<int> DIContext_GridDestinationBuffer, in int3 GridIndex, in int DataOffset, in float FloatValue)
+ void DIVelocityGrid_InterlockedMinScalar(in FDIVelocityGridContext DIContext, in RWBuffer<int> DIContext_GridDestinationBuffer, in int3 GridIndex, in int DataOffset, in float FloatValue)
 {
 #if FAST_MODE == 1
 	const int IntValue = FloatValue * FLOAT_PRECISION;
@@ -305,7 +306,7 @@ void DIVelocityGrid_InterlockedMinFloat(in FDIVelocityGridContext DIContext, in
 #endif
 }
 
-void DIVelocityGrid_InterlockedAddScalar(in FDIVelocityGridContext DIContext, in RWTexture3D<int> DIContext_GridDestinationBuffer, in int3 GridIndex, in int DataOffset, in float FloatValue)
+void DIVelocityGrid_InterlockedAddScalar(in FDIVelocityGridContext DIContext, in RWBuffer<int> DIContext_GridDestinationBuffer, in int3 GridIndex, in int DataOffset, in float FloatValue)
 {
 #if FAST_MODE == 1
 	const int IntValue = FloatValue * FLOAT_PRECISION;
@@ -326,7 +327,7 @@ float DIVelocityGrid_ConvertGridScalar(in FDIVelocityGridContext DIContext, in i
 #endif
 }
 
-void DIVelocityGrid_AddGridScalar(in FDIVelocityGridContext DIContext, in RWTexture3D<int> DIContext_GridDestinationBuffer, in int3 GridIndex, in int DataOffset, in float FloatValue)
+void DIVelocityGrid_AddGridScalar(in FDIVelocityGridContext DIContext, in RWBuffer<int> DIContext_GridDestinationBuffer, in int3 GridIndex, in int DataOffset, in float FloatValue)
 {
 	const float SumValue = asfloat(DIContext.GridCurrentBuffer[DIVelocityGrid_GetDataIndex(DIContext,GridIndex,DataOffset)])+FloatValue;
 	DIContext_GridDestinationBuffer[DIVelocityGrid_GetDataIndex(DIContext,GridIndex,DataOffset)] = asint(SumValue);
@@ -337,19 +338,19 @@ float DIVelocityGrid_GetGridScalar(in FDIVelocityGridContext DIContext, in int3
 	return asfloat(DIContext.GridCurrentBuffer[DIVelocityGrid_GetDataIndex(DIContext,GridIndex,DataOffset)]);
 }
 
-void DIVelocityGrid_SetGridScalar(in FDIVelocityGridContext DIContext, in RWTexture3D<int> DIContext_GridDestinationBuffer, in int3 GridIndex, in int DataOffset, in float FloatValue)
+void DIVelocityGrid_SetGridScalar(in FDIVelocityGridContext DIContext, in RWBuffer<int> DIContext_GridDestinationBuffer, in int3 GridIndex, in int DataOffset, in float FloatValue)
 {
 	DIContext_GridDestinationBuffer[DIVelocityGrid_GetDataIndex(DIContext,GridIndex,DataOffset)] = asint(FloatValue);
 }
 
-void DIVelocityGrid_SetGridVector( in FDIVelocityGridContext DIContext, in RWTexture3D<int> DIContext_GridDestinationBuffer, in int3 GridIndex, in int DataOffset, in float3 NodeVector)
+void DIVelocityGrid_SetGridVector( in FDIVelocityGridContext DIContext, in RWBuffer<int> DIContext_GridDestinationBuffer, in int3 GridIndex, in int DataOffset, in float3 NodeVector)
 {
 	DIVelocityGrid_SetGridScalar(DIContext,DIContext_GridDestinationBuffer,GridIndex,DataOffset,NodeVector.x);
 	DIVelocityGrid_SetGridScalar(DIContext,DIContext_GridDestinationBuffer,GridIndex,DataOffset+1,NodeVector.y);
 	DIVelocityGrid_SetGridScalar(DIContext,DIContext_GridDestinationBuffer,GridIndex,DataOffset+2,NodeVector.z);
 }
 
-void DIVelocityGrid_AddGridVector( in FDIVelocityGridContext DIContext, in RWTexture3D<int> DIContext_GridDestinationBuffer, in int3 GridIndex, in int DataOffset, in float3 NodeVector)
+void DIVelocityGrid_AddGridVector( in FDIVelocityGridContext DIContext, in RWBuffer<int> DIContext_GridDestinationBuffer, in int3 GridIndex, in int DataOffset, in float3 NodeVector)
 {
 	DIVelocityGrid_AddGridScalar(DIContext,DIContext_GridDestinationBuffer,GridIndex,DataOffset,NodeVector.x);
 	DIVelocityGrid_AddGridScalar(DIContext,DIContext_GridDestinationBuffer,GridIndex,DataOffset+1,NodeVector.y);
@@ -398,7 +399,7 @@ float3 DIVelocityGrid_GetGridVector( in FDIVelocityGridContext DIContext, in int
  */
  
 // Transfer the velocity field component on a face from the particle to the grid
-void DIVelocityGrid_BuildVelocityFace(in FDIVelocityGridContext DIContext, in RWTexture3D<int> DIContext_GridDestinationBuffer, in int3 GridIndex, in float ShapeFunction, in int3 ijk, in float3 GridPosition, in float SampleMass, in float SampleMomentum,
+void DIVelocityGrid_BuildVelocityFace(in FDIVelocityGridContext DIContext, in RWBuffer<int> DIContext_GridDestinationBuffer, in int3 GridIndex, in float ShapeFunction, in int3 ijk, in float3 GridPosition, in float SampleMass, in float SampleMomentum,
 			in float3 MomentumGradient, in float3 ComponentOffset, in int ComponentIndex, in float GridLength )
 {
 	const int3 GridLocation = int3(GridIndex.x+ijk.x,GridIndex.y+ijk.y,GridIndex.z+ijk.z);
@@ -414,7 +415,7 @@ void DIVelocityGrid_BuildVelocityFace(in FDIVelocityGridContext DIContext, in RW
 }
 
 // Transfer the velocity field component from the particle to the grid
-void DIVelocityGrid_BuildVelocityComponent(in FDIVelocityGridContext DIContext, in RWTexture3D<int> DIContext_GridDestinationBuffer, in float3 GridPosition, in float SampleMass, in float SampleMomentum,
+void DIVelocityGrid_BuildVelocityComponent(in FDIVelocityGridContext DIContext, in RWBuffer<int> DIContext_GridDestinationBuffer, in float3 GridPosition, in float SampleMass, in float SampleMomentum,
 			in float3 MomentumGradient, in float3 ComponentOffset, in int ComponentIndex, in float GridLength )
 {
 	float3 GridFraction = float3(0,0,0);
@@ -439,7 +440,7 @@ void DIVelocityGrid_BuildVelocityComponent(in FDIVelocityGridContext DIContext,
 }
 
 // Transfer the velocity field from the particle to the grid
-void DIVelocityGrid_BuildVelocityField(in FDIVelocityGridContext DIContext, in RWTexture3D<int> DIContext_GridDestinationBuffer, in float3 GridOrigin, in float GridLength, in float3 ParticlePosition, in float ParticleMass, in float3 ParticleVelocity, 
+void DIVelocityGrid_BuildVelocityField(in FDIVelocityGridContext DIContext, in RWBuffer<int> DIContext_GridDestinationBuffer, in float3 GridOrigin, in float GridLength, in float3 ParticlePosition, in float ParticleMass, in float3 ParticleVelocity, 
 			in float4x4 VelocityGradient, out bool OutFunctionStatus )
 {
 	const float3 ParticleMomentum = ParticleVelocity * ParticleMass;
@@ -553,7 +554,7 @@ void DIVelocityGrid_SampleVelocityField(in FDIVelocityGridContext DIContext, in
  * -----------------------------------------------------------------
  */
 
-void DIVelocityGrid_AddGridVelocity(in FDIVelocityGridContext DIContext, in RWTexture3D<int> DIContext_GridDestinationBuffer, in int GridHash, in float3 GridVelocity, out bool OutAddStatus)
+void DIVelocityGrid_AddGridVelocity(in FDIVelocityGridContext DIContext, in RWBuffer<int> DIContext_GridDestinationBuffer, in int GridHash, in float3 GridVelocity, out bool OutAddStatus)
 {
 	const int3 GridIndex = DIVelocityGrid_GetGridIndex(DIContext,GridHash);
 	OutAddStatus = true;
diff --git a/Engine/Plugins/Runtime/HairStrands/Source/HairStrandsCore/Private/GroomComponent.cpp b/Engine/Plugins/Runtime/HairStrands/Source/HairStrandsCore/Private/GroomComponent.cpp
--- a/Engine/Plugins/Runtime/HairStrands/Source/HairStrandsCore/Private/GroomComponent.cpp
+++ b/Engine/Plugins/Runtime/HairStrands/Source/HairStrandsCore/Private/GroomComponent.cpp
@@ -1410,7 +1410,6 @@ void UGroomComponent::CreateHairSimulation(const int32 GroupIndex, const int32 L
 					(GroomAsset->HairGroupsPhysics[GroupIndex].SolverSettings.NiagaraSolver == EGroomNiagaraSolvers::AngularSprings) ? ToRawPtr(AngularSpringsSystem) :
 					(GroomAsset->HairGroupsPhysics[GroupIndex].SolverSettings.NiagaraSolver == EGroomNiagaraSolvers::CosseratRods) ? ToRawPtr(CosseratRodsSystem) :
 					GroomAsset->HairGroupsPhysics[GroupIndex].SolverSettings.CustomSystem.LoadSynchronous();
-
 				NiagaraComponent->SetVisibleFlag(SimulationSettings.SimulationSetup.bDebugSimulation);
 				NiagaraComponent->SetAsset(NiagaraAsset);
 				NiagaraComponent->ReinitializeSystem();
@@ -1733,6 +1732,19 @@ void UGroomComponent::SetBindingAsset(UGroomBindingAsset* InBinding)
 	}
 }
 
+void UGroomComponent::SetPhysicsAsset(UPhysicsAsset* InPhysicsAsset)
+{
+	if (InPhysicsAsset && PhysicsAsset != InPhysicsAsset)
+	{
+		PhysicsAsset = InPhysicsAsset;
+			
+		ReleaseResources();
+		UpdateHairGroupsDesc();
+		UpdateHairSimulation();
+		InitResources();
+	}
+}
+
 void UGroomComponent::SetEnableSimulation(bool bInEnableSimulation)
 {
 	if (SimulationSettings.bOverrideSettings)
@@ -3182,6 +3194,10 @@ void UGroomComponent::BuildSimulationTransform(FTransform& SimulationTransform)
 
 			SimulationTransform = SkeletelMeshComponent->GetBoneTransform(BoneIndex);
 		}
+		else
+		{
+			SimulationTransform = GetComponentTransform();
+		}
 	}
 }
 
@@ -3325,6 +3341,11 @@ void UGroomComponent::TickComponent(float DeltaTime, enum ELevelTick TickType, F
 			}
 		}
 	});
+	
+	if(GetHairSwapBufferType() == EHairBufferSwapType::RenderFrame)
+	{
+		MarkRenderDynamicDataDirty();
+	}
 }
 
 void UGroomComponent::SendRenderTransform_Concurrent()
@@ -3340,6 +3361,25 @@ void UGroomComponent::SendRenderTransform_Concurrent()
 	Super::SendRenderTransform_Concurrent();
 }
 
+void UGroomComponent::SendRenderDynamicData_Concurrent()
+{
+	Super::SendRenderDynamicData_Concurrent();
+
+	if(GetHairSwapBufferType() == EHairBufferSwapType::RenderFrame)
+	{
+		TArray<FHairGroupInstance*> LocalHairGroupInstances = HairGroupInstances;
+		ENQUEUE_RENDER_COMMAND(FHairStrandsTick_TransformUpdate)(
+			[LocalHairGroupInstances](FRHICommandListImmediate& RHICmdList)
+		{
+			for (FHairGroupInstance* Instance : LocalHairGroupInstances)
+			{
+				if (Instance->Guides.DeformedResource)  { Instance->Guides.DeformedResource->SwapBuffer(); }
+				if (Instance->Strands.DeformedResource) { Instance->Strands.DeformedResource->SwapBuffer(); }
+			}
+		});
+	}
+}
+
 void UGroomComponent::GetUsedMaterials(TArray<UMaterialInterface*>& OutMaterials, bool bGetDebugMaterials) const
 {
 	UMeshComponent::GetUsedMaterials(OutMaterials, bGetDebugMaterials);
@@ -3483,7 +3523,7 @@ void UGroomComponent::PostEditChangeProperty(FPropertyChangedEvent& PropertyChan
 		{
 			// Set the callback on the new GroomAsset being assigned
 			GroomAsset->GetOnGroomAssetChanged().AddUObject(this, &UGroomComponent::Invalidate);
-			GroomAsset->GetOnGroomAssetResourcesChanged().AddUObject(this, &UGroomComponent::Invalidate);
+			GroomAsset->GetOnGroomAssetResourcesChanged().AddUObject(this, &UGroomComponent::InvalidateAndRecreate);
 			bIsGroomAssetCallbackRegistered = true;
 		}
 	}
diff --git a/Engine/Plugins/Runtime/HairStrands/Source/HairStrandsCore/Private/GroomManager.cpp b/Engine/Plugins/Runtime/HairStrands/Source/HairStrandsCore/Private/GroomManager.cpp
--- a/Engine/Plugins/Runtime/HairStrands/Source/HairStrandsCore/Private/GroomManager.cpp
+++ b/Engine/Plugins/Runtime/HairStrands/Source/HairStrandsCore/Private/GroomManager.cpp
@@ -25,7 +25,7 @@ static FAutoConsoleVariableRef CVarGHairStrandsMinLOD(TEXT("r.HairStrands.MinLOD
 static int32 GHairStrands_UseCards = 0;
 static FAutoConsoleVariableRef CVarHairStrands_UseCards(TEXT("r.HairStrands.UseCardsInsteadOfStrands"), GHairStrands_UseCards, TEXT("Force cards geometry on all groom elements. If no cards data is available, nothing will be displayed"), ECVF_Scalability);
 
-static int32 GHairStrands_SwapBufferType = 1;
+static int32 GHairStrands_SwapBufferType = 3;
 static FAutoConsoleVariableRef CVarGHairStrands_SwapBufferType(TEXT("r.HairStrands.SwapType"), GHairStrands_SwapBufferType, TEXT("Swap rendering buffer at the end of frame. This is an experimental toggle. Default:1"));
 
 static int32 GHairStrands_ManualSkinCache = 0;
@@ -41,6 +41,7 @@ EHairBufferSwapType GetHairSwapBufferType()
 	case 0: return EHairBufferSwapType::BeginOfFrame;
 	case 1: return EHairBufferSwapType::EndOfFrame;
 	case 2: return EHairBufferSwapType::Tick;
+	case 3: return EHairBufferSwapType::RenderFrame;
 	}
 	return EHairBufferSwapType::EndOfFrame;
 }
diff --git a/Engine/Plugins/Runtime/HairStrands/Source/HairStrandsCore/Private/GroomManager.h b/Engine/Plugins/Runtime/HairStrands/Source/HairStrandsCore/Private/GroomManager.h
--- a/Engine/Plugins/Runtime/HairStrands/Source/HairStrandsCore/Private/GroomManager.h
+++ b/Engine/Plugins/Runtime/HairStrands/Source/HairStrandsCore/Private/GroomManager.h
@@ -52,7 +52,8 @@ enum EHairBufferSwapType
 {
 	BeginOfFrame,
 	EndOfFrame,
-	Tick
+	Tick,
+	RenderFrame
 };
 
 EHairBufferSwapType GetHairSwapBufferType();
diff --git a/Engine/Plugins/Runtime/HairStrands/Source/HairStrandsCore/Private/Niagara/NiagaraDataInterfaceHairStrands.cpp b/Engine/Plugins/Runtime/HairStrands/Source/HairStrandsCore/Private/Niagara/NiagaraDataInterfaceHairStrands.cpp
--- a/Engine/Plugins/Runtime/HairStrands/Source/HairStrandsCore/Private/Niagara/NiagaraDataInterfaceHairStrands.cpp
+++ b/Engine/Plugins/Runtime/HairStrands/Source/HairStrandsCore/Private/Niagara/NiagaraDataInterfaceHairStrands.cpp
@@ -323,6 +323,7 @@ void FNDIHairStrandsBuffer::Initialize(
 	SourceDeformedRootResources = HairStrandsDeformedRootResource;
 	ParamsScale = InParamsScale;
 	BoundingBoxOffsets = FIntVector4(0,1,2,3);
+
 	bValidGeometryType = false;
 }
 
@@ -511,6 +512,13 @@ void FNDIHairStrandsData::Update(UNiagaraDataInterfaceHairStrands* Interface, FN
 				LocalSimulation = SimulationSettings.SimulationSetup.bLocalSimulation;
 				Interface->SourceComponent->BuildSimulationTransform(BoneTransform);
 
+				// Convert to double for LWC
+				FTransform3d BoneTransformDouble = BoneTransform;
+				const FTransform3d WorldTransformDouble = WorldTransform;
+				
+				BoneTransformDouble = BoneTransformDouble * WorldTransformDouble.Inverse();
+				BoneTransform = BoneTransformDouble;
+
 				if (SimulationSettings.bOverrideSettings)
 				{
 					GravityVector = SimulationSettings.ExternalForces.GravityVector;
@@ -648,7 +656,7 @@ struct FNDIHairStrandsParametersCS : public FNiagaraDataInterfaceParametersCS
 			// TEMP: These check are only temporary for avoiding crashes while we find the bottom of the issue.
 			ProxyData->HairStrandsBuffer->CurvesOffsetsBuffer.SRV && ProxyData->HairStrandsBuffer->ParamsScaleBuffer.SRV && ProxyData->HairStrandsBuffer->BoundingBoxBuffer.UAV;
 
-		const bool bIsGeometryValid = bIsHairValid&& ProxyData->HairGroupInstance && (ProxyData->HairGroupInstance->GeometryType != EHairGeometryType::NoneGeometry);
+		const bool bIsGeometryValid = bIsHairValid && (!ProxyData->HairGroupInstance || (ProxyData->HairGroupInstance && (ProxyData->HairGroupInstance->GeometryType != EHairGeometryType::NoneGeometry)));
 		const bool bIsDeformedValid = bIsHairValid && ProxyData->HairStrandsBuffer->SourceDeformedResources && ProxyData->HairStrandsBuffer->SourceDeformedResources->IsInitialized();
 
 		if (bIsHairValid && bIsRestValid && bIsGeometryValid)
@@ -707,19 +715,21 @@ struct FNDIHairStrandsParametersCS : public FNiagaraDataInterfaceParametersCS
 			// Offsets / Transforms
 			FVector3f RestPositionOffsetValue = ProxyData->HairStrandsBuffer->SourceRestResources->GetPositionOffset();
 
+			FMatrix44f RigidTransformFloat = ProxyData->HairGroupInstance ? ProxyData->HairGroupInstance->Debug.RigidCurrentLocalToWorld.ToMatrixWithScale():
+																			ProxyData->WorldTransform.ToMatrixWithScale();
 			FMatrix44f WorldTransformFloat = ProxyData->HairGroupInstance ? ProxyData->HairGroupInstance->GetCurrentLocalToWorld().ToMatrixWithScale() :
 																			ProxyData->WorldTransform.ToMatrixWithScale();
-			FMatrix44f BoneTransformFloat = ProxyData->BoneTransform.ToMatrixWithScale();
-
+			FMatrix44f BoneTransformFloat = ProxyData->BoneTransform.ToMatrixWithScale() * RigidTransformFloat;
+			
 			if (ProxyData->LocalSimulation)
 			{
-				FMatrix44d WorldTransformDouble = WorldTransformFloat;
-				FMatrix44d BoneTransformDouble = BoneTransformFloat;
+				const FMatrix44d WorldTransformDouble = WorldTransformFloat;
+				const FMatrix44d BoneTransformDouble = BoneTransformFloat;
 
 				// Due to large world coordinate we store the relative world transform in double precision
 				WorldTransformFloat = WorldTransformDouble * BoneTransformDouble.Inverse();
 			}
-
+			
 			if (!bIsRootValid && bHasSkinningBinding)
 			{
 				UE_LOG(LogHairStrands, Log, TEXT("FNDIHairStrandsParametersCS() Groom Asset %s from component %s is set to use skinning interpolation but the skin resources are not valid"),
@@ -778,7 +788,6 @@ struct FNDIHairStrandsParametersCS : public FNiagaraDataInterfaceParametersCS
 			{
 				ProxyData->HairStrandsBuffer->bValidGeometryType = false;
 			}
-
 			// Set shader constants
 			SetShaderValue(RHICmdList, ComputeShaderRHI, BoundingBoxOffsets, FIntVector4(0,1,2,3));
 			SetShaderValue(RHICmdList, ComputeShaderRHI, WorldTransform, FMatrix44f::Identity);
diff --git a/Engine/Plugins/Runtime/HairStrands/Source/HairStrandsCore/Private/Niagara/NiagaraDataInterfacePhysicsAsset.cpp b/Engine/Plugins/Runtime/HairStrands/Source/HairStrandsCore/Private/Niagara/NiagaraDataInterfacePhysicsAsset.cpp
--- a/Engine/Plugins/Runtime/HairStrands/Source/HairStrandsCore/Private/Niagara/NiagaraDataInterfacePhysicsAsset.cpp
+++ b/Engine/Plugins/Runtime/HairStrands/Source/HairStrandsCore/Private/Niagara/NiagaraDataInterfacePhysicsAsset.cpp
@@ -781,7 +781,7 @@ void UNiagaraDataInterfacePhysicsAsset::ExtractSourceComponent(FNiagaraSystemIns
 	}
 	else if (USceneComponent* AttachComponent = SystemInstance->GetAttachComponent())
 	{
-		// Try to find the component by walking the attachment hierarchy
+		// Try to find the first component by walking the attachment hierarchy
 		for (USceneComponent* Curr = AttachComponent; Curr; Curr = Curr->GetAttachParent())
 		{
 			USkeletalMeshComponent* SkelMeshComp = Cast<USkeletalMeshComponent>(Curr);
@@ -801,7 +801,10 @@ void UNiagaraDataInterfacePhysicsAsset::ExtractSourceComponent(FNiagaraSystemIns
 			}
 		}
 	}
-
+	
+	SourceComponents.Empty();
+	PhysicsAssets.Empty();
+	
 	// Try to find the groom physics asset by walking the attachment hierarchy
 	UPhysicsAsset* GroomPhysicsAsset = DefaultSource;
 	for (USceneComponent* Curr = SystemInstance->GetAttachComponent(); Curr; Curr = Curr->GetAttachParent())
@@ -814,13 +817,22 @@ void UNiagaraDataInterfacePhysicsAsset::ExtractSourceComponent(FNiagaraSystemIns
 			if (GroomComponent->PhysicsAsset)
 			{
 				GroomPhysicsAsset = GroomComponent->PhysicsAsset;
-				break;
 			}
+			TArray<UPrimitiveComponent*> OverlappingComponents;
+			GroomComponent->GetOverlappingComponents(OverlappingComponents);
+
+			for(auto& OverlappingComponent : OverlappingComponents)
+			{
+				if (USkeletalMeshComponent* SkelMeshComp = Cast<USkeletalMeshComponent>(OverlappingComponent))
+				{
+					SourceComponents.Add(SkelMeshComp);
+					PhysicsAssets.Add(SkelMeshComp->GetPhysicsAsset());
+				}
+			}
+			break;
 		}
 	}
-
-	SourceComponents.Empty();
-	PhysicsAssets.Empty();
+	
 	if (SourceComponent != nullptr)
 	{
 		SourceComponents.Add(SourceComponent);
@@ -831,25 +843,6 @@ void UNiagaraDataInterfacePhysicsAsset::ExtractSourceComponent(FNiagaraSystemIns
 		else
 		{
 			PhysicsAssets.Add(SourceComponent->GetPhysicsAsset());
-
-			if (USkeletalMeshComponent* ParentComp = Cast<USkeletalMeshComponent>(SourceComponent->GetAttachParent()))
-			{
-				SourceComponents.Add(ParentComp);
-				PhysicsAssets.Add(ParentComp->GetPhysicsAsset());
-
-				TArray<USceneComponent*> SceneComponents;
-				ParentComp->GetChildrenComponents(true, SceneComponents);
-
-				for (USceneComponent* ActorComp : SceneComponents)
-				{
-					USkeletalMeshComponent* SourceComp = Cast<USkeletalMeshComponent>(ActorComp);
-					if (SourceComp && SourceComp->SkeletalMesh && SourceComp->GetPhysicsAsset() && SourceComp != SourceComponent)
-					{
-						SourceComponents.Add(SourceComp);
-						PhysicsAssets.Add(SourceComp->GetPhysicsAsset());
-					}
-				}
-			}
 		}
 	}
 	else if (GroomPhysicsAsset != nullptr)
diff --git a/Engine/Plugins/Runtime/HairStrands/Source/HairStrandsCore/Private/Niagara/NiagaraDataInterfacePressureGrid.cpp b/Engine/Plugins/Runtime/HairStrands/Source/HairStrandsCore/Private/Niagara/NiagaraDataInterfacePressureGrid.cpp
--- a/Engine/Plugins/Runtime/HairStrands/Source/HairStrandsCore/Private/Niagara/NiagaraDataInterfacePressureGrid.cpp
+++ b/Engine/Plugins/Runtime/HairStrands/Source/HairStrandsCore/Private/Niagara/NiagaraDataInterfacePressureGrid.cpp
@@ -440,7 +440,7 @@ void UNiagaraDataInterfacePressureGrid::GetParameterDefinitionHLSL(const FNiagar
 
 //------------------------------------------------------------------------------------------------------------
 
-#define NIAGARA_HAIR_STRANDS_THREAD_COUNT_PRESSURE 64
+#define NIAGARA_HAIR_STRANDS_THREAD_COUNT_PRESSURE 4
 
 class FClearPressureGridCS : public FGlobalShader
 {
@@ -519,11 +519,13 @@ inline void ClearBuffer(FRHICommandList& RHICmdList, ERHIFeatureLevel::Type Feat
 
 		const uint32 GroupSize = NIAGARA_HAIR_STRANDS_THREAD_COUNT_PRESSURE;
 		const uint32 NumElements = (GridSize.X + 1) * (GridSize.Y + 1) * (GridSize.Z + 1);
-
-		const uint32 DispatchCount = FMath::DivideAndRoundUp(NumElements, GroupSize);
+		
+		const uint32 DispatchCountX = FMath::DivideAndRoundUp((uint32)(GridSize.X+1), GroupSize);
+		const uint32 DispatchCountY = FMath::DivideAndRoundUp((uint32)(GridSize.Y+1), GroupSize);
+		const uint32 DispatchCountZ = FMath::DivideAndRoundUp((uint32)(GridSize.Z+1), GroupSize);
 
 		ComputeShader->SetParameters(RHICmdList, CurrentGridBufferSRV, DestinationGridBufferUAV, GridSize, CopyPressure);
-		DispatchComputeShader(RHICmdList, ComputeShader.GetShader(), DispatchCount, 1, 1);
+		DispatchComputeShader(RHICmdList, ComputeShader.GetShader(), DispatchCountX, DispatchCountY, DispatchCountZ);
 		ComputeShader->UnsetParameters(RHICmdList);
 	}
 }
diff --git a/Engine/Plugins/Runtime/HairStrands/Source/HairStrandsCore/Private/Niagara/NiagaraDataInterfaceVelocityGrid.cpp b/Engine/Plugins/Runtime/HairStrands/Source/HairStrandsCore/Private/Niagara/NiagaraDataInterfaceVelocityGrid.cpp
--- a/Engine/Plugins/Runtime/HairStrands/Source/HairStrandsCore/Private/Niagara/NiagaraDataInterfaceVelocityGrid.cpp
+++ b/Engine/Plugins/Runtime/HairStrands/Source/HairStrandsCore/Private/Niagara/NiagaraDataInterfaceVelocityGrid.cpp
@@ -61,8 +61,8 @@ void FNDIVelocityGridBuffer::InitRHI()
 	if (GridSize.X != 0 && GridSize.Y != 0 && GridSize.Z != 0)
 	{
 		static const uint32 NumComponents = NumAttributes;
-		GridDataBuffer.Initialize(TEXT("FNDIVelocityGridBuffer"), sizeof(int32), (GridSize.X + 1)*NumComponents, (GridSize.Y + 1),
-			(GridSize.Z + 1), EPixelFormat::PF_R32_SINT);
+		GridDataBuffer.Initialize(TEXT("FNDIVelocityGridBuffer"), sizeof(int32),
+			(GridSize.X + 1)*NumComponents * (GridSize.Y + 1) * (GridSize.Z + 1), EPixelFormat::PF_R32_SINT);
 	}
 }
 
diff --git a/Engine/Plugins/Runtime/HairStrands/Source/HairStrandsCore/Public/GroomAssetPhysics.h b/Engine/Plugins/Runtime/HairStrands/Source/HairStrandsCore/Public/GroomAssetPhysics.h
--- a/Engine/Plugins/Runtime/HairStrands/Source/HairStrandsCore/Public/GroomAssetPhysics.h
+++ b/Engine/Plugins/Runtime/HairStrands/Source/HairStrandsCore/Public/GroomAssetPhysics.h
@@ -356,11 +356,11 @@ struct HAIRSTRANDSCORE_API FHairSimulationSetup
 
 	/** Strands simulation is done in local space */
 	UPROPERTY(EditAnywhere, BlueprintReadWrite, Interp, Category = "SimulationSetup", meta = (ToolTip = "Strands simulation is done in local space"))
-	bool bLocalSimulation = false;
+	bool bLocalSimulation = true;
 
 	/** Bone used for the simulation local space */
 	UPROPERTY(EditAnywhere, BlueprintReadWrite, Interp, Category = "SimulationSetup", meta = (ToolTip = "Bone used for the simulation local space", EditCondition = "bLocalSimulation"))
-	FString LocalBone;
+	FString LocalBone = "root";
 
 	/** Teleport distance threshold to reset the simulation */
 	UPROPERTY(EditAnywhere, BlueprintReadWrite, Interp, Category = "SimulationSetup", meta = (ToolTip = "Teleport distance threshold to reset the simulation"))
diff --git a/Engine/Plugins/Runtime/HairStrands/Source/HairStrandsCore/Public/GroomComponent.h b/Engine/Plugins/Runtime/HairStrands/Source/HairStrandsCore/Public/GroomComponent.h
--- a/Engine/Plugins/Runtime/HairStrands/Source/HairStrandsCore/Public/GroomComponent.h
+++ b/Engine/Plugins/Runtime/HairStrands/Source/HairStrandsCore/Public/GroomComponent.h
@@ -118,6 +118,7 @@ public:
 	virtual void DetachFromComponent(const FDetachmentTransformRules& DetachmentRules) override;
 	virtual void TickComponent(float DeltaTime, enum ELevelTick TickType, FActorComponentTickFunction *ThisTickFunction) override;
 	virtual void SendRenderTransform_Concurrent() override;
+	virtual void SendRenderDynamicData_Concurrent() override;
 	//~ End UActorComponent Interface.
 
 	//~ Begin USceneComponent Interface.
@@ -164,6 +165,10 @@ public:
 	UFUNCTION(BlueprintCallable, Category = "Groom")
 	void SetBindingAsset(UGroomBindingAsset* InBinding);
 
+	/* Accessor function for changing Groom physics asset from blueprint/sequencer */
+	UFUNCTION(BlueprintCallable, Category = "Simulation")
+	void SetPhysicsAsset(UPhysicsAsset* InPhysicsAsset);
+
 	/* Accessor function for changing the enable simulation flag from blueprint/sequencer */
 	UFUNCTION(BlueprintCallable, Category = "Simulation")
 	void SetEnableSimulation(bool bInEnableSimulation);
diff --git a/Engine/Plugins/Runtime/HairStrands/Source/HairStrandsCore/Public/Niagara/NiagaraDataInterfaceVelocityGrid.h b/Engine/Plugins/Runtime/HairStrands/Source/HairStrandsCore/Public/Niagara/NiagaraDataInterfaceVelocityGrid.h
--- a/Engine/Plugins/Runtime/HairStrands/Source/HairStrandsCore/Public/Niagara/NiagaraDataInterfaceVelocityGrid.h
+++ b/Engine/Plugins/Runtime/HairStrands/Source/HairStrandsCore/Public/Niagara/NiagaraDataInterfaceVelocityGrid.h
@@ -25,7 +25,7 @@ struct FNDIVelocityGridBuffer : public FRenderResource
 	virtual FString GetFriendlyName() const override { return TEXT("FNDIVelocityGridBuffer"); }
 
 	/** Grid data texture */
-	FTextureRWBuffer3D GridDataBuffer;
+	FRWBuffer GridDataBuffer;
 
 	/** Grid size that will be used for the collision*/
 	FIntVector GridSize;
