Fix double conversion in hair DI

diff --git a/Engine/Plugins/Runtime/HairStrands/Source/HairStrandsCore/Private/Niagara/NiagaraDataInterfaceHairStrands.cpp b/Engine/Plugins/Runtime/HairStrands/Source/HairStrandsCore/Private/Niagara/NiagaraDataInterfaceHairStrands.cpp
--- a/Engine/Plugins/Runtime/HairStrands/Source/HairStrandsCore/Private/Niagara/NiagaraDataInterfaceHairStrands.cpp
+++ b/Engine/Plugins/Runtime/HairStrands/Source/HairStrandsCore/Private/Niagara/NiagaraDataInterfaceHairStrands.cpp
@@ -513,11 +513,12 @@ void FNDIHairStrandsData::Update(UNiagaraDataInterfaceHairStrands* Interface, FN
 				Interface->SourceComponent->BuildSimulationTransform(BoneTransform);
 
 				// Convert to double for LWC
-				FTransform3d BoneTransformDouble = BoneTransform;
-				const FTransform3d WorldTransformDouble = WorldTransform;
+				FMatrix44d BoneTransformDouble = BoneTransform.ToMatrixWithScale();
+				const FMatrix44d WorldTransformDouble = WorldTransform.ToMatrixWithScale();
 				
 				BoneTransformDouble = BoneTransformDouble * WorldTransformDouble.Inverse();
-				BoneTransform = BoneTransformDouble;
+				const FMatrix44d WorldTransformFloat = BoneTransformDouble;
+				BoneTransform = FTransform(WorldTransformFloat);
 
 				if (SimulationSettings.bOverrideSettings)
 				{
