Fix crash when loading MetaHumanSample on Linux/Mac

This is a tentative/short term fix.

diff --git a/Engine/Plugins/Runtime/HairStrands/Source/HairStrandsCore/Private/HairCardsBuilder.cpp b/Engine/Plugins/Runtime/HairStrands/Source/HairStrandsCore/Private/HairCardsBuilder.cpp
--- a/Engine/Plugins/Runtime/HairStrands/Source/HairStrandsCore/Private/HairCardsBuilder.cpp
+++ b/Engine/Plugins/Runtime/HairStrands/Source/HairStrandsCore/Private/HairCardsBuilder.cpp
@@ -3864,7 +3864,7 @@ bool InternalImportGeometry(
 
 				FCardsRootData& Cards = CardsRoots.AddDefaulted_GetRef();
 				Cards.CardsIndex = CardIt;
-				Cards.Position = AverageRootPosition / AverageCount;
+				Cards.Position = AverageCount > 0 ? AverageRootPosition / AverageCount : FVector::ZeroVector;
 			}
 		}
 
@@ -3887,10 +3887,13 @@ bool InternalImportGeometry(
 				}
 
 				// 3.2 Apply root UV to all cards vertices
-				check(CurveIndex != ~0);
-				if (CurveIndex != ~0)
 				{
-					const FVector2D RootUV = InStrandsData.StrandsCurves.CurvesRootUV[CurveIndex];
+					FVector2D RootUV = FVector2D::ZeroVector;
+					const bool bFoundValidCurve = CurveIndex != ~0;
+					if (bFoundValidCurve)
+					{
+						RootUV = InStrandsData.StrandsCurves.CurvesRootUV[CurveIndex];
+					}
 
 					const uint32 CardsIndexCount = Out.Cards.IndexCounts[CardsRoot.CardsIndex];
 					const uint32 CardsIndexOffset = Out.Cards.IndexOffsets[CardsRoot.CardsIndex];
