Change some type in hair code from double to float where it was not needed.

diff --git a/Engine/Plugins/Runtime/HairStrands/Source/HairStrandsCore/Private/GroomBindingBuilder.cpp b/Engine/Plugins/Runtime/HairStrands/Source/HairStrandsCore/Private/GroomBindingBuilder.cpp
--- a/Engine/Plugins/Runtime/HairStrands/Source/HairStrandsCore/Private/GroomBindingBuilder.cpp
+++ b/Engine/Plugins/Runtime/HairStrands/Source/HairStrandsCore/Private/GroomBindingBuilder.cpp
@@ -75,12 +75,12 @@ void FHairStrandsRootUtils::DecodeTriangleIndex(uint32 Encoded, uint32& OutTrian
 	OutTriangleIndex = Encoded & 0xFFFFFF;
 }
 
-uint32 FHairStrandsRootUtils::EncodeBarycentrics(const FVector2D& B)
+uint32 FHairStrandsRootUtils::EncodeBarycentrics(const FVector2f& B)
 {
 	return uint32(FFloat16(B.X).Encoded) | (uint32(FFloat16(B.Y).Encoded)<<16);
 }
 
-FVector2D FHairStrandsRootUtils::DecodeBarycentrics(uint32 B)
+FVector2f FHairStrandsRootUtils::DecodeBarycentrics(uint32 B)
 {
 	FFloat16 BX;
 	BX.Encoded = (B & 0xFFFF);
@@ -88,15 +88,15 @@ FVector2D FHairStrandsRootUtils::DecodeBarycentrics(uint32 B)
 	FFloat16 BY;
 	BY.Encoded = (B >> 16) & 0xFFFF;
 
-	return FVector2D(BX, BY);
+	return FVector2f(BX, BY);
 }
 
-uint32 FHairStrandsRootUtils::PackUVs(const FVector2D& UV)
+uint32 FHairStrandsRootUtils::PackUVs(const FVector2f& UV)
 {
 	return (FFloat16(UV.X).Encoded & 0xFFFF) | ((FFloat16(UV.Y).Encoded & 0xFFFF) << 16);
 }
 
-float FHairStrandsRootUtils::PackUVsToFloat(const FVector2D& UV)
+float FHairStrandsRootUtils::PackUVsToFloat(const FVector2f& UV)
 {
 	uint32 Encoded = PackUVs(UV);
 	return *((float*)(&Encoded));
diff --git a/Engine/Plugins/Runtime/HairStrands/Source/HairStrandsCore/Private/HairCardsBuilder.cpp b/Engine/Plugins/Runtime/HairStrands/Source/HairStrandsCore/Private/HairCardsBuilder.cpp
--- a/Engine/Plugins/Runtime/HairStrands/Source/HairStrandsCore/Private/HairCardsBuilder.cpp
+++ b/Engine/Plugins/Runtime/HairStrands/Source/HairStrandsCore/Private/HairCardsBuilder.cpp
@@ -68,6 +68,13 @@ static FAutoConsoleVariableRef CVarHairCardsAtlasMaxSample(TEXT("r.HairStrands.C
 
 ///////////////////////////////////////////////////////////////////////////////////////////////////
 
+static FBox ToFBox3d(const FBox3f& In)
+{
+	return FBox(FVector(In.Min), FVector(In.Max));
+}
+
+///////////////////////////////////////////////////////////////////////////////////////////////////
+
 class FHairCardAtlasTextureRectVS : public FGlobalShader
 {
 	DECLARE_GLOBAL_SHADER(FHairCardAtlasTextureRectVS);
@@ -1388,7 +1395,7 @@ namespace HairCards
 		TArray<uint32>  GuidePointCount;
 		TArray<uint32>  GuidePointOffset;
 		TArray<FVector3f> GuidePoints;
-		TArray<FVector2D> GuideRootUVs;
+		TArray<FVector2f> GuideRootUVs;
 
 		TArray<FHairOrientedBound> Bounds;
 
@@ -1630,7 +1637,7 @@ namespace HairCards
 		{
 			const uint32 SimCount  = SimStrandsData.StrandsCurves.CurvesCount[GuideIndex];
 			const uint32 SimOffset = SimStrandsData.StrandsCurves.CurvesOffset[GuideIndex];
-			const FVector2D SimRootUV = SimStrandsData.StrandsCurves.CurvesRootUV[GuideIndex];
+			const FVector2f SimRootUV = SimStrandsData.StrandsCurves.CurvesRootUV[GuideIndex];
 
 			const uint32 ClusterIt = GuideIndexToClusterIndex[GuideIndex];
 			if (GuideIndexToClusterIndex[GuideIndex] >= 0)
@@ -2007,7 +2014,7 @@ namespace HairCards
 			FVector3f RootCenter = FVector3f::ZeroVector;
 			uint32 LongestCurveIndex = 0;
 			float LongestLength = 0;
-			FVector2D RootUV(0, 0);
+			FVector2f RootUV(0, 0);
 			TArray<FVector3f> Points;
 			Points.Reserve(CurveCount * 20); // Guess number for avoiding too much reallocation
 			for (uint32 StrandIt = 0; StrandIt < CurveCount; ++StrandIt)
@@ -2046,7 +2053,7 @@ namespace HairCards
 			{
 				const uint32 GuidePointCount  = Out.GuidePointCount[ClusterIt];
 				const uint32 GuidePointOffset = Out.GuidePointOffset[ClusterIt];
-				const FVector2D GuideRootUV = Out.GuideRootUVs[ClusterIt];
+				const FVector2f GuideRootUV = Out.GuideRootUVs[ClusterIt];
 				TArray<FVector3f> BoundPoints;
 				TArray<float> BoundLengths;
 				float TotalLength = 0;
@@ -2291,9 +2298,9 @@ namespace HairCards
 						}
 					}
 
-					const FVector2D& RootUV = InClusters.GuideRootUVs[ClusterIt];
-					OutCards.UVs[VertexIt    ] = FVector4(CardLength, 0, RootUV.X, RootUV.Y);
-					OutCards.UVs[VertexIt + 1] = FVector4(CardLength, 1, RootUV.X, RootUV.Y);
+					const FVector2f& RootUV = InClusters.GuideRootUVs[ClusterIt];
+					OutCards.UVs[VertexIt    ] = FVector4f(CardLength, 0, RootUV.X, RootUV.Y);
+					OutCards.UVs[VertexIt + 1] = FVector4f(CardLength, 1, RootUV.X, RootUV.Y);
 
 					OutCards.CardIndices[VertexIt]	   = CardIndex;
 					OutCards.CardIndices[VertexIt + 1] = CardIndex;
@@ -2367,7 +2374,7 @@ namespace HairCards
 						OutCards.CoordU[VIndex] = OutCards.UVs[VIndex].X;
 						
 						// Absolute/Global UVs
-						const FVector2D AtlasCoord = FVector2D(OutCards.UVs[VIndex].X, OutCards.UVs[VIndex].Y) * Rect.Resolution + FVector2D(Rect.Offset.X, Rect.Offset.Y) + FVector2D(0.5f, 0.5f);
+						const FVector2f AtlasCoord = FVector2f(OutCards.UVs[VIndex].X, OutCards.UVs[VIndex].Y) * Rect.Resolution + FVector2f(Rect.Offset.X, Rect.Offset.Y) + FVector2f(0.5f, 0.5f);
 						OutCards.UVs[VIndex].X = AtlasCoord.X / AtlasResolution.X;
 						OutCards.UVs[VIndex].Y = AtlasCoord.Y / AtlasResolution.Y;
 					}
@@ -2772,7 +2779,7 @@ namespace HairCards
 			OutGuides.StrandsCurves.CurvesCount[ClusterIt] = PointCount;
 			OutGuides.StrandsCurves.CurvesOffset[ClusterIt] = PointOffset;
 			OutGuides.StrandsCurves.CurvesLength[ClusterIt] = TotalLength;
-			OutGuides.StrandsCurves.CurvesRootUV[ClusterIt] = FVector2D(0, 0);
+			OutGuides.StrandsCurves.CurvesRootUV[ClusterIt] = FVector2f(0, 0);
 			//OutGuides.StrandsCurves.StrandIDs[ClusterIt] = ClusterIt;
 			//OutGuides.StrandsCurves.GroomIDToIndex[ClusterIt] = ;
 			OutGuides.StrandsCurves.MaxLength = FMath::Max(OutGuides.StrandsCurves.MaxLength, TotalLength);
@@ -2835,7 +2842,7 @@ namespace HairCards
 
 		const float GuideRadius = 0.01f;
 
-		FVector4 TriangleUVs[3];
+		FVector4f TriangleUVs[3];
 		for (uint32 CardIt = 0; CardIt < NumCards; ++CardIt)
 		{
 			const uint32 NumTriangles = InCards.IndexCounts[CardIt] / 3;
@@ -2882,7 +2889,7 @@ namespace HairCards
 				for (uint32 VertexIt = 0; VertexIt < 3; ++VertexIt)
 				{
 					const uint32 VertexIndex = InCards.Indices[VertexIndexOffset + VertexIt];
-					const FVector4 UV = TriangleUVs[VertexIt] = InCards.UVs[VertexIndex];
+					const FVector4f UV = TriangleUVs[VertexIt] = InCards.UVs[VertexIndex];
 					AddSimilarUV(SimilarVertexU, VertexIndex, UV.X, UVCoordTreshold);
 					AddSimilarUV(SimilarVertexV, VertexIndex, UV.Y, UVCoordTreshold);
 				}
@@ -2967,7 +2974,7 @@ namespace HairCards
 			OutGuides.StrandsCurves.CurvesCount[CardIt] = PointCount;
 			OutGuides.StrandsCurves.CurvesOffset[CardIt] = PointOffset;
 			OutGuides.StrandsCurves.CurvesLength[CardIt] = TotalLength;
-			OutGuides.StrandsCurves.CurvesRootUV[CardIt] = FVector2D(0, 0);
+			OutGuides.StrandsCurves.CurvesRootUV[CardIt] = FVector2f(0, 0);
 			OutGuides.StrandsCurves.MaxLength = FMath::Max(OutGuides.StrandsCurves.MaxLength, TotalLength);
 			OutGuides.StrandsCurves.MaxRadius = GuideRadius;
 
@@ -3163,7 +3170,7 @@ namespace HairCards
 					const uint32 PointIndex = PointOffset + PointIt;
 
 					const float StrandsU = InStrands.StrandsPoints.PointsCoordU[PointIndex];
-					const FVector2D RootUV = InStrands.StrandsCurves.CurvesRootUV[CurveIndex];
+					const FVector2f RootUV = InStrands.StrandsCurves.CurvesRootUV[CurveIndex];
 					const float Seed = 0.f; // TODO float(InStrands.RenderData.Attributes[PointIndex].Seed) / 255.f;
 
 					const FVector3f& P0_Rest = InStrands.StrandsPoints.PointsPosition[PointIndex];
@@ -3222,8 +3229,8 @@ namespace HairCards
 					OutRect.MaxBound.X = FMath::Max(P0_Deformed.X, OutRect.MaxBound.X);
 					OutRect.MaxBound.Y = FMath::Max(P0_Deformed.Y, OutRect.MaxBound.Y);
 					OutRect.MaxBound.Z = FMath::Max(P0_Deformed.Z, OutRect.MaxBound.Z);
-					InAtlas.StrandsPositions.Add(FVector4(P0_Deformed, (PointIt < PointCount - 1) ? Radius : 0));
-					InAtlas.StrandsAttributes.Add(FVector4(RootUV.X, RootUV.Y, StrandsU, Seed));
+					InAtlas.StrandsPositions.Add(FVector4f(P0_Deformed, (PointIt < PointCount - 1) ? Radius : 0));
+					InAtlas.StrandsAttributes.Add(FVector4f(RootUV.X, RootUV.Y, StrandsU, Seed));
 						
 					++VertexCount;
 				}
@@ -3458,12 +3465,12 @@ void BuildGeometry(
 	for (uint32 PointIt = 0; PointIt < PointCount; ++PointIt)
 	{
 		check(FMath::IsFinite(Out.Cards.Positions[PointIt].X) && FMath::IsFinite(Out.Cards.Positions[PointIt].Y) && FMath::IsFinite(Out.Cards.Positions[PointIt].Z));
-		Out.RenderData.Positions[PointIt] = FVector4(Out.Cards.Positions[PointIt], *((float*)&Out.Cards.CardIndices[PointIt]));
+		Out.RenderData.Positions[PointIt] = FVector4f(Out.Cards.Positions[PointIt], *((float*)&Out.Cards.CardIndices[PointIt]));
 		Out.RenderData.UVs[PointIt] = Out.Cards.UVs[PointIt];
-		Out.RenderData.Normals[PointIt * 2] = FVector4(Out.Cards.Tangents[PointIt], 0);
-		Out.RenderData.Normals[PointIt * 2 + 1] = FVector4(Out.Cards.Normals[PointIt], 1);
+		Out.RenderData.Normals[PointIt * 2] = FVector4f(Out.Cards.Tangents[PointIt], 0);
+		Out.RenderData.Normals[PointIt * 2 + 1] = FVector4f(Out.Cards.Normals[PointIt], 1);
 
-		Out.Cards.BoundingBox += FVector4(Out.RenderData.Positions[PointIt]);
+		Out.Cards.BoundingBox += FVector4f(Out.RenderData.Positions[PointIt]);
 	}
 
 	const uint32 IndexCount = Out.Cards.Indices.Num();
@@ -3478,7 +3485,7 @@ void BuildGeometry(
 	OutBulk.Normals 	= Out.RenderData.Normals;
 	OutBulk.Indices 	= Out.RenderData.Indices;
 	OutBulk.UVs 		= Out.RenderData.UVs;
-	OutBulk.BoundingBox = Out.Cards.BoundingBox;
+	OutBulk.BoundingBox = ToFBox3d(Out.Cards.BoundingBox);
 
 	const uint32 CardCount = Out.Cards.Lengths.Num();
 	Out.RenderData.CardsRect.SetNum(CardCount);
@@ -3690,7 +3697,7 @@ bool InternalImportGeometry(
 		const uint32 VertexIndex = VertexId.GetValue();
 		check(VertexIndex < VertexCount);
 		Out.Cards.Positions[VertexIndex]	= VertexPositions[VertexId];
-		Out.Cards.UVs[VertexIndex]			= FVector4(VertexInstanceUVs[VertexInstanceId0].Component(0), VertexInstanceUVs[VertexInstanceId0].Component(1), 0, 0); // RootUV are not set here, but will be 'patched' later once guides & interpolation data are built
+		Out.Cards.UVs[VertexIndex]			= FVector4f(VertexInstanceUVs[VertexInstanceId0].Component(0), VertexInstanceUVs[VertexInstanceId0].Component(1), 0, 0); // RootUV are not set here, but will be 'patched' later once guides & interpolation data are built
 		Out.Cards.Tangents[VertexIndex]		= VertexInstanceTangents[VertexInstanceId0];
 		Out.Cards.Normals[VertexIndex]		= VertexInstanceNormals[VertexInstanceId0];
 
@@ -3698,7 +3705,7 @@ bool InternalImportGeometry(
 
 		Out.Cards.BoundingBox += Out.Cards.Positions[VertexIndex];
 	}
-	OutBulk.BoundingBox = Out.Cards.BoundingBox;
+	OutBulk.BoundingBox = ToFBox3d(Out.Cards.BoundingBox);
 
 	// Fill in render resources (do we need to keep it separated? e.g, format compression, packing)
 	const uint32 PointCount = Out.Cards.Positions.Num();
@@ -3707,10 +3714,10 @@ bool InternalImportGeometry(
 	OutBulk.UVs.SetNum(PointCount);
 	for (uint32 PointIt = 0; PointIt < PointCount; ++PointIt)
 	{
-		OutBulk.Positions[PointIt] = FVector4(Out.Cards.Positions[PointIt], 0);
+		OutBulk.Positions[PointIt] = FVector4f(Out.Cards.Positions[PointIt], 0);
 		OutBulk.UVs[PointIt] = Out.Cards.UVs[PointIt];
-		OutBulk.Normals[PointIt * 2] = FVector4(Out.Cards.Tangents[PointIt], 0);
-		OutBulk.Normals[PointIt * 2 + 1] = FVector4(Out.Cards.Normals[PointIt], TangentFrameSigns[PointIt]);
+		OutBulk.Normals[PointIt * 2] = FVector4f(Out.Cards.Tangents[PointIt], 0);
+		OutBulk.Normals[PointIt * 2 + 1] = FVector4f(Out.Cards.Normals[PointIt], TangentFrameSigns[PointIt]);
 	}
 
 	OutBulk.Indices.SetNum(IndexCount);
@@ -3818,9 +3825,9 @@ bool InternalImportGeometry(
 		// 1. Extract all roots
 		struct FStrandsRootData
 		{
-			FVector   Position;
-			uint32	  CurveIndex;
-			FVector2D RootUV;
+			FVector3f Position;
+			uint32    CurveIndex;
+			FVector2f RootUV;
 		};
 		TArray<FStrandsRootData> StrandsRoots;
 		{
@@ -3839,7 +3846,7 @@ bool InternalImportGeometry(
 		// 2. Extract cards root points
 		struct FCardsRootData
 		{
-			FVector Position;
+			FVector3f Position;
 			uint32  CardsIndex;
 		};
 		TArray<FCardsRootData> CardsRoots;
@@ -3851,7 +3858,7 @@ bool InternalImportGeometry(
 				const uint32 CardsIndexOffset = Out.Cards.IndexOffsets[CardIt];
 
 				uint32 AverageCount = 0;
-				FVector AverageRootPosition = FVector::ZeroVector;
+				FVector3f AverageRootPosition = FVector3f::ZeroVector;
 				for (uint32 IndexIt = 0; IndexIt < CardsIndexCount; ++IndexIt)
 				{
 					const uint32 VertexIndex = Out.Cards.Indices[CardsIndexOffset + IndexIt];
@@ -3864,7 +3871,7 @@ bool InternalImportGeometry(
 
 				FCardsRootData& Cards = CardsRoots.AddDefaulted_GetRef();
 				Cards.CardsIndex = CardIt;
-				Cards.Position = AverageCount > 0 ? AverageRootPosition / AverageCount : FVector::ZeroVector;
+				Cards.Position = AverageCount > 0 ? AverageRootPosition / AverageCount : FVector3f::ZeroVector;
 			}
 		}
 
@@ -3878,7 +3885,7 @@ bool InternalImportGeometry(
 				float ClosestDistance = FLT_MAX;
 				for (const FStrandsRootData& StrandsRoot : StrandsRoots)
 				{
-					const float Distance = FVector::Distance(StrandsRoot.Position, CardsRoot.Position);
+					const float Distance = FVector3f::Distance(StrandsRoot.Position, CardsRoot.Position);
 					if (Distance < ClosestDistance)
 					{
 						ClosestDistance = Distance;
@@ -3888,12 +3895,7 @@ bool InternalImportGeometry(
 
 				// 3.2 Apply root UV to all cards vertices
 				{
-					FVector2D RootUV = FVector2D::ZeroVector;
-					const bool bFoundValidCurve = CurveIndex != ~0;
-					if (bFoundValidCurve)
-					{
-						RootUV = InStrandsData.StrandsCurves.CurvesRootUV[CurveIndex];
-					}
+					FVector2f RootUV = FVector2f::ZeroVector;
 
 					const uint32 CardsIndexCount = Out.Cards.IndexCounts[CardsRoot.CardsIndex];
 					const uint32 CardsIndexOffset = Out.Cards.IndexOffsets[CardsRoot.CardsIndex];
@@ -4132,14 +4134,14 @@ void BuildGeometry(
 	Out.Meshes.Positions[6] = Center + FVector3f(+Extent.X, +Extent.Y, +Extent.Z);
 	Out.Meshes.Positions[7] = Center + FVector3f(-Extent.X, +Extent.Y, +Extent.Z);
 
-	Out.Meshes.UVs[0] = FVector2D(0, 0);
-	Out.Meshes.UVs[1] = FVector2D(1, 0);
-	Out.Meshes.UVs[2] = FVector2D(1, 1);
-	Out.Meshes.UVs[3] = FVector2D(0, 1);
-	Out.Meshes.UVs[4] = FVector2D(0, 0);
-	Out.Meshes.UVs[5] = FVector2D(1, 0);
-	Out.Meshes.UVs[6] = FVector2D(1, 1);
-	Out.Meshes.UVs[7] = FVector2D(0, 1);
+	Out.Meshes.UVs[0] = FVector2f(0, 0);
+	Out.Meshes.UVs[1] = FVector2f(1, 0);
+	Out.Meshes.UVs[2] = FVector2f(1, 1);
+	Out.Meshes.UVs[3] = FVector2f(0, 1);
+	Out.Meshes.UVs[4] = FVector2f(0, 0);
+	Out.Meshes.UVs[5] = FVector2f(1, 0);
+	Out.Meshes.UVs[6] = FVector2f(1, 1);
+	Out.Meshes.UVs[7] = FVector2f(0, 1);
 
 	Out.Meshes.Normals[0] = FVector3f(0, 0, 1);
 	Out.Meshes.Normals[1] = FVector3f(0, 0, 1);
@@ -4210,14 +4212,14 @@ void BuildGeometry(
 	OutBulk.UVs.SetNum(PointCount);
 	for (uint32 PointIt = 0; PointIt < PointCount; ++PointIt)
 	{
-		OutBulk.Positions[PointIt] = FVector4(Out.Meshes.Positions[PointIt], 0);
-		OutBulk.UVs[PointIt] = FVector4(Out.Meshes.UVs[PointIt].X, Out.Meshes.UVs[PointIt].Y, 0, 0);
-		OutBulk.Normals[PointIt * 2] = FVector4(Out.Meshes.Tangents[PointIt], 0);
-		OutBulk.Normals[PointIt * 2 + 1] = FVector4(Out.Meshes.Normals[PointIt], 1);
+		OutBulk.Positions[PointIt] = FVector4f(Out.Meshes.Positions[PointIt], 0);
+		OutBulk.UVs[PointIt] = FVector4f(Out.Meshes.UVs[PointIt].X, Out.Meshes.UVs[PointIt].Y, 0, 0);
+		OutBulk.Normals[PointIt * 2] = FVector4f(Out.Meshes.Tangents[PointIt], 0);
+		OutBulk.Normals[PointIt * 2 + 1] = FVector4f(Out.Meshes.Normals[PointIt], 1);
 
-		Out.Meshes.BoundingBox += FVector4(OutBulk.Positions[PointIt]);
+		Out.Meshes.BoundingBox += FVector4f(OutBulk.Positions[PointIt]);
 	}
-	OutBulk.BoundingBox = Out.Meshes.BoundingBox;
+	OutBulk.BoundingBox = ToFBox3d(Out.Meshes.BoundingBox);
 
 	const uint32 IndexCount = Out.Meshes.Indices.Num();
 	OutBulk.Indices.SetNum(IndexCount);
@@ -4263,17 +4265,18 @@ void ImportGeometry(
 	}
 
 	// Fill in render resources (do we need to keep it separated? e.g, format compression, packing)
-	OutBulk.BoundingBox = Out.Meshes.BoundingBox;
+	OutBulk.BoundingBox = ToFBox3d(Out.Meshes.BoundingBox);
+
 	const uint32 PointCount = Out.Meshes.Positions.Num();
 	OutBulk.Positions.SetNum(PointCount);
 	OutBulk.Normals.SetNum(PointCount * FHairCardsNormalFormat::ComponentCount);
 	OutBulk.UVs.SetNum(PointCount);
 	for (uint32 PointIt = 0; PointIt < PointCount; ++PointIt)
 	{
-		OutBulk.Positions[PointIt] = FVector4(Out.Meshes.Positions[PointIt], 0);
-		OutBulk.UVs[PointIt] = FVector4(Out.Meshes.UVs[PointIt].X, Out.Meshes.UVs[PointIt].Y, 0, 0);
-		OutBulk.Normals[PointIt * 2] = FVector4(Out.Meshes.Tangents[PointIt], 0);
-		OutBulk.Normals[PointIt * 2 + 1] = FVector4(Out.Meshes.Normals[PointIt], 1);
+		OutBulk.Positions[PointIt] = FVector4f(Out.Meshes.Positions[PointIt], 0);
+		OutBulk.UVs[PointIt] = FVector4f(Out.Meshes.UVs[PointIt].X, Out.Meshes.UVs[PointIt].Y, 0, 0);
+		OutBulk.Normals[PointIt * 2] = FVector4f(Out.Meshes.Tangents[PointIt], 0);
+		OutBulk.Normals[PointIt * 2 + 1] = FVector4f(Out.Meshes.Normals[PointIt], 1);
 	}
 
 	OutBulk.Indices.SetNum(IndexCount);
@@ -4305,7 +4308,7 @@ public:
 	void SetInstanceNormal(const FVertexInstanceID& InstanceID, const FVector3f& Normal);
 
 	/** Set the UV of a vertex instance */
-	void SetInstanceUV(const FVertexInstanceID& InstanceID, const FVector2D& InstanceUV, int32 UVLayerIndex = 0);
+	void SetInstanceUV(const FVertexInstanceID& InstanceID, const FVector2f& InstanceUV, int32 UVLayerIndex = 0);
 
 	/** Set the number of UV layers */
 	void SetNumUVLayers(int32 NumUVLayers);
@@ -4323,13 +4326,13 @@ public:
 	FPolygonID AppendTriangle(const FVertexID& Vertex0, const FVertexID& Vertex1, const FVertexID& Vertex2, const FPolygonGroupID& PolygonGroup);
 
 	/** Append a triangle to the mesh with the given PolygonGroup ID, and optionally with triangle-vertex UVs and Normals */
-	FPolygonID AppendTriangle(const FVertexID* Triangle, const FPolygonGroupID& PolygonGroup, const FVector2D* VertexUVs = nullptr, const FVector3f* VertexNormals = nullptr);
+	FPolygonID AppendTriangle(const FVertexID* Triangle, const FPolygonGroupID& PolygonGroup, const FVector2f* VertexUVs = nullptr, const FVector3f* VertexNormals = nullptr);
 
 	/**
 	 * Append an arbitrary polygon to the mesh with the given PolygonGroup ID, and optionally with polygon-vertex UVs and Normals
 	 * Unique Vertex instances will be created for each polygon-vertex.
 	 */
-	FPolygonID AppendPolygon(const TArray<FVertexID>& Vertices, const FPolygonGroupID& PolygonGroup, const TArray<FVector2D>* VertexUVs = nullptr, const TArray<FVector3f>* VertexNormals = nullptr);
+	FPolygonID AppendPolygon(const TArray<FVertexID>& Vertices, const FPolygonGroupID& PolygonGroup, const TArray<FVector2f>* VertexUVs = nullptr, const TArray<FVector3f>* VertexNormals = nullptr);
 
 	/**
 	 * Append a triangle to the mesh using the given vertex instances and PolygonGroup ID
@@ -4398,7 +4401,7 @@ void FMeshDescriptionBuilder::SetInstanceNormal(const FVertexInstanceID& Instanc
 	}
 }
 
-void FMeshDescriptionBuilder::SetInstanceUV(const FVertexInstanceID& InstanceID, const FVector2D& InstanceUV, int32 UVLayerIndex)
+void FMeshDescriptionBuilder::SetInstanceUV(const FVertexInstanceID& InstanceID, const FVector2f& InstanceUV, int32 UVLayerIndex)
 {
 	if (InstanceUVs.IsValid() && ensure(UVLayerIndex < InstanceUVs.GetNumChannels()))
 	{
@@ -4448,7 +4451,7 @@ void ConvertCardsGeometryToMeshDescription(const FHairCardsGeometry& In, FMeshDe
 	{
 		FVertexID VertexID = Builder.AppendVertex(In.Positions[VIndex]);
 		FVertexInstanceID InstanceID = Builder.AppendInstance(VertexID);
-		FVector2D UV = FVector2D(In.UVs[VIndex].X, In.UVs[VIndex].Y);
+		FVector2f UV = FVector2f(In.UVs[VIndex].X, In.UVs[VIndex].Y);
 		Builder.SetInstanceNormal(InstanceID, In.Normals[VIndex]);
 		Builder.SetInstanceUV(InstanceID, UV, 0);
 
diff --git a/Engine/Plugins/Runtime/HairStrands/Source/HairStrandsCore/Public/HairAttributes.h b/Engine/Plugins/Runtime/HairStrands/Source/HairStrandsCore/Public/HairAttributes.h
--- a/Engine/Plugins/Runtime/HairStrands/Source/HairStrandsCore/Public/HairAttributes.h
+++ b/Engine/Plugins/Runtime/HairStrands/Source/HairStrandsCore/Public/HairAttributes.h
@@ -21,7 +21,7 @@ namespace HairAttribute
 		extern HAIRSTRANDSCORE_API const FName GroupID;			// int
 		extern HAIRSTRANDSCORE_API const FName Guide;			// int
 		extern HAIRSTRANDSCORE_API const FName ID;				// int
-		extern HAIRSTRANDSCORE_API const FName RootUV;			// FVector2D
+		extern HAIRSTRANDSCORE_API const FName RootUV;			// FVector2f
 		extern HAIRSTRANDSCORE_API const FName VertexCount;		// int
 		extern HAIRSTRANDSCORE_API const FName Width;			// float
 		extern HAIRSTRANDSCORE_API const FName ClosestGuides;	// FVector
diff --git a/Engine/Plugins/Runtime/HairStrands/Source/HairStrandsCore/Public/HairCardsDatas.h b/Engine/Plugins/Runtime/HairStrands/Source/HairStrandsCore/Public/HairCardsDatas.h
--- a/Engine/Plugins/Runtime/HairStrands/Source/HairStrandsCore/Public/HairCardsDatas.h
+++ b/Engine/Plugins/Runtime/HairStrands/Source/HairStrandsCore/Public/HairCardsDatas.h
@@ -176,7 +176,7 @@ struct FHairOrientedBound
 struct FHairCardsGeometry
 {
 	// Geometry
-	TArray<FVector4> UVs;
+	TArray<FVector4f>UVs;
 	TArray<FVector3f>Normals;
 	TArray<FVector3f>Tangents;
 	TArray<FVector3f>Positions;
@@ -192,7 +192,7 @@ struct FHairCardsGeometry
 	TArray<uint32> IndexOffsets;
 	TArray<uint32> IndexCounts;
 
-	FBox BoundingBox;
+	FBox3f BoundingBox;
 
 	void Reset()
 	{
@@ -357,8 +357,8 @@ struct FHairCardsProceduralAtlas
 
 	FIntPoint Resolution;
 	TArray<Rect> Rects;
-	TArray<FVector4> StrandsPositions;
-	TArray<FVector4> StrandsAttributes;
+	TArray<FVector4f> StrandsPositions;
+	TArray<FVector4f> StrandsAttributes;
 	bool bIsDirty = true;
 };
 
@@ -406,13 +406,13 @@ struct HAIRSTRANDSCORE_API FHairCardsSourceData
 struct FHairMeshes
 {
 	// Geometry
-	TArray<FVector2D> UVs;
+	TArray<FVector2f> UVs;
 	TArray<FVector3f> Normals;
 	TArray<FVector3f> Tangents;
 	TArray<FVector3f> Positions;
 	TArray<uint32>    Indices;
 
-	FBox BoundingBox;
+	FBox3f BoundingBox;
 
 	void SetNum(uint32 Count)
 	{
diff --git a/Engine/Plugins/Runtime/HairStrands/Source/HairStrandsCore/Public/HairStrandsDatas.h b/Engine/Plugins/Runtime/HairStrands/Source/HairStrandsCore/Public/HairStrandsDatas.h
--- a/Engine/Plugins/Runtime/HairStrands/Source/HairStrandsCore/Public/HairStrandsDatas.h
+++ b/Engine/Plugins/Runtime/HairStrands/Source/HairStrandsCore/Public/HairStrandsDatas.h
@@ -267,10 +267,10 @@ struct FHairStrandsRootUtils
 {
 	static uint32	 EncodeTriangleIndex(uint32 TriangleIndex, uint32 SectionIndex);
 	static void		 DecodeTriangleIndex(uint32 Encoded, uint32& OutTriangleIndex, uint32& OutSectionIndex);
-	static uint32	 EncodeBarycentrics(const FVector2D& B);
-	static FVector2D DecodeBarycentrics(uint32 B);
-	static uint32	 PackUVs(const FVector2D& UV);
-	static float	 PackUVsToFloat(const FVector2D& UV);
+	static uint32	 EncodeBarycentrics(const FVector2f& B);
+	static FVector2f DecodeBarycentrics(uint32 B);
+	static uint32	 PackUVs(const FVector2f& UV);
+	static float	 PackUVsToFloat(const FVector2f& UV);
 };
 
 struct FHairStrandsInterpolationBulkData;
diff --git a/Engine/Source/Runtime/Renderer/Private/HairStrands/HairStrandsComposition.cpp b/Engine/Source/Runtime/Renderer/Private/HairStrands/HairStrandsComposition.cpp
--- a/Engine/Source/Runtime/Renderer/Private/HairStrands/HairStrandsComposition.cpp
+++ b/Engine/Source/Runtime/Renderer/Private/HairStrands/HairStrandsComposition.cpp
@@ -27,7 +27,7 @@ static FAutoConsoleVariableRef CVarHairStrandsComposeDOFDepth(TEXT("r.HairStrand
 
 float GetHairFastResolveVelocityThreshold(const FIntPoint& Resolution)
 {
-	FVector2D PixelVelocity(1.f / (Resolution.X * 2), 1.f / (Resolution.Y * 2));
+	FVector2f PixelVelocity(1.f / (Resolution.X * 2), 1.f / (Resolution.Y * 2));
 	const float VelocityThreshold = FMath::Clamp(GHairFastResolveVelocityThreshold, 0, 512) * FMath::Min(PixelVelocity.X, PixelVelocity.Y);
 	return VelocityThreshold;
 }
diff --git a/Engine/Source/Runtime/Renderer/Private/HairStrands/HairStrandsDebug.cpp b/Engine/Source/Runtime/Renderer/Private/HairStrands/HairStrandsDebug.cpp
--- a/Engine/Source/Runtime/Renderer/Private/HairStrands/HairStrandsDebug.cpp
+++ b/Engine/Source/Runtime/Renderer/Private/HairStrands/HairStrandsDebug.cpp
@@ -483,13 +483,13 @@ static void AddDebugDeepShadowTexturePass(
 	check(OutTarget);
 
 	FIntPoint AtlasResolution(0, 0);
-	FVector2D AltasOffset(0, 0);
-	FVector2D AltasScale(0, 0);
+	FVector2f AltasOffset(0, 0);
+	FVector2f AltasScale(0, 0);
 	if (ShadowData && Resources)
 	{
 		AtlasResolution = FIntPoint(Resources->DepthAtlasTexture->Desc.Extent.X, Resources->DepthAtlasTexture->Desc.Extent.Y);
-		AltasOffset = FVector2D(ShadowData->AtlasRect.Min.X / float(AtlasResolution.X), ShadowData->AtlasRect.Min.Y / float(AtlasResolution.Y));
-		AltasScale = FVector2D((ShadowData->AtlasRect.Max.X - ShadowData->AtlasRect.Min.X) / float(AtlasResolution.X), (ShadowData->AtlasRect.Max.Y - ShadowData->AtlasRect.Min.Y) / float(AtlasResolution.Y));
+		AltasOffset = FVector2f(ShadowData->AtlasRect.Min.X / float(AtlasResolution.X), ShadowData->AtlasRect.Min.Y / float(AtlasResolution.Y));
+		AltasScale = FVector2f((ShadowData->AtlasRect.Max.X - ShadowData->AtlasRect.Min.X) / float(AtlasResolution.X), (ShadowData->AtlasRect.Max.Y - ShadowData->AtlasRect.Min.Y) / float(AtlasResolution.Y));
 	}
 
 	const FIntRect Viewport = View->ViewRect;
@@ -500,7 +500,7 @@ static void AddDebugDeepShadowTexturePass(
 	Parameters->DomAtlasOffset = AltasOffset;
 	Parameters->DomAtlasScale = AltasScale;
 	Parameters->OutputResolution = Resolution;
-	Parameters->InvOutputResolution = FVector2D(1.f / Resolution.X, 1.f / Resolution.Y);
+	Parameters->InvOutputResolution = FVector2f(1.f / Resolution.X, 1.f / Resolution.Y);
 	Parameters->DeepShadowDepthTexture = Resources ? Resources->DepthAtlasTexture : nullptr;
 	Parameters->DeepShadowLayerTexture = Resources ? Resources->LayersAtlasTexture : nullptr;
 	Parameters->LinearSampler = TStaticSamplerState<SF_Bilinear, AM_Clamp, AM_Clamp, AM_Clamp>::GetRHI();
diff --git a/Engine/Source/Runtime/Renderer/Private/HairStrands/HairStrandsEnvironment.cpp b/Engine/Source/Runtime/Renderer/Private/HairStrands/HairStrandsEnvironment.cpp
--- a/Engine/Source/Runtime/Renderer/Private/HairStrands/HairStrandsEnvironment.cpp
+++ b/Engine/Source/Runtime/Renderer/Private/HairStrands/HairStrandsEnvironment.cpp
@@ -151,7 +151,7 @@ static void AddHairStrandsEnvironmentAOPass(
 	PassParameters->AO_DistanceThreshold = FMath::Max(GHairStrandsSkyAODistanceThreshold, 1.f);
 	PassParameters->RenderTargets[0] = FRenderTargetBinding(Output, ERenderTargetLoadAction::ELoad);
 	PassParameters->Output_bHalfRes = bHalfRes;
-	PassParameters->Output_InvResolution = FVector2D(1.f/Output->Desc.Extent.X, 1.f/Output->Desc.Extent.Y);
+	PassParameters->Output_InvResolution = FVector2f(1.f/Output->Desc.Extent.X, 1.f/Output->Desc.Extent.Y);
 
 	FIntRect ViewRect;
 	if (bHalfRes)
diff --git a/Engine/Source/Runtime/Renderer/Private/HairStrands/HairStrandsVisibility.cpp b/Engine/Source/Runtime/Renderer/Private/HairStrands/HairStrandsVisibility.cpp
--- a/Engine/Source/Runtime/Renderer/Private/HairStrands/HairStrandsVisibility.cpp
+++ b/Engine/Source/Runtime/Renderer/Private/HairStrands/HairStrandsVisibility.cpp
@@ -256,7 +256,7 @@ static uint32 GetTotalSampleCountForAllocation(FIntPoint Resolution)
 
 static void SetUpViewHairRenderInfo(const FViewInfo& ViewInfo, bool bEnableMSAA, FVector4f& OutHairRenderInfo, uint32& OutHairRenderInfoBits, uint32& OutHairComponents)
 {
-	FVector2D PixelVelocity(1.f / (ViewInfo.ViewRect.Width() * 2), 1.f / (ViewInfo.ViewRect.Height() * 2));
+	FVector2f PixelVelocity(1.f / (ViewInfo.ViewRect.Width() * 2), 1.f / (ViewInfo.ViewRect.Height() * 2));
 	const float VelocityMagnitudeScale = FMath::Clamp(CVarHairVelocityMagnitudeScale.GetValueOnAnyThread(), 0, 512) * FMath::Min(PixelVelocity.X, PixelVelocity.Y);
 
 	// In the case we render coverage, we need to override some view uniform shader parameters to account for the change in MSAA sample count.
