Fix assert when hair simulation is enabled on scaled groom.

diff --git a/Engine/Plugins/Runtime/HairStrands/Source/HairStrandsCore/Private/Niagara/NiagaraDataInterfaceHairStrands.cpp b/Engine/Plugins/Runtime/HairStrands/Source/HairStrandsCore/Private/Niagara/NiagaraDataInterfaceHairStrands.cpp
--- a/Engine/Plugins/Runtime/HairStrands/Source/HairStrandsCore/Private/Niagara/NiagaraDataInterfaceHairStrands.cpp
+++ b/Engine/Plugins/Runtime/HairStrands/Source/HairStrandsCore/Private/Niagara/NiagaraDataInterfaceHairStrands.cpp
@@ -754,7 +754,7 @@ struct FNDIHairStrandsParametersCS : public FNiagaraDataInterfaceParametersCS
 			SetShaderValue(RHICmdList, ComputeShaderRHI, StrandSize, ProxyData->StrandsSize);
 			SetShaderValue(RHICmdList, ComputeShaderRHI, BoneTransform, BoneTransformFloat);
 			SetShaderValue(RHICmdList, ComputeShaderRHI, BoneInverse, BoneTransformFloat.Inverse());
-			SetShaderValue(RHICmdList, ComputeShaderRHI, BoneRotation, BoneTransformFloat.ToQuat());
+			SetShaderValue(RHICmdList, ComputeShaderRHI, BoneRotation, BoneTransformFloat.GetMatrixWithoutScale().ToQuat());
 			SetShaderValue(RHICmdList, ComputeShaderRHI, ResetSimulation, NeedResetValue);
 			SetShaderValue(RHICmdList, ComputeShaderRHI, InterpolationMode, int32(InterpolationModeValue));
 			SetShaderValue(RHICmdList, ComputeShaderRHI, RestUpdate, RestUpdateValue);
