World/Local hair inertial forces

diff --git a/Engine/Plugins/Runtime/HairStrands/Shaders/Private/NiagaraDataInterfaceHairStrands.ush b/Engine/Plugins/Runtime/HairStrands/Shaders/Private/NiagaraDataInterfaceHairStrands.ush
--- a/Engine/Plugins/Runtime/HairStrands/Shaders/Private/NiagaraDataInterfaceHairStrands.ush
+++ b/Engine/Plugins/Runtime/HairStrands/Shaders/Private/NiagaraDataInterfaceHairStrands.ush
@@ -47,6 +47,10 @@ struct FDIHairStrandsContext
 	Buffer<float4>		RestSamplePositionsBuffer;
 	Buffer<float4>		MeshSampleWeightsBuffer;
 	Buffer<float>		ParamsScaleBuffer;
+	float3				BoneLinearVelocity;
+	float3				BoneAngularVelocity;
+	float3				BoneLinearAcceleration;
+	float3				BoneAngularAcceleration;
 }; 
 
 #define DIHAIRSTRANDS_DECLARE_CONSTANTS(NAME)\
@@ -82,6 +86,10 @@ int4				BoundingBoxOffsets_##NAME;\
 Buffer<float4>		RestSamplePositionsBuffer_##NAME;\
 Buffer<float4>		MeshSampleWeightsBuffer_##NAME;\
 Buffer<float>		ParamsScaleBuffer_##NAME;\
+float3				BoneLinearVelocity_##NAME;\
+float3				BoneAngularVelocity_##NAME;\
+float3				BoneLinearAcceleration_##NAME;\
+float3				BoneAngularAcceleration_##NAME;\
 
 //@todo-lh: BoundingBoxBuffer is passed around separately as SPIRV-Cross cannot translate texture atomics properly atm.
 #define DIHAIRSTRANDS_MAKE_CONTEXT(NAME)\
@@ -118,6 +126,10 @@ DIContext.BoundingBoxOffsets = BoundingBoxOffsets_##NAME;\
 DIContext.RestSamplePositionsBuffer = RestSamplePositionsBuffer_##NAME;\
 DIContext.MeshSampleWeightsBuffer = MeshSampleWeightsBuffer_##NAME;\
 DIContext.ParamsScaleBuffer = ParamsScaleBuffer_##NAME;\
+DIContext.BoneLinearVelocity = BoneLinearVelocity_##NAME;\
+DIContext.BoneAngularVelocity = BoneAngularVelocity_##NAME;\
+DIContext.BoneLinearAcceleration = BoneLinearAcceleration_##NAME;\
+DIContext.BoneAngularAcceleration = BoneAngularAcceleration_##NAME;\
 
 /* -----------------------------------------------------------------
  * Vector utilities
@@ -780,8 +792,16 @@ void DIHairStrands_AdvectNodePosition(in FDIHairStrandsContext DIContext, in flo
 	{
 		const float3 ImplicitGradient =  float3(NodeMass,NodeMass,NodeMass) / DeltaTime - ForceGradient;
 		const float3 InverseGradient = float3(1.0/ImplicitGradient.x, 1.0/ImplicitGradient.y, 1.0/ImplicitGradient.z);
-
-		OutLinearVelocity += InverseGradient * ExternalForce;
+		
+		if(DIContext.LocalSimulation)
+		{
+			OutLinearVelocity += InverseGradient * ExternalForce - DIContext.BoneLinearAcceleration * DeltaTime - ( cross(DIContext.BoneAngularAcceleration, OutNodePosition) +
+				2.0 * cross(DIContext.BoneAngularVelocity, OutLinearVelocity) + cross(DIContext.BoneAngularVelocity, cross(DIContext.BoneAngularVelocity,OutNodePosition))) * DeltaTime;
+		}
+		else
+		{
+			OutLinearVelocity += InverseGradient * ExternalForce;
+		}
 		OutNodePosition += OutLinearVelocity * DeltaTime;
 	}
 }
diff --git a/Engine/Plugins/Runtime/HairStrands/Source/HairStrandsCore/Private/Niagara/NiagaraDataInterfaceHairStrands.cpp b/Engine/Plugins/Runtime/HairStrands/Source/HairStrandsCore/Private/Niagara/NiagaraDataInterfaceHairStrands.cpp
--- a/Engine/Plugins/Runtime/HairStrands/Source/HairStrandsCore/Private/Niagara/NiagaraDataInterfaceHairStrands.cpp
+++ b/Engine/Plugins/Runtime/HairStrands/Source/HairStrandsCore/Private/Niagara/NiagaraDataInterfaceHairStrands.cpp
@@ -179,6 +179,11 @@ const FString UNiagaraDataInterfaceHairStrands::BoneTransformName(TEXT("BoneTran
 const FString UNiagaraDataInterfaceHairStrands::BoneInverseName(TEXT("BoneInverse_"));
 const FString UNiagaraDataInterfaceHairStrands::BoneRotationName(TEXT("BoneRotation_"));
 
+const FString UNiagaraDataInterfaceHairStrands::BoneLinearVelocityName(TEXT("BoneLinearVelocity_"));
+const FString UNiagaraDataInterfaceHairStrands::BoneAngularVelocityName(TEXT("BoneAngularVelocity_"));
+const FString UNiagaraDataInterfaceHairStrands::BoneLinearAccelerationName(TEXT("BoneLinearAcceleration_"));
+const FString UNiagaraDataInterfaceHairStrands::BoneAngularAccelerationName(TEXT("BoneAngularAcceleration_"));
+
 const FString UNiagaraDataInterfaceHairStrands::DeformedPositionBufferName(TEXT("DeformedPositionBuffer_"));
 const FString UNiagaraDataInterfaceHairStrands::CurvesOffsetsBufferName(TEXT("CurvesOffsetsBuffer_"));
 const FString UNiagaraDataInterfaceHairStrands::RestPositionBufferName(TEXT("RestPositionBuffer_"));
@@ -233,6 +238,11 @@ struct FNDIHairStrandsParametersName
 		BoneTransformName = UNiagaraDataInterfaceHairStrands::BoneTransformName + Suffix;
 		BoneInverseName = UNiagaraDataInterfaceHairStrands::BoneInverseName + Suffix;
 		BoneRotationName = UNiagaraDataInterfaceHairStrands::BoneRotationName + Suffix;
+		
+		BoneLinearVelocityName = UNiagaraDataInterfaceHairStrands::BoneLinearVelocityName + Suffix;
+		BoneAngularVelocityName = UNiagaraDataInterfaceHairStrands::BoneAngularVelocityName + Suffix;
+		BoneLinearAccelerationName = UNiagaraDataInterfaceHairStrands::BoneLinearAccelerationName + Suffix;
+		BoneAngularAccelerationName = UNiagaraDataInterfaceHairStrands::BoneAngularAccelerationName + Suffix;
 
 		DeformedPositionBufferName = UNiagaraDataInterfaceHairStrands::DeformedPositionBufferName + Suffix;
 		CurvesOffsetsBufferName = UNiagaraDataInterfaceHairStrands::CurvesOffsetsBufferName + Suffix;
@@ -275,6 +285,11 @@ struct FNDIHairStrandsParametersName
 	FString BoneTransformName;
 	FString BoneInverseName;
 	FString BoneRotationName;
+	
+	FString BoneLinearVelocityName;
+	FString BoneLinearAccelerationName;
+	FString BoneAngularVelocityName;
+	FString BoneAngularAccelerationName;
 
 	FString DeformedPositionBufferName;
 	FString CurvesOffsetsBufferName;
@@ -441,7 +456,7 @@ void FNDIHairStrandsData::Release()
 }
 
 void FNDIHairStrandsData::Update(UNiagaraDataInterfaceHairStrands* Interface, FNiagaraSystemInstance* SystemInstance, const FHairStrandsBulkData* StrandsDatas,
-	UGroomAsset* GroomAsset, const int32 GroupIndex, const int32 LODIndex, const FTransform& LocalToWorld)
+	UGroomAsset* GroomAsset, const int32 GroupIndex, const int32 LODIndex, const FTransform& LocalToWorld, const float DeltaSeconds)
 {
 	if (Interface != nullptr)
 	{
@@ -515,7 +530,39 @@ void FNDIHairStrandsData::Update(UNiagaraDataInterfaceHairStrands* Interface, FN
 				// Convert to double for LWC
 				FMatrix44d BoneTransformDouble = BoneTransform.ToMatrixWithScale();
 				const FMatrix44d WorldTransformDouble = WorldTransform.ToMatrixWithScale();
+
+				if(DeltaSeconds != 0.0f)
+				{
+					const FMatrix44d PreviousBoneTransformDouble = PreviousBoneTransform.ToMatrixWithScale();
+					const FMatrix44d DeltaTransformDouble =  BoneTransformDouble * PreviousBoneTransformDouble.Inverse();
+					
+					const FMatrix44f DeltaTransformFloat = DeltaTransformDouble;
+					const FTransform DeltaTransform = FTransform(DeltaTransformFloat);
+					const FQuat DeltaRotation = DeltaTransform.GetRotation();
+					
+					// Apply linear velocity scale
+					BoneLinearVelocity = FMath::Clamp(1.f - SimulationSettings.SimulationSetup.LinearVelocityScale, 0.f, 1.f) * DeltaTransform.GetTranslation() / DeltaSeconds;
+					BoneLinearAcceleration = (BoneLinearVelocity-PreviousBoneLinearVelocity) / DeltaSeconds;
+
+					
+					// Apply angular velocity scale
+					BoneAngularVelocity = BoneTransform.TransformVector(DeltaRotation.GetRotationAxis() * DeltaRotation.GetAngle() *
+						FMath::Clamp(1.f - SimulationSettings.SimulationSetup.AngularVelocityScale, 0.f, 1.f)) / DeltaSeconds;
+					BoneAngularAcceleration = (BoneAngularVelocity-PreviousBoneAngularVelocity) / DeltaSeconds;
+				}
+				else
+				{
+					BoneLinearVelocity = FVector3f::Zero();
+					BoneAngularVelocity = FVector3f::Zero();
+
+					BoneLinearAcceleration = FVector3f::Zero();
+					BoneAngularAcceleration = FVector3f::Zero();
+				}
 				
+				PreviousBoneTransform = BoneTransform;
+				PreviousBoneLinearVelocity = BoneLinearVelocity;
+				PreviousBoneAngularVelocity = BoneAngularVelocity;
+
 				BoneTransformDouble = BoneTransformDouble * WorldTransformDouble.Inverse();
 				const FMatrix44d WorldTransformFloat = BoneTransformDouble;
 				BoneTransform = FTransform(WorldTransformFloat);
@@ -563,7 +610,7 @@ bool FNDIHairStrandsData::Init(UNiagaraDataInterfaceHairStrands* Interface, FNia
 		{
 			FTransform LocalToWorld = FTransform::Identity;
 			Interface->ExtractDatasAndResources(SystemInstance, StrandsRestResource, StrandsDeformedResource, StrandsRestRootResource, StrandsDeformedRootResource, GroomAsset, GroupIndex, LODIndex, LocalToWorld);
-			Update(Interface, SystemInstance, StrandsRestResource ? &StrandsRestResource->BulkData : nullptr, GroomAsset, GroupIndex, LODIndex, LocalToWorld);
+			Update(Interface, SystemInstance, StrandsRestResource ? &StrandsRestResource->BulkData : nullptr, GroomAsset, GroupIndex, LODIndex, LocalToWorld, 0.0f);
 
 			HairStrandsBuffer = new FNDIHairStrandsBuffer();
 			HairStrandsBuffer->Initialize(StrandsRestResource, StrandsDeformedResource, StrandsRestRootResource, StrandsDeformedRootResource, ParamsScale);
@@ -604,6 +651,11 @@ struct FNDIHairStrandsParametersCS : public FNiagaraDataInterfaceParametersCS
 		BoneInverse.Bind(ParameterMap, *ParamNames.BoneInverseName);
 		BoneRotation.Bind(ParameterMap, *ParamNames.BoneRotationName);
 
+		BoneLinearVelocity.Bind(ParameterMap, *ParamNames.BoneLinearVelocityName);
+		BoneAngularVelocity.Bind(ParameterMap, *ParamNames.BoneAngularVelocityName);
+		BoneLinearAcceleration.Bind(ParameterMap, *ParamNames.BoneLinearAccelerationName);
+		BoneAngularAcceleration.Bind(ParameterMap, *ParamNames.BoneAngularAccelerationName);
+
 		DeformedPositionBuffer.Bind(ParameterMap, *ParamNames.DeformedPositionBufferName);
 		CurvesOffsetsBuffer.Bind(ParameterMap, *ParamNames.CurvesOffsetsBufferName);
 		RestPositionBuffer.Bind(ParameterMap, *ParamNames.RestPositionBufferName);
@@ -755,6 +807,10 @@ struct FNDIHairStrandsParametersCS : public FNiagaraDataInterfaceParametersCS
 			SetShaderValue(RHICmdList, ComputeShaderRHI, BoneTransform, BoneTransformFloat);
 			SetShaderValue(RHICmdList, ComputeShaderRHI, BoneInverse, BoneTransformFloat.Inverse());
 			SetShaderValue(RHICmdList, ComputeShaderRHI, BoneRotation, BoneTransformFloat.GetMatrixWithoutScale().ToQuat());
+			SetShaderValue(RHICmdList, ComputeShaderRHI, BoneLinearVelocity, ProxyData->BoneLinearVelocity);
+			SetShaderValue(RHICmdList, ComputeShaderRHI, BoneAngularVelocity, ProxyData->BoneAngularVelocity);
+			SetShaderValue(RHICmdList, ComputeShaderRHI, BoneLinearAcceleration, ProxyData->BoneLinearAcceleration);
+			SetShaderValue(RHICmdList, ComputeShaderRHI, BoneAngularAcceleration, ProxyData->BoneAngularAcceleration);
 			SetShaderValue(RHICmdList, ComputeShaderRHI, ResetSimulation, NeedResetValue);
 			SetShaderValue(RHICmdList, ComputeShaderRHI, InterpolationMode, int32(InterpolationModeValue));
 			SetShaderValue(RHICmdList, ComputeShaderRHI, RestUpdate, RestUpdateValue);
@@ -799,6 +855,10 @@ struct FNDIHairStrandsParametersCS : public FNiagaraDataInterfaceParametersCS
 			SetShaderValue(RHICmdList, ComputeShaderRHI, BoneTransform, FMatrix44f::Identity);
 			SetShaderValue(RHICmdList, ComputeShaderRHI, BoneInverse, FMatrix44f::Identity);
 			SetShaderValue(RHICmdList, ComputeShaderRHI, BoneRotation, FQuat4f::Identity);
+			SetShaderValue(RHICmdList, ComputeShaderRHI, BoneLinearVelocity, FVector3f::ZeroVector);
+			SetShaderValue(RHICmdList, ComputeShaderRHI, BoneAngularVelocity, FVector3f::ZeroVector);
+			SetShaderValue(RHICmdList, ComputeShaderRHI, BoneLinearAcceleration, FVector3f::ZeroVector);
+			SetShaderValue(RHICmdList, ComputeShaderRHI, BoneAngularAcceleration, FVector3f::ZeroVector);
 			SetShaderValue(RHICmdList, ComputeShaderRHI, ResetSimulation, 0);
 			SetShaderValue(RHICmdList, ComputeShaderRHI, InterpolationMode, 0);
 			SetShaderValue(RHICmdList, ComputeShaderRHI, RestUpdate, 0);
@@ -846,6 +906,10 @@ private:
 	LAYOUT_FIELD(FShaderParameter, BoneTransform);
 	LAYOUT_FIELD(FShaderParameter, BoneInverse);
 	LAYOUT_FIELD(FShaderParameter, BoneRotation);
+	LAYOUT_FIELD(FShaderParameter, BoneLinearVelocity);
+	LAYOUT_FIELD(FShaderParameter, BoneAngularVelocity);
+	LAYOUT_FIELD(FShaderParameter, BoneLinearAcceleration);
+	LAYOUT_FIELD(FShaderParameter, BoneAngularAcceleration);
 	LAYOUT_FIELD(FShaderParameter, ResetSimulation);
 	LAYOUT_FIELD(FShaderParameter, InterpolationMode);
 	LAYOUT_FIELD(FShaderParameter, RestUpdate);
@@ -1130,7 +1194,7 @@ bool UNiagaraDataInterfaceHairStrands::PerInstanceTick(void* PerInstanceData, FN
 		}
 		InstanceData->ForceReset = SourceComponent->bResetSimulation;
 	}
-	InstanceData->Update(this, SystemInstance, StrandsRestResource ? &StrandsRestResource->BulkData : nullptr, GroomAsset, GroupIndex, LODIndex, LocalToWorld);
+	InstanceData->Update(this, SystemInstance, StrandsRestResource ? &StrandsRestResource->BulkData : nullptr, GroomAsset, GroupIndex, LODIndex, LocalToWorld, InDeltaSeconds);
 	return false;
 }
 
diff --git a/Engine/Plugins/Runtime/HairStrands/Source/HairStrandsCore/Public/GroomAssetPhysics.h b/Engine/Plugins/Runtime/HairStrands/Source/HairStrandsCore/Public/GroomAssetPhysics.h
--- a/Engine/Plugins/Runtime/HairStrands/Source/HairStrandsCore/Public/GroomAssetPhysics.h
+++ b/Engine/Plugins/Runtime/HairStrands/Source/HairStrandsCore/Public/GroomAssetPhysics.h
@@ -357,6 +357,14 @@ struct HAIRSTRANDSCORE_API FHairSimulationSetup
 	/** Strands simulation is done in local space */
 	UPROPERTY(EditAnywhere, BlueprintReadWrite, Interp, Category = "SimulationSetup", meta = (ToolTip = "Strands simulation is done in local space"))
 	bool bLocalSimulation = true;
+	
+	/** Scaling of the bone world linear velocity contribution */
+	UPROPERTY(EditAnywhere, BlueprintReadWrite, Interp, Category = "SimulationSetup", meta = (ToolTip = "Scaling of the bone world linear velocity contribution", EditCondition = "bLocalSimulation"))
+	float LinearVelocityScale = 0.0f;
+	
+	/** Scaling of the bone world angular velocity contribution */
+	UPROPERTY(EditAnywhere, BlueprintReadWrite, Interp, Category = "SimulationSetup", meta = (ToolTip = "Scaling of the bone world angular velocity contribution", EditCondition = "bLocalSimulation"))
+	float AngularVelocityScale = 0.0f;
 
 	/** Bone used for the simulation local space */
 	UPROPERTY(EditAnywhere, BlueprintReadWrite, Interp, Category = "SimulationSetup", meta = (ToolTip = "Bone used for the simulation local space", EditCondition = "bLocalSimulation"))
diff --git a/Engine/Plugins/Runtime/HairStrands/Source/HairStrandsCore/Public/Niagara/NiagaraDataInterfaceHairStrands.h b/Engine/Plugins/Runtime/HairStrands/Source/HairStrandsCore/Public/Niagara/NiagaraDataInterfaceHairStrands.h
--- a/Engine/Plugins/Runtime/HairStrands/Source/HairStrandsCore/Public/Niagara/NiagaraDataInterfaceHairStrands.h
+++ b/Engine/Plugins/Runtime/HairStrands/Source/HairStrandsCore/Public/Niagara/NiagaraDataInterfaceHairStrands.h
@@ -97,12 +97,22 @@ struct FNDIHairStrandsData
 	void Release();
 
 	/** Update the buffers */
-	void Update(UNiagaraDataInterfaceHairStrands* Interface, FNiagaraSystemInstance* SystemInstance, const FHairStrandsBulkData* HairStrandsDatas, UGroomAsset* GroomAsset, const int32 GroupIndex, const int32 LODIndex, const FTransform& LocalToWorld);
+	void Update(UNiagaraDataInterfaceHairStrands* Interface, FNiagaraSystemInstance* SystemInstance, const FHairStrandsBulkData* HairStrandsDatas, UGroomAsset* GroomAsset, const int32 GroupIndex, const int32 LODIndex, const FTransform& LocalToWorld, const float DeltaSeconds);
 
 	inline void ResetDatas()
 	{
 		WorldTransform.SetIdentity();
 		BoneTransform.SetIdentity();
+		BoneLinearVelocity = FVector3f::Zero();
+		BoneAngularVelocity = FVector3f::Zero();
+			
+		BoneLinearAcceleration = FVector3f::Zero();
+		BoneAngularAcceleration = FVector3f::Zero();
+			
+		PreviousBoneTransform.SetIdentity();
+		PreviousBoneLinearVelocity = FVector3f::Zero();
+		PreviousBoneAngularVelocity = FVector3f::Zero();
+		
 		GlobalInterpolation = false;
 		HairGroupInstance = nullptr;
 
@@ -159,6 +169,15 @@ struct FNDIHairStrandsData
 
 			WorldTransform = OtherDatas->WorldTransform;
 			BoneTransform = OtherDatas->BoneTransform;
+			BoneLinearVelocity = OtherDatas->BoneLinearVelocity;
+			BoneAngularVelocity = OtherDatas->BoneAngularVelocity;
+			
+			BoneLinearAcceleration = OtherDatas->BoneLinearAcceleration;
+			BoneAngularAcceleration = OtherDatas->BoneAngularAcceleration;
+			
+			PreviousBoneTransform = OtherDatas->PreviousBoneTransform;
+			PreviousBoneLinearVelocity = OtherDatas->PreviousBoneLinearVelocity;
+			PreviousBoneAngularVelocity = OtherDatas->PreviousBoneAngularVelocity;
 
 			GlobalInterpolation = OtherDatas->GlobalInterpolation;
 			BindingType = OtherDatas->BindingType;
@@ -213,6 +232,27 @@ struct FNDIHairStrandsData
 
 	/** Bone transform that will be used for local strands simulation */
 	FTransform BoneTransform;
+	
+	/** Bone transform that will be used for local strands simulation */
+	FTransform PreviousBoneTransform;
+
+	/** Bone Linear Velocity */
+	FVector3f BoneLinearVelocity;
+
+	/** Bone Previous Linear Velocity */
+	FVector3f PreviousBoneLinearVelocity;
+
+	/** Bone Angular Velocity */
+	FVector3f BoneAngularVelocity;
+
+	/** Bone Previous Angular Velocity */
+	FVector3f PreviousBoneAngularVelocity;
+
+	/** Bone Linear Acceleration */
+	FVector3f BoneLinearAcceleration;
+
+	/** Bone Angular Acceleration */
+	FVector3f BoneAngularAcceleration;
 
 	/** Global Interpolation */
 	bool GlobalInterpolation;
@@ -621,6 +661,18 @@ public:
 	/** Name of the world rotation */
 	static const FString BoneRotationName;
 
+	/** Name of the bone linear velocity */
+	static const FString BoneLinearVelocityName;
+
+	/** Name of the bone linear acceleration */
+	static const FString BoneLinearAccelerationName;
+
+	/** Name of the bone angular velocity */
+	static const FString BoneAngularVelocityName;
+
+	/** Name of the bone angular acceleration */
+	static const FString BoneAngularAccelerationName;
+
 	/** Name of the number of strands */
 	static const FString NumStrandsName;
 
