Fix live edit of groom paramaters.

The serialization of the hair description bulk data during the undo/redo stack change was causing the edit to freeze for a few second. Since this serizliation is not needed for transient/transactional object, we bypass it.

diff --git a/Engine/Plugins/Runtime/HairStrands/Source/HairStrandsCore/Private/GroomAsset.cpp b/Engine/Plugins/Runtime/HairStrands/Source/HairStrandsCore/Private/GroomAsset.cpp
--- a/Engine/Plugins/Runtime/HairStrands/Source/HairStrandsCore/Private/GroomAsset.cpp
+++ b/Engine/Plugins/Runtime/HairStrands/Source/HairStrandsCore/Private/GroomAsset.cpp
@@ -747,7 +747,11 @@ void UGroomAsset::Serialize(FArchive& Ar)
 				HairDescriptionBulkData = MakeUnique<FHairDescriptionBulkData>();
 			}
 
-			HairDescriptionBulkData->Serialize(Ar, this);
+			// We don' serialize the hair description if it is a transactional object (i.e., for undo/redo)
+			if (!Ar.IsTransacting())
+			{
+				HairDescriptionBulkData->Serialize(Ar, this);
+			}
 		}
 #endif // WITH_EDITORONLY_DATA
 	}
