Fix for hair group texture generation

diff --git a/Engine/Plugins/Runtime/HairStrands/Source/HairStrandsCore/Private/GroomTextureBuilder.cpp b/Engine/Plugins/Runtime/HairStrands/Source/HairStrandsCore/Private/GroomTextureBuilder.cpp
--- a/Engine/Plugins/Runtime/HairStrands/Source/HairStrandsCore/Private/GroomTextureBuilder.cpp
+++ b/Engine/Plugins/Runtime/HairStrands/Source/HairStrandsCore/Private/GroomTextureBuilder.cpp
@@ -485,6 +485,8 @@ static void InternalAllocateStrandsTexture(UTexture2D* Out, const FIntPoint& Res
 {
 	FTextureFormatSettings FormatSettings;
 	FormatSettings.SRGB = false;
+	FormatSettings.CompressionNone = true;
+	FormatSettings.CompressionSettings = TC_Masks;
 #if WITH_EDITORONLY_DATA
 	Out->Source.Init(Resolution.X, Resolution.Y, 1, MipCount, SourceFormat, nullptr);
 	Out->SetLayerFormatSettings(0, FormatSettings);
@@ -829,13 +831,8 @@ static void AddReadbackPass(
 		int32 Width = 0, Height = 0;
 		RHICmdList.MapStagingSurface(StagingTexture, InData, Width, Height);
 		uint8* InDataRGBA8 = (uint8*)InData;
-
-		uint8 MipIndex = 0;
-		const uint32 SizeInBytes = BytePerPixel * Resolution.X * Resolution.Y;
-		uint8* OutData = OutTexture->Source.LockMip(0);
-		FMemory::Memcpy(OutData, InDataRGBA8, SizeInBytes);
-		OutTexture->Source.UnlockMip(0);
-
+		ETextureSourceFormat SrcFormat = OutTexture->Source.GetFormat();
+		OutTexture->Source.Init(Width, Height, 1, 1, SrcFormat, InDataRGBA8);
 		RHICmdList.UnmapStagingSurface(StagingTexture);
 
 		OutTexture->DeferCompression = true; // This forces reloading data when the asset is saved
diff --git a/Engine/Shaders/Private/HairStrands/HairStrandsTexturesGeneration.usf b/Engine/Shaders/Private/HairStrands/HairStrandsTexturesGeneration.usf
--- a/Engine/Shaders/Private/HairStrands/HairStrandsTexturesGeneration.usf
+++ b/Engine/Shaders/Private/HairStrands/HairStrandsTexturesGeneration.usf
@@ -334,7 +334,7 @@ void MainPS(
 	out float4 OutTangent	: SV_Target1,
 	out float  OutCoverage	: SV_Target2,
 	out float4 OutAttribute : SV_Target3,
-	out float  OutGroupIndex : SV_Target4,
+	out float4 OutGroupIndex : SV_Target4,
 	out uint   OutTriangleMask : SV_Target5,
 	out float  OutDepthTest	: SV_Depth)
 {
@@ -424,7 +424,7 @@ void MainPS(
 		OutCoverage /= SampleCount;
 		OutTangent.xyz = (Tangent_FrameSpace + 1) * 0.5f;
 		OutTangent.w = 1;
-		OutGroupIndex = InVF_GroupIndex / 255.f;
+		OutGroupIndex = float4(InVF_GroupIndex / 255.f, 0, 0, 0);
 	}
 	else
 	{
