Add material attributes support for hair cards & meshes.

This CLs add transfer (through voxelization or software raytracing) of strands material properties (base color, roughness), in order to get parity between the different geometric representation.

diff --git a/Engine/Plugins/Runtime/HairStrands/Source/HairStrandsCore/Private/GroomAsset.cpp b/Engine/Plugins/Runtime/HairStrands/Source/HairStrandsCore/Private/GroomAsset.cpp
--- a/Engine/Plugins/Runtime/HairStrands/Source/HairStrandsCore/Private/GroomAsset.cpp
+++ b/Engine/Plugins/Runtime/HairStrands/Source/HairStrandsCore/Private/GroomAsset.cpp
@@ -416,10 +416,10 @@ static void InitAtlasTexture(ResourceType* InResource, UTexture2D* InTexture, EH
 			InResource->AuxilaryDataTexture = InTexture->TextureReference.TextureReferenceRHI;
 			InResource->AuxilaryDataSampler = DefaultSampler;
 		} break;
-		case EHairAtlasTextureType::GroupIndex:
+		case EHairAtlasTextureType::Material:
 		{
-			InResource->GroupIndexTexture = InTexture->TextureReference.TextureReferenceRHI;
-			InResource->GroupIndexSampler = DefaultSampler;
+			InResource->MaterialTexture = InTexture->TextureReference.TextureReferenceRHI;
+			InResource->MaterialSampler = DefaultSampler;
 		} break;
 		}
 	});
@@ -1294,7 +1294,7 @@ static bool IsCardsTextureResources(const FName PropertyName)
 		|| PropertyName == GET_MEMBER_NAME_CHECKED(FHairGroupCardsTextures, TangentTexture)
 		|| PropertyName == GET_MEMBER_NAME_CHECKED(FHairGroupCardsTextures, AttributeTexture)
 		|| PropertyName == GET_MEMBER_NAME_CHECKED(FHairGroupCardsTextures, AuxilaryDataTexture)
-		|| PropertyName == GET_MEMBER_NAME_CHECKED(FHairGroupCardsTextures, GroupIndexTexture);
+		|| PropertyName == GET_MEMBER_NAME_CHECKED(FHairGroupCardsTextures, MaterialTexture);
 }
 static void InitCardsTextureResources(UGroomAsset* GroomAsset);
 
@@ -2469,7 +2469,7 @@ bool UGroomAsset::BuildCardsGeometry(uint32 GroupIndex)
 					if (Textures.AttributeTexture != nullptr)	Textures.AttributeTexture->UpdateResource();
 					if (Textures.CoverageTexture != nullptr)	Textures.CoverageTexture->UpdateResource();
 					if (Textures.AuxilaryDataTexture != nullptr)Textures.AuxilaryDataTexture->UpdateResource();
-					if (Textures.GroupIndexTexture != nullptr)	Textures.GroupIndexTexture->UpdateResource();
+					if (Textures.MaterialTexture != nullptr)	Textures.MaterialTexture->UpdateResource();
 				}
 			}
 		}
@@ -2571,7 +2571,7 @@ bool UGroomAsset::BuildCardsGeometry(uint32 GroupIndex)
 				InitAtlasTexture(LOD.RestResource, Desc->Textures.AttributeTexture, EHairAtlasTextureType::Attribute);
 				InitAtlasTexture(LOD.RestResource, Desc->Textures.CoverageTexture, EHairAtlasTextureType::Coverage);
 				InitAtlasTexture(LOD.RestResource, Desc->Textures.AuxilaryDataTexture, EHairAtlasTextureType::AuxilaryData);
-				InitAtlasTexture(LOD.RestResource, Desc->Textures.GroupIndexTexture, EHairAtlasTextureType::GroupIndex);
+				InitAtlasTexture(LOD.RestResource, Desc->Textures.MaterialTexture, EHairAtlasTextureType::Material);
 				LOD.RestResource->bInvertUV = Desc->SourceType == EHairCardsSourceType::Procedural;
 				
 				// 2.2 Load interoplatino resources
@@ -2733,7 +2733,7 @@ bool UGroomAsset::BuildMeshesGeometry(uint32 GroupIndex)
 					if (Textures.AttributeTexture != nullptr)	Textures.AttributeTexture->UpdateResource();
 					if (Textures.CoverageTexture != nullptr)	Textures.CoverageTexture->UpdateResource();
 					if (Textures.AuxilaryDataTexture != nullptr)Textures.AuxilaryDataTexture->UpdateResource();
-					if (Textures.GroupIndexTexture != nullptr)	Textures.GroupIndexTexture->UpdateResource();
+					if (Textures.MaterialTexture != nullptr)	Textures.MaterialTexture->UpdateResource();
 				}
 			}
 		}
@@ -2994,7 +2994,7 @@ static void InitCardsTextureResources(UGroomAsset* GroomAsset)
 					InitAtlasTexture(LOD.RestResource, Desc->Textures.AttributeTexture, EHairAtlasTextureType::Attribute);
 					InitAtlasTexture(LOD.RestResource, Desc->Textures.CoverageTexture, EHairAtlasTextureType::Coverage);
 					InitAtlasTexture(LOD.RestResource, Desc->Textures.AuxilaryDataTexture, EHairAtlasTextureType::AuxilaryData);
-					InitAtlasTexture(LOD.RestResource, Desc->Textures.GroupIndexTexture, EHairAtlasTextureType::GroupIndex);
+					InitAtlasTexture(LOD.RestResource, Desc->Textures.MaterialTexture, EHairAtlasTextureType::Material);
 					if (LOD.RestResource)
 					{
 						LOD.RestResource->bInvertUV = Desc->SourceType == EHairCardsSourceType::Procedural; // Should fix procedural texture so that this does not happen
@@ -3053,7 +3053,7 @@ void UGroomAsset::InitCardsResources()
 					InitAtlasTexture(LOD.RestResource, Desc->Textures.AttributeTexture, EHairAtlasTextureType::Attribute);
 					InitAtlasTexture(LOD.RestResource, Desc->Textures.CoverageTexture, EHairAtlasTextureType::Coverage);
 					InitAtlasTexture(LOD.RestResource, Desc->Textures.AuxilaryDataTexture, EHairAtlasTextureType::AuxilaryData);
-					InitAtlasTexture(LOD.RestResource, Desc->Textures.GroupIndexTexture, EHairAtlasTextureType::GroupIndex);
+					InitAtlasTexture(LOD.RestResource, Desc->Textures.MaterialTexture, EHairAtlasTextureType::Material);
 					LOD.RestResource->bInvertUV = Desc->SourceType == EHairCardsSourceType::Procedural; // Should fix procedural texture so that this does not happen
 				}
 			}
@@ -3106,7 +3106,7 @@ void UGroomAsset::InitMeshesResources()
 					InitAtlasTexture(LOD.RestResource, Desc->Textures.AttributeTexture, EHairAtlasTextureType::Attribute);
 					InitAtlasTexture(LOD.RestResource, Desc->Textures.CoverageTexture, EHairAtlasTextureType::Coverage);
 					InitAtlasTexture(LOD.RestResource, Desc->Textures.AuxilaryDataTexture, EHairAtlasTextureType::AuxilaryData);
-					InitAtlasTexture(LOD.RestResource, Desc->Textures.GroupIndexTexture, EHairAtlasTextureType::GroupIndex);
+					InitAtlasTexture(LOD.RestResource, Desc->Textures.MaterialTexture, EHairAtlasTextureType::Material);
 				}
 			}
 		}
@@ -3569,7 +3569,7 @@ void UGroomAsset::SavePendingProceduralAssets()
 					if (Q->Textures->AuxilaryDataTexture)	FHairStrandsCore::SaveAsset(Q->Textures->AuxilaryDataTexture);
 					if (Q->Textures->CoverageTexture)		FHairStrandsCore::SaveAsset(Q->Textures->CoverageTexture);
 					if (Q->Textures->TangentTexture)		FHairStrandsCore::SaveAsset(Q->Textures->TangentTexture);
-					if (Q->Textures->GroupIndexTexture)		FHairStrandsCore::SaveAsset(Q->Textures->GroupIndexTexture);
+					if (Q->Textures->MaterialTexture)		FHairStrandsCore::SaveAsset(Q->Textures->MaterialTexture);
 					Q->Textures->bNeedToBeSaved = false;
 
 					InternalReleaseResource(Q->Resources);
diff --git a/Engine/Plugins/Runtime/HairStrands/Source/HairStrandsCore/Private/GroomBuilder.cpp b/Engine/Plugins/Runtime/HairStrands/Source/HairStrandsCore/Private/GroomBuilder.cpp
--- a/Engine/Plugins/Runtime/HairStrands/Source/HairStrandsCore/Private/GroomBuilder.cpp
+++ b/Engine/Plugins/Runtime/HairStrands/Source/HairStrandsCore/Private/GroomBuilder.cpp
@@ -2066,9 +2066,11 @@ struct FHairGrid
 {
 	struct FCurve
 	{
-		uint32 CurveIndex = 0;
-		float  Radius = 0;
-		uint8  GroupIndex = 0;
+		FVector3f BaseColor = FVector3f::ZeroVector;
+		float     Roughness = 0;
+		uint32    CurveIndex = 0;
+		float     Radius = 0;
+		uint8     GroupIndex = 0;
 	};
 
 	struct FVoxel
@@ -2133,6 +2135,10 @@ static void Voxelize(const FHairDescriptionGroups& InGroups, FHairGrid& Out)
 				const FVector3f& P1 = In.StrandsPoints.PointsPosition[Index1];
 				const float R0 = In.StrandsPoints.PointsRadius[Index0];
 				const float R1 = In.StrandsPoints.PointsRadius[Index1];
+				const FVector3f C0(In.StrandsPoints.PointsBaseColor[Index0].R, In.StrandsPoints.PointsBaseColor[Index0].G, In.StrandsPoints.PointsBaseColor[Index0].B);
+				const FVector3f C1(In.StrandsPoints.PointsBaseColor[Index1].R, In.StrandsPoints.PointsBaseColor[Index1].G, In.StrandsPoints.PointsBaseColor[Index1].B);
+				const float Rough0 = In.StrandsPoints.PointsRoughness[Index0];
+				const float Rough1 = In.StrandsPoints.PointsRoughness[Index1];
 				const FVector3f Segment = P1 - P0;
 
 				// This is a coarse/non-conservative voxelization, by ray-marching segment, instead of walking intersected voxels
@@ -2142,7 +2148,6 @@ static void Voxelize(const FHairDescriptionGroups& InGroups, FHairGrid& Out)
 				{
 					const float T = float(StepIt) / float(StepCount);
 					const FVector3f P = P0 + Segment * T;
-					const float R = FMath::Lerp(R0, R1, T);
 					const FIntVector Coord = ToCoord(P, Out.Resolution, Out.MinBound, Out.VoxelSize);
 					const uint32 LinearCoord = ToLinearCoord(Coord, Out.Resolution);
 					if (LinearCoord != PrevLinearCoord)
@@ -2150,7 +2155,9 @@ static void Voxelize(const FHairDescriptionGroups& InGroups, FHairGrid& Out)
 						FHairGrid::FCurve RCurve;
 						RCurve.CurveIndex = CurveIt;
 						RCurve.GroupIndex = GroupIndex;
-						RCurve.Radius = R;
+						RCurve.Radius = FMath::Lerp(R0, R1, T);
+						RCurve.BaseColor = FVector3f(FMath::Lerp(C0.X, C1.X, T), FMath::Lerp(C0.Y, C1.Y, T), FMath::Lerp(C0.Z, C1.Z, T));
+						RCurve.Roughness = FMath::Lerp(Rough0, Rough1, T);
 						Out.Voxels[LinearCoord].VoxelCurves.Add(RCurve);
 
 						PrevLinearCoord = LinearCoord;
@@ -2163,30 +2170,33 @@ static void Voxelize(const FHairDescriptionGroups& InGroups, FHairGrid& Out)
 	}
 }
 
-FORCEINLINE void SearchCell(const FHairStrandsVoxelData& In, const FIntVector& C, uint8& Out)
+FORCEINLINE void SearchCell(const FHairStrandsVoxelData& In, const FIntVector& C, FHairStrandsVoxelData::FData& Out)
 {
 	const uint32 I = GroomBuilder_Voxelization::ToIndex(C, In.Resolution);
-	if (In.GroupIndices[I] != FHairStrandsVoxelData::InvalidGroupIndex)
+	if (In.Datas[I].GroupIndex != FHairStrandsVoxelData::InvalidGroupIndex)
 	{
-		Out = In.GroupIndices[I];
+		Out = In.Datas[I];
 	}
 }
 
 } // namespace GroomBuilder_Voxelization
 
-uint8 FHairStrandsVoxelData::GetGroupIndex(const FVector3f& P) const
+FHairStrandsVoxelData::FData FHairStrandsVoxelData::GetData(const FVector3f& P) const
 {
 	const uint32 MaxLookupDistance = FMath::Max3(Resolution.X, Resolution.Y, Resolution.Z) * 0.5f;
 	const FIntVector C = GroomBuilder_Voxelization::ToCellCoord(P, MinBound, MaxBound, Resolution);
 
-	uint8 GroupIndex = FHairStrandsVoxelData::InvalidGroupIndex;
+	FHairStrandsVoxelData::FData Out;
+	Out.GroupIndex = FHairStrandsVoxelData::InvalidGroupIndex;
+	Out.BaseColor = FVector3f::ZeroVector;
+	Out.Roughness = 0.f;
 	for (int32 Offset = 1; Offset <= int32(MaxLookupDistance); ++Offset)
 	{
 		// Center
 		{
 			bool bIsValid = false;
 			const FIntVector CellCoord = GroomBuilder_Voxelization::ClampToVolume(C, Resolution, bIsValid);
-			if (bIsValid) GroomBuilder_Voxelization::SearchCell(*this, CellCoord, GroupIndex);
+			if (bIsValid) GroomBuilder_Voxelization::SearchCell(*this, CellCoord, Out);
 		}
 
 		// Top & Bottom
@@ -2196,8 +2206,8 @@ uint8 FHairStrandsVoxelData::GetGroupIndex(const FVector3f& P) const
 			bool bIsValid0 = false, bIsValid1 = false;
 			const FIntVector CellCoord0 = GroomBuilder_Voxelization::ClampToVolume(C + FIntVector(X, Y, Offset), Resolution, bIsValid0);
 			const FIntVector CellCoord1 = GroomBuilder_Voxelization::ClampToVolume(C + FIntVector(X, Y, -Offset), Resolution, bIsValid1);
-			if (bIsValid0) GroomBuilder_Voxelization::SearchCell(*this, CellCoord0, GroupIndex);
-			if (bIsValid1) GroomBuilder_Voxelization::SearchCell(*this, CellCoord1, GroupIndex);
+			if (bIsValid0) GroomBuilder_Voxelization::SearchCell(*this, CellCoord0, Out);
+			if (bIsValid1) GroomBuilder_Voxelization::SearchCell(*this, CellCoord1, Out);
 		}
 
 		const int32 OffsetMinusOne = Offset - 1;
@@ -2208,8 +2218,8 @@ uint8 FHairStrandsVoxelData::GetGroupIndex(const FVector3f& P) const
 			bool bIsValid0 = false, bIsValid1 = false;
 			const FIntVector CellCoord0 = GroomBuilder_Voxelization::ClampToVolume(C + FIntVector(X, Offset, Z), Resolution, bIsValid0);
 			const FIntVector CellCoord1 = GroomBuilder_Voxelization::ClampToVolume(C + FIntVector(X, -Offset, Z), Resolution, bIsValid1);
-			if (bIsValid0) GroomBuilder_Voxelization::SearchCell(*this, CellCoord0, GroupIndex);
-			if (bIsValid1) GroomBuilder_Voxelization::SearchCell(*this, CellCoord1, GroupIndex);
+			if (bIsValid0) GroomBuilder_Voxelization::SearchCell(*this, CellCoord0, Out);
+			if (bIsValid1) GroomBuilder_Voxelization::SearchCell(*this, CellCoord1, Out);
 		}
 
 		// Left & Right
@@ -2219,22 +2229,22 @@ uint8 FHairStrandsVoxelData::GetGroupIndex(const FVector3f& P) const
 			bool bIsValid0 = false, bIsValid1 = false;
 			const FIntVector CellCoord0 = GroomBuilder_Voxelization::ClampToVolume(C + FIntVector(Offset, Y, Z), Resolution, bIsValid0);
 			const FIntVector CellCoord1 = GroomBuilder_Voxelization::ClampToVolume(C + FIntVector(-Offset, Y, Z), Resolution, bIsValid1);
-			if (bIsValid0) GroomBuilder_Voxelization::SearchCell(*this, CellCoord0, GroupIndex);
-			if (bIsValid1) GroomBuilder_Voxelization::SearchCell(*this, CellCoord1, GroupIndex);
+			if (bIsValid0) GroomBuilder_Voxelization::SearchCell(*this, CellCoord0, Out);
+			if (bIsValid1) GroomBuilder_Voxelization::SearchCell(*this, CellCoord1, Out);
 		}
 
 		// Early out if we have found a valid group index during a ring/layer step.
-		if (GroupIndex != InvalidGroupIndex)
+		if (Out.GroupIndex != InvalidGroupIndex)
 		{
 			break;
 		}
 	}
 
 	// No valid group index has been found. In this case force the group index to 0
-	check(GroupIndex != FHairStrandsVoxelData::InvalidGroupIndex);
-	if (GroupIndex == FHairStrandsVoxelData::InvalidGroupIndex) { GroupIndex = 0; } 
+	check(Out.GroupIndex != FHairStrandsVoxelData::InvalidGroupIndex);
+	if (Out.GroupIndex == FHairStrandsVoxelData::InvalidGroupIndex) { Out.GroupIndex = 0; }
 
-	return GroupIndex;
+	return Out;
 }
 
 void FGroomBuilder::VoxelizeGroupIndex(const FHairDescriptionGroups& In, FHairStrandsVoxelData& Out)
@@ -2244,12 +2254,14 @@ void FGroomBuilder::VoxelizeGroupIndex(const FHairDescriptionGroups& In, FHairSt
 	Out.MinBound = Grid.MinBound;
 	Out.MaxBound = Grid.MaxBound;
 	Out.Resolution = Grid.Resolution;
-	Out.GroupIndices.SetNum(Grid.Voxels.Num());
+	Out.Datas.SetNum(Grid.Voxels.Num());
 
-	TArray<uint8> GroupBins;
+	TArray<uint16> GroupBins;
 	for (uint32 I=0,VoxelCount=Grid.Voxels.Num();I<VoxelCount;++I)
 	{
-		Out.GroupIndices[I] = FHairStrandsVoxelData::InvalidGroupIndex;
+		Out.Datas[I].GroupIndex = FHairStrandsVoxelData::InvalidGroupIndex;
+		Out.Datas[I].BaseColor = FVector3f::ZeroVector;
+		Out.Datas[I].Roughness = 0.f;
 
 		// 1. For a given voxel, build an histogram of all the groups falling into this voxel
 		GroupBins.Init(0, 64u);
@@ -2258,24 +2270,33 @@ void FGroomBuilder::VoxelizeGroupIndex(const FHairDescriptionGroups& In, FHairSt
 			check(Curve.GroupIndex != FHairStrandsVoxelData::InvalidGroupIndex);
 			check(Curve.GroupIndex < GroupBins.Num());
 			GroupBins[Curve.GroupIndex]++;
+			Out.Datas[I].BaseColor += Curve.BaseColor;
+			Out.Datas[I].Roughness += Curve.Roughness;
+		}
+
+		const uint32 CurveCount = Grid.Voxels[I].VoxelCurves.Num() > 0;
+		if (CurveCount > 0)
+		{
+			Out.Datas[I].BaseColor /= float(CurveCount);
+			Out.Datas[I].Roughness /= float(CurveCount);
 		}
 
 		// 2. Select the hair group having the largest count
-		if (Grid.Voxels[I].VoxelCurves.Num() > 1)
+		if (CurveCount > 1)
 		{
 			uint32 MaxBinCount = 0u;
 			for (uint8 GroupIndex = 0, GroupCount = GroupBins.Num(); GroupIndex < GroupCount; ++GroupIndex)
 			{
 				if (GroupBins[GroupIndex] > MaxBinCount)
 				{
-					Out.GroupIndices[I] = GroupIndex;
+					Out.Datas[I].GroupIndex = GroupIndex;
 					MaxBinCount = GroupBins[GroupIndex];
 				}
 			}
 		}
-		else if (Grid.Voxels[I].VoxelCurves.Num() == 1)
+		else if (CurveCount == 1)
 		{
-			Out.GroupIndices[I] = Grid.Voxels[I].VoxelCurves[0].GroupIndex;
+			Out.Datas[I].GroupIndex = Grid.Voxels[I].VoxelCurves[0].GroupIndex;
 		}
 	}
 }
diff --git a/Engine/Plugins/Runtime/HairStrands/Source/HairStrandsCore/Private/GroomResources.cpp b/Engine/Plugins/Runtime/HairStrands/Source/HairStrandsCore/Private/GroomResources.cpp
--- a/Engine/Plugins/Runtime/HairStrands/Source/HairStrandsCore/Private/GroomResources.cpp
+++ b/Engine/Plugins/Runtime/HairStrands/Source/HairStrandsCore/Private/GroomResources.cpp
@@ -570,6 +570,7 @@ FHairCardsRestResource::FHairCardsRestResource(const FHairCardsBulkData& InBulkD
 	RestIndexBuffer(InBulkData.Indices),
 	NormalsBuffer(),
 	UVsBuffer(),
+	MaterialsBuffer(),
 	BulkData(InBulkData)
 {
 
@@ -581,12 +582,14 @@ void FHairCardsRestResource::InternalAllocate()
 	CreateBuffer<FHairCardsPositionFormat>(BulkData.Positions, RestPositionBuffer, ToHairResourceDebugName(TEXT("Hair.CardsRest_PositionBuffer"), ResourceName));
 	CreateBuffer<FHairCardsNormalFormat>(BulkData.Normals, NormalsBuffer, ToHairResourceDebugName(TEXT("Hair.CardsRest_NormalBuffer"), ResourceName));
 	CreateBuffer<FHairCardsUVFormat>(BulkData.UVs, UVsBuffer, ToHairResourceDebugName(TEXT("Hair.CardsRest_UVBuffer"), ResourceName));
+	CreateBuffer<FHairCardsMaterialFormat>(BulkData.Materials, MaterialsBuffer, ToHairResourceDebugName(TEXT("Hair.CardsRest_MaterialBuffer"), ResourceName));
 
 	FSamplerStateRHIRef DefaultSampler = TStaticSamplerState<SF_Bilinear, AM_Clamp, AM_Clamp, AM_Clamp>::GetRHI();
 	DepthSampler = DefaultSampler;
 	TangentSampler = DefaultSampler;
 	CoverageSampler = DefaultSampler;
 	AttributeSampler = DefaultSampler;
+	MaterialSampler = DefaultSampler;
 }
 
 void FHairCardsRestResource::InternalRelease()
@@ -600,6 +603,7 @@ void FHairCardsRestResource::InitResource()
 	RestPositionBuffer.InitResource();
 	NormalsBuffer.InitResource();
 	UVsBuffer.InitResource();
+	MaterialsBuffer.InitResource();
 }
 
 void FHairCardsRestResource::ReleaseResource()
@@ -609,6 +613,7 @@ void FHairCardsRestResource::ReleaseResource()
 	RestPositionBuffer.ReleaseResource();
 	NormalsBuffer.ReleaseResource();
 	UVsBuffer.ReleaseResource();
+	MaterialsBuffer.ReleaseResource();
 }
 
 /////////////////////////////////////////////////////////////////////////////////////////
diff --git a/Engine/Plugins/Runtime/HairStrands/Source/HairStrandsCore/Private/GroomTextureBuilder.cpp b/Engine/Plugins/Runtime/HairStrands/Source/HairStrandsCore/Private/GroomTextureBuilder.cpp
--- a/Engine/Plugins/Runtime/HairStrands/Source/HairStrandsCore/Private/GroomTextureBuilder.cpp
+++ b/Engine/Plugins/Runtime/HairStrands/Source/HairStrandsCore/Private/GroomTextureBuilder.cpp
@@ -416,14 +416,14 @@ class FHairStrandTextureDilationCS : public FGlobalShader
 		SHADER_PARAMETER_RDG_TEXTURE(Texture2D, Source_CoverageTexture)
 		SHADER_PARAMETER_RDG_TEXTURE(Texture2D, Source_TangentTexture)
 		SHADER_PARAMETER_RDG_TEXTURE(Texture2D, Source_AttributeTexture)
-		SHADER_PARAMETER_RDG_TEXTURE(Texture2D, Source_GroupIndexTexture)
+		SHADER_PARAMETER_RDG_TEXTURE(Texture2D, Source_MaterialTexture)
 
 		SHADER_PARAMETER_RDG_TEXTURE_UAV(RWTexture2D, TriangleMaskTexture)
 		SHADER_PARAMETER_RDG_TEXTURE_UAV(RWTexture2D, Target_DepthTexture)
 		SHADER_PARAMETER_RDG_TEXTURE_UAV(RWTexture2D, Target_CoverageTexture)
 		SHADER_PARAMETER_RDG_TEXTURE_UAV(RWTexture2D, Target_TangentTexture)
 		SHADER_PARAMETER_RDG_TEXTURE_UAV(RWTexture2D, Target_AttributeTexture)
-		SHADER_PARAMETER_RDG_TEXTURE_UAV(RWTexture2D, Target_GroupIndexTexture)
+		SHADER_PARAMETER_RDG_TEXTURE_UAV(RWTexture2D, Target_MaterialTexture)
 	END_SHADER_PARAMETER_STRUCT()
 
 public:
@@ -447,13 +447,13 @@ static void AddTextureDilationPass(
 	FRDGTextureRef Source_CoverageTexture,
 	FRDGTextureRef Source_TangentTexture,
 	FRDGTextureRef Source_AttributeTexture,
-	FRDGTextureRef Source_GroupIndexTexture,
+	FRDGTextureRef Source_MaterialTexture,
 
 	FRDGTextureRef Target_DepthTexture,
 	FRDGTextureRef Target_CoverageTexture,
 	FRDGTextureRef Target_TangentTexture,
 	FRDGTextureRef Target_AttributeTexture,
-	FRDGTextureRef Target_GroupIndexTexture)
+	FRDGTextureRef Target_MaterialTexture)
 {
 	FHairStrandTextureDilationCS::FParameters* Parameters = GraphBuilder.AllocParameters<FHairStrandTextureDilationCS::FParameters>();
 	Parameters->Resolution				= Resolution;
@@ -462,12 +462,12 @@ static void AddTextureDilationPass(
 	Parameters->Source_CoverageTexture	= Source_CoverageTexture;
 	Parameters->Source_TangentTexture	= Source_TangentTexture;
 	Parameters->Source_AttributeTexture = Source_AttributeTexture;
-	Parameters->Source_GroupIndexTexture= Source_GroupIndexTexture;
+	Parameters->Source_MaterialTexture  = Source_MaterialTexture;
 	Parameters->Target_DepthTexture		= GraphBuilder.CreateUAV(Target_DepthTexture);
 	Parameters->Target_CoverageTexture	= GraphBuilder.CreateUAV(Target_CoverageTexture);
 	Parameters->Target_TangentTexture	= GraphBuilder.CreateUAV(Target_TangentTexture);
 	Parameters->Target_AttributeTexture = GraphBuilder.CreateUAV(Target_AttributeTexture);
-	Parameters->Target_GroupIndexTexture= GraphBuilder.CreateUAV(Target_GroupIndexTexture);
+	Parameters->Target_MaterialTexture  = GraphBuilder.CreateUAV(Target_MaterialTexture);
 
 	TShaderMapRef<FHairStrandTextureDilationCS> ComputeShader(ShaderMap);
 	FIntVector DispatchCount = FComputeShaderUtils::GetGroupCount(Resolution, FIntPoint(8, 4));
@@ -520,7 +520,7 @@ static void InternalAllocateStrandsTexture_Attribute(UTexture2D* Out, const FInt
 	InternalAllocateStrandsTexture(Out, Resolution, 1, PF_B8G8R8A8, ETextureSourceFormat::TSF_BGRA8);
 }
 
-static void InternalAllocateStrandsTexture_GroupIndex(UTexture2D* Out, const FIntPoint& Resolution, uint32 MipCount)
+static void InternalAllocateStrandsTexture_Material(UTexture2D* Out, const FIntPoint& Resolution, uint32 MipCount)
 {
 	InternalAllocateStrandsTexture(Out, Resolution, 1, PF_B8G8R8A8, ETextureSourceFormat::TSF_BGRA8);
 }
@@ -541,7 +541,7 @@ FStrandsTexturesOutput FGroomTextureBuilder::CreateGroomStrandsTexturesTexture(c
 	Output.Coverage = FHairStrandsCore::CreateTexture(PackageName, Resolution, TEXT("_Opacity"), InternalAllocateStrandsTexture_Coverage);
 	Output.Tangent = FHairStrandsCore::CreateTexture(PackageName, Resolution, TEXT("_Tangent"), InternalAllocateStrandsTexture_Tangent);
 	Output.Attribute = FHairStrandsCore::CreateTexture(PackageName, Resolution, TEXT("_Attribute"), InternalAllocateStrandsTexture_Attribute);
-	Output.GroupIndex = FHairStrandsCore::CreateTexture(PackageName, Resolution, TEXT("_GroupIndex"), InternalAllocateStrandsTexture_GroupIndex);
+	Output.Material = FHairStrandsCore::CreateTexture(PackageName, Resolution, TEXT("_Material"), InternalAllocateStrandsTexture_Material);
 	return Output;
 }
 #endif
@@ -582,6 +582,7 @@ class FHairStrandsTexturePS : public FGlobalShader
 		SHADER_PARAMETER(uint32, VertexCount)
 		SHADER_PARAMETER(float, MaxDistance)
 		SHADER_PARAMETER(int32, TracingDirection)
+		SHADER_PARAMETER(uint32, bHasMaterialData)
 
 		SHADER_PARAMETER(uint32, UVsChannelIndex)
 		SHADER_PARAMETER(uint32, UVsChannelCount)
@@ -591,6 +592,7 @@ class FHairStrandsTexturePS : public FGlobalShader
 		SHADER_PARAMETER(FVector3f, InVF_PositionOffset)
 		SHADER_PARAMETER_SRV(Buffer, InVF_PositionBuffer)
 		SHADER_PARAMETER_SRV(Buffer, InVF_Attribute0Buffer)
+		SHADER_PARAMETER_SRV(Buffer, InVF_MaterialBuffer)
 		SHADER_PARAMETER(uint32, InVF_ControlPointCount)
 		SHADER_PARAMETER(uint32, InVF_GroupIndex)
 
@@ -651,6 +653,7 @@ static void InternalGenerateHairStrandsTextures(
 	
 	FRHIShaderResourceView* InHairStrands_PositionBuffer,
 	FRHIShaderResourceView* InHairStrands_AttributeBuffer,
+	FRHIShaderResourceView* InHairStrands_MaterialBuffer,
 	const FVector& InHairStrands_PositionOffset,
 	float InHairStrands_Radius,
 	float InHairStrands_Length,
@@ -662,10 +665,11 @@ static void InternalGenerateHairStrandsTextures(
 	FRDGTextureRef OutTangentTexture,
 	FRDGTextureRef OutCoverageTexture,
 	FRDGTextureRef OutRootUVStrandsUSeedTexture,
-	FRDGTextureRef OutGroupIndexTexture,
+	FRDGTextureRef OutMaterialTexture,
 	FRDGTextureRef OutTriangleMask)
 {
 	const FIntPoint OutputResolution = OutDepthTexture->Desc.Extent;
+	const bool bHasMaterialData = InHairStrands_MaterialBuffer != nullptr;
 
 	FHairStrandsTexturePS::FParameters* ParametersPS = GraphBuilder.AllocParameters<FHairStrandsTexturePS::FParameters>();
 	ParametersPS->OutputResolution = OutputResolution;
@@ -675,6 +679,7 @@ static void InternalGenerateHairStrandsTextures(
 	ParametersPS->NormalsBuffer = InMeshNormalsBuffer;
 	ParametersPS->MaxDistance = FMath::Abs(InMaxDistance);
 	ParametersPS->TracingDirection = InTracingDirection;
+	ParametersPS->bHasMaterialData = bHasMaterialData ? 1u : 0u;
 
 	ParametersPS->UVsChannelIndex = UVsChannelIndex;
 	ParametersPS->UVsChannelCount = UVsChannelCount;
@@ -682,6 +687,7 @@ static void InternalGenerateHairStrandsTextures(
 	ParametersPS->InVF_PositionBuffer = InHairStrands_PositionBuffer;
 	ParametersPS->InVF_PositionOffset = InHairStrands_PositionOffset;
 	ParametersPS->InVF_Attribute0Buffer = InHairStrands_AttributeBuffer;
+	ParametersPS->InVF_MaterialBuffer = bHasMaterialData ? InHairStrands_MaterialBuffer : InHairStrands_AttributeBuffer;
 	ParametersPS->InVF_Radius = InHairStrands_Radius;
 	ParametersPS->InVF_Length = InHairStrands_Length;
 	ParametersPS->InVF_ControlPointCount = InHairStrands_ControlPointCount;
@@ -704,7 +710,7 @@ static void InternalGenerateHairStrandsTextures(
 	ParametersPS->RenderTargets[1] = FRenderTargetBinding(OutTangentTexture, bClear ? ERenderTargetLoadAction::EClear : ERenderTargetLoadAction::ELoad);
 	ParametersPS->RenderTargets[2] = FRenderTargetBinding(OutCoverageTexture, bClear ? ERenderTargetLoadAction::EClear : ERenderTargetLoadAction::ELoad);
 	ParametersPS->RenderTargets[3] = FRenderTargetBinding(OutRootUVStrandsUSeedTexture, bClear ? ERenderTargetLoadAction::EClear : ERenderTargetLoadAction::ELoad);
-	ParametersPS->RenderTargets[4] = FRenderTargetBinding(OutGroupIndexTexture, bClear ? ERenderTargetLoadAction::EClear : ERenderTargetLoadAction::ELoad);
+	ParametersPS->RenderTargets[4] = FRenderTargetBinding(OutMaterialTexture, bClear ? ERenderTargetLoadAction::EClear : ERenderTargetLoadAction::ELoad);
 	ParametersPS->RenderTargets[5] = FRenderTargetBinding(OutTriangleMask, bClear ? ERenderTargetLoadAction::EClear : ERenderTargetLoadAction::ELoad);
 	ParametersPS->RenderTargets.DepthStencil = FDepthStencilBinding(
 		OutDepthTestTexture,
@@ -890,9 +896,9 @@ static void InternalBuildStrandsTextures_GPU(
 
 	Desc.Format = PF_B8G8R8A8;
 	Desc.ClearValue = FClearValueBinding::Transparent;
-	FRDGTextureRef GroupIndexTexture[2];
-	GroupIndexTexture[0] = GraphBuilder.CreateTexture(Desc, TEXT("Hair.GroupIndexTexture"));
-	GroupIndexTexture[1] = GraphBuilder.CreateTexture(Desc, TEXT("Hair.GroupIndexTexture"));
+	FRDGTextureRef MaterialTexture[2];
+	MaterialTexture[0] = GraphBuilder.CreateTexture(Desc, TEXT("Hair.MaterialTexture"));
+	MaterialTexture[1] = GraphBuilder.CreateTexture(Desc, TEXT("Hair.MaterialTexture"));
 
 	Desc.Format = PF_R32_UINT;
 	Desc.ClearValue = FClearValueBinding::Black;
@@ -1011,6 +1017,7 @@ static void InternalBuildStrandsTextures_GPU(
 
 				GroupData.Strands.RestResource->PositionBuffer.SRV,
 				GroupData.Strands.RestResource->Attribute0Buffer.SRV,
+				GroupData.Strands.RestResource->MaterialBuffer.SRV,
 				GroupData.Strands.RestResource->GetPositionOffset(),
 				RenderingData.GeometrySettings.HairWidth * 0.5f,
 				GroupData.Strands.RestResource->BulkData.MaxLength,
@@ -1022,7 +1029,7 @@ static void InternalBuildStrandsTextures_GPU(
 				TangentTexture[0],
 				CoverageTexture[0],
 				AttributeTexture[0],
-				GroupIndexTexture[0],
+				MaterialTexture[0],
 				TriangleMaskTexture);
 
 			bClear = false;
@@ -1046,13 +1053,13 @@ static void InternalBuildStrandsTextures_GPU(
 			CoverageTexture[SourceIndex],
 			TangentTexture[SourceIndex],
 			AttributeTexture[SourceIndex],
-			GroupIndexTexture[SourceIndex],
+			MaterialTexture[SourceIndex],
 
 			DepthTexture[TargetIndex],
 			CoverageTexture[TargetIndex],
 			TangentTexture[TargetIndex],
 			AttributeTexture[TargetIndex],
-			GroupIndexTexture[TargetIndex]);
+			MaterialTexture[TargetIndex]);
 
 		SourceIndex = TargetIndex;
 	}
@@ -1061,7 +1068,7 @@ static void InternalBuildStrandsTextures_GPU(
 	AddReadbackPass(GraphBuilder, 4, CoverageTexture[TargetIndex]	, Output.Coverage);
 	AddReadbackPass(GraphBuilder, 4, TangentTexture[TargetIndex]	, Output.Tangent);
 	AddReadbackPass(GraphBuilder, 4, AttributeTexture[TargetIndex]	, Output.Attribute);
-	AddReadbackPass(GraphBuilder, 4, GroupIndexTexture[TargetIndex]	, Output.GroupIndex);
+	AddReadbackPass(GraphBuilder, 4, MaterialTexture[TargetIndex]	, Output.Material);
 }
 
 ///////////////////////////////////////////////////////////////////////////////////////////////////
diff --git a/Engine/Plugins/Runtime/HairStrands/Source/HairStrandsCore/Private/HairCardsBuilder.cpp b/Engine/Plugins/Runtime/HairStrands/Source/HairStrandsCore/Private/HairCardsBuilder.cpp
--- a/Engine/Plugins/Runtime/HairStrands/Source/HairStrandsCore/Private/HairCardsBuilder.cpp
+++ b/Engine/Plugins/Runtime/HairStrands/Source/HairStrandsCore/Private/HairCardsBuilder.cpp
@@ -3335,7 +3335,7 @@ namespace FHairCardsBuilder
 FString GetVersion()
 {
 	// Important to update the version when cards building or importing changes
-	return TEXT("9c");
+	return TEXT("9d");
 }
 
 void AllocateAtlasTexture(UTexture2D* Out, const FIntPoint& Resolution, uint32 MipCount, EPixelFormat PixelFormat, ETextureSourceFormat SourceFormat)
@@ -3712,12 +3712,14 @@ bool InternalImportGeometry(
 	OutBulk.Positions.SetNum(PointCount);
 	OutBulk.Normals.SetNum(PointCount * FHairCardsNormalFormat::ComponentCount);
 	OutBulk.UVs.SetNum(PointCount);
+	OutBulk.Materials.SetNum(PointCount);
 	for (uint32 PointIt = 0; PointIt < PointCount; ++PointIt)
 	{
 		OutBulk.Positions[PointIt] = FVector4f(Out.Cards.Positions[PointIt], 0);
 		OutBulk.UVs[PointIt] = Out.Cards.UVs[PointIt];
 		OutBulk.Normals[PointIt * 2] = FVector4f(Out.Cards.Tangents[PointIt], 0);
 		OutBulk.Normals[PointIt * 2 + 1] = FVector4f(Out.Cards.Normals[PointIt], TangentFrameSigns[PointIt]);
+		OutBulk.Materials[PointIt] = 0;
 	}
 
 	OutBulk.Indices.SetNum(IndexCount);
@@ -3755,8 +3757,12 @@ bool InternalImportGeometry(
 	}
 
 	// Used voxelized hair group (from hair strands) to assign hair group index to card vertices
-	TArray<uint32> CardGroupIndices;
-	CardGroupIndices.Init(0u, CardsCount);
+	TArray<FVector3f> CardBaseColor;
+	TArray<float> CardsRoughness;
+	TArray<uint8> CardGroupIndices;
+	CardGroupIndices.Init(0u, CardsCount);					// Per-card
+	CardBaseColor.Init(FVector3f::ZeroVector, PointCount);	// Per-vertex
+	CardsRoughness.Init(0.f, PointCount);					// Per-vertex
 	if (InStrandsVoxelData.IsValid())
 	{
 		for (uint32 CardIt = 0; CardIt < CardsCount; ++CardIt)
@@ -3773,11 +3779,13 @@ bool InternalImportGeometry(
 				const uint32 VertexIndex = Out.Cards.Indices[CardsIndexOffset + IndexIt];
 				const FVector3f& P = Out.Cards.Positions[VertexIndex];
 
-				const uint8 GroupIndex = InStrandsVoxelData.GetGroupIndex(P);
-				if (GroupIndex != FHairStrandsVoxelData::InvalidGroupIndex)
+				FHairStrandsVoxelData::FData VoxelData = InStrandsVoxelData.GetData(P);
+				if (VoxelData.GroupIndex != FHairStrandsVoxelData::InvalidGroupIndex)
 				{
-					check(GroupIndex < GroupBins.Num());
-					GroupBins[GroupIndex]++;
+					check(VoxelData.GroupIndex < GroupBins.Num());
+					GroupBins[VoxelData.GroupIndex]++;
+					CardBaseColor[VertexIndex] = VoxelData.BaseColor;
+					CardsRoughness[VertexIndex] = VoxelData.Roughness;
 				}
 			}
 
@@ -3808,18 +3816,25 @@ bool InternalImportGeometry(
 			const float CoordU = Out.Cards.CoordU[VertexIndex];
 			const float InterpolatedCardLength = FMath::Lerp(0.f, CardLength, CoordU);
 
-			// Instead of storing the interoplated card length, store the actual max length of the card, 
-			// as reconstructing the strands length, based on interpolatd CardLength will be too prone to numerical issue.
+			// Instead of storing the interpolated card length, store the actual max length of the card, 
+			// as reconstructing the strands length, based on interpolated CardLength will be too prone to numerical issue.
 			// This means that the strand length retrieves in shader will be an over estimate of the actual length
-			const FFloat16 hCardLength = CardLength;  // InterpolatedCardLength;			
+			const FFloat16 hCardLength = CardLength;  // InterpolatedCardLength;
 
 			// Encode cards length & group index into the .W component of position
 			const uint32 EncodedW = hCardLength.Encoded | (CardGroupIndices[CardIt] << 16u);
 			OutBulk.Positions[VertexIndex].W = *(float*)&EncodedW;
+
+			// Encode the base color in (cheap) sRGB in XYZ. The W component remains unused
+			OutBulk.Materials[VertexIndex] =
+				(uint32(FMath::Sqrt(CardBaseColor[VertexIndex].X) * 255.f)    )|
+				(uint32(FMath::Sqrt(CardBaseColor[VertexIndex].Y) * 255.f)<<8 )|
+				(uint32(FMath::Sqrt(CardBaseColor[VertexIndex].Z) * 255.f)<<16)|
+				(uint32(CardsRoughness[VertexIndex]               * 255.f)<<24);
 		}
 	}
 
-	// Patch Cards RootUV by transfering the guides root UV onto the cards
+	// Patch Cards RootUV by transferring the guides root UV onto the cards
 	if (InStrandsData.IsValid())
 	{
 		// 1. Extract all roots
diff --git a/Engine/Plugins/Runtime/HairStrands/Source/HairStrandsCore/Private/HairCardsDatas.cpp b/Engine/Plugins/Runtime/HairStrands/Source/HairStrandsCore/Private/HairCardsDatas.cpp
--- a/Engine/Plugins/Runtime/HairStrands/Source/HairStrandsCore/Private/HairCardsDatas.cpp
+++ b/Engine/Plugins/Runtime/HairStrands/Source/HairStrandsCore/Private/HairCardsDatas.cpp
@@ -20,6 +20,7 @@ void FHairCardsBulkData::Serialize(FArchive& Ar)
 	Ar << Positions;
 	Ar << Normals;
 	Ar << UVs;
+	Ar << Materials;
 	Ar << Indices;
 	Ar << BoundingBox;
 }
diff --git a/Engine/Plugins/Runtime/HairStrands/Source/HairStrandsCore/Private/HairCardsVertexFactory.cpp b/Engine/Plugins/Runtime/HairStrands/Source/HairStrandsCore/Private/HairCardsVertexFactory.cpp
--- a/Engine/Plugins/Runtime/HairStrands/Source/HairStrandsCore/Private/HairCardsVertexFactory.cpp
+++ b/Engine/Plugins/Runtime/HairStrands/Source/HairStrandsCore/Private/HairCardsVertexFactory.cpp
@@ -17,6 +17,16 @@
 template<typename T> inline void VFC_BindParam(FMeshDrawSingleShaderBindings& ShaderBindings, const FShaderResourceParameter& Param, T* Value) { if (Param.IsBound() && Value) ShaderBindings.Add(Param, Value); }
 template<typename T> inline void VFC_BindParam(FMeshDrawSingleShaderBindings& ShaderBindings, const FShaderParameter& Param, const T& Value) { if (Param.IsBound()) ShaderBindings.Add(Param, Value); }
 
+enum class EHairCardsFactoryFlags : uint32
+{
+	InvertedUV = 1,
+	TextureRootUV = 2,
+	TextureGroupIndex = 4,
+	TextureBaseColor = 8,
+	TextureRoughness = 10,
+	TextureAttribute = 12
+};
+
 ////////////////////////////////////////////////////////////////////////////////////////////////////////////////////////
 // Cards based vertex factory
 IMPLEMENT_GLOBAL_SHADER_PARAMETER_STRUCT(FHairCardsVertexFactoryUniformShaderParameters, "HairCardsVF");
@@ -34,9 +44,8 @@ FHairCardsUniformBuffer CreateHairCardsVFUniformBuffer(
 		const FHairGroupInstance::FCards::FLOD& LOD = Instance->Cards.LODs[LODIndex];
 
 		// Cards atlas UV are inverted so fetching needs to be inverted on the y-axis
-		UniformParameters.bInvertUV = LOD.RestResource->bInvertUV;
-		UniformParameters.bUseTextureRootUV = 0;
-		UniformParameters.bUseTextureGroupIndex = 0;
+		UniformParameters.Flags = 0;
+		if (LOD.RestResource->bInvertUV) { UniformParameters.Flags |= uint32(EHairCardsFactoryFlags::InvertedUV);  }
 		UniformParameters.MaxVertexCount = LOD.RestResource->GetVertexCount();
 
 		// When the geometry is not-dynamic (no binding to skeletal mesh, no simulation), only a single vertex buffer is allocated. 
@@ -56,6 +65,7 @@ FHairCardsUniformBuffer CreateHairCardsVFUniformBuffer(
 		}
 
 		UniformParameters.UVsBuffer = LOD.RestResource->UVsBuffer.ShaderResourceViewRHI.GetReference();
+		UniformParameters.MaterialsBuffer = LOD.RestResource->MaterialsBuffer.ShaderResourceViewRHI.GetReference();
 
 		UniformParameters.DepthTexture = LOD.RestResource->DepthTexture;
 		UniformParameters.DepthSampler = LOD.RestResource->DepthSampler;
@@ -67,8 +77,8 @@ FHairCardsUniformBuffer CreateHairCardsVFUniformBuffer(
 		UniformParameters.AttributeSampler = LOD.RestResource->AttributeSampler;
 		UniformParameters.AuxilaryDataTexture = LOD.RestResource->AuxilaryDataTexture;
 		UniformParameters.AuxilaryDataSampler = LOD.RestResource->AuxilaryDataSampler;
-		UniformParameters.GroupIndexTexture = nullptr; // Group index are stored on vertices for cards
-		UniformParameters.GroupIndexSampler = nullptr; // Group index are stored on vertices for cards
+		UniformParameters.MaterialTexture = nullptr; // Material properties & Group index are stored on vertices for cards
+		UniformParameters.MaterialSampler = nullptr; // Material properties & Group index are stored on vertices for cards
 	}
 	else if (GeometryType == EHairGeometryType::Meshes)
 	{
@@ -89,12 +99,17 @@ FHairCardsUniformBuffer CreateHairCardsVFUniformBuffer(
 		}
 
 		// Meshes UV are not inverted so no need to invert the y-axis
-		UniformParameters.bInvertUV = 0;
-		UniformParameters.bUseTextureRootUV = 1;
-		UniformParameters.bUseTextureGroupIndex = 1;
+		UniformParameters.Flags = 0;
+		UniformParameters.Flags |= uint32(EHairCardsFactoryFlags::TextureRootUV);
+		UniformParameters.Flags |= uint32(EHairCardsFactoryFlags::TextureAttribute);
+		UniformParameters.Flags |= uint32(EHairCardsFactoryFlags::TextureGroupIndex);
+		UniformParameters.Flags |= uint32(EHairCardsFactoryFlags::TextureBaseColor);
+		UniformParameters.Flags |= uint32(EHairCardsFactoryFlags::TextureRoughness);
+
 		UniformParameters.MaxVertexCount = LOD.RestResource->GetVertexCount();
 		UniformParameters.NormalsBuffer = LOD.RestResource->NormalsBuffer.ShaderResourceViewRHI.GetReference();
 		UniformParameters.UVsBuffer = LOD.RestResource->UVsBuffer.ShaderResourceViewRHI.GetReference();
+		UniformParameters.MaterialsBuffer = UniformParameters.NormalsBuffer; // Reuse normal buffer as a dummy input for material buffer, since it not used for meshes (material data is fetch through textures)
 
 		UniformParameters.DepthTexture = LOD.RestResource->DepthTexture;
 		UniformParameters.DepthSampler = LOD.RestResource->DepthSampler;
@@ -106,8 +121,8 @@ FHairCardsUniformBuffer CreateHairCardsVFUniformBuffer(
 		UniformParameters.AttributeSampler = LOD.RestResource->AttributeSampler;
 		UniformParameters.AuxilaryDataTexture = LOD.RestResource->AuxilaryDataTexture;
 		UniformParameters.AuxilaryDataSampler = LOD.RestResource->AuxilaryDataSampler;
-		UniformParameters.GroupIndexTexture = LOD.RestResource->GroupIndexTexture;
-		UniformParameters.GroupIndexSampler = LOD.RestResource->GroupIndexSampler;
+		UniformParameters.MaterialTexture = LOD.RestResource->MaterialTexture;
+		UniformParameters.MaterialSampler = LOD.RestResource->MaterialSampler;
 	}
 
 	if (!bSupportsManualVertexFetch)
@@ -116,6 +131,7 @@ FHairCardsUniformBuffer CreateHairCardsVFUniformBuffer(
 		UniformParameters.PreviousPositionBuffer = GNullVertexBuffer.VertexBufferSRV;
 		UniformParameters.NormalsBuffer = GNullVertexBuffer.VertexBufferSRV;
 		UniformParameters.UVsBuffer = GNullVertexBuffer.VertexBufferSRV;
+		UniformParameters.MaterialsBuffer = GNullVertexBuffer.VertexBufferSRV;
 	}
 
 	FRHITexture* DefaultTexture = GBlackTexture->TextureRHI;
@@ -126,14 +142,14 @@ FHairCardsUniformBuffer CreateHairCardsVFUniformBuffer(
 	if (!UniformParameters.CoverageTexture)		{ UniformParameters.CoverageTexture = DefaultTexture; }
 	if (!UniformParameters.AttributeTexture)	{ UniformParameters.AttributeTexture = DefaultTexture;}
 	if (!UniformParameters.AuxilaryDataTexture)	{ UniformParameters.AuxilaryDataTexture = DefaultTexture; }
-	if (!UniformParameters.GroupIndexTexture)	{ UniformParameters.GroupIndexTexture = DefaultTexture; }
+	if (!UniformParameters.MaterialTexture)		{ UniformParameters.MaterialTexture = DefaultTexture; }
 
 	if (!UniformParameters.DepthSampler)		{ UniformParameters.DepthSampler = DefaultSampler;	  }
 	if (!UniformParameters.TangentSampler)		{ UniformParameters.TangentSampler = DefaultSampler;  }
 	if (!UniformParameters.CoverageSampler)		{ UniformParameters.CoverageSampler = DefaultSampler; }
 	if (!UniformParameters.AttributeSampler)	{ UniformParameters.AttributeSampler = DefaultSampler;}
 	if (!UniformParameters.AuxilaryDataSampler) { UniformParameters.AuxilaryDataSampler = DefaultSampler; }
-	if (!UniformParameters.GroupIndexSampler)	{ UniformParameters.GroupIndexSampler = DefaultSampler;}
+	if (!UniformParameters.MaterialSampler)		{ UniformParameters.MaterialSampler = DefaultSampler;}
 
 	return TUniformBufferRef<FHairCardsVertexFactoryUniformShaderParameters>::CreateUniformBufferImmediate(UniformParameters, UniformBuffer_MultiFrame);
 }
@@ -285,19 +301,22 @@ void FHairCardsVertexFactory::InitResources()
 		{
 			const FHairGroupInstance::FCards::FLOD& LOD = Data.Instance->Cards.LODs[Data.LODIndex];
 		
-			Elements.Add(AccessStreamComponent(FVertexStreamComponent(&LOD.RestResource->RestPositionBuffer,	0, 0,									FHairCardsPositionFormat::SizeInByte, FHairCardsPositionFormat::VertexElementType,	EVertexStreamUsage::Default), 0));
-			Elements.Add(AccessStreamComponent(FVertexStreamComponent(&LOD.RestResource->NormalsBuffer,			0, 0,									FHairCardsNormalFormat::SizeInByte * FHairCardsNormalFormat::ComponentCount,	FHairCardsNormalFormat::VertexElementType,		EVertexStreamUsage::Default), 1));
-			Elements.Add(AccessStreamComponent(FVertexStreamComponent(&LOD.RestResource->NormalsBuffer,			0, FHairCardsNormalFormat::SizeInByte,	FHairCardsNormalFormat::SizeInByte * FHairCardsNormalFormat::ComponentCount,	FHairCardsNormalFormat::VertexElementType,		EVertexStreamUsage::Default), 2));
-			Elements.Add(AccessStreamComponent(FVertexStreamComponent(&LOD.RestResource->UVsBuffer,				0, 0,									FHairCardsUVFormat::SizeInByte,		FHairCardsUVFormat::VertexElementType,			EVertexStreamUsage::Default), 3));
+			Elements.Add(AccessStreamComponent(FVertexStreamComponent(&LOD.RestResource->RestPositionBuffer,	0, 0,									FHairCardsPositionFormat::SizeInByte,												FHairCardsPositionFormat::VertexElementType,	EVertexStreamUsage::Default), 0));
+			Elements.Add(AccessStreamComponent(FVertexStreamComponent(&LOD.RestResource->NormalsBuffer,			0, 0,									FHairCardsNormalFormat::SizeInByte * FHairCardsNormalFormat::ComponentCount,		FHairCardsNormalFormat::VertexElementType,		EVertexStreamUsage::Default), 1));
+			Elements.Add(AccessStreamComponent(FVertexStreamComponent(&LOD.RestResource->NormalsBuffer,			0, FHairCardsNormalFormat::SizeInByte,	FHairCardsNormalFormat::SizeInByte * FHairCardsNormalFormat::ComponentCount,		FHairCardsNormalFormat::VertexElementType,		EVertexStreamUsage::Default), 2));
+			Elements.Add(AccessStreamComponent(FVertexStreamComponent(&LOD.RestResource->UVsBuffer,				0, 0,									FHairCardsUVFormat::SizeInByte,														FHairCardsUVFormat::VertexElementType,			EVertexStreamUsage::Default), 3));
+			Elements.Add(AccessStreamComponent(FVertexStreamComponent(&LOD.RestResource->MaterialsBuffer,		0, 0,									FHairCardsMaterialFormat::SizeInByte,												FHairCardsMaterialFormat::VertexElementType,	EVertexStreamUsage::Default), 4));
 		}
 		else if (Data.GeometryType == EHairGeometryType::Meshes)
 		{
 			const FHairGroupInstance::FMeshes::FLOD& LOD = Data.Instance->Meshes.LODs[Data.LODIndex];
 
-			Elements.Add(AccessStreamComponent(FVertexStreamComponent(&LOD.RestResource->RestPositionBuffer,	0, 0,									FHairCardsPositionFormat::SizeInByte, FHairCardsPositionFormat::VertexElementType,	EVertexStreamUsage::Default), 0));
-			Elements.Add(AccessStreamComponent(FVertexStreamComponent(&LOD.RestResource->NormalsBuffer,			0, 0,									FHairCardsNormalFormat::SizeInByte * FHairCardsNormalFormat::ComponentCount,	FHairCardsNormalFormat::VertexElementType,		EVertexStreamUsage::Default), 1));
-			Elements.Add(AccessStreamComponent(FVertexStreamComponent(&LOD.RestResource->NormalsBuffer,			0, FHairCardsNormalFormat::SizeInByte,	FHairCardsNormalFormat::SizeInByte * FHairCardsNormalFormat::ComponentCount,	FHairCardsNormalFormat::VertexElementType,		EVertexStreamUsage::Default), 2));
-			Elements.Add(AccessStreamComponent(FVertexStreamComponent(&LOD.RestResource->UVsBuffer,				0, 0,									FHairCardsUVFormat::SizeInByte,		FHairCardsUVFormat::VertexElementType,			EVertexStreamUsage::Default), 3));
+			// Note: Use the 'Normal' buffer as a dummy input for 'Material' buffer, as Material data is fetched through textures for meshes
+			Elements.Add(AccessStreamComponent(FVertexStreamComponent(&LOD.RestResource->RestPositionBuffer,	0, 0,									FHairCardsPositionFormat::SizeInByte,												FHairCardsPositionFormat::VertexElementType,	EVertexStreamUsage::Default), 0));
+			Elements.Add(AccessStreamComponent(FVertexStreamComponent(&LOD.RestResource->NormalsBuffer,			0, 0,									FHairCardsNormalFormat::SizeInByte * FHairCardsNormalFormat::ComponentCount,		FHairCardsNormalFormat::VertexElementType,		EVertexStreamUsage::Default), 1));
+			Elements.Add(AccessStreamComponent(FVertexStreamComponent(&LOD.RestResource->NormalsBuffer,			0, FHairCardsNormalFormat::SizeInByte,	FHairCardsNormalFormat::SizeInByte * FHairCardsNormalFormat::ComponentCount,		FHairCardsNormalFormat::VertexElementType,		EVertexStreamUsage::Default), 2));
+			Elements.Add(AccessStreamComponent(FVertexStreamComponent(&LOD.RestResource->UVsBuffer,				0, 0,									FHairCardsUVFormat::SizeInByte,														FHairCardsUVFormat::VertexElementType,			EVertexStreamUsage::Default), 3));
+			Elements.Add(AccessStreamComponent(FVertexStreamComponent(&LOD.RestResource->NormalsBuffer,			0, 0,									FHairCardsMaterialFormat::SizeInByte,												FHairCardsMaterialFormat::VertexElementType,	EVertexStreamUsage::Default), 4)); 
 		}
 
 		bNeedsDeclaration = true;
diff --git a/Engine/Plugins/Runtime/HairStrands/Source/HairStrandsCore/Public/GroomAsset.h b/Engine/Plugins/Runtime/HairStrands/Source/HairStrandsCore/Public/GroomAsset.h
--- a/Engine/Plugins/Runtime/HairStrands/Source/HairStrandsCore/Public/GroomAsset.h
+++ b/Engine/Plugins/Runtime/HairStrands/Source/HairStrandsCore/Public/GroomAsset.h
@@ -397,7 +397,7 @@ enum class EHairAtlasTextureType : uint8
 	Attribute,
 	Coverage,
 	AuxilaryData,
-	GroupIndex
+	Material
 };
 
 /**
diff --git a/Engine/Plugins/Runtime/HairStrands/Source/HairStrandsCore/Public/GroomAssetCards.h b/Engine/Plugins/Runtime/HairStrands/Source/HairStrandsCore/Public/GroomAssetCards.h
--- a/Engine/Plugins/Runtime/HairStrands/Source/HairStrandsCore/Public/GroomAssetCards.h
+++ b/Engine/Plugins/Runtime/HairStrands/Source/HairStrandsCore/Public/GroomAssetCards.h
@@ -183,9 +183,9 @@ struct HAIRSTRANDSCORE_API FHairGroupCardsTextures
 	UPROPERTY(EditAnywhere, Category = "CardsAttributes")
 	TObjectPtr<UTexture2D> AuxilaryDataTexture = nullptr;
 
-	/* Store hair group indices. This texture is only used by hair meshes and can be generated through the hair-textures option on Groom Asset. */
+	/* This texture is only used by hair meshes and can be generated through the hair-textures option on Groom Asset. */
 	UPROPERTY(EditAnywhere, Category = "CardsAttributes")
-	TObjectPtr<UTexture2D> GroupIndexTexture = nullptr;
+	TObjectPtr<UTexture2D> MaterialTexture = nullptr;
 
 	void SetTexture(EHairAtlasTextureType SlotID, UTexture2D* Texture);
 
diff --git a/Engine/Plugins/Runtime/HairStrands/Source/HairStrandsCore/Public/GroomBuilder.h b/Engine/Plugins/Runtime/HairStrands/Source/HairStrandsCore/Public/GroomBuilder.h
--- a/Engine/Plugins/Runtime/HairStrands/Source/HairStrandsCore/Public/GroomBuilder.h
+++ b/Engine/Plugins/Runtime/HairStrands/Source/HairStrandsCore/Public/GroomBuilder.h
@@ -27,15 +27,21 @@ class UGroomAsset;
 struct FHairStrandsVoxelData
 {
 	static const uint8 InvalidGroupIndex = 0xFF;
+	struct FData
+	{
+		FVector3f BaseColor;
+		float Roughness;
+		uint8 GroupIndex;
+	};
 
 	FVector3f MinBound;
 	FVector3f MaxBound;
 	FIntVector Resolution;
-	TArray<uint8> GroupIndices;
+	TArray<FData> Datas;
 
 	// Return the (closest) group index corresponding to position P
-	uint8 GetGroupIndex(const FVector3f& P) const;
-	bool IsValid() const { return GroupIndices.Num() > 0; }
+	FData GetData(const FVector3f& P) const;
+	bool IsValid() const { return Datas.Num() > 0; }
 };
 
 // Data flow overview
diff --git a/Engine/Plugins/Runtime/HairStrands/Source/HairStrandsCore/Public/GroomInstance.h b/Engine/Plugins/Runtime/HairStrands/Source/HairStrandsCore/Public/GroomInstance.h
--- a/Engine/Plugins/Runtime/HairStrands/Source/HairStrandsCore/Public/GroomInstance.h
+++ b/Engine/Plugins/Runtime/HairStrands/Source/HairStrandsCore/Public/GroomInstance.h
@@ -19,14 +19,13 @@ enum class EGroomCacheType : uint8;
 //  * 10/10 bits for UV -> max 1024/1024 rect resolution
 //  * 12 bits for cards count -> 4000 cards for a hair group
 BEGIN_GLOBAL_SHADER_PARAMETER_STRUCT(FHairCardsVertexFactoryUniformShaderParameters, HAIRSTRANDSCORE_API)
-	SHADER_PARAMETER(uint32, bInvertUV)
-	SHADER_PARAMETER(uint32, bUseTextureRootUV)
-	SHADER_PARAMETER(uint32, bUseTextureGroupIndex)
+	SHADER_PARAMETER(uint32, Flags)
 	SHADER_PARAMETER(uint32, MaxVertexCount)
 	SHADER_PARAMETER_SRV(Buffer<float4>, PositionBuffer)
 	SHADER_PARAMETER_SRV(Buffer<float4>, PreviousPositionBuffer)
 	SHADER_PARAMETER_SRV(Buffer<float4>, NormalsBuffer)
 	SHADER_PARAMETER_SRV(Buffer<float4>, UVsBuffer)
+	SHADER_PARAMETER_SRV(Buffer<float4>, MaterialsBuffer)
 
 	SHADER_PARAMETER_TEXTURE(Texture2D<float4>, DepthTexture)
 	SHADER_PARAMETER_SAMPLER(SamplerState, DepthSampler)
@@ -38,8 +37,8 @@ BEGIN_GLOBAL_SHADER_PARAMETER_STRUCT(FHairCardsVertexFactoryUniformShaderParamet
 	SHADER_PARAMETER_SAMPLER(SamplerState, AttributeSampler)
 	SHADER_PARAMETER_TEXTURE(Texture2D<float4>, AuxilaryDataTexture)
 	SHADER_PARAMETER_SAMPLER(SamplerState, AuxilaryDataSampler)
-	SHADER_PARAMETER_TEXTURE(Texture2D<float4>, GroupIndexTexture)
-	SHADER_PARAMETER_SAMPLER(SamplerState, GroupIndexSampler)
+	SHADER_PARAMETER_TEXTURE(Texture2D<float4>, MaterialTexture)
+	SHADER_PARAMETER_SAMPLER(SamplerState, MaterialSampler)
 END_GLOBAL_SHADER_PARAMETER_STRUCT()
 
 typedef TUniformBufferRef<FHairCardsVertexFactoryUniformShaderParameters> FHairCardsUniformBuffer;
@@ -66,7 +65,7 @@ BEGIN_GLOBAL_SHADER_PARAMETER_STRUCT(FHairStrandsVertexFactoryUniformShaderParam
 	SHADER_PARAMETER_SRV(Buffer<uint4>, PreviousPositionBuffer)
 
 	SHADER_PARAMETER_SRV(Buffer<float2>, Attribute0Buffer)
-	SHADER_PARAMETER_SRV(Buffer<uint>, Attribute1Buffer)
+	SHADER_PARAMETER_SRV(Buffer<uint>,   Attribute1Buffer)
 	SHADER_PARAMETER_SRV(Buffer<float4>, MaterialBuffer)
 	SHADER_PARAMETER_SRV(Buffer<float4>, TangentBuffer)
 
diff --git a/Engine/Plugins/Runtime/HairStrands/Source/HairStrandsCore/Public/GroomResources.h b/Engine/Plugins/Runtime/HairStrands/Source/HairStrandsCore/Public/GroomResources.h
--- a/Engine/Plugins/Runtime/HairStrands/Source/HairStrandsCore/Public/GroomResources.h
+++ b/Engine/Plugins/Runtime/HairStrands/Source/HairStrandsCore/Public/GroomResources.h
@@ -522,6 +522,7 @@ struct FHairCardsRestResource : public FHairCommonResource
 		Total += BulkData.Normals.GetAllocatedSize();
 		Total += BulkData.Indices.GetAllocatedSize();
 		Total += BulkData.UVs.GetAllocatedSize();
+		Total += BulkData.Materials.GetAllocatedSize();
 		return Total;
 	}
 
@@ -535,20 +536,21 @@ struct FHairCardsRestResource : public FHairCommonResource
 
 	FHairCardsVertexBuffer NormalsBuffer;
 	FHairCardsVertexBuffer UVsBuffer;
+	FHairCardsVertexBuffer MaterialsBuffer;
 
 	FSamplerStateRHIRef DepthSampler = nullptr;
 	FSamplerStateRHIRef TangentSampler = nullptr;
 	FSamplerStateRHIRef CoverageSampler = nullptr;
 	FSamplerStateRHIRef AttributeSampler = nullptr;
 	FSamplerStateRHIRef AuxilaryDataSampler = nullptr;
-	FSamplerStateRHIRef GroupIndexSampler = nullptr;
+	FSamplerStateRHIRef MaterialSampler = nullptr;
 
 	FTextureReferenceRHIRef	DepthTexture = nullptr;
 	FTextureReferenceRHIRef	CoverageTexture = nullptr;
 	FTextureReferenceRHIRef	TangentTexture = nullptr;
 	FTextureReferenceRHIRef	AttributeTexture = nullptr;
 	FTextureReferenceRHIRef	AuxilaryDataTexture = nullptr;
-	FTextureReferenceRHIRef	GroupIndexTexture = nullptr;
+	FTextureReferenceRHIRef	MaterialTexture = nullptr;
 
 	/* Reference to the hair strands render data */
 	const FHairCardsBulkData& BulkData;
@@ -738,14 +740,14 @@ struct FHairMeshesRestResource : public FHairCommonResource
 	FSamplerStateRHIRef CoverageSampler = nullptr;
 	FSamplerStateRHIRef AttributeSampler = nullptr;
 	FSamplerStateRHIRef AuxilaryDataSampler = nullptr;
-	FSamplerStateRHIRef GroupIndexSampler = nullptr;
+	FSamplerStateRHIRef MaterialSampler = nullptr;
 
 	FTextureReferenceRHIRef	DepthTexture = nullptr;
 	FTextureReferenceRHIRef	CoverageTexture = nullptr;
 	FTextureReferenceRHIRef	TangentTexture = nullptr;
 	FTextureReferenceRHIRef	AttributeTexture = nullptr;
 	FTextureReferenceRHIRef	AuxilaryDataTexture = nullptr;
-	FTextureReferenceRHIRef	GroupIndexTexture = nullptr;
+	FTextureReferenceRHIRef	MaterialTexture = nullptr;
 
 	/* Reference to the hair strands render data */
 	const FHairMeshesBulkData& BulkData;
diff --git a/Engine/Plugins/Runtime/HairStrands/Source/HairStrandsCore/Public/GroomTextureBuilder.h b/Engine/Plugins/Runtime/HairStrands/Source/HairStrandsCore/Public/GroomTextureBuilder.h
--- a/Engine/Plugins/Runtime/HairStrands/Source/HairStrandsCore/Public/GroomTextureBuilder.h
+++ b/Engine/Plugins/Runtime/HairStrands/Source/HairStrandsCore/Public/GroomTextureBuilder.h
@@ -51,8 +51,8 @@ struct FStrandsTexturesOutput
 	class UTexture2D* Tangent = nullptr;
 	class UTexture2D* Coverage = nullptr;
 	class UTexture2D* Attribute = nullptr;
-	class UTexture2D* GroupIndex = nullptr;
-	bool IsValid() const { return Depth && Tangent && Coverage && Attribute && GroupIndex; }
+	class UTexture2D* Material = nullptr;
+	bool IsValid() const { return Depth && Tangent && Coverage && Attribute && Material; }
 };
 
 ////////////////////////////////////////////////////////////////////////////////////////////////////////////////////////
diff --git a/Engine/Plugins/Runtime/HairStrands/Source/HairStrandsCore/Public/HairCardsDatas.h b/Engine/Plugins/Runtime/HairStrands/Source/HairStrandsCore/Public/HairCardsDatas.h
--- a/Engine/Plugins/Runtime/HairStrands/Source/HairStrandsCore/Public/HairCardsDatas.h
+++ b/Engine/Plugins/Runtime/HairStrands/Source/HairStrandsCore/Public/HairCardsDatas.h
@@ -45,6 +45,15 @@ struct FHairCardsUVFormat
 	static const EPixelFormat Format = PF_A32B32G32R32F; // TODO
 };
 
+struct FHairCardsMaterialFormat
+{
+	typedef uint32 Type;
+	static const uint32 ComponentCount = 1;
+	static const uint32 SizeInByte = sizeof(Type);
+	static const EVertexElementType VertexElementType = VET_UByte4;
+	static const EPixelFormat Format = PF_R8G8B8A8;
+};
+
 struct FHairCardsNormalFormat
 {
 	typedef FPackedNormal Type;
@@ -261,6 +270,7 @@ struct FHairCardsBulkData
 		Positions.Empty();
 		Normals.Empty();
 		UVs.Empty();
+		Materials.Empty();
 		Indices.Empty();
 	}
 
@@ -276,6 +286,7 @@ struct FHairCardsBulkData
 	TArray<FHairCardsPositionFormat::Type> Positions;
 	TArray<FHairCardsNormalFormat::Type> Normals;
 	TArray<FHairCardsUVFormat::Type> UVs;
+	TArray<FHairCardsMaterialFormat::Type> Materials;
 	TArray<FHairCardsIndexFormat::Type> Indices;
 	FBox BoundingBox;
 
diff --git a/Engine/Plugins/Runtime/HairStrands/Source/HairStrandsEditor/Private/GroomAssetDetails.cpp b/Engine/Plugins/Runtime/HairStrands/Source/HairStrandsEditor/Private/GroomAssetDetails.cpp
--- a/Engine/Plugins/Runtime/HairStrands/Source/HairStrandsEditor/Private/GroomAssetDetails.cpp
+++ b/Engine/Plugins/Runtime/HairStrands/Source/HairStrandsEditor/Private/GroomAssetDetails.cpp
@@ -848,6 +848,7 @@ bool FGroomRenderingDetails::CommonResetToDefault(TSharedPtr<IPropertyHandle> Ch
 			HAIR_RESET1(HairGroupsCards, FHairGroupCardsTextures, Textures, CoverageTexture);
 			HAIR_RESET1(HairGroupsCards, FHairGroupCardsTextures, Textures, TangentTexture);
 			HAIR_RESET1(HairGroupsCards, FHairGroupCardsTextures, Textures, AttributeTexture);
+			HAIR_RESET1(HairGroupsCards, FHairGroupCardsTextures, Textures, MaterialTexture);
 			HAIR_RESET1(HairGroupsCards, FHairGroupCardsTextures, Textures, AuxilaryDataTexture);
 		}
 	}
@@ -868,6 +869,7 @@ bool FGroomRenderingDetails::CommonResetToDefault(TSharedPtr<IPropertyHandle> Ch
 			HAIR_RESET1(HairGroupsMeshes, FHairGroupCardsTextures, Textures, CoverageTexture);
 			HAIR_RESET1(HairGroupsMeshes, FHairGroupCardsTextures, Textures, TangentTexture);
 			HAIR_RESET1(HairGroupsMeshes, FHairGroupCardsTextures, Textures, AttributeTexture);
+			HAIR_RESET1(HairGroupsMeshes, FHairGroupCardsTextures, Textures, MaterialTexture);
 			HAIR_RESET1(HairGroupsMeshes, FHairGroupCardsTextures, Textures, AuxilaryDataTexture);
 		}
 	}
diff --git a/Engine/Shaders/Private/HairStrands/HairCardsAttributeCommon.ush b/Engine/Shaders/Private/HairStrands/HairCardsAttributeCommon.ush
--- a/Engine/Shaders/Private/HairStrands/HairCardsAttributeCommon.ush
+++ b/Engine/Shaders/Private/HairStrands/HairCardsAttributeCommon.ush
@@ -8,6 +8,13 @@
 
 ////////////////////////////////////////////////////////////////////////////////
 
+#define EHairCardsFVFlags_InvertedUV 1
+#define EHairCardsFVFlags_TextureRootUV 2
+#define EHairCardsFVFlags_TextureGroupIndex 4
+#define EHairCardsFVFlags_TextureBaseColor 8
+#define EHairCardsFVFlags_TextureRoughness 10
+#define EHairCardsFVFlags_TextureAttribute 12
+
 struct FHairVertexAttributes
 {
 	float2 HairUV;
@@ -21,7 +28,8 @@ struct FHairVertexAttributes
 
 float GetHairStrandsDepth(float2 InAtlasUV, float InDeviceZ)
 {
-	InAtlasUV.y = HairCardsVF.bInvertUV ? 1 - InAtlasUV.y : InAtlasUV.y;
+	const bool bInvertUV = HairCardsVF.Flags & EHairCardsFVFlags_InvertedUV;
+	InAtlasUV.y = bInvertUV ? 1 - InAtlasUV.y : InAtlasUV.y;
 	const float SceneDepthOffset = Texture2DSample(HairCardsVF.DepthTexture, HairCardsVF.DepthSampler, InAtlasUV).x;
 	const float SceneDepth = ConvertFromDeviceZ(InDeviceZ);
 
@@ -30,7 +38,8 @@ float GetHairStrandsDepth(float2 InAtlasUV, float InDeviceZ)
 
 float3 GetHairStrandsTangent(float2 InAtlasUV, half3x3 TangentToWorld, bool bUseTangentSpace)
 {
-	InAtlasUV.y = HairCardsVF.bInvertUV ? 1 - InAtlasUV.y : InAtlasUV.y;
+	const bool bInvertUV = HairCardsVF.Flags & EHairCardsFVFlags_InvertedUV;
+	InAtlasUV.y = bInvertUV ? 1 - InAtlasUV.y : InAtlasUV.y;
 	const float3 LocalTangent = Texture2DSample(HairCardsVF.TangentTexture, HairCardsVF.TangentSampler, InAtlasUV).xyz * 2 - 1;
 	const float3 WorldTangent = mul(LocalTangent, TangentToWorld);
 	return bUseTangentSpace ? LocalTangent : WorldTangent;
@@ -38,13 +47,15 @@ float3 GetHairStrandsTangent(float2 InAtlasUV, half3x3 TangentToWorld, bool bUse
 
 float GetHairStrandsCoverage(float2 InAtlasUV)
 {
-	InAtlasUV.y = HairCardsVF.bInvertUV ? 1 - InAtlasUV.y : InAtlasUV.y;
+	const bool bInvertUV = HairCardsVF.Flags & EHairCardsFVFlags_InvertedUV;
+	InAtlasUV.y = bInvertUV ? 1 - InAtlasUV.y : InAtlasUV.y;
 	return Texture2DSample(HairCardsVF.CoverageTexture, HairCardsVF.CoverageSampler, InAtlasUV).x;
 }
 
-FHairVertexAttributes GetHairStrandsAttributes(float2 InAtlasUV, float2 InRootUV = float2(0, 0), float InCardLength = 0.f, float InCardGroupIndex = 0.f)
+FHairVertexAttributes GetHairStrandsAttributes(float2 InAtlasUV, float2 InRootUV = float2(0, 0), float InCardLength = 0.f, float InCardGroupIndex = 0.f, float3 InCardBaseColor = 0.f, float InCardRoughness = 0.f)
 {
-	InAtlasUV.y = HairCardsVF.bInvertUV ? 1 - InAtlasUV.y : InAtlasUV.y;
+	const bool bInvertUV = HairCardsVF.Flags & EHairCardsFVFlags_InvertedUV;
+	InAtlasUV.y = bInvertUV ? 1 - InAtlasUV.y : InAtlasUV.y;
 	const float4 Value = Texture2DSample(HairCardsVF.AttributeTexture, HairCardsVF.AttributeSampler, InAtlasUV);
 	const float CoordU = Value.z;
 	const float Seed = Value.w;
@@ -57,11 +68,11 @@ FHairVertexAttributes GetHairStrandsAttributes(float2 InAtlasUV, float2 InRootUV
 
 	FHairVertexAttributes Out;
 	Out.HairUV = float2(CoordU, 0.5f);
-	Out.HairRootUV = HairCardsVF.bUseTextureRootUV ? Value.xy : InRootUV;
+	Out.HairRootUV = (HairCardsVF.Flags & EHairCardsFVFlags_TextureRootUV) ? Value.xy : InRootUV;
 	Out.HairDimensions = float2(WorldLength, WorldRadius);
-	Out.HairBaseColor = 0;
+	Out.HairBaseColor = InCardBaseColor;
 	Out.HairSeed = Seed;
-	Out.HairRoughness = 0;
+	Out.HairRoughness = InCardRoughness;
 	Out.HairGroupIndex = InCardGroupIndex;
 	return Out;
 }
@@ -88,7 +99,7 @@ float2 GetHairStrandsRootUV(float2 HairPrimitiveUV, float2 InRootUV)
 	InRootUV.y = 1.0 - InRootUV.y; 
 
 	// If VF use texture data (for helmet/meshes), fetch root UV from the attribute texture
-	if (HairCardsVF.bUseTextureRootUV)
+	if (HairCardsVF.Flags & EHairCardsFVFlags_TextureRootUV)
 	{
 		InRootUV = Texture2DSample(HairCardsVF.AttributeTexture, HairCardsVF.AttributeSampler, HairPrimitiveUV).xy;
 	}
@@ -100,23 +111,29 @@ float  GetHairStrandsSeed(float2 HairPrimitiveUV)
 	return GetHairStrandsAttributes(HairPrimitiveUV).HairSeed;
 }
 
-float3 GetHairStrandsBaseColor(float2 HairPrimitiveUV)
+float3 GetHairStrandsBaseColor(float2 InAtlasUV, float3 InBaseColor)
 {
-	return GetHairStrandsAttributes(HairPrimitiveUV).HairBaseColor;
+	if (HairCardsVF.Flags & EHairCardsFVFlags_TextureBaseColor)
+	{
+		const float3 BaseColorInSRGB = Texture2DSample(HairCardsVF.MaterialTexture, HairCardsVF.MaterialSampler, InAtlasUV).xyz;
+		InBaseColor = BaseColorInSRGB*BaseColorInSRGB; // Cheap sRGB -> Linear encoding
+	}
+	return InBaseColor;
 }
 
-float GetHairStrandsRoughness(float2 HairPrimitiveUV)
+float GetHairStrandsRoughness(float2 HairPrimitiveUV, float InRoughness)
 {
-	return GetHairStrandsAttributes(HairPrimitiveUV).HairRoughness;
+	// For meshes, we don't support roughness yet
+	return (HairCardsVF.Flags & EHairCardsFVFlags_TextureRoughness) ? 0 : InRoughness;
 }
 
 float GetHairStrandsGroupIndex(float2 InAtlasUV, float InHairPrimitiveGroupIndex)
 {
 	// * For hair cards,  group index is stored onto vertices
 	// * For hair meshes, group index is stored within a texture
-	if (HairCardsVF.bUseTextureGroupIndex)
+	if (HairCardsVF.Flags & EHairCardsFVFlags_TextureGroupIndex)
 	{
-		InHairPrimitiveGroupIndex = Texture2DSample(HairCardsVF.GroupIndexTexture, HairCardsVF.GroupIndexSampler, InAtlasUV).x * 255.f;
+		InHairPrimitiveGroupIndex = Texture2DSample(HairCardsVF.MaterialTexture, HairCardsVF.MaterialSampler, InAtlasUV).w * 255.f;
 	}
 	return InHairPrimitiveGroupIndex;
 }
diff --git a/Engine/Shaders/Private/HairStrands/HairCardsVertexFactory.ush b/Engine/Shaders/Private/HairStrands/HairCardsVertexFactory.ush
--- a/Engine/Shaders/Private/HairStrands/HairCardsVertexFactory.ush
+++ b/Engine/Shaders/Private/HairStrands/HairCardsVertexFactory.ush
@@ -48,6 +48,7 @@ struct FVertexFactoryInput
 	half4	TangentZ : ATTRIBUTE2;
   #endif
 	float4 UVs		 : ATTRIBUTE3;
+	float4 Materials : ATTRIBUTE4;
 #endif
 
 	VF_GPUSCENE_DECLARE_INPUT_BLOCK(13)
@@ -83,6 +84,7 @@ struct FVertexFactoryIntermediates
 	float2 HairPrimitiveUV;
 	float2 HairPrimitiveRootUV;
 	float2 HairPrimitiveLength;
+	float4 HairPrimitiveMaterial;
 	float  HairPrimitiveGroupIndex;
 
 	/** Cached primitive and instance data */
@@ -110,6 +112,7 @@ struct FHairCardsPoint
 	float2 RootUVs;
 	float  CardLength;
 	float  GroupIndex;
+	float4 Materials;
 };
 
 FHairCardsPoint GetVertexPosition(FVertexFactoryInput Input)
@@ -118,9 +121,11 @@ FHairCardsPoint GetVertexPosition(FVertexFactoryInput Input)
 	const uint VertexFetchId = GetHairVertexFetchId(Input);
 	const float4 Data = HairCardsVF.PositionBuffer[VertexFetchId];
 	const float4 UVs  = HairCardsVF.UVsBuffer[VertexFetchId];
+	const float4 Materials = HairCardsVF.MaterialsBuffer[VertexFetchId];
 #else
 	const float4 Data = Input.Position;
 	const float4 UVs  = Input.UVs;
+	const float4 Materials = Input.Materials;
 #endif
 
 	const uint EncodedW = asuint(Data.w);
@@ -131,6 +136,7 @@ FHairCardsPoint GetVertexPosition(FVertexFactoryInput Input)
 	Out.RootUVs		= UVs.zw;
 	Out.CardLength	= f16tof32(EncodedW & 0xFFFF);
 	Out.GroupIndex	= (EncodedW>>16u);
+	Out.Materials	= float4(Materials.xyz*Materials.xyz, Materials.w); // Decoding (cheap) sRGB -> linear color
 	return Out;
 }
 
@@ -179,6 +185,7 @@ FMaterialPixelParameters GetMaterialPixelParameters(FVertexFactoryInterpolantsVS
 	Result.HairPrimitiveRootUV = float2(Interpolants.HairPrimitiveRootUV);
 	Result.HairPrimitiveLength = Interpolants.HairPrimitiveLength;
 	Result.HairPrimitiveGroupIndex = Interpolants.HairPrimitiveGroupIndex;
+	Result.HairPrimitiveMaterial = Interpolants.HairPrimitiveMaterial;
 #if NUM_TEX_COORD_INTERPOLATORS
 	UNROLL
 	for (int CoordinateIndex = 0; CoordinateIndex < NUM_TEX_COORD_INTERPOLATORS; CoordinateIndex++)
@@ -266,6 +273,7 @@ FVertexFactoryIntermediates GetVertexFactoryIntermediates(FVertexFactoryInput In
 	Intermediates.HairPrimitiveRootUV		= CardsPoint.RootUVs;
 	Intermediates.HairPrimitiveLength		= CardsPoint.CardLength; 
 	Intermediates.HairPrimitiveGroupIndex	= CardsPoint.GroupIndex;
+	Intermediates.HairPrimitiveMaterial		= CardsPoint.Materials;
 
 	float TangentSign;
 	Intermediates.TangentToLocal = CalcTangentToLocal(Input, TangentSign);
@@ -323,6 +331,7 @@ FVertexFactoryInterpolantsVSToPS VertexFactoryGetInterpolantsVSToPS(FVertexFacto
 	Interpolants.HairPrimitiveRootUV		= CardsPoint.RootUVs;
 	Interpolants.HairPrimitiveLength		= CardsPoint.CardLength;
 	Interpolants.HairPrimitiveGroupIndex	= CardsPoint.GroupIndex;
+	Interpolants.HairPrimitiveMaterial		= CardsPoint.Materials;
 
 	SetPrimitiveId(Interpolants, Intermediates.SceneData.PrimitiveId);
 #if NUM_TEX_COORD_INTERPOLATORS 
@@ -448,6 +457,7 @@ FVertexFactoryRayTracingInterpolants VertexFactoryInterpolate(FVertexFactoryRayT
 	INTERPOLATE_MEMBER(InterpolantsVSToPS.HairPrimitiveRootUV);
 	INTERPOLATE_MEMBER(InterpolantsVSToPS.HairPrimitiveLength);
 	INTERPOLATE_MEMBER(InterpolantsVSToPS.HairPrimitiveGroupIndex);
+	INTERPOLATE_MEMBER(InterpolantsVSToPS.HairPrimitiveMaterial);
 #endif
 
 	return O;
diff --git a/Engine/Shaders/Private/HairStrands/HairStrandsTexturesGeneration.usf b/Engine/Shaders/Private/HairStrands/HairStrandsTexturesGeneration.usf
--- a/Engine/Shaders/Private/HairStrands/HairStrandsTexturesGeneration.usf
+++ b/Engine/Shaders/Private/HairStrands/HairStrandsTexturesGeneration.usf
@@ -81,9 +81,12 @@ float3 ComputeTangent(
 int2 OutputResolution;
 uint VertexCount;
 uint PrimitiveCount;
+uint bHasMaterialData;
 
 Buffer<uint4> InVF_PositionBuffer;
-Buffer<float4> InVF_Attribute0Buffer;
+Buffer<float2> InVF_Attribute0Buffer;
+Buffer<float4> InVF_MaterialBuffer;
+
 float3 InVF_PositionOffset;
 float InVF_Radius;
 float InVF_Length;
@@ -132,7 +135,8 @@ bool Trace(
 	inout float3 OutTangent,
 	inout float2 OutRootUV, 
 	inout float  OutU, 
-	inout float  OutSeed)
+	inout float  OutSeed,
+	inout float4 OutMaterial)
 {
 	bool bIsValid = false;
 
@@ -266,6 +270,14 @@ bool Trace(
 						const uint4 PackedCP_P0 = InVF_PositionBuffer[P0];
 						const uint4 PackedCP_P1 = InVF_PositionBuffer[P1];
 
+						float4 Material0 = 0;
+						float4 Material1 = 0;
+						if (bHasMaterialData)
+						{
+							Material0 = InVF_MaterialBuffer[P0];
+							Material1 = InVF_MaterialBuffer[P1];
+						}
+
 						const FHairControlPoint CP_Curr = UnpackHairControlPoint(PackedCP_P0, InVF_PositionOffset, InVF_Radius, InVF_RootScale, InVF_TipScale);
 						const FHairControlPoint CP_Next = UnpackHairControlPoint(PackedCP_P1, InVF_PositionOffset, InVF_Radius, InVF_RootScale, InVF_TipScale);
 
@@ -291,6 +303,7 @@ bool Trace(
 							OutSeed		= Attributes0.Seed;
 							OutRootUV	= Attributes0.RootUV;
 							OutU		= lerp(CP_Curr.UCoord, CP_Next.UCoord, VertexU);
+							OutMaterial = lerp(Material0, Material1, VertexU);
 
 							bIsValid = true;
 
@@ -334,7 +347,7 @@ void MainPS(
 	out float4 OutTangent	: SV_Target1,
 	out float  OutCoverage	: SV_Target2,
 	out float4 OutAttribute : SV_Target3,
-	out float4 OutGroupIndex : SV_Target4,
+	out float4 OutMaterial  : SV_Target4,
 	out uint   OutTriangleMask : SV_Target5,
 	out float  OutDepthTest	: SV_Depth)
 {
@@ -349,7 +362,7 @@ void MainPS(
 	OutCoverage = 0;
 	OutTangent = 0;
 	OutAttribute = 0;
-	OutGroupIndex = 0;
+	OutMaterial = 0;
 	OutDepthTest = 1;
 
 	float3 Tangent_ObjectSpace = 0;
@@ -397,13 +410,15 @@ void MainPS(
 		float2 SampleRootUV = 0;
 		float  SampleU = 0;
 		float  SampleSeed = 0;
-		if (Trace(RayO, RayD, MaxDistance, bDebugEnable, SampleDepth, SampleCoverage, SampleTangent, SampleRootUV, SampleU, SampleSeed))
+		float4 SampleMaterial = 0;
+		if (Trace(RayO, RayD, MaxDistance, bDebugEnable, SampleDepth, SampleCoverage, SampleTangent, SampleRootUV, SampleU, SampleSeed, SampleMaterial))
 		{
 			if (SampleDepth < OutDepth)
 			{
 				OutAttribute.xy = SampleRootUV;
 				OutAttribute.z = SampleU;
 				OutAttribute.w = SampleSeed;
+				OutMaterial.xyz= SampleMaterial.xyz;
 			}
 
 			OutDepth = min(OutDepth, SampleDepth);
@@ -424,7 +439,7 @@ void MainPS(
 		OutCoverage /= SampleCount;
 		OutTangent.xyz = (Tangent_FrameSpace + 1) * 0.5f;
 		OutTangent.w = 1;
-		OutGroupIndex = float4(InVF_GroupIndex / 255.f, 0, 0, 0);
+		OutMaterial.w = InVF_GroupIndex / 255.f;
 	}
 	else
 	{
@@ -432,7 +447,7 @@ void MainPS(
 		OutCoverage = 0;
 		OutTangent = 0;
 		OutAttribute = 0;
-		OutGroupIndex = 0;
+		OutMaterial = 0;
 	}
 
 	OutDepth = saturate(OutDepth / MaxDistance);
@@ -457,13 +472,13 @@ Texture2D<float>  Source_DepthTexture;
 Texture2D<float>  Source_CoverageTexture;
 Texture2D<float4> Source_TangentTexture;
 Texture2D<float4> Source_AttributeTexture;
-Texture2D<float4> Source_GroupIndexTexture;
+Texture2D<float4> Source_MaterialTexture;
 
 RWTexture2D<float>  Target_DepthTexture;
 RWTexture2D<float>  Target_CoverageTexture;
 RWTexture2D<float4> Target_TangentTexture;
 RWTexture2D<float4> Target_AttributeTexture;
-RWTexture2D<float4> Target_GroupIndexTexture;
+RWTexture2D<float4> Target_MaterialTexture;
 
 // Dilate the source texture (Depth/Coverage/Tanget/Attribute), using a distance to triangle edge metric. 
 // Original valid pixel are marked into the triangleMasktexture during the texture generation. Then 
@@ -490,7 +505,7 @@ void MainCS(int2 DispatchThreadId : SV_DispatchThreadID)
 		Target_CoverageTexture[PixelCoord]	= Source_CoverageTexture.Load(uint3(PixelCoord, 0));
 		Target_TangentTexture[PixelCoord]	= Source_TangentTexture.Load(uint3(PixelCoord, 0));
 		Target_AttributeTexture[PixelCoord] = Source_AttributeTexture.Load(uint3(PixelCoord, 0));
-		Target_GroupIndexTexture[PixelCoord] = Source_GroupIndexTexture.Load(uint3(PixelCoord, 0));
+		Target_MaterialTexture[PixelCoord]  = Source_MaterialTexture.Load(uint3(PixelCoord, 0));
 		return;
 	}
 	else
@@ -499,7 +514,7 @@ void MainCS(int2 DispatchThreadId : SV_DispatchThreadID)
 		float Target_Coverage = 0;
 		float4 Target_Tangent = 0;
 		float4 Target_Attribute = 0;
-		float Target_GroupIndex = 0;
+		float4 Target_Material = 0;
 
 		// Find the closest neighbor from triangle, amoung valid neighbord
 		bool bIsDilatedPixelValid = false;
@@ -522,7 +537,7 @@ void MainCS(int2 DispatchThreadId : SV_DispatchThreadID)
 				Target_Coverage		= Source_CoverageTexture.Load(uint3(NPixelCoord, 0));
 				Target_Tangent		= Source_TangentTexture.Load(uint3(NPixelCoord, 0));
 				Target_Attribute	= Source_AttributeTexture.Load(uint3(NPixelCoord, 0));
-				Target_GroupIndex	= Source_GroupIndexTexture.Load(uint3(NPixelCoord, 0)).x;
+				Target_Material		= Source_MaterialTexture.Load(uint3(NPixelCoord, 0));
 				MinDistance			= NDistance;
 				bIsDilatedPixelValid= true;
 			}
@@ -535,7 +550,7 @@ void MainCS(int2 DispatchThreadId : SV_DispatchThreadID)
 			Target_CoverageTexture[PixelCoord]	= Target_Coverage;
 			Target_TangentTexture[PixelCoord]	= Target_Tangent;
 			Target_AttributeTexture[PixelCoord] = Target_Attribute;
-			Target_GroupIndexTexture[PixelCoord] = float4(Target_GroupIndex,0,0,0);
+			Target_MaterialTexture[PixelCoord]  = Target_Material;
 
 			TriangleMaskTexture[PixelCoord] = PackDistance(MinDistance);
 		}
diff --git a/Engine/Shaders/Private/LocalVertexFactoryCommon.ush b/Engine/Shaders/Private/LocalVertexFactoryCommon.ush
--- a/Engine/Shaders/Private/LocalVertexFactoryCommon.ush
+++ b/Engine/Shaders/Private/LocalVertexFactoryCommon.ush
@@ -43,6 +43,7 @@ struct FVertexFactoryInterpolantsVSToPS
 #if HAIR_CARD_MESH_FACTORY
 	float2 HairPrimitiveUV							: HAIR_PRIMITIVE_UV;	// Atlas UV
 	float2 HairPrimitiveRootUV						: HAIR_PRIMITIVE_ROOTUV;// UV at the root point of the cards
+	float4 HairPrimitiveMaterial					: HAIR_PRIMITIVE_MATERIAL;
 	float  HairPrimitiveLength						: HAIR_PRIMITIVE_LENGTH;
 	nointerpolation float HairPrimitiveGroupIndex	: HAIR_PRIMITIVE_GROUPINDEX;
 #endif
diff --git a/Engine/Shaders/Private/MaterialTemplate.ush b/Engine/Shaders/Private/MaterialTemplate.ush
--- a/Engine/Shaders/Private/MaterialTemplate.ush
+++ b/Engine/Shaders/Private/MaterialTemplate.ush
@@ -381,6 +381,7 @@ struct FMaterialPixelParameters
 #if HAIR_CARD_MESH_FACTORY
 	float2	HairPrimitiveUV;		// AtlasUV
 	float2  HairPrimitiveRootUV;	// RootUV
+	float4	HairPrimitiveMaterial;	// Card material
 	float   HairPrimitiveLength;	// Card length
 	float	HairPrimitiveGroupIndex;// Card group index
 #endif
@@ -1339,7 +1340,7 @@ float3 MaterialExpressionGetHairBaseColor(FMaterialPixelParameters Parameters)
 #if HAIR_STRAND_MESH_FACTORY
 	return GetHairStrandsBaseColor(Parameters.HairPrimitiveId, Parameters.HairPrimitiveUV);
 #elif HAIR_CARD_MESH_FACTORY
-	return GetHairStrandsBaseColor(Parameters.HairPrimitiveUV);
+	return GetHairStrandsBaseColor(Parameters.HairPrimitiveUV, Parameters.HairPrimitiveMaterial.xyz);
 #else
 	return float3(0,0,0);
 #endif
@@ -1350,7 +1351,7 @@ float MaterialExpressionGetHairRoughness(FMaterialPixelParameters Parameters)
 #if HAIR_STRAND_MESH_FACTORY
 	return GetHairStrandsRoughness(Parameters.HairPrimitiveId, Parameters.HairPrimitiveUV);
 #elif HAIR_CARD_MESH_FACTORY
-	return GetHairStrandsRoughness(Parameters.HairPrimitiveUV);
+	return GetHairStrandsRoughness(Parameters.HairPrimitiveUV, Parameters.HairPrimitiveMaterial.w);
 #else
 	return 0;
 #endif
diff --git a/Engine/Source/Runtime/Renderer/Private/HairStrands/HairStrandsDebug.cpp b/Engine/Source/Runtime/Renderer/Private/HairStrands/HairStrandsDebug.cpp
--- a/Engine/Source/Runtime/Renderer/Private/HairStrands/HairStrandsDebug.cpp
+++ b/Engine/Source/Runtime/Renderer/Private/HairStrands/HairStrandsDebug.cpp
@@ -203,6 +203,7 @@ EHairStrandsDebugMode GetHairStrandsDebugStrandsMode()
 	case 11:  return EHairStrandsDebugMode::RenderVisCluster;
 	case 12:  return EHairStrandsDebugMode::RenderHairTangent;
 	case 13:  return EHairStrandsDebugMode::RenderHairControlPoints;
+	case 14:  return EHairStrandsDebugMode::RenderHairGroup;
 	default:  return EHairStrandsDebugMode::NoneDebug;
 	};
 }
