Remove null material validation for groom.

This is confusing for user and spams log. When no material is setup, we just fallback onto default material, withint any validation/warning.

diff --git a/Engine/Plugins/Runtime/HairStrands/Source/HairStrandsCore/Private/GroomComponent.cpp b/Engine/Plugins/Runtime/HairStrands/Source/HairStrandsCore/Private/GroomComponent.cpp
--- a/Engine/Plugins/Runtime/HairStrands/Source/HairStrandsCore/Private/GroomComponent.cpp
+++ b/Engine/Plugins/Runtime/HairStrands/Source/HairStrandsCore/Private/GroomComponent.cpp
@@ -3658,6 +3658,8 @@ void UGroomComponent::ValidateMaterials(bool bMapCheck) const
 					UE_LOG(LogHairStrands, Warning, TEXT("[Groom] %s - Groom's material needs to have Opaque blend mode. Groom's material will be replaced with default hair strands shader in editor."), *Name);
 				}
 			}break;
+			// Disable this warning as it is not really useful, and spams the log when editing parameters, since it is invoke on Invalidation & on PostEditChange
+			#if 0
 			case EHairMaterialCompatibility::Invalid_IsNull:
 			{
 				if (bMapCheck)
@@ -3672,6 +3674,7 @@ void UGroomComponent::ValidateMaterials(bool bMapCheck) const
 					UE_LOG(LogHairStrands, Warning, TEXT("[Groom] %s - Groom's material is not set and will fallback on default hair strands shader in editor."), *Name);
 				}
 			}break;
+			#endif
 		}
 	}
 }
