Fix new groom material slot.

diff --git a/Engine/Plugins/Runtime/HairStrands/Source/HairStrandsEditor/Private/GroomMaterialDetails.cpp b/Engine/Plugins/Runtime/HairStrands/Source/HairStrandsEditor/Private/GroomMaterialDetails.cpp
--- a/Engine/Plugins/Runtime/HairStrands/Source/HairStrandsEditor/Private/GroomMaterialDetails.cpp
+++ b/Engine/Plugins/Runtime/HairStrands/Source/HairStrandsEditor/Private/GroomMaterialDetails.cpp
@@ -249,7 +249,7 @@ void FGroomMaterialDetails::OnMaterialArrayChanged(UMaterialInterface* NewMateri
 		GroomAsset->Modify();
 		GroomAsset->HairGroupsMaterials[SlotIndex].Material = NewMaterial;
 
-		//Add a default name to the material slot if this slot was manually add and there is no name yet
+		// Add a default name to the material slot if this slot was manually add and there is no name yet
 		if (NewMaterial != nullptr && GroomAsset->HairGroupsMaterials[SlotIndex].SlotName == NAME_None)
 		{
 			if (GroomAsset->HairGroupsMaterials[SlotIndex].SlotName == NAME_None)
@@ -316,7 +316,31 @@ FReply FGroomMaterialDetails::AddMaterialSlot()
 
 	FScopedTransaction Transaction(LOCTEXT("PersonaAddMaterialSlotTransaction", "Persona editor: Add material slot"));
 	GroomAsset->Modify();
-	GroomAsset->HairGroupsMaterials.Add(FHairGroupsMaterial());
+	FHairGroupsMaterial NewMaterial;
+	NewMaterial.SlotName = FName(TEXT("Material"));
+
+	// Build a unique name
+	FName SlotName = NewMaterial.SlotName;
+	uint32 UniqueId = 0;
+	bool bHasUniqueName = true;
+	do
+	{
+		bHasUniqueName = true;
+		for (const FHairGroupsMaterial& Group : GroomAsset->HairGroupsMaterials)
+		{
+			if (Group.SlotName == SlotName)
+			{
+				bHasUniqueName = false;				
+				FString NewSlotName = NewMaterial.SlotName.ToString() + FString::FromInt(++UniqueId);
+				SlotName = FName(*NewSlotName);
+				break;
+			}
+		}
+	} while (!bHasUniqueName);
+
+	// Add new material
+	NewMaterial.SlotName = SlotName;
+	GroomAsset->HairGroupsMaterials.Add(NewMaterial);
 	GroomAsset->PostEditChange();
 
 	return FReply::Handled();
