Fix groom asset settings being partially wiped when reimporting with a different number of groups.

diff --git a/Engine/Plugins/Runtime/HairStrands/Source/HairStrandsCore/Private/GroomAsset.cpp b/Engine/Plugins/Runtime/HairStrands/Source/HairStrandsCore/Private/GroomAsset.cpp
--- a/Engine/Plugins/Runtime/HairStrands/Source/HairStrandsCore/Private/GroomAsset.cpp
+++ b/Engine/Plugins/Runtime/HairStrands/Source/HairStrandsCore/Private/GroomAsset.cpp
@@ -1648,7 +1648,7 @@ bool UGroomAsset::AreGroupsValid() const
 		HairGroupsLOD.Num() == GroupCount;
 }
 
-void UGroomAsset::SetNumGroup(uint32 InGroupCount, bool bResetGroupData)
+void UGroomAsset::SetNumGroup(uint32 InGroupCount, bool bResetGroupData, bool bResetOtherData)
 {
 	ReleaseResource();
 	if (bResetGroupData)
@@ -1664,37 +1664,80 @@ void UGroomAsset::SetNumGroup(uint32 InGroupCount, bool bResetGroupData)
 
 	if (InGroupCount != HairGroupsPhysics.Num())
 	{
-		HairGroupsPhysics.Init(FHairGroupsPhysics(), InGroupCount);
+		if (bResetOtherData)
+		{
+			HairGroupsPhysics.Init(FHairGroupsPhysics(), InGroupCount);
+		}
+		else
+		{
+			HairGroupsPhysics.SetNum(InGroupCount);
+		}
 	}
 
 	if (InGroupCount != HairGroupsRendering.Num())
 	{
-		HairGroupsRendering.Init(FHairGroupsRendering(), InGroupCount);
+		if (bResetOtherData)
+		{
+			HairGroupsRendering.Init(FHairGroupsRendering(), InGroupCount);
+		}
+		else
+		{
+			HairGroupsRendering.SetNum(InGroupCount);
+		}
 	}
 
 	if (InGroupCount != HairGroupsInterpolation.Num())
 	{
-		HairGroupsInterpolation.Init(FHairGroupsInterpolation(), InGroupCount);
+		if (bResetOtherData)
+		{
+			HairGroupsInterpolation.Init(FHairGroupsInterpolation(), InGroupCount);
+		}
+		else
+		{
+			HairGroupsInterpolation.SetNum(InGroupCount);
+		}
 	}
 
 	if (InGroupCount != HairGroupsLOD.Num())
 	{
-		HairGroupsLOD.Init(FHairGroupsLOD(), InGroupCount);
+		if (bResetOtherData)
+		{
+			HairGroupsLOD.Init(FHairGroupsLOD(), InGroupCount);
+		}
+		else
+		{
+			HairGroupsLOD.SetNum(InGroupCount);
+		}
 
 		// Insure that each group has at least one LOD
 		for (FHairGroupsLOD& GroupLOD : HairGroupsLOD)
 		{
-			FHairLODSettings& S = GroupLOD.LODs.AddDefaulted_GetRef();
-			S.ScreenSize = 1;
-			S.CurveDecimation = 1;
+			if (GroupLOD.LODs.IsEmpty())
+			{
+				FHairLODSettings& S = GroupLOD.LODs.AddDefaulted_GetRef();
+				S.ScreenSize = 1;
+				S.CurveDecimation = 1;
+			}
 		}
 	}
 
 	if (InGroupCount != HairGroupsInfo.Num())
 	{
-		HairGroupsInfo.Init(FHairGroupInfoWithVisibility(), InGroupCount);
+		if (bResetOtherData)
+		{
+			HairGroupsInfo.Init(FHairGroupInfoWithVisibility(), InGroupCount);
+		}
+		else
+		{
+			HairGroupsInfo.SetNum(InGroupCount);
+		}
 	}
 
+	// The following groups remain unchanged:
+	// * HairGroupsCards
+	// * HairGroupsMeshes
+	// * HairGroupsMaterials
+
 	EffectiveLODBias.SetNum(InGroupCount);
 
 #if WITH_EDITORONLY_DATA
diff --git a/Engine/Plugins/Runtime/HairStrands/Source/HairStrandsCore/Public/GroomAsset.h b/Engine/Plugins/Runtime/HairStrands/Source/HairStrandsCore/Public/GroomAsset.h
--- a/Engine/Plugins/Runtime/HairStrands/Source/HairStrandsCore/Public/GroomAsset.h
+++ b/Engine/Plugins/Runtime/HairStrands/Source/HairStrandsCore/Public/GroomAsset.h
@@ -525,7 +525,7 @@ public:
 	void ReleaseCardsResource(uint32 GroupIndex);
 	void ReleaseMeshesResource(uint32 GroupIndex);
 
-	void SetNumGroup(uint32 InGroupCount, bool bResetGroupData=true);
+	void SetNumGroup(uint32 InGroupCount, bool bResetGroupData=true, bool bResetOtherData=true);
 	bool AreGroupsValid() const;
 	int32 GetNumHairGroups() const;
 
diff --git a/Engine/Plugins/Runtime/HairStrands/Source/HairStrandsEditor/Private/ReimportHairStrandsFactory.cpp b/Engine/Plugins/Runtime/HairStrands/Source/HairStrandsEditor/Private/ReimportHairStrandsFactory.cpp
--- a/Engine/Plugins/Runtime/HairStrands/Source/HairStrandsEditor/Private/ReimportHairStrandsFactory.cpp
+++ b/Engine/Plugins/Runtime/HairStrands/Source/HairStrandsEditor/Private/ReimportHairStrandsFactory.cpp
@@ -349,7 +349,8 @@ EReimportResult::Type UReimportHairStrandsFactory::Reimport(UObject* Obj)
 		const uint32 GroupCount = OutDescription.HairGroups.Num();
 		if (uint32(HairAsset->GetNumHairGroups()) != GroupCount)
 		{
-			HairAsset->SetNumGroup(GroupCount);
+			// When reimporting an asset, if the number of groups has changed, we try to preserve existing settings
+			HairAsset->SetNumGroup(GroupCount, true /*bResetGroupData*/, false /*bResetOtherData*/);
 		}
 
 		// Initialized the settings value based on the assets interpolation options 
