Improve groom editor group readability and LOD info display.

diff --git a/Engine/Plugins/Runtime/HairStrands/Source/HairStrandsCore/Private/GroomComponent.cpp b/Engine/Plugins/Runtime/HairStrands/Source/HairStrandsCore/Private/GroomComponent.cpp
--- a/Engine/Plugins/Runtime/HairStrands/Source/HairStrandsCore/Private/GroomComponent.cpp
+++ b/Engine/Plugins/Runtime/HairStrands/Source/HairStrandsCore/Private/GroomComponent.cpp
@@ -2482,6 +2482,9 @@ void UGroomComponent::InitResources(bool bIsBindingReloading)
 		{
 			HairGroupInstance->HairGroupPublicData = new FHairGroupPublicData(GroupIt);
 			HairGroupInstance->HairGroupPublicData->Instance = HairGroupInstance;
+			HairGroupInstance->HairGroupPublicData->bDebugDrawLODInfo = bPreviewMode;
+			HairGroupInstance->HairGroupPublicData->DebugScreenSize = 0.f;
+			HairGroupInstance->HairGroupPublicData->DebugGroupColor = GetHairGroupDebugColor(GroupIt);
 			TArray<float> CPULODScreenSize;
 			TArray<bool> LODVisibility;
 			TArray<EHairGeometryType> LODGeometryTypes;
diff --git a/Engine/Plugins/Runtime/HairStrands/Source/HairStrandsCore/Private/GroomManager.cpp b/Engine/Plugins/Runtime/HairStrands/Source/HairStrandsCore/Private/GroomManager.cpp
--- a/Engine/Plugins/Runtime/HairStrands/Source/HairStrandsCore/Private/GroomManager.cpp
+++ b/Engine/Plugins/Runtime/HairStrands/Source/HairStrandsCore/Private/GroomManager.cpp
@@ -740,12 +740,14 @@ static void RunHairLODSelection(
 		float LODIndex = Instance->Debug.LODForcedIndex >= 0 ? FMath::Max(Instance->Debug.LODForcedIndex, MinLOD) : -1.0f;
 		float LODViewIndex = -1;
 		{
+			float MaxScreenSize = 0.f;
 			const FSphere SphereBound = Instance->ProxyBounds ? Instance->ProxyBounds->GetSphere() : FSphere(0);
 			for (const FSceneView* View : Views)
 			{
 				const float ScreenSize = ComputeBoundsScreenSize(FVector4(SphereBound.Center, 1), SphereBound.W, *View);
 				const float LODBias = Instance->Strands.Modifier.LODBias;
 				const float CurrLODViewIndex = FMath::Max(MinLOD, GetHairInstanceLODIndex(Instance->HairGroupPublicData->GetLODScreenSizes(), ScreenSize, LODBias));
+				MaxScreenSize = FMath::Max(MaxScreenSize, ScreenSize);
 
 				// Select highest LOD accross all views
 				LODViewIndex = LODViewIndex < 0 ? CurrLODViewIndex : FMath::Min(LODViewIndex, CurrLODViewIndex);
@@ -760,6 +762,7 @@ static void RunHairLODSelection(
 
 			// Feedback game thread with LOD selection 
 			Instance->Debug.LODPredictedIndex = LODViewIndex;
+			Instance->HairGroupPublicData->DebugScreenSize = MaxScreenSize;
 		}
 
 		// Function for selecting, loading, & initializing LOD resources
diff --git a/Engine/Plugins/Runtime/HairStrands/Source/HairStrandsEditor/Private/GroomAssetDetails.cpp b/Engine/Plugins/Runtime/HairStrands/Source/HairStrandsEditor/Private/GroomAssetDetails.cpp
--- a/Engine/Plugins/Runtime/HairStrands/Source/HairStrandsEditor/Private/GroomAssetDetails.cpp
+++ b/Engine/Plugins/Runtime/HairStrands/Source/HairStrandsEditor/Private/GroomAssetDetails.cpp
@@ -1140,44 +1140,51 @@ void FGroomRenderingDetails::OnGenerateElementForLODs(TSharedRef<IPropertyHandle
 	const FSlateFontInfo DetailFontInfo = IDetailLayoutBuilder::GetDetailFont();
 	FProperty* Property = StructProperty->GetProperty();
 
+	const FLinearColor LODColorBlock = GetHairGroupDebugColor(GroupIndex) * 0.25f;
+	const FLinearColor LODNameColor(FLinearColor::White);
+	static const FSlateBrush* GenericBrush = FCoreStyle::Get().GetBrush("GenericWhiteBox");
+	float OtherMargin = 2.0f;
+
 	ChildrenBuilder.AddCustomRow(LOCTEXT("HairInfo_Separator", "Separator"))
-	.ValueContent()
+	.WholeRowContent()
+	.VAlign(VAlign_Fill)
 	.HAlign(HAlign_Fill)
 	[
-		SNew(SHorizontalBox)
-		+ SHorizontalBox::Slot()
-		.VAlign(VAlign_Center)
+		SNew(SOverlay)
+		+ SOverlay::Slot()
 		[
-			SNew(SSeparator)
-			.Thickness(2)
+			SNew(SImage)
+			.Image(GenericBrush)
+			.ColorAndOpacity(LODColorBlock)
 		]
-		+ SHorizontalBox::Slot()
-		.AutoWidth()
-		.VAlign(VAlign_Center)
-		[
-			SNew(STextBlock)
-			.Font(IDetailLayoutBuilder::GetDetailFont())
-			.ColorAndOpacity(HairLODColor)
-			.Text(FText::Format(LOCTEXT("LOD", "LOD {0}"), FText::AsNumber(LODIndex)))
-		]
-		+ SHorizontalBox::Slot()
+		+ SOverlay::Slot()
+		.HAlign(HAlign_Right)
 		.VAlign(VAlign_Center)
 		[
-			SNew(SSeparator)
-			.Thickness(2)
-		]
-		+ SHorizontalBox::Slot()
-		.AutoWidth()
-		[
-			SNew(SButton)
+			SNew(SHorizontalBox)
+			+ SHorizontalBox::Slot()
 			.VAlign(VAlign_Center)
-			.HAlign(HAlign_Right)
-			.ButtonStyle(FEditorStyle::Get(), "HoverHintOnly")
-			.OnClicked(this, &FGroomRenderingDetails::OnRemoveLODClicked, GroupIndex, LODIndex, Property)
+			.Padding(OtherMargin)
 			[
-				SNew(SImage)
-				.Image(FEditorStyle::GetBrush("Cross"))
+				SNew(STextBlock)
+				.Font(IDetailLayoutBuilder::GetDetailFont())
+				.ColorAndOpacity(LODNameColor)
+				.Text(FText::Format(LOCTEXT("LOD", "LOD {0}"), FText::AsNumber(LODIndex)))
 			]
+			+ SHorizontalBox::Slot()
+			.AutoWidth()
+			[
+				SNew(SButton)
+				.VAlign(VAlign_Center)
+				.HAlign(HAlign_Right)
+				.ButtonStyle(FEditorStyle::Get(), "HoverHintOnly")
+				.OnClicked(this, &FGroomRenderingDetails::OnRemoveLODClicked, GroupIndex, LODIndex, Property)
+				[
+					SNew(SImage)
+					.Image(FEditorStyle::GetBrush("Cross"))
+				]
+			]
+			
 		]
 	];
 
@@ -1243,7 +1250,7 @@ FName FGroomRenderingDetails::GetGroupName(int32 GroupIndex) const
 	return NAME_None;
 }
 
-TSharedRef<SWidget> FGroomRenderingDetails::MakeGroupNameCustomization(int32 GroupIndex)
+TSharedRef<SWidget> FGroomRenderingDetails::MakeGroupNameCustomization(int32 GroupIndex, const FLinearColor& GroupColor)
 {
 	switch (PanelType)
 	{
@@ -1251,7 +1258,7 @@ TSharedRef<SWidget> FGroomRenderingDetails::MakeGroupNameCustomization(int32 Gro
 	{
 		return SNew(STextBlock)
 			.Font(IDetailLayoutBuilder::GetDetailFont())
-			.ColorAndOpacity(HairGroupColor)
+			.ColorAndOpacity(GroupColor)
 			.Text(FText::Format(LOCTEXT("Cards", "Cards {0} "), FText::AsNumber(GroupIndex)));
 	}
 	break;
@@ -1259,7 +1266,7 @@ TSharedRef<SWidget> FGroomRenderingDetails::MakeGroupNameCustomization(int32 Gro
 	{
 		return SNew(STextBlock)
 			.Font(IDetailLayoutBuilder::GetDetailFont())
-			.ColorAndOpacity(HairGroupColor)
+			.ColorAndOpacity(GroupColor)
 			.Text(FText::Format(LOCTEXT("Meshes", "Meshes {0} "), FText::AsNumber(GroupIndex)));
 	}
 	break;
@@ -1270,14 +1277,14 @@ TSharedRef<SWidget> FGroomRenderingDetails::MakeGroupNameCustomization(int32 Gro
 		{
 			return SNew(STextBlock)
 				.Font(IDetailLayoutBuilder::GetDetailFont())
-				.ColorAndOpacity(HairGroupColor)
+				.ColorAndOpacity(GroupColor)
 				.Text(FText::Format(LOCTEXT("Group", "Group ID {0} - {1}"), FText::AsNumber(GroupIndex), FText::FromName(GroupName)));
 		}
 		else
 		{
 			return SNew(STextBlock)
 				.Font(IDetailLayoutBuilder::GetDetailFont())
-				.ColorAndOpacity(HairGroupColor)
+				.ColorAndOpacity(GroupColor)
 				.Text(FText::Format(LOCTEXT("Group", "Group ID {0}"), FText::AsNumber(GroupIndex)));
 		}
 	}
@@ -1291,53 +1298,49 @@ TSharedRef<SWidget> FGroomRenderingDetails::MakeGroupNameCustomization(int32 Gro
 void FGroomRenderingDetails::OnGenerateElementForHairGroup(TSharedRef<IPropertyHandle> StructProperty, int32 GroupIndex, IDetailChildrenBuilder& ChildrenBuilder, IDetailLayoutBuilder* DetailLayout)
 {
 	const FSlateFontInfo DetailFontInfo = IDetailLayoutBuilder::GetDetailFont();
-	const FLinearColor GroupColor = GetHairGroupDebugColor(GroupIndex);
+	const FLinearColor GroupColorBlock = GetHairGroupDebugColor(GroupIndex) * 0.75f;
+	const FLinearColor GroupNameColor(FLinearColor::White);
 
 	FProperty* Property = StructProperty->GetProperty();
 
+	static const FSlateBrush* GenericBrush = FCoreStyle::Get().GetBrush("GenericWhiteBox");
+
+	float OtherMargin = 2.0f;
+	float RightMargin = 2.0f;
+	if (PanelType != EMaterialPanelType::LODs && PanelType != EMaterialPanelType::Cards && PanelType != EMaterialPanelType::Meshes)
+	{
+		RightMargin = 10.0f;
+	}
+
 	ChildrenBuilder.AddCustomRow(LOCTEXT("HairInfo_Separator", "Separator"))
-	.NameContent()
-	[
-		SNew(SHorizontalBox)
-		+ SHorizontalBox::Slot()
-		.VAlign(VAlign_Center)
-		[
-			SNew(SSeparator)
-			.Thickness(2)
-			.ColorAndOpacity(HairGroupColor)
-		]
-	]
-	.ValueContent()
+	.WholeRowContent()
+	.VAlign(VAlign_Fill)
 	.HAlign(HAlign_Fill)
 	[
-		SNew(SHorizontalBox)
-		+ SHorizontalBox::Slot()
-		.VAlign(VAlign_Center)
+		SNew(SOverlay)
+		+ SOverlay::Slot()
 		[
-			SNew(SSeparator)
-			.Thickness(2)
+			SNew(SImage)
+			.Image(GenericBrush)
+			.ColorAndOpacity(GroupColorBlock)
 		]
-		+ SHorizontalBox::Slot()
-		.AutoWidth()
-		.VAlign(VAlign_Center)
-		.HAlign(HAlign_Right)
-		[
-			MakeGroupNameCustomization(GroupIndex)
-		]
-		+ SHorizontalBox::Slot()
-		.AutoWidth()
+		+ SOverlay::Slot()
 		.HAlign(HAlign_Right)
 		.VAlign(VAlign_Center)
 		[
-			SNew(SColorBlock)
-			.Color(GroupColor)
-			.Size(FVector2D(8,8))
-			.CornerRadius(FVector4f(4.0f, 4.0f, 4.0f, 4.0f))
-		]
-		+ SHorizontalBox::Slot()
-		.AutoWidth()
-		[
-			MakeGroupNameButtonCustomization(GroupIndex, Property)
+			SNew(SHorizontalBox)
+			+ SHorizontalBox::Slot()
+			.VAlign(VAlign_Center)
+			.Padding(OtherMargin, OtherMargin, RightMargin, OtherMargin)
+			[
+				MakeGroupNameCustomization(GroupIndex, GroupNameColor)
+			]
+			+ SHorizontalBox::Slot()
+			.AutoWidth()
+			[
+				MakeGroupNameButtonCustomization(GroupIndex, Property)
+			]
+			
 		]
 	];
 
diff --git a/Engine/Plugins/Runtime/HairStrands/Source/HairStrandsEditor/Private/GroomAssetDetails.h b/Engine/Plugins/Runtime/HairStrands/Source/HairStrandsEditor/Private/GroomAssetDetails.h
--- a/Engine/Plugins/Runtime/HairStrands/Source/HairStrandsEditor/Private/GroomAssetDetails.h
+++ b/Engine/Plugins/Runtime/HairStrands/Source/HairStrandsEditor/Private/GroomAssetDetails.h
@@ -74,7 +74,7 @@ private:
 	void OnResetToBaseClicked(TSharedPtr<IPropertyHandle> PropertyHandle);
 	TSharedRef<SWidget> CreateMaterialSwatch(const TSharedPtr<FAssetThumbnailPool>& ThumbnailPool, int32 GroupIndex);
 
-	TSharedRef<SWidget> MakeGroupNameCustomization(int32 GroupIndex);
+	TSharedRef<SWidget> MakeGroupNameCustomization(int32 GroupIndex, const FLinearColor& GroupColor);
 	TSharedRef<SWidget> MakeGroupNameButtonCustomization(int32 GroupIndex, FProperty* Property);
 
 	FReply OnRemoveLODClicked(int32 GroupIndex, int32 LODIndex, FProperty* Property);
diff --git a/Engine/Plugins/Runtime/HairStrands/Source/HairStrandsEditor/Private/GroomCustomAssetEditorToolkit.cpp b/Engine/Plugins/Runtime/HairStrands/Source/HairStrandsEditor/Private/GroomCustomAssetEditorToolkit.cpp
--- a/Engine/Plugins/Runtime/HairStrands/Source/HairStrandsEditor/Private/GroomCustomAssetEditorToolkit.cpp
+++ b/Engine/Plugins/Runtime/HairStrands/Source/HairStrandsEditor/Private/GroomCustomAssetEditorToolkit.cpp
@@ -107,8 +107,6 @@ void FGroomCustomAssetEditorToolkit::UnregisterTabSpawners(const TSharedRef<clas
 #endif
 }
 
-static uint32 GOpenedGroomEditorCount = 0;
-
 void FGroomCustomAssetEditorToolkit::DocPropChanged(UObject *InObject, FPropertyChangedEvent &Property)
 {
 	if (!GroomAsset.Get())
@@ -324,7 +322,7 @@ void FGroomCustomAssetEditorToolkit::InitCustomAssetEditor(const EToolkitMode::T
 	ViewportTab = SNew(SGroomEditorViewport);
 	
 	// Default layout
-	const TSharedRef<FTabManager::FLayout> StandaloneDefaultLayout = FTabManager::NewLayout("Standalone_GroomAssetEditor_Layout_v15")
+	const TSharedRef<FTabManager::FLayout> StandaloneDefaultLayout = FTabManager::NewLayout("Standalone_GroomAssetEditor_Layout_v15a")
 		->AddArea
 		(
 			FTabManager::NewPrimaryArea()
@@ -524,24 +522,11 @@ FGroomCustomAssetEditorToolkit::FGroomCustomAssetEditorToolkit()
 	PreviewGroomComponent			= nullptr;	
 	PreviewStaticMeshComponent		= nullptr;
 	PreviewSkeletalMeshComponent	= nullptr;
-
-	// Set LOD debug mode to show current LOD index & screen size
-	if (GOpenedGroomEditorCount == 0)
-	{
-		SetHairScreenLODInfo(true);
-	}
-	++GOpenedGroomEditorCount;
 }
 
 FGroomCustomAssetEditorToolkit::~FGroomCustomAssetEditorToolkit()
 {
-	// Disable LOD debug mode info
-	--GOpenedGroomEditorCount;
-	if (GOpenedGroomEditorCount == 0)
-	{
-		SetHairScreenLODInfo(false);
-	}
-	check(GOpenedGroomEditorCount >= 0);
+
 }
 
 FText FGroomCustomAssetEditorToolkit::GetToolkitName() const
diff --git a/Engine/Shaders/Private/HairStrands/HairStrandsDebugPrint.usf b/Engine/Shaders/Private/HairStrands/HairStrandsDebugPrint.usf
--- a/Engine/Shaders/Private/HairStrands/HairStrandsDebugPrint.usf
+++ b/Engine/Shaders/Private/HairStrands/HairStrandsDebugPrint.usf
@@ -375,4 +375,65 @@ void MainCS(int2 DispatchThreadId : SV_DispatchThreadID)
 }
 #endif // SHADER_TANGENT
 
-////////////////////////////////////////////////////////////////////////////////////////////////
\ No newline at end of file
+////////////////////////////////////////////////////////////////////////////////////////////////
+
+#if SHADER_LOD_INFO
+
+int2 MaxResolution;
+uint GroupIndex;
+uint GeometryType;
+float LOD;
+float ScreenSize;
+float3 GroupColor;
+
+[numthreads(1, 1, 1)]
+void MainCS(uint3 DispatchThreadId : SV_DispatchThreadID)
+{
+	// Legend
+	if (GroupIndex == 0)
+	{
+		float2 OriginalPos = float2(15, 50) / float2(MaxResolution);
+		float2 Pos = OriginalPos;
+		Pos = ShaderPrintText(Pos, TEXT("Group Geom.   LOD      Screen Size"), FontSilver); AddNewLine(OriginalPos, Pos);
+	}
+
+	// Values
+	{
+		float2 OriginalPos = float2(30, 50) / float2(MaxResolution);
+		OriginalPos.y = float(65 + GroupIndex * 15) / float(MaxResolution.y);
+		float2 Pos = OriginalPos;
+
+		// Group Color
+		Pos = ShaderPrintText(Pos, TEXT("* "), InitFontColor(GroupColor));
+
+		// Group index
+		ShaderPrint(Pos, GroupIndex, FontEmerald);
+		Pos.x += 20.f / float(MaxResolution.x);
+
+		// Geometry Type
+		if (GeometryType == 0)
+		{
+			Pos = ShaderPrintText(Pos, TEXT("Strands "), FontOrange);
+		}
+		else if (GeometryType == 1)
+		{
+			Pos = ShaderPrintText(Pos, TEXT("Cards   "), FontOrange);
+		}
+		else if (GeometryType == 2)
+		{
+			Pos = ShaderPrintText(Pos, TEXT("Meshes  "), FontOrange);
+		}
+
+		// LOD
+		ShaderPrint(Pos, LOD, FontWhite);
+		Pos.x += 80.f / float(MaxResolution.x);
+
+		// Screen size
+		ShaderPrint(Pos, saturate(ScreenSize), FontWhite);
+		Pos.x += 80.f / float(MaxResolution.x);
+
+		AddNewLine(OriginalPos, Pos);
+	}
+}
+
+#endif // SHADER_LOD_INFO
\ No newline at end of file
diff --git a/Engine/Source/Runtime/Renderer/Private/HairStrands/HairStrandsDebug.cpp b/Engine/Source/Runtime/Renderer/Private/HairStrands/HairStrandsDebug.cpp
--- a/Engine/Source/Runtime/Renderer/Private/HairStrands/HairStrandsDebug.cpp
+++ b/Engine/Source/Runtime/Renderer/Private/HairStrands/HairStrandsDebug.cpp
@@ -86,16 +86,6 @@ bool IsHairStrandsClusterDebugAABBEnable()
 	return GHairStrandsClusterDebug > 1;
 }
 
-void SetHairScreenLODInfo(bool bEnable)
-{
-	if (!ShaderPrint::IsEnabled())
-	{
-		ShaderPrint::SetEnabled(true);
-		ShaderPrint::SetFontSize(8);
-	}
-	GHairStrandsClusterDebug = bEnable ? 4 : 1;
-}
-
 FHairStrandsDebugData::Data FHairStrandsDebugData::CreateData(FRDGBuilder& GraphBuilder)
 {
 	FHairStrandsDebugData::Data Out;
@@ -230,6 +220,79 @@ static const TCHAR* ToString(EHairStrandsDebugMode DebugMode)
 	};
 }
 
+///////////////////////////////////////////////////////////////////////////////////////////////////
+class FHairPrintLODInfoCS : public FGlobalShader
+{
+	DECLARE_GLOBAL_SHADER(FHairPrintLODInfoCS);
+	SHADER_USE_PARAMETER_STRUCT(FHairPrintLODInfoCS, FGlobalShader);
+
+	BEGIN_SHADER_PARAMETER_STRUCT(FParameters, )
+		SHADER_PARAMETER(FIntPoint, MaxResolution)
+		SHADER_PARAMETER(FVector3f, GroupColor)
+		SHADER_PARAMETER(uint32, GroupIndex)
+		SHADER_PARAMETER(uint32, GeometryType)
+		SHADER_PARAMETER(float, ScreenSize)
+		SHADER_PARAMETER(float, LOD)
+		SHADER_PARAMETER_STRUCT_INCLUDE(ShaderPrint::FShaderParameters, ShaderPrintUniformBuffer)
+		END_SHADER_PARAMETER_STRUCT()
+
+public:
+	static bool ShouldCompilePermutation(const FGlobalShaderPermutationParameters& Parameters) { return IsHairStrandsSupported(EHairStrandsShaderType::Tool, Parameters.Platform); }
+	static void ModifyCompilationEnvironment(const FGlobalShaderPermutationParameters& Parameters, FShaderCompilerEnvironment& OutEnvironment)
+	{
+		// Skip optimization for avoiding long compilation time due to large UAV writes
+		FGlobalShader::ModifyCompilationEnvironment(Parameters, OutEnvironment);
+		OutEnvironment.CompilerFlags.Add(CFLAG_Debug);
+		OutEnvironment.SetDefine(TEXT("SHADER_LOD_INFO"), 1);
+	}
+};
+
+IMPLEMENT_GLOBAL_SHADER(FHairPrintLODInfoCS, "/Engine/Private/HairStrands/HairStrandsDebugPrint.usf", "MainCS", SF_Compute);
+
+static void AddPrintLODInfoPass(
+	FRDGBuilder& GraphBuilder,
+	const FViewInfo& View,
+	const FHairGroupPublicData* Data)
+{
+	if (!ShaderPrint::IsSupported(View.Family->GetShaderPlatform()))
+	{
+		return;
+	}
+
+	if (!ShaderPrint::IsEnabled(View))
+	{
+		ShaderPrint::SetEnabled(true);
+		ShaderPrint::RequestSpaceForCharacters(2000);
+	}
+
+	const uint32 GroupIndex = Data->GetGroupIndex();
+	const FLinearColor GroupColor = Data->DebugGroupColor;
+	const uint32 IntLODIndex = Data->LODIndex;
+
+	FHairPrintLODInfoCS::FParameters* Parameters = GraphBuilder.AllocParameters<FHairPrintLODInfoCS::FParameters>();
+	Parameters->MaxResolution = FIntPoint(View.ViewRect.Width(), View.ViewRect.Height());
+	Parameters->GroupIndex = GroupIndex;
+	Parameters->LOD = Data->LODIndex;
+	Parameters->GroupColor = FVector3f(GroupColor.R, GroupColor.G, GroupColor.B);
+	Parameters->ScreenSize = Data->DebugScreenSize;
+	switch (Data->VFInput.GeometryType)
+	{
+	case EHairGeometryType::Strands: Parameters->GeometryType = 0; break;
+	case EHairGeometryType::Cards  : Parameters->GeometryType = 1; break;
+	case EHairGeometryType::Meshes : Parameters->GeometryType = 2; break;
+	}
+	ShaderPrint::SetParameters(GraphBuilder, View, Parameters->ShaderPrintUniformBuffer);
+	TShaderMapRef<FHairPrintLODInfoCS> ComputeShader(View.ShaderMap);
+
+	ClearUnusedGraphResources(ComputeShader, Parameters);
+
+	FComputeShaderUtils::AddPass(
+		GraphBuilder,
+		RDG_EVENT_NAME("HairStrands::PrintLODInfo(%d/%d)", Parameters->GroupIndex, Parameters->GroupIndex),
+		ComputeShader,
+		Parameters,
+		FIntVector(1, 1, 1));
+}
 
 ///////////////////////////////////////////////////////////////////////////////////////////////////
 class FHairDebugPrintCS : public FGlobalShader
@@ -1143,6 +1206,25 @@ static void InternalRenderHairStrandsDebugInfo(
 		}
 	}
 
+	// Draw LOD info 
+	for (const FMeshBatchAndRelevance& Mesh : View.HairStrandsMeshElements)
+	{
+		const FHairGroupPublicData* GroupData = HairStrands::GetHairData(Mesh.Mesh);
+		if (GroupData->bDebugDrawLODInfo)
+		{
+			AddPrintLODInfoPass(GraphBuilder, View, GroupData);
+		}
+	}
+	for (const FMeshBatchAndRelevance& Mesh : View.HairCardsMeshElements)
+	{
+		const FHairGroupPublicData* GroupData = HairStrands::GetHairData(Mesh.Mesh);
+		if (GroupData->bDebugDrawLODInfo)
+		{
+			AddPrintLODInfoPass(GraphBuilder, View, GroupData);
+		}
+	}
+
+
 	// Pass this point, all debug rendering concern only hair strands data
 	if (!HairStrands::HasViewHairStrandsData(View))
 	{
diff --git a/Engine/Source/Runtime/Renderer/Private/ShaderPrint.cpp b/Engine/Source/Runtime/Renderer/Private/ShaderPrint.cpp
--- a/Engine/Source/Runtime/Renderer/Private/ShaderPrint.cpp
+++ b/Engine/Source/Runtime/Renderer/Private/ShaderPrint.cpp
@@ -46,6 +46,7 @@ namespace ShaderPrint
 		ECVF_Cheat | ECVF_RenderThreadSafe);
 
 	static uint32 GCharacterRequestCount = 0;
+	static bool GShaderPrintEnableOverride = false;
 
 	// Structure used by shader buffers to store values and symbols
 	struct ShaderPrintItem
@@ -126,7 +127,14 @@ namespace ShaderPrint
 
 	void SetEnabled(bool bInEnabled)
 	{
-		CVarEnable->Set(bInEnabled);
+		if (IsInGameThread())
+		{
+			CVarEnable->Set(bInEnabled);
+		}
+		else
+		{
+			GShaderPrintEnableOverride = bInEnabled;
+		}
 	}
 
 	void SetFontSize(int32 InFontSize)
@@ -146,7 +154,7 @@ namespace ShaderPrint
 
 	bool IsEnabled()
 	{
-		return CVarEnable.GetValueOnAnyThread() != 0;
+		return CVarEnable.GetValueOnAnyThread() != 0 || GShaderPrintEnableOverride;
 	}
 
 	bool IsEnabled(const FViewInfo& View)
diff --git a/Engine/Source/Runtime/Renderer/Public/HairStrandsInterface.h b/Engine/Source/Runtime/Renderer/Public/HairStrandsInterface.h
--- a/Engine/Source/Runtime/Renderer/Public/HairStrandsInterface.h
+++ b/Engine/Source/Runtime/Renderer/Public/HairStrandsInterface.h
@@ -341,6 +341,11 @@ public:
 	float LODIndex = -1;		// Current LOD used for all views
 	float LODBias = 0;			// Current LOD bias
 	bool bLODVisibility = true; // Enable/disable hair rendering for this component
+
+	// Debug
+	bool  bDebugDrawLODInfo = false; // Enable/disable hair LOD info
+	float DebugScreenSize = 0.f;
+	FLinearColor DebugGroupColor;
 	EHairStrandsDebugMode DebugMode = EHairStrandsDebugMode::NoneDebug;
 };
 
@@ -422,9 +427,6 @@ RENDERER_API float GetHairCoverage(uint32 HairCount, float AverageHairRadius);
 /// Return the average hair normalized radius for a given hair count and a given coverage value
 RENDERER_API float GetHairAvgRadius(uint32 InCount, float InCoverage);
 
-/// Helper to enable debug information about hair LOD
-RENDERER_API void SetHairScreenLODInfo(bool bEnable);
-
 
 //////////////////////////////////////////////////////////////////////////////////////////////////////////////////////
 // HairStrands Bookmark API
