Update Groom editor Look&Feel UI to be consistent.

diff --git a/Engine/Plugins/Runtime/HairStrands/Source/HairStrandsEditor/Private/GroomAssetDetails.cpp b/Engine/Plugins/Runtime/HairStrands/Source/HairStrandsEditor/Private/GroomAssetDetails.cpp
--- a/Engine/Plugins/Runtime/HairStrands/Source/HairStrandsEditor/Private/GroomAssetDetails.cpp
+++ b/Engine/Plugins/Runtime/HairStrands/Source/HairStrandsEditor/Private/GroomAssetDetails.cpp
@@ -328,7 +328,7 @@ void FGroomRenderingDetails::CustomizeDetails(IDetailLayoutBuilder& DetailLayout
 //////////////////////////////////////////////////////////////////////////////////////////////////////////////////
 // Custom widget for material slot for hair rendering
 
-void FGroomRenderingDetails::AddNewGroupButton(IDetailCategoryBuilder& FilesCategory, FProperty* Property)
+void FGroomRenderingDetails::AddNewGroupButton(IDetailCategoryBuilder& FilesCategory, FProperty* Property, const FText& HeaderText)
 {
 	// Add a button for adding element to the hair groups array
 	FilesCategory.AddCustomRow(FText::FromString(TEXT("AddGroup")))
@@ -337,6 +337,15 @@ void FGroomRenderingDetails::AddNewGroupButton(IDetailCategoryBuilder& FilesCate
 		SNew(SHorizontalBox)
 		+ SHorizontalBox::Slot()
 		.AutoWidth()
+		.VAlign(VAlign_Center)
+		.Padding(2.f, 2.f, 10.f, 2.f)
+		[
+			SNew(STextBlock)
+			.Font(IDetailLayoutBuilder::GetDetailFont())
+			.Text(HeaderText)
+		]
+		+ SHorizontalBox::Slot()
+		.AutoWidth()
 		[
 			SNew(SButton)
 			.VAlign(VAlign_Center)
@@ -345,7 +354,7 @@ void FGroomRenderingDetails::AddNewGroupButton(IDetailCategoryBuilder& FilesCate
 			.OnClicked(this, &FGroomRenderingDetails::OnAddGroup, Property)
 			[
 				SNew(SImage)
-				.Image(FEditorStyle::GetBrush("Plus"))
+				.Image(FEditorStyle::GetBrush("Icons.PlusCircle"))
 			]
 		]
 	];
@@ -458,7 +467,7 @@ void FGroomRenderingDetails::CustomizeStrandsGroupProperties(IDetailLayoutBuilde
 		case EMaterialPanelType::Cards:		
 		{
 			TSharedRef<IPropertyHandle> Property = DetailLayout.GetProperty(GET_MEMBER_NAME_CHECKED(UGroomAsset, HairGroupsCards), UGroomAsset::StaticClass());
-			AddNewGroupButton(FilesCategory, Property->GetProperty());
+			AddNewGroupButton(FilesCategory, Property->GetProperty(), FText::FromString(TEXT("Add Card asset")));
 			if (Property->IsValidHandle())
 			{
 				TSharedRef<FDetailArrayBuilder> PropertyBuilder = MakeShareable(new FDetailArrayBuilder(Property, false, false, false));
@@ -471,7 +480,7 @@ void FGroomRenderingDetails::CustomizeStrandsGroupProperties(IDetailLayoutBuilde
 		case EMaterialPanelType::Meshes:	
 		{
 			TSharedRef<IPropertyHandle> Property = DetailLayout.GetProperty(GET_MEMBER_NAME_CHECKED(UGroomAsset, HairGroupsMeshes), UGroomAsset::StaticClass());
-			AddNewGroupButton(FilesCategory, Property->GetProperty());
+			AddNewGroupButton(FilesCategory, Property->GetProperty(), FText::FromString(TEXT("Add Mesh asset")));
 			if (Property->IsValidHandle())
 			{
 				TSharedRef<FDetailArrayBuilder> PropertyBuilder = MakeShareable(new FDetailArrayBuilder(Property, false, false, false));
@@ -1181,7 +1190,7 @@ void FGroomRenderingDetails::OnGenerateElementForLODs(TSharedRef<IPropertyHandle
 				.OnClicked(this, &FGroomRenderingDetails::OnRemoveLODClicked, GroupIndex, LODIndex, Property)
 				[
 					SNew(SImage)
-					.Image(FEditorStyle::GetBrush("Cross"))
+					.Image(FEditorStyle::GetBrush("Icons.Delete"))
 				]
 			]
 			
@@ -1206,7 +1215,7 @@ TSharedRef<SWidget> FGroomRenderingDetails::MakeGroupNameButtonCustomization(int
 			.OnClicked(this, &FGroomRenderingDetails::OnAddLODClicked, GroupIndex, Property)
 			[
 				SNew(SImage)
-				.Image(FEditorStyle::GetBrush("Plus"))
+				.Image(FEditorStyle::GetBrush("Icons.PlusCircle"))
 			];
 	}
 	break;
@@ -1219,7 +1228,7 @@ TSharedRef<SWidget> FGroomRenderingDetails::MakeGroupNameButtonCustomization(int
 			.OnClicked(this, &FGroomRenderingDetails::OnRemoveGroupClicked, GroupIndex, Property)
 			[
 				SNew(SImage)
-				.Image(FEditorStyle::GetBrush("Cross"))
+				.Image(FEditorStyle::GetBrush("Icons.Delete"))
 			];
 	}
 	break;
@@ -1232,7 +1241,7 @@ TSharedRef<SWidget> FGroomRenderingDetails::MakeGroupNameButtonCustomization(int
 			.OnClicked(this, &FGroomRenderingDetails::OnRemoveGroupClicked, GroupIndex, Property)
 			[
 				SNew(SImage)
-				.Image(FEditorStyle::GetBrush("Cross"))
+				.Image(FEditorStyle::GetBrush("Icons.Delete"))
 			];
 	}
 	break;
diff --git a/Engine/Plugins/Runtime/HairStrands/Source/HairStrandsEditor/Private/GroomAssetDetails.h b/Engine/Plugins/Runtime/HairStrands/Source/HairStrandsEditor/Private/GroomAssetDetails.h
--- a/Engine/Plugins/Runtime/HairStrands/Source/HairStrandsEditor/Private/GroomAssetDetails.h
+++ b/Engine/Plugins/Runtime/HairStrands/Source/HairStrandsEditor/Private/GroomAssetDetails.h
@@ -52,7 +52,7 @@ private:
 	int32 GetGroupCount() const;
 	FName GetGroupName(int32 GroupIndex) const;
 
-	void AddNewGroupButton(IDetailCategoryBuilder& StrandsGroupFilesCategory, FProperty* Property);
+	void AddNewGroupButton(IDetailCategoryBuilder& StrandsGroupFilesCategory, FProperty* Property, const FText& HeaderText);
 
 	void ApplyChanges();
 	void SetMaterialSlot(int32 GroupIndex, int32 MaterialIndex);
