Fix bad rbf deformation +  Add collision components + Fix clear/copy velocity grid + Remove rest update while resetting hairs

diff --git a/Engine/Plugins/Runtime/HairStrands/Shaders/Private/NiagaraClearPressureGrid.usf b/Engine/Plugins/Runtime/HairStrands/Shaders/Private/NiagaraClearPressureGrid.usf
--- a/Engine/Plugins/Runtime/HairStrands/Shaders/Private/NiagaraClearPressureGrid.usf
+++ b/Engine/Plugins/Runtime/HairStrands/Shaders/Private/NiagaraClearPressureGrid.usf
@@ -19,24 +19,24 @@ void MainCS(uint3 DispatchThreadId : SV_DispatchThreadID)
 {
 	if (all(DispatchThreadId < GridSize))
 	{
-		const int LinearIndexYZ = DispatchThreadId.y * (GridSize.z+1) + DispatchThreadId.z;
-		const int GridSizeYZ = (GridSize.y+1) * (GridSize.z+1);
+		const int LinearIndexYZ = DispatchThreadId.y * GridSize.z + DispatchThreadId.z;
+		const int GridSizeYZ = GridSize.y * GridSize.z;
 		
 		[unroll]
 		for(int i = 0; i < NUM_ATTRIBUTES; ++i)
 		{
-			const int AttrIndex = DispatchThreadId.x+i*(GridSize.x+1);
+			const int AttrIndex = DispatchThreadId.x+i*GridSize.x;
 			const int LinearIndex = AttrIndex.x * GridSizeYZ + LinearIndexYZ;
 			GridDestinationBuffer[LinearIndex] = 0;
 		}
 		if(CopyPressure)
 		{
-			const int PressureIndex = DispatchThreadId.x+FLUID_PRESSURE_OFFSET*(GridSize.x+1);
+			const int PressureIndex = DispatchThreadId.x+FLUID_PRESSURE_OFFSET*GridSize.x;
 			const int LinearIndex = PressureIndex.x * GridSizeYZ + LinearIndexYZ;
 			GridDestinationBuffer[LinearIndex] = GridCurrentBuffer[LinearIndex];
 		}
 
-		const int DistanceIndex = DispatchThreadId.x+FLUID_DISTANCE_OFFSET*(GridSize.x+1);
+		const int DistanceIndex = DispatchThreadId.x+FLUID_DISTANCE_OFFSET*GridSize.x;
 		const int LinearIndex = DistanceIndex.x * GridSizeYZ + LinearIndexYZ;
 		GridDestinationBuffer[LinearIndex] = MAX_DISTANCE;
 	}
diff --git a/Engine/Plugins/Runtime/HairStrands/Shaders/Private/NiagaraClearVelocityGrid.usf b/Engine/Plugins/Runtime/HairStrands/Shaders/Private/NiagaraClearVelocityGrid.usf
--- a/Engine/Plugins/Runtime/HairStrands/Shaders/Private/NiagaraClearVelocityGrid.usf
+++ b/Engine/Plugins/Runtime/HairStrands/Shaders/Private/NiagaraClearVelocityGrid.usf
@@ -10,7 +10,7 @@ void MainCS(uint3 DispatchThreadId : SV_DispatchThreadID)
 {
 	if (all(DispatchThreadId < GridSize))
 	{
-		const int GridIndex = DispatchThreadId.x * (GridSize.y+1) * (GridSize.z+1) + DispatchThreadId.y * (GridSize.z+1) + DispatchThreadId.z;
+		const int GridIndex = DispatchThreadId.x * GridSize.y * GridSize.z + DispatchThreadId.y * GridSize.z + DispatchThreadId.z;
 		GridDestinationBuffer[GridIndex] = 0;
 	}
 }
\ No newline at end of file
diff --git a/Engine/Plugins/Runtime/HairStrands/Shaders/Private/NiagaraCopyVelocityGrid.usf b/Engine/Plugins/Runtime/HairStrands/Shaders/Private/NiagaraCopyVelocityGrid.usf
--- a/Engine/Plugins/Runtime/HairStrands/Shaders/Private/NiagaraCopyVelocityGrid.usf
+++ b/Engine/Plugins/Runtime/HairStrands/Shaders/Private/NiagaraCopyVelocityGrid.usf
@@ -11,7 +11,7 @@ void MainCS(uint3 DispatchThreadId : SV_DispatchThreadID)
 {
 	if (all(DispatchThreadId < GridSize))
 	{
-		const int GridIndex = DispatchThreadId.x * (GridSize.y+1) * (GridSize.z+1) + DispatchThreadId.y * (GridSize.z+1) + DispatchThreadId.z;
+		const int GridIndex = DispatchThreadId.x * GridSize.y * GridSize.z + DispatchThreadId.y * GridSize.z + DispatchThreadId.z;
 		GridDestinationBuffer[GridIndex] = GridCurrentBuffer[GridIndex];
 	}
 }
\ No newline at end of file
diff --git a/Engine/Plugins/Runtime/HairStrands/Shaders/Private/NiagaraDataInterfaceHairStrands.ush b/Engine/Plugins/Runtime/HairStrands/Shaders/Private/NiagaraDataInterfaceHairStrands.ush
--- a/Engine/Plugins/Runtime/HairStrands/Shaders/Private/NiagaraDataInterfaceHairStrands.ush
+++ b/Engine/Plugins/Runtime/HairStrands/Shaders/Private/NiagaraDataInterfaceHairStrands.ush
@@ -909,7 +909,7 @@ float3 DIHairStrands_ApplyRBF(in FDIHairStrandsContext DIContext, in float3 Rest
 	ControlPoint += DIContext.MeshSampleWeightsBuffer[DIContext.SampleCount + 1].xyz * RestControlPoint.x;
 	ControlPoint += DIContext.MeshSampleWeightsBuffer[DIContext.SampleCount + 2].xyz * RestControlPoint.y;
 	ControlPoint += DIContext.MeshSampleWeightsBuffer[DIContext.SampleCount + 3].xyz * RestControlPoint.z;
-	return ControlPoint + DIContext.MeshSampleWeightsBuffer[DIContext.SampleCount + 4].xyz; // offset added to workaround numerical issues when deformed and rest positions are far apart
+	return ControlPoint;
 }
 
  // Compute the global position
diff --git a/Engine/Plugins/Runtime/HairStrands/Source/HairStrandsCore/Private/GroomComponent.cpp b/Engine/Plugins/Runtime/HairStrands/Source/HairStrandsCore/Private/GroomComponent.cpp
--- a/Engine/Plugins/Runtime/HairStrands/Source/HairStrandsCore/Private/GroomComponent.cpp
+++ b/Engine/Plugins/Runtime/HairStrands/Source/HairStrandsCore/Private/GroomComponent.cpp
@@ -1745,6 +1745,19 @@ void UGroomComponent::SetPhysicsAsset(UPhysicsAsset* InPhysicsAsset)
 	}
 }
 
+void UGroomComponent::AddCollisionComponent(USkeletalMeshComponent* SkeletalMeshComponent)
+{
+	if(SkeletalMeshComponent && (CollisionComponents.Find(SkeletalMeshComponent) == INDEX_NONE))
+	{
+		CollisionComponents.Add(SkeletalMeshComponent);
+	}
+}
+
+void UGroomComponent::ResetCollisionComponents()
+{
+	CollisionComponents.Reset();
+}
+
 void UGroomComponent::SetEnableSimulation(bool bInEnableSimulation)
 {
 	if (SimulationSettings.bOverrideSettings)
diff --git a/Engine/Plugins/Runtime/HairStrands/Source/HairStrandsCore/Private/GroomRBFDeformer.cpp b/Engine/Plugins/Runtime/HairStrands/Source/HairStrandsCore/Private/GroomRBFDeformer.cpp
--- a/Engine/Plugins/Runtime/HairStrands/Source/HairStrandsCore/Private/GroomRBFDeformer.cpp
+++ b/Engine/Plugins/Runtime/HairStrands/Source/HairStrandsCore/Private/GroomRBFDeformer.cpp
@@ -39,8 +39,7 @@ void UpdateMeshSamples(
 	TArray<FVector3f>& OutSampleDeformationsBuffer
 )
 {
-	OutSampleDeformationsBuffer.SetNum(MaxSampleCount + 5);
-	OutSampleDeformationsBuffer[MaxSampleCount + 4] = SampleDeformedPositionsBuffer[0]; // offset used to workaround numerical issues when deformed and rest positions are far apart
+	OutSampleDeformationsBuffer.SetNum(MaxSampleCount + 4);
 
 	for (uint32 SampleIndex = 0; SampleIndex < MaxSampleCount + 4; ++SampleIndex)
 	{
@@ -49,7 +48,7 @@ void UpdateMeshSamples(
 		for (uint32 i = 0; i < MaxSampleCount; ++i, ++WeightsOffset)
 		{
 			SampleDeformation += InterpolationWeightsBuffer[WeightsOffset] *
-				(SampleDeformedPositionsBuffer[i] - (SampleRestPositionsBuffer[i] + SampleDeformedPositionsBuffer[0]) ); // offset added to workaround numerical issues when deformed and rest positions are far apart
+				(SampleDeformedPositionsBuffer[i] - SampleRestPositionsBuffer[i]); 
 		}
 
 		OutSampleDeformationsBuffer[SampleIndex] = SampleDeformation;
@@ -77,7 +76,7 @@ FVector3f DisplacePosition(
 	ControlPoint += MeshSampleWeightsBuffer[SampleCount + 1] * RestControlPoint.X;
 	ControlPoint += MeshSampleWeightsBuffer[SampleCount + 2] * RestControlPoint.Y;
 	ControlPoint += MeshSampleWeightsBuffer[SampleCount + 3] * RestControlPoint.Z;
-	return ControlPoint + MeshSampleWeightsBuffer[SampleCount + 4]; // offset added to workaround numerical issues when deformed and rest positions are far apart
+	return ControlPoint;
 }
 
 void DeformStrands(	
diff --git a/Engine/Plugins/Runtime/HairStrands/Source/HairStrandsCore/Private/GroomResources.cpp b/Engine/Plugins/Runtime/HairStrands/Source/HairStrandsCore/Private/GroomResources.cpp
--- a/Engine/Plugins/Runtime/HairStrands/Source/HairStrandsCore/Private/GroomResources.cpp
+++ b/Engine/Plugins/Runtime/HairStrands/Source/HairStrandsCore/Private/GroomResources.cpp
@@ -1250,7 +1250,7 @@ void FHairStrandsDeformedRootResource::InternalAllocateLOD(FRDGBuilder& GraphBui
 			if (LOD.SampleCount > 0)
 			{
 				InternalCreateVertexBufferRDG<FHairStrandsMeshTrianglePositionFormat>(GraphBuilder, LOD.SampleCount, LOD.DeformedSamplePositionsBuffer, ToHairResourceDebugName(HAIRSTRANDS_RESOUCE_NAME(CurveType, Hair.StrandsRootDeformed_DeformedSamplePositionsBuffer), ResourceName), EHairResourceUsageType::Dynamic);
-				InternalCreateVertexBufferRDG<FHairStrandsMeshTrianglePositionFormat>(GraphBuilder, LOD.SampleCount + 5, LOD.MeshSampleWeightsBuffer, ToHairResourceDebugName(HAIRSTRANDS_RESOUCE_NAME(CurveType, Hair.StrandsRootDeformed_MeshSampleWeightsBuffer), ResourceName), EHairResourceUsageType::Dynamic); // offset added to workaround numerical issues when deformed and rest positions are far apart
+				InternalCreateVertexBufferRDG<FHairStrandsMeshTrianglePositionFormat>(GraphBuilder, LOD.SampleCount + 4, LOD.MeshSampleWeightsBuffer, ToHairResourceDebugName(HAIRSTRANDS_RESOUCE_NAME(CurveType, Hair.StrandsRootDeformed_MeshSampleWeightsBuffer), ResourceName), EHairResourceUsageType::Dynamic); // offset added to workaround numerical issues when deformed and rest positions are far apart
 			}
 
 			InternalCreateVertexBufferRDG<FHairStrandsMeshTrianglePositionFormat>(GraphBuilder, RootCount, LOD.DeformedRootTrianglePosition0Buffer, ToHairResourceDebugName(HAIRSTRANDS_RESOUCE_NAME(CurveType, Hair.StrandsRootDeformed_DeformedRootTrianglePosition0Buffer), ResourceName), EHairResourceUsageType::Dynamic);
diff --git a/Engine/Plugins/Runtime/HairStrands/Source/HairStrandsCore/Private/Niagara/NiagaraDataInterfaceHairStrands.cpp b/Engine/Plugins/Runtime/HairStrands/Source/HairStrandsCore/Private/Niagara/NiagaraDataInterfaceHairStrands.cpp
--- a/Engine/Plugins/Runtime/HairStrands/Source/HairStrandsCore/Private/Niagara/NiagaraDataInterfaceHairStrands.cpp
+++ b/Engine/Plugins/Runtime/HairStrands/Source/HairStrandsCore/Private/Niagara/NiagaraDataInterfaceHairStrands.cpp
@@ -759,7 +759,7 @@ struct FNDIHairStrandsParametersCS : public FNiagaraDataInterfaceParametersCS
 
 			// Simulation setup
 			const int32 NeedResetValue = (ProxyData->TickCount <= GHairSimulationMaxDelay) || !HairStrandsBuffer->bValidGeometryType;
-			const int32 RestUpdateValue = NeedResetValue || GHairSimulationRestUpdate;
+			const int32 RestUpdateValue = GHairSimulationRestUpdate;
 			const int32 LocalSimulationValue = ProxyData->LocalSimulation;
 
 			HairStrandsBuffer->bValidGeometryType = true;
@@ -1188,6 +1188,7 @@ bool UNiagaraDataInterfaceHairStrands::PerInstanceTick(void* PerInstanceData, FN
 	if (SourceComponent != nullptr)
 	{
 		if (SourceComponent->bResetSimulation || RequiresSimulationReset(SystemInstance, InstanceData->SkeletalMeshes))
+			
 		{
 			InstanceData->TickCount = 0;
 		}
diff --git a/Engine/Plugins/Runtime/HairStrands/Source/HairStrandsCore/Private/Niagara/NiagaraDataInterfacePhysicsAsset.cpp b/Engine/Plugins/Runtime/HairStrands/Source/HairStrandsCore/Private/Niagara/NiagaraDataInterfacePhysicsAsset.cpp
--- a/Engine/Plugins/Runtime/HairStrands/Source/HairStrandsCore/Private/Niagara/NiagaraDataInterfacePhysicsAsset.cpp
+++ b/Engine/Plugins/Runtime/HairStrands/Source/HairStrandsCore/Private/Niagara/NiagaraDataInterfacePhysicsAsset.cpp
@@ -314,7 +314,7 @@ void CreateInternalArrays(const TArray<TWeakObjectPtr<UPhysicsAsset>>& PhysicsAs
 		}
 		else
 		{
-			UE_LOG(LogPhysicsAsset, Warning, TEXT("Number of hysics asset primitives is higher than the niagara %d limit"), PHYSICS_ASSET_MAX_PRIMITIVES);
+			UE_LOG(LogPhysicsAsset, Warning, TEXT("Number of physics asset primitives is higher than the niagara %d limit"), PHYSICS_ASSET_MAX_PRIMITIVES);
 		}
 	}
 }
@@ -330,73 +330,76 @@ void UpdateInternalArrays(const TArray<TWeakObjectPtr<UPhysicsAsset>>& PhysicsAs
 
 		GetNumPrimitives(PhysicsAssets, SkeletalMeshs, NumBoxes, NumSpheres, NumCapsules);
 
-		if (((OutAssetArrays->ElementOffsets.SphereOffset - OutAssetArrays->ElementOffsets.BoxOffset) != NumBoxes) || 
-			((OutAssetArrays->ElementOffsets.CapsuleOffset - OutAssetArrays->ElementOffsets.SphereOffset) != NumSpheres) ||
-			((OutAssetArrays->ElementOffsets.NumElements - OutAssetArrays->ElementOffsets.CapsuleOffset) != NumCapsules))
+		if ((NumBoxes + NumSpheres + NumCapsules) < PHYSICS_ASSET_MAX_PRIMITIVES)
 		{
-			CreateInternalArrays(PhysicsAssets, SkeletalMeshs, OutAssetArrays, InWorldTransform, LocalTransform);
-		}
-
-		OutAssetArrays->PreviousTransform = OutAssetArrays->CurrentTransform;
-		OutAssetArrays->PreviousInverse = OutAssetArrays->CurrentInverse;
+			if (((OutAssetArrays->ElementOffsets.SphereOffset - OutAssetArrays->ElementOffsets.BoxOffset) != NumBoxes) || 
+				((OutAssetArrays->ElementOffsets.CapsuleOffset - OutAssetArrays->ElementOffsets.SphereOffset) != NumSpheres) ||
+				((OutAssetArrays->ElementOffsets.NumElements - OutAssetArrays->ElementOffsets.CapsuleOffset) != NumCapsules))
+			{
+				CreateInternalArrays(PhysicsAssets, SkeletalMeshs, OutAssetArrays, InWorldTransform, LocalTransform);
+			}
 
-		uint32 BoxCount = OutAssetArrays->ElementOffsets.BoxOffset;
-		uint32 SphereCount = OutAssetArrays->ElementOffsets.SphereOffset;
-		uint32 CapsuleCount = OutAssetArrays->ElementOffsets.CapsuleOffset;
+			OutAssetArrays->PreviousTransform = OutAssetArrays->CurrentTransform;
+			OutAssetArrays->PreviousInverse = OutAssetArrays->CurrentInverse;
 
-		for (int32 ComponentIndex = 0; ComponentIndex < SkeletalMeshs.Num(); ++ComponentIndex)
-		{
-			TWeakObjectPtr<USkeletalMeshComponent> SkeletalMesh = SkeletalMeshs[ComponentIndex];
-			const bool IsSkelMeshValid = SkeletalMesh.IsValid() && SkeletalMesh.Get() != nullptr;
-			const FTransform WorldTransform = IsSkelMeshValid ? SkeletalMesh->GetComponentTransform() : InWorldTransform;
+			uint32 BoxCount = OutAssetArrays->ElementOffsets.BoxOffset;
+			uint32 SphereCount = OutAssetArrays->ElementOffsets.SphereOffset;
+			uint32 CapsuleCount = OutAssetArrays->ElementOffsets.CapsuleOffset;
 
-			TWeakObjectPtr<UPhysicsAsset> PhysicsAsset = PhysicsAssets[ComponentIndex];
-			if (PhysicsAsset.IsValid() && PhysicsAsset.Get() != nullptr)
+			for (int32 ComponentIndex = 0; ComponentIndex < SkeletalMeshs.Num(); ++ComponentIndex)
 			{
-				USkeletalMesh* SkelMesh = (SkeletalMeshs[ComponentIndex].Get() && SkeletalMeshs[ComponentIndex]->SkeletalMesh) ? ToRawPtr(SkeletalMeshs[ComponentIndex]->SkeletalMesh) : PhysicsAsset->GetPreviewMesh();
-				if (!SkelMesh)
-				{
-					continue;
-				}
-				const FReferenceSkeleton& RefSkeleton = SkelMesh->GetRefSkeleton();
+				TWeakObjectPtr<USkeletalMeshComponent> SkeletalMesh = SkeletalMeshs[ComponentIndex];
+				const bool IsSkelMeshValid = SkeletalMesh.IsValid() && SkeletalMesh.Get() != nullptr;
+				const FTransform WorldTransform = IsSkelMeshValid ? SkeletalMesh->GetComponentTransform() : InWorldTransform;
 
+				TWeakObjectPtr<UPhysicsAsset> PhysicsAsset = PhysicsAssets[ComponentIndex];
+				if (PhysicsAsset.IsValid() && PhysicsAsset.Get() != nullptr)
 				{
-					TArray<FTransform> RestTransforms;
-					FAnimationRuntime::FillUpComponentSpaceTransforms(RefSkeleton, RefSkeleton.GetRefBonePose(), RestTransforms);
-					
-					for (const UBodySetup* BodySetup : PhysicsAsset->SkeletalBodySetups)
+					USkeletalMesh* SkelMesh = (SkeletalMeshs[ComponentIndex].Get() && SkeletalMeshs[ComponentIndex]->SkeletalMesh) ? ToRawPtr(SkeletalMeshs[ComponentIndex]->SkeletalMesh) : PhysicsAsset->GetPreviewMesh();
+					if (!SkelMesh)
 					{
-						const FName BoneName = BodySetup->BoneName;
-						const int32 BoneIndex = RefSkeleton.FindBoneIndex(BoneName);
-						if (BoneIndex != INDEX_NONE && BoneIndex < RestTransforms.Num())
-						{
-							const FTransform BoneTransform = IsSkelMeshValid ? ComputeRelativeTransform(SkeletalMesh->GetBoneTransform(BoneIndex), LocalTransform) : 
-								RestTransforms[BoneIndex] * ComputeRelativeTransform(WorldTransform, LocalTransform);
+						continue;
+					}
+					const FReferenceSkeleton& RefSkeleton = SkelMesh->GetRefSkeleton();
 
-							for (const FKBoxElem& BoxElem : BodySetup->AggGeom.BoxElems)
+					{
+						TArray<FTransform> RestTransforms;
+						FAnimationRuntime::FillUpComponentSpaceTransforms(RefSkeleton, RefSkeleton.GetRefBonePose(), RestTransforms);
+						
+						for (const UBodySetup* BodySetup : PhysicsAsset->SkeletalBodySetups)
+						{
+							const FName BoneName = BodySetup->BoneName;
+							const int32 BoneIndex = RefSkeleton.FindBoneIndex(BoneName);
+							if (BoneIndex != INDEX_NONE && BoneIndex < RestTransforms.Num())
 							{
-								if (CollisionEnabledHasPhysics(BoxElem.GetCollisionEnabled()))
+								const FTransform BoneTransform = IsSkelMeshValid ? ComputeRelativeTransform(SkeletalMesh->GetBoneTransform(BoneIndex), LocalTransform) : 
+									RestTransforms[BoneIndex] * ComputeRelativeTransform(WorldTransform, LocalTransform);
+
+								for (const FKBoxElem& BoxElem : BodySetup->AggGeom.BoxElems)
 								{
-									const FTransform ElementTransform = FTransform(BoxElem.Rotation, BoxElem.Center) * BoneTransform;
-									FillCurrentTransforms(ElementTransform, BoxCount, OutAssetArrays->CurrentTransform, OutAssetArrays->CurrentInverse);
+									if (CollisionEnabledHasPhysics(BoxElem.GetCollisionEnabled()))
+									{
+										const FTransform ElementTransform = FTransform(BoxElem.Rotation, BoxElem.Center) * BoneTransform;
+										FillCurrentTransforms(ElementTransform, BoxCount, OutAssetArrays->CurrentTransform, OutAssetArrays->CurrentInverse);
+									}
 								}
-							}
 
-							for (const FKSphereElem& SphereElem : BodySetup->AggGeom.SphereElems)
-							{
-								if (CollisionEnabledHasPhysics(SphereElem.GetCollisionEnabled()))
+								for (const FKSphereElem& SphereElem : BodySetup->AggGeom.SphereElems)
 								{
-									const FTransform ElementTransform = FTransform(SphereElem.Center) * BoneTransform;
-									FillCurrentTransforms(ElementTransform, SphereCount, OutAssetArrays->CurrentTransform, OutAssetArrays->CurrentInverse);
+									if (CollisionEnabledHasPhysics(SphereElem.GetCollisionEnabled()))
+									{
+										const FTransform ElementTransform = FTransform(SphereElem.Center) * BoneTransform;
+										FillCurrentTransforms(ElementTransform, SphereCount, OutAssetArrays->CurrentTransform, OutAssetArrays->CurrentInverse);
+									}
 								}
-							}
 
-							for (const FKSphylElem& CapsuleElem : BodySetup->AggGeom.SphylElems)
-							{
-								if (CollisionEnabledHasPhysics(CapsuleElem.GetCollisionEnabled()))
+								for (const FKSphylElem& CapsuleElem : BodySetup->AggGeom.SphylElems)
 								{
-									const FTransform ElementTransform = FTransform(CapsuleElem.Rotation, CapsuleElem.Center) * BoneTransform;
-									FillCurrentTransforms(ElementTransform, CapsuleCount, OutAssetArrays->CurrentTransform, OutAssetArrays->CurrentInverse);
+									if (CollisionEnabledHasPhysics(CapsuleElem.GetCollisionEnabled()))
+									{
+										const FTransform ElementTransform = FTransform(CapsuleElem.Rotation, CapsuleElem.Center) * BoneTransform;
+										FillCurrentTransforms(ElementTransform, CapsuleCount, OutAssetArrays->CurrentTransform, OutAssetArrays->CurrentInverse);
+									}
 								}
 							}
 						}
@@ -818,15 +821,12 @@ void UNiagaraDataInterfacePhysicsAsset::ExtractSourceComponent(FNiagaraSystemIns
 			{
 				GroomPhysicsAsset = GroomComponent->PhysicsAsset;
 			}
-			TArray<UPrimitiveComponent*> OverlappingComponents;
-			GroomComponent->GetOverlappingComponents(OverlappingComponents);
-
-			for(auto& OverlappingComponent : OverlappingComponents)
+			for(auto& CollisionComponent : GroomComponent->CollisionComponents)
 			{
-				if (USkeletalMeshComponent* SkelMeshComp = Cast<USkeletalMeshComponent>(OverlappingComponent))
+				if(CollisionComponent.IsValid())
 				{
-					SourceComponents.Add(SkelMeshComp);
-					PhysicsAssets.Add(SkelMeshComp->GetPhysicsAsset());
+					SourceComponents.Add(CollisionComponent);
+					PhysicsAssets.Add(CollisionComponent->GetPhysicsAsset());
 				}
 			}
 			break;
diff --git a/Engine/Plugins/Runtime/HairStrands/Source/HairStrandsCore/Private/Niagara/NiagaraDataInterfacePressureGrid.cpp b/Engine/Plugins/Runtime/HairStrands/Source/HairStrandsCore/Private/Niagara/NiagaraDataInterfacePressureGrid.cpp
--- a/Engine/Plugins/Runtime/HairStrands/Source/HairStrandsCore/Private/Niagara/NiagaraDataInterfacePressureGrid.cpp
+++ b/Engine/Plugins/Runtime/HairStrands/Source/HairStrandsCore/Private/Niagara/NiagaraDataInterfacePressureGrid.cpp
@@ -518,13 +518,13 @@ inline void ClearBuffer(FRHICommandList& RHICmdList, ERHIFeatureLevel::Type Feat
 		RHICmdList.Transition(MakeArrayView(Transitions, UE_ARRAY_COUNT(Transitions)));
 
 		const uint32 GroupSize = NIAGARA_HAIR_STRANDS_THREAD_COUNT_PRESSURE;
-		const uint32 NumElements = (GridSize.X + 1) * (GridSize.Y + 1) * (GridSize.Z + 1);
+		const FIntVector ExtendedSize = GridSize + FIntVector(1, 1, 1);
 		
-		const uint32 DispatchCountX = FMath::DivideAndRoundUp((uint32)(GridSize.X+1), GroupSize);
-		const uint32 DispatchCountY = FMath::DivideAndRoundUp((uint32)(GridSize.Y+1), GroupSize);
-		const uint32 DispatchCountZ = FMath::DivideAndRoundUp((uint32)(GridSize.Z+1), GroupSize);
+		const uint32 DispatchCountX = FMath::DivideAndRoundUp((uint32)(ExtendedSize.X), GroupSize);
+		const uint32 DispatchCountY = FMath::DivideAndRoundUp((uint32)(ExtendedSize.Y), GroupSize);
+		const uint32 DispatchCountZ = FMath::DivideAndRoundUp((uint32)(ExtendedSize.Z), GroupSize);
 
-		ComputeShader->SetParameters(RHICmdList, CurrentGridBufferSRV, DestinationGridBufferUAV, GridSize, CopyPressure);
+		ComputeShader->SetParameters(RHICmdList, CurrentGridBufferSRV, DestinationGridBufferUAV, ExtendedSize, CopyPressure);
 		DispatchComputeShader(RHICmdList, ComputeShader.GetShader(), DispatchCountX, DispatchCountY, DispatchCountZ);
 		ComputeShader->UnsetParameters(RHICmdList);
 	}
diff --git a/Engine/Plugins/Runtime/HairStrands/Source/HairStrandsCore/Public/GroomComponent.h b/Engine/Plugins/Runtime/HairStrands/Source/HairStrandsCore/Public/GroomComponent.h
--- a/Engine/Plugins/Runtime/HairStrands/Source/HairStrandsCore/Public/GroomComponent.h
+++ b/Engine/Plugins/Runtime/HairStrands/Source/HairStrandsCore/Public/GroomComponent.h
@@ -47,6 +47,9 @@ public:
 	UPROPERTY(EditAnywhere, BlueprintReadWrite, Category = "Simulation")
 	TObjectPtr<class UPhysicsAsset> PhysicsAsset;
 
+	/** List of collision components to be used */
+	TArray<TWeakObjectPtr<class USkeletalMeshComponent>> CollisionComponents;
+	
 	/** Groom's simulation settings */
 	UPROPERTY(EditAnywhere, BlueprintReadWrite, interp, Category = "Simulation")
 	FHairSimulationSettings SimulationSettings;
@@ -169,6 +172,14 @@ public:
 	UFUNCTION(BlueprintCallable, Category = "Simulation")
 	void SetPhysicsAsset(UPhysicsAsset* InPhysicsAsset);
 
+	/* Add a skeletal mesh to the collision components */
+	UFUNCTION(BlueprintCallable, Category = "Simulation")
+	void AddCollisionComponent(USkeletalMeshComponent* SkeletalMeshComponent);
+	
+	/* Reset the collision components */
+	UFUNCTION(BlueprintCallable, Category = "Simulation")
+	void ResetCollisionComponents();
+
 	/* Accessor function for changing the enable simulation flag from blueprint/sequencer */
 	UFUNCTION(BlueprintCallable, Category = "Simulation")
 	void SetEnableSimulation(bool bInEnableSimulation);
diff --git a/Engine/Shaders/Private/HairStrands/HairStrandsMeshRBF.ush b/Engine/Shaders/Private/HairStrands/HairStrandsMeshRBF.ush
--- a/Engine/Shaders/Private/HairStrands/HairStrandsMeshRBF.ush
+++ b/Engine/Shaders/Private/HairStrands/HairStrandsMeshRBF.ush
@@ -16,9 +16,9 @@ float3 ApplyRBF(
 		const float FunctionValue = sqrt(dot(PositionDelta, PositionDelta) + 1);
 		ControlPoint += FunctionValue * InMeshSampleWeightsBuffer[i].xyz;
 	}
-	ControlPoint += InMeshSampleWeightsBuffer[InMaxSampleCount].xyz;
-	ControlPoint += InMeshSampleWeightsBuffer[InMaxSampleCount + 1].xyz * RestControlPoint.x;
-	ControlPoint += InMeshSampleWeightsBuffer[InMaxSampleCount + 2].xyz * RestControlPoint.y;
-	ControlPoint += InMeshSampleWeightsBuffer[InMaxSampleCount + 3].xyz * RestControlPoint.z;
-	return ControlPoint + InMeshSampleWeightsBuffer[InMaxSampleCount + 4].xyz; // offset added to workaround numerical issues when deformed and rest positions are far apart
+	 ControlPoint += InMeshSampleWeightsBuffer[InMaxSampleCount].xyz;
+	 ControlPoint += InMeshSampleWeightsBuffer[InMaxSampleCount + 1].xyz * RestControlPoint.x;
+	 ControlPoint += InMeshSampleWeightsBuffer[InMaxSampleCount + 2].xyz * RestControlPoint.y;
+	 ControlPoint += InMeshSampleWeightsBuffer[InMaxSampleCount + 3].xyz * RestControlPoint.z;
+	return ControlPoint ;
 }
\ No newline at end of file
diff --git a/Engine/Shaders/Private/HairStrands/HairStrandsSamplesUpdate.usf b/Engine/Shaders/Private/HairStrands/HairStrandsSamplesUpdate.usf
--- a/Engine/Shaders/Private/HairStrands/HairStrandsSamplesUpdate.usf
+++ b/Engine/Shaders/Private/HairStrands/HairStrandsSamplesUpdate.usf
@@ -17,10 +17,6 @@ void MainCS(uint3 DispatchThreadId : SV_DispatchThreadID)
 	const uint SampleIndex = DispatchThreadId.x;
 	if (SampleIndex >= (MaxSampleCount + 4))
 	{
-		if (SampleIndex == (MaxSampleCount + 4))
-		{
-			OutSampleDeformationsBuffer[SampleIndex] = float4(SampleDeformedPositionsBuffer[0].xyz, 1); // offset used to workaround numerical issues when deformed and rest positions are far apart
-		}
 		return;
 	}
 
@@ -29,7 +25,7 @@ void MainCS(uint3 DispatchThreadId : SV_DispatchThreadID)
 	for(uint i = 0; i < MaxSampleCount; ++i, ++WeightsOffset)
 	{
 		SampleDeformation += InterpolationWeightsBuffer[WeightsOffset] * 
-				(SampleDeformedPositionsBuffer[i] - (SampleRestPositionsBuffer[i] + SampleDeformedPositionsBuffer[0])).xyz; // offset added to workaround numerical issues when deformed and rest positions are far apart
+				(SampleDeformedPositionsBuffer[i] - SampleRestPositionsBuffer[i]).xyz;
 	}
 
 	OutSampleDeformationsBuffer[SampleIndex] = float4(SampleDeformation,1);
