Enable teleporting even when local simulation is on

diff --git a/Engine/Plugins/Runtime/HairStrands/Source/HairStrandsCore/Private/GroomComponent.cpp b/Engine/Plugins/Runtime/HairStrands/Source/HairStrandsCore/Private/GroomComponent.cpp
--- a/Engine/Plugins/Runtime/HairStrands/Source/HairStrandsCore/Private/GroomComponent.cpp
+++ b/Engine/Plugins/Runtime/HairStrands/Source/HairStrandsCore/Private/GroomComponent.cpp
@@ -3251,7 +3251,10 @@ void UGroomComponent::TickComponent(float DeltaTime, enum ELevelTick TickType, F
 	bResetSimulation = bInitSimulation;
 	if (!bInitSimulation)
 	{
-		if (!SimulationSettings.SimulationSetup.bLocalSimulation)
+		const bool bLocalSimulation = SimulationSettings.SimulationSetup.bLocalSimulation &&
+			(SimulationSettings.SimulationSetup.LinearVelocityScale == 1.0) &&
+			(SimulationSettings.SimulationSetup.AngularVelocityScale == 1.0);
+		if (!bLocalSimulation)
 		{
 			if (USkeletalMeshComponent* ParentComp = Cast<USkeletalMeshComponent>(GetAttachParent()))
 			{
diff --git a/Engine/Plugins/Runtime/HairStrands/Source/HairStrandsCore/Private/Niagara/NiagaraDataInterfaceHairStrands.cpp b/Engine/Plugins/Runtime/HairStrands/Source/HairStrandsCore/Private/Niagara/NiagaraDataInterfaceHairStrands.cpp
--- a/Engine/Plugins/Runtime/HairStrands/Source/HairStrandsCore/Private/Niagara/NiagaraDataInterfaceHairStrands.cpp
+++ b/Engine/Plugins/Runtime/HairStrands/Source/HairStrandsCore/Private/Niagara/NiagaraDataInterfaceHairStrands.cpp
@@ -531,7 +531,7 @@ void FNDIHairStrandsData::Update(UNiagaraDataInterfaceHairStrands* Interface, FN
 				FMatrix44d BoneTransformDouble = BoneTransform.ToMatrixWithScale();
 				const FMatrix44d WorldTransformDouble = WorldTransform.ToMatrixWithScale();
 
-				if(DeltaSeconds != 0.0f)
+				if(DeltaSeconds != 0.0f && (TickCount > GHairSimulationMaxDelay))
 				{
 					const FMatrix44d PreviousBoneTransformDouble = PreviousBoneTransform.ToMatrixWithScale();
 					const FMatrix44d DeltaTransformDouble =  BoneTransformDouble * PreviousBoneTransformDouble.Inverse();
@@ -800,7 +800,7 @@ struct FNDIHairStrandsParametersCS : public FNiagaraDataInterfaceParametersCS
 			SetShaderValue(RHICmdList, ComputeShaderRHI, BoundingBoxOffsets, HairStrandsBuffer->BoundingBoxOffsets);
 			SetShaderValue(RHICmdList, ComputeShaderRHI, WorldTransform, WorldTransformFloat);
 			SetShaderValue(RHICmdList, ComputeShaderRHI, WorldInverse, WorldTransformFloat.Inverse());
-			SetShaderValue(RHICmdList, ComputeShaderRHI, WorldRotation, WorldTransformFloat.ToQuat());
+			SetShaderValue(RHICmdList, ComputeShaderRHI, WorldRotation, WorldTransformFloat.GetMatrixWithoutScale().ToQuat());
 			SetShaderValue(RHICmdList, ComputeShaderRHI, NumStrands, ProxyData->NumStrands);
 			SetShaderValue(RHICmdList, ComputeShaderRHI, StrandSize, ProxyData->StrandsSize);
 			SetShaderValue(RHICmdList, ComputeShaderRHI, BoneTransform, BoneTransformFloat);
