Fix duplicated Localization identifier in Groom UI.

diff --git a/Engine/Plugins/Runtime/HairStrands/Source/HairStrandsEditor/Private/GroomAssetDetails.cpp b/Engine/Plugins/Runtime/HairStrands/Source/HairStrandsEditor/Private/GroomAssetDetails.cpp
--- a/Engine/Plugins/Runtime/HairStrands/Source/HairStrandsEditor/Private/GroomAssetDetails.cpp
+++ b/Engine/Plugins/Runtime/HairStrands/Source/HairStrandsEditor/Private/GroomAssetDetails.cpp
@@ -1287,14 +1287,14 @@ TSharedRef<SWidget> FGroomRenderingDetails::MakeGroupNameCustomization(int32 Gro
 			return SNew(STextBlock)
 				.Font(IDetailLayoutBuilder::GetDetailFont())
 				.ColorAndOpacity(GroupColor)
-				.Text(FText::Format(LOCTEXT("Group", "Group ID {0} - {1}"), FText::AsNumber(GroupIndex), FText::FromName(GroupName)));
+				.Text(FText::Format(LOCTEXT("GroupWithName", "Group ID {0} - {1}"), FText::AsNumber(GroupIndex), FText::FromName(GroupName)));
 		}
 		else
 		{
 			return SNew(STextBlock)
 				.Font(IDetailLayoutBuilder::GetDetailFont())
 				.ColorAndOpacity(GroupColor)
-				.Text(FText::Format(LOCTEXT("Group", "Group ID {0}"), FText::AsNumber(GroupIndex)));
+				.Text(FText::Format(LOCTEXT("GroupWithoutName", "Group ID {0}"), FText::AsNumber(GroupIndex)));
 		}
 	}
 	break;
