Fix crash when building follicule texture using the GPU path.

The error arises due to access of game thread data on the rendering thread. This CL ensures that the rendering thread only access valid data by transfering data ownership from game thread to rendering thread.

diff --git a/Engine/Plugins/Runtime/HairStrands/Source/HairStrandsCore/Private/GroomTextureBuilder.cpp b/Engine/Plugins/Runtime/HairStrands/Source/HairStrandsCore/Private/GroomTextureBuilder.cpp
--- a/Engine/Plugins/Runtime/HairStrands/Source/HairStrandsCore/Private/GroomTextureBuilder.cpp
+++ b/Engine/Plugins/Runtime/HairStrands/Source/HairStrandsCore/Private/GroomTextureBuilder.cpp
@@ -275,11 +275,19 @@ static void InternalBuildFollicleTexture_CPU(const TArray<FFollicleInfo>& InInfo
 #endif // #if WITH_EDITORONLY_DATA
 }
 
+struct FFollicleInfoGPU
+{
+	FHairStrandsDatas*		StrandsData = nullptr;
+	FFollicleInfo::EChannel	Channel = FFollicleInfo::R;
+	uint32					KernelSizeInPixels = 0;
+	bool					bGPUOnly = false;
+};
+
 // GPU path
 static void InternalBuildFollicleTexture_GPU(
 	FRDGBuilder& GraphBuilder,
 	FGlobalShaderMap* ShaderMap,
-	const TArray<FFollicleInfo>& InInfos,
+	const TArray<FFollicleInfoGPU>& InInfos,
 	const FIntPoint Resolution,
 	const uint32 MipCount,
 	const EPixelFormat InFormat,
@@ -289,14 +297,8 @@ static void InternalBuildFollicleTexture_GPU(
 	TArray<FRDGBufferRef> RootUVBuffers[4];
 	bool bCopyDataBackToCPU = false;
 
-	for (const FFollicleInfo& Info : InInfos)
+	for (const FFollicleInfoGPU& Info : InInfos)
 	{
-		if (!Info.GroomAsset || Info.GroomAsset->GetNumHairGroups() == 0)
-		{
-			UE_LOG(LogHairStrands, Warning, TEXT("[Groom] Error - Groom follicle texture can be entirely created/rebuilt as some groom assets seems invalid."));
-			continue;
-		}
-
 		if (KernelSizeInPixels == ~0)
 		{
 			KernelSizeInPixels = Info.KernelSizeInPixels;
@@ -304,16 +306,9 @@ static void InternalBuildFollicleTexture_GPU(
 		}
 
 		// Create root UVs buffer
-		uint32 GroupIndex = 0;
-		for (const FHairGroupData& GroupData : Info.GroomAsset->HairGroupsData)
+		if (Info.StrandsData)
 		{
-			FHairStrandsDatas StrandsData;
-			FHairStrandsDatas GuidesData;
-#if WITH_EDITORONLY_DATA
-			Info.GroomAsset->GetHairStrandsDatas(GroupIndex, StrandsData, GuidesData);
-#endif
-			StrandsData.StrandsCurves.CurvesRootUV.GetData();
-			const uint32 DataCount = StrandsData.StrandsCurves.CurvesRootUV.Num();
+			const uint32 DataCount = Info.StrandsData->StrandsCurves.CurvesRootUV.Num();
 			const uint32 DataSizeInBytes = sizeof(FVector2f) * DataCount;
 			check(DataSizeInBytes != 0);
 
@@ -322,12 +317,10 @@ static void InternalBuildFollicleTexture_GPU(
 				GraphBuilder,
 				TEXT("RootUVBuffer"),
 				Desc,
-				StrandsData.StrandsCurves.CurvesRootUV.GetData(),
+				Info.StrandsData->StrandsCurves.CurvesRootUV.GetData(),
 				DataSizeInBytes,
 				ERDGInitialDataFlags::None);
 			RootUVBuffers[Info.Channel].Add(RootBuffer);
-
-			GroupIndex++;
 		}
 	}
 
@@ -339,7 +332,7 @@ static void InternalBuildFollicleTexture_GPU(
 // Asynchronous queuing for follicle texture mask generation
 struct FFollicleQuery
 {
-	TArray<FFollicleInfo> Infos;
+	TArray<FFollicleInfoGPU> Infos;
 	FIntPoint Resolution = 0;
 	uint32 MipCount = 0;
 	EPixelFormat Format = PF_R8G8B8A8;
@@ -361,13 +354,23 @@ void RunHairStrandsFolliculeMaskQueries(FRDGBuilder& GraphBuilder, FGlobalShader
 		if (Q.Infos.Num() > 0 && Q.MipCount>0 && Q.Resolution.X > 0 && Q.Resolution.Y > 0)
 		{
 			InternalBuildFollicleTexture_GPU(GraphBuilder, ShaderMap, Q.Infos, Q.Resolution, Q.MipCount, Q.Format, Q.OutTexture);
+
+			// Release all strands data
+			for (FFollicleInfoGPU& Info : Q.Infos)
+			{
+				if (Info.StrandsData)
+				{
+					delete Info.StrandsData;
+					Info.StrandsData = nullptr;
+				}
+			}
 		}
 	}
 }
 
-void FGroomTextureBuilder::BuildFollicleTexture(const TArray<FFollicleInfo>& InInfos, UTexture2D* OutTexture, bool bUseGPU)
+void FGroomTextureBuilder::BuildFollicleTexture(const TArray<FFollicleInfo>& InCPUInfos, UTexture2D* OutTexture, bool bUseGPU)
 {
-	if (!OutTexture || InInfos.Num() == 0)
+	if (!OutTexture || InCPUInfos.Num() == 0)
 	{
 		UE_LOG(LogGroomTextureBuilder, Warning, TEXT("[Groom] Error - Follicle texture can't be created/rebuilt."));
 		return;
@@ -380,13 +383,45 @@ void FGroomTextureBuilder::BuildFollicleTexture(const TArray<FFollicleInfo>& InI
 		const uint32 MipCount = OutTexture->GetNumMips();
 		if (MipCount > 0 && Resolution.X > 0 && Resolution.Y > 0)
 		{
+			TArray<FFollicleInfoGPU> GPUInfos;
+			for (const FFollicleInfo& CPUInfo : InCPUInfos)
+			{
+				if (!CPUInfo.GroomAsset || CPUInfo.GroomAsset->GetNumHairGroups() == 0)
+				{
+					UE_LOG(LogHairStrands, Warning, TEXT("[Groom] Error - Groom follicle texture can be entirely created/rebuilt as some groom assets seems invalid."));
+					continue;
+				}
+
+				for (uint32 GroupIndex = 0, GroupCount = CPUInfo.GroomAsset->GetNumHairGroups(); GroupIndex < GroupCount; ++GroupIndex)
+				{
+					FHairStrandsDatas* StrandsData = new FHairStrandsDatas();
+					FHairStrandsDatas GuidesData;
+				#if WITH_EDITORONLY_DATA
+					CPUInfo.GroomAsset->GetHairStrandsDatas(GroupIndex, *StrandsData, GuidesData);
+				#endif
+
+					if (StrandsData->StrandsCurves.CurvesRootUV.Num())
+					{
+						FFollicleInfoGPU& GPUInfo = GPUInfos.AddDefaulted_GetRef();
+						GPUInfo.Channel = CPUInfo.Channel;
+						GPUInfo.KernelSizeInPixels = CPUInfo.KernelSizeInPixels;
+						GPUInfo.StrandsData = StrandsData;
+						GPUInfo.bGPUOnly = CPUInfo.bGPUOnly;
+					}
+					else
+					{
+						delete StrandsData;
+					}
+				}
+			}
+
 			const EPixelFormat Format = OutTexture->GetPixelFormat(0);
 			ENQUEUE_RENDER_COMMAND(FFollicleTextureQuery)(
-			[Resolution, MipCount, Format, InInfos, OutTexture](FRHICommandListImmediate& RHICmdList)
+			[Resolution, MipCount, Format, GPUInfos, OutTexture](FRHICommandListImmediate& RHICmdList)
 			{
 				if (OutTexture->GetResource())
 				{
-					GFollicleQueries.Enqueue({ InInfos, Resolution, MipCount, Format, OutTexture });
+					GFollicleQueries.Enqueue({ GPUInfos, Resolution, MipCount, Format, OutTexture });
 				}
 			});
 		}
@@ -394,8 +429,7 @@ void FGroomTextureBuilder::BuildFollicleTexture(const TArray<FFollicleInfo>& InI
 	else
 	{
 		// Synchronous (CPU)
-		InternalBuildFollicleTexture_CPU(InInfos, OutTexture);
-		// maybe warning here?
+		InternalBuildFollicleTexture_CPU(InCPUInfos, OutTexture);
 	}
 }
 
