* Add RootUV, Color attributes, Roughness attriutes, Precomputed Guide Weights infos when importing groom asset.
 * Fix groom component UI to match groom asset UI.

diff --git a/Engine/Plugins/Runtime/HairStrands/Source/HairStrandsCore/Private/GroomBuilder.cpp b/Engine/Plugins/Runtime/HairStrands/Source/HairStrandsCore/Private/GroomBuilder.cpp
--- a/Engine/Plugins/Runtime/HairStrands/Source/HairStrandsCore/Private/GroomBuilder.cpp
+++ b/Engine/Plugins/Runtime/HairStrands/Source/HairStrandsCore/Private/GroomBuilder.cpp
@@ -1394,6 +1394,40 @@ FORCEINLINE FIntVector VectorToIntVector(const FVector& Index)
 	return FIntVector(Index.X, Index.Y, Index.Z);
 }
 
+// These accessors are defined here to ease mirror logic with BuildHairDescriptionGroups, if any type/attributes changes
+bool FHairDescription::HasRootUV() const
+{
+	return StrandAttributes().GetAttributesRef<FVector2f>(HairAttribute::Strand::RootUV).IsValid();
+}
+
+bool FHairDescription::HasGuideWeights() const
+{
+	return
+		// Single
+		(StrandAttributes().GetAttributesRef<int>(HairAttribute::Strand::ClosestGuides).IsValid() && 
+		 StrandAttributes().GetAttributesRef<float>(HairAttribute::Strand::GuideWeights).IsValid())
+		||
+		// Triplet
+		(StrandAttributes().GetAttributesRef<FVector3f>(HairAttribute::Strand::ClosestGuides).IsValid() &&
+		 StrandAttributes().GetAttributesRef<FVector3f>(HairAttribute::Strand::GuideWeights).IsValid());
+}
+
+bool FHairDescription::HasColorAttributes() const
+{
+	return
+		VertexAttributes().GetAttributesRef<FVector3f>(HairAttribute::Vertex::Color).IsValid() ||
+		StrandAttributes().GetAttributesRef<FVector3f>(HairAttribute::Strand::Color).IsValid() ||
+		GroomAttributes().GetAttributesRef<FVector3f>(HairAttribute::Groom::Color).IsValid();
+}
+
+bool FHairDescription::HasRoughnessAttributes() const
+{
+	return
+		VertexAttributes().GetAttributesRef<FVector3f>(HairAttribute::Vertex::Roughness).IsValid() ||
+		StrandAttributes().GetAttributesRef<FVector3f>(HairAttribute::Strand::Roughness).IsValid() ||
+		GroomAttributes().GetAttributesRef<FVector3f>(HairAttribute::Groom::Roughness).IsValid();
+}
+
 bool FGroomBuilder::BuildHairDescriptionGroups(const FHairDescription& HairDescription, FHairDescriptionGroups& Out)
 {
 	TRACE_CPUPROFILER_EVENT_SCOPE(FGroomBuilder::BuildHairDescriptionGroups);
diff --git a/Engine/Plugins/Runtime/HairStrands/Source/HairStrandsCore/Public/GroomImportOptions.h b/Engine/Plugins/Runtime/HairStrands/Source/HairStrandsCore/Public/GroomImportOptions.h
--- a/Engine/Plugins/Runtime/HairStrands/Source/HairStrandsCore/Public/GroomImportOptions.h
+++ b/Engine/Plugins/Runtime/HairStrands/Source/HairStrandsCore/Public/GroomImportOptions.h
@@ -31,6 +31,9 @@ struct HAIRSTRANDSCORE_API FGroomHairGroupPreview
 	: GroupID(0)
 	, CurveCount(0)
 	, GuideCount(0)
+	, bHasRootUV(false)
+	, bHasColorAttributes(false)
+	, bHasRoughnessAttributes(false)
 	, bHasPrecomputedWeights(false)
 	, InterpolationSettings()
 	{}
@@ -47,6 +50,15 @@ struct HAIRSTRANDSCORE_API FGroomHairGroupPreview
 	UPROPERTY(VisibleAnywhere, BlueprintReadWrite, Category = Preview)
 	int32 GuideCount;
 
+	UPROPERTY(VisibleAnywhere, BlueprintReadWrite, Category = Preview)
+	bool bHasRootUV;
+
+	UPROPERTY(VisibleAnywhere, BlueprintReadWrite, Category = Preview)
+	bool bHasColorAttributes;
+
+	UPROPERTY(VisibleAnywhere, BlueprintReadWrite, Category = Preview)
+	bool bHasRoughnessAttributes;
+
 	UPROPERTY(VisibleAnywhere, BlueprintReadWrite, Category = Preview)
 	bool bHasPrecomputedWeights;
 
diff --git a/Engine/Plugins/Runtime/HairStrands/Source/HairStrandsCore/Public/HairDescription.h b/Engine/Plugins/Runtime/HairStrands/Source/HairStrandsCore/Public/HairDescription.h
--- a/Engine/Plugins/Runtime/HairStrands/Source/HairStrandsCore/Public/HairDescription.h
+++ b/Engine/Plugins/Runtime/HairStrands/Source/HairStrandsCore/Public/HairDescription.h
@@ -94,6 +94,11 @@ public:
 
 	// #ueent_todo: Expose "remove" operations to allow editing HairDescription
 
+	bool HasRootUV() const;
+	bool HasGuideWeights() const;
+	bool HasColorAttributes() const;
+	bool HasRoughnessAttributes() const;
+
 private:
 	TAttributesSet<FVertexID> VertexAttributesSet;
 	TAttributesSet<FStrandID> StrandAttributesSet;
diff --git a/Engine/Plugins/Runtime/HairStrands/Source/HairStrandsEditor/Private/GroomActions.cpp b/Engine/Plugins/Runtime/HairStrands/Source/HairStrandsEditor/Private/GroomActions.cpp
--- a/Engine/Plugins/Runtime/HairStrands/Source/HairStrandsEditor/Private/GroomActions.cpp
+++ b/Engine/Plugins/Runtime/HairStrands/Source/HairStrandsEditor/Private/GroomActions.cpp
@@ -190,8 +190,11 @@ void FGroomActions::ExecuteRebuild(TArray<TWeakObjectPtr<UGroomAsset>> Objects)
 						OutGroup.GroupName	= GroomAsset->HairGroupsInfo[GroupIndex].GroupName;
 						OutGroup.CurveCount = GroomAsset->HairGroupsData[GroupIndex].Strands.BulkData.GetNumCurves();
 						OutGroup.GuideCount = GroomAsset->HairGroupsData[GroupIndex].Guides.BulkData.GetNumCurves();
-						OutGroup.InterpolationSettings = GroomAsset->HairGroupsInterpolation[GroupIndex];
+						OutGroup.bHasRootUV = false;
+						OutGroup.bHasColorAttributes = false;
+						OutGroup.bHasRoughnessAttributes = false;
 						OutGroup.bHasPrecomputedWeights = false;
+						OutGroup.InterpolationSettings = GroomAsset->HairGroupsInterpolation[GroupIndex];
 					}
 				}
 				TSharedPtr<SGroomImportOptionsWindow> GroomOptionWindow = SGroomImportOptionsWindow::DisplayRebuildOptions(CurrentOptions, GroupsPreview, Filename);
diff --git a/Engine/Plugins/Runtime/HairStrands/Source/HairStrandsEditor/Private/GroomAssetDetails.cpp b/Engine/Plugins/Runtime/HairStrands/Source/HairStrandsEditor/Private/GroomAssetDetails.cpp
--- a/Engine/Plugins/Runtime/HairStrands/Source/HairStrandsEditor/Private/GroomAssetDetails.cpp
+++ b/Engine/Plugins/Runtime/HairStrands/Source/HairStrandsEditor/Private/GroomAssetDetails.cpp
@@ -1250,7 +1250,7 @@ TSharedRef<SWidget> FGroomRenderingDetails::MakeGroupNameButtonCustomization(int
 	return SNullWidget::NullWidget;
 }
 
-FName FGroomRenderingDetails::GetGroupName(int32 GroupIndex) const
+static FName GetGroupName(const UGroomAsset* GroomAsset, int32 GroupIndex)
 {
 	if (GroomAsset && GroupIndex >= 0 && GroupIndex < GroomAsset->HairGroupsInfo.Num())
 	{
@@ -1259,7 +1259,26 @@ FName FGroomRenderingDetails::GetGroupName(int32 GroupIndex) const
 	return NAME_None;
 }
 
-TSharedRef<SWidget> FGroomRenderingDetails::MakeGroupNameCustomization(int32 GroupIndex, const FLinearColor& GroupColor)
+TSharedRef<SWidget> GetGroupNameWidget(const UGroomAsset* GroomAsset, int32 GroupIndex, const FLinearColor& GroupColor)
+{
+	FName GroupName = GetGroupName(GroomAsset, GroupIndex);
+	if (GroupName != NAME_None)
+	{
+		return SNew(STextBlock)
+			.Font(IDetailLayoutBuilder::GetDetailFont())
+			.ColorAndOpacity(GroupColor)
+			.Text(FText::Format(LOCTEXT("GroupWithName", "Group ID {0} - {1}"), FText::AsNumber(GroupIndex), FText::FromName(GroupName)));
+	}
+	else
+	{
+		return SNew(STextBlock)
+			.Font(IDetailLayoutBuilder::GetDetailFont())
+			.ColorAndOpacity(GroupColor)
+			.Text(FText::Format(LOCTEXT("GroupWithoutName", "Group ID {0}"), FText::AsNumber(GroupIndex)));
+	}
+}
+
+TSharedRef<SWidget> FGroomRenderingDetails::MakeGroupNameCustomization(int32 GroupIndex, const FLinearColor& GroupTextColor)
 {
 	switch (PanelType)
 	{
@@ -1267,7 +1286,7 @@ TSharedRef<SWidget> FGroomRenderingDetails::MakeGroupNameCustomization(int32 Gro
 	{
 		return SNew(STextBlock)
 			.Font(IDetailLayoutBuilder::GetDetailFont())
-			.ColorAndOpacity(GroupColor)
+			.ColorAndOpacity(GroupTextColor)
 			.Text(FText::Format(LOCTEXT("Cards", "Cards {0} "), FText::AsNumber(GroupIndex)));
 	}
 	break;
@@ -1275,27 +1294,13 @@ TSharedRef<SWidget> FGroomRenderingDetails::MakeGroupNameCustomization(int32 Gro
 	{
 		return SNew(STextBlock)
 			.Font(IDetailLayoutBuilder::GetDetailFont())
-			.ColorAndOpacity(GroupColor)
+			.ColorAndOpacity(GroupTextColor)
 			.Text(FText::Format(LOCTEXT("Meshes", "Meshes {0} "), FText::AsNumber(GroupIndex)));
 	}
 	break;
 	default:
 	{
-		FName GroupName = GetGroupName(GroupIndex);
-		if (GroupName != NAME_None)
-		{
-			return SNew(STextBlock)
-				.Font(IDetailLayoutBuilder::GetDetailFont())
-				.ColorAndOpacity(GroupColor)
-				.Text(FText::Format(LOCTEXT("GroupWithName", "Group ID {0} - {1}"), FText::AsNumber(GroupIndex), FText::FromName(GroupName)));
-		}
-		else
-		{
-			return SNew(STextBlock)
-				.Font(IDetailLayoutBuilder::GetDetailFont())
-				.ColorAndOpacity(GroupColor)
-				.Text(FText::Format(LOCTEXT("GroupWithoutName", "Group ID {0}"), FText::AsNumber(GroupIndex)));
-		}
+		return GetGroupNameWidget(GroomAsset, GroupIndex, GroupTextColor);
 	}
 	break;
 	}
diff --git a/Engine/Plugins/Runtime/HairStrands/Source/HairStrandsEditor/Private/GroomAssetDetails.h b/Engine/Plugins/Runtime/HairStrands/Source/HairStrandsEditor/Private/GroomAssetDetails.h
--- a/Engine/Plugins/Runtime/HairStrands/Source/HairStrandsEditor/Private/GroomAssetDetails.h
+++ b/Engine/Plugins/Runtime/HairStrands/Source/HairStrandsEditor/Private/GroomAssetDetails.h
@@ -50,7 +50,6 @@ private:
 	FName& GetMaterialSlotName(int32 GroupIndex);
 	const FName& GetMaterialSlotName(int32 GroupIndex) const;
 	int32 GetGroupCount() const;
-	FName GetGroupName(int32 GroupIndex) const;
 
 	void AddNewGroupButton(IDetailCategoryBuilder& StrandsGroupFilesCategory, FProperty* Property, const FText& HeaderText);
 
diff --git a/Engine/Plugins/Runtime/HairStrands/Source/HairStrandsEditor/Private/GroomComponentDetailsCustomization.cpp b/Engine/Plugins/Runtime/HairStrands/Source/HairStrandsEditor/Private/GroomComponentDetailsCustomization.cpp
--- a/Engine/Plugins/Runtime/HairStrands/Source/HairStrandsEditor/Private/GroomComponentDetailsCustomization.cpp
+++ b/Engine/Plugins/Runtime/HairStrands/Source/HairStrandsEditor/Private/GroomComponentDetailsCustomization.cpp
@@ -203,38 +203,40 @@ static TSharedRef<SUniformGridPanel> MakeHairInfoGrid(const FSlateFontInfo& Deta
 }
 
 // Hair group custom display
+TSharedRef<SWidget> GetGroupNameWidget(const UGroomAsset* GroomAsset, int32 GroupIndex, const FLinearColor& GroupColor);
 void FGroomComponentDetailsCustomization::OnGenerateElementForHairGroup(TSharedRef<IPropertyHandle> StructProperty, int32 GroupIndex, IDetailChildrenBuilder& ChildrenBuilder, IDetailLayoutBuilder* DetailLayout)
 {
 	const FSlateFontInfo DetailFontInfo = IDetailLayoutBuilder::GetDetailFont();
 
+	static const FSlateBrush* GenericBrush = FCoreStyle::Get().GetBrush("GenericWhiteBox");
+	float OtherMargin = 2.0f;
+	float RightMargin = 10.0f;
+
+	const FLinearColor GroupColorBlock = GetHairGroupDebugColor(GroupIndex) * 0.75f;
+
 	ChildrenBuilder.AddCustomRow(LOCTEXT("HairInfo_Separator", "Separator"))
-	.NameContent()
-	[
-		SNew(SHorizontalBox)
-		+ SHorizontalBox::Slot()
-		.VAlign(VAlign_Center)
-		[
-			SNew(SSeparator)
-			.Thickness(2)
-		]
-	]
-	.ValueContent()
+	.WholeRowContent()
+	.VAlign(VAlign_Fill)
 	.HAlign(HAlign_Fill)
 	[
-		SNew(SHorizontalBox)
-		+ SHorizontalBox::Slot()
-		.VAlign(VAlign_Center)
+		SNew(SOverlay)
+		+ SOverlay::Slot()
 		[
-			SNew(SSeparator)
-			.Thickness(2)
+			SNew(SImage)
+			.Image(GenericBrush)
+			.ColorAndOpacity(GroupColorBlock)
 		]
-		+ SHorizontalBox::Slot()
-		.AutoWidth()
+		+ SOverlay::Slot()
 		.HAlign(HAlign_Right)
+		.VAlign(VAlign_Center)
 		[
-			SNew(STextBlock)
-			.Font(IDetailLayoutBuilder::GetDetailFont())
-			.Text(FText::Format(LOCTEXT("Group", "Group ID {0}"), FText::AsNumber(GroupIndex)))
+			SNew(SHorizontalBox)
+			+ SHorizontalBox::Slot()
+			.VAlign(VAlign_Center)
+			.Padding(OtherMargin, OtherMargin, RightMargin, OtherMargin)
+			[
+				GetGroupNameWidget(GroomComponentPtr.IsValid() ? GroomComponentPtr->GroomAsset : nullptr, GroupIndex, FLinearColor::White)
+			]			
 		]
 	];
 
diff --git a/Engine/Plugins/Runtime/HairStrands/Source/HairStrandsEditor/Private/GroomImportOptionsWindow.cpp b/Engine/Plugins/Runtime/HairStrands/Source/HairStrandsEditor/Private/GroomImportOptionsWindow.cpp
--- a/Engine/Plugins/Runtime/HairStrands/Source/HairStrandsEditor/Private/GroomImportOptionsWindow.cpp
+++ b/Engine/Plugins/Runtime/HairStrands/Source/HairStrandsEditor/Private/GroomImportOptionsWindow.cpp
@@ -94,6 +94,28 @@ void SGroomImportOptionsWindow::Construct(const FArguments& InArgs)
 		ValidationText = LOCTEXT("GroomOptionsWindow_ValidationText3", "Unknown");
 	}
 
+	const FSlateFontInfo AttributeFont = FEditorStyle::GetFontStyle("CurveEd.InfoFont");
+	const FSlateFontInfo AttributeResultFont = FEditorStyle::GetFontStyle("CurveEd.InfoFont");
+	const FLinearColor AttributeColor(0.80f, 0.80f, 0.80f);
+	const FText TrueText  = LOCTEXT("GroomOptionsWindow_AttributeTrue", "True");
+	const FText FalseText = LOCTEXT("GroomOptionsWindow_AttributeFalse", "False");
+
+	FText HasRootUVText = FalseText;
+	FText HasColorAttributeText = FalseText;
+	FText HasRoughnessAttributeText = FalseText;
+	FText HasGuideWeightsText = FalseText;
+
+	if (GroupsPreview)
+	{
+		for (const FGroomHairGroupPreview& Group : GroupsPreview->Groups)
+		{
+			if (Group.bHasRootUV)				{ HasRootUVText = TrueText; }
+			if (Group.bHasColorAttributes)		{ HasColorAttributeText = TrueText; }
+			if (Group.bHasRoughnessAttributes)	{ HasRoughnessAttributeText = TrueText; }
+			if (Group.bHasPrecomputedWeights)	{ HasGuideWeightsText = TrueText; }
+		}
+	}
+
 	this->ChildSlot
 	[
 		SNew(SVerticalBox)
@@ -155,6 +177,127 @@ void SGroomImportOptionsWindow::Construct(const FArguments& InArgs)
 			]
 		]
 
+		// Root UV
+		+ SVerticalBox::Slot()
+		.AutoHeight()
+		.Padding(2)
+		[
+			SNew(SBorder)
+			.Padding(FMargin(3))
+			.BorderImage(FEditorStyle::GetBrush("ToolPanel.GroupBorder"))
+			[
+				SNew(SHorizontalBox)
+				+ SHorizontalBox::Slot()
+				.AutoWidth()
+				[
+					SNew(STextBlock)
+					.Font(AttributeFont)
+					.Text(LOCTEXT("GroomOptionsWindow_HasRootUV", "Has Root UV: "))
+				]
+				+ SHorizontalBox::Slot()
+				.Padding(5, 0, 0, 0)
+				.AutoWidth()
+				.VAlign(VAlign_Center)
+				[
+					SNew(STextBlock)
+					.Font(AttributeResultFont)
+					.Text(HasRootUVText)
+					.ColorAndOpacity(AttributeColor)
+				]
+			]
+		]
+
+
+		// Color attributes
+		+ SVerticalBox::Slot()
+		.AutoHeight()
+		.Padding(2)
+		[
+			SNew(SBorder)
+			.Padding(FMargin(3))
+			.BorderImage(FEditorStyle::GetBrush("ToolPanel.GroupBorder"))
+			[
+				SNew(SHorizontalBox)
+				+ SHorizontalBox::Slot()
+				.AutoWidth()
+				[
+					SNew(STextBlock)
+					.Font(AttributeFont)
+					.Text(LOCTEXT("GroomOptionsWindow_HasColor", "Has Color Attributes: "))
+				]
+				+ SHorizontalBox::Slot()
+				.Padding(5, 0, 0, 0)
+				.AutoWidth()
+				.VAlign(VAlign_Center)
+				[
+					SNew(STextBlock)
+					.Font(AttributeResultFont)
+					.Text(HasColorAttributeText)
+					.ColorAndOpacity(AttributeColor)
+				]
+			]
+		]
+
+		// Roughness attributes
+		+ SVerticalBox::Slot()
+		.AutoHeight()
+		.Padding(2)
+		[
+			SNew(SBorder)
+			.Padding(FMargin(3))
+			.BorderImage(FEditorStyle::GetBrush("ToolPanel.GroupBorder"))
+			[
+				SNew(SHorizontalBox)
+				+ SHorizontalBox::Slot()
+				.AutoWidth()
+				[
+					SNew(STextBlock)
+					.Font(AttributeFont)
+					.Text(LOCTEXT("GroomOptionsWindow_HasRoughness", "Has Roughness Attributes: "))
+				]
+				+ SHorizontalBox::Slot()
+				.Padding(5, 0, 0, 0)
+				.AutoWidth()
+				.VAlign(VAlign_Center)
+				[
+					SNew(STextBlock)
+					.Font(AttributeResultFont)
+					.Text(HasRoughnessAttributeText)
+					.ColorAndOpacity(AttributeColor)
+				]
+			]
+		]
+
+		// Guide weights
+		+ SVerticalBox::Slot()
+		.AutoHeight()
+		.Padding(2)
+		[
+			SNew(SBorder)
+			.Padding(FMargin(3))
+			.BorderImage(FEditorStyle::GetBrush("ToolPanel.GroupBorder"))
+			[
+				SNew(SHorizontalBox)
+				+ SHorizontalBox::Slot()
+				.AutoWidth()
+				[
+					SNew(STextBlock)
+					.Font(AttributeFont)
+					.Text(LOCTEXT("GroomOptionsWindow_HasGuideWeights", "Has Pre-Computed Guides Weights: "))
+				]
+				+ SHorizontalBox::Slot()
+				.Padding(5, 0, 0, 0)
+				.AutoWidth()
+				.VAlign(VAlign_Center)
+				[
+					SNew(STextBlock)
+					.Font(AttributeResultFont)
+					.Text(HasGuideWeightsText)
+					.ColorAndOpacity(AttributeColor)
+				]
+			]
+		]
+
 		+ SVerticalBox::Slot()
 		.Padding(2)
 		.MaxHeight(500.0f)
diff --git a/Engine/Plugins/Runtime/HairStrands/Source/HairStrandsEditor/Private/HairStrandsFactory.cpp b/Engine/Plugins/Runtime/HairStrands/Source/HairStrandsEditor/Private/HairStrandsFactory.cpp
--- a/Engine/Plugins/Runtime/HairStrands/Source/HairStrandsEditor/Private/HairStrandsFactory.cpp
+++ b/Engine/Plugins/Runtime/HairStrands/Source/HairStrandsEditor/Private/HairStrandsFactory.cpp
@@ -82,6 +82,10 @@ UObject* UHairStrandsFactory::FactoryCreateFile(UClass* InClass, UObject* InPare
 	{
 		// Load the alembic file upfront to preview & report any potential issue
 		FHairDescriptionGroups OutDescription;
+		bool bHasRootUV = false;
+		bool bHasPrecomputedWeights = false;
+		bool bHasColorAttributes = false;
+		bool bHasRoughnessAttributes = false;
 		{
 			FScopedSlowTask Progress((float)1, LOCTEXT("ImportHairAssetForPreview", "Importing hair asset for preview..."), true);
 			Progress.MakeDialog(true);
@@ -93,7 +97,11 @@ UObject* UHairStrandsFactory::FactoryCreateFile(UClass* InClass, UObject* InPare
 			}
 
 			FGroomBuilder::BuildHairDescriptionGroups(HairDescription, OutDescription);
-		
+			bHasRootUV = HairDescription.HasRootUV();
+			bHasPrecomputedWeights = HairDescription.HasGuideWeights();
+			bHasColorAttributes = HairDescription.HasColorAttributes();
+			bHasRoughnessAttributes = HairDescription.HasRoughnessAttributes();
+
 			// Populate the interpolation settings based on the group count, as this is used later during the ImportHair() to define 
 			// the exact number of group to create
 			const uint32 GroupCount = OutDescription.HairGroups.Num();
@@ -113,7 +121,10 @@ UObject* UHairStrandsFactory::FactoryCreateFile(UClass* InClass, UObject* InPare
 				OutGroup.GroupID	= Group.Info.GroupID;
 				OutGroup.CurveCount = Group.Info.NumCurves;
 				OutGroup.GuideCount = Group.Info.NumGuides;
-				OutGroup.bHasPrecomputedWeights = Group.Strands.StrandsCurves.HasPrecomputedWeights();
+				OutGroup.bHasRootUV = bHasRootUV;
+				OutGroup.bHasPrecomputedWeights = bHasPrecomputedWeights;
+				OutGroup.bHasColorAttributes = bHasColorAttributes;
+				OutGroup.bHasRoughnessAttributes = bHasRoughnessAttributes;
 
 				if (OutGroup.GroupID < OutDescription.HairGroups.Num())
 				{				
