Fix CIS issue.

diff --git a/Engine/Plugins/Runtime/HairStrands/Source/HairStrandsEditor/Private/GroomComponentDetailsCustomization.cpp b/Engine/Plugins/Runtime/HairStrands/Source/HairStrandsEditor/Private/GroomComponentDetailsCustomization.cpp
--- a/Engine/Plugins/Runtime/HairStrands/Source/HairStrandsEditor/Private/GroomComponentDetailsCustomization.cpp
+++ b/Engine/Plugins/Runtime/HairStrands/Source/HairStrandsEditor/Private/GroomComponentDetailsCustomization.cpp
@@ -11,6 +11,8 @@
 #include "Widgets/Input/SButton.h"
 #include "Widgets/Layout/SSeparator.h"
 #include "Widgets/Layout/SUniformGridPanel.h"
+#include "Widgets/Images/SImage.h"
+#include "Widgets/SOverlay.h"
 #include "EditorStyleSet.h"
 #include "EditorModeManager.h"
 #include "PropertyCustomizationHelpers.h"
