Change hair instance listing to use shader print for better readability.

diff --git a/Engine/Plugins/Runtime/HairStrands/Source/HairStrandsCore/Private/GroomDebug.cpp b/Engine/Plugins/Runtime/HairStrands/Source/HairStrandsCore/Private/GroomDebug.cpp
--- a/Engine/Plugins/Runtime/HairStrands/Source/HairStrandsCore/Private/GroomDebug.cpp
+++ b/Engine/Plugins/Runtime/HairStrands/Source/HairStrandsCore/Private/GroomDebug.cpp
@@ -12,6 +12,7 @@
 #include "UnrealEngine.h"
 #include "SystemTextures.h"
 #include "CanvasTypes.h"
+#include "ShaderCompilerCore.h"
 
 ///////////////////////////////////////////////////////////////////////////////////////////////////
 
@@ -886,106 +887,172 @@ FCachedGeometry GetCacheGeometryForHair(
 	const FGPUSkinCache* SkinCache,
 	FGlobalShaderMap* ShaderMap);
 
-void RunHairStrandsDebug(
-	FRDGBuilder& GraphBuilder,
+
+///////////////////////////////////////////////////////////////////////////////////////////////////
+
+class FHairDebugPrintInstanceCS : public FGlobalShader
+{
+	DECLARE_GLOBAL_SHADER(FHairDebugPrintInstanceCS);
+	SHADER_USE_PARAMETER_STRUCT(FHairDebugPrintInstanceCS, FGlobalShader);
+
+	BEGIN_SHADER_PARAMETER_STRUCT(FParameters, )
+		SHADER_PARAMETER(uint32, InstanceCount)
+		SHADER_PARAMETER(uint32, NameInfoCount)
+		SHADER_PARAMETER(uint32, NameCharacterCount)
+		SHADER_PARAMETER_RDG_BUFFER_SRV(StructuredBuffer<uint2>, NameInfos)
+		SHADER_PARAMETER_RDG_BUFFER_SRV(Buffer<uint8>, Names)
+		SHADER_PARAMETER_RDG_BUFFER_SRV(StructuredBuffer<uint4>, Infos)
+		SHADER_PARAMETER_STRUCT_INCLUDE(ShaderPrint::FShaderParameters, ShaderPrintUniformBuffer)
+		SHADER_PARAMETER_STRUCT_INCLUDE(ShaderDrawDebug::FShaderParameters, ShaderDrawUniformBuffer)
+	END_SHADER_PARAMETER_STRUCT()
+
+public:
+	static bool ShouldCompilePermutation(const FGlobalShaderPermutationParameters& Parameters) { return IsHairStrandsSupported(EHairStrandsShaderType::Tool, Parameters.Platform); }
+	static void ModifyCompilationEnvironment(const FGlobalShaderPermutationParameters& Parameters, FShaderCompilerEnvironment& OutEnvironment)
+	{
+		// Skip optimization for avoiding long compilation time due to large UAV writes
+		FGlobalShader::ModifyCompilationEnvironment(Parameters, OutEnvironment);
+		OutEnvironment.CompilerFlags.Add(CFLAG_Debug);
+		OutEnvironment.SetDefine(TEXT("SHADER_PRINT_INSTANCE"), 1);
+	}
+};
+
+IMPLEMENT_GLOBAL_SHADER(FHairDebugPrintInstanceCS, "/Engine/Private/HairStrands/HairStrandsDebugPrint.usf", "MainCS", SF_Compute);
+
+struct FHairDebugNameInfo
+{
+	uint32 PrimitiveID;
+	uint16 Offset;
+	uint8  Length;
+	uint8  Pad0;
+};
+
+static void AddHairDebugPrintInstancePass(
+	FRDGBuilder& GraphBuilder, 
 	FGlobalShaderMap* ShaderMap,
-	const FSceneView& View,
-	const FHairStrandsInstances& Instances,
-	const FGPUSkinCache* SkinCache,
 	const FShaderDrawDebugData* ShaderDrawData,
-	FRDGTextureRef SceneColorTexture,
-	FRDGTextureRef SceneDepthTexture,
-	FIntRect Viewport,
-	const TUniformBufferRef<FViewUniformShaderParameters>& ViewUniformBuffer)
+	const FShaderPrintData* ShaderPrintData,
+	const FHairStrandsInstances& Instances)
 {
-	const EHairDebugMode HairDebugMode = GetHairStrandsDebugMode();
+	const uint32 InstanceCount = Instances.Num();
 
-	if (HairDebugMode == EHairDebugMode::MacroGroups)
-	{
-		FHairDebugCanvasParameter* PassParameters = GraphBuilder.AllocParameters<FHairDebugCanvasParameter>();
-		PassParameters->View = View.ViewUniformBuffer;
-		PassParameters->RenderTargets[0] = FRenderTargetBinding(SceneColorTexture, ERenderTargetLoadAction::ELoad, 0);
+	// Request more drawing primitives & characters for printing if needed	
+	ShaderDrawDebug::SetEnabled(true);
+	ShaderDrawDebug::RequestSpaceForElements(InstanceCount * 16u);
+	ShaderPrint::RequestSpaceForCharacters(InstanceCount * 256 + 512);
 
-		const FSceneViewFamily& ViewFamily = *View.Family;
-		FCanvas& Canvas = *FCanvas::Create(GraphBuilder, SceneColorTexture, nullptr, ViewFamily.Time, View.FeatureLevel);
-		Canvas.SetRenderTargetRect(Viewport);
+	if (!ShaderPrintData || !ShaderDrawData || InstanceCount == 0) { return; }
+
+	const uint32 MaxPrimitiveNameCount = 128u;
+	check(sizeof(FHairDebugNameInfo) == 8);
+
+	TArray<FHairDebugNameInfo> NameInfos;
+	TArray<uint8> Names;
+	Names.Reserve(MaxPrimitiveNameCount * 30u);
+
+	TArray<FUintVector4> Infos;
+	Infos.Reserve(InstanceCount);
+	for (uint32 InstanceIndex = 0; InstanceIndex < InstanceCount; ++InstanceIndex)
+	{
+		const FHairStrandsInstance* AbstractInstance = Instances[InstanceIndex];
+		const FHairGroupInstance* Instance = static_cast<const FHairGroupInstance*>(AbstractInstance);
 
+		// Collect Names
+		if (InstanceIndex < MaxPrimitiveNameCount)
 		{
-			const float YStep = 14;
-			float ClusterY = 68;
-			float X = 20;
-			float Y = ClusterY;
-			const FLinearColor InactiveColor(0.5, 0.5, 0.5);
-			const FLinearColor DebugColor(1, 1, 0);
-			const FLinearColor DebugGroupColor(1, 0.5f, 0);
-			const FLinearColor DebugWarningColor(0.75f, 0, 0);
-			FString Line;
-			// Active groom
-			// Name | Group x / x | LOD x / x | GeometryType | BindingType | Sim | RBF | VertexCount | Warning
-
-			Line = FString::Printf(TEXT("----------------------------------------------------------------"));
-			Canvas.DrawShadowedString(X, Y += YStep, *Line, GetStatsFont(), DebugColor);
-
-			Line = FString::Printf(TEXT("Registered hair groups count : %d"), Instances.Num());
-			Canvas.DrawShadowedString(X, Y += YStep, *Line, GetStatsFont(), DebugColor);
-
-			for (FHairStrandsInstance* AbstractInstance : Instances)
+			const FString Name = *Instance->Debug.GroomAssetName;
+			const uint32 NameOffset = Names.Num();
+			const uint32 NameLength = Name.Len();
+			for (TCHAR C : Name)
 			{
-				FHairGroupInstance* Instance = static_cast<FHairGroupInstance*>(AbstractInstance);
+				Names.Add(uint8(C));
+			}
 
-				// Check if a groom request a Skinning binding for the current LOD, that the parent has its skin cache enabled
-				bool bInvaliSkinBinding = false;
-				{
-					const int32 GroomLODIndex = Instance->HairGroupPublicData->LODIndex;
-					EHairBindingType OriginalBindingType = Instance->HairGroupPublicData->GetBindingType(GroomLODIndex);
+			FHairDebugNameInfo& NameInfo = NameInfos.AddDefaulted_GetRef();
+			NameInfo.PrimitiveID = InstanceIndex;
+			NameInfo.Length = NameLength;
+			NameInfo.Offset = NameOffset;
+		}
 
-					if (OriginalBindingType == EHairBindingType::Skinning)
-					{
-						bool bHasValidSkinCacheData = false;
-						FHairStrandsProjectionMeshData::LOD MeshDataLOD;
-						const FCachedGeometry CachedGeometry = GetCacheGeometryForHair(GraphBuilder, Instance, SkinCache, ShaderMap);
-						for (const FCachedGeometry::Section& Section : CachedGeometry.Sections)
-						{
-							if (Section.LODIndex >= 0) { bHasValidSkinCacheData = true; break; }
-						}
-						bInvaliSkinBinding = !bHasValidSkinCacheData;
-					}
-				}
+		const float LODIndex = Instance->HairGroupPublicData->LODIndex;
+		const uint32 LODCount = Instance->HairGroupPublicData->GetLODScreenSizes().Num();
 
-				Line = FString::Printf(TEXT(" * Group:%d/%d | LOD:%1.2f/%d (%s) | GeometryType:%s | BindingType:%s | Sim:%d | RBF:%d | Clip:%d(%1.2f) | VertexCount:%d | Name: %s"),
-					Instance->Debug.GroupIndex,
-					Instance->Debug.GroupCount,
+		const uint32 DataX =
+			((Instance->Debug.GroupIndex & 0xFF)) |
+			((Instance->Debug.GroupCount & 0xFF) << 8) |
+			((LODCount & 0xFF) << 16) |
+			((uint32(Instance->GeometryType) & 0x7)<<24) |
+			((uint32(Instance->BindingType) & 0x7)<<27) |
+			((Instance->Guides.bIsSimulationEnable ? 0x1 : 0x0) << 30) |
+			((Instance->Guides.bHasGlobalInterpolation ? 0x1 : 0x0) << 31);
 
-					Instance->HairGroupPublicData->LODIndex,
-					Instance->HairGroupPublicData->GetLODScreenSizes().Num(),
-					ToString(Instance->Debug.LODSelectionTypeForDebug),
+		const uint32 DataY =
+			(FFloat16(LODIndex).Encoded) |
+			(FFloat16(Instance->Strands.Modifier.HairLengthScale_Override ? Instance->Strands.Modifier.HairLengthScale : -1.f).Encoded << 16);
 
-					ToString(Instance->GeometryType),
-					ToString(Instance->BindingType),
-					Instance->Guides.bIsSimulationEnable,
-					Instance->Guides.bHasGlobalInterpolation,
 
-					Instance->Strands.Modifier.HairLengthScale_Override,
-					Instance->Strands.Modifier.HairLengthScale,
+		const uint32 DataZ = Instance->Strands.Data->GetNumCurves(); // Change this later on for having dynamic value
+		const uint32 DataW = Instance->Strands.Data->GetNumPoints(); // Change this later on for having dynamic value
 
-					Instance->HairGroupPublicData->VertexCount,
-					*Instance->Debug.GroomAssetName);
-				int32 XOffset = Canvas.DrawShadowedString(X, Y += YStep, *Line, GetStatsFont(), DebugGroupColor);
+		Infos.Add(FUintVector4(DataX, DataY, DataZ, DataW));
+	}
 
-				// Warning/error
-				if (bInvaliSkinBinding)
-				{
-					Line = FString::Printf(TEXT(" - Warning: Request SKINNING binding but skel. mesh has no Skin Cache (%s)"),						
-						*Instance->Debug.MeshComponentName);
-					Canvas.DrawShadowedString(X + XOffset, Y += YStep, *Line, GetStatsFont(), DebugWarningColor);
-				}
-			}
+	if (NameInfos.IsEmpty())
+	{
+		FHairDebugNameInfo& NameInfo = NameInfos.AddDefaulted_GetRef();
+		NameInfo.PrimitiveID = ~0;
+		NameInfo.Length = 4;
+		NameInfo.Offset = 0;
+		Names.Add(uint8('N'));
+		Names.Add(uint8('o'));
+		Names.Add(uint8('n'));
+		Names.Add(uint8('e'));
+	}	
 
-			const bool bFlush = false;
-			Canvas.Flush_RenderThread(GraphBuilder, bFlush);
+	const uint32 InfoInBytes = 16u;
+	FRDGBufferRef NameBuffer = CreateVertexBuffer(GraphBuilder, TEXT("Hair.Debug.InstanceNames"), FRDGBufferDesc::CreateBufferDesc(1, Names.Num()), Names.GetData(), Names.Num());
+	FRDGBufferRef NameInfoBuffer = CreateStructuredBuffer(GraphBuilder, TEXT("Hair.Debug.InstanceNameInfos"), NameInfos);	
+	FRDGBufferRef InfoBuffer = CreateVertexBuffer(GraphBuilder, TEXT("Hair.Debug.InstanceInfos"), FRDGBufferDesc::CreateBufferDesc(InfoInBytes, Infos.Num()), Infos.GetData(), InfoInBytes * Infos.Num());
 
-			ClusterY = Y;
-		}
+	FHairDebugPrintInstanceCS::FParameters* Parameters = GraphBuilder.AllocParameters<FHairDebugPrintInstanceCS::FParameters>();
+	Parameters->InstanceCount = InstanceCount;
+	Parameters->NameInfoCount = NameInfos.Num();
+	Parameters->NameCharacterCount = Names.Num();
+	Parameters->Names = GraphBuilder.CreateSRV(NameBuffer, PF_R8_UINT);
+	Parameters->NameInfos = GraphBuilder.CreateSRV(NameInfoBuffer);
+	Parameters->Infos = GraphBuilder.CreateSRV(InfoBuffer, PF_R32G32B32A32_UINT);
+	ShaderPrint::SetParameters(GraphBuilder, *ShaderPrintData, Parameters->ShaderPrintUniformBuffer);
+	ShaderDrawDebug::SetParameters(GraphBuilder, *ShaderDrawData, Parameters->ShaderDrawUniformBuffer);
+	TShaderMapRef<FHairDebugPrintInstanceCS> ComputeShader(ShaderMap);
+
+	ClearUnusedGraphResources(ComputeShader, Parameters);
+
+	FComputeShaderUtils::AddPass(
+		GraphBuilder,
+		RDG_EVENT_NAME("HairStrands::DebugPrintInstance(Instances:%d)", InstanceCount),
+		ComputeShader,
+		Parameters,
+		FIntVector(1, 1, 1));
+}
+
+void RunHairStrandsDebug(
+	FRDGBuilder& GraphBuilder,
+	FGlobalShaderMap* ShaderMap,
+	const FSceneView& View,
+	const FHairStrandsInstances& Instances,
+	const FGPUSkinCache* SkinCache,
+	const FShaderDrawDebugData* ShaderDrawData,
+	const FShaderPrintData* ShaderPrintData,
+	FRDGTextureRef SceneColorTexture,
+	FRDGTextureRef SceneDepthTexture,
+	FIntRect Viewport,
+	const TUniformBufferRef<FViewUniformShaderParameters>& ViewUniformBuffer)
+{
+	const EHairDebugMode HairDebugMode = GetHairStrandsDebugMode();
+
+	if (HairDebugMode == EHairDebugMode::MacroGroups)
+	{
+		AddHairDebugPrintInstancePass(GraphBuilder, ShaderMap, ShaderDrawData, ShaderPrintData, Instances);
 	}
 
 	if (HairDebugMode == EHairDebugMode::MeshProjection)
diff --git a/Engine/Plugins/Runtime/HairStrands/Source/HairStrandsCore/Private/GroomManager.cpp b/Engine/Plugins/Runtime/HairStrands/Source/HairStrandsCore/Private/GroomManager.cpp
--- a/Engine/Plugins/Runtime/HairStrands/Source/HairStrandsCore/Private/GroomManager.cpp
+++ b/Engine/Plugins/Runtime/HairStrands/Source/HairStrandsCore/Private/GroomManager.cpp
@@ -997,6 +997,7 @@ void RunHairStrandsDebug(
 	const FHairStrandsInstances& Instances,
 	const FGPUSkinCache* SkinCache,
 	const FShaderDrawDebugData* ShaderDrawData,
+	const FShaderPrintData* ShaderPrintData,
 	FRDGTextureRef SceneColor,
 	FRDGTextureRef SceneDepth,
 	FIntRect Viewport,
@@ -1103,6 +1104,7 @@ void ProcessHairStrandsBookmark(
 			Instances,
 			Parameters.SkinCache,
 			Parameters.ShaderDebugData,
+			Parameters.ShaderPrintData,
 			Parameters.SceneColorTexture,
 			Parameters.SceneDepthTexture,
 			Parameters.View->UnscaledViewRect,
diff --git a/Engine/Shaders/Private/HairStrands/HairStrandsDebugPrint.usf b/Engine/Shaders/Private/HairStrands/HairStrandsDebugPrint.usf
--- a/Engine/Shaders/Private/HairStrands/HairStrandsDebugPrint.usf
+++ b/Engine/Shaders/Private/HairStrands/HairStrandsDebugPrint.usf
@@ -52,8 +52,9 @@ FFontColor GetOccupancyColor(float In)
 void MainCS(uint GroupIndex : SV_GroupIndex, uint3 DispatchThreadId : SV_DispatchThreadID)
 {	
 	// Draw AABB of the instance groups
-	if (DispatchThreadId.x == 0 && DispatchThreadId.y == 0 && DispatchThreadId.z == 0)
+	if (all(DispatchThreadId == 0))
 	{
+		// Macro groups bounds
 		for (uint GroupIndex =0;GroupIndex< HairMacroGroupCount; GroupIndex++)
 		{
 			const float4 InstanceBoundColor = float4(1.f, 0.5f, 0.5f, 1.0f);
@@ -67,27 +68,26 @@ void MainCS(uint GroupIndex : SV_GroupIndex, uint3 DispatchThreadId : SV_Dispatc
 		const uint DispatchedSampleCount = DispatchX * DispatchY * HairVisibilityNodeGroupSize;
 
 		const uint2 EquivalentResolution = ceil(sqrt(UsedSampleCount));
-		
-
+	
 		// Pixel coord
-		float2 OriginalPos = float2(50, 50) / float2(MaxResolution);
+		float2 OriginalPos = float2(450, 50) / float2(MaxResolution);
 		float2 Pos = OriginalPos;
-		Pos = ShaderPrintText(Pos, TEXT("Sample Count (Used/Allocated)"), FontCyan);
+		Pos = ShaderPrintText(Pos, TEXT("Sample Count (Used/Alloc)  "), FontCyan);
 		Pos = ShaderPrint(Pos, UsedSampleCount, FontCyan);
 		Pos = ShaderPrint(Pos, AllocatedSampleCount, FontCyan);
 
 		const float AllocationPercentage = float(UsedSampleCount) / float(AllocatedSampleCount);
 		const FFontColor OccupancyColor = GetOccupancyColor(AllocationPercentage);
-		Pos = ShaderPrintText(Pos, TEXT(" \% "), OccupancyColor);
+		Pos = ShaderPrintText(Pos, TEXT(" %% "), OccupancyColor);
 		Pos = ShaderPrint(Pos, AllocationPercentage * 100.f, OccupancyColor);
 
 		AddNewLine(OriginalPos, Pos);
 
-		Pos = ShaderPrintText(Pos, TEXT("Equivalent Res     "), FontCyan);
+		Pos = ShaderPrintText(Pos, TEXT("Sample Lighting Resolution "), FontCyan);
 		Pos = ShaderPrint(Pos, EquivalentResolution, FontCyan);
 		AddNewLine(OriginalPos, Pos);
 
-		Pos = ShaderPrintText(Pos, TEXT("MacroGroup Count   "), FontYellow);
+		Pos = ShaderPrintText(Pos, TEXT("MacroGroup Count           "), FontYellow);
 		Pos = ShaderPrint(Pos, HairMacroGroupCount, FontYellow);
 		AddNewLine(OriginalPos, Pos);
 	}
@@ -253,6 +253,203 @@ void MainCS(uint GroupIndex : SV_GroupIndex, uint3 DispatchThreadId : SV_Dispatc
 
 ////////////////////////////////////////////////////////////////////////////////////////////////
 
+#if SHADER_PRINT_INSTANCE
+
+int2					MaxResolution;
+uint					InstanceCount;
+
+Buffer<uint4>			Infos;
+
+uint					NameInfoCount;
+uint					NameCharacterCount;
+StructuredBuffer<uint2> NameInfos;
+Buffer<uint>			Names;
+
+struct FNameInfo
+{
+	uint PrimitiveID;
+	uint Length;
+	uint Offset;
+	uint Pad0;
+	uint Pad1;
+};
+
+FNameInfo UnpackDebugNameInfo(uint2 In)
+{
+	FNameInfo Out = (FNameInfo)0;
+	Out.PrimitiveID = In.x;
+	Out.Offset = (In.y) & 0xFFFF;
+	Out.Length = (In.y >> 16) & 0xFF;
+
+	return Out;
+}
+
+FNameInfo FindNameInfo(uint PrimitiveID)
+{
+	FNameInfo Out = (FNameInfo)0;
+	Out.PrimitiveID = ~0;
+	for (uint It = 0; It < NameInfoCount; ++It)
+	{
+		if (NameInfos[It].x == PrimitiveID)
+		{
+			return UnpackDebugNameInfo(NameInfos[It]);
+		}
+	}
+	return Out;
+}
+
+float2 PrintInstanceName(float2 InPos, uint PrimitiveID, FFontColor InColor)
+{
+	const FNameInfo Info = FindNameInfo(PrimitiveID);
+	if (Info.PrimitiveID != ~0 && (Info.Length + Info.Offset) <= NameCharacterCount)
+	{
+		for (uint It = 0; It < Info.Length; ++It)
+		{
+			const uint Char = Names[It + Info.Offset];
+			InPos = ShaderPrintSymbol(InPos, Char, InColor);
+		}
+	}
+	return InPos;
+}
+
+
+float2 Print(float2 Pos, uint X, uint DigitCount, FFontColor Color)
+{
+	Pos = ShaderPrint(Pos, X, Color);
+	Pos.x -= ShaderPrintUniform.FontSize.z * (ShaderPrintTabCount - min(ShaderPrintTabCount, DigitCount));
+	return Pos;
+}
+
+float2 Print(float2 Pos, float X, uint DigitCount, FFontColor Color)
+{
+	Pos = ShaderPrint(Pos, X, Color);
+	Pos.x -= ShaderPrintUniform.FontSize.z * (ShaderPrintTabCount - min(ShaderPrintTabCount, DigitCount));
+	return Pos;
+}
+
+[numthreads(1, 1, 1)]
+void MainCS(uint GroupIndex : SV_GroupIndex, uint3 DispatchThreadId : SV_DispatchThreadID)
+{
+	// Draw AABB of the instance groups
+	if (any(DispatchThreadId != 0))
+	{
+		return;
+	}
+
+	// Draw legend
+	float2 OriginalPos = ShaderPrintGetPos(uint2(50, 50));
+	float2 Pos = OriginalPos;
+
+
+	// Draw totals
+	const uint TotalGroup = InstanceCount;
+	uint TotalStrandsCount = 0;
+	uint TotalVertexCount = 0;
+	for (uint InstanceIndex = 0; InstanceIndex < InstanceCount; ++InstanceIndex)
+	{
+		const uint4 InstanceData = Infos[InstanceIndex];
+		TotalStrandsCount += InstanceData.z;
+		TotalVertexCount += InstanceData.w;
+	}
+
+	Pos = ShaderPrintText(Pos, TEXT("Total Groups        : "), FontWhite);
+	Pos = Print(Pos, TotalGroup, 8, FontYellow);
+	AddNewLine(OriginalPos, Pos);
+
+	Pos = ShaderPrintText(Pos, TEXT("Total Strand Curves : "), FontWhite);
+	Pos = Print(Pos, TotalStrandsCount, 8, FontOrange);
+	AddNewLine(OriginalPos, Pos);
+
+	Pos = ShaderPrintText(Pos, TEXT("Total Strand Vertex : "), FontWhite);
+	Pos = Print(Pos, TotalVertexCount, 8, FontSilver);
+	AddNewLine(OriginalPos, Pos);
+
+	AddNewLine(OriginalPos, Pos);
+
+	// Draw instance details
+	Pos = ShaderPrintText(Pos, TEXT("   Group  LOD        Geometry Binding Sim   RBF   Clip    Strands Vertex  Name"));
+	AddNewLine(OriginalPos, Pos);
+
+	for (uint InstanceIndex = 0; InstanceIndex < InstanceCount; ++InstanceIndex)
+	{
+		const uint4 InstanceData	= Infos[InstanceIndex];
+		const uint GroupIndex		= (InstanceData.x>>0)  & 0xFF;
+		const uint GroupCount		= (InstanceData.x>>8)  & 0xFF;
+		const float LODCount		= (InstanceData.x>>16) & 0xFF;
+		const uint GeometryType		= (InstanceData.x>>24) & 0x7;
+		const uint BindingType		= (InstanceData.x>>27) & 0x7;
+		const bool bHasSim			= (InstanceData.x>>30) & 0x1;
+		const bool bHasRBF			= (InstanceData.x>>31) & 0x1;
+
+		const float LOD				= f16tof32((InstanceData.y>>0)  & 0xFFFF);
+		const float ClipValue		= f16tof32((InstanceData.y>>16) & 0xFFFF);
+
+		const uint StrandsCount		= InstanceData.z;
+		const uint VertexCount		= InstanceData.w;
+
+		// Index
+		Pos = Print(Pos, InstanceIndex, 2, FontSilver);
+		Pos = ShaderPrintText(Pos, TEXT(" "));
+
+		// Group
+		Pos = Print(Pos, GroupIndex, 1, FontYellow);
+		Pos = ShaderPrintText(Pos, TEXT("/"));
+		Pos = Print(Pos, GroupCount, 1, FontYellow);
+		Pos = ShaderPrintText(Pos, TEXT("    "));
+
+		// LOD
+		Pos = Print(Pos, LOD, 7, FontYellow);	
+		Pos = ShaderPrintText(Pos, TEXT("/"));
+		Pos = Print(Pos, LODCount, 1, FontYellow);
+		Pos = ShaderPrintText(Pos, TEXT("  "));
+
+		// Geometry
+		if (GeometryType == 0)		{ Pos = ShaderPrintText(Pos, TEXT("Strands "), FontEmerald); }
+		else if (GeometryType == 1) { Pos = ShaderPrintText(Pos, TEXT("Cards   "), FontEmerald); }
+		else if (GeometryType == 2) { Pos = ShaderPrintText(Pos, TEXT("Meshes  "), FontEmerald); }
+		Pos = ShaderPrintText(Pos, TEXT(" "));
+
+		// Binding
+		if (BindingType == 0)		{ Pos = ShaderPrintText(Pos, TEXT("Skinned "), FontTurquoise); }
+		else if (BindingType == 1)	{ Pos = ShaderPrintText(Pos, TEXT("Rigid   "), FontTurquoise); }
+		else if (BindingType == 2)	{ Pos = ShaderPrintText(Pos, TEXT("None    "), FontTurquoise); }
+
+		// Sim
+		if (bHasSim)
+			Pos = ShaderPrintText(Pos, TEXT("True"), FontGreen);
+		else
+			Pos = ShaderPrintText(Pos, TEXT("False"), FontRed);
+		Pos = ShaderPrintText(Pos, TEXT(" "), FontWhite);
+
+		// RBF
+		if (bHasRBF)
+			Pos = ShaderPrintText(Pos, TEXT("True"), FontGreen);
+		else
+			Pos = ShaderPrintText(Pos, TEXT("False"), FontRed);
+		Pos = ShaderPrintText(Pos, TEXT(" "));
+
+		// Clip
+		if (ClipValue >= 0)
+			Pos = Print(Pos, ClipValue, 7, FontGreen);
+		else 
+			Pos = Print(Pos, 0u, 7, FontSilver);
+		Pos = ShaderPrintText(Pos, TEXT(" "));
+
+		// Count
+		Pos = Print(Pos, StrandsCount, 8, FontOrange);
+		Pos = Print(Pos, VertexCount,  8, FontSilver);
+
+		// Name
+		Pos = PrintInstanceName(Pos, InstanceIndex, FontWhite);
+
+		AddNewLine(OriginalPos, Pos);
+	}
+}
+
+#endif // SHADER_PRINT_INSTANCE
+
+////////////////////////////////////////////////////////////////////////////////////////////////
+
 #if SHADER_TANGENT
 uint TileSize;
 int2 TileCount;
