Change
* Fix groom stats per geometry type
* Add per hair group AABB display

diff --git a/Engine/Plugins/Runtime/HairStrands/Source/HairStrandsCore/Private/GroomDebug.cpp b/Engine/Plugins/Runtime/HairStrands/Source/HairStrandsCore/Private/GroomDebug.cpp
--- a/Engine/Plugins/Runtime/HairStrands/Source/HairStrandsCore/Private/GroomDebug.cpp
+++ b/Engine/Plugins/Runtime/HairStrands/Source/HairStrandsCore/Private/GroomDebug.cpp
@@ -895,6 +895,9 @@ class FHairDebugPrintInstanceCS : public FGlobalShader
 	DECLARE_GLOBAL_SHADER(FHairDebugPrintInstanceCS);
 	SHADER_USE_PARAMETER_STRUCT(FHairDebugPrintInstanceCS, FGlobalShader);
 
+	class FOutputType : SHADER_PERMUTATION_INT("PERMUTATION_OUTPUT_TYPE", 2);
+	using FPermutationDomain = TShaderPermutationDomain<FOutputType>;
+
 	BEGIN_SHADER_PARAMETER_STRUCT(FParameters, )
 		SHADER_PARAMETER(uint32, InstanceCount)
 		SHADER_PARAMETER(uint32, NameInfoCount)
@@ -902,6 +905,7 @@ class FHairDebugPrintInstanceCS : public FGlobalShader
 		SHADER_PARAMETER_RDG_BUFFER_SRV(StructuredBuffer<uint2>, NameInfos)
 		SHADER_PARAMETER_RDG_BUFFER_SRV(Buffer<uint8>, Names)
 		SHADER_PARAMETER_RDG_BUFFER_SRV(StructuredBuffer<uint4>, Infos)
+		SHADER_PARAMETER_RDG_BUFFER_SRV(Buffer<int>, InstanceAABB)
 		SHADER_PARAMETER_STRUCT_INCLUDE(ShaderPrint::FShaderParameters, ShaderPrintUniformBuffer)
 		SHADER_PARAMETER_STRUCT_INCLUDE(ShaderDrawDebug::FShaderParameters, ShaderDrawUniformBuffer)
 	END_SHADER_PARAMETER_STRUCT()
@@ -975,6 +979,7 @@ static void AddHairDebugPrintInstancePass(
 		}
 
 		const float LODIndex = Instance->HairGroupPublicData->LODIndex;
+		const uint32 IntLODIndex = Instance->HairGroupPublicData->LODIndex;
 		const uint32 LODCount = Instance->HairGroupPublicData->GetLODScreenSizes().Num();
 
 		const uint32 DataX =
@@ -990,9 +995,32 @@ static void AddHairDebugPrintInstancePass(
 			(FFloat16(LODIndex).Encoded) |
 			(FFloat16(Instance->Strands.Modifier.HairLengthScale_Override ? Instance->Strands.Modifier.HairLengthScale : -1.f).Encoded << 16);
 
-
-		const uint32 DataZ = Instance->Strands.Data->GetNumCurves(); // Change this later on for having dynamic value
-		const uint32 DataW = Instance->Strands.Data->GetNumPoints(); // Change this later on for having dynamic value
+		uint32 DataZ = 0;
+		uint32 DataW = 0;
+		switch (Instance->GeometryType)
+		{
+		case EHairGeometryType::Strands:
+			if (Instance->Strands.IsValid())
+			{
+				DataZ = Instance->Strands.Data->GetNumCurves(); // Change this later on for having dynamic value
+				DataW = Instance->Strands.Data->GetNumPoints(); // Change this later on for having dynamic value
+			}
+			break;
+		case EHairGeometryType::Cards:
+			if (Instance->Cards.IsValid(IntLODIndex))
+			{
+				DataZ = Instance->Cards.LODs[IntLODIndex].Guides.IsValid() ? Instance->Cards.LODs[IntLODIndex].Guides.Data->GetNumCurves() : 0;
+				DataW = Instance->Cards.LODs[IntLODIndex].Data->GetNumVertices();
+			}
+			break;
+		case EHairGeometryType::Meshes:
+			if (Instance->Meshes.IsValid(IntLODIndex))
+			{
+				DataZ = 0;
+				DataW = Instance->Meshes.LODs[IntLODIndex].Data->GetNumVertices();
+			}
+			break;
+		}
 
 		Infos.Add(FUintVector4(DataX, DataY, DataZ, DataW));
 	}
@@ -1014,25 +1042,56 @@ static void AddHairDebugPrintInstancePass(
 	FRDGBufferRef NameInfoBuffer = CreateStructuredBuffer(GraphBuilder, TEXT("Hair.Debug.InstanceNameInfos"), NameInfos);	
 	FRDGBufferRef InfoBuffer = CreateVertexBuffer(GraphBuilder, TEXT("Hair.Debug.InstanceInfos"), FRDGBufferDesc::CreateBufferDesc(InfoInBytes, Infos.Num()), Infos.GetData(), InfoInBytes * Infos.Num());
 
-	FHairDebugPrintInstanceCS::FParameters* Parameters = GraphBuilder.AllocParameters<FHairDebugPrintInstanceCS::FParameters>();
-	Parameters->InstanceCount = InstanceCount;
-	Parameters->NameInfoCount = NameInfos.Num();
-	Parameters->NameCharacterCount = Names.Num();
-	Parameters->Names = GraphBuilder.CreateSRV(NameBuffer, PF_R8_UINT);
-	Parameters->NameInfos = GraphBuilder.CreateSRV(NameInfoBuffer);
-	Parameters->Infos = GraphBuilder.CreateSRV(InfoBuffer, PF_R32G32B32A32_UINT);
-	ShaderPrint::SetParameters(GraphBuilder, *ShaderPrintData, Parameters->ShaderPrintUniformBuffer);
-	ShaderDrawDebug::SetParameters(GraphBuilder, *ShaderDrawData, Parameters->ShaderDrawUniformBuffer);
-	TShaderMapRef<FHairDebugPrintInstanceCS> ComputeShader(ShaderMap);
+	// Draw general information for all instances (one pass for all instances)
+	{
+		FHairDebugPrintInstanceCS::FParameters* Parameters = GraphBuilder.AllocParameters<FHairDebugPrintInstanceCS::FParameters>();
+		Parameters->InstanceCount = InstanceCount;
+		Parameters->NameInfoCount = NameInfos.Num();
+		Parameters->NameCharacterCount = Names.Num();
+		Parameters->Names = GraphBuilder.CreateSRV(NameBuffer, PF_R8_UINT);
+		Parameters->NameInfos = GraphBuilder.CreateSRV(NameInfoBuffer);
+		Parameters->Infos = GraphBuilder.CreateSRV(InfoBuffer, PF_R32G32B32A32_UINT);
+		ShaderPrint::SetParameters(GraphBuilder, *ShaderPrintData, Parameters->ShaderPrintUniformBuffer);
+		ShaderDrawDebug::SetParameters(GraphBuilder, *ShaderDrawData, Parameters->ShaderDrawUniformBuffer);
+		FHairDebugPrintInstanceCS::FPermutationDomain PermutationVector;
+		PermutationVector.Set<FHairDebugPrintInstanceCS::FOutputType>(0);
+		TShaderMapRef<FHairDebugPrintInstanceCS> ComputeShader(ShaderMap, PermutationVector);
+
+		ClearUnusedGraphResources(ComputeShader, Parameters);
+
+		FComputeShaderUtils::AddPass(
+			GraphBuilder,
+			RDG_EVENT_NAME("HairStrands::DebugPrintInstance(Info,Instances:%d)", InstanceCount),
+			ComputeShader,
+			Parameters,
+			FIntVector(1, 1, 1));
+	}
 
-	ClearUnusedGraphResources(ComputeShader, Parameters);
+	// Draw instances bound (one pass for each instance, due to separate AABB resources)
+	FHairDebugPrintInstanceCS::FPermutationDomain PermutationVector;
+	PermutationVector.Set<FHairDebugPrintInstanceCS::FOutputType>(1);
+	TShaderMapRef<FHairDebugPrintInstanceCS> ComputeShader(ShaderMap, PermutationVector);
+	for (uint32 InstanceIndex = 0; InstanceIndex < InstanceCount; ++InstanceIndex)
+	{
+		const FHairStrandsInstance* AbstractInstance = Instances[InstanceIndex];
+		const FHairGroupInstance* Instance = static_cast<const FHairGroupInstance*>(AbstractInstance);
 
-	FComputeShaderUtils::AddPass(
-		GraphBuilder,
-		RDG_EVENT_NAME("HairStrands::DebugPrintInstance(Instances:%d)", InstanceCount),
-		ComputeShader,
-		Parameters,
-		FIntVector(1, 1, 1));
+		if (Instance->GeometryType == EHairGeometryType::Strands)
+		{
+			FHairDebugPrintInstanceCS::FParameters* Parameters = GraphBuilder.AllocParameters<FHairDebugPrintInstanceCS::FParameters>();
+			Parameters->InstanceAABB = Register(GraphBuilder, Instance->HairGroupPublicData->GroupAABBBuffer, ERDGImportedBufferFlags::CreateSRV).SRV;
+			ShaderPrint::SetParameters(GraphBuilder, *ShaderPrintData, Parameters->ShaderPrintUniformBuffer);
+			ShaderDrawDebug::SetParameters(GraphBuilder, *ShaderDrawData, Parameters->ShaderDrawUniformBuffer);
+			ClearUnusedGraphResources(ComputeShader, Parameters);
+
+			FComputeShaderUtils::AddPass(
+				GraphBuilder,
+				RDG_EVENT_NAME("HairStrands::DebugPrintInstance(Bound)", InstanceCount),
+				ComputeShader,
+				Parameters,
+				FIntVector(1, 1, 1));
+		}
+	}
 }
 
 void RunHairStrandsDebug(
diff --git a/Engine/Shaders/Private/HairStrands/HairStrandsDebugPrint.usf b/Engine/Shaders/Private/HairStrands/HairStrandsDebugPrint.usf
--- a/Engine/Shaders/Private/HairStrands/HairStrandsDebugPrint.usf
+++ b/Engine/Shaders/Private/HairStrands/HairStrandsDebugPrint.usf
@@ -305,6 +305,7 @@ uint					NameInfoCount;
 uint					NameCharacterCount;
 StructuredBuffer<uint2> NameInfos;
 Buffer<uint>			Names;
+Buffer<int>				InstanceAABB;
 
 struct FNameInfo
 {
@@ -370,6 +371,7 @@ float2 Print(float2 Pos, float X, uint DigitCount, FFontColor Color)
 
 [numthreads(1, 1, 1)]
 void MainCS(uint GroupIndex : SV_GroupIndex, uint3 DispatchThreadId : SV_DispatchThreadID)
+#if PERMUTATION_OUTPUT_TYPE == 0
 {
 	// Draw AABB of the instance groups
 	if (any(DispatchThreadId != 0))
@@ -457,14 +459,14 @@ void MainCS(uint GroupIndex : SV_GroupIndex, uint3 DispatchThreadId : SV_Dispatc
 
 		// Sim
 		if (bHasSim)
-			Pos = ShaderPrintText(Pos, TEXT("True"), FontGreen);
+			Pos = ShaderPrintText(Pos, TEXT("True "), FontGreen);
 		else
 			Pos = ShaderPrintText(Pos, TEXT("False"), FontRed);
 		Pos = ShaderPrintText(Pos, TEXT(" "), FontWhite);
 
 		// RBF
 		if (bHasRBF)
-			Pos = ShaderPrintText(Pos, TEXT("True"), FontGreen);
+			Pos = ShaderPrintText(Pos, TEXT("True "), FontGreen);
 		else
 			Pos = ShaderPrintText(Pos, TEXT("False"), FontRed);
 		Pos = ShaderPrintText(Pos, TEXT(" "));
@@ -486,6 +488,15 @@ void MainCS(uint GroupIndex : SV_GroupIndex, uint3 DispatchThreadId : SV_Dispatc
 		AddNewLine(OriginalPos, Pos);
 	}
 }
+#else  // PERMUTATION_OUTPUT_TYPE
+{
+	if (all(DispatchThreadId == 0))
+	{
+		const FHairAABB AABB = ReadHairAABB(0, InstanceAABB);
+		AddAABBTWS(AABB.Min, AABB.Max, ColorDarkGreen);
+	}
+}
+#endif // PERMUTATION_OUTPUT_TYPE
 
 #endif // SHADER_PRINT_INSTANCE
 
diff --git a/Engine/Shaders/Private/ShaderDrawDebug.ush b/Engine/Shaders/Private/ShaderDrawDebug.ush
--- a/Engine/Shaders/Private/ShaderDrawDebug.ush
+++ b/Engine/Shaders/Private/ShaderDrawDebug.ush
@@ -23,6 +23,7 @@
 #define ColorSilver		float4(189.f / 255.f, 195.f / 255.f, 199.f / 255.f, 1)
 #define ColorEmerald	float4( 46.f / 255.f, 204.f / 255.f, 113.f / 255.f, 1)
 #define ColorLightGreen	float4(0.5f, 0.75f, 0.5f, 1.0f)
+#define ColorDarkGreen	float4(0.25f,0.50f, 0.25f, 1.0f)
 
 ///////////////////////////////////////////////////////////////////////////////////////////////////
 // Struct and core functions
