* Fix invalid groom component bound when attached to a skel. mesh.
* Fix missing stencil clear for hair shadow mask

The bound computation was causing the bounding box to be larger than needed.
The missing stencil clear was causing invalid shadow mask computation when several point/directional lights were present in the scene.

diff --git a/Engine/Plugins/Runtime/HairStrands/Source/HairStrandsCore/Private/GroomComponent.cpp b/Engine/Plugins/Runtime/HairStrands/Source/HairStrandsCore/Private/GroomComponent.cpp
--- a/Engine/Plugins/Runtime/HairStrands/Source/HairStrandsCore/Private/GroomComponent.cpp
+++ b/Engine/Plugins/Runtime/HairStrands/Source/HairStrandsCore/Private/GroomComponent.cpp
@@ -1860,11 +1860,22 @@ FBoxSphereBounds UGroomComponent::CalcBounds(const FTransform& InLocalToWorld) c
 		{
 			const FBox LocalSkeletalBound = RegisteredMeshComponent->CalcBounds(FTransform::Identity).GetBox();
 
+			// Compute an extra 'radius' which will be added to the skel. mesh bound. 
+			// It is based on the different of the groom bound and the skel. bound at rest position
+			float BoundExtraRadius = 0.f;
+			if (USkeletalMeshComponent* SkeletalMeshComponent = Cast<USkeletalMeshComponent>(RegisteredMeshComponent))
+			{
+				FBoxSphereBounds RestBounds;
+				SkeletalMeshComponent->GetPreSkinnedLocalBounds(RestBounds);
+				const FBox RestBox = RestBounds.GetBox();
+				const FVector3f MinDiff = LocalHairBound.Min - RestBox.Min;
+				const FVector3f MaxDiff = LocalHairBound.Max - RestBox.Max;
+				BoundExtraRadius = FMath::Max3(0.f, FMath::Max3(MinDiff.X, MinDiff.Y, MinDiff.Z), FMath::Max3(MaxDiff.X, MaxDiff.Y, MaxDiff.Z));
+			}
+			
 			FBox LocalBound(EForceInit::ForceInitToZero);
-			LocalBound += LocalSkeletalBound.Min;
-			LocalBound += LocalSkeletalBound.Max;
-			LocalBound += LocalHairBound.Min;
-			LocalBound += LocalHairBound.Max;
+			LocalBound += LocalSkeletalBound.Min - BoundExtraRadius;
+			LocalBound += LocalSkeletalBound.Max + BoundExtraRadius;
 			return FBoxSphereBounds(LocalBound.TransformBy(InLocalToWorld));
 		}
 		else
