Fix reset for the groom simulation

diff --git a/Engine/Plugins/Runtime/HairStrands/Source/HairStrandsCore/Private/Niagara/NiagaraDataInterfaceHairStrands.cpp b/Engine/Plugins/Runtime/HairStrands/Source/HairStrandsCore/Private/Niagara/NiagaraDataInterfaceHairStrands.cpp
--- a/Engine/Plugins/Runtime/HairStrands/Source/HairStrandsCore/Private/Niagara/NiagaraDataInterfaceHairStrands.cpp
+++ b/Engine/Plugins/Runtime/HairStrands/Source/HairStrandsCore/Private/Niagara/NiagaraDataInterfaceHairStrands.cpp
@@ -462,8 +462,12 @@ void FNDIHairStrandsData::Update(UNiagaraDataInterfaceHairStrands* Interface, FN
 	{
 		WorldTransform = LocalToWorld;
 
-		GlobalInterpolation = (Interface->IsComponentValid() && Interface->SourceComponent->BindingAsset && Interface->SourceComponent->GroomAsset) ?
-			Interface->SourceComponent->GroomAsset->EnableGlobalInterpolation : false;
+		const bool bHasValidBindingAsset = (Interface->IsComponentValid() && Interface->SourceComponent->BindingAsset && Interface->SourceComponent->GroomAsset);
+
+		GlobalInterpolation = bHasValidBindingAsset ? Interface->SourceComponent->GroomAsset->EnableGlobalInterpolation : false;
+		bSkinningTransfer = bHasValidBindingAsset ?
+			(Interface->SourceComponent->BindingAsset->SourceSkeletalMesh && Interface->SourceComponent->BindingAsset->TargetSkeletalMesh &&
+			 Interface->SourceComponent->BindingAsset->SourceSkeletalMesh != Interface->SourceComponent->BindingAsset->TargetSkeletalMesh) : false;
 		
 		TickingGroup = Interface->IsComponentValid() ? ComputeTickingGroup(Interface->SourceComponent) : NiagaraFirstTickGroup;
 
@@ -757,9 +761,11 @@ struct FNDIHairStrandsParametersCS : public FNiagaraDataInterfaceParametersCS
 			FShaderResourceViewRHIRef MeshSampleWeightsBufferSRV = (bHasSamples && DeformedMeshProjection) ?
 				DeformedMeshProjection->MeshSampleWeightsBuffer.SRV.GetReference() : FNiagaraRenderer::GetDummyFloatBuffer();
 
-			// Simulation setup
+		
+			// Simulation setup (we update the rest configuration based on the deformed positions 
+			// if in restupdate mode or if we are resetting the sim and using RBF transfer since the rest positions are not matrching the physics asset)
 			const int32 NeedResetValue = (ProxyData->TickCount <= GHairSimulationMaxDelay) || !HairStrandsBuffer->bValidGeometryType;
-			const int32 RestUpdateValue = GHairSimulationRestUpdate;
+			const int32 RestUpdateValue = GHairSimulationRestUpdate || (NeedResetValue && ProxyData->bSkinningTransfer);
 			const int32 LocalSimulationValue = ProxyData->LocalSimulation;
 
 			HairStrandsBuffer->bValidGeometryType = true;
diff --git a/Engine/Plugins/Runtime/HairStrands/Source/HairStrandsCore/Public/Niagara/NiagaraDataInterfaceHairStrands.h b/Engine/Plugins/Runtime/HairStrands/Source/HairStrandsCore/Public/Niagara/NiagaraDataInterfaceHairStrands.h
--- a/Engine/Plugins/Runtime/HairStrands/Source/HairStrandsCore/Public/Niagara/NiagaraDataInterfaceHairStrands.h
+++ b/Engine/Plugins/Runtime/HairStrands/Source/HairStrandsCore/Public/Niagara/NiagaraDataInterfaceHairStrands.h
@@ -114,6 +114,7 @@ struct FNDIHairStrandsData
 		PreviousBoneAngularVelocity = FVector3f::Zero();
 		
 		GlobalInterpolation = false;
+		bSkinningTransfer = false;
 		HairGroupInstance = nullptr;
 
 		TickCount = 0;
@@ -180,6 +181,7 @@ struct FNDIHairStrandsData
 			PreviousBoneAngularVelocity = OtherDatas->PreviousBoneAngularVelocity;
 
 			GlobalInterpolation = OtherDatas->GlobalInterpolation;
+			bSkinningTransfer = OtherDatas->bSkinningTransfer;
 			BindingType = OtherDatas->BindingType;
 			HairGroupInstance = OtherDatas->HairGroupInstance;
 
@@ -256,6 +258,9 @@ struct FNDIHairStrandsData
 
 	/** Global Interpolation */
 	bool GlobalInterpolation;
+	
+	/** Skinning transfer from a source to a target skelmesh */
+    bool bSkinningTransfer;
 
 	/** Number of strands*/
 	int32 NumStrands;
