Fix CIS issue.

diff --git a/Engine/Plugins/Runtime/HairStrands/Source/HairStrandsCore/Private/GroomDebug.cpp b/Engine/Plugins/Runtime/HairStrands/Source/HairStrandsCore/Private/GroomDebug.cpp
--- a/Engine/Plugins/Runtime/HairStrands/Source/HairStrandsCore/Private/GroomDebug.cpp
+++ b/Engine/Plugins/Runtime/HairStrands/Source/HairStrandsCore/Private/GroomDebug.cpp
@@ -942,7 +942,7 @@ static void AddHairDebugPrintInstancePass(
 	ShaderPrint::RequestSpaceForLines(InstanceCount * 16u);
 	ShaderPrint::RequestSpaceForCharacters(InstanceCount * 256 + 512);
 
-	if (!ShaderPrintData || !ShaderPrintData || InstanceCount == 0) { return; }
+	if (!ShaderPrintData || InstanceCount == 0) { return; }
 
 	const uint32 MaxPrimitiveNameCount = 128u;
 	check(sizeof(FHairDebugNameInfo) == 8);
