Fix CIS issue.

diff --git a/Engine/Plugins/Runtime/HairStrands/Source/HairStrandsCore/Private/GroomGeometryCache.cpp b/Engine/Plugins/Runtime/HairStrands/Source/HairStrandsCore/Private/GroomGeometryCache.cpp
--- a/Engine/Plugins/Runtime/HairStrands/Source/HairStrandsCore/Private/GroomGeometryCache.cpp
+++ b/Engine/Plugins/Runtime/HairStrands/Source/HairStrandsCore/Private/GroomGeometryCache.cpp
@@ -13,6 +13,7 @@
 #include "Components/SkeletalMeshComponent.h"
 #include "SkeletalRenderPublic.h"
 #include "Rendering/SkeletalMeshLODRenderData.h"
+#include "RenderGraphUtils.h"
 
 ///////////////////////////////////////////////////////////////////////////////////////////////////
 
