Fix for Hair with Simulation on Multi-GPU setups

diff --git a/Engine/Plugins/Runtime/HairStrands/Source/HairStrandsCore/Private/Niagara/NiagaraDataInterfaceHairStrands.cpp b/Engine/Plugins/Runtime/HairStrands/Source/HairStrandsCore/Private/Niagara/NiagaraDataInterfaceHairStrands.cpp
--- a/Engine/Plugins/Runtime/HairStrands/Source/HairStrandsCore/Private/Niagara/NiagaraDataInterfaceHairStrands.cpp
+++ b/Engine/Plugins/Runtime/HairStrands/Source/HairStrandsCore/Private/Niagara/NiagaraDataInterfaceHairStrands.cpp
@@ -4216,4 +4216,22 @@ void FNDIHairStrandsProxy::PreStage(FRHICommandList& RHICmdList, const FNiagaraD
 	}
 }
 
+void FNDIHairStrandsProxy::PostSimulate(FRHICommandList& RHICmdList, const FNiagaraDataInterfaceArgs& Context)
+{
+	FNDIHairStrandsProxy* InterfaceProxy = static_cast<FNDIHairStrandsProxy*>(Context.DataInterface);
+	FNDIHairStrandsData* ProxyData = InterfaceProxy->SystemInstancesToProxyData.Find(Context.SystemInstanceID);
+
+	const bool bIsHairValid = ProxyData != nullptr && ProxyData->HairStrandsBuffer && ProxyData->HairStrandsBuffer->IsInitialized();
+	const bool bIsDeformedValid = bIsHairValid && ProxyData->HairStrandsBuffer->SourceDeformedResources && ProxyData->HairStrandsBuffer->SourceDeformedResources->IsInitialized();
+
+	// MGPU DeformedPositionBuffer copy after simulation
+	if (bIsDeformedValid)
+	{
+		const FNDIHairStrandsBuffer* HairStrandsBuffer = ProxyData->HairStrandsBuffer;
+		FRHIBuffer* DeformedPositionBuffer = HairStrandsBuffer->SourceDeformedResources->DeformedPositionBuffer[HairStrandsBuffer->SourceDeformedResources->CurrentIndex].Buffer->GetRHI();
+
+		Context.ComputeDispatchInterface->MultiGPUResourceModified(RHICmdList, DeformedPositionBuffer, false, true);
+	}
+}
+
 #undef LOCTEXT_NAMESPACE
diff --git a/Engine/Plugins/Runtime/HairStrands/Source/HairStrandsCore/Public/Niagara/NiagaraDataInterfaceHairStrands.h b/Engine/Plugins/Runtime/HairStrands/Source/HairStrandsCore/Public/Niagara/NiagaraDataInterfaceHairStrands.h
--- a/Engine/Plugins/Runtime/HairStrands/Source/HairStrandsCore/Public/Niagara/NiagaraDataInterfaceHairStrands.h
+++ b/Engine/Plugins/Runtime/HairStrands/Source/HairStrandsCore/Public/Niagara/NiagaraDataInterfaceHairStrands.h
@@ -776,6 +776,9 @@ struct FNDIHairStrandsProxy : public FNiagaraDataInterfaceProxy
 	/** Launch all pre stage functions */
 	virtual void PreStage(FRHICommandList& RHICmdList, const FNiagaraDataInterfaceStageArgs& Context) override;
 
+	/** MGPU buffer copy after simulation*/
+	virtual void PostSimulate(FRHICommandList& RHICmdList, const FNiagaraDataInterfaceArgs& Context) override;
+
 	/** List of proxy data for each system instances*/
 	TMap<FNiagaraSystemInstanceID, FNDIHairStrandsData> SystemInstancesToProxyData;
 };
