Fix groom asset DDC cache issue.

The HairDescription*Group* struct wasn't updated when the HairDescription struct was updated. This caused the groom data being built & cached upon the 'old' HairDescription data. This CL ensure the HairDescriptionGroup struct is updated at the same time as the HairDescription.

This CL also rename HairDescription and HairDescriptionGroup with a 'Cached' prefix to highlight the fact these data structure are transient.
This CL also change groom asset DDC key so that previous invalid cached data are wiped out.

This issue is critical for MH & has been reported by several users.

diff --git a/Engine/Plugins/Runtime/HairStrands/Source/HairStrandsCore/Private/GroomAsset.cpp b/Engine/Plugins/Runtime/HairStrands/Source/HairStrandsCore/Private/GroomAsset.cpp
--- a/Engine/Plugins/Runtime/HairStrands/Source/HairStrandsCore/Private/GroomAsset.cpp
+++ b/Engine/Plugins/Runtime/HairStrands/Source/HairStrandsCore/Private/GroomAsset.cpp
@@ -538,26 +538,6 @@ static bool BuildHairGroupCluster(
 
 /////////////////////////////////////////////////////////////////////////////////////////
 
-#if WITH_EDITORONLY_DATA
-const FHairDescriptionGroups& UGroomAsset::GetHairDescriptionGroups()
-{
-	check(HairDescriptionBulkData);
-
-	if (!HairDescription)
-	{
-		HairDescription = MakeUnique<FHairDescription>();
-		HairDescriptionBulkData->LoadHairDescription(*HairDescription);
-	}
-	if (!HairDescriptionGroups)
-	{
-		HairDescriptionGroups = MakeUnique<FHairDescriptionGroups>();
-		FGroomBuilder::BuildHairDescriptionGroups(*HairDescription, *HairDescriptionGroups);
-	}
-
-	return *HairDescriptionGroups;
-}
-#endif // WITH_EDITORONLY_DATA
-
 uint8 UGroomAsset::GenerateClassStripFlags(FArchive& Ar)
 {
 #if WITH_EDITOR
@@ -1775,7 +1755,7 @@ void UGroomAsset::SetHairWidth(float Width)
 // differences, etc.) replace the version GUID below with a new one.
 // In case of merge conflicts with DDC versions, you MUST generate a new GUID
 // and set this new GUID as the version.
-#define GROOM_DERIVED_DATA_VERSION TEXT("6D16EF2A68094D16A06AE6A04F5D422E")
+#define GROOM_DERIVED_DATA_VERSION TEXT("5251F4CEC0B84A6B8175D0E33755556A")
 
 #if WITH_EDITORONLY_DATA
 
@@ -2073,13 +2053,36 @@ FString UGroomAsset::GetDerivedDataKeyForMeshes(uint32 GroupIndex)
 
 void UGroomAsset::CommitHairDescription(FHairDescription&& InHairDescription)
 {
-	HairDescription = MakeUnique<FHairDescription>(InHairDescription);
+	CachedHairDescription = MakeUnique<FHairDescription>(InHairDescription);
 
 	if (!HairDescriptionBulkData)
 	{
 		HairDescriptionBulkData = MakeUnique<FHairDescriptionBulkData>();
 	}
-	HairDescriptionBulkData->SaveHairDescription(*HairDescription);
+
+	// Update the cached hair description groups with the new hair description data
+	CachedHairDescriptionGroups = MakeUnique<FHairDescriptionGroups>();
+	FGroomBuilder::BuildHairDescriptionGroups(*CachedHairDescription, *CachedHairDescriptionGroups);
+
+	HairDescriptionBulkData->SaveHairDescription(*CachedHairDescription);
+}
+
+const FHairDescriptionGroups& UGroomAsset::GetHairDescriptionGroups()
+{
+	check(HairDescriptionBulkData);
+
+	if (!CachedHairDescription)
+	{
+		CachedHairDescription = MakeUnique<FHairDescription>();
+		HairDescriptionBulkData->LoadHairDescription(*CachedHairDescription);
+	}
+	if (!CachedHairDescriptionGroups)
+	{
+		CachedHairDescriptionGroups = MakeUnique<FHairDescriptionGroups>();
+		FGroomBuilder::BuildHairDescriptionGroups(*CachedHairDescription, *CachedHairDescriptionGroups);
+	}
+
+	return *CachedHairDescriptionGroups;
 }
 
 FHairDescription UGroomAsset::GetHairDescription() const
diff --git a/Engine/Plugins/Runtime/HairStrands/Source/HairStrandsCore/Public/GroomAsset.h b/Engine/Plugins/Runtime/HairStrands/Source/HairStrandsCore/Public/GroomAsset.h
--- a/Engine/Plugins/Runtime/HairStrands/Source/HairStrandsCore/Public/GroomAsset.h
+++ b/Engine/Plugins/Runtime/HairStrands/Source/HairStrandsCore/Public/GroomAsset.h
@@ -644,9 +644,13 @@ private:
 	const FHairDescriptionGroups& GetHairDescriptionGroups();
 	FString BuildDerivedDataKeySuffix(uint32 GroupIndex, const FHairGroupsInterpolation& InterpolationSettings, const FHairGroupsLOD& LODSettings) const;
 	bool IsFullyCached();
-	TUniquePtr<FHairDescription> HairDescription;
 	TUniquePtr<FHairDescriptionBulkData> HairDescriptionBulkData;
-	TUniquePtr<FHairDescriptionGroups> HairDescriptionGroups;
+
+	// Transient HairDescription & HairDescriptionGroups, which are built from HairDescriptionBulkData.
+	// All these data (bulk/desc/groups) needs to be in sync. I.e., when the HairDescription is updated, 
+	// HairDescriptionGroups needs to also be updated
+	TUniquePtr<FHairDescription> CachedHairDescription;
+	TUniquePtr<FHairDescriptionGroups> CachedHairDescriptionGroups;
 
 	TArray<FString> StrandsDerivedDataKey;
 	TArray<FString> CardsDerivedDataKey;
