Fix Attribute and RootUV transfer from strands to cards.

diff --git a/Engine/Plugins/Runtime/HairStrands/Source/HairStrandsCore/Private/GroomBuilder.cpp b/Engine/Plugins/Runtime/HairStrands/Source/HairStrandsCore/Private/GroomBuilder.cpp
--- a/Engine/Plugins/Runtime/HairStrands/Source/HairStrandsCore/Private/GroomBuilder.cpp
+++ b/Engine/Plugins/Runtime/HairStrands/Source/HairStrandsCore/Private/GroomBuilder.cpp
@@ -2308,7 +2308,7 @@ void FGroomBuilder::VoxelizeGroupIndex(const FHairDescriptionGroups& In, FHairSt
 			Out.Datas[I].Roughness += Curve.Roughness;
 		}
 
-		const uint32 CurveCount = Grid.Voxels[I].VoxelCurves.Num() > 0;
+		const uint32 CurveCount = Grid.Voxels[I].VoxelCurves.Num();
 		if (CurveCount > 0)
 		{
 			Out.Datas[I].BaseColor /= float(CurveCount);
diff --git a/Engine/Plugins/Runtime/HairStrands/Source/HairStrandsCore/Private/HairCardsBuilder.cpp b/Engine/Plugins/Runtime/HairStrands/Source/HairStrandsCore/Private/HairCardsBuilder.cpp
--- a/Engine/Plugins/Runtime/HairStrands/Source/HairStrandsCore/Private/HairCardsBuilder.cpp
+++ b/Engine/Plugins/Runtime/HairStrands/Source/HairStrandsCore/Private/HairCardsBuilder.cpp
@@ -3335,7 +3335,7 @@ namespace FHairCardsBuilder
 FString GetVersion()
 {
 	// Important to update the version when cards building or importing changes
-	return TEXT("9d");
+	return TEXT("9e");
 }
 
 void AllocateAtlasTexture(UTexture2D* Out, const FIntPoint& Resolution, uint32 MipCount, EPixelFormat PixelFormat, ETextureSourceFormat SourceFormat)
@@ -3451,7 +3451,7 @@ void BuildGeometry(
 
 		AllocateAtlasTexture(static_cast<UTexture2D*&>(OutTextures.DepthTexture),		Out.Atlas.Resolution, LODName, TEXT("_CardsAtlas_Depth"), AllocateAtlasTexture_Depth);
 		AllocateAtlasTexture(static_cast<UTexture2D*&>(OutTextures.CoverageTexture),	Out.Atlas.Resolution, LODName, TEXT("_CardsAtlas_Coverage"), AllocateAtlasTexture_Coverage);
-		AllocateAtlasTexture(static_cast<UTexture2D*&>(OutTextures.TangentTexture),	Out.Atlas.Resolution, LODName, TEXT("_CardsAtlas_Tangent"), AllocateAtlasTexture_Tangent);
+		AllocateAtlasTexture(static_cast<UTexture2D*&>(OutTextures.TangentTexture),		Out.Atlas.Resolution, LODName, TEXT("_CardsAtlas_Tangent"), AllocateAtlasTexture_Tangent);
 		AllocateAtlasTexture(static_cast<UTexture2D*&>(OutTextures.AttributeTexture),	Out.Atlas.Resolution, LODName, TEXT("_CardsAtlas_Attribute"), AllocateAtlasTexture_Attribute);
 	}
 
@@ -3898,6 +3898,7 @@ bool InternalImportGeometry(
 				// /!\ N^2 loop: the number of cards should be relatively small
 				uint32 CurveIndex = ~0;
 				float ClosestDistance = FLT_MAX;
+				FVector2f RootUV = FVector2f::ZeroVector;
 				for (const FStrandsRootData& StrandsRoot : StrandsRoots)
 				{
 					const float Distance = FVector3f::Distance(StrandsRoot.Position, CardsRoot.Position);
@@ -3905,13 +3906,12 @@ bool InternalImportGeometry(
 					{
 						ClosestDistance = Distance;
 						CurveIndex = StrandsRoot.CurveIndex;
+						RootUV = StrandsRoot.RootUV;
 					}
 				}
 
 				// 3.2 Apply root UV to all cards vertices
 				{
-					FVector2f RootUV = FVector2f::ZeroVector;
-
 					const uint32 CardsIndexCount = Out.Cards.IndexCounts[CardsRoot.CardsIndex];
 					const uint32 CardsIndexOffset = Out.Cards.IndexOffsets[CardsRoot.CardsIndex];
 					for (uint32 IndexIt = 0; IndexIt < CardsIndexCount; ++IndexIt)
