Removed a check from Hair Mesh interpolation code path - which prevents an editor crash - and is consistent with Hair Strand and Hair Card interpolation code paths.

diff --git a/Engine/Plugins/Runtime/HairStrands/Source/HairStrandsCore/Private/HairStrandsRendering.cpp b/Engine/Plugins/Runtime/HairStrands/Source/HairStrandsCore/Private/HairStrandsRendering.cpp
--- a/Engine/Plugins/Runtime/HairStrands/Source/HairStrandsCore/Private/HairStrandsRendering.cpp
+++ b/Engine/Plugins/Runtime/HairStrands/Source/HairStrandsCore/Private/HairStrandsRendering.cpp
@@ -2176,7 +2176,6 @@ void ComputeHairStrandsInterpolation(
 				check(Instance->Guides.IsValid());
 				check(Instance->Guides.HasValidRootData());
 				check(Instance->Guides.DeformedRootResource);
-				check(Instance->Guides.DeformedRootResource->IsValid(MeshLODIndex));
 
 				AddHairMeshesRBFInterpolationPass(
 					GraphBuilder,
