Re-adding check to Mesh hair interpolation based on feedback. The check has been modified to not fire off all the time in the edge case where the skin cache has been disabled in the project settings.

diff --git a/Engine/Plugins/Runtime/HairStrands/Source/HairStrandsCore/Private/HairStrandsRendering.cpp b/Engine/Plugins/Runtime/HairStrands/Source/HairStrandsCore/Private/HairStrandsRendering.cpp
--- a/Engine/Plugins/Runtime/HairStrands/Source/HairStrandsCore/Private/HairStrandsRendering.cpp
+++ b/Engine/Plugins/Runtime/HairStrands/Source/HairStrandsCore/Private/HairStrandsRendering.cpp
@@ -2176,6 +2176,7 @@ void ComputeHairStrandsInterpolation(
 				check(Instance->Guides.IsValid());
 				check(Instance->Guides.HasValidRootData());
 				check(Instance->Guides.DeformedRootResource);
+				check(MeshLODIndex == -1 || Instance->Guides.DeformedRootResource->IsValid(MeshLODIndex)); //MeshLODIndex -1 indicates that skin cache is disabled and this is a workaround to prevent the editor from crashing - an editor setting guildeline popup exists to inform the user that the skin cache should be enabled. 
 
 				AddHairMeshesRBFInterpolationPass(
 					GraphBuilder,
