Fix crash when undo operation on groom component within a blueprint.

diff --git a/Engine/Plugins/Runtime/HairStrands/Source/HairStrandsCore/Private/GroomComponent.cpp b/Engine/Plugins/Runtime/HairStrands/Source/HairStrandsCore/Private/GroomComponent.cpp
--- a/Engine/Plugins/Runtime/HairStrands/Source/HairStrandsCore/Private/GroomComponent.cpp
+++ b/Engine/Plugins/Runtime/HairStrands/Source/HairStrandsCore/Private/GroomComponent.cpp
@@ -113,7 +113,7 @@ static EHairInterpolationType ToHairInterpolationType(EGroomInterpolationType In
 
 static FHairGroupDesc GetGroomGroupsDesc(const UGroomAsset* Asset, UGroomComponent* Component, uint32 GroupIndex)
 {
-	if (!Asset)
+	if (!Asset || GroupIndex >= uint32(Component->GroomGroupsDesc.Num()))
 	{
 		return FHairGroupDesc();
 	}
