#ue5 -  LWC: FVector4f <-> FVector4d conversion is now explicit.

FLinearColor has also been modified to make any double->float conversions explicit. Previously all 3D TVector versions were allowed to be implicit and thus could convert TVector<double> => FLinearColor => TVector4<float>.

Fixed up all engine and game casts. Added "//LWC_TODO: precision loss" around any explicit casts that previously were silently explicit and we may need to revisit for precision loss analysis.

diff --git a/Engine/Plugins/Runtime/HairStrands/Source/HairStrandsCore/Private/HairCardsBuilder.cpp b/Engine/Plugins/Runtime/HairStrands/Source/HairStrandsCore/Private/HairCardsBuilder.cpp
--- a/Engine/Plugins/Runtime/HairStrands/Source/HairStrandsCore/Private/HairCardsBuilder.cpp
+++ b/Engine/Plugins/Runtime/HairStrands/Source/HairStrandsCore/Private/HairCardsBuilder.cpp
@@ -3466,7 +3466,7 @@ void BuildGeometry(
 	{
 		check(FMath::IsFinite(Out.Cards.Positions[PointIt].X) && FMath::IsFinite(Out.Cards.Positions[PointIt].Y) && FMath::IsFinite(Out.Cards.Positions[PointIt].Z));
 		Out.RenderData.Positions[PointIt] = FVector4f(Out.Cards.Positions[PointIt], *((float*)&Out.Cards.CardIndices[PointIt]));
-		Out.RenderData.UVs[PointIt] = Out.Cards.UVs[PointIt];
+		Out.RenderData.UVs[PointIt] = FVector4f(Out.Cards.UVs[PointIt]);
 		Out.RenderData.Normals[PointIt * 2] = FVector4f(Out.Cards.Tangents[PointIt], 0);
 		Out.RenderData.Normals[PointIt * 2 + 1] = FVector4f(Out.Cards.Normals[PointIt], 1);
 
@@ -3539,13 +3539,13 @@ void BuildGeometry(
 	Out.RenderData.CardsStrandsPositions.SetNum(StrandsPointCount);
 	for (uint32 PointIt = 0; PointIt < StrandsPointCount; ++PointIt)
 	{
-		Out.RenderData.CardsStrandsPositions[PointIt] = Out.Atlas.StrandsPositions[PointIt];
+		Out.RenderData.CardsStrandsPositions[PointIt] = FVector4f(Out.Atlas.StrandsPositions[PointIt]);
 	}
 
 	Out.RenderData.CardsStrandsAttributes.SetNum(StrandsPointCount);
 	for (uint32 PointIt = 0; PointIt < StrandsPointCount; ++PointIt)
 	{
-		Out.RenderData.CardsStrandsAttributes[PointIt] = Out.Atlas.StrandsAttributes[PointIt];
+		Out.RenderData.CardsStrandsAttributes[PointIt] = FVector4f(Out.Atlas.StrandsAttributes[PointIt]);
 	}
 
 	// Interpolation data
@@ -3716,7 +3716,7 @@ bool InternalImportGeometry(
 	for (uint32 PointIt = 0; PointIt < PointCount; ++PointIt)
 	{
 		OutBulk.Positions[PointIt] = FVector4f(Out.Cards.Positions[PointIt], 0);
-		OutBulk.UVs[PointIt] = Out.Cards.UVs[PointIt];
+		OutBulk.UVs[PointIt] = FVector4f(Out.Cards.UVs[PointIt]);
 		OutBulk.Normals[PointIt * 2] = FVector4f(Out.Cards.Tangents[PointIt], 0);
 		OutBulk.Normals[PointIt * 2 + 1] = FVector4f(Out.Cards.Normals[PointIt], TangentFrameSigns[PointIt]);
 		OutBulk.Materials[PointIt] = 0;
diff --git a/Engine/Source/Editor/UnrealEd/Private/Fbx/FbxMainExport.cpp b/Engine/Source/Editor/UnrealEd/Private/Fbx/FbxMainExport.cpp
--- a/Engine/Source/Editor/UnrealEd/Private/Fbx/FbxMainExport.cpp
+++ b/Engine/Source/Editor/UnrealEd/Private/Fbx/FbxMainExport.cpp
@@ -4884,7 +4884,7 @@ void FFbxExporter::ExportSplineMeshToFbx(const USplineMeshComponent* SplineMeshC
 	{
 		FVector Position = (FVector)RenderMesh.VertexBuffers.PositionVertexBuffer.VertexPosition(VertIndex);
 		const FTransform SliceTransform = SplineMeshComp->CalcSliceTransform(USplineMeshComponent::GetAxisValue(Position, SplineMeshComp->ForwardAxis));
-		FVector Normal = FVector(RenderMesh.VertexBuffers.StaticMeshVertexBuffer.VertexTangentZ(VertIndex));
+		FVector Normal = FVector4(RenderMesh.VertexBuffers.StaticMeshVertexBuffer.VertexTangentZ(VertIndex));
 		Normal = SliceTransform.TransformVector(Normal);
 		FbxVector4& FbxNormal = FbxNormals[VertIndex];
 		FbxNormal = FbxVector4(Normal.X, -Normal.Y, Normal.Z);
diff --git a/Engine/Source/Runtime/Renderer/Private/LightRendering.h b/Engine/Source/Runtime/Renderer/Private/LightRendering.h
--- a/Engine/Source/Runtime/Renderer/Private/LightRendering.h
+++ b/Engine/Source/Runtime/Renderer/Private/LightRendering.h
@@ -154,7 +154,7 @@ namespace StencilingGeometry
 			for (int32 i = 0; i < NumRings + 1; i++)
 			{
 				const float Angle = i * RadiansPerRingSegment;
-				ArcVerts.Add(UE::Math::TVector<float>(0.0f, FMath::Sin(Angle), FMath::Cos(Angle)));
+				ArcVerts.Add(FVector3f(0.0f, FMath::Sin(Angle), FMath::Cos(Angle)));
 			}
 
 			TResourceArray<VectorType, VERTEXBUFFER_ALIGNMENT> Verts;
