Fix missing RootUV from groom generated texture attribute.

diff --git a/Engine/Plugins/Runtime/HairStrands/Source/HairStrandsCore/Private/GroomTextureBuilder.cpp b/Engine/Plugins/Runtime/HairStrands/Source/HairStrandsCore/Private/GroomTextureBuilder.cpp
--- a/Engine/Plugins/Runtime/HairStrands/Source/HairStrandsCore/Private/GroomTextureBuilder.cpp
+++ b/Engine/Plugins/Runtime/HairStrands/Source/HairStrandsCore/Private/GroomTextureBuilder.cpp
@@ -626,6 +626,7 @@ class FHairStrandsTexturePS : public FGlobalShader
 		SHADER_PARAMETER(FVector3f, InVF_PositionOffset)
 		SHADER_PARAMETER_SRV(Buffer, InVF_PositionBuffer)
 		SHADER_PARAMETER_SRV(Buffer, InVF_Attribute0Buffer)
+		SHADER_PARAMETER_SRV(Buffer, InVF_Attribute1Buffer)
 		SHADER_PARAMETER_SRV(Buffer, InVF_MaterialBuffer)
 		SHADER_PARAMETER(uint32, InVF_ControlPointCount)
 		SHADER_PARAMETER(uint32, InVF_GroupIndex)
@@ -686,7 +687,8 @@ static void InternalGenerateHairStrandsTextures(
 	FRDGBufferRef VoxelData,
 	
 	FRHIShaderResourceView* InHairStrands_PositionBuffer,
-	FRHIShaderResourceView* InHairStrands_AttributeBuffer,
+	FRHIShaderResourceView* InHairStrands_Attribute0Buffer,
+	FRHIShaderResourceView* InHairStrands_Attribute1Buffer,
 	FRHIShaderResourceView* InHairStrands_MaterialBuffer,
 	const FVector& InHairStrands_PositionOffset,
 	float InHairStrands_Radius,
@@ -720,8 +722,9 @@ static void InternalGenerateHairStrandsTextures(
 
 	ParametersPS->InVF_PositionBuffer = InHairStrands_PositionBuffer;
 	ParametersPS->InVF_PositionOffset = (FVector3f)InHairStrands_PositionOffset;
-	ParametersPS->InVF_Attribute0Buffer = InHairStrands_AttributeBuffer;
-	ParametersPS->InVF_MaterialBuffer = bHasMaterialData ? InHairStrands_MaterialBuffer : InHairStrands_AttributeBuffer;
+	ParametersPS->InVF_Attribute0Buffer = InHairStrands_Attribute0Buffer;
+	ParametersPS->InVF_Attribute1Buffer = InHairStrands_Attribute1Buffer;
+	ParametersPS->InVF_MaterialBuffer = bHasMaterialData ? InHairStrands_MaterialBuffer : InHairStrands_Attribute0Buffer;
 	ParametersPS->InVF_Radius = InHairStrands_Radius;
 	ParametersPS->InVF_Length = InHairStrands_Length;
 	ParametersPS->InVF_ControlPointCount = InHairStrands_ControlPointCount;
@@ -1051,6 +1054,7 @@ static void InternalBuildStrandsTextures_GPU(
 
 				GroupData.Strands.RestResource->PositionBuffer.SRV,
 				GroupData.Strands.RestResource->Attribute0Buffer.SRV,
+				GroupData.Strands.RestResource->Attribute1Buffer.SRV,
 				GroupData.Strands.RestResource->MaterialBuffer.SRV,
 				GroupData.Strands.RestResource->GetPositionOffset(),
 				RenderingData.GeometrySettings.HairWidth * 0.5f,
diff --git a/Engine/Shaders/Private/HairStrands/HairStrandsTexturesGeneration.usf b/Engine/Shaders/Private/HairStrands/HairStrandsTexturesGeneration.usf
--- a/Engine/Shaders/Private/HairStrands/HairStrandsTexturesGeneration.usf
+++ b/Engine/Shaders/Private/HairStrands/HairStrandsTexturesGeneration.usf
@@ -85,6 +85,7 @@ uint bHasMaterialData;
 
 Buffer<uint4> InVF_PositionBuffer;
 Buffer<float2> InVF_Attribute0Buffer;
+Buffer<uint> InVF_Attribute1Buffer;
 Buffer<float4> InVF_MaterialBuffer;
 
 float3 InVF_PositionOffset;
@@ -295,10 +296,8 @@ bool Trace(
 							OutTangent = Tangent;
 							OutDepth = Distance;
 
-
-							const uint4 Attribute1 = 0;
-							const FHairAttributes Attributes0 = DecodeHairAttribute(InVF_Attribute0Buffer[P0], Attribute1, InVF_Length);
-							const FHairAttributes Attributes1 = DecodeHairAttribute(InVF_Attribute0Buffer[P1], Attribute1, InVF_Length);
+							const FHairAttributes Attributes0 = DecodeHairAttribute(InVF_Attribute0Buffer[P0], InVF_Attribute1Buffer[P0], InVF_Length);
+							const FHairAttributes Attributes1 = DecodeHairAttribute(InVF_Attribute0Buffer[P1], InVF_Attribute1Buffer[P0], InVF_Length);
 
 							OutSeed		= Attributes0.Seed;
 							OutRootUV	= Attributes0.RootUV;
