Fix hair strands texture readback, being sometime incomplete.

This CL change the read to use FRHIGPUTextureReadback and do the copy to the final UTexture on the frame later.

diff --git a/Engine/Plugins/Runtime/HairStrands/Source/HairStrandsCore/Private/GroomTextureBuilder.cpp b/Engine/Plugins/Runtime/HairStrands/Source/HairStrandsCore/Private/GroomTextureBuilder.cpp
--- a/Engine/Plugins/Runtime/HairStrands/Source/HairStrandsCore/Private/GroomTextureBuilder.cpp
+++ b/Engine/Plugins/Runtime/HairStrands/Source/HairStrandsCore/Private/GroomTextureBuilder.cpp
@@ -831,64 +831,31 @@ static void InternalGenerateHairStrandsTextures(
 		});
 }
 
-static void AddReadbackPass(
-	FRDGBuilder& GraphBuilder,
-	uint32 BytePerPixel,
-	FRDGTextureRef InputTexture,
-	UTexture2D* OutTexture)
+struct FStrandsTexturesReadback
 {
-	AddReadbackTexturePass(
-		GraphBuilder,
-		RDG_EVENT_NAME("CopyRDGToTexture2D"),
-		InputTexture,
-	[InputTexture, OutTexture, BytePerPixel](FRHICommandListImmediate& RHICmdList)
+	TUniquePtr<FRHIGPUTextureReadback> Depth;
+	TUniquePtr<FRHIGPUTextureReadback> Tangent;
+	TUniquePtr<FRHIGPUTextureReadback> Coverage;
+	TUniquePtr<FRHIGPUTextureReadback> Attribute;
+	TUniquePtr<FRHIGPUTextureReadback> Material;
+
+	FStrandsTexturesOutput Output;
+	bool IsReady() const
 	{
-		const FIntPoint Resolution = InputTexture->Desc.Extent;
-		check(OutTexture);
-
-		FRHIResourceCreateInfo CreateInfo(TEXT("ReadbackPass_StagingTexture"));
-		FTexture2DRHIRef StagingTexture = RHICreateTexture2D(
-			InputTexture->Desc.Extent.X,
-			InputTexture->Desc.Extent.Y,
-			InputTexture->Desc.Format,
-			InputTexture->Desc.NumMips,
-			1, TexCreate_CPUReadback, CreateInfo);
-
-		FRHICopyTextureInfo CopyInfo;
-		CopyInfo.NumMips = InputTexture->Desc.NumMips;
-		RHICmdList.CopyTexture(
-			InputTexture->GetRHI(),
-			StagingTexture->GetTexture2D(),
-			CopyInfo);
-
-		// Flush, to ensure that all texture generation is done
-		GDynamicRHI->RHISubmitCommandsAndFlushGPU();
-		GDynamicRHI->RHIBlockUntilGPUIdle();
-
-		// don't think we can mark package as a dirty in the package build
-	#if WITH_EDITORONLY_DATA
-		check(OutTexture->Source.GetSizeX() == Resolution.X);
-		check(OutTexture->Source.GetSizeY() == Resolution.Y);
-
-		void* InData = nullptr;
-		int32 Width = 0, Height = 0;
-		RHICmdList.MapStagingSurface(StagingTexture, InData, Width, Height);
-		uint8* InDataRGBA8 = (uint8*)InData;
-		ETextureSourceFormat SrcFormat = OutTexture->Source.GetFormat();
-		OutTexture->Source.Init(Width, Height, 1, 1, SrcFormat, InDataRGBA8);
-		RHICmdList.UnmapStagingSurface(StagingTexture);
-
-		OutTexture->DeferCompression = true; // This forces reloading data when the asset is saved
-		OutTexture->MarkPackageDirty();
-	#endif // WITH_EDITORONLY_DATA
-	});
-}
+		return
+			Depth->IsReady() &&
+			Tangent->IsReady() &&
+			Coverage->IsReady() &&
+			Attribute->IsReady() &&
+			Material->IsReady();
+	}
+};
 
 static void InternalBuildStrandsTextures_GPU(
 	FRDGBuilder& GraphBuilder,
 	FGlobalShaderMap* ShaderMap,
 	const FStrandsTexturesInfo& InInfo,
-	const FStrandsTexturesOutput& Output,
+	FStrandsTexturesReadback& Output,
 	const struct FShaderPrintData* DebugShaderData)
 {
 	USkeletalMesh* SkeletalMesh = (USkeletalMesh*)InInfo.SkeletalMesh;
@@ -1102,11 +1069,40 @@ static void InternalBuildStrandsTextures_GPU(
 		SourceIndex = TargetIndex;
 	}
 
-	AddReadbackPass(GraphBuilder, 2, DepthTexture[TargetIndex]		, Output.Depth);
-	AddReadbackPass(GraphBuilder, 4, CoverageTexture[TargetIndex]	, Output.Coverage);
-	AddReadbackPass(GraphBuilder, 4, TangentTexture[TargetIndex]	, Output.Tangent);
-	AddReadbackPass(GraphBuilder, 4, AttributeTexture[TargetIndex]	, Output.Attribute);
-	AddReadbackPass(GraphBuilder, 4, MaterialTexture[TargetIndex]	, Output.Material);
+	AddEnqueueCopyPass(GraphBuilder, Output.Depth.Get(), DepthTexture[TargetIndex]);
+	AddEnqueueCopyPass(GraphBuilder, Output.Coverage.Get(), CoverageTexture[TargetIndex]);
+	AddEnqueueCopyPass(GraphBuilder, Output.Tangent.Get(), TangentTexture[TargetIndex]);
+	AddEnqueueCopyPass(GraphBuilder, Output.Attribute.Get(), AttributeTexture[TargetIndex]);
+	AddEnqueueCopyPass(GraphBuilder, Output.Material.Get(), MaterialTexture[TargetIndex]);
+}
+
+void CopyReadbackToTexture(TUniquePtr<FRHIGPUTextureReadback>& In, UTexture2D* Out)
+{
+#if WITH_EDITORONLY_DATA
+	FRHICommandListImmediate& RHICmdList = FRHICommandListExecutor::GetImmediateCommandList();
+
+	uint32 BytePerPixels = 0;
+	uint32 Width = Out->Source.GetSizeX();
+	uint32 Height = Out->Source.GetSizeY();
+
+	ETextureSourceFormat SrcFormat = Out->Source.GetFormat();
+	switch (SrcFormat)
+	{
+	case TSF_BGRA8:
+		BytePerPixels = 4; break;
+	case TSF_G16:
+		BytePerPixels = 2; break;
+	};
+
+	int32 RowPitchInPixels = 0;
+	void* LockedData = In->Lock(RowPitchInPixels); // This forces a GPU stall
+	uint8* InDataRGBA8 = (uint8*)LockedData;
+	Out->Source.Init(Width, Height, 1, 1, SrcFormat, InDataRGBA8);
+	In->Unlock();
+
+	Out->DeferCompression = true; // This forces reloading data when the asset is saved
+	Out->MarkPackageDirty();
+#endif // WITH_EDITORONLY_DATA
 }
 
 ///////////////////////////////////////////////////////////////////////////////////////////////////
@@ -1118,18 +1114,62 @@ struct FStrandsTexturesQuery
 	FStrandsTexturesOutput Output;
 };
 TQueue<FStrandsTexturesQuery> GStrandsTexturesQueries;
+TQueue<FStrandsTexturesReadback*> GStrandsTexturesReadbacks;
 
 bool HasHairStrandsTexturesQueries()
 {
-	return !GStrandsTexturesQueries.IsEmpty();
+	return !GStrandsTexturesQueries.IsEmpty() || !GStrandsTexturesReadbacks.IsEmpty();
 }
 
 void RunHairStrandsTexturesQueries(FRDGBuilder& GraphBuilder, FGlobalShaderMap* ShaderMap, const struct FShaderPrintData* DebugShaderData)
 {
-	FStrandsTexturesQuery Q;
-	while (GStrandsTexturesQueries.Dequeue(Q))
+	// Readback data
+	{
+		FStrandsTexturesReadback* R = nullptr;
+		while (GStrandsTexturesReadbacks.Dequeue(R))
+		{
+			if (R)
+			{
+				if (R->IsReady())
+				{
+					CopyReadbackToTexture(R->Depth,    R->Output.Depth);
+					CopyReadbackToTexture(R->Tangent,  R->Output.Tangent);
+					CopyReadbackToTexture(R->Coverage, R->Output.Coverage);
+					CopyReadbackToTexture(R->Attribute,R->Output.Attribute);
+					CopyReadbackToTexture(R->Material, R->Output.Material);
+
+					R->Depth    = nullptr;
+					R->Tangent  = nullptr;
+					R->Coverage = nullptr;
+					R->Attribute= nullptr;
+					R->Material = nullptr;
+
+					delete R;
+				}
+				else
+				{
+					GStrandsTexturesReadbacks.Enqueue(R);
+				}
+			}
+		}
+	}
+
+	// Render textures
 	{
-		InternalBuildStrandsTextures_GPU(GraphBuilder, ShaderMap, Q.Info, Q.Output, DebugShaderData);
+		FStrandsTexturesQuery Q;
+		while (GStrandsTexturesQueries.Dequeue(Q))
+		{
+			FStrandsTexturesReadback* R = new FStrandsTexturesReadback();
+			R->Depth     = MakeUnique<FRHIGPUTextureReadback>(TEXT("Readback.Depth"));
+			R->Tangent   = MakeUnique<FRHIGPUTextureReadback>(TEXT("Readback.Tangent"));
+			R->Coverage  = MakeUnique<FRHIGPUTextureReadback>(TEXT("Readback.Coverage"));
+			R->Attribute = MakeUnique<FRHIGPUTextureReadback>(TEXT("Readback.Attribute"));
+			R->Material  = MakeUnique<FRHIGPUTextureReadback>(TEXT("Readback.Material"));		
+			R->Output    = Q.Output;
+
+			InternalBuildStrandsTextures_GPU(GraphBuilder, ShaderMap, Q.Info, *R, DebugShaderData);
+			GStrandsTexturesReadbacks.Enqueue(R);
+		}
 	}
 }
 
