Fix hair collision components

diff --git a/Engine/Plugins/Runtime/HairStrands/Source/HairStrandsCore/Private/Niagara/NiagaraDataInterfacePhysicsAsset.cpp b/Engine/Plugins/Runtime/HairStrands/Source/HairStrandsCore/Private/Niagara/NiagaraDataInterfacePhysicsAsset.cpp
--- a/Engine/Plugins/Runtime/HairStrands/Source/HairStrandsCore/Private/Niagara/NiagaraDataInterfacePhysicsAsset.cpp
+++ b/Engine/Plugins/Runtime/HairStrands/Source/HairStrandsCore/Private/Niagara/NiagaraDataInterfacePhysicsAsset.cpp
@@ -823,7 +823,7 @@ void UNiagaraDataInterfacePhysicsAsset::ExtractSourceComponent(FNiagaraSystemIns
 			}
 			for(auto& CollisionComponent : GroomComponent->CollisionComponents)
 			{
-				if(CollisionComponent.IsValid())
+				if(CollisionComponent.IsValid() && CollisionComponent->GetPhysicsAsset())
 				{
 					SourceComponents.Add(CollisionComponent);
 					PhysicsAssets.Add(CollisionComponent->GetPhysicsAsset());
@@ -836,13 +836,17 @@ void UNiagaraDataInterfacePhysicsAsset::ExtractSourceComponent(FNiagaraSystemIns
 	if (SourceComponent != nullptr)
 	{
 		SourceComponents.Add(SourceComponent);
-		if (GroomPhysicsAsset)
+		if (GroomPhysicsAsset || SourceComponent->GetPhysicsAsset())
 		{
-			PhysicsAssets.Add(GroomPhysicsAsset);
-		}
-		else
-		{
-			PhysicsAssets.Add(SourceComponent->GetPhysicsAsset());
+			SourceComponents.Add(SourceComponent);
+			if (GroomPhysicsAsset)
+			{
+				PhysicsAssets.Add(GroomPhysicsAsset);
+			}
+			else
+			{
+				PhysicsAssets.Add(SourceComponent->GetPhysicsAsset());
+			}
 		}
 	}
 	else if (GroomPhysicsAsset != nullptr)
