Fix bad filter collision for hairs

diff --git a/Engine/Plugins/Runtime/HairStrands/Source/HairStrandsCore/Private/Niagara/NiagaraDataInterfacePhysicsAsset.cpp b/Engine/Plugins/Runtime/HairStrands/Source/HairStrandsCore/Private/Niagara/NiagaraDataInterfacePhysicsAsset.cpp
--- a/Engine/Plugins/Runtime/HairStrands/Source/HairStrandsCore/Private/Niagara/NiagaraDataInterfacePhysicsAsset.cpp
+++ b/Engine/Plugins/Runtime/HairStrands/Source/HairStrandsCore/Private/Niagara/NiagaraDataInterfacePhysicsAsset.cpp
@@ -835,7 +835,6 @@ void UNiagaraDataInterfacePhysicsAsset::ExtractSourceComponent(FNiagaraSystemIns
 	
 	if (SourceComponent != nullptr)
 	{
-		SourceComponents.Add(SourceComponent);
 		if (GroomPhysicsAsset || SourceComponent->GetPhysicsAsset())
 		{
 			SourceComponents.Add(SourceComponent);
