Change groom bound computation when attached to a parent skel. mesh, dependeing if is attached relatively or not.

The relative attachment is done based on attachment bone naming.

diff --git a/Engine/Plugins/Runtime/HairStrands/Source/HairStrandsCore/Private/GroomComponent.cpp b/Engine/Plugins/Runtime/HairStrands/Source/HairStrandsCore/Private/GroomComponent.cpp
--- a/Engine/Plugins/Runtime/HairStrands/Source/HairStrandsCore/Private/GroomComponent.cpp
+++ b/Engine/Plugins/Runtime/HairStrands/Source/HairStrandsCore/Private/GroomComponent.cpp
@@ -1856,7 +1856,11 @@ FBoxSphereBounds UGroomComponent::CalcBounds(const FTransform& InLocalToWorld) c
 			LocalHairBound = Buffers->GetBoundingBox();
 		}
 
-		if (RegisteredMeshComponent)
+		// If the attachment is done onto a skel. mesh:
+		// * If the attachment is simple, use the skel. mesh expanded by the groom bounds.
+		// * If the attachment is 'relative' to the skel mesh (using a bone anchor, which is provided with AttachmentName), 
+		//   we use the simple groom bound.
+		if (RegisteredMeshComponent && AttachmentName.IsEmpty())
 		{
 			const FBox LocalSkeletalBound = RegisteredMeshComponent->CalcBounds(FTransform::Identity).GetBox();
 
