Fix hair strands texture missing strands geometry on certain groom.

The underlying issue was two fold:
* An invalid clamp on the GPU, causing the segment fetching index to be incorrect.
* A too large voxel size & too low max segments per vertex.

This CL also ensures that voxel tracing is as short as possible, reducing avoiding TDR on large groom.

diff --git a/Engine/Plugins/Runtime/HairStrands/Source/HairStrandsCore/Private/GroomResources.cpp b/Engine/Plugins/Runtime/HairStrands/Source/HairStrandsCore/Private/GroomResources.cpp
--- a/Engine/Plugins/Runtime/HairStrands/Source/HairStrandsCore/Private/GroomResources.cpp
+++ b/Engine/Plugins/Runtime/HairStrands/Source/HairStrandsCore/Private/GroomResources.cpp
@@ -35,8 +35,8 @@ static FAutoConsoleVariableRef CVarHairCardsBulkData_AsyncLoading(TEXT("r.HairSt
 static int32 GHairStrandsBulkData_Validation = 1;
 static FAutoConsoleVariableRef CVarHairStrandsBulkData_Validation(TEXT("r.HairStrands.Strands.BulkData.Validation"), GHairStrandsBulkData_Validation, TEXT("Validate some hair strands data at serialization/loading time."));
 
-static float GHairStrandsDebugVoxel_WorldSize = 1.0f;
-static int32 GHairStrandsDebugVoxel_MaxSegmentPerVoxel = 512;
+static float GHairStrandsDebugVoxel_WorldSize = 0.3f;
+static int32 GHairStrandsDebugVoxel_MaxSegmentPerVoxel = 2048;
 static FAutoConsoleVariableRef CVarHairStrandsDebugVoxel_WorldSize(TEXT("r.HairStrands.DebugData.VoxelSize"), GHairStrandsDebugVoxel_WorldSize, TEXT("Voxel size use for creating debug data."));
 static FAutoConsoleVariableRef CVarHairStrandsDebugVoxel_MaxSegmentPerVoxel(TEXT("r.HairStrands.DebugData.MaxSegmentPerVoxel"), GHairStrandsDebugVoxel_MaxSegmentPerVoxel, TEXT("Max number of segments per Voxel size when creating debug data."));
 
@@ -1587,9 +1587,9 @@ static uint32 ToLinearCoord(const FIntVector& T, const FIntVector& Resolution)
 	return T.X + T.Y * Resolution.X + T.Z * Resolution.X * Resolution.Y;
 }
 
-static FIntVector ToCoord(const FVector& T, const FIntVector& Resolution, const FVector& MinBound, const float VoxelSize)
+static FIntVector ToCoord(const FVector3f& T, const FIntVector& Resolution, const FVector3f& MinBound, const float VoxelSize)
 {
-	const FVector C = (T - MinBound) / VoxelSize;
+	const FVector3f C = (T - MinBound) / VoxelSize;
 	return FIntVector(
 		FMath::Clamp(FMath::FloorToInt(C.X), 0, Resolution.X - 1),
 		FMath::Clamp(FMath::FloorToInt(C.Y), 0, Resolution.Y - 1),
@@ -1600,17 +1600,19 @@ void CreateHairStrandsDebugDatas(
 	const FHairStrandsDatas& InData,
 	FHairStrandsDebugDatas& Out)
 {
-	const FVector BoundSize = InData.BoundingBox.Max - InData.BoundingBox.Min;
+	const FVector3f BoundSize = FVector3f(InData.BoundingBox.Max) - FVector3f(InData.BoundingBox.Min);
 	Out.VoxelDescription.VoxelSize = FMath::Clamp(GHairStrandsDebugVoxel_WorldSize, 0.1f, 10.f);
 	Out.VoxelDescription.VoxelResolution = FIntVector(FMath::CeilToInt(BoundSize.X / Out.VoxelDescription.VoxelSize), FMath::CeilToInt(BoundSize.Y / Out.VoxelDescription.VoxelSize), FMath::CeilToInt(BoundSize.Z / Out.VoxelDescription.VoxelSize));
-	Out.VoxelDescription.VoxelMinBound = InData.BoundingBox.Min;
-	Out.VoxelDescription.VoxelMaxBound = FVector(Out.VoxelDescription.VoxelResolution) * Out.VoxelDescription.VoxelSize + InData.BoundingBox.Min;
+	Out.VoxelDescription.VoxelMinBound = FVector3f(InData.BoundingBox.Min);
+	Out.VoxelDescription.VoxelMaxBound = FVector3f(Out.VoxelDescription.VoxelResolution) * Out.VoxelDescription.VoxelSize + FVector3f(InData.BoundingBox.Min);
 	Out.VoxelOffsetAndCount.Init(FHairStrandsDebugDatas::FOffsetAndCount(), Out.VoxelDescription.VoxelResolution.X * Out.VoxelDescription.VoxelResolution.Y * Out.VoxelDescription.VoxelResolution.Z);
 	Out.VoxelDescription.MaxSegmentPerVoxel = 0;
 
 	uint32 AllocationCount = 0;
 	TArray<TArray<FHairStrandsDebugDatas::FVoxel>> TempVoxelData;
 
+	const uint32 MaxNumberOfSegmentPerVoxel = uint32(FMath::Clamp(GHairStrandsDebugVoxel_MaxSegmentPerVoxel, 16, 64000));
+
 	// Fill in voxel (TODO: make it parallel)
 	const uint32 CurveCount = InData.StrandsCurves.Num();
 	for (uint32 CurveIndex = 0; CurveIndex < CurveCount; ++CurveIndex)
@@ -1622,16 +1624,16 @@ void CreateHairStrandsDebugDatas(
 		{
 			const uint32 Index0 = PointOffset + PointIndex;
 			const uint32 Index1 = PointOffset + PointIndex + 1;
-			const FVector& P0 = (FVector)InData.StrandsPoints.PointsPosition[Index0];
-			const FVector& P1 = (FVector)InData.StrandsPoints.PointsPosition[Index1];
-			const FVector Segment = P1 - P0;
+			const FVector3f& P0 = InData.StrandsPoints.PointsPosition[Index0];
+			const FVector3f& P1 = InData.StrandsPoints.PointsPosition[Index1];
+			const FVector3f Segment = P1 - P0;
 
 			const float Length = Segment.Size();
-			const uint32 StepCount = FMath::CeilToInt(Length / Out.VoxelDescription.VoxelSize);
+			const uint32 StepCount = FMath::CeilToInt(Length / (0.25f * Out.VoxelDescription.VoxelSize));
 			uint32 PrevLinearCoord = ~0;
 			for (uint32 StepIt = 0; StepIt < StepCount + 1; ++StepIt)
 			{
-				const FVector P = P0 + Segment * StepIt / float(StepCount);
+				const FVector3f P = P0 + Segment * StepIt / float(StepCount);
 				const FIntVector Coord = ToCoord(P, Out.VoxelDescription.VoxelResolution, Out.VoxelDescription.VoxelMinBound, Out.VoxelDescription.VoxelSize);
 				const uint32 LinearCoord = ToLinearCoord(Coord, Out.VoxelDescription.VoxelResolution);
 				if (LinearCoord != PrevLinearCoord)
@@ -1642,9 +1644,12 @@ void CreateHairStrandsDebugDatas(
 						TempVoxelData.Add(TArray<FHairStrandsDebugDatas::FVoxel>());
 					}
 
-					const uint32 Offset = Out.VoxelOffsetAndCount[LinearCoord].Offset;
-					Out.VoxelOffsetAndCount[LinearCoord].Count += 1;
-					TempVoxelData[Offset].Add({Index0, Index1});
+					if (Out.VoxelOffsetAndCount[LinearCoord].Count+1 < MaxNumberOfSegmentPerVoxel)
+					{
+						const uint32 Offset = Out.VoxelOffsetAndCount[LinearCoord].Offset;
+						Out.VoxelOffsetAndCount[LinearCoord].Count += 1;
+						TempVoxelData[Offset].Add({Index0, Index1});
+					}
 
 					Out.VoxelDescription.MaxSegmentPerVoxel = FMath::Max(Out.VoxelDescription.MaxSegmentPerVoxel, Out.VoxelOffsetAndCount[LinearCoord].Count);
 
@@ -1655,16 +1660,27 @@ void CreateHairStrandsDebugDatas(
 			}
 		}
 	}
-	Out.VoxelDescription.MaxSegmentPerVoxel = FMath::Min(Out.VoxelDescription.MaxSegmentPerVoxel, uint32(FMath::Clamp(GHairStrandsDebugVoxel_MaxSegmentPerVoxel, 16, 64000)));
+
+	check(Out.VoxelDescription.MaxSegmentPerVoxel < MaxNumberOfSegmentPerVoxel);
 	Out.VoxelData.Reserve(AllocationCount);
 
 	for (int32 Index = 0, Count = Out.VoxelOffsetAndCount.Num(); Index < Count; ++Index)
 	{
-		const uint32 ArrayIndex = Out.VoxelOffsetAndCount[Index].Offset;
-		Out.VoxelOffsetAndCount[Index].Offset = Out.VoxelData.Num();
 		if (Out.VoxelOffsetAndCount[Index].Count > 0)
 		{
-			Out.VoxelData.Append(TempVoxelData[ArrayIndex]);
+			const uint32 ArrayIndex = Out.VoxelOffsetAndCount[Index].Offset;
+			const uint32 NewOffset = Out.VoxelData.Num();
+			Out.VoxelOffsetAndCount[Index].Offset = NewOffset;
+
+			for (FHairStrandsDebugDatas::FVoxel& Voxel : TempVoxelData[ArrayIndex])
+			{
+				Voxel.Index1 = NewOffset;
+				Out.VoxelData.Add(Voxel);
+			}
+		}
+		else
+		{
+			Out.VoxelOffsetAndCount[Index].Offset = 0;
 		}
 
 		// Sanity check
diff --git a/Engine/Plugins/Runtime/HairStrands/Source/HairStrandsCore/Private/GroomTextureBuilder.cpp b/Engine/Plugins/Runtime/HairStrands/Source/HairStrandsCore/Private/GroomTextureBuilder.cpp
--- a/Engine/Plugins/Runtime/HairStrands/Source/HairStrandsCore/Private/GroomTextureBuilder.cpp
+++ b/Engine/Plugins/Runtime/HairStrands/Source/HairStrandsCore/Private/GroomTextureBuilder.cpp
@@ -636,6 +636,10 @@ class FHairStrandsTexturePS : public FGlobalShader
 		SHADER_PARAMETER(FIntVector, Voxel_Resolution)
 		SHADER_PARAMETER(float, Voxel_Size)
 		SHADER_PARAMETER(uint32, Voxel_MaxSegmentPerVoxel)
+
+		SHADER_PARAMETER(uint32, Voxel_OffsetAndCount_MaxCount)
+		SHADER_PARAMETER(uint32, Voxel_Data_MaxCount)
+
 		SHADER_PARAMETER_RDG_BUFFER_SRV(StructuredBuffer, Voxel_OffsetAndCount)
 		SHADER_PARAMETER_RDG_BUFFER_SRV(StructuredBuffer, Voxel_Data)
 
@@ -678,8 +682,8 @@ static void InternalGenerateHairStrandsTextures(
 	FRHIShaderResourceView* InMeshUVsBuffer,
 	FRHIShaderResourceView* InMeshNormalsBuffer,
 
-	const FVector& VoxelMinBound,
-	const FVector& VoxelMaxBound,
+	const FVector3f& VoxelMinBound,
+	const FVector3f& VoxelMaxBound,
 	const FIntVector& VoxelResolution,
 	float VoxelSize,
 	uint32 VoxelMaxSegmentPerVoxel,
@@ -730,14 +734,17 @@ static void InternalGenerateHairStrandsTextures(
 	ParametersPS->InVF_ControlPointCount = InHairStrands_ControlPointCount;
 	ParametersPS->InVF_GroupIndex = InHairStrands_GroupIndex;
 
-	ParametersPS->Voxel_MinBound = (FVector3f)VoxelMinBound;
-	ParametersPS->Voxel_MaxBound = (FVector3f)VoxelMaxBound;
+	ParametersPS->Voxel_MinBound = VoxelMinBound;
+	ParametersPS->Voxel_MaxBound = VoxelMaxBound;
 	ParametersPS->Voxel_Resolution = VoxelResolution;
 	ParametersPS->Voxel_Size = VoxelSize;
 	ParametersPS->Voxel_MaxSegmentPerVoxel = VoxelMaxSegmentPerVoxel;
 	ParametersPS->Voxel_OffsetAndCount = GraphBuilder.CreateSRV(VoxelOffsetAndCount);
 	ParametersPS->Voxel_Data = GraphBuilder.CreateSRV(VoxelData);
 
+	ParametersPS->Voxel_OffsetAndCount_MaxCount = VoxelOffsetAndCount->Desc.NumElements;
+	ParametersPS->Voxel_Data_MaxCount = VoxelData->Desc.NumElements;
+
 	if (ShaderPrintData)
 	{
 		ShaderPrint::SetParameters(GraphBuilder, *ShaderPrintData, ParametersPS->ShaderPrintParameters);
@@ -838,6 +845,7 @@ struct FStrandsTexturesReadback
 	TUniquePtr<FRHIGPUTextureReadback> Coverage;
 	TUniquePtr<FRHIGPUTextureReadback> Attribute;
 	TUniquePtr<FRHIGPUTextureReadback> Material;
+	uint32 NotReadyFrameCount = 0;
 
 	FStrandsTexturesOutput Output;
 	bool IsReady() const
@@ -1125,12 +1133,13 @@ void RunHairStrandsTexturesQueries(FRDGBuilder& GraphBuilder, FGlobalShaderMap*
 {
 	// Readback data
 	{
+		TArray<FStrandsTexturesReadback*> NotReadyQueue;
 		FStrandsTexturesReadback* R = nullptr;
 		while (GStrandsTexturesReadbacks.Dequeue(R))
 		{
 			if (R)
 			{
-				if (R->IsReady())
+				if (R->IsReady() || ++R->NotReadyFrameCount > 32)
 				{
 					CopyReadbackToTexture(R->Depth,    R->Output.Depth);
 					CopyReadbackToTexture(R->Tangent,  R->Output.Tangent);
@@ -1148,10 +1157,15 @@ void RunHairStrandsTexturesQueries(FRDGBuilder& GraphBuilder, FGlobalShaderMap*
 				}
 				else
 				{
-					GStrandsTexturesReadbacks.Enqueue(R);
+					NotReadyQueue.Add(R);
 				}
 			}
 		}
+
+		for (FStrandsTexturesReadback* E : NotReadyQueue)
+		{
+			GStrandsTexturesReadbacks.Enqueue(E);
+		}
 	}
 
 	// Render textures
diff --git a/Engine/Plugins/Runtime/HairStrands/Source/HairStrandsCore/Public/HairStrandsDatas.h b/Engine/Plugins/Runtime/HairStrands/Source/HairStrandsCore/Public/HairStrandsDatas.h
--- a/Engine/Plugins/Runtime/HairStrands/Source/HairStrandsCore/Public/HairStrandsDatas.h
+++ b/Engine/Plugins/Runtime/HairStrands/Source/HairStrandsCore/Public/HairStrandsDatas.h
@@ -481,8 +481,8 @@ struct HAIRSTRANDSCORE_API FHairStrandsDebugDatas
 
 	struct FDesc
 	{
-		FVector VoxelMinBound = FVector::ZeroVector;
-		FVector VoxelMaxBound = FVector::ZeroVector;
+		FVector3f VoxelMinBound = FVector3f::ZeroVector;
+		FVector3f VoxelMaxBound = FVector3f::ZeroVector;
 		FIntVector VoxelResolution = FIntVector::ZeroValue;
 		float VoxelSize = 0;
 		uint32 MaxSegmentPerVoxel = 0;
diff --git a/Engine/Shaders/Private/HairStrands/HairStrandsTexturesGeneration.usf b/Engine/Shaders/Private/HairStrands/HairStrandsTexturesGeneration.usf
--- a/Engine/Shaders/Private/HairStrands/HairStrandsTexturesGeneration.usf
+++ b/Engine/Shaders/Private/HairStrands/HairStrandsTexturesGeneration.usf
@@ -106,6 +106,10 @@ float3 Voxel_MaxBound;
 int3 Voxel_Resolution;
 float Voxel_Size;
 uint Voxel_MaxSegmentPerVoxel;
+
+uint Voxel_OffsetAndCount_MaxCount;
+uint Voxel_Data_MaxCount;
+
 StructuredBuffer<uint2> Voxel_OffsetAndCount;
 StructuredBuffer<uint2> Voxel_Data;
 
@@ -209,7 +213,6 @@ bool Trace(
 			PageIndexCoord = clamp(uint3(Coords), uint3(0, 0, 0), Voxel_Resolution - 1);
 		}
 
-		bool bLastStep = false;
 		const uint LoopCount = 256u;
 		for (uint LoopIt = 0; LoopIt < LoopCount; ++LoopIt)
 		{
@@ -225,16 +228,13 @@ bool Trace(
 				return bIsValid;
 			}
 
-			if (bLastStep)
+			const float3 WPCurr = (PageIndexCoord + 0.5f) * Voxel_Size  + Voxel_MinBound;
+			const bool bExit = dot((WPCurr - O), InRayD) > (MaxDistance + Voxel_Size);
+			if (bExit)
 			{
 				return bIsValid;
 			}
 
-			if (all(PageIndexCoord == EndCoord))
-			{
-				bLastStep = true;
-			}
-
 			const uint LinearPageIndexCoord = CoordToIndex(PageIndexCoord, Voxel_Resolution, 0);
 
 			#if DEBUG_ENABLE
@@ -254,19 +254,19 @@ bool Trace(
 				uint2 OffsetAndCount = 0;
 				{
 					const uint LinearIndexCoord = CoordToIndex(VolumeCoord, Voxel_Resolution, 0);
-					OffsetAndCount = Voxel_OffsetAndCount[LinearIndexCoord];
+					OffsetAndCount = Voxel_OffsetAndCount[clamp(LinearIndexCoord, 0, Voxel_OffsetAndCount_MaxCount-1)];
 					OffsetAndCount.y = min(OffsetAndCount.y, Voxel_MaxSegmentPerVoxel);
 				}
 
 				if (OffsetAndCount.y > 0)
 				{
-					const uint MaxVoxelCount = Voxel_Resolution.x * Voxel_Resolution.y * Voxel_Resolution.z;
 					for (uint PointIt = 0; PointIt < OffsetAndCount.y; ++PointIt)
 					{
-						const uint VoxelIndex = clamp(OffsetAndCount.x + PointIt, 0, MaxVoxelCount);
-						const uint2 P0_P1 = Voxel_Data[OffsetAndCount.x + PointIt];
-						const uint P0 = clamp(P0_P1.x, 0, InVF_ControlPointCount);
-						const uint P1 = clamp(P0_P1.y, 0, InVF_ControlPointCount);
+						const uint VoxelIndex = clamp(OffsetAndCount.x + PointIt, 0, Voxel_Data_MaxCount-1);
+						const uint2 P0_P1 = Voxel_Data[VoxelIndex].xx + uint2(0,1);
+
+						const uint P0 = P0_P1.x;
+						const uint P1 = P0_P1.y;
 
 						const uint4 PackedCP_P0 = InVF_PositionBuffer[P0];
 						const uint4 PackedCP_P1 = InVF_PositionBuffer[P1];
@@ -290,9 +290,19 @@ bool Trace(
 						float VertexU = 0;
 						const float WorldRadius = 2 * max(CP_Curr.WorldRadius, CP_Next.WorldRadius); //  0.05f;
 						const float Distance = Intersection(CP_Curr.Position, CP_Next.Position, RayP0, RayP1, WorldRadius, VertexU) * InMaxDistance;
-						if (Distance > 0 && Distance < OutDepth)
+
+						const bool bIntersect = Distance > 0;
+						const bool bClosest = Distance < OutDepth;
+
+						const bool bIntersectAndClosest = bIntersect && bClosest;
+
+						if (bIntersect)
 						{
 							OutCoverage = 1;
+						}
+
+						if (bIntersectAndClosest)
+						{
 							OutTangent = Tangent;
 							OutDepth = Distance;
 
@@ -314,7 +324,7 @@ bool Trace(
 								AddCrossWS(CP_Curr.Position, 0.05f, Color);
 								AddCrossWS(CP_Next.Position, 0.05f, Color);
 							}
-							#endif							
+							#endif
 						}
 					}
 				}
@@ -354,8 +364,8 @@ void MainPS(
 	const float3 B = normalize(InBitangent);
 	const float3 N = normalize(InNormal);
 
-	const float3 PixelAxisX = ddx(InPosition);
-	const float3 PixelAxisY = ddy(InPosition);
+	const float3 PixelAxisX = ddx_fine(InPosition);
+	const float3 PixelAxisY = ddy_fine(InPosition);
 
 	OutDepth = max(0, MaxDistance);
 	OutCoverage = 0;
@@ -369,14 +379,14 @@ void MainPS(
 
 	const bool bDebugEnable =
 	#if DEBUG_ENABLE
-		uint(GetCursorPos().x) == uint(InUV.x * OutputResolution.x) &&
-		uint(GetCursorPos().y) == uint(InUV.y * OutputResolution.y);
+		all(PixelCoord == GetCursorPos());
 	#else
 		false;
 	#endif
 
-	const uint SampleCount = 8;
 	uint bValidCount = 0;
+
+	const uint SampleCount = 16;
 	for (uint SampleIt = 0; SampleIt < SampleCount; ++SampleIt)
 	{
 		const float2 Jitter = Hammersley(SampleIt, SampleCount, 0) * 2 - 1;
