Readd GPU flush to avoid TDR when generating texture for large grooms.

diff --git a/Engine/Plugins/Runtime/HairStrands/Source/HairStrandsCore/Private/GroomTextureBuilder.cpp b/Engine/Plugins/Runtime/HairStrands/Source/HairStrandsCore/Private/GroomTextureBuilder.cpp
--- a/Engine/Plugins/Runtime/HairStrands/Source/HairStrandsCore/Private/GroomTextureBuilder.cpp
+++ b/Engine/Plugins/Runtime/HairStrands/Source/HairStrandsCore/Private/GroomTextureBuilder.cpp
@@ -811,7 +811,7 @@ static void InternalGenerateHairStrandsTextures(
 				RHICmdList.DrawIndexedPrimitive(InMeshIndexBuffer, VertexBaseIndex, 0, VertexCount, IndexBaseIndex, PrimitiveCount, 1);
 
 				// Flush, to ensure that all texture generation is done (TDR)
-				#if 0
+				#if 1
 				GDynamicRHI->RHISubmitCommandsAndFlushGPU();
 				GDynamicRHI->RHIBlockUntilGPUIdle();
 				#endif
@@ -829,7 +829,7 @@ static void InternalGenerateHairStrandsTextures(
 					RHICmdList.DrawIndexedPrimitive(InMeshIndexBuffer, VertexBaseIndex, 0, VertexCount, IndexBaseIndex, PrimitiveCount, 1);
 
 					// Flush, to ensure that all texture generation is done (TDR)
-					#if 0
+					#if 1
 					GDynamicRHI->RHISubmitCommandsAndFlushGPU();
 					GDynamicRHI->RHIBlockUntilGPUIdle();
 					#endif
