Change how the strands textures are computed.

This changes process strands texture in tile, and split computation over several frames to avoid TDR on large grooms.

diff --git a/Engine/Plugins/Runtime/HairStrands/Source/HairStrandsCore/Private/GroomTextureBuilder.cpp b/Engine/Plugins/Runtime/HairStrands/Source/HairStrandsCore/Private/GroomTextureBuilder.cpp
--- a/Engine/Plugins/Runtime/HairStrands/Source/HairStrandsCore/Private/GroomTextureBuilder.cpp
+++ b/Engine/Plugins/Runtime/HairStrands/Source/HairStrandsCore/Private/GroomTextureBuilder.cpp
@@ -660,6 +660,70 @@ class FHairStrandsTexturePS : public FGlobalShader
 IMPLEMENT_GLOBAL_SHADER(FHairStrandsTextureVS, "/Engine/Private/HairStrands/HairStrandsTexturesGeneration.usf", "MainVS", SF_Vertex);
 IMPLEMENT_GLOBAL_SHADER(FHairStrandsTexturePS, "/Engine/Private/HairStrands/HairStrandsTexturesGeneration.usf", "MainPS", SF_Pixel);
 
+
+const int32 HairStrandsTextureTileSize = 1024;
+
+static FIntPoint GetTileResolution(FIntPoint In)
+{
+	FIntPoint Out;
+	Out.X = FMath::DivideAndRoundUp(In.X, HairStrandsTextureTileSize);
+	Out.Y = FMath::DivideAndRoundUp(In.Y, HairStrandsTextureTileSize);
+	return Out;
+}
+
+struct FHairStrandsRDGTextures
+{
+	FIntPoint TileCount; // Total tile count
+	FIntPoint TileCoord; // Current tile coord
+
+	FRDGTextureRef DepthTexture;
+	FRDGTextureRef CoverageTexture;
+	FRDGTextureRef TangentTexture;
+	FRDGTextureRef AttributeTexture;
+	FRDGTextureRef MaterialTexture;
+	FRDGTextureRef TriangleMaskTexture;
+	FRDGTextureRef DepthTestTexture;
+
+	bool IsCompleted() const { return TileCoord.X >= TileCount.X || TileCoord.Y >= TileCount.Y;  }
+};
+
+struct FHairStrandsRTs
+{
+	FIntPoint TileCount; // Total tile count
+	FIntPoint TileCoord; // Current tile coord
+
+	TRefCountPtr<IPooledRenderTarget> DepthTexture;
+	TRefCountPtr<IPooledRenderTarget> CoverageTexture;
+	TRefCountPtr<IPooledRenderTarget> TangentTexture;
+	TRefCountPtr<IPooledRenderTarget> AttributeTexture;
+	TRefCountPtr<IPooledRenderTarget> MaterialTexture;
+	TRefCountPtr<IPooledRenderTarget> TriangleMaskTexture;
+	TRefCountPtr<IPooledRenderTarget> DepthTestTexture;
+	
+	bool IsCompleted() const { return TileCoord.X >= TileCount.X || TileCoord.Y >= TileCount.Y; }
+};
+
+struct FStrandsTexturesReadback
+{
+	TUniquePtr<FRHIGPUTextureReadback> Depth;
+	TUniquePtr<FRHIGPUTextureReadback> Tangent;
+	TUniquePtr<FRHIGPUTextureReadback> Coverage;
+	TUniquePtr<FRHIGPUTextureReadback> Attribute;
+	TUniquePtr<FRHIGPUTextureReadback> Material;
+	uint32 NotReadyFrameCount = 0;
+
+	bool IsReady() const
+	{
+		return
+			Depth->IsReady() &&
+			Tangent->IsReady() &&
+			Coverage->IsReady() &&
+			Attribute->IsReady() &&
+			Material->IsReady();
+	}
+};
+
+// Decompose into tile
 static void InternalGenerateHairStrandsTextures(
 	FRDGBuilder& GraphBuilder,
 	FGlobalShaderMap* ShaderMap,
@@ -700,16 +764,11 @@ static void InternalGenerateHairStrandsTextures(
 	uint32 InHairStrands_ControlPointCount,
 	uint32 InHairStrands_GroupIndex,
 
-	FRDGTextureRef OutDepthTestTexture,
-	FRDGTextureRef OutDepthTexture,
-	FRDGTextureRef OutTangentTexture,
-	FRDGTextureRef OutCoverageTexture,
-	FRDGTextureRef OutRootUVStrandsUSeedTexture,
-	FRDGTextureRef OutMaterialTexture,
-	FRDGTextureRef OutTriangleMask)
+	FHairStrandsRDGTextures& Out)
 {
-	const FIntPoint OutputResolution = OutDepthTexture->Desc.Extent;
+	const FIntPoint OutputResolution = Out.DepthTexture->Desc.Extent;
 	const bool bHasMaterialData = InHairStrands_MaterialBuffer != nullptr;
+	const FIntPoint OutTileCoord = Out.TileCoord;
 
 	FHairStrandsTexturePS::FParameters* ParametersPS = GraphBuilder.AllocParameters<FHairStrandsTexturePS::FParameters>();
 	ParametersPS->OutputResolution = OutputResolution;
@@ -750,14 +809,14 @@ static void InternalGenerateHairStrandsTextures(
 		ShaderPrint::SetParameters(GraphBuilder, *ShaderPrintData, ParametersPS->ShaderPrintParameters);
 	}
 
-	ParametersPS->RenderTargets[0] = FRenderTargetBinding(OutDepthTexture, bClear ? ERenderTargetLoadAction::EClear : ERenderTargetLoadAction::ELoad);
-	ParametersPS->RenderTargets[1] = FRenderTargetBinding(OutTangentTexture, bClear ? ERenderTargetLoadAction::EClear : ERenderTargetLoadAction::ELoad);
-	ParametersPS->RenderTargets[2] = FRenderTargetBinding(OutCoverageTexture, bClear ? ERenderTargetLoadAction::EClear : ERenderTargetLoadAction::ELoad);
-	ParametersPS->RenderTargets[3] = FRenderTargetBinding(OutRootUVStrandsUSeedTexture, bClear ? ERenderTargetLoadAction::EClear : ERenderTargetLoadAction::ELoad);
-	ParametersPS->RenderTargets[4] = FRenderTargetBinding(OutMaterialTexture, bClear ? ERenderTargetLoadAction::EClear : ERenderTargetLoadAction::ELoad);
-	ParametersPS->RenderTargets[5] = FRenderTargetBinding(OutTriangleMask, bClear ? ERenderTargetLoadAction::EClear : ERenderTargetLoadAction::ELoad);
+	ParametersPS->RenderTargets[0] = FRenderTargetBinding(Out.DepthTexture, bClear ? ERenderTargetLoadAction::EClear : ERenderTargetLoadAction::ELoad);
+	ParametersPS->RenderTargets[1] = FRenderTargetBinding(Out.TangentTexture, bClear ? ERenderTargetLoadAction::EClear : ERenderTargetLoadAction::ELoad);
+	ParametersPS->RenderTargets[2] = FRenderTargetBinding(Out.CoverageTexture, bClear ? ERenderTargetLoadAction::EClear : ERenderTargetLoadAction::ELoad);
+	ParametersPS->RenderTargets[3] = FRenderTargetBinding(Out.AttributeTexture, bClear ? ERenderTargetLoadAction::EClear : ERenderTargetLoadAction::ELoad);
+	ParametersPS->RenderTargets[4] = FRenderTargetBinding(Out.MaterialTexture, bClear ? ERenderTargetLoadAction::EClear : ERenderTargetLoadAction::ELoad);
+	ParametersPS->RenderTargets[5] = FRenderTargetBinding(Out.TriangleMaskTexture, bClear ? ERenderTargetLoadAction::EClear : ERenderTargetLoadAction::ELoad);
 	ParametersPS->RenderTargets.DepthStencil = FDepthStencilBinding(
-		OutDepthTestTexture,
+		Out.DepthTestTexture,
 		bClear ? ERenderTargetLoadAction::EClear : ERenderTargetLoadAction::ELoad,
 		ERenderTargetLoadAction::ENoAction,
 		FExclusiveDepthStencil::DepthWrite_StencilNop);
@@ -769,7 +828,7 @@ static void InternalGenerateHairStrandsTextures(
 		RDG_EVENT_NAME("HairStrandsTexturePS"),
 		ParametersPS,
 		ERDGPassFlags::Raster,
-		[ParametersPS, VertexShader, PixelShader, InMeshIndexBuffer, VertexCount, PrimitiveCount, IndexBaseIndex, VertexBaseIndex, OutputResolution](FRHICommandList& RHICmdList)
+		[ParametersPS, VertexShader, PixelShader, InMeshIndexBuffer, VertexCount, PrimitiveCount, IndexBaseIndex, VertexBaseIndex, OutputResolution, OutTileCoord](FRHICommandList& RHICmdList)
 		{
 			FHairStrandsTextureVS::FParameters ParametersVS;
 			ParametersVS.OutputResolution = ParametersPS->OutputResolution;
@@ -805,31 +864,29 @@ static void InternalGenerateHairStrandsTextures(
 			
 			// Divide the rendering work into small batches to reduce risk of TDR as the texture projection implies heavy works 
 			// (i.e. long thread running due the the large amount of strands a groom can have)
-			const int32 TileSize = 1024;
+			const int32 TileSize = HairStrandsTextureTileSize;
 			if (OutputResolution.X < TileSize)
 			{
 				RHICmdList.DrawIndexedPrimitive(InMeshIndexBuffer, VertexBaseIndex, 0, VertexCount, IndexBaseIndex, PrimitiveCount, 1);
 
 				// Flush, to ensure that all texture generation is done (TDR)
-				#if 1
+				#if 0
 				GDynamicRHI->RHISubmitCommandsAndFlushGPU();
 				GDynamicRHI->RHIBlockUntilGPUIdle();
 				#endif
 			}
 			else
 			{
-				const int32 TileCountX = FMath::DivideAndRoundUp(OutputResolution.X, TileSize);
-				const int32 TileCountY = FMath::DivideAndRoundUp(OutputResolution.Y, TileSize);
-				for (int32 TileY = 0; TileY < TileCountY; ++TileY)
-				for (int32 TileX = 0; TileX < TileCountX; ++TileX)
+				//for (int32 TileY = 0; TileY < TileCountY; ++TileY)
+				//for (int32 TileX = 0; TileX < TileCountX; ++TileX)
 				{
-					const uint32 OffsetX = TileX * TileSize;
-					const uint32 OffsetY = TileY * TileSize;
+					const uint32 OffsetX = OutTileCoord.X * TileSize;
+					const uint32 OffsetY = OutTileCoord.Y * TileSize;
 					RHICmdList.SetScissorRect(true, OffsetX, OffsetY, OffsetX + TileSize, OffsetY + TileSize);
 					RHICmdList.DrawIndexedPrimitive(InMeshIndexBuffer, VertexBaseIndex, 0, VertexCount, IndexBaseIndex, PrimitiveCount, 1);
 
 					// Flush, to ensure that all texture generation is done (TDR)
-					#if 1
+					#if 0
 					GDynamicRHI->RHISubmitCommandsAndFlushGPU();
 					GDynamicRHI->RHIBlockUntilGPUIdle();
 					#endif
@@ -838,90 +895,94 @@ static void InternalGenerateHairStrandsTextures(
 		});
 }
 
-struct FStrandsTexturesReadback
-{
-	TUniquePtr<FRHIGPUTextureReadback> Depth;
-	TUniquePtr<FRHIGPUTextureReadback> Tangent;
-	TUniquePtr<FRHIGPUTextureReadback> Coverage;
-	TUniquePtr<FRHIGPUTextureReadback> Attribute;
-	TUniquePtr<FRHIGPUTextureReadback> Material;
-	uint32 NotReadyFrameCount = 0;
-
-	FStrandsTexturesOutput Output;
-	bool IsReady() const
-	{
-		return
-			Depth->IsReady() &&
-			Tangent->IsReady() &&
-			Coverage->IsReady() &&
-			Attribute->IsReady() &&
-			Material->IsReady();
-	}
-};
 
-static void InternalBuildStrandsTextures_GPU(
-	FRDGBuilder& GraphBuilder,
-	FGlobalShaderMap* ShaderMap,
-	const FStrandsTexturesInfo& InInfo,
-	FStrandsTexturesReadback& Output,
-	const struct FShaderPrintData* DebugShaderData)
+static FHairStrandsRDGTextures CreateTextures(FRDGBuilder& GraphBuilder, FIntPoint InResolution)
 {
-	USkeletalMesh* SkeletalMesh = (USkeletalMesh*)InInfo.SkeletalMesh;
-	UStaticMesh* StaticMesh = (UStaticMesh*)InInfo.StaticMesh;
+	FHairStrandsRDGTextures Out;
 
-	if (!SkeletalMesh && !StaticMesh)
-	{
-		return;
-	}
-
-	const bool bUseSkeletalMesh = SkeletalMesh != nullptr;
-
-	const FIntPoint OutputResolution(
-		FMath::Clamp(InInfo.Resolution, 512u, 16384u),
-		FMath::Clamp(InInfo.Resolution, 512u, 16384u));
-
-	FRDGTextureDesc Desc = FRDGTextureDesc::Create2D(OutputResolution, PF_A8R8G8B8, FClearValueBinding::Black, TexCreate_RenderTargetable | TexCreate_UAV | TexCreate_ShaderResource);
+	FRDGTextureDesc Desc = FRDGTextureDesc::Create2D(InResolution, PF_A8R8G8B8, FClearValueBinding::Black, TexCreate_RenderTargetable | TexCreate_UAV | TexCreate_ShaderResource);
 
 	Desc.Format = PF_G16;
 	Desc.ClearValue = FClearValueBinding::White;
-	FRDGTextureRef DepthTexture[2];
-	DepthTexture[0] = GraphBuilder.CreateTexture(Desc, TEXT("Hair.DepthTexture"));
-	DepthTexture[1] = GraphBuilder.CreateTexture(Desc, TEXT("Hair.DepthTexture"));
+	Out.DepthTexture = GraphBuilder.CreateTexture(Desc, TEXT("Hair.DepthTexture"));
 
 	Desc.Format = PF_B8G8R8A8;
 	Desc.ClearValue = FClearValueBinding::Transparent;
-	FRDGTextureRef CoverageTexture[2];
-	CoverageTexture[0] = GraphBuilder.CreateTexture(Desc, TEXT("Hair.CoverageTexture"));
-	CoverageTexture[1] = GraphBuilder.CreateTexture(Desc, TEXT("Hair.CoverageTexture"));
+	Out.CoverageTexture = GraphBuilder.CreateTexture(Desc, TEXT("Hair.CoverageTexture"));
 
 	Desc.Format = PF_B8G8R8A8;
 	Desc.ClearValue = FClearValueBinding::Transparent;
-	FRDGTextureRef TangentTexture[2];
-	TangentTexture[0] = GraphBuilder.CreateTexture(Desc, TEXT("Hair.TangentTexture"));
-	TangentTexture[1] = GraphBuilder.CreateTexture(Desc, TEXT("Hair.TangentTexture"));
+	Out.TangentTexture = GraphBuilder.CreateTexture(Desc, TEXT("Hair.TangentTexture"));
 
 	Desc.Format = PF_B8G8R8A8;
 	Desc.ClearValue = FClearValueBinding::Transparent;
-	FRDGTextureRef AttributeTexture[2];
-	AttributeTexture[0] = GraphBuilder.CreateTexture(Desc, TEXT("Hair.AttributeTexture"));
-	AttributeTexture[1] = GraphBuilder.CreateTexture(Desc, TEXT("Hair.AttributeTexture"));
+	Out.AttributeTexture = GraphBuilder.CreateTexture(Desc, TEXT("Hair.AttributeTexture"));
 
 	Desc.Format = PF_B8G8R8A8;
 	Desc.ClearValue = FClearValueBinding::Transparent;
-	FRDGTextureRef MaterialTexture[2];
-	MaterialTexture[0] = GraphBuilder.CreateTexture(Desc, TEXT("Hair.MaterialTexture"));
-	MaterialTexture[1] = GraphBuilder.CreateTexture(Desc, TEXT("Hair.MaterialTexture"));
+	Out.MaterialTexture = GraphBuilder.CreateTexture(Desc, TEXT("Hair.MaterialTexture"));
 
 	Desc.Format = PF_R32_UINT;
 	Desc.ClearValue = FClearValueBinding::Black;
-	FRDGTextureRef TriangleMaskTexture = GraphBuilder.CreateTexture(Desc, TEXT("Hair.TriangleMaskTexture"));
+	Out.TriangleMaskTexture = GraphBuilder.CreateTexture(Desc, TEXT("Hair.TriangleMaskTexture"));
 	
 	Desc.ClearValue = FClearValueBinding(1);
 	Desc.Format = PF_DepthStencil;
 	Desc.Flags = TexCreate_DepthStencilTargetable;
-	FRDGTextureRef DepthTestTexture = GraphBuilder.CreateTexture(Desc, TEXT("Hair.DepthTestTexture"));
+	Out.DepthTestTexture = GraphBuilder.CreateTexture(Desc, TEXT("Hair.DepthTestTexture"));
+
+	Out.TileCount = GetTileResolution(InResolution);
+	Out.TileCoord = FIntPoint(0,0);
+
+	return Out;
+}
 
-	bool bClear = true;
+static FHairStrandsRDGTextures RegisterTextures(FRDGBuilder& GraphBuilder, const FHairStrandsRTs& In)
+{
+	FHairStrandsRDGTextures Out;
+	Out.DepthTexture		= GraphBuilder.RegisterExternalTexture(In.DepthTexture);
+	Out.CoverageTexture		= GraphBuilder.RegisterExternalTexture(In.CoverageTexture);
+	Out.TangentTexture		= GraphBuilder.RegisterExternalTexture(In.TangentTexture);
+	Out.AttributeTexture	= GraphBuilder.RegisterExternalTexture(In.AttributeTexture);
+	Out.MaterialTexture		= GraphBuilder.RegisterExternalTexture(In.MaterialTexture);
+	Out.TriangleMaskTexture = GraphBuilder.RegisterExternalTexture(In.TriangleMaskTexture);
+	Out.DepthTestTexture	= GraphBuilder.RegisterExternalTexture(In.DepthTestTexture);
+	Out.TileCoord			= In.TileCoord;
+	Out.TileCount			= In.TileCount;
+	return Out;
+}
+
+static void ExportTextures(FRDGBuilder& GraphBuilder, const FHairStrandsRDGTextures& In, FHairStrandsRTs& Out)
+{
+	GraphBuilder.QueueTextureExtraction(In.DepthTexture,		&Out.DepthTexture);
+	GraphBuilder.QueueTextureExtraction(In.CoverageTexture,		&Out.CoverageTexture);
+	GraphBuilder.QueueTextureExtraction(In.TangentTexture,		&Out.TangentTexture);
+	GraphBuilder.QueueTextureExtraction(In.AttributeTexture,	&Out.AttributeTexture);
+	GraphBuilder.QueueTextureExtraction(In.MaterialTexture,		&Out.MaterialTexture);
+	GraphBuilder.QueueTextureExtraction(In.TriangleMaskTexture, &Out.TriangleMaskTexture);
+	GraphBuilder.QueueTextureExtraction(In.DepthTestTexture,	&Out.DepthTestTexture);
+	Out.TileCoord = In.TileCoord;
+	Out.TileCount = In.TileCount;
+}
+
+static bool TraceTextures(
+	FRDGBuilder& GraphBuilder,
+	FGlobalShaderMap* ShaderMap,
+	const FStrandsTexturesInfo& InInfo,
+	const FShaderPrintData* ShaderPrintData,
+	FHairStrandsRDGTextures& Output)
+{
+	USkeletalMesh* SkeletalMesh = (USkeletalMesh*)InInfo.SkeletalMesh;
+	UStaticMesh* StaticMesh = (UStaticMesh*)InInfo.StaticMesh;
+
+	if (!SkeletalMesh && !StaticMesh)
+	{
+		return true;
+	}
+
+	const bool bUseSkeletalMesh = SkeletalMesh != nullptr;
+
+	bool bClear = Output.TileCoord == FIntPoint(0,0);
 	const uint32 GroupCount = InInfo.GroomAsset->GetNumHairGroups();
 	for (uint32 GroupIndex = 0; GroupIndex < GroupCount; ++GroupIndex)
 	{
@@ -1001,7 +1062,7 @@ static void InternalBuildStrandsTextures_GPU(
 			InternalGenerateHairStrandsTextures(
 				GraphBuilder,
 				ShaderMap,
-				DebugShaderData,
+				ShaderPrintData,
 				bClear,
 				InInfo.MaxTracingDistance,
 				InInfo.TracingDirection,
@@ -1037,18 +1098,32 @@ static void InternalBuildStrandsTextures_GPU(
 				GroupData.Strands.RestResource->BulkData.PointCount,
 				GroupIndex,
 
-				DepthTestTexture,
-				DepthTexture[0],
-				TangentTexture[0],
-				CoverageTexture[0],
-				AttributeTexture[0],
-				MaterialTexture[0],
-				TriangleMaskTexture);
+				Output);
+
 
 			bClear = false;
 		}
 	}
 
+	// Increment processed tile
+	++Output.TileCoord.X;
+	if (Output.TileCoord.X >= Output.TileCount.X)
+	{
+		Output.TileCoord.X = 0;
+		++Output.TileCoord.Y;
+	}
+
+	return true;
+}
+
+static FHairStrandsRDGTextures DilateTextures(FRDGBuilder& GraphBuilder, FGlobalShaderMap* ShaderMap, FHairStrandsRDGTextures& In)
+{
+	const FIntPoint OutputResolution = In.CoverageTexture->Desc.Extent;
+
+	FHairStrandsRDGTextures Out[2];
+	Out[0] = In;
+	Out[1] = CreateTextures(GraphBuilder, OutputResolution);
+
 	uint32 SourceIndex = 0;
 	uint32 TargetIndex = 0;
 	const uint32 DilationPassCount = FMath::Max(0, GHairStrandsTextureDilationPassCount);
@@ -1060,31 +1135,37 @@ static void InternalBuildStrandsTextures_GPU(
 			GraphBuilder,
 			ShaderMap,
 			OutputResolution,
-			TriangleMaskTexture,
+			In.TriangleMaskTexture,
 
-			DepthTexture[SourceIndex],
-			CoverageTexture[SourceIndex],
-			TangentTexture[SourceIndex],
-			AttributeTexture[SourceIndex],
-			MaterialTexture[SourceIndex],
+			Out[SourceIndex].DepthTexture,
+			Out[SourceIndex].CoverageTexture,
+			Out[SourceIndex].TangentTexture,
+			Out[SourceIndex].AttributeTexture,
+			Out[SourceIndex].MaterialTexture,
 
-			DepthTexture[TargetIndex],
-			CoverageTexture[TargetIndex],
-			TangentTexture[TargetIndex],
-			AttributeTexture[TargetIndex],
-			MaterialTexture[TargetIndex]);
+			Out[TargetIndex].DepthTexture,
+			Out[TargetIndex].CoverageTexture,
+			Out[TargetIndex].TangentTexture,
+			Out[TargetIndex].AttributeTexture,
+			Out[TargetIndex].MaterialTexture);
 
 		SourceIndex = TargetIndex;
 	}
 
-	AddEnqueueCopyPass(GraphBuilder, Output.Depth.Get(), DepthTexture[TargetIndex]);
-	AddEnqueueCopyPass(GraphBuilder, Output.Coverage.Get(), CoverageTexture[TargetIndex]);
-	AddEnqueueCopyPass(GraphBuilder, Output.Tangent.Get(), TangentTexture[TargetIndex]);
-	AddEnqueueCopyPass(GraphBuilder, Output.Attribute.Get(), AttributeTexture[TargetIndex]);
-	AddEnqueueCopyPass(GraphBuilder, Output.Material.Get(), MaterialTexture[TargetIndex]);
+	Out[SourceIndex].TriangleMaskTexture = In.TriangleMaskTexture;
+	return Out[SourceIndex];
+}
+
+static void CopyTextures(FRDGBuilder& GraphBuilder, const FHairStrandsRDGTextures& In, FStrandsTexturesReadback& Out)
+{
+	AddEnqueueCopyPass(GraphBuilder, Out.Depth.Get(),		In.DepthTexture);
+	AddEnqueueCopyPass(GraphBuilder, Out.Coverage.Get(),	In.CoverageTexture);
+	AddEnqueueCopyPass(GraphBuilder, Out.Tangent.Get(),		In.TangentTexture);
+	AddEnqueueCopyPass(GraphBuilder, Out.Attribute.Get(),	In.AttributeTexture);
+	AddEnqueueCopyPass(GraphBuilder, Out.Material.Get(),	In.MaterialTexture);
 }
 
-void CopyReadbackToTexture(TUniquePtr<FRHIGPUTextureReadback>& In, UTexture2D* Out)
+static void CopyReadbackToTexture(TUniquePtr<FRHIGPUTextureReadback>& In, UTexture2D* Out)
 {
 #if WITH_EDITORONLY_DATA
 	FRHICommandListImmediate& RHICmdList = FRHICommandListExecutor::GetImmediateCommandList();
@@ -1118,71 +1199,129 @@ void CopyReadbackToTexture(TUniquePtr<FRHIGPUTextureReadback>& In, UTexture2D* O
 
 struct FStrandsTexturesQuery
 {
-	FStrandsTexturesInfo Info;
-	FStrandsTexturesOutput Output;
+	FStrandsTexturesInfo		Info;
+	FStrandsTexturesOutput		Output;
+};
+
+struct FStrandsTexturesProcess
+{
+	FStrandsTexturesInfo		Info;
+	FHairStrandsRTs				InProgress;
+	FStrandsTexturesReadback	Readback;
+	FStrandsTexturesOutput		Output;
 };
-TQueue<FStrandsTexturesQuery> GStrandsTexturesQueries;
-TQueue<FStrandsTexturesReadback*> GStrandsTexturesReadbacks;
+
+TQueue<FStrandsTexturesQuery>		GStrandsTexturesQueries;
+TQueue<FStrandsTexturesProcess*>	GStrandsTexturesProcessQueue;
+TQueue<FStrandsTexturesProcess*>	GStrandsTexturesReadbacks;
 
 bool HasHairStrandsTexturesQueries()
 {
-	return !GStrandsTexturesQueries.IsEmpty() || !GStrandsTexturesReadbacks.IsEmpty();
+	return !GStrandsTexturesQueries.IsEmpty() || !GStrandsTexturesProcessQueue.IsEmpty() || !GStrandsTexturesReadbacks.IsEmpty();
 }
 
 void RunHairStrandsTexturesQueries(FRDGBuilder& GraphBuilder, FGlobalShaderMap* ShaderMap, const struct FShaderPrintData* DebugShaderData)
 {
-	// Readback data
+	// Operations are ordered in reverse to ensure they are processed on independent frames to avoid TDR on heavy grooms.
+	
+	// 3. Readback data
 	{
-		TArray<FStrandsTexturesReadback*> NotReadyQueue;
-		FStrandsTexturesReadback* R = nullptr;
+		TArray<FStrandsTexturesProcess*> NotReady;
+		FStrandsTexturesProcess* R = nullptr;
 		while (GStrandsTexturesReadbacks.Dequeue(R))
 		{
 			if (R)
 			{
-				if (R->IsReady() || ++R->NotReadyFrameCount > 32)
+				if (R->Readback.IsReady() || ++R->Readback.NotReadyFrameCount > 32)
 				{
-					CopyReadbackToTexture(R->Depth,    R->Output.Depth);
-					CopyReadbackToTexture(R->Tangent,  R->Output.Tangent);
-					CopyReadbackToTexture(R->Coverage, R->Output.Coverage);
-					CopyReadbackToTexture(R->Attribute,R->Output.Attribute);
-					CopyReadbackToTexture(R->Material, R->Output.Material);
-
-					R->Depth    = nullptr;
-					R->Tangent  = nullptr;
-					R->Coverage = nullptr;
-					R->Attribute= nullptr;
-					R->Material = nullptr;
-
+					CopyReadbackToTexture(R->Readback.Depth,    R->Output.Depth);
+					CopyReadbackToTexture(R->Readback.Tangent,  R->Output.Tangent);
+					CopyReadbackToTexture(R->Readback.Coverage, R->Output.Coverage);
+					CopyReadbackToTexture(R->Readback.Attribute,R->Output.Attribute);
+					CopyReadbackToTexture(R->Readback.Material, R->Output.Material);
+	
+					R->Readback.Depth    = nullptr;
+					R->Readback.Tangent  = nullptr;
+					R->Readback.Coverage = nullptr;
+					R->Readback.Attribute= nullptr;
+					R->Readback.Material = nullptr;
+	
 					delete R;
 				}
 				else
 				{
-					NotReadyQueue.Add(R);
+					NotReady.Add(R);
 				}
 			}
 		}
-
-		for (FStrandsTexturesReadback* E : NotReadyQueue)
+	
+		for (FStrandsTexturesProcess* E : NotReady)
 		{
 			GStrandsTexturesReadbacks.Enqueue(E);
 		}
 	}
 
-	// Render textures
+	// 2. Process passes
+	{
+		TArray<FStrandsTexturesProcess*> NotCompleted;
+		FStrandsTexturesProcess* R = nullptr;
+		while (GStrandsTexturesProcessQueue.Dequeue(R))
+		{
+			if (R)
+			{
+				FHairStrandsRDGTextures Textures = RegisterTextures(GraphBuilder, R->InProgress);
+
+				TraceTextures(GraphBuilder, ShaderMap, R->Info, DebugShaderData, Textures);
+
+				if (Textures.IsCompleted())
+				{
+					Textures = DilateTextures(GraphBuilder, ShaderMap, Textures);
+					CopyTextures(GraphBuilder, Textures, R->Readback);
+					GStrandsTexturesReadbacks.Enqueue(R);
+				}
+				else
+				{
+					ExportTextures(GraphBuilder, Textures, R->InProgress);
+					NotCompleted.Add(R);
+				}
+			}
+		}
+		
+		for (FStrandsTexturesProcess* E : NotCompleted)
+		{
+			GStrandsTexturesProcessQueue.Enqueue(E);
+		}
+	}
+
+	// 1. Create textures
 	{
 		FStrandsTexturesQuery Q;
 		while (GStrandsTexturesQueries.Dequeue(Q))
 		{
-			FStrandsTexturesReadback* R = new FStrandsTexturesReadback();
-			R->Depth     = MakeUnique<FRHIGPUTextureReadback>(TEXT("Readback.Depth"));
-			R->Tangent   = MakeUnique<FRHIGPUTextureReadback>(TEXT("Readback.Tangent"));
-			R->Coverage  = MakeUnique<FRHIGPUTextureReadback>(TEXT("Readback.Coverage"));
-			R->Attribute = MakeUnique<FRHIGPUTextureReadback>(TEXT("Readback.Attribute"));
-			R->Material  = MakeUnique<FRHIGPUTextureReadback>(TEXT("Readback.Material"));		
-			R->Output    = Q.Output;
-
-			InternalBuildStrandsTextures_GPU(GraphBuilder, ShaderMap, Q.Info, *R, DebugShaderData);
-			GStrandsTexturesReadbacks.Enqueue(R);
+			FStrandsTexturesProcess* R = new FStrandsTexturesProcess();
+			R->Readback.Depth     = MakeUnique<FRHIGPUTextureReadback>(TEXT("Readback.Depth"));
+			R->Readback.Tangent   = MakeUnique<FRHIGPUTextureReadback>(TEXT("Readback.Tangent"));
+			R->Readback.Coverage  = MakeUnique<FRHIGPUTextureReadback>(TEXT("Readback.Coverage"));
+			R->Readback.Attribute = MakeUnique<FRHIGPUTextureReadback>(TEXT("Readback.Attribute"));
+			R->Readback.Material  = MakeUnique<FRHIGPUTextureReadback>(TEXT("Readback.Material"));		
+			R->Output			  = Q.Output;
+			R->Info				  = Q.Info;
+
+			FHairStrandsRDGTextures Textures = CreateTextures(GraphBuilder, FIntPoint(FMath::Clamp(Q.Info.Resolution, 512u, 16384u), FMath::Clamp(Q.Info.Resolution, 512u, 16384u)));
+
+			TraceTextures(GraphBuilder, ShaderMap, Q.Info, DebugShaderData, Textures);
+
+			if (Textures.IsCompleted())
+			{
+				Textures = DilateTextures(GraphBuilder, ShaderMap, Textures);
+				CopyTextures(GraphBuilder, Textures, R->Readback);
+				GStrandsTexturesReadbacks.Enqueue(R);
+			}
+			else
+			{
+				ExportTextures(GraphBuilder, Textures, R->InProgress);
+				GStrandsTexturesProcessQueue.Enqueue(R);
+			}
 		}
 	}
 }
