Fix hair warning in log when a groom is attached to a skel. mesh.

This happens due to non-deterministic destruction order.

diff --git a/Engine/Plugins/Runtime/HairStrands/Source/HairStrandsCore/Private/GroomComponent.cpp b/Engine/Plugins/Runtime/HairStrands/Source/HairStrandsCore/Private/GroomComponent.cpp
--- a/Engine/Plugins/Runtime/HairStrands/Source/HairStrandsCore/Private/GroomComponent.cpp
+++ b/Engine/Plugins/Runtime/HairStrands/Source/HairStrandsCore/Private/GroomComponent.cpp
@@ -3108,7 +3108,7 @@ void UGroomComponent::OnComponentDestroyed(bool bDestroyingHierarchy)
 void UGroomComponent::OnAttachmentChanged()
 {
 	Super::OnAttachmentChanged();
-	if (GroomAsset && !IsBeingDestroyed() && HasBeenCreated())
+	if (GroomAsset && !IsBeingDestroyed() && HasBeenCreated() && IsValidChecked(this))
 	{
 		UMeshComponent* NewMeshComponent = Cast<UMeshComponent>(GetAttachParent());
 		const bool bHasAttachmentChanged = RegisteredMeshComponent != NewMeshComponent;
