Fix missing cards geometry when switching cards/strands representation at runtime.

diff --git a/Engine/Plugins/Runtime/HairStrands/Source/HairStrandsCore/Private/GroomAsset.cpp b/Engine/Plugins/Runtime/HairStrands/Source/HairStrandsCore/Private/GroomAsset.cpp
--- a/Engine/Plugins/Runtime/HairStrands/Source/HairStrandsCore/Private/GroomAsset.cpp
+++ b/Engine/Plugins/Runtime/HairStrands/Source/HairStrandsCore/Private/GroomAsset.cpp
@@ -2272,7 +2272,8 @@ bool UGroomAsset::HasValidCardsData(uint32 GroupIndex) const
 	const FHairGroupsLOD& GroupsLOD = HairGroupsLOD[GroupIndex];
 	for (int32 LODIt = 0; LODIt < GroupsLOD.LODs.Num(); ++LODIt)
 	{
-		if (GroupsLOD.LODs[LODIt].GeometryType == EGroomGeometryType::Cards)
+		// Note: We don't condition/filter resource allocation based on GeometryType == EGroomGeometryType::Cards, because 
+		// geometry type can be switched at runtime (between cards/strands)
 		{
 			int32 SourceIt = 0;
 			if (const FHairGroupsCardsSourceDescription* Desc = GetSourceDescription(HairGroupsCards, GroupIndex, LODIt, SourceIt))
@@ -2319,7 +2320,8 @@ bool UGroomAsset::CacheCardsGeometry(uint32 GroupIndex, const FString& StrandsKe
 		int32 SourceIt = 0;
 		if (FHairGroupsCardsSourceDescription* Desc = GetSourceDescription(HairGroupsCards, GroupIndex, LODIt, SourceIt))
 		{
-			if (GroupsLOD.LODs[LODIt].GeometryType == EGroomGeometryType::Cards)
+			// Note: We don't condition/filter resource allocation based on GeometryType == EGroomGeometryType::Cards, because 
+			// geometry type can be switched at runtime (between cards/strands)
 			{
 				UStaticMesh* CardsMesh = nullptr;
 				if (Desc->SourceType == EHairCardsSourceType::Procedural)
diff --git a/Engine/Plugins/Runtime/HairStrands/Source/HairStrandsCore/Private/GroomBindingBuilder.cpp b/Engine/Plugins/Runtime/HairStrands/Source/HairStrandsCore/Private/GroomBindingBuilder.cpp
--- a/Engine/Plugins/Runtime/HairStrands/Source/HairStrandsCore/Private/GroomBindingBuilder.cpp
+++ b/Engine/Plugins/Runtime/HairStrands/Source/HairStrandsCore/Private/GroomBindingBuilder.cpp
@@ -57,7 +57,7 @@ static FAutoConsoleVariableRef CVarHairStrandsBindingBuilderWarningEnable(TEXT("
 FString FGroomBindingBuilder::GetVersion()
 {
 	// Important to update the version when groom building changes
-	return TEXT("2e");
+	return TEXT("2f");
 }
 
 ///////////////////////////////////////////////////////////////////////////////////////////////////
