Fix CIS static analysis error.

diff --git a/Engine/Plugins/Runtime/HairStrands/Source/HairStrandsCore/Private/GroomTextureBuilder.cpp b/Engine/Plugins/Runtime/HairStrands/Source/HairStrandsCore/Private/GroomTextureBuilder.cpp
--- a/Engine/Plugins/Runtime/HairStrands/Source/HairStrandsCore/Private/GroomTextureBuilder.cpp
+++ b/Engine/Plugins/Runtime/HairStrands/Source/HairStrandsCore/Private/GroomTextureBuilder.cpp
@@ -673,16 +673,16 @@ static FIntPoint GetTileResolution(FIntPoint In)
 
 struct FHairStrandsRDGTextures
 {
-	FIntPoint TileCount; // Total tile count
-	FIntPoint TileCoord; // Current tile coord
-
-	FRDGTextureRef DepthTexture;
-	FRDGTextureRef CoverageTexture;
-	FRDGTextureRef TangentTexture;
-	FRDGTextureRef AttributeTexture;
-	FRDGTextureRef MaterialTexture;
-	FRDGTextureRef TriangleMaskTexture;
-	FRDGTextureRef DepthTestTexture;
+	FIntPoint TileCount = FIntPoint::ZeroValue; // Total tile count
+	FIntPoint TileCoord = FIntPoint::ZeroValue; // Current tile coord
+
+	FRDGTextureRef DepthTexture = nullptr;
+	FRDGTextureRef CoverageTexture = nullptr;
+	FRDGTextureRef TangentTexture = nullptr;
+	FRDGTextureRef AttributeTexture = nullptr;
+	FRDGTextureRef MaterialTexture = nullptr;
+	FRDGTextureRef TriangleMaskTexture = nullptr;
+	FRDGTextureRef DepthTestTexture = nullptr;
 
 	bool IsCompleted() const { return TileCoord.X >= TileCount.X || TileCoord.Y >= TileCount.Y;  }
 };
