Add different mode for computing hair strands bound when attached to a skel. mesh.

This is controlled with r.HairStrands.BoundMode.

diff --git a/Engine/Plugins/Runtime/HairStrands/Source/HairStrandsCore/Private/GroomComponent.cpp b/Engine/Plugins/Runtime/HairStrands/Source/HairStrandsCore/Private/GroomComponent.cpp
--- a/Engine/Plugins/Runtime/HairStrands/Source/HairStrandsCore/Private/GroomComponent.cpp
+++ b/Engine/Plugins/Runtime/HairStrands/Source/HairStrandsCore/Private/GroomComponent.cpp
@@ -51,6 +51,9 @@ static FAutoConsoleVariableRef CVarUseProxyLocalToWorld(TEXT("r.HairStrands.UseP
 static bool GHairStrands_Streaming_Prediction = false;
 static FAutoConsoleVariableRef CVarHairStrands_Streaming_Prediction(TEXT("r.HairStrands.Streaming.Prediction"), GHairStrands_Streaming_Prediction, TEXT("Enable LOD streaming prediction."));
 
+static int32 GHairStrands_BoundsMode = 0;
+static FAutoConsoleVariableRef CVarHairStrands_BoundsMode(TEXT("r.HairStrands.BoundMode"), GHairStrands_BoundsMode, TEXT("Define how hair bound are computed at runtime when attached to a skel. mesh.\n 0: Use skel.mesh extented with grooms bounds.\n 1: Use skel.mesh bounds.\n 2: Use skel.mesh extented with grooms bounds (conservative)"));
+
 #define LOCTEXT_NAMESPACE "GroomComponent"
 
 #define USE_HAIR_TRIANGLE_STRIP 0
@@ -1826,7 +1829,7 @@ FPrimitiveSceneProxy* UGroomComponent::CreateSceneProxy()
 FBoxSphereBounds UGroomComponent::CalcBounds(const FTransform& InLocalToWorld) const
 {
 	if (GroomAsset && GroomAsset->GetNumHairGroups() > 0)
-	{		
+	{
 		FBox LocalHairBound(EForceInit::ForceInitToZero);
 		if (!GroomCacheBuffers.IsValid())
 		{
@@ -1856,11 +1859,27 @@ FBoxSphereBounds UGroomComponent::CalcBounds(const FTransform& InLocalToWorld) c
 			LocalHairBound = Buffers->GetBoundingBox();
 		}
 
-		// If the attachment is done onto a skel. mesh:
+		// If the attachment is done onto a skel. mesh, by default (i.e., GHairStrands_BoundsMode == 0):
 		// * If the attachment is simple, use the skel. mesh expanded by the groom bounds.
 		// * If the attachment is 'relative' to the skel mesh (using a bone anchor, which is provided with AttachmentName), 
 		//   we use the simple groom bound.
-		if (RegisteredMeshComponent && AttachmentName.IsEmpty())
+		// Otherwise:
+		// * GHairStrands_BoundsMode=1 : use the skel mesh bounds
+		// * GHairStrands_BoundsMode=2 : use the skel mesh bounds + groom bounds. This is more conservative compares to GHairStrands_BoundsMode=0.
+		if (RegisteredMeshComponent && GHairStrands_BoundsMode == 1)
+		{
+			return RegisteredMeshComponent->Bounds;
+		}
+		else if (RegisteredMeshComponent && GHairStrands_BoundsMode == 2)
+		{
+			const FVector3d GroomExtends = LocalHairBound.GetExtent();
+			const float BoundExtraRadius = 0.5f * static_cast<float>(FMath::Max(0.0, FMath::Max3(GroomExtends.X, GroomExtends.Y, GroomExtends.Z)));
+			FBox EffectiveBound = RegisteredMeshComponent->Bounds.GetBox();
+			EffectiveBound.Min -= FVector3d(BoundExtraRadius);
+			EffectiveBound.Max += FVector3d(BoundExtraRadius);
+			return FBoxSphereBounds(EffectiveBound);
+		}
+		else if (RegisteredMeshComponent && AttachmentName.IsEmpty())
 		{
 			const FBox LocalSkeletalBound = RegisteredMeshComponent->CalcBounds(FTransform::Identity).GetBox();
 
