Prevent crash when a MH is imported, but skin cache is not enabled (yet).

This CL ensure in such case, the groom LOD is skipped to avoid any crash.

diff --git a/Engine/Plugins/Runtime/HairStrands/Source/HairStrandsCore/Private/GroomManager.cpp b/Engine/Plugins/Runtime/HairStrands/Source/HairStrandsCore/Private/GroomManager.cpp
--- a/Engine/Plugins/Runtime/HairStrands/Source/HairStrandsCore/Private/GroomManager.cpp
+++ b/Engine/Plugins/Runtime/HairStrands/Source/HairStrandsCore/Private/GroomManager.cpp
@@ -897,6 +897,12 @@ static void RunHairLODSelection(
 				Instance->Strands.bIsCullingEnabled = bCullingEnable;
 			}
 
+			// Ensure that if the binding type is set to Skinning (or has RBF enabled), the skel. mesh is valid and support skin cache.
+			if ((Instance->Guides.bHasGlobalInterpolation || Instance->BindingType == EHairBindingType::Skinning) && MeshLODIndex < 0)
+			{
+				return false;
+			}
+
 			// If requested LOD's resources are not ready yet, and the previous LOD was invalid or if the MeshLODIndex has changed 
 			// (which would required extra data loading) change the geometry to be invalid, so that it don't get processed this frame.
 			const bool bIsLODValid = PrevLODIndex >= 0;
