Fix hair cards invalidation when changing cards group index.

This CL fixes a missing invalidation when a hair cards group index changes. However, during the same session, it can happen that the groom binding validation fails due to stale data. Restarting the editor fix the issue.

diff --git a/Engine/Plugins/Runtime/HairStrands/Source/HairStrandsCore/Private/GroomAsset.cpp b/Engine/Plugins/Runtime/HairStrands/Source/HairStrandsCore/Private/GroomAsset.cpp
--- a/Engine/Plugins/Runtime/HairStrands/Source/HairStrandsCore/Private/GroomAsset.cpp
+++ b/Engine/Plugins/Runtime/HairStrands/Source/HairStrandsCore/Private/GroomAsset.cpp
@@ -1755,7 +1755,7 @@ void UGroomAsset::SetHairWidth(float Width)
 // differences, etc.) replace the version GUID below with a new one.
 // In case of merge conflicts with DDC versions, you MUST generate a new GUID
 // and set this new GUID as the version.
-#define GROOM_DERIVED_DATA_VERSION TEXT("5251F4CEC0B84A6B8175D0E33755556A")
+#define GROOM_DERIVED_DATA_VERSION TEXT("02B028B2ADF345B1B68BA35C872A849C")
 
 #if WITH_EDITORONLY_DATA
 
@@ -1846,7 +1846,7 @@ namespace GroomDerivedDataCacheUtils
 		for (int32 Index = 0; Index < SourceDescriptions.Num(); ++Index)
 		{
 			FHairGroupsCardsSourceDescription& Desc = SourceDescriptions[Index];
-			if (Desc.GroupIndex != GroupIndex && Desc.LODIndex < 0)
+			if (Desc.GroupIndex != GroupIndex || Desc.LODIndex < 0)
 			{
 				continue;
 			}
@@ -1898,7 +1898,7 @@ namespace GroomDerivedDataCacheUtils
 		for (int32 Index = 0; Index < SourceDescriptions.Num(); ++Index)
 		{
 			FHairGroupsMeshesSourceDescription& Desc = SourceDescriptions[Index];
-			if (Desc.GroupIndex != GroupIndex && Desc.LODIndex < 0)
+			if (Desc.GroupIndex != GroupIndex || Desc.LODIndex < 0)
 			{
 				continue;
 			}
diff --git a/Engine/Plugins/Runtime/HairStrands/Source/HairStrandsCore/Private/GroomBindingBuilder.cpp b/Engine/Plugins/Runtime/HairStrands/Source/HairStrandsCore/Private/GroomBindingBuilder.cpp
--- a/Engine/Plugins/Runtime/HairStrands/Source/HairStrandsCore/Private/GroomBindingBuilder.cpp
+++ b/Engine/Plugins/Runtime/HairStrands/Source/HairStrandsCore/Private/GroomBindingBuilder.cpp
@@ -2028,8 +2028,10 @@ static bool InternalBuildBinding_CPU(UGroomBindingAsset* BindingAsset, bool bIni
 				{
 					FHairStrandsDatas LODGuidesData;
 					const bool bIsValid = GroomAsset->GetHairCardsGuidesDatas(GroupIndex, CardsLODIt, LODGuidesData);
-					check(bIsValid);
-					InitHairStrandsRootData(OutData.CardsRootData[CardsLODIt], &LODGuidesData, MeshLODCount, NumInterpolationPoints);
+					if (bIsValid)
+					{
+						InitHairStrandsRootData(OutData.CardsRootData[CardsLODIt], &LODGuidesData, MeshLODCount, NumInterpolationPoints);
+					}
 				}
 			}
 			++GroupIndex;
@@ -2106,18 +2108,19 @@ static bool InternalBuildBinding_CPU(UGroomBindingAsset* BindingAsset, bool bIni
 				{
 					FHairStrandsDatas LODGuidesData;
 					const bool bIsValid = GroomAsset->GetHairCardsGuidesDatas(GroupIt, CardsLODIt, LODGuidesData);
-					check(bIsValid);
-				
-					bSucceed = GroomBinding_RootProjection::Project(
-						LODGuidesData,
-						TargetMeshData.Get(),
-						TransferredPositions,
-						OutHairGroupDatas[GroupIt].CardsRootData[CardsLODIt]);
-
-					if (!bSucceed) 
+					if (bIsValid)
 					{
-						UE_LOG(LogHairStrands, Error, TEXT("[Groom] Binding asset could not be built. Some cards guide roots are not close enough to the target mesh to be projected onto it."));
-						return false; 
+						bSucceed = GroomBinding_RootProjection::Project(
+							LODGuidesData,
+							TargetMeshData.Get(),
+							TransferredPositions,
+							OutHairGroupDatas[GroupIt].CardsRootData[CardsLODIt]);
+
+						if (!bSucceed) 
+						{
+							UE_LOG(LogHairStrands, Error, TEXT("[Groom] Binding asset could not be built. Some cards guide roots are not close enough to the target mesh to be projected onto it."));
+							return false; 
+						}
 					}
 				}
 
