Fix for RHI validation failure in hair rendering.

diff --git a/Engine/Plugins/Runtime/HairStrands/Source/HairStrandsCore/Private/GroomManager.cpp b/Engine/Plugins/Runtime/HairStrands/Source/HairStrandsCore/Private/GroomManager.cpp
--- a/Engine/Plugins/Runtime/HairStrands/Source/HairStrandsCore/Private/GroomManager.cpp
+++ b/Engine/Plugins/Runtime/HairStrands/Source/HairStrandsCore/Private/GroomManager.cpp
@@ -57,13 +57,6 @@ DEFINE_LOG_CATEGORY_STATIC(LogGroomManager, Log, All);
 
 ///////////////////////////////////////////////////////////////////////////////////////////////////
 
-namespace HairTransition
-{
-	void TransitToSRV(FRDGBuilder& GraphBuilder, FRDGBufferSRVRef InBuffer, ERDGPassFlags Flags);
-}
-
-///////////////////////////////////////////////////////////////////////////////////////////////////
-
 void FHairGroupInstance::FCards::FLOD::InitVertexFactory()
 {
 	VertexFactory->InitResources();
@@ -317,17 +310,17 @@ static void RunInternalHairStrandsInterpolation(
 						// Add manual transition for the GPU solver as Niagara does not track properly the RDG buffer, and so doesn't issue the correct transitions
 						if (MeshLODIndex >= 0)
 						{
-							HairTransition::TransitToSRV(GraphBuilder, Register(GraphBuilder, Instance->Guides.DeformedRootResource->LODs[MeshLODIndex].DeformedRootTrianglePosition0Buffer, ERDGImportedBufferFlags::CreateSRV).SRV, ERDGPassFlags::Compute);
-							HairTransition::TransitToSRV(GraphBuilder, Register(GraphBuilder, Instance->Guides.DeformedRootResource->LODs[MeshLODIndex].DeformedRootTrianglePosition1Buffer, ERDGImportedBufferFlags::CreateSRV).SRV, ERDGPassFlags::Compute);
-							HairTransition::TransitToSRV(GraphBuilder, Register(GraphBuilder, Instance->Guides.DeformedRootResource->LODs[MeshLODIndex].DeformedRootTrianglePosition2Buffer, ERDGImportedBufferFlags::CreateSRV).SRV, ERDGPassFlags::Compute);
+							GraphBuilder.UseExternalAccessMode(Register(GraphBuilder, Instance->Guides.DeformedRootResource->LODs[MeshLODIndex].DeformedRootTrianglePosition0Buffer, ERDGImportedBufferFlags::CreateSRV).Buffer, ERHIAccess::SRVMask);
+							GraphBuilder.UseExternalAccessMode(Register(GraphBuilder, Instance->Guides.DeformedRootResource->LODs[MeshLODIndex].DeformedRootTrianglePosition1Buffer, ERDGImportedBufferFlags::CreateSRV).Buffer, ERHIAccess::SRVMask);
+							GraphBuilder.UseExternalAccessMode(Register(GraphBuilder, Instance->Guides.DeformedRootResource->LODs[MeshLODIndex].DeformedRootTrianglePosition2Buffer, ERDGImportedBufferFlags::CreateSRV).Buffer, ERHIAccess::SRVMask);
 
 							if (bGlobalDeformationEnable)
 							{
-								HairTransition::TransitToSRV(GraphBuilder, Register(GraphBuilder, Instance->Guides.DeformedRootResource->LODs[MeshLODIndex].DeformedSamplePositionsBuffer, ERDGImportedBufferFlags::CreateSRV).SRV, ERDGPassFlags::Compute);
-								HairTransition::TransitToSRV(GraphBuilder, Register(GraphBuilder, Instance->Guides.DeformedRootResource->LODs[MeshLODIndex].MeshSampleWeightsBuffer, ERDGImportedBufferFlags::CreateSRV).SRV, ERDGPassFlags::Compute);
+								GraphBuilder.UseExternalAccessMode(Register(GraphBuilder, Instance->Guides.DeformedRootResource->LODs[MeshLODIndex].DeformedSamplePositionsBuffer, ERDGImportedBufferFlags::CreateSRV).Buffer, ERHIAccess::SRVMask);
+								GraphBuilder.UseExternalAccessMode(Register(GraphBuilder, Instance->Guides.DeformedRootResource->LODs[MeshLODIndex].MeshSampleWeightsBuffer, ERDGImportedBufferFlags::CreateSRV).Buffer, ERHIAccess::SRVMask);
 							}
 						}
-						HairTransition::TransitToSRV(GraphBuilder, Register(GraphBuilder, Instance->Guides.DeformedResource->GetPositionOffsetBuffer(FHairStrandsDeformedResource::EFrameType::Current), ERDGImportedBufferFlags::CreateSRV).SRV, ERDGPassFlags::Compute);
+						GraphBuilder.UseExternalAccessMode(Register(GraphBuilder, Instance->Guides.DeformedResource->GetPositionOffsetBuffer(FHairStrandsDeformedResource::EFrameType::Current), ERDGImportedBufferFlags::CreateSRV).Buffer, ERHIAccess::SRVMask);
 					}
 					else if (bSimulationEnable)
 					{
@@ -341,7 +334,7 @@ static void RunInternalHairStrandsInterpolation(
 							Instance->Guides.DeformedResource);
 
 						// Add manual transition for the GPU solver as Niagara does not track properly the RDG buffer, and so doesn't issue the correct transitions
-						HairTransition::TransitToSRV(GraphBuilder, Register(GraphBuilder, Instance->Guides.DeformedResource->GetPositionOffsetBuffer(FHairStrandsDeformedResource::EFrameType::Current), ERDGImportedBufferFlags::CreateSRV).SRV, ERDGPassFlags::Compute);
+						GraphBuilder.UseExternalAccessMode(Register(GraphBuilder, Instance->Guides.DeformedResource->GetPositionOffsetBuffer(FHairStrandsDeformedResource::EFrameType::Current), ERDGImportedBufferFlags::CreateSRV).Buffer, ERHIAccess::SRVMask);
 					}
 				}
 			}
diff --git a/Engine/Plugins/Runtime/HairStrands/Source/HairStrandsCore/Private/GroomResources.cpp b/Engine/Plugins/Runtime/HairStrands/Source/HairStrandsCore/Private/GroomResources.cpp
--- a/Engine/Plugins/Runtime/HairStrands/Source/HairStrandsCore/Private/GroomResources.cpp
+++ b/Engine/Plugins/Runtime/HairStrands/Source/HairStrandsCore/Private/GroomResources.cpp
@@ -57,11 +57,6 @@ EHairResourceLoadingType GetHairResourceLoadingType(EHairGeometryType InGeometry
 #endif
 }
 
-namespace HairTransition
-{
-	void TransitToSRV(FRDGBuilder& GraphBuilder, FRDGBufferSRVRef InBuffer, ERDGPassFlags Flags);	
-}
-
 enum class EHairResourceUsageType : uint8
 {
 	Static,
@@ -691,13 +686,12 @@ void FHairCardsDeformedResource::InternalAllocate(FRDGBuilder& GraphBuilder)
 
 		// Manually transit to SRVs, in case of the cards are rendered not visible, but still rendered (in shadows for instance). In such a case, the cards deformation pass is not called, and thus the 
 		// buffers are never transit from UAV (clear) to SRV for rasterization. 
-		HairTransition::TransitToSRV(GraphBuilder, Register(GraphBuilder, DeformedPositionBuffer[0], ERDGImportedBufferFlags::CreateSRV).SRV, ERDGPassFlags::Raster);
-		HairTransition::TransitToSRV(GraphBuilder, Register(GraphBuilder, DeformedPositionBuffer[1], ERDGImportedBufferFlags::CreateSRV).SRV, ERDGPassFlags::Raster);
+		GraphBuilder.UseExternalAccessMode(Register(GraphBuilder, DeformedPositionBuffer[0], ERDGImportedBufferFlags::CreateSRV).Buffer, ERHIAccess::SRVMask);
+		GraphBuilder.UseExternalAccessMode(Register(GraphBuilder, DeformedPositionBuffer[1], ERDGImportedBufferFlags::CreateSRV).Buffer, ERHIAccess::SRVMask);
 	}
 
 	FRDGImportedBuffer CardsDeformedNormalRDGBuffer = Register(GraphBuilder, DeformedNormalBuffer, ERDGImportedBufferFlags::CreateSRV);
-	HairTransition::TransitToSRV(GraphBuilder, CardsDeformedNormalRDGBuffer.SRV, ERDGPassFlags::Raster);
-	GraphBuilder.SetBufferAccessFinal(CardsDeformedNormalRDGBuffer.Buffer, ERHIAccess::SRVGraphics);
+	GraphBuilder.UseExternalAccessMode(CardsDeformedNormalRDGBuffer.Buffer, ERHIAccess::SRVMask);
 }
 
 void FHairCardsDeformedResource::InternalRelease()
@@ -824,7 +818,7 @@ void FHairStrandsRestResource::InternalAllocate(FRDGBuilder& GraphBuilder)
 	TArray<FVector4f> RestOffset;
 	RestOffset.Add((FVector3f)BulkData.GetPositionOffset());// LWC_TODO: precision loss
 	InternalCreateVertexBufferRDG<FHairStrandsPositionOffsetFormat>(GraphBuilder, RestOffset, PositionOffsetBuffer, ToHairResourceDebugName(HAIRSTRANDS_RESOUCE_NAME(CurveType, Hair.StrandsRest_PositionOffsetBuffer), ResourceName), EHairResourceUsageType::Static, ERDGInitialDataFlags::None);
-	HairTransition::TransitToSRV(GraphBuilder, Register(GraphBuilder, PositionOffsetBuffer, ERDGImportedBufferFlags::CreateSRV).SRV, ERDGPassFlags::Compute);
+	GraphBuilder.UseExternalAccessMode(Register(GraphBuilder, PositionOffsetBuffer, ERDGImportedBufferFlags::CreateSRV).Buffer, ERHIAccess::SRVMask);
 }
 
 void AddHairTangentPass(
diff --git a/Engine/Plugins/Runtime/HairStrands/Source/HairStrandsCore/Private/HairStrandsRendering.cpp b/Engine/Plugins/Runtime/HairStrands/Source/HairStrandsCore/Private/HairStrandsRendering.cpp
--- a/Engine/Plugins/Runtime/HairStrands/Source/HairStrandsCore/Private/HairStrandsRendering.cpp
+++ b/Engine/Plugins/Runtime/HairStrands/Source/HairStrandsCore/Private/HairStrandsRendering.cpp
@@ -50,37 +50,6 @@ static FAutoConsoleVariableRef CVarHairRTGeomForceRebuild(TEXT("r.HairStrands.St
 
 ////////////////////////////////////////////////////////////////////////////////////////////////////////////////////////
 
-namespace HairTransition
-{
-	BEGIN_SHADER_PARAMETER_STRUCT(FTransitionParameters, )
-		SHADER_PARAMETER_RDG_BUFFER_SRV(Buffer, Buffer)
-	END_SHADER_PARAMETER_STRUCT()
-
-	void TransitToSRV(FRDGBuilder& GraphBuilder, FRDGBufferSRVRef InBuffer, ERDGPassFlags Flags)
-	{
-		FTransitionParameters* Parameters = GraphBuilder.AllocParameters<FTransitionParameters>();
-		Parameters->Buffer = InBuffer;
-
-		ensure(Flags == ERDGPassFlags::Raster || Flags == ERDGPassFlags::Compute);
-		if (Flags == ERDGPassFlags::Raster)
-		{
-			Flags |= ERDGPassFlags::SkipRenderPass;
-		}
-		Flags |= ERDGPassFlags::NeverCull;
-
-		GraphBuilder.AddPass(
-			RDG_EVENT_NAME("HairStrands::TransitionToSRVPass(%s)", Flags == ERDGPassFlags::Raster ? TEXT("Raster") : TEXT("Compute")),
-			Parameters,
-			Flags,
-			[Parameters](FRHICommandList& RHICmdList)
-			{
-				// Nothing to do, we just want GraphBuilder to force the buffer to be transited to SRV Graphics state.
-			});
-	}
-}
-
-////////////////////////////////////////////////////////////////////////////////////////////////////////////////////////
-
 enum class EHairCardsSimulationType
 {
 	None,
@@ -1031,6 +1000,12 @@ static void AddHairCardsDeformationPass(
 	}
 	
 	FRDGImportedBuffer CardsDeformedNormalBuffer = Register(GraphBuilder, LOD.DeformedResource->DeformedNormalBuffer, ERDGImportedBufferFlags::CreateViews);
+	
+	GraphBuilder.UseInternalAccessMode({
+		CardsDeformedPositionBuffer_Curr.Buffer,
+		CardsDeformedPositionBuffer_Prev.Buffer,
+		CardsDeformedNormalBuffer.Buffer
+	});
 
 	FHairCardsDeformationCS::FParameters* Parameters = GraphBuilder.AllocParameters<FHairCardsDeformationCS::FParameters>();
 	Parameters->LocalToWorld				= FMatrix44f(Instance->GetCurrentLocalToWorld().ToMatrixWithScale());		// LWC_TODO: Precision loss
@@ -1105,10 +1080,11 @@ static void AddHairCardsDeformationPass(
 		AddCopyBufferPass(GraphBuilder, CardsDeformedPositionBuffer_Prev.Buffer, CardsDeformedPositionBuffer_Curr.Buffer);
 	}
 
-	HairTransition::TransitToSRV(GraphBuilder, CardsDeformedPositionBuffer_Curr.SRV, ERDGPassFlags::Raster);
-	HairTransition::TransitToSRV(GraphBuilder, CardsDeformedPositionBuffer_Prev.SRV, ERDGPassFlags::Raster);
-	HairTransition::TransitToSRV(GraphBuilder, CardsDeformedNormalBuffer.SRV, ERDGPassFlags::Raster);
-	GraphBuilder.SetBufferAccessFinal(CardsDeformedNormalBuffer.Buffer, ERHIAccess::SRVGraphics);
+	GraphBuilder.UseExternalAccessMode({
+		CardsDeformedPositionBuffer_Curr.Buffer,
+		CardsDeformedPositionBuffer_Prev.Buffer,
+		CardsDeformedNormalBuffer.Buffer
+	}, ERHIAccess::SRVMask);
 }
 
 ///////////////////////////////////////////////////////////////////////////////////////////////////
