GroomCache: Fixed loading GroomCache that was imported in UE4.

diff --git a/Engine/Plugins/Runtime/HairStrands/Source/HairStrandsCore/Private/GroomCache.cpp b/Engine/Plugins/Runtime/HairStrands/Source/HairStrandsCore/Private/GroomCache.cpp
--- a/Engine/Plugins/Runtime/HairStrands/Source/HairStrandsCore/Private/GroomCache.cpp
+++ b/Engine/Plugins/Runtime/HairStrands/Source/HairStrandsCore/Private/GroomCache.cpp
@@ -7,6 +7,11 @@ void UGroomCache::Serialize(FArchive& Ar)
 {
 	Super::Serialize(Ar);
 
+	if (Ar.IsLoading())
+	{
+		ArchiveVersion = Ar.UEVer();
+	}
+
 	int32 NumChunks = Chunks.Num();
 	Ar << NumChunks;
 
@@ -207,6 +212,11 @@ bool UGroomCache::GetGroomDataAtFrameIndex(int32 FrameIndex, FGroomCacheAnimatio
 	{
 		TRACE_CPUPROFILER_EVENT_SCOPE(UGroomCache::GetGroomDataAtFrameIndex_Serialize);
 		FMemoryReader Ar(TempBytes, true);
+		// Propagate the GroomCache archive version to the memory archive for proper serialization
+		if (ArchiveVersion.IsSet())
+		{
+			Ar.SetUEVer(ArchiveVersion.GetValue());
+		}
 		AnimData.Serialize(Ar);
 	}
 
diff --git a/Engine/Plugins/Runtime/HairStrands/Source/HairStrandsCore/Private/GroomCacheStreamingData.cpp b/Engine/Plugins/Runtime/HairStrands/Source/HairStrandsCore/Private/GroomCacheStreamingData.cpp
--- a/Engine/Plugins/Runtime/HairStrands/Source/HairStrandsCore/Private/GroomCacheStreamingData.cpp
+++ b/Engine/Plugins/Runtime/HairStrands/Source/HairStrandsCore/Private/GroomCacheStreamingData.cpp
@@ -198,6 +198,11 @@ void FGroomCacheStreamingData::UpdateStreamingStatus(bool bAsyncDeletionAllowed)
 						// The bulk data buffer is then serialized into GroomCacheAnimationData
 						TArrayView<uint8> TempView(ReadBuffer, DataSize);
 						FMemoryReaderView Ar(TempView, true);
+						// Propagate the GroomCache archive version to the memory archive for proper serialization
+						if (GroomCache->ArchiveVersion.IsSet())
+						{
+							Ar.SetUEVer(GroomCache->ArchiveVersion.GetValue());
+						}
 						AnimData->Serialize(Ar);
 
 						// We became the owner of ReadBuffer when GetReadResults was called so free it now
diff --git a/Engine/Plugins/Runtime/HairStrands/Source/HairStrandsCore/Public/GroomCache.h b/Engine/Plugins/Runtime/HairStrands/Source/HairStrandsCore/Public/GroomCache.h
--- a/Engine/Plugins/Runtime/HairStrands/Source/HairStrandsCore/Public/GroomCache.h
+++ b/Engine/Plugins/Runtime/HairStrands/Source/HairStrandsCore/Public/GroomCache.h
@@ -50,6 +50,8 @@ public:
 
 	TArray<FGroomCacheChunk>& GetChunks() { return Chunks; }
 
+	TOptional<FPackageFileVersion> ArchiveVersion;
+
 #if WITH_EDITORONLY_DATA
 	/** Import options used for this GroomCache */
 	UPROPERTY(Category = ImportSettings, VisibleAnywhere, Instanced)
