Groom: Added support for using the groom/cache import settings from Python script, if provided. Importing a groom asset and groom cache from the same file must be done in separate import tasks since they can take only one set of import settings.

diff --git a/Engine/Plugins/Runtime/HairStrands/Source/HairStrandsEditor/Private/HairStrandsFactory.cpp b/Engine/Plugins/Runtime/HairStrands/Source/HairStrandsEditor/Private/HairStrandsFactory.cpp
--- a/Engine/Plugins/Runtime/HairStrands/Source/HairStrandsEditor/Private/HairStrandsFactory.cpp
+++ b/Engine/Plugins/Runtime/HairStrands/Source/HairStrandsEditor/Private/HairStrandsFactory.cpp
@@ -2,6 +2,7 @@
 
 #include "HairStrandsFactory.h"
 
+#include "AssetImportTask.h"
 #include "EditorFramework/AssetImportData.h"
 #include "GroomAsset.h"
 #include "GroomAssetImportData.h"
@@ -71,6 +72,11 @@ UObject* UHairStrandsFactory::FactoryCreateFile(UClass* InClass, UObject* InPare
 {
 	bOutOperationCanceled = false;
 
+	const bool bIsUnattended = (IsAutomatedImport()
+		|| FApp::IsUnattended()
+		|| IsRunningCommandlet()
+		|| GIsRunningUnattendedScript);
+
 	// Translate the hair data from the file
 	TSharedPtr<IGroomTranslator> SelectedTranslator = GetTranslator(Filename);
 	if (!SelectedTranslator.IsValid())
@@ -78,6 +84,24 @@ UObject* UHairStrandsFactory::FactoryCreateFile(UClass* InClass, UObject* InPare
 		return nullptr;
 	}
 
+	// Use the settings from the script if provided
+	if (AssetImportTask)
+	{
+		if (UGroomImportOptions* GroomOptions = Cast<UGroomImportOptions>(AssetImportTask->Options))
+		{
+			ImportOptions = GroomOptions;
+		}
+
+		if (UGroomCacheImportOptions* CacheOptions = Cast<UGroomCacheImportOptions>(AssetImportTask->Options))
+		{
+			GroomCacheImportOptions = CacheOptions;
+		}
+		else
+		{
+			GroomCacheImportOptions->ImportSettings.bImportGroomCache = false;
+		}
+	}
+
 	FGroomAnimationInfo AnimInfo;
 	{
 		// Load the alembic file upfront to preview & report any potential issue
@@ -135,7 +159,8 @@ UObject* UHairStrandsFactory::FactoryCreateFile(UClass* InClass, UObject* InPare
 
 		FGroomCacheImporter::SetupImportSettings(GroomCacheImportOptions->ImportSettings, AnimInfo);
 
-		if (!GIsRunningUnattendedScript && !IsAutomatedImport())
+		// Don't bother saving the options coming from script
+		if (!bIsUnattended)
 		{
 			// Display import options and handle user cancellation
 			TSharedPtr<SGroomImportOptionsWindow> GroomOptionWindow = SGroomImportOptionsWindow::DisplayImportOptions(ImportOptions, GroomCacheImportOptions, GroupsPreview, Filename);
@@ -144,17 +169,17 @@ UObject* UHairStrandsFactory::FactoryCreateFile(UClass* InClass, UObject* InPare
 				bOutOperationCanceled = true;
 				return nullptr;
 			}
-		}
 
-		// Save the options as the new default
-		for (const FGroomHairGroupPreview& GroupPreview : GroupsPreview->Groups)
-		{
-			if (GroupPreview.GroupID < OutDescription.HairGroups.Num())
+			// Save the options as the new default
+			for (const FGroomHairGroupPreview& GroupPreview : GroupsPreview->Groups)
 			{
-				ImportOptions->InterpolationSettings[GroupPreview.GroupID] = GroupPreview.InterpolationSettings;
+				if (GroupPreview.GroupID < OutDescription.HairGroups.Num())
+				{
+					ImportOptions->InterpolationSettings[GroupPreview.GroupID] = GroupPreview.InterpolationSettings;
+				}
 			}
+			ImportOptions->SaveConfig();
 		}
-		ImportOptions->SaveConfig();
 	}
 
 	FGroomCacheImporter::ApplyImportSettings(GroomCacheImportOptions->ImportSettings, AnimInfo);
@@ -168,19 +193,19 @@ UObject* UHairStrandsFactory::FactoryCreateFile(UClass* InClass, UObject* InPare
 		return nullptr;
 	}
 
-	// Might try to import the same file in the same folder, so if an asset already exists there, reuse and update it
-	// Since we are importing (not reimporting) we reset the object completely. All previous settings will be lost.
-	UGroomAsset* ExistingAsset = FindObject<UGroomAsset>(InParent, *InName.ToString());
-	if (ExistingAsset)
-	{
-		ExistingAsset->SetNumGroup(0);
-	}
-
 	UObject* CurrentAsset = nullptr;
 	UGroomAsset* GroomAssetForCache = nullptr;
 	FHairImportContext HairImportContext(ImportOptions, InParent, InClass, InName, Flags);
 	if (GroomCacheImportOptions->ImportSettings.bImportGroomAsset)
 	{
+		// Might try to import the same file in the same folder, so if an asset already exists there, reuse and update it
+		// Since we are importing (not reimporting) we reset the object completely. All previous settings will be lost.
+		UGroomAsset* ExistingAsset = FindObject<UGroomAsset>(InParent, *InName.ToString());
+		if (ExistingAsset)
+		{
+			ExistingAsset->SetNumGroup(0);
+		}
+
 		UGroomAsset* ImportedAsset = FHairStrandsImporter::ImportHair(HairImportContext, HairDescription, ExistingAsset);
 		if (ImportedAsset)
 		{
