Deformed Skel Mesh to drive groom deformation

diff --git a/Engine/Plugins/Runtime/HairStrands/Source/HairStrandsCore/Private/GroomAsset.cpp b/Engine/Plugins/Runtime/HairStrands/Source/HairStrandsCore/Private/GroomAsset.cpp
--- a/Engine/Plugins/Runtime/HairStrands/Source/HairStrandsCore/Private/GroomAsset.cpp
+++ b/Engine/Plugins/Runtime/HairStrands/Source/HairStrandsCore/Private/GroomAsset.cpp
@@ -30,6 +30,8 @@
 #include "HairStrandsCore.h"
 #include "UObject/ObjectSaveContext.h"
 #include "GroomBindingAsset.h"
+#include "GroomDeformerBuilder.h"
+
 
 #if WITH_EDITOR
 #include "Interfaces/ITargetPlatform.h"
@@ -1305,6 +1307,10 @@ static bool IsStrandsInterpolationAttributes(const FName PropertyName)
 		   PropertyName == GET_MEMBER_NAME_CHECKED(FHairDecimationSettings, CurveDecimation)
 		|| PropertyName == GET_MEMBER_NAME_CHECKED(FHairDecimationSettings, VertexDecimation)
 
+		|| PropertyName == GET_MEMBER_NAME_CHECKED(FHairDeformationSettings, NumCurves)
+		|| PropertyName == GET_MEMBER_NAME_CHECKED(FHairDeformationSettings, NumPoints)
+		|| PropertyName == GET_MEMBER_NAME_CHECKED(FHairDeformationSettings, bEnableDeformation)
+
 		// Add dependency on simulation and per LOD-simulation/global-interoplation to strip-out interoplation data if there are not needed
 		|| PropertyName == GET_MEMBER_NAME_CHECKED(UGroomAsset, EnableGlobalInterpolation)
 		|| PropertyName == GET_MEMBER_NAME_CHECKED(FHairSolverSettings, EnableSimulation)
@@ -1353,8 +1359,16 @@ void UGroomAsset::PostEditChangeProperty(FPropertyChangedEvent& PropertyChangedE
 	if (bStrandsInterpolationChanged)
 	{
 		CacheDerivedDatas();
-	}
 
+		if (PropertyName == GET_MEMBER_NAME_CHECKED(FHairInterpolationSettings, bOverrideGuides)
+		 || PropertyName == GET_MEMBER_NAME_CHECKED(FHairDeformationSettings, NumCurves)
+		 || PropertyName == GET_MEMBER_NAME_CHECKED(FHairDeformationSettings, NumPoints)
+		 || PropertyName == GET_MEMBER_NAME_CHECKED(FHairDeformationSettings, bEnableDeformation))
+		{
+			DeformedSkeletalMesh = FGroomDeformerBuilder::CreateSkeletalMesh(this);  
+		}
+	}
+	
 	const bool bCardsArrayChanged = PropertyName == GET_MEMBER_NAME_CHECKED(UGroomAsset, HairGroupsCards);
 	if (bCardsArrayChanged && PropertyChangedEvent.ChangeType == EPropertyChangeType::ArrayAdd)
 	{
@@ -3206,6 +3220,11 @@ bool UGroomAsset::IsVisible(int32 GroupIndex, int32 LODIndex) const
 	return HairGroupsLOD[GroupIndex].LODs[LODIndex].bVisible;
 }
 
+bool UGroomAsset::IsDeformationEnable(int32 GroupIndex) const
+{
+	return HairGroupsInterpolation.IsValidIndex(GroupIndex) && HairGroupsInterpolation[GroupIndex].DeformationSettings.bEnableDeformation && HairGroupsInterpolation[GroupIndex].InterpolationSettings.bOverrideGuides;
+}
+
 bool UGroomAsset::IsSimulationEnable(int32 GroupIndex, int32 LODIndex) const 
 {
 	if (GroupIndex < 0 || GroupIndex >= HairGroupsLOD.Num() || !IsHairStrandsSimulationEnable())
@@ -3285,7 +3304,7 @@ bool UGroomAsset::NeedsInterpolationData(int32 GroupIndex) const
 		}
 	}
 
-	return false;
+	return IsDeformationEnable(GroupIndex);
 }
 
 bool UGroomAsset::NeedsInterpolationData() const
diff --git a/Engine/Plugins/Runtime/HairStrands/Source/HairStrandsCore/Private/GroomAssetInterpolation.cpp b/Engine/Plugins/Runtime/HairStrands/Source/HairStrandsCore/Private/GroomAssetInterpolation.cpp
--- a/Engine/Plugins/Runtime/HairStrands/Source/HairStrandsCore/Private/GroomAssetInterpolation.cpp
+++ b/Engine/Plugins/Runtime/HairStrands/Source/HairStrandsCore/Private/GroomAssetInterpolation.cpp
@@ -18,6 +18,13 @@ FHairInterpolationSettings::FHairInterpolationSettings()
 	bUseUniqueGuide = false;
 }
 
+FHairDeformationSettings::FHairDeformationSettings()
+{
+	bEnableDeformation = false;
+	NumCurves = 10;
+	NumPoints = 4;
+}
+
 FHairGroupsInterpolation::FHairGroupsInterpolation()
 {
 	DecimationSettings = FHairDecimationSettings();
@@ -87,6 +94,14 @@ bool FHairGroupsInterpolation::operator==(const FHairGroupsInterpolation& A) con
 		InterpolationSettings == A.InterpolationSettings;
 }
 
+bool FHairDeformationSettings::operator==(const FHairDeformationSettings& A) const
+{
+	return
+		bEnableDeformation == A.bEnableDeformation &&
+		NumCurves == A.NumCurves &&
+		NumPoints == A.NumPoints;
+}
+
 void FHairGroupsInterpolation::BuildDDCKey(FArchive& Ar)
 {
 	Ar << DecimationSettings.CurveDecimation;
@@ -97,6 +112,9 @@ void FHairGroupsInterpolation::BuildDDCKey(FArchive& Ar)
 	Ar << InterpolationSettings.InterpolationDistance;
 	Ar << InterpolationSettings.bRandomizeGuide;
 	Ar << InterpolationSettings.bUseUniqueGuide;
+	Ar << DeformationSettings.NumCurves;
+	Ar << DeformationSettings.NumPoints;
+	Ar << DeformationSettings.bEnableDeformation;
 }
 
 void FHairGroupsLOD::BuildDDCKey(FArchive& Ar)
diff --git a/Engine/Plugins/Runtime/HairStrands/Source/HairStrandsCore/Private/GroomBindingBuilder.cpp b/Engine/Plugins/Runtime/HairStrands/Source/HairStrandsCore/Private/GroomBindingBuilder.cpp
--- a/Engine/Plugins/Runtime/HairStrands/Source/HairStrandsCore/Private/GroomBindingBuilder.cpp
+++ b/Engine/Plugins/Runtime/HairStrands/Source/HairStrandsCore/Private/GroomBindingBuilder.cpp
@@ -456,29 +456,8 @@ private:
 ///////////////////////////////////////////////////////////////////////////////////////////////////
 // RBF weighting
 
-#if WITH_EDITORONLY_DATA
 namespace GroomBinding_RBFWeighting
 {
-	struct FPointsSampler
-	{
-		FPointsSampler(TArray<bool>& ValidPoints, const FVector3f* PointPositions, const int32 NumSamples);
-
-		/** Build the sample position from the sample indices */
-		void BuildPositions(const FVector3f* PointPositions);
-
-		/** Compute the furthest point */
-		void FurthestPoint(const int32 NumPoints, const FVector3f* PointPositions, const uint32 SampleIndex, TArray<bool>& ValidPoints, TArray<float>& PointsDistance);
-
-		/** Compute the starting point */
-		int32 StartingPoint(const TArray<bool>& ValidPoints, int32& NumPoints) const;
-
-		/** List of sampled points */
-		TArray<uint32> SampleIndices;
-
-		/** List of sampled positions */
-		TArray<FVector3f> SamplePositions;
-	};
-
 	int32 FPointsSampler::StartingPoint(const TArray<bool>& ValidPoints, int32& NumPoints) const
 	{
 		int32 StartIndex = -1;
@@ -548,7 +527,10 @@ namespace GroomBinding_RBFWeighting
 			BuildPositions(PointPositions);
 		}
 	}
-
+}
+#if WITH_EDITORONLY_DATA
+namespace GroomBinding_RBFWeighting
+{
 	struct FWeightsBuilder
 	{
 		FWeightsBuilder(const uint32 NumRows, const uint32 NumColumns,
@@ -2127,7 +2109,7 @@ static bool InternalBuildBinding_CPU(UGroomBindingAsset* BindingAsset, bool bIni
 				SlowTask.EnterProgressFrame();
 			}
 		}
-
+		
 		GroomBinding_RBFWeighting::ComputeInterpolationWeights(OutHairGroupDatas, BindingAsset->NumInterpolationPoints, BindingAsset->MatchingSection, TargetMeshData.Get(), TransferredPositions);
 		SlowTask.EnterProgressFrame();
 	}
diff --git a/Engine/Plugins/Runtime/HairStrands/Source/HairStrandsCore/Private/GroomBuilder.cpp b/Engine/Plugins/Runtime/HairStrands/Source/HairStrandsCore/Private/GroomBuilder.cpp
--- a/Engine/Plugins/Runtime/HairStrands/Source/HairStrandsCore/Private/GroomBuilder.cpp
+++ b/Engine/Plugins/Runtime/HairStrands/Source/HairStrandsCore/Private/GroomBuilder.cpp
@@ -6,6 +6,7 @@
 #include "GroomSettings.h"
 #include "HairDescription.h"
 #include "GroomSettings.h"
+#include "GroomBindingBuilder.h"
 
 #include "Async/ParallelFor.h"
 #include "HAL/IConsoleManager.h"
@@ -49,6 +50,12 @@ namespace FHairStrandsDecimation
 		float VertexDecimationPercentage,
 		bool bContinuousDecimationReordering,
 		FHairStrandsDatas& OutData);
+
+	void Decimate(
+		const FHairStrandsDatas& InData, 
+		const uint32 NumCurves,
+		const int32 NumVertices,
+		FHairStrandsDatas& OutData);
 }
 
 namespace HairStrandsBuilder
@@ -1027,6 +1034,45 @@ namespace HairInterpolationBuilder
 		return Desc;
 	}
 
+	// Find the vertex along a sim curve 'SimCurveIndex', which has the smallest euclidean distance than the render vertex
+	static FVertexInterpolationDesc FindMatchingVertex(const FVector3f& RenPointPosition, const FHairStrandsDatas& SimStrandsData, const uint32 SimCurveIndex)
+	{
+		const uint32 SimOffset = SimStrandsData.StrandsCurves.CurvesOffset[SimCurveIndex];
+		float MinPointDistance = FLT_MAX;
+		int32 MinPointIndex = 0;
+
+		// Find with with vertex the vertex should be paired
+		const uint32 SimPointCount = SimStrandsData.StrandsCurves.CurvesCount[SimCurveIndex];
+		for (uint32 SimPointIndex = 0; SimPointIndex < SimPointCount-1; ++SimPointIndex)
+		{
+			
+			const float SimRenDistance = (SimStrandsData.StrandsPoints.PointsPosition[SimPointIndex + SimOffset] - RenPointPosition).Size();
+			if (SimRenDistance <= MinPointDistance)
+			{
+				MinPointDistance = SimRenDistance;
+				MinPointIndex = SimPointIndex;
+			}
+		}
+		FVertexInterpolationDesc Desc;
+		const FVector3f EdgeDirection = (SimStrandsData.StrandsPoints.PointsPosition[MinPointIndex +1+ SimOffset] - SimStrandsData.StrandsPoints.PointsPosition[MinPointIndex + SimOffset]).GetSafeNormal();
+		if(((RenPointPosition- SimStrandsData.StrandsPoints.PointsPosition[MinPointIndex + SimOffset]).Dot(EdgeDirection) > 0.0) || (MinPointIndex == 0))
+		{
+			Desc.Index0 = MinPointIndex;
+			Desc.Index1 = MinPointIndex + 1;
+		}
+		else
+		{
+			Desc.Index0 = MinPointIndex-1;
+			Desc.Index1 = MinPointIndex;
+		}
+		
+		const float SimRenDistance0 = (SimStrandsData.StrandsPoints.PointsPosition[Desc.Index0 + SimOffset] - RenPointPosition).Size();
+		const float SimRenDistance1 = (SimStrandsData.StrandsPoints.PointsPosition[Desc.Index1 + SimOffset] - RenPointPosition).Size();
+		
+		Desc.T = SimRenDistance0 / (SimRenDistance0+SimRenDistance1);
+		return Desc;
+	}
+
 	static void BuildInterpolationData(
 		FHairStrandsInterpolationDatas& InterpolationData,
 		const FHairStrandsDatas& SimStrandsData,
@@ -1134,11 +1180,14 @@ namespace HairInterpolationBuilder
 				for (uint32 KIndex = 0; KIndex < FClosestGuides::Count; ++KIndex)
 				{
 					// Find the closest vertex on the guide which matches the strand vertex distance along its curve
-					if (Settings.InterpolationDistance == EHairInterpolationWeight::Parametric)
+					if (Settings.InterpolationDistance == EHairInterpolationWeight::Parametric || Settings.InterpolationDistance == EHairInterpolationWeight::Distance)
 					{
 						const uint32 SimCurveIndex = ClosestGuides.Indices[KIndex];
 						const uint32 SimOffset = SimStrandsData.StrandsCurves.CurvesOffset[SimCurveIndex];
-						const FVertexInterpolationDesc Desc = FindMatchingVertex(RenPointDistance, SimStrandsData, SimCurveIndex);
+						const FVertexInterpolationDesc Desc =  (Settings.InterpolationDistance == EHairInterpolationWeight::Parametric) ?
+							FindMatchingVertex(RenPointDistance, SimStrandsData, SimCurveIndex) :
+							FindMatchingVertex(RenStrandsData.StrandsPoints.PointsPosition[PointGlobalIndex], SimStrandsData, SimCurveIndex);
+						
 						const FVector& SimPointPosition0 = (FVector)SimStrandsData.StrandsPoints.PointsPosition[Desc.Index0 + SimOffset];
 						const FVector& SimPointPosition1 = (FVector)SimStrandsData.StrandsPoints.PointsPosition[Desc.Index1 + SimOffset];
 						const float Weight = 1.0f / FMath::Max(MinWeightDistance, FVector::Distance(RenPointPosition, FMath::Lerp(SimPointPosition0, SimPointPosition1, Desc.T)));
@@ -1794,7 +1843,14 @@ void FGroomBuilder::BuildData(
 			else
 			{
 				OutSim.Reset();
-				FHairStrandsDecimation::Decimate(OutRen, HairToGuideDensity, 1, false, OutSim);
+				if(InSettings.DeformationSettings.bEnableDeformation)
+				{
+					FHairStrandsDecimation::Decimate(OutRen, InSettings.DeformationSettings.NumCurves, InSettings.DeformationSettings.NumPoints, OutSim);
+				}
+				else
+				{
+					FHairStrandsDecimation::Decimate(OutRen, HairToGuideDensity, 1, false, OutSim);
+				}
 			}
 
 			OutGroupInfo.NumGuides			= OutSim.GetNumCurves();
@@ -1924,7 +1980,86 @@ static void DecimateCurve(
 		OutIndices.RemoveAt(ElementToRemove);
 	}
 }
+	
+void Decimate(
+		const FHairStrandsDatas& InData, 
+		const uint32 NumCurves,
+		const int32 NumVertices,
+		FHairStrandsDatas& OutData)
+{
+	const uint32 InCurveCount = InData.StrandsCurves.Num();
+	const uint32 OutCurveCount = FMath::Clamp(NumCurves, 1u, InCurveCount);
+	
+	// Fill root positions and curves offsets for sampling
+	TArray<FVector3f> RootPositions;
+	RootPositions.Init(FVector3f::ZeroVector, InData.GetNumCurves());
+	for(uint32 StrandIndex = 0; StrandIndex < InData.GetNumCurves(); ++StrandIndex)
+	{
+		RootPositions[StrandIndex] =
+			InData.StrandsPoints.PointsPosition[InData.StrandsCurves.CurvesOffset[StrandIndex]];
+	}
+				
+	TArray<bool> ValidPoints;
+	ValidPoints.Init(true, InData.GetNumCurves());
+
+	// Pick NumCurves strands from the guides
+	GroomBinding_RBFWeighting::FPointsSampler PointsSampler(ValidPoints, RootPositions.GetData(), OutCurveCount);
+
+	const uint32 CurveCount = PointsSampler.SampleIndices.Num();
+	uint32 PointCount = CurveCount * NumVertices;
+
+	OutData.StrandsCurves.SetNum(CurveCount);
+	OutData.StrandsPoints.SetNum(PointCount);
+	OutData.HairDensity = InData.HairDensity;
+
+	const bool bHasPrecomputedWeights = InData.StrandsCurves.CurvesClosestGuideIDs.Num() > 0 && InData.StrandsCurves.CurvesClosestGuideWeights.Num() > 0;
+	
+	PointCount = 0;
+	for(uint32 CurveIndex = 0; CurveIndex < CurveCount; ++CurveIndex)
+	{
+		const int32 GuideIndex = PointsSampler.SampleIndices[CurveIndex];
+		const int32 GuideOffset = InData.StrandsCurves.CurvesOffset[GuideIndex];
+		const int32 GuideCount = InData.StrandsCurves.CurvesCount[GuideIndex];
+		
+		const int32 MaxPoints = FMath::Min(GuideCount, NumVertices);
+		const float BoneSpace = FMath::Max(1.0,float(GuideCount-1) / (MaxPoints-1));
 
+		OutData.StrandsCurves.CurvesCount[CurveIndex]  = MaxPoints;
+		OutData.StrandsCurves.CurvesRootUV[CurveIndex] = InData.StrandsCurves.CurvesRootUV[GuideIndex];
+		OutData.StrandsCurves.CurvesOffset[CurveIndex] = PointCount;
+		OutData.StrandsCurves.CurvesLength[CurveIndex] = InData.StrandsCurves.CurvesLength[GuideIndex];
+		if (bHasPrecomputedWeights)
+		{
+			OutData.StrandsCurves.CurvesClosestGuideIDs[CurveIndex] = InData.StrandsCurves.CurvesClosestGuideIDs[GuideIndex];
+			OutData.StrandsCurves.CurvesClosestGuideWeights[CurveIndex] = InData.StrandsCurves.CurvesClosestGuideWeights[GuideIndex];
+		}
+		OutData.StrandsCurves.MaxLength = InData.StrandsCurves.MaxLength;
+		OutData.StrandsCurves.MaxRadius = InData.StrandsCurves.MaxRadius;
+	
+		for(int32 PointIndex = 0; PointIndex < MaxPoints; ++PointIndex, ++PointCount)
+		{
+			const float PointSpace = GuideOffset + PointIndex * BoneSpace;
+	
+			const int32 PointBegin = FMath::Min(FMath::FloorToInt32(PointSpace), GuideOffset + GuideCount-2);
+			const int32 PointEnd = PointBegin+1;
+	
+			const float PointAlpha = PointSpace - PointBegin;
+			
+			OutData.StrandsPoints.PointsPosition[PointCount] = InData.StrandsPoints.PointsPosition[PointBegin] * (1.0-PointAlpha) +
+				InData.StrandsPoints.PointsPosition[PointEnd] * PointAlpha;
+			OutData.StrandsPoints.PointsCoordU[PointCount] = InData.StrandsPoints.PointsCoordU[PointBegin] * (1.0-PointAlpha) +
+				InData.StrandsPoints.PointsCoordU[PointEnd] * PointAlpha;
+			OutData.StrandsPoints.PointsRadius[PointCount] = InData.StrandsPoints.PointsRadius[PointBegin] * (1.0-PointAlpha) +
+				InData.StrandsPoints.PointsRadius[PointEnd] * PointAlpha;
+			OutData.StrandsPoints.PointsBaseColor[PointCount] = InData.StrandsPoints.PointsBaseColor[PointBegin] * (1.0-PointAlpha) +
+				InData.StrandsPoints.PointsBaseColor[PointEnd] * PointAlpha;
+			OutData.StrandsPoints.PointsRoughness[PointCount] = InData.StrandsPoints.PointsRoughness[PointBegin] * (1.0-PointAlpha) +
+				InData.StrandsPoints.PointsRoughness[PointEnd] * PointAlpha;
+		}
+	}
+	HairStrandsBuilder::BuildInternalData(OutData, false);
+}
+	
 void Decimate(
 	const FHairStrandsDatas& InData, 
 	float CurveDecimationPercentage, 
diff --git a/Engine/Plugins/Runtime/HairStrands/Source/HairStrandsCore/Private/GroomComponent.cpp b/Engine/Plugins/Runtime/HairStrands/Source/HairStrandsCore/Private/GroomComponent.cpp
--- a/Engine/Plugins/Runtime/HairStrands/Source/HairStrandsCore/Private/GroomComponent.cpp
+++ b/Engine/Plugins/Runtime/HairStrands/Source/HairStrandsCore/Private/GroomComponent.cpp
@@ -370,7 +370,7 @@ public:
 
 			{
 				// If one of the group has simulation enable, then we enable velocity rendering for meshes/cards
-				if (IsHairStrandsSimulationEnable() && HairInstance->Guides.IsValid() && HairInstance->Guides.bIsSimulationEnable)
+				if (IsHairStrandsSimulationEnable() && HairInstance->Guides.IsValid() && (HairInstance->Guides.bIsSimulationEnable || HairInstance->Guides.bIsDeformationEnable))
 				{
 					bAlwaysHasVelocity = true;
 				}
@@ -1565,6 +1565,7 @@ void UGroomComponent::SetScatterSceneLighting(bool Enable)
 void UGroomComponent::SetUseCards(bool InbUseCards)
 {
 	bUseCards = InbUseCards;
+	ResetSimulation();
 	UpdateHairGroupsDescAndInvalidateRenderState();
 }
 
@@ -2166,6 +2167,7 @@ void UGroomComponent::UpdateSimulatedGroups()
 					check(Instance->HairGroupPublicData);
 					Instance->Guides.bIsSimulationEnable	 = Instance->HairGroupPublicData->IsSimulationEnable(LODIndex);
 					Instance->Guides.bHasGlobalInterpolation = Instance->HairGroupPublicData->IsGlobalInterpolationEnable(LODIndex);
+					Instance->Guides.bIsDeformationEnable = Instance->HairGroupPublicData->bIsDeformationEnable;
 				}
 				++GroupIt;
 			}
@@ -2356,10 +2358,13 @@ void UGroomComponent::InitResources(bool bIsBindingReloading)
 	bool bHasNeedSkeletalMesh = false;
 	TArray<bool> bHasNeedGlobalDeformation;
 	TArray<bool> bHasNeedSimulation;
+	TArray<bool> bHasNeedDeformation;
 	bHasNeedGlobalDeformation.Init(false, GroomAsset->HairGroupsData.Num());
 	bHasNeedSimulation.Init(false, GroomAsset->HairGroupsData.Num());
-
+	bHasNeedDeformation.Init(false, GroomAsset->HairGroupsData.Num());
+	
 	bool bHasAnyNeedSimulation = false;
+	bool bHasAnyNeedDeformation = false;
 	bool bHasAnyNeedGlobalDeformation = false;
 	for (int32 GroupIt = 0, GroupCount = GroomAsset->HairGroupsData.Num(); GroupIt < GroupCount; ++GroupIt)
 	{
@@ -2379,6 +2384,8 @@ void UGroomComponent::InitResources(bool bIsBindingReloading)
 			bHasAnyNeedSimulation				= bHasAnyNeedSimulation			|| bHasNeedSimulation[GroupIt];
 			bHasAnyNeedGlobalDeformation		= bHasAnyNeedGlobalDeformation	|| bHasNeedGlobalDeformation[GroupIt];
 		}
+		bHasNeedDeformation[GroupIt] = GroomAsset->IsDeformationEnable(GroupIt);
+		bHasAnyNeedDeformation = bHasAnyNeedDeformation || bHasNeedDeformation[GroupIt];
 	}
 	const bool bHasNeedBindingData = bHasNeedSkinningBinding || bHasAnyNeedGlobalDeformation;
 
@@ -2392,7 +2399,7 @@ void UGroomComponent::InitResources(bool bIsBindingReloading)
 	}
 	else
 	{
-		LODSelectionType = bHasAnyNeedSimulation ? EHairLODSelectionType::Predicted : EHairLODSelectionType::Immediate;
+		LODSelectionType = (bHasAnyNeedSimulation || bHasAnyNeedDeformation) ? EHairLODSelectionType::Predicted : EHairLODSelectionType::Immediate;
 	}
 
 	// 2. Insure that the binding asset is compatible, otherwise no binding
@@ -2419,6 +2426,28 @@ void UGroomComponent::InitResources(bool bIsBindingReloading)
 		}
 	}
 
+	// Grab deformed mesh if it exists 
+	if(GroomAsset && GroomAsset->DeformedSkeletalMesh)
+	{
+		const USkeleton* TargetSkeleton = GroomAsset->DeformedSkeletalMesh->GetSkeleton();
+		
+		DeformedMeshComponent = nullptr;
+		// Try to find the first component by walking the attachment hierarchy
+		for (USceneComponent* Curr = this; Curr; Curr = Curr->GetAttachParent())
+		{
+			USkeletalMeshComponent* SkelMeshComp = Cast<USkeletalMeshComponent>(Curr);
+			if (SkelMeshComp && SkelMeshComp->SkeletalMesh && SkelMeshComp->SkeletalMesh->GetSkeleton() == TargetSkeleton)
+			{
+				DeformedMeshComponent = SkelMeshComp;
+				break;
+			}
+		}
+		if (DeformedMeshComponent)
+		{
+			AddTickPrerequisiteComponent(DeformedMeshComponent);
+		}
+	}
+
 	// Validate that if we are bound to a skel. mesh, and we have binding, and the binding type is set to skinning, that the skin cache is enabled for these
 	if (bHasNeedSkinningBinding && ParentMeshComponent)
 	{
@@ -2504,6 +2533,8 @@ void UGroomComponent::InitResources(bool bIsBindingReloading)
 		HairGroupInstance->Debug.LODForcedIndex = LODForcedIndex;
 		HairGroupInstance->Debug.LODPredictedIndex = LODPredictedIndex;
 		HairGroupInstance->Debug.LODSelectionTypeForDebug = LODSelectionType;
+		HairGroupInstance->DeformedComponent = DeformedMeshComponent;
+		HairGroupInstance->DeformedSection = GroomAsset->DeformedGroupSections.IsValidIndex(GroupIt) ? GroomAsset->DeformedGroupSections[GroupIt] : INDEX_NONE; 
 
 		if (RegisteredMeshComponent)
 		{
@@ -2561,6 +2592,7 @@ void UGroomComponent::InitResources(bool bIsBindingReloading)
 				HairGroupInstance->HairGroupPublicData->LODSimulations.Add(LODSimulation);
 				HairGroupInstance->HairGroupPublicData->LODGlobalInterpolations.Add(LODGlobalInterpolation);
 			}
+			HairGroupInstance->HairGroupPublicData->bIsDeformationEnable = GroomAsset->IsDeformationEnable(GroupIt);
 			HairGroupInstance->HairGroupPublicData->SetLODScreenSizes(CPULODScreenSize);
 			HairGroupInstance->HairGroupPublicData->SetLODVisibilities(LODVisibility);
 			HairGroupInstance->HairGroupPublicData->SetLODGeometryTypes(LODGeometryTypes);
@@ -2576,7 +2608,7 @@ void UGroomComponent::InitResources(bool bIsBindingReloading)
 		// * Physics simulation
 		// * RBF deformation.
 		// Therefore, even if simulation is disabled, we need to run partially the update if the binding system is enabled (skin deformation + RBF correction)
-		const bool bNeedGuides = (GroupData.Guides.HasValidData() && (bHasNeedSimulation[GroupIt] || bHasNeedGlobalDeformation[GroupIt])) || (HairGroupInstance->Debug.GroomCacheType == EGroomCacheType::Guides);
+		const bool bNeedGuides = (GroupData.Guides.HasValidData() && (bHasNeedSimulation[GroupIt] || bHasNeedGlobalDeformation[GroupIt] || bHasNeedDeformation[GroupIt])) || (HairGroupInstance->Debug.GroomCacheType == EGroomCacheType::Guides);
 		if (bNeedGuides)
 		{
 			HairGroupInstance->Guides.Data = &GroupData.Guides.BulkData;
@@ -2606,6 +2638,7 @@ void UGroomComponent::InitResources(bool bIsBindingReloading)
 			const int32 LODIndex = -1;
 			HairGroupInstance->Guides.bIsSimulationEnable = IsSimulationEnable(GroupIt,LODIndex);
 			HairGroupInstance->Guides.bHasGlobalInterpolation = LocalBindingAsset && GroomAsset->IsGlobalInterpolationEnable(GroupIt,LODIndex);
+			HairGroupInstance->Guides.bIsDeformationEnable = GroomAsset->IsDeformationEnable(GroupIt);
 		}
 
 		// LODBias is in the Modifier which is needed for LOD selection regardless if the strands are there or not
@@ -2649,6 +2682,7 @@ void UGroomComponent::InitResources(bool bIsBindingReloading)
 						{
 							bNeedDynamicResources = true;
 						}
+						
 					}
 				}
 
@@ -2664,6 +2698,11 @@ void UGroomComponent::InitResources(bool bIsBindingReloading)
 				{
 					bNeedDynamicResources = true;
 				}
+				// the bones deformation needs to setup the deformed resources
+				if(HairGroupInstance->HairGroupPublicData->bIsDeformationEnable)
+				{
+					bNeedDynamicResources = true;
+				}
 			}
 
 			#if RHI_RAYTRACING
@@ -2746,8 +2785,9 @@ void UGroomComponent::InitResources(bool bIsBindingReloading)
 
 				const EHairBindingType BindingType	= HairGroupInstance->HairGroupPublicData->GetBindingType(CardsLODIndex);
 				const bool bHasSimulation			= HairGroupInstance->HairGroupPublicData->IsSimulationEnable(CardsLODIndex);
+				const bool bHasDeformation			= HairGroupInstance->HairGroupPublicData->bIsDeformationEnable;
 				const bool bHasGlobalDeformation	= HairGroupInstance->HairGroupPublicData->IsGlobalInterpolationEnable(CardsLODIndex);
-				const bool bNeedDeformedPositions	= bHasSimulation || bHasGlobalDeformation || BindingType == EHairBindingType::Skinning;
+				const bool bNeedDeformedPositions	= bHasSimulation || bHasDeformation || bHasGlobalDeformation || BindingType == EHairBindingType::Skinning;
 				const bool bNeedRootData			= bHasGlobalDeformation || BindingType == EHairBindingType::Skinning;
 
 				// Sanity check
@@ -3365,6 +3405,7 @@ void UGroomComponent::TickComponent(float DeltaTime, enum ELevelTick TickType, F
 	{
 		SkelLocalToTransform = RegisteredMeshComponent->GetComponentTransform();
 	}
+	
 	TArray<FHairGroupInstance*> LocalHairGroupInstances = HairGroupInstances;
 	const EHairLODSelectionType LocalLODSelectionType = LODSelectionType;
 	ENQUEUE_RENDER_COMMAND(FHairStrandsTick_TransformUpdate)(
diff --git a/Engine/Plugins/Runtime/HairStrands/Source/HairStrandsCore/Private/GroomDeformerBuilder.cpp b/Engine/Plugins/Runtime/HairStrands/Source/HairStrandsCore/Private/GroomDeformerBuilder.cpp
new file mode 100644
--- /dev/null
+++ b/Engine/Plugins/Runtime/HairStrands/Source/HairStrandsCore/Private/GroomDeformerBuilder.cpp
@@ -0,0 +1,421 @@
+﻿// Copyright Epic Games, Inc. All Rights Reserved.
+
+#include "GroomDeformerBuilder.h"
+#include "GroomBuilder.h"
+#include "HairStrandsCore.h"
+#include "Rendering/SkeletalMeshLODModel.h"
+#include "Rendering/SkeletalMeshModel.h"
+#include "PhysicsEngine/PhysicsConstraintTemplate.h"
+#include "GroomAsset.h"
+
+#if WITH_EDITOR
+
+static void AddGuideBones(UGroomAsset* GroomAsset, USkeletalMesh* SkeletalMesh, const FString& NameSkeleton, FSkeletalMeshLODModel* LodRenderData,
+	const FVector& RootPosition, TArray<FVector>& BonesPositions, TArray<int32>& BonesOffsets, TArray<int32>& GroupsOffsets, const TArray<FHairStrandsDatas>& GuidesDatas)
+{
+	BonesOffsets.Reset();
+	GroupsOffsets.Reset();
+	
+	int32 BoneIndex = 1;
+	for(int32 GroupIndex = 0; GroupIndex < GroomAsset->GetNumHairGroups(); ++GroupIndex)
+	{
+		GroupsOffsets.Add(BonesOffsets.Num());
+		if(GroomAsset->IsDeformationEnable(GroupIndex))
+		{
+			const FHairStrandsDatas& GuidesData = GuidesDatas[GroupIndex];
+			
+			for(uint32 CurveIndex = 0; CurveIndex < GuidesData.StrandsCurves.Num(); ++CurveIndex)
+			{
+				// Compute the node positions and rotations
+				FVector TangentPrev = FVector(0,0,1);
+				FQuat PointRotation = FQuat::Identity;
+				FTransform PointTransform = FTransform::Identity;
+				FTransform ParentTransform = FTransform::Identity;
+				ParentTransform.SetTranslation(RootPosition);
+
+				BonesOffsets.Add(BoneIndex);
+				const uint32 NumPoints = GuidesData.StrandsCurves.CurvesCount[CurveIndex];
+				for(uint32 PointIndex = 0; PointIndex < GuidesData.StrandsCurves.CurvesCount[CurveIndex]; ++PointIndex, ++BoneIndex)
+				{
+					const int32 PositionIndex = PointIndex+GuidesData.StrandsCurves.CurvesOffset[CurveIndex];
+					const FVector PointPosition(GuidesData.StrandsPoints.PointsPosition[PositionIndex]);
+
+					FVector TangentNext = (PointIndex < (NumPoints-1)) ? (FVector(GuidesData.StrandsPoints.PointsPosition[PositionIndex+1]) - PointPosition) :
+																		 (PointPosition - FVector(GuidesData.StrandsPoints.PointsPosition[PositionIndex-1]));
+					TangentNext.Normalize();
+						
+					PointRotation = ( FQuat::FindBetweenVectors(TangentPrev,TangentNext) * PointRotation ).GetNormalized();
+					TangentPrev = TangentNext;
+					
+					FString BoneString = NameSkeleton + TEXT("_Bone_") + FString::FromInt(BoneIndex);
+
+					PointTransform.SetTranslation(PointPosition);
+					PointTransform.SetRotation(PointRotation);
+
+					BonesPositions.Add(PointPosition);
+
+					const int32 ParentIndex = (PointIndex == 0) ? 0 : BoneIndex -1;
+					FReferenceSkeletonModifier Modifier = FReferenceSkeletonModifier(SkeletalMesh->GetRefSkeleton(), nullptr);
+					Modifier.Add(FMeshBoneInfo(FName(*BoneString), BoneString, ParentIndex), PointTransform * ParentTransform.Inverse());
+					ParentTransform = PointTransform;
+
+					LodRenderData->ActiveBoneIndices.Add(BoneIndex);
+					LodRenderData->RequiredBones.Add(BoneIndex);
+				}
+			}
+		}
+	}
+	GroupsOffsets.Add(BonesOffsets.Num());
+	BonesOffsets.Add(BoneIndex);
+}
+
+static FVector AddRootBone(UGroomAsset* GroomAsset, USkeletalMesh* SkeletalMesh, const FString& NameSkeleton, FSkeletalMeshLODModel* LodRenderData, TArray<FVector>& BonesPositions, TArray<FHairStrandsDatas>& GuidesDatas)
+{
+	GuidesDatas.SetNum(GroomAsset->GetNumHairGroups());
+	
+	// Get the number of bones
+	int32 NumBones = 1;
+	int32 NumRoots = 0;
+	
+	FVector CenterBone(0,0,0);
+	int32 GroupIndex = 0;
+	for(auto& HairGroup : GroomAsset->GetHairDescriptionGroups().HairGroups)
+	{
+		FHairStrandsDatas StrandsData;
+		FHairStrandsDatas& GuidesData = GuidesDatas[GroupIndex];
+		FGroomBuilder::BuildData(HairGroup, GroomAsset->HairGroupsInterpolation[GroupIndex], GroomAsset->HairGroupsInfo[GroupIndex], StrandsData, GuidesData);
+		
+		if(GroomAsset->IsDeformationEnable(GroupIndex))
+		{
+			for(uint32 CurveIndex = 0; CurveIndex < GuidesData.StrandsCurves.Num(); ++CurveIndex)
+			{
+				CenterBone += FVector(GuidesData.StrandsPoints.PointsPosition[GuidesData.StrandsCurves.CurvesOffset[CurveIndex]]);
+			}
+			NumRoots += GuidesData.StrandsCurves.Num();
+			NumBones += GuidesData.StrandsPoints.Num();
+		}
+		++GroupIndex;
+	}
+	if(NumRoots != 0)
+	{
+		CenterBone /= float(NumRoots);
+	}
+
+	// Add all the bones from the guides 
+	SkeletalMesh->GetRefSkeleton().Empty(NumBones);
+	FTransform Transform = FTransform::Identity;
+
+	// Add root bone
+	{
+		FString BoneString = NameSkeleton + TEXT("_Bone_") + FString::FromInt(0);
+		FReferenceSkeletonModifier Modifier = FReferenceSkeletonModifier(SkeletalMesh->GetRefSkeleton(), nullptr);
+		Transform.SetTranslation(CenterBone);
+		BonesPositions.Add(CenterBone);
+			
+		Modifier.Add(FMeshBoneInfo(FName(*BoneString), BoneString, INDEX_NONE), Transform);
+
+		LodRenderData->ActiveBoneIndices.Add(0);
+		LodRenderData->RequiredBones.Add(0); 
+	}
+	return CenterBone;
+}
+
+static void BuildSkeletonBones(UGroomAsset* GroomAsset, USkeletalMesh* SkeletalMesh, FSkeletalMeshLODModel* LodRenderData, TArray<FVector>& BonesPositions, TArray<int32>& BonesOffsets, TArray<int32>& GroupsOffsets)
+{
+	FString PackageNameSkeleton, NameSkeleton;
+	FHairStrandsCore::AssetHelper().CreateFilename(GroomAsset->GetOutermost()->GetName(), TEXT("_Skeleton"), PackageNameSkeleton, NameSkeleton);
+	
+	// Create packages for skeleton and skelmesh
+	UPackage* PackageSkeleton = CreatePackage(*PackageNameSkeleton);
+	USkeleton* Skeleton = SkeletalMesh->GetSkeleton() ? SkeletalMesh->GetSkeleton() :
+    			NewObject<USkeleton>( PackageSkeleton, *NameSkeleton, RF_Public | RF_Standalone | RF_Transactional );
+
+	TArray<FName> SkeletonBones;
+	const FReferenceSkeleton& RefSkeleton = Skeleton->GetReferenceSkeleton();
+
+	for (int32 BoneIndex = 0; BoneIndex < RefSkeleton.GetRawBoneNum(); ++BoneIndex)
+	{
+		SkeletonBones.Add(RefSkeleton.GetBoneName(BoneIndex));
+	}
+	Skeleton->RemoveBonesFromSkeleton(SkeletonBones, true);
+	
+	// Add the bones
+	TArray<FHairStrandsDatas> GuidesDatas;
+	const FVector RootPosition = AddRootBone(GroomAsset, SkeletalMesh, NameSkeleton, LodRenderData, BonesPositions, GuidesDatas);
+	AddGuideBones(GroomAsset, SkeletalMesh, NameSkeleton, LodRenderData, RootPosition, BonesPositions, BonesOffsets, GroupsOffsets, GuidesDatas);
+	
+	SkeletalMesh->SetSkeleton(Skeleton);
+	Skeleton->SetPreviewMesh(SkeletalMesh);
+	Skeleton->RecreateBoneTree(SkeletalMesh);
+	
+	SkeletalMesh->GetRefBasesInvMatrix().Reset();
+	SkeletalMesh->CalculateInvRefMatrices();
+}
+
+static void BuildMeshSection(UGroomAsset* GroomAsset, USkeletalMesh* SkeletalMesh, FSkeletalMeshLODModel* LodRenderData, const TArray<FVector>& BonesPositions,
+	const TArray<int32>& BonesOffsets, const TArray<int32>& GroupsOffsets, TArray<int32>& GroupSections)
+{
+	// Initialize mesh section
+	LodRenderData->Sections.Reset();
+	LodRenderData->IndexBuffer.Reset();
+	LodRenderData->NumTexCoords = 1;
+	LodRenderData->NumVertices = 0;
+	
+	const int32 NumGroups = GroupsOffsets.Num()-1;
+	GroupSections.Init(INDEX_NONE, NumGroups);
+	for(int32 GroupIndex = 0; GroupIndex < NumGroups; ++GroupIndex)
+	{
+		const int32 CurveStart = GroupsOffsets[GroupIndex], CurveStop = GroupsOffsets[GroupIndex+1];
+		const int32 PointStart = BonesOffsets[CurveStart], PointStop = BonesOffsets[CurveStop];
+		
+		const int32 NumCurves = CurveStop - CurveStart;
+		const int32 NumPoints = PointStop - PointStart;
+		GroupSections[GroupIndex] = INDEX_NONE;
+		if(NumCurves > 0 && NumPoints > 0)
+		{
+			const int32 SectionIndex = LodRenderData->Sections.Add(FSkelMeshSection());
+			auto& MeshSection = LodRenderData->Sections[SectionIndex];
+			GroupSections[GroupIndex] = SectionIndex;
+		
+			MeshSection.MaterialIndex = 0;
+			MeshSection.BaseIndex = LodRenderData->IndexBuffer.Num();
+	
+			MeshSection.BaseVertexIndex = LodRenderData->NumVertices;
+			MeshSection.MaxBoneInfluences = 1;
+
+			MeshSection.SoftVertices.Reset();
+			MeshSection.NumVertices = 1 + NumPoints;
+			
+			for (int32_t VertexIndex = 0; VertexIndex < MeshSection.NumVertices; VertexIndex++)
+			{
+				const int32 BoneIndex = (VertexIndex == 0) ? 0 : PointStart + VertexIndex - 1;
+				
+				FSoftSkinVertex SoftVertex;
+				SoftVertex.Color = FColor::White;
+				SoftVertex.Position = FVector3f(BonesPositions[BoneIndex]);
+			
+				FMemory::Memzero(SoftVertex.InfluenceBones);
+				FMemory::Memzero(SoftVertex.InfluenceWeights);
+		
+				SoftVertex.InfluenceWeights[0] = 1.0f;
+				SoftVertex.InfluenceBones[0] = MeshSection.BoneMap.Num();
+		
+				MeshSection.SoftVertices.Add(SoftVertex);
+				MeshSection.BoneMap.Add(BoneIndex);
+			}
+			LodRenderData->NumVertices += MeshSection.NumVertices;
+			
+			int32 NumTriangles = 0;
+			for(int32 CurveIndex = CurveStart; CurveIndex < CurveStop; ++CurveIndex)
+			{
+				for(int32 BoneIndex = BonesOffsets[CurveIndex]; BoneIndex < BonesOffsets[CurveIndex+1]-1; ++BoneIndex, ++NumTriangles)
+				{
+					LodRenderData->IndexBuffer.Add(MeshSection.BaseVertexIndex);
+					LodRenderData->IndexBuffer.Add(MeshSection.BaseVertexIndex+1+BoneIndex-PointStart);
+					LodRenderData->IndexBuffer.Add(MeshSection.BaseVertexIndex+1+BoneIndex+1-PointStart);
+				}
+			}
+			MeshSection.NumTriangles = NumTriangles;
+		}
+	}
+}
+
+static FSkeletalMeshLODModel* BuildRenderData(UGroomAsset* GroomAsset, USkeletalMesh* SkeletalMesh)
+{
+	// Create Skeletal Mesh LOD Render Data and LOD Info
+	FSkeletalMeshLODModel* LodRenderData = new FSkeletalMeshLODModel();
+	SkeletalMesh->GetImportedModel()->LODModels.Add(LodRenderData);
+
+	// Set default LOD Info
+	FSkeletalMeshLODInfo& LodInfo = SkeletalMesh->AddLODInfo();
+	LodInfo.ScreenSize = 0.3f;
+	LodInfo.LODHysteresis = 0.2f;
+	LodInfo.BuildSettings.bUseFullPrecisionUVs = true;
+	LodInfo.ReductionSettings.NumOfTrianglesPercentage = 1.0;
+	LodInfo.ReductionSettings.NumOfVertPercentage = 1.0;
+	
+	//Add default material as backup
+	LodInfo.LODMaterialMap.Add(0);
+	UMaterialInterface* DefaultMaterial = UMaterial::GetDefaultMaterial(MD_Surface);
+	SkeletalMesh->GetMaterials().Add(DefaultMaterial);
+	SkeletalMesh->GetMaterials()[0].UVChannelData.bInitialized = true;
+
+	return LodRenderData;
+}
+
+static void BuildBodySetup(UPhysicsAsset* PhysicsAsset, USkeletalMesh* SkeletalMesh, const int32 BoneIndex, const int32 ParentIndex, const FName& BoneName)
+{
+	FString BodyString = BoneName.ToString() + TEXT("_Body_") + FString::FromInt(BoneIndex);
+			
+	int32 BodyIndex = PhysicsAsset->FindBodyIndex(FName(BodyString));
+	if(BodyIndex == INDEX_NONE)
+	{
+		USkeletalBodySetup* BodySetup = NewObject<USkeletalBodySetup>(PhysicsAsset, NAME_None, RF_Transactional);
+		// make default to be use complex as simple 
+		BodySetup->CollisionTraceFlag = CTF_UseSimpleAsComplex;
+		BodySetup->PhysicsType = (ParentIndex == INDEX_NONE) ? EPhysicsType::PhysType_Kinematic : EPhysicsType::PhysType_Default;;
+			
+		BodyIndex = PhysicsAsset->SkeletalBodySetups.Add(BodySetup );
+		BodySetup->BoneName = BoneName;
+
+		PhysicsAsset->UpdateBodySetupIndexMap();
+		PhysicsAsset->UpdateBoundsBodiesArray();
+	}
+
+	UBodySetup* BodySetup = PhysicsAsset->SkeletalBodySetups[BodyIndex];
+	
+	// Empty any existing collision.
+	BodySetup->RemoveSimpleCollision();
+	
+	FKSphereElem SphereElem;
+	SphereElem.Center = FVector::ZeroVector;
+	SphereElem.Radius = 1.0f;
+
+	BodySetup->AggGeom.SphereElems.Add(SphereElem);
+}
+
+static void BuildConstraintSetup(UPhysicsAsset* PhysicsAsset, const FTransform& RelativeTransform, const int32 BoneIndex, const FName& BoneName, const FName& ParentName)
+{
+	FString ConstraintString = BoneName.ToString() + TEXT("_Constraint_") + FString::FromInt(BoneIndex);
+
+	// constraintClass must be a subclass of UPhysicsConstraintTemplate
+	int32 ConstraintIndex = PhysicsAsset->FindConstraintIndex(FName(ConstraintString));
+	if(ConstraintIndex == INDEX_NONE)
+	{
+		UPhysicsConstraintTemplate* ConstraintSetup = NewObject<UPhysicsConstraintTemplate>(PhysicsAsset, NAME_None, RF_Transactional);
+                
+		ConstraintIndex = PhysicsAsset->ConstraintSetup.Add( ConstraintSetup );
+		ConstraintSetup->DefaultInstance.JointName = BoneName;
+	}
+
+	UPhysicsConstraintTemplate* ConstraintSetup = PhysicsAsset->ConstraintSetup[ConstraintIndex];
+					
+	// set angular constraint mode
+	ConstraintSetup->DefaultInstance.SetAngularSwing1Motion(EAngularConstraintMotion::ACM_Free);
+	ConstraintSetup->DefaultInstance.SetAngularSwing2Motion(EAngularConstraintMotion::ACM_Free);
+	ConstraintSetup->DefaultInstance.SetAngularTwistMotion(EAngularConstraintMotion::ACM_Free);
+
+	// Place joint at origin of child
+	ConstraintSetup->DefaultInstance.ConstraintBone1 = BoneName;
+	ConstraintSetup->DefaultInstance.Pos1 = FVector::ZeroVector;
+	ConstraintSetup->DefaultInstance.PriAxis1 = FVector(1, 0, 0);
+	ConstraintSetup->DefaultInstance.SecAxis1 = FVector(0, 1, 0);
+
+	ConstraintSetup->DefaultInstance.ConstraintBone2 = ParentName;
+	ConstraintSetup->DefaultInstance.Pos2 = RelativeTransform.GetLocation();
+	ConstraintSetup->DefaultInstance.PriAxis2 = RelativeTransform.GetUnitAxis(EAxis::X);
+	ConstraintSetup->DefaultInstance.SecAxis2 = RelativeTransform.GetUnitAxis(EAxis::Y);
+
+	ConstraintSetup->SetDefaultProfile(ConstraintSetup->DefaultInstance);
+}
+
+static void BuildPhysicsAsset(UGroomAsset* GroomAsset, USkeletalMesh* SkeletalMesh, FSkeletalMeshLODModel* LodRenderData)
+{
+	// Fill the package/asset names for the skelmesh and the skeleton
+	FString PackageNamePhysicsAsset, NamePhysicsAsset;
+	FHairStrandsCore::AssetHelper().CreateFilename(GroomAsset->GetOutermost()->GetName(), TEXT("_PhysicsAsset"), PackageNamePhysicsAsset, NamePhysicsAsset);
+	
+	// Create packages for skeleton and skelmesh
+	UPackage* PackagePhysicsAsset = CreatePackage(*PackageNamePhysicsAsset);
+	UPhysicsAsset* PhysicsAsset = SkeletalMesh->GetPhysicsAsset() ? SkeletalMesh->GetPhysicsAsset() :
+			NewObject<UPhysicsAsset>( PackagePhysicsAsset, *NamePhysicsAsset, RF_Public | RF_Standalone | RF_Transactional );
+
+	PhysicsAsset->SkeletalBodySetups.Empty();
+	PhysicsAsset->ConstraintSetup.Empty();
+	
+	PhysicsAsset->UpdateBodySetupIndexMap();
+	PhysicsAsset->UpdateBoundsBodiesArray();
+
+	PhysicsAsset->CollisionDisableTable.Empty();
+	for (int32 BoneIndex = 0; BoneIndex < SkeletalMesh->GetRefSkeleton().GetNum(); BoneIndex++)
+	{
+		const int32 ParentIndex = SkeletalMesh->GetRefSkeleton().GetParentIndex(BoneIndex);
+
+		// Go ahead and make this bone physical.
+		FName BoneName = SkeletalMesh->GetRefSkeleton().GetBoneName(BoneIndex);
+		FTransform BoneTransform(SkeletalMesh->GetComposedRefPoseMatrix(BoneIndex));
+		
+		BuildBodySetup(PhysicsAsset, SkeletalMesh, BoneIndex, ParentIndex, BoneName);
+		if(ParentIndex != INDEX_NONE)
+		{
+			FName ParentName = SkeletalMesh->GetRefSkeleton().GetBoneName(ParentIndex);
+			FTransform ParentTransform(FTransform(SkeletalMesh->GetComposedRefPoseMatrix(ParentIndex)));
+			BuildConstraintSetup(PhysicsAsset, BoneTransform * ParentTransform.Inverse(), BoneIndex, BoneName, ParentName);
+		}
+	}
+	SkeletalMesh->SetPhysicsAsset(PhysicsAsset);
+	PhysicsAsset->SetPreviewMesh(SkeletalMesh);
+}
+
+static USkeletalMesh* BuildSkeletalMesh(UGroomAsset* GroomAsset)
+{
+	// Get package name
+	FString PackageNameSkeletalMesh, NameSkeletalMesh;
+	FHairStrandsCore::AssetHelper().CreateFilename(GroomAsset->GetOutermost()->GetName(), TEXT("_SkeletalMesh"), PackageNameSkeletalMesh, NameSkeletalMesh);
+	
+	// Create package for the skeletal mesh
+	UPackage* PackageSkeletalMesh = CreatePackage(*PackageNameSkeletalMesh);
+	USkeletalMesh* SkeletalMesh = GroomAsset->DeformedSkeletalMesh ? GroomAsset->DeformedSkeletalMesh.Get() :
+		NewObject<USkeletalMesh>(PackageSkeletalMesh, *NameSkeletalMesh, RF_Public | RF_Standalone | RF_Transactional);
+
+	// Reset the skeletal mesh
+	SkeletalMesh->GetImportedModel()->LODModels.Reset();
+	SkeletalMesh->ResetLODInfo();
+	SkeletalMesh->ReleaseResources();
+
+	// Build the LOD render data
+	FSkeletalMeshLODModel* LodRenderData = BuildRenderData(GroomAsset, SkeletalMesh);
+	TArray<FVector> BonesPositions;
+	TArray<int32> BonesOffsets;
+	TArray<int32> GroupsOffsets;
+
+	// Build the skeleton bones
+	BuildSkeletonBones(GroomAsset, SkeletalMesh, LodRenderData, BonesPositions, BonesOffsets, GroupsOffsets);
+
+	// Build the mesh section
+	BuildMeshSection(GroomAsset, SkeletalMesh, LodRenderData, BonesPositions, BonesOffsets, GroupsOffsets, GroomAsset->DeformedGroupSections);
+
+	// Build the physics asset
+	BuildPhysicsAsset(GroomAsset, SkeletalMesh, LodRenderData);
+	
+	SkeletalMesh->InvalidateDeriveDataCacheGUID();
+	SkeletalMesh->PostEditChange();
+	SkeletalMesh->InitResources();
+	if(SkeletalMesh->MarkPackageDirty())
+	{
+		FHairStrandsCore::AssetHelper().RegisterAsset(SkeletalMesh);
+		FHairStrandsCore::SaveAsset(SkeletalMesh);
+	}
+	
+	SkeletalMesh->GetSkeleton()->PostEditChange();
+	if(SkeletalMesh->GetSkeleton()->MarkPackageDirty())
+	{
+		FHairStrandsCore::AssetHelper().RegisterAsset(SkeletalMesh->GetSkeleton());
+		FHairStrandsCore::SaveAsset(SkeletalMesh->GetSkeleton());
+	}
+
+	SkeletalMesh->GetPhysicsAsset()->PostEditChange();
+	if(SkeletalMesh->GetPhysicsAsset()->MarkPackageDirty())
+	{
+		FHairStrandsCore::AssetHelper().RegisterAsset(SkeletalMesh->GetPhysicsAsset());
+		FHairStrandsCore::SaveAsset(SkeletalMesh->GetPhysicsAsset());
+	}
+	
+	return SkeletalMesh;
+}
+
+#endif
+
+USkeletalMesh* FGroomDeformerBuilder::CreateSkeletalMesh(UGroomAsset* GroomAsset)
+{
+#if WITH_EDITOR
+	
+	if (!GroomAsset)
+	{
+		return nullptr;
+	}
+	return BuildSkeletalMesh(GroomAsset);
+	
+#endif
+	return nullptr;
+}
\ No newline at end of file
diff --git a/Engine/Plugins/Runtime/HairStrands/Source/HairStrandsCore/Private/GroomManager.cpp b/Engine/Plugins/Runtime/HairStrands/Source/HairStrandsCore/Private/GroomManager.cpp
--- a/Engine/Plugins/Runtime/HairStrands/Source/HairStrandsCore/Private/GroomManager.cpp
+++ b/Engine/Plugins/Runtime/HairStrands/Source/HairStrandsCore/Private/GroomManager.cpp
@@ -174,6 +174,7 @@ static void RunInternalHairStrandsInterpolation(
 			const uint32 HairLODIndex = Instance->HairGroupPublicData->LODIndex;
 			const EHairBindingType BindingType = Instance->HairGroupPublicData->GetBindingType(HairLODIndex);
 			const bool bSimulationEnable = Instance->HairGroupPublicData->IsSimulationEnable(HairLODIndex);
+			const bool bDeformationEnable = Instance->HairGroupPublicData->bIsDeformationEnable;
 			const bool bGlobalDeformationEnable = Instance->HairGroupPublicData->IsGlobalInterpolationEnable(HairLODIndex);
 			check(InstanceBindingType == BindingType);
 
@@ -202,7 +203,7 @@ static void RunInternalHairStrandsInterpolation(
 							Instance->Strands.DeformedRootResource,
 							Instance->Strands.DeformedResource);
 					}
-					else if (bSimulationEnable)
+					else if (bSimulationEnable || bDeformationEnable)
 					{
 						check(Instance->Strands.HasValidData());
 
@@ -241,7 +242,7 @@ static void RunInternalHairStrandsInterpolation(
 								CardsInstance.Guides.DeformedRootResource,
 								CardsInstance.Guides.DeformedResource);
 						}
-						else if (bSimulationEnable)
+						else if (bSimulationEnable || bDeformationEnable)
 						{
 							check(CardsInstance.Guides.IsValid());
 							AddHairStrandUpdatePositionOffsetPass(
@@ -261,7 +262,7 @@ static void RunInternalHairStrandsInterpolation(
 			else if (EHairStrandsInterpolationType::SimulationStrands == Type)
 			{
 				// Guide update need to run only if simulation is enabled, or if RBF is enabled (since RFB are transfer through guides)
-				if (bGlobalDeformationEnable || bSimulationEnable)
+				if (bGlobalDeformationEnable || bSimulationEnable || bDeformationEnable)
 				{
 					check(Instance->Guides.IsValid());
 
@@ -322,7 +323,7 @@ static void RunInternalHairStrandsInterpolation(
 						}
 						GraphBuilder.UseExternalAccessMode(Register(GraphBuilder, Instance->Guides.DeformedResource->GetPositionOffsetBuffer(FHairStrandsDeformedResource::EFrameType::Current), ERDGImportedBufferFlags::CreateSRV).Buffer, ERHIAccess::SRVMask);
 					}
-					else if (bSimulationEnable)
+					else if (bSimulationEnable || bDeformationEnable)
 					{
 						check(Instance->Guides.IsValid());
 
@@ -624,8 +625,9 @@ void AddHairStreamingRequest(FHairGroupInstance* Instance, int32 InLODIndex)
 		}
 
 		const bool bSimulationEnable			= Instance->HairGroupPublicData->IsSimulationEnable(LODIndex);
+		const bool bDeformationEnable			= Instance->HairGroupPublicData->bIsDeformationEnable;
 		const bool bGlobalInterpolationEnable	= Instance->HairGroupPublicData->IsGlobalInterpolationEnable(LODIndex);
-		const bool bLODNeedsGuides				= bSimulationEnable || bGlobalInterpolationEnable;
+		const bool bLODNeedsGuides				= bSimulationEnable || bDeformationEnable || bGlobalInterpolationEnable;
 
 		const EHairResourceLoadingType LoadingType = GetHairResourceLoadingType(GeometryType, int32(LODIndex));
 		if (LoadingType != EHairResourceLoadingType::Async || GeometryType == EHairGeometryType::NoneGeometry)
@@ -795,8 +797,9 @@ static void RunHairLODSelection(
 			}
 
 			const bool bSimulationEnable = Instance->HairGroupPublicData->IsSimulationEnable(IntLODIndex);
+			const bool bDeformationEnable			= Instance->HairGroupPublicData->bIsDeformationEnable;
 			const bool bGlobalInterpolationEnable =  Instance->HairGroupPublicData->IsGlobalInterpolationEnable(IntLODIndex);
-			const bool bLODNeedsGuides = bSimulationEnable || bGlobalInterpolationEnable;
+			const bool bLODNeedsGuides = bSimulationEnable || bDeformationEnable || bGlobalInterpolationEnable;
 
 			const EHairResourceLoadingType LoadingType = GetHairResourceLoadingType(GeometryType, IntLODIndex);
 			EHairResourceStatus ResourceStatus = EHairResourceStatus::None;
@@ -889,6 +892,7 @@ static void RunHairLODSelection(
 				Instance->BindingType = BindingType;
 				Instance->Guides.bIsSimulationEnable = Instance->HairGroupPublicData->IsSimulationEnable(IntLODIndex);
 				Instance->Guides.bHasGlobalInterpolation = Instance->HairGroupPublicData->IsGlobalInterpolationEnable(IntLODIndex);
+				Instance->Guides.bIsDeformationEnable = Instance->HairGroupPublicData->bIsDeformationEnable;
 				Instance->Strands.bIsCullingEnabled = bCullingEnable;
 			}
 
diff --git a/Engine/Plugins/Runtime/HairStrands/Source/HairStrandsCore/Private/HairStrandsCore.cpp b/Engine/Plugins/Runtime/HairStrands/Source/HairStrandsCore/Private/HairStrandsCore.cpp
--- a/Engine/Plugins/Runtime/HairStrands/Source/HairStrandsCore/Private/HairStrandsCore.cpp
+++ b/Engine/Plugins/Runtime/HairStrands/Source/HairStrandsCore/Private/HairStrandsCore.cpp
@@ -17,8 +17,6 @@ void ProcessHairStrandsBookmark(
 	EHairStrandsBookmark Bookmark,
 	FHairStrandsBookmarkParameters& Parameters);
 
-FHairAssetHelper HairStrandsCore_AssetHelper;
-
 void FHairStrandsCore::StartupModule()
 {
 	RegisterBookmarkFunction(ProcessHairStrandsBookmark);
@@ -37,10 +35,18 @@ void FHairStrandsCore::ShutdownModule()
 void FHairStrandsCore::RegisterAssetHelper(const FHairAssetHelper& Helper)
 {
 #if WITH_EDITOR
-	HairStrandsCore_AssetHelper = Helper;
+	AssetHelper() = Helper;
 #endif
 }
 
+#if WITH_EDITOR
+FHairAssetHelper& FHairStrandsCore::AssetHelper()
+{
+	static FHairAssetHelper AssetHelper;
+	return AssetHelper;
+}
+#endif
+
 // This is a workaround to be able to create & register UTexture2D from the HairStrandsCore project 
 // without requiring editor dependencies. This is used for the hair cards generation, which creates 
 // UTexture2D assets when adding a new LOD. Ideally this should be changed and move the logic to the 
@@ -75,14 +81,14 @@ void FHairStrandsCore::RegisterAssetHelper(const FHairAssetHelper& Helper)
 #if WITH_EDITOR
 UTexture2D* FHairStrandsCore::CreateTexture(const FString& InPackageName, const FIntPoint& Resolution, const FString& Suffix, TTextureAllocation TextureAllocation)
 {
-	if (HairStrandsCore_AssetHelper.CreateFilename == nullptr || HairStrandsCore_AssetHelper.RegisterAsset == nullptr)
+	if (AssetHelper().CreateFilename == nullptr || AssetHelper().RegisterAsset == nullptr)
 	{
 		return nullptr;
 	}
 
 	FString Name;
 	FString PackageName;
-	HairStrandsCore_AssetHelper.CreateFilename(InPackageName, Suffix, PackageName, Name);
+	AssetHelper().CreateFilename(InPackageName, Suffix, PackageName, Name);
 
 	UObject* InParent = nullptr;
 	UPackage* Package = Cast<UPackage>(InParent);
@@ -104,7 +110,7 @@ UTexture2D* FHairStrandsCore::CreateTexture(const FString& InPackageName, const
 		Out->MarkPackageDirty();
 
 		// Notify the asset registry
-		HairStrandsCore_AssetHelper.RegisterAsset(Out);
+		AssetHelper().RegisterAsset(Out);
 		return Out;
 	}
 
@@ -123,14 +129,14 @@ void FHairStrandsCore::ResizeTexture(UTexture2D* Out, const FIntPoint& Resolutio
 
 UStaticMesh* FHairStrandsCore::CreateStaticMesh(const FString& InPackageName, const FString& Suffix)
 {
-	if (HairStrandsCore_AssetHelper.CreateFilename == nullptr || HairStrandsCore_AssetHelper.RegisterAsset == nullptr)
+	if (AssetHelper().CreateFilename == nullptr || AssetHelper().RegisterAsset == nullptr)
 	{
 		return nullptr;
 	}
 
 	FString Name;
 	FString PackageName;
-	HairStrandsCore_AssetHelper.CreateFilename(InPackageName, Suffix, PackageName, Name);
+	AssetHelper().CreateFilename(InPackageName, Suffix, PackageName, Name);
 
 	UObject* InParent = nullptr;
 	UPackage* Package = Cast<UPackage>(InParent);
@@ -152,7 +158,7 @@ UStaticMesh* FHairStrandsCore::CreateStaticMesh(const FString& InPackageName, co
 		Out->GetSourceModel(0).BuildSettings.bRecomputeNormals = false;
 		Out->GetSourceModel(0).BuildSettings.bRecomputeTangents = true;
 		Out->MarkPackageDirty();
-		HairStrandsCore_AssetHelper.RegisterAsset(Out);
+		AssetHelper().RegisterAsset(Out);
 		return Out;
 	}
 
@@ -190,11 +196,11 @@ UGroomBindingAsset* FHairStrandsCore::CreateGroomBindingAsset(EGroomBindingMeshT
 			Suffix += TEXT("_") + Target->GetName();
 		}
 		Suffix += TEXT("_Binding");
-		HairStrandsCore_AssetHelper.CreateFilename(GroomAsset->GetOutermost()->GetName(), Suffix, PackageName, Name);		
+		AssetHelper().CreateFilename(GroomAsset->GetOutermost()->GetName(), Suffix, PackageName, Name);		
 	}
 	else
 	{
-		HairStrandsCore_AssetHelper.CreateFilename(InPackageName, TEXT(""), PackageName, Name);
+		AssetHelper().CreateFilename(InPackageName, TEXT(""), PackageName, Name);
 	}
 
 	UPackage* Package = Cast<UPackage>(InParent);
@@ -227,7 +233,7 @@ UGroomBindingAsset* FHairStrandsCore::CreateGroomBindingAsset(EGroomBindingMeshT
 		Out->NumInterpolationPoints = NumInterpolationPoints;
 		Out->MatchingSection = MatchingSection;
 		Out->MarkPackageDirty();
-		HairStrandsCore_AssetHelper.RegisterAsset(Out);
+		AssetHelper().RegisterAsset(Out);
 		return Out;
 	}
 #endif
@@ -246,11 +252,11 @@ UGroomBindingAsset* FHairStrandsCore::CreateGroomBindingAsset(UGroomAsset* Groom
 
 void FHairStrandsCore::SaveAsset(UObject* Object)
 {
-	if (HairStrandsCore_AssetHelper.SaveAsset == nullptr)
+	if (AssetHelper().SaveAsset == nullptr)
 	{
 		return;
 	}
 
-	HairStrandsCore_AssetHelper.SaveAsset(Object);
+	AssetHelper().SaveAsset(Object);
 }
 #endif
diff --git a/Engine/Plugins/Runtime/HairStrands/Source/HairStrandsCore/Private/HairStrandsRendering.cpp b/Engine/Plugins/Runtime/HairStrands/Source/HairStrandsCore/Private/HairStrandsRendering.cpp
--- a/Engine/Plugins/Runtime/HairStrands/Source/HairStrandsCore/Private/HairStrandsRendering.cpp
+++ b/Engine/Plugins/Runtime/HairStrands/Source/HairStrandsCore/Private/HairStrandsRendering.cpp
@@ -20,6 +20,8 @@
 #include "RenderTargetPool.h"
 #include "GroomTextureBuilder.h"
 #include "GroomBindingBuilder.h"
+#include "SkeletalRenderPublic.h"
+#include "SkeletalMeshDeformerHelpers.h"
 #include "GroomAsset.h" 
 #include "GroomManager.h"
 #include "GroomInstance.h"
@@ -280,7 +282,7 @@ class FDeformGuideCS : public FGlobalShader
 	SHADER_USE_PARAMETER_STRUCT(FDeformGuideCS, FGlobalShader);
 
 	class FGroupSize : SHADER_PERMUTATION_SPARSE_INT("PERMUTATION_GROUP_SIZE", 32, 64);
-	class FDeformationType : SHADER_PERMUTATION_INT("PERMUTATION_DEFORMATION", 4);
+	class FDeformationType : SHADER_PERMUTATION_INT("PERMUTATION_DEFORMATION", 5);
 	using FPermutationDomain = TShaderPermutationDomain<FGroupSize, FDeformationType>;
 
 	BEGIN_SHADER_PARAMETER_STRUCT(FParameters, )
@@ -306,6 +308,8 @@ class FDeformGuideCS : public FGlobalShader
 		SHADER_PARAMETER(uint32, SampleCount)
 		SHADER_PARAMETER_RDG_BUFFER_SRV(Buffer, RestSamplePositionsBuffer)
 		SHADER_PARAMETER_RDG_BUFFER_SRV(Buffer, MeshSampleWeightsBuffer)
+
+		SHADER_PARAMETER_SRV(Buffer<float4>, BoneDeformedPositionBuffer)
 	END_SHADER_PARAMETER_STRUCT()
 
 public:
@@ -326,7 +330,8 @@ static void AddDeformSimHairStrandsPass(
 	FRDGImportedBuffer OutSimDeformedPositionBuffer,
 	const FVector& SimRestOffset,
 	FRDGBufferSRVRef SimDeformedOffsetBuffer,
-	const bool bHasGlobalInterpolation)
+	const bool bHasGlobalInterpolation,
+	FRHIShaderResourceView* BoneBufferSRV)
 {
 	enum EInternalDeformationType
 	{
@@ -334,6 +339,7 @@ static void AddDeformSimHairStrandsPass(
 		InternalDeformationType_Offset = 1,
 		InternalDeformationType_Skinned = 2,
 		InternalDeformationType_RBF = 3,
+		InternalDeformationType_Bones = 4,
 		InternalDeformationTypeCount
 	};
 
@@ -368,41 +374,50 @@ static void AddDeformSimHairStrandsPass(
 		}
 
 		const uint32 RootCount = SimRestRootResources ? SimRestRootResources->BulkData.RootCount : 0;
-		const bool bSupportDynamicMesh = 
-			RootCount > 0 && 
-			MeshLODIndex >= 0 && 
-			MeshLODIndex < uint32(SimRestRootResources->LODs.Num()) &&
-			MeshLODIndex < uint32(SimDeformedRootResources->LODs.Num()) &&
-			SimRestRootResources->LODs[MeshLODIndex].IsValid() &&
-			SimDeformedRootResources->LODs[MeshLODIndex].IsValid() &&
-			bIsVertexToCurveBuffersValid;
-		
-		bool bSupportGlobalInterpolation = false;
-		if (bSupportDynamicMesh)
-		{
-			FHairStrandsRestRootResource::FLOD& RestLODDatas = SimRestRootResources->LODs[MeshLODIndex];
-			FHairStrandsDeformedRootResource::FLOD& DeformedLODDatas = SimDeformedRootResources->LODs[MeshLODIndex];
 
-			bSupportGlobalInterpolation = bHasGlobalInterpolation && (RestLODDatas.SampleCount > 0);
-			if (!bSupportGlobalInterpolation) 
+		if(BoneBufferSRV)
+		{
+			InternalDeformationType = InternalDeformationType_Bones;
+			Parameters->BoneDeformedPositionBuffer = BoneBufferSRV;
+		}
+		else
+		{
+			const bool bSupportDynamicMesh = 
+				RootCount > 0 && 
+				MeshLODIndex >= 0 && 
+				MeshLODIndex < uint32(SimRestRootResources->LODs.Num()) &&
+				MeshLODIndex < uint32(SimDeformedRootResources->LODs.Num()) &&
+				SimRestRootResources->LODs[MeshLODIndex].IsValid() &&
+				SimDeformedRootResources->LODs[MeshLODIndex].IsValid() &&
+				bIsVertexToCurveBuffersValid;
+			
+			bool bSupportGlobalInterpolation = false;
+			if (bSupportDynamicMesh)
 			{
-				InternalDeformationType = InternalDeformationType_Skinned;
-				Parameters->SimRestPosition0Buffer = RegisterAsSRV(GraphBuilder, RestLODDatas.RestRootTrianglePosition0Buffer);
-				Parameters->SimRestPosition1Buffer = RegisterAsSRV(GraphBuilder, RestLODDatas.RestRootTrianglePosition1Buffer);
-				Parameters->SimRestPosition2Buffer = RegisterAsSRV(GraphBuilder, RestLODDatas.RestRootTrianglePosition2Buffer);
+				FHairStrandsRestRootResource::FLOD& RestLODDatas = SimRestRootResources->LODs[MeshLODIndex];
+				FHairStrandsDeformedRootResource::FLOD& DeformedLODDatas = SimDeformedRootResources->LODs[MeshLODIndex];
 
-				Parameters->SimDeformedPosition0Buffer = RegisterAsSRV(GraphBuilder, DeformedLODDatas.DeformedRootTrianglePosition0Buffer);
-				Parameters->SimDeformedPosition1Buffer = RegisterAsSRV(GraphBuilder, DeformedLODDatas.DeformedRootTrianglePosition1Buffer);
-				Parameters->SimDeformedPosition2Buffer = RegisterAsSRV(GraphBuilder, DeformedLODDatas.DeformedRootTrianglePosition2Buffer);
+				bSupportGlobalInterpolation = bHasGlobalInterpolation && (RestLODDatas.SampleCount > 0);
+				if (!bSupportGlobalInterpolation) 
+				{
+					InternalDeformationType = InternalDeformationType_Skinned;
+					Parameters->SimRestPosition0Buffer = RegisterAsSRV(GraphBuilder, RestLODDatas.RestRootTrianglePosition0Buffer);
+					Parameters->SimRestPosition1Buffer = RegisterAsSRV(GraphBuilder, RestLODDatas.RestRootTrianglePosition1Buffer);
+					Parameters->SimRestPosition2Buffer = RegisterAsSRV(GraphBuilder, RestLODDatas.RestRootTrianglePosition2Buffer);
 
-				Parameters->SimRootBarycentricBuffer = RegisterAsSRV(GraphBuilder, RestLODDatas.RootTriangleBarycentricBuffer);
-			}
-			else
-			{
-				InternalDeformationType = InternalDeformationType_RBF;
-				Parameters->MeshSampleWeightsBuffer = RegisterAsSRV(GraphBuilder, DeformedLODDatas.MeshSampleWeightsBuffer);
-				Parameters->RestSamplePositionsBuffer = RegisterAsSRV(GraphBuilder, RestLODDatas.RestSamplePositionsBuffer);
-				Parameters->SampleCount = RestLODDatas.SampleCount;
+					Parameters->SimDeformedPosition0Buffer = RegisterAsSRV(GraphBuilder, DeformedLODDatas.DeformedRootTrianglePosition0Buffer);
+					Parameters->SimDeformedPosition1Buffer = RegisterAsSRV(GraphBuilder, DeformedLODDatas.DeformedRootTrianglePosition1Buffer);
+					Parameters->SimDeformedPosition2Buffer = RegisterAsSRV(GraphBuilder, DeformedLODDatas.DeformedRootTrianglePosition2Buffer);
+
+					Parameters->SimRootBarycentricBuffer = RegisterAsSRV(GraphBuilder, RestLODDatas.RootTriangleBarycentricBuffer);
+				}
+				else
+				{
+					InternalDeformationType = InternalDeformationType_RBF;
+					Parameters->MeshSampleWeightsBuffer = RegisterAsSRV(GraphBuilder, DeformedLODDatas.MeshSampleWeightsBuffer);
+					Parameters->RestSamplePositionsBuffer = RegisterAsSRV(GraphBuilder, RestLODDatas.RestSamplePositionsBuffer);
+					Parameters->SampleCount = RestLODDatas.SampleCount;
+				}
 			}
 		}
 	}
@@ -677,7 +692,7 @@ static void AddHairStrandsInterpolationPass(
 
 	// Guides
 	bool bSupportGlobalInterpolation = false;
-	if (bSupportDynamicMesh && (Instance->Guides.bIsSimulationEnable || Instance->Guides.bHasGlobalInterpolation))
+	if (bSupportDynamicMesh && (Instance->Guides.bIsSimulationEnable || Instance->Guides.bHasGlobalInterpolation || Instance->Guides.bIsDeformationEnable))
 	{
 		const FHairStrandsRestRootResource::FLOD& Sim_RestLODDatas = SimRestRootResources->LODs[MeshLODIndex];
 		const FHairStrandsDeformedRootResource::FLOD& Sim_DeformedLODDatas = SimDeformedRootResources->LODs[MeshLODIndex];
@@ -714,7 +729,7 @@ static void AddHairStrandsInterpolationPass(
 		}
 	}
 
-	const bool bHasLocalDeformation = Instance->Guides.bIsSimulationEnable || bSupportGlobalInterpolation;
+	const bool bHasLocalDeformation = Instance->Guides.bIsSimulationEnable || bSupportGlobalInterpolation || Instance->Guides.bIsDeformationEnable;
 	const bool bCullingEnable		= InstanceGeometryType == EHairGeometryType::Strands && CullingData.bCullingResultAvailable;
 	Parameters->HairStrandsVF_bIsCullingEnable = bCullingEnable ? 1 : 0;
 
@@ -1759,7 +1774,7 @@ void ComputeHairStrandsInterpolation(
 				Strands_TangentSRV = Strands_DeformedTangent.SRV;
 
 				// "WITH_EDITOR && bDebugModePatchedAttributeBuffer" is a special path when visualizing groom within the groom editor, so that we can diplay guides even when there is no simulation or global interpolation enabled
-				const bool bValidGuide = Instance->Guides.bIsSimulationEnable || Instance->Guides.bHasGlobalInterpolation || (WITH_EDITOR && bDebugModePatchedAttributeBuffer);
+				const bool bValidGuide = Instance->Guides.bIsSimulationEnable || Instance->Guides.bHasGlobalInterpolation || Instance->Guides.bIsDeformationEnable || (WITH_EDITOR && bDebugModePatchedAttributeBuffer);
 				const bool bHasSkinning = Instance->BindingType == EHairBindingType::Skinning;
 
 				bool bHasValidGroomCacheBuffers = Instance->Debug.GroomCacheBuffers.IsValid() && Instance->Debug.GroomCacheBuffers->GetInterpolatedFrameBuffer().GroupsData.IsValidIndex(Instance->Debug.GroupIndex);
@@ -1787,6 +1802,38 @@ void ComputeHairStrandsInterpolation(
 						RegisterAsSRV(GraphBuilder, Instance->Guides.DeformedResource->GetPositionOffsetBuffer(FHairStrandsDeformedResource::EFrameType::Current)),
 						RegisterAsUAV(GraphBuilder, Instance->Guides.DeformedResource->GetBuffer(FHairStrandsDeformedResource::Current)));
 				}
+				else if( Instance->Guides.bIsDeformationEnable && Instance->Guides.DeformedResource && Instance->DeformedComponent && (Instance->DeformedSection != INDEX_NONE))
+				{
+					if (USkeletalMeshComponent* SkeletalMeshComponent = Cast<USkeletalMeshComponent>(Instance->DeformedComponent))
+					{
+						if(FSkeletalMeshObject* SkeletalMeshObject = SkeletalMeshComponent->MeshObject)
+						{
+							const int32 LodIndex = SkeletalMeshObject->GetLOD();
+							FSkeletalMeshLODRenderData const* LodRenderData = &SkeletalMeshObject->GetSkeletalMeshRenderData().LODRenderData[LodIndex];
+				
+							if(LodRenderData->RenderSections.Num() > Instance->DeformedSection) 
+							{
+								FRHIShaderResourceView* BoneBufferSRV = FSkeletalMeshDeformerHelpers::GetBoneBufferForReading(SkeletalMeshObject, LodIndex, Instance->DeformedSection, false);
+
+								// Guides deformation based on the skeletal mesh bones
+								AddDeformSimHairStrandsPass(
+									GraphBuilder,
+									ShaderMap,
+									EDeformationType::OffsetGuide,
+									MeshLODIndex,
+									Instance->Guides.RestResource->GetVertexCount(),
+									Instance->Guides.RestRootResource,
+									Instance->Guides.DeformedRootResource,
+									RegisterAsSRV(GraphBuilder, Instance->Guides.RestResource->PositionBuffer),
+									Register(GraphBuilder, Instance->Guides.DeformedResource->GetBuffer(FHairStrandsDeformedResource::Current), ERDGImportedBufferFlags::CreateUAV),
+									Instance->Guides.RestResource->GetPositionOffset(),
+									RegisterAsSRV(GraphBuilder, Instance->Guides.DeformedResource->GetPositionOffsetBuffer(FHairStrandsDeformedResource::Current)),
+									Instance->Guides.bHasGlobalInterpolation,
+									BoneBufferSRV);
+							}
+						}
+					}
+				}
 
 				// 2.1 Compute deformation position based on simulation/skinning/RBF
 				if (Instance->Debug.GroomCacheType != EGroomCacheType::Strands)
@@ -2016,7 +2063,7 @@ void ComputeHairStrandsInterpolation(
 		const bool bIsCardsValid = Instance->Cards.IsValid(HairLODIndex);
 		if (bIsCardsValid)
 		{
-			const bool bValidGuide = Instance->Guides.bIsSimulationEnable || Instance->Guides.bHasGlobalInterpolation;
+			const bool bValidGuide = Instance->Guides.bIsSimulationEnable || Instance->Guides.bHasGlobalInterpolation || Instance->Guides.bIsDeformationEnable;
 			const bool bHasSkinning = Instance->BindingType == EHairBindingType::Skinning && MeshLODIndex >= 0;
 			const bool bNeedDeformation = bValidGuide || bHasSkinning;
 
@@ -2214,8 +2261,8 @@ void ResetHairStrandsInterpolation(
 	int32 MeshLODIndex)
 {
 	if (!Instance || 
-		(Instance && Instance->Guides.bIsSimulationEnable) ||
-		(!Instance->Guides.bHasGlobalInterpolation && !Instance->Guides.bIsSimulationEnable) ||
+		(Instance && (Instance->Guides.bIsSimulationEnable || Instance->Guides.bIsDeformationEnable)) ||
+		(Instance && !Instance->Guides.bHasGlobalInterpolation && !Instance->Guides.bIsSimulationEnable && !Instance->Guides.bIsDeformationEnable) ||
 		!IsHairStrandsBindingEnable()) return;
 
 	DECLARE_GPU_STAT(HairStrandsGuideDeform);
@@ -2237,5 +2284,6 @@ void ResetHairStrandsInterpolation(
 		DeformedPositionBuffer,
 		Instance->Guides.RestResource->GetPositionOffset(),
 		RegisterAsSRV(GraphBuilder, Instance->Guides.DeformedResource->GetPositionOffsetBuffer(FHairStrandsDeformedResource::Current)),
-		Instance->Guides.bHasGlobalInterpolation);
+		Instance->Guides.bHasGlobalInterpolation,
+		nullptr);
 }
diff --git a/Engine/Plugins/Runtime/HairStrands/Source/HairStrandsCore/Public/GroomAsset.h b/Engine/Plugins/Runtime/HairStrands/Source/HairStrandsCore/Public/GroomAsset.h
--- a/Engine/Plugins/Runtime/HairStrands/Source/HairStrandsCore/Public/GroomAsset.h
+++ b/Engine/Plugins/Runtime/HairStrands/Source/HairStrandsCore/Public/GroomAsset.h
@@ -455,6 +455,14 @@ public:
 	UPROPERTY(EditAnywhere, EditFixedSize, BlueprintReadWrite, Category = "HairInterpolation", meta = (ToolTip = "Type of interpolation used (WIP)"))
 	EGroomInterpolationType HairInterpolationType = EGroomInterpolationType::SmoothTransform;
 
+	/** Deformed skeletal mesh that will drive the groom deformation/simulation. For creating this skeletal mesh, enable EnableDeformation within the interpolation settings*/
+	UPROPERTY(VisibleAnywhere, BlueprintReadWrite, Category = "HairInterpolation")
+	TObjectPtr<USkeletalMesh> DeformedSkeletalMesh;
+
+	/** Deformed skeletal mesh mapping from groups to sections */
+	UPROPERTY()
+	TArray<int32> DeformedGroupSections;
+
 	/** Minimum LOD to cook */
 	UPROPERTY(EditAnywhere, Category = "LOD", meta = (DisplayName = "Minimum LOD"))
 	FPerPlatformInt MinLOD;
@@ -553,6 +561,7 @@ public:
 	bool IsVisible(int32 GroupIndex, int32 LODIndex) const;
 	bool IsSimulationEnable(int32 GroupIndex, int32 LODIndex) const;
 	bool IsSimulationEnable() const;
+	bool IsDeformationEnable(int32 GroupIndex) const;
 	bool IsGlobalInterpolationEnable(int32 GroupIndex, int32 LODIndex) const;
 	bool NeedsInterpolationData(int32 GroupIndex) const;
 	bool NeedsInterpolationData() const;
@@ -570,6 +579,7 @@ public:
 
 	// Save out a static mesh based on generated cards
 	void SaveProceduralCards(uint32 CardsGroupIndex);
+	
 #endif
 
 	/** Array of user data stored with the asset */
@@ -640,8 +650,9 @@ public:
 	FString GetDerivedDataKeyForStrands(uint32 GroupIndex);
 	FString GetDerivedDataKeyForMeshes(uint32 GroupIndex);
 
-private:
 	const FHairDescriptionGroups& GetHairDescriptionGroups();
+private:
+	
 	FString BuildDerivedDataKeySuffix(uint32 GroupIndex, const FHairGroupsInterpolation& InterpolationSettings, const FHairGroupsLOD& LODSettings) const;
 	bool IsFullyCached();
 	TUniquePtr<FHairDescriptionBulkData> HairDescriptionBulkData;
diff --git a/Engine/Plugins/Runtime/HairStrands/Source/HairStrandsCore/Public/GroomAssetInterpolation.h b/Engine/Plugins/Runtime/HairStrands/Source/HairStrandsCore/Public/GroomAssetInterpolation.h
--- a/Engine/Plugins/Runtime/HairStrands/Source/HairStrandsCore/Public/GroomAssetInterpolation.h
+++ b/Engine/Plugins/Runtime/HairStrands/Source/HairStrandsCore/Public/GroomAssetInterpolation.h
@@ -22,6 +22,7 @@ enum class EHairInterpolationWeight : uint8
 	Parametric	UMETA(DisplayName = "Parametric", ToolTip = "Build interpolation data based on curve parametric distance"),
 	Root		UMETA(DisplayName = "Root", ToolTip = "Build interpolation data based on distance between guide's root and strands's root"),
 	Index		UMETA(DisplayName = "Index", ToolTip = "Build interpolation data based on guide and strands vertex indices"),
+	Distance	UMETA(DisplayName = "Distance", ToolTip = "Build interpolation data based on curve euclidean distance"),
 	Unknown		UMETA(Hidden),
 };
 
@@ -146,6 +147,28 @@ struct HAIRSTRANDSCORE_API FHairInterpolationSettings
 	bool operator==(const FHairInterpolationSettings& A) const;
 };
 
+USTRUCT(BlueprintType)
+struct HAIRSTRANDSCORE_API FHairDeformationSettings
+{
+	GENERATED_USTRUCT_BODY()
+
+	FHairDeformationSettings();
+
+	/** Enable the generation of a skeletal mesh that will drive the guides deformation */
+	UPROPERTY(EditAnywhere, BlueprintReadWrite, Category = "InterpolationSettings", meta = (ToolTip = "Enable the generation of a skeletal mesh that will drive the guides deformation. The total number of bones per group is for now 256 (NumCurves * NumPoints)"))
+	bool bEnableDeformation;
+
+	/** Number of curves to generate in the skel mesh */
+	UPROPERTY(EditAnywhere, BlueprintReadWrite, Category = "InterpolationSettings", meta = (ToolTip = "Number of guides that will be generated on the groom and the skeletal mesh", EditCondition = "bEnableDeformation"))
+	int32 NumCurves;
+
+	/** Number of points per curve */
+	UPROPERTY(EditAnywhere, BlueprintReadWrite, Category = "InterpolationSettings", meta = (ToolTip = "Number of points/bones per generated guide", EditCondition = "bEnableDeformation"))
+	int32 NumPoints;
+
+	bool operator==(const FHairDeformationSettings& A) const;
+};
+
 USTRUCT(BlueprintType)
 struct HAIRSTRANDSCORE_API FHairGroupsInterpolation
 {
@@ -159,6 +182,9 @@ struct HAIRSTRANDSCORE_API FHairGroupsInterpolation
 	UPROPERTY(EditAnywhere, Category = "InterpolationSettings", meta = (ToolTip = "Interpolation settings"))
 	FHairInterpolationSettings InterpolationSettings;
 
+	UPROPERTY(EditAnywhere, Category = "DeformationSettings", meta = (ToolTip = "Deformation settings"))
+	FHairDeformationSettings DeformationSettings;
+
 	bool operator==(const FHairGroupsInterpolation& A) const;
 
 	void BuildDDCKey(FArchive& Ar);
diff --git a/Engine/Plugins/Runtime/HairStrands/Source/HairStrandsCore/Public/GroomBindingBuilder.h b/Engine/Plugins/Runtime/HairStrands/Source/HairStrandsCore/Public/GroomBindingBuilder.h
--- a/Engine/Plugins/Runtime/HairStrands/Source/HairStrandsCore/Public/GroomBindingBuilder.h
+++ b/Engine/Plugins/Runtime/HairStrands/Source/HairStrandsCore/Public/GroomBindingBuilder.h
@@ -24,4 +24,27 @@ struct HAIRSTRANDSCORE_API FGroomBindingBuilder
 
 	// Extract root data from bulk data
 	static void GetRootData(FHairStrandsRootData& Out, const FHairStrandsRootBulkData& In);
-};
\ No newline at end of file
+};
+
+namespace GroomBinding_RBFWeighting
+{
+	struct FPointsSampler
+	{
+		FPointsSampler(TArray<bool>& ValidPoints, const FVector3f* PointPositions, const int32 NumSamples);
+
+		/** Build the sample position from the sample indices */
+		void BuildPositions(const FVector3f* PointPositions);
+
+		/** Compute the furthest point */
+		void FurthestPoint(const int32 NumPoints, const FVector3f* PointPositions, const uint32 SampleIndex, TArray<bool>& ValidPoints, TArray<float>& PointsDistance);
+
+		/** Compute the starting point */
+		int32 StartingPoint(const TArray<bool>& ValidPoints, int32& NumPoints) const;
+
+		/** List of sampled points */
+		TArray<uint32> SampleIndices;
+
+		/** List of sampled positions */
+		TArray<FVector3f> SamplePositions;
+	};
+}
\ No newline at end of file
diff --git a/Engine/Plugins/Runtime/HairStrands/Source/HairStrandsCore/Public/GroomComponent.h b/Engine/Plugins/Runtime/HairStrands/Source/HairStrandsCore/Public/GroomComponent.h
--- a/Engine/Plugins/Runtime/HairStrands/Source/HairStrandsCore/Public/GroomComponent.h
+++ b/Engine/Plugins/Runtime/HairStrands/Source/HairStrandsCore/Public/GroomComponent.h
@@ -302,6 +302,7 @@ private:
 	void DeleteDeferredHairGroupInstances();
 	void* InitializedResources;
 	class UMeshComponent* RegisteredMeshComponent;
+	class UMeshComponent* DeformedMeshComponent;
 	FVector SkeletalPreviousPositionOffset;
 	bool bIsGroomAssetCallbackRegistered;
 	bool bIsGroomBindingAssetCallbackRegistered;
diff --git a/Engine/Plugins/Runtime/HairStrands/Source/HairStrandsCore/Public/GroomDeformerBuilder.h b/Engine/Plugins/Runtime/HairStrands/Source/HairStrandsCore/Public/GroomDeformerBuilder.h
new file mode 100644
--- /dev/null
+++ b/Engine/Plugins/Runtime/HairStrands/Source/HairStrandsCore/Public/GroomDeformerBuilder.h
@@ -0,0 +1,14 @@
+﻿// Copyright Epic Games, Inc. All Rights Reserved.
+
+#pragma once
+
+#include "CoreMinimal.h"
+
+class UGroomAsset;
+class USkeletalMesh;
+
+struct HAIRSTRANDSCORE_API FGroomDeformerBuilder
+{
+	// Create skeletal mesh bones+physics asset from groom asset guides
+	static USkeletalMesh* CreateSkeletalMesh(UGroomAsset* GroomAsset);
+};
diff --git a/Engine/Plugins/Runtime/HairStrands/Source/HairStrandsCore/Public/GroomInstance.h b/Engine/Plugins/Runtime/HairStrands/Source/HairStrandsCore/Public/GroomInstance.h
--- a/Engine/Plugins/Runtime/HairStrands/Source/HairStrandsCore/Public/GroomInstance.h
+++ b/Engine/Plugins/Runtime/HairStrands/Source/HairStrandsCore/Public/GroomInstance.h
@@ -132,6 +132,7 @@ struct HAIRSTRANDSCORE_API FHairGroupInstance : public FHairStrandsInstance
 	struct FGuides : FStrandsBase
 	{
 		bool bIsSimulationEnable = false;
+		bool bIsDeformationEnable = false;
 		bool bHasGlobalInterpolation = false;
 	} Guides;
 
@@ -278,7 +279,13 @@ struct HAIRSTRANDSCORE_API FHairGroupInstance : public FHairStrandsInstance
 	bool					bForceCards = false;
 	bool					bUpdatePositionOffset = false;
 	bool					bCastShadow = true;
-
+	
+	// Deformed component to extract the bone buffer 
+	UMeshComponent*	 DeformedComponent = nullptr;
+	
+	// Section of the deformed component to be used 
+	int32	 DeformedSection = INDEX_NONE;
+	
 	bool IsValid() const 
 	{
 		return Meshes.IsValid() || Cards.IsValid() || Strands.IsValid();
diff --git a/Engine/Plugins/Runtime/HairStrands/Source/HairStrandsCore/Public/HairStrandsCore.h b/Engine/Plugins/Runtime/HairStrands/Source/HairStrandsCore/Public/HairStrandsCore.h
--- a/Engine/Plugins/Runtime/HairStrands/Source/HairStrandsCore/Public/HairStrandsCore.h
+++ b/Engine/Plugins/Runtime/HairStrands/Source/HairStrandsCore/Public/HairStrandsCore.h
@@ -47,7 +47,7 @@ public:
 	}
 
 	HAIRSTRANDSCORE_API static void RegisterAssetHelper(const FHairAssetHelper& Helper);
-
+	
 #if WITH_EDITOR
 	static UTexture2D* CreateTexture(const FString& PackgeName, const FIntPoint& Resolution, const FString& Suffix, TTextureAllocation TextureAllocation);
 	static void ResizeTexture(UTexture2D* InTexture, const FIntPoint& Resolution, TTextureAllocation TextureAllocation);
@@ -64,5 +64,8 @@ public:
 	HAIRSTRANDSCORE_API static UGroomBindingAsset* CreateGroomBindingAsset(EGroomBindingMeshType BindingType, const FString& InPackageName, UObject* InParent, UGroomAsset* GroomAsset, UObject* Source, UObject* Target, const int32 NumInterpolationPoints, const int32 MatchingSection);
 
 	static void SaveAsset(UObject* Object);
+
+	static FHairAssetHelper& AssetHelper();
+	
 #endif
 };
diff --git a/Engine/Plugins/Runtime/HairStrands/Source/HairStrandsEditor/Private/GroomAssetDetails.cpp b/Engine/Plugins/Runtime/HairStrands/Source/HairStrandsEditor/Private/GroomAssetDetails.cpp
--- a/Engine/Plugins/Runtime/HairStrands/Source/HairStrandsEditor/Private/GroomAssetDetails.cpp
+++ b/Engine/Plugins/Runtime/HairStrands/Source/HairStrandsEditor/Private/GroomAssetDetails.cpp
@@ -818,6 +818,13 @@ bool FGroomRenderingDetails::CommonResetToDefault(TSharedPtr<IPropertyHandle> Ch
 			HAIR_RESET1(HairGroupsInterpolation, FHairInterpolationSettings, InterpolationSettings, bRandomizeGuide);
 			HAIR_RESET1(HairGroupsInterpolation, FHairInterpolationSettings, InterpolationSettings, bUseUniqueGuide);
 		}
+
+		{
+			FHairDeformationSettings Default;
+			HAIR_RESET1(HairGroupsInterpolation, FHairDeformationSettings, DeformationSettings, NumCurves);
+			HAIR_RESET1(HairGroupsInterpolation, FHairDeformationSettings, DeformationSettings, NumPoints);
+			HAIR_RESET1(HairGroupsInterpolation, FHairDeformationSettings, DeformationSettings, bEnableDeformation);
+		}
 	}
 
 	// LODs
diff --git a/Engine/Shaders/Private/HairStrands/HairStrandsGuideDeform.usf b/Engine/Shaders/Private/HairStrands/HairStrandsGuideDeform.usf
--- a/Engine/Shaders/Private/HairStrands/HairStrandsGuideDeform.usf
+++ b/Engine/Shaders/Private/HairStrands/HairStrandsGuideDeform.usf
@@ -25,6 +25,7 @@ RWBuffer<uint4> OutSimDeformedPositionBuffer;
 #define DEFORMATION_OFFSET_GUIDE 1
 #define DEFORMATION_OFFSET_GUIDE_DYNAMIC 2
 #define DEFORMATION_OFFSET_GUIDE_GLOBAL 3
+#define DEFORMATION_OFFSET_GUIDE_BONES 4
 
 #if PERMUTATION_DEFORMATION == DEFORMATION_OFFSET_GUIDE_DYNAMIC
 
@@ -55,6 +56,13 @@ float3 DeformGuideByTriangle(
 
 #endif
 
+#if PERMUTATION_DEFORMATION == DEFORMATION_OFFSET_GUIDE_BONES
+
+Buffer<float4>  BoneDeformedPositionBuffer;
+
+#endif
+
+
 #if PERMUTATION_DEFORMATION == DEFORMATION_OFFSET_GUIDE_GLOBAL
 uint SampleCount;
 Buffer<float4>  RestSamplePositionsBuffer;
@@ -89,12 +97,21 @@ float3 DisplacePosition(uint VertexIndex, float3 Pos)
 	const float3 ControlPoint = ApplyRBF(RestControlPoint, SampleCount, RestSamplePositionsBuffer, MeshSampleWeightsBuffer);
 	return ControlPoint - OutSimHairPositionOffset;
 #endif
+
+#if PERMUTATION_DEFORMATION == DEFORMATION_OFFSET_GUIDE_BONES
+	const float4 RestControlPoint = float4(Pos + SimRestOffset,1);
+	const int BufferOffset = 3 * (VertexIndex+1);
+
+	const float3x4 BoneMatrix = float3x4(BoneDeformedPositionBuffer[BufferOffset], BoneDeformedPositionBuffer[BufferOffset+1], BoneDeformedPositionBuffer[BufferOffset+2]);
+	return mul(BoneMatrix,RestControlPoint) - OutSimHairPositionOffset;
+#endif
 }
 
 [numthreads(GROUP_SIZE, 1, 1)]
 void MainCS(uint2 DispatchThreadId : SV_DispatchThreadID)
 {
 	const uint VertexIndex = GetHairStrandsVertexFetchIndex(DispatchThreadId, GROUP_SIZE, DispatchCountX);
+	
 	if (VertexIndex < VertexCount)
 	{
 		const uint4  Data 				= SimRestPosePositionBuffer[VertexIndex];
diff --git a/Engine/Source/Runtime/Renderer/Private/HairStrands/HairStrandsVisibility.cpp b/Engine/Source/Runtime/Renderer/Private/HairStrands/HairStrandsVisibility.cpp
--- a/Engine/Source/Runtime/Renderer/Private/HairStrands/HairStrandsVisibility.cpp
+++ b/Engine/Source/Runtime/Renderer/Private/HairStrands/HairStrandsVisibility.cpp
@@ -3627,7 +3627,7 @@ bool HasPositionsChanged(FRDGBuilder& GraphBuilder, const FViewInfo& View)
 		const int32 LODIndex = HairGroupPublicData->GetIntLODIndex();
 		const bool bHasSimulationOrSkinning = 
 			HairGroupPublicData->GetGeometryType(LODIndex) == EHairGeometryType::Strands &&
-			(HairGroupPublicData->IsSimulationEnable(LODIndex) || HairGroupPublicData->GetBindingType(LODIndex) == EHairBindingType::Skinning);
+			(HairGroupPublicData->IsSimulationEnable(LODIndex) || HairGroupPublicData->bIsDeformationEnable || HairGroupPublicData->GetBindingType(LODIndex) == EHairBindingType::Skinning);
 		if (bHasSimulationOrSkinning)
 		{
 			GroupDatas.Add(HairGroupPublicData);
diff --git a/Engine/Source/Runtime/Renderer/Public/HairStrandsInterface.h b/Engine/Source/Runtime/Renderer/Public/HairStrandsInterface.h
--- a/Engine/Source/Runtime/Renderer/Public/HairStrandsInterface.h
+++ b/Engine/Source/Runtime/Renderer/Public/HairStrandsInterface.h
@@ -331,6 +331,7 @@ public:
 	TArray<float>LODScreenSizes;
 	TArray<bool> LODSimulations;
 	TArray<bool> LODGlobalInterpolations;
+	bool bIsDeformationEnable;
 	TArray<EHairGeometryType> LODGeometryTypes;
 
 	TArray<EHairBindingType> BindingTypes;
