Debug of RBF samples

diff --git a/Engine/Plugins/Runtime/HairStrands/Source/HairStrandsCore/Private/GroomDebug.cpp b/Engine/Plugins/Runtime/HairStrands/Source/HairStrandsCore/Private/GroomDebug.cpp
--- a/Engine/Plugins/Runtime/HairStrands/Source/HairStrandsCore/Private/GroomDebug.cpp
+++ b/Engine/Plugins/Runtime/HairStrands/Source/HairStrandsCore/Private/GroomDebug.cpp
@@ -18,28 +18,38 @@
 
 static int32 GHairDebugMeshProjection_SkinCacheMesh = 0;
 static int32 GHairDebugMeshProjection_SkinCacheMeshInUVsSpace = 0;
+
 static int32 GHairDebugMeshProjection_Sim_HairRestTriangles = 0;
 static int32 GHairDebugMeshProjection_Sim_HairRestFrames = 0;
+static int32 GHairDebugMeshProjection_Sim_HairRestSamples = 0;
 static int32 GHairDebugMeshProjection_Sim_HairDeformedTriangles = 0;
 static int32 GHairDebugMeshProjection_Sim_HairDeformedFrames = 0;
+static int32 GHairDebugMeshProjection_Sim_HairDeformedSamples = 0;
 
 static int32 GHairDebugMeshProjection_Render_HairRestTriangles = 0;
 static int32 GHairDebugMeshProjection_Render_HairRestFrames = 0;
+static int32 GHairDebugMeshProjection_Render_HairRestSamples = 0;
 static int32 GHairDebugMeshProjection_Render_HairDeformedTriangles = 0;
 static int32 GHairDebugMeshProjection_Render_HairDeformedFrames = 0;
+static int32 GHairDebugMeshProjection_Render_HairDeformedSamples = 0;
 
 
 static FAutoConsoleVariableRef CVarHairDebugMeshProjection_SkinCacheMeshInUVsSpace(TEXT("r.HairStrands.MeshProjection.DebugInUVsSpace"), GHairDebugMeshProjection_SkinCacheMeshInUVsSpace, TEXT("Render debug mes projection in UVs space"));
-static FAutoConsoleVariableRef CVarHairDebugMeshProjection_SkinCacheMesh(TEXT("r.HairStrands.MeshProjection.DebugSkinCache"), GHairDebugMeshProjection_SkinCacheMesh, TEXT("Render debug mes projection"));
-static FAutoConsoleVariableRef CVarHairDebugMeshProjection_Render_HairRestTriangles(TEXT("r.HairStrands.MeshProjection.Render.Rest.Triangles"), GHairDebugMeshProjection_Render_HairRestTriangles, TEXT("Render debug mes projection"));
-static FAutoConsoleVariableRef CVarHairDebugMeshProjection_Render_HairRestFrames(TEXT("r.HairStrands.MeshProjection.Render.Rest.Frames"), GHairDebugMeshProjection_Render_HairRestFrames, TEXT("Render debug mes projection"));
-static FAutoConsoleVariableRef CVarHairDebugMeshProjection_Render_HairDeformedTriangles(TEXT("r.HairStrands.MeshProjection.Render.Deformed.Triangles"), GHairDebugMeshProjection_Render_HairDeformedTriangles, TEXT("Render debug mes projection"));
-static FAutoConsoleVariableRef CVarHairDebugMeshProjection_Render_HairDeformedFrames(TEXT("r.HairStrands.MeshProjection.Render.Deformed.Frames"), GHairDebugMeshProjection_Render_HairDeformedFrames, TEXT("Render debug mes projection"));
-
-static FAutoConsoleVariableRef CVarHairDebugMeshProjection_Sim_HairRestTriangles(TEXT("r.HairStrands.MeshProjection.Sim.Rest.Triangles"), GHairDebugMeshProjection_Sim_HairRestTriangles, TEXT("Render debug mes projection"));
-static FAutoConsoleVariableRef CVarHairDebugMeshProjection_Sim_HairRestFrames(TEXT("r.HairStrands.MeshProjection.Sim.Rest.Frames"), GHairDebugMeshProjection_Sim_HairRestFrames, TEXT("Render debug mes projection"));
-static FAutoConsoleVariableRef CVarHairDebugMeshProjection_Sim_HairDeformedTriangles(TEXT("r.HairStrands.MeshProjection.Sim.Deformed.Triangles"), GHairDebugMeshProjection_Sim_HairDeformedTriangles, TEXT("Render debug mes projection"));
-static FAutoConsoleVariableRef CVarHairDebugMeshProjection_Sim_HairDeformedFrames(TEXT("r.HairStrands.MeshProjection.Sim.Deformed.Frames"), GHairDebugMeshProjection_Sim_HairDeformedFrames, TEXT("Render debug mes projection"));
+static FAutoConsoleVariableRef CVarHairDebugMeshProjection_SkinCacheMesh(TEXT("r.HairStrands.MeshProjection.DebugSkinCache"), GHairDebugMeshProjection_SkinCacheMesh, TEXT("Render the skin cache used in projection"));
+
+static FAutoConsoleVariableRef CVarHairDebugMeshProjection_Render_HairRestTriangles(TEXT("r.HairStrands.MeshProjection.Render.Rest.Triangles"), GHairDebugMeshProjection_Render_HairRestTriangles, TEXT("Render strands rest triangles"));
+static FAutoConsoleVariableRef CVarHairDebugMeshProjection_Render_HairRestFrames(TEXT("r.HairStrands.MeshProjection.Render.Rest.Frames"), GHairDebugMeshProjection_Render_HairRestFrames, TEXT("Render strands rest frames"));
+static FAutoConsoleVariableRef CVarHairDebugMeshProjection_Render_HairRestSamples(TEXT("r.HairStrands.MeshProjection.Render.Rest.Samples"), GHairDebugMeshProjection_Render_HairRestSamples, TEXT("Render strands rest samples"));
+static FAutoConsoleVariableRef CVarHairDebugMeshProjection_Render_HairDeformedTriangles(TEXT("r.HairStrands.MeshProjection.Render.Deformed.Triangles"), GHairDebugMeshProjection_Render_HairDeformedTriangles, TEXT("Render strands deformed triangles"));
+static FAutoConsoleVariableRef CVarHairDebugMeshProjection_Render_HairDeformedFrames(TEXT("r.HairStrands.MeshProjection.Render.Deformed.Frames"), GHairDebugMeshProjection_Render_HairDeformedFrames, TEXT("Render strands deformed frames"));
+static FAutoConsoleVariableRef CVarHairDebugMeshProjection_Render_HairDeformedSamples(TEXT("r.HairStrands.MeshProjection.Render.Deformed.Samples"), GHairDebugMeshProjection_Render_HairDeformedSamples, TEXT("Render strands deformed samples"));
+
+static FAutoConsoleVariableRef CVarHairDebugMeshProjection_Sim_HairRestTriangles(TEXT("r.HairStrands.MeshProjection.Sim.Rest.Triangles"), GHairDebugMeshProjection_Sim_HairRestTriangles, TEXT("Render guides rest triangles"));
+static FAutoConsoleVariableRef CVarHairDebugMeshProjection_Sim_HairRestFrames(TEXT("r.HairStrands.MeshProjection.Sim.Rest.Frames"), GHairDebugMeshProjection_Sim_HairRestFrames, TEXT("Render guides rest frames"));
+static FAutoConsoleVariableRef CVarHairDebugMeshProjection_Sim_HairRestSamples(TEXT("r.HairStrands.MeshProjection.Sim.Rest.Samples"), GHairDebugMeshProjection_Sim_HairRestSamples, TEXT("Render guides rest samples"));
+static FAutoConsoleVariableRef CVarHairDebugMeshProjection_Sim_HairDeformedTriangles(TEXT("r.HairStrands.MeshProjection.Sim.Deformed.Triangles"), GHairDebugMeshProjection_Sim_HairDeformedTriangles, TEXT("Render guides deformed triangles"));
+static FAutoConsoleVariableRef CVarHairDebugMeshProjection_Sim_HairDeformedFrames(TEXT("r.HairStrands.MeshProjection.Sim.Deformed.Frames"), GHairDebugMeshProjection_Sim_HairDeformedFrames, TEXT("Render guides deformed frames"));
+static FAutoConsoleVariableRef CVarHairDebugMeshProjection_Sim_HairDeformedSamples(TEXT("r.HairStrands.MeshProjection.Sim.Deformed.Samples"), GHairDebugMeshProjection_Sim_HairDeformedSamples, TEXT("Render guides deformed samples"));
 
 static int32 GHairCardsAtlasDebug = 0;
 static FAutoConsoleVariableRef CVarHairCardsAtlasDebug(TEXT("r.HairStrands.Cards.DebugAtlas"), GHairCardsAtlasDebug, TEXT("Draw debug hair cards atlas."));
@@ -304,6 +314,9 @@ BEGIN_SHADER_PARAMETER_STRUCT(FHairProjectionHairDebugParameters, )
 	SHADER_PARAMETER_RDG_BUFFER_SRV(StructuredBuffer, DeformedPosition1Buffer)
 	SHADER_PARAMETER_RDG_BUFFER_SRV(StructuredBuffer, DeformedPosition2Buffer)
 
+	SHADER_PARAMETER_RDG_BUFFER_SRV(StructuredBuffer, RestSamplePositionsBuffer)
+	SHADER_PARAMETER_RDG_BUFFER_SRV(StructuredBuffer, DeformedSamplePositionsBuffer)
+
 	SHADER_PARAMETER_RDG_BUFFER_SRV(StructuredBuffer, RootBarycentricBuffer)
 
 	SHADER_PARAMETER_STRUCT_REF(FViewUniformShaderParameters, ViewUniformBuffer)
@@ -326,7 +339,7 @@ public:
 class FHairProjectionHairDebugVS : public FHairProjectionHairDebug
 {
 public:
-	class FInputType : SHADER_PERMUTATION_INT("PERMUTATION_INPUT_TYPE", 2);
+	class FInputType : SHADER_PERMUTATION_INT("PERMUTATION_INPUT_TYPE", 3);
 	using FPermutationDomain = TShaderPermutationDomain<FInputType>;
 
 	DECLARE_GLOBAL_SHADER(FHairProjectionHairDebugVS);
@@ -336,6 +349,7 @@ public:
 		SHADER_PARAMETER_STRUCT_INCLUDE(FHairProjectionHairDebugParameters, Pass)
 		END_SHADER_PARAMETER_STRUCT()
 };
+ 
 
 class FHairProjectionHairDebugPS : public FHairProjectionHairDebug
 {
@@ -354,6 +368,7 @@ enum class EDebugProjectionHairType
 {
 	HairFrame,
 	HairTriangle,
+	HairSamples,
 };
 
 static void AddDebugProjectionHairPass(
@@ -371,17 +386,24 @@ static void AddDebugProjectionHairPass(
 	FRDGTextureRef ColorTarget,
 	FRDGTextureRef DepthTexture)
 {
-	const EPrimitiveType PrimitiveType = GeometryType == EDebugProjectionHairType::HairFrame ? PT_LineList : PT_TriangleList;
-	const uint32 RootCount = RestRootResources->BulkData.RootCount;
+	if (MeshLODIndex < 0 || MeshLODIndex >= RestRootResources->LODs.Num() || MeshLODIndex >= DeformedRootResources->LODs.Num())
+		return;
+	
+	const EPrimitiveType PrimitiveType = GeometryType == EDebugProjectionHairType::HairFrame ? PT_LineList : GeometryType == EDebugProjectionHairType::HairTriangle ? PT_TriangleList : PT_LineList;
+	const uint32 RootCount = EDebugProjectionHairType::HairSamples == GeometryType ? 3 * RestRootResources->LODs[MeshLODIndex].SampleCount : RestRootResources->BulkData.RootCount;
 	const uint32 PrimitiveCount = RootCount;
 
-	if (PrimitiveCount == 0 || MeshLODIndex < 0 || MeshLODIndex >= RestRootResources->LODs.Num() || MeshLODIndex >= DeformedRootResources->LODs.Num())
+	if (PrimitiveCount == 0)
 		return;
 
 	if (EDebugProjectionHairType::HairFrame == GeometryType &&
 		!RestRootResources->LODs[MeshLODIndex].RootTriangleBarycentricBuffer.Buffer)
 		return;
 
+	if (EDebugProjectionHairType::HairSamples == GeometryType &&
+		!RestRootResources->LODs[MeshLODIndex].RestSamplePositionsBuffer.Buffer)
+			return;
+
 	const FHairStrandsRestRootResource::FLOD& RestLODDatas = RestRootResources->LODs[MeshLODIndex];
 	const FHairStrandsDeformedRootResource::FLOD& DeformedLODDatas = DeformedRootResources->LODs[MeshLODIndex];
 
@@ -414,13 +436,16 @@ static void AddDebugProjectionHairPass(
 	Parameters->DeformedPosition1Buffer = RegisterAsSRV(GraphBuilder, DeformedLODDatas.DeformedRootTrianglePosition1Buffer);
 	Parameters->DeformedPosition2Buffer = RegisterAsSRV(GraphBuilder, DeformedLODDatas.DeformedRootTrianglePosition2Buffer);
 
+	Parameters->RestSamplePositionsBuffer = RegisterAsSRV(GraphBuilder, RestLODDatas.RestSamplePositionsBuffer);
+	Parameters->DeformedSamplePositionsBuffer = RegisterAsSRV(GraphBuilder, DeformedLODDatas.DeformedSamplePositionsBuffer);
+
 	Parameters->ViewUniformBuffer = ViewUniformBuffer;
 	Parameters->RenderTargets[0] = FRenderTargetBinding(ColorTarget, ERenderTargetLoadAction::ELoad, 0);
 	Parameters->RenderTargets.DepthStencil = FDepthStencilBinding(DepthTexture, bClearDepth ? ERenderTargetLoadAction::EClear : ERenderTargetLoadAction::ELoad, ERenderTargetLoadAction::ENoAction, FExclusiveDepthStencil::DepthWrite_StencilNop);
 
 	FHairProjectionHairDebugVS::FPermutationDomain PermutationVector;
-	PermutationVector.Set<FHairProjectionHairDebugVS::FInputType>(PrimitiveType == PT_LineList ? 0 : 1);
-
+	PermutationVector.Set<FHairProjectionHairDebugVS::FInputType>(GeometryType == EDebugProjectionHairType::HairFrame ? 0 : GeometryType == EDebugProjectionHairType::HairTriangle ? 1 : 2);
+					
 	TShaderMapRef<FHairProjectionHairDebugVS> VertexShader(ShaderMap, PermutationVector);
 	TShaderMapRef<FHairProjectionHairDebugPS> PixelShader(ShaderMap);
 
@@ -1137,7 +1162,7 @@ void RunHairStrandsDebug(
 				RenderMeshProjection(GraphBuilder, EHairStrandsProjectionMeshType::TargetMesh);
 			}
 
-			auto RenderProjectionData = [&GraphBuilder, ShaderMap, Viewport, &ViewUniformBuffer, Instances, SkinCache, &bClearDepth, SceneColorTexture, DepthTexture](EHairStrandsInterpolationType StrandType, bool bRestTriangle, bool bRestFrame, bool bDeformedTriangle, bool bDeformedFrame)
+			auto RenderProjectionData = [&GraphBuilder, ShaderMap, Viewport, &ViewUniformBuffer, Instances, SkinCache, &bClearDepth, SceneColorTexture, DepthTexture](EHairStrandsInterpolationType StrandType, bool bRestTriangle, bool bRestFrame, bool bRestSamples, bool bDeformedTriangle, bool bDeformedFrame, bool bDeformedSamples)
 			{
 				TArray<int32> HairLODIndices;
 				for (FHairStrandsInstance* AbstractInstance : Instances)
@@ -1190,6 +1215,24 @@ void RunHairStrandsDebug(
 							DepthTexture);
 						bClearDepth = false;
 					}
+					if (bRestSamples)
+					{
+						AddDebugProjectionHairPass(
+							GraphBuilder,
+							ShaderMap,
+							Viewport,
+							ViewUniformBuffer,
+							bClearDepth,
+							EDebugProjectionHairType::HairSamples,
+							HairStrandsTriangleType::RestPose,
+							MeshLODIndex,
+							RestRootResource,
+							DeformedRootResource,
+							Instance->HairGroupPublicData->VFInput.LocalToWorldTransform,
+							SceneColorTexture,
+							DepthTexture);
+						bClearDepth = false;
+					}
 					if (bDeformedTriangle)
 					{
 						AddDebugProjectionHairPass(
@@ -1226,33 +1269,59 @@ void RunHairStrandsDebug(
 							DepthTexture);
 						bClearDepth = false;
 					}
+					if (bDeformedSamples)
+					{
+						AddDebugProjectionHairPass(
+							GraphBuilder,
+							ShaderMap,
+							Viewport,
+							ViewUniformBuffer,
+							bClearDepth,
+							EDebugProjectionHairType::HairSamples,
+							HairStrandsTriangleType::DeformedPose,
+							MeshLODIndex,
+							RestRootResource,
+							DeformedRootResource,
+							Instance->HairGroupPublicData->VFInput.LocalToWorldTransform,
+							SceneColorTexture,
+							DepthTexture);
+						bClearDepth = false;
+					}
 				}
 			};
 
 			if (GHairDebugMeshProjection_Render_HairRestTriangles > 0 ||
 				GHairDebugMeshProjection_Render_HairRestFrames > 0 ||
 				GHairDebugMeshProjection_Render_HairDeformedTriangles > 0 ||
-				GHairDebugMeshProjection_Render_HairDeformedFrames > 0)
+				GHairDebugMeshProjection_Render_HairDeformedFrames > 0 ||
+				GHairDebugMeshProjection_Render_HairDeformedSamples > 0 ||
+				GHairDebugMeshProjection_Render_HairRestSamples > 0)
 			{
 				RenderProjectionData(
 					EHairStrandsInterpolationType::RenderStrands,
 					GHairDebugMeshProjection_Render_HairRestTriangles > 0,
 					GHairDebugMeshProjection_Render_HairRestFrames > 0,
+					GHairDebugMeshProjection_Render_HairRestSamples > 0,
 					GHairDebugMeshProjection_Render_HairDeformedTriangles > 0,
-					GHairDebugMeshProjection_Render_HairDeformedFrames > 0);
+					GHairDebugMeshProjection_Render_HairDeformedFrames > 0,
+					GHairDebugMeshProjection_Render_HairDeformedSamples > 0);
 			}
 
 			if (GHairDebugMeshProjection_Sim_HairRestTriangles > 0 ||
 				GHairDebugMeshProjection_Sim_HairRestFrames > 0 ||
 				GHairDebugMeshProjection_Sim_HairDeformedTriangles > 0 ||
-				GHairDebugMeshProjection_Sim_HairDeformedFrames > 0)
+				GHairDebugMeshProjection_Sim_HairDeformedFrames > 0 ||
+				GHairDebugMeshProjection_Sim_HairDeformedSamples > 0 ||
+				GHairDebugMeshProjection_Sim_HairRestSamples > 0)
 			{
 				RenderProjectionData(
 					EHairStrandsInterpolationType::SimulationStrands,
 					GHairDebugMeshProjection_Sim_HairRestTriangles > 0,
 					GHairDebugMeshProjection_Sim_HairRestFrames > 0,
+					GHairDebugMeshProjection_Sim_HairRestSamples > 0,
 					GHairDebugMeshProjection_Sim_HairDeformedTriangles > 0,
-					GHairDebugMeshProjection_Sim_HairDeformedFrames > 0);
+					GHairDebugMeshProjection_Sim_HairDeformedFrames > 0,
+					GHairDebugMeshProjection_Sim_HairDeformedSamples > 0);
 			}
 		}
 	}
diff --git a/Engine/Shaders/Private/HairStrands/HairStrandsMeshProjectionHairDebug.usf b/Engine/Shaders/Private/HairStrands/HairStrandsMeshProjectionHairDebug.usf
--- a/Engine/Shaders/Private/HairStrands/HairStrandsMeshProjectionHairDebug.usf
+++ b/Engine/Shaders/Private/HairStrands/HairStrandsMeshProjectionHairDebug.usf
@@ -7,6 +7,8 @@
 
 #define INPUT_ROOT_FRAME 0
 #define INPUT_ROOT_TRIANGLE 1
+#define INPUT_ROOT_SAMPLES 2
+
 uint MaxRootCount;
 float4x4 RootLocalToWorld;
 uint DeformedFrameEnable;
@@ -21,7 +23,10 @@ Buffer<float4>  RestPosition2Buffer;
 
 Buffer<float4>  DeformedPosition0Buffer;
 Buffer<float4>  DeformedPosition1Buffer;
-Buffer<float4>  DeformedPosition2Buffer; 
+Buffer<float4>  DeformedPosition2Buffer;
+
+Buffer<float4> RestSamplePositionsBuffer;
+Buffer<float4> DeformedSamplePositionsBuffer;
 
 void MainVS(
 	uint VertexId : SV_VertexID,
@@ -130,6 +135,37 @@ void MainVS(
 	P = mul(float4(P, 1), RootLocalToWorld).xyz;
 #endif
 
+#if PERMUTATION_INPUT_TYPE == INPUT_ROOT_SAMPLES 
+
+	const uint LineVertex = (VertexId%2);
+	const uint LineIndex = VertexId/2;
+	const uint PointIndex = VertexId/6;
+	const uint AxisIndex = (LineIndex%3);
+	
+	if (LineIndex > MaxRootCount) return;
+
+	float3 N;
+	if(AxisIndex == 0)
+	{
+		N = float3(1,0,0);
+		UVs = float2(1,0);
+	}
+	else if(AxisIndex == 1)
+	{
+		N = float3(0,1,0);
+		UVs = float2(0,1);
+	}
+	else if(AxisIndex == 2)
+	{
+		N = float3(0,0,1);
+		UVs = float2(0.5,0.5);
+	}
+	float3 P = (DeformedFrameEnable > 0) ? DeformedSamplePositionsBuffer[PointIndex].xyz : RestSamplePositionsBuffer[PointIndex].xyz;
+	P = (LineVertex == 0) ? P - 0.5 * N : P + 0.5 * N;
+	P = mul(float4(P, 1), RootLocalToWorld).xyz;
+	
+#endif
+	
 	OutVertexUVs = UVs;
 	OutFaceNormal = N;
 	OutPosition = mul(float4(P,1), LWCHackToFloat(PrimaryView.WorldToClip));
