Graivity Preloading for hairs

diff --git a/Engine/Plugins/Runtime/HairStrands/Shaders/Private/NiagaraDataInterfaceHairStrands.ush b/Engine/Plugins/Runtime/HairStrands/Shaders/Private/NiagaraDataInterfaceHairStrands.ush
--- a/Engine/Plugins/Runtime/HairStrands/Shaders/Private/NiagaraDataInterfaceHairStrands.ush
+++ b/Engine/Plugins/Runtime/HairStrands/Shaders/Private/NiagaraDataInterfaceHairStrands.ush
@@ -581,23 +581,59 @@ void DIHairStrands_ComputeEdgeRotation(in FDIHairStrandsContext DIContext, in fl
 	}
 }
 
-void DIHairStrands_ComputeEdgeDirection(in FDIHairStrandsContext DIContext, in float3 NodePosition, in float4 NodeOrientation, out float3 OutRestDirection)
+void DIHairStrands_ComputeEdgeDirection(in FDIHairStrandsContext DIContext, in float3 NodePosition, in float4 NodeOrientation, 
+	in float NodeMass, in float3 GravityVector, in float GravityPreloading, in float BendStiffness, in float StrandThickness, in float RestLength, out float3 OutRestDirection)
 {
+	const int LocalIndex = (GGroupThreadId.x % DIContext.StrandSize);
+
 	SharedNodePosition[GGroupThreadId.x] = NodePosition;
 	SharedNodeOrientation[GGroupThreadId.x] = NodeOrientation;
 	GroupMemoryBarrier(); 
 
-	//ComputeMaterialFrame(StrandsSize);
-
 	OutRestDirection = float3(0,0,0);
-	const int LocalIndex = (GGroupThreadId.x % DIContext.StrandSize);
+
 	if( LocalIndex > 1 )
 	{
-		const float3 EdgeDirection = SharedNodePosition[GGroupThreadId.x] - SharedNodePosition[GGroupThreadId.x-1];
+		float3 EdgeDirection = SharedNodePosition[GGroupThreadId.x] - SharedNodePosition[GGroupThreadId.x-1];
+
+		SharedInverseInertia[GGroupThreadId.x] = BendStiffness*PI*StrandThickness*StrandThickness / (4.0*RestLength);
+		SharedPreviousPosition[GGroupThreadId.x] = EdgeDirection;
+		SharedPreviousOrientation[GGroupThreadId.x] = float4(EdgeDirection,0);
 
-		const float4 ParentFrame = SlerpQuat(SharedNodeOrientation[GGroupThreadId.x-2],SharedNodeOrientation[GGroupThreadId.x-2],0.5);
-		OutRestDirection = UnrotateVectorByQuat(EdgeDirection,ParentFrame);
+		GroupMemoryBarrier(); 
+
+		const float GradientNext = ( LocalIndex < (DIContext.StrandSize-1) ) ?  SharedInverseInertia[GGroupThreadId.x+1]  : 0.0;
+		const float GradientPrev = SharedInverseInertia[GGroupThreadId.x];
+		const float SchurComplement =  (GradientNext * GradientNext + GradientPrev * GradientPrev);
+
+		for(int Index = 0; Index < 100; ++Index)
+		{
+			const int IsRed = (GGroupThreadId.x % 2) == 0;
+			if (!IsRed)
+			{	
+				const float3 DeltaConstraint = ((SharedPreviousPosition[GGroupThreadId.x] - SharedPreviousOrientation[GGroupThreadId.x].xyz) * GradientPrev - 
+												(SharedPreviousPosition[GGroupThreadId.x+1] - SharedPreviousOrientation[GGroupThreadId.x+1].xyz) * GradientNext) - NodeMass * GravityVector;
+
+				SharedPreviousOrientation[GGroupThreadId.x].xyz += DeltaConstraint * GradientPrev / SchurComplement;
+				SharedPreviousOrientation[GGroupThreadId.x+1].xyz -= DeltaConstraint * GradientNext / SchurComplement;
+			}
+			GroupMemoryBarrier(); 
+
+			if (IsRed)
+			{
+				const float3 DeltaConstraint = ((SharedPreviousPosition[GGroupThreadId.x] - SharedPreviousOrientation[GGroupThreadId.x].xyz) * GradientPrev - 
+												(SharedPreviousPosition[GGroupThreadId.x+1] - SharedPreviousOrientation[GGroupThreadId.x+1].xyz) * GradientNext) - NodeMass * GravityVector;
+
+
+				SharedPreviousOrientation[GGroupThreadId.x].xyz += DeltaConstraint * GradientPrev / SchurComplement;
+				SharedPreviousOrientation[GGroupThreadId.x+1].xyz -= DeltaConstraint * GradientNext / SchurComplement;
+			}
+			GroupMemoryBarrier();
+		}
+		EdgeDirection = SharedPreviousOrientation[GGroupThreadId.x].xyz * GravityPreloading + (1.0-GravityPreloading) * EdgeDirection;
+		OutRestDirection = UnrotateVectorByQuat(EdgeDirection,SharedNodeOrientation[GGroupThreadId.x-2]);
 	}
+
 	GroupMemoryBarrier(); 
 }
 
diff --git a/Engine/Plugins/Runtime/HairStrands/Source/HairStrandsCore/Private/GroomAssetPhysics.cpp b/Engine/Plugins/Runtime/HairStrands/Source/HairStrandsCore/Private/GroomAssetPhysics.cpp
--- a/Engine/Plugins/Runtime/HairStrands/Source/HairStrandsCore/Private/GroomAssetPhysics.cpp
+++ b/Engine/Plugins/Runtime/HairStrands/Source/HairStrandsCore/Private/GroomAssetPhysics.cpp
@@ -8,6 +8,7 @@ FHairSolverSettings::FHairSolverSettings()
 	NiagaraSolver = EGroomNiagaraSolvers::AngularSprings;
 	SubSteps = 5;
 	IterationCount = 5;
+	GravityPreloading = 0.0f;
 }
 
 FHairExternalForces::FHairExternalForces()
diff --git a/Engine/Plugins/Runtime/HairStrands/Source/HairStrandsCore/Private/GroomComponent.cpp b/Engine/Plugins/Runtime/HairStrands/Source/HairStrandsCore/Private/GroomComponent.cpp
--- a/Engine/Plugins/Runtime/HairStrands/Source/HairStrandsCore/Private/GroomComponent.cpp
+++ b/Engine/Plugins/Runtime/HairStrands/Source/HairStrandsCore/Private/GroomComponent.cpp
@@ -3324,8 +3324,8 @@ void UGroomComponent::TickComponent(float DeltaTime, enum ELevelTick TickType, F
 	if (!bInitSimulation)
 	{
 		const bool bLocalSimulation = SimulationSettings.SimulationSetup.bLocalSimulation &&
-			(SimulationSettings.SimulationSetup.LinearVelocityScale == 1.0) &&
-			(SimulationSettings.SimulationSetup.AngularVelocityScale == 1.0);
+			(SimulationSettings.SimulationSetup.LinearVelocityScale == 0.0) &&
+			(SimulationSettings.SimulationSetup.AngularVelocityScale == 0.0);
 		if (!bLocalSimulation)
 		{
 			if (USkeletalMeshComponent* ParentComp = Cast<USkeletalMeshComponent>(GetAttachParent()))
diff --git a/Engine/Plugins/Runtime/HairStrands/Source/HairStrandsCore/Private/Niagara/NiagaraDataInterfaceHairStrands.cpp b/Engine/Plugins/Runtime/HairStrands/Source/HairStrandsCore/Private/Niagara/NiagaraDataInterfaceHairStrands.cpp
--- a/Engine/Plugins/Runtime/HairStrands/Source/HairStrandsCore/Private/Niagara/NiagaraDataInterfaceHairStrands.cpp
+++ b/Engine/Plugins/Runtime/HairStrands/Source/HairStrandsCore/Private/Niagara/NiagaraDataInterfaceHairStrands.cpp
@@ -36,6 +36,7 @@ static const FName GetSubStepsName("GetSubSteps");
 static const FName GetIterationCountName("GetIterationCount");
 
 static const FName GetGravityVectorName("GetGravityVector");
+static const FName GetGravityPreloadingName("GetGravityPreloading");
 static const FName GetAirDragName("GetAirDrag");
 static const FName GetAirVelocityName("GetAirVelocity");
 
@@ -485,6 +486,7 @@ void FNDIHairStrandsData::Update(UNiagaraDataInterfaceHairStrands* Interface, FN
 			IterationCount = HairPhysics.SolverSettings.IterationCount;
 
 			GravityVector = HairPhysics.ExternalForces.GravityVector;
+			GravityPreloading = HairPhysics.SolverSettings.GravityPreloading;
 			AirDrag = HairPhysics.ExternalForces.AirDrag;
 			AirVelocity = HairPhysics.ExternalForces.AirVelocity;
 
@@ -544,13 +546,13 @@ void FNDIHairStrandsData::Update(UNiagaraDataInterfaceHairStrands* Interface, FN
 					const FQuat DeltaRotation = DeltaTransform.GetRotation();
 					
 					// Apply linear velocity scale
-					BoneLinearVelocity = FVector3f(FMath::Clamp(1.f - SimulationSettings.SimulationSetup.LinearVelocityScale, 0.f, 1.f) * DeltaTransform.GetTranslation() / DeltaSeconds);
+					BoneLinearVelocity = FVector3f(FMath::Clamp(SimulationSettings.SimulationSetup.LinearVelocityScale, 0.f, 1.f) * DeltaTransform.GetTranslation() / DeltaSeconds);
 					BoneLinearAcceleration = (BoneLinearVelocity-PreviousBoneLinearVelocity) / DeltaSeconds;
 
 					
 					// Apply angular velocity scale
 					BoneAngularVelocity = (FVector3f)BoneTransform.TransformVector(DeltaRotation.GetRotationAxis() * DeltaRotation.GetAngle() *
-						FMath::Clamp(1.f - SimulationSettings.SimulationSetup.AngularVelocityScale, 0.f, 1.f)) / DeltaSeconds;
+						FMath::Clamp(SimulationSettings.SimulationSetup.AngularVelocityScale, 0.f, 1.f)) / DeltaSeconds;
 					BoneAngularAcceleration = (BoneAngularVelocity-PreviousBoneAngularVelocity) / DeltaSeconds;
 				}
 				else
@@ -1298,6 +1300,16 @@ void UNiagaraDataInterfaceHairStrands::GetFunctions(TArray<FNiagaraFunctionSigna
 
 		OutFunctions.Add(Sig);
 	}
+	{
+		FNiagaraFunctionSignature Sig;
+		Sig.Name = GetGravityPreloadingName;
+		Sig.bMemberFunction = true;
+		Sig.bRequiresContext = false;
+		Sig.Inputs.Add(FNiagaraVariable(FNiagaraTypeDefinition(GetClass()), TEXT("Hair Strands")));
+		Sig.Outputs.Add(FNiagaraVariable(FNiagaraTypeDefinition::GetFloatDef(), TEXT("Gravity Preloading")));
+
+		OutFunctions.Add(Sig);
+	}
 	{
 		FNiagaraFunctionSignature Sig;
 		Sig.Name = GetAirDragName;
@@ -2160,6 +2172,12 @@ void UNiagaraDataInterfaceHairStrands::GetFunctions(TArray<FNiagaraFunctionSigna
 		Sig.Inputs.Add(FNiagaraVariable(FNiagaraTypeDefinition(GetClass()), TEXT("Hair Strands")));
 		Sig.Inputs.Add(FNiagaraVariable(FNiagaraTypeDefinition::GetVec3Def(), TEXT("Node Position")));
 		Sig.Inputs.Add(FNiagaraVariable(FNiagaraTypeDefinition::GetQuatDef(), TEXT("Node Orientation")));
+		Sig.Inputs.Add(FNiagaraVariable(FNiagaraTypeDefinition::GetFloatDef(), TEXT("Node Mass")));
+		Sig.Inputs.Add(FNiagaraVariable(FNiagaraTypeDefinition::GetVec3Def(), TEXT("Gravity Vector")));
+		Sig.Inputs.Add(FNiagaraVariable(FNiagaraTypeDefinition::GetFloatDef(), TEXT("Gravity Preloading")));
+		Sig.Inputs.Add(FNiagaraVariable(FNiagaraTypeDefinition::GetFloatDef(), TEXT("Bend Stiffness")));
+		Sig.Inputs.Add(FNiagaraVariable(FNiagaraTypeDefinition::GetFloatDef(), TEXT("Node Thickness")));
+		Sig.Inputs.Add(FNiagaraVariable(FNiagaraTypeDefinition::GetFloatDef(), TEXT("Rest Length")));
 		Sig.Outputs.Add(FNiagaraVariable(FNiagaraTypeDefinition::GetVec3Def(), TEXT("Rest Direction")));
 
 		OutFunctions.Add(Sig);
@@ -2275,6 +2293,7 @@ DEFINE_NDI_DIRECT_FUNC_BINDER(UNiagaraDataInterfaceHairStrands, GetStrandSize);
 DEFINE_NDI_DIRECT_FUNC_BINDER(UNiagaraDataInterfaceHairStrands, GetSubSteps);
 DEFINE_NDI_DIRECT_FUNC_BINDER(UNiagaraDataInterfaceHairStrands, GetIterationCount);
 
+DEFINE_NDI_DIRECT_FUNC_BINDER(UNiagaraDataInterfaceHairStrands, GetGravityPreloading);
 DEFINE_NDI_DIRECT_FUNC_BINDER(UNiagaraDataInterfaceHairStrands, GetGravityVector);
 DEFINE_NDI_DIRECT_FUNC_BINDER(UNiagaraDataInterfaceHairStrands, GetAirDrag);
 DEFINE_NDI_DIRECT_FUNC_BINDER(UNiagaraDataInterfaceHairStrands, GetAirVelocity);
@@ -2395,6 +2414,11 @@ void UNiagaraDataInterfaceHairStrands::GetVMExternalFunction(const FVMExternalFu
 		check(BindingInfo.GetNumInputs() == 1 && BindingInfo.GetNumOutputs() == 3);
 		NDI_FUNC_BINDER(UNiagaraDataInterfaceHairStrands, GetGravityVector)::Bind(this, OutFunc);
 	}
+	else if (BindingInfo.Name == GetGravityPreloadingName)
+	{
+		check(BindingInfo.GetNumInputs() == 1 && BindingInfo.GetNumOutputs() == 1);
+		NDI_FUNC_BINDER(UNiagaraDataInterfaceHairStrands, GetGravityPreloading)::Bind(this, OutFunc);
+	}
 	else if (BindingInfo.Name == GetAirDragName)
 	{
 		check(BindingInfo.GetNumInputs() == 1 && BindingInfo.GetNumOutputs() == 1);
@@ -2727,7 +2751,7 @@ void UNiagaraDataInterfaceHairStrands::GetVMExternalFunction(const FVMExternalFu
 	}
 	else if (BindingInfo.Name == ComputeEdgeDirectionName)
 	{
-		check(BindingInfo.GetNumInputs() == 8 && BindingInfo.GetNumOutputs() == 3);
+		check(BindingInfo.GetNumInputs() == 16 && BindingInfo.GetNumOutputs() == 3);
 		NDI_FUNC_BINDER(UNiagaraDataInterfaceHairStrands, ComputeEdgeDirection)::Bind(this, OutFunc);
 	}
 	else if (BindingInfo.Name == UpdateMaterialFrameName)
@@ -2871,6 +2895,17 @@ void UNiagaraDataInterfaceHairStrands::GetGravityVector(FVectorVMExternalFunctio
 	}
 }
 
+void UNiagaraDataInterfaceHairStrands::GetGravityPreloading(FVectorVMExternalFunctionContext& Context)
+{
+	VectorVM::FUserPtrHandler<FNDIHairStrandsData> InstData(Context);
+	VectorVM::FExternalFuncRegisterHandler<float> OutGravityPreloading(Context);
+
+	for (int32 InstanceIdx = 0; InstanceIdx < Context.GetNumInstances(); ++InstanceIdx)
+	{
+		*OutGravityPreloading.GetDestAndAdvance() = InstData->GravityPreloading;
+	}
+}
+
 void UNiagaraDataInterfaceHairStrands::GetAirDrag(FVectorVMExternalFunctionContext& Context)
 {
 	VectorVM::FUserPtrHandler<FNDIHairStrandsData> InstData(Context);
@@ -4027,9 +4062,10 @@ in float RestLength, in float DeltaTime, in int NodeOffset, in float MaterialDam
 	else if (FunctionInfo.DefinitionName == ComputeEdgeDirectionName)
 	{
 		static const TCHAR *FormatSample = TEXT(R"(
-							void {InstanceFunctionName} (in float3 NodePosition, in float4 NodeOrientation, out float3 OutRestDirection)
+							void {InstanceFunctionName} (in float3 NodePosition, in float4 NodeOrientation, 
+	in float NodeMass, in float3 GravityVector, in float GravityPreloading, in float BendStiffness, in float StrandThickness, in float RestLength, out float3 OutRestDirection)
 							{
-								{HairStrandsContextName} DIHairStrands_ComputeEdgeDirection(DIContext,NodePosition,NodeOrientation,OutRestDirection);
+								{HairStrandsContextName} DIHairStrands_ComputeEdgeDirection(DIContext,NodePosition,NodeOrientation,NodeMass,GravityVector,GravityPreloading,BendStiffness,StrandThickness,RestLength,OutRestDirection);
 							}
 							)");
 		OutHLSL += FString::Format(FormatSample, ArgsSample);
diff --git a/Engine/Plugins/Runtime/HairStrands/Source/HairStrandsCore/Public/GroomAssetPhysics.h b/Engine/Plugins/Runtime/HairStrands/Source/HairStrandsCore/Public/GroomAssetPhysics.h
--- a/Engine/Plugins/Runtime/HairStrands/Source/HairStrandsCore/Public/GroomAssetPhysics.h
+++ b/Engine/Plugins/Runtime/HairStrands/Source/HairStrandsCore/Public/GroomAssetPhysics.h
@@ -60,6 +60,10 @@ struct HAIRSTRANDSCORE_API FHairSolverSettings
 	UPROPERTY(EditAnywhere, Category = "SolverSettings", meta = (EditCondition = "NiagaraSolver == EGroomNiagaraSolvers::CustomSolver", ToolTip = "Custom niagara system to be used if custom solver is picked"))
 	TSoftObjectPtr<UNiagaraSystem> CustomSystem;
 
+	/** Optimisation of the rest state configuration to compensate from the gravity */
+	UPROPERTY(EditAnywhere, Category = "SolverSettings", meta = (EditCondition = "NiagaraSolver == EGroomNiagaraSolvers::AngularSprings", ToolTip = "Optimisation of the rest state configuration to compensate from the gravity "))
+	float GravityPreloading;
+
 	/** Number of substeps to be used */
 	UPROPERTY(EditAnywhere, Category = "SolverSettings", meta = (ToolTip = "Number of sub steps to be done per frame. The actual solver calls are done at 24 fps"))
 	int32 SubSteps;
@@ -358,13 +362,13 @@ struct HAIRSTRANDSCORE_API FHairSimulationSetup
 	UPROPERTY(EditAnywhere, BlueprintReadWrite, Interp, Category = "SimulationSetup", meta = (ToolTip = "Strands simulation is done in local space"))
 	bool bLocalSimulation = true;
 	
-	/** Scaling of the bone world linear velocity contribution */
-	UPROPERTY(EditAnywhere, BlueprintReadWrite, Interp, Category = "SimulationSetup", meta = (ToolTip = "Scaling of the bone world linear velocity contribution", EditCondition = "bLocalSimulation"))
-	float LinearVelocityScale = 0.0f;
+	/** The amount of linear velocities sent to the local groom space from the reference bone */
+	UPROPERTY(EditAnywhere, BlueprintReadWrite, Interp, Category = "SimulationSetup", meta = (ToolTip = "The amount of linear velocities sent to the local groom space from the reference bone", EditCondition = "bLocalSimulation", UIMin = "0", UIMax = "1", ClampMin = "0", ClampMax = "1"))
+	float LinearVelocityScale = 1.0f;
 	
-	/** Scaling of the bone world angular velocity contribution */
-	UPROPERTY(EditAnywhere, BlueprintReadWrite, Interp, Category = "SimulationSetup", meta = (ToolTip = "Scaling of the bone world angular velocity contribution", EditCondition = "bLocalSimulation"))
-	float AngularVelocityScale = 0.0f;
+	/** The amount of angular velocities sent to the local groom space from the reference bone */
+	UPROPERTY(EditAnywhere, BlueprintReadWrite, Interp, Category = "SimulationSetup", meta = (ToolTip = "The amount of angular velocities sent to the local groom space from the reference bone", EditCondition = "bLocalSimulation", UIMin = "0", UIMax = "1", ClampMin = "0", ClampMax = "1"))
+	float AngularVelocityScale = 1.0f;
 
 	/** Bone used for the simulation local space */
 	UPROPERTY(EditAnywhere, BlueprintReadWrite, Interp, Category = "SimulationSetup", meta = (ToolTip = "Bone used for the simulation local space", EditCondition = "bLocalSimulation"))
diff --git a/Engine/Plugins/Runtime/HairStrands/Source/HairStrandsCore/Public/Niagara/NiagaraDataInterfaceHairStrands.h b/Engine/Plugins/Runtime/HairStrands/Source/HairStrandsCore/Public/Niagara/NiagaraDataInterfaceHairStrands.h
--- a/Engine/Plugins/Runtime/HairStrands/Source/HairStrandsCore/Public/Niagara/NiagaraDataInterfaceHairStrands.h
+++ b/Engine/Plugins/Runtime/HairStrands/Source/HairStrandsCore/Public/Niagara/NiagaraDataInterfaceHairStrands.h
@@ -127,6 +127,7 @@ struct FNDIHairStrandsData
 		IterationCount = 20;
 
 		GravityVector = FVector(0.0, 0.0, -981.0);
+		GravityPreloading = 0.0;
 		AirDrag = 0.1;
 		AirVelocity = FVector(0, 0, 0);
 
@@ -195,6 +196,7 @@ struct FNDIHairStrandsData
 			IterationCount = OtherDatas->IterationCount;
 
 			GravityVector = OtherDatas->GravityVector;
+			GravityPreloading = OtherDatas->GravityPreloading;
 			AirDrag = OtherDatas->AirDrag;
 			AirVelocity = OtherDatas->AirVelocity;
 
@@ -291,6 +293,9 @@ struct FNDIHairStrandsData
 
 	/** Acceleration vector in cm/s2 to be used for the gravity*/
 	FVector GravityVector;
+	
+	/** Optimisation of the rest state configuration to compensate from the gravity */
+	float GravityPreloading;
 
 	/** Coefficient between 0 and 1 to be used for the air drag */
 	float AirDrag;
@@ -441,6 +446,8 @@ public:
 
 	void GetGravityVector(FVectorVMExternalFunctionContext& Context);
 
+	void GetGravityPreloading(FVectorVMExternalFunctionContext& Context);
+
 	void GetAirDrag(FVectorVMExternalFunctionContext& Context);
 
 	void GetAirVelocity(FVectorVMExternalFunctionContext& Context);
