Fix hair resources transition for raytracing geometry prior to BVH building.

diff --git a/Engine/Plugins/Runtime/HairStrands/Source/HairStrandsCore/Private/HairStrandsRendering.cpp b/Engine/Plugins/Runtime/HairStrands/Source/HairStrandsCore/Private/HairStrandsRendering.cpp
--- a/Engine/Plugins/Runtime/HairStrands/Source/HairStrandsCore/Private/HairStrandsRendering.cpp
+++ b/Engine/Plugins/Runtime/HairStrands/Source/HairStrandsCore/Private/HairStrandsRendering.cpp
@@ -1224,6 +1224,15 @@ public:
 	}
 };
 
+BEGIN_SHADER_PARAMETER_STRUCT(FStrandsResourceBLASParameters, )
+	RDG_BUFFER_ACCESS(PositionBuffer, ERHIAccess::SRVCompute)
+	RDG_BUFFER_ACCESS(IndexBuffer, ERHIAccess::SRVCompute)
+END_SHADER_PARAMETER_STRUCT()
+
+BEGIN_SHADER_PARAMETER_STRUCT(FCardsOrMeshesResourceBLASParameters, )
+RDG_BUFFER_ACCESS(PositionBuffer, ERHIAccess::SRVCompute)
+END_SHADER_PARAMETER_STRUCT()
+
 IMPLEMENT_GLOBAL_SHADER(FHairRaytracingGeometryCS, "/Engine/Private/HairStrands/HairStrandsRaytracingGeometry.usf", "MainCS", SF_Compute);
 
 static void AddGenerateRaytracingGeometryPass(
@@ -2010,9 +2019,14 @@ void ComputeHairStrandsInterpolation(
 					Instance->Strands.CachedHairRootScale	= HairRootScaleRT;
 					Instance->Strands.CachedHairTipScale	= HairTipScaleRT;
 
+					FStrandsResourceBLASParameters* Parameters = GraphBuilder.AllocParameters<FStrandsResourceBLASParameters>();
+					Parameters->PositionBuffer = Raytracing_PositionBuffer.Buffer;
+					Parameters->IndexBuffer = Raytracing_IndexBuffer.Buffer;
+
 					GraphBuilder.AddPass(
 						RDG_EVENT_NAME("HairStrands::UpdateBLAS(Strands)"),
-						ERDGPassFlags::NeverCull,
+						Parameters,
+						ERDGPassFlags::NeverCull | ERDGPassFlags::Compute,
 					[Instance, HairLODIndex, bNeedUpdate, bProceduralPrimitive, ProceduralSplits](FRHICommandListImmediate& RHICmdList)
 					{
 						SCOPED_GPU_MASK(RHICmdList, FRHIGPUMask::All());
@@ -2149,6 +2163,9 @@ void ComputeHairStrandsInterpolation(
 			#if RHI_RAYTRACING
 			if (LOD.RaytracingResource)
 			{
+				FCardsOrMeshesResourceBLASParameters* Parameters = GraphBuilder.AllocParameters<FCardsOrMeshesResourceBLASParameters>();
+				Parameters->PositionBuffer = LOD.DeformedResource ? Register(GraphBuilder, LOD.DeformedResource->GetBuffer(FHairCardsDeformedResource::Current), ERDGImportedBufferFlags::None).Buffer : nullptr;
+
 				// * When cards are dynamic, RT geometry is built once and then update.
 				// * When cards are static (i.e., no sim, not attached to  skinning()), in which case RT geometry needs only to be built once. This static geometry is shared between all the instances, 
 				//   and owned by the asset, rathter than the component. In this case the RT geometry will be built only once for all the instances
@@ -2158,7 +2175,8 @@ void ComputeHairStrandsInterpolation(
 				{
 					GraphBuilder.AddPass(
 						RDG_EVENT_NAME("HairStrands::UpdateBLAS(Cards)"),
-						ERDGPassFlags::NeverCull,
+						Parameters,
+						ERDGPassFlags::NeverCull | ERDGPassFlags::Compute,
 						[Instance, HairLODIndex, bNeedUpdate](FRHICommandListImmediate& RHICmdList)
 					{
 						SCOPED_GPU_MASK(RHICmdList, FRHIGPUMask::All());
@@ -2218,13 +2236,17 @@ void ComputeHairStrandsInterpolation(
 			FHairGroupInstance::FMeshes::FLOD& LOD = Instance->Meshes.LODs[HairLODIndex];
 			if (LOD.RaytracingResource)
 			{
+				FCardsOrMeshesResourceBLASParameters* Parameters = GraphBuilder.AllocParameters<FCardsOrMeshesResourceBLASParameters>();
+				Parameters->PositionBuffer = LOD.DeformedResource ? Register(GraphBuilder, LOD.DeformedResource->GetBuffer(FHairMeshesDeformedResource::Current), ERDGImportedBufferFlags::None).Buffer : nullptr;
+
 				const bool bNeedUpdate = bNeedDeformation;
 				const bool bNeedBuid = !LOD.RaytracingResource->bIsRTGeometryInitialized;
 				if (bNeedBuid || bNeedUpdate)
 				{
 					GraphBuilder.AddPass(
 						RDG_EVENT_NAME("HairStrands::UpdateBLAS(Meshes)"),
-						ERDGPassFlags::NeverCull,
+						Parameters,
+						ERDGPassFlags::NeverCull | ERDGPassFlags::Compute,
 						[Instance, HairLODIndex, bNeedUpdate](FRHICommandListImmediate& RHICmdList)
 					{
 						SCOPED_GPU_MASK(RHICmdList, FRHIGPUMask::All());
