Fix a crash issue related to HairStrands

diff --git a/Engine/Plugins/Runtime/HairStrands/Source/HairStrandsCore/Private/Niagara/NiagaraDataInterfaceHairStrands.cpp b/Engine/Plugins/Runtime/HairStrands/Source/HairStrandsCore/Private/Niagara/NiagaraDataInterfaceHairStrands.cpp
--- a/Engine/Plugins/Runtime/HairStrands/Source/HairStrandsCore/Private/Niagara/NiagaraDataInterfaceHairStrands.cpp
+++ b/Engine/Plugins/Runtime/HairStrands/Source/HairStrandsCore/Private/Niagara/NiagaraDataInterfaceHairStrands.cpp
@@ -481,6 +481,7 @@ void FNDIHairStrandsData::Update(UNiagaraDataInterfaceHairStrands* Interface, FN
 			StrandsSize = static_cast<uint8>(HairPhysics.StrandsParameters.StrandsSize);
 
 			HairGroupInstance = Interface->IsComponentValid() ? Interface->SourceComponent->GetGroupInstance(GroupIndex) : nullptr;
+			HairGroupInstSource = Interface->IsComponentValid() ? Interface->SourceComponent : nullptr;
 
 			SubSteps = HairPhysics.SolverSettings.SubSteps;
 			IterationCount = HairPhysics.SolverSettings.IterationCount;
@@ -704,7 +705,8 @@ struct FNDIHairStrandsParametersCS : public FNiagaraDataInterfaceParametersCS
 		FNDIHairStrandsData* ProxyData = InterfaceProxy->SystemInstancesToProxyData.Find(Context.SystemInstanceID);
 
 		const bool bIsHairValid = ProxyData != nullptr && ProxyData->HairStrandsBuffer && ProxyData->HairStrandsBuffer->IsInitialized();
-		const bool bHasSkinningBinding = bIsHairValid && ProxyData->HairGroupInstance && ProxyData->HairGroupInstance->BindingType == EHairBindingType::Skinning;
+		const bool bIsHairGroupInstValid = ProxyData->HairGroupInstSource != nullptr && ProxyData->HairGroupInstSource->ContainsGroupInstance(ProxyData->HairGroupInstance);
+		const bool bHasSkinningBinding = bIsHairValid && bIsHairGroupInstValid && ProxyData->HairGroupInstance->BindingType == EHairBindingType::Skinning;
 		const bool bIsRootValid = bIsHairValid && ProxyData->HairStrandsBuffer->SourceDeformedRootResources && ProxyData->HairStrandsBuffer->SourceDeformedRootResources->IsInitialized() && bHasSkinningBinding;
 		const bool bIsRestValid = bIsHairValid && ProxyData->HairStrandsBuffer->SourceRestResources && ProxyData->HairStrandsBuffer->SourceRestResources->IsInitialized() &&
 			 
@@ -714,10 +716,10 @@ struct FNDIHairStrandsParametersCS : public FNiagaraDataInterfaceParametersCS
 			// TEMP: These check are only temporary for avoiding crashes while we find the bottom of the issue.
 			ProxyData->HairStrandsBuffer->CurvesOffsetsBuffer.SRV && ProxyData->HairStrandsBuffer->ParamsScaleBuffer.SRV && ProxyData->HairStrandsBuffer->BoundingBoxBuffer.UAV;
 
-		const bool bIsGeometryValid = bIsHairValid && (!ProxyData->HairGroupInstance || (ProxyData->HairGroupInstance && (ProxyData->HairGroupInstance->GeometryType != EHairGeometryType::NoneGeometry)));
+		const bool bIsGeometryValid = bIsHairValid && (!bIsHairGroupInstValid || (bIsHairGroupInstValid && (ProxyData->HairGroupInstance->GeometryType != EHairGeometryType::NoneGeometry)));
 		const bool bIsDeformedValid = bIsHairValid && ProxyData->HairStrandsBuffer->SourceDeformedResources && ProxyData->HairStrandsBuffer->SourceDeformedResources->IsInitialized();
 
-		if (bIsHairValid && bIsRestValid && bIsGeometryValid)
+		if (bIsHairValid && bIsRestValid && bIsGeometryValid && bIsHairGroupInstValid)
 		{
 			check(ProxyData);
 
diff --git a/Engine/Plugins/Runtime/HairStrands/Source/HairStrandsCore/Public/GroomComponent.h b/Engine/Plugins/Runtime/HairStrands/Source/HairStrandsCore/Public/GroomComponent.h
--- a/Engine/Plugins/Runtime/HairStrands/Source/HairStrandsCore/Public/GroomComponent.h
+++ b/Engine/Plugins/Runtime/HairStrands/Source/HairStrandsCore/Public/GroomComponent.h
@@ -229,7 +229,8 @@ public:
 	/** Hair group instance access */
 	uint32 GetGroupCount() const { return HairGroupInstances.Num();  }
 	FHairGroupInstance* GetGroupInstance(int32 Index) { return Index >= 0 && Index < HairGroupInstances.Num() ? HairGroupInstances[Index] : nullptr; }
-	const FHairGroupInstance* GetGroupInstance(int32 Index) const { return Index >= 0 && Index < HairGroupInstances.Num() ? HairGroupInstances[Index] : nullptr; } 
+	const FHairGroupInstance* GetGroupInstance(int32 Index) const { return Index >= 0 && Index < HairGroupInstances.Num() ? HairGroupInstances[Index] : nullptr; }
+	bool ContainsGroupInstance(FHairGroupInstance* Instance) const { return HairGroupInstances.Contains(Instance); }
 
 	//~ Begin UPrimitiveComponent Interface
 	EHairGeometryType GetMaterialGeometryType(int32 ElementIndex) const;
diff --git a/Engine/Plugins/Runtime/HairStrands/Source/HairStrandsCore/Public/Niagara/NiagaraDataInterfaceHairStrands.h b/Engine/Plugins/Runtime/HairStrands/Source/HairStrandsCore/Public/Niagara/NiagaraDataInterfaceHairStrands.h
--- a/Engine/Plugins/Runtime/HairStrands/Source/HairStrandsCore/Public/Niagara/NiagaraDataInterfaceHairStrands.h
+++ b/Engine/Plugins/Runtime/HairStrands/Source/HairStrandsCore/Public/Niagara/NiagaraDataInterfaceHairStrands.h
@@ -115,6 +115,7 @@ struct FNDIHairStrandsData
 		
 		GlobalInterpolation = false;
 		bSkinningTransfer = false;
+		HairGroupInstSource = nullptr;
 		HairGroupInstance = nullptr;
 
 		TickCount = 0;
@@ -184,6 +185,7 @@ struct FNDIHairStrandsData
 			GlobalInterpolation = OtherDatas->GlobalInterpolation;
 			bSkinningTransfer = OtherDatas->bSkinningTransfer;
 			BindingType = OtherDatas->BindingType;
+			HairGroupInstSource = OtherDatas->HairGroupInstSource;
 			HairGroupInstance = OtherDatas->HairGroupInstance;
 
 			TickCount = OtherDatas->TickCount;
@@ -282,6 +284,9 @@ struct FNDIHairStrandsData
 	/** Hair group instance */
 	FHairGroupInstance* HairGroupInstance;
 
+	/** Source component of the hair group instance */
+	TWeakObjectPtr<class UGroomComponent> HairGroupInstSource;
+
 	/** Binding type between the groom asset and the attached skeletal mesh */
 	EHairBindingType BindingType;
 	
