Fix static analysis warning.

diff --git a/Engine/Plugins/Runtime/HairStrands/Source/HairStrandsCore/Private/Niagara/NiagaraDataInterfaceHairStrands.cpp b/Engine/Plugins/Runtime/HairStrands/Source/HairStrandsCore/Private/Niagara/NiagaraDataInterfaceHairStrands.cpp
--- a/Engine/Plugins/Runtime/HairStrands/Source/HairStrandsCore/Private/Niagara/NiagaraDataInterfaceHairStrands.cpp
+++ b/Engine/Plugins/Runtime/HairStrands/Source/HairStrandsCore/Private/Niagara/NiagaraDataInterfaceHairStrands.cpp
@@ -705,7 +705,7 @@ struct FNDIHairStrandsParametersCS : public FNiagaraDataInterfaceParametersCS
 		FNDIHairStrandsData* ProxyData = InterfaceProxy->SystemInstancesToProxyData.Find(Context.SystemInstanceID);
 
 		const bool bIsHairValid = ProxyData != nullptr && ProxyData->HairStrandsBuffer && ProxyData->HairStrandsBuffer->IsInitialized();
-		const bool bIsHairGroupInstValid = ProxyData->HairGroupInstSource != nullptr && ProxyData->HairGroupInstSource->ContainsGroupInstance(ProxyData->HairGroupInstance);
+		const bool bIsHairGroupInstValid = ProxyData != nullptr && ProxyData->HairGroupInstSource != nullptr && ProxyData->HairGroupInstSource->ContainsGroupInstance(ProxyData->HairGroupInstance);
 		const bool bHasSkinningBinding = bIsHairValid && bIsHairGroupInstValid && ProxyData->HairGroupInstance->BindingType == EHairBindingType::Skinning;
 		const bool bIsRootValid = bIsHairValid && ProxyData->HairStrandsBuffer->SourceDeformedRootResources && ProxyData->HairStrandsBuffer->SourceDeformedRootResources->IsInitialized() && bHasSkinningBinding;
 		const bool bIsRestValid = bIsHairValid && ProxyData->HairStrandsBuffer->SourceRestResources && ProxyData->HairStrandsBuffer->SourceRestResources->IsInitialized() &&
