Groom: Fixed crash when applying a guides groom cache on groom component with groom asset that hasn't enabled simulation.

diff --git a/Engine/Plugins/Runtime/HairStrands/Source/HairStrandsCore/Private/GroomComponent.cpp b/Engine/Plugins/Runtime/HairStrands/Source/HairStrandsCore/Private/GroomComponent.cpp
--- a/Engine/Plugins/Runtime/HairStrands/Source/HairStrandsCore/Private/GroomComponent.cpp
+++ b/Engine/Plugins/Runtime/HairStrands/Source/HairStrandsCore/Private/GroomComponent.cpp
@@ -3609,6 +3609,31 @@ void UGroomComponent::PostEditChangeProperty(FPropertyChangedEvent& PropertyChan
 		ResetAnimationTime();
 	}
 
+	// Check to see if simulation is required, which is needed when using with the guides groom cache,
+	// and enable it on all groom groups if it's not already
+	const EGroomCacheType GroomCacheType = GroomCache ? GroomCache->GetType() : EGroomCacheType::None;
+	const bool bSimulationRequired = GroomCacheType == EGroomCacheType::Guides;
+	if (bSimulationRequired && GroomAsset)
+	{
+		bool bGroomAssetChanged = false;
+		for (int32 Index = 0; Index < GroomAsset->GetNumHairGroups(); ++Index)
+		{
+			if (!IsSimulationEnable(Index, -1))
+			{
+				if (!bGroomAssetChanged)
+				{
+					GroomAsset->Modify();
+					bGroomAssetChanged = true;
+				}
+				GroomAsset->HairGroupsPhysics[Index].SolverSettings.EnableSimulation = true;
+			}
+		}
+		if (bGroomAssetChanged)
+		{
+			GroomAsset->CacheDerivedDatas();
+		}
+	}
+
 	if (bAssetChanged)
 	{
 		if (GroomAsset)
