Add edit API for groom.

This allows to convert a groom asset into a editable groom struct, easier to manipulation.

This CL also contains a few examples/tests on how this API can be used. These tests are not hooked, but have been tested with some local changes.

diff --git a/Engine/Plugins/Runtime/HairStrands/Source/HairStrandsCore/Private/GroomAsset.cpp b/Engine/Plugins/Runtime/HairStrands/Source/HairStrandsCore/Private/GroomAsset.cpp
--- a/Engine/Plugins/Runtime/HairStrands/Source/HairStrandsCore/Private/GroomAsset.cpp
+++ b/Engine/Plugins/Runtime/HairStrands/Source/HairStrandsCore/Private/GroomAsset.cpp
@@ -723,16 +723,16 @@ void UGroomAsset::Serialize(FArchive& Ar)
 		{
 			// When serializing data for editor, serialize the HairDescription as bulk data
 			// The computed groom data is fetched from the Derived Data Cache
-			if (!HairDescriptionBulkData)
+			if (!HairDescriptionBulkData[EHairDescriptionType::Source])
 			{
 				// When loading, bulk data can be null so instantiate a new one to serialize into
-				HairDescriptionBulkData = MakeUnique<FHairDescriptionBulkData>();
+				HairDescriptionBulkData[EHairDescriptionType::Source] = MakeUnique<FHairDescriptionBulkData>();
 			}
 
 			// We don' serialize the hair description if it is a transactional object (i.e., for undo/redo)
 			if (!Ar.IsTransacting())
 			{
-				HairDescriptionBulkData->Serialize(Ar, this);
+				HairDescriptionBulkData[EHairDescriptionType::Source]->Serialize(Ar, this);
 			}
 		}
 #endif // WITH_EDITORONLY_DATA
@@ -1036,7 +1036,7 @@ void UGroomAsset::PostLoad()
 	const bool bIsLegacyAsset = HairGroupsInterpolation.Num() == 0;
 	if (bIsLegacyAsset)
 	{
-		if (HairDescriptionBulkData)
+		if (HairDescriptionBulkData[HairDescriptionType])
 		{
 			const FHairDescriptionGroups& LocalHairDescriptionGroups = GetHairDescriptionGroups();
 
@@ -1625,7 +1625,7 @@ const TArray<UAssetUserData*>* UGroomAsset::GetAssetUserDataArray() const
 bool UGroomAsset::CanRebuildFromDescription() const
 {
 #if WITH_EDITORONLY_DATA
-	return HairDescriptionBulkData.IsValid() && !HairDescriptionBulkData->IsEmpty();
+	return HairDescriptionBulkData[HairDescriptionType].IsValid() && !HairDescriptionBulkData[HairDescriptionType]->IsEmpty();
 #else
 	return false;
 #endif
@@ -2003,12 +2003,12 @@ FString UGroomAsset::BuildDerivedDataKeySuffix(uint32 GroupIndex, const FHairGro
 	GroomDerivedDataCacheUtils::SerializeHairInterpolationSettingsForDDC(Ar, GroupIndex, const_cast<FHairGroupsInterpolation&>(InterpolationSettings), const_cast<FHairGroupsLOD&>(LODSettings), bNeedInterpolationData);
 
 	FString KeySuffix;
-	if (HairDescriptionBulkData)
+	if (HairDescriptionBulkData[HairDescriptionType])
 	{
 		// Reserve twice the size of TempBytes because of ByteToHex below + 3 for "ID" and \0
-		KeySuffix.Reserve(HairDescriptionBulkData->GetIdString().Len() + TempBytes.Num() * 2 + 3);
+		KeySuffix.Reserve(HairDescriptionBulkData[HairDescriptionType]->GetIdString().Len() + TempBytes.Num() * 2 + 3);
 		KeySuffix += TEXT("ID");
-		KeySuffix += HairDescriptionBulkData->GetIdString();
+		KeySuffix += HairDescriptionBulkData[HairDescriptionType]->GetIdString();
 	}
 	else
 	{
@@ -2065,46 +2065,48 @@ FString UGroomAsset::GetDerivedDataKeyForMeshes(uint32 GroupIndex)
 }
 
 
-void UGroomAsset::CommitHairDescription(FHairDescription&& InHairDescription)
+void UGroomAsset::CommitHairDescription(FHairDescription&& InHairDescription, EHairDescriptionType Type)
 {
-	CachedHairDescription = MakeUnique<FHairDescription>(InHairDescription);
+	HairDescriptionType = Type;
 
-	if (!HairDescriptionBulkData)
+	CachedHairDescription[HairDescriptionType] = MakeUnique<FHairDescription>(InHairDescription);
+
+	if (!HairDescriptionBulkData[HairDescriptionType])
 	{
-		HairDescriptionBulkData = MakeUnique<FHairDescriptionBulkData>();
+		HairDescriptionBulkData[HairDescriptionType] = MakeUnique<FHairDescriptionBulkData>();
 	}
 
 	// Update the cached hair description groups with the new hair description data
-	CachedHairDescriptionGroups = MakeUnique<FHairDescriptionGroups>();
-	FGroomBuilder::BuildHairDescriptionGroups(*CachedHairDescription, *CachedHairDescriptionGroups);
+	CachedHairDescriptionGroups[HairDescriptionType] = MakeUnique<FHairDescriptionGroups>();
+	FGroomBuilder::BuildHairDescriptionGroups(*CachedHairDescription[HairDescriptionType], *CachedHairDescriptionGroups[HairDescriptionType]);
 
-	HairDescriptionBulkData->SaveHairDescription(*CachedHairDescription);
+	HairDescriptionBulkData[HairDescriptionType]->SaveHairDescription(*CachedHairDescription[HairDescriptionType]);
 }
 
 const FHairDescriptionGroups& UGroomAsset::GetHairDescriptionGroups()
 {
-	check(HairDescriptionBulkData);
+	check(HairDescriptionBulkData[HairDescriptionType]);
 
-	if (!CachedHairDescription)
+	if (!CachedHairDescription[HairDescriptionType])
 	{
-		CachedHairDescription = MakeUnique<FHairDescription>();
-		HairDescriptionBulkData->LoadHairDescription(*CachedHairDescription);
+		CachedHairDescription[HairDescriptionType] = MakeUnique<FHairDescription>();
+		HairDescriptionBulkData[HairDescriptionType]->LoadHairDescription(*CachedHairDescription[HairDescriptionType]);
 	}
-	if (!CachedHairDescriptionGroups)
+	if (!CachedHairDescriptionGroups[HairDescriptionType])
 	{
-		CachedHairDescriptionGroups = MakeUnique<FHairDescriptionGroups>();
-		FGroomBuilder::BuildHairDescriptionGroups(*CachedHairDescription, *CachedHairDescriptionGroups);
+		CachedHairDescriptionGroups[HairDescriptionType] = MakeUnique<FHairDescriptionGroups>();
+		FGroomBuilder::BuildHairDescriptionGroups(*CachedHairDescription[HairDescriptionType], *CachedHairDescriptionGroups[HairDescriptionType]);
 	}
 
-	return *CachedHairDescriptionGroups;
+	return *CachedHairDescriptionGroups[HairDescriptionType];
 }
 
 FHairDescription UGroomAsset::GetHairDescription() const
 {
 	FHairDescription OutHairDescription;
-	if (HairDescriptionBulkData)
+	if (HairDescriptionBulkData[HairDescriptionType])
 	{
-		HairDescriptionBulkData->LoadHairDescription(OutHairDescription);
+		HairDescriptionBulkData[HairDescriptionType]->LoadHairDescription(OutHairDescription);
 	}
 	return MoveTemp(OutHairDescription);
 }
@@ -2165,7 +2167,7 @@ bool UGroomAsset::CacheDerivedData(uint32 GroupIndex)
 
 bool UGroomAsset::CacheStrandsData(uint32 GroupIndex, FString& OutDerivedDataKey)
 {
-	if (!HairDescriptionBulkData)
+	if (!HairDescriptionBulkData[HairDescriptionType])
 	{
 		UE_LOG(LogHairStrands, Error, TEXT("[Groom] The groom asset %s is too old. Please reimported the groom from its original source file."));
 		return false;
diff --git a/Engine/Plugins/Runtime/HairStrands/Source/HairStrandsCore/Private/GroomEdit.cpp b/Engine/Plugins/Runtime/HairStrands/Source/HairStrandsCore/Private/GroomEdit.cpp
new file mode 100644
--- /dev/null
+++ b/Engine/Plugins/Runtime/HairStrands/Source/HairStrandsCore/Private/GroomEdit.cpp
@@ -0,0 +1,374 @@
+// Copyright Epic Games, Inc. All Rights Reserved.
+
+#pragma once
+
+#include "GroomEdit.h"
+#include "GroomAsset.h"
+#include "GroomBuilder.h"
+
+// Notes:
+// This is a work-in-progress intermediate representation for editable groom.
+// At the moment we do a full in-out conversion. As the conversion refines, we 
+// will optimize conversion toward minimum edits
+
+#if WITH_EDITORONLY_DATA
+
+static void HairStrandsDataToEditableHairGuides(const FHairStrandsDatas& In, TArray<FEditableHairGuide>& Out)
+{
+	const bool bHasValidStrandsIDs = In.StrandsCurves.StrandIDs.Num() > 0;
+	const uint32 StrandsCount = In.GetNumCurves();
+	Out.SetNum(StrandsCount);
+	for (uint32 StrandsIt = 0; StrandsIt < StrandsCount; ++StrandsIt)
+	{
+		Out[StrandsIt].GuideID= bHasValidStrandsIDs ? In.StrandsCurves.StrandIDs[StrandsIt] : StrandsIt;
+		Out[StrandsIt].RootUV = In.StrandsCurves.CurvesRootUV[StrandsIt];
+
+		const uint32 PointCount = In.StrandsCurves.CurvesCount[StrandsIt];
+		const uint32 PointOffset = In.StrandsCurves.CurvesOffset[StrandsIt];
+
+		Out[StrandsIt].ControlPoints.SetNum(PointCount);
+		for (uint32 PointIt = 0; PointIt < PointCount; ++PointIt)
+		{
+			const uint32 EffectiveIndex = PointIt + PointOffset;
+			Out[StrandsIt].ControlPoints[PointIt].Position = In.StrandsPoints.PointsPosition[EffectiveIndex];
+			Out[StrandsIt].ControlPoints[PointIt].U = In.StrandsPoints.PointsCoordU[EffectiveIndex];
+		}
+	}
+}
+
+static void HairStrandsDataToEditableHairStrands(const FHairStrandsDatas& In, TArray<FEditableHairStrand>& Out)
+{
+	const bool bHasValidStrandsIDs = In.StrandsCurves.StrandIDs.Num() > 0;
+	const bool bHasValidClosestGuide = In.StrandsCurves.CurvesClosestGuideIDs.Num() > 0;
+	const uint32 StrandsCount = In.GetNumCurves();
+	Out.SetNum(StrandsCount);
+	for (uint32 StrandsIt = 0; StrandsIt<StrandsCount; ++StrandsIt)
+	{			
+		const FIntVector GuideIDs	 = bHasValidClosestGuide ? In.StrandsCurves.CurvesClosestGuideIDs[StrandsIt] : FIntVector(-1,-1,-1);
+		const FVector3f GuideWeights = bHasValidClosestGuide ? (FVector3f)In.StrandsCurves.CurvesClosestGuideIDs[StrandsIt] : FVector3f::ZeroVector;
+
+		Out[StrandsIt].StrandID		= bHasValidStrandsIDs ? In.StrandsCurves.StrandIDs[StrandsIt] : StrandsIt;
+		Out[StrandsIt].RootUV		= In.StrandsCurves.CurvesRootUV[StrandsIt];
+		Out[StrandsIt].GuideIDs[0]  = GuideIDs.X;
+		Out[StrandsIt].GuideIDs[1]  = GuideIDs.Y;
+		Out[StrandsIt].GuideIDs[2]  = GuideIDs.Z;
+		Out[StrandsIt].GuideWeights[0] = GuideWeights.X;
+		Out[StrandsIt].GuideWeights[1] = GuideWeights.Y;
+		Out[StrandsIt].GuideWeights[2] = GuideWeights.Z;
+
+		const uint32 PointCount  = In.StrandsCurves.CurvesCount[StrandsIt];
+		const uint32 PointOffset = In.StrandsCurves.CurvesOffset[StrandsIt];
+
+		Out[StrandsIt].ControlPoints.SetNum(PointCount);
+		for (uint32 PointIt = 0; PointIt < PointCount; ++PointIt)
+		{				
+			const uint32 EffectiveIndex = PointIt + PointOffset;
+			Out[StrandsIt].ControlPoints[PointIt].Position	= In.StrandsPoints.PointsPosition[EffectiveIndex];
+			Out[StrandsIt].ControlPoints[PointIt].U			= In.StrandsPoints.PointsCoordU[EffectiveIndex];
+			Out[StrandsIt].ControlPoints[PointIt].Radius	= In.StrandsPoints.PointsRadius[EffectiveIndex] * In.StrandsCurves.MaxRadius;
+			Out[StrandsIt].ControlPoints[PointIt].BaseColor	= In.StrandsPoints.PointsBaseColor[EffectiveIndex];
+			Out[StrandsIt].ControlPoints[PointIt].Roughness	= In.StrandsPoints.PointsRoughness[EffectiveIndex];
+		}
+	}
+}
+
+static void EditableGroomGroupToHairDescription(FHairDescription& OutHairDescription, const TArray<FEditableHairGuide>& In, uint32 GroupID, const FName& GroupName, uint32& CurveIndex, uint32& VertexIndex)
+{
+	for (const FEditableHairGuide& Guide : In)
+	{
+		FStrandID StrandID = (FStrandID)CurveIndex++;
+
+		const uint32 PointCount = Guide.ControlPoints.Num();
+		SetHairStrandAttribute(OutHairDescription, StrandID, HairAttribute::Strand::VertexCount, (int32)PointCount);
+		SetHairStrandAttribute(OutHairDescription, StrandID, HairAttribute::Strand::RootUV,      Guide.RootUV);
+		SetHairStrandAttribute(OutHairDescription, StrandID, HairAttribute::Strand::ID,          int(Guide.GuideID));
+		SetHairStrandAttribute(OutHairDescription, StrandID, HairAttribute::Strand::GroupID,     int(GroupID));
+		SetHairStrandAttribute(OutHairDescription, StrandID, HairAttribute::Strand::GroupName,   GroupName);
+		SetHairStrandAttribute(OutHairDescription, StrandID, HairAttribute::Strand::Guide,       1);
+
+		for (const FEditableHairGuideControlPoint& Point : Guide.ControlPoints)
+		{
+			//FVertexID VertexID = OutHairDescription.AddVertex();
+			FVertexID VertexID = (FVertexID)VertexIndex++; // FVertexID VertexID = OutHairDescription.AddVertex();
+			SetHairVertexAttribute(OutHairDescription, VertexID, HairAttribute::Vertex::Position,  Point.Position);
+		}
+	}
+}
+
+static void EditableGroomGroupToHairDescription(FHairDescription& OutHairDescription, const TArray<FEditableHairStrand>& In, uint32 GroupID, const FName& GroupName, uint32& CurveIndex, uint32& VertexIndex)
+{
+	// Check if the groom as precomputed weights
+	bool bHasSingleWeights = false;
+	bool bHasTripleWeights = false;
+	for (const FEditableHairStrand& Strands : In)
+	{
+		bHasTripleWeights = Strands.GuideIDs[0] != -1 && Strands.GuideIDs[1] != -1 && Strands.GuideIDs[2] != -1;
+		bHasSingleWeights = Strands.GuideIDs[0] != -1;
+	}
+
+	for (const FEditableHairStrand& Strands : In)
+	{
+		FStrandID StrandID = (FStrandID)CurveIndex++;
+
+		const uint32 PointCount = Strands.ControlPoints.Num();
+		SetHairStrandAttribute(OutHairDescription, StrandID, HairAttribute::Strand::VertexCount, (int32)PointCount);
+		SetHairStrandAttribute(OutHairDescription, StrandID, HairAttribute::Strand::RootUV,      Strands.RootUV);
+		SetHairStrandAttribute(OutHairDescription, StrandID, HairAttribute::Strand::ID,          int(Strands.StrandID));
+		SetHairStrandAttribute(OutHairDescription, StrandID, HairAttribute::Strand::GroupID,     int(GroupID));
+		SetHairStrandAttribute(OutHairDescription, StrandID, HairAttribute::Strand::GroupName,   GroupName);
+		SetHairStrandAttribute(OutHairDescription, StrandID, HairAttribute::Strand::Guide,       0);
+
+		if (bHasTripleWeights)
+		{
+			const FVector3f ClosestGuides = FVector3f(Strands.GuideIDs[0], Strands.GuideIDs[1], Strands.GuideIDs[2]);
+			const FVector3f GuideWeights = FVector3f(Strands.GuideWeights[0], Strands.GuideWeights[1], Strands.GuideWeights[2]);
+			SetHairStrandAttribute(OutHairDescription, StrandID, HairAttribute::Strand::ClosestGuides, ClosestGuides);
+			SetHairStrandAttribute(OutHairDescription, StrandID, HairAttribute::Strand::GuideWeights, GuideWeights);
+		}
+		else if (bHasSingleWeights)
+		{
+			int ClosestGuide = Strands.GuideIDs[0];
+			float GuideWeight = Strands.GuideWeights[0];
+			SetHairStrandAttribute(OutHairDescription, StrandID, HairAttribute::Strand::ClosestGuides, ClosestGuide);
+			SetHairStrandAttribute(OutHairDescription, StrandID, HairAttribute::Strand::GuideWeights, GuideWeight);
+		}
+
+		for (const FEditableHairStrandControlPoint& Point : Strands.ControlPoints)
+		{
+			FVertexID VertexID = (FVertexID)VertexIndex++;
+			SetHairVertexAttribute(OutHairDescription, VertexID, HairAttribute::Vertex::Position,  Point.Position);
+			SetHairVertexAttribute(OutHairDescription, VertexID, HairAttribute::Vertex::Width,     Point.Radius * 2.f);
+			SetHairVertexAttribute(OutHairDescription, VertexID, HairAttribute::Vertex::Color,     FVector3f(Point.BaseColor.R, Point.BaseColor.G, Point.BaseColor.B));
+			SetHairVertexAttribute(OutHairDescription, VertexID, HairAttribute::Vertex::Roughness, Point.Roughness);
+		}
+	}
+}
+#endif // WITH_EDITORONLY_DATA
+
+void ConvertFromGroomAsset(UGroomAsset* InGroom, FEditableGroom* OutGroom)
+{
+#if WITH_EDITORONLY_DATA
+	check(InGroom);
+	check(OutGroom);
+	UGroomAsset& In = *InGroom;
+	FEditableGroom& Out = *OutGroom;
+
+	const FHairDescriptionGroups& HairDescriptionGroups = In.GetHairDescriptionGroups();
+
+	const uint32 GroupCount = In.GetNumHairGroups();
+	Out.Groups.SetNum(GroupCount);
+
+	for (uint32 GroupIndex = 0; GroupIndex < GroupCount; ++GroupIndex)
+	{
+		FHairGroupInfo GroupInfo;
+		FHairStrandsDatas StrandsData;
+		FHairStrandsDatas GuidesData;
+		FGroomBuilder::BuildData(HairDescriptionGroups.HairGroups[GroupIndex], In.HairGroupsInterpolation[GroupIndex], GroupInfo, StrandsData, GuidesData);
+
+		Out.Groups[GroupIndex].GroupID = GroupInfo.GroupID;
+		Out.Groups[GroupIndex].GroupName = GroupInfo.GroupName;
+		HairStrandsDataToEditableHairGuides(GuidesData, Out.Groups[GroupIndex].Guides);
+
+		HairStrandsDataToEditableHairStrands(StrandsData, Out.Groups[GroupIndex].Strands);
+	}
+#endif
+}
+
+void ConvertToGroomAsset(UGroomAsset* OutGroom, const FEditableGroom* InGroom, uint32 Operations)
+{
+#if WITH_EDITORONLY_DATA
+	check(InGroom);
+	check(OutGroom);
+	const FEditableGroom& In = *InGroom;
+	const UGroomAsset& Out = *OutGroom;
+
+	// If group count differs ensure, extent/shrink it accordingly to the edited groom
+	const uint32 GroupCount = InGroom->Groups.Num();
+	if (uint32(OutGroom->GetNumHairGroups()) != GroupCount)
+	{
+		OutGroom->SetNumGroup(GroupCount, false, false);
+	}
+
+	// Count strands/vertex for later allocation
+	uint32 VertexCount = 0;
+	uint32 StrandCount = 0;
+	for (const FEditableGroomGroup& Group : In.Groups)
+	{
+		StrandCount += Group.Strands.Num();
+		StrandCount += Group.Guides.Num();
+		for (const FEditableHairStrand& Strand : Group.Strands)
+		{
+			VertexCount += Strand.ControlPoints.Num();
+		}
+		for (const FEditableHairGuide& Guide : Group.Guides)
+		{
+			VertexCount += Guide.ControlPoints.Num();
+		}
+	}
+
+	// Create hair description and allocate vertex/strands
+	FHairDescription HairDescription;
+	HairDescription.InitializeStrands(StrandCount);
+	HairDescription.InitializeVertices(VertexCount);
+
+	// Convert guides/strands to the final
+	uint32 CurveIndex = 0;
+	uint32 VertexIndex = 0;
+	for (const FEditableGroomGroup& Group : In.Groups)
+	{
+		EditableGroomGroupToHairDescription(HairDescription, Group.Guides,  Group.GroupID, Group.GroupName, CurveIndex, VertexIndex);
+		EditableGroomGroupToHairDescription(HairDescription, Group.Strands, Group.GroupID, Group.GroupName, CurveIndex, VertexIndex);
+	}
+
+	// Commit & build data for the groom asset
+	// Internally, this will derived the build data, recreate resources, and 
+	// recreate all components using this groom assets
+	OutGroom->CommitHairDescription(MoveTemp(HairDescription), UGroomAsset::EHairDescriptionType::Edit);
+	OutGroom->CacheDerivedDatas();
+	OutGroom->OnGroomAssetChanged.Broadcast();
+#endif
+}
+
+///////////////////////////////////////////////////////////////////////////////////////////////////
+// Examples / Tests (Temporary)
+
+#if WITH_EDITORONLY_DATA
+static int32 GHairStrandsEditTestCase = 0;
+static FAutoConsoleVariableRef CVarHairStrandsEditTestCase(TEXT("r.HairStrands.Edit.TestCase"), GHairStrandsEditTestCase, TEXT("Test case for testing the groom edit API"));
+
+// Test 1: Change radius of all strands
+static void InternalEditTest1(UGroomAsset* InGroom)
+{
+	FEditableGroom EditGroom;
+	ConvertFromGroomAsset(InGroom, &EditGroom);
+
+	// Override the radius of all vertices
+	for (FEditableGroomGroup& G : EditGroom.Groups)
+	{
+		for (FEditableHairStrand& S : G.Strands)
+		{
+			for (FEditableHairStrandControlPoint& C : S.ControlPoints)
+			{
+				C.Radius = 1.0f;
+			}
+		}
+	}
+
+	ConvertToGroomAsset(InGroom, &EditGroom, EEditableGroomOperations::ControlPoints_Modified);
+}
+
+// Test 2: Orient all strands based on their root direction
+static void InternalEditTest2(UGroomAsset* InGroom)
+{
+	FEditableGroom EditGroom;
+	ConvertFromGroomAsset(InGroom, &EditGroom);
+
+	// Override the radius of all vertices
+	for (FEditableGroomGroup& G : EditGroom.Groups)
+	{
+		for (FEditableHairStrand& S : G.Strands)
+		{
+			const uint32 PointCount = S.ControlPoints.Num();
+			// Compute normal direction based
+			FVector3f Direction = FVector3f(0, 0, 1);
+			if (PointCount >= 2)
+			{
+				Direction = S.ControlPoints[1].Position - S.ControlPoints[0].Position;
+				Direction.Normalize();
+			}
+
+			for (uint32 PointIt = 2; PointIt < PointCount; ++PointIt)
+			{
+				const float Distance = FVector3f::Distance(S.ControlPoints[PointIt].Position, S.ControlPoints[PointIt-1].Position);
+				S.ControlPoints[PointIt].Position = S.ControlPoints[PointIt - 1].Position + Direction * Distance;
+			}
+		}
+	}
+
+	ConvertToGroomAsset(InGroom, &EditGroom, EEditableGroomOperations::ControlPoints_Modified);
+}
+
+// Test 3: Remove some strands
+static void InternalEditTest3(UGroomAsset* InGroom)
+{
+	FEditableGroom EditGroom;
+	ConvertFromGroomAsset(InGroom, &EditGroom);
+
+	// Override the radius of all vertices
+	for (FEditableGroomGroup& G : EditGroom.Groups)
+	{
+		G.Strands.SetNum(uint32(G.Strands.Num() * 0.1f));
+	}
+
+	ConvertToGroomAsset(InGroom, &EditGroom, EEditableGroomOperations::Strands_Deleted);
+}
+
+// Test 4: Add a new group
+static void InternalEditTest4(UGroomAsset* InGroom)
+{
+	FEditableGroom EditGroom;
+	ConvertFromGroomAsset(InGroom, &EditGroom);
+
+	// Override the radius of all vertices
+	if (EditGroom.Groups.Num() > 0)
+	{
+		const uint32 GroupIt = EditGroom.Groups.Num();
+		EditGroom.Groups.SetNum(EditGroom.Groups.Num() + 1);
+
+		const FEditableGroomGroup& SrcG = EditGroom.Groups[0];
+		FEditableGroomGroup& DstG = EditGroom.Groups[GroupIt];
+
+		DstG.GroupID = GroupIt;
+		DstG.GroupName = FName(TEXT("MyNewGroup"));
+
+		const FVector3f Translation = FVector3f(100, 0, 0);
+		for (const FEditableHairGuide& SrcS : SrcG.Guides)
+		{
+			DstG.Guides.Add(SrcS);
+			for (FEditableHairGuideControlPoint& C : DstG.Guides[DstG.Guides.Num()-1].ControlPoints)
+			{
+				C.Position += Translation;
+			}
+		}
+	
+		for (const FEditableHairStrand& SrcS : SrcG.Strands)
+		{
+			DstG.Strands.Add(SrcS);
+			for (FEditableHairStrandControlPoint& C : DstG.Strands[DstG.Strands.Num() - 1].ControlPoints)
+			{
+				C.Position += Translation;
+			}
+		}
+	}
+
+	ConvertToGroomAsset(InGroom, &EditGroom, EEditableGroomOperations::Group_Added);
+}
+
+// Test 5: Remove new group
+static void InternalEditTest5(UGroomAsset* InGroom)
+{
+	FEditableGroom EditGroom;
+	ConvertFromGroomAsset(InGroom, &EditGroom);
+
+	// Override the radius of all vertices
+	if (EditGroom.Groups.Num() > 1)
+	{
+		EditGroom.Groups.SetNum(EditGroom.Groups.Num() - 1);
+	}
+
+	ConvertToGroomAsset(InGroom, &EditGroom, EEditableGroomOperations::Group_Deleted);
+}
+
+void EditTest(UGroomAsset* InGroom)
+{
+	switch (GHairStrandsEditTestCase)
+	{
+		case 1: InternalEditTest1(InGroom); break;
+		case 2: InternalEditTest2(InGroom); break;
+		case 3: InternalEditTest3(InGroom); break;
+		case 4: InternalEditTest4(InGroom); break;
+		case 5: InternalEditTest5(InGroom); break;
+	}
+}
+
+#endif // WITH_EDITORONLY_DATA
\ No newline at end of file
diff --git a/Engine/Plugins/Runtime/HairStrands/Source/HairStrandsCore/Private/GroomRBFDeformer.cpp b/Engine/Plugins/Runtime/HairStrands/Source/HairStrandsCore/Private/GroomRBFDeformer.cpp
--- a/Engine/Plugins/Runtime/HairStrands/Source/HairStrandsCore/Private/GroomRBFDeformer.cpp
+++ b/Engine/Plugins/Runtime/HairStrands/Source/HairStrandsCore/Private/GroomRBFDeformer.cpp
@@ -415,7 +415,7 @@ static void ApplyDeformationToGroom(const TArray<FRBFDeformedPositions>& Deforme
 		}
 	}
 	
-	GroomAsset->CommitHairDescription(MoveTemp(HairDescription));
+	GroomAsset->CommitHairDescription(MoveTemp(HairDescription), UGroomAsset::EHairDescriptionType::Source);
 	GroomAsset->UpdateHairGroupsInfo();
 
 	// Update/reimport the cards/meshes geometry which have been deformed prior to call this function
diff --git a/Engine/Plugins/Runtime/HairStrands/Source/HairStrandsCore/Public/GroomAsset.h b/Engine/Plugins/Runtime/HairStrands/Source/HairStrandsCore/Public/GroomAsset.h
--- a/Engine/Plugins/Runtime/HairStrands/Source/HairStrandsCore/Public/GroomAsset.h
+++ b/Engine/Plugins/Runtime/HairStrands/Source/HairStrandsCore/Public/GroomAsset.h
@@ -628,8 +628,15 @@ private:
 	bool HasValidCardsData(uint32 GroupIndex) const;
 	bool HasValidMeshesData(uint32 GroupIndex) const;
 public:
+	enum EHairDescriptionType
+	{
+		Source,
+		Edit,
+		Count
+	};
+
 	/** Commits a HairDescription to buffer for serialization */
-	void CommitHairDescription(FHairDescription&& HairDescription);
+	void CommitHairDescription(FHairDescription&& HairDescription, EHairDescriptionType Type);
 	FHairDescription GetHairDescription() const;
 
 	/** Get/Build render & guides data based on the hair description and interpolation settings */
@@ -655,13 +662,14 @@ private:
 	
 	FString BuildDerivedDataKeySuffix(uint32 GroupIndex, const FHairGroupsInterpolation& InterpolationSettings, const FHairGroupsLOD& LODSettings) const;
 	bool IsFullyCached();
-	TUniquePtr<FHairDescriptionBulkData> HairDescriptionBulkData;
+	TUniquePtr<FHairDescriptionBulkData> HairDescriptionBulkData[EHairDescriptionType::Count];
+	EHairDescriptionType HairDescriptionType = EHairDescriptionType::Source;
 
 	// Transient HairDescription & HairDescriptionGroups, which are built from HairDescriptionBulkData.
 	// All these data (bulk/desc/groups) needs to be in sync. I.e., when the HairDescription is updated, 
 	// HairDescriptionGroups needs to also be updated
-	TUniquePtr<FHairDescription> CachedHairDescription;
-	TUniquePtr<FHairDescriptionGroups> CachedHairDescriptionGroups;
+	TUniquePtr<FHairDescription> CachedHairDescription[EHairDescriptionType::Count];
+	TUniquePtr<FHairDescriptionGroups> CachedHairDescriptionGroups[EHairDescriptionType::Count];
 
 	TArray<FString> StrandsDerivedDataKey;
 	TArray<FString> CardsDerivedDataKey;
diff --git a/Engine/Plugins/Runtime/HairStrands/Source/HairStrandsCore/Public/GroomEdit.h b/Engine/Plugins/Runtime/HairStrands/Source/HairStrandsCore/Public/GroomEdit.h
new file mode 100644
--- /dev/null
+++ b/Engine/Plugins/Runtime/HairStrands/Source/HairStrandsCore/Public/GroomEdit.h
@@ -0,0 +1,85 @@
+// Copyright Epic Games, Inc. All Rights Reserved.
+
+#pragma once
+
+#include "CoreMinimal.h"
+#include "UObject/ObjectMacros.h"
+#include "Engine/EngineTypes.h"
+
+class UGroomAsset;
+
+///////////////////////////////////////////////////////////////////////////////////////////////////
+// Strands
+struct FEditableHairStrandControlPoint
+{
+	FVector3f Position;
+	float Radius;
+	float U;
+	FLinearColor BaseColor;
+	float Roughness;
+};
+
+struct FEditableHairStrand
+{
+	TArray<FEditableHairStrandControlPoint> ControlPoints;
+
+	uint32 StrandID;
+	FVector2f RootUV;
+
+	// Pre-computed Simulation
+	uint32 GuideIDs[3];
+	float GuideWeights[3];
+};
+
+///////////////////////////////////////////////////////////////////////////////////////////////////
+// Guides
+struct FEditableHairGuideControlPoint
+{
+	FVector3f Position;
+	float U;
+};
+
+struct FEditableHairGuide
+{
+	TArray<FEditableHairGuideControlPoint> ControlPoints;
+
+	uint32 GuideID;
+	FVector2f RootUV;
+};
+
+///////////////////////////////////////////////////////////////////////////////////////////////////
+// Group & Groom
+struct FEditableGroomGroup
+{
+	uint32 GroupID;
+	FName GroupName;
+	TArray<FEditableHairStrand> Strands;
+	TArray<FEditableHairGuide> Guides;
+};
+
+struct FEditableGroom
+{
+	TArray<FEditableGroomGroup> Groups;
+};
+
+enum EEditableGroomOperations
+{
+	ControlPoints_Added    = 0x1,
+	ControlPoints_Modified = 0x2,
+	ControlPoints_Deleted  = 0x4,
+
+	Strands_Added          = 0x8,
+	Strands_Modified       = 0x10,
+	Strands_Deleted        = 0x20,
+
+	Group_Added            = 0x40,
+	Group_Deleted          = 0x80
+};
+
+// Convert a groom asset into an editable groom asset
+HAIRSTRANDSCORE_API void ConvertFromGroomAsset(UGroomAsset* In, FEditableGroom* Out);
+
+// Convert an editable groom asset into a groom asset
+// The 'operations' flag indicates what type of modifications have been done onto the editable groom. 
+// This helps to drive the update mechanism.
+HAIRSTRANDSCORE_API void ConvertToGroomAsset(UGroomAsset* Out, const FEditableGroom* In, uint32 Operations);
\ No newline at end of file
diff --git a/Engine/Plugins/Runtime/HairStrands/Source/HairStrandsEditor/Private/HairStrandsImporter.cpp b/Engine/Plugins/Runtime/HairStrands/Source/HairStrandsEditor/Private/HairStrandsImporter.cpp
--- a/Engine/Plugins/Runtime/HairStrands/Source/HairStrandsEditor/Private/HairStrandsImporter.cpp
+++ b/Engine/Plugins/Runtime/HairStrands/Source/HairStrandsEditor/Private/HairStrandsImporter.cpp
@@ -41,7 +41,7 @@ UGroomAsset* FHairStrandsImporter::ImportHair(const FHairImportContext& ImportCo
 	check(OutHairAsset->AreGroupsValid());
 	check(uint32(OutHairAsset->GetNumHairGroups()) == GroupCount);
 
-	OutHairAsset->CommitHairDescription(MoveTemp(HairDescription));
+	OutHairAsset->CommitHairDescription(MoveTemp(HairDescription), UGroomAsset::EHairDescriptionType::Source);
 
 	// Populate the interpolation settings with the new settings from the importer	
 	for (uint32 GroupIndex = 0; GroupIndex < GroupCount; ++GroupIndex)
