Sim Cache post read to update groom positions

diff --git a/Engine/Plugins/FX/Niagara/Source/Niagara/Classes/NiagaraDataInterface.h b/Engine/Plugins/FX/Niagara/Source/Niagara/Classes/NiagaraDataInterface.h
--- a/Engine/Plugins/FX/Niagara/Source/Niagara/Classes/NiagaraDataInterface.h
+++ b/Engine/Plugins/FX/Niagara/Source/Niagara/Classes/NiagaraDataInterface.h
@@ -388,6 +388,9 @@ public:
 	/** Ticks the per instance data for this interface, if it has any. */
 	virtual bool PerInstanceTick(void* PerInstanceData, FNiagaraSystemInstance* SystemInstance, float DeltaSeconds) { return false; }
 	virtual bool PerInstanceTickPostSimulate(void* PerInstanceData, FNiagaraSystemInstance* SystemInstance, float DeltaSeconds) { return false; }
+	
+	/** Update internal datas after reading the cache. */
+	virtual void SimCachePostReadFrame(void* OptionalPerInstanceData, FNiagaraSystemInstance* SystemInstance) {}
 
 #if WITH_EDITORONLY_DATA
 	/** Allows the generic class defaults version of this class to specify any dependencies/version/etc that might invalidate the compile. It should never depend on the value of specific properties.*/
diff --git a/Engine/Plugins/FX/Niagara/Source/Niagara/Private/NiagaraSimCache.cpp b/Engine/Plugins/FX/Niagara/Source/Niagara/Private/NiagaraSimCache.cpp
--- a/Engine/Plugins/FX/Niagara/Source/Niagara/Private/NiagaraSimCache.cpp
+++ b/Engine/Plugins/FX/Niagara/Source/Niagara/Private/NiagaraSimCache.cpp
@@ -665,6 +665,13 @@ bool UNiagaraSimCache::ReadFrame(int32 FrameIndex, FNiagaraSystemInstance* Syste
 			Helper.ReadDataBuffer(CacheLayout.EmitterLayouts[i], CacheEmitterFrame.ParticleDataBuffers, EmitterInstance.GetData());
 		}
 	}
+	for(TPair<TWeakObjectPtr<UNiagaraDataInterface>, int32>& DataInterfacePair : SystemInstance->DataInterfaceInstanceDataOffsets)
+	{
+		if (UNiagaraDataInterface* Interface = DataInterfacePair.Key.Get())
+		{
+			Interface->SimCachePostReadFrame(&SystemInstance->DataInterfaceInstanceData[DataInterfacePair.Value], SystemInstance);
+		}
+	}
 	return true;
 }
 
diff --git a/Engine/Plugins/Runtime/HairStrands/Shaders/Private/NiagaraGroomInterpolationUtils.ush b/Engine/Plugins/Runtime/HairStrands/Shaders/Private/NiagaraGroomInterpolationUtils.ush
new file mode 100644
--- /dev/null
+++ b/Engine/Plugins/Runtime/HairStrands/Shaders/Private/NiagaraGroomInterpolationUtils.ush
@@ -0,0 +1,56 @@
+﻿// Copyright Epic Games, Inc. All Rights Reserved.
+
+/*=============================================================================
+NiagaraGroomInterpolationUtils.ush 
+=============================================================================*/
+
+#pragma once
+
+#include "NiagaraQuaternionUtils.ush"
+
+// Compute the projection triangle
+void DIHairStrands_ComputeProjectionTriangle(in float2 ProjectionUV, in float3 PA, in float3 PB, in float3 PC, in float3 Offset, 
+			out float3 OutTrianglePosition, out float4 OutTriangleRotation)
+{
+	OutTrianglePosition = PA * ProjectionUV.x + PB * ProjectionUV.y + PC * (1.0 - ProjectionUV.x - ProjectionUV.y) + Offset;
+
+	float3 TangentU = PB - PA;
+	float3 TangentV = PC - PA;
+
+	const float3 Normal = normalize(cross(TangentU, TangentV));
+	TangentV = normalize(TangentV);
+	TangentU = normalize(cross(TangentV, Normal));
+	
+	const float3 RotationMatrix[3] = { TangentU, TangentV, Normal };
+	OutTriangleRotation = QuatFromMatrix(RotationMatrix);
+}
+
+// Eval the triangle local position
+float3 DIHairStrands_TriangleLocalPosition(in float3 TrianglePosition, in float4 TriangleRotation, in float3 WorldPosition)
+{
+	return RotateVectorByQuat(WorldPosition - TrianglePosition, InverseQuat(TriangleRotation));
+}
+
+// Eval the triangle local orientation
+float4 DIHairStrands_TriangleLocalOrientation(in float3 TrianglePosition, in float4 TriangleRotation, in float4 WorldOrientation)
+{
+	return NormalizeQuat(MultiplyQuat(InverseQuat(TriangleRotation), WorldOrientation));
+}
+
+// Eval the triangle world position
+float3 DIHairStrands_TriangleWorldPosition(in float3 TrianglePosition, in float4 TriangleRotation, in float3 LocalPosition)
+{
+	return RotateVectorByQuat(LocalPosition, TriangleRotation) + TrianglePosition;
+}
+
+// Eval the triangle local orientation
+float4 DIHairStrands_TriangleWorldOrientation(in float3 TrianglePosition, in float4 TriangleRotation, in float4 LocalOrientation)
+{
+	return NormalizeQuat(MultiplyQuat(TriangleRotation, LocalOrientation));
+}
+
+float2 ExtractBarycentrics(in uint PackedCoordinates)
+{
+	return float2(f16tof32(PackedCoordinates & 0xFFFF), f16tof32((PackedCoordinates >> 16) & 0xFFFF));
+}
+
diff --git a/Engine/Plugins/Runtime/HairStrands/Shaders/Private/NiagaraInterpolateGroomGuides.usf b/Engine/Plugins/Runtime/HairStrands/Shaders/Private/NiagaraInterpolateGroomGuides.usf
new file mode 100644
--- /dev/null
+++ b/Engine/Plugins/Runtime/HairStrands/Shaders/Private/NiagaraInterpolateGroomGuides.usf
@@ -0,0 +1,143 @@
+﻿// Copyright Epic Games, Inc. All Rights Reserved.
+
+/*=============================================================================
+NiagaraInterpolateGroomGuides.usf 
+=============================================================================*/
+
+#pragma once
+
+#define LOCAL_INTERPOLATION 0
+#define WORLD_INTERPOLATION 1
+
+#include "/Engine/Public/Platform.ush"
+
+#if PERMUTATION_INTERPOLATION == LOCAL_INTERPOLATION
+
+#include "NiagaraGroomInterpolationUtils.ush"
+
+Buffer<float4>		RestTrianglePositionABuffer;
+Buffer<float4>		RestTrianglePositionBBuffer;
+Buffer<float4>		RestTrianglePositionCBuffer;
+Buffer<float4>		DeformedTrianglePositionABuffer;
+Buffer<float4>		DeformedTrianglePositionBBuffer;
+Buffer<float4>		DeformedTrianglePositionCBuffer;
+Buffer<uint>		RootBarycentricCoordinatesBuffer;
+
+#endif
+
+Buffer<uint>		VertexToCurveIndexBuffer;
+RWBuffer<uint4>		DeformedPositionBuffer;
+Buffer<uint4>		RestPositionBuffer;
+Buffer<uint>		CurvesOffsetsBuffer;
+Buffer<float4>		DeformedPositionOffset;
+float3				RestPositionOffset;
+
+int					StrandsSize;
+int					NumPoints;
+
+Buffer<float>		NiagaraFloatBuffer;
+int					NiagaraFloatStride;
+
+int					NodePositionComponent;
+int					RestPositionComponent;
+
+#if PERMUTATION_INTERPOLATION == LOCAL_INTERPOLATION
+
+// Get the deformed triangle
+void DIHairStrands_BuildDeformedTriangle(in float2 ProjectionUV, in int StrandIndex, out float3 OutTrianglePosition, out float4 OutTriangleRotation)
+{
+	DIHairStrands_ComputeProjectionTriangle(ProjectionUV, DeformedTrianglePositionABuffer[StrandIndex].xyz,
+		DeformedTrianglePositionBBuffer[StrandIndex].xyz, 
+		DeformedTrianglePositionCBuffer[StrandIndex].xyz, 
+		float3(0,0,0), OutTrianglePosition, OutTriangleRotation);
+}
+
+// Get the rest triangle
+void DIHairStrands_BuildRestTriangle(in float2 ProjectionUV, in int StrandIndex, out float3 OutTrianglePosition, out float4 OutTriangleRotation)
+{
+	DIHairStrands_ComputeProjectionTriangle(ProjectionUV, RestTrianglePositionABuffer[StrandIndex].xyz,
+		RestTrianglePositionBBuffer[StrandIndex].xyz,
+		RestTrianglePositionCBuffer[StrandIndex].xyz,
+		float3(0,0,0), OutTrianglePosition, OutTriangleRotation);
+}
+
+#endif
+
+[numthreads(THREAD_COUNT, 1, 1)]
+void MainCS(uint3 DispatchThreadId : SV_DispatchThreadID)
+{
+	if (DispatchThreadId.x < NumPoints)
+	{
+		const int StrandIndex = VertexToCurveIndexBuffer[DispatchThreadId.x];
+		
+		const int PointCount = CurvesOffsetsBuffer[StrandIndex + 1] - CurvesOffsetsBuffer[StrandIndex];
+		const int PointIndex = DispatchThreadId.x - CurvesOffsetsBuffer[StrandIndex];
+		
+		const float PointCoord = (float) (PointIndex) / (PointCount - 1.0);
+		const float PointNode = PointCoord * (StrandsSize - 2);
+
+		int NodePrev = (PointIndex == 0) ? 0 : (PointIndex == (PointCount - 1.0)) ? StrandsSize - 3 : floor(PointNode);
+		const int NodeNext = NodePrev + 1;
+
+		const float PointAlpha = PointNode - (float) NodePrev;
+		
+		const int NodePositionXOffset = NodePositionComponent * NiagaraFloatStride + StrandIndex * StrandsSize;
+		const int NodePositionYOffset = NodePositionXOffset + NiagaraFloatStride;
+		const int NodePositionZOffset = NodePositionYOffset + NiagaraFloatStride;
+		
+		const float3 WorldPrevNodePosition = float3(NiagaraFloatBuffer[NodePositionXOffset+NodePrev],
+		                                   NiagaraFloatBuffer[NodePositionYOffset+NodePrev],
+		                                   NiagaraFloatBuffer[NodePositionZOffset+NodePrev]);
+
+		const float3 WorldNextNodePosition = float3(NiagaraFloatBuffer[NodePositionXOffset+NodeNext],
+										   NiagaraFloatBuffer[NodePositionYOffset+NodeNext],
+										   NiagaraFloatBuffer[NodePositionZOffset+NodeNext]);
+
+		const int RestPositionXOffset = RestPositionComponent * NiagaraFloatStride + StrandIndex * StrandsSize;
+		const int RestPositionYOffset = RestPositionXOffset + NiagaraFloatStride;
+		const int RestPositionZOffset = RestPositionYOffset + NiagaraFloatStride;
+		
+		const float3 WorldPrevRestPosition = float3(NiagaraFloatBuffer[RestPositionXOffset+NodePrev],
+										   NiagaraFloatBuffer[RestPositionYOffset+NodePrev],
+										   NiagaraFloatBuffer[RestPositionZOffset+NodePrev]);
+
+		const float3 WorldNextRestPosition = float3(NiagaraFloatBuffer[RestPositionXOffset+NodeNext],
+										   NiagaraFloatBuffer[RestPositionYOffset+NodeNext],
+										   NiagaraFloatBuffer[RestPositionZOffset+NodeNext]);
+
+#if PERMUTATION_INTERPOLATION == LOCAL_INTERPOLATION
+		
+		const float2 ProjectionUV = ExtractBarycentrics(RootBarycentricCoordinatesBuffer[StrandIndex]);
+
+		float3 RestTrianglePosition = float3(0, 0, 0), DeformedTrianglePosition = float3(0, 0, 0);
+		float4 RestTriangleOrientation = QUATERNION_IDENTITY, DeformedTriangleOrientation = QUATERNION_IDENTITY;
+		
+		DIHairStrands_BuildRestTriangle(ProjectionUV, StrandIndex, RestTrianglePosition, RestTriangleOrientation);
+		DIHairStrands_BuildDeformedTriangle(ProjectionUV, StrandIndex, DeformedTrianglePosition, DeformedTriangleOrientation);
+		
+		const float3 LocalPrevRestPosition = DIHairStrands_TriangleLocalPosition(RestTrianglePosition, RestTriangleOrientation, WorldPrevRestPosition);
+		const float3 LocalPrevNodePosition = DIHairStrands_TriangleLocalPosition(DeformedTrianglePosition, DeformedTriangleOrientation, WorldPrevNodePosition);
+			
+		const float3 LocalNextRestPosition = DIHairStrands_TriangleLocalPosition(RestTrianglePosition, RestTriangleOrientation, WorldNextRestPosition);
+		const float3 LocalNextNodePosition = DIHairStrands_TriangleLocalPosition(DeformedTrianglePosition, DeformedTriangleOrientation, WorldNextNodePosition);
+
+		const float3 LocalPointDisplace = (LocalPrevNodePosition - LocalPrevRestPosition)  * (1.0 - PointAlpha) + (LocalNextNodePosition - LocalNextRestPosition) * PointAlpha;
+		
+		const float3 LocalTrianglePosition = DIHairStrands_TriangleLocalPosition(RestTrianglePosition, RestTriangleOrientation, 
+					f16tof32(RestPositionBuffer[DispatchThreadId.x].xyz) + RestPositionOffset) + LocalPointDisplace;
+		
+		const float3 LocalComponentPosition = DIHairStrands_TriangleWorldPosition(DeformedTrianglePosition, DeformedTriangleOrientation, LocalTrianglePosition) - DeformedPositionOffset[0].xyz;
+		DeformedPositionBuffer[DispatchThreadId.x] = uint4(f32tof16(LocalComponentPosition), RestPositionBuffer[DispatchThreadId.x].w);
+
+#endif
+		
+#if PERMUTATION_INTERPOLATION == WORLD_INTERPOLATION
+		
+		const float3 WorldPointDisplace = (WorldPrevNodePosition - WorldPrevRestPosition)  * (1.0 - PointAlpha) + (WorldNextNodePosition - WorldNextRestPosition) * PointAlpha;
+
+		const float3 WorldComponentPosition = f16tof32(RestPositionBuffer[DispatchThreadId.x].xyz) + WorldPointDisplace + RestPositionOffset - DeformedPositionOffset[0].xyz;
+		DeformedPositionBuffer[DispatchThreadId.x] = uint4(f32tof16(WorldComponentPosition), RestPositionBuffer[DispatchThreadId.x].w);
+		
+#endif
+	}
+}
diff --git a/Engine/Plugins/Runtime/HairStrands/Source/HairStrandsCore/Private/Niagara/NiagaraDataInterfaceHairStrands.cpp b/Engine/Plugins/Runtime/HairStrands/Source/HairStrandsCore/Private/Niagara/NiagaraDataInterfaceHairStrands.cpp
--- a/Engine/Plugins/Runtime/HairStrands/Source/HairStrandsCore/Private/Niagara/NiagaraDataInterfaceHairStrands.cpp
+++ b/Engine/Plugins/Runtime/HairStrands/Source/HairStrandsCore/Private/Niagara/NiagaraDataInterfaceHairStrands.cpp
@@ -384,6 +384,37 @@ void FNDIHairStrandsBuffer::InitRHI()
 			RHIUnlockBuffer(CurvesOffsetsBuffer.Buffer);
 			SourceDatas->CurveOffsets.Unlock();
 		}
+		{
+			// First retrieve the curves offsets from the bulk datas
+			TArray<FHairStrandsRootIndexFormat::Type> CurvesOffsets;
+			const uint32 InDataSize = SourceDatas->CurveOffsets.GetBulkDataSize();
+			const uint32 ElementCount = InDataSize / sizeof(FHairStrandsRootIndexFormat::BulkType);
+			CurvesOffsets.SetNum(ElementCount);
+
+			const uint8* InData = (const uint8*)SourceDatas->CurveOffsets.LockReadOnly();
+			FMemory::Memcpy(CurvesOffsets.GetData(), InData, InDataSize);
+			SourceDatas->CurveOffsets.Unlock();
+			
+			TArray<uint32> PointsCurve;
+			PointsCurve.SetNum(SourceDatas->GetNumPoints());
+
+			// Build the curve index array for each points using the curve offsets array
+			for (uint32 CurveIndex = 0; CurveIndex < SourceDatas->GetNumCurves(); ++CurveIndex)
+			{
+				for (uint32 PointIndex = CurvesOffsets[CurveIndex]; PointIndex < CurvesOffsets[CurveIndex+1]; ++PointIndex)
+				{
+					PointsCurve[PointIndex] = CurveIndex;
+				}
+			}
+			const uint32 OffsetCount = SourceDatas->GetNumPoints();
+			const uint32 OffsetBytes = sizeof(uint32)*OffsetCount;
+
+			// Copy the points buffer data to gpu for interpolation
+			PointsCurveBuffer.Initialize(TEXT("PointsCurveBuffer"), sizeof(uint32), OffsetCount, EPixelFormat::PF_R32_UINT, BUF_Static);
+			void* OffsetBufferData = RHILockBuffer(PointsCurveBuffer.Buffer, 0, OffsetBytes, RLM_WriteOnly);
+			FMemory::Memcpy(OffsetBufferData, PointsCurve.GetData(), OffsetBytes);
+			RHIUnlockBuffer(PointsCurveBuffer.Buffer);
+		}
 		{
 			static const TArray<uint32> ZeroData = { UINT_MAX,UINT_MAX,UINT_MAX,0,0,0,
 													 UINT_MAX,UINT_MAX,UINT_MAX,0,0,0,
@@ -1208,6 +1239,201 @@ bool UNiagaraDataInterfaceHairStrands::PerInstanceTick(void* PerInstanceData, FN
 	return false;
 }
 
+//------------------------------------------------------------------------------------------------------------
+
+#define NIAGARA_HAIR_STRANDS_THREAD_COUNT_INTERPOLATE 32
+
+/** Compute shader to interpolate the groom position from the sim cache */
+class FInterpolateGroomGuidesCS : public FGlobalShader
+{
+	DECLARE_SHADER_TYPE(FInterpolateGroomGuidesCS, Global)
+	
+	class FInterpolationType : SHADER_PERMUTATION_INT("PERMUTATION_INTERPOLATION", 2);
+	using FPermutationDomain = TShaderPermutationDomain<FInterpolationType>;
+
+public:
+	static bool ShouldCompilePermutation(const FGlobalShaderPermutationParameters& Parameters) { return IsHairStrandsSupported(EHairStrandsShaderType::All, Parameters.Platform); }
+
+	static void ModifyCompilationEnvironment(const FGlobalShaderPermutationParameters& Parameters, FShaderCompilerEnvironment& OutEnvironment)
+	{
+		FGlobalShader::ModifyCompilationEnvironment(Parameters, OutEnvironment);
+		OutEnvironment.SetDefine(TEXT("THREAD_COUNT"), NIAGARA_HAIR_STRANDS_THREAD_COUNT_INTERPOLATE);
+	}
+
+	FInterpolateGroomGuidesCS(const ShaderMetaType::CompiledShaderInitializerType& Initializer)
+		: FGlobalShader(Initializer)
+	{
+		RestTrianglePositionABuffer.Bind(Initializer.ParameterMap, TEXT("RestTrianglePositionABuffer"));
+		RestTrianglePositionBBuffer.Bind(Initializer.ParameterMap, TEXT("RestTrianglePositionBBuffer"));
+		RestTrianglePositionCBuffer.Bind(Initializer.ParameterMap, TEXT("RestTrianglePositionCBuffer"));
+		DeformedTrianglePositionABuffer.Bind(Initializer.ParameterMap, TEXT("DeformedTrianglePositionABuffer"));
+		DeformedTrianglePositionBBuffer.Bind(Initializer.ParameterMap, TEXT("DeformedTrianglePositionBBuffer"));
+		DeformedTrianglePositionCBuffer.Bind(Initializer.ParameterMap, TEXT("DeformedTrianglePositionCBuffer"));
+		RootBarycentricCoordinatesBuffer.Bind(Initializer.ParameterMap, TEXT("RootBarycentricCoordinatesBuffer"));
+		
+		VertexToCurveIndexBuffer.Bind(Initializer.ParameterMap, TEXT("VertexToCurveIndexBuffer"));
+		CurvesOffsetsBuffer.Bind(Initializer.ParameterMap, TEXT("CurvesOffsetsBuffer"));
+		
+		DeformedPositionOffset.Bind(Initializer.ParameterMap, TEXT("DeformedPositionOffset"));
+		RestPositionOffset.Bind(Initializer.ParameterMap, TEXT("RestPositionOffset"));
+		DeformedPositionBuffer.Bind(Initializer.ParameterMap, TEXT("DeformedPositionBuffer"));
+		RestPositionBuffer.Bind(Initializer.ParameterMap, TEXT("RestPositionBuffer"));
+		
+		StrandsSize.Bind(Initializer.ParameterMap, TEXT("StrandsSize"));
+		NumPoints.Bind(Initializer.ParameterMap, TEXT("NumPoints"));
+		
+		NiagaraFloatBuffer.Bind(Initializer.ParameterMap, TEXT("NiagaraFloatBuffer"));
+		NiagaraFloatStride.Bind(Initializer.ParameterMap, TEXT("NiagaraFloatStride"));
+		NodePositionComponent.Bind(Initializer.ParameterMap, TEXT("NodePositionComponent"));
+		RestPositionComponent.Bind(Initializer.ParameterMap, TEXT("RestPositionComponent"));
+	}
+
+	FInterpolateGroomGuidesCS()
+	{}
+
+	LAYOUT_FIELD(FShaderResourceParameter, RestTrianglePositionABuffer)
+	LAYOUT_FIELD(FShaderResourceParameter, RestTrianglePositionBBuffer)
+	LAYOUT_FIELD(FShaderResourceParameter, RestTrianglePositionCBuffer)
+	LAYOUT_FIELD(FShaderResourceParameter, DeformedTrianglePositionABuffer)
+	LAYOUT_FIELD(FShaderResourceParameter, DeformedTrianglePositionBBuffer)
+	LAYOUT_FIELD(FShaderResourceParameter, DeformedTrianglePositionCBuffer)
+	LAYOUT_FIELD(FShaderResourceParameter, RootBarycentricCoordinatesBuffer)
+	
+	LAYOUT_FIELD(FShaderResourceParameter, VertexToCurveIndexBuffer)
+	LAYOUT_FIELD(FShaderResourceParameter, CurvesOffsetsBuffer)
+	
+	LAYOUT_FIELD(FShaderResourceParameter, DeformedPositionOffset)
+	LAYOUT_FIELD(FShaderParameter,RestPositionOffset)
+	LAYOUT_FIELD(FShaderResourceParameter, DeformedPositionBuffer)
+	LAYOUT_FIELD(FShaderResourceParameter, RestPositionBuffer)
+	
+	LAYOUT_FIELD(FShaderParameter, StrandsSize)
+	LAYOUT_FIELD(FShaderParameter, NumPoints)
+
+	LAYOUT_FIELD(FShaderResourceParameter, NiagaraFloatBuffer)
+	LAYOUT_FIELD(FShaderParameter, NiagaraFloatStride)
+	LAYOUT_FIELD(FShaderParameter, NodePositionComponent)
+	LAYOUT_FIELD(FShaderParameter, RestPositionComponent)
+};
+
+IMPLEMENT_SHADER_TYPE(, FInterpolateGroomGuidesCS, TEXT("/Plugin/Runtime/HairStrands/Private/NiagaraInterpolateGroomGuides.usf"), TEXT("MainCS"), SF_Compute);
+
+//------------------------------------------------------------------------------------------------------------
+
+DECLARE_GPU_STAT(InterpolateGroomGuidesCS);
+
+static void InterpolateGroomGuides(FRHICommandListImmediate& RHICmdList, FNiagaraDataBuffer* ParticlesBuffer,
+	const uint32 NodePositionComponent,  const uint32 RestPositionComponent, FNDIHairStrandsBuffer* HairStrandsBuffer,
+	const uint32 StrandsSize, const bool bHasSkinningBinding, const bool bHasValidGeometry)
+{
+	const uint32 NumPoints = HairStrandsBuffer->SourceRestResources->GetVertexCount();
+
+	const bool bIsHairValid = HairStrandsBuffer && HairStrandsBuffer->IsInitialized();
+	const bool bIsRestValid = bIsHairValid && HairStrandsBuffer->SourceRestResources && HairStrandsBuffer->SourceRestResources->IsInitialized();
+	const bool bIsDeformedValid = bIsHairValid && HairStrandsBuffer->SourceDeformedResources && HairStrandsBuffer->SourceDeformedResources->IsInitialized();
+
+	if(bIsRestValid && bIsDeformedValid && bHasValidGeometry)
+	{
+		SCOPED_GPU_STAT(RHICmdList, InterpolateGroomGuidesCS);
+		SCOPED_DRAW_EVENT(RHICmdList, InterpolateGroomGuidesCS);
+		
+		bool bIsRootValid = bIsHairValid && HairStrandsBuffer->SourceDeformedRootResources && HairStrandsBuffer->SourceDeformedRootResources->IsInitialized() &&
+											HairStrandsBuffer->SourceRestRootResources && HairStrandsBuffer->SourceRestRootResources->IsInitialized() && bHasSkinningBinding;;
+
+		FInterpolateGroomGuidesCS::FPermutationDomain InterpolationDomain;
+		InterpolationDomain.Set<FInterpolateGroomGuidesCS::FInterpolationType>(!bIsRootValid);
+		
+		TShaderMapRef<FInterpolateGroomGuidesCS> ComputeShader(GetGlobalShaderMap(GMaxRHIFeatureLevel), InterpolationDomain);
+		SetComputePipelineState(RHICmdList, ComputeShader.GetComputeShader());
+		
+		FRHIComputeShader* ShaderRHI = RHICmdList.GetBoundComputeShader();
+		
+		const int32 MeshLODIndex = bIsRootValid ? HairStrandsBuffer->SourceDeformedRootResources->MeshLODIndex : -1;
+		bIsRootValid = bIsRootValid && HairStrandsBuffer->SourceDeformedRootResources->IsValid(MeshLODIndex);
+		
+		SetShaderValue(RHICmdList, ShaderRHI, ComputeShader->StrandsSize, StrandsSize);
+		SetShaderValue(RHICmdList, ShaderRHI, ComputeShader->NumPoints, NumPoints);
+
+		SetSRVParameter(RHICmdList, ShaderRHI, ComputeShader->NiagaraFloatBuffer, ParticlesBuffer->GetGPUBufferFloat().SRV);
+		SetShaderValue(RHICmdList, ShaderRHI, ComputeShader->NiagaraFloatStride, ParticlesBuffer->GetFloatStride() / sizeof(float));
+		SetShaderValue(RHICmdList, ShaderRHI, ComputeShader->NodePositionComponent, NodePositionComponent);
+		SetShaderValue(RHICmdList, ShaderRHI, ComputeShader->RestPositionComponent, RestPositionComponent);
+
+		const int32 BufferIndex = HairStrandsBuffer->SourceDeformedResources->GetIndex(FHairStrandsDeformedResource::Current);
+
+		SetSRVParameter(RHICmdList, ShaderRHI, ComputeShader->DeformedPositionOffset, HairStrandsBuffer->SourceDeformedResources->DeformedOffsetBuffer[BufferIndex].SRV);
+        SetShaderValue(RHICmdList, ShaderRHI, ComputeShader->RestPositionOffset, FVector3f(HairStrandsBuffer->SourceRestResources->GetPositionOffset()));
+        SetUAVParameter(RHICmdList, ShaderRHI, ComputeShader->DeformedPositionBuffer, HairStrandsBuffer->SourceDeformedResources->DeformedPositionBuffer[BufferIndex].UAV);
+        SetSRVParameter(RHICmdList, ShaderRHI, ComputeShader->RestPositionBuffer, HairStrandsBuffer->SourceRestResources->PositionBuffer.SRV);
+	
+		SetSRVParameter(RHICmdList, ShaderRHI, ComputeShader->VertexToCurveIndexBuffer, HairStrandsBuffer->PointsCurveBuffer.SRV);
+		SetSRVParameter(RHICmdList, ShaderRHI, ComputeShader->CurvesOffsetsBuffer, HairStrandsBuffer->CurvesOffsetsBuffer.SRV);
+		
+		if(bIsRootValid)
+		{
+			const FHairStrandsRestRootResource::FLOD* RestMeshProjection = &(HairStrandsBuffer->SourceRestRootResources->LODs[MeshLODIndex]) ;
+			const FHairStrandsDeformedRootResource::FLOD* DeformedMeshProjection =  &(HairStrandsBuffer->SourceDeformedRootResources->LODs[MeshLODIndex]);
+
+			SetSRVParameter(RHICmdList, ShaderRHI, ComputeShader->RestTrianglePositionABuffer, RestMeshProjection->RestRootTrianglePosition0Buffer.SRV);
+			SetSRVParameter(RHICmdList, ShaderRHI, ComputeShader->RestTrianglePositionBBuffer, RestMeshProjection->RestRootTrianglePosition1Buffer.SRV);
+			SetSRVParameter(RHICmdList, ShaderRHI, ComputeShader->RestTrianglePositionCBuffer, RestMeshProjection->RestRootTrianglePosition2Buffer.SRV);
+			SetSRVParameter(RHICmdList, ShaderRHI, ComputeShader->DeformedTrianglePositionABuffer, DeformedMeshProjection->DeformedRootTrianglePosition0Buffer.SRV);
+			SetSRVParameter(RHICmdList, ShaderRHI, ComputeShader->DeformedTrianglePositionBBuffer, DeformedMeshProjection->DeformedRootTrianglePosition1Buffer.SRV);
+			SetSRVParameter(RHICmdList, ShaderRHI, ComputeShader->DeformedTrianglePositionCBuffer, DeformedMeshProjection->DeformedRootTrianglePosition2Buffer.SRV);
+			SetSRVParameter(RHICmdList, ShaderRHI, ComputeShader->RootBarycentricCoordinatesBuffer, RestMeshProjection->RootTriangleBarycentricBuffer.SRV);
+		}
+		const uint32 GroupSize = NIAGARA_HAIR_STRANDS_THREAD_COUNT_INTERPOLATE;
+		const uint32 DispatchCount = FMath::DivideAndRoundUp(NumPoints, GroupSize);
+		
+		RHICmdList.Transition(FRHITransitionInfo(HairStrandsBuffer->SourceDeformedResources->DeformedPositionBuffer[BufferIndex].UAV,
+			ERHIAccess::SRVMask, ERHIAccess::UAVCompute));
+		
+		DispatchComputeShader(RHICmdList, ComputeShader.GetShader(), DispatchCount, 1, 1);
+		
+		RHICmdList.Transition(FRHITransitionInfo(HairStrandsBuffer->SourceDeformedResources->DeformedPositionBuffer[BufferIndex].UAV,
+			ERHIAccess::UAVCompute, ERHIAccess::SRVMask));
+	}
+}
+
+void UNiagaraDataInterfaceHairStrands::SimCachePostReadFrame(void* OptionalPerInstanceData, FNiagaraSystemInstance* SystemInstance)
+{
+	for (TSharedRef<FNiagaraEmitterInstance, ESPMode::ThreadSafe> EmitterInst : SystemInstance->GetEmitters())
+	{
+		FNiagaraDataSet& EmitterDatas  = EmitterInst->GetData();
+
+		const FNiagaraVariableBase NodePositionVariable(FNiagaraDataSetAccessorTypeInfo<FVector3f>::GetFloatType(), FName("NodePosition"));
+		const int32 NodePositionIndex = EmitterDatas.GetCompiledData().Variables.IndexOfByKey(NodePositionVariable);
+
+		const FNiagaraVariableBase RestPositionVariable(FNiagaraDataSetAccessorTypeInfo<FVector3f>::GetFloatType(), FName("RestPosition"));
+		const int32 RestPositionIndex = EmitterDatas.GetCompiledData().Variables.IndexOfByKey(RestPositionVariable);
+				
+		if (RestPositionIndex != INDEX_NONE && NodePositionIndex != INDEX_NONE)
+		{
+			const FNDIHairStrandsData* InstanceData = static_cast<FNDIHairStrandsData*>(OptionalPerInstanceData);
+			if( EmitterDatas.GetCompiledData().VariableLayouts.IsValidIndex(NodePositionIndex) &&
+				EmitterDatas.GetCompiledData().VariableLayouts.IsValidIndex(RestPositionIndex) )
+			{
+				const uint32 NodePositionComponent = EmitterDatas.GetCompiledData().VariableLayouts[NodePositionIndex].FloatComponentStart;
+				const uint32 RestPositionComponent = EmitterDatas.GetCompiledData().VariableLayouts[RestPositionIndex].FloatComponentStart;
+				
+				FNiagaraDataBuffer* ParticlesBuffer = EmitterDatas.GetCurrentData();
+				FNDIHairStrandsBuffer* HairstrandsBuffer = InstanceData->HairStrandsBuffer;
+
+				const uint32 StrandsSize = InstanceData->StrandsSize;
+				const bool bHasSkinningBinding = InstanceData->HairGroupInstance->BindingType == EHairBindingType::Skinning;
+				const bool bHasValidGeometry = InstanceData->HairGroupInstance->GeometryType != EHairGeometryType::NoneGeometry;
+				
+				ENQUEUE_RENDER_COMMAND(NiagaraInterpolateGroomSimCache)(
+				[ParticlesBuffer, HairstrandsBuffer, NodePositionComponent, RestPositionComponent, StrandsSize, bHasSkinningBinding, bHasValidGeometry](FRHICommandListImmediate& RHICmdList)
+				{
+					InterpolateGroomGuides(RHICmdList, ParticlesBuffer,
+							NodePositionComponent, RestPositionComponent, HairstrandsBuffer, StrandsSize, bHasSkinningBinding, bHasValidGeometry);
+				});
+			}
+		}
+	}
+}
+
 bool UNiagaraDataInterfaceHairStrands::CopyToInternal(UNiagaraDataInterface* Destination) const
 {
 	if (!Super::CopyToInternal(Destination))
diff --git a/Engine/Plugins/Runtime/HairStrands/Source/HairStrandsCore/Public/GroomComponent.h b/Engine/Plugins/Runtime/HairStrands/Source/HairStrandsCore/Public/GroomComponent.h
--- a/Engine/Plugins/Runtime/HairStrands/Source/HairStrandsCore/Public/GroomComponent.h
+++ b/Engine/Plugins/Runtime/HairStrands/Source/HairStrandsCore/Public/GroomComponent.h
@@ -187,6 +187,10 @@ public:
 	/* Reset the simulation, if enabled */
 	UFUNCTION(BlueprintCallable, Category = "Simulation")
 	void ResetSimulation();
+	
+	/* Given the group index return the matching niagara component */
+	UFUNCTION(BlueprintCallable, Category = "Simulation")
+	UNiagaraComponent* GetNiagaraComponent(const int32 GroupIndex) {return NiagaraComponents.IsValidIndex(GroupIndex) ? NiagaraComponents[GroupIndex] : nullptr;}
 
 	/* Accessor function for changing hair length scale from blueprint/sequencer */
 	UFUNCTION(BlueprintCallable, Category = "Groom")
diff --git a/Engine/Plugins/Runtime/HairStrands/Source/HairStrandsCore/Public/Niagara/NiagaraDataInterfaceHairStrands.h b/Engine/Plugins/Runtime/HairStrands/Source/HairStrandsCore/Public/Niagara/NiagaraDataInterfaceHairStrands.h
--- a/Engine/Plugins/Runtime/HairStrands/Source/HairStrandsCore/Public/Niagara/NiagaraDataInterfaceHairStrands.h
+++ b/Engine/Plugins/Runtime/HairStrands/Source/HairStrandsCore/Public/Niagara/NiagaraDataInterfaceHairStrands.h
@@ -61,6 +61,9 @@ struct FNDIHairStrandsBuffer : public FRenderResource
 	/** Params scale buffer */
 	FRWBuffer ParamsScaleBuffer;
 
+	/** Points curve index for fast query */
+	FRWBuffer PointsCurveBuffer;
+
 	/** The strand asset resource from which to sample */
 	const FHairStrandsRestResource* SourceRestResources;
 
@@ -411,6 +414,7 @@ public:
 	virtual bool HasPreSimulateTick() const override { return true; }
 	virtual bool HasTickGroupPrereqs() const override { return true; }
 	virtual ETickingGroup CalculateTickGroup(const void* PerInstanceData) const override;
+	virtual void SimCachePostReadFrame(void* OptionalPerInstanceData, FNiagaraSystemInstance* SystemInstance) override;
 
 	/** GPU simulation  functionality */
 #if WITH_EDITORONLY_DATA
