Fix Min/Max alpha blend state settings (d3d debug layer warning):
ID3D12Device::CreateBlendState: When BlendOpAlpha is MIN or MAX, SrcBlendAlpha and DestBlendAlpha are ignored, and treated as if they were D3D12_BLEND_ONE.

diff --git a/Engine/Plugins/Runtime/HairStrands/Source/HairStrandsCore/Private/HairStrandsMeshProjection.cpp b/Engine/Plugins/Runtime/HairStrands/Source/HairStrandsCore/Private/HairStrandsMeshProjection.cpp
--- a/Engine/Plugins/Runtime/HairStrands/Source/HairStrandsCore/Private/HairStrandsMeshProjection.cpp
+++ b/Engine/Plugins/Runtime/HairStrands/Source/HairStrandsCore/Private/HairStrandsMeshProjection.cpp
@@ -1594,7 +1594,7 @@ static void AddFollicleMaskPass(
 		// Apply additive blending pipeline state.
 		FGraphicsPipelineStateInitializer GraphicsPSOInit;
 		RHICmdList.ApplyCachedRenderTargets(GraphicsPSOInit);
-		GraphicsPSOInit.BlendState = TStaticBlendState<CW_RGBA, BO_Max, BF_SourceColor, BF_DestColor, BO_Max, BF_SourceAlpha, BF_DestAlpha>::GetRHI();
+		GraphicsPSOInit.BlendState = TStaticBlendState<CW_RGBA, BO_Max, BF_SourceColor, BF_DestColor, BO_Max, BF_One, BF_One>::GetRHI();
 		GraphicsPSOInit.RasterizerState = TStaticRasterizerState<>::GetRHI();
 		GraphicsPSOInit.DepthStencilState = TStaticDepthStencilState<false, CF_Always>::GetRHI();
 		GraphicsPSOInit.BoundShaderState.VertexDeclarationRHI = GEmptyVertexDeclaration.VertexDeclarationRHI;
@@ -1656,7 +1656,7 @@ static void AddFollicleMaskPass(
 			// Apply additive blending pipeline state.
 			FGraphicsPipelineStateInitializer GraphicsPSOInit;
 			RHICmdList.ApplyCachedRenderTargets(GraphicsPSOInit);
-			GraphicsPSOInit.BlendState = TStaticBlendState<CW_RGBA, BO_Max, BF_SourceColor, BF_DestColor, BO_Max, BF_SourceAlpha, BF_DestAlpha>::GetRHI();
+			GraphicsPSOInit.BlendState = TStaticBlendState<CW_RGBA, BO_Max, BF_SourceColor, BF_DestColor, BO_Max, BF_One, BF_One>::GetRHI();
 			GraphicsPSOInit.RasterizerState = TStaticRasterizerState<>::GetRHI();
 			GraphicsPSOInit.DepthStencilState = TStaticDepthStencilState<false, CF_Always>::GetRHI();
 			GraphicsPSOInit.BoundShaderState.VertexDeclarationRHI = GEmptyVertexDeclaration.VertexDeclarationRHI;
diff --git a/Engine/Source/Runtime/Renderer/Private/HairStrands/HairStrandsComposition.cpp b/Engine/Source/Runtime/Renderer/Private/HairStrands/HairStrandsComposition.cpp
--- a/Engine/Source/Runtime/Renderer/Private/HairStrands/HairStrandsComposition.cpp
+++ b/Engine/Source/Runtime/Renderer/Private/HairStrands/HairStrandsComposition.cpp
@@ -90,7 +90,7 @@ void InternalCommonDrawPass(
 		}
 		else
 		{
-			GraphicsPSOInit.BlendState = TStaticBlendState<CW_RGBA, BO_Add, BF_One, BF_Zero, BO_Max, BF_One, BF_Zero>::GetRHI();
+			GraphicsPSOInit.BlendState = TStaticBlendState<CW_RGBA, BO_Add, BF_One, BF_Zero, BO_Max, BF_One, BF_One>::GetRHI();
 		}
 
 		GraphicsPSOInit.RasterizerState = TStaticRasterizerState<>::GetRHI();
diff --git a/Engine/Source/Runtime/Renderer/Private/HairStrands/HairStrandsDeepShadow.cpp b/Engine/Source/Runtime/Renderer/Private/HairStrands/HairStrandsDeepShadow.cpp
--- a/Engine/Source/Runtime/Renderer/Private/HairStrands/HairStrandsDeepShadow.cpp
+++ b/Engine/Source/Runtime/Renderer/Private/HairStrands/HairStrandsDeepShadow.cpp
@@ -123,7 +123,7 @@ void AddInjectHairVoxelShadowCaster(
 			// Apply additive blending pipeline state.
 			FGraphicsPipelineStateInitializer GraphicsPSOInit;
 			RHICmdList.ApplyCachedRenderTargets(GraphicsPSOInit);
-			GraphicsPSOInit.BlendState = TStaticBlendState<CW_RGBA, BO_Max, BF_SourceColor, BF_DestColor, BO_Max, BF_SourceAlpha, BF_DestAlpha>::GetRHI();
+			GraphicsPSOInit.BlendState = TStaticBlendState<CW_RGBA, BO_Max, BF_SourceColor, BF_DestColor, BO_Max, BF_One, BF_One>::GetRHI();
 			GraphicsPSOInit.RasterizerState = TStaticRasterizerState<>::GetRHI();
 			GraphicsPSOInit.DepthStencilState = TStaticDepthStencilState<true, CF_Greater>::GetRHI();
 			GraphicsPSOInit.BoundShaderState.VertexDeclarationRHI = GEmptyVertexDeclaration.VertexDeclarationRHI;
diff --git a/Engine/Source/Runtime/Renderer/Private/HairStrands/HairStrandsEnvironment.cpp b/Engine/Source/Runtime/Renderer/Private/HairStrands/HairStrandsEnvironment.cpp
--- a/Engine/Source/Runtime/Renderer/Private/HairStrands/HairStrandsEnvironment.cpp
+++ b/Engine/Source/Runtime/Renderer/Private/HairStrands/HairStrandsEnvironment.cpp
@@ -435,7 +435,7 @@ static void AddHairStrandsEnvironmentLightingPassPS(
 
 		FGraphicsPipelineStateInitializer GraphicsPSOInit;
 		RHICmdList.ApplyCachedRenderTargets(GraphicsPSOInit);
-		GraphicsPSOInit.BlendState = TStaticBlendState<CW_RGBA, BO_Add, BF_One, BF_One, BO_Max, BF_SourceAlpha, BF_DestAlpha>::GetRHI();
+		GraphicsPSOInit.BlendState = TStaticBlendState<CW_RGBA, BO_Add, BF_One, BF_One, BO_Max, BF_One, BF_One>::GetRHI();
 		GraphicsPSOInit.RasterizerState = TStaticRasterizerState<>::GetRHI();
 		GraphicsPSOInit.DepthStencilState = TStaticDepthStencilState<false, CF_Always>::GetRHI();
 		GraphicsPSOInit.BoundShaderState.VertexDeclarationRHI = GFilterVertexDeclaration.VertexDeclarationRHI;
diff --git a/Engine/Source/Runtime/Renderer/Private/HairStrands/HairStrandsTile.cpp b/Engine/Source/Runtime/Renderer/Private/HairStrands/HairStrandsTile.cpp
--- a/Engine/Source/Runtime/Renderer/Private/HairStrands/HairStrandsTile.cpp
+++ b/Engine/Source/Runtime/Renderer/Private/HairStrands/HairStrandsTile.cpp
@@ -333,7 +333,7 @@ void AddHairStrandsDebugTilePass(
 
 				FGraphicsPipelineStateInitializer GraphicsPSOInit;
 				RHICmdList.ApplyCachedRenderTargets(GraphicsPSOInit);
-				GraphicsPSOInit.BlendState = TStaticBlendState<CW_RGBA, BO_Add, BF_One, BF_One, BO_Max, BF_SourceAlpha, BF_DestAlpha>::GetRHI();
+				GraphicsPSOInit.BlendState = TStaticBlendState<CW_RGBA, BO_Add, BF_One, BF_One, BO_Max, BF_One, BF_One>::GetRHI();
 				GraphicsPSOInit.RasterizerState = TStaticRasterizerState<>::GetRHI();
 				GraphicsPSOInit.DepthStencilState = TStaticDepthStencilState<false, CF_Always>::GetRHI();
 
diff --git a/Engine/Source/Runtime/Renderer/Private/LightRendering.cpp b/Engine/Source/Runtime/Renderer/Private/LightRendering.cpp
--- a/Engine/Source/Runtime/Renderer/Private/LightRendering.cpp
+++ b/Engine/Source/Runtime/Renderer/Private/LightRendering.cpp
@@ -2351,7 +2351,7 @@ void FDeferredShadingSceneRenderer::RenderLightForHair(
 
 		FGraphicsPipelineStateInitializer GraphicsPSOInit;
 		RHICmdList.ApplyCachedRenderTargets(GraphicsPSOInit);
-		GraphicsPSOInit.BlendState = TStaticBlendState<CW_RGBA, BO_Add, BF_One, BF_One, BO_Max, BF_SourceAlpha, BF_DestAlpha>::GetRHI();
+		GraphicsPSOInit.BlendState = TStaticBlendState<CW_RGBA, BO_Add, BF_One, BF_One, BO_Max, BF_One, BF_One>::GetRHI();
 		GraphicsPSOInit.PrimitiveType = PT_TriangleList;
 
 		GraphicsPSOInit.BoundShaderState.VertexDeclarationRHI = GetVertexDeclarationFVector4();
