Fix warning V595: The 'HairStrandsBuffer' pointer was utilized before it was verified against nullptr.
#rnx

diff --git a/Engine/Plugins/Runtime/HairStrands/Source/HairStrandsCore/Private/Niagara/NiagaraDataInterfaceHairStrands.cpp b/Engine/Plugins/Runtime/HairStrands/Source/HairStrandsCore/Private/Niagara/NiagaraDataInterfaceHairStrands.cpp
--- a/Engine/Plugins/Runtime/HairStrands/Source/HairStrandsCore/Private/Niagara/NiagaraDataInterfaceHairStrands.cpp
+++ b/Engine/Plugins/Runtime/HairStrands/Source/HairStrandsCore/Private/Niagara/NiagaraDataInterfaceHairStrands.cpp
@@ -1326,8 +1326,6 @@ static void InterpolateGroomGuides(FRHICommandListImmediate& RHICmdList, FNiagar
 	const uint32 NodePositionComponent,  const uint32 RestPositionComponent, FNDIHairStrandsBuffer* HairStrandsBuffer,
 	const uint32 StrandsSize, const bool bHasSkinningBinding, const bool bHasValidGeometry)
 {
-	const uint32 NumPoints = HairStrandsBuffer->SourceRestResources->GetVertexCount();
-
 	const bool bIsHairValid = HairStrandsBuffer && HairStrandsBuffer->IsInitialized();
 	const bool bIsRestValid = bIsHairValid && HairStrandsBuffer->SourceRestResources && HairStrandsBuffer->SourceRestResources->IsInitialized();
 	const bool bIsDeformedValid = bIsHairValid && HairStrandsBuffer->SourceDeformedResources && HairStrandsBuffer->SourceDeformedResources->IsInitialized();
@@ -1337,8 +1335,10 @@ static void InterpolateGroomGuides(FRHICommandListImmediate& RHICmdList, FNiagar
 		SCOPED_GPU_STAT(RHICmdList, InterpolateGroomGuidesCS);
 		SCOPED_DRAW_EVENT(RHICmdList, InterpolateGroomGuidesCS);
 		
-		bool bIsRootValid = bIsHairValid && HairStrandsBuffer->SourceDeformedRootResources && HairStrandsBuffer->SourceDeformedRootResources->IsInitialized() &&
-											HairStrandsBuffer->SourceRestRootResources && HairStrandsBuffer->SourceRestRootResources->IsInitialized() && bHasSkinningBinding;;
+		bool bIsRootValid = HairStrandsBuffer->SourceDeformedRootResources && HairStrandsBuffer->SourceDeformedRootResources->IsInitialized() &&
+								HairStrandsBuffer->SourceRestRootResources && HairStrandsBuffer->SourceRestRootResources->IsInitialized() && bHasSkinningBinding;;
+
+		const uint32 NumPoints = HairStrandsBuffer->SourceRestResources->GetVertexCount();
 
 		FInterpolateGroomGuidesCS::FPermutationDomain InterpolationDomain;
 		InterpolationDomain.Set<FInterpolateGroomGuidesCS::FInterpolationType>(!bIsRootValid);
