Groom: Fixed quantization issue with guides groom cache when groom cache has large displacement.

diff --git a/Engine/Plugins/Runtime/HairStrands/Source/HairStrandsCore/Private/HairStrandsRendering.cpp b/Engine/Plugins/Runtime/HairStrands/Source/HairStrandsCore/Private/HairStrandsRendering.cpp
--- a/Engine/Plugins/Runtime/HairStrands/Source/HairStrandsCore/Private/HairStrandsRendering.cpp
+++ b/Engine/Plugins/Runtime/HairStrands/Source/HairStrandsCore/Private/HairStrandsRendering.cpp
@@ -1802,6 +1802,16 @@ void ComputeHairStrandsInterpolation(
 						Instance->Guides.DeformedRootResource,
 						Instance->Guides.DeformedResource);
 
+					// Apply the same offset to the render strands for proper rendering
+					Instance->Strands.DeformedResource->SetPositionOffset(FHairStrandsDeformedResource::EFrameType::Current, GroomCacheGroupData.BoundingBox.GetCenter());
+
+					AddHairStrandUpdatePositionOffsetPass(
+						GraphBuilder,
+						ShaderMap,
+						MeshLODIndex,
+						Instance->Strands.DeformedRootResource,
+						Instance->Strands.DeformedResource);
+
 					// Pass to upload GroomCache guide positions
 					AddGroomCacheUpdatePass(
 						GraphBuilder,
