Groom: Fixed crash when applying groom binding on geometry cache.

diff --git a/Engine/Plugins/Runtime/HairStrands/Source/HairStrandsCore/Private/GroomGeometryCache.cpp b/Engine/Plugins/Runtime/HairStrands/Source/HairStrandsCore/Private/GroomGeometryCache.cpp
--- a/Engine/Plugins/Runtime/HairStrands/Source/HairStrandsCore/Private/GroomGeometryCache.cpp
+++ b/Engine/Plugins/Runtime/HairStrands/Source/HairStrandsCore/Private/GroomGeometryCache.cpp
@@ -152,12 +152,7 @@ static void BuildBoneMatrices(USkeletalMeshComponent* SkeletalMeshComponent, con
 				CachedSection.PositionBuffer = TrackProxy->PositionBuffers[PositionIndex].GetBufferSRV();
 				CachedSection.UVsBuffer = TrackProxy->TextureCoordinatesBuffer.GetBufferSRV();
 				CachedSection.TotalVertexCount = TrackProxy->MeshData->Positions.Num();
-				// It is unclear why the existing SRV is invalid. When the index buffer is changed/recreated its SRV is supposed to be update.
-				//CachedSection.IndexBuffer = TrackProxy->IndexBuffer.GetBufferSRV();
-				if (TrackProxy->IndexBuffer.IsInitialized())
-				{
-					CachedSection.IndexBuffer = RHICreateShaderResourceView(TrackProxy->IndexBuffer.IndexBufferRHI);
-				}
+				CachedSection.IndexBuffer = TrackProxy->IndexBuffer.GetBufferSRV();
 				CachedSection.TotalIndexCount = TrackProxy->IndexBuffer.NumValidIndices;
 				CachedSection.UVsChannelCount = 1;
 				CachedSection.NumPrimitives = BatchInfo.NumTriangles;
