Groom: Fixed missing argument from some log messages.

diff --git a/Engine/Plugins/Runtime/HairStrands/Source/HairStrandsCore/Private/GroomAsset.cpp b/Engine/Plugins/Runtime/HairStrands/Source/HairStrandsCore/Private/GroomAsset.cpp
--- a/Engine/Plugins/Runtime/HairStrands/Source/HairStrandsCore/Private/GroomAsset.cpp
+++ b/Engine/Plugins/Runtime/HairStrands/Source/HairStrandsCore/Private/GroomAsset.cpp
@@ -740,7 +740,7 @@ void UGroomAsset::Serialize(FArchive& Ar)
 	else
 	{
 		// Old format serialized the computed groom data directly, but are no longer supported
-		UE_LOG(LogHairStrands, Error, TEXT("[Groom] The groom asset %s is too old. Please reimported the groom from its original source file."));
+		UE_LOG(LogHairStrands, Error, TEXT("[Groom] The groom asset %s is too old. Please reimported the groom from its original source file."), *GetName());
 	}
 }
 
@@ -1046,7 +1046,7 @@ void UGroomAsset::PostLoad()
 		else
 		{
 			// Old format serialized the computed groom data directly, but are no longer supported
-			UE_LOG(LogHairStrands, Error, TEXT("[Groom] The groom asset %s is too old. Please reimported the groom from its original source file."));
+			UE_LOG(LogHairStrands, Error, TEXT("[Groom] The groom asset %s is too old. Please reimported the groom from its original source file."), *GetName());
 			SetNumGroup(0, false);
 		}
 	}
@@ -2169,7 +2169,7 @@ bool UGroomAsset::CacheStrandsData(uint32 GroupIndex, FString& OutDerivedDataKey
 {
 	if (!HairDescriptionBulkData[HairDescriptionType])
 	{
-		UE_LOG(LogHairStrands, Error, TEXT("[Groom] The groom asset %s is too old. Please reimported the groom from its original source file."));
+		UE_LOG(LogHairStrands, Error, TEXT("[Groom] The groom asset %s is too old. Please reimported the groom from its original source file."), *GetName());
 		return false;
 	}
 
@@ -2208,7 +2208,7 @@ bool UGroomAsset::CacheStrandsData(uint32 GroupIndex, FString& OutDerivedDataKey
 
 		if (!LocalHairDescriptionGroups.IsValid())
 		{
-			UE_LOG(LogHairStrands, Warning, TEXT("[Groom] The groom asset %s does not have valid hair groups."));
+			UE_LOG(LogHairStrands, Warning, TEXT("[Groom] The groom asset %s does not have valid hair groups."), *GetName());
 			return false;
 		}
 
