Replace direct access to SkeletalMesh object from USkinnedMeshComponent with GetSkeletalMesh function

diff --git a/Engine/Plugins/Animation/ControlRig/Source/ControlRig/Private/AnimNode_ControlRig.cpp b/Engine/Plugins/Animation/ControlRig/Source/ControlRig/Private/AnimNode_ControlRig.cpp
--- a/Engine/Plugins/Animation/ControlRig/Source/ControlRig/Private/AnimNode_ControlRig.cpp
+++ b/Engine/Plugins/Animation/ControlRig/Source/ControlRig/Private/AnimNode_ControlRig.cpp
@@ -280,7 +280,7 @@ void FAnimNode_ControlRig::UpdateControlRigRefPoseIfNeeded(const FAnimInstancePr
 
 	if(InProxy->GetSkelMeshComponent())
 	{
-		ExpectedHash = HashCombine(ExpectedHash, (int32)reinterpret_cast<uintptr_t>(InProxy->GetSkelMeshComponent()->SkeletalMesh.Get()));
+		ExpectedHash = HashCombine(ExpectedHash, (int32)reinterpret_cast<uintptr_t>(InProxy->GetSkelMeshComponent()->GetSkeletalMesh().Get()));
 	}
 
 	if(bIncludePoseInHash)
diff --git a/Engine/Plugins/Animation/ControlRig/Source/ControlRig/Private/AnimNode_ControlRigBase.cpp b/Engine/Plugins/Animation/ControlRig/Source/ControlRig/Private/AnimNode_ControlRigBase.cpp
--- a/Engine/Plugins/Animation/ControlRig/Source/ControlRig/Private/AnimNode_ControlRigBase.cpp
+++ b/Engine/Plugins/Animation/ControlRig/Source/ControlRig/Private/AnimNode_ControlRigBase.cpp
@@ -38,7 +38,7 @@ void FAnimNode_ControlRigBase::OnInitializeAnimInstance(const FAnimInstanceProxy
 
 	USkeletalMeshComponent* Component = InAnimInstance->GetOwningComponent();
 	UControlRig* ControlRig = GetControlRig();
-	if (Component && Component->SkeletalMesh && ControlRig)
+	if (Component && Component->GetSkeletalMesh() && ControlRig)
 	{
 #if WITH_EDITORONLY_DATA
 		UBlueprintGeneratedClass* BlueprintClass = Cast<UBlueprintGeneratedClass>(ControlRig->GetClass());
@@ -46,7 +46,7 @@ void FAnimNode_ControlRigBase::OnInitializeAnimInstance(const FAnimInstanceProxy
 		{
 			UBlueprint* Blueprint = Cast<UBlueprint>(BlueprintClass->ClassGeneratedBy);
 			// node mapping container will be saved on the initialization part
-			NodeMappingContainer = Component->SkeletalMesh->GetNodeMappingContainer(Blueprint);
+			NodeMappingContainer = Component->GetSkeletalMesh()->GetNodeMappingContainer(Blueprint);
 		}
 #endif
 
diff --git a/Engine/Plugins/Animation/ControlRig/Source/ControlRig/Private/ControlRig.cpp b/Engine/Plugins/Animation/ControlRig/Source/ControlRig/Private/ControlRig.cpp
--- a/Engine/Plugins/Animation/ControlRig/Source/ControlRig/Private/ControlRig.cpp
+++ b/Engine/Plugins/Animation/ControlRig/Source/ControlRig/Private/ControlRig.cpp
@@ -2971,8 +2971,8 @@ void UControlRig::SetBoneInitialTransformsFromSkeletalMeshComponent(USkeletalMes
 	}
 	else
 	{
-		SetBoneInitialTransformsFromSkeletalMesh(InSkelMeshComp->SkeletalMesh);
-		}
+		SetBoneInitialTransformsFromSkeletalMesh(InSkelMeshComp->GetSkeletalMesh());
+	}
 }
 
 
diff --git a/Engine/Plugins/Animation/ControlRig/Source/ControlRig/Private/ControlRigComponent.cpp b/Engine/Plugins/Animation/ControlRig/Source/ControlRig/Private/ControlRigComponent.cpp
--- a/Engine/Plugins/Animation/ControlRig/Source/ControlRig/Private/ControlRigComponent.cpp
+++ b/Engine/Plugins/Animation/ControlRig/Source/ControlRig/Private/ControlRigComponent.cpp
@@ -569,7 +569,7 @@ void UControlRigComponent::AddMappedSkeletalMesh(USkeletalMeshComponent* Skeleta
 	TArray<FControlRigComponentMappedBone> BonesToMap = Bones;
 	if (BonesToMap.Num() == 0)
 	{
-		if (const USkeletalMesh* SkeletalMesh = SkeletalMeshComponent->SkeletalMesh)
+		if (const USkeletalMesh* SkeletalMesh = SkeletalMeshComponent->GetSkeletalMesh())
 		{
 			if (const USkeleton* Skeleton = SkeletalMesh->GetSkeleton())
 			{
@@ -595,7 +595,7 @@ void UControlRigComponent::AddMappedSkeletalMesh(USkeletalMeshComponent* Skeleta
 	TArray<FControlRigComponentMappedCurve> CurvesToMap = Curves;
 	if (CurvesToMap.Num() == 0)
 	{
-		if (USkeletalMesh* SkeletalMesh = SkeletalMeshComponent->SkeletalMesh)
+		if (USkeletalMesh* SkeletalMesh = SkeletalMeshComponent->GetSkeletalMesh())
 		{
 			if (USkeleton* Skeleton = SkeletalMesh->GetSkeleton())
 			{
@@ -1278,7 +1278,7 @@ void UControlRigComponent::ValidateMappingData()
 				}
 				else if (!MappedElement.TransformName.IsNone())
 				{
-					if (USkeletalMesh* SkeletalMesh = SkeletalMeshComponent->SkeletalMesh)
+					if (USkeletalMesh* SkeletalMesh = SkeletalMeshComponent->GetSkeletalMesh())
 					{
 						if (USkeleton* Skeleton = SkeletalMesh->GetSkeleton())
 						{
diff --git a/Engine/Plugins/Animation/ControlRig/Source/ControlRig/Private/Rigs/AdditiveControlRig.cpp b/Engine/Plugins/Animation/ControlRig/Source/ControlRig/Private/Rigs/AdditiveControlRig.cpp
--- a/Engine/Plugins/Animation/ControlRig/Source/ControlRig/Private/Rigs/AdditiveControlRig.cpp
+++ b/Engine/Plugins/Animation/ControlRig/Source/ControlRig/Private/Rigs/AdditiveControlRig.cpp
@@ -76,7 +76,7 @@ void UAdditiveControlRig::Initialize(bool bInitRigUnits /*= true*/)
 	// create hierarchy from the incoming skeleton
 	if (USkeletalMeshComponent* SkeletalMeshComponent = Cast<USkeletalMeshComponent>(GetObjectBinding()->GetBoundObject()))
 	{
-		CreateRigElements(SkeletalMeshComponent->SkeletalMesh);
+		CreateRigElements(SkeletalMeshComponent->GetSkeletalMesh());
 	}
 
 	// add units and initialize
diff --git a/Engine/Plugins/Animation/ControlRig/Source/ControlRig/Private/Rigs/FKControlRig.cpp b/Engine/Plugins/Animation/ControlRig/Source/ControlRig/Private/Rigs/FKControlRig.cpp
--- a/Engine/Plugins/Animation/ControlRig/Source/ControlRig/Private/Rigs/FKControlRig.cpp
+++ b/Engine/Plugins/Animation/ControlRig/Source/ControlRig/Private/Rigs/FKControlRig.cpp
@@ -264,7 +264,7 @@ void UFKControlRig::Initialize(bool bInitRigUnits /*= true*/)
 	// create hierarchy from the incoming skeleton
 	if (USkeletalMeshComponent* SkeletalMeshComponent = Cast<USkeletalMeshComponent>(GetObjectBinding()->GetBoundObject()))
 	{
-		CreateRigElements(SkeletalMeshComponent->SkeletalMesh);
+		CreateRigElements(SkeletalMeshComponent->GetSkeletalMesh());
 	}
 	else if (USkeleton* Skeleton = Cast<USkeleton>(GetObjectBinding()->GetBoundObject()))
 	{
diff --git a/Engine/Plugins/Animation/ControlRig/Source/ControlRig/Private/Sequencer/ControlRigLayerInstance.cpp b/Engine/Plugins/Animation/ControlRig/Source/ControlRig/Private/Sequencer/ControlRigLayerInstance.cpp
--- a/Engine/Plugins/Animation/ControlRig/Source/ControlRig/Private/Sequencer/ControlRigLayerInstance.cpp
+++ b/Engine/Plugins/Animation/ControlRig/Source/ControlRig/Private/Sequencer/ControlRigLayerInstance.cpp
@@ -121,7 +121,7 @@ void UControlRigLayerInstance::SavePose()
 {
 	if (USkeletalMeshComponent* SkeletalMeshComponent = GetSkelMeshComponent())
 	{
-		if (SkeletalMeshComponent->SkeletalMesh && SkeletalMeshComponent->GetComponentSpaceTransforms().Num() > 0)
+		if (SkeletalMeshComponent->GetSkeletalMesh() && SkeletalMeshComponent->GetComponentSpaceTransforms().Num() > 0)
 		{
 			SavePoseSnapshot(UControlRigLayerInstance::SequencerPoseName);
 		}
diff --git a/Engine/Plugins/Animation/ControlRig/Source/ControlRig/Private/Sequencer/MovieSceneControlRigParameterSection.cpp b/Engine/Plugins/Animation/ControlRig/Source/ControlRig/Private/Sequencer/MovieSceneControlRigParameterSection.cpp
--- a/Engine/Plugins/Animation/ControlRig/Source/ControlRig/Private/Sequencer/MovieSceneControlRigParameterSection.cpp
+++ b/Engine/Plugins/Animation/ControlRig/Source/ControlRig/Private/Sequencer/MovieSceneControlRigParameterSection.cpp
@@ -2586,12 +2586,12 @@ bool UMovieSceneControlRigParameterSection::LoadAnimSequenceIntoThisSection(UAni
 {
 	USkeletalMeshComponent* SkelMeshComp = Cast<USkeletalMeshComponent>(BoundObject);
 	
-	if (SkelMeshComp != nullptr && (SkelMeshComp->SkeletalMesh == nullptr || SkelMeshComp->SkeletalMesh->GetSkeleton() == nullptr))
+	if (SkelMeshComp != nullptr && (SkelMeshComp->GetSkeletalMesh() == nullptr || SkelMeshComp->GetSkeletalMesh()->GetSkeleton() == nullptr))
 	{
 		return false;
 	}
 	
-	USkeleton* Skeleton = SkelMeshComp ? SkelMeshComp->SkeletalMesh->GetSkeleton() : Cast<USkeleton>(BoundObject);	
+	USkeleton* Skeleton = SkelMeshComp ? SkelMeshComp->GetSkeletalMesh()->GetSkeleton() : Cast<USkeleton>(BoundObject);
 	if (Skeleton == nullptr)
 	{
 		return false;
@@ -2695,7 +2695,7 @@ bool UMovieSceneControlRigParameterSection::LoadAnimSequenceIntoThisSection(UAni
 		// retrieve the pose using the services that persona and sequencer rely on
 		// rather than accessing the low level raw tracks.
 		FAnimPoseEvaluationOptions EvaluationOptions;
-		EvaluationOptions.OptionalSkeletalMesh = SkelMeshComp ? SkelMeshComp->SkeletalMesh : nullptr;
+		EvaluationOptions.OptionalSkeletalMesh = SkelMeshComp ? SkelMeshComp->GetSkeletalMesh() : nullptr;
 		EvaluationOptions.bShouldRetarget = false;
 		EvaluationOptions.EvaluationType = EAnimDataEvalType::Raw;
 
diff --git a/Engine/Plugins/Animation/ControlRig/Source/ControlRig/Private/Sequencer/MovieSceneControlRigParameterTemplate.cpp b/Engine/Plugins/Animation/ControlRig/Source/ControlRig/Private/Sequencer/MovieSceneControlRigParameterTemplate.cpp
--- a/Engine/Plugins/Animation/ControlRig/Source/ControlRig/Private/Sequencer/MovieSceneControlRigParameterTemplate.cpp
+++ b/Engine/Plugins/Animation/ControlRig/Source/ControlRig/Private/Sequencer/MovieSceneControlRigParameterTemplate.cpp
@@ -588,7 +588,7 @@ void FControlRigBindingHelper::BindToSequencerInstance(UControlRig* ControlRig)
 	check(ControlRig);
 	if (USkeletalMeshComponent* SkeletalMeshComponent = Cast<USkeletalMeshComponent>(ControlRig->GetObjectBinding()->GetBoundObject()))
 	{
-		if(SkeletalMeshComponent->SkeletalMesh)
+		if(SkeletalMeshComponent->GetSkeletalMesh())
 		{
 			bool bWasCreated = false;
 			if (UControlRigLayerInstance* AnimInstance = FAnimCustomInstanceHelper::BindToSkeletalMeshComponent<UControlRigLayerInstance>(SkeletalMeshComponent, bWasCreated))
diff --git a/Engine/Plugins/Animation/ControlRig/Source/ControlRig/Private/Units/Animation/RigUnit_AnimAttribute.cpp b/Engine/Plugins/Animation/ControlRig/Source/ControlRig/Private/Units/Animation/RigUnit_AnimAttribute.cpp
--- a/Engine/Plugins/Animation/ControlRig/Source/ControlRig/Private/Units/Animation/RigUnit_AnimAttribute.cpp
+++ b/Engine/Plugins/Animation/ControlRig/Source/ControlRig/Private/Units/Animation/RigUnit_AnimAttribute.cpp
@@ -67,7 +67,7 @@ FORCEINLINE_DEBUGGABLE static T* GetAnimAttributeValue(
 	const USkeletalMeshComponent* OwningComponent = Cast<USkeletalMeshComponent>(Context.OwningComponent);
 
 	if (!OwningComponent ||
-		!OwningComponent->SkeletalMesh)
+		!OwningComponent->GetSkeletalMesh())
 	{
 		return nullptr;
 	}
@@ -82,7 +82,7 @@ FORCEINLINE_DEBUGGABLE static T* GetAnimAttributeValue(
 		// Invalidate cache if input changed
 		if (CachedBoneName != BoneName)
 		{
-			CachedBoneIndex = OwningComponent->SkeletalMesh->GetRefSkeleton().FindBoneIndex(BoneName);
+			CachedBoneIndex = OwningComponent->GetSkeletalMesh()->GetRefSkeleton().FindBoneIndex(BoneName);
 		}
 	}
 	
diff --git a/Engine/Plugins/Animation/ControlRig/Source/ControlRig/Public/Tools/ControlRigToolAsset.h b/Engine/Plugins/Animation/ControlRig/Source/ControlRig/Public/Tools/ControlRigToolAsset.h
--- a/Engine/Plugins/Animation/ControlRig/Source/ControlRig/Public/Tools/ControlRigToolAsset.h
+++ b/Engine/Plugins/Animation/ControlRig/Source/ControlRig/Public/Tools/ControlRigToolAsset.h
@@ -59,7 +59,7 @@ public:
 			USkeletalMesh* SkelMesh = nullptr;
 			if (USkeletalMeshComponent* RigMeshComp = Cast<USkeletalMeshComponent>(InControlRig->GetObjectBinding()->GetBoundObject()))
 			{
-				SkelMesh = RigMeshComp->SkeletalMesh;
+				SkelMesh = RigMeshComp->GetSkeletalMesh();
 			}
 
 			Asset->SavePose(InControlRig, bUseAllControls);
diff --git a/Engine/Plugins/Animation/ControlRig/Source/ControlRigEditor/Private/ControlRigBlueprintActions.cpp b/Engine/Plugins/Animation/ControlRig/Source/ControlRigEditor/Private/ControlRigBlueprintActions.cpp
--- a/Engine/Plugins/Animation/ControlRig/Source/ControlRigEditor/Private/ControlRigBlueprintActions.cpp
+++ b/Engine/Plugins/Animation/ControlRig/Source/ControlRigEditor/Private/ControlRigBlueprintActions.cpp
@@ -381,7 +381,7 @@ void FControlRigBlueprintActions::OnSpawnedSkeletalMeshActorChanged(UObject* InO
 
 			if (!ControlRig || ControlRig->GetClass() != ControlRigClass)
 			{		
-				USkeletalMesh* SkeletalMesh = MeshActor->GetSkeletalMeshComponent()->SkeletalMesh;
+				USkeletalMesh* SkeletalMesh = MeshActor->GetSkeletalMeshComponent()->GetSkeletalMesh();
 				USkeleton* Skeleton = SkeletalMesh->GetSkeleton();
 				
 				ObjectName.RemoveFromEnd(TEXT("_C"));
diff --git a/Engine/Plugins/Animation/ControlRig/Source/ControlRigEditor/Private/ControlRigEditorModule.cpp b/Engine/Plugins/Animation/ControlRig/Source/ControlRigEditor/Private/ControlRigEditorModule.cpp
--- a/Engine/Plugins/Animation/ControlRig/Source/ControlRigEditor/Private/ControlRigEditorModule.cpp
+++ b/Engine/Plugins/Animation/ControlRig/Source/ControlRigEditor/Private/ControlRigEditorModule.cpp
@@ -708,7 +708,7 @@ void FControlRigEditorModule::BakeToControlRig(UClass* ControlRigClass, UAnimSeq
 			if (Track)
 			{
 				USkeletalMeshComponent* SkelMeshComp = MeshActor->GetSkeletalMeshComponent();
-				USkeletalMesh* SkeletalMesh = SkelMeshComp->SkeletalMesh;
+				USkeletalMesh* SkeletalMesh = SkelMeshComp->GetSkeletalMesh();
 
 				FString ObjectName = (ControlRigClass->GetName());
 				ObjectName.RemoveFromEnd(TEXT("_C"));
diff --git a/Engine/Plugins/Animation/ControlRig/Source/ControlRigEditor/Private/ControlRigSequencerEditorLibrary.cpp b/Engine/Plugins/Animation/ControlRig/Source/ControlRigEditor/Private/ControlRigSequencerEditorLibrary.cpp
--- a/Engine/Plugins/Animation/ControlRig/Source/ControlRigEditor/Private/ControlRigSequencerEditorLibrary.cpp
+++ b/Engine/Plugins/Animation/ControlRig/Source/ControlRigEditor/Private/ControlRigSequencerEditorLibrary.cpp
@@ -116,9 +116,9 @@ static void AcquireSkeletonAndSkelMeshCompFromObject(UObject* BoundObject, USkel
 			if (SkeletalMeshComp)
 			{
 				*OutSkeletalMeshComponent = SkeletalMeshComp;
-				if (SkeletalMeshComp->SkeletalMesh && SkeletalMeshComp->SkeletalMesh->GetSkeleton())
+				if (SkeletalMeshComp->GetSkeletalMesh() && SkeletalMeshComp->GetSkeletalMesh()->GetSkeleton())
 				{
-					*OutSkeleton = SkeletalMeshComp->SkeletalMesh->GetSkeleton();
+					*OutSkeleton = SkeletalMeshComp->GetSkeletalMesh()->GetSkeleton();
 				}
 				return;
 			}
@@ -133,9 +133,9 @@ static void AcquireSkeletonAndSkelMeshCompFromObject(UObject* BoundObject, USkel
 				if (SkeletalMeshComp)
 				{
 					*OutSkeletalMeshComponent = SkeletalMeshComp;
-					if (SkeletalMeshComp->SkeletalMesh && SkeletalMeshComp->SkeletalMesh->GetSkeleton())
+					if (SkeletalMeshComp->GetSkeletalMesh() && SkeletalMeshComp->GetSkeletalMesh()->GetSkeleton())
 					{
-						*OutSkeleton = SkeletalMeshComp->SkeletalMesh->GetSkeleton();
+						*OutSkeleton = SkeletalMeshComp->GetSkeletalMesh()->GetSkeleton();
 					}
 					return;
 				}
@@ -155,9 +155,9 @@ static void AcquireSkeletonAndSkelMeshCompFromObject(UObject* BoundObject, USkel
 					if (SkeletalMeshComp)
 					{
 						*OutSkeletalMeshComponent = SkeletalMeshComp;
-						if (SkeletalMeshComp->SkeletalMesh && SkeletalMeshComp->SkeletalMesh->GetSkeleton())
+						if (SkeletalMeshComp->GetSkeletalMesh() && SkeletalMeshComp->GetSkeletalMesh()->GetSkeleton())
 						{
-							*OutSkeleton = SkeletalMeshComp->SkeletalMesh->GetSkeleton();
+							*OutSkeleton = SkeletalMeshComp->GetSkeletalMesh()->GetSkeleton();
 						}
 					}
 				}
@@ -167,9 +167,9 @@ static void AcquireSkeletonAndSkelMeshCompFromObject(UObject* BoundObject, USkel
 	else if (USkeletalMeshComponent* SkeletalMeshComponent = Cast<USkeletalMeshComponent>(BoundObject))
 	{
 		*OutSkeletalMeshComponent = SkeletalMeshComponent;
-		if (SkeletalMeshComponent->SkeletalMesh && SkeletalMeshComponent->SkeletalMesh->GetSkeleton())
+		if (SkeletalMeshComponent->GetSkeletalMesh() && SkeletalMeshComponent->GetSkeletalMesh()->GetSkeleton())
 		{
-			*OutSkeleton = SkeletalMeshComponent->SkeletalMesh->GetSkeleton();
+			*OutSkeleton = SkeletalMeshComponent->GetSkeletalMesh()->GetSkeleton();
 		}
 	}
 }
@@ -669,7 +669,7 @@ bool UControlRigSequencerEditorLibrary::BakeToControlRig(UWorld* World, ULevelSe
 			USkeleton* Skeleton = nullptr;
 			USkeletalMeshComponent* SkeletalMeshComp = nullptr;
 			AcquireSkeletonAndSkelMeshCompFromObject(BoundObject, &Skeleton, &SkeletalMeshComp);
-			if (SkeletalMeshComp && SkeletalMeshComp->SkeletalMesh && SkeletalMeshComp->SkeletalMesh->GetSkeleton())
+			if (SkeletalMeshComp && SkeletalMeshComp->GetSkeletalMesh() && SkeletalMeshComp->GetSkeletalMesh()->GetSkeleton())
 			{
 				UAnimSequence* TempAnimSequence = NewObject<UAnimSequence>(GetTransientPackage(), NAME_None);
 				TempAnimSequence->SetSkeleton(Skeleton);
diff --git a/Engine/Plugins/Animation/ControlRig/Source/ControlRigEditor/Private/Sequencer/ControlRigParameterTrackEditor.cpp b/Engine/Plugins/Animation/ControlRig/Source/ControlRigEditor/Private/Sequencer/ControlRigParameterTrackEditor.cpp
--- a/Engine/Plugins/Animation/ControlRig/Source/ControlRigEditor/Private/Sequencer/ControlRigParameterTrackEditor.cpp
+++ b/Engine/Plugins/Animation/ControlRig/Source/ControlRigEditor/Private/Sequencer/ControlRigParameterTrackEditor.cpp
@@ -102,7 +102,7 @@ static USkeletalMeshComponent* AcquireSkeletalMeshFromObject(UObject* BoundObjec
 	}
 	else if (USkeletalMeshComponent* SkeletalMeshComponent = Cast<USkeletalMeshComponent>(BoundObject))
 	{
-		if (SkeletalMeshComponent->SkeletalMesh)
+		if (SkeletalMeshComponent->GetSkeletalMesh())
 		{
 			return SkeletalMeshComponent;
 		}
@@ -115,10 +115,10 @@ static USkeletalMeshComponent* AcquireSkeletalMeshFromObject(UObject* BoundObjec
 static USkeleton* GetSkeletonFromComponent(UActorComponent* InComponent)
 {
 	USkeletalMeshComponent* SkeletalMeshComp = Cast<USkeletalMeshComponent>(InComponent);
-	if (SkeletalMeshComp && SkeletalMeshComp->SkeletalMesh && SkeletalMeshComp->SkeletalMesh->GetSkeleton())
+	if (SkeletalMeshComp && SkeletalMeshComp->GetSkeletalMesh() && SkeletalMeshComp->GetSkeletalMesh()->GetSkeleton())
 	{
 		// @todo Multiple actors, multiple components
-		return SkeletalMeshComp->SkeletalMesh->GetSkeleton();
+		return SkeletalMeshComp->GetSkeletalMesh()->GetSkeleton();
 	}
 
 	return nullptr;
diff --git a/Engine/Plugins/Animation/ControlRig/Source/ControlRigEditor/Private/Tools/ControlRigPoseThumbnailRenderer.cpp b/Engine/Plugins/Animation/ControlRig/Source/ControlRigEditor/Private/Tools/ControlRigPoseThumbnailRenderer.cpp
--- a/Engine/Plugins/Animation/ControlRig/Source/ControlRigEditor/Private/Tools/ControlRigPoseThumbnailRenderer.cpp
+++ b/Engine/Plugins/Animation/ControlRig/Source/ControlRigEditor/Private/Tools/ControlRigPoseThumbnailRenderer.cpp
@@ -149,7 +149,7 @@ void FControlRigPoseThumbnailScene::CleanupComponentChildren(USceneComponent* Co
 void FControlRigPoseThumbnailScene::GetViewMatrixParameters(const float InFOVDegrees, FVector& OutOrigin, float& OutOrbitPitch, float& OutOrbitYaw, float& OutOrbitZoom) const
 {
 	check(PreviewActor->GetSkeletalMeshComponent());
-	check(PreviewActor->GetSkeletalMeshComponent()->SkeletalMesh);
+	check(PreviewActor->GetSkeletalMeshComponent()->GetSkeletalMesh());
 
 	const float HalfFOVRadians = FMath::DegreesToRadians<float>(InFOVDegrees) * 0.5f;
 	// No need to add extra size to view slightly outside of the sphere to compensate for perspective since skeletal meshes already buffer bounds.
diff --git a/Engine/Plugins/Animation/GameplayInsights/Source/GameplayInsightsEditor/Private/InsightsSkeletalMeshComponent.cpp b/Engine/Plugins/Animation/GameplayInsights/Source/GameplayInsightsEditor/Private/InsightsSkeletalMeshComponent.cpp
--- a/Engine/Plugins/Animation/GameplayInsights/Source/GameplayInsightsEditor/Private/InsightsSkeletalMeshComponent.cpp
+++ b/Engine/Plugins/Animation/GameplayInsights/Source/GameplayInsightsEditor/Private/InsightsSkeletalMeshComponent.cpp
@@ -29,9 +29,9 @@ void UInsightsSkeletalMeshComponent::SetPoseFromProvider(const IAnimationProvide
 
 void UInsightsSkeletalMeshComponent::InitAnim(bool bForceReInit)
 {
-	if(SkeletalMesh)
+	if(GetSkeletalMesh())
 	{
-		const FReferenceSkeleton& SkeletalMeshRefSkeleton = SkeletalMesh->GetRefSkeleton();
+		const FReferenceSkeleton& SkeletalMeshRefSkeleton = GetSkeletalMesh()->GetRefSkeleton();
 		// set up bone visibility states as this gets skipped since we allocate the component array before registration
 		for (int32 BaseIndex = 0; BaseIndex < 2; ++BaseIndex)
 		{
@@ -43,6 +43,6 @@ void UInsightsSkeletalMeshComponent::InitAnim(bool bForceReInit)
 			}
 		}
 
-		MorphTargetWeights.SetNumZeroed(SkeletalMesh->GetMorphTargets().Num());
+		MorphTargetWeights.SetNumZeroed(GetSkeletalMesh()->GetMorphTargets().Num());
 	}
 }
\ No newline at end of file
diff --git a/Engine/Plugins/Animation/IKRig/Source/IKRig/Private/AnimNodes/AnimNode_IKRig.cpp b/Engine/Plugins/Animation/IKRig/Source/IKRig/Private/AnimNodes/AnimNode_IKRig.cpp
--- a/Engine/Plugins/Animation/IKRig/Source/IKRig/Private/AnimNodes/AnimNode_IKRig.cpp
+++ b/Engine/Plugins/Animation/IKRig/Source/IKRig/Private/AnimNodes/AnimNode_IKRig.cpp
@@ -216,7 +216,7 @@ void FAnimNode_IKRig::PreUpdate(const UAnimInstance* InAnimInstance)
 	// initialize the IK Rig (will only try once on the current version of the rig asset)
 	if (!IKRigProcessor->IsInitialized())
 	{
- 		IKRigProcessor->Initialize(RigDefinitionAsset, InAnimInstance->GetSkelMeshComponent()->SkeletalMesh);
+ 		IKRigProcessor->Initialize(RigDefinitionAsset, InAnimInstance->GetSkelMeshComponent()->GetSkeletalMesh());
 	}
 	
 	// cache list of goal creator components on the actor
diff --git a/Engine/Plugins/Animation/IKRig/Source/IKRig/Private/AnimNodes/AnimNode_RetargetPoseFromMesh.cpp b/Engine/Plugins/Animation/IKRig/Source/IKRig/Private/AnimNodes/AnimNode_RetargetPoseFromMesh.cpp
--- a/Engine/Plugins/Animation/IKRig/Source/IKRig/Private/AnimNodes/AnimNode_RetargetPoseFromMesh.cpp
+++ b/Engine/Plugins/Animation/IKRig/Source/IKRig/Private/AnimNodes/AnimNode_RetargetPoseFromMesh.cpp
@@ -23,7 +23,7 @@ void FAnimNode_RetargetPoseFromMesh::CacheBones_AnyThread(const FAnimationCacheB
 		return;
 	}
 	
-	if (!Context.AnimInstanceProxy->GetSkelMeshComponent()->SkeletalMesh)
+	if (!Context.AnimInstanceProxy->GetSkelMeshComponent()->GetSkeletalMesh())
 	{
 		return;
 	}
@@ -32,7 +32,7 @@ void FAnimNode_RetargetPoseFromMesh::CacheBones_AnyThread(const FAnimationCacheB
 	RequiredToTargetBoneMapping.Reset();
 
 	const FReferenceSkeleton& RefSkeleton = RequiredBones.GetReferenceSkeleton();
-	const FReferenceSkeleton& TargetSkeleton = Context.AnimInstanceProxy->GetSkelMeshComponent()->SkeletalMesh->GetRefSkeleton();
+	const FReferenceSkeleton& TargetSkeleton = Context.AnimInstanceProxy->GetSkelMeshComponent()->GetSkeletalMesh()->GetRefSkeleton();
 	const TArray<FBoneIndexType>& RequiredBonesArray = RequiredBones.GetBoneIndicesArray();
 	for (int32 Index = 0; Index < RequiredBonesArray.Num(); ++Index)
 	{
@@ -78,8 +78,8 @@ void FAnimNode_RetargetPoseFromMesh::Evaluate_AnyThread(FPoseContext& Output)
 
 	// ensure processor was initialized with the currently used assets (source/target meshes and retarget asset)
 	// if processor is not ready this tick, it will be next tick as this state will trigger re-initialization
-	const TObjectPtr<USkeletalMesh> SourceMesh = SourceMeshComponent->SkeletalMesh;
-	const TObjectPtr<USkeletalMesh> TargetMesh = Output.AnimInstanceProxy->GetSkelMeshComponent()->SkeletalMesh;
+	const TObjectPtr<USkeletalMesh> SourceMesh = SourceMeshComponent->GetSkeletalMesh();
+	const TObjectPtr<USkeletalMesh> TargetMesh = Output.AnimInstanceProxy->GetSkelMeshComponent()->GetSkeletalMesh();
 	const bool bIsProcessorReady = Processor->WasInitializedWithTheseAssets(SourceMesh, TargetMesh, IKRetargeterAsset);
 
 	// if not ready to run, skip retarget and output the ref pose
@@ -234,8 +234,8 @@ bool FAnimNode_RetargetPoseFromMesh::EnsureProcessorIsInitialized(const TObjectP
 	}
 
 	// check that both a source and target mesh exist
-	const TObjectPtr<USkeletalMesh> SourceMesh = SourceMeshComponent->SkeletalMesh;
-	const TObjectPtr<USkeletalMesh> TargetMesh = TargetMeshComponent->SkeletalMesh;
+	const TObjectPtr<USkeletalMesh> SourceMesh = SourceMeshComponent->GetSkeletalMesh();
+	const TObjectPtr<USkeletalMesh> TargetMesh = TargetMeshComponent->GetSkeletalMesh();
 	if (SourceMesh.IsNull() || TargetMesh.IsNull())
 	{
 		return false; // cannot initialize if components are missing skeletal mesh references
diff --git a/Engine/Plugins/Animation/IKRig/Source/IKRigEditor/Private/RetargetEditor/IKRetargetEditorController.cpp b/Engine/Plugins/Animation/IKRig/Source/IKRigEditor/Private/RetargetEditor/IKRetargetEditorController.cpp
--- a/Engine/Plugins/Animation/IKRig/Source/IKRigEditor/Private/RetargetEditor/IKRetargetEditorController.cpp
+++ b/Engine/Plugins/Animation/IKRig/Source/IKRigEditor/Private/RetargetEditor/IKRetargetEditorController.cpp
@@ -344,7 +344,7 @@ void FIKRetargetEditorController::EditBoneSelection(
 {
 	// must have a skeletal mesh
 	UDebugSkelMeshComponent* DebugComponent = GetEditedSkeletalMesh();
-	if (!DebugComponent->SkeletalMesh)
+	if (!DebugComponent->GetSkeletalMesh())
 	{
 		return;
 	}
@@ -378,7 +378,7 @@ void FIKRetargetEditorController::EditBoneSelection(
 	}
 
 	// convert to bone indices
-	const FReferenceSkeleton& RefSkeleton = DebugComponent->SkeletalMesh->GetRefSkeleton();
+	const FReferenceSkeleton& RefSkeleton = DebugComponent->GetSkeletalMesh()->GetRefSkeleton();
 	TArray<int32> SelectedBoneIndices;
 	for (const FName& Bone : SelectedBones)
 	{
diff --git a/Engine/Plugins/Animation/IKRig/Source/IKRigEditor/Private/RigEditor/IKRigEditorController.cpp b/Engine/Plugins/Animation/IKRig/Source/IKRigEditor/Private/RigEditor/IKRigEditorController.cpp
--- a/Engine/Plugins/Animation/IKRig/Source/IKRigEditor/Private/RigEditor/IKRigEditorController.cpp
+++ b/Engine/Plugins/Animation/IKRig/Source/IKRigEditor/Private/RigEditor/IKRigEditorController.cpp
@@ -52,7 +52,7 @@ TOptional<FTransform> UIKRigBoneDetails::GetTransform(EIKRigTransformType::Type
 			
 			USkeletalMeshComponent* SkeletalMeshComponent = AnimInstancePtr->GetSkelMeshComponent();
 			const bool IsSkelMeshValid = SkeletalMeshComponent != nullptr &&
-										SkeletalMeshComponent->SkeletalMesh != nullptr;
+										SkeletalMeshComponent->GetSkeletalMesh() != nullptr;
 			if (IsSkelMeshValid)
 			{
 				GlobalTransform = SkeletalMeshComponent->GetBoneTransform(BoneIndex);
diff --git a/Engine/Plugins/Animation/RigLogic/Source/RigLogicModule/Private/AnimNode_RigLogic.cpp b/Engine/Plugins/Animation/RigLogic/Source/RigLogicModule/Private/AnimNode_RigLogic.cpp
--- a/Engine/Plugins/Animation/RigLogic/Source/RigLogicModule/Private/AnimNode_RigLogic.cpp
+++ b/Engine/Plugins/Animation/RigLogic/Source/RigLogicModule/Private/AnimNode_RigLogic.cpp
@@ -39,7 +39,7 @@ void FAnimNode_RigLogic::Initialize_AnyThread(const FAnimationInitializeContext&
 	FAnimNode_Base::Initialize_AnyThread(Context);
 	AnimSequence.Initialize(Context);
 	USkeletalMeshComponent* SkeletalMeshComponent = Context.AnimInstanceProxy->GetSkelMeshComponent();
-	USkeletalMesh* SkeletalMesh = SkeletalMeshComponent->SkeletalMesh;
+	USkeletalMesh* SkeletalMesh = SkeletalMeshComponent->GetSkeletalMesh();
 
 	static FRWLock GlobalCreationLock;
 	FRWScopeLock ScopeLock(GlobalCreationLock, SLT_ReadOnly);
@@ -66,7 +66,7 @@ void FAnimNode_RigLogic::CacheBones_AnyThread(const FAnimationCacheBonesContext&
 
 	// Initialize things that depend on the skeleton of Anim BP
 	USkeletalMeshComponent* SkeletalMeshComponent = Context.AnimInstanceProxy->GetSkelMeshComponent();
-	USkeletalMesh* SkeletalMesh = SkeletalMeshComponent->SkeletalMesh;
+	USkeletalMesh* SkeletalMesh = SkeletalMeshComponent->GetSkeletalMesh();
 	USkeleton* Skeleton = Context.AnimInstanceProxy->GetSkeleton();
 	UDNAAsset* DNAAsset = Cast<UDNAAsset>(SkeletalMesh->GetAssetUserDataOfClass(UDNAAsset::StaticClass()));
 	UDNAIndexMapping* DNAIndexMappingContainer = Cast<UDNAIndexMapping>(SkeletalMesh->GetAssetUserDataOfClass(UDNAIndexMapping::StaticClass()));
diff --git a/Engine/Plugins/Animation/RigLogic/Source/RigLogicModule/Private/RigUnit_RigLogic.cpp b/Engine/Plugins/Animation/RigLogic/Source/RigLogicModule/Private/RigUnit_RigLogic.cpp
--- a/Engine/Plugins/Animation/RigLogic/Source/RigLogicModule/Private/RigUnit_RigLogic.cpp
+++ b/Engine/Plugins/Animation/RigLogic/Source/RigLogicModule/Private/RigUnit_RigLogic.cpp
@@ -345,14 +345,14 @@ FRigUnit_RigLogic_Execute()
 					// however, during unit testing we cannot fetch it from DataSourceRegistry 
 					// in that case, a mock version will be inserted into Data by unit test beforehand
 				}
-				if (!Data.SkelMeshComponent.IsValid() || Data.SkelMeshComponent->SkeletalMesh == nullptr)
+				if (!Data.SkelMeshComponent.IsValid() || Data.SkelMeshComponent->GetSkeletalMesh() == nullptr)
 				{
 					return;
 				}
 				Data.CurrentLOD = Data.SkelMeshComponent->GetPredictedLODLevel();
 
 				// Fetch shared runtime context of rig from DNAAsset
-				TSharedPtr<FSharedRigRuntimeContext> RigRuntimeContext = GetSharedRigRuntimeContext(Data.SkelMeshComponent->SkeletalMesh);
+				TSharedPtr<FSharedRigRuntimeContext> RigRuntimeContext = GetSharedRigRuntimeContext(Data.SkelMeshComponent->GetSkeletalMesh());
 				// Context is initialized with a BehaviorReader, which can be imported into SkeletalMesh from DNA file
 				// or overwritten by GeneSplicer when making a new character
 				Data.InitializeRigLogic(Hierarchy, RigRuntimeContext);
diff --git a/Engine/Plugins/Animation/RigLogic/Source/RigLogicModule/Private/RigUnit_RigLogic_Test.cpp b/Engine/Plugins/Animation/RigLogic/Source/RigLogicModule/Private/RigUnit_RigLogic_Test.cpp
--- a/Engine/Plugins/Animation/RigLogic/Source/RigLogicModule/Private/RigUnit_RigLogic_Test.cpp
+++ b/Engine/Plugins/Animation/RigLogic/Source/RigLogicModule/Private/RigUnit_RigLogic_Test.cpp
@@ -277,7 +277,7 @@ IMPLEMENT_RIGUNIT_AUTOMATION_TEST(FRigUnit_RigLogic)
 
 	//Test
 	InitAndExecute();
-	TSharedPtr<FSharedRigRuntimeContext> SharedRigRuntimeContext = Test.GetSharedRigRuntimeContext(Test.GetData()->SkelMeshComponent->SkeletalMesh);
+	TSharedPtr<FSharedRigRuntimeContext> SharedRigRuntimeContext = Test.GetSharedRigRuntimeContext(Test.GetData()->SkelMeshComponent->GetSkeletalMesh());
 	//Assert
 	AddErrorIfFalse(
 		// Check rig logic initialized
diff --git a/Engine/Plugins/Developer/AnimationSharing/Source/AnimationSharing/Private/AnimationSharingManager.cpp b/Engine/Plugins/Developer/AnimationSharing/Source/AnimationSharing/Private/AnimationSharingManager.cpp
--- a/Engine/Plugins/Developer/AnimationSharing/Source/AnimationSharing/Private/AnimationSharingManager.cpp
+++ b/Engine/Plugins/Developer/AnimationSharing/Source/AnimationSharing/Private/AnimationSharingManager.cpp
@@ -366,7 +366,7 @@ void UAnimationSharingManager::RegisterActor(AActor* InActor, FUpdateActorHandle
 			const USkeleton* CurrentSkeleton = nullptr;
 			for (USkeletalMeshComponent* SkeletalMeshComponent : OwnedComponents)
 			{
-				const USkeletalMesh* Mesh = SkeletalMeshComponent->SkeletalMesh;
+				const USkeletalMesh* Mesh = SkeletalMeshComponent->GetSkeletalMesh();
 				const USkeleton* Skeleton = Mesh->GetSkeleton();
 
 				if (CurrentSkeleton == nullptr)
diff --git a/Engine/Plugins/Editor/SequencerAnimTools/Source/SequencerAnimTools/Private/Sequencer/AnimationBoneTrail.cpp b/Engine/Plugins/Editor/SequencerAnimTools/Source/SequencerAnimTools/Private/Sequencer/AnimationBoneTrail.cpp
--- a/Engine/Plugins/Editor/SequencerAnimTools/Source/SequencerAnimTools/Private/Sequencer/AnimationBoneTrail.cpp
+++ b/Engine/Plugins/Editor/SequencerAnimTools/Source/SequencerAnimTools/Private/Sequencer/AnimationBoneTrail.cpp
@@ -29,14 +29,14 @@ void FAnimTrajectoryCache::Evaluate(FTrajectoryCache* ParentTransformCache)
 	TSharedPtr<ISequencer> Sequencer = WeakSequencer.Pin();
 	check(Sequencer);
 
-	if (!SkeletalMeshComponent.IsValid() || !SkeletalMeshComponent->SkeletalMesh || !SkeletalMeshComponent->SkeletalMesh->GetSkeleton())
+	if (!SkeletalMeshComponent.IsValid() || !SkeletalMeshComponent->GetSkeletalMesh() || !SkeletalMeshComponent->GetSkeletalMesh()->GetSkeleton())
 	{
 		return;
 	}
 
 	SkeletalMeshComponent->MasterPoseComponent = nullptr;
 
-	CachedAnimSequence->SetSkeleton(SkeletalMeshComponent->SkeletalMesh->GetSkeleton());
+	CachedAnimSequence->SetSkeleton(SkeletalMeshComponent->GetSkeletalMesh()->GetSkeleton());
 
 	// TODO: for some reason SkeletalMeshComponent becomes invalid sometimes when evaluating, no clue why yet
 	FMovieSceneSequenceTransform MovieSceneSequenceTransform;
diff --git a/Engine/Plugins/Editor/SequencerAnimTools/Source/SequencerAnimTools/Private/Sequencer/AnimationBoneTrail.h b/Engine/Plugins/Editor/SequencerAnimTools/Source/SequencerAnimTools/Private/Sequencer/AnimationBoneTrail.h
--- a/Engine/Plugins/Editor/SequencerAnimTools/Source/SequencerAnimTools/Private/Sequencer/AnimationBoneTrail.h
+++ b/Engine/Plugins/Editor/SequencerAnimTools/Source/SequencerAnimTools/Private/Sequencer/AnimationBoneTrail.h
@@ -28,7 +28,7 @@ public:
 		, Spacing()
 		, bDirty(true)
 	{
-		CachedAnimSequence->SetSkeleton(InSkeletalMeshComponent->SkeletalMesh->GetSkeleton());
+		CachedAnimSequence->SetSkeleton(InSkeletalMeshComponent->GetSkeletalMesh()->GetSkeleton());
 	}
 
 	void Evaluate(FTrajectoryCache* ParentTrajectoryCache);
diff --git a/Engine/Plugins/Editor/SequencerAnimTools/Source/SequencerAnimTools/Private/Sequencer/SequencerTrailHierarchy.cpp b/Engine/Plugins/Editor/SequencerAnimTools/Source/SequencerAnimTools/Private/Sequencer/SequencerTrailHierarchy.cpp
--- a/Engine/Plugins/Editor/SequencerAnimTools/Source/SequencerAnimTools/Private/Sequencer/SequencerTrailHierarchy.cpp
+++ b/Engine/Plugins/Editor/SequencerAnimTools/Source/SequencerAnimTools/Private/Sequencer/SequencerTrailHierarchy.cpp
@@ -412,7 +412,7 @@ void FSequencerTrailHierarchy::OnBoneVisibilityChanged(class USkeleton* Skeleton
 				BoundComponent = BoundActor->FindComponentByClass<USkeletalMeshComponent>();
 			}
 
-			if (!BoundComponent || !BoundComponent->SkeletalMesh || !BoundComponent->SkeletalMesh->GetSkeleton() || !(BoundComponent->SkeletalMesh->GetSkeleton() == Skeleton) || !BonesTracked.Contains(BoundComponent))
+			if (!BoundComponent || !BoundComponent->GetSkeletalMesh() || !BoundComponent->GetSkeletalMesh()->GetSkeleton() || !(BoundComponent->GetSkeletalMesh()->GetSkeleton() == Skeleton) || !BonesTracked.Contains(BoundComponent))
 			{
 				continue;
 			}
@@ -534,7 +534,7 @@ void FSequencerTrailHierarchy::UpdateSequencerBindings(const TArray<FGuid>& Sequ
 					BoundComponent = BoundActor->FindComponentByClass<USkeletalMeshComponent>();
 				}
 
-				if (!BoundComponent || !BoundComponent->SkeletalMesh || !BoundComponent->SkeletalMesh->GetSkeleton())
+				if (!BoundComponent || !BoundComponent->GetSkeletalMesh() || !BoundComponent->GetSkeletalMesh()->GetSkeleton())
 				{
 					continue;
 				}
@@ -586,7 +586,7 @@ void FSequencerTrailHierarchy::UpdateSequencerBindings(const TArray<FGuid>& Sequ
 					BoundComponent = BoundActor->FindComponentByClass<USkeletalMeshComponent>();
 				}
 
-				if (!BoundComponent || !BoundComponent->SkeletalMesh || !BoundComponent->SkeletalMesh->GetSkeleton())
+				if (!BoundComponent || !BoundComponent->GetSkeletalMesh() || !BoundComponent->GetSkeletalMesh()->GetSkeleton())
 				{
 					continue;
 				};
@@ -676,7 +676,7 @@ void FSequencerTrailHierarchy::AddSkeletonToHierarchy(class USkeletalMeshCompone
 	TSharedPtr<FAnimTrajectoryCache> AnimTrajectoryCache = MakeShared<FAnimTrajectoryCache>(CompToAdd, Sequencer);
 	TMap<FName, FGuid>& BoneMap = BonesTracked.Add(CompToAdd, TMap<FName, FGuid>());
 	
-	USkeleton* MySkeleton = CompToAdd->SkeletalMesh->GetSkeleton();
+	USkeleton* MySkeleton = CompToAdd->GetSkeletalMesh()->GetSkeleton();
 	const int32 NumBones = MySkeleton->GetReferenceSkeleton().GetNum();
 	for (int32 BoneIdx = 0; BoneIdx < NumBones; BoneIdx++)
 	{
diff --git a/Engine/Plugins/Enterprise/DataprepEditor/Source/DataprepEditor/Private/DataprepStats.cpp b/Engine/Plugins/Enterprise/DataprepEditor/Source/DataprepEditor/Private/DataprepStats.cpp
--- a/Engine/Plugins/Enterprise/DataprepEditor/Source/DataprepEditor/Private/DataprepStats.cpp
+++ b/Engine/Plugins/Enterprise/DataprepEditor/Source/DataprepEditor/Private/DataprepStats.cpp
@@ -296,7 +296,7 @@ private:
 			}
 			else if( SkeletalMeshComponent )
 			{
-				USkeletalMesh* SkeletalMesh = SkeletalMeshComponent->SkeletalMesh;
+				USkeletalMesh* SkeletalMesh = SkeletalMeshComponent->GetSkeletalMesh();
 				if( SkeletalMesh )
 				{
 					SkeletalMeshes.Add(SkeletalMesh);
diff --git a/Engine/Plugins/Experimental/Animation/CurveExpression/Source/Runtime/Private/AnimNode_RemapCurvesFromMesh.cpp b/Engine/Plugins/Experimental/Animation/CurveExpression/Source/Runtime/Private/AnimNode_RemapCurvesFromMesh.cpp
--- a/Engine/Plugins/Experimental/Animation/CurveExpression/Source/Runtime/Private/AnimNode_RemapCurvesFromMesh.cpp
+++ b/Engine/Plugins/Experimental/Animation/CurveExpression/Source/Runtime/Private/AnimNode_RemapCurvesFromMesh.cpp
@@ -136,7 +136,7 @@ void FAnimNode_RemapCurvesFromMesh::GatherDebugData(
 	DECLARE_SCOPE_HIERARCHICAL_COUNTER_ANIMNODE(GatherDebugData)
 	FString DebugLine = DebugData.GetNodeName(this);
 
-	DebugLine += FString::Printf(TEXT("('%s')"), *GetNameSafe(CurrentlyUsedSourceMeshComponent.IsValid() ? CurrentlyUsedSourceMeshComponent.Get()->SkeletalMesh : nullptr));
+	DebugLine += FString::Printf(TEXT("('%s')"), *GetNameSafe(CurrentlyUsedSourceMeshComponent.IsValid() ? CurrentlyUsedSourceMeshComponent.Get()->GetSkeletalMesh() : nullptr));
 	DebugData.AddDebugItem(DebugLine, true);
 }
 
@@ -153,7 +153,7 @@ void FAnimNode_RemapCurvesFromMesh::PreUpdate(
 
 	const USkeletalMeshComponent* CurrentMeshComponent = CurrentlyUsedSourceMeshComponent.IsValid() ? CurrentlyUsedSourceMeshComponent.Get() : nullptr;
 
-	if (CurrentMeshComponent && CurrentMeshComponent->SkeletalMesh && CurrentMeshComponent->IsRegistered())
+	if (CurrentMeshComponent && CurrentMeshComponent->GetSkeletalMesh() && CurrentMeshComponent->IsRegistered())
 	{
 		// If our source is running under master-pose, then get bone data from there
 		if(USkeletalMeshComponent* MasterPoseComponent = Cast<USkeletalMeshComponent>(CurrentMeshComponent->MasterPoseComponent.Get()))
@@ -162,7 +162,7 @@ void FAnimNode_RemapCurvesFromMesh::PreUpdate(
 		}
 
 		// re-check mesh component validity as it may have changed to master
-		if(CurrentMeshComponent->SkeletalMesh && CurrentMeshComponent->IsRegistered())
+		if(CurrentMeshComponent->GetSkeletalMesh() && CurrentMeshComponent->IsRegistered())
 		{
 			if (ExpressionEngine.IsSet())
 			{
@@ -197,10 +197,10 @@ void FAnimNode_RemapCurvesFromMesh::ReinitializeMeshComponent(
 	CurveNameToUIDMap.Reset();
 	CachedExpressions.Reset();
 
-	if (InTargetMeshComponent && IsValid(InNewSkeletalMeshComponent) && InNewSkeletalMeshComponent->SkeletalMesh)
+	if (InTargetMeshComponent && IsValid(InNewSkeletalMeshComponent) && InNewSkeletalMeshComponent->GetSkeletalMesh())
 	{
-		USkeletalMesh* SourceSkelMesh = InNewSkeletalMeshComponent->SkeletalMesh;
-		USkeletalMesh* TargetSkelMesh = InTargetMeshComponent->SkeletalMesh;
+		USkeletalMesh* SourceSkelMesh = InNewSkeletalMeshComponent->GetSkeletalMesh();
+		USkeletalMesh* TargetSkelMesh = InTargetMeshComponent->GetSkeletalMesh();
 		
 		if (IsValid(SourceSkelMesh) && !SourceSkelMesh->HasAnyFlags(RF_NeedPostLoad) &&
 			IsValid(TargetSkelMesh) && !TargetSkelMesh->HasAnyFlags(RF_NeedPostLoad))
@@ -265,14 +265,14 @@ void FAnimNode_RemapCurvesFromMesh::RefreshMeshComponent(
 				ReinitializeMeshComponent(InMeshComponent, InTargetMeshComponent);
 			}
 			// if component is still same but mesh has been changed, we have to reinitialize
-			else if (CurrentMeshComponent->SkeletalMesh != CurrentlyUsedSourceMesh.Get())
+			else if (CurrentMeshComponent->GetSkeletalMesh() != CurrentlyUsedSourceMesh.Get())
 			{
 				ReinitializeMeshComponent(InMeshComponent, InTargetMeshComponent);
 			}
 			else if (InTargetMeshComponent)
 			{
 				// see if target mesh has changed
-				if (InTargetMeshComponent->SkeletalMesh != CurrentlyUsedTargetMesh.Get())
+				if (InTargetMeshComponent->GetSkeletalMesh() != CurrentlyUsedTargetMesh.Get())
 				{
 					ReinitializeMeshComponent(InMeshComponent, InTargetMeshComponent);
 				}
diff --git a/Engine/Plugins/Experimental/Animation/PoseSearch/Source/Editor/Private/PoseSearchDebugger.cpp b/Engine/Plugins/Experimental/Animation/PoseSearch/Source/Editor/Private/PoseSearchDebugger.cpp
--- a/Engine/Plugins/Experimental/Animation/PoseSearch/Source/Editor/Private/PoseSearchDebugger.cpp
+++ b/Engine/Plugins/Experimental/Animation/PoseSearch/Source/Editor/Private/PoseSearchDebugger.cpp
@@ -128,7 +128,7 @@ void FPoseSearchDebuggerPoseVector::ExtractFeatures(const UE::PoseSearch::FFeatu
 void UPoseSearchMeshComponent::Initialize(const FTransform& InComponentToWorld)
 {
 	SetComponentToWorld(InComponentToWorld);
-	const FReferenceSkeleton& SkeletalMeshRefSkeleton = SkeletalMesh->GetRefSkeleton();
+	const FReferenceSkeleton& SkeletalMeshRefSkeleton = GetSkeletalMesh()->GetRefSkeleton();
 
 	// set up bone visibility states as this gets skipped since we allocate the component array before registration
 	for (int32 BaseIndex = 0; BaseIndex < 2; ++BaseIndex)
diff --git a/Engine/Plugins/Experimental/ChaosCloth/Source/Chaos/Private/ChaosCloth/ChaosClothConfig.cpp b/Engine/Plugins/Experimental/ChaosCloth/Source/Chaos/Private/ChaosCloth/ChaosClothConfig.cpp
--- a/Engine/Plugins/Experimental/ChaosCloth/Source/Chaos/Private/ChaosCloth/ChaosClothConfig.cpp
+++ b/Engine/Plugins/Experimental/ChaosCloth/Source/Chaos/Private/ChaosCloth/ChaosClothConfig.cpp
@@ -261,7 +261,7 @@ void UChaosClothSharedSimConfig::PostEditChangeChainProperty(FPropertyChangedCha
 			{
 				if (const USkeletalMeshComponent* const Component = *It)
 				{
-					if (Component->SkeletalMesh == OwnerMesh)
+					if (Component->GetSkeletalMesh() == OwnerMesh)
 					{
 						if (UClothingSimulationInteractor* const CurInteractor = Component->GetClothingSimulationInteractor())
 						{
diff --git a/Engine/Plugins/Experimental/ChaosVehiclesPlugin/Source/ChaosVehicles/Private/ChaosVehicleMovementComponent.cpp b/Engine/Plugins/Experimental/ChaosVehiclesPlugin/Source/ChaosVehicles/Private/ChaosVehicleMovementComponent.cpp
--- a/Engine/Plugins/Experimental/ChaosVehiclesPlugin/Source/ChaosVehicles/Private/ChaosVehicleMovementComponent.cpp
+++ b/Engine/Plugins/Experimental/ChaosVehiclesPlugin/Source/ChaosVehicles/Private/ChaosVehicleMovementComponent.cpp
@@ -1325,15 +1325,15 @@ FVector UChaosVehicleMovementComponent::LocateBoneOffset(const FName InBoneName,
 	{
 		if (USkinnedMeshComponent* Mesh = Cast<USkinnedMeshComponent>(GetMesh()))
 		{
-			if (ensureMsgf(Mesh->SkeletalMesh, TEXT("Expected skeletal mesh when locating bone offset. Asset might be missing references.")))
+			if (ensureMsgf(Mesh->GetSkeletalMesh(), TEXT("Expected skeletal mesh when locating bone offset. Asset might be missing references.")))
 			{
-				const FVector BonePosition = Mesh->SkeletalMesh->GetComposedRefPoseMatrix(InBoneName).GetOrigin() * Mesh->GetRelativeScale3D();
+				const FVector BonePosition = Mesh->GetSkeletalMesh()->GetComposedRefPoseMatrix(InBoneName).GetOrigin() * Mesh->GetRelativeScale3D();
 				//BonePosition is local for the root BONE of the skeletal mesh - however, we are using the Root BODY which may have its own transform, so we need to return the position local to the root BODY
 				FMatrix RootBodyMTX = FMatrix::Identity;
 
 				if (Mesh->GetBodyInstance() && Mesh->GetBodyInstance()->BodySetup.IsValid())
 				{
-					RootBodyMTX = Mesh->SkeletalMesh->GetComposedRefPoseMatrix(Mesh->GetBodyInstance()->BodySetup->BoneName);
+					RootBodyMTX = Mesh->GetSkeletalMesh()->GetComposedRefPoseMatrix(Mesh->GetBodyInstance()->BodySetup->BoneName);
 				}
 				const FVector LocalBonePosition = RootBodyMTX.InverseTransformPosition(BonePosition);
 				Offset += LocalBonePosition;
diff --git a/Engine/Plugins/Experimental/GeometryCollectionPlugin/Source/GeometryCollectionEditor/Private/GeometryCollection/GeometryCollectionConversion.cpp b/Engine/Plugins/Experimental/GeometryCollectionPlugin/Source/GeometryCollectionEditor/Private/GeometryCollection/GeometryCollectionConversion.cpp
--- a/Engine/Plugins/Experimental/GeometryCollectionPlugin/Source/GeometryCollectionEditor/Private/GeometryCollection/GeometryCollectionConversion.cpp
+++ b/Engine/Plugins/Experimental/GeometryCollectionPlugin/Source/GeometryCollectionEditor/Private/GeometryCollection/GeometryCollectionConversion.cpp
@@ -146,14 +146,14 @@ void FGeometryCollectionConversion::CreateFromSelectedActorsCommand(UWorld * Wor
 			{
 				for (int Index = 0; Index < SkeletalMeshComponents.Num(); Index++)
 				{
-					if (SkeletalMeshComponents[Index]->SkeletalMesh)
+					if (SkeletalMeshComponents[Index]->GetSkeletalMesh())
 					{
 						if (!GeometryCollectionAsset.IsCreated())
 						{
 							GeometryCollectionAsset.Create();
 						}
 						FGeometryCollectionEngineConversion::AppendSkeletalMesh(
-							SkeletalMeshComponents[Index]->SkeletalMesh,
+							SkeletalMeshComponents[Index]->GetSkeletalMesh(),
 							SkeletalMeshComponents[Index],
 							Actor->GetTransform(),
 							GeometryCollectionAsset.GeometryCollection);
diff --git a/Engine/Plugins/Experimental/GeometryCollectionPlugin/Source/GeometryCollectionEditor/Private/GeometryCollection/GeometryCollectionFactory.cpp b/Engine/Plugins/Experimental/GeometryCollectionPlugin/Source/GeometryCollectionEditor/Private/GeometryCollection/GeometryCollectionFactory.cpp
--- a/Engine/Plugins/Experimental/GeometryCollectionPlugin/Source/GeometryCollectionEditor/Private/GeometryCollection/GeometryCollectionFactory.cpp
+++ b/Engine/Plugins/Experimental/GeometryCollectionPlugin/Source/GeometryCollectionEditor/Private/GeometryCollection/GeometryCollectionFactory.cpp
@@ -66,10 +66,10 @@ void AppendActorComponentsRecursive(
 	{
 		for (int Index = 0; Index < SkeletalMeshComponents.Num(); Index++)
 		{
-			if (SkeletalMeshComponents[Index]->SkeletalMesh)
+			if (SkeletalMeshComponents[Index]->GetSkeletalMesh())
 			{
 				SkeletalMeshList.Add(GeometryCollectionSkeletalMeshConversionTuple(
-					SkeletalMeshComponents[Index]->SkeletalMesh,
+					SkeletalMeshComponents[Index]->GetSkeletalMesh(),
 					SkeletalMeshComponents[Index],
 					SkeletalMeshComponents[Index]->GetComponentTransform()));
 			}
diff --git a/Engine/Plugins/Experimental/MLDeformer/MLDeformerFramework/Source/MLDeformerFramework/Private/MLDeformerModelInstance.cpp b/Engine/Plugins/Experimental/MLDeformer/MLDeformerFramework/Source/MLDeformerFramework/Private/MLDeformerModelInstance.cpp
--- a/Engine/Plugins/Experimental/MLDeformer/MLDeformerFramework/Source/MLDeformerFramework/Private/MLDeformerModelInstance.cpp
+++ b/Engine/Plugins/Experimental/MLDeformer/MLDeformerFramework/Source/MLDeformerFramework/Private/MLDeformerModelInstance.cpp
@@ -35,7 +35,7 @@ void UMLDeformerModelInstance::Init(USkeletalMeshComponent* SkelMeshComponent)
 		return;
 	}
 
-	USkeletalMesh* SkelMesh = SkelMeshComponent->SkeletalMesh;
+	USkeletalMesh* SkelMesh = SkelMeshComponent->GetSkeletalMesh();
 	if (SkelMesh)
 	{
 		// Init the bone mapping table.
@@ -60,7 +60,7 @@ void UMLDeformerModelInstance::Init(USkeletalMeshComponent* SkelMeshComponent)
 
 void UMLDeformerModelInstance::UpdateCompatibilityStatus()
 {
-	bIsCompatible = SkeletalMeshComponent->SkeletalMesh && CheckCompatibility(SkeletalMeshComponent, true).IsEmpty();
+	bIsCompatible = SkeletalMeshComponent->GetSkeletalMesh() && CheckCompatibility(SkeletalMeshComponent, true).IsEmpty();
 }
 
 FString UMLDeformerModelInstance::CheckCompatibility(USkeletalMeshComponent* InSkelMeshComponent, bool LogIssues)
@@ -68,7 +68,7 @@ FString UMLDeformerModelInstance::CheckCompatibility(USkeletalMeshComponent* InS
 	ErrorText = FString();
 
 	// If we're not compatible, generate a compatibility string.
-	USkeletalMesh* SkelMesh = InSkelMeshComponent ? InSkelMeshComponent->SkeletalMesh.Get() : nullptr;
+	USkeletalMesh* SkelMesh = InSkelMeshComponent ? InSkelMeshComponent->GetSkeletalMesh() : nullptr;
 	UMLDeformerInputInfo* InputInfo = Model->GetInputInfo();
 	if (SkelMesh && !InputInfo->IsCompatible(SkelMesh) && Model->GetDeformerAsset())
 	{
@@ -109,7 +109,7 @@ void UMLDeformerModelInstance::UpdateBoneTransforms()
 	if (MasterPoseComponent)
 	{
 		const TArray<FTransform>& MasterTransforms = MasterPoseComponent->GetComponentSpaceTransforms();
-		USkeletalMesh* Mesh = MasterPoseComponent->SkeletalMesh;
+		USkeletalMesh* Mesh = MasterPoseComponent->GetSkeletalMesh();
 		const int32 NumTrainingBones = AssetBonesToSkelMeshMappings.Num();
 		for (int32 Index = 0; Index < NumTrainingBones; ++Index)
 		{
diff --git a/Engine/Plugins/Experimental/MLDeformer/MLDeformerFramework/Source/MLDeformerFrameworkEditor/Private/MLDeformerEditorModel.cpp b/Engine/Plugins/Experimental/MLDeformer/MLDeformerFramework/Source/MLDeformerFrameworkEditor/Private/MLDeformerEditorModel.cpp
--- a/Engine/Plugins/Experimental/MLDeformer/MLDeformerFramework/Source/MLDeformerFrameworkEditor/Private/MLDeformerEditorModel.cpp
+++ b/Engine/Plugins/Experimental/MLDeformer/MLDeformerFramework/Source/MLDeformerFrameworkEditor/Private/MLDeformerEditorModel.cpp
@@ -1215,7 +1215,7 @@ namespace UE::MLDeformer
 			const UMLDeformerModelInstance* ModelInstance = DeformerComponent->GetModelInstance();
 			if (ModelInstance &&
 				ModelInstance->GetSkeletalMeshComponent() && 
-				ModelInstance->GetSkeletalMeshComponent()->SkeletalMesh &&
+				ModelInstance->GetSkeletalMeshComponent()->GetSkeletalMesh() &&
 				!ModelInstance->IsCompatible() )
 			{
 				return FText::FromString( ModelInstance->GetCompatibilityErrorText() );
diff --git a/Engine/Plugins/Experimental/MLDeformer/MLDeformerFramework/Source/MLDeformerFrameworkEditor/Private/MLDeformerGeomCacheSampler.cpp b/Engine/Plugins/Experimental/MLDeformer/MLDeformerFramework/Source/MLDeformerFrameworkEditor/Private/MLDeformerGeomCacheSampler.cpp
--- a/Engine/Plugins/Experimental/MLDeformer/MLDeformerFramework/Source/MLDeformerFrameworkEditor/Private/MLDeformerGeomCacheSampler.cpp
+++ b/Engine/Plugins/Experimental/MLDeformer/MLDeformerFramework/Source/MLDeformerFrameworkEditor/Private/MLDeformerGeomCacheSampler.cpp
@@ -57,7 +57,7 @@ namespace UE::MLDeformer
 		// This will also calculate the skinned positions if the delta space is set to PostSkinning.
 		FMLDeformerSampler::Sample(InAnimFrameIndex);
 
-		USkeletalMesh* SkeletalMesh = SkeletalMeshComponent.Get() ? SkeletalMeshComponent->SkeletalMesh : nullptr;
+		USkeletalMesh* SkeletalMesh = SkeletalMeshComponent.Get() ? SkeletalMeshComponent->GetSkeletalMesh() : nullptr;
 		UGeometryCache* GeometryCache = GeometryCacheComponent->GetGeometryCache();
 		if (SkeletalMeshComponent && SkeletalMesh && GeometryCacheComponent && GeometryCache)
 		{
diff --git a/Engine/Plugins/Experimental/MLDeformer/MLDeformerFramework/Source/MLDeformerFrameworkEditor/Private/MLDeformerSampler.cpp b/Engine/Plugins/Experimental/MLDeformer/MLDeformerFramework/Source/MLDeformerFrameworkEditor/Private/MLDeformerSampler.cpp
--- a/Engine/Plugins/Experimental/MLDeformer/MLDeformerFramework/Source/MLDeformerFrameworkEditor/Private/MLDeformerSampler.cpp
+++ b/Engine/Plugins/Experimental/MLDeformer/MLDeformerFramework/Source/MLDeformerFrameworkEditor/Private/MLDeformerSampler.cpp
@@ -59,7 +59,7 @@ namespace UE::MLDeformer
 
 	void FMLDeformerSampler::UpdateSkeletalMeshComponent()
 	{
-		if (SkeletalMeshComponent && SkeletalMeshComponent->SkeletalMesh)
+		if (SkeletalMeshComponent && SkeletalMeshComponent->GetSkeletalMesh())
 		{
 			// Sample the transforms at the frame time.
 			SkeletalMeshComponent->SetPosition(SampleTime);
@@ -92,7 +92,7 @@ namespace UE::MLDeformer
 		AnimFrameIndex = InAnimFrameIndex;
 		SampleTime = GetTimeAtFrame(InAnimFrameIndex);
 
-		const USkeletalMesh* SkeletalMesh = SkeletalMeshComponent->SkeletalMesh;
+		const USkeletalMesh* SkeletalMesh = SkeletalMeshComponent->GetSkeletalMesh();
 		if (SkeletalMeshComponent && SkeletalMesh)
 		{
 			UpdateSkeletalMeshComponent();
@@ -114,7 +114,7 @@ namespace UE::MLDeformer
 	FMatrix44f FMLDeformerSampler::CalcInverseSkinningTransform(int32 VertexIndex, const FSkeletalMeshLODRenderData& SkelMeshLODData, const FSkinWeightVertexBuffer& SkinWeightBuffer) const
 	{
 		check(SkeletalMeshComponent);
-		const USkeletalMesh* Mesh = SkeletalMeshComponent->SkeletalMesh;
+		const USkeletalMesh* Mesh = SkeletalMeshComponent->GetSkeletalMesh();
 		check(Mesh);
 
 		// Find the render section, which we need to find the right bone index.
@@ -156,7 +156,7 @@ namespace UE::MLDeformer
 			return;
 		}
 
-		USkeletalMesh* Mesh = SkeletalMeshComponent->SkeletalMesh;
+		USkeletalMesh* Mesh = SkeletalMeshComponent->GetSkeletalMesh();
 		if (Mesh == nullptr)
 		{
 			return;
diff --git a/Engine/Plugins/Experimental/MeshModelingToolsetExp/Source/MeshModelingToolsExp/Private/SkinWeightsBindingTool.cpp b/Engine/Plugins/Experimental/MeshModelingToolsetExp/Source/MeshModelingToolsExp/Private/SkinWeightsBindingTool.cpp
--- a/Engine/Plugins/Experimental/MeshModelingToolsetExp/Source/MeshModelingToolsExp/Private/SkinWeightsBindingTool.cpp
+++ b/Engine/Plugins/Experimental/MeshModelingToolsetExp/Source/MeshModelingToolsExp/Private/SkinWeightsBindingTool.cpp
@@ -770,9 +770,9 @@ void USkinWeightsBindingTool::Setup()
 	
 	const USkeletalMeshComponent* SkelMeshComponent = Cast<USkeletalMeshComponent>(UE::ToolTarget::GetTargetComponent(Targets[0]));
 	
-	if (SkelMeshComponent && SkelMeshComponent->SkeletalMesh)
+	if (SkelMeshComponent && SkelMeshComponent->GetSkeletalMesh())
 	{
-		USkeletalMesh* SkeletalMesh = SkelMeshComponent->SkeletalMesh;
+		USkeletalMesh* SkeletalMesh = SkelMeshComponent->GetSkeletalMesh();
 
 		// Initialize the bone browser
 		FCurveEvaluationOption CurveEvalOption(
diff --git a/Engine/Plugins/Experimental/MeshModelingToolsetExp/Source/MeshModelingToolsExp/Private/SkinWeightsPaintTool.cpp b/Engine/Plugins/Experimental/MeshModelingToolsetExp/Source/MeshModelingToolsExp/Private/SkinWeightsPaintTool.cpp
--- a/Engine/Plugins/Experimental/MeshModelingToolsetExp/Source/MeshModelingToolsExp/Private/SkinWeightsPaintTool.cpp
+++ b/Engine/Plugins/Experimental/MeshModelingToolsetExp/Source/MeshModelingToolsExp/Private/SkinWeightsPaintTool.cpp
@@ -68,11 +68,11 @@ void USkinWeightsPaintTool::Setup()
 
 	ToolProps = NewObject<USkinWeightsPaintToolProperties>(this);
 	ToolProps->RestoreProperties(this);
-	if (Component && Component->SkeletalMesh)
+	if (Component && Component->GetSkeletalMesh())
 	{
 		// Get all non-virtual bones
 		// TArray<int32> BoneIndices;
-		USkeletalMesh& SkeletalMesh = *Component->SkeletalMesh;
+		USkeletalMesh& SkeletalMesh = *Component->GetSkeletalMesh();
 		// SkeletalMesh.RefSkeleton.GetRawRefBoneInfo();
 		
 
@@ -84,11 +84,11 @@ void USkinWeightsPaintTool::Setup()
 				);
 		BoneContainer.InitializeTo(Component->RequiredBones, CurveEvalOption, SkeletalMesh);
 
-		ToolProps->SkeletalMesh = Component->SkeletalMesh;
+		ToolProps->SkeletalMesh = Component->GetSkeletalMesh();
 		ToolProps->CurrentBone.Initialize(BoneContainer);
 
 		// Pick the first root bone as the initial selection.
-		PendingCurrentBone = Component->SkeletalMesh->GetRefSkeleton().GetBoneName(0);
+		PendingCurrentBone = Component->GetSkeletalMesh()->GetRefSkeleton().GetBoneName(0);
 		ToolProps->CurrentBone.BoneName = PendingCurrentBone.GetValue();
 
 		// Update the skeleton drawing information from the original bind pose
diff --git a/Engine/Plugins/FX/Niagara/Source/Niagara/Classes/NiagaraDataInterfaceSkeletalMesh.h b/Engine/Plugins/FX/Niagara/Source/Niagara/Classes/NiagaraDataInterfaceSkeletalMesh.h
--- a/Engine/Plugins/FX/Niagara/Source/Niagara/Classes/NiagaraDataInterfaceSkeletalMesh.h
+++ b/Engine/Plugins/FX/Niagara/Source/Niagara/Classes/NiagaraDataInterfaceSkeletalMesh.h
@@ -694,7 +694,7 @@ struct FNDISkeletalMesh_InstanceData
 	FORCEINLINE_DEBUGGABLE const FSkinWeightVertexBuffer* GetSkinWeights()
 	{
 		USkeletalMeshComponent* SkelComp = Cast<USkeletalMeshComponent>(SceneComponent.Get());
-		if (SkelComp != nullptr && SkelComp->SkeletalMesh != nullptr)
+		if (SkelComp != nullptr && SkelComp->GetSkeletalMesh() != nullptr)
 		{
 			return SkelComp->GetSkinWeightBuffer(CachedLODIdx);
 		}
diff --git a/Engine/Plugins/FX/Niagara/Source/Niagara/Private/NiagaraDataInterfaceRigidMeshCollisionQuery.cpp b/Engine/Plugins/FX/Niagara/Source/Niagara/Private/NiagaraDataInterfaceRigidMeshCollisionQuery.cpp
--- a/Engine/Plugins/FX/Niagara/Source/Niagara/Private/NiagaraDataInterfaceRigidMeshCollisionQuery.cpp
+++ b/Engine/Plugins/FX/Niagara/Source/Niagara/Private/NiagaraDataInterfaceRigidMeshCollisionQuery.cpp
@@ -213,7 +213,7 @@ void CollectComponents(TConstArrayView<AActor*> Actors, TConstArrayView<FName> C
 	{
 		if (UPhysicsAsset* PhysicsAsset = Component->GetPhysicsAsset())
 		{
-			USkeletalMesh* MeshAsset = Component->SkeletalMesh ? Component->SkeletalMesh.Get() : PhysicsAsset->GetPreviewMesh();
+			USkeletalMesh* MeshAsset = Component->GetSkeletalMesh() ? Component->GetSkeletalMesh().Get() : PhysicsAsset->GetPreviewMesh();
 			if (!MeshAsset || !MeshAsset->GetRefSkeleton().GetNum())
 			{
 				return false;
@@ -233,7 +233,7 @@ void ForEachBodySetup(USkeletalMeshComponent* Component, TBodySetupPredicate Pre
 {
 	if (UPhysicsAsset* PhysicsAsset = Component->GetPhysicsAsset())
 	{
-		USkeletalMesh* SkeletalMesh = Component->SkeletalMesh ? Component->SkeletalMesh.Get() : PhysicsAsset->GetPreviewMesh();
+		USkeletalMesh* SkeletalMesh = Component->GetSkeletalMesh() ? Component->GetSkeletalMesh().Get() : PhysicsAsset->GetPreviewMesh();
 		const FReferenceSkeleton& RefSkeleton = SkeletalMesh->GetRefSkeleton();
 
 		for (const UBodySetup* BodySetup : PhysicsAsset->SkeletalBodySetups)
@@ -327,7 +327,7 @@ FTransform CreateElementTransform(const TComponentType* Component, const UBodySe
 template<>
 FTransform CreateElementTransform<USkeletalMeshComponent>(const USkeletalMeshComponent* Component, const UBodySetup* BodySetup)
 {
-	if (USkeletalMesh* SkeletalMesh = Component->SkeletalMesh.Get())
+	if (USkeletalMesh* SkeletalMesh = Component->GetSkeletalMesh().Get())
 	{
 		const FReferenceSkeleton& RefSkeleton = SkeletalMesh->GetRefSkeleton();
 		const int32 BoneCount = RefSkeleton.GetNum();
diff --git a/Engine/Plugins/FX/Niagara/Source/Niagara/Private/NiagaraDataInterfaceSkeletalMesh.cpp b/Engine/Plugins/FX/Niagara/Source/Niagara/Private/NiagaraDataInterfaceSkeletalMesh.cpp
--- a/Engine/Plugins/FX/Niagara/Source/Niagara/Private/NiagaraDataInterfaceSkeletalMesh.cpp
+++ b/Engine/Plugins/FX/Niagara/Source/Niagara/Private/NiagaraDataInterfaceSkeletalMesh.cpp
@@ -240,7 +240,7 @@ void FSkeletalMeshSkinningData::RegisterUser(FSkeletalMeshSkinningDataUsage Usag
 	USkeletalMeshComponent* SkelComp = MeshComp.Get();
 	check(SkelComp);
 
-	USkeletalMesh* SkelMesh = SkelComp->SkeletalMesh;
+	USkeletalMesh* SkelMesh = SkelComp->GetSkeletalMesh();
 	int32 LODIndex = 0;
 	int32 NumLODInfo = 1;
 
@@ -317,7 +317,7 @@ void FSkeletalMeshSkinningData::UnregisterUser(FSkeletalMeshSkinningDataUsage Us
 	int32 LODIndex = 0;
 
 	USkeletalMeshComponent* SkelComp = MeshComp.Get();
-	if (SkelComp && SkelComp->SkeletalMesh)
+	if (SkelComp && SkelComp->GetSkeletalMesh())
 	{
 		LODIndex = Usage.GetLODIndex();
 	}
@@ -359,7 +359,7 @@ void FSkeletalMeshSkinningData::UpdateBoneTransforms()
 	USkeletalMeshComponent* SkelComp = MeshComp.Get();
 	check(SkelComp);
 
-	const USkeletalMesh* SkelMesh = SkelComp->SkeletalMesh;
+	const USkeletalMesh* SkelMesh = SkelComp->GetSkeletalMesh();
 	if (SkelMesh == nullptr)
 	{
 		return;
@@ -454,9 +454,9 @@ bool FSkeletalMeshSkinningData::Tick(float InDeltaSeconds, bool bRequirePreskin)
 		PrevComponentTransforms() = CurrComponentTransforms();
 	}
 
-	if (bRequirePreskin && SkelComp->SkeletalMesh != nullptr)
+	if (bRequirePreskin && SkelComp->GetSkeletalMesh() != nullptr)
 	{
-		const USkeletalMesh* SkeletalMesh = SkelComp->SkeletalMesh;
+		const USkeletalMesh* SkeletalMesh = SkelComp->GetSkeletalMesh();
 		const FSkeletalMeshRenderData* RenderData = SkeletalMesh->GetResourceForRendering();
 		check(RenderData);
 
@@ -1099,7 +1099,7 @@ void FSkeletalMeshGpuDynamicBufferProxy::NewFrame(const FNDISkeletalMesh_Instanc
 		SkelComp = Cast<USkeletalMeshComponent>(InstanceData->SceneComponent.Get());
 		if ( SkelComp != nullptr )
 		{
-			SkelMesh = SkelComp->SkeletalMesh;
+			SkelMesh = SkelComp->GetSkeletalMesh();
 		}
 		if (SkelMesh == nullptr)
 		{
@@ -1741,7 +1741,7 @@ USkeletalMesh* UNiagaraDataInterfaceSkeletalMesh::GetSkeletalMesh(FNiagaraSystem
 			for (UActorComponent* ActorComp : Actor->GetComponents())
 			{
 				USkeletalMeshComponent* Comp = Cast<USkeletalMeshComponent>(ActorComp);
-				if (IsValid(Comp) && Comp->SkeletalMesh != nullptr)
+				if (IsValid(Comp) && Comp->GetSkeletalMesh() != nullptr)
 				{
 					return Comp;
 				}
@@ -1837,7 +1837,7 @@ USkeletalMesh* UNiagaraDataInterfaceSkeletalMesh::GetSkeletalMesh(FNiagaraSystem
 	SceneComponent = nullptr;
 	if (FoundSkelComp)
 	{
-		Mesh = FoundSkelComp->SkeletalMesh;
+		Mesh = FoundSkelComp->GetSkeletalMesh();
 		SceneComponent = FoundSkelComp;
 	}
 #if WITH_EDITORONLY_DATA
@@ -2365,7 +2365,7 @@ bool FNDISkeletalMesh_InstanceData::ResetRequired(UNiagaraDataInterfaceSkeletalM
 	// Reset if the skeletal mesh on the cached skeletal mesh component changed.
 	if (USkeletalMeshComponent* SkelComp = Cast<USkeletalMeshComponent>(Comp))
 	{
-		if (SkelComp->SkeletalMesh != SkelMesh)
+		if (SkelComp->GetSkeletalMesh() != SkelMesh)
 		{
 			if (SkinningData.SkinningData.IsValid())
 			{
diff --git a/Engine/Plugins/FX/Niagara/Source/Niagara/Private/NiagaraFunctionLibrary.cpp b/Engine/Plugins/FX/Niagara/Source/Niagara/Private/NiagaraFunctionLibrary.cpp
--- a/Engine/Plugins/FX/Niagara/Source/Niagara/Private/NiagaraFunctionLibrary.cpp
+++ b/Engine/Plugins/FX/Niagara/Source/Niagara/Private/NiagaraFunctionLibrary.cpp
@@ -428,7 +428,7 @@ void UNiagaraFunctionLibrary::OverrideSystemUserVariableSkeletalMeshComponent(UN
 	{
 		if (FNiagaraUtilities::LogVerboseWarnings())
 		{
-			UE_LOG(LogNiagara, Warning, TEXT("NiagaraSystem in \"Set Niagara Skeletal Mesh Component\" is NULL, OverrideName \"%s\" and SkeletalMeshComponent \"%s\" SkeletalMesh \"%s\", skipping."), *OverrideName, *GetFullNameSafe(SkeletalMeshComponent), *GetFullNameSafe(SkeletalMeshComponent ? SkeletalMeshComponent->SkeletalMesh : nullptr));
+			UE_LOG(LogNiagara, Warning, TEXT("NiagaraSystem in \"Set Niagara Skeletal Mesh Component\" is NULL, OverrideName \"%s\" and SkeletalMeshComponent \"%s\" SkeletalMesh \"%s\", skipping."), *OverrideName, *GetFullNameSafe(SkeletalMeshComponent), *GetFullNameSafe(SkeletalMeshComponent ? SkeletalMeshComponent->GetSkeletalMesh() : nullptr));
 		}
 		return;
 	}
diff --git a/Engine/Plugins/Importers/USDImporter/Source/USDSchemas/Private/USDSkelRootTranslator.cpp b/Engine/Plugins/Importers/USDImporter/Source/USDSchemas/Private/USDSkelRootTranslator.cpp
--- a/Engine/Plugins/Importers/USDImporter/Source/USDSchemas/Private/USDSkelRootTranslator.cpp
+++ b/Engine/Plugins/Importers/USDImporter/Source/USDSchemas/Private/USDSkelRootTranslator.cpp
@@ -276,17 +276,17 @@ namespace UsdSkelRootTranslatorImpl
 
 	void SetMorphTargetWeight( USkeletalMeshComponent& SkeletalMeshComponent, const FString& MorphTargetName, float Weight )
 	{
-		USkeletalMesh* SkeletalMesh = SkeletalMeshComponent.SkeletalMesh;
+		USkeletalMesh* SkeletalMesh = SkeletalMeshComponent.GetSkeletalMesh();
 
 		// We try keeping a perfect correspondence between SkeletalMesh->GetMorphTargets() and SkeletalMeshComponent.ActiveMorphTargets
 		int32 IndexInSkeletalMesh = INDEX_NONE;
-		SkeletalMeshComponent.SkeletalMesh->FindMorphTargetAndIndex( *MorphTargetName, IndexInSkeletalMesh );
+		SkeletalMeshComponent.GetSkeletalMesh()->FindMorphTargetAndIndex( *MorphTargetName, IndexInSkeletalMesh );
 		if ( IndexInSkeletalMesh == INDEX_NONE )
 		{
 			return;
 		}
 
-		UMorphTarget* MorphTarget = SkeletalMeshComponent.SkeletalMesh->GetMorphTargets()[ IndexInSkeletalMesh ];
+		UMorphTarget* MorphTarget = SkeletalMeshComponent.GetSkeletalMesh()->GetMorphTargets()[ IndexInSkeletalMesh ];
 		if ( !MorphTarget )
 		{
 			return;
@@ -958,14 +958,14 @@ namespace UsdSkelRootTranslatorImpl
 
 	void UpdateLiveLinkProperties( const FUsdSchemaTranslationContext& Context, USkeletalMeshComponent* Component, const pxr::UsdPrim& Prim )
 	{
-		if ( !Component || !Component->SkeletalMesh || !Prim )
+		if ( !Component || !Component->GetSkeletalMesh() || !Prim )
 		{
 			return;
 		}
 
 		FString PrimName = UsdToUnreal::ConvertString( Prim.GetName() );
 
-		USkeleton* Skeleton = Component->SkeletalMesh->GetSkeleton();
+		USkeleton* Skeleton = Component->GetSkeletalMesh()->GetSkeleton();
 		if ( !Skeleton )
 		{
 			return;
@@ -1271,7 +1271,7 @@ void FUsdSkelRootTranslator::UpdateComponents( USceneComponent* SceneComponent )
 #if WITH_EDITOR
 	// Re-set the skeletal mesh if we created a new one (maybe the hash changed, a skinned UsdGeomMesh was hidden, etc.)
 	USkeletalMesh* TargetSkeletalMesh = Cast< USkeletalMesh >( Context->AssetCache->GetAssetForPrim( PrimPath.GetString() ) );
-	if ( SkeletalMeshComponent->SkeletalMesh != TargetSkeletalMesh )
+	if ( SkeletalMeshComponent->GetSkeletalMesh() != TargetSkeletalMesh )
 	{
 		SkeletalMeshComponent->SetSkeletalMesh(TargetSkeletalMesh);
 
@@ -1305,7 +1305,7 @@ void FUsdSkelRootTranslator::UpdateComponents( USceneComponent* SceneComponent )
 	// Update the animation state
 	// Don't try animating ourselves if the sequencer is animating as it will just overwrite the animation state on next
 	// tick anyway, and all this would do is lead to flickering and other issues
-	if ( !Context->bSequencerIsAnimating && SkeletalMeshComponent->SkeletalMesh && !bPrimHasLiveLinkEnabled )
+	if ( !Context->bSequencerIsAnimating && SkeletalMeshComponent->GetSkeletalMesh() && !bPrimHasLiveLinkEnabled )
 	{
 		if ( UAnimSequence* AnimSequence = Cast<UAnimSequence>( SkeletalMeshComponent->AnimationData.AnimToPlay.Get() ) )
 		{
diff --git a/Engine/Plugins/Importers/USDImporter/Source/USDStage/Private/USDLevelSequenceHelper.cpp b/Engine/Plugins/Importers/USDImporter/Source/USDStage/Private/USDLevelSequenceHelper.cpp
--- a/Engine/Plugins/Importers/USDImporter/Source/USDStage/Private/USDLevelSequenceHelper.cpp
+++ b/Engine/Plugins/Importers/USDImporter/Source/USDStage/Private/USDLevelSequenceHelper.cpp
@@ -254,7 +254,7 @@ namespace UsdLevelSequenceHelperImpl
 	)
 	{
 		UMovieScene* MovieScene = LevelSequence->GetMovieScene();
-		if ( !MovieScene || !SkeletalMeshComp || !SkeletalMeshComp->SkeletalMesh || !SkeletalMeshComp->SkeletalMesh->GetSkeleton() )
+		if ( !MovieScene || !SkeletalMeshComp || !SkeletalMeshComp->GetSkeletalMesh() || !SkeletalMeshComp->GetSkeletalMesh()->GetSkeleton())
 		{
 			return false;
 		}
@@ -293,7 +293,7 @@ namespace UsdLevelSequenceHelperImpl
 		{
 			bCreatedTempSequence = true;
 			AnimSequence = NewObject<UAnimSequence>();
-			AnimSequence->SetSkeleton( SkeletalMeshComp->SkeletalMesh->GetSkeleton() );
+			AnimSequence->SetSkeleton( SkeletalMeshComp->GetSkeletalMesh()->GetSkeleton());
 
 			FMovieSceneSequenceIDRef Template = MovieSceneSequenceID::Root;
 			FMovieSceneSequenceTransform RootToLocalTransform;
@@ -2923,7 +2923,7 @@ void FUsdLevelSequenceHelperImpl::HandleControlRigSectionChange( UMovieSceneCont
 		return;
 	}
 
-	USkeleton* Skeleton = BoundComponent->SkeletalMesh ? BoundComponent->SkeletalMesh->GetSkeleton() : nullptr;
+	USkeleton* Skeleton = BoundComponent->GetSkeletalMesh() ? BoundComponent->GetSkeletalMesh()->GetSkeleton() : nullptr;
 	if ( !Skeleton )
 	{
 		return;
diff --git a/Engine/Plugins/Importers/USDImporter/Source/USDUtilities/Private/USDPrimConversion.cpp b/Engine/Plugins/Importers/USDImporter/Source/USDUtilities/Private/USDPrimConversion.cpp
--- a/Engine/Plugins/Importers/USDImporter/Source/USDUtilities/Private/USDPrimConversion.cpp
+++ b/Engine/Plugins/Importers/USDImporter/Source/USDUtilities/Private/USDPrimConversion.cpp
@@ -2227,7 +2227,7 @@ bool UnrealToUsd::ConvertMeshComponent( const pxr::UsdStageRefPtr& Stage, const
 			return false;
 		}
 
-		if ( const USkeletalMesh* SkeletalMesh = SkinnedMeshComponent->SkeletalMesh )
+		if ( const USkeletalMesh* SkeletalMesh = SkinnedMeshComponent->GetSkeletalMesh())
 		{
 			FSkeletalMeshRenderData* RenderData = SkeletalMesh->GetResourceForRendering();
 			if ( !RenderData )
@@ -2856,7 +2856,7 @@ bool UnrealToUsd::CreateComponentPropertyBaker( UE::FUsdPrim& Prim, const UScene
 bool UnrealToUsd::CreateSkeletalAnimationBaker( UE::FUsdPrim& SkelRoot, UE::FUsdPrim& SkelAnimation, USkeletalMeshComponent& Component, FComponentBaker& OutBaker )
 {
 #if WITH_EDITOR
-	USkeletalMesh* SkeletalMesh = Component.SkeletalMesh;
+	USkeletalMesh* SkeletalMesh = Component.GetSkeletalMesh();
 	if ( !SkeletalMesh )
 	{
 		return false;
diff --git a/Engine/Plugins/MeshPainting/Source/MeshPaintEditorMode/Private/MeshPaintMode.cpp b/Engine/Plugins/MeshPainting/Source/MeshPaintEditorMode/Private/MeshPaintMode.cpp
--- a/Engine/Plugins/MeshPainting/Source/MeshPaintEditorMode/Private/MeshPaintMode.cpp
+++ b/Engine/Plugins/MeshPainting/Source/MeshPaintEditorMode/Private/MeshPaintMode.cpp
@@ -581,9 +581,9 @@ void UMeshPaintMode::SavePaintedAssets()
 
 	for (USkeletalMeshComponent* SkeletalMeshComponent : SkeletalMeshComponents)
 	{
-		if (SkeletalMeshComponent && SkeletalMeshComponent->SkeletalMesh)
+		if (SkeletalMeshComponent && SkeletalMeshComponent->GetSkeletalMesh())
 		{
-			ObjectsToSave.Add(SkeletalMeshComponent->SkeletalMesh);
+			ObjectsToSave.Add(SkeletalMeshComponent->GetSkeletalMesh());
 		}
 	}
 
@@ -609,7 +609,7 @@ bool UMeshPaintMode::CanSaveMeshPackages() const
 		}
 		else if (USkeletalMeshComponent* SkeletalMeshComponent = Cast<USkeletalMeshComponent>(Component))
 		{
-			Object = SkeletalMeshComponent->SkeletalMesh;
+			Object = SkeletalMeshComponent->GetSkeletalMesh();
 		}
 
 		if (Object != nullptr && Object->GetOutermost()->IsDirty())
diff --git a/Engine/Plugins/MeshPainting/Source/MeshPaintEditorMode/Private/MeshPaintModeHelpers.cpp b/Engine/Plugins/MeshPainting/Source/MeshPaintEditorMode/Private/MeshPaintModeHelpers.cpp
--- a/Engine/Plugins/MeshPainting/Source/MeshPaintEditorMode/Private/MeshPaintModeHelpers.cpp
+++ b/Engine/Plugins/MeshPainting/Source/MeshPaintEditorMode/Private/MeshPaintModeHelpers.cpp
@@ -232,10 +232,10 @@ void UMeshPaintModeSubsystem::ImportVertexColorsFromTexture(UMeshComponent* Mesh
 				{
 					USkeletalMeshComponent* SkeletalMeshComponent = Cast<USkeletalMeshComponent>(MeshComponent);
 
-					if (SkeletalMeshComponent->SkeletalMesh)
+					if (SkeletalMeshComponent->GetSkeletalMesh())
 					{
 						// Import colors to skeletal mesh
-						ImportVertexColorsToSkeletalMesh(SkeletalMeshComponent->SkeletalMesh, Options, ColorTexture);
+						ImportVertexColorsToSkeletalMesh(SkeletalMeshComponent->GetSkeletalMesh(), Options, ColorTexture);
 					}
 				}
 			}
diff --git a/Engine/Plugins/MeshPainting/Source/MeshPaintingToolset/Private/MeshPaintHelpers.cpp b/Engine/Plugins/MeshPainting/Source/MeshPaintingToolset/Private/MeshPaintHelpers.cpp
--- a/Engine/Plugins/MeshPainting/Source/MeshPaintingToolset/Private/MeshPaintHelpers.cpp
+++ b/Engine/Plugins/MeshPainting/Source/MeshPaintingToolset/Private/MeshPaintHelpers.cpp
@@ -784,7 +784,7 @@ void UMeshPaintingSubsystem::FillStaticMeshVertexColors(UStaticMeshComponent* Me
 void UMeshPaintingSubsystem::FillSkeletalMeshVertexColors(USkeletalMeshComponent* MeshComponent, int32 LODIndex, const FColor FillColor, const FColor MaskColor)
 {
 	TUniquePtr< FSkinnedMeshComponentRecreateRenderStateContext > RecreateRenderStateContext;
-	USkeletalMesh* Mesh = MeshComponent->SkeletalMesh;
+	USkeletalMesh* Mesh = MeshComponent->GetSkeletalMesh();
 	if (Mesh)
 	{
 		// Dirty the mesh
@@ -1005,7 +1005,7 @@ bool UMeshPaintingSubsystem::TryGetNumberOfLODs(const UMeshComponent* MeshCompon
 	}
 	else if (const USkeletalMeshComponent* SkeletalMeshComponent = Cast<USkeletalMeshComponent>(MeshComponent))
 	{
-		const USkeletalMesh* SkeletalMesh = SkeletalMeshComponent->SkeletalMesh;
+		const USkeletalMesh* SkeletalMesh = SkeletalMeshComponent->GetSkeletalMesh();
 		if (SkeletalMesh != nullptr)
 		{
 			OutNumLODs = SkeletalMesh->GetLODNum();
@@ -1037,7 +1037,7 @@ int32 UMeshPaintingSubsystem::GetNumberOfUVs(const UMeshComponent* MeshComponent
 	}
 	else if (const USkeletalMeshComponent* SkeletalMeshComponent = Cast<USkeletalMeshComponent>(MeshComponent))
 	{
-		const USkeletalMesh* SkeletalMesh = SkeletalMeshComponent->SkeletalMesh;
+		const USkeletalMesh* SkeletalMesh = SkeletalMeshComponent->GetSkeletalMesh();
 		if (SkeletalMesh != nullptr && SkeletalMesh->GetResourceForRendering() && SkeletalMesh->GetResourceForRendering()->LODRenderData.IsValidIndex(LODIndex))
 		{
 			NumUVs = SkeletalMesh->GetResourceForRendering()->LODRenderData[LODIndex].GetNumTexCoords();
@@ -1075,7 +1075,7 @@ bool UMeshPaintingSubsystem::DoesMeshComponentContainPerLODColors(const UMeshCom
 	}
 	else if (const USkeletalMeshComponent* SkeletalMeshComponent = Cast<USkeletalMeshComponent>(MeshComponent))
 	{
-		USkeletalMesh* SkeletalMesh = SkeletalMeshComponent->SkeletalMesh;
+		USkeletalMesh* SkeletalMesh = SkeletalMeshComponent->GetSkeletalMesh();
 		if (SkeletalMesh)
 		{
 			const TArray<FSkeletalMeshLODInfo>& LODInfo = SkeletalMesh->GetLODInfoArray();
@@ -1266,7 +1266,7 @@ typedef TOctree2<FPaintedMeshVertex, FVertexColorPropogationOctreeSemantics> TVe
 void UMeshPaintingSubsystem::ApplyVertexColorsToAllLODs(IMeshPaintComponentAdapter& GeometryInfo, USkeletalMeshComponent* SkeletalMeshComponent)
 {
 	checkf(SkeletalMeshComponent != nullptr, TEXT("Invalid Skeletal Mesh Component"));
-	USkeletalMesh* Mesh = SkeletalMeshComponent->SkeletalMesh;
+	USkeletalMesh* Mesh = SkeletalMeshComponent->GetSkeletalMesh();
 	if (Mesh)
 	{
 		FSkeletalMeshRenderData* Resource = Mesh->GetResourceForRendering();
diff --git a/Engine/Plugins/MeshPainting/Source/MeshPaintingToolset/Private/MeshPaintSkeletalMeshAdapter.cpp b/Engine/Plugins/MeshPainting/Source/MeshPaintingToolset/Private/MeshPaintSkeletalMeshAdapter.cpp
--- a/Engine/Plugins/MeshPainting/Source/MeshPaintingToolset/Private/MeshPaintSkeletalMeshAdapter.cpp
+++ b/Engine/Plugins/MeshPainting/Source/MeshPaintingToolset/Private/MeshPaintSkeletalMeshAdapter.cpp
@@ -110,9 +110,9 @@ bool FMeshPaintSkeletalMeshComponentAdapter::Construct(UMeshComponent* InCompone
 	{
 		SkeletalMeshChangedHandle = SkeletalMeshComponent->RegisterOnSkeletalMeshPropertyChanged(USkeletalMeshComponent::FOnSkeletalMeshPropertyChanged::CreateSP(this, &FMeshPaintSkeletalMeshComponentAdapter::OnSkeletalMeshChanged));
 
-		if (SkeletalMeshComponent->SkeletalMesh != nullptr)
+		if (SkeletalMeshComponent->GetSkeletalMesh() != nullptr)
 		{
-			ReferencedSkeletalMesh = SkeletalMeshComponent->SkeletalMesh;
+			ReferencedSkeletalMesh = SkeletalMeshComponent->GetSkeletalMesh();
 			MeshLODIndex = InMeshLODIndex;
 			const bool bSuccess = Initialize();
 			return bSuccess;
@@ -139,8 +139,8 @@ void FMeshPaintSkeletalMeshComponentAdapter::OnSkeletalMeshChanged()
 	OnRemoved();
 	if (SkeletalMeshComponent.IsValid())
 	{
-		ReferencedSkeletalMesh = SkeletalMeshComponent->SkeletalMesh;
-		if (SkeletalMeshComponent->SkeletalMesh != nullptr)
+		ReferencedSkeletalMesh = SkeletalMeshComponent->GetSkeletalMesh();
+		if (SkeletalMeshComponent->GetSkeletalMesh() != nullptr)
 		{
 			Initialize();
 			OnAdded();
@@ -162,7 +162,7 @@ bool FMeshPaintSkeletalMeshComponentAdapter::Initialize()
 
 	if (SkeletalMeshComponent.IsValid())
 	{
-		check(ReferencedSkeletalMesh == SkeletalMeshComponent->SkeletalMesh);
+		check(ReferencedSkeletalMesh == SkeletalMeshComponent->GetSkeletalMesh());
 
 		MeshResource = ReferencedSkeletalMesh->GetResourceForRendering();
 		if (MeshResource != nullptr)
@@ -229,7 +229,7 @@ void FMeshPaintSkeletalMeshComponentAdapter::OnAdded()
 	}
 
 	checkf(ReferencedSkeletalMesh, TEXT("Invalid reference to Skeletal Mesh"));
-	checkf(ReferencedSkeletalMesh == SkeletalMeshComponent->SkeletalMesh, TEXT("Referenced Skeletal Mesh does not match one in Component"));
+	checkf(ReferencedSkeletalMesh == SkeletalMeshComponent->GetSkeletalMesh(), TEXT("Referenced Skeletal Mesh does not match one in Component"));
 
 	SkeletalMeshComponent->bUseRefPoseOnInitAnim = true;
 	SkeletalMeshComponent->InitAnim(true);
@@ -465,7 +465,7 @@ TSharedPtr<IMeshPaintComponentAdapter> FMeshPaintSkeletalMeshComponentAdapterFac
 {
 	if (USkeletalMeshComponent* SkeletalMeshComponent = Cast<USkeletalMeshComponent>(InComponent))
 	{
-		if (SkeletalMeshComponent->SkeletalMesh != nullptr)
+		if (SkeletalMeshComponent->GetSkeletalMesh() != nullptr)
 		{
 			TSharedRef<FMeshPaintSkeletalMeshComponentAdapter> Result = MakeShareable(new FMeshPaintSkeletalMeshComponentAdapter());
 			if (Result->Construct(InComponent, InMeshLODIndex))
diff --git a/Engine/Plugins/MeshPainting/Source/MeshPaintingToolset/Private/TexturePaintToolset.cpp b/Engine/Plugins/MeshPainting/Source/MeshPaintingToolset/Private/TexturePaintToolset.cpp
--- a/Engine/Plugins/MeshPainting/Source/MeshPaintingToolset/Private/TexturePaintToolset.cpp
+++ b/Engine/Plugins/MeshPainting/Source/MeshPaintingToolset/Private/TexturePaintToolset.cpp
@@ -450,7 +450,7 @@ void UTexturePaintToolset::RetrieveMeshSectionsForMaterialIndices(const UMeshCom
 	}
 	else if (const USkeletalMeshComponent* SkeletalMeshComponent = Cast<const USkeletalMeshComponent>(MeshComponent))
 	{
-		const USkeletalMesh* SkeletalMesh = SkeletalMeshComponent->SkeletalMesh;
+		const USkeletalMesh* SkeletalMesh = SkeletalMeshComponent->GetSkeletalMesh();
 		if (SkeletalMesh)
 		{
 			const FSkeletalMeshRenderData* Resource = SkeletalMesh->GetResourceForRendering();
diff --git a/Engine/Plugins/MeshPainting/Source/MeshPaintingToolset/Public/MeshPaintSkeletalMeshAdapter.h b/Engine/Plugins/MeshPainting/Source/MeshPaintingToolset/Public/MeshPaintSkeletalMeshAdapter.h
--- a/Engine/Plugins/MeshPainting/Source/MeshPaintingToolset/Public/MeshPaintSkeletalMeshAdapter.h
+++ b/Engine/Plugins/MeshPainting/Source/MeshPaintingToolset/Public/MeshPaintSkeletalMeshAdapter.h
@@ -31,7 +31,7 @@ public:
 	virtual bool Initialize() override;
 	virtual void OnAdded() override;
 	virtual void OnRemoved() override;
-	virtual bool IsValid() const override { return SkeletalMeshComponent.IsValid() && ReferencedSkeletalMesh && SkeletalMeshComponent->SkeletalMesh == ReferencedSkeletalMesh; }
+	virtual bool IsValid() const override { return SkeletalMeshComponent.IsValid() && ReferencedSkeletalMesh && SkeletalMeshComponent->GetSkeletalMesh() == ReferencedSkeletalMesh; }
 	virtual bool SupportsTexturePaint() const override { return true; }
 	virtual bool SupportsVertexPaint() const override { return SkeletalMeshComponent.IsValid(); }
 	virtual bool LineTraceComponent(struct FHitResult& OutHit, const FVector Start, const FVector End, const struct FCollisionQueryParams& Params) const override;
diff --git a/Engine/Plugins/MovieScene/SequencerScripting/Source/SequencerScriptingEditor/Private/SequencerTools.cpp b/Engine/Plugins/MovieScene/SequencerScripting/Source/SequencerScriptingEditor/Private/SequencerTools.cpp
--- a/Engine/Plugins/MovieScene/SequencerScripting/Source/SequencerScriptingEditor/Private/SequencerTools.cpp
+++ b/Engine/Plugins/MovieScene/SequencerScripting/Source/SequencerScriptingEditor/Private/SequencerTools.cpp
@@ -279,7 +279,7 @@ static USkeletalMeshComponent* GetSkelMeshComponent(IMovieScenePlayer* Player, c
 		}
 		else if (USkeletalMeshComponent* SkeletalMeshComponent = Cast<USkeletalMeshComponent>(RuntimeObject.Get()))
 		{
-			if (SkeletalMeshComponent->SkeletalMesh)
+			if (SkeletalMeshComponent->GetSkeletalMesh())
 			{
 				return SkeletalMeshComponent;
 			}
@@ -317,9 +317,9 @@ bool USequencerToolsFunctionLibrary::ExportAnimSequence(UWorld* World, ULevelSeq
 		}
  
 		USkeletalMeshComponent* SkeletalMeshComp =  GetSkelMeshComponent(Player, Binding);
-		if (SkeletalMeshComp && SkeletalMeshComp->SkeletalMesh && SkeletalMeshComp->SkeletalMesh->GetSkeleton())
+		if (SkeletalMeshComp && SkeletalMeshComp->GetSkeletalMesh() && SkeletalMeshComp->GetSkeletalMesh()->GetSkeleton())
 		{
-			AnimSequence->SetSkeleton(SkeletalMeshComp->SkeletalMesh->GetSkeleton());
+			AnimSequence->SetSkeleton(SkeletalMeshComp->GetSkeletalMesh()->GetSkeleton());
 			bResult = MovieSceneToolHelpers::ExportToAnimSequence(AnimSequence,ExportOptions, MovieScene, Player, SkeletalMeshComp, Template, RootToLocalTransform);
 		}
 	}
diff --git a/Engine/Plugins/Runtime/ApexDestruction/Source/ApexDestruction/Private/DestructibleActor.cpp b/Engine/Plugins/Runtime/ApexDestruction/Source/ApexDestruction/Private/DestructibleActor.cpp
--- a/Engine/Plugins/Runtime/ApexDestruction/Source/ApexDestruction/Private/DestructibleActor.cpp
+++ b/Engine/Plugins/Runtime/ApexDestruction/Source/ApexDestruction/Private/DestructibleActor.cpp
@@ -22,9 +22,9 @@ bool ADestructibleActor::GetReferencedContentObjects( TArray<UObject*>& Objects
 {
 	Super::GetReferencedContentObjects(Objects);
 
-	if (DestructibleComponent && DestructibleComponent->SkeletalMesh)
+	if (DestructibleComponent && DestructibleComponent->GetSkeletalMesh())
 	{
-		Objects.Add(DestructibleComponent->SkeletalMesh);
+		Objects.Add(DestructibleComponent->GetSkeletalMesh());
 	}
 	return true;
 }
diff --git a/Engine/Plugins/Runtime/ApexDestruction/Source/ApexDestruction/Private/DestructibleComponent.cpp b/Engine/Plugins/Runtime/ApexDestruction/Source/ApexDestruction/Private/DestructibleComponent.cpp
--- a/Engine/Plugins/Runtime/ApexDestruction/Source/ApexDestruction/Private/DestructibleComponent.cpp
+++ b/Engine/Plugins/Runtime/ApexDestruction/Source/ApexDestruction/Private/DestructibleComponent.cpp
@@ -122,7 +122,7 @@ void UDestructibleComponent::OnUpdateTransform(EUpdateTransformFlags UpdateTrans
 	// We are handling the physics move below, so don't handle it at higher levels
 	Super::OnUpdateTransform(UpdateTransformFlags | EUpdateTransformFlags::SkipPhysicsUpdate, Teleport);
 
-	if (SkeletalMesh == NULL)
+	if (GetSkeletalMesh() == NULL)
 	{
 		return;
 	}
@@ -147,7 +147,7 @@ void UDestructibleComponent::OnUpdateTransform(EUpdateTransformFlags UpdateTrans
 #if !(UE_BUILD_SHIPPING || UE_BUILD_TEST)
 	if( !MeshScale3D.IsUniform() )
 	{
-		UE_LOG(LogPhysics, Log, TEXT("UDestructibleComponent::SendPhysicsTransform : Non-uniform scale factor (%s) can cause physics to mismatch for %s  SkelMesh: %s"), *MeshScale3D.ToString(), *GetFullName(), SkeletalMesh ? *SkeletalMesh->GetFullName() : TEXT("NULL"));
+		UE_LOG(LogPhysics, Log, TEXT("UDestructibleComponent::SendPhysicsTransform : Non-uniform scale factor (%s) can cause physics to mismatch for %s  SkelMesh: %s"), *MeshScale3D.ToString(), *GetFullName(), GetSkeletalMesh() ? *GetSkeletalMesh()->GetFullName() : TEXT("NULL"));
 	}
 #endif
 }
@@ -173,7 +173,7 @@ void UDestructibleComponent::OnDestroyPhysicsState()
 
 UBodySetup* UDestructibleComponent::GetBodySetup()
 {
-	if (SkeletalMesh != NULL)
+	if (GetSkeletalMesh() != NULL)
 	{
 		UDestructibleMesh* TheDestructibleMesh = GetDestructibleMesh();
 
@@ -269,7 +269,7 @@ void UDestructibleComponent::SetDestructibleMesh(class UDestructibleMesh* NewMes
 
 class UDestructibleMesh* UDestructibleComponent::GetDestructibleMesh()
 {
-	return Cast<UDestructibleMesh>(SkeletalMesh);
+	return Cast<UDestructibleMesh>(GetSkeletalMesh());
 }
 
 void UDestructibleComponent::SetSkeletalMesh(USkeletalMesh* InSkelMesh, bool bReinitPose)
@@ -506,12 +506,12 @@ void UDestructibleComponent::SetMaterial(int32 ElementIndex, UMaterialInterface*
 	}
 	
 	// Update physical properties for individual bone instances as well
-	if (SkeletalMesh)
+	if (GetSkeletalMesh())
 	{
-		int32 NumBones = SkeletalMesh->GetRefSkeleton().GetRawBoneNum();
+		int32 NumBones = GetSkeletalMesh()->GetRefSkeleton().GetRawBoneNum();
 		for (int32 BoneIdx = 0; BoneIdx < NumBones; ++BoneIdx)
 		{
-			FName BoneName = SkeletalMesh->GetRefSkeleton().GetBoneName(BoneIdx);
+			FName BoneName = GetSkeletalMesh()->GetRefSkeleton().GetBoneName(BoneIdx);
 			FBodyInstance* Instance = GetBodyInstance(BoneName);
 			if (Instance)
 			{
diff --git a/Engine/Plugins/Runtime/HairStrands/Source/HairStrandsCore/Private/GroomComponent.cpp b/Engine/Plugins/Runtime/HairStrands/Source/HairStrandsCore/Private/GroomComponent.cpp
--- a/Engine/Plugins/Runtime/HairStrands/Source/HairStrandsCore/Private/GroomComponent.cpp
+++ b/Engine/Plugins/Runtime/HairStrands/Source/HairStrandsCore/Private/GroomComponent.cpp
@@ -2250,9 +2250,9 @@ static USkeletalMeshComponent* ValidateBindingAsset(
 
 	// Optional advanced check
 	bool bHasValidSectionCount = SkeletalMeshComponent && SkeletalMeshComponent->GetNumMaterials() < int32(GetHairStrandsMaxSectionCount());
-	if (SkeletalMeshComponent && SkeletalMeshComponent->SkeletalMesh && SkeletalMeshComponent->SkeletalMesh->GetResourceForRendering())
+	if (SkeletalMeshComponent && SkeletalMeshComponent->GetSkeletalMesh() && SkeletalMeshComponent->GetSkeletalMesh()->GetResourceForRendering())
 	{
-		const FSkeletalMeshRenderData* RenderData = SkeletalMeshComponent->SkeletalMesh->GetResourceForRendering();
+		const FSkeletalMeshRenderData* RenderData = SkeletalMeshComponent->GetSkeletalMesh()->GetResourceForRendering();
 		const int32 MaxSectionCount = GetHairStrandsMaxSectionCount();
 		// Check that all LOD are below the number sections
 		const uint32 MeshLODCount = RenderData->LODRenderData.Num();
@@ -2279,7 +2279,7 @@ static USkeletalMeshComponent* ValidateBindingAsset(
 	}
 
 	const bool bIsBindingCompatible =
-		UGroomBindingAsset::IsCompatible(SkeletalMeshComponent ? SkeletalMeshComponent->SkeletalMesh : nullptr, BindingAsset, bValidationEnable) &&
+		UGroomBindingAsset::IsCompatible(SkeletalMeshComponent ? SkeletalMeshComponent->GetSkeletalMesh() : nullptr, BindingAsset, bValidationEnable) &&
 		UGroomBindingAsset::IsCompatible(GroomAsset, BindingAsset, bValidationEnable) &&
 		UGroomBindingAsset::IsBindingAssetValid(BindingAsset, bIsBindingReloading, bValidationEnable);
 
@@ -2303,7 +2303,7 @@ static USkeletalMeshComponent* ValidateBindingAsset(
 					GroupIt < BindingAsset->HairGroupResources.Num() &&
 					CardsLODIndex < uint32(BindingAsset->HairGroupResources[GroupIt].CardsRootResources.Num()) &&
 					BindingAsset->HairGroupResources[GroupIt].CardsRootResources[CardsLODIndex] != nullptr &&
-					((SkeletalMeshComponent && SkeletalMeshComponent->SkeletalMesh) ? SkeletalMeshComponent->SkeletalMesh->GetLODInfoArray().Num() == BindingAsset->HairGroupResources[GroupIt].CardsRootResources[CardsLODIndex]->BulkData.MeshProjectionLODs.Num() : false);
+					((SkeletalMeshComponent && SkeletalMeshComponent->GetSkeletalMesh()) ? SkeletalMeshComponent->GetSkeletalMesh()->GetLODInfoArray().Num() == BindingAsset->HairGroupResources[GroupIt].CardsRootResources[CardsLODIndex]->BulkData.MeshProjectionLODs.Num() : false);
 
 				if (!bIsCardsBindingCompatible)
 				{
@@ -2426,7 +2426,7 @@ void UGroomComponent::InitResources(bool bIsBindingReloading)
 		if (ValidatedMeshComponent)
 		{
 			if (BindingAsset && 
-				((BindingAsset->GroomBindingType == EGroomBindingMeshType::SkeletalMesh && Cast<USkeletalMeshComponent>(ValidatedMeshComponent)->SkeletalMesh == nullptr) ||
+				((BindingAsset->GroomBindingType == EGroomBindingMeshType::SkeletalMesh && Cast<USkeletalMeshComponent>(ValidatedMeshComponent)->GetSkeletalMesh() == nullptr) ||
 				(BindingAsset->GroomBindingType == EGroomBindingMeshType::GeometryCache && Cast<UGeometryCacheComponent>(ValidatedMeshComponent)->GeometryCache == nullptr)))
 			{
 				ValidatedMeshComponent = nullptr;
@@ -2437,7 +2437,7 @@ void UGroomComponent::InitResources(bool bIsBindingReloading)
 	{
 		if (USkeletalMeshComponent* ParentSkelMeshComponent = Cast<USkeletalMeshComponent>(ParentMeshComponent))
 		{
-			ValidatedMeshComponent = ParentSkelMeshComponent->SkeletalMesh ? ParentMeshComponent : nullptr;
+			ValidatedMeshComponent = ParentSkelMeshComponent->GetSkeletalMesh() ? ParentMeshComponent : nullptr;
 		}
 	}
 
@@ -2451,7 +2451,7 @@ void UGroomComponent::InitResources(bool bIsBindingReloading)
 		for (USceneComponent* Curr = this; Curr; Curr = Curr->GetAttachParent())
 		{
 			USkeletalMeshComponent* SkelMeshComp = Cast<USkeletalMeshComponent>(Curr);
-			if (SkelMeshComp && SkelMeshComp->SkeletalMesh && SkelMeshComp->SkeletalMesh->GetSkeleton() == TargetSkeleton)
+			if (SkelMeshComp && SkelMeshComp->GetSkeletalMesh() && SkelMeshComp->GetSkeletalMesh()->GetSkeleton() == TargetSkeleton)
 			{
 				DeformedMeshComponent = SkelMeshComp;
 				break;
@@ -2635,7 +2635,7 @@ void UGroomComponent::InitResources(bool bIsBindingReloading)
 				check(GroupIt < LocalBindingAsset->HairGroupResources.Num());
 				if (USkeletalMeshComponent* SkeletalMeshComponent = Cast<USkeletalMeshComponent>(RegisteredMeshComponent))
 				{
-					check(SkeletalMeshComponent->SkeletalMesh ? SkeletalMeshComponent->SkeletalMesh->GetLODInfoArray().Num() == LocalBindingAsset->HairGroupResources[GroupIt].SimRootResources->BulkData.MeshProjectionLODs.Num() : false);
+					check(SkeletalMeshComponent->GetSkeletalMesh() ? SkeletalMeshComponent->GetSkeletalMesh()->GetLODInfoArray().Num() == LocalBindingAsset->HairGroupResources[GroupIt].SimRootResources->BulkData.MeshProjectionLODs.Num() : false);
 				}
 
 				HairGroupInstance->Guides.RestRootResource = LocalBindingAsset->HairGroupResources[GroupIt].SimRootResources;
@@ -2747,7 +2747,7 @@ void UGroomComponent::InitResources(bool bIsBindingReloading)
 				check(GroupIt < LocalBindingAsset->HairGroupResources.Num());
 				if (USkeletalMeshComponent* SkeletalMeshComponent = Cast<USkeletalMeshComponent>(RegisteredMeshComponent))
 				{
-					check(SkeletalMeshComponent->SkeletalMesh ? SkeletalMeshComponent->SkeletalMesh->GetLODInfoArray().Num() == LocalBindingAsset->HairGroupResources[GroupIt].RenRootResources->BulkData.MeshProjectionLODs.Num() : false);
+					check(SkeletalMeshComponent->GetSkeletalMesh() ? SkeletalMeshComponent->GetSkeletalMesh()->GetLODInfoArray().Num() == LocalBindingAsset->HairGroupResources[GroupIt].RenRootResources->BulkData.MeshProjectionLODs.Num() : false);
 				}
 
 				HairGroupInstance->Strands.RestRootResource = LocalBindingAsset->HairGroupResources[GroupIt].RenRootResources;
@@ -2851,7 +2851,7 @@ void UGroomComponent::InitResources(bool bIsBindingReloading)
 						check(GroupIt < LocalBindingAsset->HairGroupResources.Num());
 						if (USkeletalMeshComponent* SkeletalMeshComponent = Cast<USkeletalMeshComponent>(RegisteredMeshComponent))
 						{
-							check(SkeletalMeshComponent->SkeletalMesh ? SkeletalMeshComponent->SkeletalMesh->GetLODInfoArray().Num() == LocalBindingAsset->HairGroupResources[GroupIt].CardsRootResources[CardsLODIndex]->BulkData.MeshProjectionLODs.Num() : false);
+							check(SkeletalMeshComponent->GetSkeletalMesh() ? SkeletalMeshComponent->GetSkeletalMesh()->GetLODInfoArray().Num() == LocalBindingAsset->HairGroupResources[GroupIt].CardsRootResources[CardsLODIndex]->BulkData.MeshProjectionLODs.Num() : false);
 						}
 
 						InstanceLOD.Guides.RestRootResource = LocalBindingAsset->HairGroupResources[GroupIt].CardsRootResources[CardsLODIndex];
@@ -3289,9 +3289,9 @@ void UGroomComponent::BuildSimulationTransform(FTransform& SimulationTransform)
 	if (SimulationSettings.SimulationSetup.bLocalSimulation)
 	{
 		const USkeletalMeshComponent* SkeletelMeshComponent = Cast<const USkeletalMeshComponent>(RegisteredMeshComponent);
-		if (SkeletelMeshComponent && SkeletelMeshComponent->SkeletalMesh && !SimulationSettings.SimulationSetup.LocalBone.IsEmpty())
+		if (SkeletelMeshComponent && SkeletelMeshComponent->GetSkeletalMesh() && !SimulationSettings.SimulationSetup.LocalBone.IsEmpty())
 		{
-			const FReferenceSkeleton& RefSkeleton = SkeletelMeshComponent->SkeletalMesh->GetRefSkeleton();
+			const FReferenceSkeleton& RefSkeleton = SkeletelMeshComponent->GetSkeletalMesh()->GetRefSkeleton();
 			const FName BoneName(SimulationSettings.SimulationSetup.LocalBone);
 			const int32 BoneIndex = RefSkeleton.FindBoneIndex(BoneName);
 
@@ -3314,14 +3314,14 @@ void UGroomComponent::TickComponent(float DeltaTime, enum ELevelTick TickType, F
 	// When a groom binding and simulation are disabled, and the groom component is parented with a skeletal mesh, we can optionally 
 	// attach the groom to a particular socket/bone
 	USkeletalMeshComponent* SkeletelMeshComponent = Cast<USkeletalMeshComponent>(RegisteredMeshComponent);
-	if (SkeletelMeshComponent && SkeletelMeshComponent->SkeletalMesh && !AttachmentName.IsEmpty())
+	if (SkeletelMeshComponent && SkeletelMeshComponent->GetSkeletalMesh() && !AttachmentName.IsEmpty())
 	{
 		const FName BoneName(AttachmentName);
 		if (GetAttachSocketName() != BoneName)
 		{
 			AttachToComponent(SkeletelMeshComponent, FAttachmentTransformRules(EAttachmentRule::KeepRelative, false), BoneName);
 			const uint32 BoneIndex = SkeletelMeshComponent->GetBoneIndex(BoneName);
-			const FMatrix BoneTransformRaw = SkeletelMeshComponent->SkeletalMesh->GetComposedRefPoseMatrix(BoneIndex);
+			const FMatrix BoneTransformRaw = SkeletelMeshComponent->GetSkeletalMesh()->GetComposedRefPoseMatrix(BoneIndex);
 			const FVector BoneLocation = BoneTransformRaw.GetOrigin();
 			const FQuat BoneRotation = BoneTransformRaw.ToQuat();
 
diff --git a/Engine/Plugins/Runtime/HairStrands/Source/HairStrandsCore/Private/GroomGeometryCache.cpp b/Engine/Plugins/Runtime/HairStrands/Source/HairStrandsCore/Private/GroomGeometryCache.cpp
--- a/Engine/Plugins/Runtime/HairStrands/Source/HairStrandsCore/Private/GroomGeometryCache.cpp
+++ b/Engine/Plugins/Runtime/HairStrands/Source/HairStrandsCore/Private/GroomGeometryCache.cpp
@@ -77,7 +77,7 @@ static void BuildBoneMatrices(USkeletalMeshComponent* SkeletalMeshComponent, con
 {
 	if (SkeletalMeshComponent)
 	{
-		FSkeletalMeshRenderData* RenderData = SkeletalMeshComponent->SkeletalMesh->GetResourceForRendering();
+		FSkeletalMeshRenderData* RenderData = SkeletalMeshComponent->GetSkeletalMesh()->GetResourceForRendering();
 
 		const uint32 LODIndex = SkeletalMeshComponent->GetPredictedLODLevel();// RenderData->PendingFirstLODIdx;
 		FSkeletalMeshLODRenderData& LODData = RenderData->LODRenderData[LODIndex];
diff --git a/Engine/Plugins/Runtime/HairStrands/Source/HairStrandsCore/Private/Niagara/NiagaraDataInterfaceHairStrands.cpp b/Engine/Plugins/Runtime/HairStrands/Source/HairStrandsCore/Private/Niagara/NiagaraDataInterfaceHairStrands.cpp
--- a/Engine/Plugins/Runtime/HairStrands/Source/HairStrandsCore/Private/Niagara/NiagaraDataInterfaceHairStrands.cpp
+++ b/Engine/Plugins/Runtime/HairStrands/Source/HairStrandsCore/Private/Niagara/NiagaraDataInterfaceHairStrands.cpp
@@ -1029,9 +1029,9 @@ FORCEINLINE bool RequiresSimulationReset(FNiagaraSystemInstance* SystemInstance,
 			for (UActorComponent* ActorComp : RootActor->GetComponents())
 			{
 				USkeletalMeshComponent* SkelMeshComp = Cast<USkeletalMeshComponent>(ActorComp);
-				if (SkelMeshComp && SkelMeshComp->SkeletalMesh)
+				if (SkelMeshComp && SkelMeshComp->GetSkeletalMesh())
 				{
-					NewSkeletalMeshes += GetTypeHash(SkelMeshComp->SkeletalMesh->GetName());
+					NewSkeletalMeshes += GetTypeHash(SkelMeshComp->GetSkeletalMesh()->GetName());
 				}
 			}
 		}
