Change RootUV transfer from strand to cards, to use two root points insteand of one and interpolate root UVs in between.

This improves the root UV precision for large card.

diff --git a/Engine/Plugins/Runtime/HairStrands/Source/HairStrandsCore/Private/HairCardsBuilder.cpp b/Engine/Plugins/Runtime/HairStrands/Source/HairStrandsCore/Private/HairCardsBuilder.cpp
--- a/Engine/Plugins/Runtime/HairStrands/Source/HairStrandsCore/Private/HairCardsBuilder.cpp
+++ b/Engine/Plugins/Runtime/HairStrands/Source/HairStrandsCore/Private/HairCardsBuilder.cpp
@@ -2837,8 +2837,8 @@ namespace HairCards
 		OutGuides.StrandsCurves.SetNum(NumCards);
 		OutGuides.BoundingBox.Init();
 
-
 		InCards.CoordU.SetNum(InCards.Positions.Num());
+		InCards.LocalUVs.SetNum(InCards.Positions.Num());
 
 		const float GuideRadius = 0.01f;
 
@@ -2847,16 +2847,17 @@ namespace HairCards
 		{
 			const uint32 NumTriangles = InCards.IndexCounts[CardIt] / 3;
 			const uint32 VertexOffset = InCards.IndexOffsets[CardIt];
+			struct FIndexAndCoord { uint32 Index; float TexCoord; };
 			struct FSimilarUVVertices
 			{
 				float TexCoord = 0;
 				TArray<uint32> Indices;
-				TArray<uint32> AllIndices;
+				TArray<FIndexAndCoord> AllIndices; // Store vertex index and Tex.coord perpendicular to the principal axis (stored in TexCoord)
 			};
 			TArray<FSimilarUVVertices> SimilarVertexU;
 			TArray<FSimilarUVVertices> SimilarVertexV;
 
-			auto AddSimilarUV = [](TArray<FSimilarUVVertices>& In, uint32 Index, float TexCoord, float Threshold)
+			auto AddSimilarUV = [](TArray<FSimilarUVVertices>& In, uint32 Index, float TexCoord, float MinorTexCoord, float Threshold)
 			{
 				bool bFound = false;
 				for (int32 It = 0, Count = In.Num(); It < Count; ++It)
@@ -2865,7 +2866,7 @@ namespace HairCards
 					{
 						// We add only unique vertices per segment, so that the average position land in the center of the cards
 						In[It].Indices.AddUnique(Index);
-						In[It].AllIndices.Add(Index);
+						In[It].AllIndices.Add({Index, MinorTexCoord});
 						bFound = true;
 						break;
 					}
@@ -2876,12 +2877,12 @@ namespace HairCards
 					FSimilarUVVertices& SimilarUV = In.AddDefaulted_GetRef();
 					SimilarUV.TexCoord = TexCoord;
 					SimilarUV.Indices.Add(Index);
-					SimilarUV.AllIndices.Add(Index);
+					SimilarUV.AllIndices.Add({Index, MinorTexCoord});
 				}
 			};
 
-			// Iterate overa all the triangles of a cards, and find vertices which share either same U or same V. We add them to seperate lists. 
-			// We then use the heuristic that, the main axis will have more segments. This is what determine the principal axis of the cards
+			// Iterate over all triangles of a cards, and find vertices which share either same U or same V. We add them to separate lists. 
+			// We then use the following heuristic: the main axis will have more segments. This is what determine the principal axis of the cards
 			const float UVCoordTreshold = 1.f / 1024.f; // 1 pixel for a 1k texture
 			for (uint32 TriangleIt = 0; TriangleIt < NumTriangles; ++TriangleIt)
 			{
@@ -2890,8 +2891,8 @@ namespace HairCards
 				{
 					const uint32 VertexIndex = InCards.Indices[VertexIndexOffset + VertexIt];
 					const FVector4f UV = TriangleUVs[VertexIt] = InCards.UVs[VertexIndex];
-					AddSimilarUV(SimilarVertexU, VertexIndex, UV.X, UVCoordTreshold);
-					AddSimilarUV(SimilarVertexV, VertexIndex, UV.Y, UVCoordTreshold);
+					AddSimilarUV(SimilarVertexU, VertexIndex, UV.X, UV.Y, UVCoordTreshold);
+					AddSimilarUV(SimilarVertexV, VertexIndex, UV.Y, UV.X, UVCoordTreshold);
 				}
 			}
 
@@ -2906,6 +2907,29 @@ namespace HairCards
 					return A.TexCoord < B.TexCoord;
 				});
 
+				// For each group of vertex with similar 'principal' tex coord, sort them with growing 'perpendicular/secondary' tex coord
+				for (FSimilarUVVertices& Similar : SimilarVertex)
+				{
+					Similar.AllIndices.Sort([](const FIndexAndCoord& A, const FIndexAndCoord& B)
+					{
+						return A.TexCoord < B.TexCoord;
+					});
+
+					// Compute normalize coordinate
+					float MinTexCoord = Similar.AllIndices[0].TexCoord;
+					float MaxTexCoord = Similar.AllIndices[0].TexCoord;
+					for (const FIndexAndCoord& A : Similar.AllIndices)
+					{
+						MinTexCoord = FMath::Min(A.TexCoord, MinTexCoord);
+						MaxTexCoord = FMath::Max(A.TexCoord, MaxTexCoord);
+					}
+					MaxTexCoord = FMath::Max(MaxTexCoord, MinTexCoord + KINDA_SMALL_NUMBER);
+					for (FIndexAndCoord& A : Similar.AllIndices)
+					{
+						A.TexCoord = (A.TexCoord - MinTexCoord) / (MaxTexCoord - MinTexCoord);
+					}
+				}
+
 				CenterPoints.Reserve(SimilarVertex.Num());
 				FVector3f PrevCenterPoint = FVector3f::ZeroVector;
 				float TotalLength = 0;
@@ -2927,9 +2951,11 @@ namespace HairCards
 					}
 
 					// Update neighbor vertex with current length (to compute the parametric distance at the end)
-					for (uint32 VertexIndex : Similar.AllIndices)
+					for (const FIndexAndCoord& VertexIndex : Similar.AllIndices)
 					{
-						InCards.CoordU[VertexIndex] = TotalLength;
+						InCards.CoordU[VertexIndex.Index] = TotalLength;
+						InCards.LocalUVs[VertexIndex.Index].X = TotalLength;
+						InCards.LocalUVs[VertexIndex.Index].Y = VertexIndex.TexCoord;
 					}
 
 					CenterPoints.Add(CenterPoint);
@@ -2944,6 +2970,7 @@ namespace HairCards
 						for (uint32 VertexIndex : Similar.Indices)
 						{
 							InCards.CoordU[VertexIndex] /= TotalLength;
+							InCards.LocalUVs[VertexIndex].X = InCards.CoordU[VertexIndex];
 						}
 					}
 				}
@@ -3335,7 +3362,7 @@ namespace FHairCardsBuilder
 FString GetVersion()
 {
 	// Important to update the version when cards building or importing changes
-	return TEXT("9e");
+	return TEXT("9f");
 }
 
 void AllocateAtlasTexture(UTexture2D* Out, const FIntPoint& Resolution, uint32 MipCount, EPixelFormat PixelFormat, ETextureSourceFormat SourceFormat)
@@ -3859,10 +3886,12 @@ bool InternalImportGeometry(
 		}
 
 		// 2. Extract cards root points
+		struct FRootIndexAndPosition { uint32 Index; FVector3f Position; };
 		struct FCardsRootData
 		{
-			FVector3f Position;
-			uint32  CardsIndex;
+			FRootIndexAndPosition Root0;
+			FRootIndexAndPosition Root1;
+			uint32 CardsIndex;
 		};
 		TArray<FCardsRootData> CardsRoots;
 		{
@@ -3872,43 +3901,74 @@ bool InternalImportGeometry(
 				const uint32 CardsIndexCount = Out.Cards.IndexCounts[CardIt];
 				const uint32 CardsIndexOffset = Out.Cards.IndexOffsets[CardIt];
 
-				uint32 AverageCount = 0;
-				FVector3f AverageRootPosition = FVector3f::ZeroVector;
+				// Extract card vertices which are roots points
+				TArray<FRootIndexAndPosition> RootPositions;
 				for (uint32 IndexIt = 0; IndexIt < CardsIndexCount; ++IndexIt)
 				{
 					const uint32 VertexIndex = Out.Cards.Indices[CardsIndexOffset + IndexIt];
 					if (Out.Cards.CoordU[VertexIndex] == 0)
 					{
-						AverageRootPosition += Out.Cards.Positions[VertexIndex];
-						++AverageCount;
+						RootPositions.Add({VertexIndex, Out.Cards.Positions[VertexIndex]});
 					}
 				}
 
-				FCardsRootData& Cards = CardsRoots.AddDefaulted_GetRef();
-				Cards.CardsIndex = CardIt;
-				Cards.Position = AverageCount > 0 ? AverageRootPosition / AverageCount : FVector3f::ZeroVector;
+				// Select the two most representative root points
+				if (RootPositions.Num() > 0)
+				{
+					FCardsRootData& Cards = CardsRoots.AddDefaulted_GetRef();
+					Cards.CardsIndex = CardIt;
+					if (RootPositions.Num() == 1)
+					{
+						Cards.Root0 = RootPositions[0];
+						Cards.Root1 = RootPositions[0];
+					}
+					else if (RootPositions.Num() == 2)
+					{
+						const bool bInvert = Out.Cards.LocalUVs[RootPositions[0].Index].Y > Out.Cards.LocalUVs[RootPositions[1].Index].Y;
+						Cards.Root0 = bInvert ? RootPositions[1] : RootPositions[0];
+						Cards.Root1 = bInvert ? RootPositions[0] : RootPositions[1];
+					}
+					else
+					{
+						// Sort according to local V coord
+						RootPositions.Sort([&](const FRootIndexAndPosition& A, const FRootIndexAndPosition& B)
+						{
+							return Out.Cards.LocalUVs[A.Index].Y < Out.Cards.LocalUVs[B.Index].Y;
+						});
+
+						Cards.Root0 = RootPositions[0];
+						Cards.Root1 = RootPositions[RootPositions.Num()-1];
+					}
+				}
 			}
 		}
-
+		
 		// 3. Find cards root / curve root
+		for (const FCardsRootData& CardsRoot : CardsRoots)
 		{
-			for (const FCardsRootData& CardsRoot : CardsRoots)
+			for (uint32 CardsRootPositionIndex=0; CardsRootPositionIndex <2; CardsRootPositionIndex++)
 			{
 				// 3.1 Find closet root UV
 				// /!\ N^2 loop: the number of cards should be relatively small
-				uint32 CurveIndex = ~0;
-				float ClosestDistance = FLT_MAX;
-				FVector2f RootUV = FVector2f::ZeroVector;
-				for (const FStrandsRootData& StrandsRoot : StrandsRoots)
+				auto FindRootUV = [&](const FVector3f CardsRootPosition)
 				{
-					const float Distance = FVector3f::Distance(StrandsRoot.Position, CardsRoot.Position);
-					if (Distance < ClosestDistance)
+					uint32 CurveIndex = ~0;
+					float ClosestDistance = FLT_MAX;
+					FVector2f RootUV = FVector2f::ZeroVector;
+					for (const FStrandsRootData& StrandsRoot : StrandsRoots)
 					{
-						ClosestDistance = Distance;
-						CurveIndex = StrandsRoot.CurveIndex;
-						RootUV = StrandsRoot.RootUV;
+						const float Distance = FVector3f::Distance(StrandsRoot.Position, CardsRootPosition);
+						if (Distance < ClosestDistance)
+						{
+							ClosestDistance = Distance;
+							CurveIndex = StrandsRoot.CurveIndex;
+							RootUV = StrandsRoot.RootUV;
+						}
 					}
-				}
+					return RootUV;
+				};
+				const FVector2f RootUV0 = FindRootUV(CardsRoot.Root0.Position);
+				const FVector2f RootUV1 = FindRootUV(CardsRoot.Root1.Position);
 
 				// 3.2 Apply root UV to all cards vertices
 				{
@@ -3917,6 +3977,21 @@ bool InternalImportGeometry(
 					for (uint32 IndexIt = 0; IndexIt < CardsIndexCount; ++IndexIt)
 					{
 						const uint32 VertexIndex = Out.Cards.Indices[CardsIndexOffset + IndexIt];
+
+						// Linearly interpolate between the two roots based on the local V coordinate 
+						// * LocalU is along the card 
+						// * LocalV is across the card
+						//  U
+						//  ^  ____
+						//  | |    |
+						//  | |____| Card
+						//  | |    |
+						//  | |____|
+						//     ----> V
+						// Root0  Root1
+						const float TexCoordV = Out.Cards.LocalUVs[VertexIndex].Y;
+						const FVector2f RootUV = FMath::Lerp(RootUV0, RootUV1, TexCoordV);
+
 						Out.Cards.UVs[VertexIndex].Z = RootUV.X;
 						Out.Cards.UVs[VertexIndex].W = RootUV.Y;
 						OutBulk.UVs[VertexIndex].Z   = RootUV.X;
diff --git a/Engine/Plugins/Runtime/HairStrands/Source/HairStrandsCore/Public/HairCardsDatas.h b/Engine/Plugins/Runtime/HairStrands/Source/HairStrandsCore/Public/HairCardsDatas.h
--- a/Engine/Plugins/Runtime/HairStrands/Source/HairStrandsCore/Public/HairCardsDatas.h
+++ b/Engine/Plugins/Runtime/HairStrands/Source/HairStrandsCore/Public/HairCardsDatas.h
@@ -184,7 +184,8 @@ struct FHairCardsGeometry
 	TArray<FVector3f>Tangents;
 	TArray<FVector3f>Positions;
 	TArray<uint32>   Indices;
-	TArray<float>    CoordU; // Transient data storing [0..1] parametric value along main axis. This is used for generating guides & interpolation data
+	TArray<float>    CoordU;  // Transient data storing [0..1] parametric value along main axis. This is used for generating guides & interpolation data
+	TArray<FVector2f>LocalUVs;// Transient data storing [0..1, 0..1] local UV of the cards. U is aligned with the card's main direction, and V is aligned perpendicularity
 
 	// Vertex offset and vertex count of each cards geometry
 	// No longer used, kept only for backward compatibility
