Fix crash when enabling/disabling hair simulation.

When data were reloading from DDC, the bulkdata were not properly cleaned up, causing some internal check to fire.

diff --git a/Engine/Plugins/Runtime/HairStrands/Source/HairStrandsCore/Private/GroomAsset.cpp b/Engine/Plugins/Runtime/HairStrands/Source/HairStrandsCore/Private/GroomAsset.cpp
--- a/Engine/Plugins/Runtime/HairStrands/Source/HairStrandsCore/Private/GroomAsset.cpp
+++ b/Engine/Plugins/Runtime/HairStrands/Source/HairStrandsCore/Private/GroomAsset.cpp
@@ -2190,6 +2190,10 @@ bool UGroomAsset::CacheStrandsData(uint32 GroupIndex, FString& OutDerivedDataKey
 		Ar.SerializeCompressed(DecompressionBuffer, UncompressedSize, NAME_Zlib);
 
 		FHairGroupData& HairGroupData = HairGroupsData[GroupIndex];
+
+		// Reset hair group data to ensure previously loaded bulk data are cleared/cleaned priori to load new data.
+		HairGroupData = FHairGroupData();
+
 		FLargeMemoryReader LargeMemReader(DecompressionBuffer, UncompressedSize, ELargeMemoryReaderFlags::Persistent | ELargeMemoryReaderFlags::TakeOwnership);
 		InternalSerialize(LargeMemReader, this, HairGroupData);
 	}
