This CL optimizes render passes in the hair interpolation stage (interpolation, tangents, calcAABB):
- when compute rasterization with continuous LOD and/or temporal layering is enabled
- by calculating the interpolated previous positions in an additional interpolation pass which enables a reduction in control points processed in the passes.
- known issue: correct previous deformed skinned triangle are not available when skinned meshes switches LOD - this is planned to be fixed in a subsequent CL

diff --git a/Engine/Plugins/Runtime/HairStrands/Source/HairStrandsCore/Private/GroomDebug.cpp b/Engine/Plugins/Runtime/HairStrands/Source/HairStrandsCore/Private/GroomDebug.cpp
--- a/Engine/Plugins/Runtime/HairStrands/Source/HairStrandsCore/Private/GroomDebug.cpp
+++ b/Engine/Plugins/Runtime/HairStrands/Source/HairStrandsCore/Private/GroomDebug.cpp
@@ -410,11 +410,17 @@ static void AddDebugProjectionHairPass(
 	if (!RestLODDatas.RestRootTrianglePosition0Buffer.Buffer ||
 		!RestLODDatas.RestRootTrianglePosition1Buffer.Buffer ||
 		!RestLODDatas.RestRootTrianglePosition2Buffer.Buffer ||
-		!DeformedLODDatas.DeformedRootTrianglePosition0Buffer.Buffer ||
-		!DeformedLODDatas.DeformedRootTrianglePosition1Buffer.Buffer ||
-		!DeformedLODDatas.DeformedRootTrianglePosition2Buffer.Buffer)
+		!DeformedLODDatas.DeformedRootTrianglePosition0Buffer[0].Buffer ||
+		!DeformedLODDatas.DeformedRootTrianglePosition1Buffer[0].Buffer ||
+		!DeformedLODDatas.DeformedRootTrianglePosition2Buffer[0].Buffer	)
 		return;
 
+	// Double buffering is disabled by default unless the read-only cvar r.HairStrands.ContinuousDecimationReordering is set
+	if (IsHairStrandContinuousDecimationReorderingEnabled() && (!DeformedLODDatas.DeformedRootTrianglePosition0Buffer[1].Buffer || !DeformedLODDatas.DeformedRootTrianglePosition1Buffer[1].Buffer || !DeformedLODDatas.DeformedRootTrianglePosition2Buffer[1].Buffer))
+	{
+		return;
+	}
+
 	const FIntPoint Resolution(Viewport.Width(), Viewport.Height());
 
 	FHairProjectionHairDebugParameters* Parameters = GraphBuilder.AllocParameters<FHairProjectionHairDebugParameters>();
@@ -432,12 +438,12 @@ static void AddDebugProjectionHairPass(
 	Parameters->RestPosition1Buffer = RegisterAsSRV(GraphBuilder, RestLODDatas.RestRootTrianglePosition1Buffer);
 	Parameters->RestPosition2Buffer = RegisterAsSRV(GraphBuilder, RestLODDatas.RestRootTrianglePosition2Buffer);
 
-	Parameters->DeformedPosition0Buffer = RegisterAsSRV(GraphBuilder, DeformedLODDatas.DeformedRootTrianglePosition0Buffer);
-	Parameters->DeformedPosition1Buffer = RegisterAsSRV(GraphBuilder, DeformedLODDatas.DeformedRootTrianglePosition1Buffer);
-	Parameters->DeformedPosition2Buffer = RegisterAsSRV(GraphBuilder, DeformedLODDatas.DeformedRootTrianglePosition2Buffer);
+	Parameters->DeformedPosition0Buffer = RegisterAsSRV(GraphBuilder, DeformedLODDatas.GetDeformedRootTrianglePosition0Buffer(FHairStrandsDeformedRootResource::FLOD::Current));
+	Parameters->DeformedPosition1Buffer = RegisterAsSRV(GraphBuilder, DeformedLODDatas.GetDeformedRootTrianglePosition1Buffer(FHairStrandsDeformedRootResource::FLOD::Current));
+	Parameters->DeformedPosition2Buffer = RegisterAsSRV(GraphBuilder, DeformedLODDatas.GetDeformedRootTrianglePosition2Buffer(FHairStrandsDeformedRootResource::FLOD::Current));
 
 	Parameters->RestSamplePositionsBuffer = RegisterAsSRV(GraphBuilder, RestLODDatas.RestSamplePositionsBuffer);
-	Parameters->DeformedSamplePositionsBuffer = RegisterAsSRV(GraphBuilder, DeformedLODDatas.DeformedSamplePositionsBuffer);
+	Parameters->DeformedSamplePositionsBuffer = RegisterAsSRV(GraphBuilder, DeformedLODDatas.GetDeformedSamplePositionsBuffer(FHairStrandsDeformedRootResource::FLOD::Current));
 
 	Parameters->ViewUniformBuffer = ViewUniformBuffer;
 	Parameters->RenderTargets[0] = FRenderTargetBinding(ColorTarget, ERenderTargetLoadAction::ELoad, 0);
diff --git a/Engine/Plugins/Runtime/HairStrands/Source/HairStrandsCore/Private/GroomManager.cpp b/Engine/Plugins/Runtime/HairStrands/Source/HairStrandsCore/Private/GroomManager.cpp
--- a/Engine/Plugins/Runtime/HairStrands/Source/HairStrandsCore/Private/GroomManager.cpp
+++ b/Engine/Plugins/Runtime/HairStrands/Source/HairStrandsCore/Private/GroomManager.cpp
@@ -311,14 +311,14 @@ static void RunInternalHairStrandsInterpolation(
 						// Add manual transition for the GPU solver as Niagara does not track properly the RDG buffer, and so doesn't issue the correct transitions
 						if (MeshLODIndex >= 0)
 						{
-							GraphBuilder.UseExternalAccessMode(Register(GraphBuilder, Instance->Guides.DeformedRootResource->LODs[MeshLODIndex].DeformedRootTrianglePosition0Buffer, ERDGImportedBufferFlags::CreateSRV).Buffer, ERHIAccess::SRVMask);
-							GraphBuilder.UseExternalAccessMode(Register(GraphBuilder, Instance->Guides.DeformedRootResource->LODs[MeshLODIndex].DeformedRootTrianglePosition1Buffer, ERDGImportedBufferFlags::CreateSRV).Buffer, ERHIAccess::SRVMask);
-							GraphBuilder.UseExternalAccessMode(Register(GraphBuilder, Instance->Guides.DeformedRootResource->LODs[MeshLODIndex].DeformedRootTrianglePosition2Buffer, ERDGImportedBufferFlags::CreateSRV).Buffer, ERHIAccess::SRVMask);
+							GraphBuilder.UseExternalAccessMode(Register(GraphBuilder, Instance->Guides.DeformedRootResource->LODs[MeshLODIndex].GetDeformedRootTrianglePosition0Buffer(FHairStrandsDeformedRootResource::FLOD::Current), ERDGImportedBufferFlags::CreateSRV).Buffer, ERHIAccess::SRVMask);
+							GraphBuilder.UseExternalAccessMode(Register(GraphBuilder, Instance->Guides.DeformedRootResource->LODs[MeshLODIndex].GetDeformedRootTrianglePosition1Buffer(FHairStrandsDeformedRootResource::FLOD::Current), ERDGImportedBufferFlags::CreateSRV).Buffer, ERHIAccess::SRVMask);
+							GraphBuilder.UseExternalAccessMode(Register(GraphBuilder, Instance->Guides.DeformedRootResource->LODs[MeshLODIndex].GetDeformedRootTrianglePosition2Buffer(FHairStrandsDeformedRootResource::FLOD::Current), ERDGImportedBufferFlags::CreateSRV).Buffer, ERHIAccess::SRVMask);
 
 							if (bGlobalDeformationEnable)
 							{
-								GraphBuilder.UseExternalAccessMode(Register(GraphBuilder, Instance->Guides.DeformedRootResource->LODs[MeshLODIndex].DeformedSamplePositionsBuffer, ERDGImportedBufferFlags::CreateSRV).Buffer, ERHIAccess::SRVMask);
-								GraphBuilder.UseExternalAccessMode(Register(GraphBuilder, Instance->Guides.DeformedRootResource->LODs[MeshLODIndex].MeshSampleWeightsBuffer, ERDGImportedBufferFlags::CreateSRV).Buffer, ERHIAccess::SRVMask);
+								GraphBuilder.UseExternalAccessMode(Register(GraphBuilder, Instance->Guides.DeformedRootResource->LODs[MeshLODIndex].GetDeformedSamplePositionsBuffer(FHairStrandsDeformedRootResource::FLOD::Current), ERDGImportedBufferFlags::CreateSRV).Buffer, ERHIAccess::SRVMask);
+								GraphBuilder.UseExternalAccessMode(Register(GraphBuilder, Instance->Guides.DeformedRootResource->LODs[MeshLODIndex].GetMeshSampleWeightsBuffer(FHairStrandsDeformedRootResource::FLOD::Current), ERDGImportedBufferFlags::CreateSRV).Buffer, ERHIAccess::SRVMask);
 							}
 						}
 						GraphBuilder.UseExternalAccessMode(Register(GraphBuilder, Instance->Guides.DeformedResource->GetPositionOffsetBuffer(FHairStrandsDeformedResource::EFrameType::Current), ERDGImportedBufferFlags::CreateSRV).Buffer, ERHIAccess::SRVMask);
@@ -538,6 +538,14 @@ static void RunHairBufferSwap(const FHairStrandsInstances& Instances, const TArr
 		{
 			if (Instance->Guides.DeformedResource) { Instance->Guides.DeformedResource->SwapBuffer(); }
 			if (Instance->Strands.DeformedResource) { Instance->Strands.DeformedResource->SwapBuffer(); }
+
+			// Double buffering is disabled by default unless the read-only cvar r.HairStrands.ContinuousDecimationReordering is set
+			if (IsHairStrandContinuousDecimationReorderingEnabled())
+			{
+				if (Instance->Guides.DeformedRootResource) { Instance->Guides.DeformedRootResource->SwapBuffer(); }
+				if (Instance->Strands.DeformedRootResource) { Instance->Strands.DeformedRootResource->SwapBuffer(); }
+			}
+
 			Instance->Debug.LastFrameIndex = Views[0]->Family->FrameNumber;
 		}
 	}
diff --git a/Engine/Plugins/Runtime/HairStrands/Source/HairStrandsCore/Private/GroomResources.cpp b/Engine/Plugins/Runtime/HairStrands/Source/HairStrandsCore/Private/GroomResources.cpp
--- a/Engine/Plugins/Runtime/HairStrands/Source/HairStrandsCore/Private/GroomResources.cpp
+++ b/Engine/Plugins/Runtime/HairStrands/Source/HairStrandsCore/Private/GroomResources.cpp
@@ -823,6 +823,7 @@ void FHairStrandsRestResource::InternalAllocate(FRDGBuilder& GraphBuilder)
 void AddHairTangentPass(
 	FRDGBuilder& GraphBuilder,
 	FGlobalShaderMap* ShaderMap,
+	uint32 VertexStart,
 	uint32 VertexCount,
 	FHairGroupPublicData* HairGroupPublicData,
 	FRDGBufferSRVRef PositionBuffer,
@@ -838,6 +839,7 @@ FRDGExternalBuffer FHairStrandsRestResource::GetTangentBuffer(FRDGBuilder& Graph
 		AddHairTangentPass(
 			GraphBuilder,
 			ShaderMap,
+			0,
 			BulkData.PointCount,
 			nullptr,
 			RegisterAsSRV(GraphBuilder, PositionBuffer),
@@ -1242,13 +1244,28 @@ void FHairStrandsDeformedRootResource::InternalAllocateLOD(FRDGBuilder& GraphBui
 			LOD.Status = FLOD::EStatus::Initialized;
 			if (LOD.SampleCount > 0)
 			{
-				InternalCreateVertexBufferRDG<FHairStrandsMeshTrianglePositionFormat>(GraphBuilder, LOD.SampleCount, LOD.DeformedSamplePositionsBuffer, ToHairResourceDebugName(HAIRSTRANDS_RESOUCE_NAME(CurveType, Hair.StrandsRootDeformed_DeformedSamplePositionsBuffer), ResourceName), EHairResourceUsageType::Dynamic);
-				InternalCreateVertexBufferRDG<FHairStrandsMeshTrianglePositionFormat>(GraphBuilder, LOD.SampleCount + 4, LOD.MeshSampleWeightsBuffer, ToHairResourceDebugName(HAIRSTRANDS_RESOUCE_NAME(CurveType, Hair.StrandsRootDeformed_MeshSampleWeightsBuffer), ResourceName), EHairResourceUsageType::Dynamic); // offset added to workaround numerical issues when deformed and rest positions are far apart
+				InternalCreateVertexBufferRDG<FHairStrandsMeshTrianglePositionFormat>(GraphBuilder, LOD.SampleCount, LOD.DeformedSamplePositionsBuffer[0], ToHairResourceDebugName(HAIRSTRANDS_RESOUCE_NAME(CurveType, Hair.StrandsRootDeformed_DeformedSamplePositionsBuffer0), ResourceName), EHairResourceUsageType::Dynamic);
+				InternalCreateVertexBufferRDG<FHairStrandsMeshTrianglePositionFormat>(GraphBuilder, LOD.SampleCount + 4, LOD.MeshSampleWeightsBuffer[0], ToHairResourceDebugName(HAIRSTRANDS_RESOUCE_NAME(CurveType, Hair.StrandsRootDeformed_MeshSampleWeightsBuffer0), ResourceName), EHairResourceUsageType::Dynamic);
+
+				// Double buffering is disabled by default unless the read-only cvar r.HairStrands.ContinuousDecimationReordering is set
+				if (IsHairStrandContinuousDecimationReorderingEnabled())
+				{
+					InternalCreateVertexBufferRDG<FHairStrandsMeshTrianglePositionFormat>(GraphBuilder, LOD.SampleCount, LOD.DeformedSamplePositionsBuffer[1], ToHairResourceDebugName(HAIRSTRANDS_RESOUCE_NAME(CurveType, Hair.StrandsRootDeformed_DeformedSamplePositionsBuffer1), ResourceName), EHairResourceUsageType::Dynamic);
+					InternalCreateVertexBufferRDG<FHairStrandsMeshTrianglePositionFormat>(GraphBuilder, LOD.SampleCount + 4, LOD.MeshSampleWeightsBuffer[1], ToHairResourceDebugName(HAIRSTRANDS_RESOUCE_NAME(CurveType, Hair.StrandsRootDeformed_MeshSampleWeightsBuffer1), ResourceName), EHairResourceUsageType::Dynamic);
+				}
 			}
 
-			InternalCreateVertexBufferRDG<FHairStrandsMeshTrianglePositionFormat>(GraphBuilder, RootCount, LOD.DeformedRootTrianglePosition0Buffer, ToHairResourceDebugName(HAIRSTRANDS_RESOUCE_NAME(CurveType, Hair.StrandsRootDeformed_DeformedRootTrianglePosition0Buffer), ResourceName), EHairResourceUsageType::Dynamic);
-			InternalCreateVertexBufferRDG<FHairStrandsMeshTrianglePositionFormat>(GraphBuilder, RootCount, LOD.DeformedRootTrianglePosition1Buffer, ToHairResourceDebugName(HAIRSTRANDS_RESOUCE_NAME(CurveType, Hair.StrandsRootDeformed_DeformedRootTrianglePosition1Buffer), ResourceName), EHairResourceUsageType::Dynamic);
-			InternalCreateVertexBufferRDG<FHairStrandsMeshTrianglePositionFormat>(GraphBuilder, RootCount, LOD.DeformedRootTrianglePosition2Buffer, ToHairResourceDebugName(HAIRSTRANDS_RESOUCE_NAME(CurveType, Hair.StrandsRootDeformed_DeformedRootTrianglePosition2Buffer), ResourceName), EHairResourceUsageType::Dynamic);
+			InternalCreateVertexBufferRDG<FHairStrandsMeshTrianglePositionFormat>(GraphBuilder, RootCount, LOD.DeformedRootTrianglePosition0Buffer[0], ToHairResourceDebugName(HAIRSTRANDS_RESOUCE_NAME(CurveType, Hair.StrandsRootDeformed_DeformedRootTrianglePosition0Buffer0), ResourceName), EHairResourceUsageType::Dynamic);
+			InternalCreateVertexBufferRDG<FHairStrandsMeshTrianglePositionFormat>(GraphBuilder, RootCount, LOD.DeformedRootTrianglePosition1Buffer[0], ToHairResourceDebugName(HAIRSTRANDS_RESOUCE_NAME(CurveType, Hair.StrandsRootDeformed_DeformedRootTrianglePosition1Buffer0), ResourceName), EHairResourceUsageType::Dynamic);
+			InternalCreateVertexBufferRDG<FHairStrandsMeshTrianglePositionFormat>(GraphBuilder, RootCount, LOD.DeformedRootTrianglePosition2Buffer[0], ToHairResourceDebugName(HAIRSTRANDS_RESOUCE_NAME(CurveType, Hair.StrandsRootDeformed_DeformedRootTrianglePosition2Buffer0), ResourceName), EHairResourceUsageType::Dynamic);
+
+			// Double buffering is disabled by default unless the read-only cvar r.HairStrands.ContinuousDecimationReordering is set
+			if (IsHairStrandContinuousDecimationReorderingEnabled())
+			{
+				InternalCreateVertexBufferRDG<FHairStrandsMeshTrianglePositionFormat>(GraphBuilder, RootCount, LOD.DeformedRootTrianglePosition0Buffer[1], ToHairResourceDebugName(HAIRSTRANDS_RESOUCE_NAME(CurveType, Hair.StrandsRootDeformed_DeformedRootTrianglePosition0Buffer1), ResourceName), EHairResourceUsageType::Dynamic);
+				InternalCreateVertexBufferRDG<FHairStrandsMeshTrianglePositionFormat>(GraphBuilder, RootCount, LOD.DeformedRootTrianglePosition1Buffer[1], ToHairResourceDebugName(HAIRSTRANDS_RESOUCE_NAME(CurveType, Hair.StrandsRootDeformed_DeformedRootTrianglePosition1Buffer1), ResourceName), EHairResourceUsageType::Dynamic);
+				InternalCreateVertexBufferRDG<FHairStrandsMeshTrianglePositionFormat>(GraphBuilder, RootCount, LOD.DeformedRootTrianglePosition2Buffer[1], ToHairResourceDebugName(HAIRSTRANDS_RESOUCE_NAME(CurveType, Hair.StrandsRootDeformed_DeformedRootTrianglePosition2Buffer1), ResourceName), EHairResourceUsageType::Dynamic);
+			}
 		}
 	}
 }
@@ -1258,11 +1275,21 @@ void FHairStrandsDeformedRootResource::InternalRelease()
 	for (FLOD& GPUData : LODs)
 	{
 		GPUData.Status = FLOD::EStatus::Invalid;
-		GPUData.DeformedRootTrianglePosition0Buffer.Release();
-		GPUData.DeformedRootTrianglePosition1Buffer.Release();
-		GPUData.DeformedRootTrianglePosition2Buffer.Release();
-		GPUData.DeformedSamplePositionsBuffer.Release();
-		GPUData.MeshSampleWeightsBuffer.Release();
+		GPUData.DeformedRootTrianglePosition0Buffer[0].Release();
+		GPUData.DeformedRootTrianglePosition1Buffer[0].Release();
+		GPUData.DeformedRootTrianglePosition2Buffer[0].Release();
+		GPUData.DeformedSamplePositionsBuffer[0].Release();
+		GPUData.MeshSampleWeightsBuffer[0].Release();
+	
+		// Double buffering is disabled by default unless the read-only cvar r.HairStrands.ContinuousDecimationReordering is set
+		if (IsHairStrandContinuousDecimationReorderingEnabled())
+		{
+			GPUData.DeformedRootTrianglePosition0Buffer[1].Release();
+			GPUData.DeformedRootTrianglePosition1Buffer[1].Release();
+			GPUData.DeformedRootTrianglePosition2Buffer[1].Release();
+			GPUData.DeformedSamplePositionsBuffer[1].Release();
+			GPUData.MeshSampleWeightsBuffer[1].Release();
+		}
 	}
 	LODs.Empty();
 }
diff --git a/Engine/Plugins/Runtime/HairStrands/Source/HairStrandsCore/Private/HairStrandsMeshProjection.cpp b/Engine/Plugins/Runtime/HairStrands/Source/HairStrandsCore/Private/HairStrandsMeshProjection.cpp
--- a/Engine/Plugins/Runtime/HairStrands/Source/HairStrandsCore/Private/HairStrandsMeshProjection.cpp
+++ b/Engine/Plugins/Runtime/HairStrands/Source/HairStrandsCore/Private/HairStrandsMeshProjection.cpp
@@ -653,9 +653,9 @@ void AddHairStrandUpdateMeshTrianglesPass(
 	else if (Type == HairStrandsTriangleType::DeformedPose)
 	{
 		FHairStrandsDeformedRootResource::FLOD& DeformedLODData = DeformedResources->LODs[LODIndex];
-		OutputBuffers[0] = Register(GraphBuilder, DeformedLODData.DeformedRootTrianglePosition0Buffer, ERDGImportedBufferFlags::CreateUAV, UAVFlags);
-		OutputBuffers[1] = Register(GraphBuilder, DeformedLODData.DeformedRootTrianglePosition1Buffer, ERDGImportedBufferFlags::CreateUAV, UAVFlags);
-		OutputBuffers[2] = Register(GraphBuilder, DeformedLODData.DeformedRootTrianglePosition2Buffer, ERDGImportedBufferFlags::CreateUAV, UAVFlags);
+		OutputBuffers[0] = Register(GraphBuilder, DeformedLODData.GetDeformedRootTrianglePosition0Buffer(FHairStrandsDeformedRootResource::FLOD::Current), ERDGImportedBufferFlags::CreateUAV, UAVFlags);
+		OutputBuffers[1] = Register(GraphBuilder, DeformedLODData.GetDeformedRootTrianglePosition1Buffer(FHairStrandsDeformedRootResource::FLOD::Current), ERDGImportedBufferFlags::CreateUAV, UAVFlags);
+		OutputBuffers[2] = Register(GraphBuilder, DeformedLODData.GetDeformedRootTrianglePosition2Buffer(FHairStrandsDeformedRootResource::FLOD::Current), ERDGImportedBufferFlags::CreateUAV, UAVFlags);
 		DeformedLODData.Status = FHairStrandsDeformedRootResource::FLOD::EStatus::Completed;
 	}
 	else
@@ -946,9 +946,9 @@ void AddHairStrandInterpolateMeshTrianglesPass(
 	Parameters->MaxRootCount = RootCount;
 	Parameters->MaxSampleCount = RestLODData.SampleCount;
 
-	FRDGImportedBuffer OutDeformedRootTrianglePosition0 = Register(GraphBuilder, DeformedLODData.DeformedRootTrianglePosition0Buffer, ERDGImportedBufferFlags::CreateUAV);
-	FRDGImportedBuffer OutDeformedRootTrianglePosition1 = Register(GraphBuilder, DeformedLODData.DeformedRootTrianglePosition1Buffer, ERDGImportedBufferFlags::CreateUAV);
-	FRDGImportedBuffer OutDeformedRootTrianglePosition2 = Register(GraphBuilder, DeformedLODData.DeformedRootTrianglePosition2Buffer, ERDGImportedBufferFlags::CreateUAV);
+	FRDGImportedBuffer OutDeformedRootTrianglePosition0 = Register(GraphBuilder, DeformedLODData.GetDeformedRootTrianglePosition0Buffer(FHairStrandsDeformedRootResource::FLOD::Current), ERDGImportedBufferFlags::CreateUAV);
+	FRDGImportedBuffer OutDeformedRootTrianglePosition1 = Register(GraphBuilder, DeformedLODData.GetDeformedRootTrianglePosition1Buffer(FHairStrandsDeformedRootResource::FLOD::Current), ERDGImportedBufferFlags::CreateUAV);
+	FRDGImportedBuffer OutDeformedRootTrianglePosition2 = Register(GraphBuilder, DeformedLODData.GetDeformedRootTrianglePosition2Buffer(FHairStrandsDeformedRootResource::FLOD::Current), ERDGImportedBufferFlags::CreateUAV);
 
 	Parameters->RestRootTrianglePosition0 = RegisterAsSRV(GraphBuilder, RestLODData.RestRootTrianglePosition0Buffer);
 	Parameters->RestRootTrianglePosition1 = RegisterAsSRV(GraphBuilder, RestLODData.RestRootTrianglePosition1Buffer);
@@ -958,7 +958,7 @@ void AddHairStrandInterpolateMeshTrianglesPass(
 	Parameters->OutDeformedRootTrianglePosition1 = OutDeformedRootTrianglePosition1.UAV;
 	Parameters->OutDeformedRootTrianglePosition2 = OutDeformedRootTrianglePosition2.UAV;
 
-	Parameters->MeshSampleWeightsBuffer		= RegisterAsSRV(GraphBuilder, DeformedLODData.MeshSampleWeightsBuffer);
+	Parameters->MeshSampleWeightsBuffer		= RegisterAsSRV(GraphBuilder, DeformedLODData.GetMeshSampleWeightsBuffer(FHairStrandsDeformedRootResource::FLOD::Current));
 	Parameters->RestSamplePositionsBuffer	= RegisterAsSRV(GraphBuilder, RestLODData.RestSamplePositionsBuffer);
 
 	const FIntVector DispatchGroupCount = FComputeShaderUtils::GetGroupCount(RootCount, 128);
@@ -1041,7 +1041,7 @@ void InternalAddHairRBFInterpolationPass(
 	Parameters->OutDeformedPositionBuffer	= DeformedPositionBuffer_Curr.UAV;
 
 	Parameters->RestSamplePositionsBuffer	= RegisterAsSRV(GraphBuilder, RestLODData.RestSamplePositionsBuffer);
-	Parameters->MeshSampleWeightsBuffer		= RegisterAsSRV(GraphBuilder, DeformedLODData.MeshSampleWeightsBuffer);
+	Parameters->MeshSampleWeightsBuffer		= RegisterAsSRV(GraphBuilder, DeformedLODData.GetMeshSampleWeightsBuffer(FHairStrandsDeformedRootResource::FLOD::Current));
 
 	const FIntVector DispatchGroupCount = FComputeShaderUtils::GetGroupCount(VertexCount, 128);
 	check(DispatchGroupCount.X < 65536);
@@ -1211,7 +1211,7 @@ void AddHairStrandInitMeshSamplesPass(
 			FHairStrandsDeformedRootResource::FLOD& DeformedLODData = DeformedResources->LODs[LODIndex];
 			check(DeformedLODData.LODIndex == LODIndex);
 
-			OutBuffer = Register(GraphBuilder, DeformedLODData.DeformedSamplePositionsBuffer, ERDGImportedBufferFlags::CreateUAV);
+			OutBuffer = Register(GraphBuilder, DeformedLODData.GetDeformedSamplePositionsBuffer(FHairStrandsDeformedRootResource::FLOD::Current), ERDGImportedBufferFlags::CreateUAV);
 		}
 		else
 		{
@@ -1452,13 +1452,13 @@ void AddHairStrandUpdateMeshSamplesPass(
 	{
 		FHairUpdateMeshSamplesCS::FParameters* Parameters = GraphBuilder.AllocParameters<FHairUpdateMeshSamplesCS::FParameters>();
 
-		FRDGImportedBuffer OutWeightsBuffer = Register(GraphBuilder, DeformedLODData.MeshSampleWeightsBuffer, ERDGImportedBufferFlags::CreateUAV);
+		FRDGImportedBuffer OutWeightsBuffer = Register(GraphBuilder, DeformedLODData.GetMeshSampleWeightsBuffer(FHairStrandsDeformedRootResource::FLOD::Current), ERDGImportedBufferFlags::CreateUAV);
 
 		Parameters->MaxSampleCount					= RestLODData.SampleCount;
 		Parameters->SampleIndicesBuffer				= RegisterAsSRV(GraphBuilder, RestLODData.MeshSampleIndicesBuffer);
 		Parameters->InterpolationWeightsBuffer		= RegisterAsSRV(GraphBuilder, RestLODData.MeshInterpolationWeightsBuffer);
 		Parameters->SampleRestPositionsBuffer		= RegisterAsSRV(GraphBuilder, RestLODData.RestSamplePositionsBuffer);
-		Parameters->SampleDeformedPositionsBuffer	= RegisterAsSRV(GraphBuilder, DeformedLODData.DeformedSamplePositionsBuffer);
+		Parameters->SampleDeformedPositionsBuffer	= RegisterAsSRV(GraphBuilder, DeformedLODData.GetDeformedSamplePositionsBuffer(FHairStrandsDeformedRootResource::FLOD::Current));
 		Parameters->OutSampleDeformationsBuffer		= OutWeightsBuffer.UAV;
 
 		const FIntVector DispatchGroupCount = FComputeShaderUtils::GetGroupCount(RestLODData.SampleCount+4, 128);
@@ -1835,7 +1835,7 @@ void AddHairStrandUpdatePositionOffsetPass(
 	FRDGImportedBuffer RootTrianglePositionBuffer;
 	if (DeformedRootResources)
 	{
-		RootTrianglePositionBuffer = Register(GraphBuilder, DeformedRootResources->LODs[LODIndex].DeformedRootTrianglePosition0Buffer, ERDGImportedBufferFlags::CreateSRV);
+		RootTrianglePositionBuffer = Register(GraphBuilder, DeformedRootResources->LODs[LODIndex].GetDeformedRootTrianglePosition0Buffer(FHairStrandsDeformedRootResource::FLOD::Current), ERDGImportedBufferFlags::CreateSRV);
 	}
 
 	const bool bUseGPUOffset = DeformedRootResources != nullptr && GHairStrandsUseGPUPositionOffset > 0;
diff --git a/Engine/Plugins/Runtime/HairStrands/Source/HairStrandsCore/Private/HairStrandsRendering.cpp b/Engine/Plugins/Runtime/HairStrands/Source/HairStrandsCore/Private/HairStrandsRendering.cpp
--- a/Engine/Plugins/Runtime/HairStrands/Source/HairStrandsCore/Private/HairStrandsRendering.cpp
+++ b/Engine/Plugins/Runtime/HairStrands/Source/HairStrandsCore/Private/HairStrandsRendering.cpp
@@ -406,16 +406,16 @@ static void AddDeformSimHairStrandsPass(
 					Parameters->SimRestPosition1Buffer = RegisterAsSRV(GraphBuilder, RestLODDatas.RestRootTrianglePosition1Buffer);
 					Parameters->SimRestPosition2Buffer = RegisterAsSRV(GraphBuilder, RestLODDatas.RestRootTrianglePosition2Buffer);
 
-					Parameters->SimDeformedPosition0Buffer = RegisterAsSRV(GraphBuilder, DeformedLODDatas.DeformedRootTrianglePosition0Buffer);
-					Parameters->SimDeformedPosition1Buffer = RegisterAsSRV(GraphBuilder, DeformedLODDatas.DeformedRootTrianglePosition1Buffer);
-					Parameters->SimDeformedPosition2Buffer = RegisterAsSRV(GraphBuilder, DeformedLODDatas.DeformedRootTrianglePosition2Buffer);
+					Parameters->SimDeformedPosition0Buffer = RegisterAsSRV(GraphBuilder, DeformedLODDatas.GetDeformedRootTrianglePosition0Buffer(FHairStrandsDeformedRootResource::FLOD::Current));
+					Parameters->SimDeformedPosition1Buffer = RegisterAsSRV(GraphBuilder, DeformedLODDatas.GetDeformedRootTrianglePosition1Buffer(FHairStrandsDeformedRootResource::FLOD::Current));
+					Parameters->SimDeformedPosition2Buffer = RegisterAsSRV(GraphBuilder, DeformedLODDatas.GetDeformedRootTrianglePosition2Buffer(FHairStrandsDeformedRootResource::FLOD::Current));
 
 					Parameters->SimRootBarycentricBuffer = RegisterAsSRV(GraphBuilder, RestLODDatas.RootTriangleBarycentricBuffer);
 				}
 				else
 				{
 					InternalDeformationType = InternalDeformationType_RBF;
-					Parameters->MeshSampleWeightsBuffer = RegisterAsSRV(GraphBuilder, DeformedLODDatas.MeshSampleWeightsBuffer);
+					Parameters->MeshSampleWeightsBuffer = RegisterAsSRV(GraphBuilder, DeformedLODDatas.GetMeshSampleWeightsBuffer(FHairStrandsDeformedRootResource::FLOD::Current));
 					Parameters->RestSamplePositionsBuffer = RegisterAsSRV(GraphBuilder, RestLODDatas.RestSamplePositionsBuffer);
 					Parameters->SampleCount = RestLODDatas.SampleCount;
 				}
@@ -486,6 +486,7 @@ class FHairInterpolationCS : public FGlobalShader
 
 	BEGIN_SHADER_PARAMETER_STRUCT(FParameters, )
 		SHADER_PARAMETER_STRUCT_INCLUDE(ShaderPrint::FShaderParameters, ShaderDrawParameters)
+		SHADER_PARAMETER(uint32, VertexStart)
 		SHADER_PARAMETER(uint32, VertexCount)
 		SHADER_PARAMETER(uint32, DispatchCountX)
 		SHADER_PARAMETER(uint32, HairDebugMode)
@@ -568,6 +569,7 @@ static void AddHairStrandsInterpolationPass(
 	FGlobalShaderMap* ShaderMap,
 	const FShaderPrintData* ShaderPrintData,
 	const FHairGroupInstance* Instance,
+	const uint32 VertexStart,
 	const uint32 VertexCount,
 	const int32 MeshLODIndex,
 	const float HairLengthScale,
@@ -596,7 +598,8 @@ static void AddHairStrandsInterpolationPass(
 	FRDGImportedBuffer* OutRenderPrevPositionBuffer,
 	FRDGImportedBuffer* OutRenderAttributeBuffer,
 	const FRDGBufferSRVRef& VertexToClusterIdBuffer,
-	const FRDGBufferSRVRef& SimRootPointIndexBuffer)
+	const FRDGBufferSRVRef& SimRootPointIndexBuffer,
+	const FHairStrandsDeformedRootResource::FLOD::EFrameType DeformedFrame)
 {
 	const uint32 GroupSize = ComputeGroupSize();
 	const FIntVector DispatchCount = ComputeDispatchCount(VertexCount, GroupSize);
@@ -616,6 +619,8 @@ static void AddHairStrandsInterpolationPass(
 		Parameters->Interpolation1Buffer = Interpolation1Buffer;
 	}
 	Parameters->OutRenderDeformedPositionBuffer = OutRenderPositionBuffer.UAV;
+
+	Parameters->VertexStart = VertexStart;
 	Parameters->VertexCount = VertexCount;
 	Parameters->InRenderHairPositionOffset = (FVector3f)InRenderHairWorldOffset;
 	Parameters->InSimHairPositionOffset = (FVector3f)InSimHairWorldOffset;
@@ -706,9 +711,9 @@ static void AddHairStrandsInterpolationPass(
 			Parameters->SimRootBarycentricBuffer = RegisterAsSRV(GraphBuilder, Sim_RestLODDatas.RootTriangleBarycentricBuffer);
 		}
 		{
-			Parameters->SimDeformedPosition0Buffer = RegisterAsSRV(GraphBuilder, Sim_DeformedLODDatas.DeformedRootTrianglePosition0Buffer);
-			Parameters->SimDeformedPosition1Buffer = RegisterAsSRV(GraphBuilder, Sim_DeformedLODDatas.DeformedRootTrianglePosition1Buffer);
-			Parameters->SimDeformedPosition2Buffer = RegisterAsSRV(GraphBuilder, Sim_DeformedLODDatas.DeformedRootTrianglePosition2Buffer);
+			Parameters->SimDeformedPosition0Buffer = RegisterAsSRV(GraphBuilder, Sim_DeformedLODDatas.GetDeformedRootTrianglePosition0Buffer(DeformedFrame));
+			Parameters->SimDeformedPosition1Buffer = RegisterAsSRV(GraphBuilder, Sim_DeformedLODDatas.GetDeformedRootTrianglePosition1Buffer(DeformedFrame));
+			Parameters->SimDeformedPosition2Buffer = RegisterAsSRV(GraphBuilder, Sim_DeformedLODDatas.GetDeformedRootTrianglePosition2Buffer(DeformedFrame));
 		}
 	}
 
@@ -724,14 +729,14 @@ static void AddHairStrandsInterpolationPass(
 			Parameters->RootBarycentricBuffer = RegisterAsSRV(GraphBuilder, Ren_RestLODDatas.RootTriangleBarycentricBuffer);
 		}
 		{
-			Parameters->DeformedPosition0Buffer = RegisterAsSRV(GraphBuilder, Ren_DeformedLODDatas.DeformedRootTrianglePosition0Buffer);
-			Parameters->DeformedPosition1Buffer = RegisterAsSRV(GraphBuilder, Ren_DeformedLODDatas.DeformedRootTrianglePosition1Buffer);
-			Parameters->DeformedPosition2Buffer = RegisterAsSRV(GraphBuilder, Ren_DeformedLODDatas.DeformedRootTrianglePosition2Buffer);
+			Parameters->DeformedPosition0Buffer = RegisterAsSRV(GraphBuilder, Ren_DeformedLODDatas.GetDeformedRootTrianglePosition0Buffer(DeformedFrame));
+			Parameters->DeformedPosition1Buffer = RegisterAsSRV(GraphBuilder, Ren_DeformedLODDatas.GetDeformedRootTrianglePosition1Buffer(DeformedFrame));
+			Parameters->DeformedPosition2Buffer = RegisterAsSRV(GraphBuilder, Ren_DeformedLODDatas.GetDeformedRootTrianglePosition2Buffer(DeformedFrame));
 		}
 	}
 
 	const bool bHasLocalDeformation = Instance->Guides.bIsSimulationEnable || bSupportGlobalInterpolation || Instance->Guides.bIsDeformationEnable;
-	const bool bCullingEnable		= InstanceGeometryType == EHairGeometryType::Strands && CullingData.bCullingResultAvailable;
+	const bool bCullingEnable = IsHairStrandContinuousDecimationReorderingEnabled() ? false : (InstanceGeometryType == EHairGeometryType::Strands && CullingData.bCullingResultAvailable); 	// TODO: improve reordering so that culling can be used effectively
 	Parameters->HairStrandsVF_bIsCullingEnable = bCullingEnable ? 1 : 0;
 
 	// Select dynamic geometry permutation, based on the Simulation/RBF/InterpolationType
@@ -792,7 +797,7 @@ static void AddHairStrandsInterpolationPass(
 		GraphBuilder.SetBufferAccessFinal(OutRenderAttributeBuffer->Buffer, ERHIAccess::SRVMask);
 	}
 
-	if (Instance->HairGroupPublicData->VFInput.bHasLODSwitch && OutRenderPrevPositionBuffer && GHairStrandsTransferPositionOnLODChange > 0)
+	if (!(IsHairVisibilityComputeRasterContinuousLODEnabled() || IsHairVisibilityComputeRasterTemporalLayeringEnabled()) && Instance->HairGroupPublicData->VFInput.bHasLODSwitch && OutRenderPrevPositionBuffer && GHairStrandsTransferPositionOnLODChange > 0)
 	{
 		AddTransferPositionPass(GraphBuilder, ShaderMap, VertexCount, OutRenderPositionBuffer.SRV, OutRenderPrevPositionBuffer->UAV);
 		GraphBuilder.SetBufferAccessFinal(OutRenderPrevPositionBuffer->Buffer, ERHIAccess::SRVMask);
@@ -813,6 +818,10 @@ class FHairClusterAABBCS : public FGlobalShader
 	BEGIN_SHADER_PARAMETER_STRUCT(FParameters, )
 		SHADER_PARAMETER_STRUCT_INCLUDE(ShaderPrint::FShaderParameters, ShaderDrawParameters)
 		SHADER_PARAMETER(uint32, TotalClusterCount)
+		SHADER_PARAMETER(uint32, VertexStart)
+		SHADER_PARAMETER(uint32, VertexCount)
+		SHADER_PARAMETER(float, RcpSampleWeight)
+
 		SHADER_PARAMETER(FVector3f, CPUBoundMin)
 		SHADER_PARAMETER(FVector3f, CPUBoundMax)
 		SHADER_PARAMETER(FMatrix44f, LocalToTranslatedWorldMatrix)
@@ -880,6 +889,11 @@ static void AddHairClusterAABBPass(
 	Parameters->RenderDeformedPositionBuffer = RenderPositionBufferSRV;
 	Parameters->RenderDeformedOffsetBuffer = RenderDeformedOffsetBuffer;
 	Parameters->TotalClusterCount = 1;
+
+	Parameters->VertexStart = Instance->HairGroupPublicData->GetActiveStrandsVertexStart(Instance->Strands.RestResource->GetVertexCount());
+	Parameters->VertexCount = Instance->HairGroupPublicData->GetActiveStrandsVertexCount(Instance->Strands.RestResource->GetVertexCount(), Instance->HairGroupPublicData->MaxScreenSize);
+	Parameters->RcpSampleWeight = 1.0 / Instance->HairGroupPublicData->GetActiveStrandsSampleWeight(true, Instance->HairGroupPublicData->MaxScreenSize);
+
 	if (ShaderPrintData)
 	{
 		ShaderPrint::SetParameters(GraphBuilder, *ShaderPrintData, Parameters->ShaderDrawParameters);
@@ -1065,9 +1079,9 @@ static void AddHairCardsDeformationPass(
 		Parameters->TriangleRestPosition1Buffer = RegisterAsSRV(GraphBuilder, RestLODDatas.RestRootTrianglePosition1Buffer);
 		Parameters->TriangleRestPosition2Buffer = RegisterAsSRV(GraphBuilder, RestLODDatas.RestRootTrianglePosition2Buffer);
 
-		Parameters->TriangleDeformedPosition0Buffer = RegisterAsSRV(GraphBuilder, DeformedLODDatas.DeformedRootTrianglePosition0Buffer);
-		Parameters->TriangleDeformedPosition1Buffer = RegisterAsSRV(GraphBuilder, DeformedLODDatas.DeformedRootTrianglePosition1Buffer);
-		Parameters->TriangleDeformedPosition2Buffer = RegisterAsSRV(GraphBuilder, DeformedLODDatas.DeformedRootTrianglePosition2Buffer);
+		Parameters->TriangleDeformedPosition0Buffer = RegisterAsSRV(GraphBuilder, DeformedLODDatas.GetDeformedRootTrianglePosition0Buffer(FHairStrandsDeformedRootResource::FLOD::Current));
+		Parameters->TriangleDeformedPosition1Buffer = RegisterAsSRV(GraphBuilder, DeformedLODDatas.GetDeformedRootTrianglePosition1Buffer(FHairStrandsDeformedRootResource::FLOD::Current));
+		Parameters->TriangleDeformedPosition2Buffer = RegisterAsSRV(GraphBuilder, DeformedLODDatas.GetDeformedRootTrianglePosition2Buffer(FHairStrandsDeformedRootResource::FLOD::Current));
 	}
 
 	if (ShaderPrintData)
@@ -1115,6 +1129,7 @@ class FHairTangentCS : public FGlobalShader
 	using FPermutationDomain = TShaderPermutationDomain<FGroupSize, FCulling>;
 
 	BEGIN_SHADER_PARAMETER_STRUCT(FParameters, )
+		SHADER_PARAMETER(uint32, VertexStart)
 		SHADER_PARAMETER(uint32, VertexCount)
 		SHADER_PARAMETER(uint32, DispatchCountX)
 		SHADER_PARAMETER(uint32, HairStrandsVF_bIsCullingEnable)
@@ -1134,6 +1149,7 @@ IMPLEMENT_GLOBAL_SHADER(FHairTangentCS, "/Engine/Private/HairStrands/HairStrands
 void AddHairTangentPass(
 	FRDGBuilder& GraphBuilder,
 	FGlobalShaderMap* ShaderMap,
+	uint32 VertexStart,
 	uint32 VertexCount,
 	FHairGroupPublicData* HairGroupPublicData,
 	FRDGBufferSRVRef PositionBuffer,
@@ -1141,11 +1157,12 @@ void AddHairTangentPass(
 {
 	const uint32 GroupSize = ComputeGroupSize();
 	const FIntVector DispatchCount = ComputeDispatchCount(VertexCount, GroupSize);
-	const bool bCullingEnable = HairGroupPublicData && HairGroupPublicData->GetCullingResultAvailable();
+	const bool bCullingEnable = IsHairStrandContinuousDecimationReorderingEnabled() ? false : (HairGroupPublicData && HairGroupPublicData->GetCullingResultAvailable());
 
 	FHairTangentCS::FParameters* Parameters = GraphBuilder.AllocParameters<FHairTangentCS::FParameters>();
 	Parameters->PositionBuffer = PositionBuffer;
 	Parameters->OutputTangentBuffer = OutTangentBuffer.UAV;
+	Parameters->VertexStart = VertexStart;
 	Parameters->VertexCount = VertexCount;
 	Parameters->DispatchCountX = DispatchCount.X;
 	Parameters->HairStrandsVF_bIsCullingEnable = bCullingEnable ? 1 : 0;
@@ -1726,6 +1743,7 @@ void ComputeHairStrandsInterpolation(
 			AddHairTangentPass(
 				GraphBuilder,
 				ShaderMap,
+				0,
 				Instance->Guides.RestResource->GetVertexCount(),
 				Instance->HairGroupPublicData,
 				RegisterAsSRV(GraphBuilder, Instance->Guides.DeformedResource->GetBuffer(FHairStrandsDeformedResource::Current)),
@@ -1753,6 +1771,8 @@ void ComputeHairStrandsInterpolation(
 			// 2. Deform hair if needed (e.g.: skinning, simulation, RBF) and recompute tangent
 			FRDGBufferSRVRef Strands_PositionSRV = nullptr;
 			FRDGBufferSRVRef Strands_PositionOffsetSRV = nullptr;
+			FRDGBufferSRVRef Strands_PrevPositionOffsetSRV = nullptr;
+
 			FRDGBufferSRVRef Strands_TangentSRV = nullptr;
 			if (bNeedDeformation)
 			{
@@ -1855,16 +1875,23 @@ void ComputeHairStrandsInterpolation(
 					}
 				}
 
+				uint32 VertexStart = 0;
+				uint32 VertexCount = Instance->Strands.RestResource->GetVertexCount();
+
 				// 2.1 Compute deformation position based on simulation/skinning/RBF
 				if (Instance->Debug.GroomCacheType != EGroomCacheType::Strands)
 				{
+					VertexStart = Instance->HairGroupPublicData->GetActiveStrandsVertexStart(Instance->Strands.RestResource->GetVertexCount());
+					VertexCount = Instance->HairGroupPublicData->GetActiveStrandsVertexCount(Instance->Strands.RestResource->GetVertexCount(), Instance->HairGroupPublicData->MaxScreenSize);
+
 					const bool bUseSingleGuide = bValidGuide && Instance->Strands.InterpolationResource->UseSingleGuide();
 					AddHairStrandsInterpolationPass(
 						GraphBuilder,
 						ShaderMap,
 						ShaderPrintData,
 						Instance,
-						Instance->Strands.RestResource->GetVertexCount(),
+						VertexStart,
+						VertexCount,
 						MeshLODIndex,
 						Instance->Strands.Modifier.HairLengthScale,
 						bDebugModePatchedAttributeBuffer,
@@ -1892,7 +1919,51 @@ void ComputeHairStrandsInterpolation(
 						&Strands_DeformedPrevPosition,
 						Strands_DebugAttribute,
 						RegisterAsSRV(GraphBuilder, Instance->Strands.ClusterCullingResource->VertexToClusterIdBuffer),
-						bValidGuide ? RegisterAsSRV(GraphBuilder, Instance->Strands.InterpolationResource->SimRootPointIndexBuffer) : nullptr);
+						bValidGuide ? RegisterAsSRV(GraphBuilder, Instance->Strands.InterpolationResource->SimRootPointIndexBuffer) : nullptr, 
+						FHairStrandsDeformedRootResource::FLOD::Current);
+					
+						// second interpolation pass to calculate previous positions if needed
+						if (IsHairVisibilityComputeRasterContinuousLODEnabled() || IsHairVisibilityComputeRasterTemporalLayeringEnabled())
+						{
+							Strands_PrevPositionOffsetSRV = RegisterAsSRV(GraphBuilder, Instance->Strands.DeformedResource->GetPositionOffsetBuffer(FHairStrandsDeformedResource::EFrameType::Previous));
+
+							AddHairStrandsInterpolationPass(
+								GraphBuilder,
+								ShaderMap,
+								ShaderPrintData,
+								Instance,
+								VertexStart,
+								VertexCount,
+								MeshLODIndex,
+								Instance->Strands.Modifier.HairLengthScale,
+								bDebugModePatchedAttributeBuffer,
+								Instance->Strands.HairInterpolationType,
+								InstanceGeometryType,
+								CullingData,
+								Instance->Strands.RestResource->GetPositionOffset(),
+								bValidGuide ? Instance->Guides.RestResource->GetPositionOffset() : FVector::ZeroVector,
+								Strands_PrevPositionOffsetSRV,
+								bValidGuide ? RegisterAsSRV(GraphBuilder, Instance->Guides.DeformedResource->GetPositionOffsetBuffer(FHairStrandsDeformedResource::Previous)) : nullptr,
+								bHasSkinning ? Instance->Strands.RestRootResource : nullptr,
+								bHasSkinning&& bValidGuide ? Instance->Guides.RestRootResource : nullptr,
+								bHasSkinning ? Instance->Strands.DeformedRootResource : nullptr,
+								bHasSkinning&& bValidGuide ? Instance->Guides.DeformedRootResource : nullptr,
+								RegisterAsSRV(GraphBuilder, Instance->Strands.RestResource->PositionBuffer),
+								RegisterAsSRV(GraphBuilder, Instance->Strands.RestResource->Attribute0Buffer),
+								bUseSingleGuide,
+								bValidGuide&& bUseSingleGuide ? RegisterAsSRV(GraphBuilder, Instance->Strands.InterpolationResource->InterpolationBuffer) : nullptr,
+								bValidGuide && !bUseSingleGuide ? RegisterAsSRV(GraphBuilder, Instance->Strands.InterpolationResource->Interpolation0Buffer) : nullptr,
+								bValidGuide && !bUseSingleGuide ? RegisterAsSRV(GraphBuilder, Instance->Strands.InterpolationResource->Interpolation1Buffer) : nullptr,
+								bValidGuide ? RegisterAsSRV(GraphBuilder, Instance->Guides.RestResource->PositionBuffer) : nullptr,
+								bValidGuide ? RegisterAsSRV(GraphBuilder, Instance->Guides.DeformedResource->GetBuffer(FHairStrandsDeformedResource::Previous)) : nullptr,
+								bValidGuide ? RegisterAsSRV(GraphBuilder, Instance->Guides.RestResource->Attribute0Buffer) : nullptr,
+								Strands_DeformedPrevPosition,
+								nullptr,
+								nullptr,//Strands_DebugAttribute,
+								RegisterAsSRV(GraphBuilder, Instance->Strands.ClusterCullingResource->VertexToClusterIdBuffer),
+								bValidGuide ? RegisterAsSRV(GraphBuilder, Instance->Strands.InterpolationResource->SimRootPointIndexBuffer) : nullptr,
+								FHairStrandsDeformedRootResource::FLOD::Previous);
+						}
 				}
 				else if (bHasValidGroomCacheBuffers)
 				{
@@ -1929,7 +2000,8 @@ void ComputeHairStrandsInterpolation(
 				AddHairTangentPass(
 					GraphBuilder,
 					ShaderMap,
-					Instance->Strands.RestResource->GetVertexCount(),
+					VertexStart,
+					VertexCount,
 					Instance->HairGroupPublicData,
 					Strands_DeformedPosition.SRV,
 					Strands_DeformedTangent);
@@ -2116,6 +2188,7 @@ void ComputeHairStrandsInterpolation(
 						ShaderMap,
 						ShaderPrintData,
 						Instance,
+						0,
 						LOD.Guides.RestResource->GetVertexCount(),
 						MeshLODIndex,
 						1.0f,
@@ -2144,7 +2217,8 @@ void ComputeHairStrandsInterpolation(
 						nullptr,
 						nullptr,
 						nullptr,
-						RegisterAsSRV(GraphBuilder, LOD.Guides.InterpolationResource->SimRootPointIndexBuffer)); // <- this should be optional
+						RegisterAsSRV(GraphBuilder, LOD.Guides.InterpolationResource->SimRootPointIndexBuffer),
+						FHairStrandsDeformedRootResource::FLOD::Current); // <- this should be optional
 
 					AddHairCardsDeformationPass(
 						GraphBuilder,
diff --git a/Engine/Plugins/Runtime/HairStrands/Source/HairStrandsCore/Private/Niagara/NiagaraDataInterfaceHairStrands.cpp b/Engine/Plugins/Runtime/HairStrands/Source/HairStrandsCore/Private/Niagara/NiagaraDataInterfaceHairStrands.cpp
--- a/Engine/Plugins/Runtime/HairStrands/Source/HairStrandsCore/Private/Niagara/NiagaraDataInterfaceHairStrands.cpp
+++ b/Engine/Plugins/Runtime/HairStrands/Source/HairStrandsCore/Private/Niagara/NiagaraDataInterfaceHairStrands.cpp
@@ -779,11 +779,11 @@ struct FNDIHairStrandsParametersCS : public FNiagaraDataInterfaceParametersCS
 				RestMeshProjection->RestRootTrianglePosition2Buffer.SRV.GetReference() : FNiagaraRenderer::GetDummyFloatBuffer();
 
 			FRHIShaderResourceView* DeformedTrianglePositionASRV = (bHasSkinnedInterpolation && DeformedMeshProjection) ?
-				DeformedMeshProjection->DeformedRootTrianglePosition0Buffer.SRV.GetReference() : FNiagaraRenderer::GetDummyFloatBuffer();
+				DeformedMeshProjection->GetDeformedRootTrianglePosition0Buffer(FHairStrandsDeformedRootResource::FLOD::Current).SRV.GetReference() : FNiagaraRenderer::GetDummyFloatBuffer();
 			FRHIShaderResourceView* DeformedTrianglePositionBSRV = (bHasSkinnedInterpolation && DeformedMeshProjection) ?
-				DeformedMeshProjection->DeformedRootTrianglePosition1Buffer.SRV.GetReference() : FNiagaraRenderer::GetDummyFloatBuffer();
+				DeformedMeshProjection->GetDeformedRootTrianglePosition1Buffer(FHairStrandsDeformedRootResource::FLOD::Current).SRV.GetReference() : FNiagaraRenderer::GetDummyFloatBuffer();
 			FRHIShaderResourceView* DeformedTrianglePositionCSRV = (bHasSkinnedInterpolation && DeformedMeshProjection) ?
-				DeformedMeshProjection->DeformedRootTrianglePosition2Buffer.SRV.GetReference() : FNiagaraRenderer::GetDummyFloatBuffer();
+				DeformedMeshProjection->GetDeformedRootTrianglePosition2Buffer(FHairStrandsDeformedRootResource::FLOD::Current).SRV.GetReference() : FNiagaraRenderer::GetDummyFloatBuffer();
 			FRHIShaderResourceView* RootBarycentricCoordinatesSRV = (bHasSkinnedInterpolation && RestMeshProjection) ?
 				RestMeshProjection->RootTriangleBarycentricBuffer.SRV.GetReference() : FNiagaraRenderer::GetDummyFloatBuffer();
 
@@ -794,7 +794,7 @@ struct FNDIHairStrandsParametersCS : public FNiagaraDataInterfaceParametersCS
 			FShaderResourceViewRHIRef RestSamplePositionsBufferSRV = (bHasSamples && RestMeshProjection) ?
 				RestMeshProjection->RestSamplePositionsBuffer.SRV.GetReference() : FNiagaraRenderer::GetDummyFloatBuffer();
 			FShaderResourceViewRHIRef MeshSampleWeightsBufferSRV = (bHasSamples && DeformedMeshProjection) ?
-				DeformedMeshProjection->MeshSampleWeightsBuffer.SRV.GetReference() : FNiagaraRenderer::GetDummyFloatBuffer();
+				DeformedMeshProjection->GetMeshSampleWeightsBuffer(FHairStrandsDeformedRootResource::FLOD::Current).SRV.GetReference() : FNiagaraRenderer::GetDummyFloatBuffer();
 
 		
 			// Simulation setup (we update the rest configuration based on the deformed positions 
@@ -1377,9 +1377,9 @@ static void InterpolateGroomGuides(FRHICommandListImmediate& RHICmdList, FNiagar
 			SetSRVParameter(RHICmdList, ShaderRHI, ComputeShader->RestTrianglePositionABuffer, RestMeshProjection->RestRootTrianglePosition0Buffer.SRV);
 			SetSRVParameter(RHICmdList, ShaderRHI, ComputeShader->RestTrianglePositionBBuffer, RestMeshProjection->RestRootTrianglePosition1Buffer.SRV);
 			SetSRVParameter(RHICmdList, ShaderRHI, ComputeShader->RestTrianglePositionCBuffer, RestMeshProjection->RestRootTrianglePosition2Buffer.SRV);
-			SetSRVParameter(RHICmdList, ShaderRHI, ComputeShader->DeformedTrianglePositionABuffer, DeformedMeshProjection->DeformedRootTrianglePosition0Buffer.SRV);
-			SetSRVParameter(RHICmdList, ShaderRHI, ComputeShader->DeformedTrianglePositionBBuffer, DeformedMeshProjection->DeformedRootTrianglePosition1Buffer.SRV);
-			SetSRVParameter(RHICmdList, ShaderRHI, ComputeShader->DeformedTrianglePositionCBuffer, DeformedMeshProjection->DeformedRootTrianglePosition2Buffer.SRV);
+			SetSRVParameter(RHICmdList, ShaderRHI, ComputeShader->DeformedTrianglePositionABuffer, DeformedMeshProjection->GetDeformedRootTrianglePosition0Buffer(FHairStrandsDeformedRootResource::FLOD::Current).SRV);
+			SetSRVParameter(RHICmdList, ShaderRHI, ComputeShader->DeformedTrianglePositionBBuffer, DeformedMeshProjection->GetDeformedRootTrianglePosition1Buffer(FHairStrandsDeformedRootResource::FLOD::Current).SRV);
+			SetSRVParameter(RHICmdList, ShaderRHI, ComputeShader->DeformedTrianglePositionCBuffer, DeformedMeshProjection->GetDeformedRootTrianglePosition2Buffer(FHairStrandsDeformedRootResource::FLOD::Current).SRV);
 			SetSRVParameter(RHICmdList, ShaderRHI, ComputeShader->RootBarycentricCoordinatesBuffer, RestMeshProjection->RootTriangleBarycentricBuffer.SRV);
 		}
 		const uint32 GroupSize = NIAGARA_HAIR_STRANDS_THREAD_COUNT_INTERPOLATE;
diff --git a/Engine/Plugins/Runtime/HairStrands/Source/HairStrandsCore/Public/GroomResources.h b/Engine/Plugins/Runtime/HairStrands/Source/HairStrandsCore/Public/GroomResources.h
--- a/Engine/Plugins/Runtime/HairStrands/Source/HairStrandsCore/Public/GroomResources.h
+++ b/Engine/Plugins/Runtime/HairStrands/Source/HairStrandsCore/Public/GroomResources.h
@@ -205,15 +205,33 @@ struct FHairStrandsDeformedRootResource : public FHairCommonResource
 		uint32 Total = 0;
 		for (const FLOD& LOD : LODs)
 		{
-			Total += GetBufferTotalNumBytes(LOD.DeformedRootTrianglePosition0Buffer);
-			Total += GetBufferTotalNumBytes(LOD.DeformedRootTrianglePosition1Buffer);
-			Total += GetBufferTotalNumBytes(LOD.DeformedRootTrianglePosition2Buffer);
-			Total += GetBufferTotalNumBytes(LOD.DeformedSamplePositionsBuffer);
-			Total += GetBufferTotalNumBytes(LOD.MeshSampleWeightsBuffer);
+			Total += GetBufferTotalNumBytes(LOD.DeformedRootTrianglePosition0Buffer[0]);
+			Total += GetBufferTotalNumBytes(LOD.DeformedRootTrianglePosition1Buffer[0]);
+			Total += GetBufferTotalNumBytes(LOD.DeformedRootTrianglePosition2Buffer[0]);
+			Total += GetBufferTotalNumBytes(LOD.DeformedSamplePositionsBuffer[0]);
+			Total += GetBufferTotalNumBytes(LOD.MeshSampleWeightsBuffer[0]);
+
+			// Double buffering is disabled by default unless the read-only cvar r.HairStrands.ContinuousDecimationReordering is set
+			if (IsHairStrandContinuousDecimationReorderingEnabled())
+			{
+				Total += GetBufferTotalNumBytes(LOD.DeformedRootTrianglePosition0Buffer[1]);
+				Total += GetBufferTotalNumBytes(LOD.DeformedRootTrianglePosition1Buffer[1]);
+				Total += GetBufferTotalNumBytes(LOD.DeformedRootTrianglePosition2Buffer[1]);
+				Total += GetBufferTotalNumBytes(LOD.DeformedSamplePositionsBuffer[1]);
+				Total += GetBufferTotalNumBytes(LOD.MeshSampleWeightsBuffer[1]);
+			}
 		}
 		return Total;
 	}
 
+	void SwapBuffer()
+	{
+		for (FLOD& LOD : LODs)
+		{
+			LOD.SwapBuffer();
+		}
+	}
+
 	struct FLOD
 	{
 		enum class EStatus { Invalid, Initialized, Completed };
@@ -226,14 +244,32 @@ struct FHairStrandsDeformedRootResource : public FHairCommonResource
 		int32 LODIndex = -1;
 
 		/* Strand hair roots translation and rotation in triangle-deformed position relative to the bound triangle. Positions are relative the deformed root center*/
-		FRDGExternalBuffer DeformedRootTrianglePosition0Buffer;
-		FRDGExternalBuffer DeformedRootTrianglePosition1Buffer;
-		FRDGExternalBuffer DeformedRootTrianglePosition2Buffer;
+		FRDGExternalBuffer DeformedRootTrianglePosition0Buffer[2];
+		FRDGExternalBuffer DeformedRootTrianglePosition1Buffer[2];
+		FRDGExternalBuffer DeformedRootTrianglePosition2Buffer[2];
 
 		/* Strand hair mesh interpolation matrix and sample indices */
 		uint32 SampleCount = 0;
-		FRDGExternalBuffer DeformedSamplePositionsBuffer;
-		FRDGExternalBuffer MeshSampleWeightsBuffer;
+		FRDGExternalBuffer DeformedSamplePositionsBuffer[2];
+		FRDGExternalBuffer MeshSampleWeightsBuffer[2];
+
+		/* Whether the GPU data should be initialized with the asset data or not */
+		uint32 CurrentIndex = 0;
+
+		enum EFrameType
+		{
+			Previous,
+			Current
+		};
+
+		// Double buffering is disabled by default unless the read-only cvar r.HairStrands.ContinuousDecimationReordering is set
+		inline uint32 GetIndex(EFrameType T) const { return T == EFrameType::Current ? CurrentIndex : 1u - CurrentIndex; }
+		inline const FRDGExternalBuffer& GetDeformedRootTrianglePosition0Buffer(EFrameType T) const { return IsHairStrandContinuousDecimationReorderingEnabled() ? DeformedRootTrianglePosition0Buffer[GetIndex(T)] : DeformedRootTrianglePosition0Buffer[0]; }
+		inline const FRDGExternalBuffer& GetDeformedRootTrianglePosition1Buffer(EFrameType T) const { return IsHairStrandContinuousDecimationReorderingEnabled() ? DeformedRootTrianglePosition1Buffer[GetIndex(T)] : DeformedRootTrianglePosition1Buffer[0]; }
+		inline const FRDGExternalBuffer& GetDeformedRootTrianglePosition2Buffer(EFrameType T) const { return IsHairStrandContinuousDecimationReorderingEnabled() ? DeformedRootTrianglePosition2Buffer[GetIndex(T)] : DeformedRootTrianglePosition2Buffer[0]; }
+		inline const FRDGExternalBuffer& GetDeformedSamplePositionsBuffer(EFrameType T) const { return IsHairStrandContinuousDecimationReorderingEnabled() ? DeformedSamplePositionsBuffer[GetIndex(T)] : DeformedSamplePositionsBuffer[0]; }
+		inline const FRDGExternalBuffer& GetMeshSampleWeightsBuffer(EFrameType T) const { return IsHairStrandContinuousDecimationReorderingEnabled() ? MeshSampleWeightsBuffer[GetIndex(T)] : MeshSampleWeightsBuffer[0]; }
+		inline void SwapBuffer() { CurrentIndex = 1u - CurrentIndex; }
 	};
 
 	/* Store the hair projection information for each mesh LOD */
diff --git a/Engine/Shaders/Private/HairStrands/HairStrandsInterpolation.usf b/Engine/Shaders/Private/HairStrands/HairStrandsInterpolation.usf
--- a/Engine/Shaders/Private/HairStrands/HairStrandsInterpolation.usf
+++ b/Engine/Shaders/Private/HairStrands/HairStrandsInterpolation.usf
@@ -41,6 +41,7 @@ float UnpackHairControlPointCoordU(uint InPackedDataW)
 
 #if SHADER_HAIRINTERPOLATION
 
+uint VertexStart;
 uint VertexCount;
 uint DispatchCountX;
 uint HairDebugMode;
@@ -207,7 +208,7 @@ void MainCS(uint2 DispatchThreadId : SV_DispatchThreadID)
 	const float3 SimHairPositionOffset = OutSimHairPositionOffsetBuffer[0].xyz;
 	const float3 OutRenHairPositionOffset = OutRenHairPositionOffsetBuffer[0].xyz;
 
-	uint VertexIndex = GetHairStrandsVertexFetchIndex(DispatchThreadId, GROUP_SIZE, DispatchCountX);
+	uint VertexIndex = VertexStart + GetHairStrandsVertexFetchIndex(DispatchThreadId, GROUP_SIZE, DispatchCountX);
 #if PERMUTATION_CULLING
 	if (HairStrandsVF_bIsCullingEnable)
 	{
@@ -221,7 +222,7 @@ void MainCS(uint2 DispatchThreadId : SV_DispatchThreadID)
 	}
 #endif
 
-	if (VertexIndex < VertexCount)
+	if (VertexIndex < (VertexStart + VertexCount))
 	{
 		uint3  GuideVertexIndices = 0;
 		float3 GuideVertexLerps = 0;
@@ -549,6 +550,10 @@ uint TotalClusterCount;
 float3 CPUBoundMin;
 float3 CPUBoundMax;
 
+uint VertexStart;
+uint VertexCount;
+float RcpSampleWeight;
+
 float4x4 LocalToTranslatedWorldMatrix;
 
 Buffer<uint4> RenderDeformedPositionBuffer;
@@ -594,8 +599,8 @@ void ClusterAABBEvaluationCS(uint3 DispatchThreadId : SV_DispatchThreadID, uint3
 	}
 	const uint ClusterId   = ClusterIdBuffer[ClusterFetchIndex];
 	const uint IndexOffset = ClusterIndexOffsetBuffer[ClusterFetchIndex];
-	const uint IndexCount  = ClusterIndexCountBuffer[ClusterFetchIndex];
-
+	const uint IndexCount = max(min(ClusterIndexCountBuffer[ClusterFetchIndex], GROUP_SIZE),ClusterIndexCountBuffer[ClusterFetchIndex] * RcpSampleWeight);
+	
 	const uint VertexCountAfterCulling = HairStrandsVF_CullingIndirectBuffer[3];
 
 	int3 ClusterAABBMin = 9999999.0f;
@@ -611,16 +616,19 @@ void ClusterAABBEvaluationCS(uint3 DispatchThreadId : SV_DispatchThreadID, uint3
 		{
 			const float3 OutHairPositionOffset = RenderDeformedOffsetBuffer[0].xyz;
 			const uint VertexIndex = ClusterVertexIdBuffer[IndexOffset + ThreadReadOffset];
-			const FHairControlPoint CP = UnpackHairControlPoint(RenderDeformedPositionBuffer[VertexIndex], OutHairPositionOffset, 1, 1, 1);
-
-			if (all(IsFinite(CP.Position)))
+			if ((VertexIndex >= VertexStart) && (VertexIndex < (VertexStart + VertexCount)))
 			{
-				const int3 WorldSpaceCentimeters = int3(mul(float4(CP.Position, 1.0f), LocalToTranslatedWorldMatrix).xyz);
-				const int3 WorldSpaceCentimetersMin = WorldSpaceCentimeters - 1;
-				const int3 WorldSpaceCentimetersMax = WorldSpaceCentimeters + 1;
+				const FHairControlPoint CP = UnpackHairControlPoint(RenderDeformedPositionBuffer[VertexIndex], OutHairPositionOffset, 1, 1, 1);
 
-				ClusterAABBMin = min(ClusterAABBMin, WorldSpaceCentimetersMin);
-				ClusterAABBMax = max(ClusterAABBMax, WorldSpaceCentimetersMax);
+				if (all(IsFinite(CP.Position)))
+				{
+					const int3 WorldSpaceCentimeters = int3(mul(float4(CP.Position, 1.0f), LocalToTranslatedWorldMatrix).xyz);
+					const int3 WorldSpaceCentimetersMin = WorldSpaceCentimeters - 1;
+					const int3 WorldSpaceCentimetersMax = WorldSpaceCentimeters + 1;
+
+					ClusterAABBMin = min(ClusterAABBMin, WorldSpaceCentimetersMin);
+					ClusterAABBMax = max(ClusterAABBMax, WorldSpaceCentimetersMax);
+				}
 			}
 		}
 	}
diff --git a/Engine/Shaders/Private/HairStrands/HairStrandsTangent.usf b/Engine/Shaders/Private/HairStrands/HairStrandsTangent.usf
--- a/Engine/Shaders/Private/HairStrands/HairStrandsTangent.usf
+++ b/Engine/Shaders/Private/HairStrands/HairStrandsTangent.usf
@@ -4,6 +4,7 @@
 #include "../MonteCarlo.ush"
 #include "HairStrandsClusterCommon.ush"
 
+uint VertexStart;
 uint VertexCount;
 uint DispatchCountX;
 
@@ -34,8 +35,8 @@ float4 PackNormal(float3 V, uint A)
 [numthreads(GROUP_SIZE, 1, 1)]
 void MainCS(uint2 DispatchThreadId : SV_DispatchThreadID)
 {
-	const int MaxVertexCount = int(VertexCount);
-	int IndexCurr = GetHairStrandsVertexFetchIndex(DispatchThreadId, GROUP_SIZE, DispatchCountX);
+	const int MaxVertexCount = VertexStart + int(VertexCount);
+	int IndexCurr = VertexStart + GetHairStrandsVertexFetchIndex(DispatchThreadId, GROUP_SIZE, DispatchCountX);
 	int IndexPrev = clamp(IndexCurr-1, 0, MaxVertexCount-1);
 	int IndexNext = clamp(IndexCurr+1, 0, MaxVertexCount-1);
 	bool bValid = IndexCurr < MaxVertexCount;
diff --git a/Engine/Source/Runtime/Renderer/Private/HairStrands/HairStrandsComposition.cpp b/Engine/Source/Runtime/Renderer/Private/HairStrands/HairStrandsComposition.cpp
--- a/Engine/Source/Runtime/Renderer/Private/HairStrands/HairStrandsComposition.cpp
+++ b/Engine/Source/Runtime/Renderer/Private/HairStrands/HairStrandsComposition.cpp
@@ -486,7 +486,6 @@ static void AddHairVisibilityGBufferWritePass(
 		Parameters);
 }
 
-bool IsHairVisibilityComputeRasterTemporalLayeringEnabled();
 
 static void InternalRenderHairComposition(
 	FRDGBuilder& GraphBuilder,
diff --git a/Engine/Source/Runtime/Renderer/Private/HairStrands/HairStrandsInterface.cpp b/Engine/Source/Runtime/Renderer/Private/HairStrands/HairStrandsInterface.cpp
--- a/Engine/Source/Runtime/Renderer/Private/HairStrands/HairStrandsInterface.cpp
+++ b/Engine/Source/Runtime/Renderer/Private/HairStrands/HairStrandsInterface.cpp
@@ -455,6 +455,7 @@ bool IsHairStrandsVisibleInShadows(const FViewInfo& View, const FHairStrandsInst
 
 bool IsHairStrandContinuousDecimationReorderingEnabled()
 {
+	//NB this is a readonly cvar - and causes changes in platform data and runtime allocations 
 	return CVarHairStrandsContinuousDecimationReordering.GetValueOnAnyThread() > 0;
 }
 
@@ -538,6 +539,21 @@ float GetHairVisibilityComputeRasterSampleWeight(float ScreenSize, bool bUseTemp
 	return SampleWeight;
 }
 
+uint32 FHairGroupPublicData::GetActiveStrandsVertexStart(uint32 InVertexCount) const
+{
+	return GetHairVisibilityComputeRasterVertexStart(TemporalIndex, InVertexCount);
+}
+
+uint32 FHairGroupPublicData::GetActiveStrandsVertexCount(uint32 InVertexCount, float ScreenSize) const
+{
+	return GetHairVisibilityComputeRasterVertexCount(FMath::Min(MaxScreenSize, ScreenSize), InVertexCount);
+}
+
+float FHairGroupPublicData::GetActiveStrandsSampleWeight(bool bUseTemporalWeight, float ScreenSize) const
+{
+	return GetHairVisibilityComputeRasterSampleWeight(FMath::Min(MaxScreenSize, ScreenSize), bUseTemporalWeight);
+}
+
 void FHairGroupPublicData::UpdateTemporalIndex()
 {
 	TemporalIndex = IsHairVisibilityComputeRasterTemporalLayeringEnabled() ? ((TemporalIndex + 1) % GetHairVisibilityComputeRasterTemporalLayerCount()) : 0;
diff --git a/Engine/Source/Runtime/Renderer/Private/HairStrands/HairStrandsVisibility.cpp b/Engine/Source/Runtime/Renderer/Private/HairStrands/HairStrandsVisibility.cpp
--- a/Engine/Source/Runtime/Renderer/Private/HairStrands/HairStrandsVisibility.cpp
+++ b/Engine/Source/Runtime/Renderer/Private/HairStrands/HairStrandsVisibility.cpp
@@ -3013,9 +3013,6 @@ public:
 
 IMPLEMENT_GLOBAL_SHADER(FVisiblityRasterComputeRasterizeCS, "/Engine/Private/HairStrands/HairStrandsVisibilityRasterCompute.usf", "RasterCS", SF_Compute);
 
-uint32 GetHairVisibilityComputeRasterVertexStart(uint32 TemporalAASampleIndex, uint32 InVertexCount);
-uint32 GetHairVisibilityComputeRasterVertexCount(float ScreenSize, uint32 InVertexCount);
-float GetHairVisibilityComputeRasterSampleWeight(float ScreenSize, bool bUseTemporalWeight);
 bool IsHairStrandContinuousDecimationReorderingEnabled();
 
 static FRasterComputeOutput AddVisibilityComputeRasterPass(
@@ -3147,13 +3144,13 @@ static FRasterComputeOutput AddVisibilityComputeRasterPass(
 
 			const bool bCullingEnable = GHairVisibilityComputeRaster_Culling ? HairGroupPublicData->GetCullingResultAvailable() : false;
 			
-			// allow fewer strands to be rasterized if screen size for current view is smaller than max screen size
+			// calculate current view screen size - which can result in fewer strands rasterized in current view
 			const FSphere BoundsSphere = HairGroupPublicData->ContinuousLODBounds.GetSphere();
-			const float ScreenSize = FMath::Min(ComputeBoundsScreenSize(FVector4(BoundsSphere.Center, 1), BoundsSphere.W, ViewInfo), HairGroupPublicData->MaxScreenSize);
+			const float ScreenSize = ComputeBoundsScreenSize(FVector4(BoundsSphere.Center, 1), BoundsSphere.W, ViewInfo);
 
-			const uint32 VertexStart = GetHairVisibilityComputeRasterVertexStart(HairGroupPublicData->TemporalIndex, VFInput.Strands.VertexCount);
-			const uint32 VertexCount = GetHairVisibilityComputeRasterVertexCount(ScreenSize, VFInput.Strands.VertexCount);
-			const float SampleWeight = GetHairVisibilityComputeRasterSampleWeight(ScreenSize, false);
+			const uint32 VertexStart = HairGroupPublicData->GetActiveStrandsVertexStart(VFInput.Strands.VertexCount);
+			const uint32 VertexCount = HairGroupPublicData->GetActiveStrandsVertexCount(VFInput.Strands.VertexCount, ScreenSize);
+			const float SampleWeight = HairGroupPublicData->GetActiveStrandsSampleWeight(false, ScreenSize);
 
 			//reset buffers
 			uint32 IndexGridClearValues[4] = { 0x0,0x0,0x0,0x0 };
diff --git a/Engine/Source/Runtime/Renderer/Private/HairStrands/HairStrandsVoxelization.cpp b/Engine/Source/Runtime/Renderer/Private/HairStrands/HairStrandsVoxelization.cpp
--- a/Engine/Source/Runtime/Renderer/Private/HairStrands/HairStrandsVoxelization.cpp
+++ b/Engine/Source/Runtime/Renderer/Private/HairStrands/HairStrandsVoxelization.cpp
@@ -1265,9 +1265,6 @@ public:
 
 IMPLEMENT_GLOBAL_SHADER(FVoxelRasterComputeCS, "/Engine/Private/HairStrands/HairStrandsVoxelRasterCompute.usf", "MainCS", SF_Compute);
 
-uint32 GetHairVisibilityComputeRasterVertexStart(uint32 TemporalIndex, uint32 InVertexCount);
-uint32 GetHairVisibilityComputeRasterVertexCount(float ScreenSize, uint32 InVertexCount);
-float GetHairVisibilityComputeRasterSampleWeight(float ScreenSize, bool bUseTemporalWeight);
 bool IsHairStrandContinuousDecimationReorderingEnabled();
 
 static void AddVirtualVoxelizationComputeRasterPass(
@@ -1315,9 +1312,9 @@ static void AddVirtualVoxelizationComputeRasterPass(
 					continue;
 				}
 
-				const uint32 VertexStart = GetHairVisibilityComputeRasterVertexStart(HairGroupPublicData->TemporalIndex, VFInput.Strands.VertexCount);
-				const uint32 VertexCount = GetHairVisibilityComputeRasterVertexCount(HairGroupPublicData->MaxScreenSize, VFInput.Strands.VertexCount);
-				const float SampleWeight = GetHairVisibilityComputeRasterSampleWeight(HairGroupPublicData->MaxScreenSize, true);
+				const uint32 VertexStart = HairGroupPublicData->GetActiveStrandsVertexStart(VFInput.Strands.VertexCount);
+				const uint32 VertexCount = HairGroupPublicData->GetActiveStrandsVertexCount(VFInput.Strands.VertexCount, HairGroupPublicData->MaxScreenSize);
+				const float SampleWeight = HairGroupPublicData->GetActiveStrandsSampleWeight(false, HairGroupPublicData->MaxScreenSize);
 
 				FTransform LocalToTranslatedWorldTransform = VFInput.LocalToWorldTransform;
 				LocalToTranslatedWorldTransform.AddToTranslation(TranslatedWorldOffset);
diff --git a/Engine/Source/Runtime/Renderer/Public/HairStrandsInterface.h b/Engine/Source/Runtime/Renderer/Public/HairStrandsInterface.h
--- a/Engine/Source/Runtime/Renderer/Public/HairStrandsInterface.h
+++ b/Engine/Source/Runtime/Renderer/Public/HairStrandsInterface.h
@@ -242,6 +242,10 @@ public:
 
 	uint32 GetClusterCount() const { return ClusterCount;  }
 
+	uint32 GetActiveStrandsVertexStart(uint32 InVertexCount) const;
+	uint32 GetActiveStrandsVertexCount(uint32 InVertexCount, float ScreenSize) const;
+	float GetActiveStrandsSampleWeight(bool bUseTemporalWeight, float ScreenSize) const;
+
 	void UpdateTemporalIndex();
 
 	struct FVertexFactoryInput 
@@ -429,6 +433,9 @@ RENDERER_API bool IsHairStrandContinuousDecimationReorderingEnabled();
 // Return true if continuous LOD is enabled - implies computer raster and continuous decimation reordering is true
 RENDERER_API bool IsHairVisibilityComputeRasterContinuousLODEnabled();
 
+// Return true if temporal layering is enabled - implies computer raster and continuous decimation reordering is true
+RENDERER_API bool IsHairVisibilityComputeRasterTemporalLayeringEnabled();
+
 ////////////////////////////////////////////////////////////////////////////////////////////////////////////////////////
 
 typedef TArray<FRHIUnorderedAccessView*> FBufferTransitionQueue;
