Fix for crash while opening MotionVectors_Hair level. Shader was passed dummy resources with wrong format.
Also fixed wrong usage flag on an index buffer that caused a validation error.

#rnx

diff --git a/Engine/Plugins/Runtime/HairStrands/Source/HairStrandsCore/Private/Niagara/NiagaraDataInterfaceHairStrands.cpp b/Engine/Plugins/Runtime/HairStrands/Source/HairStrandsCore/Private/Niagara/NiagaraDataInterfaceHairStrands.cpp
--- a/Engine/Plugins/Runtime/HairStrands/Source/HairStrandsCore/Private/Niagara/NiagaraDataInterfaceHairStrands.cpp
+++ b/Engine/Plugins/Runtime/HairStrands/Source/HairStrandsCore/Private/Niagara/NiagaraDataInterfaceHairStrands.cpp
@@ -785,7 +785,7 @@ struct FNDIHairStrandsParametersCS : public FNiagaraDataInterfaceParametersCS
 			FRHIShaderResourceView* DeformedTrianglePositionCSRV = (bHasSkinnedInterpolation && DeformedMeshProjection) ?
 				DeformedMeshProjection->GetDeformedRootTrianglePosition2Buffer(FHairStrandsDeformedRootResource::FLOD::Current).SRV.GetReference() : FNiagaraRenderer::GetDummyFloatBuffer();
 			FRHIShaderResourceView* RootBarycentricCoordinatesSRV = (bHasSkinnedInterpolation && RestMeshProjection) ?
-				RestMeshProjection->RootTriangleBarycentricBuffer.SRV.GetReference() : FNiagaraRenderer::GetDummyFloatBuffer();
+				RestMeshProjection->RootTriangleBarycentricBuffer.SRV.GetReference() : FNiagaraRenderer::GetDummyUIntBuffer();
 
 			// RBF buffers
 			const bool bHasSamples = (RestMeshProjection && RestMeshProjection->SampleCount > 0);
@@ -924,7 +924,7 @@ struct FNDIHairStrandsParametersCS : public FNiagaraDataInterfaceParametersCS
 			SetSRVParameter(RHICmdList, ComputeShaderRHI, DeformedTrianglePositionCBuffer, FNiagaraRenderer::GetDummyFloatBuffer());
 			SetSRVParameter(RHICmdList, ComputeShaderRHI, RestSamplePositionsBuffer, FNiagaraRenderer::GetDummyFloatBuffer());
 			SetSRVParameter(RHICmdList, ComputeShaderRHI, MeshSampleWeightsBuffer, FNiagaraRenderer::GetDummyFloatBuffer());
-			SetSRVParameter(RHICmdList, ComputeShaderRHI, RootBarycentricCoordinatesBuffer, FNiagaraRenderer::GetDummyFloatBuffer());
+			SetSRVParameter(RHICmdList, ComputeShaderRHI, RootBarycentricCoordinatesBuffer, FNiagaraRenderer::GetDummyUIntBuffer());
 			SetSRVParameter(RHICmdList, ComputeShaderRHI, ParamsScaleBuffer, FNiagaraRenderer::GetDummyFloatBuffer());
 		}
 	}
diff --git a/Engine/Source/Runtime/Engine/Private/BatchedElements.cpp b/Engine/Source/Runtime/Engine/Private/BatchedElements.cpp
--- a/Engine/Source/Runtime/Engine/Private/BatchedElements.cpp
+++ b/Engine/Source/Runtime/Engine/Private/BatchedElements.cpp
@@ -1149,7 +1149,7 @@ bool FBatchedElements::Draw(FRHICommandList& RHICmdList, const FMeshPassProcesso
 					// Set the appropriate pixel shader for the mesh.
 					PrepareShaders(RHICmdList, GraphicsPSOInit, StencilRef, FeatureLevel, MeshElement.BlendMode, RelativeMatrices, MeshElement.BatchedElementParameters, MeshElement.Texture, bHitTesting, Gamma, &MeshElement.GlowInfo, &View);
 
-					FBufferRHIRef IndexBufferRHI = RHICmdList.CreateBuffer(sizeof(uint16) * MeshElement.Indices.Num(), BUF_VertexBuffer | BUF_Volatile, sizeof(uint16), ERHIAccess::VertexOrIndexBuffer, CreateInfo);
+					FBufferRHIRef IndexBufferRHI = RHICmdList.CreateBuffer(sizeof(uint16) * MeshElement.Indices.Num(), BUF_IndexBuffer | BUF_Volatile, sizeof(uint16), ERHIAccess::VertexOrIndexBuffer, CreateInfo);
 					void* VoidPtr2 = RHICmdList.LockBuffer(IndexBufferRHI, 0, sizeof(uint16) * MeshElement.Indices.Num(), RLM_WriteOnly);
 					FPlatformMemory::Memcpy(VoidPtr2, MeshElement.Indices.GetData(), sizeof(uint16) * MeshElement.Indices.Num());
 					RHICmdList.UnlockBuffer(IndexBufferRHI);
