- Hair conversion to RDG

diff --git a/Engine/Plugins/FX/Niagara/Source/Niagara/Public/NiagaraRenderGraphUtils.h b/Engine/Plugins/FX/Niagara/Source/Niagara/Public/NiagaraRenderGraphUtils.h
--- a/Engine/Plugins/FX/Niagara/Source/Niagara/Public/NiagaraRenderGraphUtils.h
+++ b/Engine/Plugins/FX/Niagara/Source/Niagara/Public/NiagaraRenderGraphUtils.h
@@ -5,7 +5,7 @@
 
 //////////////////////////////////////////////////////////////////////////
 // Helper for wrapping a persistent RW buffer
-struct FNiagaraPooledRWBuffer
+struct NIAGARA_API FNiagaraPooledRWBuffer
 {
 	void Initialize(FRDGBuilder& GraphBuilder, const TCHAR* ResourceName, EPixelFormat PixelFormat, const FRDGBufferDesc& BufferDesc);
 	void Initialize(FRDGBuilder& GraphBuilder, const TCHAR* ResourceName, EPixelFormat PixelFormat, const uint32 BytesPerElemenet, const uint32 NumElements, EBufferUsageFlags UsageFlags = EBufferUsageFlags::None);
@@ -36,7 +36,7 @@ private:
 
 //////////////////////////////////////////////////////////////////////////
 // Helper for wrapping a persistent RW texture
-struct FNiagaraPooledRWTexture
+struct NIAGARA_API FNiagaraPooledRWTexture
 {
 	void Initialize(FRDGBuilder& GraphBuilder, const TCHAR* ResourceName, const FRDGTextureDesc& TextureDesc);
 	void Release();
diff --git a/Engine/Plugins/Runtime/HairStrands/Shaders/Private/NiagaraDataInterfaceHairStrands.ush b/Engine/Plugins/Runtime/HairStrands/Shaders/Private/NiagaraDataInterfaceHairStrands.ush
--- a/Engine/Plugins/Runtime/HairStrands/Shaders/Private/NiagaraDataInterfaceHairStrands.ush
+++ b/Engine/Plugins/Runtime/HairStrands/Shaders/Private/NiagaraDataInterfaceHairStrands.ush
@@ -6,136 +6,9 @@ NiagaraDataInterfaceHairStrands.ush
 
 #pragma once
 
-/* -----------------------------------------------------------------
- * Hair Strands constants and context
- * -----------------------------------------------------------------
- */
-
 static const float M_PI = 3.14159265;
 static const float3 MAX_POS = float3(1e+8,1e+8,1e+8);
 
-struct FDIHairStrandsContext
-{
-	RWBuffer<uint4>		DeformedPositionBuffer;
-	Buffer<uint4>		RestPositionBuffer;
-	Buffer<uint>		CurvesOffsetsBuffer;
-	Buffer<float4>		RestTrianglePositionABuffer;
-	Buffer<float4>		RestTrianglePositionBBuffer;
-	Buffer<float4>		RestTrianglePositionCBuffer;
-	Buffer<float4>		DeformedTrianglePositionABuffer;
-	Buffer<float4>		DeformedTrianglePositionBBuffer;
-	Buffer<float4>		DeformedTrianglePositionCBuffer;
-	Buffer<uint>		RootBarycentricCoordinatesBuffer;
-	float4x4			WorldTransform;
-	float4x4			WorldInverse;
-	float4				WorldRotation;
-	float4x4			BoneTransform;
-	float4x4			BoneInverse;
-	float4				BoneRotation;
-	int					NumStrands;
-	int					StrandSize;
-	int					InterpolationMode;
-	float3				RestRootOffset;
-	float3				DeformedRootOffset;
-	float3				RestPositionOffset;
-	Buffer<float4>		DeformedPositionOffset;
-	bool				ResetSimulation;
-	bool				RestUpdate;
-	bool				LocalSimulation;
-	int					SampleCount;
-	int4				BoundingBoxOffsets;
-	Buffer<float4>		RestSamplePositionsBuffer;
-	Buffer<float4>		MeshSampleWeightsBuffer;
-	Buffer<float>		ParamsScaleBuffer;
-	float3				BoneLinearVelocity;
-	float3				BoneAngularVelocity;
-	float3				BoneLinearAcceleration;
-	float3				BoneAngularAcceleration;
-}; 
-
-#define DIHAIRSTRANDS_DECLARE_CONSTANTS(NAME)\
-RWBuffer<uint4>		DeformedPositionBuffer_##NAME;\
-Buffer<uint4>		RestPositionBuffer_##NAME;\
-Buffer<uint>		CurvesOffsetsBuffer_##NAME;\
-Buffer<float4>		RestTrianglePositionABuffer_##NAME;\
-Buffer<float4>		RestTrianglePositionBBuffer_##NAME;\
-Buffer<float4>		RestTrianglePositionCBuffer_##NAME;\
-Buffer<float4>		DeformedTrianglePositionABuffer_##NAME;\
-Buffer<float4>		DeformedTrianglePositionBBuffer_##NAME;\
-Buffer<float4>		DeformedTrianglePositionCBuffer_##NAME;\
-Buffer<uint>		RootBarycentricCoordinatesBuffer_##NAME;\
-float4x4			WorldTransform_##NAME;\
-float4x4			WorldInverse_##NAME;\
-float4				WorldRotation_##NAME;\
-float4x4			BoneTransform_##NAME;\
-float4x4			BoneInverse_##NAME;\
-float4				BoneRotation_##NAME;\
-int					NumStrands_##NAME;\
-int					StrandSize_##NAME;\
-int					InterpolationMode_##NAME;\
-float3				RestRootOffset_##NAME;\
-float3				DeformedRootOffset_##NAME;\
-float3				RestPositionOffset_##NAME;\
-Buffer<float4>		DeformedPositionOffset_##NAME;\
-bool				ResetSimulation_##NAME;\
-bool				RestUpdate_##NAME;\
-bool				LocalSimulation_##NAME;\
-RWBuffer<uint>	    BoundingBoxBuffer_##NAME;\
-int					SampleCount_##NAME;\
-int4				BoundingBoxOffsets_##NAME;\
-Buffer<float4>		RestSamplePositionsBuffer_##NAME;\
-Buffer<float4>		MeshSampleWeightsBuffer_##NAME;\
-Buffer<float>		ParamsScaleBuffer_##NAME;\
-float3				BoneLinearVelocity_##NAME;\
-float3				BoneAngularVelocity_##NAME;\
-float3				BoneLinearAcceleration_##NAME;\
-float3				BoneAngularAcceleration_##NAME;\
-
-//@todo-lh: BoundingBoxBuffer is passed around separately as SPIRV-Cross cannot translate texture atomics properly atm.
-#define DIHAIRSTRANDS_MAKE_CONTEXT(NAME)\
-FDIHairStrandsContext DIContext;\
-DIContext.DeformedPositionBuffer = DeformedPositionBuffer_##NAME;\
-DIContext.RestPositionBuffer = RestPositionBuffer_##NAME;\
-DIContext.CurvesOffsetsBuffer = CurvesOffsetsBuffer_##NAME;\
-DIContext.RestTrianglePositionABuffer = RestTrianglePositionABuffer_##NAME;\
-DIContext.RestTrianglePositionBBuffer = RestTrianglePositionBBuffer_##NAME;\
-DIContext.RestTrianglePositionCBuffer = RestTrianglePositionCBuffer_##NAME;\
-DIContext.DeformedTrianglePositionABuffer = DeformedTrianglePositionABuffer_##NAME;\
-DIContext.DeformedTrianglePositionBBuffer = DeformedTrianglePositionBBuffer_##NAME;\
-DIContext.DeformedTrianglePositionCBuffer = DeformedTrianglePositionCBuffer_##NAME;\
-DIContext.RootBarycentricCoordinatesBuffer = RootBarycentricCoordinatesBuffer_##NAME;\
-DIContext.WorldTransform = WorldTransform_##NAME;\
-DIContext.WorldInverse = WorldInverse_##NAME;\
-DIContext.WorldRotation = WorldRotation_##NAME;\
-DIContext.BoneTransform = BoneTransform_##NAME;\
-DIContext.BoneInverse = BoneInverse_##NAME;\
-DIContext.BoneRotation = BoneRotation_##NAME;\
-DIContext.NumStrands = NumStrands_##NAME;\
-DIContext.StrandSize = StrandSize_##NAME;\
-DIContext.InterpolationMode =	InterpolationMode_##NAME;\
-DIContext.RestRootOffset =	RestRootOffset_##NAME;\
-DIContext.DeformedRootOffset = DeformedRootOffset_##NAME;\
-DIContext.RestPositionOffset =	RestPositionOffset_##NAME;\
-DIContext.DeformedPositionOffset = DeformedPositionOffset_##NAME;\
-DIContext.ResetSimulation = ResetSimulation_##NAME;\
-DIContext.RestUpdate = RestUpdate_##NAME;\
-DIContext.LocalSimulation = LocalSimulation_##NAME;\
-RWBuffer<uint> DIContext_BoundingBoxBuffer = BoundingBoxBuffer_##NAME;\
-DIContext.SampleCount = SampleCount_##NAME;\
-DIContext.BoundingBoxOffsets = BoundingBoxOffsets_##NAME;\
-DIContext.RestSamplePositionsBuffer = RestSamplePositionsBuffer_##NAME;\
-DIContext.MeshSampleWeightsBuffer = MeshSampleWeightsBuffer_##NAME;\
-DIContext.ParamsScaleBuffer = ParamsScaleBuffer_##NAME;\
-DIContext.BoneLinearVelocity = BoneLinearVelocity_##NAME;\
-DIContext.BoneAngularVelocity = BoneAngularVelocity_##NAME;\
-DIContext.BoneLinearAcceleration = BoneLinearAcceleration_##NAME;\
-DIContext.BoneAngularAcceleration = BoneAngularAcceleration_##NAME;\
-
-/* -----------------------------------------------------------------
- * Vector utilities
- * -----------------------------------------------------------------
- */
-
 float2 ExtractBarycentrics(in uint PackedCoordinates)
 {
 	return float2(f16tof32(PackedCoordinates & 0xFFFF), f16tof32((PackedCoordinates >> 16) & 0xFFFF));
@@ -169,868 +42,3 @@ float2 ExtractBarycentrics(in uint PackedCoordinates)
 	return UnsignedValue;
 }
 
-/* -----------------------------------------------------------------
- * Box utilities
- * -----------------------------------------------------------------
- */
-
-// Get the boundingh box
-void DIHairStrands_GetBoundingBox(in FDIHairStrandsContext DIContext, in RWBuffer<uint> DIContext_BoundingBoxBuffer, in int BoxIndex, out float3 OutBoxCenter, out float3 OutBoxExtent)
-{
-	const int BufferOffset = DIContext.BoundingBoxOffsets[BoxIndex] * 6;
-
-	const float3 BoxMin = Uint3ToFloat3(uint3(DIContext_BoundingBoxBuffer[0+BufferOffset],
-											  DIContext_BoundingBoxBuffer[1+BufferOffset],
-											  DIContext_BoundingBoxBuffer[2+BufferOffset]));
-
-	const float3 BoxMax = Uint3ToFloat3(uint3(DIContext_BoundingBoxBuffer[3+BufferOffset],
-											  DIContext_BoundingBoxBuffer[4+BufferOffset],
-											  DIContext_BoundingBoxBuffer[5+BufferOffset]));
-
-	OutBoxExtent = (BoxMax-BoxMin);
-	OutBoxCenter = 0.5 * (BoxMin+BoxMax);
-}
-
-// Reset the bounding box
-void DIHairStrands_ResetBoundingBox(in FDIHairStrandsContext DIContext, in RWBuffer<uint> DIContext_BoundingBoxBuffer, out bool FunctionStatus)
-{
-	FunctionStatus = false;
-
-	if(GLinearThreadId == 0)
-	{
-		FunctionStatus = true;
-
-		const int BufferOffset = DIContext.BoundingBoxOffsets[3] * 6;
-
-		const uint UINT_MAX = FloatToUint(1e+8);
-		const uint UINT_MIN = FloatToUint(-1e+8);
-
-		DIContext_BoundingBoxBuffer[0+BufferOffset] = UINT_MAX;
-		DIContext_BoundingBoxBuffer[1+BufferOffset] = UINT_MAX;
-		DIContext_BoundingBoxBuffer[2+BufferOffset] = UINT_MAX;
-
-		DIContext_BoundingBoxBuffer[3+BufferOffset] = UINT_MIN;
-		DIContext_BoundingBoxBuffer[4+BufferOffset] = UINT_MIN;
-		DIContext_BoundingBoxBuffer[5+BufferOffset] = UINT_MIN;
-	}
-	DeviceMemoryBarrier();
-}
-
-void DIHairStrands_BuildBoundingBox(in FDIHairStrandsContext DIContext, in RWBuffer<uint> DIContext_BoundingBoxBuffer, in float3 NodePosition, out bool FunctionStatus)
-{
-	FunctionStatus = false;
-
-	const uint3 LocalPosition = Float3ToUint3(NodePosition);
-
-	if(isfinite(NodePosition.x) && isfinite(NodePosition.y) && isfinite(NodePosition.z))
-	{
-		FunctionStatus = true;
-
-		const int BufferOffset = DIContext.BoundingBoxOffsets[2] * 6;
-
-		InterlockedMin(DIContext_BoundingBoxBuffer[0+BufferOffset],LocalPosition.x);
-		InterlockedMin(DIContext_BoundingBoxBuffer[1+BufferOffset],LocalPosition.y);
-		InterlockedMin(DIContext_BoundingBoxBuffer[2+BufferOffset],LocalPosition.z);
-
-		InterlockedMax(DIContext_BoundingBoxBuffer[3+BufferOffset],LocalPosition.x);
-		InterlockedMax(DIContext_BoundingBoxBuffer[4+BufferOffset],LocalPosition.y);
-		InterlockedMax(DIContext_BoundingBoxBuffer[5+BufferOffset],LocalPosition.z);
-	}
-	DeviceMemoryBarrier();
-}
-
-/* -----------------------------------------------------------------
- * Utilities regarding strands indexing
- * -----------------------------------------------------------------
- */
-
-// Given a node index return the strand index and the local index within the strand
-int DIHairStrands_GetNumNodes(in FDIHairStrandsContext DIContext)
-{
-	return DIContext.NumStrands * DIContext.StrandSize;
-}
-
-/* -----------------------------------------------------------------
- * Strands properties
- * -----------------------------------------------------------------
- */
-
-// Given a strand index return the corresponding strand length
-float DIHairStrands_ComputeStrandLength(in FDIHairStrandsContext DIContext, in int StrandIndex)
-{
-	float StrandLength = 0.0;
-	if( StrandIndex < DIContext.NumStrands )
-	{
-		int PointOffset = DIContext.CurvesOffsetsBuffer[StrandIndex];
-		const int EdgeCount = DIContext.CurvesOffsetsBuffer[StrandIndex+1]-PointOffset-1;
-
-		float3 PointNext = f16tof32(DIContext.RestPositionBuffer[PointOffset++].xyz);
-		float3 PointPrev = PointNext;
-
-		for (int EdgeIndex = 0; EdgeIndex < EdgeCount; ++EdgeIndex)
-		{
-			PointPrev = PointNext;
-			PointNext = f16tof32(DIContext.RestPositionBuffer[PointOffset++].xyz);
-			StrandLength += length(PointNext-PointPrev);
-		}
-	}
-	return StrandLength;
-}
-
-/* -----------------------------------------------------------------
- * Node Mass/Inertia/Position/Orientation computation
- * -----------------------------------------------------------------
- */
-
- // Compute the node mass
-void DIHairStrands_ComputeNodeMass(in FDIHairStrandsContext DIContext, in float StrandsDensity, in float NodeThickness, out float OutNodeMass)
-{
-	const int ReducedSize =  DIContext.StrandSize-1;
-	const float CoordScale = 1.0 / (ReducedSize-1.0);
-
-	const int StrandIndex = ExecIndex() / DIContext.StrandSize;
-
-	const float StrandLength = DIHairStrands_ComputeStrandLength(DIContext,StrandIndex);
-
-	int  LocalIndex = GGroupThreadId.x % DIContext.StrandSize;
-	const float EdgeLength = StrandLength * CoordScale;
-	const float MeanRadius = 0.5 * NodeThickness;
-
-	LocalIndex = max(0,LocalIndex - 1);
-	const float EdgeScale = (LocalIndex == 0 || LocalIndex == (ReducedSize-1)) ? 0.5 : 1.0;
-
-	OutNodeMass = StrandsDensity * EdgeLength * EdgeScale * M_PI * MeanRadius * MeanRadius;
-}
-
-// Compute the node inertia
-void DIHairStrands_ComputeNodeInertia(in FDIHairStrandsContext DIContext, in float StrandsDensity, in float NodeThickness, out float3 OutNodeInertia)
-{
-	const int ReducedSize =  DIContext.StrandSize-1;
-	const float CoordScale = 1.0 / (ReducedSize-1.0);
-
-	const int StrandIndex = ExecIndex() / DIContext.StrandSize;
-
-	const float StrandLength = DIHairStrands_ComputeStrandLength(DIContext,StrandIndex);
-
-	const float EdgeLength = StrandLength * CoordScale;
-	const float MeanRadius = 0.5 * NodeThickness;
-
-	const float EdgeMass = StrandsDensity * EdgeLength * M_PI * MeanRadius * MeanRadius;
-	const float RadialInertia = EdgeMass * ( 3.0 * MeanRadius * MeanRadius + EdgeLength * EdgeLength) / 12.0;
-	OutNodeInertia = float3( RadialInertia, RadialInertia, EdgeMass * MeanRadius * MeanRadius * 0.5 );
-}
-
-// Set the node position by linear interpolation over the points
-void DIHairStrands_SetNodePosition(in FDIHairStrandsContext DIContext, in bool RootNode, const float NodePoint, const float PointCount, 
-		const float CoordScale, const int PointPrev, const int PointNext, const int PointOffset, in int StrandIndex, out float3 OutNodePosition)
-{
-	if(RootNode)
-	{
-		const float FirstPoint = CoordScale * PointCount;
-		const int FirstPrev = floor(FirstPoint);
-		const int FirstNext = PointPrev+1;
-
-		const float FirstAlpha = FirstPoint - (float)FirstPrev;
-		const float3 FirstPosition =  f16tof32(DIContext.RestPositionBuffer[PointOffset+FirstPrev].xyz) * (1.0-FirstAlpha) + 
-						   f16tof32(DIContext.RestPositionBuffer[PointOffset+FirstNext].xyz) * FirstAlpha;
-		const float3 DiffPosition = FirstPosition - f16tof32(DIContext.RestPositionBuffer[PointOffset].xyz);
-		const float3 EdgeDirection =  normalize(DiffPosition);
-
-		const float EdgeLength = DIHairStrands_ComputeStrandLength(DIContext,StrandIndex) * CoordScale;
-		OutNodePosition = 	f16tof32(DIContext.RestPositionBuffer[PointOffset+PointPrev].xyz) - EdgeLength * normalize(EdgeDirection);					 
-	}
-	else
-	{
-		const float NodeAlpha = NodePoint - (float)PointPrev;
-		OutNodePosition =  f16tof32(DIContext.RestPositionBuffer[PointOffset+PointPrev].xyz) * (1.0-NodeAlpha) + 
-						   f16tof32(DIContext.RestPositionBuffer[PointOffset+PointNext].xyz) * NodeAlpha;
-	}
-	OutNodePosition = mul(float4( OutNodePosition+DIContext.RestPositionOffset, 1.0), DIContext.WorldTransform).xyz;
-}
-
-// Compute the node position by linear interpolation over the points
-void DIHairStrands_ComputeNodePosition(in FDIHairStrandsContext DIContext, out float3 OutNodePosition)
-{
-	const int ReducedSize =  DIContext.StrandSize-1;
-	const float CoordScale = 1.0 / (ReducedSize-1.0);
-
-	int LocalIndex = GGroupThreadId.x % DIContext.StrandSize;
-	int StrandIndex = ExecIndex() / DIContext.StrandSize;
-
-	const bool RootNode = (LocalIndex == 0);
-	LocalIndex = max(0,LocalIndex - 1);
-
-	const int PointOffset = DIContext.CurvesOffsetsBuffer[StrandIndex];
-
-	const float NodeCoord = (float)(LocalIndex) * CoordScale;
-	const float PointCount = DIContext.CurvesOffsetsBuffer[StrandIndex+1]-PointOffset-1;
-
-	const float NodePoint = NodeCoord * PointCount;
-	const int PointPrev = (LocalIndex==0) ? 0 : (LocalIndex==(ReducedSize-1)) ? PointCount-1 : floor(NodePoint);
-	const int PointNext = PointPrev+1;
-
-	DIHairStrands_SetNodePosition( DIContext, RootNode, NodePoint, PointCount, CoordScale, PointPrev, PointNext, PointOffset, StrandIndex, OutNodePosition);
-}
-
-// Compute the node position by linear interpolation over the points
-void DIHairStrands_SmoothNodePosition(in FDIHairStrandsContext DIContext, in float SmoothingFilter, inout float3 OutNodePosition)
-{
-	SharedNodePosition[GGroupThreadId.x] = OutNodePosition;
-	GroupMemoryBarrier();
-
-	const int LocalIndex = (GGroupThreadId.x % DIContext.StrandSize);
-	if(LocalIndex == 0)
-	{
-		const float Beta = SmoothingFilter;
-		float3 DirM1 = SharedNodePosition[GGroupThreadId.x+1] - SharedNodePosition[GGroupThreadId.x];
-		float3 DirM2 = DirM1;
-
-		const float Gamma1 = 2.0 * (1.0-Beta);
-		const float Gamma2 = - (1.0-Beta)*(1.0-Beta);
-		const float Gamma3 = Beta*Beta;
-
-		float3 NodePosition = SharedNodePosition[GGroupThreadId.x];
-		SharedPreviousPosition[GGroupThreadId.x] = NodePosition;
-		
-		for( int i = GGroupThreadId.x, end = GGroupThreadId.x+DIContext.StrandSize-1; i < end; ++i)
-		{
-			const float3 DirM3 = SharedNodePosition[i+1] - SharedNodePosition[i];
-			const float3 DirMi = Gamma1 * DirM1 + Gamma2 * DirM2 + Gamma3 * DirM3;
-
-			SharedPreviousPosition[i+1]  = SharedPreviousPosition[i] + DirMi;
-			GroupMemoryBarrier();
-
-			DirM2 = DirM1;
-			DirM1 = DirMi;
-		}
-	}
-	GroupMemoryBarrier();
-	OutNodePosition = SharedPreviousPosition[GGroupThreadId.x];
-}
-
-// Compute the root orientation
-void DIHairStrands_ComputeRootOrientation(in FDIHairStrandsContext DIContext)
-{
-	const int LocalIndex = (GGroupThreadId.x % DIContext.StrandSize);
-	if(LocalIndex == 0)
-	{
-		const float3 EdgeDirection = normalize(SharedNodePosition[GGroupThreadId.x+1] - SharedNodePosition[GGroupThreadId.x]);
-		const float4 RootQuaternion = FindQuatBetweenNormals(float3(0,0,1),EdgeDirection);
-
-		//const float3 TangentPrev = normalize(RotateVectorByQuat( float3(1,0,0), RootQuaternion));
-		//const float3 TangentNext = normalize(cross( normalize(cross(EdgeDirection,float3(0,0,1))), EdgeDirection));
-
-		SharedNodeOrientation[GGroupThreadId.x] = RootQuaternion;
-
-		//SharedNodeOrientation[GGroupThreadId.x] = NormalizeQuat( MultiplyQuat( RootQuaternion, FindQuatBetweenNormals(TangentPrev,TangentNext) ) );
-	}
-	GroupMemoryBarrier();
-}
-
-// Update the root orientation
-void DIHairStrands_UpdateRootOrientation(in FDIHairStrandsContext DIContext)
-{
-	const int LocalIndex = (GGroupThreadId.x % DIContext.StrandSize);
-	if(LocalIndex == 0)
-	{
-		float4 NodeQuaternion = SharedNodeOrientation[GGroupThreadId.x];
-		float3 TangentPrev = RotateVectorByQuat( float3(0,0,1), NodeQuaternion);
-		float3 TangentNext = normalize(SharedNodePosition[GGroupThreadId.x+1] - SharedNodePosition[GGroupThreadId.x]);
-
-		SharedNodeOrientation[GGroupThreadId.x] = NormalizeQuat( MultiplyQuat( FindQuatBetweenNormals(TangentPrev,TangentNext), NodeQuaternion) );
-	}
-	GroupMemoryBarrier();
-}
-
-// Compute the node orientation
-void DIHairStrands_ComputeNodeOrientation(in FDIHairStrandsContext DIContext, in float3 NodePosition, out float4 OutNodeOrientation)
-{
-	SharedNodePosition[GGroupThreadId.x] = NodePosition;
-	SharedNodeOrientation[GGroupThreadId.x] = QUATERNION_IDENTITY;
-	GroupMemoryBarrier();
-
-	DIHairStrands_ComputeRootOrientation(DIContext);
-
-	ComputeMaterialFrame(DIContext.StrandSize);
-	OutNodeOrientation = 	SharedNodeOrientation[GGroupThreadId.x];
-}
-
-/* -----------------------------------------------------------------
- * Edge volume, length, rotation, direction
- * -----------------------------------------------------------------
- */
-
-// Init the samples along the strands that will be used to transfer informations to the grid
-void DIHairStrands_InitGridSamples(in FDIHairStrandsContext DIContext, in float3 NodePosition, in float3 NodeVelocity, 
-	in float NodeMass, in float GridLength, out int OutNumSamples,
-						out float3 OutDeltaPosition, out float3 OutDeltaVelocity, out float OutSampleMass)
-{
-	SharedNodePosition[GGroupThreadId.x] = NodePosition;
-	SharedPreviousPosition[GGroupThreadId.x] = NodeVelocity;
-	SharedInverseMass[GGroupThreadId.x] = NodeMass;
-	GroupMemoryBarrier();
-
-	const int LocalIndex = GGroupThreadId.x % DIContext.StrandSize;
-
-	if( LocalIndex > 0 )
-	{
-		OutDeltaVelocity = SharedPreviousPosition[GGroupThreadId.x] - SharedPreviousPosition[GGroupThreadId.x-1];
-		OutDeltaPosition = SharedNodePosition[GGroupThreadId.x] - SharedNodePosition[GGroupThreadId.x-1];
-
-		const float2 SegmentWeight = (LocalIndex == 1) ? float2(1.0,0.5) : (LocalIndex == (DIContext.StrandSize-1) ) ? float2(0.5,1.0) : float2(0.5,0.5);
-		const float SegmentMass =  (SharedInverseMass[GGroupThreadId.x-1] * SegmentWeight.x + SharedInverseMass[GGroupThreadId.x] * SegmentWeight.y);
-		const float SegmentLength = length(OutDeltaPosition);
-
-		OutNumSamples = ceil(SegmentLength / GridLength);
-		OutSampleMass = SegmentMass / OutNumSamples;
-	}
-	else
-	{
-		OutNumSamples = 0;
-		OutDeltaPosition = float3(0,0,0);
-		OutDeltaVelocity = float3(0,0,0);
-		OutSampleMass = 0.0;
-	}
-}
-
-// Get the sample state given an index and a delta position/velocity
-void DIHairStrands_GetSampleState(in FDIHairStrandsContext DIContext, in float3 NodePosition, in float3 NodeVelocity, in float3 DeltaPosition, in float3 DeltaVelocity, 
-			in int NumSamples, in int SampleIndex, out float3 OutSamplePosition, out float3 OutSampleVelocity)
-{
-	const int LocalIndex = GGroupThreadId.x % DIContext.StrandSize;
-
-	if( LocalIndex > 0 )
-	{
-		const float SampleCoord = (0.5+SampleIndex) / NumSamples;
-		OutSamplePosition = NodePosition + SampleCoord * DeltaPosition;
-		OutSampleVelocity = NodeVelocity + SampleCoord * DeltaVelocity;
-	}
-	else
-	{
-		OutSamplePosition = float3(0,0,0);
-		OutSampleVelocity = float3(0,0,0);
-	}
-}
-
-/* -----------------------------------------------------------------
- * Edge volume, length, rotation, direction
- * -----------------------------------------------------------------
- */
-
-// Compute the edge volume value
-void DIHairStrands_ComputeEdgeVolume(in FDIHairStrandsContext DIContext, in float3 NodePosition, out float OutEdgeVolume)
-{
-	SharedNodePosition[GGroupThreadId.x] = NodePosition;
-	GroupMemoryBarrier(); 
-
-	const int LocalIndex = (GGroupThreadId.x % DIContext.StrandSize);
-
-	// L = ||P1-P0||
-	OutEdgeVolume = 0.0;
-	if(LocalIndex > 2)
-	{
-		const int NodeIndexA = GGroupThreadId.x;
-		const int NodeIndexB = GGroupThreadId.x-1;
-		const int NodeIndexC = GGroupThreadId.x-2;
-		const int NodeIndexD = GGroupThreadId.x-3;
-
-		const float3 EdgeVectorA = SharedNodePosition[NodeIndexB] - SharedNodePosition[NodeIndexA];
-		const float3 EdgeVectorB = SharedNodePosition[NodeIndexC] - SharedNodePosition[NodeIndexA];
-		const float3 EdgeVectorC = SharedNodePosition[NodeIndexD] - SharedNodePosition[NodeIndexA];
-
-		OutEdgeVolume = dot(cross(EdgeVectorB,EdgeVectorC),EdgeVectorA);
-	}
-}
-
-// Compute the edge length value
-void DIHairStrands_ComputeEdgeLength(in FDIHairStrandsContext DIContext, in float3 NodePosition, in int NodeOffset, out float OutEdgeLength)
-{
-	SharedNodePosition[GGroupThreadId.x] = NodePosition;
-	GroupMemoryBarrier(); 
-
-	const int LocalIndex = (GGroupThreadId.x % DIContext.StrandSize);
-
-	// L = ||P1-P0||
-	OutEdgeLength = (LocalIndex>NodeOffset) ? length(SharedNodePosition[GGroupThreadId.x]- SharedNodePosition[GGroupThreadId.x-1-NodeOffset]) : 0.0;
-}
-
-// Compute the edge darboux vector (diff between consecutive edge orientations)
-void DIHairStrands_ComputeEdgeRotation(in FDIHairStrandsContext DIContext, in float4 NodeOrientation, out float4 OutEdgeRotation)
-{
-	SharedNodeOrientation[GGroupThreadId.x] = NodeOrientation;
-	GroupMemoryBarrier(); 
-
-	const int LocalIndex = (GGroupThreadId.x % DIContext.StrandSize);
-
-	// D = Q0^-1 * Q1
-	if(LocalIndex>0)
-	{
-		float4 q0 = SharedNodeOrientation[GGroupThreadId.x-1];
-		float4 q1 = SharedNodeOrientation[GGroupThreadId.x];
-		OutEdgeRotation = float4(
-				q1.xyz * q0.w - q0.xyz * q1.w + cross(-q0.xyz, q1.xyz),
-				q0.w * q1.w - dot(-q0.xyz, q1.xyz));
-		float4 OmegaPlus = OutEdgeRotation + float4(0,0,0,1);
-		float4 OmegaMinus = OutEdgeRotation - float4(0,0,0,1);
-		if( dot(OmegaMinus,OmegaMinus) > dot(OmegaPlus,OmegaPlus) ) OutEdgeRotation = -OutEdgeRotation;
-		//OutEdgeRotation = MultiplyQuat(InverseQuat(SharedNodeOrientation[GGroupThreadId.x-1]),SharedNodeOrientation[GGroupThreadId.x]);
-	}
-	else
-	{
-		OutEdgeRotation = QUATERNION_IDENTITY;
-	}
-}
-
-void DIHairStrands_ComputeEdgeDirection(in FDIHairStrandsContext DIContext, in float3 NodePosition, in float4 NodeOrientation, 
-	in float NodeMass, in float3 GravityVector, in float GravityPreloading, in float BendStiffness, in float StrandThickness, in float RestLength, out float3 OutRestDirection)
-{
-	const int LocalIndex = (GGroupThreadId.x % DIContext.StrandSize);
-
-	SharedNodePosition[GGroupThreadId.x] = NodePosition;
-	SharedNodeOrientation[GGroupThreadId.x] = NodeOrientation;
-	GroupMemoryBarrier(); 
-
-	OutRestDirection = float3(0,0,0);
-
-	if( LocalIndex > 1 )
-	{
-		float3 EdgeDirection = SharedNodePosition[GGroupThreadId.x] - SharedNodePosition[GGroupThreadId.x-1];
-
-		SharedInverseInertia[GGroupThreadId.x] = BendStiffness*PI*StrandThickness*StrandThickness / (4.0*RestLength);
-		SharedPreviousPosition[GGroupThreadId.x] = EdgeDirection;
-		SharedPreviousOrientation[GGroupThreadId.x] = float4(EdgeDirection,0);
-
-		GroupMemoryBarrier(); 
-
-		const float GradientNext = ( LocalIndex < (DIContext.StrandSize-1) ) ?  SharedInverseInertia[GGroupThreadId.x+1]  : 0.0;
-		const float GradientPrev = SharedInverseInertia[GGroupThreadId.x];
-		const float SchurComplement =  (GradientNext * GradientNext + GradientPrev * GradientPrev);
-
-		for(int Index = 0; Index < 100; ++Index)
-		{
-			const int IsRed = (GGroupThreadId.x % 2) == 0;
-			if (!IsRed)
-			{	
-				const float3 DeltaConstraint = ((SharedPreviousPosition[GGroupThreadId.x] - SharedPreviousOrientation[GGroupThreadId.x].xyz) * GradientPrev - 
-												(SharedPreviousPosition[GGroupThreadId.x+1] - SharedPreviousOrientation[GGroupThreadId.x+1].xyz) * GradientNext) - NodeMass * GravityVector;
-
-				SharedPreviousOrientation[GGroupThreadId.x].xyz += DeltaConstraint * GradientPrev / SchurComplement;
-				SharedPreviousOrientation[GGroupThreadId.x+1].xyz -= DeltaConstraint * GradientNext / SchurComplement;
-			}
-			GroupMemoryBarrier(); 
-
-			if (IsRed)
-			{
-				const float3 DeltaConstraint = ((SharedPreviousPosition[GGroupThreadId.x] - SharedPreviousOrientation[GGroupThreadId.x].xyz) * GradientPrev - 
-												(SharedPreviousPosition[GGroupThreadId.x+1] - SharedPreviousOrientation[GGroupThreadId.x+1].xyz) * GradientNext) - NodeMass * GravityVector;
-
-
-				SharedPreviousOrientation[GGroupThreadId.x].xyz += DeltaConstraint * GradientPrev / SchurComplement;
-				SharedPreviousOrientation[GGroupThreadId.x+1].xyz -= DeltaConstraint * GradientNext / SchurComplement;
-			}
-			GroupMemoryBarrier();
-		}
-		EdgeDirection = SharedPreviousOrientation[GGroupThreadId.x].xyz * GravityPreloading + (1.0-GravityPreloading) * EdgeDirection;
-		OutRestDirection = UnrotateVectorByQuat(EdgeDirection,SharedNodeOrientation[GGroupThreadId.x-2]);
-	}
-
-	GroupMemoryBarrier(); 
-}
-
-
-/* -----------------------------------------------------------------
- * Points position update
- * -----------------------------------------------------------------
- */
-
-// Reset the deformed points position to the rest ones
-void DIHairStrands_ResetPointPosition(in FDIHairStrandsContext DIContext, out bool ReportStatus) 
-{
-	int LocalIndex = GGroupThreadId.x % DIContext.StrandSize;
-	int StrandIndex = ExecIndex() / DIContext.StrandSize;
-
-	ReportStatus = false;
-
-	if(LocalIndex == 0)
-	{
-		const int PointBegin = DIContext.CurvesOffsetsBuffer[StrandIndex];
-		const int PointEnd = DIContext.CurvesOffsetsBuffer[StrandIndex+1]; 
-		const int PointCount = PointEnd - PointBegin;
-
-		const uint4 PositionOffset = uint4( f32tof16(DIContext.RestPositionOffset - DIContext.DeformedPositionOffset[0].xyz), 0 );
-
-		ReportStatus = true;
-		for (int PointIndex = 0; PointIndex < PointCount; ++PointIndex)
-		{
-			DIContext.DeformedPositionBuffer[PointBegin+PointIndex] = DIContext.RestPositionBuffer[PointBegin+PointIndex] + PositionOffset;
-		}
-	}
-}
-
-// Compute the projection triangle
-void DIHairStrands_ComputeProjectionTriangle(in float2 ProjectionUV, in float3 PA, in float3 PB, in float3 PC, in float3 Offset, 
-			out float3 OutTrianglePosition, out float4 OutTriangleRotation)
-{
-	OutTrianglePosition = PA * ProjectionUV.x + PB * ProjectionUV.y + PC * (1.0 - ProjectionUV.x - ProjectionUV.y) + Offset;
-
-	float3 TangentU = PB - PA;
-	float3 TangentV = PC - PA;
-
-	const float3 Normal = normalize(cross(TangentU, TangentV));
-	TangentV = normalize(TangentV);
-	TangentU = normalize(cross(TangentV, Normal));
-	
-	const float3 RotationMatrix[3] = { TangentU, TangentV, Normal };
-	OutTriangleRotation = QuatFromMatrix(RotationMatrix);
-}
-
-// Get the deformed triangle
-void DIHairStrands_BuildDeformedTriangle(in FDIHairStrandsContext DIContext, in float2 ProjectionUV, in int StrandIndex, out float3 OutTrianglePosition, out float4 OutTriangleRotation)
-{
-	DIHairStrands_ComputeProjectionTriangle(ProjectionUV, DIContext.DeformedTrianglePositionABuffer[StrandIndex].xyz,
-		DIContext.DeformedTrianglePositionBBuffer[StrandIndex].xyz, 
-		DIContext.DeformedTrianglePositionCBuffer[StrandIndex].xyz, 
-	    DIContext.DeformedRootOffset, OutTrianglePosition, OutTriangleRotation);
-}
-
-// Get the rest triangle
-void DIHairStrands_BuildRestTriangle(in FDIHairStrandsContext DIContext, in float2 ProjectionUV, in int StrandIndex, out float3 OutTrianglePosition, out float4 OutTriangleRotation)
-{
-	DIHairStrands_ComputeProjectionTriangle(ProjectionUV, DIContext.RestTrianglePositionABuffer[StrandIndex].xyz,
-		DIContext.RestTrianglePositionBBuffer[StrandIndex].xyz,
-		DIContext.RestTrianglePositionCBuffer[StrandIndex].xyz,
-	    DIContext.RestRootOffset, OutTrianglePosition, OutTriangleRotation);
-}
-
-// Eval the triangle local position
-float3 DIHairStrands_TriangleLocalPosition(in float3 TrianglePosition, in float4 TriangleRotation, in float3 WorldPosition)
-{
-	return RotateVectorByQuat(WorldPosition - TrianglePosition, InverseQuat(TriangleRotation));
-}
-
-// Eval the triangle local orientation
-float4 DIHairStrands_TriangleLocalOrientation(in float3 TrianglePosition, in float4 TriangleRotation, in float4 WorldOrientation)
-{
-	return NormalizeQuat(MultiplyQuat(InverseQuat(TriangleRotation), WorldOrientation));
-}
-
-// Eval the triangle world position
-float3 DIHairStrands_TriangleWorldPosition(in float3 TrianglePosition, in float4 TriangleRotation, in float3 LocalPosition)
-{
-	return RotateVectorByQuat(LocalPosition, TriangleRotation) + TrianglePosition;
-}
-
-// Eval the triangle local orientation
-float4 DIHairStrands_TriangleWorldOrientation(in float3 TrianglePosition, in float4 TriangleRotation, in float4 LocalOrientation)
-{
-	return NormalizeQuat(MultiplyQuat(TriangleRotation, LocalOrientation));
-}
-
-// Report interpolated nodes displacements onto the points positions 
-void DIHairStrands_UpdateTriangleDisplace(in FDIHairStrandsContext DIContext, in int StrandIndex, in float3 NodePosition, in float3 RestPosition, 
-		out float3 RestTrianglePosition, out float4 RestTriangleOrientation, out float3 DeformedTrianglePosition, out float4 DeformedTriangleOrientation)
-{
-	if( DIContext.InterpolationMode >= 1)
-	{
-		const float2 ProjectionUV = ExtractBarycentrics(DIContext.RootBarycentricCoordinatesBuffer[StrandIndex]);
-	
-		DIHairStrands_BuildRestTriangle(DIContext, ProjectionUV, StrandIndex, RestTrianglePosition, RestTriangleOrientation);
-		DIHairStrands_BuildDeformedTriangle(DIContext, ProjectionUV, StrandIndex, DeformedTrianglePosition, DeformedTriangleOrientation);
-	
-		const float3 LocalRestPosition = DIHairStrands_TriangleLocalPosition(RestTrianglePosition, RestTriangleOrientation, RestPosition);
-		const float3 LocalDeformedPosition = DIHairStrands_TriangleLocalPosition(DeformedTrianglePosition, DeformedTriangleOrientation, NodePosition);
-	
-		SharedNodePosition[GGroupThreadId.x] = LocalDeformedPosition - LocalRestPosition;
-		GroupMemoryBarrier();
-	}
-	else
-	{
-		SharedNodePosition[GGroupThreadId.x] = NodePosition - RestPosition;
-		GroupMemoryBarrier();
-	}
-}
-
-// Report interpolated nodes displacements onto the points positions 
-void DIHairStrands_UpdateDeformedPositions(in FDIHairStrandsContext DIContext, in float3 LocalDisplace, in int PositionIndex, 
-		in float3 RestTrianglePosition, in float4 RestTriangleOrientation, in float3 DeformedTrianglePosition, in float4 DeformedTriangleOrientation)
-{
-	if( DIContext.InterpolationMode >= 1)
-	{
-		const float3 LocalTrianglePosition = DIHairStrands_TriangleLocalPosition(RestTrianglePosition, RestTriangleOrientation, 
-			f16tof32(DIContext.RestPositionBuffer[PositionIndex].xyz) + DIContext.RestPositionOffset) + LocalDisplace;
-			
-		const float3 LocalComponentPosition = DIHairStrands_TriangleWorldPosition(DeformedTrianglePosition, DeformedTriangleOrientation, LocalTrianglePosition) - DIContext.DeformedPositionOffset[0].xyz;
-		DIContext.DeformedPositionBuffer[PositionIndex] = uint4(f32tof16(LocalComponentPosition), DIContext.RestPositionBuffer[PositionIndex].w);
-		
-	}
-	else
-	{
-		const float3 LocalComponentPosition = f16tof32(DIContext.RestPositionBuffer[PositionIndex].xyz) + LocalDisplace + DIContext.RestPositionOffset - DIContext.DeformedPositionOffset[0].xyz;
-		DIContext.DeformedPositionBuffer[PositionIndex] = uint4(f32tof16(LocalComponentPosition), DIContext.RestPositionBuffer[PositionIndex].w);	
-	}
-}
-
-// Report interpolated nodes displacements onto the points positions 
-void DIHairStrands_UpdatePointPosition(in FDIHairStrandsContext DIContext, in float3 NodePosition, in float3 RestPosition, out bool OutReportStatus)
-{
-	OutReportStatus = false;
-	const int ReducedSize = DIContext.StrandSize - 1;
-	
-	int LocalIndex = GGroupThreadId.x % DIContext.StrandSize - 1;
-	int StrandIndex = ExecIndex() / DIContext.StrandSize;
-
-	float3 RestTrianglePosition = float3(0, 0, 0), DeformedTrianglePosition = float3(0, 0, 0);
-	float4 RestTriangleOrientation = QUATERNION_IDENTITY, DeformedTriangleOrientation = QUATERNION_IDENTITY;
-
-	DIHairStrands_UpdateTriangleDisplace(DIContext, StrandIndex, NodePosition, RestPosition, 
-			RestTrianglePosition, RestTriangleOrientation, DeformedTrianglePosition, DeformedTriangleOrientation);
-	
-	const int PointBegin = DIContext.CurvesOffsetsBuffer[StrandIndex];
-	const int PointEnd = DIContext.CurvesOffsetsBuffer[StrandIndex + 1];
-	const int PointCount = PointEnd - PointBegin;
-
-
-	for (int PointIndex = 0; PointIndex < PointCount; ++PointIndex)
-	{
-		const float PointCoord = (float) (PointIndex) / (PointCount - 1.0);
-		const float PointNode = PointCoord * (ReducedSize - 1.0);
-
-		const int NodePrev = (PointIndex == 0) ? 0 : (PointIndex == (PointCount - 1.0)) ? ReducedSize - 2 : floor(PointNode);
-		const int NodeNext = NodePrev + 1;
-
-		if (NodePrev == LocalIndex)
-		{
-			OutReportStatus = true;
-			const float PointAlpha = PointNode - (float) NodePrev;
-			const float3 LocalDisplace = SharedNodePosition[GGroupThreadId.x] * (1.0 - PointAlpha) + SharedNodePosition[GGroupThreadId.x + 1] * PointAlpha;
-
-			DIHairStrands_UpdateDeformedPositions(DIContext, LocalDisplace, PointBegin + PointIndex, 
-				RestTrianglePosition, RestTriangleOrientation, DeformedTrianglePosition, DeformedTriangleOrientation);
-		}
-	}
-}
-
-// Get the point position
-void DIHairStrands_GetPointPosition(in FDIHairStrandsContext DIContext, in int PointIndex, out float3 OutPointPosition)
-{
-	OutPointPosition = mul(float4(f16tof32(DIContext.DeformedPositionBuffer[PointIndex].xyz) + DIContext.DeformedPositionOffset[0].xyz, 1.0), DIContext.WorldTransform).xyz;
-}
-
-/* -----------------------------------------------------------------
- * Nodes time integration
- * -----------------------------------------------------------------
- */
-
-// Add external force to the linear velocity and advect node position
-void DIHairStrands_AdvectNodePosition(in FDIHairStrandsContext DIContext, in float NodeMass, in bool IsPositionMobile, in float3 ExternalForce, in float3 ForceGradient, in float DeltaTime, inout float3 OutLinearVelocity, inout float3 OutNodePosition)
-{	
-	if(IsPositionMobile && NodeMass != 0.0)
-	{
-		const float3 ImplicitGradient =  float3(NodeMass,NodeMass,NodeMass) / DeltaTime - ForceGradient;
-		const float3 InverseGradient = float3(1.0/ImplicitGradient.x, 1.0/ImplicitGradient.y, 1.0/ImplicitGradient.z);
-		
-		if(DIContext.LocalSimulation)
-		{
-			OutLinearVelocity += InverseGradient * ExternalForce - DIContext.BoneLinearAcceleration * DeltaTime - ( cross(DIContext.BoneAngularAcceleration, OutNodePosition) +
-				2.0 * cross(DIContext.BoneAngularVelocity, OutLinearVelocity) + cross(DIContext.BoneAngularVelocity, cross(DIContext.BoneAngularVelocity,OutNodePosition))) * DeltaTime;
-		}
-		else
-		{
-			OutLinearVelocity += InverseGradient * ExternalForce;
-		}
-		OutNodePosition += OutLinearVelocity * DeltaTime;
-	}
-}
-
-// Add external torque to the angular velocity and advect node orientation
-void DIHairStrands_AdvectNodeOrientation(in FDIHairStrandsContext DIContext, in float3 NodeInertia, in bool IsOrientationMobile, in float3 ExternalTorque, in float3 TorqueGradient, in float DeltaTime, inout float3 OutAngularVelocity, inout float4 OutNodeOrientation)
-{	
-	if(IsOrientationMobile && NodeInertia.x != 0.0 && NodeInertia.y != 0.0 && NodeInertia.z != 0.0)
-	{
-		const float3 ImplicitGradient = NodeInertia / DeltaTime - TorqueGradient;
-		const float3 InverseGradient = float3(1.0/ImplicitGradient.x, 1.0/ImplicitGradient.y, 1.0/ImplicitGradient.z);
-
-		OutAngularVelocity += InverseGradient * (ExternalTorque - cross(OutAngularVelocity, NodeInertia * OutAngularVelocity));
-
-		OutNodeOrientation = OutNodeOrientation + 0.5 * DeltaTime * float4(
-			OutAngularVelocity.xyz * OutNodeOrientation.w + cross(OutAngularVelocity.xyz, OutNodeOrientation.xyz),
-			- dot(OutAngularVelocity.xyz, OutNodeOrientation.xyz));
-
-		OutNodeOrientation = NormalizeQuat(OutNodeOrientation);
-	}
-}
-
-// Update the node linear velocity based on the node position difference
-void DIHairStrands_UpdateLinearVelocity(in FDIHairStrandsContext DIContext, in float3 PreviousPosition, in float3 NodePosition, in float DeltaTime, out float3 OutLinearVelocity)
-{	
-	OutLinearVelocity = (NodePosition-PreviousPosition) / DeltaTime;
-}
-
-
-// Update the node angular velocity based on the node orientation difference
-void DIHairStrands_UpdateAngularVelocity(in FDIHairStrandsContext DIContext, in float4 PreviousOrientation, in float4 NodeOrientation, in float DeltaTime, out float3 OutAngularVelocity)
-{
-	const float4 DeltaQuat = MultiplyQuat(NodeOrientation,InverseQuat(PreviousOrientation));
-	const float AxisLength = length( DeltaQuat.xyz );
-
-	if (AxisLength<SMALL_NUMBER)
-	{
-		 OutAngularVelocity =  DeltaQuat.xyz * 2.0 / DeltaTime;
-	}
-	else
-	{
-		const float QuatAngle = 2.0 * atan2(AxisLength,DeltaQuat.w );
-		OutAngularVelocity = DeltaQuat.xyz * QuatAngle / (AxisLength*DeltaTime);
-	}
-}
-
-/* -----------------------------------------------------------------
- * Update node position orientation
- * -----------------------------------------------------------------
- */
-
-// Compute rest position
-void DIHairStrands_ComputeRestPosition(in FDIHairStrandsContext DIContext, in float3 NodePosition, out float3 OutRestPosition)
-{	
-	OutRestPosition = mul(float4( NodePosition, 1.0), DIContext.WorldInverse).xyz;
-}
-
-// Compute rest orientation
-void DIHairStrands_ComputeRestOrientation(in FDIHairStrandsContext DIContext, in float4 NodeOrientation, out float4 OutRestOrientation)
-{	
-	OutRestOrientation = NormalizeQuat( MultiplyQuat(InverseQuat(DIContext.WorldRotation),NodeOrientation) );
-}
-
-// Update Node Position
-void DIHairStrands_AttachNodePosition(in FDIHairStrandsContext DIContext, in float3 RestPosition, out float3 OutNodePosition)
-{
-	OutNodePosition = mul(float4( RestPosition, 1.0), DIContext.WorldTransform).xyz;
-}
-
-// Update Node Orientation
-void DIHairStrands_AttachNodeOrientation(in FDIHairStrandsContext DIContext, in float4 RestOrientation, out float4 OutNodeOrientation)
-{
-	OutNodeOrientation = NormalizeQuat( MultiplyQuat(DIContext.WorldRotation,RestOrientation) );
-}
-
-// Get World Vector
-void DIHairStrands_GetWorldVector(in FDIHairStrandsContext DIContext, in float3 LocalVector, in bool IsPosition, out float3 OutWorldVector)
-{	
-	OutWorldVector = DIContext.LocalSimulation ? mul(float4( LocalVector, (float)IsPosition), DIContext.BoneTransform).xyz : LocalVector;
-}
-
-// Get Local Vector
-void DIHairStrands_GetLocalVector(in FDIHairStrandsContext DIContext, in float3 WorldVector, in bool IsPosition, out float3 OutLocalVector)
-{
-	OutLocalVector = DIContext.LocalSimulation ? mul(float4( WorldVector, (float)IsPosition), DIContext.BoneInverse).xyz : WorldVector;
-}
-
-
-/* -----------------------------------------------------------------
- * Attached the root to the skinned cache
- * -----------------------------------------------------------------
- */
-// Duplicate from #include "HairStrands/HairStrandsMeshRBF.ush"
-float3 DIHairStrands_ApplyRBF(in FDIHairStrandsContext DIContext, in float3 RestControlPoint)
-{
-	float3 ControlPoint = RestControlPoint;
-	// Apply rbf interpolation from the samples set
-	for (int i = 0; i < DIContext.SampleCount; ++i)
-	{
-		const float3 PositionDelta = RestControlPoint - DIContext.RestSamplePositionsBuffer[i].xyz;
-		const float FunctionValue = sqrt(dot(PositionDelta, PositionDelta) + 1);
-		ControlPoint += FunctionValue * DIContext.MeshSampleWeightsBuffer[i].xyz;
-	}
-	ControlPoint += DIContext.MeshSampleWeightsBuffer[DIContext.SampleCount].xyz;
-	ControlPoint += DIContext.MeshSampleWeightsBuffer[DIContext.SampleCount + 1].xyz * RestControlPoint.x;
-	ControlPoint += DIContext.MeshSampleWeightsBuffer[DIContext.SampleCount + 2].xyz * RestControlPoint.y;
-	ControlPoint += DIContext.MeshSampleWeightsBuffer[DIContext.SampleCount + 3].xyz * RestControlPoint.z;
-	return ControlPoint;
-}
-
- // Compute the global position
-void DIHairStrands_EvalGlobalPosition(in FDIHairStrandsContext DIContext, in float3 RestPosition, out float3 OutGlobalPosition)
-{
-	const float3 RestSkinnedPosition = RestPosition;
-	OutGlobalPosition = DIHairStrands_ApplyRBF(DIContext, RestPosition);
-}
-
-// Compute the global position
-void DIHairStrands_EvalGlobalOrientation(in FDIHairStrandsContext DIContext, in float3 GlobalPosition, inout float4 OutGlobalOrientation)
-{
-	SharedNodePosition[GGroupThreadId.x] = GlobalPosition;
-	SharedNodeOrientation[GGroupThreadId.x] = OutGlobalOrientation;
-	GroupMemoryBarrier(); 
-
-	DIHairStrands_UpdateRootOrientation(DIContext);
-
-	ComputeMaterialFrame(DIContext.StrandSize);
-	OutGlobalOrientation = 	SharedNodeOrientation[GGroupThreadId.x];
-}
-
-// Compute rest position and orientation
-void DIHairStrands_ComputeLocalState(in FDIHairStrandsContext DIContext, in float3 RestPosition, in float4 RestOrientation, 
-												out float3 OutLocalPosition, out float4 OutLocalOrientation)
-{
-	OutLocalPosition = RestPosition;
-	OutLocalOrientation = RestOrientation;
-
-	if( DIContext.InterpolationMode >= 1)
-	{
-		const int StrandIndex = ExecIndex() / DIContext.StrandSize;
-		const float2 ProjectionUV = ExtractBarycentrics(DIContext.RootBarycentricCoordinatesBuffer[StrandIndex]);
-		
-		float3 RestTrianglePosition = float3(0, 0, 0);
-		float4 RestTriangleOrientation = QUATERNION_IDENTITY;
-		DIHairStrands_BuildRestTriangle(DIContext, ProjectionUV, StrandIndex, RestTrianglePosition, RestTriangleOrientation);
-		
-		OutLocalPosition = DIHairStrands_TriangleLocalPosition(RestTrianglePosition, RestTriangleOrientation, RestPosition);
-		OutLocalOrientation = DIHairStrands_TriangleLocalOrientation(RestTrianglePosition, RestTriangleOrientation, RestOrientation);
-	}
-}
-
-// Update Node Position and orientation
-void DIHairStrands_AttachNodeState(in FDIHairStrandsContext DIContext, in float3 LocalPosition, in float4 LocalOrientation, out float3 OutNodePosition, out float4 OutNodeOrientation)
-{	
-	OutNodePosition = LocalPosition;
-	OutNodeOrientation = LocalOrientation;
-
-	if( DIContext.InterpolationMode >= 1)
-	{
-		const int StrandIndex = ExecIndex() / DIContext.StrandSize;
-		const float2 ProjectionUV = ExtractBarycentrics(DIContext.RootBarycentricCoordinatesBuffer[StrandIndex]);
-		
-		float3 DeformedTrianglePosition = float3(0, 0, 0);
-		float4 DeformedTriangleOrientation = QUATERNION_IDENTITY;
-		DIHairStrands_BuildDeformedTriangle(DIContext, ProjectionUV, StrandIndex, DeformedTrianglePosition, DeformedTriangleOrientation);
-		
-		OutNodePosition = DIHairStrands_TriangleWorldPosition(DeformedTrianglePosition, DeformedTriangleOrientation, LocalPosition);
-		OutNodeOrientation = DIHairStrands_TriangleWorldOrientation(DeformedTrianglePosition, DeformedTriangleOrientation, LocalOrientation);
-	}
-}
-
-void DIHairStrands_UpdateNodeState(in FDIHairStrandsContext DIContext, in float3 RestPosition, in float3 NodePosition, in float4 NodeOrientation, out float3 OutNodePosition, out float4 OutNodeOrientation)
-{
-	float3 GlobalPosition =  NodePosition;
-	float4 GlobalOrientation = NodeOrientation;
-
-	if(DIContext.InterpolationMode == 2)
-	{
-		DIHairStrands_EvalGlobalPosition(DIContext, RestPosition, GlobalPosition);
-		DIHairStrands_EvalGlobalOrientation(DIContext,GlobalPosition,GlobalOrientation);
-	}
-	DIHairStrands_AttachNodePosition(DIContext,GlobalPosition,OutNodePosition);
-	DIHairStrands_AttachNodeOrientation(DIContext,GlobalOrientation,OutNodeOrientation);
-}
-
-
-
-
-
-
-
-
-
-
-
-
diff --git a/Engine/Plugins/Runtime/HairStrands/Shaders/Private/NiagaraDataInterfacePressureGrid.ush b/Engine/Plugins/Runtime/HairStrands/Shaders/Private/NiagaraDataInterfacePressureGrid.ush
--- a/Engine/Plugins/Runtime/HairStrands/Shaders/Private/NiagaraDataInterfacePressureGrid.ush
+++ b/Engine/Plugins/Runtime/HairStrands/Shaders/Private/NiagaraDataInterfacePressureGrid.ush
@@ -6,11 +6,6 @@ NiagaraDataInterfacePressureGrid.ush
 
 #pragma once
 
-/* -----------------------------------------------------------------
- * Pressure Grid constants and context
- * -----------------------------------------------------------------
- */
-
 static const int FLUID_PRESSURE_OFFSET = 6;
 static const int SOLID_WEIGHT_OFFSET = 7;
 static const int SOLID_VELOCITY_OFFSET = 10;
@@ -18,557 +13,3 @@ static const int SOLID_DISTANCE_OFFSET = 13;
 static const int FLUID_DISTANCE_OFFSET = 14;
 static const int FLUID_DENSITY_OFFSET = 15;
 
-/* -----------------------------------------------------------------
- * Scale the velocity, mass and distance buffer from the current to the destination buffer
- * -----------------------------------------------------------------
- */
-
-void DIVelocityGrid_ScaleCellDistance(in FDIVelocityGridContext DIContext, in RWBuffer<int> DIContext_GridDestinationBuffer, in int3 GridIndex, in float GridLength )
-{
-	float FluidDistance = DIVelocityGrid_ConvertGridScalar(DIContext,GridIndex,FLUID_DISTANCE_OFFSET);
-
-	[branch]
-	if(FluidDistance < 0.5 * GridLength)
-	{
-		const float SolidDistance = (DIVelocityGrid_GetGridScalar(DIContext,int3(GridIndex.x+0,GridIndex.y+0,GridIndex.z+0),SOLID_DISTANCE_OFFSET) + 
-									 DIVelocityGrid_GetGridScalar(DIContext,int3(GridIndex.x+0,GridIndex.y+0,GridIndex.z+1),SOLID_DISTANCE_OFFSET) +
-									 DIVelocityGrid_GetGridScalar(DIContext,int3(GridIndex.x+0,GridIndex.y+1,GridIndex.z+0),SOLID_DISTANCE_OFFSET) +
-									 DIVelocityGrid_GetGridScalar(DIContext,int3(GridIndex.x+0,GridIndex.y+1,GridIndex.z+1),SOLID_DISTANCE_OFFSET) +
-									 DIVelocityGrid_GetGridScalar(DIContext,int3(GridIndex.x+1,GridIndex.y+0,GridIndex.z+0),SOLID_DISTANCE_OFFSET) +
-									 DIVelocityGrid_GetGridScalar(DIContext,int3(GridIndex.x+1,GridIndex.y+0,GridIndex.z+1),SOLID_DISTANCE_OFFSET) +
-									 DIVelocityGrid_GetGridScalar(DIContext,int3(GridIndex.x+1,GridIndex.y+1,GridIndex.z+0),SOLID_DISTANCE_OFFSET) +
-									 DIVelocityGrid_GetGridScalar(DIContext,int3(GridIndex.x+1,GridIndex.y+1,GridIndex.z+1),SOLID_DISTANCE_OFFSET) ) / 8.0;
-
-		if(SolidDistance < 0.0) FluidDistance = -0.5 * GridLength;
-	}
-
-	DIVelocityGrid_SetGridScalar(DIContext,DIContext_GridDestinationBuffer,GridIndex,FLUID_DISTANCE_OFFSET,FluidDistance);
-}
-
-void DIVelocityGrid_ScaleFaceVelocity(in FDIVelocityGridContext DIContext, in RWBuffer<int> DIContext_GridDestinationBuffer, in int3 GridIndex, in int ComponentIndex)
-{
-	const float GridMomentum = DIVelocityGrid_ConvertGridScalar(DIContext,GridIndex,FLUID_VELOCITY_OFFSET+ComponentIndex);
-	const float GridMass = DIVelocityGrid_ConvertGridScalar(DIContext,GridIndex,FLUID_MASS_OFFSET+ComponentIndex);
-
-	float GridVelocity = (GridMass != 0.0) ? GridMomentum / GridMass : 0.0;
-
-	DIVelocityGrid_SetGridScalar(DIContext,DIContext_GridDestinationBuffer,GridIndex,FLUID_VELOCITY_OFFSET+ComponentIndex,GridVelocity);
-	DIVelocityGrid_SetGridScalar(DIContext,DIContext_GridDestinationBuffer,GridIndex,FLUID_MASS_OFFSET+ComponentIndex,GridMass);
-}
-
-void DIVelocityGrid_ScaleCellFields(in FDIVelocityGridContext DIContext, in RWBuffer<int> DIContext_GridDestinationBuffer, in int GridHash, in float GridLength, in float DeltaTime, out bool OutTransferStatus)
-{
-	const int3 GridIndex = DIVelocityGrid_GetGridIndex(DIContext,GridHash);
-
-	OutTransferStatus = false;
-	if(DIVelocityGrid_IsFaceValid(DIContext,GridIndex,0))
-	{
-		DIVelocityGrid_ScaleFaceVelocity(DIContext,DIContext_GridDestinationBuffer,GridIndex,0);
-	}
-	if(DIVelocityGrid_IsFaceValid(DIContext,GridIndex,1))
-	{
-		DIVelocityGrid_ScaleFaceVelocity(DIContext,DIContext_GridDestinationBuffer,GridIndex,1);
-	}
-	if(DIVelocityGrid_IsFaceValid(DIContext,GridIndex,2))
-	{
-		DIVelocityGrid_ScaleFaceVelocity(DIContext,DIContext_GridDestinationBuffer,GridIndex,2);
-	}
-	if(DIVelocityGrid_IsCellValid(DIContext,GridIndex))
-	{
-		DIVelocityGrid_ScaleCellDistance(DIContext,DIContext_GridDestinationBuffer,GridIndex,GridLength);
-	}
-	if(DIVelocityGrid_IsCellValid(DIContext,GridIndex))
-	{
-		const float SourceDensity = DIVelocityGrid_ConvertGridScalar(DIContext,GridIndex,FLUID_DENSITY_OFFSET);
-
-		const float CurrentDensity = max(1.0,SourceDensity);
-		//const float DensityRatio = GridLength * (1.0-CurrentDensity) / DeltaTime;
-		const float DensityRatio = (CurrentDensity != 0.0) ? GridLength * log(1.0/CurrentDensity) / DeltaTime : 0.0;
-
-		DIVelocityGrid_SetGridScalar(DIContext,DIContext_GridDestinationBuffer,GridIndex,FLUID_DENSITY_OFFSET,DensityRatio);
-	}
-}
-
-/* -----------------------------------------------------------------
- * Build Distance/Density Field from particles
- * -----------------------------------------------------------------
- */
- 
-// Compute the distance from the particle to the grid
-void DIVelocityGrid_BuildDistanceField(in FDIVelocityGridContext DIContext, in RWBuffer<int> DIContext_GridDestinationBuffer, in float3 GridOrigin, in float GridLength, in float3 ParticlePosition, out bool OutFunctionStatus )
-{
-	OutFunctionStatus = true;
-	const float3 GridPosition = DIVelocityGrid_GetGridPosition(DIContext,ParticlePosition,GridOrigin);
-
-	float3 GridFraction = float3(0,0,0);
-	const int3 GridIndex = DIVelocityGrid_ComputeGridIndex(DIContext,GridPosition,float3(0.5,0.5,0.5),GridLength,GridFraction);
-
-	const float NodeRadius = GridLength * 0.5 * 1.414213;
-
-	for(int i = -2; i <= 2; ++i)
-	for(int j = -2; j <= 2; ++j)
-	for(int k = -2; k <= 2; ++k)
-	{
-		const int3 GridLocation = int3(GridIndex.x+i,GridIndex.y+j,GridIndex.z+k);
-
-		if( DIVelocityGrid_IsCellValid(DIContext,GridLocation) )
-		{
-			const float3 CellPosition = (GridLocation+float3(0.5,0.5,0.5))*GridLength+GridOrigin;
-			const float CellDistance = length(CellPosition-ParticlePosition) - NodeRadius; 
-
-			DIVelocityGrid_InterlockedMinScalar(DIContext,DIContext_GridDestinationBuffer,GridLocation,FLUID_DISTANCE_OFFSET,CellDistance);
-		}
-	}
-}
-
-// Transfer the density from the particle to the grid
-void DIVelocityGrid_BuildDensityField(in FDIVelocityGridContext DIContext, in RWBuffer<int> DIContext_GridDestinationBuffer, in float3 GridOrigin, in float GridLength,
-				in float3 ParticlePosition, in float ParticleMass, in float ParticleDensity, out bool OutFunctionStatus )
-{
-	OutFunctionStatus = true;
-	const float3 GridPosition = DIVelocityGrid_GetGridPosition(DIContext,ParticlePosition,GridOrigin);
-
-	float3 GridFraction = float3(0,0,0);
-	const int3 GridIndex = DIVelocityGrid_ComputeGridIndex(DIContext,GridPosition,float3(0.5,0.5,0.5),GridLength,GridFraction);
-
-	const float GridScale = 1.0 / (GridLength*GridLength*GridLength);
-
-	for(int i = 0; i < 2; ++i)
-	{
-		const float GridWeightX = (i == 0) ? 1.0 - GridFraction.x : GridFraction.x;
-		for(int j = 0; j < 2; ++j)
-		{
-			const float GridWeightY = (j == 0) ? 1.0 - GridFraction.y : GridFraction.y;
-			for(int k = 0; k < 2; ++k)
-			{
-				const float GridWeightZ = (k == 0) ? 1.0 - GridFraction.z : GridFraction.z;
-				const float GridWeightXYZ = GridWeightX*GridWeightY*GridWeightZ;
-
-				const int3 GridLocation = int3(GridIndex.x+i,GridIndex.y+j,GridIndex.z+k);
-
-				if( DIVelocityGrid_IsCellValid(DIContext,GridLocation) )
-				{
-					DIVelocityGrid_InterlockedAddScalar(DIContext,DIContext_GridDestinationBuffer,GridLocation,FLUID_DENSITY_OFFSET,GridWeightXYZ*ParticleMass*GridScale);
-				}
-			}
-		}
-	}
-}
-
-/* -----------------------------------------------------------------
- * Get Density field from grid
- * -----------------------------------------------------------------
- */
-
-// Sample the grid velocity field and gradient at the particle world position 
-void DIVelocityGrid_GetDensityField(in FDIVelocityGridContext DIContext, in float3 GridOrigin, in float GridLength, 
-			in float3 ParticlePosition, out float OutParticleDensity)
-{
-	OutParticleDensity = 0.0;
-
-	if(isfinite(GridLength))
-	{
-		const float3 GridPosition = DIVelocityGrid_GetGridPosition(DIContext,ParticlePosition,GridOrigin);
-
-		float3 GridFraction = float3(0,0,0);
-		const int3 GridIndex = DIVelocityGrid_ComputeGridIndex(DIContext,GridPosition,float3(0.5,0.5,0.5),GridLength,GridFraction);
-
-		for(int i = 0; i < 2; ++i)
-		{
-			const float GridWeightX = (i == 0) ? 1.0 - GridFraction.x : GridFraction.x;
-			for(int j = 0; j < 2; ++j)
-			{
-				const float GridWeightY = (j == 0) ? 1.0 - GridFraction.y : GridFraction.y;
-				for(int k = 0; k < 2; ++k)
-				{
-					const float GridWeightZ = (k == 0) ? 1.0 - GridFraction.z : GridFraction.z;
-					const float GridWeightXYZ = GridWeightX*GridWeightY*GridWeightZ;
-
-					const int3 GridLocation = int3(GridIndex.x+i,GridIndex.y+j,GridIndex.z+k);
-
-					if( DIVelocityGrid_IsCellValid(DIContext,GridLocation) )
-					{
-						OutParticleDensity += DIVelocityGrid_GetGridScalar(DIContext,GridLocation,FLUID_DENSITY_OFFSET) * GridWeightXYZ;
-					}
-				}
-			}
-		}
-	}
-}
-
-// Sample the grid velocity field and gradient at the particle world position 
-void DIVelocityGrid_UpdateDeformationGradient(in FDIVelocityGridContext DIContext, in float DeltaTime, in float4x4 VelocityGradient, 
-				in float4x4 DeformationGradient, out float4x4 OutDeformationGradient, out float OutGradientDeterminant)
-{
-	float3x3 PositionGradient;
-	PositionGradient[0] = VelocityGradient[0].xyz * DeltaTime;
-	PositionGradient[1] = VelocityGradient[1].xyz * DeltaTime;
-	PositionGradient[2] = VelocityGradient[2].xyz * DeltaTime;
-
-	float3x3 ElasticGradient;
-	ElasticGradient[0] = DeformationGradient[0].xyz;
-	ElasticGradient[1] = DeformationGradient[1].xyz;
-	ElasticGradient[2] = DeformationGradient[2].xyz;
-
-	if(isfinite(PositionGradient[0][0]) && isfinite(PositionGradient[1][1]) && isfinite(PositionGradient[2][2]) )
-	{
-		ElasticGradient += mul(PositionGradient,ElasticGradient);
-	}
-	OutGradientDeterminant = determinant(ElasticGradient);
-
-	OutDeformationGradient[0] = float4(ElasticGradient[0].xyz,0.0);
-	OutDeformationGradient[1] = float4(ElasticGradient[1].xyz,0.0);
-	OutDeformationGradient[2] = float4(ElasticGradient[2].xyz,0.0);
-	OutDeformationGradient[3] = float4(0.0,0.0,0.0,0.0);
-}
-
-/* -----------------------------------------------------------------
- * Update boundary weights
- * -----------------------------------------------------------------
- */
-
- float DIVelocityGrid_GetFractionInsideEdge(in float SignedDistanceA, in float SignedDistanceB)
- {
-	const bool IsInsideA = SignedDistanceA<0.0;
-	const bool IsInsideB = SignedDistanceB<0.0;
-
-	return (IsInsideA && IsInsideB) ? 1.0 : (IsInsideA && !IsInsideB) ? SignedDistanceA / (SignedDistanceA - SignedDistanceB) : 
-		   (!IsInsideA && IsInsideB) ? SignedDistanceB / (SignedDistanceB - SignedDistanceA)  : 0.0;
- }
-
- float DIVelocityGrid_GetFractionInsideImplA(in float SignedDistanceA, in float SignedDistanceB, in float SignedDistanceD)
- {
-	const float Side0 = 1.0 - DIVelocityGrid_GetFractionInsideEdge(SignedDistanceA,SignedDistanceD);
-	const float Side1 = 1.0 - DIVelocityGrid_GetFractionInsideEdge(SignedDistanceA,SignedDistanceB);
-	return 1.0 - 0.5 * Side0 * Side1;
- }
-
- float DIVelocityGrid_GetFractionInsideImplB(in float SignedDistanceA, in float SignedDistanceB, in float SignedDistanceC, in float SignedDistanceD)
- {
-	const float Side0 = DIVelocityGrid_GetFractionInsideEdge(SignedDistanceA,SignedDistanceD);
-	const float Side1 = DIVelocityGrid_GetFractionInsideEdge(SignedDistanceB,SignedDistanceC);
-	return 0.5 * ( Side0 + Side1);
- }
-
- float DIVelocityGrid_GetFractionInsideImplC(in float SignedDistanceA, in float SignedDistanceB, in float SignedDistanceC, in float SignedDistanceD)
- {
-	const float MiddlePoint = 0.25 * (SignedDistanceA+SignedDistanceB+SignedDistanceC+SignedDistanceD);
-	if(MiddlePoint < 0.0)
-	{
-		float FaceArea = 0.0;
-
-		const float Side1 = 1.0 - DIVelocityGrid_GetFractionInsideEdge(SignedDistanceA,SignedDistanceD);
-		const float Side3 = 1.0 - DIVelocityGrid_GetFractionInsideEdge(SignedDistanceC,SignedDistanceD);
-		FaceArea += 0.5 * Side3 * Side1;
-
-		const float Side2 = 1.0 - DIVelocityGrid_GetFractionInsideEdge(SignedDistanceC,SignedDistanceB);
-		const float Side0 = 1.0 - DIVelocityGrid_GetFractionInsideEdge(SignedDistanceA,SignedDistanceB);
-		FaceArea += 0.5 * Side0 * Side2;
-
-		return 1.0 - FaceArea;
-	}
-	else
-	{
-		float FaceArea = 0.0;
-
-		const float Side0 = 1.0 - DIVelocityGrid_GetFractionInsideEdge(SignedDistanceA,SignedDistanceB);
-		const float Side1 = 1.0 - DIVelocityGrid_GetFractionInsideEdge(SignedDistanceA,SignedDistanceD);
-		FaceArea += 0.5 * Side0 * Side1;
-
-		const float Side2 = 1.0 - DIVelocityGrid_GetFractionInsideEdge(SignedDistanceC,SignedDistanceB);
-		const float Side3 = 1.0 - DIVelocityGrid_GetFractionInsideEdge(SignedDistanceC,SignedDistanceD);
-		FaceArea += 0.5 * Side3 * Side2;
-
-		return FaceArea;
-
-	}
- }
-
- float DIVelocityGrid_GetFractionInsideImplD(in float SignedDistanceA, in float SignedDistanceB, in float SignedDistanceD)
- {
-	const float Side0 = DIVelocityGrid_GetFractionInsideEdge(SignedDistanceA,SignedDistanceD);
-	const float Side1 = DIVelocityGrid_GetFractionInsideEdge(SignedDistanceA,SignedDistanceB);
-	return 0.5 * Side0 * Side1;
- }
-
- float DIVelocityGrid_GetFractionInsideFace(in float SignedDistanceA, in float SignedDistanceB, in float SignedDistanceC, in float SignedDistanceD )
- {
-	const bool IsInsideA = SignedDistanceA<0.0;
-	const bool IsInsideB = SignedDistanceB<0.0;
-	const bool IsInsideC = SignedDistanceC<0.0;
-	const bool IsInsideD = SignedDistanceD<0.0;
-
-	const int InsideCount = IsInsideA + IsInsideB + IsInsideC + IsInsideD;
-
-	[branch]
-	if(InsideCount == 4)
-	{
-		return 1.0;
-	}
-	else if(InsideCount == 3)
-	{
-		[branch]
-		if(!IsInsideA)
-		{
-			return DIVelocityGrid_GetFractionInsideImplA(SignedDistanceA,SignedDistanceB,SignedDistanceD);
-		}
-		else if(!IsInsideB)
-		{
-			return DIVelocityGrid_GetFractionInsideImplA(SignedDistanceB,SignedDistanceC,SignedDistanceA);
-		}
-		else if(!IsInsideC)
-		{
-			return DIVelocityGrid_GetFractionInsideImplA(SignedDistanceC,SignedDistanceD,SignedDistanceB);
-		}
-		else if(!IsInsideD)
-		{
-			return DIVelocityGrid_GetFractionInsideImplA(SignedDistanceD,SignedDistanceA,SignedDistanceC);
-		}
-	}
-	else if(InsideCount == 2)
-	{
-		[branch]
-		if(IsInsideA && IsInsideB)
-		{
-			return DIVelocityGrid_GetFractionInsideImplB(SignedDistanceA,SignedDistanceB,SignedDistanceC, SignedDistanceD);
-		}
-		else if(IsInsideB && IsInsideC)
-		{
-			return DIVelocityGrid_GetFractionInsideImplB(SignedDistanceB,SignedDistanceC,SignedDistanceD, SignedDistanceA);
-		}
-		else if(IsInsideC && IsInsideD)
-		{
-			return DIVelocityGrid_GetFractionInsideImplB(SignedDistanceC,SignedDistanceD,SignedDistanceA, SignedDistanceB);
-		}
-		else if(IsInsideD && IsInsideA)
-		{
-			return DIVelocityGrid_GetFractionInsideImplB(SignedDistanceD,SignedDistanceA,SignedDistanceB, SignedDistanceC);
-		}
-		else if(IsInsideA && IsInsideC)
-		{
-			return DIVelocityGrid_GetFractionInsideImplC(SignedDistanceA,SignedDistanceB,SignedDistanceC, SignedDistanceD);
-		}
-		else if(IsInsideB && IsInsideD)
-		{
-			return DIVelocityGrid_GetFractionInsideImplC(SignedDistanceB,SignedDistanceC,SignedDistanceD, SignedDistanceA);
-		}
-	}
-	else if(InsideCount == 1)
-	{
-		[branch]
-		if(IsInsideA)
-		{
-			return DIVelocityGrid_GetFractionInsideImplA(SignedDistanceA,SignedDistanceB,SignedDistanceD);
-		}
-		else if(IsInsideB)
-		{
-			return DIVelocityGrid_GetFractionInsideImplA(SignedDistanceB,SignedDistanceC,SignedDistanceA);
-		}
-		else if(IsInsideC)
-		{
-			return DIVelocityGrid_GetFractionInsideImplA(SignedDistanceC,SignedDistanceD,SignedDistanceB);
-		}
-		else if(IsInsideD)
-		{
-			return DIVelocityGrid_GetFractionInsideImplA(SignedDistanceD,SignedDistanceA,SignedDistanceC);
-		}
-	}
-	return 0;
- }
-
-// Update the solid weights and velocities
-void DIVelocityGrid_SetSolidBoundary(in FDIVelocityGridContext DIContext, in RWBuffer<int> DIContext_GridDestinationBuffer,
-									in int GridHash, in float SolidDistance,
-									in float3 SolidVelocity, out bool OutBoundaryStatus)
-{
-	const int3 GridIndex = DIVelocityGrid_GetGridIndex(DIContext,GridHash);
-	OutBoundaryStatus = false;
-
-	[branch]
-	if( DIVelocityGrid_IsNodeValid(DIContext,GridIndex))
-	{
-		OutBoundaryStatus = true;
-		DIVelocityGrid_SetGridScalar(DIContext,DIContext_GridDestinationBuffer,GridIndex,SOLID_DISTANCE_OFFSET,SolidDistance);
-		DIVelocityGrid_SetGridVector(DIContext,DIContext_GridDestinationBuffer,GridIndex,SOLID_VELOCITY_OFFSET,SolidVelocity);
-	}
-}
-
-void DIVelocityGrid_ComputeFaceWeights(in FDIVelocityGridContext DIContext, in RWBuffer<int> DIContext_GridDestinationBuffer, in int3 GridLocation, in int ComponentIndex, out float SignedDistance, inout float ParticleVelocity )
-{
-	SignedDistance = DIVelocityGrid_GetGridScalar(DIContext,GridLocation,SOLID_DISTANCE_OFFSET);
-	ParticleVelocity += DIVelocityGrid_GetGridScalar(DIContext,GridLocation,SOLID_VELOCITY_OFFSET+ComponentIndex);
-}
-
-void DIVelocityGrid_ComputeComponentWeights(in FDIVelocityGridContext DIContext, in RWBuffer<int> DIContext_GridDestinationBuffer, in int3 GridIndex, in int ComponentIndex)
-{
-	[branch]
-	if( DIVelocityGrid_IsFaceInside(DIContext,GridIndex,ComponentIndex))
-	{
-		float SignedDistance[4];
-		float ParticleVelocity = 0.0;
-
-		[branch]
-		if(ComponentIndex == 0)
-		{
-			DIVelocityGrid_ComputeFaceWeights(DIContext,DIContext_GridDestinationBuffer,int3(GridIndex.x,GridIndex.y+0,GridIndex.z+0),ComponentIndex,SignedDistance[0],ParticleVelocity);
-			DIVelocityGrid_ComputeFaceWeights(DIContext,DIContext_GridDestinationBuffer,int3(GridIndex.x,GridIndex.y+0,GridIndex.z+1),ComponentIndex,SignedDistance[1],ParticleVelocity);
-			DIVelocityGrid_ComputeFaceWeights(DIContext,DIContext_GridDestinationBuffer,int3(GridIndex.x,GridIndex.y+1,GridIndex.z+0),ComponentIndex,SignedDistance[2],ParticleVelocity);
-			DIVelocityGrid_ComputeFaceWeights(DIContext,DIContext_GridDestinationBuffer,int3(GridIndex.x,GridIndex.y+1,GridIndex.z+1),ComponentIndex,SignedDistance[3],ParticleVelocity);
-		}
-		else if(ComponentIndex == 1)
-		{
-			DIVelocityGrid_ComputeFaceWeights(DIContext,DIContext_GridDestinationBuffer,int3(GridIndex.x+0,GridIndex.y,GridIndex.z+0),ComponentIndex,SignedDistance[0],ParticleVelocity);
-			DIVelocityGrid_ComputeFaceWeights(DIContext,DIContext_GridDestinationBuffer,int3(GridIndex.x+0,GridIndex.y,GridIndex.z+1),ComponentIndex,SignedDistance[1],ParticleVelocity);
-			DIVelocityGrid_ComputeFaceWeights(DIContext,DIContext_GridDestinationBuffer,int3(GridIndex.x+1,GridIndex.y,GridIndex.z+0),ComponentIndex,SignedDistance[2],ParticleVelocity);
-			DIVelocityGrid_ComputeFaceWeights(DIContext,DIContext_GridDestinationBuffer,int3(GridIndex.x+1,GridIndex.y,GridIndex.z+1),ComponentIndex,SignedDistance[3],ParticleVelocity);
-		}
-		else
-		{
-			DIVelocityGrid_ComputeFaceWeights(DIContext,DIContext_GridDestinationBuffer,int3(GridIndex.x+0,GridIndex.y+0,GridIndex.z),ComponentIndex,SignedDistance[0],ParticleVelocity);
-			DIVelocityGrid_ComputeFaceWeights(DIContext,DIContext_GridDestinationBuffer,int3(GridIndex.x+0,GridIndex.y+1,GridIndex.z),ComponentIndex,SignedDistance[1],ParticleVelocity);
-			DIVelocityGrid_ComputeFaceWeights(DIContext,DIContext_GridDestinationBuffer,int3(GridIndex.x+1,GridIndex.y+0,GridIndex.z),ComponentIndex,SignedDistance[2],ParticleVelocity);
-			DIVelocityGrid_ComputeFaceWeights(DIContext,DIContext_GridDestinationBuffer,int3(GridIndex.x+1,GridIndex.y+1,GridIndex.z),ComponentIndex,SignedDistance[3],ParticleVelocity);
-		}
-
-		const float SolidWeight = max(0.0,min(1.0 - DIVelocityGrid_GetFractionInsideFace(SignedDistance[0],SignedDistance[1],SignedDistance[3],SignedDistance[2]),1.0));
-		const float SolidVelocity = 0.25 * ParticleVelocity;
-
-		DIVelocityGrid_SetGridScalar(DIContext,DIContext_GridDestinationBuffer,GridIndex,SOLID_WEIGHT_OFFSET+ComponentIndex,SolidWeight);
-		DIVelocityGrid_SetGridScalar(DIContext,DIContext_GridDestinationBuffer,GridIndex,SOLID_VELOCITY_OFFSET+ComponentIndex,SolidVelocity);
-	}
-}
-
- // Update the solid weights and velocities
-void DIVelocityGrid_ComputeBoundaryWeights(in FDIVelocityGridContext DIContext, in RWBuffer<int> DIContext_GridDestinationBuffer, in int GridHash, out bool OutWeightsStatus)
-{
-	const int3 GridIndex = DIVelocityGrid_GetGridIndex(DIContext,GridHash);
-	OutWeightsStatus = false;
-
-	DIVelocityGrid_ComputeComponentWeights(DIContext,DIContext_GridDestinationBuffer,GridIndex,0);
-	DIVelocityGrid_ComputeComponentWeights(DIContext,DIContext_GridDestinationBuffer,GridIndex,1);
-	DIVelocityGrid_ComputeComponentWeights(DIContext,DIContext_GridDestinationBuffer,GridIndex,2);
-}
-
-/* -----------------------------------------------------------------
- * Sample Velocity field from grid
- * -----------------------------------------------------------------
- */
-
-// Compute the grid pressure
-// Mdv/dt = Jt * P with C = Volume * div(V) = 0.0 = ( vnx * wnx * dx * dx - vpx * wpx * dx * dx ) + ... (divergence theorem)
-// -> dv = inv(M) * Jt * P * dt;
-// Jnx = dC/dvnx = wnx * dx * dx, Jpx = dC/dvpx = -wpx * dx * dx
-// since C = C0 + dC/dv * dv = C0 + J * dv = 0 = C0 + J * inv(M) * Jt * P * dt; and M = wnx * mnx
-// the schur complement is J * inv(M) * Jt * P * dt = -C0 -> P * dt = -C0 / (wnx * dx * dx * wnx * dx * dx / ( wnx * mnx) + wpx * ....)
-// P * dt = - sum( vnx * wnx * dx * dx - vpx * wpx * dx * dx ) / sum (wnx * dx * dx * wnx * dx * dx / ( wnx * mnx);
-//        = - sum( vnx * wnx ) / sum( wnx * dx * dx / mnx ) = -sum( vnx * wnx ) / ( dx * dx * sum( wnx / mnx ) )
-// P * dt * dx * dx = Pt = -sum( vnx * wnx ) / ( sum( wnx / mnx ) )
-// dv = wnx * dx * dx * P * dt / (wnx * mnx) = Pt / mnx
-
-float DIVelocityGrid_ComputeGridPressure(in FDIVelocityGridContext DIContext, in int3 PrevIndex, in float CentreDistance, out float3 PrevFraction, out float3 NextFraction)
-{	
-	float GridPressure = 0.0;
-	float GridDivergence = 0.0;
-	float GridFraction = 0.0;
-
-	for(int ComponentIndex = 0; ComponentIndex < 3; ++ComponentIndex)
-	{
-		int3 NextIndex = PrevIndex; NextIndex[ComponentIndex] += 1;
-
-		const float PrevDistance = DIVelocityGrid_GetGridScalar(DIContext,PrevIndex,FLUID_DISTANCE_OFFSET);
-		PrevFraction[ComponentIndex] = (PrevDistance < 0.0) ? 1.0 : max(DIVelocityGrid_GetFractionInsideEdge(CentreDistance,PrevDistance),0.01);
-		const float PrevWeight = DIVelocityGrid_GetGridScalar(DIContext,PrevIndex,SOLID_WEIGHT_OFFSET+ComponentIndex);
-			
-		GridDivergence -= DIVelocityGrid_GetGridScalar(DIContext,PrevIndex,SOLID_VELOCITY_OFFSET+ComponentIndex) * (1.0-PrevWeight);
-		GridDivergence -= DIVelocityGrid_GetGridScalar(DIContext,PrevIndex,FLUID_VELOCITY_OFFSET+ComponentIndex) * PrevWeight;
-		GridFraction += PrevWeight / PrevFraction[ComponentIndex];
-
-		const float NextDistance = DIVelocityGrid_GetGridScalar(DIContext,NextIndex,FLUID_DISTANCE_OFFSET);
-		NextFraction[ComponentIndex] = (NextDistance < 0.0) ? 1.0 : max(DIVelocityGrid_GetFractionInsideEdge(CentreDistance,NextDistance),0.01);
-		const float NextWeight = DIVelocityGrid_GetGridScalar(DIContext,NextIndex,SOLID_WEIGHT_OFFSET+ComponentIndex);
-			
-		GridDivergence += DIVelocityGrid_GetGridScalar(DIContext,NextIndex,SOLID_VELOCITY_OFFSET+ComponentIndex) * (1.0-NextWeight);
-		GridDivergence += DIVelocityGrid_GetGridScalar(DIContext,NextIndex,FLUID_VELOCITY_OFFSET+ComponentIndex) * NextWeight;
-		GridFraction += NextWeight / NextFraction[ComponentIndex];
-	}
-	if( GridFraction > 0.0)
-	{
-		//GridDivergence += DIVelocityGrid_GetGridScalar(DIContext,PrevIndex,FLUID_DENSITY_OFFSET);
-		GridPressure = -GridDivergence / (GridFraction);
-	}
-	return GridPressure;
-}
-
-// Project the pressure grid to be divergence free
-void DIVelocityGrid_SolveGridPressure(in FDIVelocityGridContext DIContext, in RWBuffer<int> DIContext_GridDestinationBuffer, in int GridHash, in int InitStage, out bool OutProjectStatus)
-{	
-	OutProjectStatus = false;
-
-	const int3 PrevIndex = DIVelocityGrid_GetGridIndex(DIContext,GridHash);
-	const bool IsInitStage = (SimulationStageIndex == InitStage) || (SimulationStageIndex == (InitStage+1));
-
-	if(DIVelocityGrid_IsCellValid(DIContext,PrevIndex))
-	{
-		if( (PrevIndex.x+PrevIndex.y+PrevIndex.z+SimulationStageIndex) % 2 == 0)
-		{
-			OutProjectStatus = true;
-			const float CentreDistance = DIVelocityGrid_GetGridScalar(DIContext,PrevIndex,FLUID_DISTANCE_OFFSET);
-			if( CentreDistance < 0.0)
-			{
-				float3 PrevFraction = float3(0,0,0);
-				float3 NextFraction = float3(0,0,0);
-
-				float GridPressure = DIVelocityGrid_ComputeGridPressure(DIContext,PrevIndex,CentreDistance,PrevFraction,NextFraction);
-				if(!IsInitStage)
-				{
-					DIVelocityGrid_AddGridScalar(DIContext,DIContext_GridDestinationBuffer,PrevIndex,FLUID_PRESSURE_OFFSET,GridPressure);
-				}
-				else 
-				{
-					GridPressure =  DIVelocityGrid_GetGridScalar(DIContext,PrevIndex,FLUID_PRESSURE_OFFSET);
-				}
-
-				for( int ComponentIndex = 0; ComponentIndex < 3; ++ComponentIndex)
-				{
-					int3 NextIndex = PrevIndex; NextIndex[ComponentIndex] += 1;
-
-					const float PrevWeight = DIVelocityGrid_GetGridScalar(DIContext,PrevIndex,SOLID_WEIGHT_OFFSET+ComponentIndex);
-					if( PrevWeight > 0.0 && PrevFraction[ComponentIndex] > 0.0)
-					{
-						DIVelocityGrid_AddGridScalar(DIContext,DIContext_GridDestinationBuffer,PrevIndex,FLUID_VELOCITY_OFFSET+ComponentIndex,-GridPressure/PrevFraction[ComponentIndex]);
-					}
-					else
-					{
-						const float SolidVelocity = DIVelocityGrid_GetGridScalar(DIContext,PrevIndex,SOLID_VELOCITY_OFFSET+ComponentIndex);
-						DIVelocityGrid_SetGridScalar(DIContext,DIContext_GridDestinationBuffer,PrevIndex,FLUID_VELOCITY_OFFSET+ComponentIndex,SolidVelocity);
-					}
-
-					const float NextWeight = DIVelocityGrid_GetGridScalar(DIContext,NextIndex,SOLID_WEIGHT_OFFSET+ComponentIndex);
-					if( NextWeight > 0.0 && NextFraction[ComponentIndex] > 0.0)
-					{
-						DIVelocityGrid_AddGridScalar(DIContext,DIContext_GridDestinationBuffer,NextIndex,FLUID_VELOCITY_OFFSET+ComponentIndex,GridPressure/NextFraction[ComponentIndex]);
-					}
-					else
-					{
-						const float SolidVelocity = DIVelocityGrid_GetGridScalar(DIContext,NextIndex,SOLID_VELOCITY_OFFSET+ComponentIndex);
-						DIVelocityGrid_SetGridScalar(DIContext,DIContext_GridDestinationBuffer,NextIndex,FLUID_VELOCITY_OFFSET+ComponentIndex,SolidVelocity);
-					}
-				}
-			}
-			else
-			{
-				DIVelocityGrid_SetGridScalar(DIContext,DIContext_GridDestinationBuffer,PrevIndex,FLUID_PRESSURE_OFFSET,0.0);
-			}
-		}
-	}
-}
-
-
-
-
-
-
-
-
-
-
diff --git a/Engine/Plugins/Runtime/HairStrands/Shaders/Private/NiagaraDataInterfaceVelocityGrid.ush b/Engine/Plugins/Runtime/HairStrands/Shaders/Private/NiagaraDataInterfaceVelocityGrid.ush
--- a/Engine/Plugins/Runtime/HairStrands/Shaders/Private/NiagaraDataInterfaceVelocityGrid.ush
+++ b/Engine/Plugins/Runtime/HairStrands/Shaders/Private/NiagaraDataInterfaceVelocityGrid.ush
@@ -6,597 +6,9 @@ NiagaraDataInterfaceVelocityGrid.ush
 
 #pragma once
 
-/* -----------------------------------------------------------------
- * Velocity Grid constants and context
- * -----------------------------------------------------------------
- */
-
 #define FAST_MODE 1
 
 static const float FLOAT_PRECISION = 1e+4;
 
 static const int FLUID_MASS_OFFSET = 0;
 static const int FLUID_VELOCITY_OFFSET = 3;
-
-struct FDIVelocityGridContext
-{
-	float4x4			WorldTransform;
-	float4x4			WorldInverse;
-	int3				GridSize;
-	Buffer<int>		GridCurrentBuffer;
-}; 
-
-#define DIVelocityGrid_DECLARE_CONSTANTS(NAME)\
-float4x4 WorldTransform_##NAME;\
-float4x4 WorldInverse_##NAME;\
-int3 GridSize_##NAME;\
-Buffer<int> GridCurrentBuffer_##NAME;\
-RWBuffer<int> GridDestinationBuffer_##NAME;\
-
-//@todo-lh: GridDestinationBuffer is passed around separately as SPIRV-Cross cannot translate texture atomics properly atm.
-#define DIVelocityGrid_MAKE_CONTEXT(NAME)\
-FDIVelocityGridContext DIContext; \
-DIContext.WorldTransform = WorldTransform_##NAME;\
-DIContext.WorldInverse = WorldInverse_##NAME;\
-DIContext.GridSize = GridSize_##NAME;\
-DIContext.GridCurrentBuffer = GridCurrentBuffer_##NAME;\
-RWBuffer<int> DIContext_GridDestinationBuffer = GridDestinationBuffer_##NAME;\
-
-/* -----------------------------------------------------------------
- * Grid topology conversion
- * -----------------------------------------------------------------
- */
-
-// Compute the grid origin and grid length
-void DIVelocityGrid_ComputeGridSize(in FDIVelocityGridContext DIContext, in float3 GridCenter, 
-									  in float3 GridExtent, out float3 OutGridOrigin, out float OutGridLength)
-{
-	const float3 GridLengths = float3(2.0*GridExtent.x/(DIContext.GridSize.x-2),
-									  2.0*GridExtent.y/(DIContext.GridSize.y-2),
-									  2.0*GridExtent.z/(DIContext.GridSize.z-2));
-
-	const float MaxLength = max(GridLengths.x,max(GridLengths.y,GridLengths.z));
-
-	const float3 RegularSize = float3(DIContext.GridSize.x*MaxLength,
-									  DIContext.GridSize.y*MaxLength,
-									  DIContext.GridSize.z*MaxLength);
-
-	OutGridOrigin = GridCenter - 0.5 * RegularSize;
-	OutGridLength = MaxLength;
-}
-
-// Transform the node world position in grid space
-float3 DIVelocityGrid_GetGridPosition(in FDIVelocityGridContext DIContext, in float3 ParticlePosition, in float3 GridOrigin)
-{
-	return ParticlePosition - GridOrigin;
-}
-
-/* -----------------------------------------------------------------
- * Grid cell/node/face validity
- * -----------------------------------------------------------------
- */
-
-// Return true if the grid cell index is in any voxel center (0,N-1)
-bool DIVelocityGrid_IsCellValid(in FDIVelocityGridContext DIContext, in int3 CellIndex)
-{
-	return (CellIndex.x >= 0 && CellIndex.x < DIContext.GridSize.x) && 
-		   (CellIndex.y >= 0 && CellIndex.y < DIContext.GridSize.y) &&
-		   (CellIndex.z >= 0 && CellIndex.z < DIContext.GridSize.z);
-}
-
-// Return true if the grid face index is for the component one (0,N) and in the voxel center for the others (0,N-1)
-bool DIVelocityGrid_IsFaceValid(in FDIVelocityGridContext DIContext, in int3 FaceIndex, in int ComponentIndex)
-{
-	switch(ComponentIndex)
-	{
-		case 0 :
-			return (FaceIndex.x >= 0 && FaceIndex.x <= DIContext.GridSize.x) && 
-			   (FaceIndex.y >= 0 && FaceIndex.y < DIContext.GridSize.y) &&
-			   (FaceIndex.z >= 0 && FaceIndex.z < DIContext.GridSize.z);
-		case 1 : 
-			return (FaceIndex.x >= 0 && FaceIndex.x < DIContext.GridSize.x) && 
-			   (FaceIndex.y >= 0 && FaceIndex.y <= DIContext.GridSize.y) &&
-			   (FaceIndex.z >= 0 && FaceIndex.z < DIContext.GridSize.z);
-		default : 
-			return (FaceIndex.x >= 0 && FaceIndex.x < DIContext.GridSize.x) && 
-			   (FaceIndex.y >= 0 && FaceIndex.y < DIContext.GridSize.y) &&
-			   (FaceIndex.z >= 0 && FaceIndex.z <= DIContext.GridSize.z);
-	}
-}
-
-// Return true if the grid node index is valid (0,N)
-bool DIVelocityGrid_IsNodeValid(in FDIVelocityGridContext DIContext, in int3 NodeIndex)
-{
-	return (NodeIndex.x >= 0 && NodeIndex.x <= DIContext.GridSize.x) && 
-		   (NodeIndex.y >= 0 && NodeIndex.y <= DIContext.GridSize.y) &&
-		   (NodeIndex.z >= 0 && NodeIndex.z <= DIContext.GridSize.z);
-}
-
-// Return true if the grid face index is inside the grid
-bool DIVelocityGrid_IsFaceInside(in FDIVelocityGridContext DIContext, in int3 FaceIndex, in int ComponentIndex)
-{
-	switch(ComponentIndex)
-	{
-		case 0 :
-			return (FaceIndex.x >  0 && FaceIndex.x < DIContext.GridSize.x) && 
-			   (FaceIndex.y >= 0 && FaceIndex.y < DIContext.GridSize.y) &&
-			   (FaceIndex.z >= 0 && FaceIndex.z < DIContext.GridSize.z);
-		case 1 :
-			return (FaceIndex.x >= 0 && FaceIndex.x < DIContext.GridSize.x) && 
-			   (FaceIndex.y >  0 && FaceIndex.y < DIContext.GridSize.y) &&
-			   (FaceIndex.z >= 0 && FaceIndex.z < DIContext.GridSize.z);
-		default : 
-			return (FaceIndex.x >= 0 && FaceIndex.x < DIContext.GridSize.x) && 
-			   (FaceIndex.y >= 0 && FaceIndex.y < DIContext.GridSize.y) &&
-			   (FaceIndex.z >  0 && FaceIndex.z < DIContext.GridSize.z);
-	}
-}
-
-bool DIVelocityGrid_IsCellInside(in FDIVelocityGridContext DIContext, in int3 CellIndex)
-{
-	return (CellIndex.x > 0 && CellIndex.x < (DIContext.GridSize.x-1)) && 
-		   (CellIndex.y > 0 && CellIndex.y < (DIContext.GridSize.y-1)) &&
-		   (CellIndex.z > 0 && CellIndex.z < (DIContext.GridSize.z-1));
-}
-
-bool DIVelocityGrid_IsNodeInside(in FDIVelocityGridContext DIContext, in int3 NodeIndex)
-{
-	return (NodeIndex.x > 0 && NodeIndex.x < DIContext.GridSize.x) && 
-		   (NodeIndex.y > 0 && NodeIndex.y < DIContext.GridSize.y) &&
-		   (NodeIndex.z > 0 && NodeIndex.z < DIContext.GridSize.z);
-}
-
-/* -----------------------------------------------------------------
- * Grid indexing utilities
- * -----------------------------------------------------------------
- */
-
-// Compute the grid index and fraction given a grid position and a grid offset
-int3 DIVelocityGrid_ComputeGridIndex(in FDIVelocityGridContext DIContext, in float3 GridPosition, in float3 GridOffset, 
-								 in float GridLength, out float3 OutGridFraction)
-{
-	const float3 GridCoordinate = GridPosition / GridLength - GridOffset;
-
-	const int3 GridIndex = floor(GridCoordinate);
-	OutGridFraction = GridCoordinate - GridIndex;
-	return GridIndex;
-}
-
-// Compute the grid index given a grid hash
-int3 DIVelocityGrid_GetGridIndex(in FDIVelocityGridContext DIContext, in int GridHash)
-{
-	const int GridSizeYZ = (DIContext.GridSize.y+1) * (DIContext.GridSize.z+1);
-	const int GridHashYZ = GridHash % GridSizeYZ;
-
-	return int3(GridHash / GridSizeYZ, GridHashYZ / (DIContext.GridSize.z+1), GridHashYZ % (DIContext.GridSize.z+1));
-}
-
-/* -----------------------------------------------------------------
- * Get Cell/Node/Face world positions
- * -----------------------------------------------------------------
- */
- 
- // Compute the world face position
-void DIVelocityGrid_GetFacePosition(in FDIVelocityGridContext DIContext, in int GridHash, 
-									in float3 GridOrigin, in float GridLength, in int ComponentIndex, out float3 OutFacePosition)
-{
-	const int3 GridIndex = DIVelocityGrid_GetGridIndex(DIContext,GridHash);
-
-	if( DIVelocityGrid_IsFaceValid(DIContext,GridIndex,ComponentIndex) )
-	{
-		float3 ComponentOffset = float3(0.5,0.5,0.5); ComponentOffset[ComponentIndex] = 0.0;
-		OutFacePosition = (GridIndex+ComponentOffset)*GridLength+GridOrigin;
-	}
-	else
-	{
-		OutFacePosition = float3(0,0,0);
-	}
-}
-
-// Compute the world cell position
-void DIVelocityGrid_GetCellPosition(in FDIVelocityGridContext DIContext, in int GridHash, 
-									in float3 GridOrigin, in float GridLength, out float3 OutCellPosition)
-{
-	const int3 GridIndex = DIVelocityGrid_GetGridIndex(DIContext,GridHash);
-
-	if( DIVelocityGrid_IsCellValid(DIContext,GridIndex) )
-	{
-		OutCellPosition = (GridIndex+float3(0.5,0.5,0.5))*GridLength+GridOrigin;
-	}
-	else
-	{
-		OutCellPosition = float3(0,0,0);
-	}
-}
-
-// Compute the world node position
-void DIVelocityGrid_GetNodePosition(in FDIVelocityGridContext DIContext, in int GridHash, 
-									in float3 GridOrigin, in float GridLength, out float3 OutNodePosition)
-{
-	const int3 GridIndex = DIVelocityGrid_GetGridIndex(DIContext,GridHash);
-
-	if( DIVelocityGrid_IsNodeValid(DIContext,GridIndex) )
-	{
-		OutNodePosition = GridIndex*GridLength+GridOrigin;
-	}
-	else
-	{
-		OutNodePosition = float3(0,0,0);
-	}
-}
-
-/* -----------------------------------------------------------------
- * Interlocked Add/Min Float (slow but accurate version) and Int (fast but approximate version)
- * -----------------------------------------------------------------
- */
-
-// Compute the data index from agrid index and a data offset
-int DIVelocityGrid_GetDataIndex( in FDIVelocityGridContext DIContext, in int3 GridIndex, in int DataOffset)
-{
-	const int GridIndexX = GridIndex.x + DataOffset * (DIContext.GridSize.x+1);
-	return GridIndexX * (DIContext.GridSize.y+1) * (DIContext.GridSize.z+1) + GridIndex.y * (DIContext.GridSize.z+1) + GridIndex.z;
-}
-
-#if FAST_MODE == 1
-
-void DIVelocityGrid_InterlockedAddInt(in FDIVelocityGridContext DIContext, in RWBuffer<int> DIContext_GridDestinationBuffer, in int3 GridIndex, in int DataOffset, in int IntValue)
-{
-	InterlockedAdd(DIContext_GridDestinationBuffer[DIVelocityGrid_GetDataIndex(DIContext,GridIndex,DataOffset)], IntValue);
-}
-
-void DIVelocityGrid_InterlockedMinInt(in FDIVelocityGridContext DIContext, in RWBuffer<int> DIContext_GridDestinationBuffer, in int3 GridIndex, in int DataOffset, in int IntValue)
-{
-	InterlockedMin(DIContext_GridDestinationBuffer[DIVelocityGrid_GetDataIndex(DIContext,GridIndex,DataOffset)], IntValue);
-}
-
-#else
-
-void DIVelocityGrid_InterlockedAddFloat(in FDIVelocityGridContext DIContext, in RWBuffer<int> DIContext_GridDestinationBuffer, in int3 GridIndex, in int DataOffset, in float FloatValue)
-{
-	int IntValue = asint(FloatValue);
-	int CompareValue = 0;
-	int StoredValue = 0;
-
-	[allow_uav_condition] while(true)
-	{
-		InterlockedCompareExchange(DIContext_GridDestinationBuffer[
-				DIVelocityGrid_GetDataIndex(DIContext,GridIndex,DataOffset)],CompareValue,IntValue,StoredValue);
-		if(StoredValue == CompareValue)
-		{
-			break;
-		}
-		CompareValue = StoredValue;
-		IntValue = asint(FloatValue+asfloat(StoredValue));
-	}
-}
-
-void DIVelocityGrid_InterlockedMinFloat(in FDIVelocityGridContext DIContext, in RWBuffer<int> DIContext_GridDestinationBuffer, in int3 GridIndex, in int DataOffset, in float FloatValue)
-{
-	int IntValue = asint(FloatValue);
-	int CompareValue = 0;
-	int StoredValue = 0;
-
-	[allow_uav_condition] while(true)
-	{
-		InterlockedCompareExchange(DIContext_GridDestinationBuffer[
-				DIVelocityGrid_GetDataIndex(DIContext,GridIndex,DataOffset)],CompareValue,IntValue,StoredValue);
-		if(StoredValue == CompareValue)
-		{
-			break;
-		}
-		CompareValue = StoredValue;
-		IntValue = asint(min(FloatValue,asfloat(StoredValue)));
-	}
-}
-
-#endif
-
-/* -----------------------------------------------------------------
- * Datas accessors
- * -----------------------------------------------------------------
- */
-
- void DIVelocityGrid_InterlockedMinScalar(in FDIVelocityGridContext DIContext, in RWBuffer<int> DIContext_GridDestinationBuffer, in int3 GridIndex, in int DataOffset, in float FloatValue)
-{
-#if FAST_MODE == 1
-	const int IntValue = FloatValue * FLOAT_PRECISION;
-	DIVelocityGrid_InterlockedMinInt(DIContext,DIContext_GridDestinationBuffer,GridIndex,DataOffset,IntValue);
-#else
-	DIVelocityGrid_InterlockedMinFloat(DIContext,DIContext_GridDestinationBuffer,GridIndex,DataOffset,FloatValue);
-#endif
-}
-
-void DIVelocityGrid_InterlockedAddScalar(in FDIVelocityGridContext DIContext, in RWBuffer<int> DIContext_GridDestinationBuffer, in int3 GridIndex, in int DataOffset, in float FloatValue)
-{
-#if FAST_MODE == 1
-	const int IntValue = FloatValue * FLOAT_PRECISION;
-	DIVelocityGrid_InterlockedAddInt(DIContext,DIContext_GridDestinationBuffer,GridIndex,DataOffset,IntValue);
-#else
-	DIVelocityGrid_InterlockedAddFloat(DIContext,DIContext_GridDestinationBuffer,GridIndex,DataOffset,FloatValue);
-#endif
-}
-
-float DIVelocityGrid_ConvertGridScalar(in FDIVelocityGridContext DIContext, in int3 GridIndex, in int DataOffset)
-{
-	const int IntValue = DIContext.GridCurrentBuffer[DIVelocityGrid_GetDataIndex(DIContext,GridIndex,DataOffset)];
-
-#if FAST_MODE == 1
-	return float(IntValue) / FLOAT_PRECISION;
-#else
-	return asfloat(IntValue);
-#endif
-}
-
-void DIVelocityGrid_AddGridScalar(in FDIVelocityGridContext DIContext, in RWBuffer<int> DIContext_GridDestinationBuffer, in int3 GridIndex, in int DataOffset, in float FloatValue)
-{
-	const float SumValue = asfloat(DIContext.GridCurrentBuffer[DIVelocityGrid_GetDataIndex(DIContext,GridIndex,DataOffset)])+FloatValue;
-	DIContext_GridDestinationBuffer[DIVelocityGrid_GetDataIndex(DIContext,GridIndex,DataOffset)] = asint(SumValue);
-}
-
-float DIVelocityGrid_GetGridScalar(in FDIVelocityGridContext DIContext, in int3 GridIndex, in int DataOffset)
-{
-	return asfloat(DIContext.GridCurrentBuffer[DIVelocityGrid_GetDataIndex(DIContext,GridIndex,DataOffset)]);
-}
-
-void DIVelocityGrid_SetGridScalar(in FDIVelocityGridContext DIContext, in RWBuffer<int> DIContext_GridDestinationBuffer, in int3 GridIndex, in int DataOffset, in float FloatValue)
-{
-	DIContext_GridDestinationBuffer[DIVelocityGrid_GetDataIndex(DIContext,GridIndex,DataOffset)] = asint(FloatValue);
-}
-
-void DIVelocityGrid_SetGridVector( in FDIVelocityGridContext DIContext, in RWBuffer<int> DIContext_GridDestinationBuffer, in int3 GridIndex, in int DataOffset, in float3 NodeVector)
-{
-	DIVelocityGrid_SetGridScalar(DIContext,DIContext_GridDestinationBuffer,GridIndex,DataOffset,NodeVector.x);
-	DIVelocityGrid_SetGridScalar(DIContext,DIContext_GridDestinationBuffer,GridIndex,DataOffset+1,NodeVector.y);
-	DIVelocityGrid_SetGridScalar(DIContext,DIContext_GridDestinationBuffer,GridIndex,DataOffset+2,NodeVector.z);
-}
-
-void DIVelocityGrid_AddGridVector( in FDIVelocityGridContext DIContext, in RWBuffer<int> DIContext_GridDestinationBuffer, in int3 GridIndex, in int DataOffset, in float3 NodeVector)
-{
-	DIVelocityGrid_AddGridScalar(DIContext,DIContext_GridDestinationBuffer,GridIndex,DataOffset,NodeVector.x);
-	DIVelocityGrid_AddGridScalar(DIContext,DIContext_GridDestinationBuffer,GridIndex,DataOffset+1,NodeVector.y);
-	DIVelocityGrid_AddGridScalar(DIContext,DIContext_GridDestinationBuffer,GridIndex,DataOffset+2,NodeVector.z);
-}
-
-float3 DIVelocityGrid_GetGridVector( in FDIVelocityGridContext DIContext, in int3 GridIndex, in int DataOffset)
-{
-	return float3(DIVelocityGrid_GetGridScalar(DIContext,GridIndex,DataOffset),
-				  DIVelocityGrid_GetGridScalar(DIContext,GridIndex,DataOffset+1),
-				  DIVelocityGrid_GetGridScalar(DIContext,GridIndex,DataOffset+2));
-}
-
-/* -----------------------------------------------------------------
- * Interpolation function and gradient
- * -----------------------------------------------------------------
- */
-
- float DIVelocityGrid_GetShapeFunction( in float3 GridFraction, in int3 ijk)
- {
-	const float GridWeightX = (ijk.x == 0) ? 1.0 - GridFraction.x : GridFraction.x;
-	const float GridWeightY = (ijk.y == 0) ? 1.0 - GridFraction.y : GridFraction.y;
-	const float GridWeightZ = (ijk.z == 0) ? 1.0 - GridFraction.z : GridFraction.z;
-
-	return GridWeightX*GridWeightY*GridWeightZ;
- }
-
- float3 DIVelocityGrid_GetShapeGradient( in float3 GridFraction, in int3 ijk)
- {
-	const float GridWeightX = (ijk.x == 0) ? 1.0 - GridFraction.x : GridFraction.x;
-	const float GridWeightY = (ijk.y == 0) ? 1.0 - GridFraction.y : GridFraction.y;
-	const float GridWeightZ = (ijk.z == 0) ? 1.0 - GridFraction.z : GridFraction.z;
-
-	const float GridGradientX = (ijk.x == 0) ? -1.0: 1.0;
-	const float GridGradientY = (ijk.y == 0) ? -1.0: 1.0;
-	const float GridGradientZ = (ijk.z == 0) ? -1.0: 1.0;
-
-	return  float3(GridGradientX*GridWeightY*GridWeightZ,
-				   GridWeightX*GridGradientY*GridWeightZ,
-				   GridWeightX*GridWeightY*GridGradientZ);
- }
-
-/* -----------------------------------------------------------------
- * Build Velocity Field from nodes
- * -----------------------------------------------------------------
- */
- 
-// Transfer the velocity field component on a face from the particle to the grid
-void DIVelocityGrid_BuildVelocityFace(in FDIVelocityGridContext DIContext, in RWBuffer<int> DIContext_GridDestinationBuffer, in int3 GridIndex, in float ShapeFunction, in int3 ijk, in float3 GridPosition, in float SampleMass, in float SampleMomentum,
-			in float3 MomentumGradient, in float3 ComponentOffset, in int ComponentIndex, in float GridLength )
-{
-	const int3 GridLocation = int3(GridIndex.x+ijk.x,GridIndex.y+ijk.y,GridIndex.z+ijk.z);
-
-	if( DIVelocityGrid_IsFaceValid(DIContext,GridLocation,ComponentIndex) )
-	{
-		const float3 DeltaPosition = (GridLocation+ComponentOffset)*GridLength - GridPosition;
-		const float GridMomentum = SampleMomentum + dot(MomentumGradient, DeltaPosition);
-
-		DIVelocityGrid_InterlockedAddScalar(DIContext,DIContext_GridDestinationBuffer,GridLocation,FLUID_MASS_OFFSET+ComponentIndex,SampleMass*ShapeFunction);
-		DIVelocityGrid_InterlockedAddScalar(DIContext,DIContext_GridDestinationBuffer,GridLocation,FLUID_VELOCITY_OFFSET+ComponentIndex,GridMomentum*ShapeFunction);
-	}
-}
-
-// Transfer the velocity field component from the particle to the grid
-void DIVelocityGrid_BuildVelocityComponent(in FDIVelocityGridContext DIContext, in RWBuffer<int> DIContext_GridDestinationBuffer, in float3 GridPosition, in float SampleMass, in float SampleMomentum,
-			in float3 MomentumGradient, in float3 ComponentOffset, in int ComponentIndex, in float GridLength )
-{
-	float3 GridFraction = float3(0,0,0);
-	const int3 GridIndex = DIVelocityGrid_ComputeGridIndex(DIContext,GridPosition,ComponentOffset,GridLength,GridFraction);
-
-	for(int i = 0; i < 2; ++i)
-	{
-		const float GridWeightX = (i == 0) ? 1.0 - GridFraction.x : GridFraction.x;
-		for(int j = 0; j < 2; ++j)
-		{
-			const float GridWeightY = (j == 0) ? 1.0 - GridFraction.y : GridFraction.y;
-			for(int k = 0; k < 2; ++k)
-			{
-				const float GridWeightZ = (k == 0) ? 1.0 - GridFraction.z : GridFraction.z;
-				const float GridWeightXYZ = GridWeightX*GridWeightY*GridWeightZ;
-
-				DIVelocityGrid_BuildVelocityFace(DIContext,DIContext_GridDestinationBuffer,GridIndex,GridWeightXYZ,int3(i,j,k),GridPosition,SampleMass,
-								SampleMomentum,MomentumGradient,ComponentOffset,ComponentIndex,GridLength);
-			}
-		}
-	}
-}
-
-// Transfer the velocity field from the particle to the grid
-void DIVelocityGrid_BuildVelocityField(in FDIVelocityGridContext DIContext, in RWBuffer<int> DIContext_GridDestinationBuffer, in float3 GridOrigin, in float GridLength, in float3 ParticlePosition, in float ParticleMass, in float3 ParticleVelocity, 
-			in float4x4 VelocityGradient, out bool OutFunctionStatus )
-{
-	const float3 ParticleMomentum = ParticleVelocity * ParticleMass;
-
-	const float3 MomentumGradientX = VelocityGradient[0].xyz * ParticleMass;
-	const float3 MomentumGradientY = VelocityGradient[1].xyz * ParticleMass;
-	const float3 MomentumGradientZ = VelocityGradient[2].xyz * ParticleMass;
-
-	const float3 GridPosition = DIVelocityGrid_GetGridPosition(DIContext,ParticlePosition,GridOrigin);
-
-	DIVelocityGrid_BuildVelocityComponent(DIContext,DIContext_GridDestinationBuffer,GridPosition,ParticleMass,ParticleMomentum.x,MomentumGradientX,float3(0.0,0.5,0.5),0,GridLength);
-	DIVelocityGrid_BuildVelocityComponent(DIContext,DIContext_GridDestinationBuffer,GridPosition,ParticleMass,ParticleMomentum.y,MomentumGradientY,float3(0.5,0.0,0.5),1,GridLength);
-	DIVelocityGrid_BuildVelocityComponent(DIContext,DIContext_GridDestinationBuffer,GridPosition,ParticleMass,ParticleMomentum.z,MomentumGradientZ,float3(0.5,0.5,0.0),2,GridLength);
-	
-	OutFunctionStatus = true;
-}
-
-/* -----------------------------------------------------------------
- * Sample Velocity field from grid
- * -----------------------------------------------------------------
- */
- 
-// Sample the grid velocity field and gradient component on a face at the particle local position 
-void DIVelocityGrid_SampleVelocityFace(in FDIVelocityGridContext DIContext, in int3 GridIndex, in float ShapeFunction, in float3 ShapeGradient, in int3 ijk, in float3 GridPosition, in float3 ComponentOffset, in int ComponentIndex, in float GridLength,
-						in bool ScaledVelocity, inout float OutGridVelocity, inout float OutGridMass, inout float3 OutGridGradient)
-{
-	const int3 GridLocation = int3(GridIndex.x+ijk.x,GridIndex.y+ijk.y,GridIndex.z+ijk.z);
-
-	if( DIVelocityGrid_IsFaceValid(DIContext,GridLocation,ComponentIndex) )
-	{
-		float GridVelocity = 0.0;
-		float GridMass = 0.0;
-
-		if(!ScaledVelocity)
-		{
-			const float GridMomentum = DIVelocityGrid_ConvertGridScalar(DIContext,GridLocation,FLUID_VELOCITY_OFFSET+ComponentIndex);
-			
-			GridMass = DIVelocityGrid_ConvertGridScalar(DIContext,GridLocation,FLUID_MASS_OFFSET+ComponentIndex);
-			GridVelocity = (GridMass != 0.0) ? GridMomentum / GridMass : 0.0;
-		}
-		else
-		{
-			GridMass = DIVelocityGrid_GetGridScalar(DIContext,GridLocation,FLUID_MASS_OFFSET+ComponentIndex);
-			GridVelocity = DIVelocityGrid_GetGridScalar(DIContext,GridLocation,FLUID_VELOCITY_OFFSET+ComponentIndex);
-		}
-
-		OutGridMass += GridMass * ShapeFunction;
-		OutGridVelocity += GridVelocity * ShapeFunction;
-		OutGridGradient += GridVelocity * ShapeGradient;
-	}
-}
-
-// Sample the grid velocity field and gradient component at the particle local position 
-void DIVelocityGrid_SampleVelocityComponent(in FDIVelocityGridContext DIContext, in float3 GridPosition, in float3 ComponentOffset, in int ComponentIndex, in float GridLength,
-						in bool ScaledVelocity, out float OutGridVelocity, out float OutGridMass, out float3 OutGridGradient)
-{
-	float3 GridFraction = float3(0,0,0); 
-	const int3 GridIndex = DIVelocityGrid_ComputeGridIndex(DIContext,GridPosition,ComponentOffset,GridLength,GridFraction);
-
-	OutGridVelocity = 0.0;
-	OutGridGradient = float3(0,0,0); 
-	OutGridMass = 0.0;
-
-	if(isfinite(GridLength))
-	{
-		for(int i = 0; i < 2; ++i)
-		{
-			const float GridWeightX = (i == 0) ? 1.0 - GridFraction.x : GridFraction.x;
-			const float GridGradientX = (i == 0) ? -1.0: 1.0;
-			for(int j = 0; j < 2; ++j)
-			{
-				const float GridWeightY = (j == 0) ? 1.0 - GridFraction.y : GridFraction.y;
-				const float GridGradientY = (j == 0) ? -1.0: 1.0;
-				for(int k = 0; k < 2; ++k)
-				{
-					const float GridWeightZ = (k == 0) ? 1.0 - GridFraction.z : GridFraction.z;
-					const float GridGradientZ = (k == 0) ? -1.0: 1.0;
-
-					const float GridWeightXYZ = GridWeightX*GridWeightY*GridWeightZ;
-					const float3 GridGradientXYZ = float3(GridGradientX*GridWeightY*GridWeightZ,
-															GridWeightX*GridGradientY*GridWeightZ,
-															GridWeightX*GridWeightY*GridGradientZ) / GridLength;
-
-					DIVelocityGrid_SampleVelocityFace(DIContext,GridIndex,GridWeightXYZ,GridGradientXYZ,int3(i,j,k),GridPosition,
-						ComponentOffset,ComponentIndex,GridLength,ScaledVelocity,OutGridVelocity,OutGridMass,OutGridGradient);
-				}
-			}
-		}
-	}
-}
-
-// Sample the grid velocity field and gradient at the particle world position 
-void DIVelocityGrid_SampleVelocityField(in FDIVelocityGridContext DIContext, in float3 GridOrigin, in float GridLength, in float3 ParticlePosition, in bool ScaledVelocity, out float OutParticleMass, out float3 OutParticleVelocity, 
-			out float4x4 OutParticleGradient)
-{
-	OutParticleVelocity = float3(0,0,0);
-	OutParticleGradient = float4x4(0,0,0,0,0,0,0,0,0,0,0,0,0,0,0,0);
-
-	float3 ParticleMass = float3(0,0,0);
-
-	const float3 GridPosition = DIVelocityGrid_GetGridPosition(DIContext,ParticlePosition,GridOrigin);
-	DIVelocityGrid_SampleVelocityComponent(DIContext,GridPosition,float3(0.0,0.5,0.5),0,GridLength,ScaledVelocity,OutParticleVelocity.x,ParticleMass.x,OutParticleGradient[0].xyz);
-	DIVelocityGrid_SampleVelocityComponent(DIContext,GridPosition,float3(0.5,0.0,0.5),1,GridLength,ScaledVelocity,OutParticleVelocity.y,ParticleMass.y,OutParticleGradient[1].xyz);
-	DIVelocityGrid_SampleVelocityComponent(DIContext,GridPosition,float3(0.5,0.5,0.0),2,GridLength,ScaledVelocity,OutParticleVelocity.z,ParticleMass.z,OutParticleGradient[2].xyz);
-
-	OutParticleMass = (ParticleMass.x+ParticleMass.y+ParticleMass.z)/3.0;
-}
-
-/* -----------------------------------------------------------------
- * Velocity Accessors
- * -----------------------------------------------------------------
- */
-
-void DIVelocityGrid_AddGridVelocity(in FDIVelocityGridContext DIContext, in RWBuffer<int> DIContext_GridDestinationBuffer, in int GridHash, in float3 GridVelocity, out bool OutAddStatus)
-{
-	const int3 GridIndex = DIVelocityGrid_GetGridIndex(DIContext,GridHash);
-	OutAddStatus = true;
-
-	if( DIVelocityGrid_IsFaceValid(DIContext,GridIndex,0) )
-	{
-		OutAddStatus = true;
-		DIVelocityGrid_AddGridScalar(DIContext,DIContext_GridDestinationBuffer,GridIndex,FLUID_VELOCITY_OFFSET,GridVelocity.x);
-	}
-	if( DIVelocityGrid_IsFaceValid(DIContext,GridIndex,1) )
-	{
-		OutAddStatus = true;
-		DIVelocityGrid_AddGridScalar(DIContext,DIContext_GridDestinationBuffer,GridIndex,FLUID_VELOCITY_OFFSET+1,GridVelocity.y);
-	}
-	if( DIVelocityGrid_IsFaceValid(DIContext,GridIndex,2) )
-	{
-		OutAddStatus = true;
-		DIVelocityGrid_AddGridScalar(DIContext,DIContext_GridDestinationBuffer,GridIndex,FLUID_VELOCITY_OFFSET+2,GridVelocity.z);
-	}
-}
-
-void DIVelocityGrid_GetGridVelocity(in FDIVelocityGridContext DIContext, in int GridHash, out float3 OutGridVelocity)
-{
-	const int3 GridIndex = DIVelocityGrid_GetGridIndex(DIContext,GridHash);
-
-	OutGridVelocity = float3(0,0,0);
-
-	if( DIVelocityGrid_IsFaceValid(DIContext,GridIndex,0) )
-	{
-		OutGridVelocity.x = DIVelocityGrid_GetGridScalar(DIContext,GridIndex,FLUID_VELOCITY_OFFSET);
-	}
-	if( DIVelocityGrid_IsFaceValid(DIContext,GridIndex,1) )
-	{
-		OutGridVelocity.y = DIVelocityGrid_GetGridScalar(DIContext,GridIndex,FLUID_VELOCITY_OFFSET+1);
-	}
-	if( DIVelocityGrid_IsFaceValid(DIContext,GridIndex,2) )
-	{
-		OutGridVelocity.z = DIVelocityGrid_GetGridScalar(DIContext,GridIndex,FLUID_VELOCITY_OFFSET+2);
-	}
-}
-
-
-
-
-
