- Hair conversion to RDG

diff --git a/Engine/Plugins/Runtime/HairStrands/Source/HairStrandsCore/Private/Niagara/NiagaraDataInterfaceHairStrands.cpp b/Engine/Plugins/Runtime/HairStrands/Source/HairStrandsCore/Private/Niagara/NiagaraDataInterfaceHairStrands.cpp
--- a/Engine/Plugins/Runtime/HairStrands/Source/HairStrandsCore/Private/Niagara/NiagaraDataInterfaceHairStrands.cpp
+++ b/Engine/Plugins/Runtime/HairStrands/Source/HairStrandsCore/Private/Niagara/NiagaraDataInterfaceHairStrands.cpp
@@ -6,6 +6,7 @@
 #include "NiagaraGpuComputeDispatchInterface.h"
 #include "NiagaraRenderer.h"
 #include "NiagaraSimStageData.h"
+#include "NiagaraShaderParametersBuilder.h"
 #include "NiagaraSystemInstance.h"
 
 #include "Components/SkeletalMeshComponent.h"
@@ -22,6 +23,63 @@
 
 #define LOCTEXT_NAMESPACE "NiagaraDataInterfaceHairStrands"
 
+//------------------------------------------------------------------------------------------------------------
+namespace NDIHairStrandsLocal
+{
+BEGIN_SHADER_PARAMETER_STRUCT(FShaderParameters,)
+	SHADER_PARAMETER_RDG_BUFFER_UAV(RWBuffer<uint4>,	DeformedPositionBuffer)
+	SHADER_PARAMETER_RDG_BUFFER_SRV(Buffer<uint4>,		RestPositionBuffer)
+	SHADER_PARAMETER_SRV(Buffer<uint>,					CurvesOffsetsBuffer)
+	SHADER_PARAMETER_SRV(Buffer<float4>,				RestTrianglePositionABuffer)
+	SHADER_PARAMETER_SRV(Buffer<float4>,				RestTrianglePositionBBuffer)
+	SHADER_PARAMETER_SRV(Buffer<float4>,				RestTrianglePositionCBuffer)
+	SHADER_PARAMETER_SRV(Buffer<float4>,				DeformedTrianglePositionABuffer)
+	SHADER_PARAMETER_SRV(Buffer<float4>,				DeformedTrianglePositionBBuffer)
+	SHADER_PARAMETER_SRV(Buffer<float4>,				DeformedTrianglePositionCBuffer)
+	SHADER_PARAMETER_SRV(Buffer<uint>,					RootBarycentricCoordinatesBuffer)
+	SHADER_PARAMETER(FMatrix44f,						WorldTransform)
+	SHADER_PARAMETER(FMatrix44f,						WorldInverse)
+	SHADER_PARAMETER(FQuat4f,							WorldRotation)
+	SHADER_PARAMETER(FMatrix44f,						BoneTransform)
+	SHADER_PARAMETER(FMatrix44f,						BoneInverse)
+	SHADER_PARAMETER(FQuat4f,							BoneRotation)
+	SHADER_PARAMETER(int,								NumStrands)
+	SHADER_PARAMETER(int,								StrandSize)
+	SHADER_PARAMETER(int,								InterpolationMode)
+	SHADER_PARAMETER(FVector3f,							RestRootOffset)
+	SHADER_PARAMETER(FVector3f,							DeformedRootOffset)
+	SHADER_PARAMETER(FVector3f,							RestPositionOffset)
+	SHADER_PARAMETER_SRV(Buffer<float4>,				DeformedPositionOffset)
+	SHADER_PARAMETER_RDG_BUFFER_UAV(RWBuffer<uint>,		BoundingBoxBuffer)
+	SHADER_PARAMETER(uint32,							ResetSimulation)
+	SHADER_PARAMETER(uint32,							RestUpdate)
+	SHADER_PARAMETER(uint32,							LocalSimulation)
+	SHADER_PARAMETER(int,								SampleCount)
+	SHADER_PARAMETER(FIntVector4,						BoundingBoxOffsets)
+	SHADER_PARAMETER_SRV(Buffer<float4>,				RestSamplePositionsBuffer)
+	SHADER_PARAMETER_SRV(Buffer<float4>,				MeshSampleWeightsBuffer)
+	SHADER_PARAMETER_SRV(Buffer<float>,					ParamsScaleBuffer)
+	SHADER_PARAMETER(FVector3f,							BoneLinearVelocity)
+	SHADER_PARAMETER(FVector3f,							BoneAngularVelocity)
+	SHADER_PARAMETER(FVector3f,							BoneLinearAcceleration)
+	SHADER_PARAMETER(FVector3f,							BoneAngularAcceleration)
+END_SHADER_PARAMETER_STRUCT()
+
+static const TCHAR* CommonShaderFiles[] =
+{
+	TEXT("/Plugin/FX/Niagara/Private/NiagaraQuaternionUtils.ush"),
+	//OutHLSL += TEXT("/Plugin/Runtime/HairStrands/Private/NiagaraDirectSolver.ush");
+	TEXT("/Plugin/Runtime/HairStrands/Private/NiagaraStrandsExternalForce.ush"),
+	TEXT("/Plugin/Runtime/HairStrands/Private/NiagaraHookeSpringMaterial.ush"),
+	TEXT("/Plugin/Runtime/HairStrands/Private/NiagaraAngularSpringMaterial.ush"),
+	TEXT("/Plugin/Runtime/HairStrands/Private/NiagaraConstantVolumeMaterial.ush"),
+	TEXT("/Plugin/Runtime/HairStrands/Private/NiagaraCosseratRodMaterial.ush"),
+	TEXT("/Plugin/Runtime/HairStrands/Private/NiagaraStaticCollisionConstraint.ush"),
+	TEXT("/Plugin/Runtime/HairStrands/Private/NiagaraDataInterfaceHairStrands.ush"),
+};
+
+static const TCHAR* TemplateShaderFile = TEXT("/Plugin/Runtime/HairStrands/Private/NiagaraDataInterfaceHairStrandsTemplate.ush");
+
 //------------------------------------------------------------------------------------------------------------
 
 static const FName GetPointPositionName(TEXT("GetPointPosition"));
@@ -169,52 +227,16 @@ static const FName GetSampleStateName(TEXT("GetSampleState"));
 
 //------------------------------------------------------------------------------------------------------------
 
-const FString UNiagaraDataInterfaceHairStrands::NumStrandsName(TEXT("NumStrands_"));
-const FString UNiagaraDataInterfaceHairStrands::StrandSizeName(TEXT("StrandSize_"));
-
-const FString UNiagaraDataInterfaceHairStrands::WorldTransformName(TEXT("WorldTransform_"));
-const FString UNiagaraDataInterfaceHairStrands::WorldInverseName(TEXT("WorldInverse_"));
-const FString UNiagaraDataInterfaceHairStrands::WorldRotationName(TEXT("WorldRotation_"));
-
-const FString UNiagaraDataInterfaceHairStrands::BoneTransformName(TEXT("BoneTransform_"));
-const FString UNiagaraDataInterfaceHairStrands::BoneInverseName(TEXT("BoneInverse_"));
-const FString UNiagaraDataInterfaceHairStrands::BoneRotationName(TEXT("BoneRotation_"));
-
-const FString UNiagaraDataInterfaceHairStrands::BoneLinearVelocityName(TEXT("BoneLinearVelocity_"));
-const FString UNiagaraDataInterfaceHairStrands::BoneAngularVelocityName(TEXT("BoneAngularVelocity_"));
-const FString UNiagaraDataInterfaceHairStrands::BoneLinearAccelerationName(TEXT("BoneLinearAcceleration_"));
-const FString UNiagaraDataInterfaceHairStrands::BoneAngularAccelerationName(TEXT("BoneAngularAcceleration_"));
+} //namespace NDIHairStrandsLocal
 
-const FString UNiagaraDataInterfaceHairStrands::DeformedPositionBufferName(TEXT("DeformedPositionBuffer_"));
-const FString UNiagaraDataInterfaceHairStrands::CurvesOffsetsBufferName(TEXT("CurvesOffsetsBuffer_"));
-const FString UNiagaraDataInterfaceHairStrands::RestPositionBufferName(TEXT("RestPositionBuffer_"));
-
-const FString UNiagaraDataInterfaceHairStrands::ResetSimulationName(TEXT("ResetSimulation_"));
-const FString UNiagaraDataInterfaceHairStrands::InterpolationModeName(TEXT("InterpolationMode_"));
-const FString UNiagaraDataInterfaceHairStrands::RestUpdateName(TEXT("RestUpdate_"));
-const FString UNiagaraDataInterfaceHairStrands::LocalSimulationName(TEXT("LocalSimulation_"));
-const FString UNiagaraDataInterfaceHairStrands::RootBarycentricCoordinatesName(TEXT("RootBarycentricCoordinatesBuffer_"));
-
-const FString UNiagaraDataInterfaceHairStrands::RestRootOffsetName(TEXT("RestRootOffset_"));
-const FString UNiagaraDataInterfaceHairStrands::RestTrianglePositionAName(TEXT("RestTrianglePositionABuffer_"));
-const FString UNiagaraDataInterfaceHairStrands::RestTrianglePositionBName(TEXT("RestTrianglePositionBBuffer_"));
-const FString UNiagaraDataInterfaceHairStrands::RestTrianglePositionCName(TEXT("RestTrianglePositionCBuffer_"));
-
-const FString UNiagaraDataInterfaceHairStrands::DeformedRootOffsetName(TEXT("DeformedRootOffset_"));
-const FString UNiagaraDataInterfaceHairStrands::DeformedTrianglePositionAName(TEXT("DeformedTrianglePositionABuffer_"));
-const FString UNiagaraDataInterfaceHairStrands::DeformedTrianglePositionBName(TEXT("DeformedTrianglePositionBBuffer_"));
-const FString UNiagaraDataInterfaceHairStrands::DeformedTrianglePositionCName(TEXT("DeformedTrianglePositionCBuffer_"));
-
-const FString UNiagaraDataInterfaceHairStrands::SampleCountName(TEXT("SampleCount_"));
-const FString UNiagaraDataInterfaceHairStrands::RestSamplePositionsName(TEXT("RestSamplePositionsBuffer_"));
-const FString UNiagaraDataInterfaceHairStrands::MeshSampleWeightsName(TEXT("MeshSampleWeightsBuffer_"));
-
-const FString UNiagaraDataInterfaceHairStrands::RestPositionOffsetName(TEXT("RestPositionOffset_"));
-const FString UNiagaraDataInterfaceHairStrands::DeformedPositionOffsetName(TEXT("DeformedPositionOffset_"));
+//------------------------------------------------------------------------------------------------------------
 
-const FString UNiagaraDataInterfaceHairStrands::BoundingBoxOffsetsName(TEXT("BoundingBoxOffsets_"));
-const FString UNiagaraDataInterfaceHairStrands::BoundingBoxBufferName(TEXT("BoundingBoxBuffer_"));
-const FString UNiagaraDataInterfaceHairStrands::ParamsScaleBufferName(TEXT("ParamsScaleBuffer_"));
+enum class EHairSimulationInterpolationMode : uint8
+{
+	Rigid = 0,
+	Skinned = 1,
+	RBF = 2
+};
 
 //------------------------------------------------------------------------------------------------------------
 
@@ -226,106 +248,6 @@ static FAutoConsoleVariableRef CVarHairSimulationRestUpdate(TEXT("r.HairStrands.
 
 //------------------------------------------------------------------------------------------------------------
 
-struct FNDIHairStrandsParametersName
-{
-	FNDIHairStrandsParametersName(const FString& Suffix)
-	{
-		NumStrandsName = UNiagaraDataInterfaceHairStrands::NumStrandsName + Suffix;
-		StrandSizeName = UNiagaraDataInterfaceHairStrands::StrandSizeName + Suffix;
-		WorldTransformName = UNiagaraDataInterfaceHairStrands::WorldTransformName + Suffix;
-		WorldInverseName = UNiagaraDataInterfaceHairStrands::WorldInverseName + Suffix;
-		WorldRotationName = UNiagaraDataInterfaceHairStrands::WorldRotationName + Suffix;
-
-		BoneTransformName = UNiagaraDataInterfaceHairStrands::BoneTransformName + Suffix;
-		BoneInverseName = UNiagaraDataInterfaceHairStrands::BoneInverseName + Suffix;
-		BoneRotationName = UNiagaraDataInterfaceHairStrands::BoneRotationName + Suffix;
-		
-		BoneLinearVelocityName = UNiagaraDataInterfaceHairStrands::BoneLinearVelocityName + Suffix;
-		BoneAngularVelocityName = UNiagaraDataInterfaceHairStrands::BoneAngularVelocityName + Suffix;
-		BoneLinearAccelerationName = UNiagaraDataInterfaceHairStrands::BoneLinearAccelerationName + Suffix;
-		BoneAngularAccelerationName = UNiagaraDataInterfaceHairStrands::BoneAngularAccelerationName + Suffix;
-
-		DeformedPositionBufferName = UNiagaraDataInterfaceHairStrands::DeformedPositionBufferName + Suffix;
-		CurvesOffsetsBufferName = UNiagaraDataInterfaceHairStrands::CurvesOffsetsBufferName + Suffix;
-		RestPositionBufferName = UNiagaraDataInterfaceHairStrands::RestPositionBufferName + Suffix;
-
-		InterpolationModeName = UNiagaraDataInterfaceHairStrands::InterpolationModeName + Suffix;
-		ResetSimulationName = UNiagaraDataInterfaceHairStrands::ResetSimulationName + Suffix;
-		RestUpdateName = UNiagaraDataInterfaceHairStrands::RestUpdateName + Suffix;
-		LocalSimulationName = UNiagaraDataInterfaceHairStrands::LocalSimulationName + Suffix;
-		RootBarycentricCoordinatesName = UNiagaraDataInterfaceHairStrands::RootBarycentricCoordinatesName + Suffix;
-
-		RestRootOffsetName = UNiagaraDataInterfaceHairStrands::RestRootOffsetName + Suffix;
-		RestTrianglePositionAName = UNiagaraDataInterfaceHairStrands::RestTrianglePositionAName + Suffix;
-		RestTrianglePositionBName = UNiagaraDataInterfaceHairStrands::RestTrianglePositionBName + Suffix;
-		RestTrianglePositionCName = UNiagaraDataInterfaceHairStrands::RestTrianglePositionCName + Suffix;
-
-		DeformedRootOffsetName = UNiagaraDataInterfaceHairStrands::DeformedRootOffsetName + Suffix;
-		DeformedTrianglePositionAName = UNiagaraDataInterfaceHairStrands::DeformedTrianglePositionAName + Suffix;
-		DeformedTrianglePositionBName = UNiagaraDataInterfaceHairStrands::DeformedTrianglePositionBName + Suffix;
-		DeformedTrianglePositionCName = UNiagaraDataInterfaceHairStrands::DeformedTrianglePositionCName + Suffix;
-
-		SampleCountName = UNiagaraDataInterfaceHairStrands::SampleCountName + Suffix;
-		RestSamplePositionsName = UNiagaraDataInterfaceHairStrands::RestSamplePositionsName + Suffix;
-		MeshSampleWeightsName = UNiagaraDataInterfaceHairStrands::MeshSampleWeightsName + Suffix;
-
-		RestPositionOffsetName = UNiagaraDataInterfaceHairStrands::RestPositionOffsetName + Suffix;
-		DeformedPositionOffsetName = UNiagaraDataInterfaceHairStrands::DeformedPositionOffsetName + Suffix;
-
-		BoundingBoxOffsetsName = UNiagaraDataInterfaceHairStrands::BoundingBoxOffsetsName + Suffix;
-		BoundingBoxBufferName = UNiagaraDataInterfaceHairStrands::BoundingBoxBufferName + Suffix;
-		ParamsScaleBufferName = UNiagaraDataInterfaceHairStrands::ParamsScaleBufferName + Suffix;
-	}
-
-	FString NumStrandsName;
-	FString StrandSizeName;
-	FString WorldTransformName;
-	FString WorldInverseName;
-	FString WorldRotationName;
-
-	FString BoneTransformName;
-	FString BoneInverseName;
-	FString BoneRotationName;
-	
-	FString BoneLinearVelocityName;
-	FString BoneLinearAccelerationName;
-	FString BoneAngularVelocityName;
-	FString BoneAngularAccelerationName;
-
-	FString DeformedPositionBufferName;
-	FString CurvesOffsetsBufferName;
-	FString RestPositionBufferName;
-
-	FString ResetSimulationName;
-	FString InterpolationModeName;
-	FString RestUpdateName;
-	FString LocalSimulationName;
-	FString RootBarycentricCoordinatesName;
-
-	FString RestRootOffsetName;
-	FString RestTrianglePositionAName;
-	FString RestTrianglePositionBName;
-	FString RestTrianglePositionCName;
-
-	FString DeformedRootOffsetName;
-	FString DeformedTrianglePositionAName;
-	FString DeformedTrianglePositionBName;
-	FString DeformedTrianglePositionCName;
-
-	FString SampleCountName;
-	FString RestSamplePositionsName;
-	FString MeshSampleWeightsName;
-
-	FString RestPositionOffsetName;
-	FString DeformedPositionOffsetName;
-
-	FString BoundingBoxBufferName;
-	FString BoundingBoxOffsetsName;
-	FString ParamsScaleBufferName;
-};
-
-//------------------------------------------------------------------------------------------------------------
-
 void FNDIHairStrandsBuffer::Initialize(
 	const FHairStrandsRestResource* HairStrandsRestResource,
 	const FHairStrandsDeformedResource*  HairStrandsDeformedResource,
@@ -372,6 +294,8 @@ void FNDIHairStrandsBuffer::InitRHI()
 {
 	if (SourceRestResources != nullptr)
 	{
+		FMemMark MemMark(FMemStack::Get());
+		FRDGBuilder GraphBuilder(FRHICommandListExecutor::GetImmediateCommandList());
 		FHairStrandsBulkData* SourceDatas = &SourceRestResources->BulkData; // This could be released by that time depending on how the initialization order is
 		{
 			const uint32 OffsetCount = SourceDatas->GetNumCurves() + 1;
@@ -416,19 +340,17 @@ void FNDIHairStrandsBuffer::InitRHI()
 			RHIUnlockBuffer(PointsCurveBuffer.Buffer);
 		}
 		{
-			static const TArray<uint32> ZeroData = { UINT_MAX,UINT_MAX,UINT_MAX,0,0,0,
-													 UINT_MAX,UINT_MAX,UINT_MAX,0,0,0,
-													 UINT_MAX,UINT_MAX,UINT_MAX,0,0,0,
-													 UINT_MAX,UINT_MAX,UINT_MAX,0,0,0 };
-
-			const uint32 BoundCount = 24;
-			const uint32 BoundBytes = sizeof(uint32)*BoundCount;
-
-			BoundingBoxBuffer.Initialize(TEXT("BoundingBoxBuffer"), sizeof(uint32), BoundCount, EPixelFormat::PF_R32_UINT, BUF_Static);
-			void* BoundBufferData = RHILockBuffer(BoundingBoxBuffer.Buffer, 0, BoundBytes, RLM_WriteOnly);
+			static const uint32 ZeroData[] = {
+				UINT_MAX,UINT_MAX,UINT_MAX,0,0,0,
+				UINT_MAX,UINT_MAX,UINT_MAX,0,0,0,
+				UINT_MAX,UINT_MAX,UINT_MAX,0,0,0,
+				UINT_MAX,UINT_MAX,UINT_MAX,0,0,0
+			};
+			const uint32 BoundCount = UE_ARRAY_COUNT(ZeroData);
 
-			FMemory::Memcpy(BoundBufferData, ZeroData.GetData(), BoundBytes);
-			RHIUnlockBuffer(BoundingBoxBuffer.Buffer);
+			BoundingBoxBuffer.Initialize(GraphBuilder, TEXT("BoundingBoxBuffer"), EPixelFormat::PF_R32_UINT, sizeof(uint32), BoundCount);
+			GraphBuilder.QueueBufferUpload(BoundingBoxBuffer.GetOrCreateBuffer(GraphBuilder), ZeroData, sizeof(ZeroData), ERDGInitialDataFlags::NoCopy);
+			BoundingBoxBuffer.EndGraphUsage();
 		}
 		{
 			const uint32 ScaleCount = 32 * NumScales;
@@ -443,15 +365,19 @@ void FNDIHairStrandsBuffer::InitRHI()
 		if (SourceDeformedResources == nullptr)
 		{
 			const uint32 PositionsCount = SourceDatas->GetNumPoints();
-			DeformedPositionBuffer.Initialize(TEXT("DeformedPositionBuffer"), FHairStrandsPositionFormat::SizeInByte, PositionsCount, FHairStrandsPositionFormat::Format, BUF_Static);
+
+			FRDGBufferDesc BufferDesc = FRDGBufferDesc::CreateBufferDesc(FHairStrandsPositionFormat::SizeInByte, PositionsCount);
+			FRDGBufferRef RDGBuffer = GraphBuilder.CreateBuffer(BufferDesc, TEXT("DeformedPositionBuffer"));
+			DeformedPositionBuffer = GraphBuilder.ConvertToExternalBuffer(RDGBuffer);
 		}
+		GraphBuilder.Execute();
 	}
 }
 
 void FNDIHairStrandsBuffer::ReleaseRHI()
 {
 	CurvesOffsetsBuffer.Release();
-	DeformedPositionBuffer.Release();
+	DeformedPositionBuffer.SafeRelease();
 	BoundingBoxBuffer.Release();
 	ParamsScaleBuffer.Release();
 }
@@ -664,326 +590,6 @@ bool FNDIHairStrandsData::Init(UNiagaraDataInterfaceHairStrands* Interface, FNia
 
 //------------------------------------------------------------------------------------------------------------
 
-enum class EHairSimulationInterpolationMode : uint8
-{
-	Rigid = 0,
-	Skinned = 1,
-	RBF = 2
-};
-
-struct FNDIHairStrandsParametersCS : public FNiagaraDataInterfaceParametersCS
-{
-	DECLARE_TYPE_LAYOUT(FNDIHairStrandsParametersCS, NonVirtual);
-	void Bind(const FNiagaraDataInterfaceGPUParamInfo& ParameterInfo, const class FShaderParameterMap& ParameterMap)
-	{
-		FNDIHairStrandsParametersName ParamNames(*ParameterInfo.DataInterfaceHLSLSymbol);
-
-		WorldTransform.Bind(ParameterMap, *ParamNames.WorldTransformName);
-		WorldInverse.Bind(ParameterMap, *ParamNames.WorldInverseName);
-		WorldRotation.Bind(ParameterMap, *ParamNames.WorldRotationName);
-		NumStrands.Bind(ParameterMap, *ParamNames.NumStrandsName);
-		StrandSize.Bind(ParameterMap, *ParamNames.StrandSizeName);
-
-		BoneTransform.Bind(ParameterMap, *ParamNames.BoneTransformName);
-		BoneInverse.Bind(ParameterMap, *ParamNames.BoneInverseName);
-		BoneRotation.Bind(ParameterMap, *ParamNames.BoneRotationName);
-
-		BoneLinearVelocity.Bind(ParameterMap, *ParamNames.BoneLinearVelocityName);
-		BoneAngularVelocity.Bind(ParameterMap, *ParamNames.BoneAngularVelocityName);
-		BoneLinearAcceleration.Bind(ParameterMap, *ParamNames.BoneLinearAccelerationName);
-		BoneAngularAcceleration.Bind(ParameterMap, *ParamNames.BoneAngularAccelerationName);
-
-		DeformedPositionBuffer.Bind(ParameterMap, *ParamNames.DeformedPositionBufferName);
-		CurvesOffsetsBuffer.Bind(ParameterMap, *ParamNames.CurvesOffsetsBufferName);
-		RestPositionBuffer.Bind(ParameterMap, *ParamNames.RestPositionBufferName);
-
-		ResetSimulation.Bind(ParameterMap, *ParamNames.ResetSimulationName);
-		InterpolationMode.Bind(ParameterMap,*ParamNames.InterpolationModeName);
-		RestUpdate.Bind(ParameterMap, *ParamNames.RestUpdateName);
-		LocalSimulation.Bind(ParameterMap, *ParamNames.LocalSimulationName);
-		RestRootOffset.Bind(ParameterMap, *ParamNames.RestRootOffsetName);
-		DeformedRootOffset.Bind(ParameterMap, *ParamNames.DeformedRootOffsetName);
-
-		RestPositionOffset.Bind(ParameterMap, *ParamNames.RestPositionOffsetName);
-		DeformedPositionOffset.Bind(ParameterMap, *ParamNames.DeformedPositionOffsetName);
-
-		RootBarycentricCoordinatesBuffer.Bind(ParameterMap, *ParamNames.RootBarycentricCoordinatesName);
-
-		RestTrianglePositionABuffer.Bind(ParameterMap, *ParamNames.RestTrianglePositionAName);
-		RestTrianglePositionBBuffer.Bind(ParameterMap, *ParamNames.RestTrianglePositionBName);
-		RestTrianglePositionCBuffer.Bind(ParameterMap, *ParamNames.RestTrianglePositionCName);
-
-		DeformedTrianglePositionABuffer.Bind(ParameterMap, *ParamNames.DeformedTrianglePositionAName);
-		DeformedTrianglePositionBBuffer.Bind(ParameterMap, *ParamNames.DeformedTrianglePositionBName);
-		DeformedTrianglePositionCBuffer.Bind(ParameterMap, *ParamNames.DeformedTrianglePositionCName);
-
-		SampleCount.Bind(ParameterMap, *ParamNames.SampleCountName);
-		RestSamplePositionsBuffer.Bind(ParameterMap, *ParamNames.RestSamplePositionsName);
-		MeshSampleWeightsBuffer.Bind(ParameterMap, *ParamNames.MeshSampleWeightsName);
-
-		BoundingBoxOffsets.Bind(ParameterMap, *ParamNames.BoundingBoxOffsetsName);
-		BoundingBoxBuffer.Bind(ParameterMap, *ParamNames.BoundingBoxBufferName);
-		ParamsScaleBuffer.Bind(ParameterMap, *ParamNames.ParamsScaleBufferName);
-	}
-
-	void Set(FRHICommandList& RHICmdList, const FNiagaraDataInterfaceSetArgs& Context) const
-	{
-		check(IsInRenderingThread());
-
-		FRHIComputeShader* ComputeShaderRHI = RHICmdList.GetBoundComputeShader();
-
-		FNDIHairStrandsProxy* InterfaceProxy = static_cast<FNDIHairStrandsProxy*>(Context.DataInterface);
-		FNDIHairStrandsData* ProxyData = InterfaceProxy->SystemInstancesToProxyData.Find(Context.SystemInstanceID);
-
-		const bool bIsHairValid = ProxyData != nullptr && ProxyData->HairStrandsBuffer && ProxyData->HairStrandsBuffer->IsInitialized();
-		const bool bIsHairGroupInstValid = ProxyData != nullptr && ProxyData->HairGroupInstSource != nullptr && ProxyData->HairGroupInstSource->ContainsGroupInstance(ProxyData->HairGroupInstance);
-		const bool bHasSkinningBinding = bIsHairValid && bIsHairGroupInstValid && ProxyData->HairGroupInstance->BindingType == EHairBindingType::Skinning;
-		const bool bIsRootValid = bIsHairValid && ProxyData->HairStrandsBuffer->SourceDeformedRootResources && ProxyData->HairStrandsBuffer->SourceDeformedRootResources->IsInitialized() && bHasSkinningBinding;
-		const bool bIsRestValid = bIsHairValid && ProxyData->HairStrandsBuffer->SourceRestResources && ProxyData->HairStrandsBuffer->SourceRestResources->IsInitialized() &&
-			 
-			// SourceRestResources->PositionBuffer is lazily allocated, when the instance LOD is scheduled (this happens after this called). So this is why this check is here. 
-			// This code should be refactor so that it reflects the lazy allocation scheme
-			ProxyData->HairStrandsBuffer->SourceRestResources->PositionBuffer.SRV && 
-			// TEMP: These check are only temporary for avoiding crashes while we find the bottom of the issue.
-			ProxyData->HairStrandsBuffer->CurvesOffsetsBuffer.SRV && ProxyData->HairStrandsBuffer->ParamsScaleBuffer.SRV && ProxyData->HairStrandsBuffer->BoundingBoxBuffer.UAV;
-
-		const bool bIsGeometryValid = bIsHairValid && (!bIsHairGroupInstValid || (bIsHairGroupInstValid && (ProxyData->HairGroupInstance->GeometryType != EHairGeometryType::NoneGeometry)));
-		const bool bIsDeformedValid = bIsHairValid && ProxyData->HairStrandsBuffer->SourceDeformedResources && ProxyData->HairStrandsBuffer->SourceDeformedResources->IsInitialized();
-
-		if (bIsHairValid && bIsRestValid && bIsGeometryValid && bIsHairGroupInstValid)
-		{
-			check(ProxyData);
-
-			FNDIHairStrandsBuffer* HairStrandsBuffer = ProxyData->HairStrandsBuffer;
-
-			FUnorderedAccessViewRHIRef DeformedPositionBufferUAV = bIsDeformedValid ?
-				HairStrandsBuffer->SourceDeformedResources->DeformedPositionBuffer[HairStrandsBuffer->SourceDeformedResources->CurrentIndex].UAV : HairStrandsBuffer->DeformedPositionBuffer.UAV;
-			FRHIShaderResourceView* DeformedPositionOffsetSRV = bIsDeformedValid ? 
-				HairStrandsBuffer->SourceDeformedResources->DeformedOffsetBuffer[HairStrandsBuffer->SourceDeformedResources->CurrentIndex].SRV.GetReference() : FNiagaraRenderer::GetDummyFloatBuffer();
-
-			const int32 MeshLODIndex = bIsRootValid ? HairStrandsBuffer->SourceDeformedRootResources->MeshLODIndex : -1;
-
-			// Projection Buffers
-			const bool bHasSkinnedInterpolation = (bIsRootValid && HairStrandsBuffer->SourceDeformedRootResources->IsValid(MeshLODIndex));
-			const EHairSimulationInterpolationMode InterpolationModeValue = bHasSkinnedInterpolation ? (ProxyData->GlobalInterpolation ? 
-				EHairSimulationInterpolationMode::RBF : EHairSimulationInterpolationMode::Skinned ): EHairSimulationInterpolationMode::Rigid;
-
-			const FHairStrandsRestRootResource::FLOD* RestMeshProjection = bHasSkinnedInterpolation ? &(HairStrandsBuffer->SourceRestRootResources->LODs[MeshLODIndex]) : nullptr;
-			const FHairStrandsDeformedRootResource::FLOD* DeformedMeshProjection = bHasSkinnedInterpolation ? &(HairStrandsBuffer->SourceDeformedRootResources->LODs[MeshLODIndex]) : nullptr;
-
-			FRHIShaderResourceView* RestTrianglePositionASRV = (bHasSkinnedInterpolation && RestMeshProjection) ?
-				RestMeshProjection->RestRootTrianglePosition0Buffer.SRV.GetReference() : FNiagaraRenderer::GetDummyFloatBuffer();
-			FRHIShaderResourceView* RestTrianglePositionBSRV = (bHasSkinnedInterpolation && RestMeshProjection) ?
-				RestMeshProjection->RestRootTrianglePosition1Buffer.SRV.GetReference() : FNiagaraRenderer::GetDummyFloatBuffer();
-			FRHIShaderResourceView* RestTrianglePositionCSRV = (bHasSkinnedInterpolation && RestMeshProjection) ?
-				RestMeshProjection->RestRootTrianglePosition2Buffer.SRV.GetReference() : FNiagaraRenderer::GetDummyFloatBuffer();
-
-			FRHIShaderResourceView* DeformedTrianglePositionASRV = (bHasSkinnedInterpolation && DeformedMeshProjection) ?
-				DeformedMeshProjection->GetDeformedRootTrianglePosition0Buffer(FHairStrandsDeformedRootResource::FLOD::Current).SRV.GetReference() : FNiagaraRenderer::GetDummyFloatBuffer();
-			FRHIShaderResourceView* DeformedTrianglePositionBSRV = (bHasSkinnedInterpolation && DeformedMeshProjection) ?
-				DeformedMeshProjection->GetDeformedRootTrianglePosition1Buffer(FHairStrandsDeformedRootResource::FLOD::Current).SRV.GetReference() : FNiagaraRenderer::GetDummyFloatBuffer();
-			FRHIShaderResourceView* DeformedTrianglePositionCSRV = (bHasSkinnedInterpolation && DeformedMeshProjection) ?
-				DeformedMeshProjection->GetDeformedRootTrianglePosition2Buffer(FHairStrandsDeformedRootResource::FLOD::Current).SRV.GetReference() : FNiagaraRenderer::GetDummyFloatBuffer();
-			FRHIShaderResourceView* RootBarycentricCoordinatesSRV = (bHasSkinnedInterpolation && RestMeshProjection) ?
-				RestMeshProjection->RootTriangleBarycentricBuffer.SRV.GetReference() : FNiagaraRenderer::GetDummyUIntBuffer();
-
-			// RBF buffers
-			const bool bHasSamples = (RestMeshProjection && RestMeshProjection->SampleCount > 0);
-			const int32 SampleCountValue = bHasSamples ? RestMeshProjection->SampleCount : 0;
-
-			FShaderResourceViewRHIRef RestSamplePositionsBufferSRV = (bHasSamples && RestMeshProjection) ?
-				RestMeshProjection->RestSamplePositionsBuffer.SRV.GetReference() : FNiagaraRenderer::GetDummyFloatBuffer();
-			FShaderResourceViewRHIRef MeshSampleWeightsBufferSRV = (bHasSamples && DeformedMeshProjection) ?
-				DeformedMeshProjection->GetMeshSampleWeightsBuffer(FHairStrandsDeformedRootResource::FLOD::Current).SRV.GetReference() : FNiagaraRenderer::GetDummyFloatBuffer();
-
-		
-			// Simulation setup (we update the rest configuration based on the deformed positions 
-			// if in restupdate mode or if we are resetting the sim and using RBF transfer since the rest positions are not matrching the physics asset)
-			const int32 NeedResetValue = (ProxyData->TickCount <= GHairSimulationMaxDelay) || !HairStrandsBuffer->bValidGeometryType;
-			const int32 RestUpdateValue = GHairSimulationRestUpdate || (NeedResetValue && ProxyData->bSkinningTransfer);
-			const int32 LocalSimulationValue = ProxyData->LocalSimulation;
-
-			HairStrandsBuffer->bValidGeometryType = true;
-			
-			// Offsets / Transforms
-			FVector3f RestPositionOffsetValue = (FVector3f)ProxyData->HairStrandsBuffer->SourceRestResources->GetPositionOffset();
-
-			FMatrix44f RigidTransformFloat = FMatrix44f(ProxyData->HairGroupInstance ? ProxyData->HairGroupInstance->Debug.RigidCurrentLocalToWorld.ToMatrixWithScale():
-																			ProxyData->WorldTransform.ToMatrixWithScale());
-			FMatrix44f WorldTransformFloat = FMatrix44f(ProxyData->HairGroupInstance ? ProxyData->HairGroupInstance->GetCurrentLocalToWorld().ToMatrixWithScale() :
-																			ProxyData->WorldTransform.ToMatrixWithScale());
-			FMatrix44f BoneTransformFloat = FMatrix44f(ProxyData->BoneTransform.ToMatrixWithScale()) * RigidTransformFloat;
-			
-			if (ProxyData->LocalSimulation)
-			{
-				const FMatrix44d WorldTransformDouble(WorldTransformFloat);
-				const FMatrix44d BoneTransformDouble(BoneTransformFloat);
-
-				// Due to large world coordinate we store the relative world transform in double precision 
-				WorldTransformFloat = FMatrix44f(WorldTransformDouble * BoneTransformDouble.Inverse());
-			}
-			
-			if (!bIsRootValid && bHasSkinningBinding)
-			{
-				UE_LOG(LogHairStrands, Log, TEXT("FNDIHairStrandsParametersCS() Groom Asset %s from component %s is set to use skinning interpolation but the skin resources are not valid"),
-					 *ProxyData->HairGroupInstance->Debug.GroomAssetName, *ProxyData->HairGroupInstance->Debug.MeshComponentName);
-			}
-
-			FRHITransitionInfo Transitions[] = {
-				FRHITransitionInfo(DeformedPositionBufferUAV, ERHIAccess::Unknown, ERHIAccess::UAVCompute),
-				FRHITransitionInfo(HairStrandsBuffer->BoundingBoxBuffer.UAV, ERHIAccess::Unknown, ERHIAccess::UAVCompute),
-				FRHITransitionInfo(HairStrandsBuffer->ParamsScaleBuffer.UAV, ERHIAccess::Unknown, ERHIAccess::SRVCompute),
-				FRHITransitionInfo(HairStrandsBuffer->CurvesOffsetsBuffer.UAV, ERHIAccess::Unknown, ERHIAccess::SRVCompute)
-			};
-			RHICmdList.Transition(MakeArrayView(Transitions, UE_ARRAY_COUNT(Transitions)));
-
-			// Set shader constants
-			SetShaderValue(RHICmdList, ComputeShaderRHI, BoundingBoxOffsets, HairStrandsBuffer->BoundingBoxOffsets);
-			SetShaderValue(RHICmdList, ComputeShaderRHI, WorldTransform, WorldTransformFloat);
-			SetShaderValue(RHICmdList, ComputeShaderRHI, WorldInverse, WorldTransformFloat.Inverse());
-			SetShaderValue(RHICmdList, ComputeShaderRHI, WorldRotation, WorldTransformFloat.GetMatrixWithoutScale().ToQuat());
-			SetShaderValue(RHICmdList, ComputeShaderRHI, NumStrands, ProxyData->NumStrands);
-			SetShaderValue(RHICmdList, ComputeShaderRHI, StrandSize, ProxyData->StrandsSize);
-			SetShaderValue(RHICmdList, ComputeShaderRHI, BoneTransform, BoneTransformFloat);
-			SetShaderValue(RHICmdList, ComputeShaderRHI, BoneInverse, BoneTransformFloat.Inverse());
-			SetShaderValue(RHICmdList, ComputeShaderRHI, BoneRotation, BoneTransformFloat.GetMatrixWithoutScale().ToQuat());
-			SetShaderValue(RHICmdList, ComputeShaderRHI, BoneLinearVelocity, ProxyData->BoneLinearVelocity);
-			SetShaderValue(RHICmdList, ComputeShaderRHI, BoneAngularVelocity, ProxyData->BoneAngularVelocity);
-			SetShaderValue(RHICmdList, ComputeShaderRHI, BoneLinearAcceleration, ProxyData->BoneLinearAcceleration);
-			SetShaderValue(RHICmdList, ComputeShaderRHI, BoneAngularAcceleration, ProxyData->BoneAngularAcceleration);
-			SetShaderValue(RHICmdList, ComputeShaderRHI, ResetSimulation, NeedResetValue);
-			SetShaderValue(RHICmdList, ComputeShaderRHI, InterpolationMode, int32(InterpolationModeValue));
-			SetShaderValue(RHICmdList, ComputeShaderRHI, RestUpdate, RestUpdateValue);
-			SetShaderValue(RHICmdList, ComputeShaderRHI, LocalSimulation, LocalSimulationValue);
-			SetShaderValue(RHICmdList, ComputeShaderRHI, RestRootOffset, FVector3f::ZeroVector);
-			SetShaderValue(RHICmdList, ComputeShaderRHI, DeformedRootOffset, FVector3f::ZeroVector);
-			SetShaderValue(RHICmdList, ComputeShaderRHI, RestPositionOffset, RestPositionOffsetValue);
-			SetShaderValue(RHICmdList, ComputeShaderRHI, SampleCount, SampleCountValue);
-
-			// Set Shader UAV
-			SetUAVParameter(RHICmdList, ComputeShaderRHI, DeformedPositionBuffer, DeformedPositionBufferUAV);
-			SetUAVParameter(RHICmdList, ComputeShaderRHI, BoundingBoxBuffer, HairStrandsBuffer->BoundingBoxBuffer.UAV);
-
-			// Set Shader SRV
-			SetSRVParameter(RHICmdList, ComputeShaderRHI, CurvesOffsetsBuffer, HairStrandsBuffer->CurvesOffsetsBuffer.SRV);
-			SetSRVParameter(RHICmdList, ComputeShaderRHI, ParamsScaleBuffer, HairStrandsBuffer->ParamsScaleBuffer.SRV);
-			SetSRVParameter(RHICmdList, ComputeShaderRHI, RestPositionBuffer, HairStrandsBuffer->SourceRestResources->PositionBuffer.SRV);
-			SetSRVParameter(RHICmdList, ComputeShaderRHI, DeformedPositionOffset, DeformedPositionOffsetSRV);
-			SetSRVParameter(RHICmdList, ComputeShaderRHI, RestTrianglePositionABuffer, RestTrianglePositionASRV);
-			SetSRVParameter(RHICmdList, ComputeShaderRHI, RestTrianglePositionBBuffer, RestTrianglePositionBSRV);
-			SetSRVParameter(RHICmdList, ComputeShaderRHI, RestTrianglePositionCBuffer, RestTrianglePositionCSRV);
-			SetSRVParameter(RHICmdList, ComputeShaderRHI, DeformedTrianglePositionABuffer, DeformedTrianglePositionASRV);
-			SetSRVParameter(RHICmdList, ComputeShaderRHI, DeformedTrianglePositionBBuffer, DeformedTrianglePositionBSRV);
-			SetSRVParameter(RHICmdList, ComputeShaderRHI, DeformedTrianglePositionCBuffer, DeformedTrianglePositionCSRV);
-			SetSRVParameter(RHICmdList, ComputeShaderRHI, RestSamplePositionsBuffer, RestSamplePositionsBufferSRV);
-			SetSRVParameter(RHICmdList, ComputeShaderRHI, MeshSampleWeightsBuffer, MeshSampleWeightsBufferSRV);
-			SetSRVParameter(RHICmdList, ComputeShaderRHI, RootBarycentricCoordinatesBuffer, RootBarycentricCoordinatesSRV);
-		}
-		else
-		{
-			if (bIsHairValid)
-			{
-				ProxyData->HairStrandsBuffer->bValidGeometryType = false;
-			}
-			// Set shader constants
-			SetShaderValue(RHICmdList, ComputeShaderRHI, BoundingBoxOffsets, FIntVector4(0,1,2,3));
-			SetShaderValue(RHICmdList, ComputeShaderRHI, WorldTransform, FMatrix44f::Identity);
-			SetShaderValue(RHICmdList, ComputeShaderRHI, WorldInverse, FMatrix44f::Identity);
-			SetShaderValue(RHICmdList, ComputeShaderRHI, WorldRotation, FQuat4f::Identity);
-			SetShaderValue(RHICmdList, ComputeShaderRHI, NumStrands, 1);
-			SetShaderValue(RHICmdList, ComputeShaderRHI, StrandSize, 1);
-			SetShaderValue(RHICmdList, ComputeShaderRHI, BoneTransform, FMatrix44f::Identity);
-			SetShaderValue(RHICmdList, ComputeShaderRHI, BoneInverse, FMatrix44f::Identity);
-			SetShaderValue(RHICmdList, ComputeShaderRHI, BoneRotation, FQuat4f::Identity);
-			SetShaderValue(RHICmdList, ComputeShaderRHI, BoneLinearVelocity, FVector3f::ZeroVector);
-			SetShaderValue(RHICmdList, ComputeShaderRHI, BoneAngularVelocity, FVector3f::ZeroVector);
-			SetShaderValue(RHICmdList, ComputeShaderRHI, BoneLinearAcceleration, FVector3f::ZeroVector);
-			SetShaderValue(RHICmdList, ComputeShaderRHI, BoneAngularAcceleration, FVector3f::ZeroVector);
-			SetShaderValue(RHICmdList, ComputeShaderRHI, ResetSimulation, 0);
-			SetShaderValue(RHICmdList, ComputeShaderRHI, InterpolationMode, 0);
-			SetShaderValue(RHICmdList, ComputeShaderRHI, RestUpdate, 0);
-			SetShaderValue(RHICmdList, ComputeShaderRHI, LocalSimulation, 0);
-			SetShaderValue(RHICmdList, ComputeShaderRHI, RestRootOffset, FVector3f::ZeroVector);
-			SetShaderValue(RHICmdList, ComputeShaderRHI, DeformedRootOffset, FVector3f::ZeroVector);
-			SetShaderValue(RHICmdList, ComputeShaderRHI, RestPositionOffset, FVector3f::ZeroVector);
-			SetShaderValue(RHICmdList, ComputeShaderRHI, SampleCount, 0);
-
-			// Set Shader UAV
-			SetUAVParameter(RHICmdList, ComputeShaderRHI, DeformedPositionBuffer, Context.ComputeDispatchInterface->GetEmptyUAVFromPool(RHICmdList, PF_R32_FLOAT, ENiagaraEmptyUAVType::Buffer));
-			SetUAVParameter(RHICmdList, ComputeShaderRHI, BoundingBoxBuffer, Context.ComputeDispatchInterface->GetEmptyUAVFromPool(RHICmdList, PF_R32_UINT, ENiagaraEmptyUAVType::Buffer));
-
-			// Set Shader SRV
-			SetSRVParameter(RHICmdList, ComputeShaderRHI, CurvesOffsetsBuffer, FNiagaraRenderer::GetDummyUIntBuffer());
-			SetSRVParameter(RHICmdList, ComputeShaderRHI, RestPositionBuffer, FNiagaraRenderer::GetDummyFloatBuffer());
-			SetSRVParameter(RHICmdList, ComputeShaderRHI, DeformedPositionOffset, FNiagaraRenderer::GetDummyFloatBuffer());
-			SetSRVParameter(RHICmdList, ComputeShaderRHI, RestTrianglePositionABuffer, FNiagaraRenderer::GetDummyFloatBuffer());
-			SetSRVParameter(RHICmdList, ComputeShaderRHI, RestTrianglePositionBBuffer, FNiagaraRenderer::GetDummyFloatBuffer());
-			SetSRVParameter(RHICmdList, ComputeShaderRHI, RestTrianglePositionCBuffer, FNiagaraRenderer::GetDummyFloatBuffer());
-			SetSRVParameter(RHICmdList, ComputeShaderRHI, DeformedTrianglePositionABuffer, FNiagaraRenderer::GetDummyFloatBuffer());
-			SetSRVParameter(RHICmdList, ComputeShaderRHI, DeformedTrianglePositionBBuffer, FNiagaraRenderer::GetDummyFloatBuffer());
-			SetSRVParameter(RHICmdList, ComputeShaderRHI, DeformedTrianglePositionCBuffer, FNiagaraRenderer::GetDummyFloatBuffer());
-			SetSRVParameter(RHICmdList, ComputeShaderRHI, RestSamplePositionsBuffer, FNiagaraRenderer::GetDummyFloatBuffer());
-			SetSRVParameter(RHICmdList, ComputeShaderRHI, MeshSampleWeightsBuffer, FNiagaraRenderer::GetDummyFloatBuffer());
-			SetSRVParameter(RHICmdList, ComputeShaderRHI, RootBarycentricCoordinatesBuffer, FNiagaraRenderer::GetDummyUIntBuffer());
-			SetSRVParameter(RHICmdList, ComputeShaderRHI, ParamsScaleBuffer, FNiagaraRenderer::GetDummyFloatBuffer());
-		}
-	}
-
-	void Unset(FRHICommandList& RHICmdList, const FNiagaraDataInterfaceSetArgs& Context) const
-	{
-		FRHIComputeShader* ShaderRHI = RHICmdList.GetBoundComputeShader();
-		SetUAVParameter(RHICmdList, ShaderRHI, DeformedPositionBuffer, nullptr);
-		SetUAVParameter(RHICmdList, ShaderRHI, BoundingBoxBuffer, nullptr);
-	}
-
-private:
-
-	LAYOUT_FIELD(FShaderParameter, WorldTransform);
-	LAYOUT_FIELD(FShaderParameter, WorldInverse);
-	LAYOUT_FIELD(FShaderParameter, WorldRotation);
-	LAYOUT_FIELD(FShaderParameter, NumStrands);
-	LAYOUT_FIELD(FShaderParameter, StrandSize);
-	LAYOUT_FIELD(FShaderParameter, BoneTransform);
-	LAYOUT_FIELD(FShaderParameter, BoneInverse);
-	LAYOUT_FIELD(FShaderParameter, BoneRotation);
-	LAYOUT_FIELD(FShaderParameter, BoneLinearVelocity);
-	LAYOUT_FIELD(FShaderParameter, BoneAngularVelocity);
-	LAYOUT_FIELD(FShaderParameter, BoneLinearAcceleration);
-	LAYOUT_FIELD(FShaderParameter, BoneAngularAcceleration);
-	LAYOUT_FIELD(FShaderParameter, ResetSimulation);
-	LAYOUT_FIELD(FShaderParameter, InterpolationMode);
-	LAYOUT_FIELD(FShaderParameter, RestUpdate);
-	LAYOUT_FIELD(FShaderParameter, LocalSimulation);
-	LAYOUT_FIELD(FShaderParameter, RestRootOffset);
-	LAYOUT_FIELD(FShaderParameter, DeformedRootOffset);
-	LAYOUT_FIELD(FShaderParameter, SampleCount);
-	LAYOUT_FIELD(FShaderParameter, RestPositionOffset);
-	LAYOUT_FIELD(FShaderParameter, BoundingBoxOffsets);
-
-	LAYOUT_FIELD(FShaderResourceParameter, DeformedPositionBuffer);
-	LAYOUT_FIELD(FShaderResourceParameter, BoundingBoxBuffer);
-
-	LAYOUT_FIELD(FShaderResourceParameter, CurvesOffsetsBuffer);
-	LAYOUT_FIELD(FShaderResourceParameter, RestPositionBuffer);
-	LAYOUT_FIELD(FShaderResourceParameter, RootBarycentricCoordinatesBuffer);
-	LAYOUT_FIELD(FShaderResourceParameter, RestTrianglePositionABuffer);
-	LAYOUT_FIELD(FShaderResourceParameter, RestTrianglePositionBBuffer);
-	LAYOUT_FIELD(FShaderResourceParameter, RestTrianglePositionCBuffer);
-	LAYOUT_FIELD(FShaderResourceParameter, DeformedTrianglePositionABuffer);
-	LAYOUT_FIELD(FShaderResourceParameter, DeformedTrianglePositionBBuffer);
-	LAYOUT_FIELD(FShaderResourceParameter, DeformedTrianglePositionCBuffer);
-	LAYOUT_FIELD(FShaderResourceParameter, RestSamplePositionsBuffer);
-	LAYOUT_FIELD(FShaderResourceParameter, MeshSampleWeightsBuffer);
-	LAYOUT_FIELD(FShaderResourceParameter, DeformedPositionOffset);
-	LAYOUT_FIELD(FShaderResourceParameter, ParamsScaleBuffer);
-};
-
-IMPLEMENT_TYPE_LAYOUT(FNDIHairStrandsParametersCS);
-
-IMPLEMENT_NIAGARA_DI_PARAMETER(UNiagaraDataInterfaceHairStrands, FNDIHairStrandsParametersCS);
-
-//------------------------------------------------------------------------------------------------------------
-
 void FNDIHairStrandsProxy::ConsumePerInstanceDataFromGameThread(void* PerInstanceData, const FNiagaraSystemInstanceID& Instance)
 {
 	FNDIHairStrandsData* SourceData = static_cast<FNDIHairStrandsData*>(PerInstanceData);
@@ -1246,83 +852,55 @@ bool UNiagaraDataInterfaceHairStrands::PerInstanceTick(void* PerInstanceData, FN
 /** Compute shader to interpolate the groom position from the sim cache */
 class FInterpolateGroomGuidesCS : public FGlobalShader
 {
-	DECLARE_SHADER_TYPE(FInterpolateGroomGuidesCS, Global)
-	
+	DECLARE_GLOBAL_SHADER(FInterpolateGroomGuidesCS)
+	SHADER_USE_PARAMETER_STRUCT(FInterpolateGroomGuidesCS, FGlobalShader);
+
 	class FInterpolationType : SHADER_PERMUTATION_INT("PERMUTATION_INTERPOLATION", 2);
 	using FPermutationDomain = TShaderPermutationDomain<FInterpolationType>;
 
-public:
-	static bool ShouldCompilePermutation(const FGlobalShaderPermutationParameters& Parameters) { return IsHairStrandsSupported(EHairStrandsShaderType::All, Parameters.Platform); }
+	BEGIN_SHADER_PARAMETER_STRUCT(FParameters, )
+		SHADER_PARAMETER_SRV(Buffer<float4>,	RestTrianglePositionABuffer)
+		SHADER_PARAMETER_SRV(Buffer<float4>,	RestTrianglePositionBBuffer)
+		SHADER_PARAMETER_SRV(Buffer<float4>,	RestTrianglePositionCBuffer)
+		SHADER_PARAMETER_SRV(Buffer<float4>,	DeformedTrianglePositionABuffer)
+		SHADER_PARAMETER_SRV(Buffer<float4>,	DeformedTrianglePositionBBuffer)
+		SHADER_PARAMETER_SRV(Buffer<float4>,	DeformedTrianglePositionCBuffer)
+		SHADER_PARAMETER_SRV(Buffer<uint>,		RootBarycentricCoordinatesBuffer)
+
+		SHADER_PARAMETER_SRV(Buffer<uint>,		VertexToCurveIndexBuffer)
+		SHADER_PARAMETER_UAV(RWBuffer<uint4>,	DeformedPositionBuffer)
+		SHADER_PARAMETER_SRV(Buffer<uint4>,		RestPositionBuffer)
+		SHADER_PARAMETER_SRV(Buffer<uint>,		CurvesOffsetsBuffer)
+		SHADER_PARAMETER_SRV(Buffer<float4>,	DeformedPositionOffset)
+		SHADER_PARAMETER(FVector3f,				RestPositionOffset)
+
+		SHADER_PARAMETER(int,					StrandsSize)
+		SHADER_PARAMETER(int,					NumPoints)
+
+		SHADER_PARAMETER_SRV(Buffer<float>,		NiagaraFloatBuffer)
+		SHADER_PARAMETER(int,					NiagaraFloatStride)
+
+		SHADER_PARAMETER(int,					NodePositionComponent)
+		SHADER_PARAMETER(int,					RestPositionComponent)
+	END_SHADER_PARAMETER_STRUCT()
+
+	static bool ShouldCompilePermutation(const FGlobalShaderPermutationParameters& Parameters)
+	{
+		return IsHairStrandsSupported(EHairStrandsShaderType::All, Parameters.Platform);
+	}
 
 	static void ModifyCompilationEnvironment(const FGlobalShaderPermutationParameters& Parameters, FShaderCompilerEnvironment& OutEnvironment)
 	{
 		FGlobalShader::ModifyCompilationEnvironment(Parameters, OutEnvironment);
 		OutEnvironment.SetDefine(TEXT("THREAD_COUNT"), NIAGARA_HAIR_STRANDS_THREAD_COUNT_INTERPOLATE);
 	}
-
-	FInterpolateGroomGuidesCS(const ShaderMetaType::CompiledShaderInitializerType& Initializer)
-		: FGlobalShader(Initializer)
-	{
-		RestTrianglePositionABuffer.Bind(Initializer.ParameterMap, TEXT("RestTrianglePositionABuffer"));
-		RestTrianglePositionBBuffer.Bind(Initializer.ParameterMap, TEXT("RestTrianglePositionBBuffer"));
-		RestTrianglePositionCBuffer.Bind(Initializer.ParameterMap, TEXT("RestTrianglePositionCBuffer"));
-		DeformedTrianglePositionABuffer.Bind(Initializer.ParameterMap, TEXT("DeformedTrianglePositionABuffer"));
-		DeformedTrianglePositionBBuffer.Bind(Initializer.ParameterMap, TEXT("DeformedTrianglePositionBBuffer"));
-		DeformedTrianglePositionCBuffer.Bind(Initializer.ParameterMap, TEXT("DeformedTrianglePositionCBuffer"));
-		RootBarycentricCoordinatesBuffer.Bind(Initializer.ParameterMap, TEXT("RootBarycentricCoordinatesBuffer"));
-		
-		VertexToCurveIndexBuffer.Bind(Initializer.ParameterMap, TEXT("VertexToCurveIndexBuffer"));
-		CurvesOffsetsBuffer.Bind(Initializer.ParameterMap, TEXT("CurvesOffsetsBuffer"));
-		
-		DeformedPositionOffset.Bind(Initializer.ParameterMap, TEXT("DeformedPositionOffset"));
-		RestPositionOffset.Bind(Initializer.ParameterMap, TEXT("RestPositionOffset"));
-		DeformedPositionBuffer.Bind(Initializer.ParameterMap, TEXT("DeformedPositionBuffer"));
-		RestPositionBuffer.Bind(Initializer.ParameterMap, TEXT("RestPositionBuffer"));
-		
-		StrandsSize.Bind(Initializer.ParameterMap, TEXT("StrandsSize"));
-		NumPoints.Bind(Initializer.ParameterMap, TEXT("NumPoints"));
-		
-		NiagaraFloatBuffer.Bind(Initializer.ParameterMap, TEXT("NiagaraFloatBuffer"));
-		NiagaraFloatStride.Bind(Initializer.ParameterMap, TEXT("NiagaraFloatStride"));
-		NodePositionComponent.Bind(Initializer.ParameterMap, TEXT("NodePositionComponent"));
-		RestPositionComponent.Bind(Initializer.ParameterMap, TEXT("RestPositionComponent"));
-	}
-
-	FInterpolateGroomGuidesCS()
-	{}
-
-	LAYOUT_FIELD(FShaderResourceParameter, RestTrianglePositionABuffer)
-	LAYOUT_FIELD(FShaderResourceParameter, RestTrianglePositionBBuffer)
-	LAYOUT_FIELD(FShaderResourceParameter, RestTrianglePositionCBuffer)
-	LAYOUT_FIELD(FShaderResourceParameter, DeformedTrianglePositionABuffer)
-	LAYOUT_FIELD(FShaderResourceParameter, DeformedTrianglePositionBBuffer)
-	LAYOUT_FIELD(FShaderResourceParameter, DeformedTrianglePositionCBuffer)
-	LAYOUT_FIELD(FShaderResourceParameter, RootBarycentricCoordinatesBuffer)
-	
-	LAYOUT_FIELD(FShaderResourceParameter, VertexToCurveIndexBuffer)
-	LAYOUT_FIELD(FShaderResourceParameter, CurvesOffsetsBuffer)
-	
-	LAYOUT_FIELD(FShaderResourceParameter, DeformedPositionOffset)
-	LAYOUT_FIELD(FShaderParameter,RestPositionOffset)
-	LAYOUT_FIELD(FShaderResourceParameter, DeformedPositionBuffer)
-	LAYOUT_FIELD(FShaderResourceParameter, RestPositionBuffer)
-	
-	LAYOUT_FIELD(FShaderParameter, StrandsSize)
-	LAYOUT_FIELD(FShaderParameter, NumPoints)
-
-	LAYOUT_FIELD(FShaderResourceParameter, NiagaraFloatBuffer)
-	LAYOUT_FIELD(FShaderParameter, NiagaraFloatStride)
-	LAYOUT_FIELD(FShaderParameter, NodePositionComponent)
-	LAYOUT_FIELD(FShaderParameter, RestPositionComponent)
 };
 
 IMPLEMENT_SHADER_TYPE(, FInterpolateGroomGuidesCS, TEXT("/Plugin/Runtime/HairStrands/Private/NiagaraInterpolateGroomGuides.usf"), TEXT("MainCS"), SF_Compute);
 
 //------------------------------------------------------------------------------------------------------------
 
-DECLARE_GPU_STAT(InterpolateGroomGuidesCS);
-
-static void InterpolateGroomGuides(FRHICommandListImmediate& RHICmdList, FNiagaraDataBuffer* ParticlesBuffer,
+static void InterpolateGroomGuides(FRDGBuilder& GraphBuilder, FNiagaraDataBuffer* ParticlesBuffer,
 	const uint32 NodePositionComponent,  const uint32 RestPositionComponent, FNDIHairStrandsBuffer* HairStrandsBuffer,
 	const uint32 StrandsSize, const bool bHasSkinningBinding, const bool bHasValidGeometry)
 {
@@ -1332,66 +910,61 @@ static void InterpolateGroomGuides(FRHICommandListImmediate& RHICmdList, FNiagar
 
 	if(bIsRestValid && bIsDeformedValid && bHasValidGeometry)
 	{
-		SCOPED_GPU_STAT(RHICmdList, InterpolateGroomGuidesCS);
-		SCOPED_DRAW_EVENT(RHICmdList, InterpolateGroomGuidesCS);
-		
 		bool bIsRootValid = HairStrandsBuffer->SourceDeformedRootResources && HairStrandsBuffer->SourceDeformedRootResources->IsInitialized() &&
-								HairStrandsBuffer->SourceRestRootResources && HairStrandsBuffer->SourceRestRootResources->IsInitialized() && bHasSkinningBinding;;
+								HairStrandsBuffer->SourceRestRootResources && HairStrandsBuffer->SourceRestRootResources->IsInitialized() && bHasSkinningBinding;
 
 		const uint32 NumPoints = HairStrandsBuffer->SourceRestResources->GetVertexCount();
+		const int32 MeshLODIndex = bIsRootValid ? HairStrandsBuffer->SourceDeformedRootResources->MeshLODIndex : -1;
+		bIsRootValid = bIsRootValid && HairStrandsBuffer->SourceDeformedRootResources->IsValid(MeshLODIndex);
 
 		FInterpolateGroomGuidesCS::FPermutationDomain InterpolationDomain;
 		InterpolationDomain.Set<FInterpolateGroomGuidesCS::FInterpolationType>(!bIsRootValid);
-		
+
 		TShaderMapRef<FInterpolateGroomGuidesCS> ComputeShader(GetGlobalShaderMap(GMaxRHIFeatureLevel), InterpolationDomain);
-		SetComputePipelineState(RHICmdList, ComputeShader.GetComputeShader());
-		
-		FRHIComputeShader* ShaderRHI = RHICmdList.GetBoundComputeShader();
-		
-		const int32 MeshLODIndex = bIsRootValid ? HairStrandsBuffer->SourceDeformedRootResources->MeshLODIndex : -1;
-		bIsRootValid = bIsRootValid && HairStrandsBuffer->SourceDeformedRootResources->IsValid(MeshLODIndex);
-		
-		SetShaderValue(RHICmdList, ShaderRHI, ComputeShader->StrandsSize, StrandsSize);
-		SetShaderValue(RHICmdList, ShaderRHI, ComputeShader->NumPoints, NumPoints);
 
-		SetSRVParameter(RHICmdList, ShaderRHI, ComputeShader->NiagaraFloatBuffer, ParticlesBuffer->GetGPUBufferFloat().SRV);
-		SetShaderValue(RHICmdList, ShaderRHI, ComputeShader->NiagaraFloatStride, ParticlesBuffer->GetFloatStride() / sizeof(float));
-		SetShaderValue(RHICmdList, ShaderRHI, ComputeShader->NodePositionComponent, NodePositionComponent);
-		SetShaderValue(RHICmdList, ShaderRHI, ComputeShader->RestPositionComponent, RestPositionComponent);
+		FInterpolateGroomGuidesCS::FParameters* PassParameters = GraphBuilder.AllocParameters<FInterpolateGroomGuidesCS::FParameters>();
+		PassParameters->StrandsSize = StrandsSize;
+		PassParameters->NumPoints = NumPoints;
+
+		PassParameters->NiagaraFloatBuffer = ParticlesBuffer->GetGPUBufferFloat().SRV;
+		PassParameters->NiagaraFloatStride = ParticlesBuffer->GetFloatStride() / sizeof(float);
+		PassParameters->NodePositionComponent = NodePositionComponent;
+		PassParameters->RestPositionComponent = RestPositionComponent;
 
 		const int32 BufferIndex = HairStrandsBuffer->SourceDeformedResources->GetIndex(FHairStrandsDeformedResource::Current);
 
-		SetSRVParameter(RHICmdList, ShaderRHI, ComputeShader->DeformedPositionOffset, HairStrandsBuffer->SourceDeformedResources->DeformedOffsetBuffer[BufferIndex].SRV);
-        SetShaderValue(RHICmdList, ShaderRHI, ComputeShader->RestPositionOffset, FVector3f(HairStrandsBuffer->SourceRestResources->GetPositionOffset()));
-        SetUAVParameter(RHICmdList, ShaderRHI, ComputeShader->DeformedPositionBuffer, HairStrandsBuffer->SourceDeformedResources->DeformedPositionBuffer[BufferIndex].UAV);
-        SetSRVParameter(RHICmdList, ShaderRHI, ComputeShader->RestPositionBuffer, HairStrandsBuffer->SourceRestResources->PositionBuffer.SRV);
-	
-		SetSRVParameter(RHICmdList, ShaderRHI, ComputeShader->VertexToCurveIndexBuffer, HairStrandsBuffer->PointsCurveBuffer.SRV);
-		SetSRVParameter(RHICmdList, ShaderRHI, ComputeShader->CurvesOffsetsBuffer, HairStrandsBuffer->CurvesOffsetsBuffer.SRV);
+		PassParameters->DeformedPositionOffset = HairStrandsBuffer->SourceDeformedResources->DeformedOffsetBuffer[BufferIndex].SRV;
+		PassParameters->RestPositionOffset = FVector3f(HairStrandsBuffer->SourceRestResources->GetPositionOffset());
+		PassParameters->DeformedPositionBuffer = HairStrandsBuffer->SourceDeformedResources->DeformedPositionBuffer[BufferIndex].UAV;
+		PassParameters->RestPositionBuffer = HairStrandsBuffer->SourceRestResources->PositionBuffer.SRV;
 		
+		PassParameters->VertexToCurveIndexBuffer = HairStrandsBuffer->PointsCurveBuffer.SRV;
+		PassParameters->CurvesOffsetsBuffer = HairStrandsBuffer->CurvesOffsetsBuffer.SRV;
 		if(bIsRootValid)
 		{
 			const FHairStrandsRestRootResource::FLOD* RestMeshProjection = &(HairStrandsBuffer->SourceRestRootResources->LODs[MeshLODIndex]) ;
 			const FHairStrandsDeformedRootResource::FLOD* DeformedMeshProjection =  &(HairStrandsBuffer->SourceDeformedRootResources->LODs[MeshLODIndex]);
 
-			SetSRVParameter(RHICmdList, ShaderRHI, ComputeShader->RestTrianglePositionABuffer, RestMeshProjection->RestRootTrianglePosition0Buffer.SRV);
-			SetSRVParameter(RHICmdList, ShaderRHI, ComputeShader->RestTrianglePositionBBuffer, RestMeshProjection->RestRootTrianglePosition1Buffer.SRV);
-			SetSRVParameter(RHICmdList, ShaderRHI, ComputeShader->RestTrianglePositionCBuffer, RestMeshProjection->RestRootTrianglePosition2Buffer.SRV);
-			SetSRVParameter(RHICmdList, ShaderRHI, ComputeShader->DeformedTrianglePositionABuffer, DeformedMeshProjection->GetDeformedRootTrianglePosition0Buffer(FHairStrandsDeformedRootResource::FLOD::Current).SRV);
-			SetSRVParameter(RHICmdList, ShaderRHI, ComputeShader->DeformedTrianglePositionBBuffer, DeformedMeshProjection->GetDeformedRootTrianglePosition1Buffer(FHairStrandsDeformedRootResource::FLOD::Current).SRV);
-			SetSRVParameter(RHICmdList, ShaderRHI, ComputeShader->DeformedTrianglePositionCBuffer, DeformedMeshProjection->GetDeformedRootTrianglePosition2Buffer(FHairStrandsDeformedRootResource::FLOD::Current).SRV);
-			SetSRVParameter(RHICmdList, ShaderRHI, ComputeShader->RootBarycentricCoordinatesBuffer, RestMeshProjection->RootTriangleBarycentricBuffer.SRV);
+			PassParameters->RestTrianglePositionABuffer = RestMeshProjection->RestRootTrianglePosition0Buffer.SRV;
+			PassParameters->RestTrianglePositionBBuffer = RestMeshProjection->RestRootTrianglePosition1Buffer.SRV;
+			PassParameters->RestTrianglePositionCBuffer = RestMeshProjection->RestRootTrianglePosition2Buffer.SRV;
+			PassParameters->DeformedTrianglePositionABuffer = DeformedMeshProjection->GetDeformedRootTrianglePosition0Buffer(FHairStrandsDeformedRootResource::FLOD::Current).SRV;
+			PassParameters->DeformedTrianglePositionBBuffer = DeformedMeshProjection->GetDeformedRootTrianglePosition1Buffer(FHairStrandsDeformedRootResource::FLOD::Current).SRV;
+			PassParameters->DeformedTrianglePositionCBuffer = DeformedMeshProjection->GetDeformedRootTrianglePosition2Buffer(FHairStrandsDeformedRootResource::FLOD::Current).SRV;
+			PassParameters->RootBarycentricCoordinatesBuffer = RestMeshProjection->RootTriangleBarycentricBuffer.SRV;
 		}
+
 		const uint32 GroupSize = NIAGARA_HAIR_STRANDS_THREAD_COUNT_INTERPOLATE;
 		const uint32 DispatchCount = FMath::DivideAndRoundUp(NumPoints, GroupSize);
-		
-		RHICmdList.Transition(FRHITransitionInfo(HairStrandsBuffer->SourceDeformedResources->DeformedPositionBuffer[BufferIndex].UAV,
-			ERHIAccess::SRVMask, ERHIAccess::UAVCompute));
-		
-		DispatchComputeShader(RHICmdList, ComputeShader.GetShader(), DispatchCount, 1, 1);
-		
-		RHICmdList.Transition(FRHITransitionInfo(HairStrandsBuffer->SourceDeformedResources->DeformedPositionBuffer[BufferIndex].UAV,
-			ERHIAccess::UAVCompute, ERHIAccess::SRVMask));
+
+		FComputeShaderUtils::AddPass(
+			GraphBuilder,
+			RDG_EVENT_NAME("InterpolateGroomGuidesCS"),
+			ERDGPassFlags::Compute,
+			ComputeShader,
+			PassParameters,
+			FIntVector(DispatchCount, 1, 1)
+		);
 	}
 }
 
@@ -1424,11 +997,14 @@ void UNiagaraDataInterfaceHairStrands::SimCachePostReadFrame(void* OptionalPerIn
 				const bool bHasValidGeometry = InstanceData->HairGroupInstance->GeometryType != EHairGeometryType::NoneGeometry;
 				
 				ENQUEUE_RENDER_COMMAND(NiagaraInterpolateGroomSimCache)(
-				[ParticlesBuffer, HairstrandsBuffer, NodePositionComponent, RestPositionComponent, StrandsSize, bHasSkinningBinding, bHasValidGeometry](FRHICommandListImmediate& RHICmdList)
-				{
-					InterpolateGroomGuides(RHICmdList, ParticlesBuffer,
-							NodePositionComponent, RestPositionComponent, HairstrandsBuffer, StrandsSize, bHasSkinningBinding, bHasValidGeometry);
-				});
+					[ParticlesBuffer, HairstrandsBuffer, NodePositionComponent, RestPositionComponent, StrandsSize, bHasSkinningBinding, bHasValidGeometry](FRHICommandListImmediate& RHICmdList)
+					{
+						FMemMark MemMark(FMemStack::Get());
+						FRDGBuilder GraphBuilder(RHICmdList);
+						InterpolateGroomGuides(GraphBuilder, ParticlesBuffer, NodePositionComponent, RestPositionComponent, HairstrandsBuffer, StrandsSize, bHasSkinningBinding, bHasValidGeometry);
+						GraphBuilder.Execute();
+					}
+				);
 			}
 		}
 	}
@@ -1478,6 +1054,7 @@ BEGIN_FUNCTION_BUILD_OPTIMIZATION
 
 void UNiagaraDataInterfaceHairStrands::GetFunctions(TArray<FNiagaraFunctionSignature>& OutFunctions)
 {
+	using namespace NDIHairStrandsLocal;
 	{
 		FNiagaraFunctionSignature Sig;
 		Sig.Name = GetNumStrandsName;
@@ -2617,6 +2194,8 @@ DEFINE_NDI_DIRECT_FUNC_BINDER(UNiagaraDataInterfaceHairStrands, GetSampleState);
 
 void UNiagaraDataInterfaceHairStrands::GetVMExternalFunction(const FVMExternalFunctionBindingInfo& BindingInfo, void* InstanceData, FVMExternalFunction &OutFunc)
 {
+	using namespace NDIHairStrandsLocal;
+
 	if (BindingInfo.Name == GetNumStrandsName)
 	{
 		check(BindingInfo.GetNumInputs() == 1 && BindingInfo.GetNumOutputs() == 1);
@@ -3652,767 +3231,80 @@ void UNiagaraDataInterfaceHairStrands::NeedRestUpdate(FVectorVMExternalFunctionC
 #if WITH_EDITORONLY_DATA
 bool UNiagaraDataInterfaceHairStrands::GetFunctionHLSL(const FNiagaraDataInterfaceGPUParamInfo& ParamInfo, const FNiagaraDataInterfaceGeneratedFunction& FunctionInfo, int FunctionInstanceIndex, FString& OutHLSL)
 {
-	FNDIHairStrandsParametersName ParamNames(ParamInfo.DataInterfaceHLSLSymbol);
-
-	TMap<FString, FStringFormatArg> ArgsSample = {
-		{TEXT("InstanceFunctionName"), FunctionInfo.InstanceName},
-		{TEXT("NumStrandsName"), ParamNames.NumStrandsName},
-		{TEXT("StrandSizeName"), ParamNames.StrandSizeName},
-		{TEXT("WorldTransformName"), ParamNames.WorldTransformName},
-		{TEXT("WorldInverseName"), ParamNames.WorldInverseName},
-		{TEXT("WorldRotationName"), ParamNames.WorldRotationName},
-		{TEXT("DeformedPositionBufferName"), ParamNames.DeformedPositionBufferName},
-		{TEXT("CurvesOffsetsBufferName"), ParamNames.CurvesOffsetsBufferName},
-		{TEXT("RestPositionBufferName"), ParamNames.RestPositionBufferName},
-		{TEXT("HairStrandsContextName"), TEXT("DIHAIRSTRANDS_MAKE_CONTEXT(") + ParamInfo.DataInterfaceHLSLSymbol + TEXT(")")},
+	using namespace NDIHairStrandsLocal;
+
+	static const TSet<FName> ValidGpuFunctions =
+	{
+		GetStrandSizeName,
+		GetNumStrandsName,
+		GetWorldTransformName,
+		GetWorldInverseName,
+		GetStretchScaleName,
+		GetBendScaleName,
+		GetRadiusScaleName,
+		GetThicknessScaleName,
+		GetPointPositionName,
+		ComputeNodePositionName,
+		ComputeNodeOrientationName,
+		ComputeNodeMassName,
+		ComputeNodeInertiaName,
+		ComputeEdgeLengthName,
+		ComputeEdgeRotationName,
+		ComputeRestPositionName,
+		ComputeRestOrientationName,
+		ComputeLocalStateName,
+		GetLocalVectorName,
+		GetWorldVectorName,
+		AttachNodePositionName,
+		AttachNodeOrientationName,
+		AttachNodeStateName,
+		UpdateNodeStateName,
+		UpdatePointPositionName,
+		ResetPointPositionName,
+		AdvectNodePositionName,
+		AdvectNodeOrientationName,
+		UpdateLinearVelocityName,
+		UpdateAngularVelocityName,
+		GetBoundingBoxName,
+		ResetBoundingBoxName,
+		BuildBoundingBoxName,
+		SetupDistanceSpringMaterialName,
+		SolveDistanceSpringMaterialName,
+		ProjectDistanceSpringMaterialName,
+		SetupAngularSpringMaterialName,
+		SolveAngularSpringMaterialName,
+		ProjectAngularSpringMaterialName,
+		SetupStretchRodMaterialName,
+		SolveStretchRodMaterialName,
+		ProjectStretchRodMaterialName,
+		SetupBendRodMaterialName,
+		SolveBendRodMaterialName,
+		ProjectBendRodMaterialName,
+		SolveHardCollisionConstraintName,
+		ProjectHardCollisionConstraintName,
+		SolveSoftCollisionConstraintName,
+		ProjectSoftCollisionConstraintName,
+		SetupSoftCollisionConstraintName,
+		ComputeEdgeDirectionName,
+		UpdateMaterialFrameName,
+		ComputeMaterialFrameName,
+		ComputeAirDragForceName,
+		InitGridSamplesName,
+		GetSampleStateName,
+		NeedSimulationResetName,
+		HasGlobalInterpolationName,
+		NeedRestUpdateName,
 	};
 
-	if (FunctionInfo.DefinitionName == GetStrandSizeName)
-	{
-		static const TCHAR *FormatSample = TEXT(R"(
-		void {InstanceFunctionName}(out int OutStrandSize)
-		{
-			OutStrandSize = {StrandSizeName};
-		}
-		)");
-		OutHLSL += FString::Format(FormatSample, ArgsSample);
-		return true;
-	}
-	else if (FunctionInfo.DefinitionName == GetNumStrandsName)
-	{
-		static const TCHAR *FormatSample = TEXT(R"(
-		void {InstanceFunctionName}(out int OutNumStrands)
-		{
-			OutNumStrands = {NumStrandsName};
-		}
-		)");
-		OutHLSL += FString::Format(FormatSample, ArgsSample);
-		return true;
-	}
-	else if (FunctionInfo.DefinitionName == GetWorldTransformName)
-	{
-		static const TCHAR *FormatSample = TEXT(R"(
-		void {InstanceFunctionName}(out float4x4 OutWorldTransform)
-		{
-			OutWorldTransform = {WorldTransformName};
-		}
-		)");
-		OutHLSL += FString::Format(FormatSample, ArgsSample);
-		return true;
-	}
-	else if (FunctionInfo.DefinitionName == GetWorldInverseName)
-	{
-		static const TCHAR *FormatSample = TEXT(R"(
-		void {InstanceFunctionName}(out float4x4 OutWorldInverse)
-		{
-			OutWorldInverse = {WorldInverseName};
-		}
-		)");
-		OutHLSL += FString::Format(FormatSample, ArgsSample);
-		return true;
-	}
-	if (FunctionInfo.DefinitionName == GetStretchScaleName)
-	{
-		static const TCHAR* FormatSample = TEXT(R"(
-		void {InstanceFunctionName}(out float OutStretchScale)
-		{
-			{HairStrandsContextName} OutStretchScale = DIContext.ParamsScaleBuffer[GGroupThreadId.x % DIContext.StrandSize];
-		}
-		)");
-		OutHLSL += FString::Format(FormatSample, ArgsSample);
-		return true;
-	}
-	if (FunctionInfo.DefinitionName == GetBendScaleName)
-	{
-		static const TCHAR* FormatSample = TEXT(R"(
-		void {InstanceFunctionName}(out float OutBendScale)
-		{
-			{HairStrandsContextName} OutBendScale = DIContext.ParamsScaleBuffer[32 + GGroupThreadId.x % DIContext.StrandSize];
-		}
-		)");
-		OutHLSL += FString::Format(FormatSample, ArgsSample);
-		return true;
-	}
-	if (FunctionInfo.DefinitionName == GetRadiusScaleName)
-	{
-		static const TCHAR* FormatSample = TEXT(R"(
-		void {InstanceFunctionName}(out float OutRadiusScale)
-		{
-			{HairStrandsContextName} OutRadiusScale = DIContext.ParamsScaleBuffer[64 + GGroupThreadId.x % DIContext.StrandSize];
-		}
-		)");
-		OutHLSL += FString::Format(FormatSample, ArgsSample);
-		return true;
-	}
-	if (FunctionInfo.DefinitionName == GetThicknessScaleName)
-	{
-		static const TCHAR* FormatSample = TEXT(R"(
-		void {InstanceFunctionName}(out float OutThicknessScale)
-		{
-			{HairStrandsContextName} OutThicknessScale = DIContext.ParamsScaleBuffer[96 + GGroupThreadId.x % DIContext.StrandSize];
-		}
-		)");
-		OutHLSL += FString::Format(FormatSample, ArgsSample);
-		return true;
-	}
-	else if (FunctionInfo.DefinitionName == GetPointPositionName)
-	{
-		static const TCHAR *FormatSample = TEXT(R"(
-			void {InstanceFunctionName} (in int PointIndex, out float3 OutPointPosition)
-			{
-				{HairStrandsContextName} DIHairStrands_GetPointPosition(DIContext,PointIndex,OutPointPosition);
-			}
-			)");
-		OutHLSL += FString::Format(FormatSample, ArgsSample);
-		return true;
-	}
-	else if (FunctionInfo.DefinitionName == ComputeNodePositionName)
-	{
-		static const TCHAR *FormatSample = TEXT(R"(
-			void {InstanceFunctionName} (in float SmoothingFilter, out float3 OutNodePosition)
-			{
-				{HairStrandsContextName} DIHairStrands_ComputeNodePosition(DIContext,OutNodePosition);
-				DIHairStrands_SmoothNodePosition(DIContext,SmoothingFilter,OutNodePosition);
-			}
-			)");
-		OutHLSL += FString::Format(FormatSample, ArgsSample);
-		return true;
-	}
-	else if (FunctionInfo.DefinitionName == ComputeNodeOrientationName)
-	{
-		static const TCHAR *FormatSample = TEXT(R"(
-			void {InstanceFunctionName} (in float3 NodePosition, out float4 OutNodeOrientation)
-			{
-				{HairStrandsContextName} DIHairStrands_ComputeNodeOrientation(DIContext,NodePosition,OutNodeOrientation);
-			}
-			)");
-		OutHLSL += FString::Format(FormatSample, ArgsSample);
-		return true;
-	}
-	else if (FunctionInfo.DefinitionName == ComputeNodeMassName)
-	{
-		static const TCHAR *FormatSample = TEXT(R"(
-			void {InstanceFunctionName} (in float StrandsDensity, in float NodeThickness, out float OutNodeMass)
-			{
-				{HairStrandsContextName} DIHairStrands_ComputeNodeMass(DIContext,StrandsDensity,NodeThickness,OutNodeMass);
-			}
-			)");
-		OutHLSL += FString::Format(FormatSample, ArgsSample);
-		return true;
-	}
-	else if (FunctionInfo.DefinitionName == ComputeNodeInertiaName)
-	{
-		static const TCHAR *FormatSample = TEXT(R"(
-			void {InstanceFunctionName} (in float StrandsDensity, in float NodeThickness, out float3 OutNodeInertia)
-			{
-				{HairStrandsContextName} DIHairStrands_ComputeNodeInertia(DIContext,StrandsDensity,NodeThickness,OutNodeInertia);
-			}
-			)");
-		OutHLSL += FString::Format(FormatSample, ArgsSample);
-		return true;
-	}
-	else if (FunctionInfo.DefinitionName == ComputeEdgeLengthName)
-	{
-		static const TCHAR *FormatSample = TEXT(R"(
-			void {InstanceFunctionName} (in float3 NodePosition, in int NodeOffset, out float OutEdgeLength)
-			{
-				{HairStrandsContextName}
-				if(NodeOffset == 2)
-				{
-					DIHairStrands_ComputeEdgeVolume(DIContext,NodePosition,OutEdgeLength);
-				}
-				else
-				{
-					DIHairStrands_ComputeEdgeLength(DIContext,NodePosition,NodeOffset,OutEdgeLength);
-				}
-			}
-			)");
-		OutHLSL += FString::Format(FormatSample, ArgsSample);
-		return true;
-	}
-	else if (FunctionInfo.DefinitionName == ComputeEdgeRotationName)
-	{
-		static const TCHAR *FormatSample = TEXT(R"(
-			void {InstanceFunctionName} (in float4 NodeOrientation, out float4 OutEdgeRotation)
-			{
-				{HairStrandsContextName} DIHairStrands_ComputeEdgeRotation(DIContext,NodeOrientation,OutEdgeRotation);
-			}
-			)");
-		OutHLSL += FString::Format(FormatSample, ArgsSample);
-		return true;
-	}
-	else if (FunctionInfo.DefinitionName == ComputeRestPositionName)
-	{
-		static const TCHAR *FormatSample = TEXT(R"(
-			void {InstanceFunctionName} (in float3 NodePosition, out float3 OutRestPosition)
-			{
-				{HairStrandsContextName} DIHairStrands_ComputeRestPosition(DIContext,NodePosition,OutRestPosition);
-			}
-			)");
-		OutHLSL += FString::Format(FormatSample, ArgsSample);
-		return true;
-	}
-	else if (FunctionInfo.DefinitionName == ComputeRestOrientationName)
-	{
-		static const TCHAR *FormatSample = TEXT(R"(
-			void {InstanceFunctionName} (in float4 NodeOrientation, out float4 OutRestOrientation)
-			{
-				{HairStrandsContextName} DIHairStrands_ComputeRestOrientation(DIContext,NodeOrientation,OutRestOrientation);
-			}
-			)");
-		OutHLSL += FString::Format(FormatSample, ArgsSample);
-		return true;
-	}
-	else if (FunctionInfo.DefinitionName == ComputeLocalStateName)
-	{
-		static const TCHAR *FormatSample = TEXT(R"(
-			void {InstanceFunctionName} (in float3 RestPosition, in float4 RestOrientation, out float3 LocalPosition, out float4 LocalOrientation)
-			{
-				{HairStrandsContextName} DIHairStrands_ComputeLocalState(DIContext,RestPosition,RestOrientation,LocalPosition,LocalOrientation);
-			}
-			)");
-		OutHLSL += FString::Format(FormatSample, ArgsSample);
-		return true;
-	}
-	else if (FunctionInfo.DefinitionName == GetLocalVectorName)
-	{
-		static const TCHAR *FormatSample = TEXT(R"(
-			void {InstanceFunctionName} (in float3 WorldVector, in bool IsPosition, out float3 LocalVector)
-			{
-				{HairStrandsContextName} DIHairStrands_GetLocalVector(DIContext,WorldVector,IsPosition,LocalVector);
-			}
-			)");
-		OutHLSL += FString::Format(FormatSample, ArgsSample);
-		return true;
-	}
-	else if (FunctionInfo.DefinitionName == GetWorldVectorName)
-	{
-		static const TCHAR* FormatSample = TEXT(R"(
-				void {InstanceFunctionName} (in float3 LocalVector, in bool IsPosition, out float3 WorldVector)
-				{
-					{HairStrandsContextName} DIHairStrands_GetWorldVector(DIContext,LocalVector,IsPosition,WorldVector);
-				}
-				)");
-		OutHLSL += FString::Format(FormatSample, ArgsSample);
-		return true;
-	}
-	else if (FunctionInfo.DefinitionName == AttachNodePositionName)
-	{
-		static const TCHAR* FormatSample = TEXT(R"(
-				void {InstanceFunctionName} (in float3 RestPosition, out float3 NodePosition)
-				{
-					{HairStrandsContextName} DIHairStrands_AttachNodePosition(DIContext,RestPosition,NodePosition);
-				}
-				)");
-		OutHLSL += FString::Format(FormatSample, ArgsSample);
-		return true;
-	}
-	else if (FunctionInfo.DefinitionName == AttachNodeOrientationName)
-	{
-		static const TCHAR *FormatSample = TEXT(R"(
-				void {InstanceFunctionName} (in float4 RestOrientation, out float4 NodeOrientation)
-				{
-					{HairStrandsContextName} DIHairStrands_AttachNodeOrientation(DIContext,RestOrientation,NodeOrientation);
-				}
-				)");
-		OutHLSL += FString::Format(FormatSample, ArgsSample);
-		return true;
-	}
-	else if (FunctionInfo.DefinitionName == AttachNodeStateName)
-	{
-		static const TCHAR *FormatSample = TEXT(R"(
-				void {InstanceFunctionName} ( in float3 LocalPosition, in float4 LocalOrientation, out float3 NodePosition, out float4 NodeOrientation)
-				{
-					{HairStrandsContextName} DIHairStrands_AttachNodeState(DIContext,LocalPosition,LocalOrientation,NodePosition,NodeOrientation);
-				}
-				)");
-		OutHLSL += FString::Format(FormatSample, ArgsSample);
-		return true;
-	}
-	else if (FunctionInfo.DefinitionName == UpdateNodeStateName)
-	{
-		static const TCHAR* FormatSample = TEXT(R"(
-					void {InstanceFunctionName} ( in float3 RestPosition, in float3 NodePosition, in float4 NodeOrientation, out float3 OutNodePosition, out float4 OutNodeOrientation)
-					{
-						{HairStrandsContextName} DIHairStrands_UpdateNodeState(DIContext,RestPosition,NodePosition,NodeOrientation,OutNodePosition,OutNodeOrientation);
-					}
-					)");
-		OutHLSL += FString::Format(FormatSample, ArgsSample);
-		return true;
-	}
-	else if (FunctionInfo.DefinitionName == UpdatePointPositionName)
-	{
-		static const TCHAR *FormatSample = TEXT(R"(
-			void {InstanceFunctionName} (in float3 NodePosition, in float3 RestPosition, out bool OutReportStatus)
-			{
-				{HairStrandsContextName} DIHairStrands_UpdatePointPosition(DIContext,NodePosition, RestPosition ,OutReportStatus);
-			}
-			)");
-		OutHLSL += FString::Format(FormatSample, ArgsSample);
-		return true;
-	}
-	else if (FunctionInfo.DefinitionName == ResetPointPositionName)
-	{
-		static const TCHAR *FormatSample = TEXT(R"(
-			void {InstanceFunctionName} (out bool OutReportStatus)
-			{
-				{HairStrandsContextName} DIHairStrands_ResetPointPosition(DIContext,OutReportStatus);
-			}
-			)");
-		OutHLSL += FString::Format(FormatSample, ArgsSample);
-		return true;
-	}
-	else if (FunctionInfo.DefinitionName == AdvectNodePositionName)
-	{
-		static const TCHAR *FormatSample = TEXT(R"(
-			void {InstanceFunctionName} (in float NodeMass, in bool IsPositionMobile, in float3 ExternalForce, in float3 ForceGradient, in float DeltaTime,
-									     in float3 LinearVelocity, in float3 NodePosition, out float3 OutLinearVelocity, out float3 OutNodePosition)
-			{
-				OutLinearVelocity = LinearVelocity;
-				OutNodePosition = NodePosition;
-				{HairStrandsContextName} DIHairStrands_AdvectNodePosition(DIContext,NodeMass,IsPositionMobile,ExternalForce,ForceGradient,DeltaTime,OutLinearVelocity,OutNodePosition);
-			}
-			)");
-		OutHLSL += FString::Format(FormatSample, ArgsSample);
-		return true;
-	}
-	else if (FunctionInfo.DefinitionName == AdvectNodeOrientationName)
-	{
-		static const TCHAR *FormatSample = TEXT(R"(
-			void {InstanceFunctionName} (in float3 NodeInertia, in bool IsOrientationMobile, in float3 ExternalTorque, in float3 TorqueGradient, in float DeltaTime,
-										 in float3 AngularVelocity, in float4 NodeOrientation, out float3 OutAngularVelocity, out float4 OutNodeOrientation)
-			{
-				OutAngularVelocity = AngularVelocity;
-				OutNodeOrientation = NodeOrientation;
-				{HairStrandsContextName} DIHairStrands_AdvectNodeOrientation(DIContext,NodeInertia,IsOrientationMobile,ExternalTorque,TorqueGradient,DeltaTime,OutAngularVelocity,OutNodeOrientation);
-			}
-			)");
-		OutHLSL += FString::Format(FormatSample, ArgsSample);
-		return true;
-	}
-	else if (FunctionInfo.DefinitionName == UpdateLinearVelocityName)
-	{
-		static const TCHAR *FormatSample = TEXT(R"(
-			void {InstanceFunctionName} (in float3 PreviousPosition, in float3 NodePosition, in float DeltaTime, out float3 OutLinearVelocity)
-			{
-				{HairStrandsContextName} DIHairStrands_UpdateLinearVelocity(DIContext,PreviousPosition,NodePosition,DeltaTime,OutLinearVelocity);
-			}
-			)");
-		OutHLSL += FString::Format(FormatSample, ArgsSample);
-		return true;
-	}
-	else if (FunctionInfo.DefinitionName == UpdateAngularVelocityName)
-	{
-		static const TCHAR *FormatSample = TEXT(R"(
-			void {InstanceFunctionName} (in float4 PreviousOrientation, in float4 NodeOrientation, in float DeltaTime, out float3 OutAngularVelocity)
-			{
-				{HairStrandsContextName} DIHairStrands_UpdateAngularVelocity(DIContext,PreviousOrientation,NodeOrientation,DeltaTime,OutAngularVelocity);
-			}
-			)");
-		OutHLSL += FString::Format(FormatSample, ArgsSample);
-		return true;
-	}
-	else if (FunctionInfo.DefinitionName == GetBoundingBoxName)
-	{
-		static const TCHAR *FormatSample = TEXT(R"(
-				void {InstanceFunctionName} (in int BoxIndex, out float3 OutBoxCenter, out float3 OutBoxExtent)
-				{
-					{HairStrandsContextName} DIHairStrands_GetBoundingBox(DIContext,DIContext_BoundingBoxBuffer,BoxIndex,OutBoxCenter,OutBoxExtent);
-				}
-				)");
-		OutHLSL += FString::Format(FormatSample, ArgsSample);
-		return true;
-	}
-	else if (FunctionInfo.DefinitionName == ResetBoundingBoxName)
-	{
-		static const TCHAR *FormatSample = TEXT(R"(
-				void {InstanceFunctionName} (out bool FunctionStatus)
-				{
-					{HairStrandsContextName} DIHairStrands_ResetBoundingBox(DIContext,DIContext_BoundingBoxBuffer,FunctionStatus);
-				}
-				)");
-		OutHLSL += FString::Format(FormatSample, ArgsSample);
-		return true;
-	}
-	else if (FunctionInfo.DefinitionName == BuildBoundingBoxName)
-	{
-		static const TCHAR *FormatSample = TEXT(R"(
-					void {InstanceFunctionName} (in float3 NodePosition, out bool OutFunctionStatus)
-					{
-						{HairStrandsContextName} DIHairStrands_BuildBoundingBox(DIContext,DIContext_BoundingBoxBuffer,NodePosition,OutFunctionStatus);
-					}
-					)");
-		OutHLSL += FString::Format(FormatSample, ArgsSample);
-		return true;
-	}
-	else if (FunctionInfo.DefinitionName == SetupDistanceSpringMaterialName)
-	{
-		static const TCHAR *FormatSample = TEXT(R"(
-				void {InstanceFunctionName} (in float YoungModulus, in float RodThickness,
-in float RestLength, in float DeltaTime, in int NodeOffset, in float MaterialDamping, out float OutMaterialCompliance, out float OutMaterialWeight, out float OutMaterialMultiplier)
-				{
-					{HairStrandsContextName}
-					if(NodeOffset == 0)
-					{
-						SetupStretchSpringMaterial(DIContext.StrandSize,YoungModulus,RodThickness,RestLength,DeltaTime,false,MaterialDamping,OutMaterialCompliance,OutMaterialWeight,OutMaterialMultiplier);
-					}
-					else if( NodeOffset == 1)
-					{
-						SetupBendSpringMaterial(DIContext.StrandSize,YoungModulus,RodThickness,RestLength,DeltaTime,false,MaterialDamping,OutMaterialCompliance,OutMaterialWeight,OutMaterialMultiplier);
-					}
-					else if( NodeOffset == 2)
-					{
-						SetupTwistSpringMaterial(DIContext.StrandSize,YoungModulus,RodThickness,RestLength,DeltaTime,false,MaterialDamping,OutMaterialCompliance,OutMaterialWeight,OutMaterialMultiplier);
-					}
-				}
-				)");
-		OutHLSL += FString::Format(FormatSample, ArgsSample);
-		return true;
-	}
-	else if (FunctionInfo.DefinitionName == SolveDistanceSpringMaterialName)
-	{
-		static const TCHAR *FormatSample = TEXT(R"(
-				void {InstanceFunctionName} (in bool EnableConstraint, in float RestLength, in float DeltaTime, in int NodeOffset, in float MaterialDamping,
-		in float MaterialCompliance, in float MaterialWeight, in float MaterialMultiplier, out float OutMaterialMultiplier)
-				{
-					{HairStrandsContextName}
-					if(NodeOffset == 0)
-					{
-						SolveStretchSpringMaterial(EnableConstraint,DIContext.StrandSize,RestLength,DeltaTime,MaterialDamping,MaterialCompliance,MaterialWeight,MaterialMultiplier,OutMaterialMultiplier);
-					}
-					else if(NodeOffset == 1)
-					{
-						SolveBendSpringMaterial(EnableConstraint,DIContext.StrandSize,RestLength,DeltaTime,MaterialDamping,MaterialCompliance,MaterialWeight,MaterialMultiplier,OutMaterialMultiplier);
-					}
-					else if(NodeOffset == 2)
-					{
-						SolveTwistSpringMaterial(EnableConstraint,DIContext.StrandSize,RestLength,DeltaTime,MaterialDamping,MaterialCompliance,MaterialWeight,MaterialMultiplier,OutMaterialMultiplier);
-					}
-				}
-				)");
-		OutHLSL += FString::Format(FormatSample, ArgsSample);
-		return true;
-	}
-	else if (FunctionInfo.DefinitionName == ProjectDistanceSpringMaterialName)
-	{
-		static const TCHAR *FormatSample = TEXT(R"(
-				void {InstanceFunctionName} (in bool EnableConstraint, in float YoungModulus, in float RodThickness, in float RestLength, in float DeltaTime, in int NodeOffset, out float3 OutNodePosition)
-				{
-					{HairStrandsContextName}
-					if(NodeOffset == 0)
-					{
-						ProjectStretchSpringMaterial(EnableConstraint,DIContext.StrandSize,YoungModulus,RodThickness,RestLength,DeltaTime,OutNodePosition);
-					}
-					if(NodeOffset == 1)
-					{
-						ProjectBendSpringMaterial(EnableConstraint,DIContext.StrandSize,YoungModulus,RodThickness,RestLength,DeltaTime,OutNodePosition);
-					}
-					if(NodeOffset == 2)
-					{
-						ProjectTwistSpringMaterial(EnableConstraint,DIContext.StrandSize,YoungModulus,RodThickness,RestLength,DeltaTime,OutNodePosition);
-					}
-				}
-				)");
-		OutHLSL += FString::Format(FormatSample, ArgsSample);
-		return true;
-	}
-	else if (FunctionInfo.DefinitionName == SetupAngularSpringMaterialName)
-	{
-		static const TCHAR *FormatSample = TEXT(R"(
-					void {InstanceFunctionName} (in float YoungModulus, in float RodThickness,
-	in float RestLength, in float DeltaTime, in float MaterialDamping, out float OutMaterialCompliance, out float OutMaterialWeight, out float3 OutMaterialMultiplier)
-					{
-						{HairStrandsContextName} SetupAngularSpringMaterial(DIContext.StrandSize,YoungModulus,RodThickness,RestLength,DeltaTime,false,MaterialDamping,OutMaterialCompliance,OutMaterialWeight,OutMaterialMultiplier);
-					}
-					)");
-		OutHLSL += FString::Format(FormatSample, ArgsSample);
-		return true;
-	}
-	else if (FunctionInfo.DefinitionName == SolveAngularSpringMaterialName)
-	{
-		static const TCHAR *FormatSample = TEXT(R"(
-					void {InstanceFunctionName} (in bool EnableConstraint, in float RestLength, in float3 RestDirection, in float DeltaTime, in float MaterialDamping,
-			in float MaterialCompliance, in float MaterialWeight, in float3 MaterialMultiplier, out float3 OutMaterialMultiplier)
-					{
-						{HairStrandsContextName} SolveAngularSpringMaterial(EnableConstraint,DIContext.StrandSize,RestLength, RestDirection,DeltaTime,MaterialDamping,MaterialCompliance,MaterialWeight,MaterialMultiplier,OutMaterialMultiplier);
-					}
-					)");
-		OutHLSL += FString::Format(FormatSample, ArgsSample);
-		return true;
-	}
-	else if (FunctionInfo.DefinitionName == ProjectAngularSpringMaterialName)
-	{
-		static const TCHAR *FormatSample = TEXT(R"(
-					void {InstanceFunctionName} (in bool EnableConstraint, in float YoungModulus, in float RodThickness, in float RestLength, in float3 RestDirection, in float DeltaTime, out float3 OutNodePosition)
-					{
-						{HairStrandsContextName} ProjectAngularSpringMaterial(EnableConstraint,DIContext.StrandSize,YoungModulus,RodThickness,RestLength,RestDirection,DeltaTime,OutNodePosition);
-					}
-					)");
-		OutHLSL += FString::Format(FormatSample, ArgsSample);
-		return true;
-	}
-	else if (FunctionInfo.DefinitionName == SetupStretchRodMaterialName)
-	{
-		static const TCHAR *FormatSample = TEXT(R"(
-					void {InstanceFunctionName} (in float YoungModulus, in float RodThickness,
-	in float RestLength, in float DeltaTime, in float MaterialDamping, out float OutMaterialCompliance, out float OutMaterialWeight, out float3 OutMaterialMultiplier)
-					{
-						{HairStrandsContextName} SetupStretchRodMaterial(DIContext.StrandSize,YoungModulus,RodThickness,RestLength,DeltaTime,false,MaterialDamping,OutMaterialCompliance,OutMaterialWeight,OutMaterialMultiplier);
-					}
-					)");
-		OutHLSL += FString::Format(FormatSample, ArgsSample);
-		return true;
-	}
-	else if (FunctionInfo.DefinitionName == SolveStretchRodMaterialName)
-	{
-		static const TCHAR *FormatSample = TEXT(R"(
-					void {InstanceFunctionName} (in bool EnableConstraint, in float RestLength, in float DeltaTime, in float MaterialDamping,
-			in float MaterialCompliance, in float MaterialWeight, in float3 MaterialMultiplier, out float3 OutMaterialMultiplier)
-					{
-						{HairStrandsContextName} SolveStretchRodMaterial(EnableConstraint,DIContext.StrandSize,RestLength,DeltaTime,MaterialDamping,MaterialCompliance,MaterialWeight,MaterialMultiplier,OutMaterialMultiplier);
-					}
-					)");
-		OutHLSL += FString::Format(FormatSample, ArgsSample);
-		return true;
-	}
-	else if (FunctionInfo.DefinitionName == ProjectStretchRodMaterialName)
-	{
-		static const TCHAR *FormatSample = TEXT(R"(
-					void {InstanceFunctionName} (in bool EnableConstraint, in float YoungModulus, in float RodThickness, in float RestLength, in float DeltaTime, out float3 OutNodePosition)
-					{
-						{HairStrandsContextName} ProjectStretchRodMaterial(EnableConstraint,DIContext.StrandSize,YoungModulus,RodThickness,RestLength,DeltaTime,OutNodePosition);
-					}
-					)");
-		OutHLSL += FString::Format(FormatSample, ArgsSample);
-		return true;
-	}
-	else if (FunctionInfo.DefinitionName == SetupBendRodMaterialName)
-	{
-		static const TCHAR *FormatSample = TEXT(R"(
-						void {InstanceFunctionName} (in float YoungModulus, in float RodThickness,
-		in float RestLength, in float DeltaTime, in float MaterialDamping, out float OutMaterialCompliance, out float OutMaterialWeight, out float3 OutMaterialMultiplier)
-						{
-							{HairStrandsContextName} SetupBendRodMaterial(DIContext.StrandSize,YoungModulus,RodThickness,RestLength,DeltaTime,false,MaterialDamping,OutMaterialCompliance,OutMaterialWeight,OutMaterialMultiplier);
-						}
-						)");
-		OutHLSL += FString::Format(FormatSample, ArgsSample);
-		return true;
-	}
-	else if (FunctionInfo.DefinitionName == SolveBendRodMaterialName)
-	{
-		static const TCHAR *FormatSample = TEXT(R"(
-						void {InstanceFunctionName} (in bool EnableConstraint, in float RestLength, in float4 RestRotation, in float DeltaTime, in float MaterialDamping,
-				in float MaterialCompliance, in float MaterialWeight, in float3 MaterialMultiplier, out float3 OutMaterialMultiplier)
-						{
-							{HairStrandsContextName} SolveBendRodMaterial(EnableConstraint,DIContext.StrandSize,RestLength,RestRotation,DeltaTime,MaterialDamping,MaterialCompliance,MaterialWeight,MaterialMultiplier,OutMaterialMultiplier);
-						}
-						)");
-		OutHLSL += FString::Format(FormatSample, ArgsSample);
-		return true;
-	}
-	else if (FunctionInfo.DefinitionName == ProjectBendRodMaterialName)
-	{
-		static const TCHAR *FormatSample = TEXT(R"(
-						void {InstanceFunctionName} (in bool EnableConstraint, in float YoungModulus, in float RodThickness, in float RestLength, in float4 RestRotation, in float DeltaTime, out float4 OutNodeOrientation)
-						{
-							{HairStrandsContextName} ProjectBendRodMaterial(EnableConstraint,DIContext.StrandSize,YoungModulus,RodThickness,RestLength,RestRotation,DeltaTime,OutNodeOrientation);
-						}
-						)");
-		OutHLSL += FString::Format(FormatSample, ArgsSample);
-		return true;
-	}
-	else if (FunctionInfo.DefinitionName == SolveHardCollisionConstraintName)
-	{
-		static const TCHAR *FormatSample = TEXT(R"(
-							void {InstanceFunctionName} (in bool EnableConstraint, in float PenetrationDepth, in float3 CollisionPosition, in float3 CollisionVelocity, in float3 CollisionNormal,
-				in float StaticFriction, in float KineticFriction, in float DeltaTime, out float3 OutMaterialMultiplier )
-							{
-								OutMaterialMultiplier = float3(0,0,0);
-								{HairStrandsContextName} SolveHardCollisionConstraint(EnableConstraint,DIContext.StrandSize,PenetrationDepth,
-									CollisionPosition,CollisionVelocity,CollisionNormal,StaticFriction,KineticFriction,false,DeltaTime);
-							}
-							)");
-		OutHLSL += FString::Format(FormatSample, ArgsSample);
-		return true;
-	}
-	else if (FunctionInfo.DefinitionName == ProjectHardCollisionConstraintName)
-	{
-			static const TCHAR *FormatSample = TEXT(R"(
-						void {InstanceFunctionName} (in bool EnableConstraint, in float PenetrationDepth, in float3 CollisionPosition, in float3 CollisionVelocity, in float3 CollisionNormal,
-			in float StaticFriction, in float KineticFriction, in float DeltaTime, out float3 OutNodePosition )
-						{
-							{HairStrandsContextName} ProjectHardCollisionConstraint(EnableConstraint,DIContext.StrandSize,PenetrationDepth,
-								CollisionPosition,CollisionVelocity,CollisionNormal,StaticFriction,KineticFriction,DeltaTime,OutNodePosition);
-						}
-						)");
-		OutHLSL += FString::Format(FormatSample, ArgsSample);
-		return true;
-	}
-	else if (FunctionInfo.DefinitionName == SolveSoftCollisionConstraintName)
-	{
-		static const TCHAR *FormatSample = TEXT(R"(
-								void {InstanceFunctionName} (in bool EnableConstraint, in float PenetrationDepth, in float3 CollisionPosition, in float3 CollisionVelocity, in float3 CollisionNormal,
-					in float StaticFriction, in float KineticFriction, in float DeltaTime, in float MaterialDamping,
-			in float MaterialCompliance, in float MaterialWeight, in float3 MaterialMultiplier, out float3 OutMaterialMultiplier )
-								{
-									{HairStrandsContextName} SolveSoftCollisionConstraint(EnableConstraint,DIContext.StrandSize,PenetrationDepth,
-										CollisionPosition,CollisionVelocity,CollisionNormal,StaticFriction,KineticFriction,false,DeltaTime,MaterialDamping,
-											MaterialCompliance,MaterialWeight,MaterialMultiplier,OutMaterialMultiplier);
-								}
-								)");
-		OutHLSL += FString::Format(FormatSample, ArgsSample);
-		return true;
-	}
-	else if (FunctionInfo.DefinitionName == ProjectSoftCollisionConstraintName)
-	{
-		static const TCHAR *FormatSample = TEXT(R"(
-							void {InstanceFunctionName} (in bool EnableConstraint, in float ConstraintStiffness, in float PenetrationDepth, in float3 CollisionPosition, in float3 CollisionVelocity, in float3 CollisionNormal,
-					in float StaticFriction, in float KineticFriction, in float DeltaTime, out float3 OutNodePosition )
-							{
-								{HairStrandsContextName} ProjectSoftCollisionConstraint(EnableConstraint,DIContext.StrandSize,ConstraintStiffness,PenetrationDepth,
-									CollisionPosition,CollisionVelocity,CollisionNormal,StaticFriction,KineticFriction,DeltaTime,OutNodePosition);
-							}
-							)");
-		OutHLSL += FString::Format(FormatSample, ArgsSample);
-		return true;
-	}
-	else if (FunctionInfo.DefinitionName == SetupSoftCollisionConstraintName)
-	{
-		static const TCHAR *FormatSample = TEXT(R"(
-								void {InstanceFunctionName} (in float ConstraintStiffness, in float DeltaTime, in float MaterialDamping, out float OutMaterialCompliance, out float OutMaterialWeight, out float3 OutMaterialMultiplier )
-								{
-									{HairStrandsContextName} SetupSoftCollisionConstraint(DIContext.StrandSize,ConstraintStiffness,DeltaTime,MaterialDamping,OutMaterialCompliance,OutMaterialWeight,OutMaterialMultiplier);
-								}
-								)");
-		OutHLSL += FString::Format(FormatSample, ArgsSample);
-		return true;
-	}
-	else if (FunctionInfo.DefinitionName == ComputeEdgeDirectionName)
-	{
-		static const TCHAR *FormatSample = TEXT(R"(
-							void {InstanceFunctionName} (in float3 NodePosition, in float4 NodeOrientation, 
-	in float NodeMass, in float3 GravityVector, in float GravityPreloading, in float BendStiffness, in float StrandThickness, in float RestLength, out float3 OutRestDirection)
-							{
-								{HairStrandsContextName} DIHairStrands_ComputeEdgeDirection(DIContext,NodePosition,NodeOrientation,NodeMass,GravityVector,GravityPreloading,BendStiffness,StrandThickness,RestLength,OutRestDirection);
-							}
-							)");
-		OutHLSL += FString::Format(FormatSample, ArgsSample);
-		return true;
-	}
-	else if (FunctionInfo.DefinitionName == UpdateMaterialFrameName)
-	{
-		static const TCHAR *FormatSample = TEXT(R"(
-					void {InstanceFunctionName} (out float4 OutNodeOrientation)
-					{
-						{HairStrandsContextName} UpdateMaterialFrame(DIContext.StrandSize);
-						OutNodeOrientation = SharedNodeOrientation[GGroupThreadId.x];
-					}
-					)");
-		OutHLSL += FString::Format(FormatSample, ArgsSample);
-		return true;
-	}
-	else if (FunctionInfo.DefinitionName == ComputeMaterialFrameName)
-	{
-		static const TCHAR *FormatSample = TEXT(R"(
-						void {InstanceFunctionName} ( out float4 OutNodeOrientation)
-						{
-							{HairStrandsContextName} ComputeMaterialFrame(DIContext.StrandSize);
-							OutNodeOrientation = SharedNodeOrientation[GGroupThreadId.x];
-						}
-						)");
-		OutHLSL += FString::Format(FormatSample, ArgsSample);
-		return true;
-	}
-	else if (FunctionInfo.DefinitionName == ComputeAirDragForceName)
-	{
-		static const TCHAR *FormatSample = TEXT(R"(
-						void {InstanceFunctionName} (  in float AirDensity, in float AirViscosity, in float AirDrag,
-		in float3 AirVelocity, in float NodeThickness, in float3 NodePosition, in float3 NodeVelocity, out float3 OutAirDrag, out float3 OutDragGradient )
-						{
-							{HairStrandsContextName} ComputeAirDragForce(DIContext.StrandSize,AirDensity,AirViscosity,AirDrag,AirVelocity,NodeThickness,NodePosition,NodeVelocity,OutAirDrag,OutDragGradient);
-						}
-						)");
-		OutHLSL += FString::Format(FormatSample, ArgsSample);
-		return true;
-	}
-	else if (FunctionInfo.DefinitionName == InitGridSamplesName)
-	{
-		static const TCHAR* FormatSample = TEXT(R"(
-							void {InstanceFunctionName} ( in float3 NodePosition, in float3 NodeVelocity,
-	in float NodeMass, in float GridLength, out int OutNumSamples,
-						out float3 OutDeltaPosition, out float3 OutDeltaVelocity, out float OutSampleMass)
-							{
-								{HairStrandsContextName} DIHairStrands_InitGridSamples(DIContext,
-										 NodePosition, NodeVelocity, NodeMass, GridLength,
-											OutNumSamples, OutDeltaPosition, OutDeltaVelocity, OutSampleMass);
-							}
-							)");
-		OutHLSL += FString::Format(FormatSample, ArgsSample);
-		return true;
-	}
-	else if (FunctionInfo.DefinitionName == GetSampleStateName)
-	{
-		static const TCHAR* FormatSample = TEXT(R"(
-								void {InstanceFunctionName} ( in float3 NodePosition, in float3 NodeVelocity, in float3 DeltaPosition, in float3 DeltaVelocity,
-			in int NumSamples, in int SampleIndex, out float3 OutSamplePosition, out float3 OutSampleVelocity)
-								{
-									{HairStrandsContextName} DIHairStrands_GetSampleState(DIContext,
-											 NodePosition, NodeVelocity, DeltaPosition, DeltaVelocity,
-											 NumSamples, SampleIndex, OutSamplePosition, OutSampleVelocity);
-								}
-								)");
-		OutHLSL += FString::Format(FormatSample, ArgsSample);
-		return true;
-	}
-	else if (FunctionInfo.DefinitionName == NeedSimulationResetName)
-	{
-		static const TCHAR *FormatSample = TEXT(R"(
-				void {InstanceFunctionName} ( out bool ResetSimulation)
-				{
-					{HairStrandsContextName} ResetSimulation  = DIContext.ResetSimulation;
-				}
-				)");
-		OutHLSL += FString::Format(FormatSample, ArgsSample);
-		return true;
-	}
-	else if (FunctionInfo.DefinitionName == HasGlobalInterpolationName)
-	{
-		static const TCHAR* FormatSample = TEXT(R"(
-				void {InstanceFunctionName} ( out bool GlobalInterpolation)
-				{
-					{HairStrandsContextName} GlobalInterpolation  = (DIContext.InterpolationMode == 2);
-				}
-				)");
-		OutHLSL += FString::Format(FormatSample, ArgsSample);
-		return true;
-	}
-	else if (FunctionInfo.DefinitionName == NeedRestUpdateName)
-	{
-		static const TCHAR* FormatSample = TEXT(R"(
-				void {InstanceFunctionName} ( out bool RestUpdate)
-				{
-					{HairStrandsContextName} RestUpdate  = DIContext.RestUpdate;
-				}
-				)");
-		OutHLSL += FString::Format(FormatSample, ArgsSample);
-		return true;
-	}
-
-	OutHLSL += TEXT("\n");
-	return false;
+	return ValidGpuFunctions.Contains(FunctionInfo.DefinitionName);
 }
 
 void UNiagaraDataInterfaceHairStrands::GetCommonHLSL(FString& OutHLSL)
 {
-	OutHLSL += TEXT("#include \"/Plugin/FX/Niagara/Private/NiagaraQuaternionUtils.ush\"\n");
-	//OutHLSL += TEXT("#include \"/Plugin/Runtime/HairStrands/Private/NiagaraDirectSolver.ush\"\n");
-	OutHLSL += TEXT("#include \"/Plugin/Runtime/HairStrands/Private/NiagaraStrandsExternalForce.ush\"\n");
-	OutHLSL += TEXT("#include \"/Plugin/Runtime/HairStrands/Private/NiagaraHookeSpringMaterial.ush\"\n");
-	OutHLSL += TEXT("#include \"/Plugin/Runtime/HairStrands/Private/NiagaraAngularSpringMaterial.ush\"\n");
-	OutHLSL += TEXT("#include \"/Plugin/Runtime/HairStrands/Private/NiagaraConstantVolumeMaterial.ush\"\n");
-	OutHLSL += TEXT("#include \"/Plugin/Runtime/HairStrands/Private/NiagaraCosseratRodMaterial.ush\"\n");
-	OutHLSL += TEXT("#include \"/Plugin/Runtime/HairStrands/Private/NiagaraStaticCollisionConstraint.ush\"\n");
-	OutHLSL += TEXT("#include \"/Plugin/Runtime/HairStrands/Private/NiagaraDataInterfaceHairStrands.ush\"\n");
+	for (const TCHAR* CommonFile : NDIHairStrandsLocal::CommonShaderFiles)
+	{
+		OutHLSL.Appendf(TEXT("#include \"%s\"\n"), CommonFile);
+	}
 }
 
 bool UNiagaraDataInterfaceHairStrands::AppendCompileHash(FNiagaraCompileHashVisitor* InVisitor) const
@@ -4420,31 +3312,223 @@ bool UNiagaraDataInterfaceHairStrands::AppendCompileHash(FNiagaraCompileHashVisi
 	if (!Super::AppendCompileHash(InVisitor))
 		return false;
 
-	for (const TCHAR* VirtualFilePath :
-		{
-			TEXT("/Plugin/FX/Niagara/Private/NiagaraQuaternionUtils.ush"),
-			TEXT("/Plugin/Runtime/HairStrands/Private/NiagaraStrandsExternalForce.ush"),
-			TEXT("/Plugin/Runtime/HairStrands/Private/NiagaraHookeSpringMaterial.ush"),
-			TEXT("/Plugin/Runtime/HairStrands/Private/NiagaraAngularSpringMaterial.ush"),
-			TEXT("/Plugin/Runtime/HairStrands/Private/NiagaraConstantVolumeMaterial.ush"),
-			TEXT("/Plugin/Runtime/HairStrands/Private/NiagaraCosseratRodMaterial.ush"),
-			TEXT("/Plugin/Runtime/HairStrands/Private/NiagaraStaticCollisionConstraint.ush"),
-			TEXT("/Plugin/Runtime/HairStrands/Private/NiagaraDataInterfaceHairStrands.ush")
-		})
+	for (const TCHAR* CommonFile : NDIHairStrandsLocal::CommonShaderFiles)
 	{
-		FSHAHash Hash = GetShaderFileHash(VirtualFilePath, EShaderPlatform::SP_PCD3D_SM5);
-		InVisitor->UpdateString(TEXT("NiagaraDataInterfaceHairStrandsHLSLSource"), Hash.ToString());
+		InVisitor->UpdateString(TEXT("NiagaraDataInterfaceHairStrandsHLSLSource"), GetShaderFileHash(CommonFile, EShaderPlatform::SP_PCD3D_SM5).ToString());
 	}
+	InVisitor->UpdateString(TEXT("NiagaraDataInterfaceHairStrandsTemplateHLSLSource"), GetShaderFileHash(NDIHairStrandsLocal::TemplateShaderFile, EShaderPlatform::SP_PCD3D_SM5).ToString());
 
 	return true;
 }
 
 void UNiagaraDataInterfaceHairStrands::GetParameterDefinitionHLSL(const FNiagaraDataInterfaceGPUParamInfo& ParamInfo, FString& OutHLSL)
 {
-	OutHLSL += TEXT("DIHAIRSTRANDS_DECLARE_CONSTANTS(") + ParamInfo.DataInterfaceHLSLSymbol + TEXT(")\n");
+	Super::GetParameterDefinitionHLSL(ParamInfo, OutHLSL);
+
+	TMap<FString, FStringFormatArg> TemplateArgs =
+	{
+		{TEXT("ParameterName"),	ParamInfo.DataInterfaceHLSLSymbol},
+	};
+
+	FString TemplateFile;
+	LoadShaderSourceFile(NDIHairStrandsLocal::TemplateShaderFile, EShaderPlatform::SP_PCD3D_SM5, &TemplateFile, nullptr);
+	OutHLSL += FString::Format(*TemplateFile, TemplateArgs);
 }
 #endif
 
+void UNiagaraDataInterfaceHairStrands::BuildShaderParameters(FNiagaraShaderParametersBuilder& ShaderParametersBuilder) const
+{
+	ShaderParametersBuilder.AddNestedStruct<NDIHairStrandsLocal::FShaderParameters>();
+}
+
+void UNiagaraDataInterfaceHairStrands::SetShaderParameters(const FNiagaraDataInterfaceSetShaderParametersContext& Context) const
+{
+	FRDGBuilder& GraphBuilder = Context.GetGraphBuilder();
+	FNDIHairStrandsProxy& DIProxy = Context.GetProxy<FNDIHairStrandsProxy>();
+	FNDIHairStrandsData* ProxyData = DIProxy.SystemInstancesToProxyData.Find(Context.GetSystemInstanceID());
+
+	const bool bIsHairValid = ProxyData != nullptr && ProxyData->HairStrandsBuffer && ProxyData->HairStrandsBuffer->IsInitialized();
+	const bool bIsHairGroupInstValid = ProxyData != nullptr && ProxyData->HairGroupInstSource != nullptr && ProxyData->HairGroupInstSource->ContainsGroupInstance(ProxyData->HairGroupInstance);
+	const bool bHasSkinningBinding = bIsHairValid && bIsHairGroupInstValid && ProxyData->HairGroupInstance->BindingType == EHairBindingType::Skinning;
+	const bool bIsRootValid = bIsHairValid && ProxyData->HairStrandsBuffer->SourceDeformedRootResources && ProxyData->HairStrandsBuffer->SourceDeformedRootResources->IsInitialized() && bHasSkinningBinding;
+	const bool bIsRestValid = bIsHairValid && ProxyData->HairStrandsBuffer->SourceRestResources && ProxyData->HairStrandsBuffer->SourceRestResources->IsInitialized()&&
+		// SourceRestResources->PositionBuffer is lazily allocated, when the instance LOD is scheduled (this happens after this called). So this is why this check is here. 
+		// This code should be refactor so that it reflects the lazy allocation scheme
+		ProxyData->HairStrandsBuffer->SourceRestResources->PositionBuffer.SRV &&
+		// TEMP: These check are only temporary for avoiding crashes while we find the bottom of the issue.
+		ProxyData->HairStrandsBuffer->CurvesOffsetsBuffer.SRV && ProxyData->HairStrandsBuffer->ParamsScaleBuffer.SRV && ProxyData->HairStrandsBuffer->BoundingBoxBuffer.IsValid();
+
+	const bool bIsGeometryValid = bIsHairValid && (!bIsHairGroupInstValid || (bIsHairGroupInstValid && (ProxyData->HairGroupInstance->GeometryType != EHairGeometryType::NoneGeometry)));
+	const bool bIsDeformedValid = bIsHairValid && ProxyData->HairStrandsBuffer->SourceDeformedResources && ProxyData->HairStrandsBuffer->SourceDeformedResources->IsInitialized();
+
+	NDIHairStrandsLocal::FShaderParameters* ShaderParameters = Context.GetParameterNestedStruct<NDIHairStrandsLocal::FShaderParameters>();
+	if (bIsHairValid && bIsRestValid && bIsGeometryValid && bIsHairGroupInstValid)
+	{
+		check(ProxyData);
+
+		FNDIHairStrandsBuffer* HairStrandsBuffer = ProxyData->HairStrandsBuffer;
+
+		FRDGBufferUAVRef DeformedPositionBufferUAV = nullptr;
+		FRHIShaderResourceView* DeformedPositionOffsetSRV = nullptr;
+		if (bIsDeformedValid)
+		{
+			FRDGPooledBuffer* DeformedPositionBuffer = HairStrandsBuffer->SourceDeformedResources->DeformedPositionBuffer[HairStrandsBuffer->SourceDeformedResources->CurrentIndex].Buffer;
+			DeformedPositionBufferUAV = GraphBuilder.CreateUAV(GraphBuilder.RegisterExternalBuffer(DeformedPositionBuffer), FHairStrandsPositionFormat::Format);	//-OPT: If constant across all dispatches can create in PreStage once
+			DeformedPositionOffsetSRV = HairStrandsBuffer->SourceDeformedResources->DeformedOffsetBuffer[HairStrandsBuffer->SourceDeformedResources->CurrentIndex].SRV.GetReference();
+		}
+		else
+		{
+			DeformedPositionBufferUAV = GraphBuilder.CreateUAV(GraphBuilder.RegisterExternalBuffer(HairStrandsBuffer->DeformedPositionBuffer), FHairStrandsPositionFormat::Format);
+			DeformedPositionOffsetSRV = FNiagaraRenderer::GetDummyFloatBuffer();
+		}
+		const int32 MeshLODIndex = bIsRootValid ? HairStrandsBuffer->SourceDeformedRootResources->MeshLODIndex : -1;
+
+		// Projection Buffers
+		const bool bHasSkinnedInterpolation = (bIsRootValid && HairStrandsBuffer->SourceDeformedRootResources->IsValid(MeshLODIndex));
+		const EHairSimulationInterpolationMode InterpolationModeValue = bHasSkinnedInterpolation ? (ProxyData->GlobalInterpolation ?
+			EHairSimulationInterpolationMode::RBF : EHairSimulationInterpolationMode::Skinned) : EHairSimulationInterpolationMode::Rigid;
+
+		const FHairStrandsRestRootResource::FLOD* RestMeshProjection = bHasSkinnedInterpolation ? &(HairStrandsBuffer->SourceRestRootResources->LODs[MeshLODIndex]) : nullptr;
+		const FHairStrandsDeformedRootResource::FLOD* DeformedMeshProjection = bHasSkinnedInterpolation ? &(HairStrandsBuffer->SourceDeformedRootResources->LODs[MeshLODIndex]) : nullptr;
+
+		FRHIShaderResourceView* RestTrianglePositionASRV = (bHasSkinnedInterpolation && RestMeshProjection) ? RestMeshProjection->RestRootTrianglePosition0Buffer.SRV.GetReference() : FNiagaraRenderer::GetDummyFloatBuffer();
+		FRHIShaderResourceView* RestTrianglePositionBSRV = (bHasSkinnedInterpolation && RestMeshProjection) ? RestMeshProjection->RestRootTrianglePosition1Buffer.SRV.GetReference() : FNiagaraRenderer::GetDummyFloatBuffer();
+		FRHIShaderResourceView* RestTrianglePositionCSRV = (bHasSkinnedInterpolation && RestMeshProjection) ? RestMeshProjection->RestRootTrianglePosition2Buffer.SRV.GetReference() : FNiagaraRenderer::GetDummyFloatBuffer();
+
+		FRHIShaderResourceView* DeformedTrianglePositionASRV = (bHasSkinnedInterpolation && DeformedMeshProjection) ? DeformedMeshProjection->GetDeformedRootTrianglePosition0Buffer(FHairStrandsDeformedRootResource::FLOD::Current).SRV.GetReference() : FNiagaraRenderer::GetDummyFloatBuffer();
+		FRHIShaderResourceView* DeformedTrianglePositionBSRV = (bHasSkinnedInterpolation && DeformedMeshProjection) ? DeformedMeshProjection->GetDeformedRootTrianglePosition1Buffer(FHairStrandsDeformedRootResource::FLOD::Current).SRV.GetReference() : FNiagaraRenderer::GetDummyFloatBuffer();
+		FRHIShaderResourceView* DeformedTrianglePositionCSRV = (bHasSkinnedInterpolation && DeformedMeshProjection) ? DeformedMeshProjection->GetDeformedRootTrianglePosition2Buffer(FHairStrandsDeformedRootResource::FLOD::Current).SRV.GetReference() : FNiagaraRenderer::GetDummyFloatBuffer();
+		FRHIShaderResourceView* RootBarycentricCoordinatesSRV = (bHasSkinnedInterpolation && RestMeshProjection) ? RestMeshProjection->RootTriangleBarycentricBuffer.SRV.GetReference() : FNiagaraRenderer::GetDummyUIntBuffer();
+
+		// RBF buffers
+		const bool bHasSamples = (RestMeshProjection && RestMeshProjection->SampleCount > 0);
+		const int32 SampleCountValue = bHasSamples ? RestMeshProjection->SampleCount : 0;
+
+		FShaderResourceViewRHIRef RestSamplePositionsBufferSRV = (bHasSamples && RestMeshProjection) ? RestMeshProjection->RestSamplePositionsBuffer.SRV.GetReference() : FNiagaraRenderer::GetDummyFloatBuffer();
+		FShaderResourceViewRHIRef MeshSampleWeightsBufferSRV = (bHasSamples && DeformedMeshProjection) ? DeformedMeshProjection->GetMeshSampleWeightsBuffer((FHairStrandsDeformedRootResource::FLOD::Current)).SRV.GetReference() : FNiagaraRenderer::GetDummyFloatBuffer();
+
+		// Simulation setup (we update the rest configuration based on the deformed positions 
+		// if in restupdate mode or if we are resetting the sim and using RBF transfer since the rest positions are not matrching the physics asset)
+		const int32 NeedResetValue = (ProxyData->TickCount <= GHairSimulationMaxDelay) || !HairStrandsBuffer->bValidGeometryType;
+		const int32 RestUpdateValue = GHairSimulationRestUpdate || (NeedResetValue && ProxyData->bSkinningTransfer);
+		const int32 LocalSimulationValue = ProxyData->LocalSimulation;
+
+		HairStrandsBuffer->bValidGeometryType = true;
+
+		// Offsets / Transforms
+		FVector3f RestPositionOffsetValue = (FVector3f)ProxyData->HairStrandsBuffer->SourceRestResources->GetPositionOffset();
+
+		FMatrix44f RigidTransformFloat = FMatrix44f(ProxyData->HairGroupInstance ? ProxyData->HairGroupInstance->Debug.RigidCurrentLocalToWorld.ToMatrixWithScale() : ProxyData->WorldTransform.ToMatrixWithScale());
+		FMatrix44f WorldTransformFloat = FMatrix44f(ProxyData->HairGroupInstance ? ProxyData->HairGroupInstance->GetCurrentLocalToWorld().ToMatrixWithScale() : ProxyData->WorldTransform.ToMatrixWithScale());
+		FMatrix44f BoneTransformFloat = FMatrix44f(ProxyData->BoneTransform.ToMatrixWithScale()) * RigidTransformFloat;
+
+		if (ProxyData->LocalSimulation)
+		{
+			const FMatrix44d WorldTransformDouble(WorldTransformFloat);
+			const FMatrix44d BoneTransformDouble(BoneTransformFloat);
+
+			// Due to large world coordinate we store the relative world transform in double precision 
+			WorldTransformFloat = FMatrix44f(WorldTransformDouble * BoneTransformDouble.Inverse());
+		}
+
+		if (!bIsRootValid && bHasSkinningBinding)
+		{
+			UE_LOG(LogHairStrands, Log, TEXT("FNDIHairStrandsParametersCS() Groom Asset %s from component %s is set to use skinning interpolation but the skin resources are not valid"),
+				*ProxyData->HairGroupInstance->Debug.GroomAssetName, *ProxyData->HairGroupInstance->Debug.MeshComponentName);
+		}
+
+		// Set shader constants
+		ShaderParameters->BoundingBoxOffsets = HairStrandsBuffer->BoundingBoxOffsets;
+		ShaderParameters->WorldTransform = WorldTransformFloat;
+		ShaderParameters->WorldInverse = WorldTransformFloat.Inverse();
+		ShaderParameters->WorldRotation = WorldTransformFloat.GetMatrixWithoutScale().ToQuat();
+		ShaderParameters->NumStrands = ProxyData->NumStrands;
+		ShaderParameters->StrandSize = ProxyData->StrandsSize;
+		ShaderParameters->BoneTransform = BoneTransformFloat;
+		ShaderParameters->BoneInverse = BoneTransformFloat.Inverse();
+		ShaderParameters->BoneRotation = BoneTransformFloat.GetMatrixWithoutScale().ToQuat();
+		ShaderParameters->BoneLinearVelocity = ProxyData->BoneLinearVelocity;
+		ShaderParameters->BoneAngularVelocity = ProxyData->BoneAngularVelocity;
+		ShaderParameters->BoneLinearAcceleration = ProxyData->BoneLinearAcceleration;
+		ShaderParameters->BoneAngularAcceleration = ProxyData->BoneAngularAcceleration;
+		ShaderParameters->ResetSimulation = NeedResetValue;
+		ShaderParameters->InterpolationMode = int32(InterpolationModeValue);
+		ShaderParameters->RestUpdate = RestUpdateValue;
+		ShaderParameters->LocalSimulation = LocalSimulationValue;
+		ShaderParameters->RestRootOffset = FVector3f::ZeroVector;
+		ShaderParameters->DeformedRootOffset = FVector3f::ZeroVector;
+		ShaderParameters->RestPositionOffset = RestPositionOffsetValue;
+		ShaderParameters->SampleCount = SampleCountValue;
+
+		// Set Shader UAV
+		ShaderParameters->DeformedPositionBuffer = DeformedPositionBufferUAV;
+		ShaderParameters->BoundingBoxBuffer = HairStrandsBuffer->BoundingBoxBuffer.GetOrCreateUAV(GraphBuilder);
+
+		// Set Shader SRV
+		ShaderParameters->CurvesOffsetsBuffer = HairStrandsBuffer->CurvesOffsetsBuffer.SRV;
+		ShaderParameters->ParamsScaleBuffer = HairStrandsBuffer->ParamsScaleBuffer.SRV;
+		ShaderParameters->RestPositionBuffer = GraphBuilder.CreateSRV(GraphBuilder.RegisterExternalBuffer(HairStrandsBuffer->SourceRestResources->PositionBuffer.Buffer), FHairStrandsPositionFormat::Format); //-OPT: If constant across all dispatches can create in PreStage once
+		ShaderParameters->DeformedPositionOffset = DeformedPositionOffsetSRV;
+		ShaderParameters->RestTrianglePositionABuffer = RestTrianglePositionASRV;
+		ShaderParameters->RestTrianglePositionBBuffer = RestTrianglePositionBSRV;
+		ShaderParameters->RestTrianglePositionCBuffer = RestTrianglePositionCSRV;
+		ShaderParameters->DeformedTrianglePositionABuffer = DeformedTrianglePositionASRV;
+		ShaderParameters->DeformedTrianglePositionBBuffer = DeformedTrianglePositionBSRV;
+		ShaderParameters->DeformedTrianglePositionCBuffer = DeformedTrianglePositionCSRV;
+		ShaderParameters->RestSamplePositionsBuffer = RestSamplePositionsBufferSRV;
+		ShaderParameters->MeshSampleWeightsBuffer = MeshSampleWeightsBufferSRV;
+		ShaderParameters->RootBarycentricCoordinatesBuffer = RootBarycentricCoordinatesSRV;
+	}
+	else
+	{
+		if (bIsHairValid)
+		{
+			ProxyData->HairStrandsBuffer->bValidGeometryType = false;
+		}
+		// Set shader constants
+		ShaderParameters->BoundingBoxOffsets = FIntVector4(0, 1, 2, 3);
+		ShaderParameters->WorldTransform = FMatrix44f::Identity;
+		ShaderParameters->WorldInverse = FMatrix44f::Identity;
+		ShaderParameters->WorldRotation = FQuat4f::Identity;
+		ShaderParameters->NumStrands = 1;
+		ShaderParameters->StrandSize = 1;
+		ShaderParameters->BoneTransform = FMatrix44f::Identity;
+		ShaderParameters->BoneInverse = FMatrix44f::Identity;
+		ShaderParameters->BoneRotation = FQuat4f::Identity;
+		ShaderParameters->BoneLinearVelocity = FVector3f::ZeroVector;
+		ShaderParameters->BoneAngularVelocity = FVector3f::ZeroVector;
+		ShaderParameters->BoneLinearAcceleration = FVector3f::ZeroVector;
+		ShaderParameters->BoneAngularAcceleration = FVector3f::ZeroVector;
+		ShaderParameters->ResetSimulation = 0;
+		ShaderParameters->InterpolationMode = 0;
+		ShaderParameters->RestUpdate = 0;
+		ShaderParameters->LocalSimulation = 0;
+		ShaderParameters->RestRootOffset = FVector3f::ZeroVector;
+		ShaderParameters->DeformedRootOffset = FVector3f::ZeroVector;
+		ShaderParameters->RestPositionOffset = FVector3f::ZeroVector;
+		ShaderParameters->SampleCount = 0;
+
+		// Set Shader UAV
+		ShaderParameters->DeformedPositionBuffer = Context.GetComputeDispatchInterface().GetEmptyBufferUAV(GraphBuilder, FHairStrandsPositionFormat::Format);
+		ShaderParameters->BoundingBoxBuffer = Context.GetComputeDispatchInterface().GetEmptyBufferUAV(GraphBuilder, PF_R32_UINT);
+
+		// Set Shader SRV
+		ShaderParameters->CurvesOffsetsBuffer = FNiagaraRenderer::GetDummyUIntBuffer();
+		ShaderParameters->RestPositionBuffer = Context.GetComputeDispatchInterface().GetEmptyBufferSRV(GraphBuilder, FHairStrandsPositionFormat::Format);
+		ShaderParameters->DeformedPositionOffset = FNiagaraRenderer::GetDummyFloatBuffer();
+		ShaderParameters->RestTrianglePositionABuffer = FNiagaraRenderer::GetDummyFloatBuffer();
+		ShaderParameters->RestTrianglePositionBBuffer = FNiagaraRenderer::GetDummyFloatBuffer();
+		ShaderParameters->RestTrianglePositionCBuffer = FNiagaraRenderer::GetDummyFloatBuffer();
+		ShaderParameters->DeformedTrianglePositionABuffer = FNiagaraRenderer::GetDummyFloatBuffer();
+		ShaderParameters->DeformedTrianglePositionBBuffer = FNiagaraRenderer::GetDummyFloatBuffer();
+		ShaderParameters->DeformedTrianglePositionCBuffer = FNiagaraRenderer::GetDummyFloatBuffer();
+		ShaderParameters->RestSamplePositionsBuffer = FNiagaraRenderer::GetDummyFloatBuffer();
+		ShaderParameters->MeshSampleWeightsBuffer = FNiagaraRenderer::GetDummyFloatBuffer();
+		ShaderParameters->RootBarycentricCoordinatesBuffer = FNiagaraRenderer::GetDummyFloatBuffer();
+		ShaderParameters->ParamsScaleBuffer = FNiagaraRenderer::GetDummyFloatBuffer();
+	}
+}
+
 void UNiagaraDataInterfaceHairStrands::ProvidePerInstanceDataForRenderThread(void* DataForRenderThread, void* PerInstanceData, const FNiagaraSystemInstanceID& SystemInstance)
 {
 	check(Proxy);
@@ -4458,12 +3542,11 @@ void UNiagaraDataInterfaceHairStrands::ProvidePerInstanceDataForRenderThread(voi
 	}
 }
 
-void FNDIHairStrandsProxy::PreStage(FRHICommandList& RHICmdList, const FNiagaraDataInterfaceStageArgs& Context)
+void FNDIHairStrandsProxy::PreStage(const FNDIGpuComputePreStageContext& Context)
 {
-	if (Context.SimStageData->bFirstStage)
+	if (Context.GetSimStageData().bFirstStage)
 	{
-		FNDIHairStrandsData* ProxyData =
-			SystemInstancesToProxyData.Find(Context.SystemInstanceID);
+		FNDIHairStrandsData* ProxyData = SystemInstancesToProxyData.Find(Context.GetSystemInstanceID());
 
 		if (ProxyData != nullptr && ProxyData->HairStrandsBuffer != nullptr)
 		{
@@ -4480,21 +3563,29 @@ void FNDIHairStrandsProxy::PreStage(FRHICommandList& RHICmdList, const FNiagaraD
 	}
 }
 
-void FNDIHairStrandsProxy::PostSimulate(FRHICommandList& RHICmdList, const FNiagaraDataInterfaceArgs& Context)
+void FNDIHairStrandsProxy::PostSimulate(const FNDIGpuComputePostSimulateContext& Context)
 {
-	FNDIHairStrandsProxy* InterfaceProxy = static_cast<FNDIHairStrandsProxy*>(Context.DataInterface);
-	FNDIHairStrandsData* ProxyData = InterfaceProxy->SystemInstancesToProxyData.Find(Context.SystemInstanceID);
+	FNDIHairStrandsData* ProxyData = SystemInstancesToProxyData.Find(Context.GetSystemInstanceID());
 
 	const bool bIsHairValid = ProxyData != nullptr && ProxyData->HairStrandsBuffer && ProxyData->HairStrandsBuffer->IsInitialized();
 	const bool bIsDeformedValid = bIsHairValid && ProxyData->HairStrandsBuffer->SourceDeformedResources && ProxyData->HairStrandsBuffer->SourceDeformedResources->IsInitialized();
 
 	// MGPU DeformedPositionBuffer copy after simulation
-	if (bIsDeformedValid)
-	{
-		const FNDIHairStrandsBuffer* HairStrandsBuffer = ProxyData->HairStrandsBuffer;
-		FRHIBuffer* DeformedPositionBuffer = HairStrandsBuffer->SourceDeformedResources->DeformedPositionBuffer[HairStrandsBuffer->SourceDeformedResources->CurrentIndex].Buffer->GetRHI();
+	//-TODO:RDG:mGPU
+	//if (bIsDeformedValid)
+	//{
+	//	const FNDIHairStrandsBuffer* HairStrandsBuffer = ProxyData->HairStrandsBuffer;
+	//	FRHIBuffer* DeformedPositionBuffer = HairStrandsBuffer->SourceDeformedResources->DeformedPositionBuffer[HairStrandsBuffer->SourceDeformedResources->CurrentIndex].Buffer->GetRHI();
+
+	//	Context.ComputeDispatchInterface->MultiGPUResourceModified(RHICmdList, DeformedPositionBuffer, false, true);
+	//}
 
-		Context.ComputeDispatchInterface->MultiGPUResourceModified(RHICmdList, DeformedPositionBuffer, false, true);
+	if (Context.IsFinalPostSimulate())
+	{
+		if (ProxyData->HairStrandsBuffer->IsInitialized())
+		{
+			ProxyData->HairStrandsBuffer->BoundingBoxBuffer.EndGraphUsage();
+		}
 	}
 }
 
