- Hair conversion to RDG

diff --git a/Engine/Plugins/Runtime/HairStrands/Source/HairStrandsCore/Private/Niagara/NiagaraDataInterfacePressureGrid.cpp b/Engine/Plugins/Runtime/HairStrands/Source/HairStrandsCore/Private/Niagara/NiagaraDataInterfacePressureGrid.cpp
--- a/Engine/Plugins/Runtime/HairStrands/Source/HairStrandsCore/Private/Niagara/NiagaraDataInterfacePressureGrid.cpp
+++ b/Engine/Plugins/Runtime/HairStrands/Source/HairStrandsCore/Private/Niagara/NiagaraDataInterfacePressureGrid.cpp
@@ -21,19 +21,22 @@ DEFINE_LOG_CATEGORY_STATIC(LogPressureGrid, Log, All);
 
 //------------------------------------------------------------------------------------------------------------
 
-static const FName BuildDistanceFieldName(TEXT("BuildDistanceField"));
-static const FName BuildDensityFieldName(TEXT("BuildDensityField"));
-static const FName SolveGridPressureName(TEXT("SolveGridPressure"));
-static const FName ScaleCellFieldsName(TEXT("ScaleCellFields"));
-static const FName SetSolidBoundaryName(TEXT("SetSolidBoundary"));
-static const FName ComputeBoundaryWeightsName(TEXT("ComputeBoundaryWeights"));
-static const FName GetNodePositionName(TEXT("GetNodePosition"));
-static const FName GetDensityFieldName(TEXT("GetDensityField"));
-static const FName UpdateDeformationGradientName(TEXT("UpdateDeformationGradient"));
+namespace NDIPressureGridLocal
+{
+	static const TCHAR* CommonShaderFile = TEXT("/Plugin/Runtime/HairStrands/Private/NiagaraDataInterfacePressureGrid.ush");
+	static const TCHAR* TemplateShaderFile = TEXT("/Plugin/Runtime/HairStrands/Private/NiagaraDataInterfacePressureGridTemplate.ush");
 
-//------------------------------------------------------------------------------------------------------------
+	static const FName BuildDistanceFieldName(TEXT("BuildDistanceField"));
+	static const FName BuildDensityFieldName(TEXT("BuildDensityField"));
+	static const FName SolveGridPressureName(TEXT("SolveGridPressure"));
+	static const FName ScaleCellFieldsName(TEXT("ScaleCellFields"));
+	static const FName SetSolidBoundaryName(TEXT("SetSolidBoundary"));
+	static const FName ComputeBoundaryWeightsName(TEXT("ComputeBoundaryWeights"));
+	static const FName GetNodePositionName(TEXT("GetNodePosition"));
+	static const FName GetDensityFieldName(TEXT("GetDensityField"));
+	static const FName UpdateDeformationGradientName(TEXT("UpdateDeformationGradient"));
 
-IMPLEMENT_NIAGARA_DI_PARAMETER(UNiagaraDataInterfacePressureGrid, FNDIVelocityGridParametersCS);
+} //namespace NDIPressureGridLocal
 
 //------------------------------------------------------------------------------------------------------------
 
@@ -46,6 +49,8 @@ UNiagaraDataInterfacePressureGrid::UNiagaraDataInterfacePressureGrid(FObjectInit
 
 void UNiagaraDataInterfacePressureGrid::GetFunctions(TArray<FNiagaraFunctionSignature>& OutFunctions)
 {
+	using namespace NDIPressureGridLocal;
+
 	Super::GetFunctions(OutFunctions);
 	{
 		FNiagaraFunctionSignature Sig;
@@ -201,6 +206,8 @@ DEFINE_NDI_DIRECT_FUNC_BINDER(UNiagaraDataInterfacePressureGrid, UpdateDeformati
 
 void UNiagaraDataInterfacePressureGrid::GetVMExternalFunction(const FVMExternalFunctionBindingInfo& BindingInfo, void* InstanceData, FVMExternalFunction &OutFunc)
 {
+	using namespace NDIPressureGridLocal;
+
 	Super::GetVMExternalFunction(BindingInfo, InstanceData, OutFunc);
 
 	if (BindingInfo.Name == BuildDistanceFieldName)
@@ -297,129 +304,33 @@ void UNiagaraDataInterfacePressureGrid::UpdateDeformationGradient(FVectorVMExter
 #if WITH_EDITORONLY_DATA
 bool UNiagaraDataInterfacePressureGrid::GetFunctionHLSL(const FNiagaraDataInterfaceGPUParamInfo& ParamInfo, const FNiagaraDataInterfaceGeneratedFunction& FunctionInfo, int FunctionInstanceIndex, FString& OutHLSL)
 {
+	using namespace NDIPressureGridLocal;
+
 	if (Super::GetFunctionHLSL(ParamInfo, FunctionInfo, FunctionInstanceIndex, OutHLSL))
 	{
 		return true;
 	}
-	else
-	{
-		TMap<FString, FStringFormatArg> ArgsSample = {
-			{TEXT("InstanceFunctionName"), FunctionInfo.InstanceName},
-			{TEXT("PressureGridContextName"), TEXT("DIVelocityGrid_MAKE_CONTEXT(") + ParamInfo.DataInterfaceHLSLSymbol + TEXT(")")},
-		};
-
-		if (FunctionInfo.DefinitionName == BuildDistanceFieldName)
-		{
-			static const TCHAR* FormatSample = TEXT(R"(
-					void {InstanceFunctionName} (in float3 GridOrigin, in float GridLength, 
-					in float3 ParticlePosition, out bool OutFunctionStatus )
-					{
-						{PressureGridContextName} DIVelocityGrid_BuildDistanceField(DIContext,DIContext_GridDestinationBuffer,GridOrigin,GridLength,ParticlePosition,OutFunctionStatus);
-					}
-					)");
-			OutHLSL += FString::Format(FormatSample, ArgsSample);
-			return true;
-		}
-		else if (FunctionInfo.DefinitionName == BuildDensityFieldName)
-		{
-			static const TCHAR* FormatSample = TEXT(R"(
-					void {InstanceFunctionName} (in float3 GridOrigin, in float GridLength, 
-					in float3 ParticlePosition, in float ParticleMass, in float ParticleDensity, out bool OutFunctionStatus )
-					{
-						{PressureGridContextName} DIVelocityGrid_BuildDensityField(DIContext,DIContext_GridDestinationBuffer,GridOrigin,GridLength,ParticlePosition,ParticleMass,ParticleDensity,OutFunctionStatus);
-					}
-					)");
-			OutHLSL += FString::Format(FormatSample, ArgsSample);
-			return true;
-		}
-		else if (FunctionInfo.DefinitionName == UpdateDeformationGradientName)
-		{
-			static const TCHAR* FormatSample = TEXT(R"(
-					void {InstanceFunctionName} (in float DeltaTime, in float4x4 VelocityGradient, 
-				in float4x4 DeformationGradient, out float4x4 OutDeformationGradient, out float OutGradientDeterminant)
-					{
-						{PressureGridContextName} DIVelocityGrid_UpdateDeformationGradient(DIContext,DeltaTime,VelocityGradient,DeformationGradient,OutDeformationGradient,OutGradientDeterminant);
-					}
-					)");
-			OutHLSL += FString::Format(FormatSample, ArgsSample);
-			return true;
-		}
-		else if (FunctionInfo.DefinitionName == SolveGridPressureName)
-		{
-			static const TCHAR* FormatSample = TEXT(R"(
-					void {InstanceFunctionName} (in int GridCell, in int InitStage, out bool OutProjectStatus)
-					{
-						{PressureGridContextName} DIVelocityGrid_SolveGridPressure(DIContext,DIContext_GridDestinationBuffer,GridCell,InitStage,OutProjectStatus);
-					}
-					)");
-			OutHLSL += FString::Format(FormatSample, ArgsSample);
-			return true;
-		}
-		else if (FunctionInfo.DefinitionName == GetNodePositionName)
-		{
-			static const TCHAR* FormatSample = TEXT(R"(
-					void {InstanceFunctionName} (in int GridCell, in float3 GridOrigin, in float GridLength, out float3 OutGridPosition)
-					{
-						{PressureGridContextName} DIVelocityGrid_GetNodePosition(DIContext,GridCell,GridOrigin,GridLength,OutGridPosition);
-					}
-					)");
-			OutHLSL += FString::Format(FormatSample, ArgsSample);
-			return true;
-		}
-		else if (FunctionInfo.DefinitionName == GetDensityFieldName)
-		{
-			static const TCHAR* FormatSample = TEXT(R"(
-					void {InstanceFunctionName} (in float3 GridOrigin, in float GridLength, in float3 ParticlePosition, out float OutParticleDensity)
-					{
-						{PressureGridContextName} DIVelocityGrid_GetDensityField(DIContext,GridOrigin,GridLength,ParticlePosition,OutParticleDensity);
-					}
-					)");
-			OutHLSL += FString::Format(FormatSample, ArgsSample);
-			return true;
-		}
-		else if (FunctionInfo.DefinitionName == SetSolidBoundaryName)
-		{
-			static const TCHAR* FormatSample = TEXT(R"(
-					void {InstanceFunctionName} (in int GridCell, in float SolidDistance, in float3 SolidVelocity, out bool OutBoundaryStatus)
-					{
-						{PressureGridContextName} DIVelocityGrid_SetSolidBoundary(DIContext,DIContext_GridDestinationBuffer,GridCell,SolidDistance,SolidVelocity,OutBoundaryStatus);
-					}
-					)");
-			OutHLSL += FString::Format(FormatSample, ArgsSample);
-			return true;
-		}
-		else if (FunctionInfo.DefinitionName == ComputeBoundaryWeightsName)
-		{
-			static const TCHAR* FormatSample = TEXT(R"(
-					void {InstanceFunctionName} (in int GridCell, out bool OutWeightsStatus)
-					{
-						{PressureGridContextName} DIVelocityGrid_ComputeBoundaryWeights(DIContext,DIContext_GridDestinationBuffer,GridCell,OutWeightsStatus);
-					}
-					)");
-			OutHLSL += FString::Format(FormatSample, ArgsSample);
-			return true;
-		}
-		else if (FunctionInfo.DefinitionName == ScaleCellFieldsName)
-		{
-			static const TCHAR* FormatSample = TEXT(R"(
-					void {InstanceFunctionName} (in int GridCell, in float GridLength, in float DeltaTime, out bool OutTransferStatus)
-					{
-						{PressureGridContextName} DIVelocityGrid_ScaleCellFields(DIContext,DIContext_GridDestinationBuffer,GridCell,GridLength,DeltaTime,OutTransferStatus);
-					}
-					)");
-			OutHLSL += FString::Format(FormatSample, ArgsSample);
-			return true;
-		}
-	}
 
-	OutHLSL += TEXT("\n");
-	return false;
+	static const TSet<FName> ValidGpuFunctions =
+	{
+		BuildDistanceFieldName,
+		BuildDensityFieldName,
+		UpdateDeformationGradientName,
+		SolveGridPressureName,
+		GetNodePositionName,
+		GetDensityFieldName,
+		SetSolidBoundaryName,
+		ComputeBoundaryWeightsName,
+		ScaleCellFieldsName,
+	};
+
+	return ValidGpuFunctions.Contains(FunctionInfo.DefinitionName);
 }
 
 void UNiagaraDataInterfacePressureGrid::GetCommonHLSL(FString& OutHLSL)
 {
 	Super::GetCommonHLSL(OutHLSL);
-	OutHLSL += TEXT("#include \"/Plugin/Runtime/HairStrands/Private/NiagaraDataInterfacePressureGrid.ush\"\n");
+	OutHLSL.Appendf(TEXT("#include \"%s\"\n"), NDIPressureGridLocal::CommonShaderFile);
 }
 
 bool UNiagaraDataInterfacePressureGrid::AppendCompileHash(FNiagaraCompileHashVisitor* InVisitor) const
@@ -427,14 +338,24 @@ bool UNiagaraDataInterfacePressureGrid::AppendCompileHash(FNiagaraCompileHashVis
 	if (!Super::AppendCompileHash(InVisitor))
 		return false;
 
-	FSHAHash Hash = GetShaderFileHash((TEXT("/Plugin/Runtime/HairStrands/Private/NiagaraDataInterfacePressureGrid.ush")), EShaderPlatform::SP_PCD3D_SM5);
-	InVisitor->UpdateString(TEXT("NiagaraDataInterfacePressureGridHLSLSource"), Hash.ToString());
+	InVisitor->UpdateString(TEXT("NiagaraDataInterfacePressureGridHLSLSource"), GetShaderFileHash(NDIPressureGridLocal::CommonShaderFile, EShaderPlatform::SP_PCD3D_SM5).ToString());
+	InVisitor->UpdateString(TEXT("NiagaraDataInterfacePressureGridTemplateHLSLSource"), GetShaderFileHash(NDIPressureGridLocal::TemplateShaderFile, EShaderPlatform::SP_PCD3D_SM5).ToString());
+
 	return true;
 }
 
 void UNiagaraDataInterfacePressureGrid::GetParameterDefinitionHLSL(const FNiagaraDataInterfaceGPUParamInfo& ParamInfo, FString& OutHLSL)
 {
-	OutHLSL += TEXT("DIVelocityGrid_DECLARE_CONSTANTS(") + ParamInfo.DataInterfaceHLSLSymbol + TEXT(")\n");
+	Super::GetParameterDefinitionHLSL(ParamInfo, OutHLSL);
+
+	TMap<FString, FStringFormatArg> TemplateArgs =
+	{
+		{TEXT("ParameterName"),	ParamInfo.DataInterfaceHLSLSymbol},
+	};
+
+	FString TemplateFile;
+	LoadShaderSourceFile(NDIPressureGridLocal::TemplateShaderFile, EShaderPlatform::SP_PCD3D_SM5, &TemplateFile, nullptr);
+	OutHLSL += FString::Format(*TemplateFile, TemplateArgs);
 }
 #endif
 
@@ -444,9 +365,9 @@ void UNiagaraDataInterfacePressureGrid::GetParameterDefinitionHLSL(const FNiagar
 
 class FClearPressureGridCS : public FGlobalShader
 {
-	DECLARE_SHADER_TYPE(FClearPressureGridCS, Global)
+	DECLARE_GLOBAL_SHADER(FClearPressureGridCS)
+	SHADER_USE_PARAMETER_STRUCT(FClearPressureGridCS, FGlobalShader);
 
-public:
 	static bool ShouldCompilePermutation(const FGlobalShaderPermutationParameters& Parameters)
 	{
 		return RHISupportsComputeShaders(Parameters.Platform);
@@ -458,92 +379,64 @@ public:
 		OutEnvironment.SetDefine(TEXT("THREAD_COUNT"), NIAGARA_HAIR_STRANDS_THREAD_COUNT_PRESSURE);
 	}
 
-	FClearPressureGridCS(const ShaderMetaType::CompiledShaderInitializerType& Initializer)
-		: FGlobalShader(Initializer)
-	{
-		GridDestinationBuffer.Bind(Initializer.ParameterMap, TEXT("GridDestinationBuffer"));
-		GridSize.Bind(Initializer.ParameterMap, TEXT("GridSize"));
-		CopyPressure.Bind(Initializer.ParameterMap, TEXT("CopyPressure"));
-		GridCurrentBuffer.Bind(Initializer.ParameterMap, TEXT("GridCurrentBuffer"));
-	}
-
-	FClearPressureGridCS()
-	{}
-
-	void SetParameters(FRHICommandList& RHICmdList, FRHIShaderResourceView* InGridCurrentBuffer,
-		FRHIUnorderedAccessView* InGridDestinationBuffer,
-		const FIntVector& InGridSize, const int32 InCopyPressure)
-	{
-		FRHIComputeShader* ShaderRHI = RHICmdList.GetBoundComputeShader();
-
-		SetUAVParameter(RHICmdList, ShaderRHI, GridDestinationBuffer, InGridDestinationBuffer);
-		SetShaderValue(RHICmdList, ShaderRHI, GridSize, InGridSize);
-		SetShaderValue(RHICmdList, ShaderRHI, CopyPressure, InCopyPressure);
-		SetSRVParameter(RHICmdList, ShaderRHI, GridCurrentBuffer, InGridCurrentBuffer);
-	}
-
-	void UnsetParameters(FRHICommandList& RHICmdList)
-	{
-		FRHIComputeShader* ShaderRHI = RHICmdList.GetBoundComputeShader();
-
-		SetUAVParameter(RHICmdList, ShaderRHI, GridDestinationBuffer, nullptr);
-	}
-
-	LAYOUT_FIELD(FShaderResourceParameter, GridDestinationBuffer);
-	LAYOUT_FIELD(FShaderParameter, GridSize);
-	LAYOUT_FIELD(FShaderParameter, CopyPressure);
-	LAYOUT_FIELD(FShaderResourceParameter, GridCurrentBuffer);
+	BEGIN_SHADER_PARAMETER_STRUCT(FParameters, )
+		SHADER_PARAMETER(FIntVector,					GridSize)
+		SHADER_PARAMETER(int,							CopyPressure)
+		SHADER_PARAMETER_RDG_BUFFER_SRV(Buffer<int>,	GridCurrentBuffer)
+		SHADER_PARAMETER_RDG_BUFFER_UAV(RWBuffer<int>,	GridDestinationBuffer)
+	END_SHADER_PARAMETER_STRUCT()
 };
 
 IMPLEMENT_SHADER_TYPE(, FClearPressureGridCS, TEXT("/Plugin/Runtime/HairStrands/Private/NiagaraClearPressureGrid.usf"), TEXT("MainCS"), SF_Compute);
 
-
 //------------------------------------------------------------------------------------------------------------
 
-inline void ClearBuffer(FRHICommandList& RHICmdList, ERHIFeatureLevel::Type FeatureLevel, FNDIVelocityGridBuffer* CurrentGridBuffer, FNDIVelocityGridBuffer* DestinationGridBuffer, const FIntVector& GridSize, const bool CopyPressure)
+inline void ClearBuffer(FRDGBuilder& GraphBuilder, ERHIFeatureLevel::Type FeatureLevel, FNDIVelocityGridBuffer* CurrentGridBuffer, FNDIVelocityGridBuffer* DestinationGridBuffer, const FIntVector& GridSize, const bool CopyPressure)
 {
-	FRHIUnorderedAccessView* DestinationGridBufferUAV = DestinationGridBuffer->GridDataBuffer.UAV;
-	FRHIShaderResourceView* CurrentGridBufferSRV = CurrentGridBuffer->GridDataBuffer.SRV;
-	FRHIUnorderedAccessView* CurrentGridBufferUAV = CurrentGridBuffer->GridDataBuffer.UAV;
-
-	if (DestinationGridBufferUAV != nullptr && CurrentGridBufferSRV != nullptr && CurrentGridBufferUAV != nullptr)
+	if (CurrentGridBuffer->GridDataBuffer.IsValid() == false || DestinationGridBuffer->GridDataBuffer.IsValid() == false)
 	{
-		TShaderMapRef<FClearPressureGridCS> ComputeShader(GetGlobalShaderMap(FeatureLevel));
-		SetComputePipelineState(RHICmdList, ComputeShader.GetComputeShader());
-
-		FRHITransitionInfo Transitions[] = {
-			FRHITransitionInfo(CurrentGridBufferUAV, ERHIAccess::Unknown, ERHIAccess::SRVCompute),
-			FRHITransitionInfo(DestinationGridBufferUAV, ERHIAccess::Unknown, ERHIAccess::UAVCompute)
-		};
-		RHICmdList.Transition(MakeArrayView(Transitions, UE_ARRAY_COUNT(Transitions)));
-
-		const uint32 GroupSize = NIAGARA_HAIR_STRANDS_THREAD_COUNT_PRESSURE;
-		const FIntVector ExtendedSize = GridSize + FIntVector(1, 1, 1);
-		
-		const uint32 DispatchCountX = FMath::DivideAndRoundUp((uint32)(ExtendedSize.X), GroupSize);
-		const uint32 DispatchCountY = FMath::DivideAndRoundUp((uint32)(ExtendedSize.Y), GroupSize);
-		const uint32 DispatchCountZ = FMath::DivideAndRoundUp((uint32)(ExtendedSize.Z), GroupSize);
-
-		ComputeShader->SetParameters(RHICmdList, CurrentGridBufferSRV, DestinationGridBufferUAV, ExtendedSize, CopyPressure);
-		DispatchComputeShader(RHICmdList, ComputeShader.GetShader(), DispatchCountX, DispatchCountY, DispatchCountZ);
-		ComputeShader->UnsetParameters(RHICmdList);
+		return;
 	}
+
+	const FIntVector ExtendedSize = GridSize + FIntVector(1, 1, 1);
+	const uint32 ThreadGroupSize = NIAGARA_HAIR_STRANDS_THREAD_COUNT_PRESSURE;
+	const FIntVector ThreadGroupCount(
+		FMath::DivideAndRoundUp((uint32)ExtendedSize.X, ThreadGroupSize),
+		FMath::DivideAndRoundUp((uint32)ExtendedSize.Y, ThreadGroupSize),
+		FMath::DivideAndRoundUp((uint32)ExtendedSize.Z, ThreadGroupSize)
+	);
+
+	TShaderMapRef<FClearPressureGridCS> ComputeShader(GetGlobalShaderMap(GMaxRHIFeatureLevel));
+
+	FClearPressureGridCS::FParameters* PassParameters = GraphBuilder.AllocParameters<FClearPressureGridCS::FParameters>();
+	PassParameters->GridSize = ExtendedSize;
+	PassParameters->CopyPressure = CopyPressure;
+	PassParameters->GridCurrentBuffer = CurrentGridBuffer->GridDataBuffer.GetOrCreateSRV(GraphBuilder);
+	PassParameters->GridDestinationBuffer = DestinationGridBuffer->GridDataBuffer.GetOrCreateUAV(GraphBuilder);
+
+	FComputeShaderUtils::AddPass(
+		GraphBuilder,
+		RDG_EVENT_NAME("VelocityGrid::ClearTexture"),
+		ERDGPassFlags::Compute,
+		ComputeShader,
+		PassParameters,
+		ThreadGroupCount
+	);
 }
 
 //------------------------------------------------------------------------------------------------------------
 
-void FNDIPressureGridProxy::PreStage(FRHICommandList& RHICmdList, const FNiagaraDataInterfaceStageArgs& Context)
+void FNDIPressureGridProxy::PreStage(const FNDIGpuComputePreStageContext& Context)
 {
-	FNDIVelocityGridData* ProxyData =
-		FNDIVelocityGridProxy::SystemInstancesToProxyData.Find(Context.SystemInstanceID);
+	FNDIVelocityGridData* ProxyData = FNDIVelocityGridProxy::SystemInstancesToProxyData.Find(Context.GetSystemInstanceID());
 
 	if (ProxyData != nullptr)
 	{
-		if (Context.SimStageData->bFirstStage)
+		if (Context.GetSimStageData().bFirstStage)
 		{
-			ClearBuffer(RHICmdList, Context.ComputeDispatchInterface->GetFeatureLevel(), ProxyData->CurrentGridBuffer, ProxyData->DestinationGridBuffer, ProxyData->GridSize, true);
+			ClearBuffer(Context.GetGraphBuilder(), Context.GetComputeDispatchInterface().GetFeatureLevel(), ProxyData->CurrentGridBuffer, ProxyData->DestinationGridBuffer, ProxyData->GridSize, true);
 		}
 	}
 }
 
-#undef LOCTEXT_NAMESPACE
\ No newline at end of file
+#undef LOCTEXT_NAMESPACE
diff --git a/Engine/Plugins/Runtime/HairStrands/Source/HairStrandsCore/Private/Niagara/NiagaraDataInterfaceVelocityGrid.cpp b/Engine/Plugins/Runtime/HairStrands/Source/HairStrandsCore/Private/Niagara/NiagaraDataInterfaceVelocityGrid.cpp
--- a/Engine/Plugins/Runtime/HairStrands/Source/HairStrandsCore/Private/Niagara/NiagaraDataInterfaceVelocityGrid.cpp
+++ b/Engine/Plugins/Runtime/HairStrands/Source/HairStrandsCore/Private/Niagara/NiagaraDataInterfaceVelocityGrid.cpp
@@ -6,6 +6,7 @@
 #include "NiagaraGpuComputeDispatchInterface.h"
 #include "NiagaraRenderer.h"
 #include "NiagaraSimStageData.h"
+#include "NiagaraShaderParametersBuilder.h"
 #include "NiagaraSystemInstance.h"
 
 #include "ShaderParameterUtils.h"
@@ -21,32 +22,26 @@ DEFINE_LOG_CATEGORY_STATIC(LogVelocityGrid, Log, All);
 
 //------------------------------------------------------------------------------------------------------------
 
-static const FName BuildVelocityFieldName(TEXT("BuildVelocityField"));
-static const FName SampleVelocityFieldName(TEXT("SampleVelocityField"));
-static const FName ComputeGridSizeName(TEXT("ComputeGridSize"));
-static const FName UpdateGridTransformName(TEXT("UpdateGridTransform"));
-static const FName SetGridDimensionName(TEXT("SetGridDimension"));
-
-//------------------------------------------------------------------------------------------------------------
-
-const FString UNiagaraDataInterfaceVelocityGrid::GridCurrentBufferName(TEXT("GridCurrentBuffer_"));
-const FString UNiagaraDataInterfaceVelocityGrid::GridDestinationBufferName(TEXT("GridDestinationBuffer_"));
-
-const FString UNiagaraDataInterfaceVelocityGrid::GridSizeName(TEXT("GridSize_"));
-const FString UNiagaraDataInterfaceVelocityGrid::WorldTransformName(TEXT("WorldTransform_"));
-const FString UNiagaraDataInterfaceVelocityGrid::WorldInverseName(TEXT("WorldInverse_"));
+namespace NDIVelocityGridLocal
+{
+	BEGIN_SHADER_PARAMETER_STRUCT(FShaderParameters,)
+		SHADER_PARAMETER(FMatrix44f,					WorldTransform)
+		SHADER_PARAMETER(FMatrix44f,					WorldInverse)
+		SHADER_PARAMETER(FIntVector,					GridSize)
+		SHADER_PARAMETER_RDG_BUFFER_SRV(Buffer<int>,	GridCurrentBuffer)
+		SHADER_PARAMETER_RDG_BUFFER_UAV(RWBuffer<int>,	GridDestinationBuffer)
+	END_SHADER_PARAMETER_STRUCT();
 
-//------------------------------------------------------------------------------------------------------------
+	static const TCHAR* CommonShaderFile = TEXT("/Plugin/Runtime/HairStrands/Private/NiagaraDataInterfaceVelocityGrid.ush");
+	static const TCHAR* TemplateShaderFile = TEXT("/Plugin/Runtime/HairStrands/Private/NiagaraDataInterfaceVelocityGridTemplate.ush");
 
-FNDIVelocityGridParametersName::FNDIVelocityGridParametersName(const FString& Suffix)
-{
-	GridCurrentBufferName = UNiagaraDataInterfaceVelocityGrid::GridCurrentBufferName + Suffix;
-	GridDestinationBufferName = UNiagaraDataInterfaceVelocityGrid::GridDestinationBufferName + Suffix;
+	static const FName BuildVelocityFieldName(TEXT("BuildVelocityField"));
+	static const FName SampleVelocityFieldName(TEXT("SampleVelocityField"));
+	static const FName ComputeGridSizeName(TEXT("ComputeGridSize"));
+	static const FName UpdateGridTransformName(TEXT("UpdateGridTransform"));
+	static const FName SetGridDimensionName(TEXT("SetGridDimension"));
 
-	GridSizeName = UNiagaraDataInterfaceVelocityGrid::GridSizeName + Suffix;
-	WorldTransformName = UNiagaraDataInterfaceVelocityGrid::WorldTransformName + Suffix;
-	WorldInverseName = UNiagaraDataInterfaceVelocityGrid::WorldInverseName + Suffix;
-}
+} //namespace NDIVelocityGridLocal
 
 //------------------------------------------------------------------------------------------------------------
 
@@ -61,8 +56,12 @@ void FNDIVelocityGridBuffer::InitRHI()
 	if (GridSize.X != 0 && GridSize.Y != 0 && GridSize.Z != 0)
 	{
 		static const uint32 NumComponents = NumAttributes;
-		GridDataBuffer.Initialize(TEXT("FNDIVelocityGridBuffer"), sizeof(int32),
-			(GridSize.X + 1)*NumComponents * (GridSize.Y + 1) * (GridSize.Z + 1), EPixelFormat::PF_R32_SINT);
+
+		FMemMark MemMark(FMemStack::Get());
+		FRDGBuilder GraphBuilder(FRHICommandListExecutor::GetImmediateCommandList());
+		GridDataBuffer.Initialize(GraphBuilder, TEXT("FNDIVelocityGridBuffer"), EPixelFormat::PF_R32_SINT, sizeof(int32), (GridSize.X + 1) * NumComponents * (GridSize.Y + 1) * (GridSize.Z + 1));
+		GridDataBuffer.EndGraphUsage();
+		GraphBuilder.Execute();
 	}
 }
 
@@ -145,70 +144,6 @@ bool FNDIVelocityGridData::Init(const FIntVector& InGridSize, const int32 InNumA
 	return true;
 }
 
-void FNDIVelocityGridParametersCS::Bind(const FNiagaraDataInterfaceGPUParamInfo& ParameterInfo, const class FShaderParameterMap& ParameterMap)
-{
-	FNDIVelocityGridParametersName ParamNames(ParameterInfo.DataInterfaceHLSLSymbol);
-
-	GridCurrentBuffer.Bind(ParameterMap, *ParamNames.GridCurrentBufferName);
-	GridDestinationBuffer.Bind(ParameterMap, *ParamNames.GridDestinationBufferName);
-
-	GridSize.Bind(ParameterMap, *ParamNames.GridSizeName);
-	WorldTransform.Bind(ParameterMap, *ParamNames.WorldTransformName);
-	WorldInverse.Bind(ParameterMap, *ParamNames.WorldInverseName);
-
-}
-
-void FNDIVelocityGridParametersCS::Set(FRHICommandList& RHICmdList, const FNiagaraDataInterfaceSetArgs& Context) const
-{
-	check(IsInRenderingThread());
-
-	FRHIComputeShader* ComputeShaderRHI = RHICmdList.GetBoundComputeShader();
-
-	FNDIVelocityGridProxy* InterfaceProxy =
-		static_cast<FNDIVelocityGridProxy*>(Context.DataInterface);
-	FNDIVelocityGridData* ProxyData =
-		InterfaceProxy->SystemInstancesToProxyData.Find(Context.SystemInstanceID);
-
-	if (ProxyData != nullptr && ProxyData->CurrentGridBuffer != nullptr && ProxyData->DestinationGridBuffer != nullptr
-		&& ProxyData->CurrentGridBuffer->IsInitialized() && ProxyData->DestinationGridBuffer->IsInitialized())
-	{
-		FNDIVelocityGridBuffer* CurrentGridBuffer = ProxyData->CurrentGridBuffer;
-		FNDIVelocityGridBuffer* DestinationGridBuffer = ProxyData->DestinationGridBuffer;
-
-		FRHITransitionInfo Transitions[] = {
-			FRHITransitionInfo(CurrentGridBuffer->GridDataBuffer.UAV, ERHIAccess::Unknown, ERHIAccess::SRVCompute),
-			FRHITransitionInfo(DestinationGridBuffer->GridDataBuffer.UAV, ERHIAccess::Unknown, ERHIAccess::UAVCompute)
-		};
-		RHICmdList.Transition(MakeArrayView(Transitions, UE_ARRAY_COUNT(Transitions)));
-
-		SetUAVParameter(RHICmdList, ComputeShaderRHI, GridDestinationBuffer, DestinationGridBuffer->GridDataBuffer.UAV);
-		SetSRVParameter(RHICmdList, ComputeShaderRHI, GridCurrentBuffer, CurrentGridBuffer->GridDataBuffer.SRV);
-
-		SetShaderValue(RHICmdList, ComputeShaderRHI, GridSize, ProxyData->GridSize);
-		SetShaderValue(RHICmdList, ComputeShaderRHI, WorldTransform, (FMatrix44f)ProxyData->WorldTransform);
-		SetShaderValue(RHICmdList, ComputeShaderRHI, WorldInverse, (FMatrix44f)ProxyData->WorldTransform.Inverse());
-	}
-	else
-	{
-		SetUAVParameter(RHICmdList, ComputeShaderRHI, GridDestinationBuffer, Context.ComputeDispatchInterface->GetEmptyUAVFromPool(RHICmdList, PF_R32_UINT, ENiagaraEmptyUAVType::Buffer));
-		SetSRVParameter(RHICmdList, ComputeShaderRHI, GridCurrentBuffer, FNiagaraRenderer::GetDummyUIntBuffer());
-
-		SetShaderValue(RHICmdList, ComputeShaderRHI, GridSize, FIntVector());
-		SetShaderValue(RHICmdList, ComputeShaderRHI, WorldTransform, FMatrix44f::Identity);
-		SetShaderValue(RHICmdList, ComputeShaderRHI, WorldInverse, FMatrix44f::Identity);
-	}
-}
-
-void FNDIVelocityGridParametersCS::Unset(FRHICommandList& RHICmdList, const FNiagaraDataInterfaceSetArgs& Context) const
-{
-	FRHIComputeShader* ShaderRHI = RHICmdList.GetBoundComputeShader();
-	SetUAVParameter(RHICmdList, ShaderRHI, GridDestinationBuffer, nullptr);
-}
-
-IMPLEMENT_TYPE_LAYOUT(FNDIVelocityGridParametersCS);
-
-IMPLEMENT_NIAGARA_DI_PARAMETER(UNiagaraDataInterfaceVelocityGrid, FNDIVelocityGridParametersCS);
-
 //------------------------------------------------------------------------------------------------------------
 
 UNiagaraDataInterfaceVelocityGrid::UNiagaraDataInterfaceVelocityGrid(FObjectInitializer const& ObjectInitializer)
@@ -300,8 +235,10 @@ bool UNiagaraDataInterfaceVelocityGrid::AppendCompileHash(FNiagaraCompileHashVis
 		return false;
 	}
 
-	FSHAHash Hash = GetShaderFileHash((TEXT("/Plugin/Runtime/HairStrands/Private/NiagaraDataInterfaceVelocityGrid.ush")), EShaderPlatform::SP_PCD3D_SM5);
-	InVisitor->UpdateString(TEXT("NiagaraDataInterfaceVelocityGridHLSLSource"), Hash.ToString());
+	InVisitor->UpdateString(TEXT("NiagaraDataInterfaceVelocityGridHLSLSource"), GetShaderFileHash(NDIVelocityGridLocal::CommonShaderFile, EShaderPlatform::SP_PCD3D_SM5).ToString());
+	InVisitor->UpdateString(TEXT("NiagaraDataInterfaceVelocityGridTemplateHLSLSource"), GetShaderFileHash(NDIVelocityGridLocal::TemplateShaderFile, EShaderPlatform::SP_PCD3D_SM5).ToString());
+	InVisitor->UpdateShaderParameters<NDIVelocityGridLocal::FShaderParameters>();
+	
 	return true;
 }
 #endif
@@ -319,6 +256,8 @@ void UNiagaraDataInterfaceVelocityGrid::PostInitProperties()
 
 void UNiagaraDataInterfaceVelocityGrid::GetFunctions(TArray<FNiagaraFunctionSignature>& OutFunctions)
 {
+	using namespace NDIVelocityGridLocal;
+
 	{
 		FNiagaraFunctionSignature Sig;
 		Sig.Name = BuildVelocityFieldName;
@@ -409,6 +348,8 @@ DEFINE_NDI_DIRECT_FUNC_BINDER(UNiagaraDataInterfaceVelocityGrid, SetGridDimensio
 
 void UNiagaraDataInterfaceVelocityGrid::GetVMExternalFunction(const FVMExternalFunctionBindingInfo& BindingInfo, void* InstanceData, FVMExternalFunction &OutFunc)
 {
+	using namespace NDIVelocityGridLocal;
+
 	if (BindingInfo.Name == BuildVelocityFieldName)
 	{
 		check(BindingInfo.GetNumInputs() == 28 && BindingInfo.GetNumOutputs() == 1);
@@ -533,61 +474,72 @@ void UNiagaraDataInterfaceVelocityGrid::UpdateGridTransform(FVectorVMExternalFun
 #if WITH_EDITORONLY_DATA
 bool UNiagaraDataInterfaceVelocityGrid::GetFunctionHLSL(const FNiagaraDataInterfaceGPUParamInfo& ParamInfo, const FNiagaraDataInterfaceGeneratedFunction& FunctionInfo, int FunctionInstanceIndex, FString& OutHLSL)
 {
-	FNDIVelocityGridParametersName ParamNames(ParamInfo.DataInterfaceHLSLSymbol);
+	using namespace NDIVelocityGridLocal;
 
-	TMap<FString, FStringFormatArg> ArgsSample = {
-		{TEXT("InstanceFunctionName"), FunctionInfo.InstanceName},
-		{TEXT("VelocityGridContextName"), TEXT("DIVelocityGrid_MAKE_CONTEXT(") + ParamInfo.DataInterfaceHLSLSymbol + TEXT(")")},
+	static const TSet<FName> ValidGpuFunctions =
+	{
+		BuildVelocityFieldName,
+		SampleVelocityFieldName,
+		ComputeGridSizeName,
 	};
 
-	if (FunctionInfo.DefinitionName == BuildVelocityFieldName)
-	{
-		static const TCHAR *FormatSample = TEXT(R"(
-				void {InstanceFunctionName} (in float3 GridOrigin, in float GridLength, in float3 ParticlePosition, in float ParticleMass, in float3 ParticleVelocity, in float4x4 VelocityGradient, out bool OutFunctionStatus)
-				{
-					{VelocityGridContextName} DIVelocityGrid_BuildVelocityField(DIContext,DIContext_GridDestinationBuffer,GridOrigin,GridLength,ParticlePosition,ParticleMass,ParticleVelocity,VelocityGradient,OutFunctionStatus);
-				}
-				)");
-		OutHLSL += FString::Format(FormatSample, ArgsSample);
-		return true;
-	}
-	else if (FunctionInfo.DefinitionName == SampleVelocityFieldName)
-	{
-		static const TCHAR *FormatSample = TEXT(R"(
-				void {InstanceFunctionName} (in float3 GridOrigin, in float GridLength, in float3 ParticlePosition, in bool ScaledVelocity, out float OutParticleMass, out float3 OutParticleVelocity, out float4x4 OutVelocityGradient)
-				{
-					{VelocityGridContextName} DIVelocityGrid_SampleVelocityField(DIContext,GridOrigin,GridLength,ParticlePosition,ScaledVelocity,OutParticleMass,OutParticleVelocity,OutVelocityGradient);
-				}
-				)");
-		OutHLSL += FString::Format(FormatSample, ArgsSample);
-		return true;
-	}
-	else if (FunctionInfo.DefinitionName == ComputeGridSizeName)
-	{
-		static const TCHAR *FormatSample = TEXT(R"(
-				void {InstanceFunctionName} (in float3 GridCenter, in float3 GridExtent, out float3 OutGridOrigin, out float OutGridLength)
-				{
-					{VelocityGridContextName} DIVelocityGrid_ComputeGridSize(DIContext,GridCenter,GridExtent,OutGridOrigin,OutGridLength);
-				}
-				)");
-		OutHLSL += FString::Format(FormatSample, ArgsSample);
-		return true;
-	}
-	OutHLSL += TEXT("\n");
-	return false;
+	return ValidGpuFunctions.Contains(FunctionInfo.DefinitionName);
 }
 
 void UNiagaraDataInterfaceVelocityGrid::GetCommonHLSL(FString& OutHLSL)
 {
-	OutHLSL += TEXT("#include \"/Plugin/Runtime/HairStrands/Private/NiagaraDataInterfaceVelocityGrid.ush\"\n");
+	OutHLSL.Appendf(TEXT("#include \"%s\"\n"), NDIVelocityGridLocal::CommonShaderFile);
 }
 
 void UNiagaraDataInterfaceVelocityGrid::GetParameterDefinitionHLSL(const FNiagaraDataInterfaceGPUParamInfo& ParamInfo, FString& OutHLSL)
 {
-	OutHLSL += TEXT("DIVelocityGrid_DECLARE_CONSTANTS(") + ParamInfo.DataInterfaceHLSLSymbol + TEXT(")\n");
+	Super::GetParameterDefinitionHLSL(ParamInfo, OutHLSL);
+
+	TMap<FString, FStringFormatArg> TemplateArgs =
+	{
+		{TEXT("ParameterName"),	ParamInfo.DataInterfaceHLSLSymbol},
+	};
+
+	FString TemplateFile;
+	LoadShaderSourceFile(NDIVelocityGridLocal::TemplateShaderFile, EShaderPlatform::SP_PCD3D_SM5, &TemplateFile, nullptr);
+	OutHLSL += FString::Format(*TemplateFile, TemplateArgs);
 }
 #endif
 
+void UNiagaraDataInterfaceVelocityGrid::BuildShaderParameters(FNiagaraShaderParametersBuilder& ShaderParametersBuilder) const
+{
+	ShaderParametersBuilder.AddNestedStruct<NDIVelocityGridLocal::FShaderParameters>();
+}
+
+void UNiagaraDataInterfaceVelocityGrid::SetShaderParameters(const FNiagaraDataInterfaceSetShaderParametersContext& Context) const
+{
+	FRDGBuilder& GraphBuilder = Context.GetGraphBuilder();
+	FNDIVelocityGridProxy& DIProxy = Context.GetProxy<FNDIVelocityGridProxy>();
+	FNDIVelocityGridData* ProxyData = DIProxy.SystemInstancesToProxyData.Find(Context.GetSystemInstanceID());
+
+	NDIVelocityGridLocal::FShaderParameters* ShaderParameters = Context.GetParameterNestedStruct<NDIVelocityGridLocal::FShaderParameters>();
+
+	if (ProxyData != nullptr && ProxyData->CurrentGridBuffer != nullptr && ProxyData->DestinationGridBuffer != nullptr
+		&& ProxyData->CurrentGridBuffer->IsInitialized() && ProxyData->DestinationGridBuffer->IsInitialized())
+	{
+		ShaderParameters->GridSize			= ProxyData->GridSize;
+		ShaderParameters->WorldTransform	= FMatrix44f(ProxyData->WorldTransform);
+		ShaderParameters->WorldInverse		= FMatrix44f(ProxyData->WorldInverse);
+
+		ShaderParameters->GridCurrentBuffer = ProxyData->CurrentGridBuffer->GridDataBuffer.GetOrCreateSRV(GraphBuilder);
+		ShaderParameters->GridDestinationBuffer = ProxyData->DestinationGridBuffer->GridDataBuffer.GetOrCreateUAV(GraphBuilder);
+	}
+	else
+	{
+		ShaderParameters->GridSize			= FIntVector::ZeroValue;
+		ShaderParameters->WorldTransform	= FMatrix44f::Identity;
+		ShaderParameters->WorldInverse		= FMatrix44f::Identity;
+
+		ShaderParameters->GridCurrentBuffer = Context.GetComputeDispatchInterface().GetEmptyBufferSRV(GraphBuilder, EPixelFormat::PF_R32_SINT);
+		ShaderParameters->GridDestinationBuffer	= Context.GetComputeDispatchInterface().GetEmptyBufferUAV(GraphBuilder, EPixelFormat::PF_R32_SINT);
+	}
+}
+
 void UNiagaraDataInterfaceVelocityGrid::ProvidePerInstanceDataForRenderThread(void* DataForRenderThread, void* PerInstanceData, const FNiagaraSystemInstanceID& SystemInstance)
 {
 	FNDIVelocityGridData* GameThreadData = static_cast<FNDIVelocityGridData*>(PerInstanceData);
@@ -602,23 +554,6 @@ void UNiagaraDataInterfaceVelocityGrid::ProvidePerInstanceDataForRenderThread(vo
 
 //------------------------------------------------------------------------------------------------------------
 
-inline void ClearBuffer(FRHICommandList& RHICmdList, FNDIVelocityGridBuffer* DestinationGridBuffer)
-{
-	FRHIUnorderedAccessView* DestinationGridBufferUAV = DestinationGridBuffer->GridDataBuffer.UAV;
-
-	if (DestinationGridBufferUAV != nullptr )
-	{
-		FRHITransitionInfo Transitions[] = {
-			FRHITransitionInfo(DestinationGridBufferUAV, ERHIAccess::Unknown, ERHIAccess::UAVCompute),
-		};
-		RHICmdList.Transition(MakeArrayView(Transitions, UE_ARRAY_COUNT(Transitions)));
-
-		RHICmdList.ClearUAVUint(DestinationGridBufferUAV, FUintVector4(0, 0, 0, 0));
-	}
-}
-
-//------------------------------------------------------------------------------------------------------------
-
 void FNDIVelocityGridProxy::ConsumePerInstanceDataFromGameThread(void* PerInstanceData, const FNiagaraSystemInstanceID& Instance)
 {
 	FNDIVelocityGridData* SourceData = static_cast<FNDIVelocityGridData*>(PerInstanceData);
@@ -646,9 +581,9 @@ void FNDIVelocityGridProxy::ConsumePerInstanceDataFromGameThread(void* PerInstan
 
 class FClearVelocityGridCS : public FGlobalShader
 {
-	DECLARE_SHADER_TYPE(FClearVelocityGridCS, Global)
+	DECLARE_GLOBAL_SHADER(FClearVelocityGridCS)
+	SHADER_USE_PARAMETER_STRUCT(FClearVelocityGridCS, FGlobalShader);
 
-public:
 	static bool ShouldCompilePermutation(const FGlobalShaderPermutationParameters& Parameters)
 	{
 		return RHISupportsComputeShaders(Parameters.Platform);
@@ -660,72 +595,52 @@ public:
 		OutEnvironment.SetDefine(TEXT("THREAD_COUNT"), NIAGARA_HAIR_STRANDS_THREAD_COUNT_VELOCITY);
 	}
 
-	FClearVelocityGridCS(const ShaderMetaType::CompiledShaderInitializerType& Initializer)
-		: FGlobalShader(Initializer)
-	{
-		GridSize.Bind(Initializer.ParameterMap, TEXT("GridSize"));
-		GridDestinationBuffer.Bind(Initializer.ParameterMap, TEXT("GridDestinationBuffer"));
-	}
-
-	FClearVelocityGridCS()
-	{}
-
-	void SetParameters(FRHICommandList& RHICmdList, FRHIUnorderedAccessView* InGridDestinationBuffer,
-		const FIntVector& InGridSize)
-	{
-		FRHIComputeShader* ShaderRHI = RHICmdList.GetBoundComputeShader();
-
-		SetUAVParameter(RHICmdList, ShaderRHI, GridDestinationBuffer, InGridDestinationBuffer);
-		SetShaderValue(RHICmdList, ShaderRHI, GridSize, InGridSize);
-	}
-
-	void UnsetParameters(FRHICommandList& RHICmdList)
-	{
-		FRHIComputeShader* ShaderRHI = RHICmdList.GetBoundComputeShader();
-
-		SetUAVParameter(RHICmdList, ShaderRHI, GridDestinationBuffer, nullptr);
-	}
-
-	LAYOUT_FIELD(FShaderResourceParameter, GridDestinationBuffer);
-	LAYOUT_FIELD(FShaderParameter, GridSize);
+	BEGIN_SHADER_PARAMETER_STRUCT(FParameters, )
+		SHADER_PARAMETER(FIntVector,					GridSize)
+		SHADER_PARAMETER_RDG_BUFFER_UAV(RWBuffer<int>,	GridDestinationBuffer)
+	END_SHADER_PARAMETER_STRUCT()
 };
 
 IMPLEMENT_SHADER_TYPE(, FClearVelocityGridCS, TEXT("/Plugin/Runtime/HairStrands/Private/NiagaraClearVelocityGrid.usf"), TEXT("MainCS"), SF_Compute);
 
-inline void ClearTexture(FRHICommandList& RHICmdList, FNDIVelocityGridBuffer* DestinationGridBuffer, const FIntVector& InGridSize)
+inline void ClearTexture(FRDGBuilder& GraphBuilder, FNDIVelocityGridBuffer* DestinationGridBuffer, const FIntVector& InGridSize)
 {
-	FRHIUnorderedAccessView* DestinationGridBufferUAV = DestinationGridBuffer->GridDataBuffer.UAV;
-
-	if (DestinationGridBufferUAV != nullptr)
+	if (DestinationGridBuffer->GridDataBuffer.IsValid() == false)
 	{
-		TShaderMapRef<FClearVelocityGridCS> ComputeShader(GetGlobalShaderMap(GMaxRHIFeatureLevel));
-		SetComputePipelineState(RHICmdList, ComputeShader.GetComputeShader());
+		return;
+	}
 
-		FRHITransitionInfo Transitions[] = {
-			FRHITransitionInfo(DestinationGridBufferUAV, ERHIAccess::Unknown, ERHIAccess::UAVCompute)
-		};
-		RHICmdList.Transition(MakeArrayView(Transitions, UE_ARRAY_COUNT(Transitions)));
+	const FIntVector GridSize((InGridSize.X + 1) * DestinationGridBuffer->NumAttributes, InGridSize.Y + 1, InGridSize.Z + 1);
+	const uint32 ThreadGroupSize = NIAGARA_HAIR_STRANDS_THREAD_COUNT_VELOCITY;
+	const FIntVector ThreadGroupCount(
+		FMath::DivideAndRoundUp((uint32)GridSize.X, ThreadGroupSize),
+		FMath::DivideAndRoundUp((uint32)GridSize.Y, ThreadGroupSize),
+		FMath::DivideAndRoundUp((uint32)GridSize.Z, ThreadGroupSize)
+	);
 
-		const uint32 GroupSize = NIAGARA_HAIR_STRANDS_THREAD_COUNT_VELOCITY;
-		const FIntVector GridSize((InGridSize.X + 1) * DestinationGridBuffer->NumAttributes, InGridSize.Y + 1, InGridSize.Z + 1);
+	TShaderMapRef<FClearVelocityGridCS> ComputeShader(GetGlobalShaderMap(GMaxRHIFeatureLevel));
 
-		const uint32 DispatchCountX = FMath::DivideAndRoundUp((uint32)GridSize.X, GroupSize);
-		const uint32 DispatchCountY = FMath::DivideAndRoundUp((uint32)GridSize.Y, GroupSize);
-		const uint32 DispatchCountZ = FMath::DivideAndRoundUp((uint32)GridSize.Z, GroupSize);
+	FClearVelocityGridCS::FParameters* PassParameters = GraphBuilder.AllocParameters<FClearVelocityGridCS::FParameters>();
+	PassParameters->GridSize = GridSize;
+	PassParameters->GridDestinationBuffer = DestinationGridBuffer->GridDataBuffer.GetOrCreateUAV(GraphBuilder);
 
-		ComputeShader->SetParameters(RHICmdList, DestinationGridBufferUAV, GridSize);
-		DispatchComputeShader(RHICmdList, ComputeShader.GetShader(), DispatchCountX, DispatchCountY, DispatchCountZ);
-		ComputeShader->UnsetParameters(RHICmdList);
-	}
+	FComputeShaderUtils::AddPass(
+		GraphBuilder,
+		RDG_EVENT_NAME("VelocityGrid::ClearTexture"),
+		ERDGPassFlags::Compute,
+		ComputeShader,
+		PassParameters,
+		ThreadGroupCount
+	);
 }
 
 //------------------------------------------------------------------------------------------------------------
 
 class FCopyVelocityGridCS : public FGlobalShader
 {
-	DECLARE_SHADER_TYPE(FCopyVelocityGridCS, Global)
+	DECLARE_GLOBAL_SHADER(FCopyVelocityGridCS)
+	SHADER_USE_PARAMETER_STRUCT(FCopyVelocityGridCS, FGlobalShader);
 
-public:
 	static bool ShouldCompilePermutation(const FGlobalShaderPermutationParameters& Parameters)
 	{
 		return RHISupportsComputeShaders(Parameters.Platform);
@@ -737,116 +652,98 @@ public:
 		OutEnvironment.SetDefine(TEXT("THREAD_COUNT"), NIAGARA_HAIR_STRANDS_THREAD_COUNT_VELOCITY);
 	}
 
-	FCopyVelocityGridCS(const ShaderMetaType::CompiledShaderInitializerType& Initializer)
-		: FGlobalShader(Initializer)
-	{
-		GridSize.Bind(Initializer.ParameterMap, TEXT("GridSize"));
-		GridDestinationBuffer.Bind(Initializer.ParameterMap, TEXT("GridDestinationBuffer"));
-		GridCurrentBuffer.Bind(Initializer.ParameterMap, TEXT("GridCurrentBuffer"));
-	}
-
-	FCopyVelocityGridCS()
-	{}
-
-	void SetParameters(FRHICommandList& RHICmdList, FRHIShaderResourceView* InGridCurrentBuffer,
-		FRHIUnorderedAccessView* InGridDestinationBuffer,
-		const FIntVector& InGridSize)
-	{
-		FRHIComputeShader* ShaderRHI = RHICmdList.GetBoundComputeShader();
-
-		SetUAVParameter(RHICmdList, ShaderRHI, GridDestinationBuffer, InGridDestinationBuffer);
-		SetShaderValue(RHICmdList, ShaderRHI, GridSize, InGridSize);
-		SetSRVParameter(RHICmdList, ShaderRHI, GridCurrentBuffer, InGridCurrentBuffer);
-	}
-
-	void UnsetParameters(FRHICommandList& RHICmdList)
-	{
-		FRHIComputeShader* ShaderRHI = RHICmdList.GetBoundComputeShader();
-
-		SetUAVParameter(RHICmdList, ShaderRHI, GridDestinationBuffer, nullptr);
-	}
-
-	LAYOUT_FIELD(FShaderResourceParameter, GridDestinationBuffer);
-	LAYOUT_FIELD(FShaderParameter, GridSize);
-	LAYOUT_FIELD(FShaderResourceParameter, GridCurrentBuffer);
+	BEGIN_SHADER_PARAMETER_STRUCT(FParameters, )
+		SHADER_PARAMETER(FIntVector,					GridSize)
+		SHADER_PARAMETER_RDG_BUFFER_SRV(Buffer<int>,	GridCurrentBuffer)
+		SHADER_PARAMETER_RDG_BUFFER_UAV(RWBuffer<int>,	GridDestinationBuffer)
+	END_SHADER_PARAMETER_STRUCT()
 };
 
 IMPLEMENT_SHADER_TYPE(, FCopyVelocityGridCS, TEXT("/Plugin/Runtime/HairStrands/Private/NiagaraCopyVelocityGrid.usf"), TEXT("MainCS"), SF_Compute);
 
-inline void CopyTexture(FRHICommandList& RHICmdList, ERHIFeatureLevel::Type FeatureLevel, FNDIVelocityGridBuffer* CurrentGridBuffer, FNDIVelocityGridBuffer* DestinationGridBuffer, const FIntVector& InGridSize)
+inline void CopyTexture(FRDGBuilder& GraphBuilder, ERHIFeatureLevel::Type FeatureLevel, FNDIVelocityGridBuffer* CurrentGridBuffer, FNDIVelocityGridBuffer* DestinationGridBuffer, const FIntVector& InGridSize)
 {
-	FRHIUnorderedAccessView* DestinationGridBufferUAV = DestinationGridBuffer->GridDataBuffer.UAV;
-	FRHIShaderResourceView* CurrentGridBufferSRV = CurrentGridBuffer->GridDataBuffer.SRV;
-	FRHIUnorderedAccessView* CurrentGridBufferUAV = CurrentGridBuffer->GridDataBuffer.UAV;
-
-	if (DestinationGridBufferUAV != nullptr && CurrentGridBufferSRV != nullptr && CurrentGridBufferUAV != nullptr)
+	if (DestinationGridBuffer->GridDataBuffer.IsValid() == false || CurrentGridBuffer->GridDataBuffer.IsValid() == false)
 	{
-		TShaderMapRef<FCopyVelocityGridCS> ComputeShader(GetGlobalShaderMap(FeatureLevel));
-		SetComputePipelineState(RHICmdList, ComputeShader.GetComputeShader());
-
-		FRHITransitionInfo Transitions[] = {
-			FRHITransitionInfo(CurrentGridBufferUAV, ERHIAccess::Unknown, ERHIAccess::SRVCompute),
-			FRHITransitionInfo(DestinationGridBufferUAV, ERHIAccess::Unknown, ERHIAccess::UAVCompute)
-		};
-		RHICmdList.Transition(MakeArrayView(Transitions, UE_ARRAY_COUNT(Transitions)));
+		return;
+	}
 
-		const uint32 GroupSize = NIAGARA_HAIR_STRANDS_THREAD_COUNT_VELOCITY;
-		const FIntVector GridSize((InGridSize.X + 1) * DestinationGridBuffer->NumAttributes, InGridSize.Y + 1, InGridSize.Z + 1);
+	TShaderMapRef<FCopyVelocityGridCS> ComputeShader(GetGlobalShaderMap(FeatureLevel));
 
-		const uint32 DispatchCountX = FMath::DivideAndRoundUp((uint32)GridSize.X, GroupSize);
-		const uint32 DispatchCountY = FMath::DivideAndRoundUp((uint32)GridSize.Y, GroupSize);
-		const uint32 DispatchCountZ = FMath::DivideAndRoundUp((uint32)GridSize.Z, GroupSize);
+	const FIntVector GridSize((InGridSize.X + 1) * DestinationGridBuffer->NumAttributes, InGridSize.Y + 1, InGridSize.Z + 1);
+	const uint32 ThreadGroupSize = NIAGARA_HAIR_STRANDS_THREAD_COUNT_VELOCITY;
+	const FIntVector ThreadGroupCount(
+		FMath::DivideAndRoundUp((uint32)GridSize.X, ThreadGroupSize),
+		FMath::DivideAndRoundUp((uint32)GridSize.Y, ThreadGroupSize),
+		FMath::DivideAndRoundUp((uint32)GridSize.Z, ThreadGroupSize)
+	);
 
-		ComputeShader->SetParameters(RHICmdList, CurrentGridBufferSRV, DestinationGridBufferUAV, GridSize);
-		DispatchComputeShader(RHICmdList, ComputeShader.GetShader(), DispatchCountX, DispatchCountY, DispatchCountZ);
-		ComputeShader->UnsetParameters(RHICmdList);
-	}
+	FCopyVelocityGridCS::FParameters* PassParameters = GraphBuilder.AllocParameters<FCopyVelocityGridCS::FParameters>();
+	PassParameters->GridSize = GridSize;
+	PassParameters->GridCurrentBuffer = CurrentGridBuffer->GridDataBuffer.GetOrCreateSRV(GraphBuilder);
+	PassParameters->GridDestinationBuffer = DestinationGridBuffer->GridDataBuffer.GetOrCreateUAV(GraphBuilder);
+
+	FComputeShaderUtils::AddPass(
+		GraphBuilder,
+		RDG_EVENT_NAME("VelocityGrid::CopyTexture"),
+		ERDGPassFlags::Compute,
+		ComputeShader,
+		PassParameters,
+		ThreadGroupCount
+	);
 }
 
-void FNDIVelocityGridProxy::PreStage(FRHICommandList& RHICmdList, const FNiagaraDataInterfaceStageArgs& Context)
+void FNDIVelocityGridProxy::PreStage(const FNDIGpuComputePreStageContext& Context)
 {
-	FNDIVelocityGridData* ProxyData =
-		SystemInstancesToProxyData.Find(Context.SystemInstanceID);
-
+	FNDIVelocityGridData* ProxyData = SystemInstancesToProxyData.Find(Context.GetSystemInstanceID());
 	if (ProxyData != nullptr)
 	{
-		if (Context.SimStageData->bFirstStage)
+		if (Context.GetSimStageData().bFirstStage)
 		{
-			ClearTexture(RHICmdList, ProxyData->DestinationGridBuffer, ProxyData->GridSize);
+			ClearTexture(Context.GetGraphBuilder(), ProxyData->DestinationGridBuffer, ProxyData->GridSize);
 		}
-
-		FRHITransitionInfo Transitions[] = {
-			// FIXME: what's the source state for these?
-			FRHITransitionInfo(ProxyData->CurrentGridBuffer->GridDataBuffer.UAV, ERHIAccess::Unknown, ERHIAccess::SRVCompute),
-			FRHITransitionInfo(ProxyData->DestinationGridBuffer->GridDataBuffer.UAV, ERHIAccess::Unknown, ERHIAccess::UAVCompute)
-		};
-		RHICmdList.Transition(MakeArrayView(Transitions, UE_ARRAY_COUNT(Transitions)));
 	}
 }
 
-void FNDIVelocityGridProxy::PostStage(FRHICommandList& RHICmdList, const FNiagaraDataInterfaceStageArgs& Context)
+void FNDIVelocityGridProxy::PostStage(const FNDIGpuComputePostStageContext& Context)
 {
-	FNDIVelocityGridData* ProxyData =
-		SystemInstancesToProxyData.Find(Context.SystemInstanceID);
-
+	FNDIVelocityGridData* ProxyData = SystemInstancesToProxyData.Find(Context.GetSystemInstanceID());
 	if (ProxyData != nullptr)
 	{
 		//ProxyData->Swap();
-		CopyTexture(RHICmdList, Context.ComputeDispatchInterface->GetFeatureLevel(), ProxyData->DestinationGridBuffer, ProxyData->CurrentGridBuffer,  ProxyData->GridSize);
+		CopyTexture(Context.GetGraphBuilder(), Context.GetComputeDispatchInterface().GetFeatureLevel(), ProxyData->DestinationGridBuffer, ProxyData->CurrentGridBuffer, ProxyData->GridSize);
 		//FRHICopyTextureInfo CopyInfo;
 		//RHICmdList.CopyTexture(ProxyData->DestinationGridBuffer->GridDataBuffer.Buffer,
 		//	ProxyData->CurrentGridBuffer->GridDataBuffer.Buffer, CopyInfo);
 	}
 }
 
-void FNDIVelocityGridProxy::ResetData(FRHICommandList& RHICmdList, const FNiagaraDataInterfaceArgs& Context)
+void FNDIVelocityGridProxy::PostSimulate(const FNDIGpuComputePostSimulateContext& Context)
+{
+	if (Context.IsFinalPostSimulate())
+	{
+		if (FNDIVelocityGridData* ProxyData = SystemInstancesToProxyData.Find(Context.GetSystemInstanceID()))
+		{
+			if (ProxyData->CurrentGridBuffer->GridDataBuffer.IsValid())
+			{
+				ProxyData->CurrentGridBuffer->GridDataBuffer.EndGraphUsage();
+			}
+			if (ProxyData->DestinationGridBuffer->GridDataBuffer.IsValid())
+			{
+				ProxyData->DestinationGridBuffer->GridDataBuffer.EndGraphUsage();
+			}
+		}
+	}
+}
+
+void FNDIVelocityGridProxy::ResetData(const FNDIGpuComputeResetContext& Context)
 {
-	FNDIVelocityGridData* ProxyData = SystemInstancesToProxyData.Find(Context.SystemInstanceID);
+	FNDIVelocityGridData* ProxyData = SystemInstancesToProxyData.Find(Context.GetSystemInstanceID());
 
 	if (ProxyData != nullptr && ProxyData->DestinationGridBuffer != nullptr && ProxyData->CurrentGridBuffer != nullptr)
 	{
-		ClearTexture(RHICmdList, ProxyData->DestinationGridBuffer, ProxyData->GridSize);
-		ClearTexture(RHICmdList, ProxyData->CurrentGridBuffer, ProxyData->GridSize);
+		ClearTexture(Context.GetGraphBuilder(), ProxyData->DestinationGridBuffer, ProxyData->GridSize);
+		ClearTexture(Context.GetGraphBuilder(), ProxyData->CurrentGridBuffer, ProxyData->GridSize);
 	}
 }
 
diff --git a/Engine/Plugins/Runtime/HairStrands/Source/HairStrandsCore/Public/Niagara/NiagaraDataInterfaceHairStrands.h b/Engine/Plugins/Runtime/HairStrands/Source/HairStrandsCore/Public/Niagara/NiagaraDataInterfaceHairStrands.h
--- a/Engine/Plugins/Runtime/HairStrands/Source/HairStrandsCore/Public/Niagara/NiagaraDataInterfaceHairStrands.h
+++ b/Engine/Plugins/Runtime/HairStrands/Source/HairStrandsCore/Public/Niagara/NiagaraDataInterfaceHairStrands.h
@@ -5,6 +5,7 @@
 #include "NiagaraDataInterface.h"
 #include "NiagaraDataInterfaceRW.h"
 #include "NiagaraCommon.h"
+#include "NiagaraRenderGraphUtils.h"
 #include "VectorVM.h"
 #include "GroomAsset.h"
 #include "GroomActor.h"
@@ -50,19 +51,19 @@ struct FNDIHairStrandsBuffer : public FRenderResource
 	virtual FString GetFriendlyName() const override { return TEXT("FNDIHairStrandsBuffer"); }
 
 	/** Strand curves point offset buffer */
-	FRWBuffer CurvesOffsetsBuffer;
+	FReadBuffer CurvesOffsetsBuffer;
 
-	/** Deformed position buffer in case no ressource are there */
-	FRWBuffer DeformedPositionBuffer;
+	/** Deformed position buffer in case no resource are there */
+	TRefCountPtr<FRDGPooledBuffer> DeformedPositionBuffer;
 
 	/** Bounding Box Buffer*/
-	FRWBuffer BoundingBoxBuffer;
+	FNiagaraPooledRWBuffer BoundingBoxBuffer;
 
 	/** Params scale buffer */
-	FRWBuffer ParamsScaleBuffer;
+	FReadBuffer ParamsScaleBuffer;
 
 	/** Points curve index for fast query */
-	FRWBuffer PointsCurveBuffer;
+	FReadBuffer PointsCurveBuffer;
 
 	/** The strand asset resource from which to sample */
 	const FHairStrandsRestResource* SourceRestResources;
@@ -385,9 +386,6 @@ class HAIRSTRANDSCORE_API UNiagaraDataInterfaceHairStrands : public UNiagaraData
 	GENERATED_UCLASS_BODY()
 
 public:
-
-	DECLARE_NIAGARA_DI_PARAMETER();
-
 	/** Hair Strands Asset used to sample from when not overridden by a source actor from the scene. Also useful for previewing in the editor. */
 	UPROPERTY(EditAnywhere, Category = "Source")
 	TObjectPtr<UGroomAsset> DefaultSource;
@@ -423,6 +421,10 @@ public:
 	virtual bool GetFunctionHLSL(const FNiagaraDataInterfaceGPUParamInfo& ParamInfo, const FNiagaraDataInterfaceGeneratedFunction& FunctionInfo, int FunctionInstanceIndex, FString& OutHLSL) override;
 	virtual bool AppendCompileHash(FNiagaraCompileHashVisitor* InVisitor) const override;
 #endif
+	virtual bool UseLegacyShaderBindings() const override { return false; }
+	virtual void BuildShaderParameters(FNiagaraShaderParametersBuilder& ShaderParametersBuilder) const override;
+	virtual void SetShaderParameters(const FNiagaraDataInterfaceSetShaderParametersContext& Context) const override;
+
 	virtual void ProvidePerInstanceDataForRenderThread(void* DataForRenderThread, void* PerInstanceData, const FNiagaraSystemInstanceID& SystemInstance) override;
 
 	/** Update the source component */
@@ -661,114 +663,6 @@ public:
 	/** Get the sample state given an index */
 	void GetSampleState(FVectorVMExternalFunctionContext& Context);
 
-	/** Name of the world transform */
-	static const FString WorldTransformName;
-
-	/** Name of the bounding box offsets*/
-	static const FString BoundingBoxOffsetsName;
-
-	/** Name of the world inverse transform */
-	static const FString WorldInverseName;
-
-	/** Name of the world rotation */
-	static const FString WorldRotationName;
-
-	/** Name of the bone transform */
-	static const FString BoneTransformName;
-
-	/** Name of the bone inverse transform */
-	static const FString BoneInverseName;
-
-	/** Name of the world rotation */
-	static const FString BoneRotationName;
-
-	/** Name of the bone linear velocity */
-	static const FString BoneLinearVelocityName;
-
-	/** Name of the bone linear acceleration */
-	static const FString BoneLinearAccelerationName;
-
-	/** Name of the bone angular velocity */
-	static const FString BoneAngularVelocityName;
-
-	/** Name of the bone angular acceleration */
-	static const FString BoneAngularAccelerationName;
-
-	/** Name of the number of strands */
-	static const FString NumStrandsName;
-
-	/** Name of the strand size */
-	static const FString StrandSizeName;
-
-	/** Name of the points positions buffer */
-	static const FString DeformedPositionBufferName;
-
-	/** Name of the curves offsets buffer */
-	static const FString CurvesOffsetsBufferName;
-
-	/** Name of bounding box buffer */
-	static const FString BoundingBoxBufferName;
-
-	/** Name of the nodes positions buffer */
-	static const FString RestPositionBufferName;
-
-	/** Param to check if the roots have been attached to the skin */
-	static const FString InterpolationModeName;
-
-	/** Param to check if we need to update the rest pose */
-	static const FString RestUpdateName;
-
-	/** Param to check if the simulation is going to be run in local space */
-	static const FString LocalSimulationName;
-
-	/** boolean to check if we need to rest the simulation*/
-	static const FString ResetSimulationName;
-
-	/** Rest center of all the roots */
-	static const FString RestRootOffsetName;
-
-	/** Rest position of the triangle vertex A */
-	static const FString RestTrianglePositionAName;
-
-	/** Rest position of the triangle vertex B */
-	static const FString RestTrianglePositionBName;
-
-	/** Rest position of the triangle vertex C */
-	static const FString RestTrianglePositionCName;
-
-	/** Deformed center of all the roots */
-	static const FString DeformedRootOffsetName;
-
-	/** Deformed position of the triangle vertex A */
-	static const FString DeformedTrianglePositionAName;
-
-	/** Deformed position of the triangle vertex A */
-	static const FString DeformedTrianglePositionBName;
-
-	/** Deformed position of the triangle vertex A */
-	static const FString DeformedTrianglePositionCName;
-
-	/** Root barycentric coordinates */
-	static const FString RootBarycentricCoordinatesName;
-
-	/** Rest center of all the position */
-	static const FString RestPositionOffsetName;
-
-	/** Deformed center of all the position */
-	static const FString DeformedPositionOffsetName;
-
-	/** Number of samples for rbf interpolation */
-	static const FString SampleCountName;
-
-	/** Rbf sample weights */
-	static const FString MeshSampleWeightsName;
-
-	/** Rbf Sample rest positions */
-	static const FString RestSamplePositionsName;
-
-	/** Params scale buffer */
-	static const FString ParamsScaleBufferName;
-
 protected:
 	/** Copy one niagara DI to this */
 	virtual bool CopyToInternal(UNiagaraDataInterface* Destination) const override;
@@ -790,10 +684,10 @@ struct FNDIHairStrandsProxy : public FNiagaraDataInterfaceProxy
 	void DestroyPerInstanceData(const FNiagaraSystemInstanceID& SystemInstance);
 
 	/** Launch all pre stage functions */
-	virtual void PreStage(FRHICommandList& RHICmdList, const FNiagaraDataInterfaceStageArgs& Context) override;
+	virtual void PreStage(const FNDIGpuComputePreStageContext& Context) override;
 
 	/** MGPU buffer copy after simulation*/
-	virtual void PostSimulate(FRHICommandList& RHICmdList, const FNiagaraDataInterfaceArgs& Context) override;
+	virtual void PostSimulate(const FNDIGpuComputePostSimulateContext& Context) override;
 
 	/** List of proxy data for each system instances*/
 	TMap<FNiagaraSystemInstanceID, FNDIHairStrandsData> SystemInstancesToProxyData;
diff --git a/Engine/Plugins/Runtime/HairStrands/Source/HairStrandsCore/Public/Niagara/NiagaraDataInterfacePressureGrid.h b/Engine/Plugins/Runtime/HairStrands/Source/HairStrandsCore/Public/Niagara/NiagaraDataInterfacePressureGrid.h
--- a/Engine/Plugins/Runtime/HairStrands/Source/HairStrandsCore/Public/Niagara/NiagaraDataInterfacePressureGrid.h
+++ b/Engine/Plugins/Runtime/HairStrands/Source/HairStrandsCore/Public/Niagara/NiagaraDataInterfacePressureGrid.h
@@ -16,9 +16,6 @@ class HAIRSTRANDSCORE_API UNiagaraDataInterfacePressureGrid : public UNiagaraDat
 	GENERATED_UCLASS_BODY()
 
 public:
-
-	DECLARE_NIAGARA_DI_PARAMETER();
-
 	/** UNiagaraDataInterface Interface */
 	virtual void GetFunctions(TArray<FNiagaraFunctionSignature>& OutFunctions) override;
 	virtual void GetVMExternalFunction(const FVMExternalFunctionBindingInfo& BindingInfo, void* InstanceData, FVMExternalFunction &OutFunc) override;
@@ -65,6 +62,6 @@ public:
 struct FNDIPressureGridProxy : public FNDIVelocityGridProxy
 {
 	/** Launch all pre stage functions */
-	virtual void PreStage(FRHICommandList& RHICmdList, const FNiagaraDataInterfaceStageArgs& Context) override;
+	virtual void PreStage(const FNDIGpuComputePreStageContext& Context) override;
 };
 
diff --git a/Engine/Plugins/Runtime/HairStrands/Source/HairStrandsCore/Public/Niagara/NiagaraDataInterfaceVelocityGrid.h b/Engine/Plugins/Runtime/HairStrands/Source/HairStrandsCore/Public/Niagara/NiagaraDataInterfaceVelocityGrid.h
--- a/Engine/Plugins/Runtime/HairStrands/Source/HairStrandsCore/Public/Niagara/NiagaraDataInterfaceVelocityGrid.h
+++ b/Engine/Plugins/Runtime/HairStrands/Source/HairStrandsCore/Public/Niagara/NiagaraDataInterfaceVelocityGrid.h
@@ -4,6 +4,7 @@
 
 #include "NiagaraDataInterface.h"
 #include "NiagaraDataInterfaceRW.h"
+#include "NiagaraRenderGraphUtils.h"
 #include "NiagaraCommon.h"
 #include "VectorVM.h"
 #include "NiagaraDataInterfaceVelocityGrid.generated.h"
@@ -25,7 +26,7 @@ struct FNDIVelocityGridBuffer : public FRenderResource
 	virtual FString GetFriendlyName() const override { return TEXT("FNDIVelocityGridBuffer"); }
 
 	/** Grid data texture */
-	FRWBuffer GridDataBuffer;
+	FNiagaraPooledRWBuffer GridDataBuffer;
 
 	/** Grid size that will be used for the collision*/
 	FIntVector GridSize;
@@ -71,39 +72,6 @@ struct FNDIVelocityGridData
 	FNDIVelocityGridBuffer* DestinationGridBuffer;
 };
 
-/** Data Interface paramaters name */
-struct FNDIVelocityGridParametersName
-{
-	FNDIVelocityGridParametersName(const FString& Suffix);
-
-	FString GridCurrentBufferName;
-	FString GridDestinationBufferName;
-
-	FString GridSizeName;
-	FString WorldTransformName;
-	FString WorldInverseName;
-};
-
-struct FNDIVelocityGridParametersCS : public FNiagaraDataInterfaceParametersCS
-{
-	DECLARE_TYPE_LAYOUT(FNDIVelocityGridParametersCS, NonVirtual);
-
-	void Bind(const FNiagaraDataInterfaceGPUParamInfo& ParameterInfo, const class FShaderParameterMap& ParameterMap);
-
-	void Set(FRHICommandList& RHICmdList, const FNiagaraDataInterfaceSetArgs& Context) const;
-
-	void Unset(FRHICommandList& RHICmdList, const FNiagaraDataInterfaceSetArgs& Context) const;
-
-private:
-
-	LAYOUT_FIELD(FShaderResourceParameter, GridCurrentBuffer);
-	LAYOUT_FIELD(FShaderResourceParameter, GridDestinationBuffer);
-
-	LAYOUT_FIELD(FShaderParameter, GridSize);
-	LAYOUT_FIELD(FShaderParameter, WorldTransform);
-	LAYOUT_FIELD(FShaderParameter, WorldInverse);
-};
-
 /** Data Interface for the strand base */
 UCLASS(EditInlineNew, Category = "Grid", meta = (DisplayName = "Velocity Grid"))
 class HAIRSTRANDSCORE_API UNiagaraDataInterfaceVelocityGrid : public UNiagaraDataInterfaceRWBase
@@ -111,9 +79,6 @@ class HAIRSTRANDSCORE_API UNiagaraDataInterfaceVelocityGrid : public UNiagaraDat
 	GENERATED_UCLASS_BODY()
 
 public:
-
-	DECLARE_NIAGARA_DI_PARAMETER();
-
 	/** Grid size along the X axis. */
 	UPROPERTY(EditAnywhere, Category = "Spawn")
 	FIntVector GridSize;
@@ -141,6 +106,10 @@ public:
 	virtual void GetParameterDefinitionHLSL(const FNiagaraDataInterfaceGPUParamInfo& ParamInfo, FString& OutHLSL) override;
 	virtual bool GetFunctionHLSL(const FNiagaraDataInterfaceGPUParamInfo& ParamInfo, const FNiagaraDataInterfaceGeneratedFunction& FunctionInfo, int FunctionInstanceIndex, FString& OutHLSL) override;
 #endif
+	virtual bool UseLegacyShaderBindings() const override { return false; }
+	virtual void BuildShaderParameters(FNiagaraShaderParametersBuilder& ShaderParametersBuilder) const override;
+	virtual void SetShaderParameters(const FNiagaraDataInterfaceSetShaderParametersContext& Context) const override;
+
 	virtual void ProvidePerInstanceDataForRenderThread(void* DataForRenderThread, void* PerInstanceData, const FNiagaraSystemInstanceID& SystemInstance) override;
 
 	/** Build the velocity field */
@@ -158,21 +127,6 @@ public:
 	/** Set the grid dimension */
 	void SetGridDimension(FVectorVMExternalFunctionContext& Context);
 
-	/** Name of the grid current buffer */
-	static const FString GridCurrentBufferName;
-
-	/** Name of the grid X velocity buffer */
-	static const FString GridDestinationBufferName;
-
-	/** Name of the grid size */
-	static const FString GridSizeName;
-
-	/** Name of the world transform  */
-	static const FString WorldTransformName;
-
-	/** Name of the World transform inverse */
-	static const FString WorldInverseName;
-
 protected:
 	/** Copy one niagara DI to this */
 	virtual bool CopyToInternal(UNiagaraDataInterface* Destination) const override;
@@ -188,13 +142,16 @@ struct FNDIVelocityGridProxy : public FNiagaraDataInterfaceProxyRW
 	virtual void ConsumePerInstanceDataFromGameThread(void* PerInstanceData, const FNiagaraSystemInstanceID& Instance) override;
 
 	/** Launch all pre stage functions */
-	virtual void PreStage(FRHICommandList& RHICmdList, const FNiagaraDataInterfaceStageArgs& Context) override;
+	virtual void PreStage(const FNDIGpuComputePreStageContext& Context) override;
 
 	/** Launch all post stage functions */
-	virtual void PostStage(FRHICommandList& RHICmdList, const FNiagaraDataInterfaceStageArgs& Context) override;
+	virtual void PostStage(const FNDIGpuComputePostStageContext& Context) override;
+
+	/** Called at the end of each simulate tick. */
+	virtual void PostSimulate(const FNDIGpuComputePostSimulateContext& Context) override;
 
 	/** Reset the buffers  */
-	virtual void ResetData(FRHICommandList& RHICmdList, const FNiagaraDataInterfaceArgs& Context) override;
+	virtual void ResetData(const FNDIGpuComputeResetContext& Context) override;
 
 	// Get the element count for this instance
 	virtual FIntVector GetElementCount(FNiagaraSystemInstanceID SystemInstanceID) const override;
diff --git a/Engine/Source/Runtime/Core/Private/UObject/UE5MainStreamObjectVersion.cpp b/Engine/Source/Runtime/Core/Private/UObject/UE5MainStreamObjectVersion.cpp
--- a/Engine/Source/Runtime/Core/Private/UObject/UE5MainStreamObjectVersion.cpp
+++ b/Engine/Source/Runtime/Core/Private/UObject/UE5MainStreamObjectVersion.cpp
@@ -14,7 +14,7 @@ TMap<FGuid, FGuid> FUE5MainStreamObjectVersion::GetSystemGuids()
 	SystemGuids.Add(DevGuids.MATERIALSHADERMAP_DERIVEDDATA_VER, FGuid("65C01D817C9A4EEFAE9E988D41A1F3DD"));
 	SystemGuids.Add(DevGuids.NANITE_DERIVEDDATA_VER, FGuid("7A3EA93546624597BEBE2ADFDE3A7628"));
 	SystemGuids.Add(DevGuids.NIAGARASHADERMAP_DERIVEDDATA_VER, FGuid("8A37C45D24F2423CBE5F8F371DE33575"));
-	SystemGuids.Add(DevGuids.Niagara_LatestScriptCompileVersion, FGuid("C4BD9AE35DC59C4C985DD5316ACC8CC9"));
+	SystemGuids.Add(DevGuids.Niagara_LatestScriptCompileVersion, FGuid("92519728F82ECC4FB1CD175F35D6442D"));
 	SystemGuids.Add(DevGuids.SkeletalMeshDerivedDataVersion, FGuid("FB7B2FB546DE49B2A2F88B7F69763CEE"));
 	SystemGuids.Add(DevGuids.STATICMESH_DERIVEDDATA_VER, FGuid("98F7D79A5811013825E75BBCC41ED3E9"));
 	SystemGuids.Add(DevGuids.POSESEARCHDB_DERIVEDDATA_VER, FGuid("4E595C2AC5E947D6BA9ABC874353E5BC"));
