Change groom mesh/hair debug to use shader print instead of custom pass.

diff --git a/Engine/Plugins/Runtime/HairStrands/Source/HairStrandsCore/Private/GroomDebug.cpp b/Engine/Plugins/Runtime/HairStrands/Source/HairStrandsCore/Private/GroomDebug.cpp
--- a/Engine/Plugins/Runtime/HairStrands/Source/HairStrandsCore/Private/GroomDebug.cpp
+++ b/Engine/Plugins/Runtime/HairStrands/Source/HairStrandsCore/Private/GroomDebug.cpp
@@ -143,65 +143,31 @@ static void GetGroomInterpolationData(
 
 ///////////////////////////////////////////////////////////////////////////////////////////////////
 
-BEGIN_SHADER_PARAMETER_STRUCT(FHairProjectionMeshDebugParameters, )
-	SHADER_PARAMETER(FMatrix44f, LocalToWorld)
-	SHADER_PARAMETER(uint32, VertexOffset)
-	SHADER_PARAMETER(uint32, IndexOffset)
-	SHADER_PARAMETER(uint32, MaxIndexCount)
-	SHADER_PARAMETER(uint32, MaxVertexCount)
-	SHADER_PARAMETER(uint32, MeshUVsChannelOffset)
-	SHADER_PARAMETER(uint32, MeshUVsChannelCount)
-	SHADER_PARAMETER(uint32, bOutputInUVsSpace)
-	SHADER_PARAMETER(uint32, MeshType)
-	SHADER_PARAMETER(uint32, SectionIndex)
-	SHADER_PARAMETER(FVector2f, OutputResolution)
-	SHADER_PARAMETER_SRV(StructuredBuffer, InputIndexBuffer)
-	SHADER_PARAMETER_SRV(StructuredBuffer, InputVertexPositionBuffer)
-	SHADER_PARAMETER_SRV(StructuredBuffer, InputVertexUVsBuffer)
-	SHADER_PARAMETER_STRUCT_REF(FViewUniformShaderParameters, ViewUniformBuffer)
-	RENDER_TARGET_BINDING_SLOTS()
-END_SHADER_PARAMETER_STRUCT()
-
-class FHairProjectionMeshDebug : public FGlobalShader
+class FHairProjectionMeshDebugCS : public FGlobalShader
 {
-public:
-	static bool ShouldCompilePermutation(const FGlobalShaderPermutationParameters& Parameters)
-	{
-		return IsHairStrandsSupported(EHairStrandsShaderType::Tool, Parameters.Platform);
-	}
-
-
-	FHairProjectionMeshDebug() = default;
-	FHairProjectionMeshDebug(const CompiledShaderInitializerType& Initializer) : FGlobalShader(Initializer) {}
-};
+	DECLARE_GLOBAL_SHADER(FHairProjectionMeshDebugCS);
+	SHADER_USE_PARAMETER_STRUCT(FHairProjectionMeshDebugCS, FGlobalShader);
 
-
-class FHairProjectionMeshDebugVS : public FHairProjectionMeshDebug
-{
-public:
 	class FInputType : SHADER_PERMUTATION_INT("PERMUTATION_INPUT_TYPE", 2);
 	using FPermutationDomain = TShaderPermutationDomain<FInputType>;
 
-	DECLARE_GLOBAL_SHADER(FHairProjectionMeshDebugVS);
-	SHADER_USE_PARAMETER_STRUCT(FHairProjectionMeshDebugVS, FHairProjectionMeshDebug);
-
-	BEGIN_SHADER_PARAMETER_STRUCT(FParameters, )
-		SHADER_PARAMETER_STRUCT_INCLUDE(FHairProjectionMeshDebugParameters, Pass)
-	END_SHADER_PARAMETER_STRUCT()
-
-	static bool ShouldCompilePermutation(const FGlobalShaderPermutationParameters& Parameters)
-	{
-		return IsHairStrandsSupported(EHairStrandsShaderType::Tool, Parameters.Platform);
-	}
-};
-
-class FHairProjectionMeshDebugPS : public FHairProjectionMeshDebug
-{
-	DECLARE_GLOBAL_SHADER(FHairProjectionMeshDebugPS);
-	SHADER_USE_PARAMETER_STRUCT(FHairProjectionMeshDebugPS, FHairProjectionMeshDebug);
-
 	BEGIN_SHADER_PARAMETER_STRUCT(FParameters, )
-		SHADER_PARAMETER_STRUCT_INCLUDE(FHairProjectionMeshDebugParameters, Pass)
+		SHADER_PARAMETER(FMatrix44f, LocalToWorld)
+		SHADER_PARAMETER(uint32, VertexOffset)
+		SHADER_PARAMETER(uint32, IndexOffset)
+		SHADER_PARAMETER(uint32, MaxIndexCount)
+		SHADER_PARAMETER(uint32, MaxVertexCount)
+		SHADER_PARAMETER(uint32, MeshUVsChannelOffset)
+		SHADER_PARAMETER(uint32, MeshUVsChannelCount)
+		SHADER_PARAMETER(uint32, bOutputInUVsSpace)
+		SHADER_PARAMETER(uint32, MeshType)
+		SHADER_PARAMETER(uint32, SectionIndex)
+		SHADER_PARAMETER(FVector2f, OutputResolution)
+		SHADER_PARAMETER_SRV(StructuredBuffer, InputIndexBuffer)
+		SHADER_PARAMETER_SRV(StructuredBuffer, InputVertexPositionBuffer)
+		SHADER_PARAMETER_SRV(StructuredBuffer, InputVertexUVsBuffer)
+		SHADER_PARAMETER_STRUCT_REF(FViewUniformShaderParameters, ViewUniformBuffer)
+		SHADER_PARAMETER_STRUCT_INCLUDE(ShaderPrint::FShaderParameters, ShaderPrintUniformBuffer)
 	END_SHADER_PARAMETER_STRUCT()
 
 	static bool ShouldCompilePermutation(const FGlobalShaderPermutationParameters& Parameters)
@@ -210,30 +176,37 @@ class FHairProjectionMeshDebugPS : public FHairProjectionMeshDebug
 	}
 };
 
-IMPLEMENT_GLOBAL_SHADER(FHairProjectionMeshDebugVS, "/Engine/Private/HairStrands/HairStrandsMeshProjectionMeshDebug.usf", "MainVS", SF_Vertex);
-IMPLEMENT_GLOBAL_SHADER(FHairProjectionMeshDebugPS, "/Engine/Private/HairStrands/HairStrandsMeshProjectionMeshDebug.usf", "MainPS", SF_Pixel);
+IMPLEMENT_GLOBAL_SHADER(FHairProjectionMeshDebugCS, "/Engine/Private/HairStrands/HairStrandsMeshProjectionMeshDebug.usf", "MainCS", SF_Compute);
 
 static void AddDebugProjectionMeshPass(
 	FRDGBuilder& GraphBuilder,
 	FGlobalShaderMap* ShaderMap,
+	const FShaderPrintData* ShaderPrintData,
 	const FIntRect Viewport,
 	const TUniformBufferRef<FViewUniformShaderParameters>& ViewUniformBuffer,
 	const EHairStrandsProjectionMeshType MeshType,
-	const bool bClearDepth,
-	FHairStrandsProjectionMeshData::Section& MeshSectionData,
-	FRDGTextureRef& ColorTexture,
-	FRDGTextureRef& DepthTexture)
+	FHairStrandsProjectionMeshData::Section& MeshSectionData)
 {
-	const EPrimitiveType PrimitiveType = PT_TriangleList;
+	// Force ShaderPrint on.
+	ShaderPrint::SetEnabled(true);
+
 	const bool bHasIndexBuffer = MeshSectionData.IndexBuffer != nullptr;
 	const uint32 PrimitiveCount = MeshSectionData.NumPrimitives;
 
 	if (!MeshSectionData.PositionBuffer || PrimitiveCount == 0)
 		return;
 
+	if (ShaderPrintData == nullptr || !ShaderPrint::IsEnabled(*ShaderPrintData))
+	{
+		return;
+	}
+
+	ShaderPrint::RequestSpaceForLines(PrimitiveCount);
+	ShaderPrint::RequestSpaceForTriangles(PrimitiveCount);
+
 	const FIntPoint Resolution(Viewport.Width(), Viewport.Height());
 
-	FHairProjectionMeshDebugParameters* Parameters = GraphBuilder.AllocParameters<FHairProjectionMeshDebugParameters>();
+	FHairProjectionMeshDebugCS::FParameters* Parameters = GraphBuilder.AllocParameters<FHairProjectionMeshDebugCS::FParameters>();
 	Parameters->LocalToWorld = FMatrix44f(MeshSectionData.LocalToWorld.ToMatrixWithScale());		// LWC_TODO: Precision loss
 	Parameters->OutputResolution = Resolution;
 	Parameters->MeshType = uint32(MeshType);
@@ -249,120 +222,61 @@ static void AddDebugProjectionMeshPass(
 	Parameters->InputVertexUVsBuffer = MeshSectionData.UVsBuffer;
 	Parameters->SectionIndex = MeshSectionData.SectionIndex;
 	Parameters->ViewUniformBuffer = ViewUniformBuffer;
-	Parameters->RenderTargets[0] = FRenderTargetBinding(ColorTexture, ERenderTargetLoadAction::ELoad, 0);
-	Parameters->RenderTargets.DepthStencil = FDepthStencilBinding(DepthTexture, bClearDepth ? ERenderTargetLoadAction::EClear : ERenderTargetLoadAction::ELoad, ERenderTargetLoadAction::ENoAction, FExclusiveDepthStencil::DepthWrite_StencilNop);
-
-	FHairProjectionMeshDebugVS::FPermutationDomain PermutationVector;
-	PermutationVector.Set<FHairProjectionMeshDebugVS::FInputType>(bHasIndexBuffer ? 1 : 0);
-
-	TShaderMapRef<FHairProjectionMeshDebugVS> VertexShader(ShaderMap, PermutationVector);
-	TShaderMapRef<FHairProjectionMeshDebugPS> PixelShader(ShaderMap);
+	ShaderPrint::SetParameters(GraphBuilder, *ShaderPrintData, Parameters->ShaderPrintUniformBuffer);
 
-	FHairProjectionMeshDebugVS::FParameters VSParameters;
-	VSParameters.Pass = *Parameters;
-	FHairProjectionMeshDebugPS::FParameters PSParameters;
-	PSParameters.Pass = *Parameters;
+	FHairProjectionMeshDebugCS::FPermutationDomain PermutationVector;
+	PermutationVector.Set<FHairProjectionMeshDebugCS::FInputType>(bHasIndexBuffer ? 1 : 0);
+	TShaderMapRef<FHairProjectionMeshDebugCS> ComputeShader(ShaderMap, PermutationVector);
+	ClearUnusedGraphResources(ComputeShader, Parameters);
 
-	GraphBuilder.AddPass(
-		RDG_EVENT_NAME("HairStrands::MeshProjectionMeshDebug"),
+	FComputeShaderUtils::AddPass(
+		GraphBuilder,
+		RDG_EVENT_NAME("HairStrands::MeshProjectionDebug(Mesh)"),
+		ComputeShader,
 		Parameters,
-		ERDGPassFlags::Raster,
-		[Parameters, VSParameters, PSParameters, VertexShader, PixelShader, Viewport, Resolution, PrimitiveCount, PrimitiveType](FRHICommandList& RHICmdList)
-		{
-
-			RHICmdList.SetViewport(
-				Viewport.Min.X,
-				Viewport.Min.Y,
-				0.0f,
-				Viewport.Max.X,
-				Viewport.Max.Y,
-				1.0f);
-
-			// Apply additive blending pipeline state.
-			FGraphicsPipelineStateInitializer GraphicsPSOInit;
-			RHICmdList.ApplyCachedRenderTargets(GraphicsPSOInit);
-			GraphicsPSOInit.BlendState = TStaticBlendState<CW_RGBA, BO_Add, BF_One, BF_Zero, BO_Add, BF_One, BF_Zero>::GetRHI();
-			GraphicsPSOInit.RasterizerState = TStaticRasterizerState<FM_Wireframe>::GetRHI();
-			GraphicsPSOInit.DepthStencilState = TStaticDepthStencilState<true, CF_DepthNearOrEqual>::GetRHI();
-			GraphicsPSOInit.BoundShaderState.VertexDeclarationRHI = GEmptyVertexDeclaration.VertexDeclarationRHI;
-			GraphicsPSOInit.BoundShaderState.VertexShaderRHI = VertexShader.GetVertexShader();
-			GraphicsPSOInit.BoundShaderState.PixelShaderRHI = PixelShader.GetPixelShader();
-			GraphicsPSOInit.PrimitiveType = PrimitiveType;
-			SetGraphicsPipelineState(RHICmdList, GraphicsPSOInit, 0);
-
-			SetShaderParameters(RHICmdList, VertexShader, VertexShader.GetVertexShader(), VSParameters);
-			SetShaderParameters(RHICmdList, PixelShader, PixelShader.GetPixelShader(), PSParameters);
-
-			// Emit an instanced quad draw call on the order of the number of pixels on the screen.	
-			RHICmdList.DrawPrimitive(0, PrimitiveCount, 1);
-		});
+		FIntVector(FMath::DivideAndRoundUp(PrimitiveCount, 256u), 1, 1));
 }
 
 ///////////////////////////////////////////////////////////////////////////////////////////////////
 
-BEGIN_SHADER_PARAMETER_STRUCT(FHairProjectionHairDebugParameters, )
-	SHADER_PARAMETER(FVector2f, OutputResolution)
-	SHADER_PARAMETER(uint32, MaxRootCount)
-	SHADER_PARAMETER(uint32, DeformedFrameEnable)
-	SHADER_PARAMETER(FMatrix44f, RootLocalToWorld)
+class FHairProjectionHairDebugCS : public FGlobalShader
+{
+	DECLARE_GLOBAL_SHADER(FHairProjectionHairDebugCS);
+	SHADER_USE_PARAMETER_STRUCT(FHairProjectionHairDebugCS, FGlobalShader);
 
-	SHADER_PARAMETER_RDG_BUFFER_SRV(StructuredBuffer, RestPosition0Buffer)
-	SHADER_PARAMETER_RDG_BUFFER_SRV(StructuredBuffer, RestPosition1Buffer)
-	SHADER_PARAMETER_RDG_BUFFER_SRV(StructuredBuffer, RestPosition2Buffer)
+	class FInputType : SHADER_PERMUTATION_INT("PERMUTATION_INPUT_TYPE", 3);
+	using FPermutationDomain = TShaderPermutationDomain<FInputType>;
 
-	SHADER_PARAMETER_RDG_BUFFER_SRV(StructuredBuffer, DeformedPosition0Buffer)
-	SHADER_PARAMETER_RDG_BUFFER_SRV(StructuredBuffer, DeformedPosition1Buffer)
-	SHADER_PARAMETER_RDG_BUFFER_SRV(StructuredBuffer, DeformedPosition2Buffer)
+	BEGIN_SHADER_PARAMETER_STRUCT(FParameters, )
+		SHADER_PARAMETER(FVector2f, OutputResolution)
+		SHADER_PARAMETER(uint32, MaxRootCount)
+		SHADER_PARAMETER(uint32, DeformedFrameEnable)
+		SHADER_PARAMETER(FMatrix44f, RootLocalToWorld)
 
-	SHADER_PARAMETER_RDG_BUFFER_SRV(StructuredBuffer, RestSamplePositionsBuffer)
-	SHADER_PARAMETER_RDG_BUFFER_SRV(StructuredBuffer, DeformedSamplePositionsBuffer)
+		SHADER_PARAMETER_RDG_BUFFER_SRV(StructuredBuffer, RestPosition0Buffer)
+		SHADER_PARAMETER_RDG_BUFFER_SRV(StructuredBuffer, RestPosition1Buffer)
+		SHADER_PARAMETER_RDG_BUFFER_SRV(StructuredBuffer, RestPosition2Buffer)
 
-	SHADER_PARAMETER_RDG_BUFFER_SRV(StructuredBuffer, RootBarycentricBuffer)
+		SHADER_PARAMETER_RDG_BUFFER_SRV(StructuredBuffer, DeformedPosition0Buffer)
+		SHADER_PARAMETER_RDG_BUFFER_SRV(StructuredBuffer, DeformedPosition1Buffer)
+		SHADER_PARAMETER_RDG_BUFFER_SRV(StructuredBuffer, DeformedPosition2Buffer)
 
-	SHADER_PARAMETER_STRUCT_REF(FViewUniformShaderParameters, ViewUniformBuffer)
-	RENDER_TARGET_BINDING_SLOTS()
-END_SHADER_PARAMETER_STRUCT()
+		SHADER_PARAMETER_RDG_BUFFER_SRV(StructuredBuffer, RestSamplePositionsBuffer)
+		SHADER_PARAMETER_RDG_BUFFER_SRV(StructuredBuffer, DeformedSamplePositionsBuffer)
+
+		SHADER_PARAMETER_RDG_BUFFER_SRV(StructuredBuffer, RootBarycentricBuffer)
+
+		SHADER_PARAMETER_STRUCT_REF(FViewUniformShaderParameters, ViewUniformBuffer)
+		SHADER_PARAMETER_STRUCT_INCLUDE(ShaderPrint::FShaderParameters, ShaderPrintUniformBuffer)
+	END_SHADER_PARAMETER_STRUCT()
 
-class FHairProjectionHairDebug : public FGlobalShader
-{
-public:
 	static bool ShouldCompilePermutation(const FGlobalShaderPermutationParameters& Parameters)
 	{
 		return IsHairStrandsSupported(EHairStrandsShaderType::Tool, Parameters.Platform);
 	}
-
-	FHairProjectionHairDebug() = default;
-	FHairProjectionHairDebug(const CompiledShaderInitializerType& Initializer) : FGlobalShader(Initializer) {}
-};
-
-
-class FHairProjectionHairDebugVS : public FHairProjectionHairDebug
-{
-public:
-	class FInputType : SHADER_PERMUTATION_INT("PERMUTATION_INPUT_TYPE", 3);
-	using FPermutationDomain = TShaderPermutationDomain<FInputType>;
-
-	DECLARE_GLOBAL_SHADER(FHairProjectionHairDebugVS);
-	SHADER_USE_PARAMETER_STRUCT(FHairProjectionHairDebugVS, FHairProjectionHairDebug);
-
-	BEGIN_SHADER_PARAMETER_STRUCT(FParameters, )
-		SHADER_PARAMETER_STRUCT_INCLUDE(FHairProjectionHairDebugParameters, Pass)
-		END_SHADER_PARAMETER_STRUCT()
-};
- 
-
-class FHairProjectionHairDebugPS : public FHairProjectionHairDebug
-{
-	DECLARE_GLOBAL_SHADER(FHairProjectionHairDebugPS);
-	SHADER_USE_PARAMETER_STRUCT(FHairProjectionHairDebugPS, FHairProjectionHairDebug);
-
-	BEGIN_SHADER_PARAMETER_STRUCT(FParameters, )
-		SHADER_PARAMETER_STRUCT_INCLUDE(FHairProjectionHairDebugParameters, Pass)
-		END_SHADER_PARAMETER_STRUCT()
 };
 
-IMPLEMENT_GLOBAL_SHADER(FHairProjectionHairDebugVS, "/Engine/Private/HairStrands/HairStrandsMeshProjectionHairDebug.usf", "MainVS", SF_Vertex);
-IMPLEMENT_GLOBAL_SHADER(FHairProjectionHairDebugPS, "/Engine/Private/HairStrands/HairStrandsMeshProjectionHairDebug.usf", "MainPS", SF_Pixel);
+IMPLEMENT_GLOBAL_SHADER(FHairProjectionHairDebugCS, "/Engine/Private/HairStrands/HairStrandsMeshProjectionHairDebug.usf", "MainCS", SF_Compute);
 
 enum class EDebugProjectionHairType
 {
@@ -374,18 +288,19 @@ enum class EDebugProjectionHairType
 static void AddDebugProjectionHairPass(
 	FRDGBuilder& GraphBuilder,
 	FGlobalShaderMap* ShaderMap,
+	const FShaderPrintData* ShaderPrintData,
 	FIntRect Viewport,
 	const TUniformBufferRef<FViewUniformShaderParameters>& ViewUniformBuffer,
-	const bool bClearDepth,
-	const EDebugProjectionHairType GeometryType,
-	const HairStrandsTriangleType PoseType,
 	const int32 MeshLODIndex,
 	const FHairStrandsRestRootResource* RestRootResources,
 	const FHairStrandsDeformedRootResource* DeformedRootResources,
 	const FTransform& LocalToWorld,
-	FRDGTextureRef ColorTarget,
-	FRDGTextureRef DepthTexture)
+	const EDebugProjectionHairType GeometryType,
+	const HairStrandsTriangleType PoseType)
 {
+	// Force ShaderPrint on.
+	ShaderPrint::SetEnabled(true);
+
 	if (MeshLODIndex < 0 || MeshLODIndex >= RestRootResources->LODs.Num() || MeshLODIndex >= DeformedRootResources->LODs.Num())
 		return;
 	
@@ -396,6 +311,14 @@ static void AddDebugProjectionHairPass(
 	if (PrimitiveCount == 0)
 		return;
 
+	if (ShaderPrintData == nullptr || !ShaderPrint::IsEnabled(*ShaderPrintData))
+	{
+		return;
+	}
+
+	ShaderPrint::RequestSpaceForLines(PrimitiveCount);
+	ShaderPrint::RequestSpaceForTriangles(PrimitiveCount);
+
 	if (EDebugProjectionHairType::HairFrame == GeometryType &&
 		!RestRootResources->LODs[MeshLODIndex].RootTriangleBarycentricBuffer.Buffer)
 		return;
@@ -423,7 +346,7 @@ static void AddDebugProjectionHairPass(
 
 	const FIntPoint Resolution(Viewport.Width(), Viewport.Height());
 
-	FHairProjectionHairDebugParameters* Parameters = GraphBuilder.AllocParameters<FHairProjectionHairDebugParameters>();
+	FHairProjectionHairDebugCS::FParameters* Parameters = GraphBuilder.AllocParameters<FHairProjectionHairDebugCS::FParameters>();
 	Parameters->OutputResolution = Resolution;
 	Parameters->MaxRootCount = RootCount;
 	Parameters->RootLocalToWorld = FMatrix44f(LocalToWorld.ToMatrixWithScale());	// LWC_TODO: Precision loss
@@ -446,53 +369,20 @@ static void AddDebugProjectionHairPass(
 	Parameters->DeformedSamplePositionsBuffer = RegisterAsSRV(GraphBuilder, DeformedLODDatas.GetDeformedSamplePositionsBuffer(FHairStrandsDeformedRootResource::FLOD::Current));
 
 	Parameters->ViewUniformBuffer = ViewUniformBuffer;
-	Parameters->RenderTargets[0] = FRenderTargetBinding(ColorTarget, ERenderTargetLoadAction::ELoad, 0);
-	Parameters->RenderTargets.DepthStencil = FDepthStencilBinding(DepthTexture, bClearDepth ? ERenderTargetLoadAction::EClear : ERenderTargetLoadAction::ELoad, ERenderTargetLoadAction::ENoAction, FExclusiveDepthStencil::DepthWrite_StencilNop);
-
-	FHairProjectionHairDebugVS::FPermutationDomain PermutationVector;
-	PermutationVector.Set<FHairProjectionHairDebugVS::FInputType>(GeometryType == EDebugProjectionHairType::HairFrame ? 0 : GeometryType == EDebugProjectionHairType::HairTriangle ? 1 : 2);
-					
-	TShaderMapRef<FHairProjectionHairDebugVS> VertexShader(ShaderMap, PermutationVector);
-	TShaderMapRef<FHairProjectionHairDebugPS> PixelShader(ShaderMap);
-
-	FHairProjectionHairDebugVS::FParameters VSParameters;
-	VSParameters.Pass = *Parameters;
-	FHairProjectionHairDebugPS::FParameters PSParameters;
-	PSParameters.Pass = *Parameters;
-
-	GraphBuilder.AddPass(
-		RDG_EVENT_NAME("HairStrands::MeshProjectionHairDebug"),
-		Parameters,
-		ERDGPassFlags::Raster,
-		[Parameters, VSParameters, PSParameters, VertexShader, PixelShader, Viewport, Resolution, PrimitiveCount, PrimitiveType](FRHICommandList& RHICmdList)
-		{
+	ShaderPrint::SetParameters(GraphBuilder, *ShaderPrintData, Parameters->ShaderPrintUniformBuffer);
+
+	FHairProjectionHairDebugCS::FPermutationDomain PermutationVector;
+	PermutationVector.Set<FHairProjectionHairDebugCS::FInputType>(GeometryType == EDebugProjectionHairType::HairFrame ? 0 : GeometryType == EDebugProjectionHairType::HairTriangle ? 1 : 2);
 
-			RHICmdList.SetViewport(
-				Viewport.Min.X,
-				Viewport.Min.Y,
-				0.0f,
-				Viewport.Max.X,
-				Viewport.Max.Y,
-				1.0f);
-
-			// Apply additive blending pipeline state.
-			FGraphicsPipelineStateInitializer GraphicsPSOInit;
-			RHICmdList.ApplyCachedRenderTargets(GraphicsPSOInit);
-			GraphicsPSOInit.BlendState = TStaticBlendState<CW_RGBA, BO_Add, BF_One, BF_Zero, BO_Add, BF_One, BF_Zero>::GetRHI();
-			GraphicsPSOInit.RasterizerState = TStaticRasterizerState<>::GetRHI();
-			GraphicsPSOInit.DepthStencilState = TStaticDepthStencilState<true, CF_DepthNearOrEqual>::GetRHI();
-			GraphicsPSOInit.BoundShaderState.VertexDeclarationRHI = GEmptyVertexDeclaration.VertexDeclarationRHI;
-			GraphicsPSOInit.BoundShaderState.VertexShaderRHI = VertexShader.GetVertexShader();
-			GraphicsPSOInit.BoundShaderState.PixelShaderRHI = PixelShader.GetPixelShader();
-			GraphicsPSOInit.PrimitiveType = PrimitiveType;
-			SetGraphicsPipelineState(RHICmdList, GraphicsPSOInit, 0);
-
-			SetShaderParameters(RHICmdList, VertexShader, VertexShader.GetVertexShader(), VSParameters);
-			SetShaderParameters(RHICmdList, PixelShader, PixelShader.GetPixelShader(), PSParameters);
-
-			// Emit an instanced quad draw call on the order of the number of pixels on the screen.	
-			RHICmdList.DrawPrimitive(0, PrimitiveCount, 1);
-		});
+	TShaderMapRef<FHairProjectionHairDebugCS> ComputeShader(ShaderMap, PermutationVector);
+	ClearUnusedGraphResources(ComputeShader, Parameters);
+
+	FComputeShaderUtils::AddPass(
+		GraphBuilder,
+		RDG_EVENT_NAME("HairStrands::MeshProjectionDebug(Hair, %s)", GeometryType == EDebugProjectionHairType::HairFrame ? TEXT("Frame") : GeometryType == EDebugProjectionHairType::HairTriangle ? TEXT("Triangle") : TEXT("Samples")),
+		ComputeShader,
+		Parameters,
+		FIntVector(FMath::DivideAndRoundUp(PrimitiveCount, 256u), 1, 1));
 }	
 
 ///////////////////////////////////////////////////////////////////////////////////////////////////
@@ -1142,23 +1032,15 @@ void RunHairStrandsDebug(
 	if (HairDebugMode == EHairDebugMode::MeshProjection)
 	{
 		{
-			bool bClearDepth = true;
-			FRDGTextureRef DepthTexture;
-			{
-				const FRDGTextureDesc Desc = FRDGTextureDesc::Create2D(SceneColorTexture->Desc.Extent, PF_DepthStencil, FClearValueBinding::DepthFar, TexCreate_DepthStencilTargetable | TexCreate_ShaderResource);
-				DepthTexture = GraphBuilder.CreateTexture(Desc, TEXT("Hair.InterpolationDepthTexture"));
-			}
-
 			if (GHairDebugMeshProjection_SkinCacheMesh > 0)
 			{
-				auto RenderMeshProjection = [&bClearDepth, ShaderMap, Viewport, &ViewUniformBuffer, Instances, SkinCache, &SceneColorTexture, &DepthTexture, &GraphBuilder](FRDGBuilder& LocalGraphBuilder, EHairStrandsProjectionMeshType MeshType)
+				auto RenderMeshProjection = [ShaderMap, ShaderPrintData, Viewport, &ViewUniformBuffer, Instances, SkinCache, &GraphBuilder](FRDGBuilder& LocalGraphBuilder, EHairStrandsProjectionMeshType MeshType)
 				{
 					FHairStrandsProjectionMeshData::LOD MeshProjectionLODData;
 					GetGroomInterpolationData(GraphBuilder, ShaderMap, Instances, MeshType, SkinCache, MeshProjectionLODData);
 					for (FHairStrandsProjectionMeshData::Section& Section : MeshProjectionLODData.Sections)
 					{
-						AddDebugProjectionMeshPass(LocalGraphBuilder, ShaderMap, Viewport, ViewUniformBuffer, MeshType, bClearDepth, Section, SceneColorTexture, DepthTexture);
-						bClearDepth = false;
+						AddDebugProjectionMeshPass(LocalGraphBuilder, ShaderMap, ShaderPrintData, Viewport, ViewUniformBuffer, MeshType, Section);
 					}
 				};
 
@@ -1168,7 +1050,7 @@ void RunHairStrandsDebug(
 				RenderMeshProjection(GraphBuilder, EHairStrandsProjectionMeshType::TargetMesh);
 			}
 
-			auto RenderProjectionData = [&GraphBuilder, ShaderMap, Viewport, &ViewUniformBuffer, Instances, SkinCache, &bClearDepth, SceneColorTexture, DepthTexture](EHairStrandsInterpolationType StrandType, bool bRestTriangle, bool bRestFrame, bool bRestSamples, bool bDeformedTriangle, bool bDeformedFrame, bool bDeformedSamples)
+			auto RenderProjectionData = [&GraphBuilder, ShaderMap, Viewport, &ViewUniformBuffer, Instances, ShaderPrintData](EHairStrandsInterpolationType StrandType, bool bRestTriangle, bool bRestFrame, bool bRestSamples, bool bDeformedTriangle, bool bDeformedFrame, bool bDeformedSamples)
 			{
 				TArray<int32> HairLODIndices;
 				for (FHairStrandsInstance* AbstractInstance : Instances)
@@ -1185,114 +1067,12 @@ void RunHairStrandsDebug(
 
 					const int32 MeshLODIndex = Instance->Debug.MeshLODIndex;
 
-					if (bRestTriangle)
-					{
-						AddDebugProjectionHairPass(
-							GraphBuilder, 
-							ShaderMap, 
-							Viewport, 
-							ViewUniformBuffer, 
-							bClearDepth, 
-							EDebugProjectionHairType::HairTriangle, 
-							HairStrandsTriangleType::RestPose, 
-							MeshLODIndex,
-							RestRootResource, 
-							DeformedRootResource, 
-							Instance->HairGroupPublicData->VFInput.LocalToWorldTransform, 
-							SceneColorTexture, 
-							DepthTexture);
-						bClearDepth = false;
-					}
-					if (bRestFrame)
-					{
-						AddDebugProjectionHairPass(
-							GraphBuilder,
-							ShaderMap,
-							Viewport,
-							ViewUniformBuffer,
-							bClearDepth,
-							EDebugProjectionHairType::HairFrame,
-							HairStrandsTriangleType::RestPose,
-							MeshLODIndex,
-							RestRootResource,
-							DeformedRootResource,
-							Instance->HairGroupPublicData->VFInput.LocalToWorldTransform,
-							SceneColorTexture,
-							DepthTexture);
-						bClearDepth = false;
-					}
-					if (bRestSamples)
-					{
-						AddDebugProjectionHairPass(
-							GraphBuilder,
-							ShaderMap,
-							Viewport,
-							ViewUniformBuffer,
-							bClearDepth,
-							EDebugProjectionHairType::HairSamples,
-							HairStrandsTriangleType::RestPose,
-							MeshLODIndex,
-							RestRootResource,
-							DeformedRootResource,
-							Instance->HairGroupPublicData->VFInput.LocalToWorldTransform,
-							SceneColorTexture,
-							DepthTexture);
-						bClearDepth = false;
-					}
-					if (bDeformedTriangle)
-					{
-						AddDebugProjectionHairPass(
-							GraphBuilder,
-							ShaderMap,
-							Viewport,
-							ViewUniformBuffer,
-							bClearDepth,
-							EDebugProjectionHairType::HairTriangle,
-							HairStrandsTriangleType::DeformedPose,
-							MeshLODIndex,
-							RestRootResource,
-							DeformedRootResource,
-							Instance->HairGroupPublicData->VFInput.LocalToWorldTransform,
-							SceneColorTexture,
-							DepthTexture);
-						bClearDepth = false;
-					}
-					if (bDeformedFrame)
-					{
-						AddDebugProjectionHairPass(
-							GraphBuilder,
-							ShaderMap,
-							Viewport,
-							ViewUniformBuffer,
-							bClearDepth,
-							EDebugProjectionHairType::HairFrame,
-							HairStrandsTriangleType::DeformedPose,
-							MeshLODIndex,
-							RestRootResource,
-							DeformedRootResource,
-							Instance->HairGroupPublicData->VFInput.LocalToWorldTransform,
-							SceneColorTexture,
-							DepthTexture);
-						bClearDepth = false;
-					}
-					if (bDeformedSamples)
-					{
-						AddDebugProjectionHairPass(
-							GraphBuilder,
-							ShaderMap,
-							Viewport,
-							ViewUniformBuffer,
-							bClearDepth,
-							EDebugProjectionHairType::HairSamples,
-							HairStrandsTriangleType::DeformedPose,
-							MeshLODIndex,
-							RestRootResource,
-							DeformedRootResource,
-							Instance->HairGroupPublicData->VFInput.LocalToWorldTransform,
-							SceneColorTexture,
-							DepthTexture);
-						bClearDepth = false;
-					}
+					if (bRestTriangle)		{ AddDebugProjectionHairPass(GraphBuilder, ShaderMap, ShaderPrintData, Viewport, ViewUniformBuffer, MeshLODIndex, RestRootResource, DeformedRootResource, Instance->HairGroupPublicData->VFInput.LocalToWorldTransform, EDebugProjectionHairType::HairTriangle, HairStrandsTriangleType::RestPose); }
+					if (bRestFrame)			{ AddDebugProjectionHairPass(GraphBuilder, ShaderMap, ShaderPrintData, Viewport, ViewUniformBuffer, MeshLODIndex, RestRootResource, DeformedRootResource, Instance->HairGroupPublicData->VFInput.LocalToWorldTransform, EDebugProjectionHairType::HairFrame,    HairStrandsTriangleType::RestPose); }
+					if (bRestSamples)		{ AddDebugProjectionHairPass(GraphBuilder, ShaderMap, ShaderPrintData, Viewport, ViewUniformBuffer, MeshLODIndex, RestRootResource, DeformedRootResource, Instance->HairGroupPublicData->VFInput.LocalToWorldTransform, EDebugProjectionHairType::HairSamples,  HairStrandsTriangleType::RestPose); }
+					if (bDeformedTriangle)	{ AddDebugProjectionHairPass(GraphBuilder, ShaderMap, ShaderPrintData, Viewport, ViewUniformBuffer, MeshLODIndex, RestRootResource, DeformedRootResource, Instance->HairGroupPublicData->VFInput.LocalToWorldTransform, EDebugProjectionHairType::HairTriangle, HairStrandsTriangleType::DeformedPose); }
+					if (bDeformedFrame)		{ AddDebugProjectionHairPass(GraphBuilder, ShaderMap, ShaderPrintData, Viewport, ViewUniformBuffer, MeshLODIndex, RestRootResource, DeformedRootResource, Instance->HairGroupPublicData->VFInput.LocalToWorldTransform, EDebugProjectionHairType::HairFrame,    HairStrandsTriangleType::DeformedPose); }
+					if (bDeformedSamples)	{ AddDebugProjectionHairPass(GraphBuilder, ShaderMap, ShaderPrintData, Viewport, ViewUniformBuffer, MeshLODIndex, RestRootResource, DeformedRootResource, Instance->HairGroupPublicData->VFInput.LocalToWorldTransform, EDebugProjectionHairType::HairSamples,  HairStrandsTriangleType::DeformedPose); }
 				}
 			};
 
diff --git a/Engine/Shaders/Private/HairStrands/HairStrandsMeshProjectionHairDebug.usf b/Engine/Shaders/Private/HairStrands/HairStrandsMeshProjectionHairDebug.usf
--- a/Engine/Shaders/Private/HairStrands/HairStrandsMeshProjectionHairDebug.usf
+++ b/Engine/Shaders/Private/HairStrands/HairStrandsMeshProjectionHairDebug.usf
@@ -4,6 +4,7 @@
 
 #include "../Common.ush"
 #include "HairStrandsMeshProjectionCommon.ush" 
+#include "../ShaderPrint.ush"
 
 #define INPUT_ROOT_FRAME 0
 #define INPUT_ROOT_TRIANGLE 1
@@ -28,26 +29,15 @@ Buffer<float4>  DeformedPosition2Buffer;
 Buffer<float4> RestSamplePositionsBuffer;
 Buffer<float4> DeformedSamplePositionsBuffer;
 
-void MainVS(
-	uint VertexId : SV_VertexID,
-	out float2 OutVertexUVs : VERTEXUVs,
-	out float3 OutFaceNormal : FACE_NORMAL,
-	out float4 OutPosition : SV_POSITION)
+[numthreads(256, 1, 1)]
+void MainCS(uint3 DispatchThreadId : SV_DispatchThreadID)
 {
 	ResolvedView = ResolveView();
 
-	float2 UVs = 0;
-
 #if PERMUTATION_INPUT_TYPE == INPUT_ROOT_FRAME  
-	const uint LineVertex = (VertexId%2);
-	const uint LineIndex = VertexId/2;
+	const uint LineIndex = DispatchThreadId.x;
 
-	OutFaceNormal 	= float3(0,0,1);
-	OutPosition 	= float4(0,0,0,0);
 	if (LineIndex > MaxRootCount) return;
-
-	float3 N;
-	float3 P;
 	{
 		const FHairMeshTriangleWithUVs RestTriangle		= GetTriangleWithUVsTransformation(LineIndex, RestPosition0Buffer, RestPosition1Buffer, RestPosition2Buffer);
 		const FHairMeshTriangleWithUVs DeformedTriangle = GetTriangleWithUVsTransformation(LineIndex, DeformedPosition0Buffer, DeformedPosition1Buffer, DeformedPosition2Buffer);
@@ -62,7 +52,7 @@ void MainVS(
 		float3 P0 = B.x * Tri.P0 + B.y * Tri.P1 + B.z * Tri.P2;
 		float3 P1 = P0 + Tri.N;
 
-		UVs = 
+		const float2 UVs = 
 			B.x * Tri.UV0 +
 			B.y * Tri.UV1 +
 			B.z * Tri.UV2;
@@ -70,14 +60,15 @@ void MainVS(
 		P0 = mul(float4(P0, 1), RootLocalToWorld).xyz;
 		P1 = mul(float4(P1, 1), RootLocalToWorld).xyz;
 
-		N = normalize(P0 - P1);
-		P = LineVertex == 0 ? P0 : P1;
+		const float3 N = normalize(P0 - P1);
+		const float4 Color = float4(UVs, 0, 1);
+		//Color = float4(N* 0.5f + 0.5f, 1);
+		AddLineWS(P0, P1, Color);
 	}
-#endif
+#endif // INPUT_ROOT_FRAME
 
 #if PERMUTATION_INPUT_TYPE == INPUT_ROOT_TRIANGLE
-	const uint TriangleVertex = (VertexId % 3);
-	const uint TriangleIndex = VertexId / 3;
+	const uint TriangleIndex = DispatchThreadId.x;
 
 	if (TriangleIndex > MaxRootCount) return;
 
@@ -89,24 +80,7 @@ void MainVS(
 		Tri = DeformedTriangle;
 	else
 		Tri = RestTriangle;
-
-	float3 N = Tri.N;
-	float3 P = 0;
-	if (TriangleVertex == 0)
-	{
-		P = Tri.P0;
-		UVs = Tri.UV0;
-	}
-	if (TriangleVertex == 1)
-	{
-		P = Tri.P1;
-		UVs = Tri.UV1;
-	}
-	if (TriangleVertex == 2)
-	{
-		P = Tri.P2;
-		UVs = Tri.UV2;
-	}
+	
 
 	// Debug: To test if the rest-triangle-to-deformed-triangle transformation works correctly
 	#if 0
@@ -132,52 +106,24 @@ void MainVS(
 	}
 	#endif
 
-	P = mul(float4(P, 1), RootLocalToWorld).xyz;
-#endif
+	float4 Color = float4(Tri.UV0, 0, 1);
+	// Color = float4((Tri.N + 1) * 0.5f, 1);
 
-#if PERMUTATION_INPUT_TYPE == INPUT_ROOT_SAMPLES 
+	AddFilledTriangleWS(
+		mul(float4(Tri.P0, 1), RootLocalToWorld).xyz, 
+		mul(float4(Tri.P1, 1), RootLocalToWorld).xyz, 
+		mul(float4(Tri.P2, 1), RootLocalToWorld).xyz, 
+		Color);
+#endif // INPUT_ROOT_TRIANGLE
 
-	const uint LineVertex = (VertexId%2);
-	const uint LineIndex = VertexId/2;
-	const uint PointIndex = VertexId/6;
-	const uint AxisIndex = (LineIndex%3);
+#if PERMUTATION_INPUT_TYPE == INPUT_ROOT_SAMPLES 
+	const uint PointIndex = DispatchThreadId.x;
 	
-	if (LineIndex > MaxRootCount) return;
+	if (PointIndex > MaxRootCount) return;
 
-	float3 N;
-	if(AxisIndex == 0)
-	{
-		N = float3(1,0,0);
-		UVs = float2(1,0);
-	}
-	else if(AxisIndex == 1)
-	{
-		N = float3(0,1,0);
-		UVs = float2(0,1);
-	}
-	else if(AxisIndex == 2)
-	{
-		N = float3(0,0,1);
-		UVs = float2(0.5,0.5);
-	}
 	float3 P = (DeformedFrameEnable > 0) ? DeformedSamplePositionsBuffer[PointIndex].xyz : RestSamplePositionsBuffer[PointIndex].xyz;
-	P = (LineVertex == 0) ? P - 0.5 * N : P + 0.5 * N;
 	P = mul(float4(P, 1), RootLocalToWorld).xyz;
-	
-#endif
-	
-	OutVertexUVs = UVs;
-	OutFaceNormal = N;
-	OutPosition = mul(float4(P,1), LWCHackToFloat(PrimaryView.WorldToClip));
-}
-
-void MainPS(
-	in float2 VertexUVs : VERTEXUVs,
-	in float3 FaceNormal : FACE_NORMAL,
-	out float4 OutColor : SV_Target0)
-{
-	FaceNormal += float3(1,1,1);
-	FaceNormal *= 0.5f;
-	OutColor = float4(FaceNormal,1);
-	OutColor = float4(VertexUVs,0, 1);
+
+	AddCrossWS(P, 1.f, ColorYellow);
+#endif // INPUT_ROOT_SAMPLES
 }
\ No newline at end of file
diff --git a/Engine/Shaders/Private/HairStrands/HairStrandsMeshProjectionMeshDebug.usf b/Engine/Shaders/Private/HairStrands/HairStrandsMeshProjectionMeshDebug.usf
--- a/Engine/Shaders/Private/HairStrands/HairStrandsMeshProjectionMeshDebug.usf
+++ b/Engine/Shaders/Private/HairStrands/HairStrandsMeshProjectionMeshDebug.usf
@@ -4,6 +4,7 @@
 
 #include "../Common.ush"
 #include "HairStrandsMeshProjectionCommon.ush"
+#include "../ShaderPrint.ush"
 
 float4x4 LocalToWorld;
 uint VertexOffset;
@@ -28,63 +29,37 @@ Buffer<uint> InputIndexBuffer;
 #define SOURCE_MESH 2
 #define TARGET_MESH 3
 
-void MainVS(
-	uint VertexId : SV_VertexID,
-	out float2 OutVertexUV : TRI_UVs,
-	out float3 OutFaceNormal : FACE_NORMAL,
-	out float4 OutPosition : SV_POSITION)
+[numthreads(256, 1, 1)]
+void MainCS(uint3 DispatchThreadId : SV_DispatchThreadID)
 {
 	ResolvedView = ResolveView();
 
+	float4 Color = 0;
+	if (MeshType == REST_MESH)		Color = float4(1, 0, 0, 1);
+	if (MeshType == DEFORMED_MESH)	Color = float4(0, 1, 0, 1);
+	if (MeshType == SOURCE_MESH)	Color = float4(0, 0, 1, 1);
+	if (MeshType == TARGET_MESH)	Color = float4(1, 1, 0, 1);
+
 #if PERMUTATION_INPUT_TYPE == 0
-	const uint VertexIndex = min(VertexOffset + VertexId, MaxVertexCount-1);
-	float3 P = InputVertexPositionBuffer[VertexIndex];
-	OutFaceNormal = float3(0,0,1);
-	OutVertexUV = InputVertexUVsBuffer[VertexIndex];
-#endif
+	const uint VertexId = DispatchThreadId.x * 3;
+	const uint VertexIndex0 = min(VertexOffset + VertexId, MaxVertexCount-1);
 
-#if PERMUTATION_INPUT_TYPE == 1
-	const uint TriangleVertexIndex = (VertexId%3);
-	const uint TriangleIndex = VertexId / 3;
-	const FHairMeshTriangleWithUVs Tri = GetTriangleData(TriangleIndex, InputVertexPositionBuffer, InputVertexUVsBuffer, InputIndexBuffer, IndexOffset, MaxIndexCount, MaxVertexCount, MeshUVsChannelOffset, MeshUVsChannelCount);
-	OutFaceNormal = Tri.N;
+	FHairMeshTriangleWithUVs Tri = (FHairMeshTriangleWithUVs)0;
+	Tri.P0 = InputVertexPositionBuffer[VertexIndex0 + 0];
+	Tri.P1 = InputVertexPositionBuffer[VertexIndex0 + 1];
+	Tri.P2 = InputVertexPositionBuffer[VertexIndex0 + 2];
 
-	float3 P = 0;
-	float2 UV = 0;
-	if (TriangleVertexIndex == 0)
-	{
-		P = Tri.P0;
-		UV = Tri.UV0;
-	}
-	if (TriangleVertexIndex == 1)
-	{
-		P = Tri.P1;
-		UV = Tri.UV1;
-	}
-	if (TriangleVertexIndex == 2)
-	{
-		P = Tri.P2;
-		UV = Tri.UV2;
-	}
+	Tri.UV0 = InputVertexUVsBuffer[VertexIndex0 + 0];
+	Tri.UV1 = InputVertexUVsBuffer[VertexIndex0 + 1];
+	Tri.UV2 = InputVertexUVsBuffer[VertexIndex0 + 2];
 
-	OutVertexUV = UV;
+	Tri.N = cross(Tri.P1 - Tri.P0, Tri.P2 - Tri.P0);
 #endif
-	if (bOutputInUVsSpace)
-	{
-		float2 PixelOffset = 0;
-		//float2 PixelOffset = 1.f / float2(OutputResolution);
-		if (MeshType == REST_MESH)		PixelOffset *= 0;
-		if (MeshType == DEFORMED_MESH)	PixelOffset *= 1;
-		if (MeshType == SOURCE_MESH)	PixelOffset *= 2;
-		if (MeshType == TARGET_MESH)	PixelOffset *= 3;
 
-		OutPosition = float4((PixelOffset + OutVertexUV) * 2 - 1, 0.5, 1);
-	}
-	else
-	{
-		P = mul(float4(P, 1), LocalToWorld).xyz;
-		OutPosition = mul(float4(P, 1), LWCHackToFloat(PrimaryView.WorldToClip));
-	}
+#if PERMUTATION_INPUT_TYPE == 1
+	const uint TriangleIndex = DispatchThreadId.x;
+	const FHairMeshTriangleWithUVs Tri = GetTriangleData(TriangleIndex, InputVertexPositionBuffer, InputVertexUVsBuffer, InputIndexBuffer, IndexOffset, MaxIndexCount, MaxVertexCount, MeshUVsChannelOffset, MeshUVsChannelCount);
+#endif
 
 	// Quick filter
 #if 0
@@ -94,24 +69,22 @@ void MainVS(
 	//if (MeshType != TARGET_MESH)
 	if (MeshType != SOURCE_MESH && MeshType != TARGET_MESH || SectionIndex != 0)
 	{
-		OutPosition = float4(1.#INF, 1.#INF, 1.#INF, 1.#INF);
-		OutVertexUV = 0;
+		return;
 	}
 #endif
-}
 
-void MainPS(
-	in float2 VertexUV : TRI_UVs,
-	in float3 FaceNormal : FACE_NORMAL,
-	out float4 OutColor : SV_Target0)
-{
-	float3 Color = 0;
-	if (MeshType == REST_MESH)		Color = float3(1, 0, 0);
-	if (MeshType == DEFORMED_MESH)	Color = float3(0, 1, 0);
-	if (MeshType == SOURCE_MESH)	Color = float3(0, 0, 1);
-	if (MeshType == TARGET_MESH)	Color = float3(1, 1, 0);
+	Color.xyz *= Tri.N * 0.5f + 0.5f;
 
-	OutColor = float4((FaceNormal + 1) * 0.5f,1);
-	OutColor = float4(VertexUV, 0, 1);
-	OutColor = float4(Color, 1);
-}
\ No newline at end of file
+	if (bOutputInUVsSpace)
+	{
+		AddFilledTriangleWS(float3(Tri.UV0, 0), float3(Tri.UV1, 0), float3(Tri.UV2, 0), Color);
+	}
+	else
+	{
+		AddFilledTriangleWS(
+			mul(float4(Tri.P0, 1), LocalToWorld).xyz, 
+			mul(float4(Tri.P1, 1), LocalToWorld).xyz, 
+			mul(float4(Tri.P2, 1), LocalToWorld).xyz,
+			Color);
+	}
+}
