- Fix functions that were not correctly templated

#rnx

diff --git a/Engine/Plugins/Runtime/HairStrands/Shaders/Private/NiagaraDataInterfacePressureGridTemplate.ush b/Engine/Plugins/Runtime/HairStrands/Shaders/Private/NiagaraDataInterfacePressureGridTemplate.ush
--- a/Engine/Plugins/Runtime/HairStrands/Shaders/Private/NiagaraDataInterfacePressureGridTemplate.ush
+++ b/Engine/Plugins/Runtime/HairStrands/Shaders/Private/NiagaraDataInterfacePressureGridTemplate.ush
@@ -423,7 +423,7 @@ void ComputeComponentWeights_{ParameterName}(in int3 GridIndex, in int Component
 }
 
  // Update the solid weights and velocities
-void DIVelocityGrid_ComputeBoundaryWeights(in int GridHash, out bool OutWeightsStatus)
+void ComputeBoundaryWeights_{ParameterName}(in int GridHash, out bool OutWeightsStatus)
 {
 	const int3 GridIndex = GetGridIndex_{ParameterName}(GridHash);
 	OutWeightsStatus = false;
@@ -505,7 +505,7 @@ void SolveGridPressure_{ParameterName}(in int GridHash, in int InitStage, out bo
 				float GridPressure = ComputeGridPressure_{ParameterName}(PrevIndex,CentreDistance,PrevFraction,NextFraction);
 				if(!IsInitStage)
 				{
-					DIVelocityGrid_AddGridScalar(PrevIndex,FLUID_PRESSURE_OFFSET,GridPressure);
+					AddGridScalar_{ParameterName}(PrevIndex,FLUID_PRESSURE_OFFSET,GridPressure);
 				}
 				else 
 				{
@@ -519,7 +519,7 @@ void SolveGridPressure_{ParameterName}(in int GridHash, in int InitStage, out bo
 					const float PrevWeight = GetGridScalar_{ParameterName}(PrevIndex,SOLID_WEIGHT_OFFSET+ComponentIndex);
 					if( PrevWeight > 0.0 && PrevFraction[ComponentIndex] > 0.0)
 					{
-						DIVelocityGrid_AddGridScalar(PrevIndex,FLUID_VELOCITY_OFFSET+ComponentIndex,-GridPressure/PrevFraction[ComponentIndex]);
+						AddGridScalar_{ParameterName}(PrevIndex,FLUID_VELOCITY_OFFSET+ComponentIndex,-GridPressure/PrevFraction[ComponentIndex]);
 					}
 					else
 					{
@@ -530,7 +530,7 @@ void SolveGridPressure_{ParameterName}(in int GridHash, in int InitStage, out bo
 					const float NextWeight = GetGridScalar_{ParameterName}(NextIndex,SOLID_WEIGHT_OFFSET+ComponentIndex);
 					if( NextWeight > 0.0 && NextFraction[ComponentIndex] > 0.0)
 					{
-						DIVelocityGrid_AddGridScalar(NextIndex,FLUID_VELOCITY_OFFSET+ComponentIndex,GridPressure/NextFraction[ComponentIndex]);
+						AddGridScalar_{ParameterName}(NextIndex,FLUID_VELOCITY_OFFSET+ComponentIndex,GridPressure/NextFraction[ComponentIndex]);
 					}
 					else
 					{
@@ -546,13 +546,3 @@ void SolveGridPressure_{ParameterName}(in int GridHash, in int InitStage, out bo
 		}
 	}
 }
-
-
-
-
-
-
-
-
-
-
