- Fix for static analysis issue

#rnx

diff --git a/Engine/Plugins/Runtime/HairStrands/Source/HairStrandsCore/Private/Niagara/NiagaraDataInterfaceHairStrands.cpp b/Engine/Plugins/Runtime/HairStrands/Source/HairStrandsCore/Private/Niagara/NiagaraDataInterfaceHairStrands.cpp
--- a/Engine/Plugins/Runtime/HairStrands/Source/HairStrandsCore/Private/Niagara/NiagaraDataInterfaceHairStrands.cpp
+++ b/Engine/Plugins/Runtime/HairStrands/Source/HairStrandsCore/Private/Niagara/NiagaraDataInterfaceHairStrands.cpp
@@ -3578,12 +3578,9 @@ void FNDIHairStrandsProxy::PostSimulate(const FNDIGpuComputePostSimulateContext&
 		Context.GetComputeDispatchInterface().MultiGPUResourceModified(Context.GetGraphBuilder(), DeformedPositionBuffer, false, true);
 	}
 
-	if (Context.IsFinalPostSimulate())
+	if (bIsHairValid && Context.IsFinalPostSimulate())
 	{
-		if (ProxyData->HairStrandsBuffer->IsInitialized())
-		{
-			ProxyData->HairStrandsBuffer->BoundingBoxBuffer.EndGraphUsage();
-		}
+		ProxyData->HairStrandsBuffer->BoundingBoxBuffer.EndGraphUsage();
 	}
 }
 
