Add previous position support for geometry cache.
Fix hair strands motion vector when CLOD is used.

diff --git a/Engine/Plugins/Runtime/HairStrands/Source/HairStrandsCore/Private/GroomGeometryCache.cpp b/Engine/Plugins/Runtime/HairStrands/Source/HairStrandsCore/Private/GroomGeometryCache.cpp
--- a/Engine/Plugins/Runtime/HairStrands/Source/HairStrandsCore/Private/GroomGeometryCache.cpp
+++ b/Engine/Plugins/Runtime/HairStrands/Source/HairStrandsCore/Private/GroomGeometryCache.cpp
@@ -22,7 +22,9 @@ FHairStrandsProjectionMeshData::Section ConvertMeshSection(const FCachedGeometry
 	FHairStrandsProjectionMeshData::Section Out;
 	Out.IndexBuffer = In.IndexBuffer;
 	Out.RDGPositionBuffer = In.RDGPositionBuffer;
+	Out.RDGPreviousPositionBuffer = In.RDGPreviousPositionBuffer;
 	Out.PositionBuffer = In.PositionBuffer;
+	Out.PreviousPositionBuffer = In.PreviousPositionBuffer;
 	Out.UVsBuffer = In.UVsBuffer;
 	Out.UVsChannelOffset = In.UVsChannelOffset;
 	Out.UVsChannelCount = In.UVsChannelCount;
@@ -86,14 +88,20 @@ static void BuildBoneMatrices(USkeletalMeshComponent* SkeletalMeshComponent, con
 		TArray<FVector4f> BoneMatrices;
 		BuildBoneMatrices(SkeletalMeshComponent, LODData, LODIndex, MatrixOffsets, BoneMatrices);
 
-		FRDGBufferRef DeformedPositionsBuffer	= GraphBuilder.CreateBuffer(FRDGBufferDesc::CreateBufferDesc(sizeof(float), LODData.StaticVertexBuffers.PositionVertexBuffer.GetNumVertices() * 3), TEXT("Hair.SkinnedDeformedPositions"));
-		FRDGBufferRef BoneMatricesBuffer		= CreateStructuredBuffer(GraphBuilder, TEXT("Hair.SkinnedBoneMatrices"), sizeof(float) * 4, BoneMatrices.Num(), BoneMatrices.GetData(), sizeof(float) * 4 * BoneMatrices.Num());
-		FRDGBufferRef MatrixOffsetsBuffer		= CreateStructuredBuffer(GraphBuilder, TEXT("Hair.SkinnedMatrixOffsets"), sizeof(uint32), MatrixOffsets.Num(), MatrixOffsets.GetData(), sizeof(uint32) * MatrixOffsets.Num());
+		const bool bNeedPreviousPosition = IsHairStrandContinuousDecimationReorderingEnabled();
 
+		FRDGBufferRef DeformedPositionsBuffer			= GraphBuilder.CreateBuffer(FRDGBufferDesc::CreateBufferDesc(sizeof(float), LODData.StaticVertexBuffers.PositionVertexBuffer.GetNumVertices() * 3), TEXT("Hair.SkinnedDeformedPositions"));
+		FRDGBufferRef DeformedPreviousPositionsBuffer	= bNeedPreviousPosition ? GraphBuilder.CreateBuffer(FRDGBufferDesc::CreateBufferDesc(sizeof(float), LODData.StaticVertexBuffers.PositionVertexBuffer.GetNumVertices() * 3), TEXT("Hair.SkinnedDeformedPreviousPositions")) : nullptr;
+		FRDGBufferRef BoneMatricesBuffer				= CreateStructuredBuffer(GraphBuilder, TEXT("Hair.SkinnedBoneMatrices"), sizeof(float) * 4, BoneMatrices.Num(), BoneMatrices.GetData(), sizeof(float) * 4 * BoneMatrices.Num());
+		FRDGBufferRef MatrixOffsetsBuffer				= CreateStructuredBuffer(GraphBuilder, TEXT("Hair.SkinnedMatrixOffsets"), sizeof(uint32), MatrixOffsets.Num(), MatrixOffsets.GetData(), sizeof(uint32) * MatrixOffsets.Num());
+
+		// TODO previous bone matrices
 		AddSkinUpdatePass(GraphBuilder, ShaderMap, SkeletalMeshComponent->GetSkinWeightBuffer(LODIndex), LODData, BoneMatricesBuffer, MatrixOffsetsBuffer, DeformedPositionsBuffer);
 
 		CachedGeometry.DeformedPositionBuffer = DeformedPositionsBuffer;
+		CachedGeometry.DeformedPreviousPositionBuffer = DeformedPreviousPositionsBuffer;
 		FRDGBufferSRVRef DeformedPositionSRV = GraphBuilder.CreateSRV(DeformedPositionsBuffer, PF_R32_FLOAT);
+		FRDGBufferSRVRef DeformedPreviousPositionSRV = bNeedPreviousPosition ? GraphBuilder.CreateSRV(DeformedPreviousPositionsBuffer, PF_R32_FLOAT) : nullptr;
 		CachedGeometry.LocalToWorld = SkeletalMeshComponent->SceneProxy ? FTransform(SkeletalMeshComponent->SceneProxy->GetLocalToWorld()) : FTransform();
 		for (int32 SectionIdx = 0; SectionIdx < LODData.RenderSections.Num(); ++SectionIdx)
 		{
@@ -101,7 +109,9 @@ static void BuildBoneMatrices(USkeletalMeshComponent* SkeletalMeshComponent, con
 			FSkelMeshRenderSection& Section = LODData.RenderSections[SectionIdx];
 
 			CachedSection.RDGPositionBuffer = DeformedPositionSRV;
+			CachedSection.RDGPreviousPositionBuffer = DeformedPreviousPositionSRV;
 			CachedSection.PositionBuffer = nullptr; // Do not use the SRV slot, but instead use the RDG buffer created above (DeformedPositionSRV)
+			CachedSection.PreviousPositionBuffer = nullptr; // Do not use the SRV slot, but instead use the RDG buffer created above (DeformedPositionSRV)
 			CachedSection.UVsBuffer = LODData.StaticVertexBuffers.StaticMeshVertexBuffer.GetTexCoordsSRV();
 			CachedSection.TotalVertexCount = LODData.StaticVertexBuffers.PositionVertexBuffer.GetNumVertices();
 			CachedSection.IndexBuffer = LODData.MultiSizeIndexContainer.GetIndexBuffer()->GetSRV();
diff --git a/Engine/Plugins/Runtime/HairStrands/Source/HairStrandsCore/Private/HairStrandsMeshProjection.cpp b/Engine/Plugins/Runtime/HairStrands/Source/HairStrandsCore/Private/HairStrandsMeshProjection.cpp
--- a/Engine/Plugins/Runtime/HairStrands/Source/HairStrandsCore/Private/HairStrandsMeshProjection.cpp
+++ b/Engine/Plugins/Runtime/HairStrands/Source/HairStrandsCore/Private/HairStrandsMeshProjection.cpp
@@ -515,9 +515,10 @@ private:
 	DECLARE_GLOBAL_SHADER(FHairUpdateMeshTriangleCS);
 	SHADER_USE_PARAMETER_STRUCT(FHairUpdateMeshTriangleCS, FGlobalShader);
 
-	class FUpdateUVs : SHADER_PERMUTATION_INT("PERMUTATION_WITHUV", 2);
-	class FPositionType : SHADER_PERMUTATION_INT("PERMUTATION_POSITION_TYPE", 2);
-	using FPermutationDomain = TShaderPermutationDomain<FUpdateUVs, FPositionType>;
+	class FUpdateUVs : SHADER_PERMUTATION_BOOL("PERMUTATION_WITHUV");
+	class FPositionType : SHADER_PERMUTATION_BOOL("PERMUTATION_POSITION_TYPE");
+	class FPrevious : SHADER_PERMUTATION_BOOL("PERMUTATION_PREVIOUS"); 
+	using FPermutationDomain = TShaderPermutationDomain<FUpdateUVs, FPositionType, FPrevious>;
 
 	BEGIN_SHADER_PARAMETER_STRUCT(FParameters, )
 		SHADER_PARAMETER(uint32, MaxRootCount)
@@ -537,6 +538,15 @@ private:
 		SHADER_PARAMETER_RDG_BUFFER_SRV(Buffer, RDGMeshPositionBuffer6)
 		SHADER_PARAMETER_RDG_BUFFER_SRV(Buffer, RDGMeshPositionBuffer7)
 
+		SHADER_PARAMETER_RDG_BUFFER_SRV(Buffer, RDGMeshPreviousPositionBuffer0)
+		SHADER_PARAMETER_RDG_BUFFER_SRV(Buffer, RDGMeshPreviousPositionBuffer1)
+		SHADER_PARAMETER_RDG_BUFFER_SRV(Buffer, RDGMeshPreviousPositionBuffer2)
+		SHADER_PARAMETER_RDG_BUFFER_SRV(Buffer, RDGMeshPreviousPositionBuffer3)
+		SHADER_PARAMETER_RDG_BUFFER_SRV(Buffer, RDGMeshPreviousPositionBuffer4)
+		SHADER_PARAMETER_RDG_BUFFER_SRV(Buffer, RDGMeshPreviousPositionBuffer5)
+		SHADER_PARAMETER_RDG_BUFFER_SRV(Buffer, RDGMeshPreviousPositionBuffer6)
+		SHADER_PARAMETER_RDG_BUFFER_SRV(Buffer, RDGMeshPreviousPositionBuffer7)
+
 		SHADER_PARAMETER_SRV(Buffer, MeshPositionBuffer0)
 		SHADER_PARAMETER_SRV(Buffer, MeshPositionBuffer1)
 		SHADER_PARAMETER_SRV(Buffer, MeshPositionBuffer2)
@@ -546,6 +556,15 @@ private:
 		SHADER_PARAMETER_SRV(Buffer, MeshPositionBuffer6)
 		SHADER_PARAMETER_SRV(Buffer, MeshPositionBuffer7)
 
+		SHADER_PARAMETER_SRV(Buffer, MeshPreviousPositionBuffer0)
+		SHADER_PARAMETER_SRV(Buffer, MeshPreviousPositionBuffer1)
+		SHADER_PARAMETER_SRV(Buffer, MeshPreviousPositionBuffer2)
+		SHADER_PARAMETER_SRV(Buffer, MeshPreviousPositionBuffer3)
+		SHADER_PARAMETER_SRV(Buffer, MeshPreviousPositionBuffer4)
+		SHADER_PARAMETER_SRV(Buffer, MeshPreviousPositionBuffer5)
+		SHADER_PARAMETER_SRV(Buffer, MeshPreviousPositionBuffer6)
+		SHADER_PARAMETER_SRV(Buffer, MeshPreviousPositionBuffer7)
+
 		SHADER_PARAMETER_SRV(Buffer, MeshIndexBuffer0)
 		SHADER_PARAMETER_SRV(Buffer, MeshIndexBuffer1)
 		SHADER_PARAMETER_SRV(Buffer, MeshIndexBuffer2)
@@ -565,9 +584,13 @@ private:
 		SHADER_PARAMETER_SRV(Buffer, MeshUVsBuffer7)
 
 		SHADER_PARAMETER_RDG_BUFFER_SRV(Buffer, RootTriangleIndex)
-		SHADER_PARAMETER_RDG_BUFFER_UAV(RWStructuredBuffer, OutRootTrianglePosition0)
-		SHADER_PARAMETER_RDG_BUFFER_UAV(RWStructuredBuffer, OutRootTrianglePosition1)
-		SHADER_PARAMETER_RDG_BUFFER_UAV(RWStructuredBuffer, OutRootTrianglePosition2)
+		SHADER_PARAMETER_RDG_BUFFER_UAV(RWStructuredBuffer, OutRootTrianglePrevPosition0)
+		SHADER_PARAMETER_RDG_BUFFER_UAV(RWStructuredBuffer, OutRootTrianglePrevPosition1)
+		SHADER_PARAMETER_RDG_BUFFER_UAV(RWStructuredBuffer, OutRootTrianglePrevPosition2)
+
+		SHADER_PARAMETER_RDG_BUFFER_UAV(RWStructuredBuffer, OutRootTriangleCurrPosition0)
+		SHADER_PARAMETER_RDG_BUFFER_UAV(RWStructuredBuffer, OutRootTriangleCurrPosition1)
+		SHADER_PARAMETER_RDG_BUFFER_UAV(RWStructuredBuffer, OutRootTriangleCurrPosition2)
 	END_SHADER_PARAMETER_STRUCT()
 
 public:
@@ -576,7 +599,7 @@ public:
 	{
 		FGlobalShader::ModifyCompilationEnvironment(Parameters, OutEnvironment);
 		OutEnvironment.SetDefine(TEXT("MAX_SECTION_COUNT"), SectionArrayCount);
-		OutEnvironment.SetDefine(TEXT("SHADER_MESH_UPDATE"), SectionArrayCount);
+		OutEnvironment.SetDefine(TEXT("SHADER_MESH_UPDATE"), SectionArrayCount);		
 	}
 };
 
@@ -637,25 +660,35 @@ void AddHairStrandUpdateMeshTrianglesPass(
 	const TArray<uint32>& ValidSectionIndices = RestResources->BulkData.GetValidSectionIndices(LODIndex);
 	const uint32 ValidSectionCount = ValidSectionIndices.Num();
 	const uint32 PassCount = FMath::DivideAndRoundUp(ValidSectionCount, FHairUpdateMeshTriangleCS::SectionArrayCount);
+	const bool bComputePreviousDeformedPosition = IsHairStrandContinuousDecimationReorderingEnabled() && Type == HairStrandsTriangleType::DeformedPose;
 
 	FHairUpdateMeshTriangleCS::FParameters CommonParameters;
 
-	FRDGImportedBuffer OutputBuffers[3];
+	FRDGImportedBuffer OutputCurrBuffers[3];
+	FRDGImportedBuffer OutputPrevBuffers[3];
 	const bool bEnableUAVOverlap = true;
 	const ERDGUnorderedAccessViewFlags UAVFlags = bEnableUAVOverlap ? ERDGUnorderedAccessViewFlags::SkipBarrier : ERDGUnorderedAccessViewFlags::None;
 	CommonParameters.RootTriangleIndex = RegisterAsSRV(GraphBuilder, RestLODData.RootTriangleIndexBuffer);
 	if (Type == HairStrandsTriangleType::RestPose)
 	{
-		OutputBuffers[0] = Register(GraphBuilder, RestLODData.RestRootTrianglePosition0Buffer, ERDGImportedBufferFlags::CreateUAV, UAVFlags);
-		OutputBuffers[1] = Register(GraphBuilder, RestLODData.RestRootTrianglePosition1Buffer, ERDGImportedBufferFlags::CreateUAV, UAVFlags);
-		OutputBuffers[2] = Register(GraphBuilder, RestLODData.RestRootTrianglePosition2Buffer, ERDGImportedBufferFlags::CreateUAV, UAVFlags);
+		OutputCurrBuffers[0] = Register(GraphBuilder, RestLODData.RestRootTrianglePosition0Buffer, ERDGImportedBufferFlags::CreateUAV, UAVFlags);
+		OutputCurrBuffers[1] = Register(GraphBuilder, RestLODData.RestRootTrianglePosition1Buffer, ERDGImportedBufferFlags::CreateUAV, UAVFlags);
+		OutputCurrBuffers[2] = Register(GraphBuilder, RestLODData.RestRootTrianglePosition2Buffer, ERDGImportedBufferFlags::CreateUAV, UAVFlags);
 	}
 	else if (Type == HairStrandsTriangleType::DeformedPose)
 	{
 		FHairStrandsDeformedRootResource::FLOD& DeformedLODData = DeformedResources->LODs[LODIndex];
-		OutputBuffers[0] = Register(GraphBuilder, DeformedLODData.GetDeformedRootTrianglePosition0Buffer(FHairStrandsDeformedRootResource::FLOD::Current), ERDGImportedBufferFlags::CreateUAV, UAVFlags);
-		OutputBuffers[1] = Register(GraphBuilder, DeformedLODData.GetDeformedRootTrianglePosition1Buffer(FHairStrandsDeformedRootResource::FLOD::Current), ERDGImportedBufferFlags::CreateUAV, UAVFlags);
-		OutputBuffers[2] = Register(GraphBuilder, DeformedLODData.GetDeformedRootTrianglePosition2Buffer(FHairStrandsDeformedRootResource::FLOD::Current), ERDGImportedBufferFlags::CreateUAV, UAVFlags);
+		OutputCurrBuffers[0] = Register(GraphBuilder, DeformedLODData.GetDeformedRootTrianglePosition0Buffer(FHairStrandsDeformedRootResource::FLOD::Current), ERDGImportedBufferFlags::CreateUAV, UAVFlags);
+		OutputCurrBuffers[1] = Register(GraphBuilder, DeformedLODData.GetDeformedRootTrianglePosition1Buffer(FHairStrandsDeformedRootResource::FLOD::Current), ERDGImportedBufferFlags::CreateUAV, UAVFlags);
+		OutputCurrBuffers[2] = Register(GraphBuilder, DeformedLODData.GetDeformedRootTrianglePosition2Buffer(FHairStrandsDeformedRootResource::FLOD::Current), ERDGImportedBufferFlags::CreateUAV, UAVFlags);
+
+		if (bComputePreviousDeformedPosition)
+		{
+			OutputPrevBuffers[0] = Register(GraphBuilder, DeformedLODData.GetDeformedRootTrianglePosition0Buffer(FHairStrandsDeformedRootResource::FLOD::Previous), ERDGImportedBufferFlags::CreateUAV, UAVFlags);
+			OutputPrevBuffers[1] = Register(GraphBuilder, DeformedLODData.GetDeformedRootTrianglePosition1Buffer(FHairStrandsDeformedRootResource::FLOD::Previous), ERDGImportedBufferFlags::CreateUAV, UAVFlags);
+			OutputPrevBuffers[2] = Register(GraphBuilder, DeformedLODData.GetDeformedRootTrianglePosition2Buffer(FHairStrandsDeformedRootResource::FLOD::Previous), ERDGImportedBufferFlags::CreateUAV, UAVFlags);
+		}
+
 		DeformedLODData.Status = FHairStrandsDeformedRootResource::FLOD::EStatus::Completed;
 	}
 	else
@@ -664,9 +697,16 @@ void AddHairStrandUpdateMeshTrianglesPass(
 		return;
 	}
 
-	CommonParameters.OutRootTrianglePosition0 = OutputBuffers[0].UAV;
-	CommonParameters.OutRootTrianglePosition1 = OutputBuffers[1].UAV;
-	CommonParameters.OutRootTrianglePosition2 = OutputBuffers[2].UAV;
+	CommonParameters.OutRootTriangleCurrPosition0 = OutputCurrBuffers[0].UAV;
+	CommonParameters.OutRootTriangleCurrPosition1 = OutputCurrBuffers[1].UAV;
+	CommonParameters.OutRootTriangleCurrPosition2 = OutputCurrBuffers[2].UAV;
+
+	if (bComputePreviousDeformedPosition)
+	{
+		CommonParameters.OutRootTrianglePrevPosition0 = OutputPrevBuffers[0].UAV;
+		CommonParameters.OutRootTrianglePrevPosition1 = OutputPrevBuffers[1].UAV;
+		CommonParameters.OutRootTrianglePrevPosition2 = OutputPrevBuffers[2].UAV;
+	}
 
 	for (uint32 PassIt = 0; PassIt < PassCount; ++PassIt)
 	{
@@ -683,7 +723,9 @@ void AddHairStrandUpdateMeshTrianglesPass(
 		{
 			uint32 MeshSectionBufferIndex = 0;
 			FRDGBufferSRVRef RDGPositionBuffer = nullptr;
+			FRDGBufferSRVRef RDGPreviousPositionBuffer = nullptr;
 			FRHIShaderResourceView* PositionBuffer = nullptr;
+			FRHIShaderResourceView* PreviousPositionBuffer = nullptr;
 			FRHIShaderResourceView* IndexBuffer = nullptr;
 			FRHIShaderResourceView* UVsBuffer = nullptr;
 		};
@@ -691,76 +733,29 @@ void AddHairStrandUpdateMeshTrianglesPass(
 		TMap<FRDGBufferSRVRef, FMeshSectionBuffers>		UniqueMeshSectionBuffersRDG;
 		uint32 UniqueMeshSectionBufferIndex = 0;
 
+		#define SETPARAMETERS(OutParameters, InMeshSectionData, Index) \
+				OutParameters->RDGMeshPositionBuffer##Index			= InMeshSectionData.RDGPositionBuffer; \
+				OutParameters->RDGMeshPreviousPositionBuffer##Index = InMeshSectionData.RDGPreviousPositionBuffer; \
+				OutParameters->MeshPositionBuffer##Index			= InMeshSectionData.PositionBuffer; \
+				OutParameters->MeshPreviousPositionBuffer##Index	= InMeshSectionData.PreviousPositionBuffer; \
+				OutParameters->MeshIndexBuffer##Index				= InMeshSectionData.IndexBuffer; \
+				OutParameters->MeshUVsBuffer##Index					= InMeshSectionData.UVsBuffer;
+
 		auto SetMeshSectionBuffers = [Parameters](uint32 UniqueIndex, const FHairStrandsProjectionMeshData::Section& MeshSectionData)
 		{
 			switch (UniqueIndex)
 			{
-			case 0:
-			{
-				Parameters->RDGMeshPositionBuffer0 = MeshSectionData.RDGPositionBuffer;
-				Parameters->MeshPositionBuffer0 = MeshSectionData.PositionBuffer;
-				Parameters->MeshIndexBuffer0 = MeshSectionData.IndexBuffer;
-				Parameters->MeshUVsBuffer0 = MeshSectionData.UVsBuffer;
-				break;
-			}
-			case 1:
-			{
-				Parameters->RDGMeshPositionBuffer1 = MeshSectionData.RDGPositionBuffer;
-				Parameters->MeshPositionBuffer1 = MeshSectionData.PositionBuffer;
-				Parameters->MeshIndexBuffer1 = MeshSectionData.IndexBuffer;
-				Parameters->MeshUVsBuffer1 = MeshSectionData.UVsBuffer;
-				break;
-			}
-			case 2:
-			{
-				Parameters->RDGMeshPositionBuffer2 = MeshSectionData.RDGPositionBuffer;
-				Parameters->MeshPositionBuffer2 = MeshSectionData.PositionBuffer;
-				Parameters->MeshIndexBuffer2 = MeshSectionData.IndexBuffer;
-				Parameters->MeshUVsBuffer2 = MeshSectionData.UVsBuffer;
-				break;
-			}
-			case 3:
-			{
-				Parameters->RDGMeshPositionBuffer3 = MeshSectionData.RDGPositionBuffer;
-				Parameters->MeshPositionBuffer3 = MeshSectionData.PositionBuffer;
-				Parameters->MeshIndexBuffer3 = MeshSectionData.IndexBuffer;
-				Parameters->MeshUVsBuffer3 = MeshSectionData.UVsBuffer;
-				break;
-			}
-			case 4:
-			{
-				Parameters->RDGMeshPositionBuffer4 = MeshSectionData.RDGPositionBuffer;
-				Parameters->MeshPositionBuffer4 = MeshSectionData.PositionBuffer;
-				Parameters->MeshIndexBuffer4 = MeshSectionData.IndexBuffer;
-				Parameters->MeshUVsBuffer4 = MeshSectionData.UVsBuffer;
-				break;
-			}
-			case 5:
-			{
-				Parameters->RDGMeshPositionBuffer5 = MeshSectionData.RDGPositionBuffer;
-				Parameters->MeshPositionBuffer5 = MeshSectionData.PositionBuffer;
-				Parameters->MeshIndexBuffer5 = MeshSectionData.IndexBuffer;
-				Parameters->MeshUVsBuffer5 = MeshSectionData.UVsBuffer;
-				break;
-			}
-			case 6:
-			{
-				Parameters->RDGMeshPositionBuffer6 = MeshSectionData.RDGPositionBuffer;
-				Parameters->MeshPositionBuffer6 = MeshSectionData.PositionBuffer;
-				Parameters->MeshIndexBuffer6 = MeshSectionData.IndexBuffer;
-				Parameters->MeshUVsBuffer6 = MeshSectionData.UVsBuffer;
-				break;
-			}
-			case 7:
-			{
-				Parameters->RDGMeshPositionBuffer7 = MeshSectionData.RDGPositionBuffer;
-				Parameters->MeshPositionBuffer7 = MeshSectionData.PositionBuffer;
-				Parameters->MeshIndexBuffer7 = MeshSectionData.IndexBuffer;
-				Parameters->MeshUVsBuffer7 = MeshSectionData.UVsBuffer;
-				break;
-			}
+			case 0: SETPARAMETERS(Parameters, MeshSectionData, 0); break;
+			case 1: SETPARAMETERS(Parameters, MeshSectionData, 1); break;
+			case 2: SETPARAMETERS(Parameters, MeshSectionData, 2); break;
+			case 3: SETPARAMETERS(Parameters, MeshSectionData, 3); break;
+			case 4: SETPARAMETERS(Parameters, MeshSectionData, 4); break;
+			case 5: SETPARAMETERS(Parameters, MeshSectionData, 5); break;
+			case 6: SETPARAMETERS(Parameters, MeshSectionData, 6); break;
+			case 7: SETPARAMETERS(Parameters, MeshSectionData, 7); break;
 			}
 		};
+		#undef SETPARAMETERS
 
 		bool bUseRDGPositionBuffer = false;
 		for (int32 SectionStartIt = Parameters->Pass_SectionStart, SectionItEnd = Parameters->Pass_SectionStart + Parameters->Pass_SectionCount; SectionStartIt < SectionItEnd; ++SectionStartIt)
@@ -800,7 +795,9 @@ void AddHairStrandUpdateMeshTrianglesPass(
 			{
 				// Insure that all buffers actually match
 				check(Buffers->RDGPositionBuffer == MeshSectionData.RDGPositionBuffer);
+				check(Buffers->RDGPreviousPositionBuffer == MeshSectionData.RDGPreviousPositionBuffer);
 				check(Buffers->PositionBuffer == MeshSectionData.PositionBuffer);
+				check(Buffers->PreviousPositionBuffer == MeshSectionData.PreviousPositionBuffer);
 				check(Buffers->IndexBuffer == MeshSectionData.IndexBuffer);
 				check(Buffers->UVsBuffer == MeshSectionData.UVsBuffer);
 
@@ -815,7 +812,9 @@ void AddHairStrandUpdateMeshTrianglesPass(
 				FMeshSectionBuffers Entry;
 				Entry.MeshSectionBufferIndex = UniqueMeshSectionBufferIndex;
 				Entry.RDGPositionBuffer = MeshSectionData.RDGPositionBuffer;
+				Entry.RDGPreviousPositionBuffer = MeshSectionData.RDGPreviousPositionBuffer;
 				Entry.PositionBuffer = MeshSectionData.PositionBuffer;
+				Entry.PreviousPositionBuffer = MeshSectionData.PreviousPositionBuffer;
 				Entry.IndexBuffer = MeshSectionData.IndexBuffer;
 				Entry.UVsBuffer = MeshSectionData.UVsBuffer;
 
@@ -869,6 +868,7 @@ void AddHairStrandUpdateMeshTrianglesPass(
 
 		FHairUpdateMeshTriangleCS::FPermutationDomain PermutationVector;
 		PermutationVector.Set<FHairUpdateMeshTriangleCS::FUpdateUVs>(1);
+		PermutationVector.Set<FHairUpdateMeshTriangleCS::FPrevious>(bComputePreviousDeformedPosition);
 		PermutationVector.Set<FHairUpdateMeshTriangleCS::FPositionType>(bUseRDGPositionBuffer ? 1 : 0);
 
 		const FIntVector DispatchGroupCount = FComputeShaderUtils::GetGroupCount(RootCount, 128);
@@ -876,15 +876,23 @@ void AddHairStrandUpdateMeshTrianglesPass(
 		TShaderMapRef<FHairUpdateMeshTriangleCS> ComputeShader(ShaderMap, PermutationVector);
 		FComputeShaderUtils::AddPass(
 			GraphBuilder,
-			RDG_EVENT_NAME("HairStrandsTriangleMeshUpdate"),
+			RDG_EVENT_NAME("HairStrands::TriangleMeshUpdate(%s,%s)", bComputePreviousDeformedPosition ? TEXT("Previous/Current, Current") : TEXT("Current"), bUseRDGPositionBuffer ? TEXT("RDGBuffer") : TEXT("SkinCacheBuffer")),
 			ComputeShader,
 			PassParameters,
 			DispatchGroupCount);
 	}
 
-	GraphBuilder.SetBufferAccessFinal(OutputBuffers[0].Buffer, ERHIAccess::SRVMask);
-	GraphBuilder.SetBufferAccessFinal(OutputBuffers[1].Buffer, ERHIAccess::SRVMask);
-	GraphBuilder.SetBufferAccessFinal(OutputBuffers[2].Buffer, ERHIAccess::SRVMask);
+	GraphBuilder.SetBufferAccessFinal(OutputCurrBuffers[0].Buffer, ERHIAccess::SRVMask);
+	GraphBuilder.SetBufferAccessFinal(OutputCurrBuffers[1].Buffer, ERHIAccess::SRVMask);
+	GraphBuilder.SetBufferAccessFinal(OutputCurrBuffers[2].Buffer, ERHIAccess::SRVMask);
+
+	if (bComputePreviousDeformedPosition)
+	{
+		GraphBuilder.SetBufferAccessFinal(OutputPrevBuffers[0].Buffer, ERHIAccess::SRVMask);
+		GraphBuilder.SetBufferAccessFinal(OutputPrevBuffers[1].Buffer, ERHIAccess::SRVMask);
+		GraphBuilder.SetBufferAccessFinal(OutputPrevBuffers[2].Buffer, ERHIAccess::SRVMask);
+	}
+
 }
 
 ///////////////////////////////////////////////////////////////////////////////////////////////////
@@ -966,7 +974,7 @@ void AddHairStrandInterpolateMeshTrianglesPass(
 	TShaderMapRef<FHairInterpolateMeshTriangleCS> ComputeShader(ShaderMap);
 	FComputeShaderUtils::AddPass(
 		GraphBuilder,
-		RDG_EVENT_NAME("HairStrandsTriangleMeshInterpolate"),
+		RDG_EVENT_NAME("HairStrands::TriangleMeshInterpolate"),
 		ComputeShader,
 		Parameters,
 		DispatchGroupCount);
@@ -1048,7 +1056,7 @@ void InternalAddHairRBFInterpolationPass(
 	TShaderMapRef<FHairMeshesInterpolateCS> ComputeShader(ShaderMap);
 	FComputeShaderUtils::AddPass(
 		GraphBuilder,
-		RDG_EVENT_NAME("HairInterpolationRBF(%s)", bCards ? TEXT("Cards") : TEXT("Meshes")),
+		RDG_EVENT_NAME("HairStrands::HairInterpolationRBF(%s)", bCards ? TEXT("Cards") : TEXT("Meshes")),
 		ComputeShader,
 		Parameters,
 		DispatchGroupCount);
@@ -1236,65 +1244,34 @@ void AddHairStrandInitMeshSamplesPass(
 			{
 				uint32 SectionBufferIndex = 0;
 				FRDGBufferSRVRef RDGPositionBuffer = nullptr;
+				FRDGBufferSRVRef RDGPreviousPositionBuffer = nullptr;
 				FRHIShaderResourceView* PositionBuffer = nullptr;
+				FRHIShaderResourceView* PreviousPositionBuffer = nullptr;
 			};
 			TMap<FRHIShaderResourceView*, FMeshSectionBuffers>	UniqueMeshSectionBuffers;
 			TMap<FRDGBufferSRVRef, FMeshSectionBuffers>		UniqueMeshSectionBuffersRDG;
 
+			#define SETPARAMETERS(OutParameters, InMeshSectionData, Index) \
+				OutParameters.RDGVertexPositionsBuffer##Index = InMeshSectionData.RDGPositionBuffer; \
+				OutParameters.VertexPositionsBuffer##Index = InMeshSectionData.PositionBuffer;
+
 			auto SetMeshSectionBuffers = [&Parameters](uint32 UniqueIndex, const FHairStrandsProjectionMeshData::Section& MeshSectionData)
 			{
 				switch (UniqueIndex)
 				{
-				case 0:
-				{
-					Parameters.RDGVertexPositionsBuffer0 = MeshSectionData.RDGPositionBuffer;
-					Parameters.VertexPositionsBuffer0 = MeshSectionData.PositionBuffer;
-					break;
-				}
-				case 1:
-				{
-					Parameters.RDGVertexPositionsBuffer1 = MeshSectionData.RDGPositionBuffer;
-					Parameters.VertexPositionsBuffer1 = MeshSectionData.PositionBuffer;
-					break;
-				}
-				case 2:
-				{
-					Parameters.RDGVertexPositionsBuffer2 = MeshSectionData.RDGPositionBuffer;
-					Parameters.VertexPositionsBuffer2 = MeshSectionData.PositionBuffer;
-					break;
-				}
-				case 3:
-				{
-					Parameters.RDGVertexPositionsBuffer3 = MeshSectionData.RDGPositionBuffer;
-					Parameters.VertexPositionsBuffer3 = MeshSectionData.PositionBuffer;
-					break;
-				}
-				case 4:
-				{
-					Parameters.RDGVertexPositionsBuffer4 = MeshSectionData.RDGPositionBuffer;
-					Parameters.VertexPositionsBuffer4 = MeshSectionData.PositionBuffer;
-					break;
-				}
-				case 5:
-				{
-					Parameters.RDGVertexPositionsBuffer5 = MeshSectionData.RDGPositionBuffer;
-					Parameters.VertexPositionsBuffer5 = MeshSectionData.PositionBuffer;
-					break;
-				}
-				case 6:
-				{
-					Parameters.RDGVertexPositionsBuffer6 = MeshSectionData.RDGPositionBuffer;
-					Parameters.VertexPositionsBuffer6 = MeshSectionData.PositionBuffer;
-					break;
-				}
-				case 7:
-				{
-					Parameters.RDGVertexPositionsBuffer7 = MeshSectionData.RDGPositionBuffer;
-					Parameters.VertexPositionsBuffer7 = MeshSectionData.PositionBuffer;
-					break;
-				}
+				case 0: SETPARAMETERS(Parameters, MeshSectionData, 0); break;
+				case 1: SETPARAMETERS(Parameters, MeshSectionData, 1); break;
+				case 2: SETPARAMETERS(Parameters, MeshSectionData, 2); break;
+				case 3: SETPARAMETERS(Parameters, MeshSectionData, 3); break;
+				case 4: SETPARAMETERS(Parameters, MeshSectionData, 4); break;
+				case 5: SETPARAMETERS(Parameters, MeshSectionData, 5); break;
+				case 6: SETPARAMETERS(Parameters, MeshSectionData, 6); break;
+				case 7: SETPARAMETERS(Parameters, MeshSectionData, 7); break;
 				}
 			};
+
+			#undef SETPARAMETERS
+
 			uint32 UniqueMeshSectionBufferIndex = 0;
 			bool bUseRDGPositionBuffer = false;
 			for (int32 SectionStartIt = PassSectionStart, SectionItEnd = PassSectionStart + PassSectionCount; SectionStartIt < SectionItEnd; ++SectionStartIt)
@@ -1466,7 +1443,7 @@ void AddHairStrandUpdateMeshSamplesPass(
 		TShaderMapRef<FHairUpdateMeshSamplesCS> ComputeShader(ShaderMap);
 		FComputeShaderUtils::AddPass(
 			GraphBuilder,
-			RDG_EVENT_NAME("HairStrandsUpdateMeshSamples"),
+			RDG_EVENT_NAME("HairStrands::UpdateMeshSamples"),
 			ComputeShader,
 			Parameters,
 			DispatchGroupCount);
@@ -1797,12 +1774,16 @@ private:
 	SHADER_USE_PARAMETER_STRUCT(FHairUpdatePositionOffsetCS, FGlobalShader);
 
 	class FUseGPUOffset : SHADER_PERMUTATION_BOOL("PERMUTATION_USE_GPU_OFFSET");
-	using FPermutationDomain = TShaderPermutationDomain<FUseGPUOffset>;
+	class FPrevious : SHADER_PERMUTATION_BOOL("PERMUTATION_PREVIOUS");
+	using FPermutationDomain = TShaderPermutationDomain<FUseGPUOffset, FPrevious>;
 
 	BEGIN_SHADER_PARAMETER_STRUCT(FParameters, )
-		SHADER_PARAMETER(FVector3f, CPUPositionOffset)
-		SHADER_PARAMETER_RDG_BUFFER_SRV(StructuredBuffer, RootTrianglePosition0Buffer)
-		SHADER_PARAMETER_RDG_BUFFER_UAV(RWBuffer, OutOffsetBuffer)
+		SHADER_PARAMETER(FVector3f, CPUCurrPositionOffset)
+		SHADER_PARAMETER(FVector3f, CPUPrevPositionOffset)
+		SHADER_PARAMETER_RDG_BUFFER_SRV(StructuredBuffer, RootTriangleCurrPosition0Buffer)
+		SHADER_PARAMETER_RDG_BUFFER_SRV(StructuredBuffer, RootTrianglePrevPosition0Buffer)
+		SHADER_PARAMETER_RDG_BUFFER_UAV(RWBuffer, OutCurrOffsetBuffer)
+		SHADER_PARAMETER_RDG_BUFFER_UAV(RWBuffer, OutPrevOffsetBuffer)
 	END_SHADER_PARAMETER_STRUCT()
 
 public:
@@ -1831,29 +1812,42 @@ void AddHairStrandUpdatePositionOffsetPass(
 		return;
 	}
 
-	FRDGImportedBuffer OutPositionOffsetBuffer    = Register(GraphBuilder, DeformedResources->GetPositionOffsetBuffer(FHairStrandsDeformedResource::Current), ERDGImportedBufferFlags::CreateUAV);
-	FRDGImportedBuffer RootTrianglePositionBuffer;
+	const bool bComputePreviousDeformedPosition = IsHairStrandContinuousDecimationReorderingEnabled();
+
+	FRDGImportedBuffer OutCurrPositionOffsetBuffer = Register(GraphBuilder, DeformedResources->GetPositionOffsetBuffer(FHairStrandsDeformedResource::Current), ERDGImportedBufferFlags::CreateUAV);
+	FRDGImportedBuffer OutPrevPositionOffsetBuffer = Register(GraphBuilder, DeformedResources->GetPositionOffsetBuffer(FHairStrandsDeformedResource::Previous), ERDGImportedBufferFlags::CreateUAV);
+
+	FRDGImportedBuffer RootTriangleCurrPositionBuffer;
+	FRDGImportedBuffer RootTrianglePrevPositionBuffer;
 	if (DeformedRootResources)
 	{
-		RootTrianglePositionBuffer = Register(GraphBuilder, DeformedRootResources->LODs[LODIndex].GetDeformedRootTrianglePosition0Buffer(FHairStrandsDeformedRootResource::FLOD::Current), ERDGImportedBufferFlags::CreateSRV);
+		RootTriangleCurrPositionBuffer = Register(GraphBuilder, DeformedRootResources->LODs[LODIndex].GetDeformedRootTrianglePosition0Buffer(FHairStrandsDeformedRootResource::FLOD::Current), ERDGImportedBufferFlags::CreateSRV);
+		RootTrianglePrevPositionBuffer = Register(GraphBuilder, DeformedRootResources->LODs[LODIndex].GetDeformedRootTrianglePosition0Buffer(FHairStrandsDeformedRootResource::FLOD::Previous), ERDGImportedBufferFlags::CreateSRV);
 	}
 
 	const bool bUseGPUOffset = DeformedRootResources != nullptr && GHairStrandsUseGPUPositionOffset > 0;
-	const uint32 OffsetIndex = DeformedResources->GetIndex(FHairStrandsDeformedResource::Current);
+	const uint32 CurrOffsetIndex = DeformedResources->GetIndex(FHairStrandsDeformedResource::Current);
+	const uint32 PrevOffsetIndex = DeformedResources->GetIndex(FHairStrandsDeformedResource::Previous);
+
 	FHairUpdatePositionOffsetCS::FParameters* PassParameters = GraphBuilder.AllocParameters<FHairUpdatePositionOffsetCS::FParameters>();
-	PassParameters->CPUPositionOffset = (FVector3f)DeformedResources->PositionOffset[OffsetIndex];
-	PassParameters->RootTrianglePosition0Buffer = RootTrianglePositionBuffer.SRV;
-	PassParameters->OutOffsetBuffer = OutPositionOffsetBuffer.UAV;
+	PassParameters->CPUCurrPositionOffset			= (FVector3f)DeformedResources->PositionOffset[CurrOffsetIndex];
+	PassParameters->CPUPrevPositionOffset			= (FVector3f)DeformedResources->PositionOffset[PrevOffsetIndex];
+	PassParameters->RootTriangleCurrPosition0Buffer	= RootTriangleCurrPositionBuffer.SRV;
+	PassParameters->RootTrianglePrevPosition0Buffer = RootTrianglePrevPositionBuffer.SRV;
+	PassParameters->OutCurrOffsetBuffer				= OutCurrPositionOffsetBuffer.UAV;
+	PassParameters->OutPrevOffsetBuffer				= OutPrevPositionOffsetBuffer.UAV;
 
 	FHairUpdatePositionOffsetCS::FPermutationDomain PermutationVector;
 	PermutationVector.Set<FHairUpdatePositionOffsetCS::FUseGPUOffset>(bUseGPUOffset);
+	PermutationVector.Set<FHairUpdatePositionOffsetCS::FPrevious>(bComputePreviousDeformedPosition);
 	TShaderMapRef<FHairUpdatePositionOffsetCS> ComputeShader(ShaderMap, PermutationVector);
 	FComputeShaderUtils::AddPass(
 		GraphBuilder,
-		RDG_EVENT_NAME("HairStrandsUpdatePositionOffset"),
+		RDG_EVENT_NAME("HairStrands::UpdatePositionOffset(%s,%s)", bUseGPUOffset ? TEXT("GPU") : TEXT("CPU"), bComputePreviousDeformedPosition ? TEXT("Current/Previous") : TEXT("Current")),
 		ComputeShader,
 		PassParameters,
 		FIntVector(1, 1, 1));
 
-	GraphBuilder.SetBufferAccessFinal(OutPositionOffsetBuffer.Buffer, ERHIAccess::SRVMask);
+	GraphBuilder.SetBufferAccessFinal(OutCurrPositionOffsetBuffer.Buffer, ERHIAccess::SRVMask);
+	GraphBuilder.SetBufferAccessFinal(OutPrevPositionOffsetBuffer.Buffer, ERHIAccess::SRVMask);
 }
\ No newline at end of file
diff --git a/Engine/Plugins/Runtime/HairStrands/Source/HairStrandsCore/Public/HairStrandsMeshProjection.h b/Engine/Plugins/Runtime/HairStrands/Source/HairStrandsCore/Public/HairStrandsMeshProjection.h
--- a/Engine/Plugins/Runtime/HairStrands/Source/HairStrandsCore/Public/HairStrandsMeshProjection.h
+++ b/Engine/Plugins/Runtime/HairStrands/Source/HairStrandsCore/Public/HairStrandsMeshProjection.h
@@ -24,7 +24,9 @@ struct FHairStrandsProjectionMeshData
 	{
 		FTransform LocalToWorld;
 		FRDGBufferSRVRef RDGPositionBuffer = nullptr;
+		FRDGBufferSRVRef RDGPreviousPositionBuffer = nullptr;
 		FRHIShaderResourceView* PositionBuffer = nullptr;
+		FRHIShaderResourceView* PreviousPositionBuffer = nullptr;
 		FRHIShaderResourceView* UVsBuffer = nullptr;
 		FRHIShaderResourceView* IndexBuffer = nullptr;
 		uint32 UVsChannelCount = 0;
diff --git a/Engine/Shaders/Private/HairStrands/HairStrandsMeshUpdate.usf b/Engine/Shaders/Private/HairStrands/HairStrandsMeshUpdate.usf
--- a/Engine/Shaders/Private/HairStrands/HairStrandsMeshUpdate.usf
+++ b/Engine/Shaders/Private/HairStrands/HairStrandsMeshUpdate.usf
@@ -41,6 +41,7 @@ Buffer<uint> MeshIndexBuffer7;
 // (while the skin cache system is not converted to RDG). If skin cache is disable, a manual skin compute 
 // update happens in the hair system, and its output are RDG buffers.
 #if PERMUTATION_POSITION_TYPE == 0
+
 Buffer<float> MeshPositionBuffer0;
 Buffer<float> MeshPositionBuffer1;
 Buffer<float> MeshPositionBuffer2;
@@ -49,7 +50,18 @@ Buffer<float> MeshPositionBuffer4;
 Buffer<float> MeshPositionBuffer5;
 Buffer<float> MeshPositionBuffer6;
 Buffer<float> MeshPositionBuffer7;
+
+Buffer<float> MeshPreviousPositionBuffer0;
+Buffer<float> MeshPreviousPositionBuffer1;
+Buffer<float> MeshPreviousPositionBuffer2;
+Buffer<float> MeshPreviousPositionBuffer3;
+Buffer<float> MeshPreviousPositionBuffer4;
+Buffer<float> MeshPreviousPositionBuffer5;
+Buffer<float> MeshPreviousPositionBuffer6;
+Buffer<float> MeshPreviousPositionBuffer7;
+
 #else
+
 Buffer<float> RDGMeshPositionBuffer0;
 Buffer<float> RDGMeshPositionBuffer1;
 Buffer<float> RDGMeshPositionBuffer2;
@@ -59,6 +71,15 @@ Buffer<float> RDGMeshPositionBuffer5;
 Buffer<float> RDGMeshPositionBuffer6;
 Buffer<float> RDGMeshPositionBuffer7;
 
+Buffer<float> RDGMeshPreviousPositionBuffer0;
+Buffer<float> RDGMeshPreviousPositionBuffer1;
+Buffer<float> RDGMeshPreviousPositionBuffer2;
+Buffer<float> RDGMeshPreviousPositionBuffer3;
+Buffer<float> RDGMeshPreviousPositionBuffer4;
+Buffer<float> RDGMeshPreviousPositionBuffer5;
+Buffer<float> RDGMeshPreviousPositionBuffer6;
+Buffer<float> RDGMeshPreviousPositionBuffer7;
+
 #define MeshPositionBuffer0 RDGMeshPositionBuffer0
 #define MeshPositionBuffer1 RDGMeshPositionBuffer1
 #define MeshPositionBuffer2 RDGMeshPositionBuffer2
@@ -67,6 +88,16 @@ Buffer<float> RDGMeshPositionBuffer7;
 #define MeshPositionBuffer5 RDGMeshPositionBuffer5
 #define MeshPositionBuffer6 RDGMeshPositionBuffer6
 #define MeshPositionBuffer7 RDGMeshPositionBuffer7
+
+#define MeshPreviousPositionBuffer0 RDGMeshPreviousPositionBuffer0
+#define MeshPreviousPositionBuffer1 RDGMeshPreviousPositionBuffer1
+#define MeshPreviousPositionBuffer2 RDGMeshPreviousPositionBuffer2
+#define MeshPreviousPositionBuffer3 RDGMeshPreviousPositionBuffer3
+#define MeshPreviousPositionBuffer4 RDGMeshPreviousPositionBuffer4
+#define MeshPreviousPositionBuffer5 RDGMeshPreviousPositionBuffer5
+#define MeshPreviousPositionBuffer6 RDGMeshPreviousPositionBuffer6
+#define MeshPreviousPositionBuffer7 RDGMeshPreviousPositionBuffer7
+
 #endif 
 
 #if PERMUTATION_WITHUV == 1
@@ -84,84 +115,52 @@ Buffer<float2> MeshUVsBuffer8;
 uint MaxRootCount;
 Buffer<uint> RootTriangleIndex;
 
-RWBuffer<float4> OutRootTrianglePosition0;
-RWBuffer<float4> OutRootTrianglePosition1;
-RWBuffer<float4> OutRootTrianglePosition2;
+RWBuffer<float4> OutRootTriangleCurrPosition0;
+RWBuffer<float4> OutRootTriangleCurrPosition1;
+RWBuffer<float4> OutRootTriangleCurrPosition2;
+
+#if PERMUTATION_PREVIOUS
+RWBuffer<float4> OutRootTrianglePrevPosition0;
+RWBuffer<float4> OutRootTrianglePrevPosition1;
+RWBuffer<float4> OutRootTrianglePrevPosition2;
+#endif
 
 #if PERMUTATION_WITHUV == 1
-FHairMeshTriangleWithUVs GetTriangleData(FHairTriangleIndex TriangleIndex, uint SectionIndex)
+FHairMeshTriangleWithUVs GetTriangleData(FHairTriangleIndex TriangleIndex, uint SectionIndex, bool bPrevious)
 {
+	#define GetTriangleDataAtIndex(Index) GetTriangleData(TriangleIndex.TriangleIndex, bPrevious ? MeshPreviousPositionBuffer##Index : MeshPositionBuffer##Index, MeshUVsBuffer##Index, MeshIndexBuffer##Index, MeshIndexOffset(SectionIndex), MeshMaxIndexCount(SectionIndex), MeshMaxVertexCount(SectionIndex), MeshUVsChannelOffset(SectionIndex), MeshUVsChannelCount(SectionIndex));
+
 	const uint MeshSectionIndex = MeshSectionBufferIndex(SectionIndex);
-	if (MeshSectionIndex == 0)
-	{
-		return GetTriangleData(TriangleIndex.TriangleIndex, MeshPositionBuffer0, MeshUVsBuffer0, MeshIndexBuffer0, MeshIndexOffset(SectionIndex), MeshMaxIndexCount(SectionIndex), MeshMaxVertexCount(SectionIndex), MeshUVsChannelOffset(SectionIndex), MeshUVsChannelCount(SectionIndex));
-	}
-	else if (MeshSectionIndex == 1)
-	{
-		return GetTriangleData(TriangleIndex.TriangleIndex, MeshPositionBuffer1, MeshUVsBuffer1, MeshIndexBuffer1, MeshIndexOffset(SectionIndex), MeshMaxIndexCount(SectionIndex), MeshMaxVertexCount(SectionIndex), MeshUVsChannelOffset(SectionIndex), MeshUVsChannelCount(SectionIndex));
-	}
-	else if (MeshSectionIndex == 2)
-	{
-		return GetTriangleData(TriangleIndex.TriangleIndex, MeshPositionBuffer2, MeshUVsBuffer2, MeshIndexBuffer2, MeshIndexOffset(SectionIndex), MeshMaxIndexCount(SectionIndex), MeshMaxVertexCount(SectionIndex), MeshUVsChannelOffset(SectionIndex), MeshUVsChannelCount(SectionIndex));
-	}
-	else if (MeshSectionIndex == 3)
-	{
-		return GetTriangleData(TriangleIndex.TriangleIndex, MeshPositionBuffer3, MeshUVsBuffer3, MeshIndexBuffer3, MeshIndexOffset(SectionIndex), MeshMaxIndexCount(SectionIndex), MeshMaxVertexCount(SectionIndex), MeshUVsChannelOffset(SectionIndex), MeshUVsChannelCount(SectionIndex));
-	}
-	else if (MeshSectionIndex == 4)
-	{
-		return GetTriangleData(TriangleIndex.TriangleIndex, MeshPositionBuffer4, MeshUVsBuffer4, MeshIndexBuffer4, MeshIndexOffset(SectionIndex), MeshMaxIndexCount(SectionIndex), MeshMaxVertexCount(SectionIndex), MeshUVsChannelOffset(SectionIndex), MeshUVsChannelCount(SectionIndex));
-	}
-	else if (MeshSectionIndex == 5)
-	{
-		return GetTriangleData(TriangleIndex.TriangleIndex, MeshPositionBuffer5, MeshUVsBuffer5, MeshIndexBuffer5, MeshIndexOffset(SectionIndex), MeshMaxIndexCount(SectionIndex), MeshMaxVertexCount(SectionIndex), MeshUVsChannelOffset(SectionIndex), MeshUVsChannelCount(SectionIndex));
-	}
-	else if (MeshSectionIndex == 6)
-	{
-		return GetTriangleData(TriangleIndex.TriangleIndex, MeshPositionBuffer6, MeshUVsBuffer6, MeshIndexBuffer6, MeshIndexOffset(SectionIndex), MeshMaxIndexCount(SectionIndex), MeshMaxVertexCount(SectionIndex), MeshUVsChannelOffset(SectionIndex), MeshUVsChannelCount(SectionIndex));
-	}
-	else if (MeshSectionIndex == 7)
-	{
-		return GetTriangleData(TriangleIndex.TriangleIndex, MeshPositionBuffer7, MeshUVsBuffer7, MeshIndexBuffer7, MeshIndexOffset(SectionIndex), MeshMaxIndexCount(SectionIndex), MeshMaxVertexCount(SectionIndex), MeshUVsChannelOffset(SectionIndex), MeshUVsChannelCount(SectionIndex));
-	}
+		 if (MeshSectionIndex == 0) { return GetTriangleDataAtIndex(0); } 
+	else if (MeshSectionIndex == 1) { return GetTriangleDataAtIndex(1); }
+	else if (MeshSectionIndex == 2) { return GetTriangleDataAtIndex(2); }
+	else if (MeshSectionIndex == 3) { return GetTriangleDataAtIndex(3); }
+	else if (MeshSectionIndex == 4) { return GetTriangleDataAtIndex(4); }
+	else if (MeshSectionIndex == 5) { return GetTriangleDataAtIndex(5); }
+	else if (MeshSectionIndex == 6) { return GetTriangleDataAtIndex(6); }
+	else if (MeshSectionIndex == 7) { return GetTriangleDataAtIndex(7); }
+	
+	#undef GetTriangleDataAtIndex
+
 	return (FHairMeshTriangleWithUVs)0;
 }
 #else
-FHairMeshTriangle GetTriangleData(FHairTriangleIndex TriangleIndex, uint SectionIndex)
+FHairMeshTriangle GetTriangleData(FHairTriangleIndex TriangleIndex, uint SectionIndex, bool bPrevious)
 {
+	#define GetTriangleDataAtIndex(Index) GetTriangleData(TriangleIndex.TriangleIndex, bPrevious ? MeshPreviousPositionBuffer##Index : MeshPositionBuffer##Index, MeshIndexBuffer##Index, MeshIndexOffset(SectionIndex), MeshMaxIndexCount(SectionIndex), MeshMaxVertexCount(SectionIndex));
+
 	const uint MeshSectionIndex = MeshSectionBufferIndex(SectionIndex);
-	if (MeshSectionIndex == 0)
-	{
-		return GetTriangleData(TriangleIndex.TriangleIndex, MeshPositionBuffer0, MeshIndexBuffer0, MeshIndexOffset(SectionIndex), MeshMaxIndexCount(SectionIndex), MeshMaxVertexCount(SectionIndex));
-	}
-	else if (MeshSectionIndex == 1)
-	{
-		return GetTriangleData(TriangleIndex.TriangleIndex, MeshPositionBuffer1, MeshIndexBuffer1, MeshIndexOffset(SectionIndex), MeshMaxIndexCount(SectionIndex), MeshMaxVertexCount(SectionIndex));
-	}
-	else if (MeshSectionIndex == 2)
-	{
-		return GetTriangleData(TriangleIndex.TriangleIndex, MeshPositionBuffer2, MeshIndexBuffer2, MeshIndexOffset(SectionIndex), MeshMaxIndexCount(SectionIndex), MeshMaxVertexCount(SectionIndex));
-	}
-	else if (MeshSectionIndex == 3)
-	{
-		return GetTriangleData(TriangleIndex.TriangleIndex, MeshPositionBuffer3, MeshIndexBuffer3, MeshIndexOffset(SectionIndex), MeshMaxIndexCount(SectionIndex), MeshMaxVertexCount(SectionIndex));
-	}
-	else if (MeshSectionIndex == 4)
-	{
-		return GetTriangleData(TriangleIndex.TriangleIndex, MeshPositionBuffer4, MeshIndexBuffer4, MeshIndexOffset(SectionIndex), MeshMaxIndexCount(SectionIndex), MeshMaxVertexCount(SectionIndex));
-	}
-	else if (MeshSectionIndex == 5)
-	{
-		return GetTriangleData(TriangleIndex.TriangleIndex, MeshPositionBuffer5, MeshIndexBuffer5, MeshIndexOffset(SectionIndex), MeshMaxIndexCount(SectionIndex), MeshMaxVertexCount(SectionIndex));
-	}
-	else if (MeshSectionIndex == 6)
-	{
-		return GetTriangleData(TriangleIndex.TriangleIndex, MeshPositionBuffer6, MeshIndexBuffer6, MeshIndexOffset(SectionIndex), MeshMaxIndexCount(SectionIndex), MeshMaxVertexCount(SectionIndex));
-	}
-	else if (MeshSectionIndex == 7)
-	{
-		return GetTriangleData(TriangleIndex.TriangleIndex, MeshPositionBuffer7, MeshIndexBuffer7, MeshIndexOffset(SectionIndex), MeshMaxIndexCount(SectionIndex), MeshMaxVertexCount(SectionIndex));
-	}
+		 if (MeshSectionIndex == 0) { return GetTriangleDataAtIndex(0); }
+	else if (MeshSectionIndex == 1) { return GetTriangleDataAtIndex(1); }
+	else if (MeshSectionIndex == 2) { return GetTriangleDataAtIndex(2); }
+	else if (MeshSectionIndex == 3) { return GetTriangleDataAtIndex(3); }
+	else if (MeshSectionIndex == 4) { return GetTriangleDataAtIndex(4); }
+	else if (MeshSectionIndex == 5) { return GetTriangleDataAtIndex(5); }
+	else if (MeshSectionIndex == 6) { return GetTriangleDataAtIndex(6); }
+	else if (MeshSectionIndex == 7) { return GetTriangleDataAtIndex(7); }
+
+	#undef GetTriangleDataAtIndex
+
 	return (FHairMeshTriangle)0;
 }
 #endif
@@ -196,17 +195,37 @@ void MainCS(uint3 DispatchThreadId : SV_DispatchThreadID)
 
 	// Compute the slod index of the resources for the section index (this is done in multiple pass
 #if PERMUTATION_WITHUV == 1
-	const FHairMeshTriangleWithUVs Tri = GetTriangleData(TriangleIndex, SectionSlotIndex);
+	{
+		const FHairMeshTriangleWithUVs Tri = GetTriangleData(TriangleIndex, SectionSlotIndex, false);
+		OutRootTriangleCurrPosition0[RootIndex] = float4(Tri.P0, PackUVsToFloat(Tri.UV0));
+		OutRootTriangleCurrPosition1[RootIndex] = float4(Tri.P1, PackUVsToFloat(Tri.UV1));
+		OutRootTriangleCurrPosition2[RootIndex] = float4(Tri.P2, PackUVsToFloat(Tri.UV2));
+	}
 
-	OutRootTrianglePosition0[RootIndex] = float4(Tri.P0, PackUVsToFloat(Tri.UV0));
-	OutRootTrianglePosition1[RootIndex] = float4(Tri.P1, PackUVsToFloat(Tri.UV1));
-	OutRootTrianglePosition2[RootIndex] = float4(Tri.P2, PackUVsToFloat(Tri.UV2));
+	#if PERMUTATION_PREVIOUS
+	{
+		const FHairMeshTriangleWithUVs Tri = GetTriangleData(TriangleIndex, SectionSlotIndex, true);
+		OutRootTrianglePrevPosition0[RootIndex] = float4(Tri.P0, PackUVsToFloat(Tri.UV0));
+		OutRootTrianglePrevPosition1[RootIndex] = float4(Tri.P1, PackUVsToFloat(Tri.UV1));
+		OutRootTrianglePrevPosition2[RootIndex] = float4(Tri.P2, PackUVsToFloat(Tri.UV2));
+	}
+	#endif
 #else
-	const FHairMeshTriangle Tri = GetTriangleData(TriangleIndex, SectionSlotIndex);
+	{
+		const FHairMeshTriangle Tri = GetTriangleData(TriangleIndex, SectionSlotIndex, false);
+		OutRootTriangleCurrPosition0[RootIndex] = (float4(Tri.P0, 1));
+		OutRootTriangleCurrPosition1[RootIndex] = (float4(Tri.P1, 1));
+		OutRootTriangleCurrPosition2[RootIndex] = (float4(Tri.P2, 1));
+	}
 
-	OutRootTrianglePosition0[RootIndex] = (float4(Tri.P0, 1));
-	OutRootTrianglePosition1[RootIndex] = (float4(Tri.P1, 1));
-	OutRootTrianglePosition2[RootIndex] = (float4(Tri.P2, 1));
+	#if PERMUTATION_PREVIOUS
+	{
+		const FHairMeshTriangle Tri = GetTriangleData(TriangleIndex, SectionSlotIndex, true);
+		OutRootTrianglePrevPosition0[RootIndex] = (float4(Tri.P0, 1));
+		OutRootTrianglePrevPosition1[RootIndex] = (float4(Tri.P1, 1));
+		OutRootTrianglePrevPosition2[RootIndex] = (float4(Tri.P2, 1));
+	}
+	#endif
 #endif
 }
 
@@ -215,20 +234,25 @@ void MainCS(uint3 DispatchThreadId : SV_DispatchThreadID)
 #if SHADER_OFFSET_UPDATE
 
 #if PERMUTATION_USE_GPU_OFFSET
-Buffer<float4> RootTrianglePosition0Buffer;
-float3 GetOffsetPosition()
-{
-	return RootTrianglePosition0Buffer[0].xyz;
-}
+	Buffer<float4> RootTriangleCurrPosition0Buffer;
+	float3 GetOffsetCurrPosition()	{ return RootTriangleCurrPosition0Buffer[0].xyz; }
+
+	#if PERMUTATION_PREVIOUS
+	Buffer<float4> RootTrianglePrevPosition0Buffer;
+	float3 GetOffsetPrevPosition()	{ return RootTrianglePrevPosition0Buffer[0].xyz; }
+	#endif
 #else
-float3 CPUPositionOffset;
-float3 GetOffsetPosition()
-{
-	return CPUPositionOffset;
-}
+	float3 CPUCurrPositionOffset;
+	float3 CPUPrevPositionOffset;
+
+	float3 GetOffsetCurrPosition()	{ return CPUCurrPositionOffset; }
+	float3 GetOffsetPrevPosition()	{ return CPUPrevPositionOffset; }
 #endif
 
-RWBuffer<float4> OutOffsetBuffer;
+RWBuffer<float4> OutCurrOffsetBuffer;
+#if PERMUTATION_PREVIOUS
+RWBuffer<float4> OutPrevOffsetBuffer;
+#endif
 
 [numthreads(1, 1, 1)]
 void MainCS(uint3 DispatchThreadId : SV_DispatchThreadID)
@@ -236,7 +260,10 @@ void MainCS(uint3 DispatchThreadId : SV_DispatchThreadID)
 	if (DispatchThreadId.x > 0)
 		return;
 
-	OutOffsetBuffer[0] = float4(GetOffsetPosition(), 1);
+	OutCurrOffsetBuffer[0] = float4(GetOffsetCurrPosition(), 1);
+#if PERMUTATION_PREVIOUS
+	OutPrevOffsetBuffer[0] = float4(GetOffsetPrevPosition(), 1);
+#endif
 }
 
 #endif // SHADER_OFFSET_UPDATE
\ No newline at end of file
diff --git a/Engine/Source/Runtime/Engine/Private/GPUSkinCache.cpp b/Engine/Source/Runtime/Engine/Private/GPUSkinCache.cpp
--- a/Engine/Source/Runtime/Engine/Private/GPUSkinCache.cpp
+++ b/Engine/Source/Runtime/Engine/Private/GPUSkinCache.cpp
@@ -436,6 +436,7 @@ public:
 		{
 			const FSkelMeshRenderSection& Section = *DispatchData[SectionIndex].Section;
 			MeshSection.PositionBuffer = DispatchData[SectionIndex].PositionBuffer->Buffer.SRV;
+			MeshSection.PreviousPositionBuffer = DispatchData[SectionIndex].PreviousPositionBuffer->Buffer.SRV;
 			MeshSection.UVsBuffer = DispatchData[SectionIndex].UVsBufferSRV;
 			MeshSection.TotalVertexCount = DispatchData[SectionIndex].PositionBuffer->Buffer.NumBytes / (sizeof(float) * 3);
 			MeshSection.NumPrimitives = Section.NumTriangles;
diff --git a/Engine/Source/Runtime/Engine/Public/GPUSkinCache.h b/Engine/Source/Runtime/Engine/Public/GPUSkinCache.h
--- a/Engine/Source/Runtime/Engine/Public/GPUSkinCache.h
+++ b/Engine/Source/Runtime/Engine/Public/GPUSkinCache.h
@@ -99,8 +99,10 @@ struct FCachedGeometry
 {
 	struct Section
 	{
-		FRDGBufferSRVRef RDGPositionBuffer = nullptr;		// Valid when the input comes from a manual skin cache (i.e. skinned run into compute on demand)
-		FRHIShaderResourceView* PositionBuffer = nullptr;	// Valid when the input comes from the skin cached (since it is not convert yet to RDG)
+		FRDGBufferSRVRef RDGPositionBuffer = nullptr;				// Valid when the input comes from a manual skin cache (i.e. skinned run into compute on demand)
+		FRDGBufferSRVRef RDGPreviousPositionBuffer = nullptr;		// Valid when the input comes from a manual skin cache (i.e. skinned run into compute on demand)
+		FRHIShaderResourceView* PositionBuffer = nullptr;			// Valid when the input comes from the skin cached (since it is not convert yet to RDG)
+		FRHIShaderResourceView* PreviousPositionBuffer = nullptr;	// Valid when the input comes from the skin cached (since it is not convert yet to RDG)
 		FRHIShaderResourceView* UVsBuffer = nullptr;
 		FRHIShaderResourceView* IndexBuffer = nullptr;
 		uint32 UVsChannelOffset = 0;
@@ -118,6 +120,7 @@ struct FCachedGeometry
 	int32 LODIndex = 0;
 	TArray<Section> Sections;
 	FRDGBufferRef DeformedPositionBuffer = nullptr;
+	FRDGBufferRef DeformedPreviousPositionBuffer = nullptr;
 	FTransform LocalToWorld = FTransform::Identity;
 };
 
