Remove unsed mesh interpolation global shader for hair code.

diff --git a/Engine/Plugins/Runtime/HairStrands/Source/HairStrandsCore/Private/HairStrandsMeshProjection.cpp b/Engine/Plugins/Runtime/HairStrands/Source/HairStrandsCore/Private/HairStrandsMeshProjection.cpp
--- a/Engine/Plugins/Runtime/HairStrands/Source/HairStrandsCore/Private/HairStrandsMeshProjection.cpp
+++ b/Engine/Plugins/Runtime/HairStrands/Source/HairStrandsCore/Private/HairStrandsMeshProjection.cpp
@@ -895,96 +895,6 @@ void AddHairStrandUpdateMeshTrianglesPass(
 
 }
 
-///////////////////////////////////////////////////////////////////////////////////////////////////
-
-class FHairInterpolateMeshTriangleCS : public FGlobalShader
-{
-private:
-	DECLARE_GLOBAL_SHADER(FHairInterpolateMeshTriangleCS);
-	SHADER_USE_PARAMETER_STRUCT(FHairInterpolateMeshTriangleCS, FGlobalShader);
-
-	using FPermutationDomain = TShaderPermutationDomain<>;
-
-	BEGIN_SHADER_PARAMETER_STRUCT(FParameters, )
-		SHADER_PARAMETER(uint32, MaxRootCount)
-		SHADER_PARAMETER(uint32, MaxSampleCount)
-
-		SHADER_PARAMETER_RDG_BUFFER_SRV(Buffer, RestSamplePositionsBuffer)
-		SHADER_PARAMETER_RDG_BUFFER_SRV(Buffer, MeshSampleWeightsBuffer)
-
-		SHADER_PARAMETER_RDG_BUFFER_SRV(StructuredBuffer, RestRootTrianglePosition0)
-		SHADER_PARAMETER_RDG_BUFFER_SRV(StructuredBuffer, RestRootTrianglePosition1)
-		SHADER_PARAMETER_RDG_BUFFER_SRV(StructuredBuffer, RestRootTrianglePosition2)
-
-		SHADER_PARAMETER_RDG_BUFFER_UAV(RWStructuredBuffer, OutDeformedRootTrianglePosition0)
-		SHADER_PARAMETER_RDG_BUFFER_UAV(RWStructuredBuffer, OutDeformedRootTrianglePosition1)
-		SHADER_PARAMETER_RDG_BUFFER_UAV(RWStructuredBuffer, OutDeformedRootTrianglePosition2)
-	END_SHADER_PARAMETER_STRUCT()
-
-public:
-	static bool ShouldCompilePermutation(const FGlobalShaderPermutationParameters& Parameters) { return IsHairStrandsSupported(EHairStrandsShaderType::All, Parameters.Platform); }
-	static void ModifyCompilationEnvironment(const FGlobalShaderPermutationParameters& Parameters, FShaderCompilerEnvironment& OutEnvironment)
-	{
-		FGlobalShader::ModifyCompilationEnvironment(Parameters, OutEnvironment);
-		OutEnvironment.SetDefine(TEXT("SHADER_ROOTTRIANGLE"), 1);
-	}
-};
-
-IMPLEMENT_GLOBAL_SHADER(FHairInterpolateMeshTriangleCS, "/Engine/Private/HairStrands/HairStrandsMeshInterpolate.usf", "MainCS", SF_Compute);
-
-void AddHairStrandInterpolateMeshTrianglesPass(
-	FRDGBuilder& GraphBuilder,
-	FGlobalShaderMap* ShaderMap,
-	const int32 LODIndex,
-	const FHairStrandsProjectionMeshData::LOD& MeshData,
-	FHairStrandsRestRootResource* RestResources,
-	FHairStrandsDeformedRootResource* DeformedResources)
-{
-	const uint32 RootCount = RestResources->BulkData.RootCount;
-	if (RootCount == 0 || LODIndex < 0 || LODIndex >= RestResources->LODs.Num() || LODIndex >= DeformedResources->LODs.Num())
-	{
-		return;
-	}
-	FHairStrandsRestRootResource::FLOD& RestLODData = RestResources->LODs[LODIndex];
-	FHairStrandsDeformedRootResource::FLOD& DeformedLODData = DeformedResources->LODs[LODIndex];
-	check(RestLODData.LODIndex == LODIndex);
-	check(DeformedLODData.LODIndex == LODIndex);
-
-	FHairInterpolateMeshTriangleCS::FParameters* Parameters = GraphBuilder.AllocParameters<FHairInterpolateMeshTriangleCS::FParameters>();
-	Parameters->MaxRootCount = RootCount;
-	Parameters->MaxSampleCount = RestLODData.SampleCount;
-
-	FRDGImportedBuffer OutDeformedRootTrianglePosition0 = Register(GraphBuilder, DeformedLODData.GetDeformedRootTrianglePosition0Buffer(FHairStrandsDeformedRootResource::FLOD::Current), ERDGImportedBufferFlags::CreateUAV);
-	FRDGImportedBuffer OutDeformedRootTrianglePosition1 = Register(GraphBuilder, DeformedLODData.GetDeformedRootTrianglePosition1Buffer(FHairStrandsDeformedRootResource::FLOD::Current), ERDGImportedBufferFlags::CreateUAV);
-	FRDGImportedBuffer OutDeformedRootTrianglePosition2 = Register(GraphBuilder, DeformedLODData.GetDeformedRootTrianglePosition2Buffer(FHairStrandsDeformedRootResource::FLOD::Current), ERDGImportedBufferFlags::CreateUAV);
-
-	Parameters->RestRootTrianglePosition0 = RegisterAsSRV(GraphBuilder, RestLODData.RestRootTrianglePosition0Buffer);
-	Parameters->RestRootTrianglePosition1 = RegisterAsSRV(GraphBuilder, RestLODData.RestRootTrianglePosition1Buffer);
-	Parameters->RestRootTrianglePosition2 = RegisterAsSRV(GraphBuilder, RestLODData.RestRootTrianglePosition2Buffer);
-
-	Parameters->OutDeformedRootTrianglePosition0 = OutDeformedRootTrianglePosition0.UAV;
-	Parameters->OutDeformedRootTrianglePosition1 = OutDeformedRootTrianglePosition1.UAV;
-	Parameters->OutDeformedRootTrianglePosition2 = OutDeformedRootTrianglePosition2.UAV;
-
-	Parameters->MeshSampleWeightsBuffer		= RegisterAsSRV(GraphBuilder, DeformedLODData.GetMeshSampleWeightsBuffer(FHairStrandsDeformedRootResource::FLOD::Current));
-	Parameters->RestSamplePositionsBuffer	= RegisterAsSRV(GraphBuilder, RestLODData.RestSamplePositionsBuffer);
-
-	const FIntVector DispatchGroupCount = FComputeShaderUtils::GetGroupCount(RootCount, 128);
-	check(DispatchGroupCount.X < 65536);
-	TShaderMapRef<FHairInterpolateMeshTriangleCS> ComputeShader(ShaderMap);
-	FComputeShaderUtils::AddPass(
-		GraphBuilder,
-		RDG_EVENT_NAME("HairStrands::TriangleMeshInterpolate"),
-		ComputeShader,
-		Parameters,
-		DispatchGroupCount);
-
-	GraphBuilder.SetBufferAccessFinal(OutDeformedRootTrianglePosition0.Buffer, ERHIAccess::SRVMask);
-	GraphBuilder.SetBufferAccessFinal(OutDeformedRootTrianglePosition1.Buffer, ERHIAccess::SRVMask);
-	GraphBuilder.SetBufferAccessFinal(OutDeformedRootTrianglePosition2.Buffer, ERHIAccess::SRVMask);
-}
-
-
 ///////////////////////////////////////////////////////////////////////////////////////////////////
 
 class FHairMeshesInterpolateCS : public FGlobalShader
diff --git a/Engine/Plugins/Runtime/HairStrands/Source/HairStrandsCore/Public/HairStrandsMeshProjection.h b/Engine/Plugins/Runtime/HairStrands/Source/HairStrandsCore/Public/HairStrandsMeshProjection.h
--- a/Engine/Plugins/Runtime/HairStrands/Source/HairStrandsCore/Public/HairStrandsMeshProjection.h
+++ b/Engine/Plugins/Runtime/HairStrands/Source/HairStrandsCore/Public/HairStrandsMeshProjection.h
@@ -115,14 +115,6 @@ void AddComputeMipsPass(
 	FGlobalShaderMap* ShaderMap,
 	FRDGTextureRef& OutTexture);
 
-void AddHairStrandInterpolateMeshTrianglesPass(
-	FRDGBuilder& GraphBuilder,
-	FGlobalShaderMap* ShaderMap,
-	const int32 LODIndex,
-	const FHairStrandsProjectionMeshData::LOD& MeshData,
-	FHairStrandsRestRootResource* RestResources,
-	FHairStrandsDeformedRootResource* DeformedResources);
-
 void AddSkinUpdatePass(
 	FRDGBuilder& GraphBuilder,
 	FGlobalShaderMap* ShaderMap,
diff --git a/Engine/Shaders/Private/HairStrands/HairStrandsMeshInterpolate.usf b/Engine/Shaders/Private/HairStrands/HairStrandsMeshInterpolate.usf
--- a/Engine/Shaders/Private/HairStrands/HairStrandsMeshInterpolate.usf
+++ b/Engine/Shaders/Private/HairStrands/HairStrandsMeshInterpolate.usf
@@ -6,48 +6,6 @@
 
 ////////////////////////////////////////////////////////////////////////
 
-#if SHADER_ROOTTRIANGLE
-
-uint MaxRootCount;
-uint MaxSampleCount;
-
-Buffer<float4> RestSamplePositionsBuffer;
-Buffer<float4> MeshSampleWeightsBuffer;
-
-Buffer<float4> RestRootTrianglePosition0;
-Buffer<float4> RestRootTrianglePosition1;
-Buffer<float4> RestRootTrianglePosition2;
-
-RWBuffer<float4> OutDeformedRootTrianglePosition0;
-RWBuffer<float4> OutDeformedRootTrianglePosition1;
-RWBuffer<float4> OutDeformedRootTrianglePosition2;
-
-void ApplyRBF(
-	uint InRootIndex,
-	Buffer<float4> InRestRootTrianglePosition,
-	RWBuffer<float4> OutDeformedRootTrianglePosition)
-{
-	const float3 RestControlPoint = InRestRootTrianglePosition[InRootIndex].xyz;
-	const float3 DeformedControlPoint = ApplyRBF(RestControlPoint, MaxSampleCount, RestSamplePositionsBuffer, MeshSampleWeightsBuffer);
-	OutDeformedRootTrianglePosition[InRootIndex + MaxRootCount] = float4(DeformedControlPoint, 1);
-}
-
-[numthreads(128, 1, 1)]
-void MainCS(uint3 DispatchThreadId : SV_DispatchThreadID)
-{ 
-	const uint RootIndex = DispatchThreadId.x;
-	if (RootIndex >= MaxRootCount) 
-		return;
-
-	ApplyRBF(RootIndex, RestRootTrianglePosition0, OutDeformedRootTrianglePosition0);
-	ApplyRBF(RootIndex, RestRootTrianglePosition1, OutDeformedRootTrianglePosition1);
-	ApplyRBF(RootIndex, RestRootTrianglePosition2, OutDeformedRootTrianglePosition2);
-}
-
-#endif // SHADER_ROOTTRIANGLE
-
-////////////////////////////////////////////////////////////////////////
-
 #if SHADER_HAIRMESHES
 
 uint VertexCount;
