* Remove unused GPU binding code & shaders
* More shader clean up. Regroup all hair strands mesh-based code (RBF sample init/sample update/deformation, root triangle update/offset, ...) into HairStrandsMesh.usf.

diff --git a/Engine/Plugins/Runtime/HairStrands/Shaders/Private/NiagaraDataInterfaceHairStrandsTemplate.ush b/Engine/Plugins/Runtime/HairStrands/Shaders/Private/NiagaraDataInterfaceHairStrandsTemplate.ush
--- a/Engine/Plugins/Runtime/HairStrands/Shaders/Private/NiagaraDataInterfaceHairStrandsTemplate.ush
+++ b/Engine/Plugins/Runtime/HairStrands/Shaders/Private/NiagaraDataInterfaceHairStrandsTemplate.ush
@@ -802,7 +802,7 @@ void GetLocalVector_{ParameterName}(in float3 WorldVector, in bool IsPosition, o
  * Attached the root to the skinned cache
  * -----------------------------------------------------------------
  */
-// Duplicate from #include "HairStrands/HairStrandsMeshRBF.ush"
+// Duplicate from #include "HairStrands/HairStrandsMeshProjectionCommon.ush" ApplyRBF()
 float3 ApplyRBF_{ParameterName}(in float3 RestControlPoint)
 {
 	float3 ControlPoint = RestControlPoint;
diff --git a/Engine/Plugins/Runtime/HairStrands/Source/HairStrandsCore/Private/GroomRBFDeformer.cpp b/Engine/Plugins/Runtime/HairStrands/Source/HairStrandsCore/Private/GroomRBFDeformer.cpp
--- a/Engine/Plugins/Runtime/HairStrands/Source/HairStrandsCore/Private/GroomRBFDeformer.cpp
+++ b/Engine/Plugins/Runtime/HairStrands/Source/HairStrandsCore/Private/GroomRBFDeformer.cpp
@@ -10,7 +10,7 @@
 #include "Engine/StaticMesh.h"
 #include "GroomBindingBuilder.h"
 
-// HairStrandsSamplesInit.usf
+// HairStrandsMesh.usf
 void InitMeshSamples(
 	uint32 MaxVertexCount,
 	const TArray<FVector3f>& VertexPositionsBuffer,
@@ -30,7 +30,7 @@ void InitMeshSamples(
 	}
 }
 
-// HairStrandsSamplesUpdate.usf
+// HairStrandsMesh.usf
 void UpdateMeshSamples(
 	uint32 MaxSampleCount,
 	const TArray<float>& InterpolationWeightsBuffer,
diff --git a/Engine/Plugins/Runtime/HairStrands/Source/HairStrandsCore/Private/HairStrandsMeshProjection.cpp b/Engine/Plugins/Runtime/HairStrands/Source/HairStrandsCore/Private/HairStrandsMeshProjection.cpp
--- a/Engine/Plugins/Runtime/HairStrands/Source/HairStrandsCore/Private/HairStrandsMeshProjection.cpp
+++ b/Engine/Plugins/Runtime/HairStrands/Source/HairStrandsCore/Private/HairStrandsMeshProjection.cpp
@@ -69,77 +69,6 @@ FHairStrandsProjectionMeshData ExtractMeshData(FSkeletalMeshRenderData* RenderDa
 	return MeshData;
 }
 
-///////////////////////////////////////////////////////////////////////////////////////////////////
-#ifndef GPU_BINDING
-#define GPU_BINDING 0
-#endif
-#if GPU_BINDING
-class FMarkMeshSectionIdCS : public FGlobalShader
-{
-	DECLARE_GLOBAL_SHADER(FMarkMeshSectionIdCS);
-	SHADER_USE_PARAMETER_STRUCT(FMarkMeshSectionIdCS, FGlobalShader);
-
-	BEGIN_SHADER_PARAMETER_STRUCT(FParameters, )
-		SHADER_PARAMETER(uint32, MeshSectionId)
-		SHADER_PARAMETER(uint32, MeshSectionPrimitiveCount)
-		SHADER_PARAMETER(uint32, MeshMaxIndexCount)
-		SHADER_PARAMETER(uint32, MeshMaxVertexCount)
-		SHADER_PARAMETER(uint32, MeshIndexOffset)
-		SHADER_PARAMETER_SRV(Buffer<uint>, MeshIndexBuffer)
-		SHADER_PARAMETER_RDG_BUFFER_UAV(RWBuffer<uint32>, OutVertexSectionId)
-
-	END_SHADER_PARAMETER_STRUCT()
-
-public:
-	static bool ShouldCompilePermutation(const FGlobalShaderPermutationParameters& Parameters) { return IsHairStrandsSupported(EHairStrandsShaderType::Strands, Parameters.Platform); }
-	static void ModifyCompilationEnvironment(const FGlobalShaderPermutationParameters& Parameters, FShaderCompilerEnvironment& OutEnvironment)
-	{
-		FGlobalShader::ModifyCompilationEnvironment(Parameters, OutEnvironment);
-		OutEnvironment.SetDefine(TEXT("SHADER_SECTIONID"), 1);
-	}
-};
-
-IMPLEMENT_GLOBAL_SHADER(FMarkMeshSectionIdCS, "/Engine/Private/HairStrands/HairStrandsMeshProjection.usf", "MainMarkSectionIdCS", SF_Compute);
-
-static FRDGBufferRef AddMeshSectionId(
-	FRDGBuilder& GraphBuilder,
-	FGlobalShaderMap* ShaderMap,
-	const FHairStrandsProjectionMeshData::LOD& MeshData)
-{
-	const int32 SectionCount = MeshData.Sections.Num();
-	if (SectionCount < 0)
-		return nullptr;
-
-	// Initialized the section ID to a large number, as the shader will do an atomic min on the section ID.
-	FRDGBufferRef VertexSectionIdBuffer = GraphBuilder.CreateBuffer(FRDGBufferDesc::CreateBufferDesc(sizeof(uint32), MeshData.Sections[0].TotalVertexCount), TEXT("Hair.SectionId"));
-	FRDGBufferUAVRef VertexSectionIdBufferUAV = GraphBuilder.CreateUAV(VertexSectionIdBuffer, PF_R32_UINT);
-	AddClearUAVPass(GraphBuilder, VertexSectionIdBufferUAV, ~0u);
-	for (const FHairStrandsProjectionMeshData::Section& MeshSection : MeshData.Sections)
-	{	
-		FMarkMeshSectionIdCS::FParameters* Parameters = GraphBuilder.AllocParameters<FMarkMeshSectionIdCS::FParameters>();
-		Parameters->MeshSectionId				= MeshSection.SectionIndex;
-		Parameters->MeshSectionPrimitiveCount	= MeshSection.NumPrimitives;
-		Parameters->MeshMaxIndexCount			= MeshSection.TotalIndexCount;
-		Parameters->MeshMaxVertexCount			= MeshSection.TotalVertexCount;
-		Parameters->MeshIndexOffset				= MeshSection.IndexBaseIndex;
-		Parameters->MeshIndexBuffer				= MeshSection.IndexBuffer;
-		Parameters->OutVertexSectionId			= VertexSectionIdBufferUAV;
-
-		const FIntVector DispatchGroupCount = FComputeShaderUtils::GetGroupCount(MeshSection.NumPrimitives*3, 128);
-		check(DispatchGroupCount.X < 65536);
-		TShaderMapRef<FMarkMeshSectionIdCS> ComputeShader(ShaderMap);
-		FComputeShaderUtils::AddPass(
-			GraphBuilder,
-			RDG_EVENT_NAME("HairStrandsMarkVertexSectionId"),
-			ComputeShader,
-			Parameters,
-			DispatchGroupCount);
-	}
-
-	return VertexSectionIdBuffer;
-}
-#endif // GPU_BINDING
-
 ///////////////////////////////////////////////////////////////////////////////////////////////////
 class FSkinUpdateCS : public FGlobalShader
 {
@@ -207,304 +136,6 @@ void AddSkinUpdatePass(
 		DispatchGroupCount);
 }
 
-///////////////////////////////////////////////////////////////////////////////////////////////////
-#if GPU_BINDING
-class FMeshTransferCS : public FGlobalShader
-{
-	DECLARE_GLOBAL_SHADER(FMeshTransferCS);
-	SHADER_USE_PARAMETER_STRUCT(FMeshTransferCS, FGlobalShader);
-
-	BEGIN_SHADER_PARAMETER_STRUCT(FParameters, )
-		SHADER_PARAMETER(uint32, bNeedClear)
-		SHADER_PARAMETER(uint32, Source_MeshPrimitiveCount_Iteration)
-		SHADER_PARAMETER(uint32, Source_MeshMaxIndexCount)
-		SHADER_PARAMETER(uint32, Source_MeshMaxVertexCount)
-		SHADER_PARAMETER(uint32, Source_MeshIndexOffset)
-		SHADER_PARAMETER(uint32, Source_MeshUVsChannelOffset)
-		SHADER_PARAMETER(uint32, Source_MeshUVsChannelCount)
-		SHADER_PARAMETER(uint32, Target_MeshMaxVertexCount)
-		SHADER_PARAMETER(uint32, Target_MeshUVsChannelOffset)
-		SHADER_PARAMETER(uint32, Target_MeshUVsChannelCount)
-		SHADER_PARAMETER(uint32, Target_SectionId)
-		SHADER_PARAMETER_SRV(Buffer<uint>, Source_MeshIndexBuffer)
-		SHADER_PARAMETER_SRV(Buffer<float>, Source_MeshPositionBuffer)
-		SHADER_PARAMETER_SRV(Buffer<float2>, Source_MeshUVsBuffer)
-		SHADER_PARAMETER_SRV(Buffer<float2>, Target_MeshUVsBuffer)
-		SHADER_PARAMETER_RDG_BUFFER_SRV(Buffer<uint>, Target_VertexSectionId)
-		SHADER_PARAMETER_UAV(RWBuffer<float>, Target_MeshPositionBuffer)
-		SHADER_PARAMETER_RDG_BUFFER_UAV(RWBuffer<uint32>, OutDistanceBuffer)
-
-	END_SHADER_PARAMETER_STRUCT()
-
-public:
-	static bool ShouldCompilePermutation(const FGlobalShaderPermutationParameters& Parameters) { return IsHairStrandsSupported(EHairStrandsShaderType::Tool, Parameters.Platform); }
-	static void ModifyCompilationEnvironment(const FGlobalShaderPermutationParameters& Parameters, FShaderCompilerEnvironment& OutEnvironment)
-	{
-		FGlobalShader::ModifyCompilationEnvironment(Parameters, OutEnvironment);
-		OutEnvironment.SetDefine(TEXT("SHADER_MESHTRANSFER"), 1);
-	}
-};
-
-IMPLEMENT_GLOBAL_SHADER(FMeshTransferCS, "/Engine/Private/HairStrands/HairStrandsMeshProjection.usf", "MainMeshTransferCS", SF_Compute);
-
-static void AddMeshTransferPass(
-	FRDGBuilder& GraphBuilder,
-	FGlobalShaderMap* ShaderMap,
-	bool bClear,
-	const FHairStrandsProjectionMeshData::Section& SourceSectionData,
-	const FHairStrandsProjectionMeshData::Section& TargetSectionData,
-	FRDGBufferRef VertexSectionId, 
-	FRWBuffer& OutTargetRestPosition,
-	FBufferTransitionQueue& OutTransitionQueue)
-{
-	if (!SourceSectionData.IndexBuffer ||
-		!SourceSectionData.PositionBuffer ||
-		 SourceSectionData.TotalIndexCount == 0 ||
-		 SourceSectionData.TotalVertexCount == 0||
-
-		!TargetSectionData.IndexBuffer ||
-		!TargetSectionData.PositionBuffer ||
-		 TargetSectionData.TotalIndexCount == 0 ||
-		 TargetSectionData.TotalVertexCount == 0)
-	{
-		return;
-	}
-	
-	FRDGBufferRef PositionDistanceBuffer = GraphBuilder.CreateBuffer(FRDGBufferDesc::CreateBufferDesc(sizeof(uint32), TargetSectionData.TotalVertexCount), TEXT("Hair.DistanceBuffer"));
-	FRDGBufferUAVRef PositionDistanceBufferUAV = GraphBuilder.CreateUAV(PositionDistanceBuffer, PF_R32_UINT);
-
-	// For projecting hair onto a skeletal mesh, 1 thread is spawn for each hair which iterates over all triangles.
-	// To avoid TDR, we split projection into multiple passes when the mesh is too large.
-	uint32 MeshPassNumPrimitive = 1024 * FMath::Clamp(GHairProjectionMaxTrianglePerProjectionIteration, 1, 256);
-	uint32 MeshPassCount = 1;
-	if (SourceSectionData.NumPrimitives < MeshPassNumPrimitive)
-	{
-		MeshPassNumPrimitive = SourceSectionData.NumPrimitives;
-	}
-	else
-	{
-		MeshPassCount = FMath::CeilToInt(SourceSectionData.NumPrimitives / float(MeshPassNumPrimitive));
-	}
-
-	FRDGBufferSRVRef VertexSectionIdSRV = GraphBuilder.CreateSRV(VertexSectionId, PF_R32_UINT);
-	for (uint32 MeshPassIt = 0; MeshPassIt < MeshPassCount; ++MeshPassIt)
-	{
-		FMeshTransferCS::FParameters* Parameters = GraphBuilder.AllocParameters<FMeshTransferCS::FParameters>();
-		Parameters->bNeedClear = bClear ? 1 : 0;
-
-		Parameters->Source_MeshPrimitiveCount_Iteration = (MeshPassIt < MeshPassCount - 1) ? MeshPassNumPrimitive : (SourceSectionData.NumPrimitives - MeshPassNumPrimitive * MeshPassIt);
-		Parameters->Source_MeshMaxIndexCount		= SourceSectionData.TotalIndexCount;
-		Parameters->Source_MeshMaxVertexCount		= SourceSectionData.TotalVertexCount;
-		Parameters->Source_MeshIndexOffset			= SourceSectionData.IndexBaseIndex + (MeshPassNumPrimitive * MeshPassIt * 3);
-		Parameters->Source_MeshUVsChannelOffset		= SourceSectionData.UVsChannelOffset;
-		Parameters->Source_MeshUVsChannelCount		= SourceSectionData.UVsChannelCount;
-		Parameters->Source_MeshIndexBuffer			= SourceSectionData.IndexBuffer;
-		Parameters->Source_MeshPositionBuffer		= SourceSectionData.PositionBuffer;
-		Parameters->Source_MeshUVsBuffer			= SourceSectionData.UVsBuffer;
-
-		Parameters->Target_MeshMaxVertexCount		= TargetSectionData.TotalVertexCount;
-		Parameters->Target_MeshUVsChannelOffset		= TargetSectionData.UVsChannelOffset;
-		Parameters->Target_MeshUVsChannelCount		= TargetSectionData.UVsChannelCount;
-		Parameters->Target_MeshUVsBuffer			= TargetSectionData.UVsBuffer;
-		Parameters->Target_MeshPositionBuffer		= OutTargetRestPosition.UAV;
-		Parameters->Target_VertexSectionId			= VertexSectionIdSRV;
-		Parameters->Target_SectionId				= TargetSectionData.SectionIndex;
-
-		Parameters->OutDistanceBuffer				= PositionDistanceBufferUAV;
-
-		const FIntVector DispatchGroupCount = FComputeShaderUtils::GetGroupCount(TargetSectionData.TotalVertexCount, 128);
-		check(DispatchGroupCount.X < 65536);
-		TShaderMapRef<FMeshTransferCS> ComputeShader(ShaderMap);
-		FComputeShaderUtils::AddPass(
-			GraphBuilder,
-			RDG_EVENT_NAME("HairStrandsTransferMesh"),
-			ComputeShader,
-			Parameters,
-			DispatchGroupCount);
-		bClear = false;
-
-		OutTransitionQueue.Add(OutTargetRestPosition.UAV);
-	}
-}
-///////////////////////////////////////////////////////////////////////////////////////////////////
-class FHairMeshProjectionCS : public FGlobalShader
-{
-	DECLARE_GLOBAL_SHADER(FHairMeshProjectionCS);
-	SHADER_USE_PARAMETER_STRUCT(FHairMeshProjectionCS, FGlobalShader);
-
-	BEGIN_SHADER_PARAMETER_STRUCT(FParameters, )
-		SHADER_PARAMETER(uint32, bClear)
-		SHADER_PARAMETER(uint32, MaxRootCount)
-
-		SHADER_PARAMETER(uint32, MeshPrimitiveOffset_Iteration)
-		SHADER_PARAMETER(uint32, MeshPrimitiveCount_Iteration)
-		SHADER_PARAMETER(uint32, MeshSectionIndex)
-		SHADER_PARAMETER(uint32, MeshMaxIndexCount)
-		SHADER_PARAMETER(uint32, MeshMaxVertexCount)
-		SHADER_PARAMETER(uint32, MeshIndexOffset)
-
-		SHADER_PARAMETER_SRV(Buffer, MeshIndexBuffer)
-		SHADER_PARAMETER_SRV(Buffer, MeshPositionBuffer)
-
-		SHADER_PARAMETER_SRV(Buffer, RootPositionBuffer)
-		SHADER_PARAMETER_SRV(Buffer, RootNormalBuffer)
-
-		SHADER_PARAMETER_UAV(RWBuffer, OutRootTriangleIndex)
-		SHADER_PARAMETER_UAV(RWBuffer, OutRootTriangleBarycentrics)
-		SHADER_PARAMETER_RDG_BUFFER_UAV(RWStructuredBuffer, OutRootTriangleDistance)
-
-	END_SHADER_PARAMETER_STRUCT()
-
-public:
-	static bool ShouldCompilePermutation(const FGlobalShaderPermutationParameters& Parameters) { return IsHairStrandsSupported(EHairStrandsShaderType::Strands, Parameters.Platform); }
-	static void ModifyCompilationEnvironment(const FGlobalShaderPermutationParameters& Parameters, FShaderCompilerEnvironment& OutEnvironment)
-	{
-		FGlobalShader::ModifyCompilationEnvironment(Parameters, OutEnvironment);
-		OutEnvironment.SetDefine(TEXT("SHADER_PROJECTION"), 1);
-	}
-};
-
-IMPLEMENT_GLOBAL_SHADER(FHairMeshProjectionCS, "/Engine/Private/HairStrands/HairStrandsMeshProjection.usf", "MainCS", SF_Compute);
-
-static void AddHairStrandMeshProjectionPass(
-	FRDGBuilder& GraphBuilder,
-	FGlobalShaderMap* ShaderMap,
-	const bool bClear,
-	const int32 LODIndex,
-	const FHairStrandsProjectionMeshData::Section& MeshSectionData,
-	FHairStrandsRestRootResource* RestResources,
-	FRDGBufferRef RootDistanceBuffer)
-{
-	const uint32 RootCount = RestResources->RootData.RootCount;
-	if (!RestResources->RootPositionBuffer.SRV ||
-		!RestResources->RootNormalBuffer.SRV ||
-		LODIndex < 0 || LODIndex >= RestResources->LODs.Num() ||
-		!MeshSectionData.IndexBuffer ||
-		!MeshSectionData.PositionBuffer ||
-		MeshSectionData.TotalIndexCount == 0 ||
-		MeshSectionData.TotalVertexCount == 0)
-	{
-		return;
-	}
-
-	FHairStrandsRestRootResource::FLOD& RestLODDatas = RestResources->LODs[LODIndex];
-	if (!RestLODDatas.RootTriangleIndexBuffer.UAV || !RestLODDatas.RootTriangleBarycentricBuffer.UAV)
-	{
-		return;
-	}
-
-	// The current shader code HairStrandsMeshProjection.usf encode the section ID onto the highest 8bits of a 32bits uint. 
-	// This limits the number of section to 64. See EncodeTriangleIndex & DecodeTriangleIndex functions in 
-	// HairStarndsMeshProjectionCommon.ush for mode details.
-	// This means that the mesh needs to have less than 67M triangles (since triangle ID is stored onto 26bits).
-	//
-	// This could be increase if necessary.
-	check(MeshSectionData.SectionIndex < GetHairStrandsMaxSectionCount());
-	check(MeshSectionData.NumPrimitives < GetHairStrandsMaxTriangleCount());
-
-	// For projecting hair onto a skeletal mesh, 1 thread is spawn for each hair which iterates over all triangles.
-	// To avoid TDR, we split projection into multiple passes when the mesh is too large.
-	uint32 MeshPassNumPrimitive = 1024 * FMath::Clamp(GHairProjectionMaxTrianglePerProjectionIteration, 1, 256);
-	uint32 MeshPassCount = 1;
-	if (MeshSectionData.NumPrimitives < MeshPassNumPrimitive)
-	{
-		MeshPassNumPrimitive = MeshSectionData.NumPrimitives;
-	}
-	else
-	{
-		MeshPassCount = FMath::CeilToInt(MeshSectionData.NumPrimitives / float(MeshPassNumPrimitive));
-	}
-
-
-	FRDGBufferUAVRef DistanceUAV = GraphBuilder.CreateUAV(FRDGBufferUAVDesc(RootDistanceBuffer, PF_R32_FLOAT));
-	for (uint32 MeshPassIt=0;MeshPassIt<MeshPassCount;++MeshPassIt)
-	{
-		FHairMeshProjectionCS::FParameters* Parameters = GraphBuilder.AllocParameters<FHairMeshProjectionCS::FParameters>();
-		Parameters->bClear				= bClear && MeshPassIt == 0 ? 1 : 0;
-		Parameters->MaxRootCount		= RootCount;
-		Parameters->RootPositionBuffer	= RestResources->RootPositionBuffer.SRV;
-		Parameters->RootNormalBuffer	= RestResources->RootNormalBuffer.SRV;
-		Parameters->MeshSectionIndex	= MeshSectionData.SectionIndex;
-		Parameters->MeshMaxIndexCount	= MeshSectionData.TotalIndexCount;
-		Parameters->MeshMaxVertexCount	= MeshSectionData.TotalVertexCount;
-		Parameters->MeshIndexOffset		= MeshSectionData.IndexBaseIndex + (MeshPassNumPrimitive * MeshPassIt * 3);
-		Parameters->MeshIndexBuffer		= MeshSectionData.IndexBuffer;
-		Parameters->MeshPositionBuffer	= MeshSectionData.PositionBuffer;
-		Parameters->MeshPrimitiveOffset_Iteration	= MeshPassNumPrimitive * MeshPassIt;
-		Parameters->MeshPrimitiveCount_Iteration	= (MeshPassIt < MeshPassCount-1) ? MeshPassNumPrimitive : (MeshSectionData.NumPrimitives - MeshPassNumPrimitive * MeshPassIt);
-
-		// The projection is always done onto the source/rest mesh
-		Parameters->OutRootTriangleIndex		= RestLODDatas.RootTriangleIndexBuffer.UAV;
-		Parameters->OutRootTriangleBarycentrics = RestLODDatas.RootTriangleBarycentricBuffer.UAV;
-		Parameters->OutRootTriangleDistance		= DistanceUAV;
-
-		const FIntVector DispatchGroupCount = FComputeShaderUtils::GetGroupCount(RootCount, 128);
-		check(DispatchGroupCount.X < 65536);
-		TShaderMapRef<FHairMeshProjectionCS> ComputeShader(ShaderMap);
-		FComputeShaderUtils::AddPass(
-			GraphBuilder,
-			RDG_EVENT_NAME("HairStrandsMeshProjection"),
-			ComputeShader,
-			Parameters,
-			DispatchGroupCount);
-	}
-}
-
-void ProjectHairStrandsOntoMesh(
-	FRDGBuilder& GraphBuilder,
-	FGlobalShaderMap* ShaderMap,
-	const int32 LODIndex,
-	const FHairStrandsProjectionMeshData& ProjectionMeshData,
-	FHairStrandsRestRootResource* RestResources)
-{
-	if (LODIndex < 0 || LODIndex >= RestResources->LODs.Num())
-		return;
-
-	const uint32 RootCount = RestResources->RootData.RootCount;
-	FHairStrandsRestRootResource::FLOD& RestLODDatas = RestResources->LODs[LODIndex];
-	FRDGBufferRef RootDistanceBuffer = GraphBuilder.CreateBuffer(FRDGBufferDesc::CreateStructuredDesc(sizeof(float), RootCount), TEXT("Hair.TriangleDistance"));
-
-	bool ClearDistance = true;
-	for (const FHairStrandsProjectionMeshData::Section& MeshSection : ProjectionMeshData.LODs[LODIndex].Sections)
-	{
-		check(RestLODDatas.LODIndex == LODIndex);
-		AddHairStrandMeshProjectionPass(GraphBuilder, ShaderMap, ClearDistance, LODIndex, MeshSection, RestResources, RootDistanceBuffer);
-		RestLODDatas.Status = FHairStrandsRestRootResource::FLOD::EStatus::Completed;
-		ClearDistance = false;
-	}
-}
-
-void TransferMesh(
-	FRDGBuilder& GraphBuilder,
-	FGlobalShaderMap* ShaderMap,
-	const int32 LODIndex,
-	const FHairStrandsProjectionMeshData& SourceMeshData,
-	const FHairStrandsProjectionMeshData& TargetMeshData,
-	FRWBuffer& OutPositionBuffer,
-	FBufferTransitionQueue& OutTransitionQueue)
-{
-	if (LODIndex < 0)
-		return;
-
-	// LODs are transfered using the LOD0 of the source mesh, as the LOD count can mismatch between source and target meshes.
-	const int32 SourceLODIndex = 0;
-	const int32 TargetLODIndex = LODIndex;
-
-	// Assume that the section 0 contains the head section, which is where the hair/facial hair should be projected on
-	const uint32 SourceSectionIndex = 0;
-	const uint32 TargetSectionIndex = 0;
-
-	const int32 SectionCount = TargetMeshData.LODs[TargetLODIndex].Sections.Num();
-	if (SectionCount < 0)
-		return;
-
-	FRDGBufferRef VertexSectionId = AddMeshSectionId(GraphBuilder, ShaderMap, TargetMeshData.LODs[TargetLODIndex]);
-	const FHairStrandsProjectionMeshData::Section& SourceMeshSection = SourceMeshData.LODs[SourceLODIndex].Sections[SourceSectionIndex];
-	const FHairStrandsProjectionMeshData::Section& TargetMeshSection = TargetMeshData.LODs[TargetLODIndex].Sections[TargetSectionIndex];
-	AddMeshTransferPass(GraphBuilder, ShaderMap, true, SourceMeshSection, TargetMeshSection, VertexSectionId, OutPositionBuffer);
-}
-#endif // GPU_BINDING
-
 ///////////////////////////////////////////////////////////////////////////////////////////////////
 
 class FHairUpdateMeshTriangleCS : public FGlobalShader
@@ -603,9 +234,9 @@ public:
 	}
 };
 
-IMPLEMENT_GLOBAL_SHADER(FHairUpdateMeshTriangleCS, "/Engine/Private/HairStrands/HairStrandsMeshUpdate.usf", "MainCS", SF_Compute);
+IMPLEMENT_GLOBAL_SHADER(FHairUpdateMeshTriangleCS, "/Engine/Private/HairStrands/HairStrandsMesh.usf", "MainCS", SF_Compute);
 
-// Manual packing for mesh section scalars.  MUST MATCH WITH HairStrandsMeshUpdate.usf
+// Manual packing for mesh section scalars.  MUST MATCH WITH HairStrandsMesh.usf
 // PackedMeshSectionScalars1: MeshSectionIndices, MeshMaxIndexCount, MeshMaxVertexCount, MeshIndexOffset
 // PackedMeshSectionScalars2: MeshUVsChannelOffset, MeshUVsChannelCount, MeshSectionBufferIndex, *free*
 inline void SetMeshSectionIndices(FHairUpdateMeshTriangleCS::FParameters& Params, uint32 SectionIndex, uint32 MeshSectionIndices)         { Params.PackedMeshSectionScalars1[SectionIndex].X = MeshSectionIndices; }
@@ -925,7 +556,7 @@ public:
 	}
 };
 
-IMPLEMENT_GLOBAL_SHADER(FHairMeshesInterpolateCS, "/Engine/Private/HairStrands/HairStrandsMeshInterpolate.usf", "MainHairMeshesCS", SF_Compute);
+IMPLEMENT_GLOBAL_SHADER(FHairMeshesInterpolateCS, "/Engine/Private/HairStrands/HairStrandsMesh.usf", "MainHairMeshesCS", SF_Compute);
 
 template<typename TRestResource, typename TDeformedResource>
 void InternalAddHairRBFInterpolationPass(
@@ -1064,12 +695,13 @@ public:
 	{
 		FGlobalShader::ModifyCompilationEnvironment(Parameters, OutEnvironment);
 		OutEnvironment.SetDefine(TEXT("MAX_SECTION_COUNT"), SectionArrayCount);
+		OutEnvironment.SetDefine(TEXT("SHADER_SAMPLE_INIT"), 1);
 	}
 };
 
-IMPLEMENT_GLOBAL_SHADER(FHairInitMeshSamplesCS, "/Engine/Private/HairStrands/HairStrandsSamplesInit.usf", "MainCS", SF_Compute);
+IMPLEMENT_GLOBAL_SHADER(FHairInitMeshSamplesCS, "/Engine/Private/HairStrands/HairStrandsMesh.usf", "MainCS", SF_Compute);
 
-// Manual packing for mesh section scalars.  MUST MATCH WITH HairStrandsSamplesInit.usf
+// Manual packing for mesh section scalars.  MUST MATCH WITH HairStrandsMesh.usf
 // PackedSectionScalars: SectionVertexOffset, SectionVertexCount, SectionBufferIndex, *free*
 inline void SetSectionVertexOffset(FHairInitMeshSamplesCS::FParameters& Params, uint32 SectionIndex, uint32 SectionVertexOffset) { Params.PackedSectionScalars[SectionIndex].X = SectionVertexOffset; }
 inline void SetSectionVertexCount(FHairInitMeshSamplesCS::FParameters& Params, uint32 SectionIndex, uint32 SectionVertexCount)   { Params.PackedSectionScalars[SectionIndex].Y = SectionVertexCount; }
@@ -1311,10 +943,11 @@ public:
 	static void ModifyCompilationEnvironment(const FGlobalShaderPermutationParameters& Parameters, FShaderCompilerEnvironment& OutEnvironment)
 	{
 		FGlobalShader::ModifyCompilationEnvironment(Parameters, OutEnvironment);
+		OutEnvironment.SetDefine(TEXT("SHADER_SAMPLE_UPDATE"), 1);
 	}
 };
 
-IMPLEMENT_GLOBAL_SHADER(FHairUpdateMeshSamplesCS, "/Engine/Private/HairStrands/HairStrandsSamplesUpdate.usf", "MainCS", SF_Compute);
+IMPLEMENT_GLOBAL_SHADER(FHairUpdateMeshSamplesCS, "/Engine/Private/HairStrands/HairStrandsMesh.usf", "MainCS", SF_Compute);
 
 void AddHairStrandUpdateMeshSamplesPass(
 	FRDGBuilder& GraphBuilder,
@@ -1708,7 +1341,7 @@ public:
 	}
 };
 
-IMPLEMENT_GLOBAL_SHADER(FHairUpdatePositionOffsetCS, "/Engine/Private/HairStrands/HairStrandsMeshUpdate.usf", "MainCS", SF_Compute);
+IMPLEMENT_GLOBAL_SHADER(FHairUpdatePositionOffsetCS, "/Engine/Private/HairStrands/HairStrandsMesh.usf", "MainCS", SF_Compute);
 
 void AddHairStrandUpdatePositionOffsetPass(
 	FRDGBuilder& GraphBuilder,
diff --git a/Engine/Shaders/Private/HairStrands/HairStrandsGuideDeform.usf b/Engine/Shaders/Private/HairStrands/HairStrandsGuideDeform.usf
--- a/Engine/Shaders/Private/HairStrands/HairStrandsGuideDeform.usf
+++ b/Engine/Shaders/Private/HairStrands/HairStrandsGuideDeform.usf
@@ -3,7 +3,6 @@
 #include "../Common.ush"
 #include "HairStrandsClusterCommon.ush"
 #include "HairStrandsMeshProjectionCommon.ush"
-#include "HairStrandsMeshRBF.ush"
 
 float3 SimRestOffset;
 uint VertexCount;
diff --git a/Engine/Shaders/Private/HairStrands/HairStrandsMesh.usf b/Engine/Shaders/Private/HairStrands/HairStrandsMesh.usf
new file mode 100644
--- /dev/null
+++ b/Engine/Shaders/Private/HairStrands/HairStrandsMesh.usf
@@ -0,0 +1,478 @@
+// Copyright Epic Games, Inc. All Rights Reserved.
+
+#include "../Common.ush"
+#include "HairStrandsMeshProjectionCommon.ush"
+
+////////////////////////////////////////////////////////////////////////
+// Deform hair mesh basd on RBF weights
+
+#if SHADER_HAIRMESHES
+
+uint VertexCount;
+uint MaxSampleCount;
+
+Buffer<float4> RestSamplePositionsBuffer;
+Buffer<float4> MeshSampleWeightsBuffer;
+
+Buffer<float4> RestPositionBuffer;
+RWBuffer<float4> OutDeformedPositionBuffer;
+
+[numthreads(128, 1, 1)]
+void MainHairMeshesCS(uint3 DispatchThreadId : SV_DispatchThreadID)
+{
+	const uint VertexIndex = DispatchThreadId.x;
+	if (VertexIndex >= VertexCount)
+		return;
+
+	const float3 RestPosition = RestPositionBuffer[VertexIndex].xyz;
+	const float3 DeformedPosition = ApplyRBF(RestPosition, MaxSampleCount, RestSamplePositionsBuffer, MeshSampleWeightsBuffer);
+	OutDeformedPositionBuffer[VertexIndex] = float4(DeformedPosition,1);
+}
+
+#endif // SHADER_HAIRMESHES
+
+////////////////////////////////////////////////////////////////////////
+// Init RBF point based on mesh position
+
+#if SHADER_SAMPLE_INIT
+
+#ifndef MAX_SECTION_COUNT
+#error MAX_SECTION_COUNT needs to be defined
+#endif
+
+uint PassSectionCount;
+
+// Manual packing for mesh section scalars.  MUST MATCH WITH HairStrandsMeshProjection.cpp
+// PackedSectionScalars: SectionVertexOffset, SectionVertexCount, SectionBufferIndex, *free*
+uint4 PackedSectionScalars[MAX_SECTION_COUNT];
+#define SectionVertexOffset(SectionIndex) PackedSectionScalars[SectionIndex].x
+#define SectionVertexCount(SectionIndex)  PackedSectionScalars[SectionIndex].y
+#define SectionBufferIndex(SectionIndex)  PackedSectionScalars[SectionIndex].z
+
+#if PERMUTATION_POSITION_TYPE == 0
+Buffer<float> VertexPositionsBuffer0;
+Buffer<float> VertexPositionsBuffer1;
+Buffer<float> VertexPositionsBuffer2;
+Buffer<float> VertexPositionsBuffer3;
+Buffer<float> VertexPositionsBuffer4;
+Buffer<float> VertexPositionsBuffer5;
+Buffer<float> VertexPositionsBuffer6;
+Buffer<float> VertexPositionsBuffer7;
+#elif PERMUTATION_POSITION_TYPE == 1
+Buffer<float> RDGVertexPositionsBuffer0;
+Buffer<float> RDGVertexPositionsBuffer1;
+Buffer<float> RDGVertexPositionsBuffer2;
+Buffer<float> RDGVertexPositionsBuffer3;
+Buffer<float> RDGVertexPositionsBuffer4;
+Buffer<float> RDGVertexPositionsBuffer5;
+Buffer<float> RDGVertexPositionsBuffer6;
+Buffer<float> RDGVertexPositionsBuffer7;
+
+#define VertexPositionsBuffer0 RDGVertexPositionsBuffer0
+#define VertexPositionsBuffer1 RDGVertexPositionsBuffer1
+#define VertexPositionsBuffer2 RDGVertexPositionsBuffer2
+#define VertexPositionsBuffer3 RDGVertexPositionsBuffer3
+#define VertexPositionsBuffer4 RDGVertexPositionsBuffer4
+#define VertexPositionsBuffer5 RDGVertexPositionsBuffer5
+#define VertexPositionsBuffer6 RDGVertexPositionsBuffer6
+#define VertexPositionsBuffer7 RDGVertexPositionsBuffer7
+#endif
+
+uint MaxVertexCount;
+uint MaxSampleCount;
+Buffer<uint> SampleIndicesBuffer;
+RWBuffer<float4> OutSamplePositionsBuffer;
+
+[numthreads(128, 1, 1)]
+void MainCS(uint3 DispatchThreadId : SV_DispatchThreadID)
+{ 
+	const uint SampleIndex = DispatchThreadId.x;
+	if (SampleIndex >= MaxSampleCount) 
+		return;
+
+	const uint VertexIndex = SampleIndicesBuffer[SampleIndex];
+	if (VertexIndex >= MaxVertexCount) 
+		return;
+	
+	bool bValid = false;  
+	uint SectionIndex = 0;
+	for (uint It = 0; It < PassSectionCount; ++It)
+	{
+		if ((VertexIndex >= SectionVertexOffset(It)) && (VertexIndex < (SectionVertexOffset(It) + SectionVertexCount(It))))
+		{
+			SectionIndex = It;
+			bValid = true;
+			break;
+		}
+	}
+	
+	if (!bValid)
+		return;
+
+	float3 SamplePosition = float3(0,0,0);
+	const uint BufferIndex = SectionBufferIndex(SectionIndex);
+	const int LocalIndex = VertexIndex;
+	if (BufferIndex == 0)
+	{
+		SamplePosition.x = VertexPositionsBuffer0.Load(LocalIndex * 3 + 0);
+		SamplePosition.y = VertexPositionsBuffer0.Load(LocalIndex * 3 + 1);
+		SamplePosition.z = VertexPositionsBuffer0.Load(LocalIndex * 3 + 2);
+	}
+	else if (BufferIndex == 1)
+	{
+		SamplePosition.x = VertexPositionsBuffer1.Load(LocalIndex * 3 + 0);
+		SamplePosition.y = VertexPositionsBuffer1.Load(LocalIndex * 3 + 1);
+		SamplePosition.z = VertexPositionsBuffer1.Load(LocalIndex * 3 + 2);
+	}
+	else if (BufferIndex == 2)
+	{
+		SamplePosition.x = VertexPositionsBuffer2.Load(LocalIndex * 3 + 0);
+		SamplePosition.y = VertexPositionsBuffer2.Load(LocalIndex * 3 + 1);
+		SamplePosition.z = VertexPositionsBuffer2.Load(LocalIndex * 3 + 2);
+	}
+	else if (BufferIndex == 3)
+	{
+		SamplePosition.x = VertexPositionsBuffer3.Load(LocalIndex * 3 + 0);
+		SamplePosition.y = VertexPositionsBuffer3.Load(LocalIndex * 3 + 1);
+		SamplePosition.z = VertexPositionsBuffer3.Load(LocalIndex * 3 + 2);
+	}
+	else if (BufferIndex == 4)
+	{
+		SamplePosition.x = VertexPositionsBuffer4.Load(LocalIndex * 3 + 0);
+		SamplePosition.y = VertexPositionsBuffer4.Load(LocalIndex * 3 + 1);
+		SamplePosition.z = VertexPositionsBuffer4.Load(LocalIndex * 3 + 2);
+	}
+	else if (BufferIndex == 5)
+	{
+		SamplePosition.x = VertexPositionsBuffer5.Load(LocalIndex * 3 + 0);
+		SamplePosition.y = VertexPositionsBuffer5.Load(LocalIndex * 3 + 1);
+		SamplePosition.z = VertexPositionsBuffer5.Load(LocalIndex * 3 + 2);
+	}
+	else if (BufferIndex == 6)
+	{
+		SamplePosition.x = VertexPositionsBuffer6.Load(LocalIndex * 3 + 0);
+		SamplePosition.y = VertexPositionsBuffer6.Load(LocalIndex * 3 + 1);
+		SamplePosition.z = VertexPositionsBuffer6.Load(LocalIndex * 3 + 2);
+	}
+	else if (BufferIndex == 7)
+	{
+		SamplePosition.x = VertexPositionsBuffer7.Load(LocalIndex * 3 + 0);
+		SamplePosition.y = VertexPositionsBuffer7.Load(LocalIndex * 3 + 1);
+		SamplePosition.z = VertexPositionsBuffer7.Load(LocalIndex * 3 + 2);
+	}
+
+	OutSamplePositionsBuffer[SampleIndex] = float4(SamplePosition,1);
+}
+
+#endif // SHADER_SAMPLE_INIT
+
+////////////////////////////////////////////////////////////////////////
+// Update RBF points
+
+#if SHADER_SAMPLE_UPDATE
+
+// Copyright Epic Games, Inc. All Rights Reserved.
+
+#include "../Common.ush"
+
+uint MaxSampleCount;
+Buffer<uint> SampleIndicesBuffer;
+Buffer<float> InterpolationWeightsBuffer;
+
+Buffer<float4> SampleRestPositionsBuffer;
+Buffer<float4> SampleDeformedPositionsBuffer;
+
+RWBuffer<float4> OutSampleDeformationsBuffer; 
+
+[numthreads(128, 1, 1)]
+void MainCS(uint3 DispatchThreadId : SV_DispatchThreadID)
+{ 
+	const uint SampleIndex = DispatchThreadId.x;
+	if (SampleIndex >= (MaxSampleCount + 4))
+	{
+		return;
+	}
+
+	uint WeightsOffset = SampleIndex * (MaxSampleCount+4);
+	float3 SampleDeformation = float3(0,0,0);
+	for(uint i = 0; i < MaxSampleCount; ++i, ++WeightsOffset)
+	{
+		SampleDeformation += InterpolationWeightsBuffer[WeightsOffset] * 
+				(SampleDeformedPositionsBuffer[i] - SampleRestPositionsBuffer[i]).xyz;
+	}
+
+	OutSampleDeformationsBuffer[SampleIndex] = float4(SampleDeformation,1);
+}
+
+#endif // SHADER_SAMPLE_UPDATE
+
+////////////////////////////////////////////////////////////////////////
+// Update strands root triangle positions
+
+#if SHADER_MESH_UPDATE
+
+#ifndef MAX_SECTION_COUNT
+#error MAX_SECTION_COUNT needs to be defined
+#endif
+
+uint Pass_SectionStart;
+uint Pass_SectionCount;
+
+uint MaxSectionCount;
+
+// Manual packing for mesh section scalars.  MUST MATCH WITH HairStrandsMeshProjection.cpp
+// PackedMeshSectionScalars1: MeshSectionIndices, MeshMaxIndexCount, MeshMaxVertexCount, MeshIndexOffset
+// PackedMeshSectionScalars2: MeshUVsChannelOffset, MeshUVsChannelCount, MeshSectionBufferIndex, *free*
+uint4 PackedMeshSectionScalars1[MAX_SECTION_COUNT];
+uint4 PackedMeshSectionScalars2[MAX_SECTION_COUNT];
+#define MeshSectionIndices(SectionIndex) PackedMeshSectionScalars1[SectionIndex].x
+#define MeshMaxIndexCount(SectionIndex) PackedMeshSectionScalars1[SectionIndex].y
+#define MeshMaxVertexCount(SectionIndex) PackedMeshSectionScalars1[SectionIndex].z
+#define MeshIndexOffset(SectionIndex) PackedMeshSectionScalars1[SectionIndex].w
+#define MeshUVsChannelOffset(SectionIndex) PackedMeshSectionScalars2[SectionIndex].x
+#define MeshUVsChannelCount(SectionIndex) PackedMeshSectionScalars2[SectionIndex].y
+#define MeshSectionBufferIndex(SectionIndex) PackedMeshSectionScalars2[SectionIndex].z
+
+Buffer<uint> MeshIndexBuffer0;
+Buffer<uint> MeshIndexBuffer1;
+Buffer<uint> MeshIndexBuffer2;
+Buffer<uint> MeshIndexBuffer3;
+Buffer<uint> MeshIndexBuffer4;
+Buffer<uint> MeshIndexBuffer5;
+Buffer<uint> MeshIndexBuffer6;
+Buffer<uint> MeshIndexBuffer7;
+
+// Note on buffer input type
+// There are two different input buffer permutation, as if skin cache is enabled, we are using raw RHI SRV 
+// (while the skin cache system is not converted to RDG). If skin cache is disable, a manual skin compute 
+// update happens in the hair system, and its output are RDG buffers.
+#if PERMUTATION_POSITION_TYPE == 0
+
+Buffer<float> MeshPositionBuffer0;
+Buffer<float> MeshPositionBuffer1;
+Buffer<float> MeshPositionBuffer2;
+Buffer<float> MeshPositionBuffer3;
+Buffer<float> MeshPositionBuffer4;
+Buffer<float> MeshPositionBuffer5;
+Buffer<float> MeshPositionBuffer6;
+Buffer<float> MeshPositionBuffer7;
+
+Buffer<float> MeshPreviousPositionBuffer0;
+Buffer<float> MeshPreviousPositionBuffer1;
+Buffer<float> MeshPreviousPositionBuffer2;
+Buffer<float> MeshPreviousPositionBuffer3;
+Buffer<float> MeshPreviousPositionBuffer4;
+Buffer<float> MeshPreviousPositionBuffer5;
+Buffer<float> MeshPreviousPositionBuffer6;
+Buffer<float> MeshPreviousPositionBuffer7;
+
+#else
+
+Buffer<float> RDGMeshPositionBuffer0;
+Buffer<float> RDGMeshPositionBuffer1;
+Buffer<float> RDGMeshPositionBuffer2;
+Buffer<float> RDGMeshPositionBuffer3;
+Buffer<float> RDGMeshPositionBuffer4;
+Buffer<float> RDGMeshPositionBuffer5;
+Buffer<float> RDGMeshPositionBuffer6;
+Buffer<float> RDGMeshPositionBuffer7;
+
+Buffer<float> RDGMeshPreviousPositionBuffer0;
+Buffer<float> RDGMeshPreviousPositionBuffer1;
+Buffer<float> RDGMeshPreviousPositionBuffer2;
+Buffer<float> RDGMeshPreviousPositionBuffer3;
+Buffer<float> RDGMeshPreviousPositionBuffer4;
+Buffer<float> RDGMeshPreviousPositionBuffer5;
+Buffer<float> RDGMeshPreviousPositionBuffer6;
+Buffer<float> RDGMeshPreviousPositionBuffer7;
+
+#define MeshPositionBuffer0 RDGMeshPositionBuffer0
+#define MeshPositionBuffer1 RDGMeshPositionBuffer1
+#define MeshPositionBuffer2 RDGMeshPositionBuffer2
+#define MeshPositionBuffer3 RDGMeshPositionBuffer3
+#define MeshPositionBuffer4 RDGMeshPositionBuffer4
+#define MeshPositionBuffer5 RDGMeshPositionBuffer5
+#define MeshPositionBuffer6 RDGMeshPositionBuffer6
+#define MeshPositionBuffer7 RDGMeshPositionBuffer7
+
+#define MeshPreviousPositionBuffer0 RDGMeshPreviousPositionBuffer0
+#define MeshPreviousPositionBuffer1 RDGMeshPreviousPositionBuffer1
+#define MeshPreviousPositionBuffer2 RDGMeshPreviousPositionBuffer2
+#define MeshPreviousPositionBuffer3 RDGMeshPreviousPositionBuffer3
+#define MeshPreviousPositionBuffer4 RDGMeshPreviousPositionBuffer4
+#define MeshPreviousPositionBuffer5 RDGMeshPreviousPositionBuffer5
+#define MeshPreviousPositionBuffer6 RDGMeshPreviousPositionBuffer6
+#define MeshPreviousPositionBuffer7 RDGMeshPreviousPositionBuffer7
+
+#endif 
+
+#if PERMUTATION_WITHUV == 1
+Buffer<float2> MeshUVsBuffer0;
+Buffer<float2> MeshUVsBuffer1;
+Buffer<float2> MeshUVsBuffer2;
+Buffer<float2> MeshUVsBuffer3;
+Buffer<float2> MeshUVsBuffer4;
+Buffer<float2> MeshUVsBuffer5;
+Buffer<float2> MeshUVsBuffer6;
+Buffer<float2> MeshUVsBuffer7;
+Buffer<float2> MeshUVsBuffer8;
+#endif
+
+uint MaxRootCount;
+Buffer<uint> RootTriangleIndex;
+
+RWBuffer<float4> OutRootTriangleCurrPosition0;
+RWBuffer<float4> OutRootTriangleCurrPosition1;
+RWBuffer<float4> OutRootTriangleCurrPosition2;
+
+#if PERMUTATION_PREVIOUS
+RWBuffer<float4> OutRootTrianglePrevPosition0;
+RWBuffer<float4> OutRootTrianglePrevPosition1;
+RWBuffer<float4> OutRootTrianglePrevPosition2;
+#endif
+
+#if PERMUTATION_WITHUV == 1
+FHairMeshTriangleWithUVs GetTriangleData(FHairTriangleIndex TriangleIndex, uint SectionIndex, bool bPrevious)
+{
+	#define GetTriangleDataAtIndex(Index) GetTriangleData(TriangleIndex.TriangleIndex, bPrevious ? MeshPreviousPositionBuffer##Index : MeshPositionBuffer##Index, MeshUVsBuffer##Index, MeshIndexBuffer##Index, MeshIndexOffset(SectionIndex), MeshMaxIndexCount(SectionIndex), MeshMaxVertexCount(SectionIndex), MeshUVsChannelOffset(SectionIndex), MeshUVsChannelCount(SectionIndex));
+
+	const uint MeshSectionIndex = MeshSectionBufferIndex(SectionIndex);
+		 if (MeshSectionIndex == 0) { return GetTriangleDataAtIndex(0); } 
+	else if (MeshSectionIndex == 1) { return GetTriangleDataAtIndex(1); }
+	else if (MeshSectionIndex == 2) { return GetTriangleDataAtIndex(2); }
+	else if (MeshSectionIndex == 3) { return GetTriangleDataAtIndex(3); }
+	else if (MeshSectionIndex == 4) { return GetTriangleDataAtIndex(4); }
+	else if (MeshSectionIndex == 5) { return GetTriangleDataAtIndex(5); }
+	else if (MeshSectionIndex == 6) { return GetTriangleDataAtIndex(6); }
+	else if (MeshSectionIndex == 7) { return GetTriangleDataAtIndex(7); }
+	
+	#undef GetTriangleDataAtIndex
+
+	return (FHairMeshTriangleWithUVs)0;
+}
+#else
+FHairMeshTriangle GetTriangleData(FHairTriangleIndex TriangleIndex, uint SectionIndex, bool bPrevious)
+{
+	#define GetTriangleDataAtIndex(Index) GetTriangleData(TriangleIndex.TriangleIndex, bPrevious ? MeshPreviousPositionBuffer##Index : MeshPositionBuffer##Index, MeshIndexBuffer##Index, MeshIndexOffset(SectionIndex), MeshMaxIndexCount(SectionIndex), MeshMaxVertexCount(SectionIndex));
+
+	const uint MeshSectionIndex = MeshSectionBufferIndex(SectionIndex);
+		 if (MeshSectionIndex == 0) { return GetTriangleDataAtIndex(0); }
+	else if (MeshSectionIndex == 1) { return GetTriangleDataAtIndex(1); }
+	else if (MeshSectionIndex == 2) { return GetTriangleDataAtIndex(2); }
+	else if (MeshSectionIndex == 3) { return GetTriangleDataAtIndex(3); }
+	else if (MeshSectionIndex == 4) { return GetTriangleDataAtIndex(4); }
+	else if (MeshSectionIndex == 5) { return GetTriangleDataAtIndex(5); }
+	else if (MeshSectionIndex == 6) { return GetTriangleDataAtIndex(6); }
+	else if (MeshSectionIndex == 7) { return GetTriangleDataAtIndex(7); }
+
+	#undef GetTriangleDataAtIndex
+
+	return (FHairMeshTriangle)0;
+}
+#endif
+
+[numthreads(128, 1, 1)]
+void MainCS(uint3 DispatchThreadId : SV_DispatchThreadID)
+{ 
+	const uint RootIndex = DispatchThreadId.x;
+	if (RootIndex >= MaxRootCount) 
+		return;
+	
+	const FHairTriangleIndex TriangleIndex = DecodeTriangleIndex(RootTriangleIndex[RootIndex]);
+	if (TriangleIndex.SectionIndex >= MaxSectionCount)
+		return;
+
+	const uint SectionIndex = TriangleIndex.SectionIndex;
+
+	bool bValid = false;
+	uint SectionSlotIndex = 0;
+	for (uint It=0;It<Pass_SectionCount;++It)
+	{
+		if (MeshSectionIndices(It) == SectionIndex)
+		{
+			SectionSlotIndex = It;
+			bValid = true;
+			break;
+		}
+	}
+	
+	if (!bValid)
+		return;
+
+	// Compute the slod index of the resources for the section index (this is done in multiple pass
+#if PERMUTATION_WITHUV == 1
+	{
+		const FHairMeshTriangleWithUVs Tri = GetTriangleData(TriangleIndex, SectionSlotIndex, false);
+		OutRootTriangleCurrPosition0[RootIndex] = float4(Tri.P0, PackUVsToFloat(Tri.UV0));
+		OutRootTriangleCurrPosition1[RootIndex] = float4(Tri.P1, PackUVsToFloat(Tri.UV1));
+		OutRootTriangleCurrPosition2[RootIndex] = float4(Tri.P2, PackUVsToFloat(Tri.UV2));
+	}
+
+	#if PERMUTATION_PREVIOUS
+	{
+		const FHairMeshTriangleWithUVs Tri = GetTriangleData(TriangleIndex, SectionSlotIndex, true);
+		OutRootTrianglePrevPosition0[RootIndex] = float4(Tri.P0, PackUVsToFloat(Tri.UV0));
+		OutRootTrianglePrevPosition1[RootIndex] = float4(Tri.P1, PackUVsToFloat(Tri.UV1));
+		OutRootTrianglePrevPosition2[RootIndex] = float4(Tri.P2, PackUVsToFloat(Tri.UV2));
+	}
+	#endif
+#else
+	{
+		const FHairMeshTriangle Tri = GetTriangleData(TriangleIndex, SectionSlotIndex, false);
+		OutRootTriangleCurrPosition0[RootIndex] = (float4(Tri.P0, 1));
+		OutRootTriangleCurrPosition1[RootIndex] = (float4(Tri.P1, 1));
+		OutRootTriangleCurrPosition2[RootIndex] = (float4(Tri.P2, 1));
+	}
+
+	#if PERMUTATION_PREVIOUS
+	{
+		const FHairMeshTriangle Tri = GetTriangleData(TriangleIndex, SectionSlotIndex, true);
+		OutRootTrianglePrevPosition0[RootIndex] = (float4(Tri.P0, 1));
+		OutRootTrianglePrevPosition1[RootIndex] = (float4(Tri.P1, 1));
+		OutRootTrianglePrevPosition2[RootIndex] = (float4(Tri.P2, 1));
+	}
+	#endif
+#endif
+}
+
+#endif // SHADER_MESH_UPDATE
+
+////////////////////////////////////////////////////////////////////////
+// Update hair strands offset position
+
+#if SHADER_OFFSET_UPDATE
+
+#if PERMUTATION_USE_GPU_OFFSET
+	Buffer<float4> RootTriangleCurrPosition0Buffer;
+	float3 GetOffsetCurrPosition()	{ return RootTriangleCurrPosition0Buffer[0].xyz; }
+
+	#if PERMUTATION_PREVIOUS
+	Buffer<float4> RootTrianglePrevPosition0Buffer;
+	float3 GetOffsetPrevPosition()	{ return RootTrianglePrevPosition0Buffer[0].xyz; }
+	#endif
+#else
+	float3 CPUCurrPositionOffset;
+	float3 CPUPrevPositionOffset;
+
+	float3 GetOffsetCurrPosition()	{ return CPUCurrPositionOffset; }
+	float3 GetOffsetPrevPosition()	{ return CPUPrevPositionOffset; }
+#endif
+
+RWBuffer<float4> OutCurrOffsetBuffer;
+#if PERMUTATION_PREVIOUS
+RWBuffer<float4> OutPrevOffsetBuffer;
+#endif
+
+[numthreads(1, 1, 1)]
+void MainCS(uint3 DispatchThreadId : SV_DispatchThreadID)
+{
+	if (DispatchThreadId.x > 0)
+		return;
+
+	OutCurrOffsetBuffer[0] = float4(GetOffsetCurrPosition(), 1);
+#if PERMUTATION_PREVIOUS
+	OutPrevOffsetBuffer[0] = float4(GetOffsetPrevPosition(), 1);
+#endif
+}
+
+#endif // SHADER_OFFSET_UPDATE
\ No newline at end of file
diff --git a/Engine/Shaders/Private/HairStrands/HairStrandsMeshInterpolate.usf b/Engine/Shaders/Private/HairStrands/HairStrandsMeshInterpolate.usf
deleted file mode 100644
diff --git a/Engine/Shaders/Private/HairStrands/HairStrandsMeshProjection.usf b/Engine/Shaders/Private/HairStrands/HairStrandsMeshProjection.usf
deleted file mode 100644
diff --git a/Engine/Shaders/Private/HairStrands/HairStrandsMeshProjectionCommon.ush b/Engine/Shaders/Private/HairStrands/HairStrandsMeshProjectionCommon.ush
--- a/Engine/Shaders/Private/HairStrands/HairStrandsMeshProjectionCommon.ush
+++ b/Engine/Shaders/Private/HairStrands/HairStrandsMeshProjectionCommon.ush
@@ -256,4 +256,26 @@ FHairTriangleIndex DecodeTriangleIndex(uint Encoded)
 	Out.SectionIndex = (Encoded>>24) & 0xFF;
 	Out.TriangleIndex = Encoded & 0xFFFFFF;
 	return Out;
+}
+
+// Apply RBF weight onto a control point
+float3 ApplyRBF(
+	float3 RestControlPoint,
+	int InMaxSampleCount,
+	Buffer<float4> InRestSamplePositionsBuffer,
+	Buffer<float4> InMeshSampleWeightsBuffer)
+{
+	float3 ControlPoint = RestControlPoint;
+	// Apply rbf interpolation from the samples set
+	for (int i = 0; i < InMaxSampleCount; ++i)
+	{
+		const float3 PositionDelta = RestControlPoint - InRestSamplePositionsBuffer[i].xyz;
+		const float FunctionValue = sqrt(dot(PositionDelta, PositionDelta) + 1);
+		ControlPoint += FunctionValue * InMeshSampleWeightsBuffer[i].xyz;
+	}
+	ControlPoint += InMeshSampleWeightsBuffer[InMaxSampleCount].xyz;
+	ControlPoint += InMeshSampleWeightsBuffer[InMaxSampleCount + 1].xyz * RestControlPoint.x;
+	ControlPoint += InMeshSampleWeightsBuffer[InMaxSampleCount + 2].xyz * RestControlPoint.y;
+	ControlPoint += InMeshSampleWeightsBuffer[InMaxSampleCount + 3].xyz * RestControlPoint.z;
+	return ControlPoint;
 }
\ No newline at end of file
diff --git a/Engine/Shaders/Private/HairStrands/HairStrandsMeshRBF.ush b/Engine/Shaders/Private/HairStrands/HairStrandsMeshRBF.ush
deleted file mode 100644
diff --git a/Engine/Shaders/Private/HairStrands/HairStrandsMeshUpdate.usf b/Engine/Shaders/Private/HairStrands/HairStrandsMeshUpdate.usf
deleted file mode 100644
diff --git a/Engine/Shaders/Private/HairStrands/HairStrandsSamplesInit.usf b/Engine/Shaders/Private/HairStrands/HairStrandsSamplesInit.usf
deleted file mode 100644
diff --git a/Engine/Shaders/Private/HairStrands/HairStrandsSamplesUpdate.usf b/Engine/Shaders/Private/HairStrands/HairStrandsSamplesUpdate.usf
deleted file mode 100644
