Reduce groom binding data size by deduplicated unique triangle data.

This reduces binding asset type around x5.

diff --git a/Engine/Plugins/Runtime/HairStrands/Shaders/Private/NiagaraDataInterfaceHairStrandsTemplate.ush b/Engine/Plugins/Runtime/HairStrands/Shaders/Private/NiagaraDataInterfaceHairStrandsTemplate.ush
--- a/Engine/Plugins/Runtime/HairStrands/Shaders/Private/NiagaraDataInterfaceHairStrandsTemplate.ush
+++ b/Engine/Plugins/Runtime/HairStrands/Shaders/Private/NiagaraDataInterfaceHairStrandsTemplate.ush
@@ -14,6 +14,7 @@ Buffer<float4>		{ParameterName}_DeformedTrianglePositionABuffer;
 Buffer<float4>		{ParameterName}_DeformedTrianglePositionBBuffer;
 Buffer<float4>		{ParameterName}_DeformedTrianglePositionCBuffer;
 Buffer<uint>		{ParameterName}_RootBarycentricCoordinatesBuffer;
+Buffer<uint>		{ParameterName}_RootToUniqueTriangleIndexBuffer;
 float4x4			{ParameterName}_WorldTransform;
 float4x4			{ParameterName}_WorldInverse;
 float4				{ParameterName}_WorldRotation;
@@ -557,20 +558,22 @@ void ComputeProjectionTriangle_{ParameterName}(in float2 ProjectionUV, in float3
 }
 
 // Get the deformed triangle
-void BuildDeformedTriangle_{ParameterName}(in float2 ProjectionUV, in int StrandIndex, out float3 OutTrianglePosition, out float4 OutTriangleRotation)
+void BuildDeformedTriangle_{ParameterName}(in float2 ProjectionUV, in int TriangleIndex, out float3 OutTrianglePosition, out float4 OutTriangleRotation)
 {
-	ComputeProjectionTriangle_{ParameterName}(ProjectionUV, {ParameterName}_DeformedTrianglePositionABuffer[StrandIndex].xyz,
-		{ParameterName}_DeformedTrianglePositionBBuffer[StrandIndex].xyz, 
-		{ParameterName}_DeformedTrianglePositionCBuffer[StrandIndex].xyz, 
+	ComputeProjectionTriangle_{ParameterName}(ProjectionUV, 
+		{ParameterName}_DeformedTrianglePositionABuffer[TriangleIndex].xyz,
+		{ParameterName}_DeformedTrianglePositionBBuffer[TriangleIndex].xyz,
+		{ParameterName}_DeformedTrianglePositionCBuffer[TriangleIndex].xyz,
 	    {ParameterName}_DeformedRootOffset, OutTrianglePosition, OutTriangleRotation);
 }
 
 // Get the rest triangle
-void BuildRestTriangle_{ParameterName}(in float2 ProjectionUV, in int StrandIndex, out float3 OutTrianglePosition, out float4 OutTriangleRotation)
+void BuildRestTriangle_{ParameterName}(in float2 ProjectionUV, in int TriangleIndex, out float3 OutTrianglePosition, out float4 OutTriangleRotation)
 {
-	ComputeProjectionTriangle_{ParameterName}(ProjectionUV, {ParameterName}_RestTrianglePositionABuffer[StrandIndex].xyz,
-		{ParameterName}_RestTrianglePositionBBuffer[StrandIndex].xyz,
-		{ParameterName}_RestTrianglePositionCBuffer[StrandIndex].xyz,
+	ComputeProjectionTriangle_{ParameterName}(ProjectionUV, 
+		{ParameterName}_RestTrianglePositionABuffer[TriangleIndex].xyz,
+		{ParameterName}_RestTrianglePositionBBuffer[TriangleIndex].xyz,
+		{ParameterName}_RestTrianglePositionCBuffer[TriangleIndex].xyz,
 	    {ParameterName}_RestRootOffset, OutTrianglePosition, OutTriangleRotation);
 }
 
@@ -605,9 +608,10 @@ void UpdateTriangleDisplace_{ParameterName}(in int StrandIndex, in float3 NodePo
 	if( {ParameterName}_InterpolationMode >= 1)
 	{
 		const float2 ProjectionUV = ExtractBarycentrics({ParameterName}_RootBarycentricCoordinatesBuffer[StrandIndex]);
-	
-		BuildRestTriangle_{ParameterName}(ProjectionUV, StrandIndex, RestTrianglePosition, RestTriangleOrientation);
-		BuildDeformedTriangle_{ParameterName}(ProjectionUV, StrandIndex, DeformedTrianglePosition, DeformedTriangleOrientation);
+		const uint TriangleIndex = {ParameterName}_RootToUniqueTriangleIndexBuffer[StrandIndex];
+
+		BuildRestTriangle_{ParameterName}(ProjectionUV, TriangleIndex, RestTrianglePosition, RestTriangleOrientation);
+		BuildDeformedTriangle_{ParameterName}(ProjectionUV, TriangleIndex, DeformedTrianglePosition, DeformedTriangleOrientation);
 	
 		const float3 LocalRestPosition = TriangleLocalPosition_{ParameterName}(RestTrianglePosition, RestTriangleOrientation, RestPosition);
 		const float3 LocalDeformedPosition = TriangleLocalPosition_{ParameterName}(DeformedTrianglePosition, DeformedTriangleOrientation, NodePosition);
@@ -851,10 +855,11 @@ void ComputeLocalState_{ParameterName}(in float3 RestPosition, in float4 RestOri
 	{
 		const int StrandIndex = ExecIndex() / {ParameterName}_StrandSize;
 		const float2 ProjectionUV = ExtractBarycentrics({ParameterName}_RootBarycentricCoordinatesBuffer[StrandIndex]);
-		
+		const uint TriangleIndex = {ParameterName}_RootToUniqueTriangleIndexBuffer[StrandIndex];
+
 		float3 RestTrianglePosition = float3(0, 0, 0);
 		float4 RestTriangleOrientation = QUATERNION_IDENTITY;
-		BuildRestTriangle_{ParameterName}(ProjectionUV, StrandIndex, RestTrianglePosition, RestTriangleOrientation);
+		BuildRestTriangle_{ParameterName}(ProjectionUV, TriangleIndex, RestTrianglePosition, RestTriangleOrientation);
 		
 		OutLocalPosition = TriangleLocalPosition_{ParameterName}(RestTrianglePosition, RestTriangleOrientation, RestPosition);
 		OutLocalOrientation = TriangleLocalOrientation_{ParameterName}(RestTrianglePosition, RestTriangleOrientation, RestOrientation);
@@ -871,10 +876,11 @@ void AttachNodeState_{ParameterName}(in float3 LocalPosition, in float4 LocalOri
 	{
 		const int StrandIndex = ExecIndex() / {ParameterName}_StrandSize;
 		const float2 ProjectionUV = ExtractBarycentrics({ParameterName}_RootBarycentricCoordinatesBuffer[StrandIndex]);
-		
+		const uint TriangleIndex = {ParameterName}_RootToUniqueTriangleIndexBuffer[StrandIndex];
+
 		float3 DeformedTrianglePosition = float3(0, 0, 0);
 		float4 DeformedTriangleOrientation = QUATERNION_IDENTITY;
-		BuildDeformedTriangle_{ParameterName}(ProjectionUV, StrandIndex, DeformedTrianglePosition, DeformedTriangleOrientation);
+		BuildDeformedTriangle_{ParameterName}(ProjectionUV, TriangleIndex, DeformedTrianglePosition, DeformedTriangleOrientation);
 		
 		OutNodePosition = TriangleWorldPosition_{ParameterName}(DeformedTrianglePosition, DeformedTriangleOrientation, LocalPosition);
 		OutNodeOrientation = TriangleWorldOrientation_{ParameterName}(DeformedTrianglePosition, DeformedTriangleOrientation, LocalOrientation);
diff --git a/Engine/Plugins/Runtime/HairStrands/Shaders/Private/NiagaraInterpolateGroomGuides.usf b/Engine/Plugins/Runtime/HairStrands/Shaders/Private/NiagaraInterpolateGroomGuides.usf
--- a/Engine/Plugins/Runtime/HairStrands/Shaders/Private/NiagaraInterpolateGroomGuides.usf
+++ b/Engine/Plugins/Runtime/HairStrands/Shaders/Private/NiagaraInterpolateGroomGuides.usf
@@ -22,6 +22,7 @@ Buffer<float4>		DeformedTrianglePositionABuffer;
 Buffer<float4>		DeformedTrianglePositionBBuffer;
 Buffer<float4>		DeformedTrianglePositionCBuffer;
 Buffer<uint>		RootBarycentricCoordinatesBuffer;
+Buffer<uint>		RootToUniqueTriangleIndexBuffer;
 
 #endif
 
@@ -44,20 +45,24 @@ int					RestPositionComponent;
 #if PERMUTATION_INTERPOLATION == LOCAL_INTERPOLATION
 
 // Get the deformed triangle
-void DIHairStrands_BuildDeformedTriangle(in float2 ProjectionUV, in int StrandIndex, out float3 OutTrianglePosition, out float4 OutTriangleRotation)
+void DIHairStrands_BuildDeformedTriangle(in float2 ProjectionUV, in uint TriangleIndex, in int StrandIndex, out float3 OutTrianglePosition, out float4 OutTriangleRotation)
 {
-	DIHairStrands_ComputeProjectionTriangle(ProjectionUV, DeformedTrianglePositionABuffer[StrandIndex].xyz,
-		DeformedTrianglePositionBBuffer[StrandIndex].xyz, 
-		DeformedTrianglePositionCBuffer[StrandIndex].xyz, 
+	DIHairStrands_ComputeProjectionTriangle(
+		ProjectionUV, 
+		DeformedTrianglePositionABuffer[TriangleIndex].xyz,
+		DeformedTrianglePositionBBuffer[TriangleIndex].xyz, 
+		DeformedTrianglePositionCBuffer[TriangleIndex].xyz, 
 		float3(0,0,0), OutTrianglePosition, OutTriangleRotation);
 }
 
 // Get the rest triangle
-void DIHairStrands_BuildRestTriangle(in float2 ProjectionUV, in int StrandIndex, out float3 OutTrianglePosition, out float4 OutTriangleRotation)
+void DIHairStrands_BuildRestTriangle(in float2 ProjectionUV, in uint TriangleIndex, in int StrandIndex, out float3 OutTrianglePosition, out float4 OutTriangleRotation)
 {
-	DIHairStrands_ComputeProjectionTriangle(ProjectionUV, RestTrianglePositionABuffer[StrandIndex].xyz,
-		RestTrianglePositionBBuffer[StrandIndex].xyz,
-		RestTrianglePositionCBuffer[StrandIndex].xyz,
+	DIHairStrands_ComputeProjectionTriangle(
+		ProjectionUV, 
+		RestTrianglePositionABuffer[TriangleIndex].xyz,
+		RestTrianglePositionBBuffer[TriangleIndex].xyz,
+		RestTrianglePositionCBuffer[TriangleIndex].xyz,
 		float3(0,0,0), OutTrianglePosition, OutTriangleRotation);
 }
 
@@ -66,7 +71,7 @@ void DIHairStrands_BuildRestTriangle(in float2 ProjectionUV, in int StrandIndex,
 [numthreads(THREAD_COUNT, 1, 1)]
 void MainCS(uint3 DispatchThreadId : SV_DispatchThreadID)
 {
-	if (DispatchThreadId.x < NumPoints)
+	if (DispatchThreadId.x < uint(NumPoints))
 	{
 		const int StrandIndex = VertexToCurveIndexBuffer[DispatchThreadId.x];
 		
@@ -112,8 +117,9 @@ void MainCS(uint3 DispatchThreadId : SV_DispatchThreadID)
 		float3 RestTrianglePosition = float3(0, 0, 0), DeformedTrianglePosition = float3(0, 0, 0);
 		float4 RestTriangleOrientation = QUATERNION_IDENTITY, DeformedTriangleOrientation = QUATERNION_IDENTITY;
 		
-		DIHairStrands_BuildRestTriangle(ProjectionUV, StrandIndex, RestTrianglePosition, RestTriangleOrientation);
-		DIHairStrands_BuildDeformedTriangle(ProjectionUV, StrandIndex, DeformedTrianglePosition, DeformedTriangleOrientation);
+		const uint TriangleIndex = RootToUniqueTriangleIndexBuffer[StrandIndex];
+		DIHairStrands_BuildRestTriangle(ProjectionUV, TriangleIndex, StrandIndex, RestTrianglePosition, RestTriangleOrientation);
+		DIHairStrands_BuildDeformedTriangle(ProjectionUV, TriangleIndex, StrandIndex, DeformedTrianglePosition, DeformedTriangleOrientation);
 		
 		const float3 LocalPrevRestPosition = DIHairStrands_TriangleLocalPosition(RestTrianglePosition, RestTriangleOrientation, WorldPrevRestPosition);
 		const float3 LocalPrevNodePosition = DIHairStrands_TriangleLocalPosition(DeformedTrianglePosition, DeformedTriangleOrientation, WorldPrevNodePosition);
diff --git a/Engine/Plugins/Runtime/HairStrands/Source/HairStrandsCore/Private/GroomBindingBuilder.cpp b/Engine/Plugins/Runtime/HairStrands/Source/HairStrandsCore/Private/GroomBindingBuilder.cpp
--- a/Engine/Plugins/Runtime/HairStrands/Source/HairStrandsCore/Private/GroomBindingBuilder.cpp
+++ b/Engine/Plugins/Runtime/HairStrands/Source/HairStrandsCore/Private/GroomBindingBuilder.cpp
@@ -57,7 +57,7 @@ static FAutoConsoleVariableRef CVarHairStrandsBindingBuilderWarningEnable(TEXT("
 FString FGroomBindingBuilder::GetVersion()
 {
 	// Important to update the version when groom building changes
-	return TEXT("2f");
+	return TEXT("2gc");
 }
 
 ///////////////////////////////////////////////////////////////////////////////////////////////////
@@ -102,6 +102,17 @@ float FHairStrandsRootUtils::PackUVsToFloat(const FVector2f& UV)
 	return *((float*)(&Encoded));
 }
 
+//////////////////////////////////////////////////////////////////////////
+// Intermediate data struct
+
+/** Binding data */
+struct FHairRootGroupData
+{
+	FHairStrandsRootData			SimRootData;
+	FHairStrandsRootData			RenRootData;
+	TArray<FHairStrandsRootData>	CardsRootData;
+};
+
 namespace
 {
 //////////////////////////////////////////////////////////////////////////
@@ -647,20 +658,18 @@ namespace GroomBinding_RBFWeighting
 
 		const bool ValidSection = (TargetSection >= 0 && TargetSection < MeshLODData.GetNumSections());
 
-		const TArray<uint32>& RootBuffers = ProjectionLOD.RootTriangleIndexBuffer;
-		for (int32 RootIt = 0; RootIt < RootBuffers.Num(); ++RootIt)
+		for (uint32 EncodedTriangleId : ProjectionLOD.UniqueTriangleIndexBuffer)
 		{
-			uint32 SectionIndex  = 0;
+			uint32 SectionIndex = 0;
 			uint32 TriangleIndex = 0;
-			FHairStrandsRootUtils::DecodeTriangleIndex(RootBuffers[RootIt], TriangleIndex, SectionIndex);
-			if (!ValidSection || (ValidSection && (SectionIndex == TargetSection) ) )
+			FHairStrandsRootUtils::DecodeTriangleIndex(EncodedTriangleId, TriangleIndex, SectionIndex);
+			if (!ValidSection || (ValidSection && (SectionIndex == TargetSection)))
 			{
 				const IMeshSectionData& Section = MeshLODData.GetSection(SectionIndex);
 				for (uint32 VertexIt = 0; VertexIt < 3; ++VertexIt)
 				{
 					const uint32 VertexIndex = TriangleIndices[Section.GetBaseIndex() + 3 * TriangleIndex + VertexIt];
-					ValidPoints[VertexIndex] = (VertexIndex >= Section.GetBaseVertexIndex()) && (VertexIndex < 
-						Section.GetBaseVertexIndex() + Section.GetNumVertices());
+					ValidPoints[VertexIndex] = (VertexIndex >= Section.GetBaseVertexIndex()) && (VertexIndex < Section.GetBaseVertexIndex() + Section.GetNumVertices());
 				}
 			}
 		}
@@ -1209,27 +1218,39 @@ namespace GroomBinding_RootProjection
 				return false;
 			}
 
-			OutRootData.MeshProjectionLODs[LODIt].RootTriangleIndexBuffer.SetNum(CurveCount);
-			OutRootData.MeshProjectionLODs[LODIt].RootTriangleBarycentricBuffer.SetNum(CurveCount);
-			OutRootData.MeshProjectionLODs[LODIt].RestRootTrianglePosition0Buffer.SetNum(CurveCount);
-			OutRootData.MeshProjectionLODs[LODIt].RestRootTrianglePosition1Buffer.SetNum(CurveCount);
-			OutRootData.MeshProjectionLODs[LODIt].RestRootTrianglePosition2Buffer.SetNum(CurveCount);
+			OutRootData.MeshProjectionLODs[LODIt].RootBarycentricBuffer.SetNum(CurveCount);
+			OutRootData.MeshProjectionLODs[LODIt].RootToUniqueTriangleIndexBuffer.SetNum(CurveCount);
 
 			// 2.3. Compute the closest triangle for each root
 			//InMeshRenderData->LODRenderData[LODIt].GetNumVertices();
-#if BINDING_PARALLEL_BUILDING
+
+			TArray<FHairStrandsCurveTriangleIndexFormat::Type> RootTriangleIndexBuffer;
+			RootTriangleIndexBuffer.SetNum(CurveCount);
+
+			TArray<FHairStrandsMeshTrianglePositionFormat::Type> RestRootTrianglePosition0Buffer;
+			TArray<FHairStrandsMeshTrianglePositionFormat::Type> RestRootTrianglePosition1Buffer;
+			TArray<FHairStrandsMeshTrianglePositionFormat::Type> RestRootTrianglePosition2Buffer;
+			RestRootTrianglePosition0Buffer.SetNum(CurveCount);
+			RestRootTrianglePosition1Buffer.SetNum(CurveCount);
+			RestRootTrianglePosition2Buffer.SetNum(CurveCount);
+
+		#if BINDING_PARALLEL_BUILDING
 			TAtomic<uint32> bIsValid(1);
 			ParallelFor(CurveCount,
 				[
 					LODIt,
 					&InStrandsData,
 					&Grid,
+					&RootTriangleIndexBuffer,
+					&RestRootTrianglePosition0Buffer,
+					&RestRootTrianglePosition1Buffer,
+					&RestRootTrianglePosition2Buffer,
 					&OutRootData,
 					&bIsValid
 				] (uint32 CurveIndex)
-#else
+		#else
 			for (uint32 CurveIndex = 0; CurveIndex < CurveCount; ++CurveIndex)
-#endif
+		#endif
 			{
 				const uint32 Offset = InStrandsData.StrandsCurves.CurvesOffset[CurveIndex];
 				const FVector& RootP = (FVector)InStrandsData.StrandsPoints.PointsPosition[Offset];
@@ -1237,11 +1258,11 @@ namespace GroomBinding_RootProjection
 
 				if (Cells.Num() == 0)
 				{
-#if BINDING_PARALLEL_BUILDING
+				#if BINDING_PARALLEL_BUILDING
 					bIsValid = 0; return;
-#else
+				#else
 					return false;
-#endif
+				#endif
 				}
 
 				float ClosestDistance = FLT_MAX;
@@ -1263,42 +1284,95 @@ namespace GroomBinding_RootProjection
 				}
 				check(ClosestDistance < FLT_MAX);
 
+				// Record closest triangle and the root's barycentrics
 				const uint32 EncodedBarycentrics = FHairStrandsRootUtils::EncodeBarycentrics(FVector2f(ClosestBarycentrics));	// LWC_TODO: Precision loss
 				const uint32 EncodedTriangleIndex = FHairStrandsRootUtils::EncodeTriangleIndex(ClosestTriangle.TriangleIndex, ClosestTriangle.SectionIndex);
-				OutRootData.MeshProjectionLODs[LODIt].RootTriangleIndexBuffer[CurveIndex] = EncodedTriangleIndex;
-				OutRootData.MeshProjectionLODs[LODIt].RootTriangleBarycentricBuffer[CurveIndex] = EncodedBarycentrics;
-				OutRootData.MeshProjectionLODs[LODIt].RestRootTrianglePosition0Buffer[CurveIndex] = FVector4f((FVector3f)ClosestTriangle.P0, FHairStrandsRootUtils::PackUVsToFloat(FVector2f(ClosestTriangle.UV0)));	// LWC_TODO: Precision loss
-				OutRootData.MeshProjectionLODs[LODIt].RestRootTrianglePosition1Buffer[CurveIndex] = FVector4f((FVector3f)ClosestTriangle.P1, FHairStrandsRootUtils::PackUVsToFloat(FVector2f(ClosestTriangle.UV1)));	// LWC_TODO: Precision loss
-				OutRootData.MeshProjectionLODs[LODIt].RestRootTrianglePosition2Buffer[CurveIndex] = FVector4f((FVector3f)ClosestTriangle.P2, FHairStrandsRootUtils::PackUVsToFloat(FVector2f(ClosestTriangle.UV2)));	// LWC_TODO: Precision loss
+				OutRootData.MeshProjectionLODs[LODIt].RootBarycentricBuffer[CurveIndex] = EncodedBarycentrics;
+
+				RootTriangleIndexBuffer[CurveIndex] = EncodedTriangleIndex;
+				RestRootTrianglePosition0Buffer[CurveIndex] = FVector4f((FVector3f)ClosestTriangle.P0, FHairStrandsRootUtils::PackUVsToFloat(FVector2f(ClosestTriangle.UV0)));	// LWC_TODO: Precision loss
+				RestRootTrianglePosition1Buffer[CurveIndex] = FVector4f((FVector3f)ClosestTriangle.P1, FHairStrandsRootUtils::PackUVsToFloat(FVector2f(ClosestTriangle.UV1)));	// LWC_TODO: Precision loss
+				RestRootTrianglePosition2Buffer[CurveIndex] = FVector4f((FVector3f)ClosestTriangle.P2, FHairStrandsRootUtils::PackUVsToFloat(FVector2f(ClosestTriangle.UV2)));	// LWC_TODO: Precision loss
 			}
-#if BINDING_PARALLEL_BUILDING
+		#if BINDING_PARALLEL_BUILDING
 			);
 			if (bIsValid == 0)
 			{
 				return false;
 			}
-#endif
+		#endif
+
+			// Build list of unique triangles
+			TArray<uint32> UniqueSectionId;
+			TArray<uint32> UniqueTriangleToRootList;
+			TMap<uint32, TArray<uint32>> UniqueTriangleToRootMap;
+			for (uint32 CurveIndex = 0; CurveIndex < CurveCount; ++CurveIndex)
+			{
+				const uint32 EncodedTriangleId = RootTriangleIndexBuffer[CurveIndex];
+				if (TArray<uint32>* CurvesList = UniqueTriangleToRootMap.Find(EncodedTriangleId))
+				{
+					CurvesList->Add(CurveIndex);
+				}
+				else
+				{
+					// Add unique section
+					uint32 TriangleIndex;
+					uint32 SectionIndex;
+					FHairStrandsRootUtils::DecodeTriangleIndex(EncodedTriangleId, TriangleIndex, SectionIndex);
+					UniqueSectionId.AddUnique(SectionIndex);
+
+					// Add unique triangle
+					UniqueTriangleToRootList.Add(EncodedTriangleId);
+					TArray<uint32>& NewCurvesList = UniqueTriangleToRootMap.Add(EncodedTriangleId);
+					NewCurvesList.Add(CurveIndex);
+				}
+			}
+
+			// Sort unique triangle per section and triangle ID (encoded triangle ID stores section ID in high bits)
+			UniqueTriangleToRootList.Sort();
+			UniqueSectionId.Sort();
+
+			// Build final unique triangle list and the root-to-unique-triangle mapping
+			const uint32 UniqueTriangleCount = UniqueTriangleToRootList.Num();
+			OutRootData.MeshProjectionLODs[LODIt].UniqueTriangleIndexBuffer.Reserve(UniqueTriangleCount);
+			OutRootData.MeshProjectionLODs[LODIt].RestUniqueTrianglePosition0Buffer.Reserve(UniqueTriangleCount);
+			OutRootData.MeshProjectionLODs[LODIt].RestUniqueTrianglePosition1Buffer.Reserve(UniqueTriangleCount);
+			OutRootData.MeshProjectionLODs[LODIt].RestUniqueTrianglePosition2Buffer.Reserve(UniqueTriangleCount);
+			for (uint32 EncodedTriangleId : UniqueTriangleToRootList)
+			{
+				auto It = UniqueTriangleToRootMap.Find(EncodedTriangleId);
+				check(It);
+
+				OutRootData.MeshProjectionLODs[LODIt].UniqueTriangleIndexBuffer.Add(EncodedTriangleId);
+
+				const uint32 FirstCurveIndex = (*It)[0];
+				OutRootData.MeshProjectionLODs[LODIt].RestUniqueTrianglePosition0Buffer.Add(RestRootTrianglePosition0Buffer[FirstCurveIndex]);
+				OutRootData.MeshProjectionLODs[LODIt].RestUniqueTrianglePosition1Buffer.Add(RestRootTrianglePosition1Buffer[FirstCurveIndex]);
+				OutRootData.MeshProjectionLODs[LODIt].RestUniqueTrianglePosition2Buffer.Add(RestRootTrianglePosition2Buffer[FirstCurveIndex]);
 
-			// Update the valid & unique sections IDs
-			FGroomBindingBuilder::BuildUniqueSections(OutRootData.MeshProjectionLODs[LODIt]);
+				// Write for each root, the index of the triangle
+				const uint32 UniqueTriangleIndex = OutRootData.MeshProjectionLODs[LODIt].UniqueTriangleIndexBuffer.Num()-1;
+				for (uint32 CurveIndex : *It)
+				{
+					OutRootData.MeshProjectionLODs[LODIt].RootToUniqueTriangleIndexBuffer[CurveIndex] = UniqueTriangleIndex;
+				}
+			}
+
+			// Sanity check
+			check(OutRootData.MeshProjectionLODs[LODIt].RootToUniqueTriangleIndexBuffer.Num() == CurveCount);
+			check(OutRootData.MeshProjectionLODs[LODIt].RestUniqueTrianglePosition0Buffer.Num() == UniqueTriangleCount);
+			check(OutRootData.MeshProjectionLODs[LODIt].RestUniqueTrianglePosition1Buffer.Num() == UniqueTriangleCount);
+			check(OutRootData.MeshProjectionLODs[LODIt].RestUniqueTrianglePosition2Buffer.Num() == UniqueTriangleCount);
+			check(OutRootData.MeshProjectionLODs[LODIt].UniqueTriangleIndexBuffer.Num() == UniqueTriangleCount);
+
+			// Update the root mesh projection data with unique valid mesh section IDs, based on the projection data
+			OutRootData.MeshProjectionLODs[LODIt].UniqueSectionIds = UniqueSectionId;
 		}
 
 		return true;
 	}
 }// namespace GroomBinding_Project
 
-void FGroomBindingBuilder::BuildUniqueSections(FHairStrandsRootData::FMeshProjectionLOD& LOD)
-{
-	LOD.ValidSectionIndices.Empty();
-	for (const FHairStrandsCurveTriangleIndexFormat::Type& Tri : LOD.RootTriangleIndexBuffer)
-	{
-		uint32 TriangleIndex;
-		uint32 SectionIndex;
-		FHairStrandsRootUtils::DecodeTriangleIndex(Tri, TriangleIndex, SectionIndex);
-		LOD.ValidSectionIndices.AddUnique(SectionIndex);
-	}
-}
-
 ///////////////////////////////////////////////////////////////////////////////////////////////////
 // Mesh transfer
 
@@ -1844,15 +1918,17 @@ static void BuildRootBulkData(
 
 		Out.MeshProjectionLODs[MeshLODIt].LODIndex = In.MeshProjectionLODs[MeshLODIt].LODIndex;
 		Out.MeshProjectionLODs[MeshLODIt].SampleCount = bHasValidSamples ? In.MeshProjectionLODs[MeshLODIt].SampleCount : 0u;
+		Out.MeshProjectionLODs[MeshLODIt].UniqueTriangleCount = In.MeshProjectionLODs[MeshLODIt].RestUniqueTrianglePosition0Buffer.Num();
 
-		CopyToBulkData<FHairStrandsCurveTriangleIndexFormat>(Out.MeshProjectionLODs[MeshLODIt].RootTriangleIndexBuffer, In.MeshProjectionLODs[MeshLODIt].RootTriangleIndexBuffer);
-		CopyToBulkData<FHairStrandsCurveTriangleBarycentricFormat>(Out.MeshProjectionLODs[MeshLODIt].RootTriangleBarycentricBuffer, In.MeshProjectionLODs[MeshLODIt].RootTriangleBarycentricBuffer);
+		CopyToBulkData<FHairStrandsCurveTriangleIndexFormat>(Out.MeshProjectionLODs[MeshLODIt].UniqueTriangleIndexBuffer, In.MeshProjectionLODs[MeshLODIt].UniqueTriangleIndexBuffer);
+		CopyToBulkData<FHairStrandsCurveTriangleBarycentricFormat>(Out.MeshProjectionLODs[MeshLODIt].RootBarycentricBuffer, In.MeshProjectionLODs[MeshLODIt].RootBarycentricBuffer);
+		CopyToBulkData<FHairStrandsCurveTriangleIndexFormat>(Out.MeshProjectionLODs[MeshLODIt].RootToUniqueTriangleIndexBuffer, In.MeshProjectionLODs[MeshLODIt].RootToUniqueTriangleIndexBuffer);
 
-		CopyToBulkData<FHairStrandsMeshTrianglePositionFormat>(Out.MeshProjectionLODs[MeshLODIt].RestRootTrianglePosition0Buffer, In.MeshProjectionLODs[MeshLODIt].RestRootTrianglePosition0Buffer);
-		CopyToBulkData<FHairStrandsMeshTrianglePositionFormat>(Out.MeshProjectionLODs[MeshLODIt].RestRootTrianglePosition1Buffer, In.MeshProjectionLODs[MeshLODIt].RestRootTrianglePosition1Buffer);
-		CopyToBulkData<FHairStrandsMeshTrianglePositionFormat>(Out.MeshProjectionLODs[MeshLODIt].RestRootTrianglePosition2Buffer, In.MeshProjectionLODs[MeshLODIt].RestRootTrianglePosition2Buffer);
+		CopyToBulkData<FHairStrandsMeshTrianglePositionFormat>(Out.MeshProjectionLODs[MeshLODIt].RestUniqueTrianglePosition0Buffer, In.MeshProjectionLODs[MeshLODIt].RestUniqueTrianglePosition0Buffer);
+		CopyToBulkData<FHairStrandsMeshTrianglePositionFormat>(Out.MeshProjectionLODs[MeshLODIt].RestUniqueTrianglePosition1Buffer, In.MeshProjectionLODs[MeshLODIt].RestUniqueTrianglePosition1Buffer);
+		CopyToBulkData<FHairStrandsMeshTrianglePositionFormat>(Out.MeshProjectionLODs[MeshLODIt].RestUniqueTrianglePosition2Buffer, In.MeshProjectionLODs[MeshLODIt].RestUniqueTrianglePosition2Buffer);
 
-		Out.MeshProjectionLODs[MeshLODIt].ValidSectionIndices = In.MeshProjectionLODs[MeshLODIt].ValidSectionIndices;
+		Out.MeshProjectionLODs[MeshLODIt].UniqueSectionIndices = In.MeshProjectionLODs[MeshLODIt].UniqueSectionIds;
 
 		if (bHasValidSamples)
 		{
@@ -1888,14 +1964,15 @@ static void BuildRootData(
 		Out.MeshProjectionLODs[MeshLODIt].LODIndex = In.MeshProjectionLODs[MeshLODIt].LODIndex;
 		Out.MeshProjectionLODs[MeshLODIt].SampleCount = bHasValidSamples ? In.MeshProjectionLODs[MeshLODIt].SampleCount : 0u;
 
-		CopyFromBulkData<FHairStrandsCurveTriangleIndexFormat>(Out.MeshProjectionLODs[MeshLODIt].RootTriangleIndexBuffer, In.MeshProjectionLODs[MeshLODIt].RootTriangleIndexBuffer);
-		CopyFromBulkData<FHairStrandsCurveTriangleBarycentricFormat>(Out.MeshProjectionLODs[MeshLODIt].RootTriangleBarycentricBuffer, In.MeshProjectionLODs[MeshLODIt].RootTriangleBarycentricBuffer);
+		CopyFromBulkData<FHairStrandsCurveTriangleIndexFormat>(Out.MeshProjectionLODs[MeshLODIt].UniqueTriangleIndexBuffer, In.MeshProjectionLODs[MeshLODIt].UniqueTriangleIndexBuffer);
+		CopyFromBulkData<FHairStrandsCurveTriangleIndexFormat>(Out.MeshProjectionLODs[MeshLODIt].RootToUniqueTriangleIndexBuffer, In.MeshProjectionLODs[MeshLODIt].RootToUniqueTriangleIndexBuffer);
+		CopyFromBulkData<FHairStrandsCurveTriangleBarycentricFormat>(Out.MeshProjectionLODs[MeshLODIt].RootBarycentricBuffer, In.MeshProjectionLODs[MeshLODIt].RootBarycentricBuffer);
 
-		CopyFromBulkData<FHairStrandsMeshTrianglePositionFormat>(Out.MeshProjectionLODs[MeshLODIt].RestRootTrianglePosition0Buffer, In.MeshProjectionLODs[MeshLODIt].RestRootTrianglePosition0Buffer);
-		CopyFromBulkData<FHairStrandsMeshTrianglePositionFormat>(Out.MeshProjectionLODs[MeshLODIt].RestRootTrianglePosition1Buffer, In.MeshProjectionLODs[MeshLODIt].RestRootTrianglePosition1Buffer);
-		CopyFromBulkData<FHairStrandsMeshTrianglePositionFormat>(Out.MeshProjectionLODs[MeshLODIt].RestRootTrianglePosition2Buffer, In.MeshProjectionLODs[MeshLODIt].RestRootTrianglePosition2Buffer);
+		CopyFromBulkData<FHairStrandsMeshTrianglePositionFormat>(Out.MeshProjectionLODs[MeshLODIt].RestUniqueTrianglePosition0Buffer, In.MeshProjectionLODs[MeshLODIt].RestUniqueTrianglePosition0Buffer);
+		CopyFromBulkData<FHairStrandsMeshTrianglePositionFormat>(Out.MeshProjectionLODs[MeshLODIt].RestUniqueTrianglePosition1Buffer, In.MeshProjectionLODs[MeshLODIt].RestUniqueTrianglePosition1Buffer);
+		CopyFromBulkData<FHairStrandsMeshTrianglePositionFormat>(Out.MeshProjectionLODs[MeshLODIt].RestUniqueTrianglePosition2Buffer, In.MeshProjectionLODs[MeshLODIt].RestUniqueTrianglePosition2Buffer);
 
-		Out.MeshProjectionLODs[MeshLODIt].ValidSectionIndices = In.MeshProjectionLODs[MeshLODIt].ValidSectionIndices;
+		Out.MeshProjectionLODs[MeshLODIt].UniqueSectionIds = In.MeshProjectionLODs[MeshLODIt].UniqueSectionIndices;
 
 		if (bHasValidSamples)
 		{
diff --git a/Engine/Plugins/Runtime/HairStrands/Source/HairStrandsCore/Private/GroomDebug.cpp b/Engine/Plugins/Runtime/HairStrands/Source/HairStrandsCore/Private/GroomDebug.cpp
--- a/Engine/Plugins/Runtime/HairStrands/Source/HairStrandsCore/Private/GroomDebug.cpp
+++ b/Engine/Plugins/Runtime/HairStrands/Source/HairStrandsCore/Private/GroomDebug.cpp
@@ -259,18 +259,19 @@ class FHairProjectionHairDebugCS : public FGlobalShader
 		SHADER_PARAMETER(uint32, DeformedFrameEnable)
 		SHADER_PARAMETER(FMatrix44f, RootLocalToWorld)
 
-		SHADER_PARAMETER_RDG_BUFFER_SRV(StructuredBuffer, RestPosition0Buffer)
-		SHADER_PARAMETER_RDG_BUFFER_SRV(StructuredBuffer, RestPosition1Buffer)
-		SHADER_PARAMETER_RDG_BUFFER_SRV(StructuredBuffer, RestPosition2Buffer)
+		SHADER_PARAMETER_RDG_BUFFER_SRV(Buffer, RestPosition0Buffer)
+		SHADER_PARAMETER_RDG_BUFFER_SRV(Buffer, RestPosition1Buffer)
+		SHADER_PARAMETER_RDG_BUFFER_SRV(Buffer, RestPosition2Buffer)
 
-		SHADER_PARAMETER_RDG_BUFFER_SRV(StructuredBuffer, DeformedPosition0Buffer)
-		SHADER_PARAMETER_RDG_BUFFER_SRV(StructuredBuffer, DeformedPosition1Buffer)
-		SHADER_PARAMETER_RDG_BUFFER_SRV(StructuredBuffer, DeformedPosition2Buffer)
+		SHADER_PARAMETER_RDG_BUFFER_SRV(Buffer, DeformedPosition0Buffer)
+		SHADER_PARAMETER_RDG_BUFFER_SRV(Buffer, DeformedPosition1Buffer)
+		SHADER_PARAMETER_RDG_BUFFER_SRV(Buffer, DeformedPosition2Buffer)
 
-		SHADER_PARAMETER_RDG_BUFFER_SRV(StructuredBuffer, RestSamplePositionsBuffer)
-		SHADER_PARAMETER_RDG_BUFFER_SRV(StructuredBuffer, DeformedSamplePositionsBuffer)
+		SHADER_PARAMETER_RDG_BUFFER_SRV(Buffer, RestSamplePositionsBuffer)
+		SHADER_PARAMETER_RDG_BUFFER_SRV(Buffer, DeformedSamplePositionsBuffer)
 
-		SHADER_PARAMETER_RDG_BUFFER_SRV(StructuredBuffer, RootBarycentricBuffer)
+		SHADER_PARAMETER_RDG_BUFFER_SRV(Buffer, RootBarycentricBuffer)
+		SHADER_PARAMETER_RDG_BUFFER_SRV(Buffer, RootToUniqueTriangleIndexBuffer)
 
 		SHADER_PARAMETER_STRUCT_REF(FViewUniformShaderParameters, ViewUniformBuffer)
 		SHADER_PARAMETER_STRUCT_INCLUDE(ShaderPrint::FShaderParameters, ShaderPrintUniformBuffer)
@@ -332,7 +333,7 @@ static void AddDebugProjectionHairPass(
 	ShaderPrint::RequestSpaceForTriangles(PrimitiveCount);
 
 	if (EDebugProjectionHairType::HairFrame == GeometryType &&
-		!RestRootResources->LODs[MeshLODIndex].RootTriangleBarycentricBuffer.Buffer)
+		!RestRootResources->LODs[MeshLODIndex].RootBarycentricBuffer.Buffer)
 		return;
 
 	if (EDebugProjectionHairType::HairSamples == GeometryType &&
@@ -342,16 +343,16 @@ static void AddDebugProjectionHairPass(
 	const FHairStrandsRestRootResource::FLOD& RestLODDatas = RestRootResources->LODs[MeshLODIndex];
 	const FHairStrandsDeformedRootResource::FLOD& DeformedLODDatas = DeformedRootResources->LODs[MeshLODIndex];
 
-	if (!RestLODDatas.RestRootTrianglePosition0Buffer.Buffer ||
-		!RestLODDatas.RestRootTrianglePosition1Buffer.Buffer ||
-		!RestLODDatas.RestRootTrianglePosition2Buffer.Buffer ||
-		!DeformedLODDatas.DeformedRootTrianglePosition0Buffer[0].Buffer ||
-		!DeformedLODDatas.DeformedRootTrianglePosition1Buffer[0].Buffer ||
-		!DeformedLODDatas.DeformedRootTrianglePosition2Buffer[0].Buffer	)
+	if (!RestLODDatas.RestUniqueTrianglePosition0Buffer.Buffer ||
+		!RestLODDatas.RestUniqueTrianglePosition1Buffer.Buffer ||
+		!RestLODDatas.RestUniqueTrianglePosition2Buffer.Buffer ||
+		!DeformedLODDatas.DeformedUniqueTrianglePosition0Buffer[0].Buffer ||
+		!DeformedLODDatas.DeformedUniqueTrianglePosition1Buffer[0].Buffer ||
+		!DeformedLODDatas.DeformedUniqueTrianglePosition2Buffer[0].Buffer	)
 		return;
 
 	// Double buffering is disabled by default unless the read-only cvar r.HairStrands.ContinuousDecimationReordering is set
-	if (IsHairStrandContinuousDecimationReorderingEnabled() && (!DeformedLODDatas.DeformedRootTrianglePosition0Buffer[1].Buffer || !DeformedLODDatas.DeformedRootTrianglePosition1Buffer[1].Buffer || !DeformedLODDatas.DeformedRootTrianglePosition2Buffer[1].Buffer))
+	if (IsHairStrandContinuousDecimationReorderingEnabled() && (!DeformedLODDatas.DeformedUniqueTrianglePosition0Buffer[1].Buffer || !DeformedLODDatas.DeformedUniqueTrianglePosition1Buffer[1].Buffer || !DeformedLODDatas.DeformedUniqueTrianglePosition2Buffer[1].Buffer))
 	{
 		return;
 	}
@@ -366,16 +367,18 @@ static void AddDebugProjectionHairPass(
 
 	if (EDebugProjectionHairType::HairFrame == GeometryType)
 	{
-		Parameters->RootBarycentricBuffer	= RegisterAsSRV(GraphBuilder, RestLODDatas.RootTriangleBarycentricBuffer);
+		Parameters->RootBarycentricBuffer	= RegisterAsSRV(GraphBuilder, RestLODDatas.RootBarycentricBuffer);
 	}
 
-	Parameters->RestPosition0Buffer = RegisterAsSRV(GraphBuilder, RestLODDatas.RestRootTrianglePosition0Buffer);
-	Parameters->RestPosition1Buffer = RegisterAsSRV(GraphBuilder, RestLODDatas.RestRootTrianglePosition1Buffer);
-	Parameters->RestPosition2Buffer = RegisterAsSRV(GraphBuilder, RestLODDatas.RestRootTrianglePosition2Buffer);
+	Parameters->RootToUniqueTriangleIndexBuffer = RegisterAsSRV(GraphBuilder, RestLODDatas.RootToUniqueTriangleIndexBuffer);
 
-	Parameters->DeformedPosition0Buffer = RegisterAsSRV(GraphBuilder, DeformedLODDatas.GetDeformedRootTrianglePosition0Buffer(FHairStrandsDeformedRootResource::FLOD::Current));
-	Parameters->DeformedPosition1Buffer = RegisterAsSRV(GraphBuilder, DeformedLODDatas.GetDeformedRootTrianglePosition1Buffer(FHairStrandsDeformedRootResource::FLOD::Current));
-	Parameters->DeformedPosition2Buffer = RegisterAsSRV(GraphBuilder, DeformedLODDatas.GetDeformedRootTrianglePosition2Buffer(FHairStrandsDeformedRootResource::FLOD::Current));
+	Parameters->RestPosition0Buffer = RegisterAsSRV(GraphBuilder, RestLODDatas.RestUniqueTrianglePosition0Buffer);
+	Parameters->RestPosition1Buffer = RegisterAsSRV(GraphBuilder, RestLODDatas.RestUniqueTrianglePosition1Buffer);
+	Parameters->RestPosition2Buffer = RegisterAsSRV(GraphBuilder, RestLODDatas.RestUniqueTrianglePosition2Buffer);
+
+	Parameters->DeformedPosition0Buffer = RegisterAsSRV(GraphBuilder, DeformedLODDatas.GetDeformedUniqueTrianglePosition0Buffer(FHairStrandsDeformedRootResource::FLOD::Current));
+	Parameters->DeformedPosition1Buffer = RegisterAsSRV(GraphBuilder, DeformedLODDatas.GetDeformedUniqueTrianglePosition1Buffer(FHairStrandsDeformedRootResource::FLOD::Current));
+	Parameters->DeformedPosition2Buffer = RegisterAsSRV(GraphBuilder, DeformedLODDatas.GetDeformedUniqueTrianglePosition2Buffer(FHairStrandsDeformedRootResource::FLOD::Current));
 
 	Parameters->RestSamplePositionsBuffer = RegisterAsSRV(GraphBuilder, RestLODDatas.RestSamplePositionsBuffer);
 	Parameters->DeformedSamplePositionsBuffer = RegisterAsSRV(GraphBuilder, DeformedLODDatas.GetDeformedSamplePositionsBuffer(FHairStrandsDeformedRootResource::FLOD::Current));
@@ -884,7 +887,12 @@ static void AddHairDebugPrintInstancePass(
 	TArray<uint8> Names;
 	Names.Reserve(MaxPrimitiveNameCount * 30u);
 
-	TArray<FUintVector4> Infos;
+	struct FInstanceInfos
+	{
+		FUintVector4 Data0;
+		FUintVector4 Data1;
+	};
+	TArray<FInstanceInfos> Infos;
 	Infos.Reserve(InstanceCount);
 	for (uint32 InstanceIndex = 0; InstanceIndex < InstanceCount; ++InstanceIndex)
 	{
@@ -912,7 +920,9 @@ static void AddHairDebugPrintInstancePass(
 		const uint32 IntLODIndex = Instance->HairGroupPublicData->LODIndex;
 		const uint32 LODCount = Instance->HairGroupPublicData->GetLODScreenSizes().Num();
 
-		const uint32 DataX =
+		FUintVector4 Data0 = { 0,0,0,0 };
+		FUintVector4 Data1 = { 0,0,0,0 };
+		Data0.X =
 			((Instance->Debug.GroupIndex & 0xFF)) |
 			((Instance->Debug.GroupCount & 0xFF) << 8) |
 			((LODCount & 0xFF) << 16) |
@@ -921,38 +931,43 @@ static void AddHairDebugPrintInstancePass(
 			((Instance->Guides.bIsSimulationEnable ? 0x1 : 0x0) << 30) |
 			((Instance->Guides.bHasGlobalInterpolation ? 0x1 : 0x0) << 31);
 
-		const uint32 DataY =
+		Data0.Y =
 			(FFloat16(LODIndex).Encoded) |
 			(FFloat16(Instance->Strands.Modifier.HairLengthScale_Override ? Instance->Strands.Modifier.HairLengthScale : -1.f).Encoded << 16);
-
-		uint32 DataZ = 0;
-		uint32 DataW = 0;
+		
 		switch (Instance->GeometryType)
 		{
 		case EHairGeometryType::Strands:
 			if (Instance->Strands.IsValid())
 			{
-				DataZ = Instance->Strands.Data->GetNumCurves(); // Change this later on for having dynamic value
-				DataW = Instance->Strands.Data->GetNumPoints(); // Change this later on for having dynamic value
+				Data0.Z = Instance->Strands.Data->GetNumCurves(); // Change this later on for having dynamic value
+				Data0.W = Instance->Strands.Data->GetNumPoints(); // Change this later on for having dynamic value
+				const int32 MeshLODIndex = Instance->HairGroupPublicData->MeshLODIndex;
+				if (MeshLODIndex>=0)
+				{
+					Data1.X = Instance->Strands.RestRootResource->BulkData.MeshProjectionLODs[MeshLODIndex].UniqueSectionIndices.Num();
+					Data1.Y = Instance->Strands.RestRootResource->BulkData.MeshProjectionLODs[MeshLODIndex].UniqueTriangleCount;
+					Data1.Z = Instance->Strands.RestRootResource->BulkData.RootCount;
+					Data1.W = Instance->Strands.RestRootResource->BulkData.PointCount;
+				}
 			}
 			break;
 		case EHairGeometryType::Cards:
 			if (Instance->Cards.IsValid(IntLODIndex))
 			{
-				DataZ = Instance->Cards.LODs[IntLODIndex].Guides.IsValid() ? Instance->Cards.LODs[IntLODIndex].Guides.Data->GetNumCurves() : 0;
-				DataW = Instance->Cards.LODs[IntLODIndex].Data->GetNumVertices();
+				Data0.Z = Instance->Cards.LODs[IntLODIndex].Guides.IsValid() ? Instance->Cards.LODs[IntLODIndex].Guides.Data->GetNumCurves() : 0;
+				Data0.W = Instance->Cards.LODs[IntLODIndex].Data->GetNumVertices();
 			}
 			break;
 		case EHairGeometryType::Meshes:
 			if (Instance->Meshes.IsValid(IntLODIndex))
 			{
-				DataZ = 0;
-				DataW = Instance->Meshes.LODs[IntLODIndex].Data->GetNumVertices();
+				Data0.Z = 0;
+				Data0.W = Instance->Meshes.LODs[IntLODIndex].Data->GetNumVertices();
 			}
 			break;
 		}
-
-		Infos.Add(FUintVector4(DataX, DataY, DataZ, DataW));
+		Infos.Add({Data0, Data1});
 	}
 
 	if (NameInfos.IsEmpty())
@@ -967,10 +982,11 @@ static void AddHairDebugPrintInstancePass(
 		Names.Add(uint8('e'));
 	}	
 
-	const uint32 InfoInBytes = 16u;
+	const uint32 InfoInBytes  = sizeof(FInstanceInfos);
+	const uint32 InfoIn4Bytes = sizeof(FInstanceInfos) / sizeof(uint8);
 	FRDGBufferRef NameBuffer = CreateVertexBuffer(GraphBuilder, TEXT("Hair.Debug.InstanceNames"), FRDGBufferDesc::CreateBufferDesc(1, Names.Num()), Names.GetData(), Names.Num());
 	FRDGBufferRef NameInfoBuffer = CreateStructuredBuffer(GraphBuilder, TEXT("Hair.Debug.InstanceNameInfos"), NameInfos);	
-	FRDGBufferRef InfoBuffer = CreateVertexBuffer(GraphBuilder, TEXT("Hair.Debug.InstanceInfos"), FRDGBufferDesc::CreateBufferDesc(InfoInBytes, Infos.Num()), Infos.GetData(), InfoInBytes * Infos.Num());
+	FRDGBufferRef InfoBuffer = CreateVertexBuffer(GraphBuilder, TEXT("Hair.Debug.InstanceInfos"), FRDGBufferDesc::CreateBufferDesc(4, InfoIn4Bytes * Infos.Num()), Infos.GetData(), InfoInBytes * Infos.Num());
 
 	// Draw general information for all instances (one pass for all instances)
 	{
@@ -980,7 +996,7 @@ static void AddHairDebugPrintInstancePass(
 		Parameters->NameCharacterCount = Names.Num();
 		Parameters->Names = GraphBuilder.CreateSRV(NameBuffer, PF_R8_UINT);
 		Parameters->NameInfos = GraphBuilder.CreateSRV(NameInfoBuffer);
-		Parameters->Infos = GraphBuilder.CreateSRV(InfoBuffer, PF_R32G32B32A32_UINT);
+		Parameters->Infos = GraphBuilder.CreateSRV(InfoBuffer, PF_R32_UINT);
 		ShaderPrint::SetParameters(GraphBuilder, *ShaderPrintData, Parameters->ShaderPrintUniformBuffer);
 		FHairDebugPrintInstanceCS::FPermutationDomain PermutationVector;
 		PermutationVector.Set<FHairDebugPrintInstanceCS::FOutputType>(0);
diff --git a/Engine/Plugins/Runtime/HairStrands/Source/HairStrandsCore/Private/GroomManager.cpp b/Engine/Plugins/Runtime/HairStrands/Source/HairStrandsCore/Private/GroomManager.cpp
--- a/Engine/Plugins/Runtime/HairStrands/Source/HairStrandsCore/Private/GroomManager.cpp
+++ b/Engine/Plugins/Runtime/HairStrands/Source/HairStrandsCore/Private/GroomManager.cpp
@@ -311,9 +311,9 @@ static void RunInternalHairStrandsInterpolation(
 						// Add manual transition for the GPU solver as Niagara does not track properly the RDG buffer, and so doesn't issue the correct transitions
 						if (MeshLODIndex >= 0)
 						{
-							GraphBuilder.UseExternalAccessMode(Register(GraphBuilder, Instance->Guides.DeformedRootResource->LODs[MeshLODIndex].GetDeformedRootTrianglePosition0Buffer(FHairStrandsDeformedRootResource::FLOD::Current), ERDGImportedBufferFlags::CreateSRV).Buffer, ERHIAccess::SRVMask);
-							GraphBuilder.UseExternalAccessMode(Register(GraphBuilder, Instance->Guides.DeformedRootResource->LODs[MeshLODIndex].GetDeformedRootTrianglePosition1Buffer(FHairStrandsDeformedRootResource::FLOD::Current), ERDGImportedBufferFlags::CreateSRV).Buffer, ERHIAccess::SRVMask);
-							GraphBuilder.UseExternalAccessMode(Register(GraphBuilder, Instance->Guides.DeformedRootResource->LODs[MeshLODIndex].GetDeformedRootTrianglePosition2Buffer(FHairStrandsDeformedRootResource::FLOD::Current), ERDGImportedBufferFlags::CreateSRV).Buffer, ERHIAccess::SRVMask);
+							GraphBuilder.UseExternalAccessMode(Register(GraphBuilder, Instance->Guides.DeformedRootResource->LODs[MeshLODIndex].GetDeformedUniqueTrianglePosition0Buffer(FHairStrandsDeformedRootResource::FLOD::Current), ERDGImportedBufferFlags::CreateSRV).Buffer, ERHIAccess::SRVMask);
+							GraphBuilder.UseExternalAccessMode(Register(GraphBuilder, Instance->Guides.DeformedRootResource->LODs[MeshLODIndex].GetDeformedUniqueTrianglePosition1Buffer(FHairStrandsDeformedRootResource::FLOD::Current), ERDGImportedBufferFlags::CreateSRV).Buffer, ERHIAccess::SRVMask);
+							GraphBuilder.UseExternalAccessMode(Register(GraphBuilder, Instance->Guides.DeformedRootResource->LODs[MeshLODIndex].GetDeformedUniqueTrianglePosition2Buffer(FHairStrandsDeformedRootResource::FLOD::Current), ERDGImportedBufferFlags::CreateSRV).Buffer, ERHIAccess::SRVMask);
 
 							if (bGlobalDeformationEnable)
 							{
diff --git a/Engine/Plugins/Runtime/HairStrands/Source/HairStrandsCore/Private/GroomRBFDeformer.cpp b/Engine/Plugins/Runtime/HairStrands/Source/HairStrandsCore/Private/GroomRBFDeformer.cpp
--- a/Engine/Plugins/Runtime/HairStrands/Source/HairStrandsCore/Private/GroomRBFDeformer.cpp
+++ b/Engine/Plugins/Runtime/HairStrands/Source/HairStrandsCore/Private/GroomRBFDeformer.cpp
@@ -82,15 +82,16 @@ FVector3f DisplacePosition(
 void DeformStrands(	
 	const FHairStrandsDatas& HairStandsData,
 	const TArray<FHairStrandsIndexFormat::Type>& VertexToCurveIndexBuffer,
-	const TArray<FHairStrandsCurveTriangleBarycentricFormat::Type>& RootTriangleBarycentricBuffer,
+	const TArray<FHairStrandsIndexFormat::Type>& RootToUniqueTriangleBuffer,
+	const TArray<FHairStrandsCurveTriangleBarycentricFormat::Type>& RootBarycentricBuffer,
 
-	const TArray<FHairStrandsMeshTrianglePositionFormat::Type>& RootTrianglePosition0Buffer_Rest,
-	const TArray<FHairStrandsMeshTrianglePositionFormat::Type>& RootTrianglePosition1Buffer_Rest,
-	const TArray<FHairStrandsMeshTrianglePositionFormat::Type>& RootTrianglePosition2Buffer_Rest,
+	const TArray<FHairStrandsMeshTrianglePositionFormat::Type>& UniqueTrianglePosition0Buffer_Rest,
+	const TArray<FHairStrandsMeshTrianglePositionFormat::Type>& UniqueTrianglePosition1Buffer_Rest,
+	const TArray<FHairStrandsMeshTrianglePositionFormat::Type>& UniqueTrianglePosition2Buffer_Rest,
 
-	const TArray<FHairStrandsMeshTrianglePositionFormat::Type>& RootTrianglePosition0Buffer_Deformed,
-	const TArray<FHairStrandsMeshTrianglePositionFormat::Type>& RootTrianglePosition1Buffer_Deformed,
-	const TArray<FHairStrandsMeshTrianglePositionFormat::Type>& RootTrianglePosition2Buffer_Deformed,
+	const TArray<FHairStrandsMeshTrianglePositionFormat::Type>& UniqueTrianglePosition0Buffer_Deformed,
+	const TArray<FHairStrandsMeshTrianglePositionFormat::Type>& UniqueTrianglePosition1Buffer_Deformed,
+	const TArray<FHairStrandsMeshTrianglePositionFormat::Type>& UniqueTrianglePosition2Buffer_Deformed,
 
 	uint32 VertexCount,
 	uint32 SampleCount,
@@ -117,6 +118,7 @@ void DeformStrands(
 	{
 		const uint32 VertexOffset = HairStandsData.StrandsCurves.CurvesOffset[CurveIndex];
 		const uint32 RootIndex = VertexToCurveIndexBuffer[VertexOffset];
+		const uint32 TriangleIndex = RootToUniqueTriangleBuffer[RootIndex];
 
 		// Sanity check
 		check(RootIndex == CurveIndex);
@@ -124,19 +126,18 @@ void DeformStrands(
 		const FVector3f& Rest_Position   = HairStandsData.StrandsPoints.PointsPosition[VertexOffset];
 		const FVector3f& Deform_Position = OutDeformedPositionBuffer[VertexOffset];
 
-
-		const uint32 PackedBarycentric = RootTriangleBarycentricBuffer[RootIndex];
+		const uint32 PackedBarycentric = RootBarycentricBuffer[RootIndex];
 		const FVector2f B0 = FVector2f(FHairStrandsRootUtils::DecodeBarycentrics(PackedBarycentric));
 		const FVector3f   B  = FVector3f(B0.X, B0.Y, 1.f - B0.X - B0.Y);
 
 		/* Strand hair roots translation and rotation in rest position relative to the bound triangle. Positions are relative to the rest root center */
-		const FVector3f& Rest_V0 = RootTrianglePosition0Buffer_Rest[RootIndex];
-		const FVector3f& Rest_V1 = RootTrianglePosition1Buffer_Rest[RootIndex];
-		const FVector3f& Rest_V2 = RootTrianglePosition2Buffer_Rest[RootIndex];
+		const FVector3f& Rest_V0 = UniqueTrianglePosition0Buffer_Rest[TriangleIndex];
+		const FVector3f& Rest_V1 = UniqueTrianglePosition1Buffer_Rest[TriangleIndex];
+		const FVector3f& Rest_V2 = UniqueTrianglePosition2Buffer_Rest[TriangleIndex];
 
-		const FVector3f& Deform_V0 = RootTrianglePosition0Buffer_Deformed[RootIndex];
-		const FVector3f& Deform_V1 = RootTrianglePosition1Buffer_Deformed[RootIndex];
-		const FVector3f& Deform_V2 = RootTrianglePosition2Buffer_Deformed[RootIndex];
+		const FVector3f& Deform_V0 = UniqueTrianglePosition0Buffer_Deformed[TriangleIndex];
+		const FVector3f& Deform_V1 = UniqueTrianglePosition1Buffer_Deformed[TriangleIndex];
+		const FVector3f& Deform_V2 = UniqueTrianglePosition2Buffer_Deformed[TriangleIndex];
 
 		const FVector3f Rest_RootPosition		=   Rest_V0 * B.X +   Rest_V1 * B.Y +   Rest_V2 * B.Z;
 		const FVector3f Deform_RootPosition	= Deform_V0 * B.X + Deform_V1 * B.Y + Deform_V2 * B.Z;
@@ -158,27 +159,26 @@ void DeformStrands(
 }
 
 // Compute the triangle positions for each curve's roots
-void ExtractRootTrianglePositions(
+void ExtractUniqueTrianglePositions(
 	const FHairStrandsRootData::FMeshProjectionLOD& RestLODData,
 	const uint32 MeshLODIndex,
-	const FHairStrandsDatas& HairStrandsData,
+	const uint32 UniqueTriangleCount,
 	const FSkeletalMeshRenderData* InMeshRenderData, 
-	TArray<FHairStrandsMeshTrianglePositionFormat::Type>& OutDeformRootTrianglePosition0Buffer,
-	TArray<FHairStrandsMeshTrianglePositionFormat::Type>& OutDeformRootTrianglePosition1Buffer,
-	TArray<FHairStrandsMeshTrianglePositionFormat::Type>& OutDeformRootTrianglePosition2Buffer)
+	TArray<FHairStrandsMeshTrianglePositionFormat::Type>& OutDeformUniqueTrianglePosition0Buffer,
+	TArray<FHairStrandsMeshTrianglePositionFormat::Type>& OutDeformUniqueTrianglePosition1Buffer,
+	TArray<FHairStrandsMeshTrianglePositionFormat::Type>& OutDeformUniqueTrianglePosition2Buffer)
 {
-	const uint32 RootCount = HairStrandsData.GetNumCurves();
-	OutDeformRootTrianglePosition0Buffer.SetNum(RootCount);
-	OutDeformRootTrianglePosition1Buffer.SetNum(RootCount);
-	OutDeformRootTrianglePosition2Buffer.SetNum(RootCount);
+	OutDeformUniqueTrianglePosition0Buffer.SetNum(UniqueTriangleCount);
+	OutDeformUniqueTrianglePosition1Buffer.SetNum(UniqueTriangleCount);
+	OutDeformUniqueTrianglePosition2Buffer.SetNum(UniqueTriangleCount);
 
 	const uint32 SectionCount = InMeshRenderData->LODRenderData[MeshLODIndex].RenderSections.Num();
 	TArray<uint32> IndexBuffer;
 	InMeshRenderData->LODRenderData[MeshLODIndex].MultiSizeIndexContainer.GetIndexBuffer(IndexBuffer);
 
-	for (uint32 RootIndex = 0; RootIndex < RootCount; ++RootIndex)
+	for (uint32 UniqueTriangleIndex = 0; UniqueTriangleIndex < UniqueTriangleCount; ++UniqueTriangleIndex)
 	{
-		const uint32 PackedTriangleIndex = RestLODData.RootTriangleIndexBuffer[RootIndex];
+		const uint32 PackedTriangleIndex = RestLODData.UniqueTriangleIndexBuffer[UniqueTriangleIndex];
 		uint32 TriangleIndex = 0;
 		uint32 SectionIndex = 0;
 		FHairStrandsRootUtils::DecodeTriangleIndex(PackedTriangleIndex, TriangleIndex, SectionIndex);
@@ -195,9 +195,9 @@ void ExtractRootTrianglePositions(
 		const FVector3f P1 = InMeshRenderData->LODRenderData[MeshLODIndex].StaticVertexBuffers.PositionVertexBuffer.VertexPosition(I1);
 		const FVector3f P2 = InMeshRenderData->LODRenderData[MeshLODIndex].StaticVertexBuffers.PositionVertexBuffer.VertexPosition(I2);
 
-		OutDeformRootTrianglePosition0Buffer[RootIndex] = P0;
-		OutDeformRootTrianglePosition1Buffer[RootIndex] = P1;
-		OutDeformRootTrianglePosition2Buffer[RootIndex] = P2;
+		OutDeformUniqueTrianglePosition0Buffer[UniqueTriangleIndex] = P0;
+		OutDeformUniqueTrianglePosition1Buffer[UniqueTriangleIndex] = P1;
+		OutDeformUniqueTrianglePosition2Buffer[UniqueTriangleIndex] = P2;
 	}
 }
 
@@ -217,6 +217,8 @@ TArray<FVector3f> GetDeformedHairStrandsPositions(
 
 	InitMeshSamples(MaxVertexCount, MeshVertexPositionsBuffer_Target, MaxSampleCount, SampleIndicesBuffer, OutSamplePositionsBuffer);
 
+	const TArray<FHairStrandsIndexFormat::Type>& RootToUniqueTriangleBuffer = RestLODData.RootToUniqueTriangleIndexBuffer;
+
 	// Update those vertices with the RBF interpolation weights
 	const TArray<float>& InterpolationWeightsBuffer = RestLODData.MeshInterpolationWeightsBuffer;
 	const TArray<FVector4f>& SampleRestPositionsBuffer = RestLODData.RestSamplePositionsBuffer;
@@ -232,21 +234,21 @@ TArray<FVector3f> GetDeformedHairStrandsPositions(
 	TArray<FVector3f> OutPositions = Points.PointsPosition;
 
 	// Use the vertex position of the binding, as the source asset might not have the same topology (in case the groom has been transfered from one mesh toanother using UV sharing)
-	TArray<FHairStrandsMeshTrianglePositionFormat::Type> RootTrianglePosition0Buffer_Rest = RestLODData.RestRootTrianglePosition0Buffer;
-	TArray<FHairStrandsMeshTrianglePositionFormat::Type> RootTrianglePosition1Buffer_Rest = RestLODData.RestRootTrianglePosition1Buffer;
-	TArray<FHairStrandsMeshTrianglePositionFormat::Type> RootTrianglePosition2Buffer_Rest = RestLODData.RestRootTrianglePosition2Buffer;
-
-	TArray<FHairStrandsMeshTrianglePositionFormat::Type> RootTrianglePosition0Buffer_Deformed;
-	TArray<FHairStrandsMeshTrianglePositionFormat::Type> RootTrianglePosition1Buffer_Deformed;
-	TArray<FHairStrandsMeshTrianglePositionFormat::Type> RootTrianglePosition2Buffer_Deformed;
-	ExtractRootTrianglePositions(
+	TArray<FHairStrandsMeshTrianglePositionFormat::Type> UniqueTrianglePosition0Buffer_Rest = RestLODData.RestUniqueTrianglePosition0Buffer;
+	TArray<FHairStrandsMeshTrianglePositionFormat::Type> UniqueTrianglePosition1Buffer_Rest = RestLODData.RestUniqueTrianglePosition1Buffer;
+	TArray<FHairStrandsMeshTrianglePositionFormat::Type> UniqueTrianglePosition2Buffer_Rest = RestLODData.RestUniqueTrianglePosition2Buffer;
+
+	TArray<FHairStrandsMeshTrianglePositionFormat::Type> UniqueTrianglePosition0Buffer_Deformed;
+	TArray<FHairStrandsMeshTrianglePositionFormat::Type> UniqueTrianglePosition1Buffer_Deformed;
+	TArray<FHairStrandsMeshTrianglePositionFormat::Type> UniqueTrianglePosition2Buffer_Deformed;
+	ExtractUniqueTrianglePositions(
 		RestLODData,
 		MeshLODIndex,
-		HairStrandsData,
+		RestLODData.RestUniqueTrianglePosition0Buffer.Num(),
 		InMeshRenderData,
-		RootTrianglePosition0Buffer_Deformed,
-		RootTrianglePosition1Buffer_Deformed,
-		RootTrianglePosition2Buffer_Deformed);
+		UniqueTrianglePosition0Buffer_Deformed,
+		UniqueTrianglePosition1Buffer_Deformed,
+		UniqueTrianglePosition2Buffer_Deformed);
 
 	// Deform the strands vertices with the deformed mesh samples
 	const TArray<FVector3f>& RestPosePositionBuffer = OutPositions;
@@ -257,15 +259,16 @@ TArray<FVector3f> GetDeformedHairStrandsPositions(
 	DeformStrands(
 		HairStrandsData,
 		VertexToCurveIndexBuffer,
-		RestLODData.RootTriangleBarycentricBuffer,
+		RootToUniqueTriangleBuffer,
+		RestLODData.RootBarycentricBuffer,
 
-		RootTrianglePosition0Buffer_Rest,
-		RootTrianglePosition1Buffer_Rest,
-		RootTrianglePosition2Buffer_Rest,
+		UniqueTrianglePosition0Buffer_Rest,
+		UniqueTrianglePosition1Buffer_Rest,
+		UniqueTrianglePosition2Buffer_Rest,
 
-		RootTrianglePosition0Buffer_Deformed,
-		RootTrianglePosition1Buffer_Deformed,
-		RootTrianglePosition2Buffer_Deformed,
+		UniqueTrianglePosition0Buffer_Deformed,
+		UniqueTrianglePosition1Buffer_Deformed,
+		UniqueTrianglePosition2Buffer_Deformed,
 
 		VertexCount, 
 		MaxSampleCount, 
diff --git a/Engine/Plugins/Runtime/HairStrands/Source/HairStrandsCore/Private/GroomResources.cpp b/Engine/Plugins/Runtime/HairStrands/Source/HairStrandsCore/Private/GroomResources.cpp
--- a/Engine/Plugins/Runtime/HairStrands/Source/HairStrandsCore/Private/GroomResources.cpp
+++ b/Engine/Plugins/Runtime/HairStrands/Source/HairStrandsCore/Private/GroomResources.cpp
@@ -1089,15 +1089,16 @@ bool FHairStrandsRestRootResource::InternalIsLODDataLoaded(int32 LODIndex)
 	if (LODIndex >= 0 && LODIndex < LODs.Num())
 	{
 		FHairStrandsRootBulkData::FMeshProjectionLOD& CPUData = BulkData.MeshProjectionLODs[LODIndex];
-		const bool bHasValidCPUData = CPUData.RootTriangleBarycentricBuffer.GetBulkDataSize() > 0;
+		const bool bHasValidCPUData = CPUData.RootBarycentricBuffer.GetBulkDataSize() > 0;
 		if (bHasValidCPUData)
 		{
-			AsyncLoadHairBulkData(CPUData.RootTriangleBarycentricBuffer, bIsLoading);
-			AsyncLoadHairBulkData(CPUData.RootTriangleIndexBuffer, bIsLoading);
+			AsyncLoadHairBulkData(CPUData.RootBarycentricBuffer, bIsLoading);
+			AsyncLoadHairBulkData(CPUData.RootToUniqueTriangleIndexBuffer, bIsLoading);
+			AsyncLoadHairBulkData(CPUData.UniqueTriangleIndexBuffer, bIsLoading);
 
-			AsyncLoadHairBulkData(CPUData.RestRootTrianglePosition0Buffer, bIsLoading);
-			AsyncLoadHairBulkData(CPUData.RestRootTrianglePosition1Buffer, bIsLoading);
-			AsyncLoadHairBulkData(CPUData.RestRootTrianglePosition2Buffer, bIsLoading);
+			AsyncLoadHairBulkData(CPUData.RestUniqueTrianglePosition0Buffer, bIsLoading);
+			AsyncLoadHairBulkData(CPUData.RestUniqueTrianglePosition1Buffer, bIsLoading);
+			AsyncLoadHairBulkData(CPUData.RestUniqueTrianglePosition2Buffer, bIsLoading);
 		}
 
 		const bool bHasValidCPUWeights = CPUData.MeshSampleIndicesBuffer.GetBulkDataSize() > 0;
@@ -1143,29 +1144,31 @@ void FHairStrandsRestRootResource::InternalAllocateLOD(FRDGBuilder& GraphBuilder
 		}
 
 		FHairStrandsRootBulkData::FMeshProjectionLOD& CPUData = BulkData.MeshProjectionLODs[LODIndex];
-		const bool bHasValidCPUData = CPUData.RootTriangleBarycentricBuffer.GetBulkDataSize() > 0;
+		const bool bHasValidCPUData = CPUData.RootBarycentricBuffer.GetBulkDataSize() > 0;
 		if (bHasValidCPUData)
 		{
 			GPUData.Status = FLOD::EStatus::Completed;
+			
+			InternalCreateVertexBufferRDG_FromBulkData<FHairStrandsCurveTriangleBarycentricFormat>(GraphBuilder, CPUData.RootBarycentricBuffer, BulkData.RootCount, GPUData.RootBarycentricBuffer, ToHairResourceDebugName(HAIRSTRANDS_RESOUCE_NAME(CurveType, Hair.StrandsRestRoot_RootTriangleBarycentricBuffer), ResourceName), EHairResourceUsageType::Static);
+			InternalCreateVertexBufferRDG_FromBulkData<FHairStrandsCurveTriangleIndexFormat>(GraphBuilder, CPUData.RootToUniqueTriangleIndexBuffer, BulkData.RootCount, GPUData.RootToUniqueTriangleIndexBuffer, ToHairResourceDebugName(HAIRSTRANDS_RESOUCE_NAME(CurveType, Hair.StrandsRestRoot_RootToUniqueTriangleIndexBuffer), ResourceName), EHairResourceUsageType::Static);
+			InternalCreateVertexBufferRDG_FromBulkData<FHairStrandsCurveTriangleIndexFormat>(GraphBuilder, CPUData.UniqueTriangleIndexBuffer, CPUData.UniqueTriangleCount, GPUData.UniqueTriangleIndexBuffer, ToHairResourceDebugName(HAIRSTRANDS_RESOUCE_NAME(CurveType, Hair.StrandsRestRoot_UniqueTriangleIndexBuffer), ResourceName), EHairResourceUsageType::Static);
 
-			InternalCreateVertexBufferRDG_FromBulkData<FHairStrandsCurveTriangleBarycentricFormat>(GraphBuilder, CPUData.RootTriangleBarycentricBuffer, BulkData.RootCount, GPUData.RootTriangleBarycentricBuffer, ToHairResourceDebugName(HAIRSTRANDS_RESOUCE_NAME(CurveType, Hair.StrandsRestRoot_RootTriangleBarycentricBuffer), ResourceName), EHairResourceUsageType::Static);
-			InternalCreateVertexBufferRDG_FromBulkData<FHairStrandsCurveTriangleIndexFormat>(GraphBuilder, CPUData.RootTriangleIndexBuffer, BulkData.RootCount, GPUData.RootTriangleIndexBuffer, ToHairResourceDebugName(HAIRSTRANDS_RESOUCE_NAME(CurveType, Hair.StrandsRestRoot_RootTriangleIndexBuffer), ResourceName), EHairResourceUsageType::Static);
-
-			InternalCreateVertexBufferRDG_FromBulkData<FHairStrandsMeshTrianglePositionFormat>(GraphBuilder, CPUData.RestRootTrianglePosition0Buffer, BulkData.RootCount, GPUData.RestRootTrianglePosition0Buffer, ToHairResourceDebugName(HAIRSTRANDS_RESOUCE_NAME(CurveType, Hair.StrandsRestRoot_RestRootTrianglePosition0Buffer), ResourceName), EHairResourceUsageType::Static);
-			InternalCreateVertexBufferRDG_FromBulkData<FHairStrandsMeshTrianglePositionFormat>(GraphBuilder, CPUData.RestRootTrianglePosition1Buffer, BulkData.RootCount, GPUData.RestRootTrianglePosition1Buffer, ToHairResourceDebugName(HAIRSTRANDS_RESOUCE_NAME(CurveType, Hair.StrandsRestRoot_RestRootTrianglePosition1Buffer), ResourceName), EHairResourceUsageType::Static);
-			InternalCreateVertexBufferRDG_FromBulkData<FHairStrandsMeshTrianglePositionFormat>(GraphBuilder, CPUData.RestRootTrianglePosition2Buffer, BulkData.RootCount, GPUData.RestRootTrianglePosition2Buffer, ToHairResourceDebugName(HAIRSTRANDS_RESOUCE_NAME(CurveType, Hair.StrandsRestRoot_RestRootTrianglePosition2Buffer), ResourceName), EHairResourceUsageType::Static);
+			InternalCreateVertexBufferRDG_FromBulkData<FHairStrandsMeshTrianglePositionFormat>(GraphBuilder, CPUData.RestUniqueTrianglePosition0Buffer, CPUData.UniqueTriangleCount, GPUData.RestUniqueTrianglePosition0Buffer, ToHairResourceDebugName(HAIRSTRANDS_RESOUCE_NAME(CurveType, Hair.StrandsRestRoot_RestUniqueTrianglePosition0Buffer), ResourceName), EHairResourceUsageType::Static);
+			InternalCreateVertexBufferRDG_FromBulkData<FHairStrandsMeshTrianglePositionFormat>(GraphBuilder, CPUData.RestUniqueTrianglePosition1Buffer, CPUData.UniqueTriangleCount, GPUData.RestUniqueTrianglePosition1Buffer, ToHairResourceDebugName(HAIRSTRANDS_RESOUCE_NAME(CurveType, Hair.StrandsRestRoot_RestUniqueTrianglePosition1Buffer), ResourceName), EHairResourceUsageType::Static);
+			InternalCreateVertexBufferRDG_FromBulkData<FHairStrandsMeshTrianglePositionFormat>(GraphBuilder, CPUData.RestUniqueTrianglePosition2Buffer, CPUData.UniqueTriangleCount, GPUData.RestUniqueTrianglePosition2Buffer, ToHairResourceDebugName(HAIRSTRANDS_RESOUCE_NAME(CurveType, Hair.StrandsRestRoot_RestUniqueTrianglePosition2Buffer), ResourceName), EHairResourceUsageType::Static);
 		}
 		else
 		{
 			GPUData.Status = FLOD::EStatus::Initialized;
 			
-			InternalCreateVertexBufferRDG<FHairStrandsCurveTriangleBarycentricFormat>(GraphBuilder, BulkData.RootCount, GPUData.RootTriangleBarycentricBuffer, ToHairResourceDebugName(HAIRSTRANDS_RESOUCE_NAME(CurveType, Hair.StrandsRestRoot_RootTriangleBarycentricBuffer), ResourceName), EHairResourceUsageType::Dynamic);
-			InternalCreateVertexBufferRDG<FHairStrandsCurveTriangleIndexFormat>(GraphBuilder, BulkData.RootCount, GPUData.RootTriangleIndexBuffer, ToHairResourceDebugName(HAIRSTRANDS_RESOUCE_NAME(CurveType, Hair.StrandsRestRoot_RootTriangleIndexBuffer), ResourceName), EHairResourceUsageType::Dynamic);
+			InternalCreateVertexBufferRDG<FHairStrandsCurveTriangleBarycentricFormat>(GraphBuilder, BulkData.RootCount, GPUData.RootBarycentricBuffer, ToHairResourceDebugName(HAIRSTRANDS_RESOUCE_NAME(CurveType, Hair.StrandsRestRoot_RootBarycentricBuffer), ResourceName), EHairResourceUsageType::Dynamic);
+			InternalCreateVertexBufferRDG<FHairStrandsCurveTriangleIndexFormat>(GraphBuilder, BulkData.RootCount, GPUData.RootToUniqueTriangleIndexBuffer, ToHairResourceDebugName(HAIRSTRANDS_RESOUCE_NAME(CurveType, Hair.StrandsRestRoot_RootToUniqueTriangleIndexBuffer), ResourceName), EHairResourceUsageType::Dynamic);
+			InternalCreateVertexBufferRDG<FHairStrandsCurveTriangleIndexFormat>(GraphBuilder, CPUData.UniqueTriangleCount, GPUData.UniqueTriangleIndexBuffer, ToHairResourceDebugName(HAIRSTRANDS_RESOUCE_NAME(CurveType, Hair.StrandsRestRoot_UniqueTriangleIndexBuffer), ResourceName), EHairResourceUsageType::Dynamic);
 			
 			// Create buffers. Initialization will be done by render passes
-			InternalCreateVertexBufferRDG<FHairStrandsMeshTrianglePositionFormat>(GraphBuilder, BulkData.RootCount, GPUData.RestRootTrianglePosition0Buffer, ToHairResourceDebugName(HAIRSTRANDS_RESOUCE_NAME(CurveType, Hair.StrandsRestRoot_RestRootTrianglePosition0Buffer), ResourceName), EHairResourceUsageType::Dynamic);
-			InternalCreateVertexBufferRDG<FHairStrandsMeshTrianglePositionFormat>(GraphBuilder, BulkData.RootCount, GPUData.RestRootTrianglePosition1Buffer, ToHairResourceDebugName(HAIRSTRANDS_RESOUCE_NAME(CurveType, Hair.StrandsRestRoot_RestRootTrianglePosition1Buffer), ResourceName), EHairResourceUsageType::Dynamic);
-			InternalCreateVertexBufferRDG<FHairStrandsMeshTrianglePositionFormat>(GraphBuilder, BulkData.RootCount, GPUData.RestRootTrianglePosition2Buffer, ToHairResourceDebugName(HAIRSTRANDS_RESOUCE_NAME(CurveType, Hair.StrandsRestRoot_RestRootTrianglePosition2Buffer), ResourceName), EHairResourceUsageType::Dynamic);
+			InternalCreateVertexBufferRDG<FHairStrandsMeshTrianglePositionFormat>(GraphBuilder, CPUData.UniqueTriangleCount, GPUData.RestUniqueTrianglePosition0Buffer, ToHairResourceDebugName(HAIRSTRANDS_RESOUCE_NAME(CurveType, Hair.StrandsRestRoot_RestUniqueTrianglePosition0Buffer), ResourceName), EHairResourceUsageType::Dynamic);
+			InternalCreateVertexBufferRDG<FHairStrandsMeshTrianglePositionFormat>(GraphBuilder, CPUData.UniqueTriangleCount, GPUData.RestUniqueTrianglePosition1Buffer, ToHairResourceDebugName(HAIRSTRANDS_RESOUCE_NAME(CurveType, Hair.StrandsRestRoot_RestUniqueTrianglePosition1Buffer), ResourceName), EHairResourceUsageType::Dynamic);
+			InternalCreateVertexBufferRDG<FHairStrandsMeshTrianglePositionFormat>(GraphBuilder, CPUData.UniqueTriangleCount, GPUData.RestUniqueTrianglePosition2Buffer, ToHairResourceDebugName(HAIRSTRANDS_RESOUCE_NAME(CurveType, Hair.StrandsRestRoot_RestUniqueTrianglePosition2Buffer), ResourceName), EHairResourceUsageType::Dynamic);
 		}
 
 		GPUData.SampleCount = CPUData.SampleCount;
@@ -1194,11 +1197,12 @@ void FHairStrandsRestRootResource::InternalRelease()
 	for (FLOD& GPUData : LODs)
 	{
 		GPUData.Status = FLOD::EStatus::Invalid;
-		GPUData.RootTriangleIndexBuffer.Release();
-		GPUData.RootTriangleBarycentricBuffer.Release();
-		GPUData.RestRootTrianglePosition0Buffer.Release();
-		GPUData.RestRootTrianglePosition1Buffer.Release();
-		GPUData.RestRootTrianglePosition2Buffer.Release();
+		GPUData.RootBarycentricBuffer.Release();
+		GPUData.RootToUniqueTriangleIndexBuffer.Release();
+		GPUData.UniqueTriangleIndexBuffer.Release();
+		GPUData.RestUniqueTrianglePosition0Buffer.Release();
+		GPUData.RestUniqueTrianglePosition1Buffer.Release();
+		GPUData.RestUniqueTrianglePosition2Buffer.Release();
 		GPUData.SampleCount = 0;
 		GPUData.MeshInterpolationWeightsBuffer.Release();
 		GPUData.MeshSampleIndicesBuffer.Release();
@@ -1255,16 +1259,16 @@ void FHairStrandsDeformedRootResource::InternalAllocateLOD(FRDGBuilder& GraphBui
 				}
 			}
 
-			InternalCreateVertexBufferRDG<FHairStrandsMeshTrianglePositionFormat>(GraphBuilder, RootCount, LOD.DeformedRootTrianglePosition0Buffer[0], ToHairResourceDebugName(HAIRSTRANDS_RESOUCE_NAME(CurveType, Hair.StrandsRootDeformed_DeformedRootTrianglePosition0Buffer0), ResourceName), EHairResourceUsageType::Dynamic);
-			InternalCreateVertexBufferRDG<FHairStrandsMeshTrianglePositionFormat>(GraphBuilder, RootCount, LOD.DeformedRootTrianglePosition1Buffer[0], ToHairResourceDebugName(HAIRSTRANDS_RESOUCE_NAME(CurveType, Hair.StrandsRootDeformed_DeformedRootTrianglePosition1Buffer0), ResourceName), EHairResourceUsageType::Dynamic);
-			InternalCreateVertexBufferRDG<FHairStrandsMeshTrianglePositionFormat>(GraphBuilder, RootCount, LOD.DeformedRootTrianglePosition2Buffer[0], ToHairResourceDebugName(HAIRSTRANDS_RESOUCE_NAME(CurveType, Hair.StrandsRootDeformed_DeformedRootTrianglePosition2Buffer0), ResourceName), EHairResourceUsageType::Dynamic);
+			InternalCreateVertexBufferRDG<FHairStrandsMeshTrianglePositionFormat>(GraphBuilder, RootCount, LOD.DeformedUniqueTrianglePosition0Buffer[0], ToHairResourceDebugName(HAIRSTRANDS_RESOUCE_NAME(CurveType, Hair.StrandsRootDeformed_DeformedUniqueTrianglePosition0Buffer0), ResourceName), EHairResourceUsageType::Dynamic);
+			InternalCreateVertexBufferRDG<FHairStrandsMeshTrianglePositionFormat>(GraphBuilder, RootCount, LOD.DeformedUniqueTrianglePosition1Buffer[0], ToHairResourceDebugName(HAIRSTRANDS_RESOUCE_NAME(CurveType, Hair.StrandsRootDeformed_DeformedUniqueTrianglePosition1Buffer0), ResourceName), EHairResourceUsageType::Dynamic);
+			InternalCreateVertexBufferRDG<FHairStrandsMeshTrianglePositionFormat>(GraphBuilder, RootCount, LOD.DeformedUniqueTrianglePosition2Buffer[0], ToHairResourceDebugName(HAIRSTRANDS_RESOUCE_NAME(CurveType, Hair.StrandsRootDeformed_DeformedUniqueTrianglePosition2Buffer0), ResourceName), EHairResourceUsageType::Dynamic);
 
 			// Double buffering is disabled by default unless the read-only cvar r.HairStrands.ContinuousDecimationReordering is set
 			if (IsHairStrandContinuousDecimationReorderingEnabled())
 			{
-				InternalCreateVertexBufferRDG<FHairStrandsMeshTrianglePositionFormat>(GraphBuilder, RootCount, LOD.DeformedRootTrianglePosition0Buffer[1], ToHairResourceDebugName(HAIRSTRANDS_RESOUCE_NAME(CurveType, Hair.StrandsRootDeformed_DeformedRootTrianglePosition0Buffer1), ResourceName), EHairResourceUsageType::Dynamic);
-				InternalCreateVertexBufferRDG<FHairStrandsMeshTrianglePositionFormat>(GraphBuilder, RootCount, LOD.DeformedRootTrianglePosition1Buffer[1], ToHairResourceDebugName(HAIRSTRANDS_RESOUCE_NAME(CurveType, Hair.StrandsRootDeformed_DeformedRootTrianglePosition1Buffer1), ResourceName), EHairResourceUsageType::Dynamic);
-				InternalCreateVertexBufferRDG<FHairStrandsMeshTrianglePositionFormat>(GraphBuilder, RootCount, LOD.DeformedRootTrianglePosition2Buffer[1], ToHairResourceDebugName(HAIRSTRANDS_RESOUCE_NAME(CurveType, Hair.StrandsRootDeformed_DeformedRootTrianglePosition2Buffer1), ResourceName), EHairResourceUsageType::Dynamic);
+				InternalCreateVertexBufferRDG<FHairStrandsMeshTrianglePositionFormat>(GraphBuilder, RootCount, LOD.DeformedUniqueTrianglePosition0Buffer[1], ToHairResourceDebugName(HAIRSTRANDS_RESOUCE_NAME(CurveType, Hair.StrandsRootDeformed_DeformedUniqueTrianglePosition0Buffer1), ResourceName), EHairResourceUsageType::Dynamic);
+				InternalCreateVertexBufferRDG<FHairStrandsMeshTrianglePositionFormat>(GraphBuilder, RootCount, LOD.DeformedUniqueTrianglePosition1Buffer[1], ToHairResourceDebugName(HAIRSTRANDS_RESOUCE_NAME(CurveType, Hair.StrandsRootDeformed_DeformedUniqueTrianglePosition1Buffer1), ResourceName), EHairResourceUsageType::Dynamic);
+				InternalCreateVertexBufferRDG<FHairStrandsMeshTrianglePositionFormat>(GraphBuilder, RootCount, LOD.DeformedUniqueTrianglePosition2Buffer[1], ToHairResourceDebugName(HAIRSTRANDS_RESOUCE_NAME(CurveType, Hair.StrandsRootDeformed_DeformedUniqueTrianglePosition2Buffer1), ResourceName), EHairResourceUsageType::Dynamic);
 			}
 		}
 	}
@@ -1275,18 +1279,18 @@ void FHairStrandsDeformedRootResource::InternalRelease()
 	for (FLOD& GPUData : LODs)
 	{
 		GPUData.Status = FLOD::EStatus::Invalid;
-		GPUData.DeformedRootTrianglePosition0Buffer[0].Release();
-		GPUData.DeformedRootTrianglePosition1Buffer[0].Release();
-		GPUData.DeformedRootTrianglePosition2Buffer[0].Release();
+		GPUData.DeformedUniqueTrianglePosition0Buffer[0].Release();
+		GPUData.DeformedUniqueTrianglePosition1Buffer[0].Release();
+		GPUData.DeformedUniqueTrianglePosition2Buffer[0].Release();
 		GPUData.DeformedSamplePositionsBuffer[0].Release();
 		GPUData.MeshSampleWeightsBuffer[0].Release();
 	
 		// Double buffering is disabled by default unless the read-only cvar r.HairStrands.ContinuousDecimationReordering is set
 		if (IsHairStrandContinuousDecimationReorderingEnabled())
 		{
-			GPUData.DeformedRootTrianglePosition0Buffer[1].Release();
-			GPUData.DeformedRootTrianglePosition1Buffer[1].Release();
-			GPUData.DeformedRootTrianglePosition2Buffer[1].Release();
+			GPUData.DeformedUniqueTrianglePosition0Buffer[1].Release();
+			GPUData.DeformedUniqueTrianglePosition1Buffer[1].Release();
+			GPUData.DeformedUniqueTrianglePosition2Buffer[1].Release();
 			GPUData.DeformedSamplePositionsBuffer[1].Release();
 			GPUData.MeshSampleWeightsBuffer[1].Release();
 		}
@@ -1316,7 +1320,7 @@ bool FHairStrandsRootBulkData::HasProjectionData() const
 const TArray<uint32>& FHairStrandsRootBulkData::GetValidSectionIndices(int32 LODIndex) const
 {
 	check(LODIndex >= 0 && LODIndex < MeshProjectionLODs.Num());
-	return MeshProjectionLODs[LODIndex].ValidSectionIndices;
+	return MeshProjectionLODs[LODIndex].UniqueSectionIndices;
 }
 
 static void InternalSerialize(FArchive& Ar, UObject* Owner, FHairStrandsRootBulkData::FMeshProjectionLOD& LOD)
@@ -1327,11 +1331,12 @@ static void InternalSerialize(FArchive& Ar, UObject* Owner, FHairStrandsRootBulk
 	if (Ar.IsSaving())
 	{
 		const uint32 BulkFlags = BULKDATA_Force_NOT_InlinePayload;
-		LOD.RootTriangleIndexBuffer.SetBulkDataFlags(BulkFlags);
-		LOD.RootTriangleBarycentricBuffer.SetBulkDataFlags(BulkFlags);
-		LOD.RestRootTrianglePosition0Buffer.SetBulkDataFlags(BulkFlags);
-		LOD.RestRootTrianglePosition1Buffer.SetBulkDataFlags(BulkFlags);
-		LOD.RestRootTrianglePosition2Buffer.SetBulkDataFlags(BulkFlags);
+		LOD.RootBarycentricBuffer.SetBulkDataFlags(BulkFlags);
+		LOD.RootToUniqueTriangleIndexBuffer.SetBulkDataFlags(BulkFlags);
+		LOD.UniqueTriangleIndexBuffer.SetBulkDataFlags(BulkFlags);
+		LOD.RestUniqueTrianglePosition0Buffer.SetBulkDataFlags(BulkFlags);
+		LOD.RestUniqueTrianglePosition1Buffer.SetBulkDataFlags(BulkFlags);
+		LOD.RestUniqueTrianglePosition2Buffer.SetBulkDataFlags(BulkFlags);
 
 		LOD.MeshInterpolationWeightsBuffer.SetBulkDataFlags(BulkFlags);
 		LOD.MeshSampleIndicesBuffer.SetBulkDataFlags(BulkFlags);
@@ -1339,11 +1344,13 @@ static void InternalSerialize(FArchive& Ar, UObject* Owner, FHairStrandsRootBulk
 	}
 
 	Ar << LOD.LODIndex;
-	LOD.RootTriangleIndexBuffer.Serialize(Ar, Owner, ChunkIndex, bAttemptFileMapping);
-	LOD.RootTriangleBarycentricBuffer.Serialize(Ar, Owner, ChunkIndex, bAttemptFileMapping);
-	LOD.RestRootTrianglePosition0Buffer.Serialize(Ar, Owner, ChunkIndex, bAttemptFileMapping);
-	LOD.RestRootTrianglePosition1Buffer.Serialize(Ar, Owner, ChunkIndex, bAttemptFileMapping);
-	LOD.RestRootTrianglePosition2Buffer.Serialize(Ar, Owner, ChunkIndex, bAttemptFileMapping);
+	Ar << LOD.UniqueTriangleCount;
+	LOD.RootBarycentricBuffer.Serialize(Ar, Owner, ChunkIndex, bAttemptFileMapping);
+	LOD.RootToUniqueTriangleIndexBuffer.Serialize(Ar, Owner, ChunkIndex, bAttemptFileMapping);
+	LOD.UniqueTriangleIndexBuffer.Serialize(Ar, Owner, ChunkIndex, bAttemptFileMapping);
+	LOD.RestUniqueTrianglePosition0Buffer.Serialize(Ar, Owner, ChunkIndex, bAttemptFileMapping);
+	LOD.RestUniqueTrianglePosition1Buffer.Serialize(Ar, Owner, ChunkIndex, bAttemptFileMapping);
+	LOD.RestUniqueTrianglePosition2Buffer.Serialize(Ar, Owner, ChunkIndex, bAttemptFileMapping);
 
 	Ar << LOD.SampleCount;
 	if (LOD.SampleCount)
@@ -1352,7 +1359,7 @@ static void InternalSerialize(FArchive& Ar, UObject* Owner, FHairStrandsRootBulk
 		LOD.MeshSampleIndicesBuffer.Serialize(Ar, Owner, ChunkIndex, bAttemptFileMapping);
 		LOD.RestSamplePositionsBuffer.Serialize(Ar, Owner, ChunkIndex, bAttemptFileMapping);
 	}
-		LOD.ValidSectionIndices.BulkSerialize(Ar);
+		LOD.UniqueSectionIndices.BulkSerialize(Ar);
 }
 
 static void InternalSerialize(FArchive& Ar, UObject* Owner, TArray<FHairStrandsRootBulkData::FMeshProjectionLOD>& LODs)
@@ -1395,26 +1402,28 @@ void FHairStrandsRootBulkData::Reset()
 	VertexToCurveIndexBuffer = FByteBulkData();
 	for (FMeshProjectionLOD& LOD : MeshProjectionLODs)
 	{
-		LOD.RootTriangleIndexBuffer.RemoveBulkData();
-		LOD.RootTriangleBarycentricBuffer.RemoveBulkData();
-		LOD.RestRootTrianglePosition0Buffer.RemoveBulkData();
-		LOD.RestRootTrianglePosition1Buffer.RemoveBulkData();
-		LOD.RestRootTrianglePosition2Buffer.RemoveBulkData();
+		LOD.RootBarycentricBuffer.RemoveBulkData();
+		LOD.RootToUniqueTriangleIndexBuffer.RemoveBulkData();
+		LOD.UniqueTriangleIndexBuffer.RemoveBulkData();
+		LOD.RestUniqueTrianglePosition0Buffer.RemoveBulkData();
+		LOD.RestUniqueTrianglePosition1Buffer.RemoveBulkData();
+		LOD.RestUniqueTrianglePosition2Buffer.RemoveBulkData();
 		LOD.MeshInterpolationWeightsBuffer.RemoveBulkData();
 		LOD.MeshSampleIndicesBuffer.RemoveBulkData();
 		LOD.RestSamplePositionsBuffer.RemoveBulkData();
 
 		// Reset the bulk byte buffer to ensure the (serialize) data size is reset to 0
-		LOD.RootTriangleIndexBuffer 		= FByteBulkData();
-		LOD.RootTriangleBarycentricBuffer 	= FByteBulkData();
-		LOD.RestRootTrianglePosition0Buffer = FByteBulkData();
-		LOD.RestRootTrianglePosition1Buffer = FByteBulkData();
-		LOD.RestRootTrianglePosition2Buffer = FByteBulkData();
-		LOD.MeshInterpolationWeightsBuffer 	= FByteBulkData();
-		LOD.MeshSampleIndicesBuffer 		= FByteBulkData();
-		LOD.RestSamplePositionsBuffer 		= FByteBulkData();
-
-		LOD.ValidSectionIndices.Empty();
+		LOD.RootBarycentricBuffer 			 = FByteBulkData();
+		LOD.RootToUniqueTriangleIndexBuffer  = FByteBulkData();
+		LOD.UniqueTriangleIndexBuffer		 = FByteBulkData();
+		LOD.RestUniqueTrianglePosition0Buffer= FByteBulkData();
+		LOD.RestUniqueTrianglePosition1Buffer= FByteBulkData();
+		LOD.RestUniqueTrianglePosition2Buffer= FByteBulkData();
+		LOD.MeshInterpolationWeightsBuffer 	 = FByteBulkData();
+		LOD.MeshSampleIndicesBuffer 		 = FByteBulkData();
+		LOD.RestSamplePositionsBuffer 		 = FByteBulkData();
+
+		LOD.UniqueSectionIndices.Empty();
 	}
 	MeshProjectionLODs.Empty();
 }
@@ -1431,14 +1440,15 @@ bool FHairStrandsRootData::HasProjectionData() const
 	bool bIsValid = MeshProjectionLODs.Num() > 0;
 	for (const FMeshProjectionLOD& LOD : MeshProjectionLODs)
 	{
-		const bool bHasValidCPUData = LOD.RootTriangleBarycentricBuffer.Num() > 0;
+		const bool bHasValidCPUData = LOD.RootBarycentricBuffer.Num() > 0;
 		if (bHasValidCPUData)
 		{
-			bIsValid = bIsValid && LOD.RootTriangleBarycentricBuffer.Num() > 0;
-			bIsValid = bIsValid && LOD.RootTriangleIndexBuffer.Num() > 0;
-			bIsValid = bIsValid && LOD.RestRootTrianglePosition0Buffer.Num() > 0;
-			bIsValid = bIsValid && LOD.RestRootTrianglePosition1Buffer.Num() > 0;
-			bIsValid = bIsValid && LOD.RestRootTrianglePosition2Buffer.Num() > 0;
+			bIsValid = bIsValid && LOD.RootBarycentricBuffer.Num() > 0;
+			bIsValid = bIsValid && LOD.RootToUniqueTriangleIndexBuffer.Num() > 0;
+			bIsValid = bIsValid && LOD.UniqueTriangleIndexBuffer.Num() > 0;
+			bIsValid = bIsValid && LOD.RestUniqueTrianglePosition0Buffer.Num() > 0;
+			bIsValid = bIsValid && LOD.RestUniqueTrianglePosition1Buffer.Num() > 0;
+			bIsValid = bIsValid && LOD.RestUniqueTrianglePosition2Buffer.Num() > 0;
 
 			if (!bIsValid) break;
 		}
diff --git a/Engine/Plugins/Runtime/HairStrands/Source/HairStrandsCore/Private/HairStrandsMeshProjection.cpp b/Engine/Plugins/Runtime/HairStrands/Source/HairStrandsCore/Private/HairStrandsMeshProjection.cpp
--- a/Engine/Plugins/Runtime/HairStrands/Source/HairStrandsCore/Private/HairStrandsMeshProjection.cpp
+++ b/Engine/Plugins/Runtime/HairStrands/Source/HairStrandsCore/Private/HairStrandsMeshProjection.cpp
@@ -152,7 +152,7 @@ private:
 	using FPermutationDomain = TShaderPermutationDomain<FUpdateUVs, FPositionType, FPrevious>;
 
 	BEGIN_SHADER_PARAMETER_STRUCT(FParameters, )
-		SHADER_PARAMETER(uint32, MaxRootCount)
+		SHADER_PARAMETER(uint32, MaxUniqueTriangleCount)
 		SHADER_PARAMETER(uint32, MaxSectionCount)
 		SHADER_PARAMETER(uint32, Pass_SectionStart)
 		SHADER_PARAMETER(uint32, Pass_SectionCount)
@@ -214,14 +214,14 @@ private:
 		SHADER_PARAMETER_SRV(Buffer, MeshUVsBuffer6)
 		SHADER_PARAMETER_SRV(Buffer, MeshUVsBuffer7)
 
-		SHADER_PARAMETER_RDG_BUFFER_SRV(Buffer, RootTriangleIndex)
-		SHADER_PARAMETER_RDG_BUFFER_UAV(RWStructuredBuffer, OutRootTrianglePrevPosition0)
-		SHADER_PARAMETER_RDG_BUFFER_UAV(RWStructuredBuffer, OutRootTrianglePrevPosition1)
-		SHADER_PARAMETER_RDG_BUFFER_UAV(RWStructuredBuffer, OutRootTrianglePrevPosition2)
+		SHADER_PARAMETER_RDG_BUFFER_SRV(Buffer, UniqueTriangleIndices)
+		SHADER_PARAMETER_RDG_BUFFER_UAV(RWStructuredBuffer, OutUniqueTrianglePrevPosition0)
+		SHADER_PARAMETER_RDG_BUFFER_UAV(RWStructuredBuffer, OutUniqueTrianglePrevPosition1)
+		SHADER_PARAMETER_RDG_BUFFER_UAV(RWStructuredBuffer, OutUniqueTrianglePrevPosition2)
 
-		SHADER_PARAMETER_RDG_BUFFER_UAV(RWStructuredBuffer, OutRootTriangleCurrPosition0)
-		SHADER_PARAMETER_RDG_BUFFER_UAV(RWStructuredBuffer, OutRootTriangleCurrPosition1)
-		SHADER_PARAMETER_RDG_BUFFER_UAV(RWStructuredBuffer, OutRootTriangleCurrPosition2)
+		SHADER_PARAMETER_RDG_BUFFER_UAV(RWStructuredBuffer, OutUniqueTriangleCurrPosition0)
+		SHADER_PARAMETER_RDG_BUFFER_UAV(RWStructuredBuffer, OutUniqueTriangleCurrPosition1)
+		SHADER_PARAMETER_RDG_BUFFER_UAV(RWStructuredBuffer, OutUniqueTriangleCurrPosition2)
 	END_SHADER_PARAMETER_STRUCT()
 
 public:
@@ -293,31 +293,33 @@ void AddHairStrandUpdateMeshTrianglesPass(
 	const uint32 PassCount = FMath::DivideAndRoundUp(ValidSectionCount, FHairUpdateMeshTriangleCS::SectionArrayCount);
 	const bool bComputePreviousDeformedPosition = IsHairStrandContinuousDecimationReorderingEnabled() && Type == HairStrandsTriangleType::DeformedPose;
 
+	const uint32 MaxUniqueTriangleCount = RestLODData.RestUniqueTrianglePosition0Buffer.Buffer->Desc.NumElements;
+
 	FHairUpdateMeshTriangleCS::FParameters CommonParameters;
 
 	FRDGImportedBuffer OutputCurrBuffers[3];
 	FRDGImportedBuffer OutputPrevBuffers[3];
 	const bool bEnableUAVOverlap = true;
 	const ERDGUnorderedAccessViewFlags UAVFlags = bEnableUAVOverlap ? ERDGUnorderedAccessViewFlags::SkipBarrier : ERDGUnorderedAccessViewFlags::None;
-	CommonParameters.RootTriangleIndex = RegisterAsSRV(GraphBuilder, RestLODData.RootTriangleIndexBuffer);
+	CommonParameters.UniqueTriangleIndices = RegisterAsSRV(GraphBuilder, RestLODData.UniqueTriangleIndexBuffer);
 	if (Type == HairStrandsTriangleType::RestPose)
 	{
-		OutputCurrBuffers[0] = Register(GraphBuilder, RestLODData.RestRootTrianglePosition0Buffer, ERDGImportedBufferFlags::CreateUAV, UAVFlags);
-		OutputCurrBuffers[1] = Register(GraphBuilder, RestLODData.RestRootTrianglePosition1Buffer, ERDGImportedBufferFlags::CreateUAV, UAVFlags);
-		OutputCurrBuffers[2] = Register(GraphBuilder, RestLODData.RestRootTrianglePosition2Buffer, ERDGImportedBufferFlags::CreateUAV, UAVFlags);
+		OutputCurrBuffers[0] = Register(GraphBuilder, RestLODData.RestUniqueTrianglePosition0Buffer, ERDGImportedBufferFlags::CreateUAV, UAVFlags);
+		OutputCurrBuffers[1] = Register(GraphBuilder, RestLODData.RestUniqueTrianglePosition1Buffer, ERDGImportedBufferFlags::CreateUAV, UAVFlags);
+		OutputCurrBuffers[2] = Register(GraphBuilder, RestLODData.RestUniqueTrianglePosition2Buffer, ERDGImportedBufferFlags::CreateUAV, UAVFlags);
 	}
 	else if (Type == HairStrandsTriangleType::DeformedPose)
 	{
 		FHairStrandsDeformedRootResource::FLOD& DeformedLODData = DeformedResources->LODs[LODIndex];
-		OutputCurrBuffers[0] = Register(GraphBuilder, DeformedLODData.GetDeformedRootTrianglePosition0Buffer(FHairStrandsDeformedRootResource::FLOD::Current), ERDGImportedBufferFlags::CreateUAV, UAVFlags);
-		OutputCurrBuffers[1] = Register(GraphBuilder, DeformedLODData.GetDeformedRootTrianglePosition1Buffer(FHairStrandsDeformedRootResource::FLOD::Current), ERDGImportedBufferFlags::CreateUAV, UAVFlags);
-		OutputCurrBuffers[2] = Register(GraphBuilder, DeformedLODData.GetDeformedRootTrianglePosition2Buffer(FHairStrandsDeformedRootResource::FLOD::Current), ERDGImportedBufferFlags::CreateUAV, UAVFlags);
+		OutputCurrBuffers[0] = Register(GraphBuilder, DeformedLODData.GetDeformedUniqueTrianglePosition0Buffer(FHairStrandsDeformedRootResource::FLOD::Current), ERDGImportedBufferFlags::CreateUAV, UAVFlags);
+		OutputCurrBuffers[1] = Register(GraphBuilder, DeformedLODData.GetDeformedUniqueTrianglePosition1Buffer(FHairStrandsDeformedRootResource::FLOD::Current), ERDGImportedBufferFlags::CreateUAV, UAVFlags);
+		OutputCurrBuffers[2] = Register(GraphBuilder, DeformedLODData.GetDeformedUniqueTrianglePosition2Buffer(FHairStrandsDeformedRootResource::FLOD::Current), ERDGImportedBufferFlags::CreateUAV, UAVFlags);
 
 		if (bComputePreviousDeformedPosition)
 		{
-			OutputPrevBuffers[0] = Register(GraphBuilder, DeformedLODData.GetDeformedRootTrianglePosition0Buffer(FHairStrandsDeformedRootResource::FLOD::Previous), ERDGImportedBufferFlags::CreateUAV, UAVFlags);
-			OutputPrevBuffers[1] = Register(GraphBuilder, DeformedLODData.GetDeformedRootTrianglePosition1Buffer(FHairStrandsDeformedRootResource::FLOD::Previous), ERDGImportedBufferFlags::CreateUAV, UAVFlags);
-			OutputPrevBuffers[2] = Register(GraphBuilder, DeformedLODData.GetDeformedRootTrianglePosition2Buffer(FHairStrandsDeformedRootResource::FLOD::Previous), ERDGImportedBufferFlags::CreateUAV, UAVFlags);
+			OutputPrevBuffers[0] = Register(GraphBuilder, DeformedLODData.GetDeformedUniqueTrianglePosition0Buffer(FHairStrandsDeformedRootResource::FLOD::Previous), ERDGImportedBufferFlags::CreateUAV, UAVFlags);
+			OutputPrevBuffers[1] = Register(GraphBuilder, DeformedLODData.GetDeformedUniqueTrianglePosition1Buffer(FHairStrandsDeformedRootResource::FLOD::Previous), ERDGImportedBufferFlags::CreateUAV, UAVFlags);
+			OutputPrevBuffers[2] = Register(GraphBuilder, DeformedLODData.GetDeformedUniqueTrianglePosition2Buffer(FHairStrandsDeformedRootResource::FLOD::Previous), ERDGImportedBufferFlags::CreateUAV, UAVFlags);
 		}
 
 		DeformedLODData.Status = FHairStrandsDeformedRootResource::FLOD::EStatus::Completed;
@@ -328,21 +330,21 @@ void AddHairStrandUpdateMeshTrianglesPass(
 		return;
 	}
 
-	CommonParameters.OutRootTriangleCurrPosition0 = OutputCurrBuffers[0].UAV;
-	CommonParameters.OutRootTriangleCurrPosition1 = OutputCurrBuffers[1].UAV;
-	CommonParameters.OutRootTriangleCurrPosition2 = OutputCurrBuffers[2].UAV;
+	CommonParameters.OutUniqueTriangleCurrPosition0 = OutputCurrBuffers[0].UAV;
+	CommonParameters.OutUniqueTriangleCurrPosition1 = OutputCurrBuffers[1].UAV;
+	CommonParameters.OutUniqueTriangleCurrPosition2 = OutputCurrBuffers[2].UAV;
 
 	if (bComputePreviousDeformedPosition)
 	{
-		CommonParameters.OutRootTrianglePrevPosition0 = OutputPrevBuffers[0].UAV;
-		CommonParameters.OutRootTrianglePrevPosition1 = OutputPrevBuffers[1].UAV;
-		CommonParameters.OutRootTrianglePrevPosition2 = OutputPrevBuffers[2].UAV;
+		CommonParameters.OutUniqueTrianglePrevPosition0 = OutputPrevBuffers[0].UAV;
+		CommonParameters.OutUniqueTrianglePrevPosition1 = OutputPrevBuffers[1].UAV;
+		CommonParameters.OutUniqueTrianglePrevPosition2 = OutputPrevBuffers[2].UAV;
 	}
 
 	for (uint32 PassIt = 0; PassIt < PassCount; ++PassIt)
 	{
 		FHairUpdateMeshTriangleCS::FParameters* Parameters = &CommonParameters;
-		Parameters->MaxRootCount = RootCount;
+		Parameters->MaxUniqueTriangleCount = MaxUniqueTriangleCount;
 		Parameters->MaxSectionCount = MeshData.Sections.Num();
 		Parameters->Pass_SectionStart = PassIt * FHairUpdateMeshTriangleCS::SectionArrayCount;
 		Parameters->Pass_SectionCount = FMath::Min(ValidSectionCount - Parameters->Pass_SectionStart, FHairUpdateMeshTriangleCS::SectionArrayCount);
@@ -361,7 +363,7 @@ void AddHairStrandUpdateMeshTrianglesPass(
 			FRHIShaderResourceView* UVsBuffer = nullptr;
 		};
 		TMap<FRHIShaderResourceView*, FMeshSectionBuffers>	UniqueMeshSectionBuffers;
-		TMap<FRDGBufferSRVRef, FMeshSectionBuffers>		UniqueMeshSectionBuffersRDG;
+		TMap<FRDGBufferSRVRef, FMeshSectionBuffers>			UniqueMeshSectionBuffersRDG;
 		uint32 UniqueMeshSectionBufferIndex = 0;
 
 		#define SETPARAMETERS(OutParameters, InMeshSectionData, Index) \
@@ -1002,12 +1004,14 @@ void AddHairStrandUpdateMeshSamplesPass(
 BEGIN_SHADER_PARAMETER_STRUCT(FHairFollicleMaskParameters, )
 	SHADER_PARAMETER(FVector2f, OutputResolution)
 	SHADER_PARAMETER(uint32, MaxRootCount)
+	SHADER_PARAMETER(uint32, MaxUniqueTriangleIndex)
 	SHADER_PARAMETER(uint32, Channel)
 	SHADER_PARAMETER(uint32, KernelSizeInPixels)
 
-	SHADER_PARAMETER_RDG_BUFFER_SRV(Buffer, TrianglePosition0Buffer)
-	SHADER_PARAMETER_RDG_BUFFER_SRV(Buffer, TrianglePosition1Buffer)
-	SHADER_PARAMETER_RDG_BUFFER_SRV(Buffer, TrianglePosition2Buffer)
+	SHADER_PARAMETER_RDG_BUFFER_SRV(Buffer, UniqueTrianglePosition0Buffer)
+	SHADER_PARAMETER_RDG_BUFFER_SRV(Buffer, UniqueTrianglePosition1Buffer)
+	SHADER_PARAMETER_RDG_BUFFER_SRV(Buffer, UniqueTrianglePosition2Buffer)
+	SHADER_PARAMETER_RDG_BUFFER_SRV(Buffer, RootToUniqueTriangleIndexBuffer)
 	SHADER_PARAMETER_RDG_BUFFER_SRV(Buffer, RootBarycentricBuffer)
 	SHADER_PARAMETER_RDG_BUFFER_SRV(Buffer, RootUVsBuffer)
 
@@ -1073,21 +1077,23 @@ static void AddFollicleMaskPass(
 		return;
 
 	FHairStrandsRestRootResource::FLOD& LODData = RestResources->LODs[LODIndex];
-	if (!LODData.RootTriangleBarycentricBuffer.Buffer ||
-		!LODData.RestRootTrianglePosition0Buffer.Buffer ||
-		!LODData.RestRootTrianglePosition1Buffer.Buffer ||
-		!LODData.RestRootTrianglePosition2Buffer.Buffer)
+	if (!LODData.RootBarycentricBuffer.Buffer ||
+		!LODData.RestUniqueTrianglePosition0Buffer.Buffer ||
+		!LODData.RestUniqueTrianglePosition1Buffer.Buffer ||
+		!LODData.RestUniqueTrianglePosition2Buffer.Buffer)
 		return;
 
 	const FIntPoint OutputResolution = OutTexture->Desc.Extent;
 	FHairFollicleMaskParameters* Parameters = GraphBuilder.AllocParameters<FHairFollicleMaskParameters>();
-	Parameters->TrianglePosition0Buffer = RegisterAsSRV(GraphBuilder, LODData.RestRootTrianglePosition0Buffer);
-	Parameters->TrianglePosition1Buffer = RegisterAsSRV(GraphBuilder, LODData.RestRootTrianglePosition1Buffer);
-	Parameters->TrianglePosition2Buffer = RegisterAsSRV(GraphBuilder, LODData.RestRootTrianglePosition2Buffer);
-	Parameters->RootBarycentricBuffer   = RegisterAsSRV(GraphBuilder, LODData.RootTriangleBarycentricBuffer);
+	Parameters->UniqueTrianglePosition0Buffer = RegisterAsSRV(GraphBuilder, LODData.RestUniqueTrianglePosition0Buffer);
+	Parameters->UniqueTrianglePosition1Buffer = RegisterAsSRV(GraphBuilder, LODData.RestUniqueTrianglePosition1Buffer);
+	Parameters->UniqueTrianglePosition2Buffer = RegisterAsSRV(GraphBuilder, LODData.RestUniqueTrianglePosition2Buffer);
+	Parameters->RootToUniqueTriangleIndexBuffer = RegisterAsSRV(GraphBuilder, LODData.RootToUniqueTriangleIndexBuffer);
+	Parameters->RootBarycentricBuffer   = RegisterAsSRV(GraphBuilder, LODData.RootBarycentricBuffer);
 	Parameters->RootUVsBuffer = nullptr;
 	Parameters->OutputResolution = OutputResolution;
 	Parameters->MaxRootCount = RootCount;
+	Parameters->MaxUniqueTriangleIndex = RestResources->BulkData.MeshProjectionLODs[LODIndex].UniqueTriangleCount;
 	Parameters->Channel = FMath::Min(Channel, 3u);
 	Parameters->KernelSizeInPixels = FMath::Clamp(KernelSizeInPixels, 2u, 200u);
 	Parameters->RenderTargets[0] = FRenderTargetBinding(OutTexture, bNeedClear ? ERenderTargetLoadAction::EClear : ERenderTargetLoadAction::ELoad, 0);
@@ -1103,7 +1109,7 @@ static void AddFollicleMaskPass(
 	ParametersPS.Pass = *Parameters;
 
 	GraphBuilder.AddPass(
-		RDG_EVENT_NAME("HairStrandsFollicleMask"),
+		RDG_EVENT_NAME("HairStrands::FollicleMask"),
 		Parameters,
 		ERDGPassFlags::Raster,
 		[Parameters, ParametersVS, ParametersPS, VertexShader, PixelShader, OutputResolution](FRHICommandList& RHICmdList)
@@ -1143,13 +1149,14 @@ static void AddFollicleMaskPass(
 {
 	const FIntPoint OutputResolution = OutTexture->Desc.Extent;
 	FHairFollicleMaskParameters* Parameters = GraphBuilder.AllocParameters<FHairFollicleMaskParameters>();
-	Parameters->TrianglePosition0Buffer = nullptr;
-	Parameters->TrianglePosition1Buffer = nullptr;
-	Parameters->TrianglePosition2Buffer = nullptr;
+	Parameters->UniqueTrianglePosition0Buffer = nullptr;
+	Parameters->UniqueTrianglePosition1Buffer = nullptr;
+	Parameters->UniqueTrianglePosition2Buffer = nullptr;
 	Parameters->RootBarycentricBuffer = nullptr;
 	Parameters->RootUVsBuffer = GraphBuilder.CreateSRV(RootUVBuffer, PF_G32R32F);
 	Parameters->OutputResolution = OutputResolution;
 	Parameters->MaxRootCount = RootCount;
+	Parameters->MaxUniqueTriangleIndex = 0;
 	Parameters->Channel = FMath::Min(Channel, 3u);
 	Parameters->KernelSizeInPixels = FMath::Clamp(KernelSizeInPixels, 2u, 200u);
 	Parameters->RenderTargets[0] = FRenderTargetBinding(OutTexture, bNeedClear ? ERenderTargetLoadAction::EClear : ERenderTargetLoadAction::ELoad, 0);
@@ -1165,7 +1172,7 @@ static void AddFollicleMaskPass(
 	ParametersPS.Pass = *Parameters;
 
 	GraphBuilder.AddPass(
-		RDG_EVENT_NAME("HairStrandsFollicleMask"),
+		RDG_EVENT_NAME("HairStrands::FollicleMask"),
 		Parameters,
 		ERDGPassFlags::Raster,
 		[Parameters, ParametersVS, ParametersPS, VertexShader, PixelShader, OutputResolution](FRHICommandList& RHICmdList)
@@ -1364,8 +1371,8 @@ void AddHairStrandUpdatePositionOffsetPass(
 	FRDGImportedBuffer RootTrianglePrevPositionBuffer;
 	if (DeformedRootResources)
 	{
-		RootTriangleCurrPositionBuffer = Register(GraphBuilder, DeformedRootResources->LODs[LODIndex].GetDeformedRootTrianglePosition0Buffer(FHairStrandsDeformedRootResource::FLOD::Current), ERDGImportedBufferFlags::CreateSRV);
-		RootTrianglePrevPositionBuffer = Register(GraphBuilder, DeformedRootResources->LODs[LODIndex].GetDeformedRootTrianglePosition0Buffer(FHairStrandsDeformedRootResource::FLOD::Previous), ERDGImportedBufferFlags::CreateSRV);
+		RootTriangleCurrPositionBuffer = Register(GraphBuilder, DeformedRootResources->LODs[LODIndex].GetDeformedUniqueTrianglePosition0Buffer(FHairStrandsDeformedRootResource::FLOD::Current), ERDGImportedBufferFlags::CreateSRV);
+		RootTrianglePrevPositionBuffer = Register(GraphBuilder, DeformedRootResources->LODs[LODIndex].GetDeformedUniqueTrianglePosition0Buffer(FHairStrandsDeformedRootResource::FLOD::Previous), ERDGImportedBufferFlags::CreateSRV);
 	}
 
 	const bool bUseGPUOffset = DeformedRootResources != nullptr && GHairStrandsUseGPUPositionOffset > 0;
diff --git a/Engine/Plugins/Runtime/HairStrands/Source/HairStrandsCore/Private/HairStrandsRendering.cpp b/Engine/Plugins/Runtime/HairStrands/Source/HairStrandsCore/Private/HairStrandsRendering.cpp
--- a/Engine/Plugins/Runtime/HairStrands/Source/HairStrandsCore/Private/HairStrandsRendering.cpp
+++ b/Engine/Plugins/Runtime/HairStrands/Source/HairStrandsCore/Private/HairStrandsRendering.cpp
@@ -301,6 +301,7 @@ class FDeformGuideCS : public FGlobalShader
 
 		SHADER_PARAMETER_RDG_BUFFER_SRV(Buffer<uint>, SimRootBarycentricBuffer)
 		SHADER_PARAMETER_RDG_BUFFER_SRV(Buffer<uint>, SimVertexToRootIndexBuffer)
+		SHADER_PARAMETER_RDG_BUFFER_SRV(Buffer<uint>, SimRootToUniqueTriangleIndexBuffer)
 
 		SHADER_PARAMETER_RDG_BUFFER_SRV(Buffer, SimDeformedOffsetBuffer)
 		SHADER_PARAMETER_RDG_BUFFER_SRV(Buffer, SimRestPosePositionBuffer)
@@ -402,15 +403,16 @@ static void AddDeformSimHairStrandsPass(
 				if (!bSupportGlobalInterpolation) 
 				{
 					InternalDeformationType = InternalDeformationType_Skinned;
-					Parameters->SimRestPosition0Buffer = RegisterAsSRV(GraphBuilder, RestLODDatas.RestRootTrianglePosition0Buffer);
-					Parameters->SimRestPosition1Buffer = RegisterAsSRV(GraphBuilder, RestLODDatas.RestRootTrianglePosition1Buffer);
-					Parameters->SimRestPosition2Buffer = RegisterAsSRV(GraphBuilder, RestLODDatas.RestRootTrianglePosition2Buffer);
+					Parameters->SimRootToUniqueTriangleIndexBuffer = RegisterAsSRV(GraphBuilder, RestLODDatas.RootToUniqueTriangleIndexBuffer);
+					Parameters->SimRestPosition0Buffer = RegisterAsSRV(GraphBuilder, RestLODDatas.RestUniqueTrianglePosition0Buffer);
+					Parameters->SimRestPosition1Buffer = RegisterAsSRV(GraphBuilder, RestLODDatas.RestUniqueTrianglePosition1Buffer);
+					Parameters->SimRestPosition2Buffer = RegisterAsSRV(GraphBuilder, RestLODDatas.RestUniqueTrianglePosition2Buffer);
 
-					Parameters->SimDeformedPosition0Buffer = RegisterAsSRV(GraphBuilder, DeformedLODDatas.GetDeformedRootTrianglePosition0Buffer(FHairStrandsDeformedRootResource::FLOD::Current));
-					Parameters->SimDeformedPosition1Buffer = RegisterAsSRV(GraphBuilder, DeformedLODDatas.GetDeformedRootTrianglePosition1Buffer(FHairStrandsDeformedRootResource::FLOD::Current));
-					Parameters->SimDeformedPosition2Buffer = RegisterAsSRV(GraphBuilder, DeformedLODDatas.GetDeformedRootTrianglePosition2Buffer(FHairStrandsDeformedRootResource::FLOD::Current));
+					Parameters->SimDeformedPosition0Buffer = RegisterAsSRV(GraphBuilder, DeformedLODDatas.GetDeformedUniqueTrianglePosition0Buffer(FHairStrandsDeformedRootResource::FLOD::Current));
+					Parameters->SimDeformedPosition1Buffer = RegisterAsSRV(GraphBuilder, DeformedLODDatas.GetDeformedUniqueTrianglePosition1Buffer(FHairStrandsDeformedRootResource::FLOD::Current));
+					Parameters->SimDeformedPosition2Buffer = RegisterAsSRV(GraphBuilder, DeformedLODDatas.GetDeformedUniqueTrianglePosition2Buffer(FHairStrandsDeformedRootResource::FLOD::Current));
 
-					Parameters->SimRootBarycentricBuffer = RegisterAsSRV(GraphBuilder, RestLODDatas.RootTriangleBarycentricBuffer);
+					Parameters->SimRootBarycentricBuffer = RegisterAsSRV(GraphBuilder, RestLODDatas.RootBarycentricBuffer);
 				}
 				else
 				{
@@ -519,8 +521,9 @@ class FHairInterpolationCS : public FGlobalShader
 		SHADER_PARAMETER_RDG_BUFFER_SRV(Buffer<float4>, DeformedPosition1Buffer)
 		SHADER_PARAMETER_RDG_BUFFER_SRV(Buffer<float4>, DeformedPosition2Buffer)
 
-		SHADER_PARAMETER_RDG_BUFFER_SRV(Buffer<uint>, RootBarycentricBuffer)
+		SHADER_PARAMETER_RDG_BUFFER_SRV(Buffer<uint>, RenRootBarycentricBuffer)
 		SHADER_PARAMETER_RDG_BUFFER_SRV(Buffer<uint>, RenVertexToRootIndexBuffer)
+		SHADER_PARAMETER_RDG_BUFFER_SRV(Buffer<uint>, RenRootToUniqueTriangleIndexBuffer)
 
 		SHADER_PARAMETER_RDG_BUFFER_SRV(Buffer<float4>, SimRestPosition0Buffer)
 		SHADER_PARAMETER_RDG_BUFFER_SRV(Buffer<float4>, SimRestPosition1Buffer)
@@ -532,6 +535,7 @@ class FHairInterpolationCS : public FGlobalShader
 
 		SHADER_PARAMETER_RDG_BUFFER_SRV(Buffer<uint>, SimRootBarycentricBuffer)
 		SHADER_PARAMETER_RDG_BUFFER_SRV(Buffer<uint>, SimVertexToRootIndexBuffer)
+		SHADER_PARAMETER_RDG_BUFFER_SRV(Buffer<uint>, SimRootToUniqueTriangleIndexBuffer)
 
 		SHADER_PARAMETER_RDG_BUFFER_SRV(Buffer<uint>, SimRootPointIndexBuffer)
 
@@ -704,16 +708,17 @@ static void AddHairStrandsInterpolationPass(
 		const FHairStrandsDeformedRootResource::FLOD& Sim_DeformedLODDatas = SimDeformedRootResources->LODs[MeshLODIndex];
 		bSupportGlobalInterpolation = Instance->Guides.bHasGlobalInterpolation && (Sim_RestLODDatas.SampleCount > 0);
 		{
-			Parameters->SimRestPosition0Buffer = RegisterAsSRV(GraphBuilder, Sim_RestLODDatas.RestRootTrianglePosition0Buffer);
-			Parameters->SimRestPosition1Buffer = RegisterAsSRV(GraphBuilder, Sim_RestLODDatas.RestRootTrianglePosition1Buffer);
-			Parameters->SimRestPosition2Buffer = RegisterAsSRV(GraphBuilder, Sim_RestLODDatas.RestRootTrianglePosition2Buffer);
+			Parameters->SimRestPosition0Buffer = RegisterAsSRV(GraphBuilder, Sim_RestLODDatas.RestUniqueTrianglePosition0Buffer);
+			Parameters->SimRestPosition1Buffer = RegisterAsSRV(GraphBuilder, Sim_RestLODDatas.RestUniqueTrianglePosition1Buffer);
+			Parameters->SimRestPosition2Buffer = RegisterAsSRV(GraphBuilder, Sim_RestLODDatas.RestUniqueTrianglePosition2Buffer);
 
-			Parameters->SimRootBarycentricBuffer = RegisterAsSRV(GraphBuilder, Sim_RestLODDatas.RootTriangleBarycentricBuffer);
+			Parameters->SimRootBarycentricBuffer = RegisterAsSRV(GraphBuilder, Sim_RestLODDatas.RootBarycentricBuffer);
+			Parameters->SimRootToUniqueTriangleIndexBuffer = RegisterAsSRV(GraphBuilder, Sim_RestLODDatas.RootToUniqueTriangleIndexBuffer);
 		}
 		{
-			Parameters->SimDeformedPosition0Buffer = RegisterAsSRV(GraphBuilder, Sim_DeformedLODDatas.GetDeformedRootTrianglePosition0Buffer(DeformedFrame));
-			Parameters->SimDeformedPosition1Buffer = RegisterAsSRV(GraphBuilder, Sim_DeformedLODDatas.GetDeformedRootTrianglePosition1Buffer(DeformedFrame));
-			Parameters->SimDeformedPosition2Buffer = RegisterAsSRV(GraphBuilder, Sim_DeformedLODDatas.GetDeformedRootTrianglePosition2Buffer(DeformedFrame));
+			Parameters->SimDeformedPosition0Buffer = RegisterAsSRV(GraphBuilder, Sim_DeformedLODDatas.GetDeformedUniqueTrianglePosition0Buffer(DeformedFrame));
+			Parameters->SimDeformedPosition1Buffer = RegisterAsSRV(GraphBuilder, Sim_DeformedLODDatas.GetDeformedUniqueTrianglePosition1Buffer(DeformedFrame));
+			Parameters->SimDeformedPosition2Buffer = RegisterAsSRV(GraphBuilder, Sim_DeformedLODDatas.GetDeformedUniqueTrianglePosition2Buffer(DeformedFrame));
 		}
 	}
 
@@ -723,15 +728,16 @@ static void AddHairStrandsInterpolationPass(
 		const FHairStrandsRestRootResource::FLOD& Ren_RestLODDatas = RenRestRootResources->LODs[MeshLODIndex];
 		const FHairStrandsDeformedRootResource::FLOD& Ren_DeformedLODDatas = RenDeformedRootResources->LODs[MeshLODIndex];
 		{
-			Parameters->RestPosition0Buffer = RegisterAsSRV(GraphBuilder, Ren_RestLODDatas.RestRootTrianglePosition0Buffer);
-			Parameters->RestPosition1Buffer = RegisterAsSRV(GraphBuilder, Ren_RestLODDatas.RestRootTrianglePosition1Buffer);
-			Parameters->RestPosition2Buffer = RegisterAsSRV(GraphBuilder, Ren_RestLODDatas.RestRootTrianglePosition2Buffer);
-			Parameters->RootBarycentricBuffer = RegisterAsSRV(GraphBuilder, Ren_RestLODDatas.RootTriangleBarycentricBuffer);
+			Parameters->RestPosition0Buffer = RegisterAsSRV(GraphBuilder, Ren_RestLODDatas.RestUniqueTrianglePosition0Buffer);
+			Parameters->RestPosition1Buffer = RegisterAsSRV(GraphBuilder, Ren_RestLODDatas.RestUniqueTrianglePosition1Buffer);
+			Parameters->RestPosition2Buffer = RegisterAsSRV(GraphBuilder, Ren_RestLODDatas.RestUniqueTrianglePosition2Buffer);
+			Parameters->RenRootBarycentricBuffer = RegisterAsSRV(GraphBuilder, Ren_RestLODDatas.RootBarycentricBuffer);
+			Parameters->RenRootToUniqueTriangleIndexBuffer = RegisterAsSRV(GraphBuilder, Ren_RestLODDatas.RootToUniqueTriangleIndexBuffer);
 		}
 		{
-			Parameters->DeformedPosition0Buffer = RegisterAsSRV(GraphBuilder, Ren_DeformedLODDatas.GetDeformedRootTrianglePosition0Buffer(DeformedFrame));
-			Parameters->DeformedPosition1Buffer = RegisterAsSRV(GraphBuilder, Ren_DeformedLODDatas.GetDeformedRootTrianglePosition1Buffer(DeformedFrame));
-			Parameters->DeformedPosition2Buffer = RegisterAsSRV(GraphBuilder, Ren_DeformedLODDatas.GetDeformedRootTrianglePosition2Buffer(DeformedFrame));
+			Parameters->DeformedPosition0Buffer = RegisterAsSRV(GraphBuilder, Ren_DeformedLODDatas.GetDeformedUniqueTrianglePosition0Buffer(DeformedFrame));
+			Parameters->DeformedPosition1Buffer = RegisterAsSRV(GraphBuilder, Ren_DeformedLODDatas.GetDeformedUniqueTrianglePosition1Buffer(DeformedFrame));
+			Parameters->DeformedPosition2Buffer = RegisterAsSRV(GraphBuilder, Ren_DeformedLODDatas.GetDeformedUniqueTrianglePosition2Buffer(DeformedFrame));
 		}
 	}
 
@@ -991,6 +997,7 @@ class FHairCardsDeformationCS : public FGlobalShader
 		SHADER_PARAMETER_RDG_BUFFER_SRV(Buffer<float4>, TriangleDeformedPosition1Buffer)
 		SHADER_PARAMETER_RDG_BUFFER_SRV(Buffer<float4>, TriangleDeformedPosition2Buffer)
 
+		SHADER_PARAMETER_RDG_BUFFER_SRV(Buffer<uint>, GuideRootToUniqueTriangleIndexBuffer)
 		SHADER_PARAMETER_RDG_BUFFER_SRV(Buffer<uint>, GuideRootBarycentricBuffer)
 		SHADER_PARAMETER_RDG_BUFFER_SRV(Buffer<uint>, GuideVertexToRootIndexBuffer)
 
@@ -1044,7 +1051,7 @@ static void AddHairCardsDeformationPass(
 	Parameters->GuideRestPositionBuffer		= RegisterAsSRV(GraphBuilder, LOD.Guides.RestResource->PositionBuffer);
 	Parameters->GuideDeformedPositionBuffer = RegisterAsSRV(GraphBuilder, LOD.Guides.DeformedResource->GetBuffer(FHairStrandsDeformedResource::Current));
 	Parameters->GuideDeformedPositionOffsetBuffer = RegisterAsSRV(GraphBuilder, LOD.Guides.DeformedResource->GetPositionOffsetBuffer(FHairStrandsDeformedResource::Current));
-	
+
 	Parameters->CardsVertexCount			= LOD.RestResource->GetVertexCount();
 	Parameters->CardsRestPositionBuffer		= LOD.RestResource->RestPositionBuffer.ShaderResourceViewRHI;
 	Parameters->CardsRestTangentBuffer		= LOD.RestResource->NormalsBuffer.ShaderResourceViewRHI;
@@ -1073,15 +1080,16 @@ static void AddHairCardsDeformationPass(
 		const FHairStrandsRestRootResource::FLOD& RestLODDatas = RestRootResources->LODs[MeshLODIndex];
 		const FHairStrandsDeformedRootResource::FLOD& DeformedLODDatas = DeformedRootResources->LODs[MeshLODIndex];
 
-		Parameters->GuideRootBarycentricBuffer = RegisterAsSRV(GraphBuilder, RestLODDatas.RootTriangleBarycentricBuffer);
+		Parameters->GuideRootBarycentricBuffer = RegisterAsSRV(GraphBuilder, RestLODDatas.RootBarycentricBuffer);
+		Parameters->GuideRootToUniqueTriangleIndexBuffer = RegisterAsSRV(GraphBuilder, RestLODDatas.RootToUniqueTriangleIndexBuffer);
 
-		Parameters->TriangleRestPosition0Buffer = RegisterAsSRV(GraphBuilder, RestLODDatas.RestRootTrianglePosition0Buffer);
-		Parameters->TriangleRestPosition1Buffer = RegisterAsSRV(GraphBuilder, RestLODDatas.RestRootTrianglePosition1Buffer);
-		Parameters->TriangleRestPosition2Buffer = RegisterAsSRV(GraphBuilder, RestLODDatas.RestRootTrianglePosition2Buffer);
+		Parameters->TriangleRestPosition0Buffer = RegisterAsSRV(GraphBuilder, RestLODDatas.RestUniqueTrianglePosition0Buffer);
+		Parameters->TriangleRestPosition1Buffer = RegisterAsSRV(GraphBuilder, RestLODDatas.RestUniqueTrianglePosition1Buffer);
+		Parameters->TriangleRestPosition2Buffer = RegisterAsSRV(GraphBuilder, RestLODDatas.RestUniqueTrianglePosition2Buffer);
 
-		Parameters->TriangleDeformedPosition0Buffer = RegisterAsSRV(GraphBuilder, DeformedLODDatas.GetDeformedRootTrianglePosition0Buffer(FHairStrandsDeformedRootResource::FLOD::Current));
-		Parameters->TriangleDeformedPosition1Buffer = RegisterAsSRV(GraphBuilder, DeformedLODDatas.GetDeformedRootTrianglePosition1Buffer(FHairStrandsDeformedRootResource::FLOD::Current));
-		Parameters->TriangleDeformedPosition2Buffer = RegisterAsSRV(GraphBuilder, DeformedLODDatas.GetDeformedRootTrianglePosition2Buffer(FHairStrandsDeformedRootResource::FLOD::Current));
+		Parameters->TriangleDeformedPosition0Buffer = RegisterAsSRV(GraphBuilder, DeformedLODDatas.GetDeformedUniqueTrianglePosition0Buffer(FHairStrandsDeformedRootResource::FLOD::Current));
+		Parameters->TriangleDeformedPosition1Buffer = RegisterAsSRV(GraphBuilder, DeformedLODDatas.GetDeformedUniqueTrianglePosition1Buffer(FHairStrandsDeformedRootResource::FLOD::Current));
+		Parameters->TriangleDeformedPosition2Buffer = RegisterAsSRV(GraphBuilder, DeformedLODDatas.GetDeformedUniqueTrianglePosition2Buffer(FHairStrandsDeformedRootResource::FLOD::Current));
 	}
 
 	if (ShaderPrintData)
diff --git a/Engine/Plugins/Runtime/HairStrands/Source/HairStrandsCore/Private/Niagara/NiagaraDataInterfaceHairStrands.cpp b/Engine/Plugins/Runtime/HairStrands/Source/HairStrandsCore/Private/Niagara/NiagaraDataInterfaceHairStrands.cpp
--- a/Engine/Plugins/Runtime/HairStrands/Source/HairStrandsCore/Private/Niagara/NiagaraDataInterfaceHairStrands.cpp
+++ b/Engine/Plugins/Runtime/HairStrands/Source/HairStrandsCore/Private/Niagara/NiagaraDataInterfaceHairStrands.cpp
@@ -37,6 +37,7 @@ BEGIN_SHADER_PARAMETER_STRUCT(FShaderParameters,)
 	SHADER_PARAMETER_SRV(Buffer<float4>,				DeformedTrianglePositionBBuffer)
 	SHADER_PARAMETER_SRV(Buffer<float4>,				DeformedTrianglePositionCBuffer)
 	SHADER_PARAMETER_SRV(Buffer<uint>,					RootBarycentricCoordinatesBuffer)
+	SHADER_PARAMETER_SRV(Buffer<uint>,					RootToUniqueTriangleIndexBuffer)
 	SHADER_PARAMETER(FMatrix44f,						WorldTransform)
 	SHADER_PARAMETER(FMatrix44f,						WorldInverse)
 	SHADER_PARAMETER(FQuat4f,							WorldRotation)
@@ -866,6 +867,7 @@ class FInterpolateGroomGuidesCS : public FGlobalShader
 		SHADER_PARAMETER_SRV(Buffer<float4>,	DeformedTrianglePositionBBuffer)
 		SHADER_PARAMETER_SRV(Buffer<float4>,	DeformedTrianglePositionCBuffer)
 		SHADER_PARAMETER_SRV(Buffer<uint>,		RootBarycentricCoordinatesBuffer)
+		SHADER_PARAMETER_SRV(Buffer<uint>,		RootToUniqueTriangleIndexBuffer)
 
 		SHADER_PARAMETER_SRV(Buffer<uint>,		VertexToCurveIndexBuffer)
 		SHADER_PARAMETER_UAV(RWBuffer<uint4>,	DeformedPositionBuffer)
@@ -945,13 +947,14 @@ static void InterpolateGroomGuides(FRDGBuilder& GraphBuilder, FNiagaraDataBuffer
 			const FHairStrandsRestRootResource::FLOD* RestMeshProjection = &(HairStrandsBuffer->SourceRestRootResources->LODs[MeshLODIndex]) ;
 			const FHairStrandsDeformedRootResource::FLOD* DeformedMeshProjection =  &(HairStrandsBuffer->SourceDeformedRootResources->LODs[MeshLODIndex]);
 
-			PassParameters->RestTrianglePositionABuffer = RestMeshProjection->RestRootTrianglePosition0Buffer.SRV;
-			PassParameters->RestTrianglePositionBBuffer = RestMeshProjection->RestRootTrianglePosition1Buffer.SRV;
-			PassParameters->RestTrianglePositionCBuffer = RestMeshProjection->RestRootTrianglePosition2Buffer.SRV;
-			PassParameters->DeformedTrianglePositionABuffer = DeformedMeshProjection->GetDeformedRootTrianglePosition0Buffer(FHairStrandsDeformedRootResource::FLOD::Current).SRV;
-			PassParameters->DeformedTrianglePositionBBuffer = DeformedMeshProjection->GetDeformedRootTrianglePosition1Buffer(FHairStrandsDeformedRootResource::FLOD::Current).SRV;
-			PassParameters->DeformedTrianglePositionCBuffer = DeformedMeshProjection->GetDeformedRootTrianglePosition2Buffer(FHairStrandsDeformedRootResource::FLOD::Current).SRV;
-			PassParameters->RootBarycentricCoordinatesBuffer = RestMeshProjection->RootTriangleBarycentricBuffer.SRV;
+			PassParameters->RestTrianglePositionABuffer = RestMeshProjection->RestUniqueTrianglePosition0Buffer.SRV;
+			PassParameters->RestTrianglePositionBBuffer = RestMeshProjection->RestUniqueTrianglePosition1Buffer.SRV;
+			PassParameters->RestTrianglePositionCBuffer = RestMeshProjection->RestUniqueTrianglePosition2Buffer.SRV;
+			PassParameters->DeformedTrianglePositionABuffer = DeformedMeshProjection->GetDeformedUniqueTrianglePosition0Buffer(FHairStrandsDeformedRootResource::FLOD::Current).SRV;
+			PassParameters->DeformedTrianglePositionBBuffer = DeformedMeshProjection->GetDeformedUniqueTrianglePosition1Buffer(FHairStrandsDeformedRootResource::FLOD::Current).SRV;
+			PassParameters->DeformedTrianglePositionCBuffer = DeformedMeshProjection->GetDeformedUniqueTrianglePosition2Buffer(FHairStrandsDeformedRootResource::FLOD::Current).SRV;
+			PassParameters->RootBarycentricCoordinatesBuffer = RestMeshProjection->RootBarycentricBuffer.SRV;
+			PassParameters->RootToUniqueTriangleIndexBuffer = RestMeshProjection->RootToUniqueTriangleIndexBuffer.SRV;
 		}
 
 		const uint32 GroupSize = NIAGARA_HAIR_STRANDS_THREAD_COUNT_INTERPOLATE;
@@ -3391,14 +3394,15 @@ void UNiagaraDataInterfaceHairStrands::SetShaderParameters(const FNiagaraDataInt
 		const FHairStrandsRestRootResource::FLOD* RestMeshProjection = bHasSkinnedInterpolation ? &(HairStrandsBuffer->SourceRestRootResources->LODs[MeshLODIndex]) : nullptr;
 		const FHairStrandsDeformedRootResource::FLOD* DeformedMeshProjection = bHasSkinnedInterpolation ? &(HairStrandsBuffer->SourceDeformedRootResources->LODs[MeshLODIndex]) : nullptr;
 
-		FRHIShaderResourceView* RestTrianglePositionASRV = (bHasSkinnedInterpolation && RestMeshProjection) ? RestMeshProjection->RestRootTrianglePosition0Buffer.SRV.GetReference() : FNiagaraRenderer::GetDummyFloatBuffer();
-		FRHIShaderResourceView* RestTrianglePositionBSRV = (bHasSkinnedInterpolation && RestMeshProjection) ? RestMeshProjection->RestRootTrianglePosition1Buffer.SRV.GetReference() : FNiagaraRenderer::GetDummyFloatBuffer();
-		FRHIShaderResourceView* RestTrianglePositionCSRV = (bHasSkinnedInterpolation && RestMeshProjection) ? RestMeshProjection->RestRootTrianglePosition2Buffer.SRV.GetReference() : FNiagaraRenderer::GetDummyFloatBuffer();
+		FRHIShaderResourceView* RestTrianglePositionASRV = (bHasSkinnedInterpolation && RestMeshProjection) ? RestMeshProjection->RestUniqueTrianglePosition0Buffer.SRV.GetReference() : FNiagaraRenderer::GetDummyFloatBuffer();
+		FRHIShaderResourceView* RestTrianglePositionBSRV = (bHasSkinnedInterpolation && RestMeshProjection) ? RestMeshProjection->RestUniqueTrianglePosition1Buffer.SRV.GetReference() : FNiagaraRenderer::GetDummyFloatBuffer();
+		FRHIShaderResourceView* RestTrianglePositionCSRV = (bHasSkinnedInterpolation && RestMeshProjection) ? RestMeshProjection->RestUniqueTrianglePosition2Buffer.SRV.GetReference() : FNiagaraRenderer::GetDummyFloatBuffer();
 
-		FRHIShaderResourceView* DeformedTrianglePositionASRV = (bHasSkinnedInterpolation && DeformedMeshProjection) ? DeformedMeshProjection->GetDeformedRootTrianglePosition0Buffer(FHairStrandsDeformedRootResource::FLOD::Current).SRV.GetReference() : FNiagaraRenderer::GetDummyFloatBuffer();
-		FRHIShaderResourceView* DeformedTrianglePositionBSRV = (bHasSkinnedInterpolation && DeformedMeshProjection) ? DeformedMeshProjection->GetDeformedRootTrianglePosition1Buffer(FHairStrandsDeformedRootResource::FLOD::Current).SRV.GetReference() : FNiagaraRenderer::GetDummyFloatBuffer();
-		FRHIShaderResourceView* DeformedTrianglePositionCSRV = (bHasSkinnedInterpolation && DeformedMeshProjection) ? DeformedMeshProjection->GetDeformedRootTrianglePosition2Buffer(FHairStrandsDeformedRootResource::FLOD::Current).SRV.GetReference() : FNiagaraRenderer::GetDummyFloatBuffer();
-		FRHIShaderResourceView* RootBarycentricCoordinatesSRV = (bHasSkinnedInterpolation && RestMeshProjection) ? RestMeshProjection->RootTriangleBarycentricBuffer.SRV.GetReference() : FNiagaraRenderer::GetDummyUIntBuffer();
+		FRHIShaderResourceView* DeformedTrianglePositionASRV = (bHasSkinnedInterpolation && DeformedMeshProjection) ? DeformedMeshProjection->GetDeformedUniqueTrianglePosition0Buffer(FHairStrandsDeformedRootResource::FLOD::Current).SRV.GetReference() : FNiagaraRenderer::GetDummyFloatBuffer();
+		FRHIShaderResourceView* DeformedTrianglePositionBSRV = (bHasSkinnedInterpolation && DeformedMeshProjection) ? DeformedMeshProjection->GetDeformedUniqueTrianglePosition1Buffer(FHairStrandsDeformedRootResource::FLOD::Current).SRV.GetReference() : FNiagaraRenderer::GetDummyFloatBuffer();
+		FRHIShaderResourceView* DeformedTrianglePositionCSRV = (bHasSkinnedInterpolation && DeformedMeshProjection) ? DeformedMeshProjection->GetDeformedUniqueTrianglePosition2Buffer(FHairStrandsDeformedRootResource::FLOD::Current).SRV.GetReference() : FNiagaraRenderer::GetDummyFloatBuffer();
+		FRHIShaderResourceView* RootBarycentricCoordinatesSRV = (bHasSkinnedInterpolation && RestMeshProjection) ? RestMeshProjection->RootBarycentricBuffer.SRV.GetReference() : FNiagaraRenderer::GetDummyUIntBuffer();
+		FRHIShaderResourceView* RootToUniqueTriangleIndexSRV = (bHasSkinnedInterpolation && RestMeshProjection) ? RestMeshProjection->RootToUniqueTriangleIndexBuffer.SRV.GetReference() : FNiagaraRenderer::GetDummyUIntBuffer();
 
 		// RBF buffers
 		const bool bHasSamples = (RestMeshProjection && RestMeshProjection->SampleCount > 0);
@@ -3478,6 +3482,7 @@ void UNiagaraDataInterfaceHairStrands::SetShaderParameters(const FNiagaraDataInt
 		ShaderParameters->RestSamplePositionsBuffer = RestSamplePositionsBufferSRV;
 		ShaderParameters->MeshSampleWeightsBuffer = MeshSampleWeightsBufferSRV;
 		ShaderParameters->RootBarycentricCoordinatesBuffer = RootBarycentricCoordinatesSRV;
+		ShaderParameters->RootToUniqueTriangleIndexBuffer = RootToUniqueTriangleIndexSRV;
 	}
 	else
 	{
@@ -3525,6 +3530,7 @@ void UNiagaraDataInterfaceHairStrands::SetShaderParameters(const FNiagaraDataInt
 		ShaderParameters->RestSamplePositionsBuffer = FNiagaraRenderer::GetDummyFloatBuffer();
 		ShaderParameters->MeshSampleWeightsBuffer = FNiagaraRenderer::GetDummyFloatBuffer();
 		ShaderParameters->RootBarycentricCoordinatesBuffer = FNiagaraRenderer::GetDummyFloatBuffer();
+		ShaderParameters->RootToUniqueTriangleIndexBuffer = FNiagaraRenderer::GetDummyFloatBuffer();
 		ShaderParameters->ParamsScaleBuffer = FNiagaraRenderer::GetDummyFloatBuffer();
 	}
 }
diff --git a/Engine/Plugins/Runtime/HairStrands/Source/HairStrandsCore/Public/GroomBindingAsset.h b/Engine/Plugins/Runtime/HairStrands/Source/HairStrandsCore/Public/GroomBindingAsset.h
--- a/Engine/Plugins/Runtime/HairStrands/Source/HairStrandsCore/Public/GroomBindingAsset.h
+++ b/Engine/Plugins/Runtime/HairStrands/Source/HairStrandsCore/Public/GroomBindingAsset.h
@@ -47,14 +47,6 @@ enum class EGroomBindingMeshType : uint8
 	GeometryCache
 };
 
-/** Binding data */
-struct FHairRootGroupData
-{
-	FHairStrandsRootData			SimRootData;
-	FHairStrandsRootData			RenRootData;
-	TArray<FHairStrandsRootData>	CardsRootData;
-};
-
 /** Binding bulk data */
 struct FHairGroupBulkData
 {
diff --git a/Engine/Plugins/Runtime/HairStrands/Source/HairStrandsCore/Public/GroomBindingBuilder.h b/Engine/Plugins/Runtime/HairStrands/Source/HairStrandsCore/Public/GroomBindingBuilder.h
--- a/Engine/Plugins/Runtime/HairStrands/Source/HairStrandsCore/Public/GroomBindingBuilder.h
+++ b/Engine/Plugins/Runtime/HairStrands/Source/HairStrandsCore/Public/GroomBindingBuilder.h
@@ -19,9 +19,6 @@ struct HAIRSTRANDSCORE_API FGroomBindingBuilder
 	// Build binding asset data
 	static bool BuildBinding(class UGroomBindingAsset* BindingAsset, bool bInitResources);
 
-	// Update the root mesh projection data with unique valid mesh section IDs, based on the projection data
-	static void BuildUniqueSections(FHairStrandsRootData::FMeshProjectionLOD& LOD);
-
 	// Extract root data from bulk data
 	static void GetRootData(FHairStrandsRootData& Out, const FHairStrandsRootBulkData& In);
 };
diff --git a/Engine/Plugins/Runtime/HairStrands/Source/HairStrandsCore/Public/GroomResources.h b/Engine/Plugins/Runtime/HairStrands/Source/HairStrandsCore/Public/GroomResources.h
--- a/Engine/Plugins/Runtime/HairStrands/Source/HairStrandsCore/Public/GroomResources.h
+++ b/Engine/Plugins/Runtime/HairStrands/Source/HairStrandsCore/Public/GroomResources.h
@@ -131,11 +131,12 @@ struct FHairStrandsRestRootResource : public FHairCommonResource
 		Total += GetBufferTotalNumBytes(VertexToCurveIndexBuffer);
 		for (const FLOD& LOD : LODs)
 		{
-			Total += GetBufferTotalNumBytes(LOD.RootTriangleIndexBuffer);
-			Total += GetBufferTotalNumBytes(LOD.RootTriangleBarycentricBuffer);
-			Total += GetBufferTotalNumBytes(LOD.RestRootTrianglePosition0Buffer);
-			Total += GetBufferTotalNumBytes(LOD.RestRootTrianglePosition1Buffer);
-			Total += GetBufferTotalNumBytes(LOD.RestRootTrianglePosition2Buffer);
+			Total += GetBufferTotalNumBytes(LOD.UniqueTriangleIndexBuffer);
+			Total += GetBufferTotalNumBytes(LOD.RootToUniqueTriangleIndexBuffer);
+			Total += GetBufferTotalNumBytes(LOD.RootBarycentricBuffer);
+			Total += GetBufferTotalNumBytes(LOD.RestUniqueTrianglePosition0Buffer);
+			Total += GetBufferTotalNumBytes(LOD.RestUniqueTrianglePosition1Buffer);
+			Total += GetBufferTotalNumBytes(LOD.RestUniqueTrianglePosition2Buffer);
 			Total += GetBufferTotalNumBytes(LOD.MeshInterpolationWeightsBuffer);
 			Total += GetBufferTotalNumBytes(LOD.MeshSampleIndicesBuffer);
 			Total += GetBufferTotalNumBytes(LOD.RestSamplePositionsBuffer);
@@ -156,13 +157,15 @@ struct FHairStrandsRestRootResource : public FHairCommonResource
 		/* Triangle on which a root is attached */
 		/* When the projection is done with source to target mesh transfer, the projection indices does not match.
 		   In this case we need to separate index computation. The barycentric coords remain the same however. */
-		FRDGExternalBuffer RootTriangleIndexBuffer;
-		FRDGExternalBuffer RootTriangleBarycentricBuffer;
+		FRDGExternalBuffer UniqueTriangleIndexBuffer;
+		/* Strands hair root to unique triangle index */
+		FRDGExternalBuffer RootToUniqueTriangleIndexBuffer;
+		FRDGExternalBuffer RootBarycentricBuffer;
 
 		/* Strand hair roots translation and rotation in rest position relative to the bound triangle. Positions are relative to the rest root center */
-		FRDGExternalBuffer RestRootTrianglePosition0Buffer;
-		FRDGExternalBuffer RestRootTrianglePosition1Buffer;
-		FRDGExternalBuffer RestRootTrianglePosition2Buffer;
+		FRDGExternalBuffer RestUniqueTrianglePosition0Buffer;
+		FRDGExternalBuffer RestUniqueTrianglePosition1Buffer;
+		FRDGExternalBuffer RestUniqueTrianglePosition2Buffer;
 
 		/* Strand hair mesh interpolation matrix and sample indices */
 		uint32 SampleCount = 0;
@@ -205,18 +208,18 @@ struct FHairStrandsDeformedRootResource : public FHairCommonResource
 		uint32 Total = 0;
 		for (const FLOD& LOD : LODs)
 		{
-			Total += GetBufferTotalNumBytes(LOD.DeformedRootTrianglePosition0Buffer[0]);
-			Total += GetBufferTotalNumBytes(LOD.DeformedRootTrianglePosition1Buffer[0]);
-			Total += GetBufferTotalNumBytes(LOD.DeformedRootTrianglePosition2Buffer[0]);
+			Total += GetBufferTotalNumBytes(LOD.DeformedUniqueTrianglePosition0Buffer[0]);
+			Total += GetBufferTotalNumBytes(LOD.DeformedUniqueTrianglePosition1Buffer[0]);
+			Total += GetBufferTotalNumBytes(LOD.DeformedUniqueTrianglePosition2Buffer[0]);
 			Total += GetBufferTotalNumBytes(LOD.DeformedSamplePositionsBuffer[0]);
 			Total += GetBufferTotalNumBytes(LOD.MeshSampleWeightsBuffer[0]);
 
 			// Double buffering is disabled by default unless the read-only cvar r.HairStrands.ContinuousDecimationReordering is set
 			if (IsHairStrandContinuousDecimationReorderingEnabled())
 			{
-				Total += GetBufferTotalNumBytes(LOD.DeformedRootTrianglePosition0Buffer[1]);
-				Total += GetBufferTotalNumBytes(LOD.DeformedRootTrianglePosition1Buffer[1]);
-				Total += GetBufferTotalNumBytes(LOD.DeformedRootTrianglePosition2Buffer[1]);
+				Total += GetBufferTotalNumBytes(LOD.DeformedUniqueTrianglePosition0Buffer[1]);
+				Total += GetBufferTotalNumBytes(LOD.DeformedUniqueTrianglePosition1Buffer[1]);
+				Total += GetBufferTotalNumBytes(LOD.DeformedUniqueTrianglePosition2Buffer[1]);
 				Total += GetBufferTotalNumBytes(LOD.DeformedSamplePositionsBuffer[1]);
 				Total += GetBufferTotalNumBytes(LOD.MeshSampleWeightsBuffer[1]);
 			}
@@ -244,9 +247,9 @@ struct FHairStrandsDeformedRootResource : public FHairCommonResource
 		int32 LODIndex = -1;
 
 		/* Strand hair roots translation and rotation in triangle-deformed position relative to the bound triangle. Positions are relative the deformed root center*/
-		FRDGExternalBuffer DeformedRootTrianglePosition0Buffer[2];
-		FRDGExternalBuffer DeformedRootTrianglePosition1Buffer[2];
-		FRDGExternalBuffer DeformedRootTrianglePosition2Buffer[2];
+		FRDGExternalBuffer DeformedUniqueTrianglePosition0Buffer[2];
+		FRDGExternalBuffer DeformedUniqueTrianglePosition1Buffer[2];
+		FRDGExternalBuffer DeformedUniqueTrianglePosition2Buffer[2];
 
 		/* Strand hair mesh interpolation matrix and sample indices */
 		uint32 SampleCount = 0;
@@ -264,9 +267,9 @@ struct FHairStrandsDeformedRootResource : public FHairCommonResource
 
 		// Double buffering is disabled by default unless the read-only cvar r.HairStrands.ContinuousDecimationReordering is set
 		inline uint32 GetIndex(EFrameType T) const { return T == EFrameType::Current ? CurrentIndex : 1u - CurrentIndex; }
-		inline const FRDGExternalBuffer& GetDeformedRootTrianglePosition0Buffer(EFrameType T) const { return IsHairStrandContinuousDecimationReorderingEnabled() ? DeformedRootTrianglePosition0Buffer[GetIndex(T)] : DeformedRootTrianglePosition0Buffer[0]; }
-		inline const FRDGExternalBuffer& GetDeformedRootTrianglePosition1Buffer(EFrameType T) const { return IsHairStrandContinuousDecimationReorderingEnabled() ? DeformedRootTrianglePosition1Buffer[GetIndex(T)] : DeformedRootTrianglePosition1Buffer[0]; }
-		inline const FRDGExternalBuffer& GetDeformedRootTrianglePosition2Buffer(EFrameType T) const { return IsHairStrandContinuousDecimationReorderingEnabled() ? DeformedRootTrianglePosition2Buffer[GetIndex(T)] : DeformedRootTrianglePosition2Buffer[0]; }
+		inline const FRDGExternalBuffer& GetDeformedUniqueTrianglePosition0Buffer(EFrameType T) const { return IsHairStrandContinuousDecimationReorderingEnabled() ? DeformedUniqueTrianglePosition0Buffer[GetIndex(T)] : DeformedUniqueTrianglePosition0Buffer[0]; }
+		inline const FRDGExternalBuffer& GetDeformedUniqueTrianglePosition1Buffer(EFrameType T) const { return IsHairStrandContinuousDecimationReorderingEnabled() ? DeformedUniqueTrianglePosition1Buffer[GetIndex(T)] : DeformedUniqueTrianglePosition1Buffer[0]; }
+		inline const FRDGExternalBuffer& GetDeformedUniqueTrianglePosition2Buffer(EFrameType T) const { return IsHairStrandContinuousDecimationReorderingEnabled() ? DeformedUniqueTrianglePosition2Buffer[GetIndex(T)] : DeformedUniqueTrianglePosition2Buffer[0]; }
 		inline const FRDGExternalBuffer& GetDeformedSamplePositionsBuffer(EFrameType T) const { return IsHairStrandContinuousDecimationReorderingEnabled() ? DeformedSamplePositionsBuffer[GetIndex(T)] : DeformedSamplePositionsBuffer[0]; }
 		inline const FRDGExternalBuffer& GetMeshSampleWeightsBuffer(EFrameType T) const { return IsHairStrandContinuousDecimationReorderingEnabled() ? MeshSampleWeightsBuffer[GetIndex(T)] : MeshSampleWeightsBuffer[0]; }
 		inline void SwapBuffer() { CurrentIndex = 1u - CurrentIndex; }
diff --git a/Engine/Plugins/Runtime/HairStrands/Source/HairStrandsCore/Public/HairStrandsDatas.h b/Engine/Plugins/Runtime/HairStrands/Source/HairStrandsCore/Public/HairStrandsDatas.h
--- a/Engine/Plugins/Runtime/HairStrands/Source/HairStrandsCore/Public/HairStrandsDatas.h
+++ b/Engine/Plugins/Runtime/HairStrands/Source/HairStrandsCore/Public/HairStrandsDatas.h
@@ -517,15 +517,16 @@ struct FHairStrandsRootBulkData
 		Total += VertexToCurveIndexBuffer.IsBulkDataLoaded() ? VertexToCurveIndexBuffer.GetBulkDataSize() : 0u;
 		for (const FMeshProjectionLOD& LOD : MeshProjectionLODs)
 		{
-			Total += LOD.RootTriangleIndexBuffer.IsBulkDataLoaded() ?			LOD.RootTriangleIndexBuffer.GetBulkDataSize() : 0u;
-			Total += LOD.RootTriangleBarycentricBuffer.IsBulkDataLoaded() ?		LOD.RootTriangleBarycentricBuffer.GetBulkDataSize() : 0u;
-			Total += LOD.RestRootTrianglePosition0Buffer.IsBulkDataLoaded() ?	LOD.RestRootTrianglePosition0Buffer.GetBulkDataSize() : 0u;
-			Total += LOD.RestRootTrianglePosition1Buffer.IsBulkDataLoaded() ?	LOD.RestRootTrianglePosition1Buffer.GetBulkDataSize() : 0u;
-			Total += LOD.RestRootTrianglePosition2Buffer.IsBulkDataLoaded() ?	LOD.RestRootTrianglePosition2Buffer.GetBulkDataSize() : 0u;
+			Total += LOD.UniqueTriangleIndexBuffer.IsBulkDataLoaded() ?			LOD.UniqueTriangleIndexBuffer.GetBulkDataSize() : 0u;
+			Total += LOD.RootToUniqueTriangleIndexBuffer.IsBulkDataLoaded() ?	LOD.RootToUniqueTriangleIndexBuffer.GetBulkDataSize() : 0u;
+			Total += LOD.RootBarycentricBuffer.IsBulkDataLoaded() ?				LOD.RootBarycentricBuffer.GetBulkDataSize() : 0u;
+			Total += LOD.RestUniqueTrianglePosition0Buffer.IsBulkDataLoaded() ?	LOD.RestUniqueTrianglePosition0Buffer.GetBulkDataSize() : 0u;
+			Total += LOD.RestUniqueTrianglePosition1Buffer.IsBulkDataLoaded() ?	LOD.RestUniqueTrianglePosition1Buffer.GetBulkDataSize() : 0u;
+			Total += LOD.RestUniqueTrianglePosition2Buffer.IsBulkDataLoaded() ?	LOD.RestUniqueTrianglePosition2Buffer.GetBulkDataSize() : 0u;
 			Total += LOD.MeshInterpolationWeightsBuffer.IsBulkDataLoaded() ?	LOD.MeshInterpolationWeightsBuffer.GetBulkDataSize() : 0u;
 			Total += LOD.MeshSampleIndicesBuffer.IsBulkDataLoaded() ?			LOD.MeshSampleIndicesBuffer.GetBulkDataSize() : 0u;
 			Total += LOD.RestSamplePositionsBuffer.IsBulkDataLoaded() ?			LOD.RestSamplePositionsBuffer.GetBulkDataSize() : 0u;
-			Total += LOD.ValidSectionIndices.GetAllocatedSize();
+			Total += LOD.UniqueSectionIndices.GetAllocatedSize();
 		}
 		return Total;
 	}
@@ -533,32 +534,35 @@ struct FHairStrandsRootBulkData
 	struct FMeshProjectionLOD
 	{
 		int32 LODIndex = -1;
+		uint32 UniqueTriangleCount = 0;
 
-		/* Triangle on which a root is attached */
-		/* When the projection is done with source to target mesh transfer, the projection indices does not match.
-		   In this case we need to separate index computation. The barycentric coords remain the same however. */
-		FByteBulkData RootTriangleIndexBuffer;
-		FByteBulkData RootTriangleBarycentricBuffer;
+		/* Map each root onto the unique triangle Id (per-root) */
+		FByteBulkData RootToUniqueTriangleIndexBuffer;
+		/* Root's barycentric (per-root) */
+		FByteBulkData RootBarycentricBuffer;
 
-		/* Strand hair roots translation and rotation in rest position relative to the bound triangle. Positions are relative to the rest root center */
-		FByteBulkData RestRootTrianglePosition0Buffer;
-		FByteBulkData RestRootTrianglePosition1Buffer;
-		FByteBulkData RestRootTrianglePosition2Buffer;
+		/* Unique triangles list from skeleton mesh section IDs and triangle IDs (per-unique-triangle) */
+		FByteBulkData UniqueTriangleIndexBuffer;
+
+		/* Rest triangle positions (per-unique-triangle) */
+		FByteBulkData RestUniqueTrianglePosition0Buffer;
+		FByteBulkData RestUniqueTrianglePosition1Buffer;
+		FByteBulkData RestUniqueTrianglePosition2Buffer;
 
 		/* Number of samples used for the mesh interpolation */
 		uint32 SampleCount = 0;
 
-		/* Store the hair interpolation weights | Size = SamplesCount * SamplesCount */
+		/* Store the hair interpolation weights | Size = SamplesCount * SamplesCount (per-sample)*/
 		FByteBulkData MeshInterpolationWeightsBuffer;
 
-		/* Store the samples vertex indices */
+		/* Store the samples vertex indices (per-sample) */
 		FByteBulkData MeshSampleIndicesBuffer;
 
-		/* Store the samples rest positions */
+		/* Store the samples rest positions (per-sample) */
 		FByteBulkData RestSamplePositionsBuffer;
 
 		/* Store the mesh section indices which are relevant for this root LOD data */
-		TArray<uint32> ValidSectionIndices;
+		TArray<uint32> UniqueSectionIndices;
 	};
 
 	/* Number of roots */
@@ -589,13 +593,14 @@ struct FHairStrandsRootData
 		/* Triangle on which a root is attached */
 		/* When the projection is done with source to target mesh transfer, the projection indices does not match.
 		   In this case we need to separate index computation. The barycentric coords remain the same however. */
-		TArray<FHairStrandsCurveTriangleIndexFormat::Type> RootTriangleIndexBuffer;
-		TArray<FHairStrandsCurveTriangleBarycentricFormat::Type> RootTriangleBarycentricBuffer;
+		TArray<uint32> RootToUniqueTriangleIndexBuffer;
+		TArray<FHairStrandsCurveTriangleBarycentricFormat::Type> RootBarycentricBuffer;
 
 		/* Strand hair roots translation and rotation in rest position relative to the bound triangle. Positions are relative to the rest root center */
-		TArray<FHairStrandsMeshTrianglePositionFormat::Type> RestRootTrianglePosition0Buffer;
-		TArray<FHairStrandsMeshTrianglePositionFormat::Type> RestRootTrianglePosition1Buffer;
-		TArray<FHairStrandsMeshTrianglePositionFormat::Type> RestRootTrianglePosition2Buffer;
+		TArray<FHairStrandsCurveTriangleIndexFormat::Type> UniqueTriangleIndexBuffer;
+		TArray<FHairStrandsMeshTrianglePositionFormat::Type> RestUniqueTrianglePosition0Buffer;
+		TArray<FHairStrandsMeshTrianglePositionFormat::Type> RestUniqueTrianglePosition1Buffer;
+		TArray<FHairStrandsMeshTrianglePositionFormat::Type> RestUniqueTrianglePosition2Buffer;
 
 		/* Number of samples used for the mesh interpolation */
 		uint32 SampleCount = 0;
@@ -610,7 +615,7 @@ struct FHairStrandsRootData
 		TArray<FHairStrandsMeshTrianglePositionFormat::Type> RestSamplePositionsBuffer;
 
 		/* Store the mesh section indices which are relevant for this root LOD data */
-		TArray<uint32> ValidSectionIndices;
+		TArray<uint32> UniqueSectionIds;
 	};
 
 	/* Number of roots */
@@ -626,7 +631,6 @@ struct FHairStrandsRootData
 	TArray<FMeshProjectionLOD> MeshProjectionLODs;
 };
 
-
 /* Structure describing the LOD settings (Screen size, vertex info, ...) for each clusters.
 	The packed version of this structure corresponds to the GPU data layout (HairStrandsClusterCommon.ush)
 	This uses by the GPU LOD selection. */
diff --git a/Engine/Shaders/Private/HairStrands/HairCardsDeformation.usf b/Engine/Shaders/Private/HairStrands/HairCardsDeformation.usf
--- a/Engine/Shaders/Private/HairStrands/HairCardsDeformation.usf
+++ b/Engine/Shaders/Private/HairStrands/HairCardsDeformation.usf
@@ -38,6 +38,7 @@ Buffer<float4>  TriangleDeformedPosition2Buffer;
 
 Buffer<uint>	GuideRootBarycentricBuffer;
 Buffer<uint>    GuideVertexToRootIndexBuffer;
+Buffer<uint>    GuideRootToUniqueTriangleIndexBuffer;
 
 
 #if PERMUTATION_GROUP_SIZE == 64
@@ -154,10 +155,12 @@ void MainCS(uint2 DispatchThreadId : SV_DispatchThreadID)
 	const float GuideLerp = ((PackedInterpolationData >> 24) & 0xFF) / float(0xFF);
 
 	const uint RootIndex = GuideVertexToRootIndexBuffer[GuideIndex0];
+	const uint TriangleIndex = GuideRootToUniqueTriangleIndexBuffer[RootIndex];
+
 	const float3 RootBarycentric = DecodeBarycentrics(GuideRootBarycentricBuffer[RootIndex]);
 
-	const FHairMeshTriangle RestTriangle = GetTriangleTransformation(RootIndex, TriangleRestPosition0Buffer, TriangleRestPosition1Buffer, TriangleRestPosition2Buffer);
-	const FHairMeshTriangle DeformedTriangle = GetTriangleTransformation(RootIndex, TriangleDeformedPosition0Buffer, TriangleDeformedPosition1Buffer, TriangleDeformedPosition2Buffer);
+	const FHairMeshTriangle RestTriangle = GetTriangleTransformation(TriangleIndex, TriangleRestPosition0Buffer, TriangleRestPosition1Buffer, TriangleRestPosition2Buffer);
+	const FHairMeshTriangle DeformedTriangle = GetTriangleTransformation(TriangleIndex, TriangleDeformedPosition0Buffer, TriangleDeformedPosition1Buffer, TriangleDeformedPosition2Buffer);
 
 	const float3 Offset0 = ComputeDynamicGeometryOffset(GuideIndex0, 1, RestTriangle, DeformedTriangle, RootBarycentric, GuideDeformedPositionOffset);
 	const float3 Offset1 = ComputeDynamicGeometryOffset(GuideIndex1, 1, RestTriangle, DeformedTriangle, RootBarycentric, GuideDeformedPositionOffset);
diff --git a/Engine/Shaders/Private/HairStrands/HairStrandsDebug.usf b/Engine/Shaders/Private/HairStrands/HairStrandsDebug.usf
--- a/Engine/Shaders/Private/HairStrands/HairStrandsDebug.usf
+++ b/Engine/Shaders/Private/HairStrands/HairStrandsDebug.usf
@@ -222,6 +222,8 @@ void MainCS(uint3 DispatchThreadId : SV_DispatchThreadID)
 #define INPUT_ROOT_TRIANGLE 1
 #define INPUT_ROOT_SAMPLES 2
 
+uint MaxUniqueSection;
+uint MaxUniqueTriangleCount;
 uint MaxRootCount;
 float4x4 RootLocalToWorld;
 uint DeformedFrameEnable;
@@ -230,6 +232,8 @@ uint DeformedFrameEnable;
 Buffer<uint>   RootBarycentricBuffer;
 #endif
 
+Buffer<uint>  RootToUniqueTriangleIndexBuffer;
+
 Buffer<float4>  RestPosition0Buffer;
 Buffer<float4>  RestPosition1Buffer;
 Buffer<float4>  RestPosition2Buffer;
@@ -247,12 +251,13 @@ void MainCS(uint3 DispatchThreadId : SV_DispatchThreadID)
 	ResolvedView = ResolveView();
 
 #if PERMUTATION_INPUT_TYPE == INPUT_ROOT_FRAME  
-	const uint LineIndex = DispatchThreadId.x;
+	const uint RootIndex = DispatchThreadId.x;
 
-	if (LineIndex > MaxRootCount) return;
+	if (RootIndex > MaxRootCount) return;
 	{
-		const FHairMeshTriangleWithUVs RestTriangle		= GetTriangleWithUVsTransformation(LineIndex, RestPosition0Buffer, RestPosition1Buffer, RestPosition2Buffer);
-		const FHairMeshTriangleWithUVs DeformedTriangle = GetTriangleWithUVsTransformation(LineIndex, DeformedPosition0Buffer, DeformedPosition1Buffer, DeformedPosition2Buffer);
+		const uint TriangleIndex = RootToUniqueTriangleIndexBuffer[RootIndex];
+		const FHairMeshTriangleWithUVs RestTriangle		= GetTriangleWithUVsTransformation(TriangleIndex, RestPosition0Buffer, RestPosition1Buffer, RestPosition2Buffer);
+		const FHairMeshTriangleWithUVs DeformedTriangle = GetTriangleWithUVsTransformation(TriangleIndex, DeformedPosition0Buffer, DeformedPosition1Buffer, DeformedPosition2Buffer);
 
 		FHairMeshTriangleWithUVs Tri;
 		if (DeformedFrameEnable > 0)
@@ -260,7 +265,7 @@ void MainCS(uint3 DispatchThreadId : SV_DispatchThreadID)
 		else
 			Tri = RestTriangle;
 
-		const float3 B  = DecodeBarycentrics(RootBarycentricBuffer[LineIndex]);
+		const float3 B  = DecodeBarycentrics(RootBarycentricBuffer[RootIndex]);
 		float3 P0 = B.x * Tri.P0 + B.y * Tri.P1 + B.z * Tri.P2;
 		float3 P1 = P0 + Tri.N;
 
@@ -1087,7 +1092,7 @@ void AddNewLine(float2 OriginalPos, inout float2 Pos)
 int2					MaxResolution;
 uint					InstanceCount;
 
-Buffer<uint4>			Infos;
+Buffer<uint>			Infos;
 
 uint					NameInfoCount;
 uint					NameCharacterCount;
@@ -1095,6 +1100,69 @@ StructuredBuffer<uint2> NameInfos;
 Buffer<uint>			Names;
 Buffer<int>				InstanceAABB;
 
+struct FInstanceInfo
+{
+	uint  GroupIndex;
+	uint  GroupCount;
+	float LODCount;
+	uint  GeometryType;
+	uint  BindingType;
+	bool  bHasSim;
+	bool  bHasRBF;
+	float LOD;
+	float ClipValue;
+
+	// Strands data
+	uint StrandsCount;
+	uint VertexCount;
+
+	// Binding
+	uint UniqueTriangle;
+	uint UniqueSection;
+	uint RootCount;
+	uint ControlPointCount;
+};
+
+FInstanceInfo ReadInstanceInfo(uint Index)
+{
+	const uint Index8 = Index * 8;
+
+	uint4 Data0 = 0;
+	Data0.x = Infos[Index8 + 0];
+	Data0.y = Infos[Index8 + 1];
+	Data0.z = Infos[Index8 + 2];
+	Data0.w = Infos[Index8 + 3];
+
+	uint4 Data1 = 0;
+	Data1.x = Infos[Index8 + 4];
+	Data1.y = Infos[Index8 + 5];
+	Data1.z = Infos[Index8 + 6];
+	Data1.w = Infos[Index8 + 7];
+
+	FInstanceInfo Out = (FInstanceInfo)0;
+	Out.GroupIndex = (Data0.x >> 0) & 0xFF;
+	Out.GroupCount = (Data0.x >> 8) & 0xFF;
+	Out.LODCount = (Data0.x >> 16) & 0xFF;
+	Out.GeometryType = (Data0.x >> 24) & 0x7;
+	Out.BindingType = (Data0.x >> 27) & 0x7;
+	Out.bHasSim = (Data0.x >> 30) & 0x1;
+	Out.bHasRBF = (Data0.x >> 31) & 0x1;
+
+	Out.LOD = f16tof32((Data0.y >> 0) & 0xFFFF);
+	Out.ClipValue = f16tof32((Data0.y >> 16) & 0xFFFF);
+
+	Out.StrandsCount = Data0.z;
+	Out.VertexCount = Data0.w;
+
+	// Strands binding data
+	Out.UniqueSection		= Data1.x;
+	Out.UniqueTriangle		= Data1.y;
+	Out.RootCount			= Data1.z;
+	Out.ControlPointCount	= Data1.w;
+
+	return Out;
+}
+
 struct FNameInfo
 {
 	uint PrimitiveID;
@@ -1161,11 +1229,13 @@ void MainCS(uint GroupIndex : SV_GroupIndex, uint3 DispatchThreadId : SV_Dispatc
 	const uint TotalGroup = InstanceCount;
 	uint TotalStrandsCount = 0;
 	uint TotalVertexCount = 0;
-	for (uint InstanceIndex = 0; InstanceIndex < InstanceCount; ++InstanceIndex)
 	{
-		const uint4 InstanceData = Infos[InstanceIndex];
-		TotalStrandsCount += InstanceData.z;
-		TotalVertexCount += InstanceData.w;
+		for (uint InstanceIndex = 0; InstanceIndex < InstanceCount; ++InstanceIndex)
+		{
+			const FInstanceInfo Info = ReadInstanceInfo(InstanceIndex);
+			TotalStrandsCount += Info.StrandsCount;
+			TotalVertexCount += Info.VertexCount;
+		}
 	}
 
 	Pos = ShaderPrintText(Pos, TEXT("Total Groups        : "), FontWhite);
@@ -1183,77 +1253,70 @@ void MainCS(uint GroupIndex : SV_GroupIndex, uint3 DispatchThreadId : SV_Dispatc
 	AddNewLine(OriginalPos, Pos);
 
 	// Draw instance details
-	Pos = ShaderPrintText(Pos, TEXT("   Group  LOD     Geometry Binding Sim   RBF   Clip    Strands Vertex  Name"));
+	Pos = ShaderPrintText(Pos, TEXT("   Group  LOD     Geometry Binding Sim   RBF   Clip    Strands Vertex  Bind.Triangle Bind.Section Name"));
 	AddNewLine(OriginalPos, Pos);
 
 	for (uint InstanceIndex = 0; InstanceIndex < InstanceCount; ++InstanceIndex)
 	{
-		const uint4 InstanceData	= Infos[InstanceIndex];
-		const uint GroupIndex		= (InstanceData.x>>0)  & 0xFF;
-		const uint GroupCount		= (InstanceData.x>>8)  & 0xFF;
-		const float LODCount		= (InstanceData.x>>16) & 0xFF;
-		const uint GeometryType		= (InstanceData.x>>24) & 0x7;
-		const uint BindingType		= (InstanceData.x>>27) & 0x7;
-		const bool bHasSim			= (InstanceData.x>>30) & 0x1;
-		const bool bHasRBF			= (InstanceData.x>>31) & 0x1;
-
-		const float LOD				= f16tof32((InstanceData.y>>0)  & 0xFFFF);
-		const float ClipValue		= f16tof32((InstanceData.y>>16) & 0xFFFF);
-
-		const uint StrandsCount		= InstanceData.z;
-		const uint VertexCount		= InstanceData.w;
+		const FInstanceInfo Info = ReadInstanceInfo(InstanceIndex);
 
 		// Index
 		Pos = Print(Pos, InstanceIndex, FontSilver, 2);
 		Pos = ShaderPrintText(Pos, TEXT(" "));
 
 		// Group
-		Pos = Print(Pos, GroupIndex, FontYellow, 1);
+		Pos = Print(Pos, Info.GroupIndex, FontYellow, 1);
 		Pos = ShaderPrintText(Pos, TEXT("/"));
-		Pos = Print(Pos, GroupCount, FontYellow, 1);
+		Pos = Print(Pos, Info.GroupCount, FontYellow, 1);
 		Pos = ShaderPrintText(Pos, TEXT("    "));
 
 		// LOD
-		Pos = Print(Pos, LOD, FontYellow, 4, 2);	
+		Pos = Print(Pos, Info.LOD, FontYellow, 4, 2);
 		Pos = ShaderPrintText(Pos, TEXT("/"));
-		Pos = Print(Pos, LODCount, FontYellow, 1);
+		Pos = Print(Pos, Info.LODCount, FontYellow, 1);
 		Pos = ShaderPrintText(Pos, TEXT("  "));
 
 		// Geometry
-		if (GeometryType == 0)		{ Pos = ShaderPrintText(Pos, TEXT("Strands "), FontEmerald); }
-		else if (GeometryType == 1) { Pos = ShaderPrintText(Pos, TEXT("Cards   "), FontEmerald); }
-		else if (GeometryType == 2) { Pos = ShaderPrintText(Pos, TEXT("Meshes  "), FontEmerald); }
+		if (Info.GeometryType == 0)		{ Pos = ShaderPrintText(Pos, TEXT("Strands "), FontEmerald); }
+		else if (Info.GeometryType == 1){ Pos = ShaderPrintText(Pos, TEXT("Cards   "), FontEmerald); }
+		else if (Info.GeometryType == 2){ Pos = ShaderPrintText(Pos, TEXT("Meshes  "), FontEmerald); }
 		Pos = ShaderPrintText(Pos, TEXT(" "));
 
 		// Binding
-		if (BindingType == 0)		{ Pos = ShaderPrintText(Pos, TEXT("Skinned "), FontTurquoise); }
-		else if (BindingType == 1)	{ Pos = ShaderPrintText(Pos, TEXT("Rigid   "), FontTurquoise); }
-		else if (BindingType == 2)	{ Pos = ShaderPrintText(Pos, TEXT("None    "), FontTurquoise); }
+		if (Info.BindingType == 0)		{ Pos = ShaderPrintText(Pos, TEXT("Skinned "), FontTurquoise); }
+		else if (Info.BindingType == 1)	{ Pos = ShaderPrintText(Pos, TEXT("Rigid   "), FontTurquoise); }
+		else if (Info.BindingType == 2)	{ Pos = ShaderPrintText(Pos, TEXT("None    "), FontTurquoise); }
 
 		// Sim
-		if (bHasSim)
+		if (Info.bHasSim)
 			Pos = ShaderPrintText(Pos, TEXT("True "), FontGreen);
 		else
 			Pos = ShaderPrintText(Pos, TEXT("False"), FontRed);
 		Pos = ShaderPrintText(Pos, TEXT(" "), FontWhite);
 
 		// RBF
-		if (bHasRBF)
+		if (Info.bHasRBF)
 			Pos = ShaderPrintText(Pos, TEXT("True "), FontGreen);
 		else
 			Pos = ShaderPrintText(Pos, TEXT("False"), FontRed);
 		Pos = ShaderPrintText(Pos, TEXT(" "));
 
 		// Clip
-		if (ClipValue >= 0)
-			Pos = Print(Pos, ClipValue, FontGreen, 7);
+		if (Info.ClipValue >= 0)
+			Pos = Print(Pos, Info.ClipValue, FontGreen, 7);
 		else 
 			Pos = Print(Pos, 0u, FontSilver, 7);
 		Pos = ShaderPrintText(Pos, TEXT(" "));
 
-		// Count
-		Pos = Print(Pos, StrandsCount, FontOrange, 8);
-		Pos = Print(Pos, VertexCount,  FontSilver, 8);
+		// Strands/Vertex sCount
+		Pos = Print(Pos, Info.StrandsCount, FontOrange, 8);
+		Pos = Print(Pos, Info.VertexCount,  FontSilver, 8);
+
+		// Binding data
+		Pos = Print(Pos, Info.UniqueTriangle, FontOrange, 8);
+		Pos = ShaderPrintText(Pos, TEXT("      "));
+		Pos = Print(Pos, Info.UniqueSection, FontSilver, 8);
+		Pos = ShaderPrintText(Pos, TEXT("     "));
 
 		// Name
 		Pos = PrintInstanceName(Pos, InstanceIndex, FontWhite);
diff --git a/Engine/Shaders/Private/HairStrands/HairStrandsFollicleMask.usf b/Engine/Shaders/Private/HairStrands/HairStrandsFollicleMask.usf
--- a/Engine/Shaders/Private/HairStrands/HairStrandsFollicleMask.usf
+++ b/Engine/Shaders/Private/HairStrands/HairStrandsFollicleMask.usf
@@ -7,6 +7,7 @@
 
 float2 OutputResolution;
 uint MaxRootCount;
+uint MaxUniqueTriangleIndex;
 uint Channel;
 uint KernelSizeInPixels;
 
@@ -14,9 +15,10 @@ uint KernelSizeInPixels;
 #define UV_ROOT 1
 
 #if PERMUTATION_UV_TYPE == UV_MESH
-Buffer<float4> TrianglePosition0Buffer;
-Buffer<float4> TrianglePosition1Buffer;
-Buffer<float4> TrianglePosition2Buffer;
+Buffer<float4> UniqueTrianglePosition0Buffer;
+Buffer<float4> UniqueTrianglePosition1Buffer;
+Buffer<float4> UniqueTrianglePosition2Buffer;
+Buffer<uint>   RootToUniqueTriangleIndexBuffer;
 Buffer<uint>   RootBarycentricBuffer; 
 #endif 
 
@@ -36,20 +38,27 @@ void MainVS(
 
 	ResolvedView = ResolveView();
 
+
+	const uint RootIndex = VertexId / 3;
 	const uint TriangleVertex = (VertexId % 3);
-	const uint TriangleIndex = VertexId / 3;
 
 #if PERMUTATION_UV_TYPE == UV_MESH
-	if (TriangleIndex > MaxRootCount) return;
-	const FHairMeshTriangleWithUVs Tri = GetTriangleWithUVsTransformation(TriangleIndex, TrianglePosition0Buffer, TrianglePosition1Buffer, TrianglePosition2Buffer);
+	uint UniqueTriangleIndex = 0;
+	if (RootIndex < MaxRootCount)
+	{
+		UniqueTriangleIndex = RootToUniqueTriangleIndexBuffer[RootIndex];
+	}
+
+	if (UniqueTriangleIndex > MaxUniqueTriangleIndex) return;
+	const FHairMeshTriangleWithUVs Tri = GetTriangleWithUVsTransformation(UniqueTriangleIndex, UniqueTrianglePosition0Buffer, UniqueTrianglePosition1Buffer, UniqueTrianglePosition2Buffer);
 
-	const uint EncodedBarycentric = RootBarycentricBuffer[TriangleIndex];
+	const uint EncodedBarycentric = RootBarycentricBuffer[RootIndex];
 	const float3 Barycentric = DecodeBarycentrics(EncodedBarycentric);
 	float2 RootUV = Barycentric.x * Tri.UV0 + Barycentric.y * Tri.UV1 + Barycentric.z * Tri.UV2;
 #endif
 
 #if PERMUTATION_UV_TYPE == UV_ROOT
-	float2 RootUV = RootUVsBuffer[TriangleIndex];
+	float2 RootUV = RootUVsBuffer[RootIndex];
 #endif
 
 	const float2 KernelRadius = float2(KernelSizeInPixels, KernelSizeInPixels);
diff --git a/Engine/Shaders/Private/HairStrands/HairStrandsGuideDeform.usf b/Engine/Shaders/Private/HairStrands/HairStrandsGuideDeform.usf
--- a/Engine/Shaders/Private/HairStrands/HairStrandsGuideDeform.usf
+++ b/Engine/Shaders/Private/HairStrands/HairStrandsGuideDeform.usf
@@ -38,6 +38,7 @@ Buffer<float4>  SimDeformedPosition2Buffer;
 
 Buffer<uint>	SimRootBarycentricBuffer;
 Buffer<uint>    SimVertexToRootIndexBuffer;
+Buffer<uint>	SimRootToUniqueTriangleIndexBuffer;
 
 // Return world position of a guide vertex deformed by a skinned triangle
 float3 DeformGuideByTriangle(
@@ -85,8 +86,9 @@ float3 DisplacePosition(uint VertexIndex, float3 Pos)
 	const uint RootIndex = SimVertexToRootIndexBuffer[VertexIndex];
 	const float3 RootBarycentric = DecodeBarycentrics(SimRootBarycentricBuffer[RootIndex]);
 
-	const FHairMeshTriangle RestTriangle = GetTriangleTransformation(RootIndex, SimRestPosition0Buffer, SimRestPosition1Buffer, SimRestPosition2Buffer);
-	const FHairMeshTriangle DeformedTriangle = GetTriangleTransformation(RootIndex, SimDeformedPosition0Buffer, SimDeformedPosition1Buffer, SimDeformedPosition2Buffer);
+	const uint TriangleIndex = SimRootToUniqueTriangleIndexBuffer[RootIndex];
+	const FHairMeshTriangle RestTriangle = GetTriangleTransformation(TriangleIndex, SimRestPosition0Buffer, SimRestPosition1Buffer, SimRestPosition2Buffer);
+	const FHairMeshTriangle DeformedTriangle = GetTriangleTransformation(TriangleIndex, SimDeformedPosition0Buffer, SimDeformedPosition1Buffer, SimDeformedPosition2Buffer);
 
 	return DeformGuideByTriangle(VertexIndex, RestTriangle, DeformedTriangle, RootBarycentric) - OutSimHairPositionOffset;
 #endif
diff --git a/Engine/Shaders/Private/HairStrands/HairStrandsInterpolation.usf b/Engine/Shaders/Private/HairStrands/HairStrandsInterpolation.usf
--- a/Engine/Shaders/Private/HairStrands/HairStrandsInterpolation.usf
+++ b/Engine/Shaders/Private/HairStrands/HairStrandsInterpolation.usf
@@ -95,7 +95,8 @@ Buffer<float4>  DeformedPosition0Buffer;
 Buffer<float4>  DeformedPosition1Buffer;
 Buffer<float4>  DeformedPosition2Buffer;
 
-Buffer<uint>	RootBarycentricBuffer;
+Buffer<uint>	RenRootBarycentricBuffer;
+Buffer<uint>	RenRootToUniqueTriangleIndexBuffer;
 
 Buffer<float4>  SimRestPosition0Buffer;
 Buffer<float4>  SimRestPosition1Buffer;
@@ -106,6 +107,7 @@ Buffer<float4>  SimDeformedPosition1Buffer;
 Buffer<float4>  SimDeformedPosition2Buffer;
 
 Buffer<uint>	SimRootBarycentricBuffer;
+Buffer<uint>	SimRootToUniqueTriangleIndexBuffer;
 #endif
 
 #if PERMUTATION_DEBUG == 1
@@ -205,6 +207,7 @@ uint4 WriteOutput(float3 WorldPosition, float3 OutPositionOffset, uint PackedDat
 struct FRootTriangleData
 {
 	uint RootIndex;
+	uint TriangleIndex;
 	float3 RootBarycentric;
 	FHairMeshTriangle RestTriangle;
 	FHairMeshTriangle DeformedTriangle;
@@ -217,16 +220,18 @@ FRootTriangleData GetRootTriangleData(uint VertexIndex, bool bSim)
 	if (bSim)
 	{
 		Out.RootIndex		 = SimVertexToRootIndexBuffer[VertexIndex];
+		Out.TriangleIndex	 = SimRootToUniqueTriangleIndexBuffer[Out.RootIndex];
 		Out.RootBarycentric  = DecodeBarycentrics(SimRootBarycentricBuffer[Out.RootIndex]);
-		Out.RestTriangle	 = GetTriangleTransformation(Out.RootIndex, SimRestPosition0Buffer, SimRestPosition1Buffer, SimRestPosition2Buffer);
-		Out.DeformedTriangle = GetTriangleTransformation(Out.RootIndex, SimDeformedPosition0Buffer, SimDeformedPosition1Buffer, SimDeformedPosition2Buffer);	
+		Out.RestTriangle	 = GetTriangleTransformation(Out.TriangleIndex, SimRestPosition0Buffer, SimRestPosition1Buffer, SimRestPosition2Buffer);
+		Out.DeformedTriangle = GetTriangleTransformation(Out.TriangleIndex, SimDeformedPosition0Buffer, SimDeformedPosition1Buffer, SimDeformedPosition2Buffer);
 	}
 	else
 	{
 		Out.RootIndex		 = RenVertexToRootIndexBuffer[VertexIndex];
-		Out.RootBarycentric  = DecodeBarycentrics(RootBarycentricBuffer[Out.RootIndex]);
-		Out.RestTriangle	 = GetTriangleTransformation(Out.RootIndex, RestPosition0Buffer, RestPosition1Buffer, RestPosition2Buffer);
-		Out.DeformedTriangle = GetTriangleTransformation(Out.RootIndex, DeformedPosition0Buffer, DeformedPosition1Buffer, DeformedPosition2Buffer);
+		Out.TriangleIndex	 = RenRootToUniqueTriangleIndexBuffer[Out.RootIndex];
+		Out.RootBarycentric  = DecodeBarycentrics(RenRootBarycentricBuffer[Out.RootIndex]);
+		Out.RestTriangle	 = GetTriangleTransformation(Out.TriangleIndex, RestPosition0Buffer, RestPosition1Buffer, RestPosition2Buffer);
+		Out.DeformedTriangle = GetTriangleTransformation(Out.TriangleIndex, DeformedPosition0Buffer, DeformedPosition1Buffer, DeformedPosition2Buffer);
 	}
 	#endif
 	return Out;
diff --git a/Engine/Shaders/Private/HairStrands/HairStrandsMesh.usf b/Engine/Shaders/Private/HairStrands/HairStrandsMesh.usf
--- a/Engine/Shaders/Private/HairStrands/HairStrandsMesh.usf
+++ b/Engine/Shaders/Private/HairStrands/HairStrandsMesh.usf
@@ -318,17 +318,17 @@ Buffer<float2> MeshUVsBuffer7;
 Buffer<float2> MeshUVsBuffer8;
 #endif
 
-uint MaxRootCount;
-Buffer<uint> RootTriangleIndex;
+uint MaxUniqueTriangleCount;
+Buffer<uint> UniqueTriangleIndices;
 
-RWBuffer<float4> OutRootTriangleCurrPosition0;
-RWBuffer<float4> OutRootTriangleCurrPosition1;
-RWBuffer<float4> OutRootTriangleCurrPosition2;
+RWBuffer<float4> OutUniqueTriangleCurrPosition0;
+RWBuffer<float4> OutUniqueTriangleCurrPosition1;
+RWBuffer<float4> OutUniqueTriangleCurrPosition2;
 
 #if PERMUTATION_PREVIOUS
-RWBuffer<float4> OutRootTrianglePrevPosition0;
-RWBuffer<float4> OutRootTrianglePrevPosition1;
-RWBuffer<float4> OutRootTrianglePrevPosition2;
+RWBuffer<float4> OutUniqueTrianglePrevPosition0;
+RWBuffer<float4> OutUniqueTrianglePrevPosition1;
+RWBuffer<float4> OutUniqueTrianglePrevPosition2;
 #endif
 
 #if PERMUTATION_WITHUV == 1
@@ -374,11 +374,11 @@ FHairMeshTriangle GetTriangleData(FHairTriangleIndex TriangleIndex, uint Section
 [numthreads(128, 1, 1)]
 void MainCS(uint3 DispatchThreadId : SV_DispatchThreadID)
 { 
-	const uint RootIndex = DispatchThreadId.x;
-	if (RootIndex >= MaxRootCount) 
+	const uint UniqueTriangleIndex = DispatchThreadId.x;
+	if (UniqueTriangleIndex >= MaxUniqueTriangleCount)
 		return;
 	
-	const FHairTriangleIndex TriangleIndex = DecodeTriangleIndex(RootTriangleIndex[RootIndex]);
+	const FHairTriangleIndex TriangleIndex = DecodeTriangleIndex(UniqueTriangleIndices[UniqueTriangleIndex]);
 	if (TriangleIndex.SectionIndex >= MaxSectionCount)
 		return;
 
@@ -403,33 +403,33 @@ void MainCS(uint3 DispatchThreadId : SV_DispatchThreadID)
 #if PERMUTATION_WITHUV == 1
 	{
 		const FHairMeshTriangleWithUVs Tri = GetTriangleData(TriangleIndex, SectionSlotIndex, false);
-		OutRootTriangleCurrPosition0[RootIndex] = float4(Tri.P0, PackUVsToFloat(Tri.UV0));
-		OutRootTriangleCurrPosition1[RootIndex] = float4(Tri.P1, PackUVsToFloat(Tri.UV1));
-		OutRootTriangleCurrPosition2[RootIndex] = float4(Tri.P2, PackUVsToFloat(Tri.UV2));
+		OutUniqueTriangleCurrPosition0[UniqueTriangleIndex] = float4(Tri.P0, PackUVsToFloat(Tri.UV0));
+		OutUniqueTriangleCurrPosition1[UniqueTriangleIndex] = float4(Tri.P1, PackUVsToFloat(Tri.UV1));
+		OutUniqueTriangleCurrPosition2[UniqueTriangleIndex] = float4(Tri.P2, PackUVsToFloat(Tri.UV2));
 	}
 
 	#if PERMUTATION_PREVIOUS
 	{
 		const FHairMeshTriangleWithUVs Tri = GetTriangleData(TriangleIndex, SectionSlotIndex, true);
-		OutRootTrianglePrevPosition0[RootIndex] = float4(Tri.P0, PackUVsToFloat(Tri.UV0));
-		OutRootTrianglePrevPosition1[RootIndex] = float4(Tri.P1, PackUVsToFloat(Tri.UV1));
-		OutRootTrianglePrevPosition2[RootIndex] = float4(Tri.P2, PackUVsToFloat(Tri.UV2));
+		OutUniqueTrianglePrevPosition0[UniqueTriangleIndex] = float4(Tri.P0, PackUVsToFloat(Tri.UV0));
+		OutUniqueTrianglePrevPosition1[UniqueTriangleIndex] = float4(Tri.P1, PackUVsToFloat(Tri.UV1));
+		OutUniqueTrianglePrevPosition2[UniqueTriangleIndex] = float4(Tri.P2, PackUVsToFloat(Tri.UV2));
 	}
 	#endif
 #else
 	{
 		const FHairMeshTriangle Tri = GetTriangleData(TriangleIndex, SectionSlotIndex, false);
-		OutRootTriangleCurrPosition0[RootIndex] = (float4(Tri.P0, 1));
-		OutRootTriangleCurrPosition1[RootIndex] = (float4(Tri.P1, 1));
-		OutRootTriangleCurrPosition2[RootIndex] = (float4(Tri.P2, 1));
+		OutUniqueTriangleCurrPosition0[UniqueTriangleIndex] = (float4(Tri.P0, 1));
+		OutUniqueTriangleCurrPosition1[UniqueTriangleIndex] = (float4(Tri.P1, 1));
+		OutUniqueTriangleCurrPosition2[UniqueTriangleIndex] = (float4(Tri.P2, 1));
 	}
 
 	#if PERMUTATION_PREVIOUS
 	{
 		const FHairMeshTriangle Tri = GetTriangleData(TriangleIndex, SectionSlotIndex, true);
-		OutRootTrianglePrevPosition0[RootIndex] = (float4(Tri.P0, 1));
-		OutRootTrianglePrevPosition1[RootIndex] = (float4(Tri.P1, 1));
-		OutRootTrianglePrevPosition2[RootIndex] = (float4(Tri.P2, 1));
+		OutUniqueTrianglePrevPosition0[UniqueTriangleIndex] = (float4(Tri.P0, 1));
+		OutUniqueTrianglePrevPosition1[UniqueTriangleIndex] = (float4(Tri.P1, 1));
+		OutUniqueTrianglePrevPosition2[UniqueTriangleIndex] = (float4(Tri.P2, 1));
 	}
 	#endif
 #endif
diff --git a/Engine/Shaders/Private/HairStrands/HairStrandsMeshProjectionCommon.ush b/Engine/Shaders/Private/HairStrands/HairStrandsMeshProjectionCommon.ush
--- a/Engine/Shaders/Private/HairStrands/HairStrandsMeshProjectionCommon.ush
+++ b/Engine/Shaders/Private/HairStrands/HairStrandsMeshProjectionCommon.ush
@@ -106,7 +106,7 @@ FHairMeshTriangleWithUVs GetTriangleData(uint TriangleIndex, Buffer<float> InPos
 }
 
 FHairMeshTriangle GetTriangleTransformation(
-	uint RootIndex,
+	uint TriangleIndex,
 	Buffer<float4>  InPosition0Buffer,
 	Buffer<float4>  InPosition1Buffer,
 	Buffer<float4>  InPosition2Buffer)
@@ -114,9 +114,9 @@ FHairMeshTriangle GetTriangleTransformation(
 	FHairMeshTriangle Out;
 
 	// Build orthonormal frame
-	Out.P0 = InPosition0Buffer[RootIndex].xyz;
-	Out.P1 = InPosition1Buffer[RootIndex].xyz;
-	Out.P2 = InPosition2Buffer[RootIndex].xyz;
+	Out.P0 = InPosition0Buffer[TriangleIndex].xyz;
+	Out.P1 = InPosition1Buffer[TriangleIndex].xyz;
+	Out.P2 = InPosition2Buffer[TriangleIndex].xyz;
 	Out.E0 = Out.P1 - Out.P0;
 	Out.E1 = Out.P2 - Out.P0;
 
@@ -128,16 +128,16 @@ FHairMeshTriangle GetTriangleTransformation(
 }
 
 FHairMeshTriangleWithUVs GetTriangleWithUVsTransformation(
-	uint RootIndex,
+	uint TriangleIndex,
 	Buffer<float4>  InPosition0Buffer,
 	Buffer<float4>  InPosition1Buffer,
 	Buffer<float4>  InPosition2Buffer)
 {
 	FHairMeshTriangleWithUVs Out;
 
-	const float4 V0 = InPosition0Buffer[RootIndex];
-	const float4 V1 = InPosition1Buffer[RootIndex];
-	const float4 V2 = InPosition2Buffer[RootIndex];
+	const float4 V0 = InPosition0Buffer[TriangleIndex];
+	const float4 V1 = InPosition1Buffer[TriangleIndex];
+	const float4 V2 = InPosition2Buffer[TriangleIndex];
 
 	// Build orthonormal frame
 	Out.P0 = V0.xyz;
diff --git a/Engine/Source/Runtime/Core/Private/UObject/UE5MainStreamObjectVersion.cpp b/Engine/Source/Runtime/Core/Private/UObject/UE5MainStreamObjectVersion.cpp
--- a/Engine/Source/Runtime/Core/Private/UObject/UE5MainStreamObjectVersion.cpp
+++ b/Engine/Source/Runtime/Core/Private/UObject/UE5MainStreamObjectVersion.cpp
@@ -14,7 +14,7 @@ TMap<FGuid, FGuid> FUE5MainStreamObjectVersion::GetSystemGuids()
 	SystemGuids.Add(DevGuids.MATERIALSHADERMAP_DERIVEDDATA_VER, FGuid("65C01D817C9A4EEFAE9E988D41A1F3DD"));
 	SystemGuids.Add(DevGuids.NANITE_DERIVEDDATA_VER, FGuid("7A3EA93546624597BEBE2ADFDE3A7628"));
 	SystemGuids.Add(DevGuids.NIAGARASHADERMAP_DERIVEDDATA_VER, FGuid("8A37C45D24F2423CBE5F8F371DE33575"));
-	SystemGuids.Add(DevGuids.Niagara_LatestScriptCompileVersion, FGuid("92519728F82ECC4FB1CD175F35D6442D"));
+	SystemGuids.Add(DevGuids.Niagara_LatestScriptCompileVersion, FGuid("ECBF3E4E3F07478998D0D5D44AD86807"));
 	SystemGuids.Add(DevGuids.SkeletalMeshDerivedDataVersion, FGuid("FB7B2FB546DE49B2A2F88B7F69763CEE"));
 	SystemGuids.Add(DevGuids.STATICMESH_DERIVEDDATA_VER, FGuid("98F7D79A5811013825E75BBCC41ED3E9"));
 	SystemGuids.Add(DevGuids.POSESEARCHDB_DERIVEDDATA_VER, FGuid("4E595C2AC5E947D6BA9ABC874353E5BC"));
