Rename root triangle format. (Decouple unique triangle & root index format).

diff --git a/Engine/Plugins/Runtime/HairStrands/Source/HairStrandsCore/Private/GroomBindingBuilder.cpp b/Engine/Plugins/Runtime/HairStrands/Source/HairStrandsCore/Private/GroomBindingBuilder.cpp
--- a/Engine/Plugins/Runtime/HairStrands/Source/HairStrandsCore/Private/GroomBindingBuilder.cpp
+++ b/Engine/Plugins/Runtime/HairStrands/Source/HairStrandsCore/Private/GroomBindingBuilder.cpp
@@ -1224,7 +1224,7 @@ namespace GroomBinding_RootProjection
 			// 2.3. Compute the closest triangle for each root
 			//InMeshRenderData->LODRenderData[LODIt].GetNumVertices();
 
-			TArray<FHairStrandsCurveTriangleIndexFormat::Type> RootTriangleIndexBuffer;
+			TArray<FHairStrandsUniqueTriangleIndexFormat::Type> RootTriangleIndexBuffer;
 			RootTriangleIndexBuffer.SetNum(CurveCount);
 
 			TArray<FHairStrandsMeshTrianglePositionFormat::Type> RestRootTrianglePosition0Buffer;
@@ -1920,9 +1920,9 @@ static void BuildRootBulkData(
 		Out.MeshProjectionLODs[MeshLODIt].SampleCount = bHasValidSamples ? In.MeshProjectionLODs[MeshLODIt].SampleCount : 0u;
 		Out.MeshProjectionLODs[MeshLODIt].UniqueTriangleCount = In.MeshProjectionLODs[MeshLODIt].RestUniqueTrianglePosition0Buffer.Num();
 
-		CopyToBulkData<FHairStrandsCurveTriangleIndexFormat>(Out.MeshProjectionLODs[MeshLODIt].UniqueTriangleIndexBuffer, In.MeshProjectionLODs[MeshLODIt].UniqueTriangleIndexBuffer);
-		CopyToBulkData<FHairStrandsCurveTriangleBarycentricFormat>(Out.MeshProjectionLODs[MeshLODIt].RootBarycentricBuffer, In.MeshProjectionLODs[MeshLODIt].RootBarycentricBuffer);
-		CopyToBulkData<FHairStrandsCurveTriangleIndexFormat>(Out.MeshProjectionLODs[MeshLODIt].RootToUniqueTriangleIndexBuffer, In.MeshProjectionLODs[MeshLODIt].RootToUniqueTriangleIndexBuffer);
+		CopyToBulkData<FHairStrandsUniqueTriangleIndexFormat>(Out.MeshProjectionLODs[MeshLODIt].UniqueTriangleIndexBuffer, In.MeshProjectionLODs[MeshLODIt].UniqueTriangleIndexBuffer);
+		CopyToBulkData<FHairStrandsRootBarycentricFormat>(Out.MeshProjectionLODs[MeshLODIt].RootBarycentricBuffer, In.MeshProjectionLODs[MeshLODIt].RootBarycentricBuffer);
+		CopyToBulkData<FHairStrandsRootToUniqueTriangleIndexFormat>(Out.MeshProjectionLODs[MeshLODIt].RootToUniqueTriangleIndexBuffer, In.MeshProjectionLODs[MeshLODIt].RootToUniqueTriangleIndexBuffer);
 
 		CopyToBulkData<FHairStrandsMeshTrianglePositionFormat>(Out.MeshProjectionLODs[MeshLODIt].RestUniqueTrianglePosition0Buffer, In.MeshProjectionLODs[MeshLODIt].RestUniqueTrianglePosition0Buffer);
 		CopyToBulkData<FHairStrandsMeshTrianglePositionFormat>(Out.MeshProjectionLODs[MeshLODIt].RestUniqueTrianglePosition1Buffer, In.MeshProjectionLODs[MeshLODIt].RestUniqueTrianglePosition1Buffer);
@@ -1964,9 +1964,9 @@ static void BuildRootData(
 		Out.MeshProjectionLODs[MeshLODIt].LODIndex = In.MeshProjectionLODs[MeshLODIt].LODIndex;
 		Out.MeshProjectionLODs[MeshLODIt].SampleCount = bHasValidSamples ? In.MeshProjectionLODs[MeshLODIt].SampleCount : 0u;
 
-		CopyFromBulkData<FHairStrandsCurveTriangleIndexFormat>(Out.MeshProjectionLODs[MeshLODIt].UniqueTriangleIndexBuffer, In.MeshProjectionLODs[MeshLODIt].UniqueTriangleIndexBuffer);
-		CopyFromBulkData<FHairStrandsCurveTriangleIndexFormat>(Out.MeshProjectionLODs[MeshLODIt].RootToUniqueTriangleIndexBuffer, In.MeshProjectionLODs[MeshLODIt].RootToUniqueTriangleIndexBuffer);
-		CopyFromBulkData<FHairStrandsCurveTriangleBarycentricFormat>(Out.MeshProjectionLODs[MeshLODIt].RootBarycentricBuffer, In.MeshProjectionLODs[MeshLODIt].RootBarycentricBuffer);
+		CopyFromBulkData<FHairStrandsUniqueTriangleIndexFormat>(Out.MeshProjectionLODs[MeshLODIt].UniqueTriangleIndexBuffer, In.MeshProjectionLODs[MeshLODIt].UniqueTriangleIndexBuffer);
+		CopyFromBulkData<FHairStrandsRootToUniqueTriangleIndexFormat>(Out.MeshProjectionLODs[MeshLODIt].RootToUniqueTriangleIndexBuffer, In.MeshProjectionLODs[MeshLODIt].RootToUniqueTriangleIndexBuffer);
+		CopyFromBulkData<FHairStrandsRootBarycentricFormat>(Out.MeshProjectionLODs[MeshLODIt].RootBarycentricBuffer, In.MeshProjectionLODs[MeshLODIt].RootBarycentricBuffer);
 
 		CopyFromBulkData<FHairStrandsMeshTrianglePositionFormat>(Out.MeshProjectionLODs[MeshLODIt].RestUniqueTrianglePosition0Buffer, In.MeshProjectionLODs[MeshLODIt].RestUniqueTrianglePosition0Buffer);
 		CopyFromBulkData<FHairStrandsMeshTrianglePositionFormat>(Out.MeshProjectionLODs[MeshLODIt].RestUniqueTrianglePosition1Buffer, In.MeshProjectionLODs[MeshLODIt].RestUniqueTrianglePosition1Buffer);
diff --git a/Engine/Plugins/Runtime/HairStrands/Source/HairStrandsCore/Private/GroomRBFDeformer.cpp b/Engine/Plugins/Runtime/HairStrands/Source/HairStrandsCore/Private/GroomRBFDeformer.cpp
--- a/Engine/Plugins/Runtime/HairStrands/Source/HairStrandsCore/Private/GroomRBFDeformer.cpp
+++ b/Engine/Plugins/Runtime/HairStrands/Source/HairStrandsCore/Private/GroomRBFDeformer.cpp
@@ -83,7 +83,7 @@ void DeformStrands(
 	const FHairStrandsDatas& HairStandsData,
 	const TArray<FHairStrandsIndexFormat::Type>& VertexToCurveIndexBuffer,
 	const TArray<FHairStrandsIndexFormat::Type>& RootToUniqueTriangleBuffer,
-	const TArray<FHairStrandsCurveTriangleBarycentricFormat::Type>& RootBarycentricBuffer,
+	const TArray<FHairStrandsRootBarycentricFormat::Type>& RootBarycentricBuffer,
 
 	const TArray<FHairStrandsMeshTrianglePositionFormat::Type>& UniqueTrianglePosition0Buffer_Rest,
 	const TArray<FHairStrandsMeshTrianglePositionFormat::Type>& UniqueTrianglePosition1Buffer_Rest,
diff --git a/Engine/Plugins/Runtime/HairStrands/Source/HairStrandsCore/Private/GroomResources.cpp b/Engine/Plugins/Runtime/HairStrands/Source/HairStrandsCore/Private/GroomResources.cpp
--- a/Engine/Plugins/Runtime/HairStrands/Source/HairStrandsCore/Private/GroomResources.cpp
+++ b/Engine/Plugins/Runtime/HairStrands/Source/HairStrandsCore/Private/GroomResources.cpp
@@ -1149,9 +1149,9 @@ void FHairStrandsRestRootResource::InternalAllocateLOD(FRDGBuilder& GraphBuilder
 		{
 			GPUData.Status = FLOD::EStatus::Completed;
 			
-			InternalCreateVertexBufferRDG_FromBulkData<FHairStrandsCurveTriangleBarycentricFormat>(GraphBuilder, CPUData.RootBarycentricBuffer, BulkData.RootCount, GPUData.RootBarycentricBuffer, ToHairResourceDebugName(HAIRSTRANDS_RESOUCE_NAME(CurveType, Hair.StrandsRestRoot_RootTriangleBarycentricBuffer), ResourceName), EHairResourceUsageType::Static);
-			InternalCreateVertexBufferRDG_FromBulkData<FHairStrandsCurveTriangleIndexFormat>(GraphBuilder, CPUData.RootToUniqueTriangleIndexBuffer, BulkData.RootCount, GPUData.RootToUniqueTriangleIndexBuffer, ToHairResourceDebugName(HAIRSTRANDS_RESOUCE_NAME(CurveType, Hair.StrandsRestRoot_RootToUniqueTriangleIndexBuffer), ResourceName), EHairResourceUsageType::Static);
-			InternalCreateVertexBufferRDG_FromBulkData<FHairStrandsCurveTriangleIndexFormat>(GraphBuilder, CPUData.UniqueTriangleIndexBuffer, CPUData.UniqueTriangleCount, GPUData.UniqueTriangleIndexBuffer, ToHairResourceDebugName(HAIRSTRANDS_RESOUCE_NAME(CurveType, Hair.StrandsRestRoot_UniqueTriangleIndexBuffer), ResourceName), EHairResourceUsageType::Static);
+			InternalCreateVertexBufferRDG_FromBulkData<FHairStrandsRootBarycentricFormat>(GraphBuilder, CPUData.RootBarycentricBuffer, BulkData.RootCount, GPUData.RootBarycentricBuffer, ToHairResourceDebugName(HAIRSTRANDS_RESOUCE_NAME(CurveType, Hair.StrandsRestRoot_RootTriangleBarycentricBuffer), ResourceName), EHairResourceUsageType::Static);
+			InternalCreateVertexBufferRDG_FromBulkData<FHairStrandsRootToUniqueTriangleIndexFormat>(GraphBuilder, CPUData.RootToUniqueTriangleIndexBuffer, BulkData.RootCount, GPUData.RootToUniqueTriangleIndexBuffer, ToHairResourceDebugName(HAIRSTRANDS_RESOUCE_NAME(CurveType, Hair.StrandsRestRoot_RootToUniqueTriangleIndexBuffer), ResourceName), EHairResourceUsageType::Static);
+			InternalCreateVertexBufferRDG_FromBulkData<FHairStrandsUniqueTriangleIndexFormat>(GraphBuilder, CPUData.UniqueTriangleIndexBuffer, CPUData.UniqueTriangleCount, GPUData.UniqueTriangleIndexBuffer, ToHairResourceDebugName(HAIRSTRANDS_RESOUCE_NAME(CurveType, Hair.StrandsRestRoot_UniqueTriangleIndexBuffer), ResourceName), EHairResourceUsageType::Static);
 
 			InternalCreateVertexBufferRDG_FromBulkData<FHairStrandsMeshTrianglePositionFormat>(GraphBuilder, CPUData.RestUniqueTrianglePosition0Buffer, CPUData.UniqueTriangleCount, GPUData.RestUniqueTrianglePosition0Buffer, ToHairResourceDebugName(HAIRSTRANDS_RESOUCE_NAME(CurveType, Hair.StrandsRestRoot_RestUniqueTrianglePosition0Buffer), ResourceName), EHairResourceUsageType::Static);
 			InternalCreateVertexBufferRDG_FromBulkData<FHairStrandsMeshTrianglePositionFormat>(GraphBuilder, CPUData.RestUniqueTrianglePosition1Buffer, CPUData.UniqueTriangleCount, GPUData.RestUniqueTrianglePosition1Buffer, ToHairResourceDebugName(HAIRSTRANDS_RESOUCE_NAME(CurveType, Hair.StrandsRestRoot_RestUniqueTrianglePosition1Buffer), ResourceName), EHairResourceUsageType::Static);
@@ -1161,9 +1161,9 @@ void FHairStrandsRestRootResource::InternalAllocateLOD(FRDGBuilder& GraphBuilder
 		{
 			GPUData.Status = FLOD::EStatus::Initialized;
 			
-			InternalCreateVertexBufferRDG<FHairStrandsCurveTriangleBarycentricFormat>(GraphBuilder, BulkData.RootCount, GPUData.RootBarycentricBuffer, ToHairResourceDebugName(HAIRSTRANDS_RESOUCE_NAME(CurveType, Hair.StrandsRestRoot_RootBarycentricBuffer), ResourceName), EHairResourceUsageType::Dynamic);
-			InternalCreateVertexBufferRDG<FHairStrandsCurveTriangleIndexFormat>(GraphBuilder, BulkData.RootCount, GPUData.RootToUniqueTriangleIndexBuffer, ToHairResourceDebugName(HAIRSTRANDS_RESOUCE_NAME(CurveType, Hair.StrandsRestRoot_RootToUniqueTriangleIndexBuffer), ResourceName), EHairResourceUsageType::Dynamic);
-			InternalCreateVertexBufferRDG<FHairStrandsCurveTriangleIndexFormat>(GraphBuilder, CPUData.UniqueTriangleCount, GPUData.UniqueTriangleIndexBuffer, ToHairResourceDebugName(HAIRSTRANDS_RESOUCE_NAME(CurveType, Hair.StrandsRestRoot_UniqueTriangleIndexBuffer), ResourceName), EHairResourceUsageType::Dynamic);
+			InternalCreateVertexBufferRDG<FHairStrandsRootBarycentricFormat>(GraphBuilder, BulkData.RootCount, GPUData.RootBarycentricBuffer, ToHairResourceDebugName(HAIRSTRANDS_RESOUCE_NAME(CurveType, Hair.StrandsRestRoot_RootBarycentricBuffer), ResourceName), EHairResourceUsageType::Dynamic);
+			InternalCreateVertexBufferRDG<FHairStrandsRootToUniqueTriangleIndexFormat>(GraphBuilder, BulkData.RootCount, GPUData.RootToUniqueTriangleIndexBuffer, ToHairResourceDebugName(HAIRSTRANDS_RESOUCE_NAME(CurveType, Hair.StrandsRestRoot_RootToUniqueTriangleIndexBuffer), ResourceName), EHairResourceUsageType::Dynamic);
+			InternalCreateVertexBufferRDG<FHairStrandsUniqueTriangleIndexFormat>(GraphBuilder, CPUData.UniqueTriangleCount, GPUData.UniqueTriangleIndexBuffer, ToHairResourceDebugName(HAIRSTRANDS_RESOUCE_NAME(CurveType, Hair.StrandsRestRoot_UniqueTriangleIndexBuffer), ResourceName), EHairResourceUsageType::Dynamic);
 			
 			// Create buffers. Initialization will be done by render passes
 			InternalCreateVertexBufferRDG<FHairStrandsMeshTrianglePositionFormat>(GraphBuilder, CPUData.UniqueTriangleCount, GPUData.RestUniqueTrianglePosition0Buffer, ToHairResourceDebugName(HAIRSTRANDS_RESOUCE_NAME(CurveType, Hair.StrandsRestRoot_RestUniqueTrianglePosition0Buffer), ResourceName), EHairResourceUsageType::Dynamic);
diff --git a/Engine/Plugins/Runtime/HairStrands/Source/HairStrandsCore/Public/HairStrandsDatas.h b/Engine/Plugins/Runtime/HairStrands/Source/HairStrandsCore/Public/HairStrandsDatas.h
--- a/Engine/Plugins/Runtime/HairStrands/Source/HairStrandsCore/Public/HairStrandsDatas.h
+++ b/Engine/Plugins/Runtime/HairStrands/Source/HairStrandsCore/Public/HairStrandsDatas.h
@@ -243,7 +243,8 @@ struct FHairStrandsMeshTrianglePositionFormat
 	static const EPixelFormat Format = PF_A32B32G32R32F;
 };
 
-struct FHairStrandsCurveTriangleIndexFormat
+// Encode Section ID and triangle Index from the source skel. mesh
+struct FHairStrandsUniqueTriangleIndexFormat
 {
 	using Type = uint32;
 	using BulkType = uint32;
@@ -253,7 +254,17 @@ struct FHairStrandsCurveTriangleIndexFormat
 	static const EPixelFormat Format = PF_R32_UINT;
 };
 
-struct FHairStrandsCurveTriangleBarycentricFormat
+struct FHairStrandsRootToUniqueTriangleIndexFormat
+{
+	using Type = uint32;
+	using BulkType = uint32;
+	static const uint32 ComponentCount = 1;
+	static const uint32 SizeInByte = sizeof(Type);
+	static const EVertexElementType VertexElementType = VET_UInt;
+	static const EPixelFormat Format = PF_R32_UINT;
+};
+
+struct FHairStrandsRootBarycentricFormat
 {
 	using Type = uint32;
 	using BulkType = uint32;
@@ -593,11 +604,11 @@ struct FHairStrandsRootData
 		/* Triangle on which a root is attached */
 		/* When the projection is done with source to target mesh transfer, the projection indices does not match.
 		   In this case we need to separate index computation. The barycentric coords remain the same however. */
-		TArray<uint32> RootToUniqueTriangleIndexBuffer;
-		TArray<FHairStrandsCurveTriangleBarycentricFormat::Type> RootBarycentricBuffer;
+		TArray<FHairStrandsRootToUniqueTriangleIndexFormat::Type> RootToUniqueTriangleIndexBuffer;
+		TArray<FHairStrandsRootBarycentricFormat::Type> RootBarycentricBuffer;
 
 		/* Strand hair roots translation and rotation in rest position relative to the bound triangle. Positions are relative to the rest root center */
-		TArray<FHairStrandsCurveTriangleIndexFormat::Type> UniqueTriangleIndexBuffer;
+		TArray<FHairStrandsUniqueTriangleIndexFormat::Type> UniqueTriangleIndexBuffer;
 		TArray<FHairStrandsMeshTrianglePositionFormat::Type> RestUniqueTrianglePosition0Buffer;
 		TArray<FHairStrandsMeshTrianglePositionFormat::Type> RestUniqueTrianglePosition1Buffer;
 		TArray<FHairStrandsMeshTrianglePositionFormat::Type> RestUniqueTrianglePosition2Buffer;
