Remove unecessary triangle evaluation when fetching skin-cache/manual-cache data during LOD level selection for groom.

diff --git a/Engine/Plugins/Runtime/HairStrands/Source/HairStrandsCore/Private/GroomDebug.cpp b/Engine/Plugins/Runtime/HairStrands/Source/HairStrandsCore/Private/GroomDebug.cpp
--- a/Engine/Plugins/Runtime/HairStrands/Source/HairStrandsCore/Private/GroomDebug.cpp
+++ b/Engine/Plugins/Runtime/HairStrands/Source/HairStrandsCore/Private/GroomDebug.cpp
@@ -69,6 +69,13 @@ static FAutoConsoleVariableRef CVarHairStrandsControlPointDebug(TEXT("r.HairStra
 
 bool IsHairStrandsSkinCacheEnable();
 
+FCachedGeometry GetCacheGeometryForHair(
+	FRDGBuilder& GraphBuilder,
+	FHairGroupInstance* Instance,
+	const FGPUSkinCache* SkinCache,
+	FGlobalShaderMap* ShaderMap,
+	const bool bOutputTriangleData);
+
 static void GetGroomInterpolationData(
 	FRDGBuilder& GraphBuilder,
 	FGlobalShaderMap* ShaderMap,
@@ -84,26 +91,7 @@ static void GetGroomInterpolationData(
 		if (!Instance || !Instance->Debug.MeshComponent)
 			continue;
 
-		FCachedGeometry CachedGeometry;
-		if (Instance->Debug.GroomBindingType == EGroomBindingMeshType::SkeletalMesh)
-		{
-			if (USkeletalMeshComponent* SkeletalMeshComponent = Cast<USkeletalMeshComponent>(Instance->Debug.MeshComponent))
-			{
-				if (SkinCache)
-				{
-					CachedGeometry = SkinCache->GetCachedGeometry(SkeletalMeshComponent->ComponentId.PrimIDValue, EGPUSkinCacheEntryMode::Raster);
-				}
-
-				if (IsHairStrandsSkinCacheEnable() && CachedGeometry.Sections.Num() == 0)
-				{
-					BuildCacheGeometry(GraphBuilder, ShaderMap, SkeletalMeshComponent, CachedGeometry);
-				}
-			}
-		}
-		else if (UGeometryCacheComponent* GeometryCacheComponent = Cast<UGeometryCacheComponent>(Instance->Debug.MeshComponent))
-		{
-			BuildCacheGeometry(GraphBuilder, ShaderMap, GeometryCacheComponent, CachedGeometry);
-		}
+		const FCachedGeometry CachedGeometry = GetCacheGeometryForHair(GraphBuilder, Instance, SkinCache,ShaderMap, true);
 		if (CachedGeometry.Sections.Num() == 0)
 			continue;
 
@@ -815,13 +803,6 @@ static const TCHAR* ToString(EHairLODSelectionType In)
 	return TEXT("None");
 }
 
-FCachedGeometry GetCacheGeometryForHair(
-	FRDGBuilder& GraphBuilder,
-	FHairGroupInstance* Instance,
-	const FGPUSkinCache* SkinCache,
-	FGlobalShaderMap* ShaderMap);
-
-
 ///////////////////////////////////////////////////////////////////////////////////////////////////
 
 class FHairDebugPrintInstanceCS : public FGlobalShader
@@ -1198,4 +1179,4 @@ void RunHairStrandsDebug(
 			AddDrawDebugStrandsCVsPass(GraphBuilder, View, ShaderMap, Instance, ShaderPrintData, SceneColorTexture, SceneDepthTexture);
 		}
 	}
-}
\ No newline at end of file
+}
diff --git a/Engine/Plugins/Runtime/HairStrands/Source/HairStrandsCore/Private/GroomGeometryCache.cpp b/Engine/Plugins/Runtime/HairStrands/Source/HairStrandsCore/Private/GroomGeometryCache.cpp
--- a/Engine/Plugins/Runtime/HairStrands/Source/HairStrandsCore/Private/GroomGeometryCache.cpp
+++ b/Engine/Plugins/Runtime/HairStrands/Source/HairStrandsCore/Private/GroomGeometryCache.cpp
@@ -42,7 +42,7 @@ FHairStrandsProjectionMeshData::Section ConvertMeshSection(const FCachedGeometry
 
 ////////////////////////////////////////////////////////////////////////////////////////////////////////////////////////
 
-static void BuildBoneMatrices(USkeletalMeshComponent* SkeletalMeshComponent, const FSkeletalMeshLODRenderData& LODData,
+static void BuildBoneMatrices(const USkeletalMeshComponent* SkeletalMeshComponent, const FSkeletalMeshLODRenderData& LODData,
 	const uint32 LODIndex, TArray<uint32>& MatrixOffsets, TArray<FVector4f>& BoneMatrices)
 {
 	TArray<FMatrix44f> BoneTransforms;
@@ -72,60 +72,63 @@ static void BuildBoneMatrices(USkeletalMeshComponent* SkeletalMeshComponent, con
 }
 
  void BuildCacheGeometry(
-	 FRDGBuilder& GraphBuilder,
+	FRDGBuilder& GraphBuilder,
 	FGlobalShaderMap* ShaderMap, 
-	USkeletalMeshComponent* SkeletalMeshComponent, 
-	FCachedGeometry& CachedGeometry)
+	const USkeletalMeshComponent* SkeletalMeshComponent, 
+	const bool bOutputTriangleData,
+	FCachedGeometry& Out)
 {
 	if (SkeletalMeshComponent)
 	{
-		FSkeletalMeshRenderData* RenderData = SkeletalMeshComponent->GetSkeletalMesh()->GetResourceForRendering();
-
 		const uint32 LODIndex = SkeletalMeshComponent->GetPredictedLODLevel();// RenderData->PendingFirstLODIdx;
-		FSkeletalMeshLODRenderData& LODData = RenderData->LODRenderData[LODIndex];
-
-		TArray<uint32> MatrixOffsets;
-		TArray<FVector4f> BoneMatrices;
-		BuildBoneMatrices(SkeletalMeshComponent, LODData, LODIndex, MatrixOffsets, BoneMatrices);
-
-		const bool bNeedPreviousPosition = IsHairStrandContinuousDecimationReorderingEnabled();
-
-		FRDGBufferRef DeformedPositionsBuffer			= GraphBuilder.CreateBuffer(FRDGBufferDesc::CreateBufferDesc(sizeof(float), LODData.StaticVertexBuffers.PositionVertexBuffer.GetNumVertices() * 3), TEXT("Hair.SkinnedDeformedPositions"));
-		FRDGBufferRef DeformedPreviousPositionsBuffer	= bNeedPreviousPosition ? GraphBuilder.CreateBuffer(FRDGBufferDesc::CreateBufferDesc(sizeof(float), LODData.StaticVertexBuffers.PositionVertexBuffer.GetNumVertices() * 3), TEXT("Hair.SkinnedDeformedPreviousPositions")) : nullptr;
-		FRDGBufferRef BoneMatricesBuffer				= CreateStructuredBuffer(GraphBuilder, TEXT("Hair.SkinnedBoneMatrices"), sizeof(float) * 4, BoneMatrices.Num(), BoneMatrices.GetData(), sizeof(float) * 4 * BoneMatrices.Num());
-		FRDGBufferRef MatrixOffsetsBuffer				= CreateStructuredBuffer(GraphBuilder, TEXT("Hair.SkinnedMatrixOffsets"), sizeof(uint32), MatrixOffsets.Num(), MatrixOffsets.GetData(), sizeof(uint32) * MatrixOffsets.Num());
-
-		// TODO previous bone matrices
-		AddSkinUpdatePass(GraphBuilder, ShaderMap, SkeletalMeshComponent->GetSkinWeightBuffer(LODIndex), LODData, BoneMatricesBuffer, MatrixOffsetsBuffer, DeformedPositionsBuffer);
+		Out.LocalToWorld = SkeletalMeshComponent->SceneProxy ? FTransform(SkeletalMeshComponent->SceneProxy->GetLocalToWorld()) : FTransform();
+		Out.LODIndex = LODIndex;
 
-		CachedGeometry.DeformedPositionBuffer = DeformedPositionsBuffer;
-		CachedGeometry.DeformedPreviousPositionBuffer = DeformedPreviousPositionsBuffer;
-		FRDGBufferSRVRef DeformedPositionSRV = GraphBuilder.CreateSRV(DeformedPositionsBuffer, PF_R32_FLOAT);
-		FRDGBufferSRVRef DeformedPreviousPositionSRV = bNeedPreviousPosition ? GraphBuilder.CreateSRV(DeformedPreviousPositionsBuffer, PF_R32_FLOAT) : nullptr;
-		CachedGeometry.LocalToWorld = SkeletalMeshComponent->SceneProxy ? FTransform(SkeletalMeshComponent->SceneProxy->GetLocalToWorld()) : FTransform();
-		for (int32 SectionIdx = 0; SectionIdx < LODData.RenderSections.Num(); ++SectionIdx)
+		if (bOutputTriangleData)
 		{
-			FCachedGeometry::Section CachedSection;
-			FSkelMeshRenderSection& Section = LODData.RenderSections[SectionIdx];
-
-			CachedSection.RDGPositionBuffer = DeformedPositionSRV;
-			CachedSection.RDGPreviousPositionBuffer = DeformedPreviousPositionSRV;
-			CachedSection.PositionBuffer = nullptr; // Do not use the SRV slot, but instead use the RDG buffer created above (DeformedPositionSRV)
-			CachedSection.PreviousPositionBuffer = nullptr; // Do not use the SRV slot, but instead use the RDG buffer created above (DeformedPositionSRV)
-			CachedSection.UVsBuffer = LODData.StaticVertexBuffers.StaticMeshVertexBuffer.GetTexCoordsSRV();
-			CachedSection.TotalVertexCount = LODData.StaticVertexBuffers.PositionVertexBuffer.GetNumVertices();
-			CachedSection.IndexBuffer = LODData.MultiSizeIndexContainer.GetIndexBuffer()->GetSRV();
-			CachedSection.TotalIndexCount = LODData.MultiSizeIndexContainer.GetIndexBuffer()->Num();
-			CachedSection.UVsChannelCount = LODData.StaticVertexBuffers.StaticMeshVertexBuffer.GetNumTexCoords();
-			CachedSection.NumPrimitives = Section.NumTriangles;
-			CachedSection.NumVertices = Section.NumVertices;
-			CachedSection.IndexBaseIndex = Section.BaseIndex;
-			CachedSection.VertexBaseIndex = Section.BaseVertexIndex;
-			CachedSection.SectionIndex = SectionIdx;
-			CachedSection.LODIndex = LODIndex;
-			CachedSection.UVsChannelOffset = 0; // Assume that we needs to pair meshes based on UVs 0
-
-			CachedGeometry.Sections.Add(CachedSection);
+			FSkeletalMeshRenderData* RenderData = SkeletalMeshComponent->GetSkeletalMesh()->GetResourceForRendering();
+			FSkeletalMeshLODRenderData& LODData = RenderData->LODRenderData[LODIndex];
+
+			TArray<uint32> MatrixOffsets;
+			TArray<FVector4f> BoneMatrices;
+			BuildBoneMatrices(SkeletalMeshComponent, LODData, LODIndex, MatrixOffsets, BoneMatrices);
+
+			const bool bNeedPreviousPosition = IsHairStrandContinuousDecimationReorderingEnabled();
+
+			FRDGBufferRef DeformedPositionsBuffer			= GraphBuilder.CreateBuffer(FRDGBufferDesc::CreateBufferDesc(sizeof(float), LODData.StaticVertexBuffers.PositionVertexBuffer.GetNumVertices() * 3), TEXT("Hair.SkinnedDeformedPositions"));
+			FRDGBufferRef DeformedPreviousPositionsBuffer	= bNeedPreviousPosition ? GraphBuilder.CreateBuffer(FRDGBufferDesc::CreateBufferDesc(sizeof(float), LODData.StaticVertexBuffers.PositionVertexBuffer.GetNumVertices() * 3), TEXT("Hair.SkinnedDeformedPreviousPositions")) : nullptr;
+			FRDGBufferRef BoneMatricesBuffer				= CreateStructuredBuffer(GraphBuilder, TEXT("Hair.SkinnedBoneMatrices"), sizeof(float) * 4, BoneMatrices.Num(), BoneMatrices.GetData(), sizeof(float) * 4 * BoneMatrices.Num());
+			FRDGBufferRef MatrixOffsetsBuffer				= CreateStructuredBuffer(GraphBuilder, TEXT("Hair.SkinnedMatrixOffsets"), sizeof(uint32), MatrixOffsets.Num(), MatrixOffsets.GetData(), sizeof(uint32) * MatrixOffsets.Num());
+
+			// TODO previous bone matrices
+			AddSkinUpdatePass(GraphBuilder, ShaderMap, SkeletalMeshComponent->GetSkinWeightBuffer(LODIndex), LODData, BoneMatricesBuffer, MatrixOffsetsBuffer, DeformedPositionsBuffer);
+
+			Out.DeformedPositionBuffer = DeformedPositionsBuffer;
+			Out.DeformedPreviousPositionBuffer = DeformedPreviousPositionsBuffer;
+			FRDGBufferSRVRef DeformedPositionSRV = GraphBuilder.CreateSRV(DeformedPositionsBuffer, PF_R32_FLOAT);
+			FRDGBufferSRVRef DeformedPreviousPositionSRV = bNeedPreviousPosition ? GraphBuilder.CreateSRV(DeformedPreviousPositionsBuffer, PF_R32_FLOAT) : nullptr;
+			for (int32 SectionIdx = 0; SectionIdx < LODData.RenderSections.Num(); ++SectionIdx)
+			{
+				const FSkelMeshRenderSection& Section = LODData.RenderSections[SectionIdx];
+
+				FCachedGeometry::Section& OutSection = Out.Sections.AddDefaulted_GetRef();
+				OutSection.RDGPositionBuffer = DeformedPositionSRV;
+				OutSection.RDGPreviousPositionBuffer = DeformedPreviousPositionSRV;
+				OutSection.PositionBuffer = nullptr; // Do not use the SRV slot, but instead use the RDG buffer created above (DeformedPositionSRV)
+				OutSection.PreviousPositionBuffer = nullptr; // Do not use the SRV slot, but instead use the RDG buffer created above (DeformedPositionSRV)
+				OutSection.UVsBuffer = LODData.StaticVertexBuffers.StaticMeshVertexBuffer.GetTexCoordsSRV();
+				OutSection.TotalVertexCount = LODData.StaticVertexBuffers.PositionVertexBuffer.GetNumVertices();
+				OutSection.IndexBuffer = LODData.MultiSizeIndexContainer.GetIndexBuffer()->GetSRV();
+				OutSection.TotalIndexCount = LODData.MultiSizeIndexContainer.GetIndexBuffer()->Num();
+				OutSection.UVsChannelCount = LODData.StaticVertexBuffers.StaticMeshVertexBuffer.GetNumTexCoords();
+				OutSection.NumPrimitives = Section.NumTriangles;
+				OutSection.NumVertices = Section.NumVertices;
+				OutSection.IndexBaseIndex = Section.BaseIndex;
+				OutSection.VertexBaseIndex = Section.BaseVertexIndex;
+				OutSection.SectionIndex = SectionIdx;
+				OutSection.LODIndex = LODIndex;
+				OutSection.UVsChannelOffset = 0; // Assume that we needs to pair meshes based on UVs 0
+			}
 		}
 	}
 }
@@ -133,50 +136,55 @@ static void BuildBoneMatrices(USkeletalMeshComponent* SkeletalMeshComponent, con
  void BuildCacheGeometry(
 	 FRDGBuilder& GraphBuilder,
 	 FGlobalShaderMap* ShaderMap, 
-	 UGeometryCacheComponent* GeometryCacheComponent, 
-	 FCachedGeometry& CachedGeometry)
+	 const UGeometryCacheComponent* GeometryCacheComponent, 
+	 const bool bOutputTriangleData,
+	 FCachedGeometry& Out)
  {
 	 if (GeometryCacheComponent)
 	 {
 		 if (FGeometryCacheSceneProxy* SceneProxy = static_cast<FGeometryCacheSceneProxy*>(GeometryCacheComponent->SceneProxy))
 		 {
-			 // Prior to getting here, the GeometryCache has been validated to be flattened (ie. has only one track)
-			 check(SceneProxy->GetTracks().Num() > 0);
-			 const FGeomCacheTrackProxy* TrackProxy = SceneProxy->GetTracks()[0];
-
-			 check(TrackProxy->MeshData && TrackProxy->NextFrameMeshData);
-			 const bool bHasMotionVectors = (
-				 TrackProxy->MeshData->VertexInfo.bHasMotionVectors &&
-				 TrackProxy->NextFrameMeshData->VertexInfo.bHasMotionVectors &&
-				 TrackProxy->MeshData->Positions.Num() == TrackProxy->MeshData->MotionVectors.Num())
-				 && (TrackProxy->NextFrameMeshData->Positions.Num() == TrackProxy->NextFrameMeshData->MotionVectors.Num());
-
-			 // PositionBuffer depends on CurrentPositionBufferIndex and on if the cache has motion vectors
-			 const uint32 PositionIndex = (TrackProxy->CurrentPositionBufferIndex == -1 || bHasMotionVectors) ? 0 : TrackProxy->CurrentPositionBufferIndex % 2;
-			 CachedGeometry.LocalToWorld = FTransform(SceneProxy->GetLocalToWorld());
-			 for (int32 SectionIdx = 0; SectionIdx < TrackProxy->MeshData->BatchesInfo.Num(); ++SectionIdx)
+			 Out.LocalToWorld = FTransform(SceneProxy->GetLocalToWorld());
+			 Out.LODIndex = 0;
+			 if (bOutputTriangleData)
 			 {
-				const FGeometryCacheMeshBatchInfo& BatchInfo = TrackProxy->MeshData->BatchesInfo[SectionIdx];
+				 // Prior to getting here, the GeometryCache has been validated to be flattened (ie. has only one track)
+				 check(SceneProxy->GetTracks().Num() > 0);
+				 const FGeomCacheTrackProxy* TrackProxy = SceneProxy->GetTracks()[0];
+
+				 check(TrackProxy->MeshData && TrackProxy->NextFrameMeshData);
+				 const bool bHasMotionVectors = (
+					 TrackProxy->MeshData->VertexInfo.bHasMotionVectors &&
+					 TrackProxy->NextFrameMeshData->VertexInfo.bHasMotionVectors &&
+					 TrackProxy->MeshData->Positions.Num() == TrackProxy->MeshData->MotionVectors.Num())
+					 && (TrackProxy->NextFrameMeshData->Positions.Num() == TrackProxy->NextFrameMeshData->MotionVectors.Num());
+
+				 // PositionBuffer depends on CurrentPositionBufferIndex and on if the cache has motion vectors
+				 const uint32 PositionIndex = (TrackProxy->CurrentPositionBufferIndex == -1 || bHasMotionVectors) ? 0 : TrackProxy->CurrentPositionBufferIndex % 2;
+				 for (int32 SectionIdx = 0; SectionIdx < TrackProxy->MeshData->BatchesInfo.Num(); ++SectionIdx)
+				 {
+					const FGeometryCacheMeshBatchInfo& BatchInfo = TrackProxy->MeshData->BatchesInfo[SectionIdx];
 				
-				FCachedGeometry::Section CachedSection;
-				CachedSection.PositionBuffer = TrackProxy->PositionBuffers[PositionIndex].GetBufferSRV();
-				CachedSection.UVsBuffer = TrackProxy->TextureCoordinatesBuffer.GetBufferSRV();
-				CachedSection.TotalVertexCount = TrackProxy->MeshData->Positions.Num();
-				CachedSection.IndexBuffer = TrackProxy->IndexBuffer.GetBufferSRV();
-				CachedSection.TotalIndexCount = TrackProxy->IndexBuffer.NumValidIndices;
-				CachedSection.UVsChannelCount = 1;
-				CachedSection.NumPrimitives = BatchInfo.NumTriangles;
-				CachedSection.NumVertices = TrackProxy->MeshData->Positions.Num();
-				CachedSection.IndexBaseIndex = BatchInfo.StartIndex;
-				CachedSection.VertexBaseIndex = 0;
-				CachedSection.SectionIndex = SectionIdx;
-				CachedSection.LODIndex = 0;
-				CachedSection.UVsChannelOffset = 0;
+					FCachedGeometry::Section OutSection;
+					OutSection.PositionBuffer = TrackProxy->PositionBuffers[PositionIndex].GetBufferSRV();
+					OutSection.UVsBuffer = TrackProxy->TextureCoordinatesBuffer.GetBufferSRV();
+					OutSection.TotalVertexCount = TrackProxy->MeshData->Positions.Num();
+					OutSection.IndexBuffer = TrackProxy->IndexBuffer.GetBufferSRV();
+					OutSection.TotalIndexCount = TrackProxy->IndexBuffer.NumValidIndices;
+					OutSection.UVsChannelCount = 1;
+					OutSection.NumPrimitives = BatchInfo.NumTriangles;
+					OutSection.NumVertices = TrackProxy->MeshData->Positions.Num();
+					OutSection.IndexBaseIndex = BatchInfo.StartIndex;
+					OutSection.VertexBaseIndex = 0;
+					OutSection.SectionIndex = SectionIdx;
+					OutSection.LODIndex = 0;
+					OutSection.UVsChannelOffset = 0;
 				
-				if (CachedSection.PositionBuffer && CachedSection.IndexBuffer)
-				{
-					CachedGeometry.Sections.Add(CachedSection);
-				}
+					if (OutSection.PositionBuffer && OutSection.IndexBuffer)
+					{
+						Out.Sections.Add(OutSection);
+					}
+				 }
 			 }
 		 }
 	 }
diff --git a/Engine/Plugins/Runtime/HairStrands/Source/HairStrandsCore/Private/GroomGeometryCache.h b/Engine/Plugins/Runtime/HairStrands/Source/HairStrandsCore/Private/GroomGeometryCache.h
--- a/Engine/Plugins/Runtime/HairStrands/Source/HairStrandsCore/Private/GroomGeometryCache.h
+++ b/Engine/Plugins/Runtime/HairStrands/Source/HairStrandsCore/Private/GroomGeometryCache.h
@@ -12,16 +12,20 @@
 #include "HairStrandsMeshProjection.h"
 #include "GPUSkinCache.h"
 
+class UGeometryCacheComponent;
+
 FHairStrandsProjectionMeshData::Section ConvertMeshSection(const FCachedGeometry::Section& In, const FTransform& T);
 
 void BuildCacheGeometry(
 	FRDGBuilder& GraphBuilder,
 	FGlobalShaderMap* ShaderMap, 
-	USkeletalMeshComponent* SkeletalMeshComponent, 
-	FCachedGeometry& CachedGeometry);
+	const USkeletalMeshComponent* SkeletalMeshComponent, 
+	const bool bOutputTriangleData,
+	FCachedGeometry& OutCachedGeometry);
 
 void BuildCacheGeometry(
 	FRDGBuilder& GraphBuilder,
 	FGlobalShaderMap* ShaderMap, 
-	class UGeometryCacheComponent* GeometryCacheComponent, 
-	FCachedGeometry& CachedGeometry);
\ No newline at end of file
+	const UGeometryCacheComponent* GeometryCacheComponent,
+	const bool bOutputTriangleData,
+	FCachedGeometry& OutCachedGeometry);
\ No newline at end of file
diff --git a/Engine/Plugins/Runtime/HairStrands/Source/HairStrandsCore/Private/GroomManager.cpp b/Engine/Plugins/Runtime/HairStrands/Source/HairStrandsCore/Private/GroomManager.cpp
--- a/Engine/Plugins/Runtime/HairStrands/Source/HairStrandsCore/Private/GroomManager.cpp
+++ b/Engine/Plugins/Runtime/HairStrands/Source/HairStrandsCore/Private/GroomManager.cpp
@@ -79,12 +79,13 @@ FCachedGeometry GetCacheGeometryForHair(
 	FRDGBuilder& GraphBuilder, 
 	FHairGroupInstance* Instance, 
 	const FGPUSkinCache* SkinCache, 
-	FGlobalShaderMap* ShaderMap)
+	FGlobalShaderMap* ShaderMap,
+	const bool bOutputTriangleData)
 {
 	FCachedGeometry Out;
 	if (Instance->Debug.GroomBindingType == EGroomBindingMeshType::SkeletalMesh)
 	{
-		if (USkeletalMeshComponent* SkeletalMeshComponent = Cast<USkeletalMeshComponent>(Instance->Debug.MeshComponent))
+		if (const USkeletalMeshComponent* SkeletalMeshComponent = Cast<USkeletalMeshComponent>(Instance->Debug.MeshComponent))
 		{
 			if (SkinCache)
 			{
@@ -94,16 +95,16 @@ FCachedGeometry GetCacheGeometryForHair(
 			if (IsHairStrandsSkinCacheEnable() && Out.Sections.Num() == 0)
 			{
 				//#hair_todo: Need to have a (frame) cache to insure that we don't recompute the same projection several time
-				// Actual populate the cache with only the needed part basd on the groom projection data. At the moment it recompute everything ...
-				BuildCacheGeometry(GraphBuilder, ShaderMap, SkeletalMeshComponent, Out);
+				// Actual populate the cache with only the needed part based on the groom projection data. At the moment it recompute everything ...
+				BuildCacheGeometry(GraphBuilder, ShaderMap, SkeletalMeshComponent, bOutputTriangleData, Out);
 			}
 		}
 	}
 	else if (Instance->Debug.GroomBindingType == EGroomBindingMeshType::GeometryCache)
 	{
-		if (UGeometryCacheComponent* GeometryCacheComponent = Cast<UGeometryCacheComponent>(Instance->Debug.MeshComponent))
+		if (const UGeometryCacheComponent* GeometryCacheComponent = Cast<const UGeometryCacheComponent>(Instance->Debug.MeshComponent))
 		{
-			BuildCacheGeometry(GraphBuilder, ShaderMap, GeometryCacheComponent, Out);
+			BuildCacheGeometry(GraphBuilder, ShaderMap, GeometryCacheComponent, bOutputTriangleData, Out);
 		}
 	}
 	return Out;
@@ -113,7 +114,7 @@ FCachedGeometry GetCacheGeometryForHair(
 // Return -1 if the hair are not bound of if the underlying geometry is invalid
 static int32 GetCacheGeometryLODIndex(const FCachedGeometry& In)
 {
-	int32 MeshLODIndex = -1;
+	int32 MeshLODIndex = In.LODIndex;
 	for (const FCachedGeometry::Section& Section : In.Sections)
 	{
 		// Ensure all mesh's sections have the same LOD index
@@ -155,7 +156,7 @@ static void RunInternalHairStrandsInterpolation(
 		check(Instance->HairGroupPublicData);
 
 		FHairStrandsProjectionMeshData::LOD MeshDataLOD;
-		const FCachedGeometry CachedGeometry = GetCacheGeometryForHair(GraphBuilder, Instance, SkinCache, ShaderMap);
+		const FCachedGeometry CachedGeometry = GetCacheGeometryForHair(GraphBuilder, Instance, SkinCache, ShaderMap, true);
 		for (const FCachedGeometry::Section& Section : CachedGeometry.Sections)
 		{
 			// Ensure all mesh's sections have the same LOD index
@@ -707,7 +708,7 @@ static void RunHairLODSelection(
 
 		check(Instance);
 		check(Instance->HairGroupPublicData);
-		const FCachedGeometry CachedGeometry = GetCacheGeometryForHair(GraphBuilder, Instance, SkinCache, ShaderMap);
+		const FCachedGeometry CachedGeometry = GetCacheGeometryForHair(GraphBuilder, Instance, SkinCache, ShaderMap, false);
 		const int32 MeshLODIndex = GetCacheGeometryLODIndex(CachedGeometry);
 
 		// Perform LOD selection based on all the views	
diff --git a/Engine/Source/Runtime/Engine/Private/GPUSkinCache.cpp b/Engine/Source/Runtime/Engine/Private/GPUSkinCache.cpp
--- a/Engine/Source/Runtime/Engine/Private/GPUSkinCache.cpp
+++ b/Engine/Source/Runtime/Engine/Private/GPUSkinCache.cpp
@@ -2148,6 +2148,7 @@ FCachedGeometry FGPUSkinCache::GetCachedGeometry(uint32 ComponentId, EGPUSkinCac
 						Out.Sections.Add(CachedSection);
 					}
 					Out.LocalToWorld = FTransform(Entry->GPUSkin->GetTransform());
+					Out.LODIndex = LODIndex;
 					return true;
 				}
 			}
diff --git a/Engine/Source/Runtime/Engine/Public/GPUSkinCache.h b/Engine/Source/Runtime/Engine/Public/GPUSkinCache.h
--- a/Engine/Source/Runtime/Engine/Public/GPUSkinCache.h
+++ b/Engine/Source/Runtime/Engine/Public/GPUSkinCache.h
@@ -114,10 +114,10 @@ struct FCachedGeometry
 		uint32 TotalVertexCount = 0;
 		uint32 TotalIndexCount = 0;
 		uint32 SectionIndex = 0;
-		int32 LODIndex = 0;
+		int32 LODIndex = -1;
 	};
 
-	int32 LODIndex = 0;
+	int32 LODIndex = -1;
 	TArray<Section> Sections;
 	FRDGBufferRef DeformedPositionBuffer = nullptr;
 	FRDGBufferRef DeformedPreviousPositionBuffer = nullptr;
