Remove enforcement of GDME being off when hair strands plugin is enabled.

This CL reverse the dependency, by ensuring the hair strands plugin does not run skel. mesh update when GDME is enabled (to prevent crash/incorrect behavior).

diff --git a/Engine/Plugins/Runtime/HairStrands/Source/HairStrandsCore/Private/GroomManager.cpp b/Engine/Plugins/Runtime/HairStrands/Source/HairStrandsCore/Private/GroomManager.cpp
--- a/Engine/Plugins/Runtime/HairStrands/Source/HairStrandsCore/Private/GroomManager.cpp
+++ b/Engine/Plugins/Runtime/HairStrands/Source/HairStrandsCore/Private/GroomManager.cpp
@@ -74,6 +74,14 @@ static bool IsInstanceFrustumCullingEnable()
 
 bool NeedsUpdateCardsMeshTriangles();
 
+static bool IsSkeletalMeshEvaluationEnabled()
+{
+	// When deferred skel. mesh update is enabled, hair strands skeletal mesh deformation is not allowed, as skin-cached update happen after 
+	// the PreInitView calls, which causes hair LOD selection and hair simulation to have invalid value & resources
+	static const auto CVarSkelMeshGDME = IConsoleManager::Get().FindTConsoleVariableDataInt(TEXT("r.DeferSkeletalDynamicDataUpdateUntilGDME"));
+	return CVarSkelMeshGDME ? CVarSkelMeshGDME->GetValueOnAnyThread() == 0 : true;
+}
+
 // Returns the cached geometry of the underlying geometry on which a hair instance is attached to
 FCachedGeometry GetCacheGeometryForHair(
 	FRDGBuilder& GraphBuilder, 
@@ -85,18 +93,21 @@ FCachedGeometry GetCacheGeometryForHair(
 	FCachedGeometry Out;
 	if (Instance->Debug.GroomBindingType == EGroomBindingMeshType::SkeletalMesh)
 	{
-		if (const USkeletalMeshComponent* SkeletalMeshComponent = Cast<USkeletalMeshComponent>(Instance->Debug.MeshComponent))
+		if (IsSkeletalMeshEvaluationEnabled())
 		{
-			if (SkinCache)
+			if (const USkeletalMeshComponent* SkeletalMeshComponent = Cast<USkeletalMeshComponent>(Instance->Debug.MeshComponent))
 			{
-				Out = SkinCache->GetCachedGeometry(SkeletalMeshComponent->ComponentId.PrimIDValue, EGPUSkinCacheEntryMode::Raster);
-			}
+				if (SkinCache)
+				{
+					Out = SkinCache->GetCachedGeometry(SkeletalMeshComponent->ComponentId.PrimIDValue, EGPUSkinCacheEntryMode::Raster);
+				}
 
-			if (IsHairStrandsSkinCacheEnable() && Out.Sections.Num() == 0)
-			{
-				//#hair_todo: Need to have a (frame) cache to insure that we don't recompute the same projection several time
-				// Actual populate the cache with only the needed part based on the groom projection data. At the moment it recompute everything ...
-				BuildCacheGeometry(GraphBuilder, ShaderMap, SkeletalMeshComponent, bOutputTriangleData, Out);
+				if (IsHairStrandsSkinCacheEnable() && Out.Sections.Num() == 0)
+				{
+					//#hair_todo: Need to have a (frame) cache to insure that we don't recompute the same projection several time
+					// Actual populate the cache with only the needed part based on the groom projection data. At the moment it recompute everything ...
+					BuildCacheGeometry(GraphBuilder, ShaderMap, SkeletalMeshComponent, bOutputTriangleData, Out);
+				}
 			}
 		}
 	}
diff --git a/Engine/Source/Runtime/Engine/Private/SkeletalRenderGPUSkin.cpp b/Engine/Source/Runtime/Engine/Private/SkeletalRenderGPUSkin.cpp
--- a/Engine/Source/Runtime/Engine/Private/SkeletalRenderGPUSkin.cpp
+++ b/Engine/Source/Runtime/Engine/Private/SkeletalRenderGPUSkin.cpp
@@ -508,11 +508,7 @@ static TAutoConsoleVariable<int32> CVarDeferSkeletalDynamicDataUpdateUntilGDME(
 
 FORCEINLINE bool IsDeferredSkeletalDynamicDataUpdateEnabled()
 {
-	// When the hair plugin is enabled, deferred skel. mesh update are not allowed, as skin-cached update need to run 
-	// prior to the PreInitView calls, because hair LOD selection and hair simulation requires the skin cache 
-	// resources to be updated at that point.
-	static bool bIsHairPluginEnable = IPluginManager::Get().FindPlugin("HairStrands") != nullptr && IPluginManager::Get().FindPlugin("HairStrands")->IsEnabled();
-	return CVarDeferSkeletalDynamicDataUpdateUntilGDME.GetValueOnRenderThread() > 0 && !bIsHairPluginEnable;
+	return CVarDeferSkeletalDynamicDataUpdateUntilGDME.GetValueOnRenderThread() > 0;
 }
 
 void FSkeletalMeshObjectGPUSkin::UpdateDynamicData_RenderThread(FGPUSkinCache* GPUSkinCache, FRHICommandListImmediate& RHICmdList, FDynamicSkelMeshObjectDataGPUSkin* InDynamicData, FSceneInterface* Scene, uint32 FrameNumberToPrepare, uint32 RevisionNumber)
