Fix simulation break when switching LOD

diff --git a/Engine/Plugins/Runtime/HairStrands/Source/HairStrandsCore/Private/GroomAsset.cpp b/Engine/Plugins/Runtime/HairStrands/Source/HairStrandsCore/Private/GroomAsset.cpp
--- a/Engine/Plugins/Runtime/HairStrands/Source/HairStrandsCore/Private/GroomAsset.cpp
+++ b/Engine/Plugins/Runtime/HairStrands/Source/HairStrandsCore/Private/GroomAsset.cpp
@@ -3256,9 +3256,9 @@ bool UGroomAsset::IsSimulationEnable(int32 GroupIndex, int32 LODIndex) const
 		return false;
 	}
 
-	return
+	return HairGroupsLOD[GroupIndex].LODs[LODIndex].bVisible && (
 		HairGroupsLOD[GroupIndex].LODs[LODIndex].Simulation == EGroomOverrideType::Enable ||
-		(HairGroupsLOD[GroupIndex].LODs[LODIndex].Simulation == EGroomOverrideType::Auto && HairGroupsPhysics[GroupIndex].SolverSettings.EnableSimulation);
+		(HairGroupsLOD[GroupIndex].LODs[LODIndex].Simulation == EGroomOverrideType::Auto && HairGroupsPhysics[GroupIndex].SolverSettings.EnableSimulation));
 }
 
 bool UGroomAsset::IsGlobalInterpolationEnable(int32 GroupIndex, int32 LODIndex) const
