Add groom binding listing & preview within the groom editor.

diff --git a/Engine/Plugins/Runtime/HairStrands/Source/HairStrandsCore/Public/GroomBindingAsset.h b/Engine/Plugins/Runtime/HairStrands/Source/HairStrandsCore/Public/GroomBindingAsset.h
--- a/Engine/Plugins/Runtime/HairStrands/Source/HairStrandsCore/Public/GroomBindingAsset.h
+++ b/Engine/Plugins/Runtime/HairStrands/Source/HairStrandsCore/Public/GroomBindingAsset.h
@@ -74,7 +74,7 @@ public:
 	EGroomBindingMeshType GroomBindingType = EGroomBindingMeshType::SkeletalMesh;
 
 	/** Groom to bind. */
-	UPROPERTY(EditAnywhere, BlueprintReadWrite, Category = "BuildSettings")
+	UPROPERTY(EditAnywhere, BlueprintReadWrite, Category = "BuildSettings", AssetRegistrySearchable)
 	TObjectPtr<UGroomAsset> Groom;
 
 	/** Skeletal mesh on which the groom has been authored. This is optional, and used only if the hair
@@ -191,3 +191,12 @@ public:
 #endif
 	bool bIsValid = false;
 };
+
+UCLASS(BlueprintType, hidecategories = (Object))
+class HAIRSTRANDSCORE_API UGroomBindingAssetList : public UObject
+{
+	GENERATED_BODY()
+public:
+	UPROPERTY(EditAnywhere, Transient, EditFixedSize, Category = "Bindings")
+	TArray<UGroomBindingAsset*> Bindings;
+};
\ No newline at end of file
diff --git a/Engine/Plugins/Runtime/HairStrands/Source/HairStrandsEditor/Private/GroomAssetDetails.cpp b/Engine/Plugins/Runtime/HairStrands/Source/HairStrandsEditor/Private/GroomAssetDetails.cpp
--- a/Engine/Plugins/Runtime/HairStrands/Source/HairStrandsEditor/Private/GroomAssetDetails.cpp
+++ b/Engine/Plugins/Runtime/HairStrands/Source/HairStrandsEditor/Private/GroomAssetDetails.cpp
@@ -280,6 +280,7 @@ FGroomRenderingDetails::FGroomRenderingDetails(IGroomCustomAssetEditorToolkit* I
 	if (InToolkit)
 	{
 		GroomAsset = InToolkit->GetCustomAsset();
+		Toolkit = InToolkit;
 	}
 	bDeleteWarningConsumed = false;
 	PanelType = Type;
@@ -305,6 +306,7 @@ FName GetCategoryName(EMaterialPanelType Type)
 	case EMaterialPanelType::Interpolation: return FName(TEXT("Interpolation"));
 	case EMaterialPanelType::LODs:			return FName(TEXT("LODs"));
 	case EMaterialPanelType::Physics:		return FName(TEXT("Physics"));
+	case EMaterialPanelType::Bindings:		return FName(TEXT("Bindings"));
 	}
 	return FName(TEXT("Unknown"));
 }
@@ -319,12 +321,22 @@ void FGroomRenderingDetails::CustomizeDetails(IDetailLayoutBuilder& DetailLayout
 	const TArray<TWeakObjectPtr<UObject>>& SelectedObjects = DetailLayout.GetSelectedObjects();
 	check(SelectedObjects.Num() <= 1); // The OnGenerateCustomWidgets delegate will not be useful if we try to process more than one object.
 
-	GroomAsset = SelectedObjects.Num() > 0 ? Cast<UGroomAsset>(SelectedObjects[0].Get()) : nullptr;
+	FName CategoryName = GetCategoryName(PanelType);
 	GroomDetailLayout = &DetailLayout;
+	if (UGroomAsset* LocalGroomAsset = Cast<UGroomAsset>(SelectedObjects[0].Get()))
+	{
+		GroomAsset = LocalGroomAsset;
 
-	FName CategoryName = GetCategoryName(PanelType);
-	IDetailCategoryBuilder& HairGroupCategory = DetailLayout.EditCategory(CategoryName, FText::GetEmpty(), ECategoryPriority::TypeSpecific);
-	CustomizeStrandsGroupProperties(DetailLayout, HairGroupCategory);
+		IDetailCategoryBuilder& HairGroupCategory = DetailLayout.EditCategory(CategoryName, FText::GetEmpty(), ECategoryPriority::TypeSpecific);
+		CustomizeStrandsGroupProperties(DetailLayout, HairGroupCategory);
+	}
+	else if (UGroomBindingAssetList* LocalGroomBindingList = Cast<UGroomBindingAssetList>(SelectedObjects[0].Get()))
+	{
+		GroomBindingAssetList = LocalGroomBindingList;
+
+		IDetailCategoryBuilder& HairGroupCategory = DetailLayout.EditCategory(CategoryName, FText::GetEmpty(), ECategoryPriority::TypeSpecific);
+		CustomizeStrandsGroupProperties(DetailLayout, HairGroupCategory);
+	}
 }
 
 //////////////////////////////////////////////////////////////////////////////////////////////////////////////////
@@ -540,6 +552,18 @@ void FGroomRenderingDetails::CustomizeStrandsGroupProperties(IDetailLayoutBuilde
 			}
 		}
 		break;
+		case EMaterialPanelType::Bindings:
+		{
+			TSharedRef<IPropertyHandle> Property = DetailLayout.GetProperty(GET_MEMBER_NAME_CHECKED(UGroomBindingAssetList, Bindings), UGroomBindingAssetList::StaticClass());
+			if (Property->IsValidHandle())
+			{
+				TSharedRef<FDetailArrayBuilder> PropertyBuilder = MakeShareable(new FDetailArrayBuilder(Property, false, false, false));
+				PropertyBuilder->OnGenerateArrayElementWidget(FOnGenerateArrayElementWidget::CreateSP(this, &FGroomRenderingDetails::OnGenerateElementForBindingAsset, &DetailLayout));
+				PropertyBuilder->SetDisplayName(FText::FromString(TEXT("Bindings")));
+				FilesCategory.AddCustomBuilder(PropertyBuilder, false);
+			}
+		}
+		break;
 	}
 }
 
@@ -592,6 +616,17 @@ FReply FGroomRenderingDetails::OnAddGroup(FProperty* Property)
 	return FReply::Handled();
 }
 
+FReply FGroomRenderingDetails::OnSelectBinding(int32 BindingIndex, FProperty* Property)
+{
+	check(GroomBindingAssetList);
+
+	// If user click twice onto the same binding index, we disable the binding;
+	BindingIndex = Toolkit->GetActiveBindingIndex() == BindingIndex ? -1 : BindingIndex;
+	Toolkit->PreviewBinding(BindingIndex);
+	ApplyChanges();
+	return FReply::Handled();
+}
+
 FName& FGroomRenderingDetails::GetMaterialSlotName(int32 GroupIndex)
 {
 	check(GroomAsset);
@@ -1637,6 +1672,11 @@ void FGroomRenderingDetails::OnGenerateElementForHairGroup(TSharedRef<IPropertyH
 			AddPropertyWithCustomReset(ChildHandle, ChildrenBuilder, GroupIndex, -1);
 		}
 		break;
+		case EMaterialPanelType::Bindings:
+		{
+			AddPropertyWithCustomReset(ChildHandle, ChildrenBuilder, GroupIndex, -1);
+		}
+		break;
 		default:
 		{
 			ChildrenBuilder.AddProperty(ChildHandle.ToSharedRef());
@@ -1647,6 +1687,35 @@ void FGroomRenderingDetails::OnGenerateElementForHairGroup(TSharedRef<IPropertyH
 	}
 }
 
+// Hair binding display
+void FGroomRenderingDetails::OnGenerateElementForBindingAsset(TSharedRef<IPropertyHandle> StructProperty, int32 BindingIndex, IDetailChildrenBuilder& ChildrenBuilder, IDetailLayoutBuilder* DetailLayout)
+{
+	FProperty* Property = StructProperty->GetProperty();
+	ChildrenBuilder.AddProperty(StructProperty);
+
+	const FLinearColor Color = BindingIndex == Toolkit->GetActiveBindingIndex() ? FLinearColor::Yellow : FLinearColor::White;
+
+	ChildrenBuilder.AddCustomRow(FText::FromString(TEXT("Preview")))
+	.ValueContent()
+	[
+		SNew(SHorizontalBox)
+		+ SHorizontalBox::Slot()
+		.AutoWidth()
+		[
+			SNew(SButton)
+			.VAlign(VAlign_Center)
+			.HAlign(HAlign_Center)
+			.ButtonStyle(FAppStyle::Get(), "HoverHintOnly")
+			.OnClicked(this, &FGroomRenderingDetails::OnSelectBinding, BindingIndex, Property)
+			[
+				SNew(SImage)
+				.Image(FAppStyle::GetBrush("Icons.Visible"))
+				.ColorAndOpacity(Color)
+			]
+		]
+	];
+}
+
 ////////////////////////////////////////////////////////////////////////////////////////////////////////////////////////
 void FGroomRenderingDetails::OnSetObject(const FAssetData& AssetData)
 {
diff --git a/Engine/Plugins/Runtime/HairStrands/Source/HairStrandsEditor/Private/GroomAssetDetails.h b/Engine/Plugins/Runtime/HairStrands/Source/HairStrandsEditor/Private/GroomAssetDetails.h
--- a/Engine/Plugins/Runtime/HairStrands/Source/HairStrandsEditor/Private/GroomAssetDetails.h
+++ b/Engine/Plugins/Runtime/HairStrands/Source/HairStrandsEditor/Private/GroomAssetDetails.h
@@ -14,6 +14,7 @@
 #include "Widgets/Input/SComboBox.h"
 
 class UGroomAsset;
+class UGroomBindingAssetList;
 class IGroomCustomAssetEditorToolkit;
 class IDetailLayoutBuilder;
 class IDetailCategoryBuilder;
@@ -27,12 +28,14 @@ enum class EMaterialPanelType
 	Meshes,
 	Interpolation,
 	LODs,
-	Physics
+	Physics,
+	Bindings
 };
 
 class FGroomRenderingDetails : public IDetailCustomization
 {
 	UGroomAsset* GroomAsset = nullptr;
+	UGroomBindingAssetList* GroomBindingAssetList = nullptr;
 	IGroomCustomAssetEditorToolkit* Toolkit = nullptr;
 
 public:
@@ -65,6 +68,7 @@ private:
 	void OnGenerateElementForHairGroup(TSharedRef<IPropertyHandle> StructProperty, int32 GroupIndex, IDetailChildrenBuilder& ChildrenBuilder, IDetailLayoutBuilder* DetailLayout);
 	void OnGenerateElementForLODs(TSharedRef<IPropertyHandle> StructProperty, int32 LODIndex, IDetailChildrenBuilder& ChildrenBuilder, IDetailLayoutBuilder* DetailLayout, int32 GroupIndex);
 	void AddLODSlot(TSharedRef<IPropertyHandle>& LODHandle, IDetailChildrenBuilder& ChildrenBuilder, int32 GroupIndex, int32 LODIndex);
+	void OnGenerateElementForBindingAsset(TSharedRef<IPropertyHandle> StructProperty, int32 BindingIndex, IDetailChildrenBuilder& ChildrenBuilder, IDetailLayoutBuilder* DetailLayout);
 
 	// Display custom thumbnail for material
 	void OnSetObject(const FAssetData& AssetData);
@@ -83,6 +87,7 @@ private:
 	FReply OnRefreshCards(int32 GroupIndex, FProperty* Property);
 	FReply OnSaveCards(int32 GroupIndex, FProperty* Property);
 	FReply OnGenerateCardDataUsingPlugin(int32 GroupIndex);
+	FReply OnSelectBinding(int32 BindingIndex, FProperty* Property);
 
 	void ExpandStructForLOD(TSharedRef<IPropertyHandle>& PropertyHandle, IDetailChildrenBuilder& ChildrenBuilder, int32 GroupIndex, int32 LODIndex, bool bOverrideReset);
 	void ExpandStruct(TSharedPtr<IPropertyHandle>& PropertyHandle, IDetailChildrenBuilder& ChildrenBuilder, int32 GroupIndex, int32 LODIndex, bool bOverrideReset);
diff --git a/Engine/Plugins/Runtime/HairStrands/Source/HairStrandsEditor/Private/GroomCustomAssetEditorToolkit.cpp b/Engine/Plugins/Runtime/HairStrands/Source/HairStrandsEditor/Private/GroomCustomAssetEditorToolkit.cpp
--- a/Engine/Plugins/Runtime/HairStrands/Source/HairStrandsEditor/Private/GroomCustomAssetEditorToolkit.cpp
+++ b/Engine/Plugins/Runtime/HairStrands/Source/HairStrandsEditor/Private/GroomCustomAssetEditorToolkit.cpp
@@ -19,6 +19,9 @@
 #include "GroomAssetDetails.h"
 #include "GroomMaterialDetails.h"
 
+#include "Misc/AssetRegistryInterface.h"
+#include "AssetRegistry/AssetRegistryModule.h"
+
 #include "Engine/SkeletalMesh.h"
 #include "Engine/StaticMesh.h"
 
@@ -35,6 +38,7 @@ const FName FGroomCustomAssetEditorToolkit::TabId_MeshesProperties(TEXT("GroomCu
 const FName FGroomCustomAssetEditorToolkit::TabId_MaterialProperties(TEXT("GroomCustomAssetEditor_MaterialProperties"));
 const FName FGroomCustomAssetEditorToolkit::TabId_PhysicsProperties(TEXT("GroomCustomAssetEditor_PhysicsProperties"));
 const FName FGroomCustomAssetEditorToolkit::TabId_PreviewGroomComponent(TEXT("GroomCustomAssetEditor_PreviewGroomComponent"));
+const FName FGroomCustomAssetEditorToolkit::TabId_BindingProperties(TEXT("GroomCustomAssetEditor_BindingProperties"));
 
 void FGroomCustomAssetEditorToolkit::RegisterTabSpawners(const TSharedRef<class FTabManager>& InTabManager)
 {
@@ -88,6 +92,11 @@ void FGroomCustomAssetEditorToolkit::RegisterTabSpawners(const TSharedRef<class
 		.SetGroup(WorkspaceMenuCategory.ToSharedRef())
 		.SetIcon(FSlateIcon(FAppStyle::GetAppStyleSetName(), "LevelEditor.Tabs.Details"));
 #endif
+
+	InTabManager->RegisterTabSpawner(TabId_BindingProperties, FOnSpawnTab::CreateSP(this, &FGroomCustomAssetEditorToolkit::SpawnTab_BindingProperties))
+		.SetDisplayName(LOCTEXT("BindingPropertiesTab", "Binding"))
+		.SetGroup(WorkspaceMenuCategory.ToSharedRef())
+		.SetIcon(FSlateIcon(FAppStyle::GetAppStyleSetName(), "LevelEditor.Tabs.Details"));
 }
 
 void FGroomCustomAssetEditorToolkit::UnregisterTabSpawners(const TSharedRef<class FTabManager>& InTabManager)
@@ -105,6 +114,7 @@ void FGroomCustomAssetEditorToolkit::UnregisterTabSpawners(const TSharedRef<clas
 #if GROOMEDITOR_ENABLE_COMPONENT_PANEL
 	InTabManager->UnregisterTabSpawner(TabId_PreviewGroomComponent);
 #endif
+	InTabManager->UnregisterTabSpawner(TabId_BindingProperties);
 }
 
 void FGroomCustomAssetEditorToolkit::DocPropChanged(UObject *InObject, FPropertyChangedEvent &Property)
@@ -141,23 +151,12 @@ void FGroomCustomAssetEditorToolkit::DocPropChanged(UObject *InObject, FProperty
 #endif
 }
 
-void FGroomCustomAssetEditorToolkit::OnStaticGroomTargetChanged(UStaticMesh *NewTarget)
-{	
-	if (PreviewStaticMeshComponent != nullptr)
-	{
-		PreviewStaticMeshComponent->SetStaticMesh(NewTarget);
-		PreviewSkeletalMeshComponent->SetVisibility(false);
-		PreviewStaticMeshComponent->SetVisibility(NewTarget != nullptr);
-	}	
-}
-
 void FGroomCustomAssetEditorToolkit::OnSkeletalGroomTargetChanged(USkeletalMesh *NewTarget)
 {	
 	if (PreviewSkeletalMeshComponent != nullptr)
 	{
 		PreviewSkeletalMeshComponent->SetSkeletalMesh(NewTarget);
 		PreviewSkeletalMeshComponent->SetVisibility(NewTarget != nullptr);
-		PreviewStaticMeshComponent->SetVisibility(false);
 	}	
 }
 
@@ -216,7 +215,6 @@ void FGroomCustomAssetEditorToolkit::InitPreviewComponents()
 	check(GroomAsset!=nullptr);
 //	check(FGroomDataIO::GetDocumentForAsset(GroomAsset) != nullptr); TODO
 	check(PreviewGroomComponent == nullptr);
-	check(PreviewStaticMeshComponent == nullptr);
 	check(PreviewSkeletalMeshComponent == nullptr);
 	
 #if 0 //TODO
@@ -226,39 +224,14 @@ void FGroomCustomAssetEditorToolkit::InitPreviewComponents()
 	FGroomDataIO::UpdateDocumentFromGroomAsset(GroomAsset.Get(), Doc);
 #endif
 
+	const bool bHasValidBindingAsset = GroomBindingAsset.IsValid();
 	PreviewGroomComponent = NewObject<UGroomComponent>(GetTransientPackage(), NAME_None, RF_Transient);
 	PreviewGroomComponent->CastShadow = 1;
 	PreviewGroomComponent->bCastDynamicShadow = 1;
 	PreviewGroomComponent->SetPreviewMode(true);
 	PreviewGroomComponent->SetGroomAsset(GroomAsset.Get());
 	PreviewGroomComponent->Activate(true);
-	
-	//if (PropertyListenDelegate.IsValid() == false)
-	//{		
-	//	PropertyListenDelegate = FCoreUObjectDelegates::OnObjectPropertyChanged.AddRaw(this, &FGroomCustomAssetEditorToolkit::DocPropChanged);
-	//}
-
-#if 0 //TODO
-	PreviewStaticMeshComponent = NewObject<UStaticMeshComponent>(GetTransientPackage(), NAME_None, RF_Transient);
-	PreviewSkeletalMeshComponent = NewObject<USkeletalMeshComponent>(GetTransientPackage(), NAME_None, RF_Transient);
-	if (Doc->GetStaticGroomTarget())
-	{
-		PreviewStaticMeshComponent->SetStaticMesh(Doc->GetStaticGroomTarget());
-	}
-	else
-	{
-		PreviewStaticMeshComponent->SetVisibility(false);
-	}
-
-	if (Doc->GetSkeletalGroomTarget())
-	{
-		PreviewSkeletalMeshComponent->SetSkeletalMesh(Doc->GetSkeletalGroomTarget());
-	}
-	else
-	{
-		PreviewSkeletalMeshComponent->SetVisibility(false);
-	}
-#endif
+	PreviewBinding(ActiveGroomBindingIndex);
 }
 
 void FGroomCustomAssetEditorToolkit::OnClose() 
@@ -282,17 +255,61 @@ void FGroomCustomAssetEditorToolkit::OnClose()
 #if GROOMEDITOR_ENABLE_COMPONENT_PANEL
 	DetailView_PreviewGroomComponent.Reset();
 #endif
+	DetailView_BindingProperties.Reset();
+}
+
+static void ListAllBindingAssets(const UGroomAsset* InGroomAsset, TWeakObjectPtr<UGroomBindingAssetList>& Out)
+{
+	FARFilter Filter;
+	Filter.ClassPaths.Add(UGroomBindingAsset::StaticClass()->GetClassPathName());
+
+	// Search for binding with matching groom asset
+	//{
+	//	const FName GroomAssetProperty = GET_MEMBER_NAME_CHECKED(UGroomBindingAsset, Groom);
+	//	FString GroomAssetName = FAssetData(InGroomAsset).GetExportTextName();
+	//	FString GroomAssetName = TSoftObjectPtr<UGroomAsset>(InGroomAsset).ToString();
+	//	Filter.TagsAndValues.Add(GroomAssetProperty, GroomAssetName);
+	//}
+
+	FAssetRegistryModule& AssetRegistryModule = FModuleManager::LoadModuleChecked<FAssetRegistryModule>("AssetRegistry");
+	TArray<FAssetData> BindingAssetData;
+	AssetRegistryModule.Get().GetAssets(Filter, BindingAssetData);
+
+	// Filter binding asset which match the groom asset (as the tag/value filter above does not work)
+	for (FAssetData& Asset :  BindingAssetData)
+	{
+		if (UGroomBindingAsset* Binding = (UGroomBindingAsset*)Asset.GetAsset())
+		{
+			if (Binding->Groom == InGroomAsset)
+			{
+				Out->Bindings.Add(Binding);
+			}
+		}
+	}
 }
 
 void FGroomCustomAssetEditorToolkit::InitCustomAssetEditor(const EToolkitMode::Type Mode, const TSharedPtr<class IToolkitHost>& InitToolkitHost, UGroomAsset* InCustomAsset)
 {
 	PreviewGroomComponent		 = nullptr;
-	PreviewStaticMeshComponent	 = nullptr;
 	PreviewSkeletalMeshComponent = nullptr;
+	GroomBindingAssetList = nullptr;
 
+	ViewportTab = SNew(SGroomEditorViewport);
+
+	GroomBindingAssetList = NewObject<UGroomBindingAssetList>(GetTransientPackage(), NAME_None, RF_Transient);
+	ListAllBindingAssets(InCustomAsset, GroomBindingAssetList);
+	if (GroomBindingAssetList->Bindings.Num() > 0)
+	{
+		ActiveGroomBindingIndex = 0;
+		GroomBindingAsset = GroomBindingAssetList->Bindings[ActiveGroomBindingIndex];
+	}
+	else
+	{
+		ActiveGroomBindingIndex = -1;
+	}
 	SetCustomAsset(InCustomAsset);
 	InitPreviewComponents();
-	
+
 	FPropertyEditorModule& PropertyEditorModule = FModuleManager::GetModuleChecked<FPropertyEditorModule>("PropertyEditor");
 
 	FDetailsViewArgs DetailsViewArgs;
@@ -308,6 +325,7 @@ void FGroomCustomAssetEditorToolkit::InitCustomAssetEditor(const EToolkitMode::T
 #if GROOMEDITOR_ENABLE_COMPONENT_PANEL
 	DetailView_PreviewGroomComponent	= PropertyEditorModule.CreateDetailView(DetailsViewArgs);
 #endif
+	DetailView_BindingProperties		= PropertyEditorModule.CreateDetailView(DetailsViewArgs);
 
 	// Customization
 	DetailView_CardsProperties->SetGenericLayoutDetailsDelegate(FOnGetDetailCustomizationInstance::CreateStatic(&FGroomRenderingDetails::MakeInstance, (IGroomCustomAssetEditorToolkit*)this, EMaterialPanelType::Cards));
@@ -317,12 +335,12 @@ void FGroomCustomAssetEditorToolkit::InitCustomAssetEditor(const EToolkitMode::T
 	DetailView_PhysicsProperties->SetGenericLayoutDetailsDelegate(FOnGetDetailCustomizationInstance::CreateStatic(&FGroomRenderingDetails::MakeInstance, (IGroomCustomAssetEditorToolkit*)this, EMaterialPanelType::Physics));
 	DetailView_LODProperties->SetGenericLayoutDetailsDelegate(FOnGetDetailCustomizationInstance::CreateStatic(&FGroomRenderingDetails::MakeInstance, (IGroomCustomAssetEditorToolkit*)this, EMaterialPanelType::LODs));
 	DetailView_MaterialProperties->SetGenericLayoutDetailsDelegate(FOnGetDetailCustomizationInstance::CreateStatic(&FGroomMaterialDetails::MakeInstance, (IGroomCustomAssetEditorToolkit*)this));
+	DetailView_BindingProperties->SetGenericLayoutDetailsDelegate(FOnGetDetailCustomizationInstance::CreateStatic(&FGroomRenderingDetails::MakeInstance, (IGroomCustomAssetEditorToolkit*)this, EMaterialPanelType::Bindings));
 
 	SEditorViewport::FArguments args;
-	ViewportTab = SNew(SGroomEditorViewport);
 	
 	// Default layout
-	const TSharedRef<FTabManager::FLayout> StandaloneDefaultLayout = FTabManager::NewLayout("Standalone_GroomAssetEditor_Layout_v15a")
+	const TSharedRef<FTabManager::FLayout> StandaloneDefaultLayout = FTabManager::NewLayout("Standalone_GroomAssetEditor_Layout_v15b")
 		->AddArea
 		(
 			FTabManager::NewPrimaryArea()
@@ -353,6 +371,7 @@ void FGroomCustomAssetEditorToolkit::InitCustomAssetEditor(const EToolkitMode::T
 				#if GROOMEDITOR_ENABLE_COMPONENT_PANEL
 					->AddTab(TabId_PreviewGroomComponent,	ETabState::OpenedTab)
 				#endif
+					->AddTab(TabId_BindingProperties,		ETabState::OpenedTab)
 				)
 			)
 		);
@@ -442,6 +461,10 @@ void FGroomCustomAssetEditorToolkit::InitCustomAssetEditor(const EToolkitMode::T
 		DetailView_PreviewGroomComponent->SetObject(PreviewGroomComponent.Get());
 	}
 #endif
+	if (DetailView_BindingProperties.IsValid())
+	{
+		DetailView_BindingProperties->SetObject(Cast<UObject>(GroomBindingAssetList));
+	}
 
 	ExtendToolbar();
 	RegenerateMenusAndToolbars();
@@ -473,6 +496,7 @@ void FGroomCustomAssetEditorToolkit::InitCustomAssetEditor(const EToolkitMode::T
 			LocalToolKit->DetailView_CardsProperties->ForceRefresh();
 			LocalToolKit->DetailView_MeshesProperties->ForceRefresh();
 			LocalToolKit->DetailView_MaterialProperties->ForceRefresh();
+			LocalToolKit->DetailView_BindingProperties->ForceRefresh();
 		};
 		PropertyListenDelegate = GroomAsset->GetOnGroomAssetResourcesChanged().AddLambda(InvalidateDetailViews);
 	}
@@ -520,7 +544,6 @@ FGroomCustomAssetEditorToolkit::FGroomCustomAssetEditorToolkit()
 {
 	GroomAsset						= nullptr;	
 	PreviewGroomComponent			= nullptr;	
-	PreviewStaticMeshComponent		= nullptr;
 	PreviewSkeletalMeshComponent	= nullptr;
 }
 
@@ -665,14 +688,21 @@ TSharedRef<SDockTab> FGroomCustomAssetEditorToolkit::SpawnTab_PreviewGroomCompon
 		];
 }
 
-UGroomComponent *FGroomCustomAssetEditorToolkit::GetPreview_GroomComponent() const
-{	
-	return PreviewGroomComponent.Get();
+TSharedRef<SDockTab> FGroomCustomAssetEditorToolkit::SpawnTab_BindingProperties(const FSpawnTabArgs& Args)
+{
+	check(Args.GetTabId() == TabId_BindingProperties);
+
+	return SNew(SDockTab)
+		.Label(LOCTEXT("BindingPropertiesTab", "Binding"))
+		.TabColorScale(GetTabColorScale())
+		[
+			DetailView_BindingProperties.ToSharedRef()
+		];
 }
 
-UStaticMeshComponent *FGroomCustomAssetEditorToolkit::GetPreview_StaticMeshComponent() const
+UGroomComponent *FGroomCustomAssetEditorToolkit::GetPreview_GroomComponent() const
 {	
-	return PreviewStaticMeshComponent.Get();
+	return PreviewGroomComponent.Get();
 }
 
 USkeletalMeshComponent *FGroomCustomAssetEditorToolkit::GetPreview_SkeletalMeshComponent() const
@@ -684,12 +714,9 @@ TSharedRef<SDockTab> FGroomCustomAssetEditorToolkit::SpawnViewportTab(const FSpa
 {
 	check(Args.GetTabId() == TabId_Viewport);
 	check(GetPreview_GroomComponent());
-//	check(GetPreview_StaticMeshComponent());
-//	check(GetPreview_SkeletalMeshComponent());
 
-	ViewportTab->SetGroomComponent(GetPreview_GroomComponent());
-	ViewportTab->SetStaticMeshComponent(GetPreview_StaticMeshComponent());
 	ViewportTab->SetSkeletalMeshComponent(GetPreview_SkeletalMeshComponent());
+	ViewportTab->SetGroomComponent(GetPreview_GroomComponent());
 
 	return SNew(SDockTab)
 		.Label(LOCTEXT("RenderTitle", "Render"))
@@ -699,4 +726,41 @@ TSharedRef<SDockTab> FGroomCustomAssetEditorToolkit::SpawnViewportTab(const FSpa
 		];
 }
 
+void FGroomCustomAssetEditorToolkit::PreviewBinding(int32 BindingIndex)
+{
+	if (BindingIndex >= 0 && GroomBindingAssetList.Get() && BindingIndex < GroomBindingAssetList->Bindings.Num())
+	{
+		GroomBindingAsset = GroomBindingAssetList->Bindings[BindingIndex];
+	}
+	else
+	{
+		GroomBindingAsset = nullptr;
+	}
+
+	if (GroomBindingAsset.Get())
+	{
+		PreviewSkeletalMeshComponent = NewObject<USkeletalMeshComponent>(GetTransientPackage(), NAME_None, RF_Transient);
+		PreviewSkeletalMeshComponent->SetSkeletalMesh(GroomBindingAsset->TargetSkeletalMesh);
+		PreviewSkeletalMeshComponent->SetVisibility(true);
+		PreviewSkeletalMeshComponent->Activate(true);
+		PreviewGroomComponent->AttachToComponent(PreviewSkeletalMeshComponent.Get(), FAttachmentTransformRules::KeepRelativeTransform);
+		PreviewGroomComponent->SetGroomAsset(GroomAsset.Get(), GroomBindingAsset.Get());
+		ActiveGroomBindingIndex = BindingIndex;
+	}
+	else
+	{
+		PreviewGroomComponent->SetBindingAsset(nullptr);
+		PreviewGroomComponent->DetachFromComponent(FDetachmentTransformRules::KeepRelativeTransform);
+		PreviewSkeletalMeshComponent = nullptr;
+		ActiveGroomBindingIndex = -1;
+	}
+
+	ViewportTab->SetSkeletalMeshComponent(GetPreview_SkeletalMeshComponent());
+}
+
+int32 FGroomCustomAssetEditorToolkit::GetActiveBindingIndex() const
+{
+	return ActiveGroomBindingIndex;
+}
+
 #undef LOCTEXT_NAMESPACE
diff --git a/Engine/Plugins/Runtime/HairStrands/Source/HairStrandsEditor/Private/Widgets/SGroomEditorViewport.cpp b/Engine/Plugins/Runtime/HairStrands/Source/HairStrandsEditor/Private/Widgets/SGroomEditorViewport.cpp
--- a/Engine/Plugins/Runtime/HairStrands/Source/HairStrandsEditor/Private/Widgets/SGroomEditorViewport.cpp
+++ b/Engine/Plugins/Runtime/HairStrands/Source/HairStrandsEditor/Private/Widgets/SGroomEditorViewport.cpp
@@ -328,7 +328,7 @@ void SGroomEditorViewport::SetSkeletalMeshComponent(USkeletalMeshComponent *Targ
 	}
 	SkeletalGroomTarget = Target;
 
-	if (StaticGroomTarget != nullptr)
+	if (SkeletalGroomTarget != nullptr)
 	{
 		AdvancedPreviewScene->AddComponent(SkeletalGroomTarget, SkeletalGroomTarget->GetRelativeTransform());
 	}
diff --git a/Engine/Plugins/Runtime/HairStrands/Source/HairStrandsEditor/Public/GroomCustomAssetEditorToolkit.h b/Engine/Plugins/Runtime/HairStrands/Source/HairStrandsEditor/Public/GroomCustomAssetEditorToolkit.h
--- a/Engine/Plugins/Runtime/HairStrands/Source/HairStrandsEditor/Public/GroomCustomAssetEditorToolkit.h
+++ b/Engine/Plugins/Runtime/HairStrands/Source/HairStrandsEditor/Public/GroomCustomAssetEditorToolkit.h
@@ -15,11 +15,12 @@ class UGroomComponent;
 class IDetailsView;
 class SDockableTab;
 class SGroomEditorViewport;
-class UStaticMeshComponent;
 class USkeletalMesh;
 class UStaticMesh;
 class USkeletalMeshComponent;
-
+class UGroomAsset;
+class UGroomBindingAsset;
+class UGroomBindingAssetList;
 
 class HAIRSTRANDSEDITOR_API IGroomCustomAssetEditorToolkit : public FAssetEditorToolkit
 {
@@ -29,6 +30,11 @@ public:
 
 	/** Set the current custom asset. */
 	virtual void SetCustomAsset(UGroomAsset* InCustomAsset) = 0;
+
+	/** Set preview of a particular binding. */
+	virtual void PreviewBinding(int32 BindingIndex) = 0;
+
+	virtual int32 GetActiveBindingIndex() const =0;
 };
 
 class HAIRSTRANDSEDITOR_API FGroomCustomAssetEditorToolkit : public IGroomCustomAssetEditorToolkit
@@ -69,6 +75,12 @@ public:
 	/** Set the current custom asset. */
 	virtual void SetCustomAsset(UGroomAsset* InCustomAsset) override;	
 
+	/** Set preview of a particular binding. */
+	virtual void PreviewBinding(int32 BindingIndex) override;
+
+	/** Return the index of the active binding. */
+	virtual int32 GetActiveBindingIndex() const override;
+
 private:
 
 	// called when The Play simulation button is pressed
@@ -89,9 +101,6 @@ private:
 	// THis should be called when the properties of the Document are changed
 	void DocPropChanged(UObject *, FPropertyChangedEvent &);
 
-	// THis is called when the groom target object changes and needs updating
-	void OnStaticGroomTargetChanged(UStaticMesh *NewTarget);
-
 	// THis is called when the groom target object changes and needs updating
 	void OnSkeletalGroomTargetChanged(USkeletalMesh *NewTarget);
 
@@ -100,7 +109,6 @@ private:
 
 	// return a pointer to the groom preview component
 	UGroomComponent*		GetPreview_GroomComponent() const;
-	UStaticMeshComponent*	GetPreview_StaticMeshComponent() const;
 	USkeletalMeshComponent*	GetPreview_SkeletalMeshComponent() const;
 
 	TSharedRef<SDockTab> SpawnViewportTab(const FSpawnTabArgs& Args);
@@ -113,7 +121,7 @@ private:
 	TSharedRef<SDockTab> SpawnTab_MaterialProperties(const FSpawnTabArgs& Args);
 	TSharedRef<SDockTab> SpawnTab_PhysicsProperties(const FSpawnTabArgs& Args);
 	TSharedRef<SDockTab> SpawnTab_PreviewGroomComponent(const FSpawnTabArgs& Args);
-
+	TSharedRef<SDockTab> SpawnTab_BindingProperties(const FSpawnTabArgs& Args);
 private:
 
 	/** Dockable tab for properties */
@@ -128,6 +136,7 @@ private:
 	TSharedPtr<class IDetailsView> DetailView_MaterialProperties;
 	TSharedPtr<class IDetailsView> DetailView_PhysicsProperties;
 	TSharedPtr<class IDetailsView> DetailView_PreviewGroomComponent;
+	TSharedPtr<class IDetailsView> DetailView_BindingProperties;
 
 	static const FName ToolkitFName;
 	static const FName TabId_Viewport;
@@ -140,12 +149,16 @@ private:
 	static const FName TabId_MaterialProperties;
 	static const FName TabId_PhysicsProperties;
 	static const FName TabId_PreviewGroomComponent;
+	static const FName TabId_BindingProperties;
+
 
+	int32 ActiveGroomBindingIndex = -1;
 	FDelegateHandle PropertyListenDelegate;
+	TWeakObjectPtr<UGroomAsset> GroomAsset;
+	TWeakObjectPtr<UGroomBindingAsset> GroomBindingAsset;
+	TWeakObjectPtr<UGroomBindingAssetList> GroomBindingAssetList;
 
-	TWeakObjectPtr  < class UGroomAsset>  GroomAsset;
-	TWeakObjectPtr  < UGroomComponent > PreviewGroomComponent;
-	TWeakObjectPtr  < UStaticMeshComponent > PreviewStaticMeshComponent;
-	TWeakObjectPtr  < USkeletalMeshComponent > PreviewSkeletalMeshComponent;
+	TWeakObjectPtr<UGroomComponent> PreviewGroomComponent;
+	TWeakObjectPtr<USkeletalMeshComponent> PreviewSkeletalMeshComponent;
 };
 
