Add skel. mesh animation support in groom editor.

diff --git a/Engine/Plugins/Runtime/HairStrands/Source/HairStrandsEditor/Private/GroomAssetDetails.cpp b/Engine/Plugins/Runtime/HairStrands/Source/HairStrandsEditor/Private/GroomAssetDetails.cpp
--- a/Engine/Plugins/Runtime/HairStrands/Source/HairStrandsEditor/Private/GroomAssetDetails.cpp
+++ b/Engine/Plugins/Runtime/HairStrands/Source/HairStrandsEditor/Private/GroomAssetDetails.cpp
@@ -1693,7 +1693,7 @@ void FGroomRenderingDetails::OnGenerateElementForBindingAsset(TSharedRef<IProper
 	FProperty* Property = StructProperty->GetProperty();
 	ChildrenBuilder.AddProperty(StructProperty);
 
-	const FLinearColor Color = BindingIndex == Toolkit->GetActiveBindingIndex() ? FLinearColor::Yellow : FLinearColor::White;
+	const FLinearColor Color = BindingIndex == Toolkit->GetActiveBindingIndex() ? FLinearColor::White : FLinearColor(0.1f, 0.1f, 0.1f, 1.f);
 
 	ChildrenBuilder.AddCustomRow(FText::FromString(TEXT("Preview")))
 	.ValueContent()
diff --git a/Engine/Plugins/Runtime/HairStrands/Source/HairStrandsEditor/Private/GroomCustomAssetEditorToolkit.cpp b/Engine/Plugins/Runtime/HairStrands/Source/HairStrandsEditor/Private/GroomCustomAssetEditorToolkit.cpp
--- a/Engine/Plugins/Runtime/HairStrands/Source/HairStrandsEditor/Private/GroomCustomAssetEditorToolkit.cpp
+++ b/Engine/Plugins/Runtime/HairStrands/Source/HairStrandsEditor/Private/GroomCustomAssetEditorToolkit.cpp
@@ -18,6 +18,8 @@
 #include "GroomEditorStyle.h"
 #include "GroomAssetDetails.h"
 #include "GroomMaterialDetails.h"
+#include "PropertyCustomizationHelpers.h"
+#include "AssetThumbnail.h"
 
 #include "Misc/AssetRegistryInterface.h"
 #include "AssetRegistry/AssetRegistryModule.h"
@@ -54,37 +56,37 @@ void FGroomCustomAssetEditorToolkit::RegisterTabSpawners(const TSharedRef<class
 	InTabManager->RegisterTabSpawner(TabId_LODProperties, FOnSpawnTab::CreateSP(this, &FGroomCustomAssetEditorToolkit::SpawnTab_LODProperties))
 		.SetDisplayName(LOCTEXT("LODPropertiesTab", "LOD"))
 		.SetGroup(WorkspaceMenuCategory.ToSharedRef())
-		.SetIcon(FSlateIcon(FAppStyle::GetAppStyleSetName(), "LevelEditor.Tabs.Details"));
+		.SetIcon(FSlateIcon(FAppStyle::GetAppStyleSetName(), "Icons.LOD"));
 
 	InTabManager->RegisterTabSpawner(TabId_InterpolationProperties, FOnSpawnTab::CreateSP(this, &FGroomCustomAssetEditorToolkit::SpawnTab_InterpolationProperties))
 		.SetDisplayName(LOCTEXT("InterpolationPropertiesTab", "Interpolation"))
 		.SetGroup(WorkspaceMenuCategory.ToSharedRef())
-		.SetIcon(FSlateIcon(FAppStyle::GetAppStyleSetName(), "LevelEditor.Tabs.Details"));
+		.SetIcon(FSlateIcon(FAppStyle::GetAppStyleSetName(), "ClassIcon.Blueprint"));
 
 	InTabManager->RegisterTabSpawner(TabId_RenderingProperties, FOnSpawnTab::CreateSP(this, &FGroomCustomAssetEditorToolkit::SpawnTab_RenderingProperties))
 		.SetDisplayName(LOCTEXT("RenderingPropertiesTab", "Strands"))
 		.SetGroup(WorkspaceMenuCategory.ToSharedRef())
-		.SetIcon(FSlateIcon(FAppStyle::GetAppStyleSetName(), "LevelEditor.Tabs.Details"));
+		.SetIcon(FSlateIcon(FAppStyle::GetAppStyleSetName(), "ClassIcon.CurveBase"));
 
 	InTabManager->RegisterTabSpawner(TabId_PhysicsProperties, FOnSpawnTab::CreateSP(this, &FGroomCustomAssetEditorToolkit::SpawnTab_PhysicsProperties))
 		.SetDisplayName(LOCTEXT("PhysicsPropertiesTab", "Physics"))
 		.SetGroup(WorkspaceMenuCategory.ToSharedRef())
-		.SetIcon(FSlateIcon(FAppStyle::GetAppStyleSetName(), "LevelEditor.Tabs.Details"));
+		.SetIcon(FSlateIcon(FAppStyle::GetAppStyleSetName(), "CollisionAnalyzer.TabIcon"));
 
 	InTabManager->RegisterTabSpawner(TabId_CardsProperties, FOnSpawnTab::CreateSP(this, &FGroomCustomAssetEditorToolkit::SpawnTab_CardsProperties))
 		.SetDisplayName(LOCTEXT("CardsPropertiesTab", "Cards"))
 		.SetGroup(WorkspaceMenuCategory.ToSharedRef())
-		.SetIcon(FSlateIcon(FAppStyle::GetAppStyleSetName(), "LevelEditor.Tabs.Details"));
+		.SetIcon(FSlateIcon(FAppStyle::GetAppStyleSetName(), "Persona.ConvertToStaticMesh"));
 
 	InTabManager->RegisterTabSpawner(TabId_MeshesProperties, FOnSpawnTab::CreateSP(this, &FGroomCustomAssetEditorToolkit::SpawnTab_MeshesProperties))
 		.SetDisplayName(LOCTEXT("MeshesPropertiesTab", "Meshes"))
 		.SetGroup(WorkspaceMenuCategory.ToSharedRef())
-		.SetIcon(FSlateIcon(FAppStyle::GetAppStyleSetName(), "LevelEditor.Tabs.Details"));
+		.SetIcon(FSlateIcon(FAppStyle::GetAppStyleSetName(), "Persona.ConvertToStaticMesh"));
 
 	InTabManager->RegisterTabSpawner(TabId_MaterialProperties, FOnSpawnTab::CreateSP(this, &FGroomCustomAssetEditorToolkit::SpawnTab_MaterialProperties))
 		.SetDisplayName(LOCTEXT("MaterialPropertiesTab", "Material"))
 		.SetGroup(WorkspaceMenuCategory.ToSharedRef())
-		.SetIcon(FSlateIcon(FAppStyle::GetAppStyleSetName(), "LevelEditor.Tabs.Details"));
+		.SetIcon(FSlateIcon(FAppStyle::GetAppStyleSetName(), "ClassIcon.Material"));
 
 #if GROOMEDITOR_ENABLE_COMPONENT_PANEL
 	InTabManager->RegisterTabSpawner(TabId_PreviewGroomComponent, FOnSpawnTab::CreateSP(this, &FGroomCustomAssetEditorToolkit::SpawnTab_PreviewGroomComponent))
@@ -96,7 +98,7 @@ void FGroomCustomAssetEditorToolkit::RegisterTabSpawners(const TSharedRef<class
 	InTabManager->RegisterTabSpawner(TabId_BindingProperties, FOnSpawnTab::CreateSP(this, &FGroomCustomAssetEditorToolkit::SpawnTab_BindingProperties))
 		.SetDisplayName(LOCTEXT("BindingPropertiesTab", "Binding"))
 		.SetGroup(WorkspaceMenuCategory.ToSharedRef())
-		.SetIcon(FSlateIcon(FAppStyle::GetAppStyleSetName(), "LevelEditor.Tabs.Details"));
+		.SetIcon(FSlateIcon(FAppStyle::GetAppStyleSetName(), "Icons.SkeletalMesh"));
 }
 
 void FGroomCustomAssetEditorToolkit::UnregisterTabSpawners(const TSharedRef<class FTabManager>& InTabManager)
@@ -160,10 +162,50 @@ void FGroomCustomAssetEditorToolkit::OnSkeletalGroomTargetChanged(USkeletalMesh
 	}	
 }
 
+bool FGroomCustomAssetEditorToolkit::OnShouldFilterAnimAsset(const FAssetData& AssetData)
+{
+	// Check the compatible skeletons.
+	if (PreviewSkeletalMeshComponent != nullptr)
+	{
+		USkeleton* Skeleton = PreviewSkeletalMeshComponent->GetSkeletalMesh()->GetSkeleton();
+		if (Skeleton && Skeleton->IsCompatibleSkeletonByAssetData(AssetData))
+		{
+			return false;
+		}
+	}
+	return true;
+}
+
+void FGroomCustomAssetEditorToolkit::OnObjectChangedAnimAsset(const FAssetData& AssetData)
+{
+	if (UAnimationAsset* AnimationAsset = Cast<UAnimationAsset>(AssetData.GetAsset()))
+	{
+		if (PreviewSkeletalMeshComponent != nullptr)
+		{
+			PreviewSkeletalAnimationAsset = AnimationAsset;
+			PreviewSkeletalMeshComponent->SetAnimationMode(EAnimationMode::AnimationSingleNode);
+			PreviewSkeletalMeshComponent->PlayAnimation(AnimationAsset, true);
+		}
+	}
+}
+
+FString FGroomCustomAssetEditorToolkit::GetCurrentAnimAssetPath() const
+{
+	if (PreviewSkeletalAnimationAsset != nullptr)
+	{
+		return PreviewSkeletalAnimationAsset->GetPathName();
+	}
+	return FString();
+}
+
+bool FGroomCustomAssetEditorToolkit::OnIsEnabledAnimAsset()
+{
+	return true;
+}
+
 void FGroomCustomAssetEditorToolkit::ExtendToolbar()
 {
 	// Disable simulation toolbar as it is currently not hooked
-#if 0 
 	struct Local
 	{
 		static TSharedRef<SWidget> FillSimulationOptionsMenu(FGroomCustomAssetEditorToolkit* Toolkit)
@@ -177,19 +219,43 @@ void FGroomCustomAssetEditorToolkit::ExtendToolbar()
 			return MenuBuilder.MakeWidget();
 		}
 
+		static TSharedRef<SWidget> FillAnimationOptionsMenu(FGroomCustomAssetEditorToolkit* Toolkit)
+		{
+			FMenuBuilder MenuBuilder(true, Toolkit->GetToolkitCommands());
+
+			MenuBuilder.AddMenuEntry(FGroomEditorCommands::Get().PlayAnimation);
+			MenuBuilder.AddMenuEntry(FGroomEditorCommands::Get().StopAnimation);
+
+			return MenuBuilder.MakeWidget();
+		}
+
 		static void FillToolbar(FToolBarBuilder& ToolbarBuilder, FGroomCustomAssetEditorToolkit* Toolkit)
 		{			
-			ToolbarBuilder.AddToolBarButton( FGroomEditorCommands::Get().Simulate );
-
-			ToolbarBuilder.BeginSection( "PlaybackOptions" );
-			{				
-				ToolbarBuilder.AddComboButton(
-					FUIAction(),
-					FOnGetContent::CreateStatic(Local::FillSimulationOptionsMenu, Toolkit),
-					LOCTEXT("SimulationOptions", "Simulation"),
-					LOCTEXT("SimulationOptionsTooltip", "Simulation options"),
-					FSlateIcon(FGroomEditorStyle::GetStyleName(), "GroomEditor.SimulationOptions")
-				);
+			ToolbarBuilder.BeginSection( "Animation" );
+			{
+				FSlateIcon PlayIcon(FAppStyle::Get().GetStyleSetName(), "Icons.Toolbar.Play");
+				FSlateIcon StopIcon(FAppStyle::Get().GetStyleSetName(), "Icons.Toolbar.Stop");
+				ToolbarBuilder.AddToolBarButton(FGroomEditorCommands::Get().PlayAnimation, NAME_None, TAttribute<FText>(LOCTEXT("Groom.EmptyPlay", "")), TAttribute<FText>(), TAttribute<FSlateIcon>(PlayIcon));
+				ToolbarBuilder.AddToolBarButton(FGroomEditorCommands::Get().StopAnimation, NAME_None, TAttribute<FText>(LOCTEXT("Groom.EmptyStop", "")), TAttribute<FText>(), TAttribute<FSlateIcon>(StopIcon));
+
+				TAttribute<FText> AnimWidgetText(LOCTEXT("AnimationOptions", "Animation"));
+				
+				FIntPoint ThumbnailOverride = FIntPoint(64, 64);
+
+				TSharedRef<SWidget> PropWidget = SNew(SObjectPropertyEntryBox)
+					.ThumbnailPool(Toolkit->ThumbnailPool)
+					.ObjectPath(Toolkit, &FGroomCustomAssetEditorToolkit::GetCurrentAnimAssetPath)
+					.AllowedClass(UAnimationAsset::StaticClass())
+					.AllowClear(true)
+					.DisplayBrowse(false)
+					.DisplayUseSelected(false)
+					.DisplayThumbnail(true)
+					.DisplayCompactSize(true)
+					.ThumbnailSizeOverride(ThumbnailOverride)
+					.OnObjectChanged(FOnAssetSelected::CreateSP(Toolkit, &FGroomCustomAssetEditorToolkit::OnObjectChangedAnimAsset))
+					.OnShouldFilterAsset(FOnShouldFilterAsset::CreateSP(Toolkit, &FGroomCustomAssetEditorToolkit::OnShouldFilterAnimAsset))
+					.OnIsEnabled(FOnIsEnabled::CreateSP(Toolkit, &FGroomCustomAssetEditorToolkit::OnIsEnabledAnimAsset));
+				ToolbarBuilder.AddToolBarWidget(PropWidget);
 			}
 			ToolbarBuilder.EndSection();
 		}
@@ -207,7 +273,6 @@ void FGroomCustomAssetEditorToolkit::ExtendToolbar()
 	AddToolbarExtender(ToolbarExtender);
 
 	FGroomEditor& GroomEditorModule = FModuleManager::LoadModuleChecked<FGroomEditor>("HairStrandsEditor"); // GroomEditor
-#endif
 }
 
 void FGroomCustomAssetEditorToolkit::InitPreviewComponents()
@@ -288,22 +353,61 @@ static void ListAllBindingAssets(const UGroomAsset* InGroomAsset, TWeakObjectPtr
 	}
 }
 
+static UAnimationAsset* GetFirstCompatibleAnimAsset(const USkeletalMeshComponent* InComponent)
+{
+	if (!InComponent)
+	{
+		return nullptr;
+	}
+
+	USkeleton* InSkeleton = InComponent->GetSkeletalMesh()->GetSkeleton();
+	if (!InSkeleton)
+	{
+		return nullptr;
+	}
+
+	const FString SkeletonString = FAssetData(InSkeleton).GetExportTextName();
+
+	FARFilter Filter;
+	Filter.ClassPaths.Add(UAnimationAsset::StaticClass()->GetClassPathName());
+	//Filter.TagsAndValues.Add(USkeletalMesh::GetSkeletonMemberName(), SkeletonString);
+
+	FAssetRegistryModule& AssetRegistryModule = FModuleManager::LoadModuleChecked<FAssetRegistryModule>("AssetRegistry");
+	TArray<FAssetData> AnimAssetData;
+	AssetRegistryModule.Get().GetAssets(Filter, AnimAssetData);
+
+	// Filter binding asset which match the groom asset (as the tag/value filter above does not work)
+	for (FAssetData& Asset : AnimAssetData)
+	{
+		if (InSkeleton->IsCompatibleSkeletonByAssetData(Asset))
+		{
+			return (UAnimationAsset*)Asset.GetAsset();
+		}
+	}
+	return nullptr;
+}
+
 void FGroomCustomAssetEditorToolkit::InitCustomAssetEditor(const EToolkitMode::Type Mode, const TSharedPtr<class IToolkitHost>& InitToolkitHost, UGroomAsset* InCustomAsset)
 {
-	PreviewGroomComponent		 = nullptr;
+	PreviewGroomComponent = nullptr;
 	PreviewSkeletalMeshComponent = nullptr;
+	PreviewSkeletalAnimationAsset = nullptr;
 	GroomBindingAssetList = nullptr;
 
 	ViewportTab = SNew(SGroomEditorViewport);
+	ThumbnailPool = MakeShared<FAssetThumbnailPool>(64);
 
 	GroomBindingAssetList = NewObject<UGroomBindingAssetList>(GetTransientPackage(), NAME_None, RF_Transient);
 	ListAllBindingAssets(InCustomAsset, GroomBindingAssetList);
+	// Automatically affect the first skelal mesh compatible with the groom asset
+	#if 0
 	if (GroomBindingAssetList->Bindings.Num() > 0)
 	{
 		ActiveGroomBindingIndex = 0;
 		GroomBindingAsset = GroomBindingAssetList->Bindings[ActiveGroomBindingIndex];
 	}
 	else
+	#endif
 	{
 		ActiveGroomBindingIndex = -1;
 	}
@@ -472,16 +576,24 @@ void FGroomCustomAssetEditorToolkit::InitCustomAssetEditor(const EToolkitMode::T
 	const TSharedRef<FUICommandList>& CommandList = GetToolkitCommands();
 		
 	CommandList->MapAction(FGroomEditorCommands::Get().ResetSimulation,
-			FExecuteAction::CreateSP(this, &FGroomCustomAssetEditorToolkit::OnResetSimulation),
-			FCanExecuteAction::CreateSP(this, &FGroomCustomAssetEditorToolkit::CanResetSimulation));
+		FExecuteAction::CreateSP(this, &FGroomCustomAssetEditorToolkit::OnResetSimulation),
+		FCanExecuteAction::CreateSP(this, &FGroomCustomAssetEditorToolkit::CanResetSimulation));
 
 	CommandList->MapAction(FGroomEditorCommands::Get().PauseSimulation,
-			FExecuteAction::CreateSP(this, &FGroomCustomAssetEditorToolkit::OnPauseSimulation),
-			FCanExecuteAction::CreateSP(this, &FGroomCustomAssetEditorToolkit::CanResetSimulation));
+		FExecuteAction::CreateSP(this, &FGroomCustomAssetEditorToolkit::OnPauseSimulation),
+		FCanExecuteAction::CreateSP(this, &FGroomCustomAssetEditorToolkit::CanResetSimulation));
 
 	CommandList->MapAction(FGroomEditorCommands::Get().PlaySimulation,
-			FExecuteAction::CreateSP(this, &FGroomCustomAssetEditorToolkit::OnPlaySimulation),
-			FCanExecuteAction::CreateSP(this, &FGroomCustomAssetEditorToolkit::CanPlaySimulation));
+		FExecuteAction::CreateSP(this, &FGroomCustomAssetEditorToolkit::OnPlaySimulation),
+		FCanExecuteAction::CreateSP(this, &FGroomCustomAssetEditorToolkit::CanPlaySimulation));
+
+	CommandList->MapAction(FGroomEditorCommands::Get().PlayAnimation,
+		FExecuteAction::CreateSP(this, &FGroomCustomAssetEditorToolkit::OnPlayAnimation),
+		FCanExecuteAction::CreateSP(this, &FGroomCustomAssetEditorToolkit::CanPlayAnimation));
+
+	CommandList->MapAction(FGroomEditorCommands::Get().StopAnimation,
+		FExecuteAction::CreateSP(this, &FGroomCustomAssetEditorToolkit::OnStopAnimation),
+		FCanExecuteAction::CreateSP(this, &FGroomCustomAssetEditorToolkit::CanStopAnimation));
 
 	// Delegate to refresh the panels (and the stats) when the groom asset change
 	if (GroomAsset.IsValid())
@@ -502,6 +614,9 @@ void FGroomCustomAssetEditorToolkit::InitCustomAssetEditor(const EToolkitMode::T
 	}
 }
 
+///////////////////////////////////////////////////////////////////////////////////////////////////
+// Simulation
+
 void FGroomCustomAssetEditorToolkit::OnPlaySimulation()
 {
 	//if (NiagaraComponent != nullptr)
@@ -540,11 +655,44 @@ bool FGroomCustomAssetEditorToolkit::CanResetSimulation() const
 	return true;
 }
 
+///////////////////////////////////////////////////////////////////////////////////////////////////
+// Animation
+
+void FGroomCustomAssetEditorToolkit::OnPlayAnimation()
+{
+	if (PreviewSkeletalMeshComponent != nullptr && PreviewSkeletalAnimationAsset != nullptr)
+	{
+		PreviewSkeletalMeshComponent->PlayAnimation(PreviewSkeletalAnimationAsset.Get(), true);
+	}
+}
+
+bool FGroomCustomAssetEditorToolkit::CanPlayAnimation() const
+{
+	return PreviewSkeletalMeshComponent != nullptr && PreviewSkeletalAnimationAsset != nullptr && !PreviewSkeletalMeshComponent->IsPlaying();
+}
+
+void FGroomCustomAssetEditorToolkit::OnStopAnimation()
+{
+	if (PreviewSkeletalMeshComponent.Get())
+	{
+		PreviewSkeletalMeshComponent->Stop();
+	}
+}
+
+bool FGroomCustomAssetEditorToolkit::CanStopAnimation() const
+{
+	return PreviewSkeletalMeshComponent != nullptr && PreviewSkeletalMeshComponent->IsPlaying();
+}
+
+///////////////////////////////////////////////////////////////////////////////////////////////////
+// Editor
+
 FGroomCustomAssetEditorToolkit::FGroomCustomAssetEditorToolkit()
 {
-	GroomAsset						= nullptr;	
-	PreviewGroomComponent			= nullptr;	
+	GroomAsset						= nullptr;
+	PreviewGroomComponent			= nullptr;
 	PreviewSkeletalMeshComponent	= nullptr;
+	PreviewSkeletalAnimationAsset	= nullptr;
 }
 
 FGroomCustomAssetEditorToolkit::~FGroomCustomAssetEditorToolkit()
@@ -745,6 +893,7 @@ void FGroomCustomAssetEditorToolkit::PreviewBinding(int32 BindingIndex)
 		PreviewSkeletalMeshComponent->Activate(true);
 		PreviewGroomComponent->AttachToComponent(PreviewSkeletalMeshComponent.Get(), FAttachmentTransformRules::KeepRelativeTransform);
 		PreviewGroomComponent->SetGroomAsset(GroomAsset.Get(), GroomBindingAsset.Get());
+		PreviewSkeletalAnimationAsset = GetFirstCompatibleAnimAsset(PreviewSkeletalMeshComponent.Get());
 		ActiveGroomBindingIndex = BindingIndex;
 	}
 	else
@@ -752,6 +901,7 @@ void FGroomCustomAssetEditorToolkit::PreviewBinding(int32 BindingIndex)
 		PreviewGroomComponent->SetBindingAsset(nullptr);
 		PreviewGroomComponent->DetachFromComponent(FDetachmentTransformRules::KeepRelativeTransform);
 		PreviewSkeletalMeshComponent = nullptr;
+		PreviewSkeletalAnimationAsset = nullptr;
 		ActiveGroomBindingIndex = -1;
 	}
 
diff --git a/Engine/Plugins/Runtime/HairStrands/Source/HairStrandsEditor/Private/GroomEditorCommands.cpp b/Engine/Plugins/Runtime/HairStrands/Source/HairStrandsEditor/Private/GroomEditorCommands.cpp
--- a/Engine/Plugins/Runtime/HairStrands/Source/HairStrandsEditor/Private/GroomEditorCommands.cpp
+++ b/Engine/Plugins/Runtime/HairStrands/Source/HairStrandsEditor/Private/GroomEditorCommands.cpp
@@ -25,6 +25,8 @@ void FGroomEditorCommands::RegisterCommands()
 	UI_COMMAND(PauseSimulation,		"Pause",	"Pause the running hair simulation",	EUserInterfaceActionType::Button,	FInputChord());
 	UI_COMMAND(PlaySimulation,		"Play",		"Play the hair simulation",				EUserInterfaceActionType::Button,	FInputChord());
 
+	UI_COMMAND(PlayAnimation,		"Play",		"Play animation",						EUserInterfaceActionType::Button, FInputChord());
+	UI_COMMAND(StopAnimation,		"Stop",		"Stop animation",						EUserInterfaceActionType::Button, FInputChord());
 
 	UI_COMMAND(ViewMode_Lit,			"Lit",				"Default",						EUserInterfaceActionType::Button, FInputChord());
 	UI_COMMAND(ViewMode_Guide,			"Guide",			"View guide (only available if the groom has simulation enabled)",					EUserInterfaceActionType::Button, FInputChord());
diff --git a/Engine/Plugins/Runtime/HairStrands/Source/HairStrandsEditor/Public/GroomCustomAssetEditorToolkit.h b/Engine/Plugins/Runtime/HairStrands/Source/HairStrandsEditor/Public/GroomCustomAssetEditorToolkit.h
--- a/Engine/Plugins/Runtime/HairStrands/Source/HairStrandsEditor/Public/GroomCustomAssetEditorToolkit.h
+++ b/Engine/Plugins/Runtime/HairStrands/Source/HairStrandsEditor/Public/GroomCustomAssetEditorToolkit.h
@@ -21,6 +21,7 @@ class USkeletalMeshComponent;
 class UGroomAsset;
 class UGroomBindingAsset;
 class UGroomBindingAssetList;
+class UAnimationAsset;
 
 class HAIRSTRANDSEDITOR_API IGroomCustomAssetEditorToolkit : public FAssetEditorToolkit
 {
@@ -95,6 +96,14 @@ private:
 	void OnResetSimulation();
 	bool CanResetSimulation() const;
 
+	// Called when the play animation button is pressed
+	void OnPlayAnimation();
+	bool CanPlayAnimation() const;
+
+	// Called when the stop animation button is pressed
+	void OnStopAnimation();
+	bool CanStopAnimation() const;
+
 	// Add buttons to to the toolbar
 	void ExtendToolbar();
 
@@ -104,6 +113,12 @@ private:
 	// THis is called when the groom target object changes and needs updating
 	void OnSkeletalGroomTargetChanged(USkeletalMesh *NewTarget);
 
+	// Return true if the animation asset should be filtered out (not compatible with the preview skel. mesh)
+	bool OnShouldFilterAnimAsset(const FAssetData& AssetData);
+	void OnObjectChangedAnimAsset(const FAssetData& AssetData);
+	bool OnIsEnabledAnimAsset();
+	FString GetCurrentAnimAssetPath() const;
+
 	// create the custom components we need
 	void InitPreviewComponents();
 
@@ -160,5 +175,8 @@ private:
 
 	TWeakObjectPtr<UGroomComponent> PreviewGroomComponent;
 	TWeakObjectPtr<USkeletalMeshComponent> PreviewSkeletalMeshComponent;
+	TWeakObjectPtr<UAnimationAsset> PreviewSkeletalAnimationAsset;
+
+	TSharedPtr<FAssetThumbnailPool> ThumbnailPool;
 };
 
diff --git a/Engine/Plugins/Runtime/HairStrands/Source/HairStrandsEditor/Public/GroomEditorCommands.h b/Engine/Plugins/Runtime/HairStrands/Source/HairStrandsEditor/Public/GroomEditorCommands.h
--- a/Engine/Plugins/Runtime/HairStrands/Source/HairStrandsEditor/Public/GroomEditorCommands.h
+++ b/Engine/Plugins/Runtime/HairStrands/Source/HairStrandsEditor/Public/GroomEditorCommands.h
@@ -18,6 +18,9 @@ public:
 	TSharedPtr<FUICommandInfo> PauseSimulation;
 	TSharedPtr<FUICommandInfo> PlaySimulation;
 
+	TSharedPtr<FUICommandInfo> PlayAnimation;
+	TSharedPtr<FUICommandInfo> StopAnimation;
+
 	TSharedPtr<FUICommandInfo> Simulate;
 
 	TSharedPtr<FUICommandInfo> ViewMode_Lit;
