#jjra UE-157463
Groom supports reading skinned mesh geometry cache from deformer graph system.

diff --git a/Engine/Plugins/Runtime/HairStrands/Source/HairStrandsCore/Private/GroomDebug.cpp b/Engine/Plugins/Runtime/HairStrands/Source/HairStrandsCore/Private/GroomDebug.cpp
--- a/Engine/Plugins/Runtime/HairStrands/Source/HairStrandsCore/Private/GroomDebug.cpp
+++ b/Engine/Plugins/Runtime/HairStrands/Source/HairStrandsCore/Private/GroomDebug.cpp
@@ -3,12 +3,13 @@
 #include "GeometryCacheComponent.h"
 #include "GroomInstance.h"
 #include "GroomManager.h"
-#include "GPUSkinCache.h"
 #include "HairStrandsMeshProjection.h"
 #include "HairStrandsInterface.h"
 #include "CommonRenderResources.h"
+#include "CachedGeometry.h"
 #include "GroomGeometryCache.h"
 #include "Components/SkeletalMeshComponent.h"
+#include "SkeletalRenderPublic.h"
 #include "UnrealEngine.h"
 #include "SystemTextures.h"
 #include "CanvasTypes.h"
@@ -72,7 +73,6 @@ bool IsHairStrandsSkinCacheEnable();
 FCachedGeometry GetCacheGeometryForHair(
 	FRDGBuilder& GraphBuilder,
 	FHairGroupInstance* Instance,
-	const FGPUSkinCache* SkinCache,
 	FGlobalShaderMap* ShaderMap,
 	const bool bOutputTriangleData);
 
@@ -81,7 +81,6 @@ static void GetGroomInterpolationData(
 	FGlobalShaderMap* ShaderMap,
 	const FHairStrandsInstances& Instances,
 	const EHairStrandsProjectionMeshType MeshType,
-	const FGPUSkinCache* SkinCache,
 	FHairStrandsProjectionMeshData::LOD& OutGeometries)
 {
 	for (FHairStrandsInstance* AbstractInstance : Instances)
@@ -91,22 +90,23 @@ static void GetGroomInterpolationData(
 		if (!Instance || !Instance->Debug.MeshComponent)
 			continue;
 
-		const FCachedGeometry CachedGeometry = GetCacheGeometryForHair(GraphBuilder, Instance, SkinCache,ShaderMap, true);
+		const FCachedGeometry CachedGeometry = GetCacheGeometryForHair(GraphBuilder, Instance, ShaderMap, true);
 		if (CachedGeometry.Sections.Num() == 0)
 			continue;
 
 		if (MeshType == EHairStrandsProjectionMeshType::DeformedMesh || MeshType == EHairStrandsProjectionMeshType::RestMesh)
 		{
-			for (const FCachedGeometry::Section& Section : CachedGeometry.Sections)
+			for (int32 SectionIndex = 0; SectionIndex < CachedGeometry.Sections.Num(); ++SectionIndex)
 			{
-				FHairStrandsProjectionMeshData::Section OutSection = ConvertMeshSection(Section, CachedGeometry.LocalToWorld);
+				FHairStrandsProjectionMeshData::Section OutSection = ConvertMeshSection(CachedGeometry, SectionIndex);
 				if (MeshType == EHairStrandsProjectionMeshType::RestMesh)
 				{					
 					// If the mesh has some mesh-tranferred data, we display that otherwise we use the rest data
-					const bool bHasTransferData = Section.LODIndex < Instance->Debug.TransferredPositions.Num();
+					const int32 SectionLodIndex = CachedGeometry.Sections[SectionIndex].LODIndex;
+					const bool bHasTransferData = SectionLodIndex < Instance->Debug.TransferredPositions.Num();
 					if (bHasTransferData)
 					{
-						OutSection.PositionBuffer = Instance->Debug.TransferredPositions[Section.LODIndex].SRV;
+						OutSection.PositionBuffer = Instance->Debug.TransferredPositions[SectionLodIndex].SRV;
 					}
 					else if (Instance->Debug.TargetMeshData.LODs.Num() > 0)
 					{
@@ -1043,7 +1043,6 @@ void RunHairStrandsDebug(
 	const FSceneView& View,
 	const FHairStrandsInstances& Instances,
 	const FUintVector4& InstanceCountPerType,
-	const FGPUSkinCache* SkinCache,
 	const FShaderPrintData* ShaderPrintData,
 	FRDGTextureRef SceneColorTexture,
 	FRDGTextureRef SceneDepthTexture,
@@ -1062,10 +1061,10 @@ void RunHairStrandsDebug(
 		{
 			if (GHairDebugMeshProjection_SkinCacheMesh > 0)
 			{
-				auto RenderMeshProjection = [ShaderMap, ShaderPrintData, Viewport, &ViewUniformBuffer, Instances, SkinCache, &GraphBuilder](FRDGBuilder& LocalGraphBuilder, EHairStrandsProjectionMeshType MeshType)
+				auto RenderMeshProjection = [ShaderMap, ShaderPrintData, Viewport, &ViewUniformBuffer, Instances, &GraphBuilder](FRDGBuilder& LocalGraphBuilder, EHairStrandsProjectionMeshType MeshType)
 				{
 					FHairStrandsProjectionMeshData::LOD MeshProjectionLODData;
-					GetGroomInterpolationData(GraphBuilder, ShaderMap, Instances, MeshType, SkinCache, MeshProjectionLODData);
+					GetGroomInterpolationData(GraphBuilder, ShaderMap, Instances, MeshType, MeshProjectionLODData);
 					for (FHairStrandsProjectionMeshData::Section& Section : MeshProjectionLODData.Sections)
 					{
 						AddDebugProjectionMeshPass(LocalGraphBuilder, ShaderMap, ShaderPrintData, Viewport, ViewUniformBuffer, MeshType, Section);
@@ -1179,4 +1178,4 @@ void RunHairStrandsDebug(
 			AddDrawDebugStrandsCVsPass(GraphBuilder, View, ShaderMap, Instance, ShaderPrintData, SceneColorTexture, SceneDepthTexture);
 		}
 	}
-}
+}
\ No newline at end of file
diff --git a/Engine/Plugins/Runtime/HairStrands/Source/HairStrandsCore/Private/GroomGeometryCache.cpp b/Engine/Plugins/Runtime/HairStrands/Source/HairStrandsCore/Private/GroomGeometryCache.cpp
--- a/Engine/Plugins/Runtime/HairStrands/Source/HairStrandsCore/Private/GroomGeometryCache.cpp
+++ b/Engine/Plugins/Runtime/HairStrands/Source/HairStrandsCore/Private/GroomGeometryCache.cpp
@@ -6,7 +6,7 @@
 #include "GeometryCacheComponent.h"
 #include "GeometryCacheMeshData.h"
 #include "GeometryCacheSceneProxy.h"
-#include "GPUSkinCache.h"
+#include "CachedGeometry.h"
 #include "Rendering/SkeletalMeshRenderData.h"
 #include "Rendering/SkinWeightVertexBuffer.h"
 #include "CommonRenderResources.h"
@@ -17,8 +17,9 @@
 
 ///////////////////////////////////////////////////////////////////////////////////////////////////
 
-FHairStrandsProjectionMeshData::Section ConvertMeshSection(const FCachedGeometry::Section& In, const FTransform& InLocalToWorld)
+FHairStrandsProjectionMeshData::Section ConvertMeshSection(FCachedGeometry const& InCachedGeometry, int32 InSectionIndex)
 {
+	FCachedGeometry::Section const& In = InCachedGeometry.Sections[InSectionIndex];
 	FHairStrandsProjectionMeshData::Section Out;
 	Out.IndexBuffer = In.IndexBuffer;
 	Out.RDGPositionBuffer = In.RDGPositionBuffer;
@@ -36,7 +37,7 @@ FHairStrandsProjectionMeshData::Section ConvertMeshSection(const FCachedGeometry
 	Out.NumVertices= In.NumVertices;
 	Out.SectionIndex = In.SectionIndex;
 	Out.LODIndex = In.LODIndex;
-	Out.LocalToWorld = InLocalToWorld;
+	Out.LocalToWorld = InCachedGeometry.LocalToWorld;
 	return Out;
 }
 
diff --git a/Engine/Plugins/Runtime/HairStrands/Source/HairStrandsCore/Private/GroomGeometryCache.h b/Engine/Plugins/Runtime/HairStrands/Source/HairStrandsCore/Private/GroomGeometryCache.h
--- a/Engine/Plugins/Runtime/HairStrands/Source/HairStrandsCore/Private/GroomGeometryCache.h
+++ b/Engine/Plugins/Runtime/HairStrands/Source/HairStrandsCore/Private/GroomGeometryCache.h
@@ -3,18 +3,15 @@
 #pragma once
 
 #include "CoreMinimal.h"
-#include "RendererInterface.h"
-#include "Containers/Array.h"
-#include "Engine/EngineTypes.h"
-#include "Shader.h"
-#include "RenderResource.h"
-#include "RenderGraphResources.h"
 #include "HairStrandsMeshProjection.h"
-#include "GPUSkinCache.h"
 
+struct FCachedGeometry;
 class UGeometryCacheComponent;
+class FGlobalShaderMap;
+class FRDGBuilder;
+class USkeletalMeshComponent;
 
-FHairStrandsProjectionMeshData::Section ConvertMeshSection(const FCachedGeometry::Section& In, const FTransform& T);
+FHairStrandsProjectionMeshData::Section ConvertMeshSection(FCachedGeometry const& InCachedGeometry, int32 InSectionIndex);
 
 void BuildCacheGeometry(
 	FRDGBuilder& GraphBuilder,
diff --git a/Engine/Plugins/Runtime/HairStrands/Source/HairStrandsCore/Private/GroomManager.cpp b/Engine/Plugins/Runtime/HairStrands/Source/HairStrandsCore/Private/GroomManager.cpp
--- a/Engine/Plugins/Runtime/HairStrands/Source/HairStrandsCore/Private/GroomManager.cpp
+++ b/Engine/Plugins/Runtime/HairStrands/Source/HairStrandsCore/Private/GroomManager.cpp
@@ -4,7 +4,6 @@
 #include "HairStrandsMeshProjection.h"
 
 #include "GeometryCacheComponent.h"
-#include "GPUSkinCache.h"
 #include "Rendering/SkeletalMeshRenderData.h"
 #include "Rendering/SkinWeightVertexBuffer.h"
 #include "CommonRenderResources.h"
@@ -20,6 +19,7 @@
 #include "HairStrandsVertexFactory.h"
 #include "RenderGraphEvent.h"
 #include "RenderGraphUtils.h"
+#include "CachedGeometry.h"
 
 static int32 GHairStrandsMinLOD = 0;
 static FAutoConsoleVariableRef CVarGHairStrandsMinLOD(TEXT("r.HairStrands.MinLOD"), GHairStrandsMinLOD, TEXT("Clamp the min hair LOD to this value, preventing to reach lower/high-quality LOD."), ECVF_Scalability);
@@ -86,7 +86,6 @@ static bool IsSkeletalMeshEvaluationEnabled()
 FCachedGeometry GetCacheGeometryForHair(
 	FRDGBuilder& GraphBuilder, 
 	FHairGroupInstance* Instance, 
-	const FGPUSkinCache* SkinCache, 
 	FGlobalShaderMap* ShaderMap,
 	const bool bOutputTriangleData)
 {
@@ -95,11 +94,12 @@ FCachedGeometry GetCacheGeometryForHair(
 	{
 		if (IsSkeletalMeshEvaluationEnabled())
 		{
+			//todo: It's unsafe to be accessing the component directly here. This can be solved when we have persistent ids available on the render thread.
 			if (const USkeletalMeshComponent* SkeletalMeshComponent = Cast<USkeletalMeshComponent>(Instance->Debug.MeshComponent))
 			{
-				if (SkinCache)
+				if (FSkeletalMeshSceneProxy* SceneProxy = static_cast<FSkeletalMeshSceneProxy*>(SkeletalMeshComponent->SceneProxy))
 				{
-					Out = SkinCache->GetCachedGeometry(SkeletalMeshComponent->ComponentId.PrimIDValue, EGPUSkinCacheEntryMode::Raster);
+					SceneProxy->GetCachedGeometry(Out);
 				}
 
 				if (IsHairStrandsSkinCacheEnable() && Out.Sections.Num() == 0)
@@ -141,7 +141,6 @@ static void RunInternalHairStrandsInterpolation(
 	const FSceneView* View,
 	const uint32 ViewUniqueID,
 	const FHairStrandsInstances& Instances,
-	const FGPUSkinCache* SkinCache,
 	const FShaderPrintData* ShaderPrintData,
 	FGlobalShaderMap* ShaderMap, 
 	EHairStrandsInterpolationType Type,
@@ -167,14 +166,15 @@ static void RunInternalHairStrandsInterpolation(
 		check(Instance->HairGroupPublicData);
 
 		FHairStrandsProjectionMeshData::LOD MeshDataLOD;
-		const FCachedGeometry CachedGeometry = GetCacheGeometryForHair(GraphBuilder, Instance, SkinCache, ShaderMap, true);
-		for (const FCachedGeometry::Section& Section : CachedGeometry.Sections)
+		const FCachedGeometry CachedGeometry = GetCacheGeometryForHair(GraphBuilder, Instance, ShaderMap, true);
+		for (int32 SectionIndex = 0; SectionIndex < CachedGeometry.Sections.Num(); ++SectionIndex)
 		{
 			// Ensure all mesh's sections have the same LOD index
-			if (MeshLODIndex < 0) MeshLODIndex = Section.LODIndex;
-			check(MeshLODIndex == Section.LODIndex);
+			const int32 SectionLodIndex = CachedGeometry.Sections[SectionIndex].LODIndex;
+			if (MeshLODIndex < 0) MeshLODIndex = SectionLodIndex;
+			check(MeshLODIndex == SectionLodIndex);
 
-			MeshDataLOD.Sections.Add(ConvertMeshSection(Section, CachedGeometry.LocalToWorld));
+			MeshDataLOD.Sections.Add(ConvertMeshSection(CachedGeometry, SectionIndex));
 		}
 
 		Instance->Debug.MeshLODIndex = MeshLODIndex;
@@ -399,7 +399,6 @@ static void RunHairStrandsInterpolation_Guide(
 	const FSceneView* View,
 	const uint32 ViewUniqueID,
 	const FHairStrandsInstances& Instances,
-	const FGPUSkinCache* SkinCache,
 	const FShaderPrintData* ShaderPrintData,
 	FGlobalShaderMap* ShaderMap,
 	FHairStrandClusterData* ClusterData)
@@ -415,7 +414,6 @@ static void RunHairStrandsInterpolation_Guide(
 		View,
 		ViewUniqueID,
 		Instances,
-		SkinCache,
 		ShaderPrintData,
 		ShaderMap,
 		EHairStrandsInterpolationType::SimulationStrands,
@@ -427,7 +425,6 @@ static void RunHairStrandsInterpolation_Strands(
 	const FSceneView* View,
 	const uint32 ViewUniqueID,
 	const FHairStrandsInstances& Instances,
-	const FGPUSkinCache* SkinCache,
 	const FShaderPrintData* ShaderPrintData,
 	FGlobalShaderMap* ShaderMap,
 	FHairStrandClusterData* ClusterData)
@@ -443,7 +440,6 @@ static void RunHairStrandsInterpolation_Strands(
 		View,
 		ViewUniqueID,
 		Instances,
-		SkinCache,
 		ShaderPrintData,
 		ShaderMap,
 		EHairStrandsInterpolationType::RenderStrands,
@@ -691,7 +687,6 @@ static void RunHairLODSelection(
 	FRDGBuilder& GraphBuilder, 
 	const FHairStrandsInstances& Instances, 
 	const TArray<const FSceneView*>& Views, 
-	const FGPUSkinCache* SkinCache, 
 	FGlobalShaderMap* ShaderMap)
 {
 	EShaderPlatform ShaderPlatform = EShaderPlatform::SP_NumPlatforms;
@@ -719,7 +714,7 @@ static void RunHairLODSelection(
 
 		check(Instance);
 		check(Instance->HairGroupPublicData);
-		const FCachedGeometry CachedGeometry = GetCacheGeometryForHair(GraphBuilder, Instance, SkinCache, ShaderMap, false);
+		const FCachedGeometry CachedGeometry = GetCacheGeometryForHair(GraphBuilder, Instance, ShaderMap, false);
 		const int32 MeshLODIndex = GetCacheGeometryLODIndex(CachedGeometry);
 
 		// Perform LOD selection based on all the views	
@@ -1032,7 +1027,6 @@ void RunHairStrandsDebug(
 	const FSceneView& View,
 	const FHairStrandsInstances& Instances,
 	const FUintVector4& InstanceCountPerType,
-	const FGPUSkinCache* SkinCache,
 	const FShaderPrintData* ShaderPrintData,
 	FRDGTextureRef SceneColor,
 	FRDGTextureRef SceneDepth,
@@ -1084,7 +1078,6 @@ void ProcessHairStrandsBookmark(
 			*GraphBuilder,
 			*Parameters.Instances,
 			Parameters.AllViews,
-			Parameters.SkinCache, 
 			Parameters.ShaderMap);
 	}
 	else if (Bookmark == EHairStrandsBookmark::ProcessEndOfFrame)
@@ -1104,7 +1097,6 @@ void ProcessHairStrandsBookmark(
 			Parameters.View,
 			Parameters.ViewUniqueID,
 			Instances,
-			Parameters.SkinCache,
 			Parameters.ShaderPrintData,
 			Parameters.ShaderMap,
 			&Parameters.HairClusterData);
@@ -1123,7 +1115,6 @@ void ProcessHairStrandsBookmark(
 			Parameters.View,
 			Parameters.ViewUniqueID,
 			Instances,
-			Parameters.SkinCache,
 			Parameters.ShaderPrintData,
 			Parameters.ShaderMap,
 			&Parameters.HairClusterData);
@@ -1137,7 +1128,6 @@ void ProcessHairStrandsBookmark(
 			*Parameters.View,
 			Instances,
 			Parameters.InstanceCountPerType,
-			Parameters.SkinCache,
 			Parameters.ShaderPrintData,
 			Parameters.SceneColorTexture,
 			Parameters.SceneDepthTexture,
diff --git a/Engine/Source/Runtime/Engine/Private/GPUSkinCache.cpp b/Engine/Source/Runtime/Engine/Private/GPUSkinCache.cpp
--- a/Engine/Source/Runtime/Engine/Private/GPUSkinCache.cpp
+++ b/Engine/Source/Runtime/Engine/Private/GPUSkinCache.cpp
@@ -595,6 +595,7 @@ public:
 #endif // RHI_RAYTRACING
 
 	TArray<FSectionDispatchData>& GetDispatchData() { return DispatchData; }
+	TArray<FSectionDispatchData> const& GetDispatchData() const { return DispatchData; }
 
 protected:
 	EGPUSkinCacheEntryMode Mode;
@@ -2081,96 +2082,40 @@ void FGPUSkinCache::InvalidateAllEntries()
 	SET_MEMORY_STAT(STAT_GPUSkinCache_TangentsIntermediateMemUsed, 0);
 }
 
-FRWBuffer* FGPUSkinCache::GetPositionBuffer(uint32 ComponentId, uint32 SectionIndex) const
+FGPUSkinCacheEntry const* FGPUSkinCache::GetSkinCacheEntry(uint32 ComponentId) const
 {
 	for (FGPUSkinCacheEntry* Entry : Entries)
 	{
 		if (Entry && Entry->GPUSkin && Entry->GPUSkin->GetComponentId() == ComponentId)
 		{
-			FGPUSkinCacheEntry::FSectionDispatchData& DispatchData = Entry->GetDispatchData()[SectionIndex];
-			FSkinCacheRWBuffer* SkinCacheRWBuffer = DispatchData.PositionBuffer;
-			return SkinCacheRWBuffer != nullptr ? &SkinCacheRWBuffer->Buffer : nullptr;
+			return Entry;
 		}
 	}
-
 	return nullptr;
 }
 
-FRWBuffer* FGPUSkinCache::GetTangentBuffer(uint32 ComponentId, uint32 SectionIndex) const
+FRWBuffer* FGPUSkinCache::GetPositionBuffer(FGPUSkinCacheEntry const* Entry, uint32 SectionIndex)
 {
-	for (FGPUSkinCacheEntry* Entry : Entries)
+	if (Entry)
 	{
-		if (Entry && Entry->GPUSkin && Entry->GPUSkin->GetComponentId() == ComponentId)
-		{
-			FGPUSkinCacheEntry::FSectionDispatchData& DispatchData = Entry->GetDispatchData()[SectionIndex];
-			FSkinCacheRWBuffer* SkinCacheRWBuffer = DispatchData.GetTangentRWBuffer();
-			return SkinCacheRWBuffer != nullptr ? &SkinCacheRWBuffer->Buffer : nullptr;
-		}
+		FGPUSkinCacheEntry::FSectionDispatchData const& DispatchData = Entry->GetDispatchData()[SectionIndex];
+		FSkinCacheRWBuffer* SkinCacheRWBuffer = DispatchData.PositionBuffer;
+		return SkinCacheRWBuffer != nullptr ? &SkinCacheRWBuffer->Buffer : nullptr;
 	}
-
 	return nullptr;
 }
 
-FRHIShaderResourceView* FGPUSkinCache::GetBoneBuffer(uint32 ComponentId, uint32 SectionIndex) const
+FRWBuffer* FGPUSkinCache::GetPreviousPositionBuffer(FGPUSkinCacheEntry const* Entry, uint32 SectionIndex)
 {
-	for (FGPUSkinCacheEntry* Entry : Entries)
+	if (Entry)
 	{
-		FGPUSkinCacheEntry::FSectionDispatchData& DispatchData = Entry->GetDispatchData()[SectionIndex];
-		FGPUBaseSkinVertexFactory::FShaderDataType& ShaderData = DispatchData.SourceVertexFactory->GetShaderData();
-		return ShaderData.GetBoneBufferForReading(false).VertexBufferSRV;
+		FGPUSkinCacheEntry::FSectionDispatchData const& DispatchData = Entry->GetDispatchData()[SectionIndex];
+		FSkinCacheRWBuffer* SkinCacheRWBuffer = DispatchData.PreviousPositionBuffer;
+		return SkinCacheRWBuffer != nullptr ? &SkinCacheRWBuffer->Buffer : nullptr;
 	}
-
 	return nullptr;
 }
 
-FCachedGeometry FGPUSkinCache::GetCachedGeometry(uint32 ComponentId, EGPUSkinCacheEntryMode Mode) const
-{
-	auto FindEntry = [&](EGPUSkinCacheEntryMode InMode, bool bByPassModeCheck, FCachedGeometry& Out) -> bool
-	{
-		for (FGPUSkinCacheEntry* Entry : Entries)
-		{
-			if (Entry && (bByPassModeCheck || Entry->Mode == InMode) && Entry->GPUSkin && Entry->GPUSkin->GetComponentId() == ComponentId && Entry->GPUSkin->HaveValidDynamicData())
-			{
-				const FSkeletalMeshRenderData& RenderData = Entry->GPUSkin->GetSkeletalMeshRenderData();
-				const int32 LODIndex = Entry->LOD;
-				if (LODIndex >= 0 && LODIndex < RenderData.LODRenderData.Num())
-				{
-					const FSkeletalMeshLODRenderData& LODData = RenderData.LODRenderData[LODIndex];
-					const uint32 SectionCount = LODData.RenderSections.Num();
-					for (uint32 SectionIdx = 0; SectionIdx < SectionCount; ++SectionIdx)
-					{
-						FCachedGeometry::Section CachedSection = Entry->GetCachedGeometry(SectionIdx);
-						CachedSection.IndexBuffer = LODData.MultiSizeIndexContainer.GetIndexBuffer()->GetSRV();
-						CachedSection.TotalIndexCount = LODData.MultiSizeIndexContainer.GetIndexBuffer()->Num();
-						CachedSection.LODIndex = LODIndex;
-						CachedSection.UVsChannelOffset = 0; // Assume that we needs to pair meshes based on UVs 0
-						CachedSection.UVsChannelCount = LODData.StaticVertexBuffers.StaticMeshVertexBuffer.GetNumTexCoords();
-						Out.Sections.Add(CachedSection);
-					}
-					Out.LocalToWorld = FTransform(Entry->GPUSkin->GetTransform());
-					Out.LODIndex = LODIndex;
-					return true;
-				}
-			}
-		}
-		return false;
-	};
-
-	// 1. Try to find a Skin cache entry which matches the requested mode (Raster/Raytracing)
-	// 2. If we can't find an entry with a matching mode, use any mode type
-	FCachedGeometry Out;
-	if (!FindEntry(Mode, false, Out))
-	{
-		FindEntry(Mode, true, Out);
-	}
-	return Out;
-}
-
-FCachedGeometry::Section FGPUSkinCache::GetCachedGeometry(FGPUSkinCacheEntry* InOutEntry, uint32 sectionIndex)
-{
-	return InOutEntry ? InOutEntry->GetCachedGeometry(sectionIndex) : FCachedGeometry::Section();
-}
-
 void FGPUSkinCache::UpdateSkinWeightBuffer(FGPUSkinCacheEntry* Entry)
 {
 	if (Entry)
diff --git a/Engine/Source/Runtime/Engine/Private/GPUSkinVertexFactory.cpp b/Engine/Source/Runtime/Engine/Private/GPUSkinVertexFactory.cpp
--- a/Engine/Source/Runtime/Engine/Private/GPUSkinVertexFactory.cpp
+++ b/Engine/Source/Runtime/Engine/Private/GPUSkinVertexFactory.cpp
@@ -879,7 +879,7 @@ IMPLEMENT_TYPE_LAYOUT(FGPUSkinVertexPassthroughFactoryShaderParameters);
 FGPUSkinPassthroughVertexFactory
 -----------------------------------------------------------------------------*/
 FGPUSkinPassthroughVertexFactory::FGPUSkinPassthroughVertexFactory(ERHIFeatureLevel::Type InFeatureLevel)
-	: FLocalVertexFactory(InFeatureLevel, "FGPUSkinPassthroughVertexFactory"), PositionStreamIndex(-1), TangentStreamIndex(-1)
+	: FLocalVertexFactory(InFeatureLevel, "FGPUSkinPassthroughVertexFactory")
 {
 	bSupportsManualVertexFetch = true;
 }
diff --git a/Engine/Source/Runtime/Engine/Private/SkeletalMesh.cpp b/Engine/Source/Runtime/Engine/Private/SkeletalMesh.cpp
--- a/Engine/Source/Runtime/Engine/Private/SkeletalMesh.cpp
+++ b/Engine/Source/Runtime/Engine/Private/SkeletalMesh.cpp
@@ -6143,6 +6143,11 @@ uint8 FSkeletalMeshSceneProxy::GetCurrentFirstLODIdx_Internal() const
 	return SkeletalMeshRenderData->CurrentFirstLODIdx;
 }
 
+bool FSkeletalMeshSceneProxy::GetCachedGeometry(FCachedGeometry& OutCachedGeometry) const 
+{
+	return MeshObject != nullptr && MeshObject->GetCachedGeometry(OutCachedGeometry); 
+}
+
 void FSkeletalMeshSceneProxy::GetDynamicElementsSection(const TArray<const FSceneView*>& Views, const FSceneViewFamily& ViewFamily, uint32 VisibilityMap, 
 	const FSkeletalMeshLODRenderData& LODData, const int32 LODIndex, const int32 SectionIndex, bool bSectionSelected,
 	const FSectionElementInfo& SectionElementInfo, bool bInSelectable, FMeshElementCollector& Collector ) const
diff --git a/Engine/Source/Runtime/Engine/Private/SkeletalRenderGPUSkin.cpp b/Engine/Source/Runtime/Engine/Private/SkeletalRenderGPUSkin.cpp
--- a/Engine/Source/Runtime/Engine/Private/SkeletalRenderGPUSkin.cpp
+++ b/Engine/Source/Runtime/Engine/Private/SkeletalRenderGPUSkin.cpp
@@ -9,6 +9,7 @@
 #include "SceneUtils.h"
 #include "SkeletalRender.h"
 #include "GPUSkinCache.h"
+#include "CachedGeometry.h"
 #include "RayTracingSkinnedGeometry.h"
 #include "Animation/MorphTarget.h"
 #include "ClearQuad.h"
@@ -476,23 +477,22 @@ void FSkeletalMeshObjectGPUSkin::UpdateSkinWeightBuffer(USkinnedMeshComponent* I
 
 			if (InMeshComponent && InMeshComponent->SceneProxy)
 			{
-				FGPUSkinCache* GPUSkinCache = InMeshComponent->SceneProxy->GetScene().GetGPUSkinCache();
 				FGPUSkinCacheEntry* SkinCacheEntryToUpdate = SkinCacheEntry;
-				if (GPUSkinCache && SkinCacheEntryToUpdate)
+				if (SkinCacheEntryToUpdate)
 				{
 					ENQUEUE_RENDER_COMMAND(UpdateSkinCacheSkinWeightBuffer)(
-						[GPUSkinCache, SkinCacheEntryToUpdate](FRHICommandListImmediate& RHICmdList)
+						[SkinCacheEntryToUpdate](FRHICommandListImmediate& RHICmdList)
 					{
-						GPUSkinCache->UpdateSkinWeightBuffer(SkinCacheEntryToUpdate);
+						FGPUSkinCache::UpdateSkinWeightBuffer(SkinCacheEntryToUpdate);
 					});
 				}
 
-				if (GPUSkinCache && SkinCacheEntryForRayTracing)
+				if (SkinCacheEntryForRayTracing)
 				{
 					ENQUEUE_RENDER_COMMAND(UpdateSkinCacheSkinWeightBuffer)(
-						[GPUSkinCache, SkinCacheEntryForRayTracing = SkinCacheEntryForRayTracing](FRHICommandListImmediate& RHICmdList)
+						[SkinCacheEntryForRayTracing = SkinCacheEntryForRayTracing](FRHICommandListImmediate& RHICmdList)
 					{
-						GPUSkinCache->UpdateSkinWeightBuffer(SkinCacheEntryForRayTracing);
+						FGPUSkinCache::UpdateSkinWeightBuffer(SkinCacheEntryForRayTracing);
 					});
 				}
 			}
@@ -2197,6 +2197,76 @@ const TArray<FMatrix44f>& FSkeletalMeshObjectGPUSkin::GetReferenceToLocalMatrice
 	return DynamicData->ReferenceToLocal;
 }
 
+bool FSkeletalMeshObjectGPUSkin::GetCachedGeometry(FCachedGeometry& OutCachedGeometry) const
+{
+	OutCachedGeometry = FCachedGeometry();
+	
+	const int32 LodIndex = GetLOD();
+	if (SkeletalMeshRenderData == nullptr || !SkeletalMeshRenderData->LODRenderData.IsValidIndex(LodIndex))
+	{
+		return false;
+	}
+
+	FSkeletalMeshLODRenderData const& LODRenderData = SkeletalMeshRenderData->LODRenderData[LodIndex];
+	const uint32 SectionCount = LODRenderData.RenderSections.Num();
+
+	FVertexFactoryData const& VertexFactories = LODs[LodIndex].GPUSkinVertexFactories;
+	if (VertexFactories.PassthroughVertexFactories.Num() != SectionCount)
+	{
+		return false;
+	}
+
+	for (uint32 SectionIndex = 0; SectionIndex < SectionCount; ++SectionIndex)
+	{
+		FCachedGeometry::Section& CachedSection = OutCachedGeometry.Sections.AddDefaulted_GetRef();
+
+		if (SkinCacheEntry != nullptr)
+		{
+			// Get the cached geometry SRVs from the skin cache.
+			CachedSection.PositionBuffer = FGPUSkinCache::GetPositionBuffer(SkinCacheEntry, SectionIndex)->SRV;
+			CachedSection.PreviousPositionBuffer = FGPUSkinCache::GetPreviousPositionBuffer(SkinCacheEntry, SectionIndex)->SRV;
+		}
+		else
+		{
+			// Pull the cached geometry SRV directly out of the vertex factory.
+			// This catches cases where we have setup deformed data outside of the skin cache (eg mesh deformers).
+			// This approach _could_ be used for the skin cache case as well, except that skin cache only 
+			// actually pokes the vertex factory once instead of every frame. It flips current/previous buffers
+			// by setting SRVs directly in GetShaderBindings().
+			FGPUSkinPassthroughVertexFactory* VertexFactory = VertexFactories.PassthroughVertexFactories[SectionIndex].Get();
+			if (VertexFactory == nullptr || VertexFactory->GetPositionsSRV() == nullptr)
+			{
+				// Reset all output if one section isn't available.
+				OutCachedGeometry.Sections.Reset();
+				return false;
+			}
+
+			CachedSection.PositionBuffer = VertexFactory->GetPositionsSRV();
+			FRHIShaderResourceView* PreviousPositionBuffer = VertexFactory->GetPreviousPositionsSRV();
+			CachedSection.PreviousPositionBuffer = PreviousPositionBuffer != nullptr ? PreviousPositionBuffer : CachedSection.PositionBuffer;
+		}
+				
+		CachedSection.IndexBuffer = LODRenderData.MultiSizeIndexContainer.GetIndexBuffer()->GetSRV();
+		CachedSection.TotalIndexCount = LODRenderData.MultiSizeIndexContainer.GetIndexBuffer()->Num();
+		CachedSection.TotalVertexCount = LODRenderData.StaticVertexBuffers.PositionVertexBuffer.GetNumVertices();
+		CachedSection.UVsBuffer = LODRenderData.StaticVertexBuffers.StaticMeshVertexBuffer.GetTexCoordsSRV();
+		CachedSection.UVsChannelOffset = 0; // Assume that we needs to pair meshes based on UVs 0
+		CachedSection.UVsChannelCount = LODRenderData.StaticVertexBuffers.StaticMeshVertexBuffer.GetNumTexCoords();
+
+		FSkelMeshRenderSection const& Section = LODRenderData.RenderSections[SectionIndex];
+		CachedSection.LODIndex = LodIndex;
+		CachedSection.SectionIndex = SectionIndex;
+		CachedSection.NumPrimitives = Section.NumTriangles;
+		CachedSection.NumVertices = Section.NumVertices;
+		CachedSection.IndexBaseIndex = Section.BaseIndex;
+		CachedSection.VertexBaseIndex = Section.BaseVertexIndex;
+	}
+			
+	OutCachedGeometry.LODIndex = LodIndex;
+	OutCachedGeometry.LocalToWorld = FTransform(GetTransform());
+	return true;
+}
+
 /*-----------------------------------------------------------------------------
 FDynamicSkelMeshObjectDataGPUSkin
 -----------------------------------------------------------------------------*/
diff --git a/Engine/Source/Runtime/Engine/Private/SkeletalRenderGPUSkin.h b/Engine/Source/Runtime/Engine/Private/SkeletalRenderGPUSkin.h
--- a/Engine/Source/Runtime/Engine/Private/SkeletalRenderGPUSkin.h
+++ b/Engine/Source/Runtime/Engine/Private/SkeletalRenderGPUSkin.h
@@ -342,6 +342,7 @@ public:
 	virtual bool IsCPUSkinned() const override { return false; }
 	virtual TArray<FTransform>* GetComponentSpaceTransforms() const override;
 	virtual const TArray<FMatrix44f>& GetReferenceToLocalMatrices() const override;
+	virtual bool GetCachedGeometry(FCachedGeometry& OutCachedGeometry) const override;
 
 #if RHI_RAYTRACING
 	/** Geometry for ray tracing. */
diff --git a/Engine/Source/Runtime/Engine/Public/CachedGeometry.h b/Engine/Source/Runtime/Engine/Public/CachedGeometry.h
new file mode 100644
--- /dev/null
+++ b/Engine/Source/Runtime/Engine/Public/CachedGeometry.h
@@ -0,0 +1,47 @@
+// Copyright Epic Games, Inc. All Rights Reserved.
+
+#pragma  once
+
+#include "CoreMinimal.h"
+#include "RHI.h"
+#include "RenderGraphDefinitions.h"
+
+/**
+ * The output of some geometry caching system.
+ * For example this might be the GPUSkinCache, MeshDeformer or any other compute process that generates vertex buffers.
+ * The contents can be ingested by some rendering system.
+ * For example this might be used to place the hair groom system against a deformed mesh.
+ */
+struct FCachedGeometry
+{
+	struct Section
+	{
+		FRDGBufferSRVRef RDGPositionBuffer = nullptr;
+		FRDGBufferSRVRef RDGPreviousPositionBuffer = nullptr;
+		
+		FRHIShaderResourceView* PositionBuffer = nullptr;			// Valid when the input comes from the GPUSkinCache (since it is doesn't use RDG)
+		FRHIShaderResourceView* PreviousPositionBuffer = nullptr;	// Valid when the input comes from the GPUSkinCache (since it is doesn't use RDG)
+		
+		FRHIShaderResourceView* UVsBuffer = nullptr;
+		FRHIShaderResourceView* IndexBuffer = nullptr;
+		
+		uint32 UVsChannelOffset = 0;
+		uint32 UVsChannelCount = 0;
+		uint32 NumPrimitives = 0;
+		uint32 NumVertices = 0;
+		uint32 VertexBaseIndex = 0;
+		uint32 IndexBaseIndex = 0;
+		uint32 TotalVertexCount = 0;
+		uint32 TotalIndexCount = 0;
+		uint32 SectionIndex = 0;
+		int32 LODIndex = -1;
+	};
+
+	int32 LODIndex = -1;
+	TArray<Section> Sections;
+
+	FRDGBufferRef DeformedPositionBuffer = nullptr;
+	FRDGBufferRef DeformedPreviousPositionBuffer = nullptr;
+
+	FTransform LocalToWorld = FTransform::Identity;
+};
diff --git a/Engine/Source/Runtime/Engine/Public/GPUSkinCache.h b/Engine/Source/Runtime/Engine/Public/GPUSkinCache.h
--- a/Engine/Source/Runtime/Engine/Public/GPUSkinCache.h
+++ b/Engine/Source/Runtime/Engine/Public/GPUSkinCache.h
@@ -42,6 +42,7 @@
 #include "GPUSkinPublicDefs.h"
 #include "VertexFactory.h"
 #include "CanvasTypes.h"
+#include "CachedGeometry.h"
 
 class FGPUSkinPassthroughVertexFactory;
 class FGPUBaseSkinVertexFactory;
@@ -94,36 +95,6 @@ struct FGPUSkinBatchElementUserData
 	int32 Section;
 };
 
-class FRDGPooledBuffer;
-struct FCachedGeometry
-{
-	struct Section
-	{
-		FRDGBufferSRVRef RDGPositionBuffer = nullptr;				// Valid when the input comes from a manual skin cache (i.e. skinned run into compute on demand)
-		FRDGBufferSRVRef RDGPreviousPositionBuffer = nullptr;		// Valid when the input comes from a manual skin cache (i.e. skinned run into compute on demand)
-		FRHIShaderResourceView* PositionBuffer = nullptr;			// Valid when the input comes from the skin cached (since it is not convert yet to RDG)
-		FRHIShaderResourceView* PreviousPositionBuffer = nullptr;	// Valid when the input comes from the skin cached (since it is not convert yet to RDG)
-		FRHIShaderResourceView* UVsBuffer = nullptr;
-		FRHIShaderResourceView* IndexBuffer = nullptr;
-		uint32 UVsChannelOffset = 0;
-		uint32 UVsChannelCount = 0;
-		uint32 NumPrimitives = 0;
-		uint32 NumVertices = 0;
-		uint32 VertexBaseIndex = 0;
-		uint32 IndexBaseIndex = 0;
-		uint32 TotalVertexCount = 0;
-		uint32 TotalIndexCount = 0;
-		uint32 SectionIndex = 0;
-		int32 LODIndex = -1;
-	};
-
-	int32 LODIndex = -1;
-	TArray<Section> Sections;
-	FRDGBufferRef DeformedPositionBuffer = nullptr;
-	FRDGBufferRef DeformedPreviousPositionBuffer = nullptr;
-	FTransform LocalToWorld = FTransform::Identity;
-};
-
 enum class EGPUSkinCacheEntryMode
 {
 	Raster,
@@ -158,9 +129,7 @@ public:
 	ENGINE_API FGPUSkinCache(ERHIFeatureLevel::Type InFeatureLevel, bool bInRequiresMemoryLimit, UWorld* InWorld);
 	ENGINE_API ~FGPUSkinCache();
 
-	ENGINE_API FCachedGeometry GetCachedGeometry(uint32 ComponentId, EGPUSkinCacheEntryMode InMode) const;
-	FCachedGeometry::Section GetCachedGeometry(FGPUSkinCacheEntry* InOutEntry, uint32 SectionId);
-	void UpdateSkinWeightBuffer(FGPUSkinCacheEntry* Entry);
+	static void UpdateSkinWeightBuffer(FGPUSkinCacheEntry* Entry);
 
 	bool ProcessEntry(
 		EGPUSkinCacheEntryMode Mode,
@@ -419,9 +388,13 @@ public:
 		const FVertexBufferAndSRV* BoneBuffers[NUM_BUFFERS];
 	};
 
-	ENGINE_API FRWBuffer* GetPositionBuffer(uint32 ComponentId, uint32 SectionIndex) const;
-	ENGINE_API FRWBuffer* GetTangentBuffer(uint32 ComponentId, uint32 SectionIndex) const;
-	ENGINE_API FRHIShaderResourceView* GetBoneBuffer(uint32 ComponentId, uint32 SectionIndex) const;
+	FGPUSkinCacheEntry const* GetSkinCacheEntry(uint32 ComponentId) const;
+	static FRWBuffer* GetPositionBuffer(FGPUSkinCacheEntry const* Entry, uint32 SectionIndex);
+	static FRWBuffer* GetPreviousPositionBuffer(FGPUSkinCacheEntry const* Entry, uint32 SectionIndex);
+
+	// Deprecated function. Can remove include of CachedGeometry.h when this is removed.
+	UE_DEPRECATED(5.1, "Use GetPositionBuffer() or similar instead.")
+	FCachedGeometry::Section GetCachedGeometry(FGPUSkinCacheEntry* InOutEntry, uint32 SectionId);
 
 #if RHI_RAYTRACING
 	void ProcessRayTracingGeometryToUpdate(FRHICommandListImmediate& RHICmdList, FGPUSkinCacheEntry* SkinCacheEntry);
diff --git a/Engine/Source/Runtime/Engine/Public/GPUSkinVertexFactory.h b/Engine/Source/Runtime/Engine/Public/GPUSkinVertexFactory.h
--- a/Engine/Source/Runtime/Engine/Public/GPUSkinVertexFactory.h
+++ b/Engine/Source/Runtime/Engine/Public/GPUSkinVertexFactory.h
@@ -536,6 +536,11 @@ public:
 		}
 	}
 
+	inline FRHIShaderResourceView* GetPreviousPositionsSRV() const
+	{
+		return PrevPositionSRVAlias;
+	}
+
 protected:
 	friend class FGPUSkinVertexPassthroughFactoryShaderParameters;
 
@@ -549,13 +554,13 @@ protected:
 	FVertexBuffer TangentVBAlias;
 	FVertexBuffer ColorVBAlias;
 	// SRVs required for binding
-	FRHIShaderResourceView* PositionSRVAlias;
-	FRHIShaderResourceView* PrevPositionSRVAlias;
-	FRHIShaderResourceView* TangentSRVAlias;
-	FRHIShaderResourceView* ColorSRVAlias;
+	FRHIShaderResourceView* PositionSRVAlias = nullptr;
+	FRHIShaderResourceView* PrevPositionSRVAlias = nullptr;
+	FRHIShaderResourceView* TangentSRVAlias = nullptr;
+	FRHIShaderResourceView* ColorSRVAlias = nullptr;
 	// Cached stream indices
-	int32 PositionStreamIndex;
-	int32 TangentStreamIndex;
+	int32 PositionStreamIndex = -1;
+	int32 TangentStreamIndex = -1;
 	// Frame number of the bone data that is last updated
 	uint32 UpdatedFrameNumber = 0;
 
diff --git a/Engine/Source/Runtime/Engine/Public/SkeletalMeshTypes.h b/Engine/Source/Runtime/Engine/Public/SkeletalMeshTypes.h
--- a/Engine/Source/Runtime/Engine/Public/SkeletalMeshTypes.h
+++ b/Engine/Source/Runtime/Engine/Public/SkeletalMeshTypes.h
@@ -235,8 +235,7 @@ public:
 		return GetCurrentFirstLODIdx_Internal();
 	}
 
-	const TArray<FMatrix>& GetMeshObjectReferenceToLocalMatrices() const;
-	const TIndirectArray<FSkeletalMeshLODRenderData>& GetSkeletalMeshRenderDataLOD() const;
+	bool GetCachedGeometry(struct FCachedGeometry& OutCachedGeometry) const;
 
 protected:
 	AActor* Owner;
diff --git a/Engine/Source/Runtime/Engine/Public/SkeletalRenderPublic.h b/Engine/Source/Runtime/Engine/Public/SkeletalRenderPublic.h
--- a/Engine/Source/Runtime/Engine/Public/SkeletalRenderPublic.h
+++ b/Engine/Source/Runtime/Engine/Public/SkeletalRenderPublic.h
@@ -18,6 +18,7 @@ class FPrimitiveDrawInterface;
 class FVertexFactory;
 class UMorphTarget;
 struct FSkelMeshRenderSection;
+struct FCachedGeometry;
 class FGPUSkinCacheEntry;
 
 /** data for a single skinned skeletal mesh vertex */
@@ -113,6 +114,11 @@ public:
 	 */
 	virtual const TArray<FMatrix44f>& GetReferenceToLocalMatrices() const = 0;
 
+	/**
+	 * If we are caching geometry deformation through skin-cache/mesh-deformers or other, then this returns the currently cached geoemtry.
+	 */
+	virtual bool GetCachedGeometry(FCachedGeometry& OutCachedGeometry) const { return false; }
+
 	/**
 	*	Will force re-evaluating which Skin Weight buffer should be used for skinning, determined by checking for any override weights or a skin weight profile being set.
 	*	This prevents re-creating the vertex factories, but rather updates the bindings in place. 
diff --git a/Engine/Source/Runtime/Renderer/Private/HairStrands/HairStrandsInterface.cpp b/Engine/Source/Runtime/Renderer/Private/HairStrands/HairStrandsInterface.cpp
--- a/Engine/Source/Runtime/Renderer/Private/HairStrands/HairStrandsInterface.cpp
+++ b/Engine/Source/Runtime/Renderer/Private/HairStrands/HairStrandsInterface.cpp
@@ -687,7 +687,6 @@ FHairStrandsBookmarkParameters CreateHairStrandsBookmarkParameters(FScene* Scene
 	Out.InstanceCountPerType[HairInstanceCount_CardsOrMeshes] = Out.VisibleInstances.Num() - Out.InstanceCountPerType[HairInstanceCount_StrandsShadowView];
 
 	Out.ShaderPrintData			= ShaderPrint::IsEnabled(View.ShaderPrintData) ? &View.ShaderPrintData : nullptr;
-	Out.SkinCache				= View.Family->Scene->GetGPUSkinCache();
 	Out.ShaderMap				= View.ShaderMap;
 	Out.Instances				= &Scene->HairStrandsSceneData.RegisteredProxies;
 	Out.View					= &View;
diff --git a/Engine/Source/Runtime/Renderer/Public/HairStrandsInterface.h b/Engine/Source/Runtime/Renderer/Public/HairStrandsInterface.h
--- a/Engine/Source/Runtime/Renderer/Public/HairStrandsInterface.h
+++ b/Engine/Source/Runtime/Renderer/Public/HairStrandsInterface.h
@@ -480,7 +480,6 @@ enum EHairInstanceCount : uint8
 
 struct FHairStrandsBookmarkParameters
 {
-	class FGPUSkinCache* SkinCache = nullptr;
 	FShaderPrintData* ShaderPrintData = nullptr;
 	class FGlobalShaderMap* ShaderMap = nullptr;
 
