Fix manual skin cache to use render thread data only.

On the the skel mesh component is used to retrieve the mesh object. This needs to be fixed in a subsequent CL.
Skel. mesh triangles are processed section by section. This is inefficient when section are small. Ideall all section should be process at once. Also, on the the used section (i.e., having roots attached to) should be process. This should/could be addressed in future CL.

diff --git a/Engine/Plugins/Runtime/HairStrands/Source/HairStrandsCore/Private/GroomDebug.cpp b/Engine/Plugins/Runtime/HairStrands/Source/HairStrandsCore/Private/GroomDebug.cpp
--- a/Engine/Plugins/Runtime/HairStrands/Source/HairStrandsCore/Private/GroomDebug.cpp
+++ b/Engine/Plugins/Runtime/HairStrands/Source/HairStrandsCore/Private/GroomDebug.cpp
@@ -68,8 +68,6 @@ static FAutoConsoleVariableRef CVarHairStrandsControlPointDebug(TEXT("r.HairStra
 
 ///////////////////////////////////////////////////////////////////////////////////////////////////
 
-bool IsHairStrandsSkinCacheEnable();
-
 FCachedGeometry GetCacheGeometryForHair(
 	FRDGBuilder& GraphBuilder,
 	FHairGroupInstance* Instance,
diff --git a/Engine/Plugins/Runtime/HairStrands/Source/HairStrandsCore/Private/GroomGeometryCache.cpp b/Engine/Plugins/Runtime/HairStrands/Source/HairStrandsCore/Private/GroomGeometryCache.cpp
--- a/Engine/Plugins/Runtime/HairStrands/Source/HairStrandsCore/Private/GroomGeometryCache.cpp
+++ b/Engine/Plugins/Runtime/HairStrands/Source/HairStrandsCore/Private/GroomGeometryCache.cpp
@@ -14,6 +14,7 @@
 #include "SkeletalRenderPublic.h"
 #include "Rendering/SkeletalMeshLODRenderData.h"
 #include "RenderGraphUtils.h"
+#include "SkeletalMeshDeformerHelpers.h"
 
 ///////////////////////////////////////////////////////////////////////////////////////////////////
 
@@ -46,117 +47,75 @@ FHairStrandsProjectionMeshData::Section ConvertMeshSection(FCachedGeometry const
 ENGINE_API void UpdateRefToLocalMatrices(TArray<FMatrix44f>& ReferenceToLocal, const USkinnedMeshComponent* InMeshComponent, const FSkeletalMeshRenderData* InSkeletalMeshRenderData, int32 LODIndex, const TArray<FBoneIndexType>* ExtraRequiredBoneIndices = NULL);
 ENGINE_API void UpdatePreviousRefToLocalMatrices(TArray<FMatrix44f>& ReferenceToLocal, const USkinnedMeshComponent* InMeshComponent, const FSkeletalMeshRenderData* InSkeletalMeshRenderData, int32 LODIndex, const TArray<FBoneIndexType>* ExtraRequiredBoneIndices = NULL);
 
-static void BuildBoneMatrices(
+ void BuildCacheGeometry(
+	FRDGBuilder& GraphBuilder,
+	FGlobalShaderMap* ShaderMap, 
 	const USkeletalMeshComponent* SkeletalMeshComponent, 
-	const FSkeletalMeshRenderData* RenderData,
-	const FSkeletalMeshLODRenderData& LODData,
-	const uint32 LODIndex, 
-	const bool bNeedPreviousPosition,
-	TArray<uint32>& OutMatrixOffsets, 
-	TArray<FVector4f>& OutBoneMatrices,
-	TArray<FVector4f>& OutPrevBoneMatrices)
+	const bool bOutputTriangleData,
+	FCachedGeometry& Out)
 {
-	//SkeletalMeshComponent->GetCurrentRefToLocalMatrices(BoneTransforms, LODIndex);
-	check(RenderData->LODRenderData.IsValidIndex(LODIndex));
+	if (SkeletalMeshComponent == nullptr)
+	{
+		return;
+	}
 
-	TArray<FMatrix44f> BoneTransforms;
-	UpdateRefToLocalMatrices(BoneTransforms, SkeletalMeshComponent, RenderData, LODIndex, nullptr);
+	const uint32 LODIndex = SkeletalMeshComponent->GetPredictedLODLevel();// RenderData->PendingFirstLODIdx;
+	Out.LocalToWorld = SkeletalMeshComponent->SceneProxy ? FTransform(SkeletalMeshComponent->SceneProxy->GetLocalToWorld()) : FTransform();
+	Out.LODIndex = LODIndex;
 
-	TArray<FMatrix44f> PrevBoneTransforms;
-	if (bNeedPreviousPosition)
+	FSkeletalMeshObject* SkeletalMeshObject = SkeletalMeshComponent->MeshObject;
+	if (!bOutputTriangleData || SkeletalMeshObject == nullptr)
 	{
-		UpdatePreviousRefToLocalMatrices(PrevBoneTransforms, SkeletalMeshComponent, RenderData, LODIndex, nullptr);
+		return;
 	}
-	
-	OutMatrixOffsets.SetNum(LODData.GetNumVertices());
+
+	FSkeletalMeshRenderData* RenderData = SkeletalMeshComponent->GetSkeletalMesh()->GetResourceForRendering();
+	FSkeletalMeshLODRenderData& LODData = RenderData->LODRenderData[LODIndex];
+
+	const bool bNeedPreviousPosition = IsHairStrandContinuousDecimationReorderingEnabled();
+
+	// Create deformed position buffer (output)
+	FRDGBufferRef DeformedPositionsBuffer			= GraphBuilder.CreateBuffer(FRDGBufferDesc::CreateBufferDesc(sizeof(float), LODData.StaticVertexBuffers.PositionVertexBuffer.GetNumVertices() * 3), TEXT("Hair.SkinnedDeformedPositions"));
+	FRDGBufferRef DeformedPreviousPositionsBuffer	= bNeedPreviousPosition ? GraphBuilder.CreateBuffer(FRDGBufferDesc::CreateBufferDesc(sizeof(float), LODData.StaticVertexBuffers.PositionVertexBuffer.GetNumVertices() * 3), TEXT("Hair.SkinnedDeformedPreviousPositions")) : nullptr;
+
+	const int32 LodIndex = SkeletalMeshObject->GetLOD();
+	FSkeletalMeshLODRenderData const* LodRenderData = &SkeletalMeshObject->GetSkeletalMeshRenderData().LODRenderData[LodIndex];
+						
+	Out.DeformedPositionBuffer = DeformedPositionsBuffer;
+	Out.DeformedPreviousPositionBuffer = DeformedPreviousPositionsBuffer;
+	FRDGBufferSRVRef DeformedPositionSRV = GraphBuilder.CreateSRV(DeformedPositionsBuffer, PF_R32_FLOAT);
+	FRDGBufferSRVRef DeformedPreviousPositionSRV = bNeedPreviousPosition ? GraphBuilder.CreateSRV(DeformedPreviousPositionsBuffer, PF_R32_FLOAT) : nullptr;
+
 	uint32 BonesOffset = 0;
 	for (int32 SectionIdx = 0; SectionIdx < LODData.RenderSections.Num(); ++SectionIdx)
 	{
 		const FSkelMeshRenderSection& Section = LODData.RenderSections[SectionIdx];
-		for (uint32 SectionVertex = 0; SectionVertex < Section.NumVertices; ++SectionVertex)
+		if (SkeletalMeshObject->HaveValidDynamicData())
 		{
-			OutMatrixOffsets[Section.BaseVertexIndex + SectionVertex] = BonesOffset;
+			FRHIShaderResourceView * BoneBuffer		= FSkeletalMeshDeformerHelpers::GetBoneBufferForReading(SkeletalMeshObject, LodIndex, SectionIdx, false);
+			FRHIShaderResourceView * BonePrevBuffer	= FSkeletalMeshDeformerHelpers::GetBoneBufferForReading(SkeletalMeshObject, LodIndex, SectionIdx, true);
+			AddSkinUpdatePass(GraphBuilder, ShaderMap, SectionIdx, BonesOffset, SkeletalMeshComponent->GetSkinWeightBuffer(LODIndex), LODData, BoneBuffer, BonePrevBuffer, DeformedPositionsBuffer, DeformedPreviousPositionsBuffer);
+
+			FCachedGeometry::Section& OutSection = Out.Sections.AddDefaulted_GetRef();
+			OutSection.RDGPositionBuffer = DeformedPositionSRV;
+			OutSection.RDGPreviousPositionBuffer = DeformedPreviousPositionSRV;
+			OutSection.PositionBuffer = nullptr; // Do not use the SRV slot, but instead use the RDG buffer created above (DeformedPositionSRV)
+			OutSection.PreviousPositionBuffer = nullptr; // Do not use the SRV slot, but instead use the RDG buffer created above (DeformedPositionSRV)
+			OutSection.UVsBuffer = LODData.StaticVertexBuffers.StaticMeshVertexBuffer.GetTexCoordsSRV();
+			OutSection.TotalVertexCount = LODData.StaticVertexBuffers.PositionVertexBuffer.GetNumVertices();
+			OutSection.IndexBuffer = LODData.MultiSizeIndexContainer.GetIndexBuffer()->GetSRV();
+			OutSection.TotalIndexCount = LODData.MultiSizeIndexContainer.GetIndexBuffer()->Num();
+			OutSection.UVsChannelCount = LODData.StaticVertexBuffers.StaticMeshVertexBuffer.GetNumTexCoords();
+			OutSection.NumPrimitives = Section.NumTriangles;
+			OutSection.NumVertices = Section.NumVertices;
+			OutSection.IndexBaseIndex = Section.BaseIndex;
+			OutSection.VertexBaseIndex = Section.BaseVertexIndex;
+			OutSection.SectionIndex = SectionIdx;
+			OutSection.LODIndex = LODIndex;
+			OutSection.UVsChannelOffset = 0; // Assume that we needs to pair meshes based on UVs 0
 		}
 		BonesOffset += Section.BoneMap.Num();
 	}
-	OutBoneMatrices.SetNum(BonesOffset * 3);
-	OutPrevBoneMatrices.SetNum(BonesOffset * 3);
-	BonesOffset = 0;
-	for (int32 SectionIdx = 0; SectionIdx < LODData.RenderSections.Num(); ++SectionIdx)
-	{
-		const FSkelMeshRenderSection& Section = LODData.RenderSections[SectionIdx];
-		for (int32 BoneIdx = 0; BoneIdx < Section.BoneMap.Num(); ++BoneIdx, ++BonesOffset)
-		{
-			BoneTransforms[Section.BoneMap[BoneIdx]].To3x4MatrixTranspose(&OutBoneMatrices[3 * BonesOffset].X);
-			if (bNeedPreviousPosition)
-			{
-				PrevBoneTransforms[Section.BoneMap[BoneIdx]].To3x4MatrixTranspose(&OutPrevBoneMatrices[3 * BonesOffset].X);
-			}
-		}
-	}
-}
-
- void BuildCacheGeometry(
-	FRDGBuilder& GraphBuilder,
-	FGlobalShaderMap* ShaderMap, 
-	const USkeletalMeshComponent* SkeletalMeshComponent, 
-	const bool bOutputTriangleData,
-	FCachedGeometry& Out)
-{
-	if (SkeletalMeshComponent)
-	{
-		const uint32 LODIndex = SkeletalMeshComponent->GetPredictedLODLevel();// RenderData->PendingFirstLODIdx;
-		Out.LocalToWorld = SkeletalMeshComponent->SceneProxy ? FTransform(SkeletalMeshComponent->SceneProxy->GetLocalToWorld()) : FTransform();
-		Out.LODIndex = LODIndex;
-
-		if (bOutputTriangleData)
-		{
-			FSkeletalMeshRenderData* RenderData = SkeletalMeshComponent->GetSkeletalMesh()->GetResourceForRendering();
-			FSkeletalMeshLODRenderData& LODData = RenderData->LODRenderData[LODIndex];
-
-			const bool bNeedPreviousPosition = IsHairStrandContinuousDecimationReorderingEnabled();
-
-			TArray<uint32> MatrixOffsets;
-			TArray<FVector4f> BoneMatrices;
-			TArray<FVector4f> PrevBoneMatrices;
-			BuildBoneMatrices(SkeletalMeshComponent, RenderData, LODData, LODIndex, bNeedPreviousPosition, MatrixOffsets, BoneMatrices, PrevBoneMatrices);
-
-			FRDGBufferRef DeformedPositionsBuffer			= GraphBuilder.CreateBuffer(FRDGBufferDesc::CreateBufferDesc(sizeof(float), LODData.StaticVertexBuffers.PositionVertexBuffer.GetNumVertices() * 3), TEXT("Hair.SkinnedDeformedPositions"));
-			FRDGBufferRef DeformedPreviousPositionsBuffer	= bNeedPreviousPosition ? GraphBuilder.CreateBuffer(FRDGBufferDesc::CreateBufferDesc(sizeof(float), LODData.StaticVertexBuffers.PositionVertexBuffer.GetNumVertices() * 3), TEXT("Hair.SkinnedDeformedPreviousPositions")) : nullptr;
-			FRDGBufferRef BoneMatricesBuffer				= CreateStructuredBuffer(GraphBuilder, TEXT("Hair.SkinnedBoneMatrices"), sizeof(float) * 4, BoneMatrices.Num(), BoneMatrices.GetData(), sizeof(float) * 4 * BoneMatrices.Num());
-			FRDGBufferRef PreviousBoneMatricesBuffer		= bNeedPreviousPosition ? CreateStructuredBuffer(GraphBuilder, TEXT("Hair.SkinnedPreviousBoneMatrices"), sizeof(float) * 4, BoneMatrices.Num(), BoneMatrices.GetData(), sizeof(float) * 4 * BoneMatrices.Num()) : nullptr;
-			FRDGBufferRef MatrixOffsetsBuffer				= CreateStructuredBuffer(GraphBuilder, TEXT("Hair.SkinnedMatrixOffsets"), sizeof(uint32), MatrixOffsets.Num(), MatrixOffsets.GetData(), sizeof(uint32) * MatrixOffsets.Num());
-
-			AddSkinUpdatePass(GraphBuilder, ShaderMap, SkeletalMeshComponent->GetSkinWeightBuffer(LODIndex), LODData, BoneMatricesBuffer, PreviousBoneMatricesBuffer, MatrixOffsetsBuffer, DeformedPositionsBuffer, DeformedPreviousPositionsBuffer);
-
-			Out.DeformedPositionBuffer = DeformedPositionsBuffer;
-			Out.DeformedPreviousPositionBuffer = DeformedPreviousPositionsBuffer;
-			FRDGBufferSRVRef DeformedPositionSRV = GraphBuilder.CreateSRV(DeformedPositionsBuffer, PF_R32_FLOAT);
-			FRDGBufferSRVRef DeformedPreviousPositionSRV = bNeedPreviousPosition ? GraphBuilder.CreateSRV(DeformedPreviousPositionsBuffer, PF_R32_FLOAT) : nullptr;
-			for (int32 SectionIdx = 0; SectionIdx < LODData.RenderSections.Num(); ++SectionIdx)
-			{
-				const FSkelMeshRenderSection& Section = LODData.RenderSections[SectionIdx];
-
-				FCachedGeometry::Section& OutSection = Out.Sections.AddDefaulted_GetRef();
-				OutSection.RDGPositionBuffer = DeformedPositionSRV;
-				OutSection.RDGPreviousPositionBuffer = DeformedPreviousPositionSRV;
-				OutSection.PositionBuffer = nullptr; // Do not use the SRV slot, but instead use the RDG buffer created above (DeformedPositionSRV)
-				OutSection.PreviousPositionBuffer = nullptr; // Do not use the SRV slot, but instead use the RDG buffer created above (DeformedPositionSRV)
-				OutSection.UVsBuffer = LODData.StaticVertexBuffers.StaticMeshVertexBuffer.GetTexCoordsSRV();
-				OutSection.TotalVertexCount = LODData.StaticVertexBuffers.PositionVertexBuffer.GetNumVertices();
-				OutSection.IndexBuffer = LODData.MultiSizeIndexContainer.GetIndexBuffer()->GetSRV();
-				OutSection.TotalIndexCount = LODData.MultiSizeIndexContainer.GetIndexBuffer()->Num();
-				OutSection.UVsChannelCount = LODData.StaticVertexBuffers.StaticMeshVertexBuffer.GetNumTexCoords();
-				OutSection.NumPrimitives = Section.NumTriangles;
-				OutSection.NumVertices = Section.NumVertices;
-				OutSection.IndexBaseIndex = Section.BaseIndex;
-				OutSection.VertexBaseIndex = Section.BaseVertexIndex;
-				OutSection.SectionIndex = SectionIdx;
-				OutSection.LODIndex = LODIndex;
-				OutSection.UVsChannelOffset = 0; // Assume that we needs to pair meshes based on UVs 0
-			}
-		}
-	}
 }
 
  void BuildCacheGeometry(
diff --git a/Engine/Plugins/Runtime/HairStrands/Source/HairStrandsCore/Private/GroomManager.cpp b/Engine/Plugins/Runtime/HairStrands/Source/HairStrandsCore/Private/GroomManager.cpp
--- a/Engine/Plugins/Runtime/HairStrands/Source/HairStrandsCore/Private/GroomManager.cpp
+++ b/Engine/Plugins/Runtime/HairStrands/Source/HairStrandsCore/Private/GroomManager.cpp
@@ -30,7 +30,7 @@ static FAutoConsoleVariableRef CVarHairStrands_UseCards(TEXT("r.HairStrands.UseC
 static int32 GHairStrands_SwapBufferType = 3;
 static FAutoConsoleVariableRef CVarGHairStrands_SwapBufferType(TEXT("r.HairStrands.SwapType"), GHairStrands_SwapBufferType, TEXT("Swap rendering buffer at the end of frame. This is an experimental toggle. Default:1"));
 
-static int32 GHairStrands_ManualSkinCache = 0;
+static int32 GHairStrands_ManualSkinCache = 1;
 static FAutoConsoleVariableRef CVarGHairStrands_ManualSkinCache(TEXT("r.HairStrands.ManualSkinCache"), GHairStrands_ManualSkinCache, TEXT("If skin cache is not enabled, and grooms use skinning method, this enable a simple skin cache mechanisme for groom. Default:disable"));
 
 static int32 GHairStrands_InterpolationFrustumCullingEnable = 1;
@@ -48,11 +48,6 @@ EHairBufferSwapType GetHairSwapBufferType()
 	return EHairBufferSwapType::EndOfFrame;
 }
 
-bool IsHairStrandsSkinCacheEnable()
-{
-	return GHairStrands_ManualSkinCache > 0;
-}
-
 DEFINE_LOG_CATEGORY_STATIC(LogGroomManager, Log, All);
 
 ///////////////////////////////////////////////////////////////////////////////////////////////////
@@ -102,7 +97,7 @@ FCachedGeometry GetCacheGeometryForHair(
 					SceneProxy->GetCachedGeometry(Out);
 				}
 
-				if (IsHairStrandsSkinCacheEnable() && Out.Sections.Num() == 0)
+				if (GHairStrands_ManualSkinCache > 0 && Out.Sections.Num() == 0)
 				{
 					//#hair_todo: Need to have a (frame) cache to insure that we don't recompute the same projection several time
 					// Actual populate the cache with only the needed part based on the groom projection data. At the moment it recompute everything ...
diff --git a/Engine/Plugins/Runtime/HairStrands/Source/HairStrandsCore/Private/HairStrandsMeshProjection.cpp b/Engine/Plugins/Runtime/HairStrands/Source/HairStrandsCore/Private/HairStrandsMeshProjection.cpp
--- a/Engine/Plugins/Runtime/HairStrands/Source/HairStrandsCore/Private/HairStrandsMeshProjection.cpp
+++ b/Engine/Plugins/Runtime/HairStrands/Source/HairStrandsCore/Private/HairStrandsMeshProjection.cpp
@@ -83,13 +83,14 @@ class FSkinUpdateCS : public FGlobalShader
 	using FPermutationDomain = TShaderPermutationDomain<FUnlimitedBoneInfluence, FUseExtraInfluence, FIndexUint16, FPrevious>;
 
 	BEGIN_SHADER_PARAMETER_STRUCT(FParameters, )
+		SHADER_PARAMETER(uint32, NumTotalVertices)
+		SHADER_PARAMETER(uint32, SectionVertexBaseIndex)
 		SHADER_PARAMETER(uint32, IndexSize)
-		SHADER_PARAMETER(uint32, NumVertices)
 		SHADER_PARAMETER(uint32, WeightStride)
+		SHADER_PARAMETER(uint32, BonesOffset)
 		SHADER_PARAMETER_SRV(Buffer<uint>, WeightLookup)
-		SHADER_PARAMETER_RDG_BUFFER_SRV(StructuredBuffer<float4>, BoneMatrices)
-		SHADER_PARAMETER_RDG_BUFFER_SRV(StructuredBuffer<float4>, PrevBoneMatrices)
-		SHADER_PARAMETER_RDG_BUFFER_SRV(StructuredBuffer<uint>, MatrixOffsets)
+		SHADER_PARAMETER_SRV(Buffer<float4>, BoneMatrices)
+		SHADER_PARAMETER_SRV(Buffer<float4>, PrevBoneMatrices)
 		SHADER_PARAMETER_SRV(Buffer<uint>, VertexWeights)
 		SHADER_PARAMETER_SRV(Buffer<float>, RestPositions)
 		SHADER_PARAMETER_RDG_BUFFER_UAV(RWBuffer<float>, DeformedPositions)
@@ -105,29 +106,39 @@ IMPLEMENT_GLOBAL_SHADER(FSkinUpdateCS, "/Engine/Private/HairStrands/HairStrandsS
 void AddSkinUpdatePass(
 	FRDGBuilder& GraphBuilder,
 	FGlobalShaderMap* ShaderMap,
+	uint32 SectionIndex,
+	uint32 BonesOffset,
 	FSkinWeightVertexBuffer* SkinWeight,
 	FSkeletalMeshLODRenderData& RenderData,
-	FRDGBufferRef BoneMatrices,
-	FRDGBufferRef PrevBoneMatrices,
-	FRDGBufferRef MatrixOffsets,
+	FRHIShaderResourceView* BoneMatrices,
+	FRHIShaderResourceView* PrevBoneMatrices,
 	FRDGBufferRef OutDeformedPosition,
 	FRDGBufferRef OutPrevDeformedPosition)
 {
-	const bool bPrevPosition = OutPrevDeformedPosition != nullptr && PrevBoneMatrices != nullptr;
+	check(BoneMatrices);
+
+	const FSkelMeshRenderSection& Section = RenderData.RenderSections[SectionIndex];
+	const uint32 NumVertexToProcess = Section.NumVertices;
+	const uint32 SectionVertexBaseIndex = Section.BaseVertexIndex;
+	const uint32 NumTotalVertices = RenderData.StaticVertexBuffers.PositionVertexBuffer.GetNumVertices();
 
+	const bool bPrevPosition = OutPrevDeformedPosition != nullptr && PrevBoneMatrices != nullptr;
+	
 	FSkinUpdateCS::FParameters* Parameters = GraphBuilder.AllocParameters<FSkinUpdateCS::FParameters>();
 	Parameters->IndexSize = SkinWeight->GetBoneIndexByteSize();
-	Parameters->NumVertices = RenderData.StaticVertexBuffers.PositionVertexBuffer.GetNumVertices();
+	Parameters->NumTotalVertices = NumTotalVertices;
+	Parameters->SectionVertexBaseIndex = SectionVertexBaseIndex;
 	Parameters->WeightStride = SkinWeight->GetConstantInfluencesVertexStride();
 	Parameters->WeightLookup = SkinWeight->GetLookupVertexBuffer()->GetSRV();
-	Parameters->BoneMatrices = GraphBuilder.CreateSRV(BoneMatrices);
-	Parameters->MatrixOffsets = GraphBuilder.CreateSRV(MatrixOffsets);
+	Parameters->BonesOffset = BonesOffset;
+	Parameters->BoneMatrices = BoneMatrices;
 	Parameters->VertexWeights = SkinWeight->GetDataVertexBuffer()->GetSRV();
 	Parameters->RestPositions = RenderData.StaticVertexBuffers.PositionVertexBuffer.GetSRV();
 	Parameters->DeformedPositions = GraphBuilder.CreateUAV(OutDeformedPosition, PF_R32_FLOAT);
 	if (bPrevPosition)
 	{
-		Parameters->PrevBoneMatrices = GraphBuilder.CreateSRV(BoneMatrices);
+		check(PrevBoneMatrices);
+		Parameters->PrevBoneMatrices = BoneMatrices;
 		Parameters->PrevDeformedPositions = GraphBuilder.CreateUAV(OutPrevDeformedPosition, PF_R32_FLOAT);
 	}
 
@@ -137,12 +148,12 @@ void AddSkinUpdatePass(
 	PermutationVector.Set<FSkinUpdateCS::FIndexUint16>(SkinWeight->Use16BitBoneIndex());
 	PermutationVector.Set<FSkinUpdateCS::FPrevious>(bPrevPosition);
 
-	const FIntVector DispatchGroupCount = FComputeShaderUtils::GetGroupCount(RenderData.StaticVertexBuffers.PositionVertexBuffer.GetNumVertices(), 64);
+	const FIntVector DispatchGroupCount = FComputeShaderUtils::GetGroupCount(NumVertexToProcess, 64);
 	check(DispatchGroupCount.X < 65536);
 	TShaderMapRef<FSkinUpdateCS> ComputeShader(ShaderMap, PermutationVector);
 	FComputeShaderUtils::AddPass(
 		GraphBuilder,
-		RDG_EVENT_NAME("HairStrands::UpdateSkinPosition(%s)", bPrevPosition ? TEXT("Curr,Prev") : TEXT("Curr")),
+		RDG_EVENT_NAME("HairStrands::UpdateSkinPosition(%s,Section:%d)", bPrevPosition ? TEXT("Curr,Prev") : TEXT("Curr"), SectionIndex),
 		ComputeShader,
 		Parameters,
 		DispatchGroupCount);
diff --git a/Engine/Plugins/Runtime/HairStrands/Source/HairStrandsCore/Public/HairStrandsMeshProjection.h b/Engine/Plugins/Runtime/HairStrands/Source/HairStrandsCore/Public/HairStrandsMeshProjection.h
--- a/Engine/Plugins/Runtime/HairStrands/Source/HairStrandsCore/Public/HairStrandsMeshProjection.h
+++ b/Engine/Plugins/Runtime/HairStrands/Source/HairStrandsCore/Public/HairStrandsMeshProjection.h
@@ -118,11 +118,12 @@ void AddComputeMipsPass(
 void AddSkinUpdatePass(
 	FRDGBuilder& GraphBuilder,
 	FGlobalShaderMap* ShaderMap,
+	uint32 SectionIndex,
+	uint32 BonesOffset, 
 	class FSkinWeightVertexBuffer* SkinWeight,
 	class FSkeletalMeshLODRenderData& RenderData,
-	FRDGBufferRef BoneMatrices,
-	FRDGBufferRef PreviousBoneMatrices,
-	FRDGBufferRef MatrixOffsets,
+	FRHIShaderResourceView* BoneMatrices,
+	FRHIShaderResourceView* PrevBoneMatrices,
 	FRDGBufferRef OutDeformedosition,
 	FRDGBufferRef OutPreviousDeformedosition);
 
diff --git a/Engine/Shaders/Private/HairStrands/HairStrandsSkinUpdate.usf b/Engine/Shaders/Private/HairStrands/HairStrandsSkinUpdate.usf
--- a/Engine/Shaders/Private/HairStrands/HairStrandsSkinUpdate.usf
+++ b/Engine/Shaders/Private/HairStrands/HairStrandsSkinUpdate.usf
@@ -7,21 +7,22 @@
 //----------------------------------------------------------------------------------
 
 uint IndexSize; // bone index size: 1 - 8bit or 2 - 16bit
-uint NumVertices;
+uint SectionVertexBaseIndex;
+uint NumTotalVertices;
 uint WeightStride;
+uint BonesOffset;
 
 Buffer<uint> WeightLookup;
-StructuredBuffer<uint> MatrixOffsets;
 Buffer<uint> VertexWeights;
 Buffer<float> RestPositions;
 
 // Current Position
-StructuredBuffer<float4> BoneMatrices;
+Buffer<float4> BoneMatrices;
 RWBuffer<float> DeformedPositions;
 
 // Previous Position
 #if PERMUTATION_PREV
-StructuredBuffer<float4> PrevBoneMatrices;
+Buffer<float4> PrevBoneMatrices;
 RWBuffer<float> PrevDeformedPositions;
 #endif
 
@@ -50,9 +51,9 @@ struct FUnpackedWeight
 
 //----------------------------------------------------------------------------------
 
-FBlendedMatrix BlendMatrix(FUnpackedWeight UnpackedWeight, uint VertexIndex, StructuredBuffer<float4> InBoneMatrices)
+FBlendedMatrix BlendMatrix(FUnpackedWeight UnpackedWeight, uint VertexIndex, Buffer<float4> InBoneMatrices)
 {
-	const uint MatrixOffset = MatrixOffsets[VertexIndex];
+	const uint MatrixOffset = BonesOffset;
 #if GPUSKIN_UNLIMITED_BONE_INFLUENCE
 	uint BlendOffsetCount = WeightLookup[VertexIndex];
 	int NumBoneInfluences = BlendOffsetCount & 0xff;
@@ -156,8 +157,8 @@ void UpdateSkinPositionCS(uint3 GroupID : SV_GroupID,
  uint3 DispatchThreadID : SV_DispatchThreadID,
  uint3 GroupThreadID : SV_GroupThreadID)
 {
-	const uint VertexIndex = DispatchThreadID.x;
-	if (VertexIndex < NumVertices)
+	const uint VertexIndex = DispatchThreadID.x + SectionVertexBaseIndex;
+	if (VertexIndex < NumTotalVertices)
 	{
 		const uint PositionOffset = VertexIndex * 3;
 		const uint WeightOffset = VertexIndex * WeightStride / 4;
@@ -174,7 +175,7 @@ void UpdateSkinPositionCS(uint3 GroupID : SV_GroupID,
 		{
 			const FBlendedMatrix BlendedMatrix = BlendMatrix(UnpackedWeight, VertexIndex, BoneMatrices);
 			const float3 OutputPosition =  mul(BlendedMatrix, InputPosition);
-			DeformedPositions[PositionOffset] = OutputPosition.x;
+			DeformedPositions[PositionOffset    ] = OutputPosition.x;
 			DeformedPositions[PositionOffset + 1] = OutputPosition.y;
 			DeformedPositions[PositionOffset + 2] = OutputPosition.z;
 		}
@@ -184,7 +185,7 @@ void UpdateSkinPositionCS(uint3 GroupID : SV_GroupID,
 		{
 			const FBlendedMatrix PrevBlendedMatrix = BlendMatrix(UnpackedWeight, VertexIndex, PrevBoneMatrices);
 			const float3 OutputPosition = mul(PrevBlendedMatrix, InputPosition);
-			PrevDeformedPositions[PositionOffset] = OutputPosition.x;
+			PrevDeformedPositions[PositionOffset    ] = OutputPosition.x;
 			PrevDeformedPositions[PositionOffset + 1] = OutputPosition.y;
 			PrevDeformedPositions[PositionOffset + 2] = OutputPosition.z;
 		}
