Use render thread structures when building groom geometry cache.
Fixes crashes caused by mixing some Component owned LOD render data with render thread copies.

diff --git a/Engine/Plugins/Runtime/HairStrands/Source/HairStrandsCore/Private/GroomGeometryCache.cpp b/Engine/Plugins/Runtime/HairStrands/Source/HairStrandsCore/Private/GroomGeometryCache.cpp
--- a/Engine/Plugins/Runtime/HairStrands/Source/HairStrandsCore/Private/GroomGeometryCache.cpp
+++ b/Engine/Plugins/Runtime/HairStrands/Source/HairStrandsCore/Private/GroomGeometryCache.cpp
@@ -59,9 +59,13 @@ ENGINE_API void UpdatePreviousRefToLocalMatrices(TArray<FMatrix44f>& ReferenceTo
 		return;
 	}
 
-	const uint32 LODIndex = SkeletalMeshComponent->GetPredictedLODLevel();// RenderData->PendingFirstLODIdx;
-	Out.LocalToWorld = SkeletalMeshComponent->SceneProxy ? FTransform(SkeletalMeshComponent->SceneProxy->GetLocalToWorld()) : FTransform();
-	Out.LODIndex = LODIndex;
+	FSkeletalMeshSceneProxy* SceneProxy = static_cast<FSkeletalMeshSceneProxy*>(SkeletalMeshComponent->SceneProxy);
+	if (SceneProxy == nullptr)
+	{
+		return;
+	}
+
+	Out.LocalToWorld = FTransform(SceneProxy->GetLocalToWorld());
 
 	FSkeletalMeshObject* SkeletalMeshObject = SkeletalMeshComponent->MeshObject;
 	if (!bOutputTriangleData || SkeletalMeshObject == nullptr)
@@ -69,8 +73,10 @@ ENGINE_API void UpdatePreviousRefToLocalMatrices(TArray<FMatrix44f>& ReferenceTo
 		return;
 	}
 
-	FSkeletalMeshRenderData* RenderData = SkeletalMeshComponent->GetSkeletalMeshAsset()->GetResourceForRendering();
-	FSkeletalMeshLODRenderData& LODData = RenderData->LODRenderData[LODIndex];
+	const int32 LODIndex = SkeletalMeshObject->GetLOD();
+	Out.LODIndex = LODIndex;
+
+	FSkeletalMeshLODRenderData& LODData = SkeletalMeshObject->GetSkeletalMeshRenderData().LODRenderData[LODIndex];
 
 	const bool bNeedPreviousPosition = IsHairStrandContinuousDecimationReorderingEnabled();
 
@@ -78,9 +84,6 @@ ENGINE_API void UpdatePreviousRefToLocalMatrices(TArray<FMatrix44f>& ReferenceTo
 	FRDGBufferRef DeformedPositionsBuffer			= GraphBuilder.CreateBuffer(FRDGBufferDesc::CreateBufferDesc(sizeof(float), LODData.StaticVertexBuffers.PositionVertexBuffer.GetNumVertices() * 3), TEXT("Hair.SkinnedDeformedPositions"));
 	FRDGBufferRef DeformedPreviousPositionsBuffer	= bNeedPreviousPosition ? GraphBuilder.CreateBuffer(FRDGBufferDesc::CreateBufferDesc(sizeof(float), LODData.StaticVertexBuffers.PositionVertexBuffer.GetNumVertices() * 3), TEXT("Hair.SkinnedDeformedPreviousPositions")) : nullptr;
 
-	const int32 LodIndex = SkeletalMeshObject->GetLOD();
-	FSkeletalMeshLODRenderData const* LodRenderData = &SkeletalMeshObject->GetSkeletalMeshRenderData().LODRenderData[LodIndex];
-						
 	Out.DeformedPositionBuffer = DeformedPositionsBuffer;
 	Out.DeformedPreviousPositionBuffer = DeformedPreviousPositionsBuffer;
 	FRDGBufferSRVRef DeformedPositionSRV = GraphBuilder.CreateSRV(DeformedPositionsBuffer, PF_R32_FLOAT);
@@ -92,8 +95,8 @@ ENGINE_API void UpdatePreviousRefToLocalMatrices(TArray<FMatrix44f>& ReferenceTo
 		const FSkelMeshRenderSection& Section = LODData.RenderSections[SectionIdx];
 		if (SkeletalMeshObject->HaveValidDynamicData())
 		{
-			FRHIShaderResourceView * BoneBuffer		= FSkeletalMeshDeformerHelpers::GetBoneBufferForReading(SkeletalMeshObject, LodIndex, SectionIdx, false);
-			FRHIShaderResourceView * BonePrevBuffer	= FSkeletalMeshDeformerHelpers::GetBoneBufferForReading(SkeletalMeshObject, LodIndex, SectionIdx, true);
+			FRHIShaderResourceView * BoneBuffer		= FSkeletalMeshDeformerHelpers::GetBoneBufferForReading(SkeletalMeshObject, LODIndex, SectionIdx, false);
+			FRHIShaderResourceView * BonePrevBuffer	= FSkeletalMeshDeformerHelpers::GetBoneBufferForReading(SkeletalMeshObject, LODIndex, SectionIdx, true);
 			AddSkinUpdatePass(GraphBuilder, ShaderMap, SectionIdx, BonesOffset, SkeletalMeshComponent->GetSkinWeightBuffer(LODIndex), LODData, BoneBuffer, BonePrevBuffer, DeformedPositionsBuffer, DeformedPreviousPositionsBuffer);
 
 			FCachedGeometry::Section& OutSection = Out.Sections.AddDefaulted_GetRef();
