Fix groom regression from LOD not being filled out by geometry cache building during RunHairLODSelection pass.

diff --git a/Engine/Plugins/Runtime/HairStrands/Source/HairStrandsCore/Private/GroomGeometryCache.cpp b/Engine/Plugins/Runtime/HairStrands/Source/HairStrandsCore/Private/GroomGeometryCache.cpp
--- a/Engine/Plugins/Runtime/HairStrands/Source/HairStrandsCore/Private/GroomGeometryCache.cpp
+++ b/Engine/Plugins/Runtime/HairStrands/Source/HairStrandsCore/Private/GroomGeometryCache.cpp
@@ -70,6 +70,8 @@ ENGINE_API void UpdatePreviousRefToLocalMatrices(TArray<FMatrix44f>& ReferenceTo
 	FSkeletalMeshObject* SkeletalMeshObject = SkeletalMeshComponent->MeshObject;
 	if (!bOutputTriangleData || SkeletalMeshObject == nullptr)
 	{
+		// Fill out LODIndex for RunHairLODSelection() pass.
+		Out.LODIndex = SkeletalMeshObject->GetLOD();
 		return;
 	}
 
