Fix static analysis warning. Reference of potentially null SkeletalMeshObject.

diff --git a/Engine/Plugins/Runtime/HairStrands/Source/HairStrandsCore/Private/GroomGeometryCache.cpp b/Engine/Plugins/Runtime/HairStrands/Source/HairStrandsCore/Private/GroomGeometryCache.cpp
--- a/Engine/Plugins/Runtime/HairStrands/Source/HairStrandsCore/Private/GroomGeometryCache.cpp
+++ b/Engine/Plugins/Runtime/HairStrands/Source/HairStrandsCore/Private/GroomGeometryCache.cpp
@@ -68,16 +68,19 @@ ENGINE_API void UpdatePreviousRefToLocalMatrices(TArray<FMatrix44f>& ReferenceTo
 	Out.LocalToWorld = FTransform(SceneProxy->GetLocalToWorld());
 
 	FSkeletalMeshObject* SkeletalMeshObject = SkeletalMeshComponent->MeshObject;
-	if (!bOutputTriangleData || SkeletalMeshObject == nullptr)
+	if (SkeletalMeshObject == nullptr)
 	{
-		// Fill out LODIndex for RunHairLODSelection() pass.
-		Out.LODIndex = SkeletalMeshObject->GetLOD();
 		return;
 	}
 
 	const int32 LODIndex = SkeletalMeshObject->GetLOD();
 	Out.LODIndex = LODIndex;
 
+	if (!bOutputTriangleData)
+	{
+		return;
+	}
+
 	FSkeletalMeshLODRenderData& LODData = SkeletalMeshObject->GetSkeletalMeshRenderData().LODRenderData[LODIndex];
 
 	const bool bNeedPreviousPosition = IsHairStrandContinuousDecimationReorderingEnabled();
