Fix for crash when opening empty groom asset after restarting editor.

#rnx

diff --git a/Engine/Plugins/Runtime/HairStrands/Source/HairStrandsCore/Private/GroomAsset.cpp b/Engine/Plugins/Runtime/HairStrands/Source/HairStrandsCore/Private/GroomAsset.cpp
--- a/Engine/Plugins/Runtime/HairStrands/Source/HairStrandsCore/Private/GroomAsset.cpp
+++ b/Engine/Plugins/Runtime/HairStrands/Source/HairStrandsCore/Private/GroomAsset.cpp
@@ -836,7 +836,7 @@ void UGroomAsset::UpdateResource()
 				InternalUpdateResource(GroupData.Strands.InterpolationResource);
 			}
 
-			if ((ChangeType & GroomChangeType_LOD) && GroupData.Strands.HasValidData())
+			if ((ChangeType & GroomChangeType_LOD) && GroupData.Strands.HasValidData() && CanRebuildFromDescription())
 			{
 				const FHairDescriptionGroups& LocalHairDescriptionGroups = GetHairDescriptionGroups();
 
@@ -1130,7 +1130,7 @@ void UGroomAsset::PostLoad()
 
 	// Build hair strands if needed
 #if WITH_EDITORONLY_DATA
-	if (!bSucceed && IsHairStrandsEnabled(EHairStrandsShaderType::Strands))
+	if (!bSucceed && IsHairStrandsEnabled(EHairStrandsShaderType::Strands) && CanRebuildFromDescription())
 	{
 		const FHairDescriptionGroups& LocalHairDescriptionGroups = GetHairDescriptionGroups();
 
@@ -1220,35 +1220,38 @@ void UGroomAsset::PreSave(FObjectPreSaveContext ObjectSaveContext)
 
 	if ((ChangeType & GroomChangeType_Interpolation) || (ChangeType & GroomChangeType_LOD))
 	{
-		const FHairDescriptionGroups& LocalHairDescriptionGroups = GetHairDescriptionGroups();
-
-		FGroomComponentRecreateRenderStateContext RecreateRenderContext(this);
-		for (uint32 GroupIt = 0; GroupIt < GroupCount; ++GroupIt)
+		if (CanRebuildFromDescription())
 		{
-			const bool bHasInterpolationChanged = !(CachedHairGroupsInterpolation[GroupIt] == HairGroupsInterpolation[GroupIt]);
-			const bool bHasLODChanged = !(CachedHairGroupsLOD[GroupIt] == HairGroupsLOD[GroupIt]);
-			const bool bNeedInterpoldationData = NeedsInterpolationData(GroupIt);
+			const FHairDescriptionGroups &LocalHairDescriptionGroups = GetHairDescriptionGroups();
 
-			if (bHasInterpolationChanged)
-			{
-				BuildHairGroup(
-					GroupIt,
-					bNeedInterpoldationData,
-					LocalHairDescriptionGroups,
-					HairGroupsInterpolation,
-					HairGroupsLOD,
-					HairGroupsInfo,
-					HairGroupsData);
-			}
-			else if (bHasLODChanged)
+			FGroomComponentRecreateRenderStateContext RecreateRenderContext(this);
+			for (uint32 GroupIt = 0; GroupIt < GroupCount; ++GroupIt)
 			{
-				BuildHairGroupCluster(
-					GroupIt,
-					LocalHairDescriptionGroups,
-					HairGroupsInterpolation,
-					HairGroupsLOD,
-					HairGroupsInfo,
-					HairGroupsData);
+				const bool bHasInterpolationChanged = !(CachedHairGroupsInterpolation[GroupIt] == HairGroupsInterpolation[GroupIt]);
+				const bool bHasLODChanged = !(CachedHairGroupsLOD[GroupIt] == HairGroupsLOD[GroupIt]);
+				const bool bNeedInterpoldationData = NeedsInterpolationData(GroupIt);
+
+				if (bHasInterpolationChanged)
+				{
+					BuildHairGroup(
+						GroupIt,
+						bNeedInterpoldationData,
+						LocalHairDescriptionGroups,
+						HairGroupsInterpolation,
+						HairGroupsLOD,
+						HairGroupsInfo,
+						HairGroupsData);
+				}
+				else if (bHasLODChanged)
+				{
+					BuildHairGroupCluster(
+						GroupIt,
+						LocalHairDescriptionGroups,
+						HairGroupsInterpolation,
+						HairGroupsLOD,
+						HairGroupsInfo,
+						HairGroupsData);
+				}
 			}
 		}
 		InitResources();
@@ -2492,7 +2495,7 @@ bool UGroomAsset::BuildCardsGeometry(uint32 GroupIndex)
 {
 	LLM_SCOPE(ELLMTag::Meshes) // This should be a Groom LLM tag, but there is no LLM tag bit left
 
-	if (!IsHairStrandsEnabled(EHairStrandsShaderType::Cards) || HairGroupsCards.Num() == 0)
+	if (!IsHairStrandsEnabled(EHairStrandsShaderType::Cards) || HairGroupsCards.Num() == 0 || !CanRebuildFromDescription())
 	{
 		return false;
 	}
@@ -3380,7 +3383,7 @@ void UGroomAsset::StripLODs(const TArray<int32>& LODsToKeep, bool bRebuildResour
 	}
 
 	// Rebuild the LOD data
-	if (bRebuildResources)
+	if (bRebuildResources && CanRebuildFromDescription())
 	{
 		const FHairDescriptionGroups& LocalHairDescriptionGroups = GetHairDescriptionGroups();
 		for (int32 GroupIt = 0; GroupIt < GroupCount; ++GroupIt)
@@ -3416,7 +3419,7 @@ bool UGroomAsset::HasDebugData() const
 void UGroomAsset::CreateDebugData()
 {
 #if WITH_EDITORONLY_DATA
-	if (!IsHairStrandsEnabled(EHairStrandsShaderType::Strands))
+	if (!IsHairStrandsEnabled(EHairStrandsShaderType::Strands) || !CanRebuildFromDescription())
 		return;
 
 	for (uint32 GroupIndex = 0, GroupCount = GetNumHairGroups(); GroupIndex < GroupCount; ++GroupIndex)
@@ -3528,7 +3531,7 @@ static TQueue<FHairProceduralCardsQuery*> HairCardsQueuries;
 // Generate geometry and textures for hair cards
 void UGroomAsset::SaveProceduralCards(uint32 DescIndex)
 {
-	if (!IsHairStrandsEnabled(EHairStrandsShaderType::Cards))
+	if (!IsHairStrandsEnabled(EHairStrandsShaderType::Cards) || !CanRebuildFromDescription())
 	{
 		return;
 	}
@@ -3667,6 +3670,10 @@ bool UGroomAsset::GetHairStrandsDatas(
 	FHairStrandsDatas& OutStrandsData,
 	FHairStrandsDatas& OutGuidesData)
 {
+	if (!CanRebuildFromDescription())
+	{
+		return false;
+	}
 	const FHairDescriptionGroups& DescriptionGroups = GetHairDescriptionGroups();
 	const bool bIsValid = DescriptionGroups.IsValid() && GroupIndex < DescriptionGroups.HairGroups.Num();
 	if (bIsValid)
