Fix crash when forcing cards geometry on groom having both cards and strands geometry for a given LOD.

diff --git a/Engine/Plugins/Runtime/HairStrands/Source/HairStrandsCore/Public/GroomInstance.h b/Engine/Plugins/Runtime/HairStrands/Source/HairStrandsCore/Public/GroomInstance.h
--- a/Engine/Plugins/Runtime/HairStrands/Source/HairStrandsCore/Public/GroomInstance.h
+++ b/Engine/Plugins/Runtime/HairStrands/Source/HairStrandsCore/Public/GroomInstance.h
@@ -293,6 +293,7 @@ struct HAIRSTRANDSCORE_API FHairGroupInstance : public FHairStrandsInstance
 	virtual const FBoxSphereBounds& GetBounds() const override { return Debug.Proxy->GetBounds(); }
 	virtual const FBoxSphereBounds& GetLocalBounds() const { return Debug.Proxy->GetLocalBounds(); }
 	virtual const FHairGroupPublicData* GetHairData() const { return HairGroupPublicData; }
+	virtual const EHairGeometryType GetHairGeometry() const { return GeometryType; }
 
 	/** Get the current local to world transform according to the internal binding type */
 	FORCEINLINE const FTransform& GetCurrentLocalToWorld() const
diff --git a/Engine/Source/Runtime/Renderer/Private/HairStrands/HairStrandsInterface.cpp b/Engine/Source/Runtime/Renderer/Private/HairStrands/HairStrandsInterface.cpp
--- a/Engine/Source/Runtime/Renderer/Private/HairStrands/HairStrandsInterface.cpp
+++ b/Engine/Source/Runtime/Renderer/Private/HairStrands/HairStrandsInterface.cpp
@@ -435,8 +435,7 @@ bool IsHairStrandsVisibleInShadows(const FViewInfo& View, const FHairStrandsInst
 	bool bIsVisibleInShadow = false;
 	if (const FHairGroupPublicData* HairData = Instance.GetHairData())
 	{
-		const int32 LODIndex = HairData->LODIndex;
-		const bool bIsStrands = LODIndex >= 0 && HairData->IsVisible(LODIndex) && HairData->GetGeometryType(LODIndex) == EHairGeometryType::Strands;
+		const bool bIsStrands = HairData->LODIndex >= 0 && Instance.GetHairGeometry() == EHairGeometryType::Strands;
 		if (!bIsStrands)
 		{
 			return false;
diff --git a/Engine/Source/Runtime/Renderer/Public/HairStrandsInterface.h b/Engine/Source/Runtime/Renderer/Public/HairStrandsInterface.h
--- a/Engine/Source/Runtime/Renderer/Public/HairStrandsInterface.h
+++ b/Engine/Source/Runtime/Renderer/Public/HairStrandsInterface.h
@@ -146,6 +146,7 @@ struct RENDERER_API FHairStrandsInstance
 	virtual const FBoxSphereBounds& GetBounds() const = 0;
 	virtual const FBoxSphereBounds& GetLocalBounds() const = 0;
 	virtual const FHairGroupPublicData* GetHairData() const { return nullptr; }
+	virtual const EHairGeometryType GetHairGeometry() const { return EHairGeometryType::NoneGeometry; }
 protected:
 	mutable uint32 RefCount = 0;
 };
