Update to hair strands to support Alembic groom_group_cards_id for hair card generation

diff --git a/Engine/Plugins/Runtime/HairStrands/Source/HairStrandsCore/Private/HairAttributes.cpp b/Engine/Plugins/Runtime/HairStrands/Source/HairStrandsCore/Private/HairAttributes.cpp
--- a/Engine/Plugins/Runtime/HairStrands/Source/HairStrandsCore/Private/HairAttributes.cpp
+++ b/Engine/Plugins/Runtime/HairStrands/Source/HairStrandsCore/Private/HairAttributes.cpp
@@ -23,6 +23,7 @@ namespace HairAttribute
 	const FName Strand::CurveType("groom_curve_type");
 	const FName Strand::Knots("groom_knots");
 	const FName Strand::GroupName("groom_group_name");
+	const FName Strand::GroupCardsID("groom_group_cards_id");
 
 	const FName Groom::Color("groom_color");
 	const FName Groom::Roughness("groom_roughness");
diff --git a/Engine/Plugins/Runtime/HairStrands/Source/HairStrandsCore/Public/GroomAsset.h b/Engine/Plugins/Runtime/HairStrands/Source/HairStrandsCore/Public/GroomAsset.h
--- a/Engine/Plugins/Runtime/HairStrands/Source/HairStrandsCore/Public/GroomAsset.h
+++ b/Engine/Plugins/Runtime/HairStrands/Source/HairStrandsCore/Public/GroomAsset.h
@@ -60,6 +60,9 @@ struct HAIRSTRANDSCORE_API FHairGroupInfo
 
 	UPROPERTY(VisibleAnywhere, Category = "Info", meta = (DisplayName = "Imported width (DCC) of the larget hair strands"))
 	float MaxImportedWidth = -1.f;
+
+	UPROPERTY(VisibleAnywhere, Category = "Info", meta = (DisplayName = "Card Group ID for split hair card generation"))
+	int32 GroupCardsID = 0;
 };
 
 USTRUCT(BlueprintType)
diff --git a/Engine/Plugins/Runtime/HairStrands/Source/HairStrandsCore/Public/HairAttributes.h b/Engine/Plugins/Runtime/HairStrands/Source/HairStrandsCore/Public/HairAttributes.h
--- a/Engine/Plugins/Runtime/HairStrands/Source/HairStrandsCore/Public/HairAttributes.h
+++ b/Engine/Plugins/Runtime/HairStrands/Source/HairStrandsCore/Public/HairAttributes.h
@@ -30,6 +30,7 @@ namespace HairAttribute
 		extern HAIRSTRANDSCORE_API const FName CurveType;		// FName (EGroomCurveType)
 		extern HAIRSTRANDSCORE_API const FName Knots;			// float[]
 		extern HAIRSTRANDSCORE_API const FName GroupName;		// FName
+		extern HAIRSTRANDSCORE_API const FName GroupCardsID;	// int
 	}
 
 	namespace Groom
