Removed an assert that was firing when rendering hair

diff --git a/Engine/Plugins/Runtime/HairStrands/Source/HairStrandsCore/Private/GroomResources.cpp b/Engine/Plugins/Runtime/HairStrands/Source/HairStrandsCore/Private/GroomResources.cpp
--- a/Engine/Plugins/Runtime/HairStrands/Source/HairStrandsCore/Private/GroomResources.cpp
+++ b/Engine/Plugins/Runtime/HairStrands/Source/HairStrandsCore/Private/GroomResources.cpp
@@ -1035,7 +1035,6 @@ bool FHairStrandsClusterCullingResource::InternalIsDataLoaded()
 
 void FHairStrandsClusterCullingResource::InternalAllocate(FRDGBuilder& GraphBuilder)
 {
-	check(BulkDataRequest.IsNone() == false)
 	BulkDataRequest = FBulkDataRequest();
 
 	if (GHairStrandsBulkData_Validation > 0)
