Add function to readback from GPU Groom's deformed position.

diff --git a/Engine/Plugins/Runtime/HairStrands/Source/HairStrandsCore/Private/GroomBuilder.cpp b/Engine/Plugins/Runtime/HairStrands/Source/HairStrandsCore/Private/GroomBuilder.cpp
--- a/Engine/Plugins/Runtime/HairStrands/Source/HairStrandsCore/Private/GroomBuilder.cpp
+++ b/Engine/Plugins/Runtime/HairStrands/Source/HairStrandsCore/Private/GroomBuilder.cpp
@@ -7,6 +7,10 @@
 #include "HairDescription.h"
 #include "GroomSettings.h"
 #include "GroomBindingBuilder.h"
+#include "RenderGraphBuilder.h"
+#include "RenderGraphUtils.h"
+#include "GlobalShader.h"
+#include "ShaderParameterStruct.h"
 
 #include "Async/ParallelFor.h"
 #include "HAL/IConsoleManager.h"
@@ -3162,4 +3166,191 @@ void FGroomBuilder::BuildClusterBulkData(
 	GroomBuilder_Cluster::BuildClusterBulkData(ClusterData, Out);
 }
 
+///////////////////////////////////////////////////////////////////////////////////////////////////
+// Asynchronous queuing for hair strands positions readback
+
+class FHairStrandCopyPositionCS : public FGlobalShader
+{
+	DECLARE_GLOBAL_SHADER(FHairStrandCopyPositionCS);
+	SHADER_USE_PARAMETER_STRUCT(FHairStrandCopyPositionCS, FGlobalShader);
+
+	BEGIN_SHADER_PARAMETER_STRUCT(FParameters, )
+		SHADER_PARAMETER(uint32, MaxVertexCount)
+		SHADER_PARAMETER_RDG_BUFFER_SRV(Buffer, InPositions)
+		SHADER_PARAMETER_RDG_BUFFER_SRV(Buffer, InPositionOffset)
+		SHADER_PARAMETER_RDG_BUFFER_UAV(RWBuffer, OutPositions)
+	END_SHADER_PARAMETER_STRUCT()
+
+public:
+	static bool ShouldCompilePermutation(const FGlobalShaderPermutationParameters& Parameters) { return IsHairStrandsSupported(EHairStrandsShaderType::Tool, Parameters.Platform); }
+	static void ModifyCompilationEnvironment(const FGlobalShaderPermutationParameters& Parameters, FShaderCompilerEnvironment& OutEnvironment)
+	{
+		FGlobalShader::ModifyCompilationEnvironment(Parameters, OutEnvironment);
+		OutEnvironment.SetDefine(TEXT("SHADER_READ_POSITIONS"), 1);
+	}
+};
+
+IMPLEMENT_GLOBAL_SHADER(FHairStrandCopyPositionCS, "/Engine/Private/HairStrands/HairStrandsTexturesGeneration.usf", "MainCS", SF_Compute);
+
+static FRDGBufferRef ReadPositions(
+	FRDGBuilder& GraphBuilder,
+	FGlobalShaderMap* ShaderMap,
+	const UGroomComponent* Component,
+	uint32 GroupIt)
+{
+	FRDGBufferRef OutPositions = nullptr;
+	if (const FHairGroupInstance* Instance = Component->GetGroupInstance(GroupIt))
+	{
+		const uint32 NumVertices = Instance->Strands.Data->PointCount;
+
+		FRDGBufferSRVRef InPositions = nullptr;
+		FRDGBufferSRVRef InPositionOffset = nullptr;
+		if (Instance->Strands.DeformedResource)
+		{
+			InPositions = Register(GraphBuilder, Instance->Strands.DeformedResource->GetBuffer(FHairStrandsDeformedResource::EFrameType::Current), ERDGImportedBufferFlags::CreateSRV).SRV;
+			InPositionOffset = Register(GraphBuilder, Instance->Strands.DeformedResource->GetPositionOffsetBuffer(FHairStrandsDeformedResource::EFrameType::Current), ERDGImportedBufferFlags::CreateSRV).SRV;
+		}
+		else
+		{
+			InPositions = Register(GraphBuilder, Instance->Strands.RestResource->PositionBuffer, ERDGImportedBufferFlags::CreateSRV).SRV;
+			InPositionOffset = Register(GraphBuilder, Instance->Strands.RestResource->PositionOffsetBuffer, ERDGImportedBufferFlags::CreateSRV).SRV;
+		}
+
+		OutPositions = GraphBuilder.CreateBuffer(FRDGBufferDesc::CreateBufferDesc(sizeof(FVector4f), NumVertices), TEXT("StrandsPositionForReadback"));
+
+		FHairStrandCopyPositionCS::FParameters* Parameters = GraphBuilder.AllocParameters<FHairStrandCopyPositionCS::FParameters>();
+		Parameters->MaxVertexCount = NumVertices;
+		Parameters->InPositionOffset = InPositionOffset;
+		Parameters->InPositions = InPositions;
+		Parameters->OutPositions = GraphBuilder.CreateUAV(OutPositions, PF_A32B32G32R32F);
+
+		FIntVector DispatchCount = FComputeShaderUtils::GetGroupCount(NumVertices, 128);
+
+		TShaderMapRef<FHairStrandCopyPositionCS> ComputeShader(ShaderMap);
+		FComputeShaderUtils::AddPass(
+			GraphBuilder,
+			RDG_EVENT_NAME("HairStrands::CopyPositions"),
+			ComputeShader,
+			Parameters,
+			DispatchCount);
+	}
+	return OutPositions;
+}
+
+struct FStrandsPositionData
+{
+	TUniquePtr<FRHIGPUBufferReadback> GPUPositions;
+	uint32 NumBytes = 0;
+	uint32 GroupIt = 0;
+	uint32 GroupCount = 0;
+	uint32 ReadbackDelay = 0;
+	const UGroomComponent* Component = nullptr;
+	FStrandsPositionOutput* Output;
+};
+
+TQueue<FStrandsPositionData*>		GStrandsPositionQueries;
+TQueue<FStrandsPositionData*>		GStrandsPositionReadbacks;
+
+bool HasHairStrandsPositionQueries()
+{
+	return !GStrandsPositionQueries.IsEmpty() || !GStrandsPositionReadbacks.IsEmpty();
+}
+
+void RunHairStrandsPositionQueries(FRDGBuilder& GraphBuilder, FGlobalShaderMap* ShaderMap, const struct FShaderPrintData* DebugShaderData)
+{
+	// Operations are ordered in reverse to ensure they are processed on independent frames to avoid TDR on heavy grooms.
+
+	// 2. Readback positions
+	{
+		TArray<FStrandsPositionData*> QueryNotReady;
+		FStrandsPositionData* Q = nullptr;
+		while (GStrandsPositionReadbacks.Dequeue(Q))
+		{			
+			if (Q->GPUPositions->IsReady() || Q->ReadbackDelay == 0)
+			{
+				const FHairGroupInstance* Instance = Q->Component->GetGroupInstance(Q->GroupIt);
+				const uint32 CurveCount = Instance->Strands.RestResource->BulkData.CurveCount;
+				const uint32 PointCount = Instance->Strands.RestResource->BulkData.PointCount;
+
+				check(Q->GroupIt < uint32(Q->Output->Groups.Num()));
+				FStrandsPositionOutput::FGroup& Group = Q->Output->Groups[Q->GroupIt];
+				Group.Reserve(CurveCount);
+
+				FStrandsPositionOutput::FStrand Strand;
+				const FVector4f* Positions = (const FVector4f*)Q->GPUPositions->Lock(Q->NumBytes);
+				for (uint32 PointIt=0; PointIt<PointCount; ++PointIt)
+				{
+					const FVector4f P = Positions[PointIt];
+					const bool bIsFirstPoint = P.W == 1;
+					if (bIsFirstPoint && Strand.Num() > 0)
+					{
+						Group.Add(Strand);
+						Strand.SetNum(0);
+					}
+					Strand.Add(FVector3f(P.X, P.Y, P.Z));
+				}
+				Q->GPUPositions->Unlock();
+
+				--Q->Output->Status;
+				check(Q->Output->Status >= 0);
+			}
+			else
+			{
+				if (Q->ReadbackDelay > 0) { --Q->ReadbackDelay; }
+				QueryNotReady.Add(Q);
+			}
+		}
+
+		for (FStrandsPositionData* N : QueryNotReady)
+		{
+			GStrandsPositionReadbacks.Enqueue(N);
+		}
+	}
+
+	// 1. Copy positions
+	{
+		FStrandsPositionData* Q;
+		while (GStrandsPositionQueries.Dequeue(Q))
+		{
+			check(Q->Component);
+			
+			FRDGBufferRef Positions = ReadPositions(GraphBuilder, ShaderMap, Q->Component, Q->GroupIt);
+			FRHIGPUBufferReadback* GPUPositions = Q->GPUPositions.Get();
+			AddEnqueueCopyPass(GraphBuilder, GPUPositions, Positions, Positions->Desc.GetSize());
+			Q->NumBytes = Positions->Desc.GetSize();
+			GStrandsPositionReadbacks.Enqueue(Q);
+		}
+	}
+}
+
+#if WITH_EDITOR
+bool RequestStrandsPosition(const UGroomComponent* Component, FStrandsPositionOutput* Output)
+{
+	if (Component == nullptr || Output == nullptr)
+	{
+		return false;
+	}
+
+	const uint32 GroupCount = Component->GetGroupCount();
+	for (uint32 GroupIt = 0; GroupIt < GroupCount; ++GroupIt)
+	{
+		FStrandsPositionData* Data = new FStrandsPositionData();
+		Data->NumBytes = 0;
+		Data->GroupIt = GroupIt;
+		Data->GroupCount = GroupCount;
+		Data->Component = Component;
+		Data->ReadbackDelay = 3;
+		Data->GPUPositions = MakeUnique<FRHIGPUBufferReadback>(TEXT("Readback.Positions"));
+		Data->Output = Output;
+		GStrandsPositionQueries.Enqueue(Data);
+	}
+
+	*Output = FStrandsPositionOutput();
+	Output->Component = Component;
+	Output->Status = GroupCount;
+	Output->Groups.SetNum(GroupCount);
+	return true;
+}
+#endif
+
 #undef LOCTEXT_NAMESPACE
diff --git a/Engine/Plugins/Runtime/HairStrands/Source/HairStrandsCore/Private/GroomManager.cpp b/Engine/Plugins/Runtime/HairStrands/Source/HairStrandsCore/Private/GroomManager.cpp
--- a/Engine/Plugins/Runtime/HairStrands/Source/HairStrandsCore/Private/GroomManager.cpp
+++ b/Engine/Plugins/Runtime/HairStrands/Source/HairStrandsCore/Private/GroomManager.cpp
@@ -985,9 +985,10 @@ void RunHairStrandsFolliculeMaskQueries(FRDGBuilder& GraphBuilder, FGlobalShader
 #if WITH_EDITOR
 bool HasHairStrandsTexturesQueries();
 void RunHairStrandsTexturesQueries(FRDGBuilder& GraphBuilder, FGlobalShaderMap* ShaderMap, const struct FShaderPrintData* ShaderPrintData);
-#endif
 
-#if WITH_EDITOR
+bool HasHairStrandsPositionQueries();
+void RunHairStrandsPositionQueries(FRDGBuilder& GraphBuilder, FGlobalShaderMap* ShaderMap, const struct FShaderPrintData* DebugShaderData);
+
 bool HasHairCardsAtlasQueries();
 void RunHairCardsAtlasQueries(FRDGBuilder& GraphBuilder, FGlobalShaderMap* ShaderMap, const struct FShaderPrintData* ShaderPrintData);
 #endif
@@ -999,6 +1000,11 @@ static void RunHairStrandsProcess(FRDGBuilder& GraphBuilder, FGlobalShaderMap* S
 	{
 		RunHairStrandsTexturesQueries(GraphBuilder, ShaderMap, ShaderPrintData);
 	}
+
+	if (HasHairStrandsPositionQueries())
+	{
+		RunHairStrandsPositionQueries(GraphBuilder, ShaderMap, ShaderPrintData);
+	}
 #endif
 
 	if (HasHairStrandsFolliculeMaskQueries())
@@ -1051,6 +1057,7 @@ void ProcessHairStrandsBookmark(
 		#if WITH_EDITOR
 			HasHairCardsAtlasQueries() ||
 			HasHairStrandsTexturesQueries() ||
+			HasHairStrandsPositionQueries() ||
 		#endif
 			HasHairStrandsFolliculeMaskQueries();
 
diff --git a/Engine/Plugins/Runtime/HairStrands/Source/HairStrandsCore/Public/GroomBuilder.h b/Engine/Plugins/Runtime/HairStrands/Source/HairStrandsCore/Public/GroomBuilder.h
--- a/Engine/Plugins/Runtime/HairStrands/Source/HairStrandsCore/Public/GroomBuilder.h
+++ b/Engine/Plugins/Runtime/HairStrands/Source/HairStrandsCore/Public/GroomBuilder.h
@@ -23,6 +23,7 @@ struct FHairDescriptionGroup;
 struct FHairGroupInfo;
 class FHairDescription;
 class UGroomAsset;
+class UGroomComponent;
 
 struct FHairStrandsVoxelData
 {
@@ -105,4 +106,21 @@ struct HAIRSTRANDSCORE_API FGroomBuilder
 	static void VoxelizeGroupIndex(
 		const FHairDescriptionGroups& In,
 		FHairStrandsVoxelData& Out);
-};
\ No newline at end of file
+};
+
+////////////////////////////////////////////////////////////////////////////////////////////////////////////////////////
+
+struct FStrandsPositionOutput
+{
+	typedef TArray<FVector3f> FStrand;
+	typedef TArray<FStrand>   FGroup;
+
+	TArray<FGroup> Groups;
+	const UGroomComponent* Component = nullptr;
+
+	int32 Status = -1;
+
+	bool IsValid() const { return Status == 0; }
+};
+
+HAIRSTRANDSCORE_API bool RequestStrandsPosition(const UGroomComponent* Component, FStrandsPositionOutput* Output);
diff --git a/Engine/Shaders/Private/HairStrands/HairStrandsTexturesGeneration.usf b/Engine/Shaders/Private/HairStrands/HairStrandsTexturesGeneration.usf
--- a/Engine/Shaders/Private/HairStrands/HairStrandsTexturesGeneration.usf
+++ b/Engine/Shaders/Private/HairStrands/HairStrandsTexturesGeneration.usf
@@ -566,3 +566,28 @@ void MainCS(int2 DispatchThreadId : SV_DispatchThreadID)
 	}
 }
 #endif // SHADER_TEXTURE_DILATION
+
+/////////////////////////////////////////////////////////////////////////////////////////////////////////////////////////////
+
+#if SHADER_READ_POSITIONS
+
+#include "HairStrandsVertexFactoryCommon.ush"
+
+uint MaxVertexCount;
+Buffer<uint4> InPositions;
+Buffer<float4> InPositionOffset;
+RWBuffer<float4> OutPositions;
+[numthreads(128, 1, 1)]
+void MainCS(int2 DispatchThreadId : SV_DispatchThreadID)
+{
+	uint VertexIndex = DispatchThreadId.x;
+	if (VertexIndex < MaxVertexCount)
+	{
+		const uint4 PackedPosition = InPositions[VertexIndex];
+		const float3 PositionOffset = InPositionOffset[0].xyz;
+		const FHairControlPoint CP = UnpackHairControlPoint(PackedPosition, PositionOffset, 1.0f /*InVF_Radius*/, 1.0f /*InVF_RootScale*/, 1.0f /*InVF_TipScale*/);
+		
+		OutPositions[VertexIndex] = float4(CP.Position, CP.Type);
+	}
+}
+#endif // SHADER_READ_POSITIONS
\ No newline at end of file
