Add proper package version support for hair description by forwarding correctly parent UE/Custom/Licensee versions.

diff --git a/Engine/Plugins/Runtime/HairStrands/Source/HairStrandsCore/Private/GroomAsset.cpp b/Engine/Plugins/Runtime/HairStrands/Source/HairStrandsCore/Private/GroomAsset.cpp
--- a/Engine/Plugins/Runtime/HairStrands/Source/HairStrandsCore/Private/GroomAsset.cpp
+++ b/Engine/Plugins/Runtime/HairStrands/Source/HairStrandsCore/Private/GroomAsset.cpp
@@ -700,6 +700,14 @@ void UGroomAsset::Serialize(FArchive& Ar)
 	Ar.UsingCustomVersion(FUE5MainStreamObjectVersion::GUID);	// Needed to support MeshDescription AttributesSet serialization
 	Ar.UsingCustomVersion(FAnimObjectVersion::GUID);    // Needed to support Cards and Cluster culling serialization
 
+#if WITH_EDITORONLY_DATA
+	// Retrieve achieve version to forward them to the hair description bulk data
+	HairDescriptionVersion.CustomVersions = Ar.GetCustomVersions();
+	HairDescriptionVersion.UEVersion = Ar.UEVer();
+	HairDescriptionVersion.LicenseeVersion = Ar.LicenseeUEVer();
+	HairDescriptionVersion.bIsValid = true;
+#endif // WITH_EDITORONLY_DATA
+
 	if (Ar.CustomVer(FPhysicsObjectVersion::GUID) >= FPhysicsObjectVersion::GroomWithDescription)
 	{
 		FStripDataFlags StripFlags(Ar, ClassDataStripFlags);
@@ -2067,6 +2075,24 @@ FString UGroomAsset::GetDerivedDataKeyForMeshes(uint32 GroupIndex)
 	return DerivedDataKey;
 }
 
+const FHairDescriptionVersion& UGroomAsset::GetHairDescriptionVersion() const
+{
+	if (!HairDescriptionVersion.IsValid())
+	{
+		FHairDescriptionVersion* MutableHairDescription = const_cast<FHairDescriptionVersion*>(&HairDescriptionVersion);
+		if (UObject* Outer = GetOuter())
+		{
+			MutableHairDescription->UEVersion = Outer->GetLinkerUEVersion();
+			MutableHairDescription->LicenseeVersion = Outer->GetLinkerLicenseeUEVersion();
+			if (FLinkerLoad* Linker = Outer->GetLinker())
+			{
+				MutableHairDescription->CustomVersions = Linker->GetCustomVersions();
+			}
+		}
+		MutableHairDescription->bIsValid = true; 
+	}
+	return HairDescriptionVersion;
+}
 
 void UGroomAsset::CommitHairDescription(FHairDescription&& InHairDescription, EHairDescriptionType Type)
 {
@@ -2083,7 +2109,7 @@ void UGroomAsset::CommitHairDescription(FHairDescription&& InHairDescription, EH
 	CachedHairDescriptionGroups[HairDescriptionType] = MakeUnique<FHairDescriptionGroups>();
 	FGroomBuilder::BuildHairDescriptionGroups(*CachedHairDescription[HairDescriptionType], *CachedHairDescriptionGroups[HairDescriptionType]);
 
-	HairDescriptionBulkData[HairDescriptionType]->SaveHairDescription(*CachedHairDescription[HairDescriptionType]);
+	HairDescriptionBulkData[HairDescriptionType]->SaveHairDescription(*CachedHairDescription[HairDescriptionType], GetHairDescriptionVersion());
 }
 
 const FHairDescriptionGroups& UGroomAsset::GetHairDescriptionGroups()
@@ -2093,7 +2119,7 @@ const FHairDescriptionGroups& UGroomAsset::GetHairDescriptionGroups()
 	if (!CachedHairDescription[HairDescriptionType])
 	{
 		CachedHairDescription[HairDescriptionType] = MakeUnique<FHairDescription>();
-		HairDescriptionBulkData[HairDescriptionType]->LoadHairDescription(*CachedHairDescription[HairDescriptionType]);
+		HairDescriptionBulkData[HairDescriptionType]->LoadHairDescription(*CachedHairDescription[HairDescriptionType], GetHairDescriptionVersion());
 	}
 	if (!CachedHairDescriptionGroups[HairDescriptionType])
 	{
@@ -2109,7 +2135,7 @@ FHairDescription UGroomAsset::GetHairDescription() const
 	FHairDescription OutHairDescription;
 	if (HairDescriptionBulkData[HairDescriptionType])
 	{
-		HairDescriptionBulkData[HairDescriptionType]->LoadHairDescription(OutHairDescription);
+		HairDescriptionBulkData[HairDescriptionType]->LoadHairDescription(OutHairDescription, GetHairDescriptionVersion());
 	}
 	return MoveTemp(OutHairDescription);
 }
diff --git a/Engine/Plugins/Runtime/HairStrands/Source/HairStrandsCore/Private/HairDescription.cpp b/Engine/Plugins/Runtime/HairStrands/Source/HairStrandsCore/Private/HairDescription.cpp
--- a/Engine/Plugins/Runtime/HairStrands/Source/HairStrandsCore/Private/HairDescription.cpp
+++ b/Engine/Plugins/Runtime/HairStrands/Source/HairStrandsCore/Private/HairDescription.cpp
@@ -87,7 +87,7 @@ void FHairDescriptionBulkData::Serialize(FArchive& Ar, UObject* Owner)
 	if (Ar.IsTransacting())
 	{
 		// If transacting, keep these members alive the other side of an undo, otherwise their values will get lost
-		CustomVersions.Serialize(Ar);
+		Version.CustomVersions.Serialize(Ar);
 		Ar << bBulkDataUpdated;
 	}
 	else
@@ -99,8 +99,8 @@ void FHairDescriptionBulkData::Serialize(FArchive& Ar, UObject* Owner)
 			if (!bBulkDataUpdated)
 			{
 				FHairDescription HairDescription;
-				LoadHairDescription(HairDescription);
-				SaveHairDescription(HairDescription);
+				LoadHairDescription(HairDescription, Version);
+				SaveHairDescription(HairDescription, Version);
 			}
 		}
 	}
@@ -113,21 +113,25 @@ void FHairDescriptionBulkData::Serialize(FArchive& Ar, UObject* Owner)
 	{
 		// If loading, take the package custom version so it can be applied to the bulk data archive
 		// when unpacking HairDescription from it
-		CustomVersions = BulkData.GetCustomVersions(Ar);
+		Version.CustomVersions = BulkData.GetCustomVersions(Ar);
 	}
 }
 
-void FHairDescriptionBulkData::SaveHairDescription(FHairDescription& HairDescription)
+void FHairDescriptionBulkData::SaveHairDescription(FHairDescription& HairDescription, const FHairDescriptionVersion& InVersion)
 {
 	BulkData.RemoveBulkData();
 
 	if (HairDescription.IsValid())
 	{
 		FBulkDataWriter Ar(BulkData, /*bIsPersistent*/ true);
+		Ar.SetCustomVersions(InVersion.CustomVersions);
+		Ar.SetUEVer(InVersion.UEVersion);
+		Ar.SetLicenseeUEVer(InVersion.LicenseeVersion);
 		HairDescription.Serialize(Ar);
 
 		// Preserve CustomVersions at save time so we can reuse the same ones when reloading direct from memory
-		CustomVersions = Ar.GetCustomVersions();
+		Version = InVersion;
+		Version.CustomVersions = Ar.GetCustomVersions();
 	}
 
 	// Use bulk data hash instead of guid to identify content to improve DDC cache hit
@@ -138,8 +142,7 @@ void FHairDescriptionBulkData::SaveHairDescription(FHairDescription& HairDescrip
 	bBulkDataUpdated = true;
 }
 
-
-void FHairDescriptionBulkData::LoadHairDescription(FHairDescription& HairDescription)
+void FHairDescriptionBulkData::LoadHairDescription(FHairDescription& HairDescription, const FHairDescriptionVersion& InVersion)
 {
 	HairDescription.Reset();
 
@@ -149,9 +152,12 @@ void FHairDescriptionBulkData::LoadHairDescription(FHairDescription& HairDescrip
 
 		// Propagate the custom version information from the package to the bulk data, so that the HairDescription
 		// is serialized with the same versioning
-		Ar.SetCustomVersions(CustomVersions);
+		Ar.SetCustomVersions(InVersion.CustomVersions);
+		Ar.SetUEVer(InVersion.UEVersion);
+		Ar.SetLicenseeUEVer(InVersion.LicenseeVersion);
 
 		HairDescription.Serialize(Ar);
+		Version.CustomVersions = Ar.GetCustomVersions();
 	}
 }
 
diff --git a/Engine/Plugins/Runtime/HairStrands/Source/HairStrandsCore/Public/GroomAsset.h b/Engine/Plugins/Runtime/HairStrands/Source/HairStrandsCore/Public/GroomAsset.h
--- a/Engine/Plugins/Runtime/HairStrands/Source/HairStrandsCore/Public/GroomAsset.h
+++ b/Engine/Plugins/Runtime/HairStrands/Source/HairStrandsCore/Public/GroomAsset.h
@@ -674,6 +674,10 @@ private:
 	TUniquePtr<FHairDescription> CachedHairDescription[EHairDescriptionType::Count];
 	TUniquePtr<FHairDescriptionGroups> CachedHairDescriptionGroups[EHairDescriptionType::Count];
 
+	// Hair description version aggregate all versions used by the asset (i.e., UE/Custom/Licensee version)
+	const FHairDescriptionVersion& GetHairDescriptionVersion() const;
+	FHairDescriptionVersion HairDescriptionVersion;
+
 	TArray<FString> StrandsDerivedDataKey;
 	TArray<FString> CardsDerivedDataKey;
 	TArray<FString> MeshesDerivedDataKey;
diff --git a/Engine/Plugins/Runtime/HairStrands/Source/HairStrandsCore/Public/HairDescription.h b/Engine/Plugins/Runtime/HairStrands/Source/HairStrandsCore/Public/HairDescription.h
--- a/Engine/Plugins/Runtime/HairStrands/Source/HairStrandsCore/Public/HairDescription.h
+++ b/Engine/Plugins/Runtime/HairStrands/Source/HairStrandsCore/Public/HairDescription.h
@@ -108,6 +108,15 @@ private:
 	int32 NumStrands;
 };
 
+struct FHairDescriptionVersion
+{
+	FCustomVersionContainer CustomVersions;
+	FPackageFileVersion UEVersion;
+	int32 LicenseeVersion;
+	bool bIsValid = false;
+	bool IsValid() const { return bIsValid; }
+};
+
 /**
  * Bulk data storage for FHairDescription
  */
@@ -123,10 +132,10 @@ public:
 	void Serialize(FArchive& Ar, UObject* Owner);
 
 	/** Stores a new hair description in the bulk data */
-	void SaveHairDescription(FHairDescription& HairDescription);
+	void SaveHairDescription(FHairDescription& HairDescription, const FHairDescriptionVersion& Version);
 
 	/** Loads the hair description from the bulk data */
-	void LoadHairDescription(FHairDescription& HairDescription);
+	void LoadHairDescription(FHairDescription& HairDescription, const FHairDescriptionVersion& Version);
 
 	/** Empties the bulk data */
 	void Empty();
@@ -149,8 +158,9 @@ private:
 	/** GUID associated with the stored bulk data */
 	FGuid Guid;
 
-	/** Custom version to propagate to archive when serializing the bulk data */
-	FCustomVersionContainer CustomVersions;
+	/** UE/Custom/Licensee version to propagate to archive when serializing the bulk data */
+	FHairDescriptionVersion Version;
+
 
 	/** Whether the bulk data has been written via SaveHairDescription */
 	bool bBulkDataUpdated = false;
