Fix invalid type when importing Groom roughness attribute per vertex.

diff --git a/Engine/Plugins/Runtime/HairStrands/Source/HairStrandsCore/Private/GroomBuilder.cpp b/Engine/Plugins/Runtime/HairStrands/Source/HairStrandsCore/Private/GroomBuilder.cpp
--- a/Engine/Plugins/Runtime/HairStrands/Source/HairStrandsCore/Private/GroomBuilder.cpp
+++ b/Engine/Plugins/Runtime/HairStrands/Source/HairStrandsCore/Private/GroomBuilder.cpp
@@ -1480,9 +1480,9 @@ bool FHairDescription::HasColorAttributes() const
 bool FHairDescription::HasRoughnessAttributes() const
 {
 	return
-		VertexAttributes().GetAttributesRef<FVector3f>(HairAttribute::Vertex::Roughness).IsValid() ||
-		StrandAttributes().GetAttributesRef<FVector3f>(HairAttribute::Strand::Roughness).IsValid() ||
-		GroomAttributes().GetAttributesRef<FVector3f>(HairAttribute::Groom::Roughness).IsValid();
+		VertexAttributes().GetAttributesRef<float>(HairAttribute::Vertex::Roughness).IsValid() ||
+		StrandAttributes().GetAttributesRef<float>(HairAttribute::Strand::Roughness).IsValid() ||
+		GroomAttributes().GetAttributesRef<float>(HairAttribute::Groom::Roughness).IsValid();
 }
 
 bool FGroomBuilder::BuildHairDescriptionGroups(const FHairDescription& HairDescription, FHairDescriptionGroups& Out)
