* Add memory tracker tag dedicated for groom and groom editor.
* Fix missing untracked allocation related to groom

diff --git a/Engine/Plugins/Importers/AlembicHairImporter/Source/AlembicHairTranslator/Private/AlembicHairTranslatorModule.cpp b/Engine/Plugins/Importers/AlembicHairImporter/Source/AlembicHairTranslator/Private/AlembicHairTranslatorModule.cpp
--- a/Engine/Plugins/Importers/AlembicHairImporter/Source/AlembicHairTranslator/Private/AlembicHairTranslatorModule.cpp
+++ b/Engine/Plugins/Importers/AlembicHairImporter/Source/AlembicHairTranslator/Private/AlembicHairTranslatorModule.cpp
@@ -4,11 +4,15 @@
 
 #include "AlembicHairTranslator.h"
 #include "HairStrandsEditor.h"
+#include "HAL/LowLevelMemTracker.h"
+
+LLM_DECLARE_TAG_API(GroomEditor, HAIRSTRANDSEDITOR_API);
 
 IMPLEMENT_MODULE(FAlembicHairTranslatorModule, AlembicHairTranslatorModule);
 
 void FAlembicHairTranslatorModule::StartupModule()
 {
+	LLM_SCOPE_BYTAG(GroomEditor)
 	FGroomEditor::Get().RegisterHairTranslator<FAlembicHairTranslator>();
 }
 
diff --git a/Engine/Plugins/Runtime/HairStrands/Source/HairStrandsCore/Private/GroomAsset.cpp b/Engine/Plugins/Runtime/HairStrands/Source/HairStrandsCore/Private/GroomAsset.cpp
--- a/Engine/Plugins/Runtime/HairStrands/Source/HairStrandsCore/Private/GroomAsset.cpp
+++ b/Engine/Plugins/Runtime/HairStrands/Source/HairStrandsCore/Private/GroomAsset.cpp
@@ -12,6 +12,7 @@
 #include "RenderingThread.h"
 #include "Engine/AssetUserData.h"
 #include "HairStrandsVertexFactory.h"
+#include "HAL/LowLevelMemTracker.h"
 #include "Misc/App.h"
 #include "Misc/Paths.h"
 #include "Serialization/LargeMemoryReader.h"
@@ -32,7 +33,6 @@
 #include "GroomBindingAsset.h"
 #include "GroomDeformerBuilder.h"
 
-
 #if WITH_EDITOR
 #include "Interfaces/ITargetPlatform.h"
 #include "PlatformInfo.h"
@@ -43,6 +43,8 @@
 #include "EditorFramework/AssetImportData.h"
 #endif
 
+LLM_DEFINE_TAG(Groom);
+
 #define LOCTEXT_NAMESPACE "GroomAsset"
 
 static int32 GHairStrandsLoadAsset = 1;
@@ -821,7 +823,7 @@ inline void InternalReleaseResource(T*& Resource)
 #if WITH_EDITOR
 void UGroomAsset::UpdateResource()
 {
-	LLM_SCOPE(ELLMTag::Meshes) // This should be a Groom LLM tag, but there is no LLM tag bit left
+	LLM_SCOPE_BYTAG(Groom);
 
 	uint32 AllChangeType = 0;
 	for (uint32 GroupIndex = 0, GroupCount = GetNumHairGroups(); GroupIndex < GroupCount; ++GroupIndex)
@@ -892,7 +894,7 @@ void UGroomAsset::UpdateResource()
 
 void UGroomAsset::InitResources()
 {
-	LLM_SCOPE(ELLMTag::Meshes) // This should be a Groom LLM tag, but there is no LLM tag bit left
+	LLM_SCOPE_BYTAG(Groom);
 
 	InitGuideResources();
 	InitStrandsResources();
@@ -1035,7 +1037,7 @@ bool ConvertMaterial(TArray<T>& Groups, TArray<FHairGroupsMaterial>& HairGroupsM
 
 void UGroomAsset::PostLoad()
 {
-	LLM_SCOPE(ELLMTag::Meshes) // This should be a Groom LLM tag, but there is no LLM tag bit left
+	LLM_SCOPE_BYTAG(Groom);
 
 	Super::PostLoad();
 
@@ -2519,7 +2521,7 @@ inline FString GetLODName(const UGroomAsset* Asset, uint32 LODIndex)
 
 bool UGroomAsset::BuildCardsGeometry(uint32 GroupIndex)
 {
-	LLM_SCOPE(ELLMTag::Meshes) // This should be a Groom LLM tag, but there is no LLM tag bit left
+	LLM_SCOPE_BYTAG(Groom);
 
 	if (!IsHairStrandsEnabled(EHairStrandsShaderType::Cards) || HairGroupsCards.Num() == 0 || !CanRebuildFromDescription())
 	{
@@ -3104,7 +3106,7 @@ static void InitCardsTextureResources(UGroomAsset* GroomAsset)
 
 void UGroomAsset::InitCardsResources()
 {
-	LLM_SCOPE(ELLMTag::Meshes) // This should be a Groom LLM tag, but there is no LLM tag bit left
+	LLM_SCOPE_BYTAG(Groom);
 
 	if (!IsHairStrandsEnabled(EHairStrandsShaderType::Cards) || HairGroupsCards.Num() == 0)
 	{
@@ -3167,7 +3169,7 @@ void UGroomAsset::InitCardsResources()
 
 void UGroomAsset::InitMeshesResources()
 {
-	LLM_SCOPE(ELLMTag::Meshes) // This should be a Groom LLM tag, but there is no LLM tag bit left
+	LLM_SCOPE_BYTAG(Groom);
 
 	if (!IsHairStrandsEnabled(EHairStrandsShaderType::Meshes))
 	{
@@ -3562,7 +3564,7 @@ void UGroomAsset::SaveProceduralCards(uint32 DescIndex)
 		return;
 	}
 
-	LLM_SCOPE(ELLMTag::Meshes) // This should be a Groom LLM tag, but there is no LLM tag bit left
+	LLM_SCOPE_BYTAG(Groom);
 
 	if (DescIndex >= uint32(HairGroupsCards.Num()))
 		return;
diff --git a/Engine/Plugins/Runtime/HairStrands/Source/HairStrandsCore/Private/GroomBindingAsset.cpp b/Engine/Plugins/Runtime/HairStrands/Source/HairStrandsCore/Private/GroomBindingAsset.cpp
--- a/Engine/Plugins/Runtime/HairStrands/Source/HairStrandsCore/Private/GroomBindingAsset.cpp
+++ b/Engine/Plugins/Runtime/HairStrands/Source/HairStrandsCore/Private/GroomBindingAsset.cpp
@@ -7,6 +7,7 @@
 #include "GeometryCacheMeshData.h"
 #include "GroomAsset.h"
 #include "GroomBindingBuilder.h"
+#include "HAL/LowLevelMemTracker.h"
 #include "UObject/ObjectSaveContext.h"
 
 #if WITH_EDITORONLY_DATA
@@ -16,6 +17,8 @@
 #include "GroomComponent.h"
 #endif
 
+LLM_DECLARE_TAG(Groom);
+
 /////////////////////////////////////////////////////////////////////////////////////////
 
 bool IsHairStrandsDDCLogEnable();
@@ -90,7 +93,7 @@ void UGroomBindingAsset::Serialize(FArchive& Ar)
 
 void UGroomBindingAsset::InitResource()
 {
-	LLM_SCOPE(ELLMTag::Meshes) // This should be a Groom LLM tag, but there is no LLM tag bit left
+	LLM_SCOPE_BYTAG(Groom);
 
 	// Ensure we are releasing binding resources before reallocating them
 	ReleaseResource();
@@ -259,7 +262,7 @@ void UGroomBindingAsset::InvalidateBinding()
 
 void UGroomBindingAsset::PostLoad()
 {
-	LLM_SCOPE(ELLMTag::Meshes) // This should be a Groom LLM tag, but there is no LLM tag bit left
+	LLM_SCOPE_BYTAG(Groom);
 
 	Super::PostLoad();
 
diff --git a/Engine/Plugins/Runtime/HairStrands/Source/HairStrandsCore/Private/GroomComponent.cpp b/Engine/Plugins/Runtime/HairStrands/Source/HairStrandsCore/Private/GroomComponent.cpp
--- a/Engine/Plugins/Runtime/HairStrands/Source/HairStrandsCore/Private/GroomComponent.cpp
+++ b/Engine/Plugins/Runtime/HairStrands/Source/HairStrandsCore/Private/GroomComponent.cpp
@@ -13,6 +13,7 @@
 #include "HairStrandsVertexFactory.h"
 #include "RayTracingInstance.h"
 #include "RayTracingDefinitions.h"
+#include "HAL/LowLevelMemTracker.h"
 #include "HairStrandsInterface.h"
 #include "UObject/UObjectIterator.h"
 #include "GlobalShader.h"
@@ -35,6 +36,8 @@
 #include "Async/ParallelFor.h"
 #include "PrimitiveSceneInfo.h"
 
+LLM_DECLARE_TAG(Groom);
+
 static int32 GHairEnableAdaptiveSubsteps = 0;  
 static FAutoConsoleVariableRef CVarHairEnableAdaptiveSubsteps(TEXT("r.HairStrands.EnableAdaptiveSubsteps"), GHairEnableAdaptiveSubsteps, TEXT("Enable adaptive solver substeps"));
 bool IsHairAdaptiveSubstepsEnabled() { return (GHairEnableAdaptiveSubsteps == 1); }
@@ -2344,7 +2347,7 @@ static EGroomGeometryType GetEffectiveGeometryType(EGroomGeometryType Type, bool
 
 void UGroomComponent::InitResources(bool bIsBindingReloading)
 {
-	LLM_SCOPE(ELLMTag::Meshes) // This should be a Groom LLM tag, but there is no LLM tag bit left
+	LLM_SCOPE_BYTAG(Groom);
 
 	ReleaseResources();
 	bInitSimulation = true;
@@ -3042,7 +3045,7 @@ void UGroomComponent::DeleteDeferredHairGroupInstances()
 
 void UGroomComponent::PostLoad()
 {
-	LLM_SCOPE(ELLMTag::Meshes) // This should be a Groom LLM tag, but there is no LLM tag bit left
+	LLM_SCOPE_BYTAG(Groom);
 
 	Super::PostLoad();
 
diff --git a/Engine/Plugins/Runtime/HairStrands/Source/HairStrandsEditor/Private/HairStrandsEditor.cpp b/Engine/Plugins/Runtime/HairStrands/Source/HairStrandsEditor/Private/HairStrandsEditor.cpp
--- a/Engine/Plugins/Runtime/HairStrands/Source/HairStrandsEditor/Private/HairStrandsEditor.cpp
+++ b/Engine/Plugins/Runtime/HairStrands/Source/HairStrandsEditor/Private/HairStrandsEditor.cpp
@@ -36,6 +36,8 @@ IMPLEMENT_MODULE(FGroomEditor, HairStrandsEditor);
 
 #define LOCTEXT_NAMESPACE "GroomEditor"
 
+LLM_DEFINE_TAG(GroomEditor);
+
 FName FGroomEditor::GroomEditorAppIdentifier(TEXT("GroomEditor"));
 
 
@@ -71,6 +73,8 @@ void SaveAsset(UObject* Object)
 
 void FGroomEditor::StartupModule()
 {
+	LLM_SCOPE_BYTAG(GroomEditor)
+
 	IAssetTools& AssetTools = FModuleManager::LoadModuleChecked<FAssetToolsModule>("AssetTools").Get();
 	TSharedRef<IAssetTypeActions> GroomAssetActions = MakeShared<FGroomActions>();
 	TSharedRef<IAssetTypeActions> BindingAssetActions = MakeShared<FGroomBindingActions>();
diff --git a/Engine/Plugins/Runtime/HairStrands/Source/HairStrandsEditor/Public/HairStrandsEditor.h b/Engine/Plugins/Runtime/HairStrands/Source/HairStrandsEditor/Public/HairStrandsEditor.h
--- a/Engine/Plugins/Runtime/HairStrands/Source/HairStrandsEditor/Public/HairStrandsEditor.h
+++ b/Engine/Plugins/Runtime/HairStrands/Source/HairStrandsEditor/Public/HairStrandsEditor.h
@@ -6,11 +6,14 @@
 #include "Modules/ModuleManager.h"
 #include "Styling/ISlateStyle.h"
 #include "Styling/SlateStyle.h"
+#include "HAL/LowLevelMemTracker.h"
 
 #define HAIRSTRANDSEDITOR_MODULE_NAME TEXT("HairStrandsEditor")
 
 class IGroomTranslator;
 
+LLM_DECLARE_TAG_API(GroomEditor, HAIRSTRANDSEDITOR_API);
+
 /** Implements the HairStrands module  */
 class HAIRSTRANDSEDITOR_API FGroomEditor : public IModuleInterface
 {
@@ -28,6 +31,7 @@ public:
 
 	static inline FGroomEditor& Get()
 	{
+		LLM_SCOPE_BYTAG(GroomEditor);
 		return FModuleManager::LoadModuleChecked<FGroomEditor>(HAIRSTRANDSEDITOR_MODULE_NAME);
 	}
 
