Fix hair instabilities

diff --git a/Engine/Plugins/Runtime/HairStrands/Shaders/Private/NiagaraDataInterfaceHairStrandsTemplate.ush b/Engine/Plugins/Runtime/HairStrands/Shaders/Private/NiagaraDataInterfaceHairStrandsTemplate.ush
--- a/Engine/Plugins/Runtime/HairStrands/Shaders/Private/NiagaraDataInterfaceHairStrandsTemplate.ush
+++ b/Engine/Plugins/Runtime/HairStrands/Shaders/Private/NiagaraDataInterfaceHairStrandsTemplate.ush
@@ -461,53 +461,78 @@ void ComputeEdgeDirection_{ParameterName}(in float3 NodePosition, in float4 Node
 
 	SharedNodePosition[GGroupThreadId.x] = NodePosition;
 	SharedNodeOrientation[GGroupThreadId.x] = NodeOrientation;
-	GroupMemoryBarrier(); 
+	GroupMemoryBarrier();
 
 	OutRestDirection = float3(0,0,0);
-
-	if( LocalIndex > 1 )
+	
+	if(GravityPreloading == 0.0)
+	{
+		if( LocalIndex > 1 )
+		{
+			const float3 EdgeDirection = SharedNodePosition[GGroupThreadId.x] - SharedNodePosition[GGroupThreadId.x-1];
+			OutRestDirection = UnrotateVectorByQuat(EdgeDirection,SharedNodeOrientation[GGroupThreadId.x-2]);
+		}
+		GroupMemoryBarrier();
+	}
+	else
 	{
-		float3 EdgeDirection = SharedNodePosition[GGroupThreadId.x] - SharedNodePosition[GGroupThreadId.x-1];
+		float3 EdgeDirection = (LocalIndex > 1) ? SharedNodePosition[GGroupThreadId.x] - SharedNodePosition[GGroupThreadId.x-1] : float3(0,0,0);
 
-		SharedInverseInertia[GGroupThreadId.x] = BendStiffness*PI*StrandThickness*StrandThickness / (4.0*RestLength);
+		SharedInverseInertia[GGroupThreadId.x] = (isfinite(RestLength) && (abs(RestLength) > 1e-8f)) ? BendStiffness*PI*StrandThickness*StrandThickness / (4.0*RestLength) : 0.0;
 		SharedPreviousPosition[GGroupThreadId.x] = EdgeDirection;
 		SharedPreviousOrientation[GGroupThreadId.x] = float4(EdgeDirection,0);
 
 		GroupMemoryBarrier(); 
 
-		const float GradientNext = ( LocalIndex < ({ParameterName}_StrandSize-1) ) ?  SharedInverseInertia[GGroupThreadId.x+1]  : 0.0;
-		const float GradientPrev = SharedInverseInertia[GGroupThreadId.x];
-		const float SchurComplement =  (GradientNext * GradientNext + GradientPrev * GradientPrev);
-
-		for(int Index = 0; Index < 100; ++Index)
 		{
-			const int IsRed = (GGroupThreadId.x % 2) == 0;
-			if (!IsRed)
-			{	
-				const float3 DeltaConstraint = ((SharedPreviousPosition[GGroupThreadId.x] - SharedPreviousOrientation[GGroupThreadId.x].xyz) * GradientPrev - 
-												(SharedPreviousPosition[GGroupThreadId.x+1] - SharedPreviousOrientation[GGroupThreadId.x+1].xyz) * GradientNext) - NodeMass * GravityVector;
-
-				SharedPreviousOrientation[GGroupThreadId.x].xyz += DeltaConstraint * GradientPrev / SchurComplement;
-				SharedPreviousOrientation[GGroupThreadId.x+1].xyz -= DeltaConstraint * GradientNext / SchurComplement;
-			}
-			GroupMemoryBarrier(); 
-
-			if (IsRed)
+			const bool IsNotLast = LocalIndex < ({ParameterName}_StrandSize-1);
+			const float GradientNext = IsNotLast ?  SharedInverseInertia[GGroupThreadId.x+1]  : 0.0;
+			const float GradientPrev = SharedInverseInertia[GGroupThreadId.x];
+			const float SchurComplement =  (GradientNext * GradientNext + GradientPrev * GradientPrev);
+			const bool bIsNodeValid = ( LocalIndex > 1 ) && (abs(SchurComplement) > 1e-8f);
+			
+			for(int Index = 0; Index < 100; ++Index)
 			{
-				const float3 DeltaConstraint = ((SharedPreviousPosition[GGroupThreadId.x] - SharedPreviousOrientation[GGroupThreadId.x].xyz) * GradientPrev - 
-												(SharedPreviousPosition[GGroupThreadId.x+1] - SharedPreviousOrientation[GGroupThreadId.x+1].xyz) * GradientNext) - NodeMass * GravityVector;
-
-
-				SharedPreviousOrientation[GGroupThreadId.x].xyz += DeltaConstraint * GradientPrev / SchurComplement;
-				SharedPreviousOrientation[GGroupThreadId.x+1].xyz -= DeltaConstraint * GradientNext / SchurComplement;
+				const int IsRed = (GGroupThreadId.x % 2) == 0;
+				float3 DeltaConstraint = float3(0,0,0);
+				if (!IsRed)
+				{
+					DeltaConstraint = IsNotLast ? ((SharedPreviousPosition[GGroupThreadId.x] - SharedPreviousOrientation[GGroupThreadId.x].xyz) * GradientPrev - 
+													(SharedPreviousPosition[GGroupThreadId.x+1] - SharedPreviousOrientation[GGroupThreadId.x+1].xyz) * GradientNext) - NodeMass * GravityVector :
+													((SharedPreviousPosition[GGroupThreadId.x] - SharedPreviousOrientation[GGroupThreadId.x].xyz) * GradientPrev) - NodeMass * GravityVector;
+
+					if( bIsNodeValid) SharedPreviousOrientation[GGroupThreadId.x].xyz += DeltaConstraint * GradientPrev / SchurComplement;
+				}
+				GroupMemoryBarrier(); 
+				if(!IsRed && IsNotLast)
+				{
+					if( bIsNodeValid) SharedPreviousOrientation[GGroupThreadId.x+1].xyz -= DeltaConstraint * GradientNext / SchurComplement;
+				}
+				GroupMemoryBarrier(); 
+
+				if (IsRed)
+				{
+					DeltaConstraint = IsNotLast ? ((SharedPreviousPosition[GGroupThreadId.x] - SharedPreviousOrientation[GGroupThreadId.x].xyz) * GradientPrev - 
+													(SharedPreviousPosition[GGroupThreadId.x+1] - SharedPreviousOrientation[GGroupThreadId.x+1].xyz) * GradientNext) - NodeMass * GravityVector :
+													((SharedPreviousPosition[GGroupThreadId.x] - SharedPreviousOrientation[GGroupThreadId.x].xyz) * GradientPrev) - NodeMass * GravityVector;
+
+
+					if( bIsNodeValid) SharedPreviousOrientation[GGroupThreadId.x].xyz += DeltaConstraint * GradientPrev / SchurComplement;
+				}
+				GroupMemoryBarrier();
+				if(IsRed && IsNotLast)
+				{
+					if( bIsNodeValid) SharedPreviousOrientation[GGroupThreadId.x+1].xyz -= DeltaConstraint * GradientNext / SchurComplement;
+				}
+				GroupMemoryBarrier(); 
 			}
-			GroupMemoryBarrier();
+			EdgeDirection = SharedPreviousOrientation[GGroupThreadId.x].xyz * GravityPreloading + (1.0-GravityPreloading) * EdgeDirection;
+			GroupMemoryBarrier(); 
+			OutRestDirection = UnrotateVectorByQuat(EdgeDirection,SharedNodeOrientation[GGroupThreadId.x-2]);
 		}
-		EdgeDirection = SharedPreviousOrientation[GGroupThreadId.x].xyz * GravityPreloading + (1.0-GravityPreloading) * EdgeDirection;
-		OutRestDirection = UnrotateVectorByQuat(EdgeDirection,SharedNodeOrientation[GGroupThreadId.x-2]);
-	}
 
-	GroupMemoryBarrier(); 
+		GroupMemoryBarrier();
+	}
 }
 
 
diff --git a/Engine/Plugins/Runtime/HairStrands/Source/HairStrandsCore/Private/HairStrandsMeshProjection.cpp b/Engine/Plugins/Runtime/HairStrands/Source/HairStrandsCore/Private/HairStrandsMeshProjection.cpp
--- a/Engine/Plugins/Runtime/HairStrands/Source/HairStrandsCore/Private/HairStrandsMeshProjection.cpp
+++ b/Engine/Plugins/Runtime/HairStrands/Source/HairStrandsCore/Private/HairStrandsMeshProjection.cpp
@@ -1398,7 +1398,7 @@ void AddHairStrandUpdatePositionOffsetPass(
 		RootTrianglePrevPositionBuffer = Register(GraphBuilder, DeformedRootResources->LODs[LODIndex].GetDeformedUniqueTrianglePosition0Buffer(FHairStrandsDeformedRootResource::FLOD::Previous), ERDGImportedBufferFlags::CreateSRV);
 	}
 
-	const bool bUseGPUOffset = DeformedRootResources != nullptr && GHairStrandsUseGPUPositionOffset > 0;
+	const bool bUseGPUOffset = DeformedRootResources != nullptr && DeformedRootResources->IsValid() && DeformedRootResources->IsInitialized() && GHairStrandsUseGPUPositionOffset > 0;
 	const uint32 CurrOffsetIndex = DeformedResources->GetIndex(FHairStrandsDeformedResource::Current);
 	const uint32 PrevOffsetIndex = DeformedResources->GetIndex(FHairStrandsDeformedResource::Previous);
 
diff --git a/Engine/Plugins/Runtime/HairStrands/Source/HairStrandsCore/Private/Niagara/NiagaraDataInterfaceHairStrands.cpp b/Engine/Plugins/Runtime/HairStrands/Source/HairStrandsCore/Private/Niagara/NiagaraDataInterfaceHairStrands.cpp
--- a/Engine/Plugins/Runtime/HairStrands/Source/HairStrandsCore/Private/Niagara/NiagaraDataInterfaceHairStrands.cpp
+++ b/Engine/Plugins/Runtime/HairStrands/Source/HairStrandsCore/Private/Niagara/NiagaraDataInterfaceHairStrands.cpp
@@ -29,15 +29,15 @@ namespace NDIHairStrandsLocal
 BEGIN_SHADER_PARAMETER_STRUCT(FShaderParameters,)
 	SHADER_PARAMETER_RDG_BUFFER_UAV(RWBuffer<uint4>,	DeformedPositionBuffer)
 	SHADER_PARAMETER_RDG_BUFFER_SRV(Buffer<uint4>,		RestPositionBuffer)
-	SHADER_PARAMETER_SRV(Buffer<uint>,					CurvesOffsetsBuffer)
-	SHADER_PARAMETER_SRV(Buffer<float4>,				RestTrianglePositionABuffer)
-	SHADER_PARAMETER_SRV(Buffer<float4>,				RestTrianglePositionBBuffer)
-	SHADER_PARAMETER_SRV(Buffer<float4>,				RestTrianglePositionCBuffer)
-	SHADER_PARAMETER_SRV(Buffer<float4>,				DeformedTrianglePositionABuffer)
-	SHADER_PARAMETER_SRV(Buffer<float4>,				DeformedTrianglePositionBBuffer)
-	SHADER_PARAMETER_SRV(Buffer<float4>,				DeformedTrianglePositionCBuffer)
-	SHADER_PARAMETER_SRV(Buffer<uint>,					RootBarycentricCoordinatesBuffer)
-	SHADER_PARAMETER_SRV(Buffer<uint>,					RootToUniqueTriangleIndexBuffer)
+	SHADER_PARAMETER_RDG_BUFFER_SRV(Buffer<uint>,		CurvesOffsetsBuffer)
+	SHADER_PARAMETER_RDG_BUFFER_SRV(Buffer<float4>,		RestTrianglePositionABuffer)
+	SHADER_PARAMETER_RDG_BUFFER_SRV(Buffer<float4>,		RestTrianglePositionBBuffer)
+	SHADER_PARAMETER_RDG_BUFFER_SRV(Buffer<float4>,		RestTrianglePositionCBuffer)
+	SHADER_PARAMETER_RDG_BUFFER_SRV(Buffer<float4>,		DeformedTrianglePositionABuffer)
+	SHADER_PARAMETER_RDG_BUFFER_SRV(Buffer<float4>,		DeformedTrianglePositionBBuffer)
+	SHADER_PARAMETER_RDG_BUFFER_SRV(Buffer<float4>,		DeformedTrianglePositionCBuffer)
+	SHADER_PARAMETER_RDG_BUFFER_SRV(Buffer<uint>,		RootBarycentricCoordinatesBuffer)
+	SHADER_PARAMETER_RDG_BUFFER_SRV(Buffer<uint>,		RootToUniqueTriangleIndexBuffer)
 	SHADER_PARAMETER(FMatrix44f,						WorldTransform)
 	SHADER_PARAMETER(FMatrix44f,						WorldInverse)
 	SHADER_PARAMETER(FQuat4f,							WorldRotation)
@@ -50,16 +50,16 @@ BEGIN_SHADER_PARAMETER_STRUCT(FShaderParameters,)
 	SHADER_PARAMETER(FVector3f,							RestRootOffset)
 	SHADER_PARAMETER(FVector3f,							DeformedRootOffset)
 	SHADER_PARAMETER(FVector3f,							RestPositionOffset)
-	SHADER_PARAMETER_SRV(Buffer<float4>,				DeformedPositionOffset)
+	SHADER_PARAMETER_RDG_BUFFER_SRV(Buffer<float4>,		DeformedPositionOffset)
 	SHADER_PARAMETER_RDG_BUFFER_UAV(RWBuffer<uint>,		BoundingBoxBuffer)
 	SHADER_PARAMETER(uint32,							ResetSimulation)
 	SHADER_PARAMETER(uint32,							RestUpdate)
 	SHADER_PARAMETER(uint32,							LocalSimulation)
 	SHADER_PARAMETER(int,								SampleCount)
 	SHADER_PARAMETER(FIntVector4,						BoundingBoxOffsets)
-	SHADER_PARAMETER_SRV(Buffer<float4>,				RestSamplePositionsBuffer)
-	SHADER_PARAMETER_SRV(Buffer<float4>,				MeshSampleWeightsBuffer)
-	SHADER_PARAMETER_SRV(Buffer<float>,					ParamsScaleBuffer)
+	SHADER_PARAMETER_RDG_BUFFER_SRV(Buffer<float4>,		RestSamplePositionsBuffer)
+	SHADER_PARAMETER_RDG_BUFFER_SRV(Buffer<float4>,		MeshSampleWeightsBuffer)
+	SHADER_PARAMETER_RDG_BUFFER_SRV(Buffer<float>,		ParamsScaleBuffer)
 	SHADER_PARAMETER(FVector3f,							BoneLinearVelocity)
 	SHADER_PARAMETER(FVector3f,							BoneAngularVelocity)
 	SHADER_PARAMETER(FVector3f,							BoneLinearAcceleration)
@@ -278,21 +278,22 @@ void FNDIHairStrandsBuffer::Update(
 	SourceDeformedRootResources = HairStrandsDeformedRootResource;
 }
 
-void FNDIHairStrandsBuffer::Transfer(const TStaticArray<float, 32 * NumScales>& InParamsScale)
+void FNDIHairStrandsBuffer::Transfer(FRDGBuilder& GraphBuilder, const TStaticArray<float, 32 * NumScales>& InParamsScale)
 {
-	if (SourceRestResources != nullptr && ParamsScaleBuffer.Buffer.IsValid())
+	if (SourceRestResources != nullptr && ParamsScaleBuffer.IsValid())
 	{
 		const uint32 ScaleCount = 32 * NumScales;
 		const uint32 ScaleBytes = sizeof(float) * ScaleCount;
-
-		void* ScaleBufferData = RHILockBuffer(ParamsScaleBuffer.Buffer, 0, ScaleBytes, RLM_WriteOnly);
-		FMemory::Memcpy(ScaleBufferData, InParamsScale.GetData(), ScaleBytes);
-		RHIUnlockBuffer(ParamsScaleBuffer.Buffer);
+		
+		GraphBuilder.QueueBufferUpload(ParamsScaleBuffer.GetOrCreateBuffer(GraphBuilder), InParamsScale.GetData(), ScaleBytes);
+		ParamsScaleBuffer.EndGraphUsage();
 	}
 }
 
 void FNDIHairStrandsBuffer::InitRHI()
 {
+	//ReadbackBuffer = new FRHIGPUBufferReadback(TEXT("Hair.PositionOffsetBuffer"));
+	
 	if (SourceRestResources != nullptr)
 	{
 		FMemMark MemMark(FMemStack::Get());
@@ -303,10 +304,11 @@ void FNDIHairStrandsBuffer::InitRHI()
 			const uint32 OffsetBytes = sizeof(uint32)*OffsetCount;
 
 			const FHairStrandsRootIndexFormat::Type* SourceData = (const FHairStrandsRootIndexFormat::Type*)SourceDatas->CurveOffsets.Lock(LOCK_READ_ONLY);
-			CurvesOffsetsBuffer.Initialize(TEXT("CurvesOffsetsBuffer"), sizeof(uint32), OffsetCount, EPixelFormat::PF_R32_UINT, BUF_Static);
-			void* OffsetBufferData = RHILockBuffer(CurvesOffsetsBuffer.Buffer, 0, OffsetBytes, RLM_WriteOnly);
-			FMemory::Memcpy(OffsetBufferData, SourceData, OffsetBytes);
-			RHIUnlockBuffer(CurvesOffsetsBuffer.Buffer);
+			
+			CurvesOffsetsBuffer.Initialize(GraphBuilder, TEXT("CurvesOffsetsBuffer"), EPixelFormat::PF_R32_UINT, sizeof(uint32), OffsetCount);
+			GraphBuilder.QueueBufferUpload(CurvesOffsetsBuffer.GetOrCreateBuffer(GraphBuilder), SourceData, OffsetBytes);
+			CurvesOffsetsBuffer.EndGraphUsage();
+
 			SourceDatas->CurveOffsets.Unlock();
 		}
 		{
@@ -335,10 +337,9 @@ void FNDIHairStrandsBuffer::InitRHI()
 			const uint32 OffsetBytes = sizeof(uint32)*OffsetCount;
 
 			// Copy the points buffer data to gpu for interpolation
-			PointsCurveBuffer.Initialize(TEXT("PointsCurveBuffer"), sizeof(uint32), OffsetCount, EPixelFormat::PF_R32_UINT, BUF_Static);
-			void* OffsetBufferData = RHILockBuffer(PointsCurveBuffer.Buffer, 0, OffsetBytes, RLM_WriteOnly);
-			FMemory::Memcpy(OffsetBufferData, PointsCurve.GetData(), OffsetBytes);
-			RHIUnlockBuffer(PointsCurveBuffer.Buffer);
+			PointsCurveBuffer.Initialize(GraphBuilder, TEXT("PointsCurveBuffer"), EPixelFormat::PF_R32_UINT, sizeof(uint32), OffsetCount);
+			GraphBuilder.QueueBufferUpload(PointsCurveBuffer.GetOrCreateBuffer(GraphBuilder), PointsCurve.GetData(), OffsetBytes);
+			PointsCurveBuffer.EndGraphUsage();
 		}
 		{
 			static const uint32 ZeroData[] = {
@@ -357,13 +358,11 @@ void FNDIHairStrandsBuffer::InitRHI()
 			const uint32 ScaleCount = 32 * NumScales;
 			const uint32 ScaleBytes = sizeof(float) * ScaleCount;
 
-			ParamsScaleBuffer.Initialize(TEXT("ParamsScaleBuffer"), sizeof(float), ScaleCount, EPixelFormat::PF_R32_FLOAT, BUF_Static);
-			void* ScaleBufferData = RHILockBuffer(ParamsScaleBuffer.Buffer, 0, ScaleBytes, RLM_WriteOnly);
-
-			FMemory::Memcpy(ScaleBufferData, ParamsScale.GetData(), ScaleBytes);
-			RHIUnlockBuffer(ParamsScaleBuffer.Buffer);
+			ParamsScaleBuffer.Initialize(GraphBuilder, TEXT("ParamsScaleBuffer"), EPixelFormat::PF_R32_FLOAT, sizeof(float), ScaleCount);
+			GraphBuilder.QueueBufferUpload(ParamsScaleBuffer.GetOrCreateBuffer(GraphBuilder), ParamsScale.GetData(), ScaleBytes);
+			ParamsScaleBuffer.EndGraphUsage();
 		}
-		if (SourceDeformedResources == nullptr)
+		if (SourceDeformedResources == nullptr || (SourceDeformedResources && !SourceDeformedResources->IsInitialized()))
 		{
 			const uint32 PositionsCount = SourceDatas->GetNumPoints();
 
@@ -377,10 +376,14 @@ void FNDIHairStrandsBuffer::InitRHI()
 
 void FNDIHairStrandsBuffer::ReleaseRHI()
 {
+	//delete ReadbackBuffer;
+	//ReadbackBuffer = nullptr;
+	
 	CurvesOffsetsBuffer.Release();
 	DeformedPositionBuffer.SafeRelease();
 	BoundingBoxBuffer.Release();
 	ParamsScaleBuffer.Release();
+	PointsCurveBuffer.Release();
 }
 
 //------------------------------------------------------------------------------------------------------------
@@ -860,20 +863,20 @@ class FInterpolateGroomGuidesCS : public FGlobalShader
 	using FPermutationDomain = TShaderPermutationDomain<FInterpolationType>;
 
 	BEGIN_SHADER_PARAMETER_STRUCT(FParameters, )
-		SHADER_PARAMETER_SRV(Buffer<float4>,	RestTrianglePositionABuffer)
-		SHADER_PARAMETER_SRV(Buffer<float4>,	RestTrianglePositionBBuffer)
-		SHADER_PARAMETER_SRV(Buffer<float4>,	RestTrianglePositionCBuffer)
-		SHADER_PARAMETER_SRV(Buffer<float4>,	DeformedTrianglePositionABuffer)
-		SHADER_PARAMETER_SRV(Buffer<float4>,	DeformedTrianglePositionBBuffer)
-		SHADER_PARAMETER_SRV(Buffer<float4>,	DeformedTrianglePositionCBuffer)
-		SHADER_PARAMETER_SRV(Buffer<uint>,		RootBarycentricCoordinatesBuffer)
-		SHADER_PARAMETER_SRV(Buffer<uint>,		RootToUniqueTriangleIndexBuffer)
-
-		SHADER_PARAMETER_SRV(Buffer<uint>,		VertexToCurveIndexBuffer)
-		SHADER_PARAMETER_UAV(RWBuffer<uint4>,	DeformedPositionBuffer)
-		SHADER_PARAMETER_SRV(Buffer<uint4>,		RestPositionBuffer)
-		SHADER_PARAMETER_SRV(Buffer<uint>,		CurvesOffsetsBuffer)
-		SHADER_PARAMETER_SRV(Buffer<float4>,	DeformedPositionOffset)
+		SHADER_PARAMETER_RDG_BUFFER_SRV(Buffer<float4>,	RestTrianglePositionABuffer)
+		SHADER_PARAMETER_RDG_BUFFER_SRV(Buffer<float4>,	RestTrianglePositionBBuffer)
+		SHADER_PARAMETER_RDG_BUFFER_SRV(Buffer<float4>,	RestTrianglePositionCBuffer)
+		SHADER_PARAMETER_RDG_BUFFER_SRV(Buffer<float4>,	DeformedTrianglePositionABuffer)
+		SHADER_PARAMETER_RDG_BUFFER_SRV(Buffer<float4>,	DeformedTrianglePositionBBuffer)
+		SHADER_PARAMETER_RDG_BUFFER_SRV(Buffer<float4>,	DeformedTrianglePositionCBuffer)
+		SHADER_PARAMETER_RDG_BUFFER_SRV(Buffer<uint>,	RootBarycentricCoordinatesBuffer)
+		SHADER_PARAMETER_RDG_BUFFER_SRV(Buffer<uint>,	RootToUniqueTriangleIndexBuffer)
+
+		SHADER_PARAMETER_RDG_BUFFER_SRV(Buffer<uint>,		VertexToCurveIndexBuffer)
+		SHADER_PARAMETER_RDG_BUFFER_UAV(RWBuffer<uint4>,	DeformedPositionBuffer)
+		SHADER_PARAMETER_RDG_BUFFER_SRV(Buffer<uint4>,		RestPositionBuffer)
+		SHADER_PARAMETER_RDG_BUFFER_SRV(Buffer<uint>,		CurvesOffsetsBuffer)
+		SHADER_PARAMETER_RDG_BUFFER_SRV(Buffer<float4>,	DeformedPositionOffset)
 		SHADER_PARAMETER(FVector3f,				RestPositionOffset)
 
 		SHADER_PARAMETER(int,					StrandsSize)
@@ -935,26 +938,29 @@ static void InterpolateGroomGuides(FRDGBuilder& GraphBuilder, FNiagaraDataBuffer
 
 		const int32 BufferIndex = HairStrandsBuffer->SourceDeformedResources->GetIndex(FHairStrandsDeformedResource::Current);
 
-		PassParameters->DeformedPositionOffset = HairStrandsBuffer->SourceDeformedResources->DeformedOffsetBuffer[BufferIndex].SRV;
+		PassParameters->DeformedPositionOffset = RegisterAsSRV(GraphBuilder, HairStrandsBuffer->SourceDeformedResources->DeformedOffsetBuffer[BufferIndex]);
 		PassParameters->RestPositionOffset = FVector3f(HairStrandsBuffer->SourceRestResources->GetPositionOffset());
-		PassParameters->DeformedPositionBuffer = HairStrandsBuffer->SourceDeformedResources->DeformedPositionBuffer[BufferIndex].UAV;
-		PassParameters->RestPositionBuffer = HairStrandsBuffer->SourceRestResources->PositionBuffer.SRV;
+		PassParameters->DeformedPositionBuffer = RegisterAsUAV(GraphBuilder, HairStrandsBuffer->SourceDeformedResources->DeformedPositionBuffer[BufferIndex]);
+		PassParameters->RestPositionBuffer = RegisterAsSRV(GraphBuilder, HairStrandsBuffer->SourceRestResources->PositionBuffer);
 		
-		PassParameters->VertexToCurveIndexBuffer = HairStrandsBuffer->PointsCurveBuffer.SRV;
-		PassParameters->CurvesOffsetsBuffer = HairStrandsBuffer->CurvesOffsetsBuffer.SRV;
+		PassParameters->VertexToCurveIndexBuffer = HairStrandsBuffer->PointsCurveBuffer.GetOrCreateSRV(GraphBuilder);
+		PassParameters->CurvesOffsetsBuffer = HairStrandsBuffer->CurvesOffsetsBuffer.GetOrCreateSRV(GraphBuilder);
 		if(bIsRootValid)
 		{
 			const FHairStrandsRestRootResource::FLOD* RestMeshProjection = &(HairStrandsBuffer->SourceRestRootResources->LODs[MeshLODIndex]) ;
 			const FHairStrandsDeformedRootResource::FLOD* DeformedMeshProjection =  &(HairStrandsBuffer->SourceDeformedRootResources->LODs[MeshLODIndex]);
 
-			PassParameters->RestTrianglePositionABuffer = RestMeshProjection->RestUniqueTrianglePosition0Buffer.SRV;
-			PassParameters->RestTrianglePositionBBuffer = RestMeshProjection->RestUniqueTrianglePosition1Buffer.SRV;
-			PassParameters->RestTrianglePositionCBuffer = RestMeshProjection->RestUniqueTrianglePosition2Buffer.SRV;
-			PassParameters->DeformedTrianglePositionABuffer = DeformedMeshProjection->GetDeformedUniqueTrianglePosition0Buffer(FHairStrandsDeformedRootResource::FLOD::Current).SRV;
-			PassParameters->DeformedTrianglePositionBBuffer = DeformedMeshProjection->GetDeformedUniqueTrianglePosition1Buffer(FHairStrandsDeformedRootResource::FLOD::Current).SRV;
-			PassParameters->DeformedTrianglePositionCBuffer = DeformedMeshProjection->GetDeformedUniqueTrianglePosition2Buffer(FHairStrandsDeformedRootResource::FLOD::Current).SRV;
-			PassParameters->RootBarycentricCoordinatesBuffer = RestMeshProjection->RootBarycentricBuffer.SRV;
-			PassParameters->RootToUniqueTriangleIndexBuffer = RestMeshProjection->RootToUniqueTriangleIndexBuffer.SRV;
+			PassParameters->RestTrianglePositionABuffer = RegisterAsSRV(GraphBuilder, RestMeshProjection->RestUniqueTrianglePosition0Buffer);
+			PassParameters->RestTrianglePositionBBuffer = RegisterAsSRV(GraphBuilder, RestMeshProjection->RestUniqueTrianglePosition1Buffer);
+			PassParameters->RestTrianglePositionCBuffer = RegisterAsSRV(GraphBuilder, RestMeshProjection->RestUniqueTrianglePosition2Buffer);
+
+			PassParameters->DeformedTrianglePositionABuffer = RegisterAsSRV(GraphBuilder, DeformedMeshProjection->GetDeformedUniqueTrianglePosition0Buffer(FHairStrandsDeformedRootResource::FLOD::Current));
+			PassParameters->DeformedTrianglePositionBBuffer = RegisterAsSRV(GraphBuilder, DeformedMeshProjection->GetDeformedUniqueTrianglePosition1Buffer(FHairStrandsDeformedRootResource::FLOD::Current));
+			PassParameters->DeformedTrianglePositionCBuffer = RegisterAsSRV(GraphBuilder, DeformedMeshProjection->GetDeformedUniqueTrianglePosition2Buffer(FHairStrandsDeformedRootResource::FLOD::Current));
+			
+			PassParameters->RootBarycentricCoordinatesBuffer = RegisterAsSRV(GraphBuilder, RestMeshProjection->RootBarycentricBuffer);
+			PassParameters->RootToUniqueTriangleIndexBuffer = RegisterAsSRV(GraphBuilder, RestMeshProjection->RootToUniqueTriangleIndexBuffer);
+
 		}
 
 		const uint32 GroupSize = NIAGARA_HAIR_STRANDS_THREAD_COUNT_INTERPOLATE;
@@ -3355,11 +3361,8 @@ void UNiagaraDataInterfaceHairStrands::SetShaderParameters(const FNiagaraDataInt
 	const bool bHasSkinningBinding = bIsHairValid && bIsHairGroupInstValid && ProxyData->HairGroupInstance->BindingType == EHairBindingType::Skinning;
 	const bool bIsRootValid = bIsHairValid && ProxyData->HairStrandsBuffer->SourceDeformedRootResources && ProxyData->HairStrandsBuffer->SourceDeformedRootResources->IsInitialized() && bHasSkinningBinding;
 	const bool bIsRestValid = bIsHairValid && ProxyData->HairStrandsBuffer->SourceRestResources && ProxyData->HairStrandsBuffer->SourceRestResources->IsInitialized()&&
-		// SourceRestResources->PositionBuffer is lazily allocated, when the instance LOD is scheduled (this happens after this called). So this is why this check is here. 
-		// This code should be refactor so that it reflects the lazy allocation scheme
-		ProxyData->HairStrandsBuffer->SourceRestResources->PositionBuffer.SRV &&
 		// TEMP: These check are only temporary for avoiding crashes while we find the bottom of the issue.
-		ProxyData->HairStrandsBuffer->CurvesOffsetsBuffer.SRV && ProxyData->HairStrandsBuffer->ParamsScaleBuffer.SRV && ProxyData->HairStrandsBuffer->BoundingBoxBuffer.IsValid();
+		ProxyData->HairStrandsBuffer->CurvesOffsetsBuffer.IsValid() && ProxyData->HairStrandsBuffer->ParamsScaleBuffer.IsValid() && ProxyData->HairStrandsBuffer->BoundingBoxBuffer.IsValid();
 
 	const bool bIsGeometryValid = bIsHairValid && (!bIsHairGroupInstValid || (bIsHairGroupInstValid && (ProxyData->HairGroupInstance->GeometryType != EHairGeometryType::NoneGeometry)));
 	const bool bIsDeformedValid = bIsHairValid && ProxyData->HairStrandsBuffer->SourceDeformedResources && ProxyData->HairStrandsBuffer->SourceDeformedResources->IsInitialized();
@@ -3372,17 +3375,28 @@ void UNiagaraDataInterfaceHairStrands::SetShaderParameters(const FNiagaraDataInt
 		FNDIHairStrandsBuffer* HairStrandsBuffer = ProxyData->HairStrandsBuffer;
 
 		FRDGBufferUAVRef DeformedPositionBufferUAV = nullptr;
-		FRHIShaderResourceView* DeformedPositionOffsetSRV = nullptr;
+		FRDGBufferSRVRef DeformedPositionOffsetSRV = nullptr;
 		if (bIsDeformedValid)
 		{
 			FRDGPooledBuffer* DeformedPositionBuffer = HairStrandsBuffer->SourceDeformedResources->DeformedPositionBuffer[HairStrandsBuffer->SourceDeformedResources->CurrentIndex].Buffer;
 			DeformedPositionBufferUAV = GraphBuilder.CreateUAV(GraphBuilder.RegisterExternalBuffer(DeformedPositionBuffer), FHairStrandsPositionFormat::Format);	//-OPT: If constant across all dispatches can create in PreStage once
-			DeformedPositionOffsetSRV = HairStrandsBuffer->SourceDeformedResources->DeformedOffsetBuffer[HairStrandsBuffer->SourceDeformedResources->CurrentIndex].SRV.GetReference();
+			DeformedPositionOffsetSRV = RegisterAsSRV(GraphBuilder, HairStrandsBuffer->SourceDeformedResources->DeformedOffsetBuffer[HairStrandsBuffer->SourceDeformedResources->CurrentIndex]);
+
+			// Keep the readback commented out for future debug
+			// FRDGBufferRef Buffer = GraphBuilder.RegisterExternalBuffer(HairStrandsBuffer->SourceDeformedResources->DeformedOffsetBuffer[HairStrandsBuffer->SourceDeformedResources->CurrentIndex].Buffer);
+			// AddEnqueueCopyPass(GraphBuilder, HairStrandsBuffer->ReadbackBuffer, Buffer, 16u);
+			//
+			// if (HairStrandsBuffer->ReadbackBuffer->IsReady())
+			// {
+			// 	const FVector4f ReadData = *(FVector4f*)(HairStrandsBuffer->ReadbackBuffer->Lock(sizeof(uint32)));
+			// 	HairStrandsBuffer->ReadbackBuffer->Unlock();
+			// 	UE_LOG(LogHairStrands, Log, TEXT("Deformed Offset = %f %f %f"), ReadData.X, ReadData.Y, ReadData.Z);
+			// }
 		}
 		else
 		{
 			DeformedPositionBufferUAV = GraphBuilder.CreateUAV(GraphBuilder.RegisterExternalBuffer(HairStrandsBuffer->DeformedPositionBuffer), FHairStrandsPositionFormat::Format);
-			DeformedPositionOffsetSRV = FNiagaraRenderer::GetDummyFloatBuffer();
+			DeformedPositionOffsetSRV = Context.GetComputeDispatchInterface().GetEmptyBufferSRV(GraphBuilder, FHairStrandsPositionOffsetFormat::Format);
 		}
 		const int32 MeshLODIndex = bIsRootValid ? HairStrandsBuffer->SourceDeformedRootResources->MeshLODIndex : -1;
 
@@ -3394,23 +3408,24 @@ void UNiagaraDataInterfaceHairStrands::SetShaderParameters(const FNiagaraDataInt
 		const FHairStrandsRestRootResource::FLOD* RestMeshProjection = bHasSkinnedInterpolation ? &(HairStrandsBuffer->SourceRestRootResources->LODs[MeshLODIndex]) : nullptr;
 		const FHairStrandsDeformedRootResource::FLOD* DeformedMeshProjection = bHasSkinnedInterpolation ? &(HairStrandsBuffer->SourceDeformedRootResources->LODs[MeshLODIndex]) : nullptr;
 
-		FRHIShaderResourceView* RestTrianglePositionASRV = (bHasSkinnedInterpolation && RestMeshProjection) ? RestMeshProjection->RestUniqueTrianglePosition0Buffer.SRV.GetReference() : FNiagaraRenderer::GetDummyFloatBuffer();
-		FRHIShaderResourceView* RestTrianglePositionBSRV = (bHasSkinnedInterpolation && RestMeshProjection) ? RestMeshProjection->RestUniqueTrianglePosition1Buffer.SRV.GetReference() : FNiagaraRenderer::GetDummyFloatBuffer();
-		FRHIShaderResourceView* RestTrianglePositionCSRV = (bHasSkinnedInterpolation && RestMeshProjection) ? RestMeshProjection->RestUniqueTrianglePosition2Buffer.SRV.GetReference() : FNiagaraRenderer::GetDummyFloatBuffer();
+		FRDGBufferSRVRef RestTrianglePositionASRV = (bHasSkinnedInterpolation && RestMeshProjection) ? RegisterAsSRV(GraphBuilder, RestMeshProjection->RestUniqueTrianglePosition0Buffer) : Context.GetComputeDispatchInterface().GetEmptyBufferSRV(GraphBuilder, FHairStrandsMeshTrianglePositionFormat::Format);
+		FRDGBufferSRVRef RestTrianglePositionBSRV = (bHasSkinnedInterpolation && RestMeshProjection) ? RegisterAsSRV(GraphBuilder, RestMeshProjection->RestUniqueTrianglePosition1Buffer) : Context.GetComputeDispatchInterface().GetEmptyBufferSRV(GraphBuilder, FHairStrandsMeshTrianglePositionFormat::Format);
+		FRDGBufferSRVRef RestTrianglePositionCSRV = (bHasSkinnedInterpolation && RestMeshProjection) ? RegisterAsSRV(GraphBuilder, RestMeshProjection->RestUniqueTrianglePosition2Buffer) : Context.GetComputeDispatchInterface().GetEmptyBufferSRV(GraphBuilder, FHairStrandsMeshTrianglePositionFormat::Format);
 
-		FRHIShaderResourceView* DeformedTrianglePositionASRV = (bHasSkinnedInterpolation && DeformedMeshProjection) ? DeformedMeshProjection->GetDeformedUniqueTrianglePosition0Buffer(FHairStrandsDeformedRootResource::FLOD::Current).SRV.GetReference() : FNiagaraRenderer::GetDummyFloatBuffer();
-		FRHIShaderResourceView* DeformedTrianglePositionBSRV = (bHasSkinnedInterpolation && DeformedMeshProjection) ? DeformedMeshProjection->GetDeformedUniqueTrianglePosition1Buffer(FHairStrandsDeformedRootResource::FLOD::Current).SRV.GetReference() : FNiagaraRenderer::GetDummyFloatBuffer();
-		FRHIShaderResourceView* DeformedTrianglePositionCSRV = (bHasSkinnedInterpolation && DeformedMeshProjection) ? DeformedMeshProjection->GetDeformedUniqueTrianglePosition2Buffer(FHairStrandsDeformedRootResource::FLOD::Current).SRV.GetReference() : FNiagaraRenderer::GetDummyFloatBuffer();
-		FRHIShaderResourceView* RootBarycentricCoordinatesSRV = (bHasSkinnedInterpolation && RestMeshProjection) ? RestMeshProjection->RootBarycentricBuffer.SRV.GetReference() : FNiagaraRenderer::GetDummyUIntBuffer();
-		FRHIShaderResourceView* RootToUniqueTriangleIndexSRV = (bHasSkinnedInterpolation && RestMeshProjection) ? RestMeshProjection->RootToUniqueTriangleIndexBuffer.SRV.GetReference() : FNiagaraRenderer::GetDummyUIntBuffer();
+		FRDGBufferSRVRef DeformedTrianglePositionASRV = (bHasSkinnedInterpolation && DeformedMeshProjection) ?  RegisterAsSRV(GraphBuilder, DeformedMeshProjection->GetDeformedUniqueTrianglePosition0Buffer(FHairStrandsDeformedRootResource::FLOD::Current)) : Context.GetComputeDispatchInterface().GetEmptyBufferSRV(GraphBuilder, FHairStrandsMeshTrianglePositionFormat::Format);
+		FRDGBufferSRVRef DeformedTrianglePositionBSRV = (bHasSkinnedInterpolation && DeformedMeshProjection) ?  RegisterAsSRV(GraphBuilder, DeformedMeshProjection->GetDeformedUniqueTrianglePosition1Buffer(FHairStrandsDeformedRootResource::FLOD::Current)) : Context.GetComputeDispatchInterface().GetEmptyBufferSRV(GraphBuilder, FHairStrandsMeshTrianglePositionFormat::Format);
+		FRDGBufferSRVRef DeformedTrianglePositionCSRV = (bHasSkinnedInterpolation && DeformedMeshProjection) ?  RegisterAsSRV(GraphBuilder, DeformedMeshProjection->GetDeformedUniqueTrianglePosition2Buffer(FHairStrandsDeformedRootResource::FLOD::Current)) : Context.GetComputeDispatchInterface().GetEmptyBufferSRV(GraphBuilder, FHairStrandsMeshTrianglePositionFormat::Format);
+		
+		FRDGBufferSRVRef RootBarycentricCoordinatesSRV = (bHasSkinnedInterpolation && RestMeshProjection) ? RegisterAsSRV(GraphBuilder, RestMeshProjection->RootBarycentricBuffer) : Context.GetComputeDispatchInterface().GetEmptyBufferSRV(GraphBuilder, FHairStrandsRootBarycentricFormat::Format);
+		FRDGBufferSRVRef RootToUniqueTriangleIndexSRV = (bHasSkinnedInterpolation && RestMeshProjection) ? RegisterAsSRV(GraphBuilder, RestMeshProjection->RootToUniqueTriangleIndexBuffer) : Context.GetComputeDispatchInterface().GetEmptyBufferSRV(GraphBuilder, FHairStrandsRootToUniqueTriangleIndexFormat::Format);
 
 		// RBF buffers
 		const bool bHasSamples = (RestMeshProjection && RestMeshProjection->SampleCount > 0);
 		const int32 SampleCountValue = bHasSamples ? RestMeshProjection->SampleCount : 0;
 
-		FShaderResourceViewRHIRef RestSamplePositionsBufferSRV = (bHasSamples && RestMeshProjection) ? RestMeshProjection->RestSamplePositionsBuffer.SRV.GetReference() : FNiagaraRenderer::GetDummyFloatBuffer();
-		FShaderResourceViewRHIRef MeshSampleWeightsBufferSRV = (bHasSamples && DeformedMeshProjection) ? DeformedMeshProjection->GetMeshSampleWeightsBuffer((FHairStrandsDeformedRootResource::FLOD::Current)).SRV.GetReference() : FNiagaraRenderer::GetDummyFloatBuffer();
-
+		FRDGBufferSRVRef RestSamplePositionsBufferSRV = (bHasSamples && RestMeshProjection) ? RegisterAsSRV(GraphBuilder,RestMeshProjection->RestSamplePositionsBuffer) : Context.GetComputeDispatchInterface().GetEmptyBufferSRV(GraphBuilder, FHairStrandsMeshTrianglePositionFormat::Format);
+		FRDGBufferSRVRef MeshSampleWeightsBufferSRV = (bHasSamples && DeformedMeshProjection) ? RegisterAsSRV(GraphBuilder,DeformedMeshProjection->GetMeshSampleWeightsBuffer((FHairStrandsDeformedRootResource::FLOD::Current))) : Context.GetComputeDispatchInterface().GetEmptyBufferSRV(GraphBuilder, FHairStrandsMeshTrianglePositionFormat::Format);
+		
 		// Simulation setup (we update the rest configuration based on the deformed positions 
 		// if in restupdate mode or if we are resetting the sim and using RBF transfer since the rest positions are not matrching the physics asset)
 		const int32 NeedResetValue = (ProxyData->TickCount <= GHairSimulationMaxDelay) || !HairStrandsBuffer->bValidGeometryType;
@@ -3435,12 +3450,27 @@ void UNiagaraDataInterfaceHairStrands::SetShaderParameters(const FNiagaraDataInt
 			WorldTransformFloat = FMatrix44f(WorldTransformDouble * BoneTransformDouble.Inverse());
 		}
 
+		if (ProxyData->BoneLinearVelocity.ContainsNaN() || ProxyData->BoneAngularVelocity.ContainsNaN() || ProxyData->BoneLinearAcceleration.ContainsNaN() || ProxyData->BoneAngularAcceleration.ContainsNaN())
+		{
+			UE_LOG(LogHairStrands, Log, TEXT("Bad bones state"));
+		}
+
 		if (!bIsRootValid && bHasSkinningBinding)
 		{
 			UE_LOG(LogHairStrands, Log, TEXT("FNDIHairStrandsParametersCS() Groom Asset %s from component %s is set to use skinning interpolation but the skin resources are not valid"),
 				*ProxyData->HairGroupInstance->Debug.GroomAssetName, *ProxyData->HairGroupInstance->Debug.MeshComponentName);
 		}
 
+		if (WorldTransformFloat.ContainsNaN())
+		{
+			WorldTransformFloat = FMatrix44f(ProxyData->WorldTransform.ToMatrixWithScale());
+		}
+
+		if (bHasSkinnedInterpolation && (!DeformedMeshProjection || !RestMeshProjection))
+		{
+			UE_LOG(LogHairStrands, Log, TEXT("Bad valid projection"));
+		}
+
 		// Set shader constants
 		ShaderParameters->BoundingBoxOffsets = HairStrandsBuffer->BoundingBoxOffsets;
 		ShaderParameters->WorldTransform = WorldTransformFloat;
@@ -3469,9 +3499,9 @@ void UNiagaraDataInterfaceHairStrands::SetShaderParameters(const FNiagaraDataInt
 		ShaderParameters->BoundingBoxBuffer = HairStrandsBuffer->BoundingBoxBuffer.GetOrCreateUAV(GraphBuilder);
 
 		// Set Shader SRV
-		ShaderParameters->CurvesOffsetsBuffer = HairStrandsBuffer->CurvesOffsetsBuffer.SRV;
-		ShaderParameters->ParamsScaleBuffer = HairStrandsBuffer->ParamsScaleBuffer.SRV;
-		ShaderParameters->RestPositionBuffer = GraphBuilder.CreateSRV(GraphBuilder.RegisterExternalBuffer(HairStrandsBuffer->SourceRestResources->PositionBuffer.Buffer), FHairStrandsPositionFormat::Format); //-OPT: If constant across all dispatches can create in PreStage once
+		ShaderParameters->CurvesOffsetsBuffer = HairStrandsBuffer->CurvesOffsetsBuffer.GetOrCreateSRV(GraphBuilder);
+		ShaderParameters->ParamsScaleBuffer = HairStrandsBuffer->ParamsScaleBuffer.GetOrCreateSRV(GraphBuilder);
+		ShaderParameters->RestPositionBuffer = RegisterAsSRV(GraphBuilder,HairStrandsBuffer->SourceRestResources->PositionBuffer); //-OPT: If constant across all dispatches can create in PreStage once
 		ShaderParameters->DeformedPositionOffset = DeformedPositionOffsetSRV;
 		ShaderParameters->RestTrianglePositionABuffer = RestTrianglePositionASRV;
 		ShaderParameters->RestTrianglePositionBBuffer = RestTrianglePositionBSRV;
@@ -3518,20 +3548,20 @@ void UNiagaraDataInterfaceHairStrands::SetShaderParameters(const FNiagaraDataInt
 		ShaderParameters->BoundingBoxBuffer = Context.GetComputeDispatchInterface().GetEmptyBufferUAV(GraphBuilder, PF_R32_UINT);
 
 		// Set Shader SRV
-		ShaderParameters->CurvesOffsetsBuffer = FNiagaraRenderer::GetDummyUIntBuffer();
+		ShaderParameters->CurvesOffsetsBuffer = Context.GetComputeDispatchInterface().GetEmptyBufferSRV(GraphBuilder, PF_R32_UINT);;
 		ShaderParameters->RestPositionBuffer = Context.GetComputeDispatchInterface().GetEmptyBufferSRV(GraphBuilder, FHairStrandsPositionFormat::Format);
-		ShaderParameters->DeformedPositionOffset = FNiagaraRenderer::GetDummyFloatBuffer();
-		ShaderParameters->RestTrianglePositionABuffer = FNiagaraRenderer::GetDummyFloatBuffer();
-		ShaderParameters->RestTrianglePositionBBuffer = FNiagaraRenderer::GetDummyFloatBuffer();
-		ShaderParameters->RestTrianglePositionCBuffer = FNiagaraRenderer::GetDummyFloatBuffer();
-		ShaderParameters->DeformedTrianglePositionABuffer = FNiagaraRenderer::GetDummyFloatBuffer();
-		ShaderParameters->DeformedTrianglePositionBBuffer = FNiagaraRenderer::GetDummyFloatBuffer();
-		ShaderParameters->DeformedTrianglePositionCBuffer = FNiagaraRenderer::GetDummyFloatBuffer();
-		ShaderParameters->RestSamplePositionsBuffer = FNiagaraRenderer::GetDummyFloatBuffer();
-		ShaderParameters->MeshSampleWeightsBuffer = FNiagaraRenderer::GetDummyFloatBuffer();
-		ShaderParameters->RootBarycentricCoordinatesBuffer = FNiagaraRenderer::GetDummyUIntBuffer();
-		ShaderParameters->RootToUniqueTriangleIndexBuffer = FNiagaraRenderer::GetDummyUIntBuffer();
-		ShaderParameters->ParamsScaleBuffer = FNiagaraRenderer::GetDummyFloatBuffer();
+		ShaderParameters->DeformedPositionOffset = Context.GetComputeDispatchInterface().GetEmptyBufferSRV(GraphBuilder, FHairStrandsPositionOffsetFormat::Format);
+		ShaderParameters->RestTrianglePositionABuffer = Context.GetComputeDispatchInterface().GetEmptyBufferSRV(GraphBuilder, FHairStrandsMeshTrianglePositionFormat::Format);
+		ShaderParameters->RestTrianglePositionBBuffer = Context.GetComputeDispatchInterface().GetEmptyBufferSRV(GraphBuilder, FHairStrandsMeshTrianglePositionFormat::Format);
+		ShaderParameters->RestTrianglePositionCBuffer = Context.GetComputeDispatchInterface().GetEmptyBufferSRV(GraphBuilder, FHairStrandsMeshTrianglePositionFormat::Format);
+		ShaderParameters->DeformedTrianglePositionABuffer = Context.GetComputeDispatchInterface().GetEmptyBufferSRV(GraphBuilder, FHairStrandsMeshTrianglePositionFormat::Format);
+		ShaderParameters->DeformedTrianglePositionBBuffer = Context.GetComputeDispatchInterface().GetEmptyBufferSRV(GraphBuilder, FHairStrandsMeshTrianglePositionFormat::Format);
+		ShaderParameters->DeformedTrianglePositionCBuffer = Context.GetComputeDispatchInterface().GetEmptyBufferSRV(GraphBuilder, FHairStrandsMeshTrianglePositionFormat::Format);
+		ShaderParameters->RestSamplePositionsBuffer = Context.GetComputeDispatchInterface().GetEmptyBufferSRV(GraphBuilder, FHairStrandsMeshTrianglePositionFormat::Format);
+		ShaderParameters->MeshSampleWeightsBuffer = Context.GetComputeDispatchInterface().GetEmptyBufferSRV(GraphBuilder, FHairStrandsMeshTrianglePositionFormat::Format);
+		ShaderParameters->RootBarycentricCoordinatesBuffer = Context.GetComputeDispatchInterface().GetEmptyBufferSRV(GraphBuilder, FHairStrandsRootBarycentricFormat::Format);
+		ShaderParameters->RootToUniqueTriangleIndexBuffer = Context.GetComputeDispatchInterface().GetEmptyBufferSRV(GraphBuilder, FHairStrandsRootToUniqueTriangleIndexFormat::Format);
+		ShaderParameters->ParamsScaleBuffer = Context.GetComputeDispatchInterface().GetEmptyBufferSRV(GraphBuilder, PF_R32_FLOAT);
 	}
 }
 
@@ -3564,7 +3594,7 @@ void FNDIHairStrandsProxy::PreStage(const FNDIGpuComputePreStageContext& Context
 			BoundingBoxOffsets[2] = BoundingBoxOffsets[3];
 			BoundingBoxOffsets[3] = FirstOffset;
 
-			ProxyData->HairStrandsBuffer->Transfer(ProxyData->ParamsScale);
+			ProxyData->HairStrandsBuffer->Transfer(Context.GetGraphBuilder(), ProxyData->ParamsScale);
 		}
 	}
 }
@@ -3587,6 +3617,9 @@ void FNDIHairStrandsProxy::PostSimulate(const FNDIGpuComputePostSimulateContext&
 	if (bIsHairValid && Context.IsFinalPostSimulate())
 	{
 		ProxyData->HairStrandsBuffer->BoundingBoxBuffer.EndGraphUsage();
+		ProxyData->HairStrandsBuffer->CurvesOffsetsBuffer.EndGraphUsage();
+		ProxyData->HairStrandsBuffer->ParamsScaleBuffer.EndGraphUsage();
+		ProxyData->HairStrandsBuffer->PointsCurveBuffer.EndGraphUsage();
 	}
 }
 
diff --git a/Engine/Plugins/Runtime/HairStrands/Source/HairStrandsCore/Public/Niagara/NiagaraDataInterfaceHairStrands.h b/Engine/Plugins/Runtime/HairStrands/Source/HairStrandsCore/Public/Niagara/NiagaraDataInterfaceHairStrands.h
--- a/Engine/Plugins/Runtime/HairStrands/Source/HairStrandsCore/Public/Niagara/NiagaraDataInterfaceHairStrands.h
+++ b/Engine/Plugins/Runtime/HairStrands/Source/HairStrandsCore/Public/Niagara/NiagaraDataInterfaceHairStrands.h
@@ -39,7 +39,7 @@ struct FNDIHairStrandsBuffer : public FRenderResource
 		const FHairStrandsDeformedRootResource* HairStrandsDeformedRootResource);
 
 	/** Transfer CPU datas to GPU */
-	void Transfer(const TStaticArray<float, 32 * NumScales>& InParamsScale);
+	void Transfer(FRDGBuilder& GraphBuilder, const TStaticArray<float, 32 * NumScales>& InParamsScale);
 
 	/** Init the buffer */
 	virtual void InitRHI() override;
@@ -51,7 +51,7 @@ struct FNDIHairStrandsBuffer : public FRenderResource
 	virtual FString GetFriendlyName() const override { return TEXT("FNDIHairStrandsBuffer"); }
 
 	/** Strand curves point offset buffer */
-	FReadBuffer CurvesOffsetsBuffer;
+	FNiagaraPooledRWBuffer CurvesOffsetsBuffer;
 
 	/** Deformed position buffer in case no resource are there */
 	TRefCountPtr<FRDGPooledBuffer> DeformedPositionBuffer;
@@ -60,10 +60,10 @@ struct FNDIHairStrandsBuffer : public FRenderResource
 	FNiagaraPooledRWBuffer BoundingBoxBuffer;
 
 	/** Params scale buffer */
-	FReadBuffer ParamsScaleBuffer;
+	FNiagaraPooledRWBuffer ParamsScaleBuffer;
 
 	/** Points curve index for fast query */
-	FReadBuffer PointsCurveBuffer;
+	FNiagaraPooledRWBuffer PointsCurveBuffer;
 
 	/** The strand asset resource from which to sample */
 	const FHairStrandsRestResource* SourceRestResources;
@@ -85,6 +85,9 @@ struct FNDIHairStrandsBuffer : public FRenderResource
 
 	/** Valid geometry type for hair (strands, cards, mesh)*/
 	bool bValidGeometryType = false;
+
+	// For debug only
+	//FRHIGPUBufferReadback* ReadbackBuffer = nullptr;
 };
 
 /** Data stored per strand base instance*/
diff --git a/Engine/Plugins/Runtime/HairStrands/Source/HairStrandsEditor/Private/GroomAssetDetails.cpp b/Engine/Plugins/Runtime/HairStrands/Source/HairStrandsEditor/Private/GroomAssetDetails.cpp
--- a/Engine/Plugins/Runtime/HairStrands/Source/HairStrandsEditor/Private/GroomAssetDetails.cpp
+++ b/Engine/Plugins/Runtime/HairStrands/Source/HairStrandsEditor/Private/GroomAssetDetails.cpp
@@ -323,19 +323,22 @@ void FGroomRenderingDetails::CustomizeDetails(IDetailLayoutBuilder& DetailLayout
 
 	FName CategoryName = GetCategoryName(PanelType);
 	GroomDetailLayout = &DetailLayout;
-	if (UGroomAsset* LocalGroomAsset = Cast<UGroomAsset>(SelectedObjects[0].Get()))
+	if(SelectedObjects.Num() > 0)
 	{
-		GroomAsset = LocalGroomAsset;
+		if (UGroomAsset* LocalGroomAsset = Cast<UGroomAsset>(SelectedObjects[0].Get()))
+		{
+			GroomAsset = LocalGroomAsset;
 
-		IDetailCategoryBuilder& HairGroupCategory = DetailLayout.EditCategory(CategoryName, FText::GetEmpty(), ECategoryPriority::TypeSpecific);
-		CustomizeStrandsGroupProperties(DetailLayout, HairGroupCategory);
-	}
-	else if (UGroomBindingAssetList* LocalGroomBindingList = Cast<UGroomBindingAssetList>(SelectedObjects[0].Get()))
-	{
-		GroomBindingAssetList = LocalGroomBindingList;
+			IDetailCategoryBuilder& HairGroupCategory = DetailLayout.EditCategory(CategoryName, FText::GetEmpty(), ECategoryPriority::TypeSpecific);
+			CustomizeStrandsGroupProperties(DetailLayout, HairGroupCategory);
+		}
+		else if (UGroomBindingAssetList* LocalGroomBindingList = Cast<UGroomBindingAssetList>(SelectedObjects[0].Get()))
+		{
+			GroomBindingAssetList = LocalGroomBindingList;
 
-		IDetailCategoryBuilder& HairGroupCategory = DetailLayout.EditCategory(CategoryName, FText::GetEmpty(), ECategoryPriority::TypeSpecific);
-		CustomizeStrandsGroupProperties(DetailLayout, HairGroupCategory);
+			IDetailCategoryBuilder& HairGroupCategory = DetailLayout.EditCategory(CategoryName, FText::GetEmpty(), ECategoryPriority::TypeSpecific);
+			CustomizeStrandsGroupProperties(DetailLayout, HairGroupCategory);
+		}
 	}
 }
 
@@ -392,6 +395,7 @@ void FGroomRenderingDetails::CustomizeStrandsGroupProperties(IDetailLayoutBuilde
 		DetailLayout.HideProperty(DetailLayout.GetProperty(GET_MEMBER_NAME_CHECKED(UGroomAsset, HairInterpolationType), UGroomAsset::StaticClass()));
 		DetailLayout.HideProperty(DetailLayout.GetProperty(GET_MEMBER_NAME_CHECKED(UGroomAsset, MinLOD), UGroomAsset::StaticClass()));
 		DetailLayout.HideProperty(DetailLayout.GetProperty(GET_MEMBER_NAME_CHECKED(UGroomAsset, DisableBelowMinLodStripping), UGroomAsset::StaticClass()));
+		DetailLayout.HideProperty(DetailLayout.GetProperty(GET_MEMBER_NAME_CHECKED(UGroomAsset, DeformedSkeletalMesh), UGroomAsset::StaticClass()));
 	}
 	break;
 	case EMaterialPanelType::Meshes:
@@ -408,6 +412,7 @@ void FGroomRenderingDetails::CustomizeStrandsGroupProperties(IDetailLayoutBuilde
 		DetailLayout.HideProperty(DetailLayout.GetProperty(GET_MEMBER_NAME_CHECKED(UGroomAsset, HairInterpolationType), UGroomAsset::StaticClass()));
 		DetailLayout.HideProperty(DetailLayout.GetProperty(GET_MEMBER_NAME_CHECKED(UGroomAsset, MinLOD), UGroomAsset::StaticClass()));
 		DetailLayout.HideProperty(DetailLayout.GetProperty(GET_MEMBER_NAME_CHECKED(UGroomAsset, DisableBelowMinLodStripping), UGroomAsset::StaticClass()));
+		DetailLayout.HideProperty(DetailLayout.GetProperty(GET_MEMBER_NAME_CHECKED(UGroomAsset, DeformedSkeletalMesh), UGroomAsset::StaticClass()));
 	}
 	break;
 	case EMaterialPanelType::Strands:
@@ -424,6 +429,7 @@ void FGroomRenderingDetails::CustomizeStrandsGroupProperties(IDetailLayoutBuilde
 		DetailLayout.HideProperty(DetailLayout.GetProperty(GET_MEMBER_NAME_CHECKED(UGroomAsset, HairInterpolationType), UGroomAsset::StaticClass()));
 		DetailLayout.HideProperty(DetailLayout.GetProperty(GET_MEMBER_NAME_CHECKED(UGroomAsset, MinLOD), UGroomAsset::StaticClass()));
 		DetailLayout.HideProperty(DetailLayout.GetProperty(GET_MEMBER_NAME_CHECKED(UGroomAsset, DisableBelowMinLodStripping), UGroomAsset::StaticClass()));
+		DetailLayout.HideProperty(DetailLayout.GetProperty(GET_MEMBER_NAME_CHECKED(UGroomAsset, DeformedSkeletalMesh), UGroomAsset::StaticClass()));
 	}
 	break;
 	case EMaterialPanelType::Physics:
@@ -440,6 +446,7 @@ void FGroomRenderingDetails::CustomizeStrandsGroupProperties(IDetailLayoutBuilde
 		DetailLayout.HideProperty(DetailLayout.GetProperty(GET_MEMBER_NAME_CHECKED(UGroomAsset, HairInterpolationType), UGroomAsset::StaticClass()));
 		DetailLayout.HideProperty(DetailLayout.GetProperty(GET_MEMBER_NAME_CHECKED(UGroomAsset, MinLOD), UGroomAsset::StaticClass()));
 		DetailLayout.HideProperty(DetailLayout.GetProperty(GET_MEMBER_NAME_CHECKED(UGroomAsset, DisableBelowMinLodStripping), UGroomAsset::StaticClass()));
+		DetailLayout.HideProperty(DetailLayout.GetProperty(GET_MEMBER_NAME_CHECKED(UGroomAsset, DeformedSkeletalMesh), UGroomAsset::StaticClass()));
 	}
 	break;
 	case EMaterialPanelType::Interpolation:
@@ -456,6 +463,7 @@ void FGroomRenderingDetails::CustomizeStrandsGroupProperties(IDetailLayoutBuilde
 //		DetailLayout.HideProperty(DetailLayout.GetProperty(GET_MEMBER_NAME_CHECKED(UGroomAsset, HairInterpolationType), UGroomAsset::StaticClass()));
 		DetailLayout.HideProperty(DetailLayout.GetProperty(GET_MEMBER_NAME_CHECKED(UGroomAsset, MinLOD), UGroomAsset::StaticClass()));
 		DetailLayout.HideProperty(DetailLayout.GetProperty(GET_MEMBER_NAME_CHECKED(UGroomAsset, DisableBelowMinLodStripping), UGroomAsset::StaticClass()));
+		//DetailLayout.HideProperty(DetailLayout.GetProperty(GET_MEMBER_NAME_CHECKED(UGroomAsset, DeformedSkeletalMesh), UGroomAsset::StaticClass()));
 	}
 	break;
 	case EMaterialPanelType::LODs:
@@ -472,6 +480,7 @@ void FGroomRenderingDetails::CustomizeStrandsGroupProperties(IDetailLayoutBuilde
 		DetailLayout.HideProperty(DetailLayout.GetProperty(GET_MEMBER_NAME_CHECKED(UGroomAsset, HairInterpolationType), UGroomAsset::StaticClass()));
 //		DetailLayout.HideProperty(DetailLayout.GetProperty(GET_MEMBER_NAME_CHECKED(UGroomAsset, MinLOD), UGroomAsset::StaticClass()));
 //		DetailLayout.HideProperty(DetailLayout.GetProperty(GET_MEMBER_NAME_CHECKED(UGroomAsset, DisableBelowMinLodStripping), UGroomAsset::StaticClass()));
+		DetailLayout.HideProperty(DetailLayout.GetProperty(GET_MEMBER_NAME_CHECKED(UGroomAsset, DeformedSkeletalMesh), UGroomAsset::StaticClass()));
 	}
 	break;
 	}
diff --git a/Engine/Source/Runtime/Core/Private/UObject/UE5MainStreamObjectVersion.cpp b/Engine/Source/Runtime/Core/Private/UObject/UE5MainStreamObjectVersion.cpp
--- a/Engine/Source/Runtime/Core/Private/UObject/UE5MainStreamObjectVersion.cpp
+++ b/Engine/Source/Runtime/Core/Private/UObject/UE5MainStreamObjectVersion.cpp
@@ -14,7 +14,7 @@ TMap<FGuid, FGuid> FUE5MainStreamObjectVersion::GetSystemGuids()
 	SystemGuids.Add(DevGuids.MATERIALSHADERMAP_DERIVEDDATA_VER, FGuid("65C01D817C9A4EEFAE9E988D41A1F3DD"));
 	SystemGuids.Add(DevGuids.NANITE_DERIVEDDATA_VER, FGuid("7A3EA93546624597BEBE2ADFDE3A7628"));
 	SystemGuids.Add(DevGuids.NIAGARASHADERMAP_DERIVEDDATA_VER, FGuid("8A37C45D24F2423CBE5F8F371DE33575"));
-	SystemGuids.Add(DevGuids.Niagara_LatestScriptCompileVersion, FGuid("E6A9C4E567DE4CB684C367F4A1C2C836"));
+	SystemGuids.Add(DevGuids.Niagara_LatestScriptCompileVersion, FGuid("D9A65544D65B4DEBB4D75186060500BB"));
 	SystemGuids.Add(DevGuids.SkeletalMeshDerivedDataVersion, FGuid("FB7B2FB546DE49B2A2F88B7F69763CEE"));
 	SystemGuids.Add(DevGuids.STATICMESH_DERIVEDDATA_VER, FGuid("98F7D79A5811013825E75BBCC41ED3E9"));
 	SystemGuids.Add(DevGuids.POSESEARCHDB_DERIVEDDATA_VER, FGuid("4E595C2AC5E947D6BA9ABC874353E5BC"));
