- Resolve merge conflicts.

Fix hair instabilities

diff --git a/Engine/Plugins/Runtime/HairStrands/Source/HairStrandsCore/Private/Niagara/NiagaraDataInterfaceHairStrands.cpp b/Engine/Plugins/Runtime/HairStrands/Source/HairStrandsCore/Private/Niagara/NiagaraDataInterfaceHairStrands.cpp
--- a/Engine/Plugins/Runtime/HairStrands/Source/HairStrandsCore/Private/Niagara/NiagaraDataInterfaceHairStrands.cpp
+++ b/Engine/Plugins/Runtime/HairStrands/Source/HairStrandsCore/Private/Niagara/NiagaraDataInterfaceHairStrands.cpp
@@ -285,8 +285,8 @@ void FNDIHairStrandsBuffer::Transfer(FRDGBuilder& GraphBuilder, const TStaticArr
 		const uint32 ScaleCount = 32 * NumScales;
 		const uint32 ScaleBytes = sizeof(float) * ScaleCount;
 		
-		GraphBuilder.QueueBufferUpload(ParamsScaleBuffer.GetOrCreateBuffer(GraphBuilder), InParamsScale.GetData(), ScaleBytes);
-		ParamsScaleBuffer.EndGraphUsage();
+		ParamsScaleBuffer.Initialize(GraphBuilder, TEXT("ParamsScaleBuffer"), EPixelFormat::PF_R32_FLOAT, sizeof(float), ScaleCount);
+		GraphBuilder.QueueBufferUpload(ParamsScaleBuffer.GetOrCreateBuffer(GraphBuilder), ParamsScale.GetData(), ScaleBytes);
 	}
 }
 
diff --git a/Engine/Source/Runtime/Core/Private/UObject/UE5ReleaseStreamObjectVersion.cpp b/Engine/Source/Runtime/Core/Private/UObject/UE5ReleaseStreamObjectVersion.cpp
--- a/Engine/Source/Runtime/Core/Private/UObject/UE5ReleaseStreamObjectVersion.cpp
+++ b/Engine/Source/Runtime/Core/Private/UObject/UE5ReleaseStreamObjectVersion.cpp
@@ -12,7 +12,7 @@ TMap<FGuid, FGuid> FUE5ReleaseStreamObjectVersion::GetSystemGuids()
 	SystemGuids.Add(DevGuids.MATERIALSHADERMAP_DERIVEDDATA_VER, FGuid("79090A66F9D94B5BB285D49E5D39468E"));
 	SystemGuids.Add(DevGuids.NANITE_DERIVEDDATA_VER, FGuid("AC88CFBDDA614A9C8488F64D2DAF699D"));
 	SystemGuids.Add(DevGuids.NIAGARASHADERMAP_DERIVEDDATA_VER, FGuid("D204CCC166ED482685DAEE8AB72895AB"));
-	SystemGuids.Add(DevGuids.Niagara_LatestScriptCompileVersion, FGuid("3B3F2E14B150418481F033E5FAD9D13A"));
+	SystemGuids.Add(DevGuids.Niagara_LatestScriptCompileVersion, FGuid("07F1E268748E4972AF354D99B195A347"));
 	SystemGuids.Add(DevGuids.SkeletalMeshDerivedDataVersion, FGuid("25C49E579B3142DDA2A8C14037267679"));
 
 	return SystemGuids;
