Unshelved from pending changelist '21972241':
- Resolving merge conflicts.

Fix hair instabilities

diff --git a/Engine/Plugins/Runtime/HairStrands/Source/HairStrandsCore/Private/Niagara/NiagaraDataInterfaceHairStrands.cpp b/Engine/Plugins/Runtime/HairStrands/Source/HairStrandsCore/Private/Niagara/NiagaraDataInterfaceHairStrands.cpp
--- a/Engine/Plugins/Runtime/HairStrands/Source/HairStrandsCore/Private/Niagara/NiagaraDataInterfaceHairStrands.cpp
+++ b/Engine/Plugins/Runtime/HairStrands/Source/HairStrandsCore/Private/Niagara/NiagaraDataInterfaceHairStrands.cpp
@@ -287,6 +287,7 @@ void FNDIHairStrandsBuffer::Transfer(FRDGBuilder& GraphBuilder, const TStaticArr
 		
 		ParamsScaleBuffer.Initialize(GraphBuilder, TEXT("ParamsScaleBuffer"), EPixelFormat::PF_R32_FLOAT, sizeof(float), ScaleCount);
 		GraphBuilder.QueueBufferUpload(ParamsScaleBuffer.GetOrCreateBuffer(GraphBuilder), ParamsScale.GetData(), ScaleBytes);
+		ParamsScaleBuffer.EndGraphUsage();
 	}
 }
 
