Change how cards guide are generated at import time.

This improves accuracy for picking the correct principal direction used for generated guides. The new approach parses the entire set of cards and choose the direction (U or V) based on statistic mean. This improves results by avoiding invalid local decision which could happen with low poly cards.

diff --git a/Engine/Plugins/Runtime/HairStrands/Source/HairStrandsCore/Private/HairCardsBuilder.cpp b/Engine/Plugins/Runtime/HairStrands/Source/HairStrandsCore/Private/HairCardsBuilder.cpp
--- a/Engine/Plugins/Runtime/HairStrands/Source/HairStrandsCore/Private/HairCardsBuilder.cpp
+++ b/Engine/Plugins/Runtime/HairStrands/Source/HairStrandsCore/Private/HairCardsBuilder.cpp
@@ -2842,6 +2842,80 @@ namespace HairCards
 
 		const float GuideRadius = 0.01f;
 
+		// Find the principal direction of the atlas by comparing the number of segment along U and along V
+		bool bIsMainDirectionU = true;
+		{
+			uint32 MainDirectionUCount = 0;
+			uint32 ValidCount = 0;
+			FVector4f TriangleUVs[3];
+			for (uint32 CardIt = 0; CardIt < NumCards; ++CardIt)
+			{
+				const uint32 NumTriangles = InCards.IndexCounts[CardIt] / 3;
+				const uint32 VertexOffset = InCards.IndexOffsets[CardIt];
+				struct FIndexAndCoord { uint32 Index; float TexCoord; };
+				struct FSimilarUVVertices
+				{
+					float TexCoord = 0;
+					TArray<uint32> Indices;
+					TArray<FIndexAndCoord> AllIndices; // Store vertex index and Tex.coord perpendicular to the principal axis (stored in TexCoord)
+				};
+				TArray<FSimilarUVVertices> SimilarVertexU;
+				TArray<FSimilarUVVertices> SimilarVertexV;
+
+				auto AddSimilarUV = [](TArray<FSimilarUVVertices>& In, uint32 Index, float TexCoord, float MinorTexCoord, float Threshold)
+				{
+					bool bFound = false;
+					for (int32 It = 0, Count = In.Num(); It < Count; ++It)
+					{
+						if (FMath::Abs(In[It].TexCoord - TexCoord) < Threshold)
+						{
+							// We add only unique vertices per segment, so that the average position land in the center of the cards
+							In[It].Indices.AddUnique(Index);
+							In[It].AllIndices.Add({ Index, MinorTexCoord });
+							bFound = true;
+							break;
+						}
+					}
+
+					if (!bFound)
+					{
+						FSimilarUVVertices& SimilarUV = In.AddDefaulted_GetRef();
+						SimilarUV.TexCoord = TexCoord;
+						SimilarUV.Indices.Add(Index);
+						SimilarUV.AllIndices.Add({ Index, MinorTexCoord });
+					}
+				};
+
+				// Iterate over all triangles of a cards, and find vertices which share either same U or same V. We add them to separate lists. 
+				// We then use the following heuristic: the main axis will have more segments. This is what determine the principal axis of the cards
+				const float UVCoordTreshold = 1.f / 1024.f; // 1 pixel for a 1k texture
+				for (uint32 TriangleIt = 0; TriangleIt < NumTriangles; ++TriangleIt)
+				{
+					const uint32 VertexIndexOffset = VertexOffset + TriangleIt * 3;
+					for (uint32 VertexIt = 0; VertexIt < 3; ++VertexIt)
+					{
+						const uint32 VertexIndex = InCards.Indices[VertexIndexOffset + VertexIt];
+						const FVector4f UV = TriangleUVs[VertexIt] = InCards.UVs[VertexIndex];
+						AddSimilarUV(SimilarVertexU, VertexIndex, UV.X, UV.Y, UVCoordTreshold);
+						AddSimilarUV(SimilarVertexV, VertexIndex, UV.Y, UV.X, UVCoordTreshold);
+					}
+				}
+
+				// Use global UV orientation
+				// Use global segment count orientation
+
+				// Find the perpendicular direction by comparing the number of segment along U and along V
+				const bool bIsValid = SimilarVertexU.Num() != SimilarVertexV.Num() ? 1u : 0u;
+				if (bIsValid)
+				{
+					++ValidCount;
+					MainDirectionUCount += SimilarVertexU.Num() > SimilarVertexV.Num() ? 1u : 0u;
+				}
+			}
+
+			bIsMainDirectionU = (float(MainDirectionUCount) / float(FMath::Max(1u,ValidCount))) > 0.5f;
+		}
+
 		FVector4f TriangleUVs[3];
 		for (uint32 CardIt = 0; CardIt < NumCards; ++CardIt)
 		{
@@ -2897,7 +2971,6 @@ namespace HairCards
 			}
 
 			// Find the perpendicular direction by comparing the number of segment along U and along V
-			const bool bIsMainDirectionU = SimilarVertexU.Num() >= SimilarVertexV.Num();
 			TArray<FVector3f> CenterPoints;
 			{
 				// Sort vertices along the main axis so that, when we iterate through them, we get a correct linear ordering
@@ -3362,7 +3435,7 @@ namespace FHairCardsBuilder
 FString GetVersion()
 {
 	// Important to update the version when cards building or importing changes
-	return TEXT("9f");
+	return TEXT("9g");
 }
 
 void AllocateAtlasTexture(UTexture2D* Out, const FIntPoint& Resolution, uint32 MipCount, EPixelFormat PixelFormat, ETextureSourceFormat SourceFormat)
