Fix Groom rebuild not taking into account new interpolation settings once rebuilding was validated.

diff --git a/Engine/Plugins/Runtime/HairStrands/Source/HairStrandsEditor/Private/GroomActions.cpp b/Engine/Plugins/Runtime/HairStrands/Source/HairStrandsEditor/Private/GroomActions.cpp
--- a/Engine/Plugins/Runtime/HairStrands/Source/HairStrandsEditor/Private/GroomActions.cpp
+++ b/Engine/Plugins/Runtime/HairStrands/Source/HairStrandsEditor/Private/GroomActions.cpp
@@ -180,9 +180,9 @@ void FGroomActions::ExecuteRebuild(TArray<TWeakObjectPtr<UGroomAsset>> Objects)
 				// Duplicate the options to prevent dirtying the asset when they are modified but the rebuild is cancelled
 				UGroomImportOptions* CurrentOptions = DuplicateObject<UGroomImportOptions>(GroomAssetImportData->ImportOptions, nullptr);
 			
+				const uint32 GroupCount = GroomAsset->GetNumHairGroups();
 				UGroomHairGroupsPreview* GroupsPreview = NewObject<UGroomHairGroupsPreview>();
 				{					
-					const uint32 GroupCount = GroomAsset->GetNumHairGroups();
 					for (uint32 GroupIndex = 0; GroupIndex < GroupCount; GroupIndex++)
 					{
 						FGroomHairGroupPreview& OutGroup = GroupsPreview->Groups.AddDefaulted_GetRef();
@@ -204,6 +204,12 @@ void FGroomActions::ExecuteRebuild(TArray<TWeakObjectPtr<UGroomAsset>> Objects)
 					continue;
 				}
 
+				// Apply new interpolation settings to the groom, prior to rebuilding the groom
+				for (uint32 GroupIndex = 0; GroupIndex < GroupCount; ++GroupIndex)
+				{
+					GroomAsset->HairGroupsInterpolation[GroupIndex] = GroupsPreview->Groups[GroupIndex].InterpolationSettings;
+				}
+
 				bool bSucceeded = GroomAsset->CacheDerivedDatas();
 				if (bSucceeded)
 				{
