Fix crash when bone name is incorrect.

diff --git a/Engine/Plugins/Runtime/HairStrands/Source/HairStrandsCore/Private/GroomComponent.cpp b/Engine/Plugins/Runtime/HairStrands/Source/HairStrandsCore/Private/GroomComponent.cpp
--- a/Engine/Plugins/Runtime/HairStrands/Source/HairStrandsCore/Private/GroomComponent.cpp
+++ b/Engine/Plugins/Runtime/HairStrands/Source/HairStrandsCore/Private/GroomComponent.cpp
@@ -3338,19 +3338,22 @@ void UGroomComponent::TickComponent(float DeltaTime, enum ELevelTick TickType, F
 		const FName BoneName(AttachmentName);
 		if (GetAttachSocketName() != BoneName)
 		{
-			AttachToComponent(SkeletelMeshComponent, FAttachmentTransformRules(EAttachmentRule::KeepRelative, false), BoneName);
-			const uint32 BoneIndex = SkeletelMeshComponent->GetBoneIndex(BoneName);
-			const FMatrix BoneTransformRaw = SkeletelMeshComponent->GetSkeletalMeshAsset()->GetComposedRefPoseMatrix(BoneIndex);
-			const FVector BoneLocation = BoneTransformRaw.GetOrigin();
-			const FQuat BoneRotation = BoneTransformRaw.ToQuat();
-
-			FTransform BoneTransform = FTransform::Identity;
-			BoneTransform.SetLocation(BoneLocation);
-			BoneTransform.SetRotation(BoneRotation);
-
-			FTransform InvBoneTransform = BoneTransform.Inverse();
-			SetRelativeLocation(InvBoneTransform.GetLocation());
-			SetRelativeRotation(InvBoneTransform.GetRotation());
+			const int32 BoneIndex = SkeletelMeshComponent->GetBoneIndex(BoneName);
+			if (BoneIndex != INDEX_NONE)
+			{
+				AttachToComponent(SkeletelMeshComponent, FAttachmentTransformRules(EAttachmentRule::KeepRelative, false), BoneName);
+				const FMatrix BoneTransformRaw = SkeletelMeshComponent->GetSkeletalMeshAsset()->GetComposedRefPoseMatrix(BoneIndex);
+				const FVector BoneLocation = BoneTransformRaw.GetOrigin();
+				const FQuat BoneRotation = BoneTransformRaw.ToQuat();
+
+				FTransform BoneTransform = FTransform::Identity;
+				BoneTransform.SetLocation(BoneLocation);
+				BoneTransform.SetRotation(BoneRotation);
+
+				FTransform InvBoneTransform = BoneTransform.Inverse();
+				SetRelativeLocation(InvBoneTransform.GetLocation());
+				SetRelativeRotation(InvBoneTransform.GetRotation());
+			}
 		}
 	}
 
