Groom component PSO precache requests

#rnx

diff --git a/Engine/Plugins/Runtime/HairStrands/Source/HairStrandsCore/Private/GroomComponent.cpp b/Engine/Plugins/Runtime/HairStrands/Source/HairStrandsCore/Private/GroomComponent.cpp
--- a/Engine/Plugins/Runtime/HairStrands/Source/HairStrandsCore/Private/GroomComponent.cpp
+++ b/Engine/Plugins/Runtime/HairStrands/Source/HairStrandsCore/Private/GroomComponent.cpp
@@ -3064,6 +3064,8 @@ void UGroomComponent::PostLoad()
 	// This call will handle the GroomAsset properly if it's still being loaded
 	SetGroomAsset(GroomAsset, BindingAsset, false);
 
+	PrecachePSOs();
+
 #if WITH_EDITOR
 	if (GroomAsset && !bIsGroomAssetCallbackRegistered)
 	{
@@ -3087,6 +3089,112 @@ void UGroomComponent::PostLoad()
 #endif
 }
 
+void UGroomComponent::PrecachePSOs()
+{
+	if (GroomAsset == nullptr || !FApp::CanEverRender() || !PipelineStateCache::IsPSOPrecachingEnabled())
+	{
+		return;
+	}
+
+	struct VFsPerMaterialData
+	{
+		int16 MaterialIndex;
+		EHairGeometryType HairGeometryType;
+		TArray<const FVertexFactoryType*, TInlineAllocator<2>> VertexFactoryTypes;
+	};
+	TArray<VFsPerMaterialData, TInlineAllocator<2>> VFsPerMaterials;
+
+	const auto AddVFPerMaterial = [&VFsPerMaterials](int32 InMaterialIndex, EHairGeometryType InHairGeometryType, const FVertexFactoryType* InVFType)
+	{
+		VFsPerMaterialData* VFsPerMaterial = VFsPerMaterials.FindByPredicate([InMaterialIndex, InHairGeometryType](const VFsPerMaterialData& Other) { return (Other.MaterialIndex == InMaterialIndex && Other.HairGeometryType == InHairGeometryType); });
+		if (VFsPerMaterial == nullptr)
+		{
+			VFsPerMaterial = &VFsPerMaterials.AddDefaulted_GetRef();
+			VFsPerMaterial->MaterialIndex = InMaterialIndex;
+			VFsPerMaterial->HairGeometryType = InHairGeometryType;
+		}
+		VFsPerMaterial->VertexFactoryTypes.AddUnique(InVFType);
+	};
+
+	const EShaderPlatform ShaderPlatform = GetScene() ? GetScene()->GetShaderPlatform() : GMaxRHIShaderPlatform;
+	const int32 GroupCount = GroomAsset->GetNumHairGroups();
+	check(GroomAsset->HairGroupsData.Num() == HairGroupInstances.Num());
+	for (int32 GroupIt = 0; GroupIt < GroupCount; GroupIt++)
+	{
+		const FHairGroupData& InGroupData = GroomAsset->HairGroupsData[GroupIt];
+		FHairGroupInstance* HairInstance = HairGroupInstances[GroupIt];
+		check(HairInstance->HairGroupPublicData);
+
+		// Strands
+		if (HairInstance->Strands.IsValid())
+		{
+			const int32 MaterialIndex = GroomAsset->GetMaterialIndex(GroomAsset->HairGroupsRendering[GroupIt].MaterialSlotName);
+			AddVFPerMaterial(MaterialIndex, EHairGeometryType::Strands, &FHairStrandsVertexFactory::StaticType);
+		}
+
+		// Cards
+		if (HairInstance->Cards.IsValid() && IsHairStrandsEnabled(EHairStrandsShaderType::Cards, ShaderPlatform))
+		{
+			uint32 CardsLODIndex = 0;
+			for (const FHairGroupData::FCards::FLOD& LOD : InGroupData.Cards.LODs)
+			{
+				if (LOD.IsValid())
+				{
+					// Material
+					int32 MaterialIndex = INDEX_NONE;
+					for (const FHairGroupsCardsSourceDescription& Desc : GroomAsset->HairGroupsCards)
+					{
+						if (Desc.GroupIndex == GroupIt && Desc.LODIndex == CardsLODIndex)
+						{
+							MaterialIndex = GroomAsset->GetMaterialIndex(Desc.MaterialSlotName);
+							break;
+						}
+					}
+					AddVFPerMaterial(MaterialIndex, EHairGeometryType::Cards, &FHairCardsVertexFactory::StaticType);
+				}
+				++CardsLODIndex;
+			}
+		}
+
+		// Meshes
+		if (HairInstance->Meshes.IsValid())
+		{
+			uint32 MeshesLODIndex = 0;
+			for (const FHairGroupData::FMeshes::FLOD& LOD : InGroupData.Meshes.LODs)
+			{
+				if (LOD.IsValid())
+				{
+					// Material
+					int32 MaterialIndex = INDEX_NONE;
+					for (const FHairGroupsMeshesSourceDescription& Desc : GroomAsset->HairGroupsMeshes)
+					{
+						if (Desc.GroupIndex == GroupIt && Desc.LODIndex == MeshesLODIndex)
+						{
+							MaterialIndex = GroomAsset->GetMaterialIndex(Desc.MaterialSlotName);
+							break;
+						}
+					}
+					AddVFPerMaterial(MaterialIndex, EHairGeometryType::Meshes, &FHairCardsVertexFactory::StaticType);
+				}
+				++MeshesLODIndex;
+			}
+		}
+	}
+
+	FPSOPrecacheParams PrecachePSOParams;
+	SetupPrecachePSOParams(PrecachePSOParams);
+
+	for (VFsPerMaterialData& VFsPerMaterial : VFsPerMaterials)
+	{
+		UMaterialInterface* MaterialInterface = GetMaterial(GetMaterialIndexWithFallback(VFsPerMaterial.MaterialIndex), VFsPerMaterial.HairGeometryType, true);
+		if (MaterialInterface)
+		{
+			MaterialInterface->PrecachePSOs(VFsPerMaterial.VertexFactoryTypes, PrecachePSOParams);
+		}
+	}
+
+}
+
 #if WITH_EDITOR
 void UGroomComponent::Invalidate()
 {
diff --git a/Engine/Plugins/Runtime/HairStrands/Source/HairStrandsCore/Public/GroomComponent.h b/Engine/Plugins/Runtime/HairStrands/Source/HairStrandsCore/Public/GroomComponent.h
--- a/Engine/Plugins/Runtime/HairStrands/Source/HairStrandsCore/Public/GroomComponent.h
+++ b/Engine/Plugins/Runtime/HairStrands/Source/HairStrandsCore/Public/GroomComponent.h
@@ -135,6 +135,7 @@ public:
 
 	//~ Begin UMeshComponent Interface.
 	virtual void PostLoad() override;
+	virtual void PrecachePSOs() override;
 	//~ End UMeshComponent Interface.
 
 	/** Return the guide hairs rest resources*/
