GroomCache streamer is now registered from the HairStrandsCore instead of HairStrandsEditor to support remote MRQ, which uses -game.

diff --git a/Engine/Plugins/Runtime/HairStrands/Source/HairStrandsCore/Private/HairStrandsCore.cpp b/Engine/Plugins/Runtime/HairStrands/Source/HairStrandsCore/Private/HairStrandsCore.cpp
--- a/Engine/Plugins/Runtime/HairStrands/Source/HairStrandsCore/Private/HairStrandsCore.cpp
+++ b/Engine/Plugins/Runtime/HairStrands/Source/HairStrandsCore/Private/HairStrandsCore.cpp
@@ -7,6 +7,7 @@
 #include "GroomManager.h"
 #include "Engine/StaticMesh.h"
 #include "GroomAsset.h"
+#include "GroomCacheStreamingManager.h"
 #include "GroomCreateBindingOptions.h"
 
 IMPLEMENT_MODULE(FHairStrandsCore, HairStrandsCore);
@@ -25,10 +26,14 @@ void FHairStrandsCore::StartupModule()
 	FString PluginShaderDir = FPaths::Combine(IPluginManager::Get().FindPlugin(TEXT("HairStrands"))->GetBaseDir(), TEXT("Shaders"));
 	AddShaderSourceDirectoryMapping(TEXT("/Plugin/Runtime/HairStrands"), PluginShaderDir);
 	SetHairStrandsEnabled(true);
+
+	IGroomCacheStreamingManager::Register();
 }
 
 void FHairStrandsCore::ShutdownModule()
 {
+	IGroomCacheStreamingManager::Unregister();
+
 	SetHairStrandsEnabled(false);
 }
 
diff --git a/Engine/Plugins/Runtime/HairStrands/Source/HairStrandsEditor/Private/HairStrandsEditor.cpp b/Engine/Plugins/Runtime/HairStrands/Source/HairStrandsEditor/Private/HairStrandsEditor.cpp
--- a/Engine/Plugins/Runtime/HairStrands/Source/HairStrandsEditor/Private/HairStrandsEditor.cpp
+++ b/Engine/Plugins/Runtime/HairStrands/Source/HairStrandsEditor/Private/HairStrandsEditor.cpp
@@ -18,7 +18,6 @@
 #include "GroomBindingDetailsCustomization.h"
 #include "GroomCacheImportOptions.h"
 #include "GroomCacheImportSettingsCustomization.h"
-#include "GroomCacheStreamingManager.h"
 #include "GroomCacheTrackEditor.h"
 #include "GroomComponentDetailsCustomization.h"
 #include "GroomCreateBindingOptions.h"
@@ -149,14 +148,10 @@ void FGroomEditor::StartupModule()
 			GetMutableDefault<UGroomPluginSettings>()
 		);
 	}
-
-	IGroomCacheStreamingManager::Register();
 }
 
 void FGroomEditor::ShutdownModule()
 {
-	IGroomCacheStreamingManager::Unregister();
-
 	ISettingsModule* SettingsModule = FModuleManager::GetModulePtr<ISettingsModule>("Settings");
 	if (SettingsModule != nullptr)
 	{
