Fix GPU hair crash due to bad LocalToWorld coming from the skelmesh proxy

diff --git a/Engine/Plugins/Runtime/HairStrands/Shaders/Private/NiagaraDataInterfaceVelocityGridTemplate.ush b/Engine/Plugins/Runtime/HairStrands/Shaders/Private/NiagaraDataInterfaceVelocityGridTemplate.ush
--- a/Engine/Plugins/Runtime/HairStrands/Shaders/Private/NiagaraDataInterfaceVelocityGridTemplate.ush
+++ b/Engine/Plugins/Runtime/HairStrands/Shaders/Private/NiagaraDataInterfaceVelocityGridTemplate.ush
@@ -365,7 +365,6 @@ void BuildVelocityFace_{ParameterName}(in int3 GridIndex, in float ShapeFunction
 	{
 		const float3 DeltaPosition = (GridLocation+ComponentOffset)*GridLength - GridPosition;
 		const float GridMomentum = SampleMomentum + dot(MomentumGradient, DeltaPosition);
-
 		InterlockedAddScalar_{ParameterName}(GridLocation,FLUID_MASS_OFFSET+ComponentIndex,SampleMass*ShapeFunction);
 		InterlockedAddScalar_{ParameterName}(GridLocation,FLUID_VELOCITY_OFFSET+ComponentIndex,GridMomentum*ShapeFunction);
 	}
@@ -400,17 +399,21 @@ void BuildVelocityComponent_{ParameterName}(in float3 GridPosition, in float Sam
 void BuildVelocityField_{ParameterName}(in float3 GridOrigin, in float GridLength, in float3 ParticlePosition, in float ParticleMass, in float3 ParticleVelocity, 
 			in float4x4 VelocityGradient, out bool OutFunctionStatus )
 {
-	const float3 ParticleMomentum = ParticleVelocity * ParticleMass;
+	if(isfinite(ParticlePosition.x) && isfinite(ParticlePosition.y) && isfinite(ParticlePosition.z) &&
+	   isfinite(ParticleVelocity.x) && isfinite(ParticleVelocity.y) && isfinite(ParticleVelocity.z))
+	{
+		const float3 ParticleMomentum = ParticleVelocity * ParticleMass;
 
-	const float3 MomentumGradientX = VelocityGradient[0].xyz * ParticleMass;
-	const float3 MomentumGradientY = VelocityGradient[1].xyz * ParticleMass;
-	const float3 MomentumGradientZ = VelocityGradient[2].xyz * ParticleMass;
+		const float3 MomentumGradientX = VelocityGradient[0].xyz * ParticleMass;
+		const float3 MomentumGradientY = VelocityGradient[1].xyz * ParticleMass;
+		const float3 MomentumGradientZ = VelocityGradient[2].xyz * ParticleMass;
 
-	const float3 GridPosition = GetGridPosition_{ParameterName}(ParticlePosition,GridOrigin);
+		const float3 GridPosition = GetGridPosition_{ParameterName}(ParticlePosition,GridOrigin);
 
-	BuildVelocityComponent_{ParameterName}(GridPosition,ParticleMass,ParticleMomentum.x,MomentumGradientX,float3(0.0,0.5,0.5),0,GridLength);
-	BuildVelocityComponent_{ParameterName}(GridPosition,ParticleMass,ParticleMomentum.y,MomentumGradientY,float3(0.5,0.0,0.5),1,GridLength);
-	BuildVelocityComponent_{ParameterName}(GridPosition,ParticleMass,ParticleMomentum.z,MomentumGradientZ,float3(0.5,0.5,0.0),2,GridLength);
+		BuildVelocityComponent_{ParameterName}(GridPosition,ParticleMass,ParticleMomentum.x,MomentumGradientX,float3(0.0,0.5,0.5),0,GridLength);
+		BuildVelocityComponent_{ParameterName}(GridPosition,ParticleMass,ParticleMomentum.y,MomentumGradientY,float3(0.5,0.0,0.5),1,GridLength);
+		BuildVelocityComponent_{ParameterName}(GridPosition,ParticleMass,ParticleMomentum.z,MomentumGradientZ,float3(0.5,0.5,0.0),2,GridLength);
+	}
 	
 	OutFunctionStatus = true;
 }
diff --git a/Engine/Plugins/Runtime/HairStrands/Source/HairStrandsCore/Private/GroomManager.cpp b/Engine/Plugins/Runtime/HairStrands/Source/HairStrandsCore/Private/GroomManager.cpp
--- a/Engine/Plugins/Runtime/HairStrands/Source/HairStrandsCore/Private/GroomManager.cpp
+++ b/Engine/Plugins/Runtime/HairStrands/Source/HairStrandsCore/Private/GroomManager.cpp
@@ -30,7 +30,7 @@ static FAutoConsoleVariableRef CVarHairStrands_UseCards(TEXT("r.HairStrands.UseC
 static int32 GHairStrands_SwapBufferType = 3;
 static FAutoConsoleVariableRef CVarGHairStrands_SwapBufferType(TEXT("r.HairStrands.SwapType"), GHairStrands_SwapBufferType, TEXT("Swap rendering buffer at the end of frame. This is an experimental toggle. Default:1"));
 
-static int32 GHairStrands_ManualSkinCache = 1;
+static int32 GHairStrands_ManualSkinCache = 0;
 static FAutoConsoleVariableRef CVarGHairStrands_ManualSkinCache(TEXT("r.HairStrands.ManualSkinCache"), GHairStrands_ManualSkinCache, TEXT("If skin cache is not enabled, and grooms use skinning method, this enable a simple skin cache mechanisme for groom. Default:disable"));
 
 static int32 GHairStrands_InterpolationFrustumCullingEnable = 1;
