Rigged Hair Mesh + Reuse of the attached naiagra component for hair sim

diff --git a/Engine/Plugins/Runtime/HairStrands/Config/BaseHairStrands.ini b/Engine/Plugins/Runtime/HairStrands/Config/BaseHairStrands.ini
--- a/Engine/Plugins/Runtime/HairStrands/Config/BaseHairStrands.ini
+++ b/Engine/Plugins/Runtime/HairStrands/Config/BaseHairStrands.ini
@@ -11,4 +11,7 @@ fx.UseShaderStages=1
 +ClassRedirects=(OldName="/Script/HairStrandsEditor.GroomImportOptions",NewName="/Script/HairStrandsCore.GroomImportOptions")
 +ClassRedirects=(OldName="/Script/HairStrandsNiagara.NiagaraDataInterfaceHairStrands",NewName="/Script/HairStrandsCore.NiagaraDataInterfaceHairStrands")
 +ClassRedirects=(OldName="/Script/HairStrandsNiagara.NiagaraDataInterfaceVelocityGrid",NewName="/Script/HairStrandsCore.NiagaraDataInterfaceVelocityGrid")
-+ClassRedirects=(OldName="/Script/HairStrandsNiagara.NiagaraDataInterfacePressureGrid",NewName="/Script/HairStrandsCore.NiagaraDataInterfacePressureGrid")
\ No newline at end of file
++ClassRedirects=(OldName="/Script/HairStrandsNiagara.NiagaraDataInterfacePressureGrid",NewName="/Script/HairStrandsCore.NiagaraDataInterfacePressureGrid")
++PropertyRedirects=(OldName="GroomAsset.DeformedSkeletalMesh",NewName="GroomAsset.RiggedSkeletalMesh")
++PropertyRedirects=(OldName="HairDeformationSettings.bEnableDeformation",NewName="HairDeformationSettings.bEnableRigging")
++PropertyRedirects=(OldName="HairGroupsInterpolation.DeformationSettings",NewName="HairGroupsInterpolation.RiggingSettings")
\ No newline at end of file
diff --git a/Engine/Plugins/Runtime/HairStrands/Source/HairStrandsCore/Private/GroomAsset.cpp b/Engine/Plugins/Runtime/HairStrands/Source/HairStrandsCore/Private/GroomAsset.cpp
--- a/Engine/Plugins/Runtime/HairStrands/Source/HairStrandsCore/Private/GroomAsset.cpp
+++ b/Engine/Plugins/Runtime/HairStrands/Source/HairStrandsCore/Private/GroomAsset.cpp
@@ -1317,7 +1317,7 @@ static bool IsStrandsInterpolationAttributes(const FName PropertyName)
 
 		|| PropertyName == GET_MEMBER_NAME_CHECKED(FHairDeformationSettings, NumCurves)
 		|| PropertyName == GET_MEMBER_NAME_CHECKED(FHairDeformationSettings, NumPoints)
-		|| PropertyName == GET_MEMBER_NAME_CHECKED(FHairDeformationSettings, bEnableDeformation)
+		|| PropertyName == GET_MEMBER_NAME_CHECKED(FHairDeformationSettings, bEnableRigging)
 
 		// Add dependency on simulation and per LOD-simulation/global-interoplation to strip-out interoplation data if there are not needed
 		|| PropertyName == GET_MEMBER_NAME_CHECKED(UGroomAsset, EnableGlobalInterpolation)
@@ -1367,14 +1367,6 @@ void UGroomAsset::PostEditChangeProperty(FPropertyChangedEvent& PropertyChangedE
 	if (bStrandsInterpolationChanged)
 	{
 		CacheDerivedDatas();
-
-		if (PropertyName == GET_MEMBER_NAME_CHECKED(FHairInterpolationSettings, bOverrideGuides)
-		 || PropertyName == GET_MEMBER_NAME_CHECKED(FHairDeformationSettings, NumCurves)
-		 || PropertyName == GET_MEMBER_NAME_CHECKED(FHairDeformationSettings, NumPoints)
-		 || PropertyName == GET_MEMBER_NAME_CHECKED(FHairDeformationSettings, bEnableDeformation))
-		{
-			DeformedSkeletalMesh = FGroomDeformerBuilder::CreateSkeletalMesh(this);  
-		}
 	}
 	
 	const bool bCardsArrayChanged = PropertyName == GET_MEMBER_NAME_CHECKED(UGroomAsset, HairGroupsCards);
@@ -1777,7 +1769,7 @@ void UGroomAsset::SetHairWidth(float Width)
 // differences, etc.) replace the version GUID below with a new one.
 // In case of merge conflicts with DDC versions, you MUST generate a new GUID
 // and set this new GUID as the version.
-#define GROOM_DERIVED_DATA_VERSION TEXT("DAF54482721C4B60885A27B2720FDC6E")
+#define GROOM_DERIVED_DATA_VERSION TEXT("FC0D84AFC0F444A9A988C651BEED113A")
 
 #if WITH_EDITORONLY_DATA
 
@@ -3236,7 +3228,7 @@ bool UGroomAsset::IsVisible(int32 GroupIndex, int32 LODIndex) const
 
 bool UGroomAsset::IsDeformationEnable(int32 GroupIndex) const
 {
-	return HairGroupsInterpolation.IsValidIndex(GroupIndex) && HairGroupsInterpolation[GroupIndex].DeformationSettings.bEnableDeformation && HairGroupsInterpolation[GroupIndex].InterpolationSettings.bOverrideGuides;
+	return HairGroupsInterpolation.IsValidIndex(GroupIndex) && HairGroupsInterpolation[GroupIndex].RiggingSettings.bEnableRigging && HairGroupsInterpolation[GroupIndex].InterpolationSettings.bOverrideGuides;
 }
 
 bool UGroomAsset::IsSimulationEnable(int32 GroupIndex, int32 LODIndex) const 
diff --git a/Engine/Plugins/Runtime/HairStrands/Source/HairStrandsCore/Private/GroomAssetInterpolation.cpp b/Engine/Plugins/Runtime/HairStrands/Source/HairStrandsCore/Private/GroomAssetInterpolation.cpp
--- a/Engine/Plugins/Runtime/HairStrands/Source/HairStrandsCore/Private/GroomAssetInterpolation.cpp
+++ b/Engine/Plugins/Runtime/HairStrands/Source/HairStrandsCore/Private/GroomAssetInterpolation.cpp
@@ -22,7 +22,8 @@ FHairInterpolationSettings::FHairInterpolationSettings()
 
 FHairDeformationSettings::FHairDeformationSettings()
 {
-	bEnableDeformation = false;
+	bCanEditRigging = false;
+	bEnableRigging = false;
 	NumCurves = 10;
 	NumPoints = 4;
 }
@@ -99,9 +100,10 @@ bool FHairGroupsInterpolation::operator==(const FHairGroupsInterpolation& A) con
 bool FHairDeformationSettings::operator==(const FHairDeformationSettings& A) const
 {
 	return
-		bEnableDeformation == A.bEnableDeformation &&
+		bEnableRigging == A.bEnableRigging &&
 		NumCurves == A.NumCurves &&
-		NumPoints == A.NumPoints;
+		NumPoints == A.NumPoints &&
+		bCanEditRigging == A.bCanEditRigging;
 }
 
 void FHairGroupsInterpolation::BuildDDCKey(FArchive& Ar)
@@ -114,9 +116,10 @@ void FHairGroupsInterpolation::BuildDDCKey(FArchive& Ar)
 	Ar << InterpolationSettings.InterpolationDistance;
 	Ar << InterpolationSettings.bRandomizeGuide;
 	Ar << InterpolationSettings.bUseUniqueGuide;
-	Ar << DeformationSettings.NumCurves;
-	Ar << DeformationSettings.NumPoints;
-	Ar << DeformationSettings.bEnableDeformation;
+	Ar << RiggingSettings.NumCurves;
+	Ar << RiggingSettings.NumPoints;
+	Ar << RiggingSettings.bEnableRigging;
+	Ar << RiggingSettings.bCanEditRigging;
 }
 
 void FHairGroupsLOD::BuildDDCKey(FArchive& Ar)
diff --git a/Engine/Plugins/Runtime/HairStrands/Source/HairStrandsCore/Private/GroomBuilder.cpp b/Engine/Plugins/Runtime/HairStrands/Source/HairStrandsCore/Private/GroomBuilder.cpp
--- a/Engine/Plugins/Runtime/HairStrands/Source/HairStrandsCore/Private/GroomBuilder.cpp
+++ b/Engine/Plugins/Runtime/HairStrands/Source/HairStrandsCore/Private/GroomBuilder.cpp
@@ -1847,9 +1847,20 @@ void FGroomBuilder::BuildData(
 			else
 			{
 				OutSim.Reset();
-				if(InSettings.DeformationSettings.bEnableDeformation)
+				if(InSettings.RiggingSettings.bEnableRigging)
 				{
-					FHairStrandsDecimation::Decimate(OutRen, InSettings.DeformationSettings.NumCurves, InSettings.DeformationSettings.NumPoints, OutSim);
+					if (InHairDescriptionGroup.Info.NumGuides > 0)
+					{
+						// We pick the new guides among the imported ones
+						FHairStrandsDatas TempSim = InHairDescriptionGroup.Guides;
+						HairStrandsBuilder::BuildInternalData(TempSim, true);
+						FHairStrandsDecimation::Decimate(TempSim, InSettings.RiggingSettings.NumCurves, InSettings.RiggingSettings.NumPoints, OutSim);
+					}
+					else
+					{
+						// Otherwise let s pick the guides among the rendered strands
+						FHairStrandsDecimation::Decimate(OutRen, InSettings.RiggingSettings.NumCurves, InSettings.RiggingSettings.NumPoints, OutSim);
+					}
 				}
 				else
 				{
diff --git a/Engine/Plugins/Runtime/HairStrands/Source/HairStrandsCore/Private/GroomComponent.cpp b/Engine/Plugins/Runtime/HairStrands/Source/HairStrandsCore/Private/GroomComponent.cpp
--- a/Engine/Plugins/Runtime/HairStrands/Source/HairStrandsCore/Private/GroomComponent.cpp
+++ b/Engine/Plugins/Runtime/HairStrands/Source/HairStrandsCore/Private/GroomComponent.cpp
@@ -58,6 +58,9 @@ static FAutoConsoleVariableRef CVarHairStrands_Streaming_Prediction(TEXT("r.Hair
 static int32 GHairStrands_BoundsMode = 0;
 static FAutoConsoleVariableRef CVarHairStrands_BoundsMode(TEXT("r.HairStrands.BoundMode"), GHairStrands_BoundsMode, TEXT("Define how hair bound are computed at runtime when attached to a skel. mesh.\n 0: Use skel.mesh extented with grooms bounds.\n 1: Use skel.mesh bounds.\n 2: Use skel.mesh extented with grooms bounds (conservative)"));
 
+static int32 GHairStrands_UseAttachedSimulationComponents = 0;
+static FAutoConsoleVariableRef CVarHairStrands_UseAttachedSimulationComponents(TEXT("r.HairStrands.UseAttachedSimulationComponents"), GHairStrands_UseAttachedSimulationComponents, TEXT("Boolean to check if we are using already attached niagara components for simulation (WIP)"));
+
 #define LOCTEXT_NAMESPACE "GroomComponent"
 
 #define USE_HAIR_TRIANGLE_STRIP 0
@@ -1449,6 +1452,18 @@ void UGroomComponent::UpdateHairSimulation()
 	const int32 NumComponents = FMath::Max(NumGroups, NiagaraComponents.Num());
 
 	NiagaraComponents.SetNumZeroed(NumComponents);
+	if( GHairStrands_UseAttachedSimulationComponents == 1)
+	{
+		int32 CompIndex = 0;
+		// Fill the niagara components with the already attached children if any
+		for(const TObjectPtr<USceneComponent>& GroomChild : GetAttachChildren())
+		{
+			if(GroomChild->IsA<UNiagaraComponent>() && (CompIndex < NumComponents) && !NiagaraComponents[CompIndex])
+			{
+				NiagaraComponents[CompIndex++] = StaticCast<UNiagaraComponent*>(GroomChild);
+			}
+		}
+	}
 	const int32 LODInit = (LODForcedIndex != -1.f) ? LODForcedIndex : (LODPredictedIndex != -1.f) ? LODPredictedIndex : 0;
 	for (int32 CompIndex = 0; CompIndex < NumComponents; ++CompIndex)
 	{
@@ -2446,9 +2461,9 @@ void UGroomComponent::InitResources(bool bIsBindingReloading)
 	}
 
 	// Grab deformed mesh if it exists 
-	if(GroomAsset && GroomAsset->DeformedSkeletalMesh)
+	if(GroomAsset && GroomAsset->RiggedSkeletalMesh)
 	{
-		const USkeleton* TargetSkeleton = GroomAsset->DeformedSkeletalMesh->GetSkeleton();
+		const USkeleton* TargetSkeleton = GroomAsset->RiggedSkeletalMesh->GetSkeleton();
 		
 		DeformedMeshComponent = nullptr;
 		// Try to find the first component by walking the attachment hierarchy
@@ -2562,7 +2577,7 @@ void UGroomComponent::InitResources(bool bIsBindingReloading)
 		HairGroupInstance->GeometryType = EHairGeometryType::NoneGeometry;
 		HairGroupInstance->BindingType = EHairBindingType::NoneBinding;
 
-		HairGroupInstance->Debug.SkinningCurrentLocalToWorld	= RegisteredMeshComponent ? RegisteredMeshComponent->GetComponentTransform() : FTransform();
+		HairGroupInstance->Debug.SkinningCurrentLocalToWorld	= RegisteredMeshComponent ? RegisteredMeshComponent->GetComponentTransform() : FTransform::Identity;
 		HairGroupInstance->Debug.SkinningPreviousLocalToWorld	= HairGroupInstance->Debug.SkinningCurrentLocalToWorld;
 		HairGroupInstance->Debug.RigidCurrentLocalToWorld		= GetComponentTransform();
 		HairGroupInstance->Debug.RigidPreviousLocalToWorld		= HairGroupInstance->Debug.RigidCurrentLocalToWorld;
diff --git a/Engine/Plugins/Runtime/HairStrands/Source/HairStrandsCore/Private/GroomDeformerBuilder.cpp b/Engine/Plugins/Runtime/HairStrands/Source/HairStrandsCore/Private/GroomDeformerBuilder.cpp
--- a/Engine/Plugins/Runtime/HairStrands/Source/HairStrandsCore/Private/GroomDeformerBuilder.cpp
+++ b/Engine/Plugins/Runtime/HairStrands/Source/HairStrandsCore/Private/GroomDeformerBuilder.cpp
@@ -358,7 +358,7 @@ static USkeletalMesh* BuildSkeletalMesh(UGroomAsset* GroomAsset)
 	
 	// Create package for the skeletal mesh
 	UPackage* PackageSkeletalMesh = CreatePackage(*PackageNameSkeletalMesh);
-	USkeletalMesh* SkeletalMesh = GroomAsset->DeformedSkeletalMesh ? GroomAsset->DeformedSkeletalMesh.Get() :
+	USkeletalMesh* SkeletalMesh = GroomAsset->RiggedSkeletalMesh ? GroomAsset->RiggedSkeletalMesh.Get() :
 		NewObject<USkeletalMesh>(PackageSkeletalMesh, *NameSkeletalMesh, RF_Public | RF_Standalone | RF_Transactional);
 
 	// Reset the skeletal mesh
diff --git a/Engine/Plugins/Runtime/HairStrands/Source/HairStrandsCore/Public/GroomAsset.h b/Engine/Plugins/Runtime/HairStrands/Source/HairStrandsCore/Public/GroomAsset.h
--- a/Engine/Plugins/Runtime/HairStrands/Source/HairStrandsCore/Public/GroomAsset.h
+++ b/Engine/Plugins/Runtime/HairStrands/Source/HairStrandsCore/Public/GroomAsset.h
@@ -460,7 +460,7 @@ public:
 
 	/** Deformed skeletal mesh that will drive the groom deformation/simulation. For creating this skeletal mesh, enable EnableDeformation within the interpolation settings*/
 	UPROPERTY(VisibleAnywhere, BlueprintReadWrite, Category = "HairInterpolation")
-	TObjectPtr<USkeletalMesh> DeformedSkeletalMesh;
+	TObjectPtr<USkeletalMesh> RiggedSkeletalMesh;
 
 	/** Deformed skeletal mesh mapping from groups to sections */
 	UPROPERTY()
diff --git a/Engine/Plugins/Runtime/HairStrands/Source/HairStrandsCore/Public/GroomAssetInterpolation.h b/Engine/Plugins/Runtime/HairStrands/Source/HairStrandsCore/Public/GroomAssetInterpolation.h
--- a/Engine/Plugins/Runtime/HairStrands/Source/HairStrandsCore/Public/GroomAssetInterpolation.h
+++ b/Engine/Plugins/Runtime/HairStrands/Source/HairStrandsCore/Public/GroomAssetInterpolation.h
@@ -153,17 +153,20 @@ struct HAIRSTRANDSCORE_API FHairDeformationSettings
 	GENERATED_USTRUCT_BODY()
 
 	FHairDeformationSettings();
+	
+	UPROPERTY(VisibleAnywhere, Category = "SkeletonGeneration")
+	bool bCanEditRigging;
 
 	/** Enable the generation of a skeletal mesh that will drive the guides deformation */
-	UPROPERTY(EditAnywhere, BlueprintReadWrite, Category = "InterpolationSettings", meta = (ToolTip = "Enable the generation of a skeletal mesh that will drive the guides deformation. The total number of bones per group is for now 256 (NumCurves * NumPoints)"))
-	bool bEnableDeformation;
+	UPROPERTY(EditAnywhere, BlueprintReadWrite, Category = "SkeletonGeneration", meta = (ToolTip = "Enable the generation of a skeletal mesh that will drive the guides deformation. The total number of bones per group is for now 256 (NumCurves * NumPoints)", EditCondition = "bCanEditRigging"))
+	bool bEnableRigging;
 
 	/** Number of curves to generate in the skel mesh */
-	UPROPERTY(EditAnywhere, BlueprintReadWrite, Category = "InterpolationSettings", meta = (ToolTip = "Number of guides that will be generated on the groom and the skeletal mesh", EditCondition = "bEnableDeformation"))
+	UPROPERTY(EditAnywhere, BlueprintReadWrite, Category = "SkeletonGeneration", meta = (ToolTip = "Number of guides that will be generated on the groom and the skeletal mesh", EditCondition = "bCanEditRigging"))
 	int32 NumCurves;
 
 	/** Number of points per curve */
-	UPROPERTY(EditAnywhere, BlueprintReadWrite, Category = "InterpolationSettings", meta = (ToolTip = "Number of points/bones per generated guide", EditCondition = "bEnableDeformation"))
+	UPROPERTY(EditAnywhere, BlueprintReadWrite, Category = "SkeletonGeneration", meta = (ToolTip = "Number of points/bones per generated guide", EditCondition = "bCanEditRigging"))
 	int32 NumPoints;
 
 	bool operator==(const FHairDeformationSettings& A) const;
@@ -182,8 +185,8 @@ struct HAIRSTRANDSCORE_API FHairGroupsInterpolation
 	UPROPERTY(EditAnywhere, Category = "InterpolationSettings", meta = (ToolTip = "Interpolation settings"))
 	FHairInterpolationSettings InterpolationSettings;
 
-	UPROPERTY(EditAnywhere, Category = "DeformationSettings", meta = (ToolTip = "Deformation settings"))
-	FHairDeformationSettings DeformationSettings;
+	UPROPERTY(EditAnywhere, Category = "SkeletonGeneration", meta = (DisplayName = "Skeleton Generation", ToolTip = "Skeleton Generation"))
+	FHairDeformationSettings RiggingSettings;
 
 	bool operator==(const FHairGroupsInterpolation& A) const;
 
diff --git a/Engine/Plugins/Runtime/HairStrands/Source/HairStrandsEditor/Private/GroomActions.cpp b/Engine/Plugins/Runtime/HairStrands/Source/HairStrandsEditor/Private/GroomActions.cpp
--- a/Engine/Plugins/Runtime/HairStrands/Source/HairStrandsEditor/Private/GroomActions.cpp
+++ b/Engine/Plugins/Runtime/HairStrands/Source/HairStrandsEditor/Private/GroomActions.cpp
@@ -9,6 +9,7 @@
 #include "GeometryCache.h"
 #include "GroomAssetImportData.h"
 #include "GroomBuilder.h"
+#include "GroomDeformerBuilder.h"
 #include "GroomImportOptions.h"
 #include "GroomImportOptionsWindow.h"
 #include "GroomCustomAssetEditorToolkit.h"
@@ -195,6 +196,7 @@ void FGroomActions::ExecuteRebuild(TArray<TWeakObjectPtr<UGroomAsset>> Objects)
 						OutGroup.bHasRoughnessAttributes = false;
 						OutGroup.bHasPrecomputedWeights = false;
 						OutGroup.InterpolationSettings = GroomAsset->HairGroupsInterpolation[GroupIndex];
+						OutGroup.InterpolationSettings.RiggingSettings.bCanEditRigging = true;
 					}
 				}
 				TSharedPtr<SGroomImportOptionsWindow> GroomOptionWindow = SGroomImportOptionsWindow::DisplayRebuildOptions(CurrentOptions, GroupsPreview, Filename);
@@ -205,14 +207,22 @@ void FGroomActions::ExecuteRebuild(TArray<TWeakObjectPtr<UGroomAsset>> Objects)
 				}
 
 				// Apply new interpolation settings to the groom, prior to rebuilding the groom
+				bool bEnableRigging = false;
 				for (uint32 GroupIndex = 0; GroupIndex < GroupCount; ++GroupIndex)
 				{
 					GroomAsset->HairGroupsInterpolation[GroupIndex] = GroupsPreview->Groups[GroupIndex].InterpolationSettings;
+					GroomAsset->HairGroupsInterpolation[GroupIndex].RiggingSettings.bCanEditRigging = false;
+					bEnableRigging |= GroomAsset->HairGroupsInterpolation[GroupIndex].RiggingSettings.bEnableRigging &&
+						GroomAsset->HairGroupsInterpolation[GroupIndex].InterpolationSettings.bOverrideGuides;
 				}
 
 				bool bSucceeded = GroomAsset->CacheDerivedDatas();
 				if (bSucceeded)
 				{
+					if(bEnableRigging)
+					{
+						GroomAsset->RiggedSkeletalMesh = FGroomDeformerBuilder::CreateSkeletalMesh(GroomAsset.Get());
+					}
 					// Move the transient ImportOptions to the asset package and set it on the GroomAssetImportData for serialization
 					CurrentOptions->Rename(nullptr, GroomAssetImportData);
 					for (const FGroomHairGroupPreview& GroupPreview : GroupsPreview->Groups)
diff --git a/Engine/Plugins/Runtime/HairStrands/Source/HairStrandsEditor/Private/GroomAssetDetails.cpp b/Engine/Plugins/Runtime/HairStrands/Source/HairStrandsEditor/Private/GroomAssetDetails.cpp
--- a/Engine/Plugins/Runtime/HairStrands/Source/HairStrandsEditor/Private/GroomAssetDetails.cpp
+++ b/Engine/Plugins/Runtime/HairStrands/Source/HairStrandsEditor/Private/GroomAssetDetails.cpp
@@ -395,7 +395,7 @@ void FGroomRenderingDetails::CustomizeStrandsGroupProperties(IDetailLayoutBuilde
 		DetailLayout.HideProperty(DetailLayout.GetProperty(GET_MEMBER_NAME_CHECKED(UGroomAsset, HairInterpolationType), UGroomAsset::StaticClass()));
 		DetailLayout.HideProperty(DetailLayout.GetProperty(GET_MEMBER_NAME_CHECKED(UGroomAsset, MinLOD), UGroomAsset::StaticClass()));
 		DetailLayout.HideProperty(DetailLayout.GetProperty(GET_MEMBER_NAME_CHECKED(UGroomAsset, DisableBelowMinLodStripping), UGroomAsset::StaticClass()));
-		DetailLayout.HideProperty(DetailLayout.GetProperty(GET_MEMBER_NAME_CHECKED(UGroomAsset, DeformedSkeletalMesh), UGroomAsset::StaticClass()));
+		DetailLayout.HideProperty(DetailLayout.GetProperty(GET_MEMBER_NAME_CHECKED(UGroomAsset, RiggedSkeletalMesh), UGroomAsset::StaticClass()));
 	}
 	break;
 	case EMaterialPanelType::Meshes:
@@ -412,7 +412,7 @@ void FGroomRenderingDetails::CustomizeStrandsGroupProperties(IDetailLayoutBuilde
 		DetailLayout.HideProperty(DetailLayout.GetProperty(GET_MEMBER_NAME_CHECKED(UGroomAsset, HairInterpolationType), UGroomAsset::StaticClass()));
 		DetailLayout.HideProperty(DetailLayout.GetProperty(GET_MEMBER_NAME_CHECKED(UGroomAsset, MinLOD), UGroomAsset::StaticClass()));
 		DetailLayout.HideProperty(DetailLayout.GetProperty(GET_MEMBER_NAME_CHECKED(UGroomAsset, DisableBelowMinLodStripping), UGroomAsset::StaticClass()));
-		DetailLayout.HideProperty(DetailLayout.GetProperty(GET_MEMBER_NAME_CHECKED(UGroomAsset, DeformedSkeletalMesh), UGroomAsset::StaticClass()));
+		DetailLayout.HideProperty(DetailLayout.GetProperty(GET_MEMBER_NAME_CHECKED(UGroomAsset, RiggedSkeletalMesh), UGroomAsset::StaticClass()));
 	}
 	break;
 	case EMaterialPanelType::Strands:
@@ -429,7 +429,7 @@ void FGroomRenderingDetails::CustomizeStrandsGroupProperties(IDetailLayoutBuilde
 		DetailLayout.HideProperty(DetailLayout.GetProperty(GET_MEMBER_NAME_CHECKED(UGroomAsset, HairInterpolationType), UGroomAsset::StaticClass()));
 		DetailLayout.HideProperty(DetailLayout.GetProperty(GET_MEMBER_NAME_CHECKED(UGroomAsset, MinLOD), UGroomAsset::StaticClass()));
 		DetailLayout.HideProperty(DetailLayout.GetProperty(GET_MEMBER_NAME_CHECKED(UGroomAsset, DisableBelowMinLodStripping), UGroomAsset::StaticClass()));
-		DetailLayout.HideProperty(DetailLayout.GetProperty(GET_MEMBER_NAME_CHECKED(UGroomAsset, DeformedSkeletalMesh), UGroomAsset::StaticClass()));
+		DetailLayout.HideProperty(DetailLayout.GetProperty(GET_MEMBER_NAME_CHECKED(UGroomAsset, RiggedSkeletalMesh), UGroomAsset::StaticClass()));
 	}
 	break;
 	case EMaterialPanelType::Physics:
@@ -446,7 +446,7 @@ void FGroomRenderingDetails::CustomizeStrandsGroupProperties(IDetailLayoutBuilde
 		DetailLayout.HideProperty(DetailLayout.GetProperty(GET_MEMBER_NAME_CHECKED(UGroomAsset, HairInterpolationType), UGroomAsset::StaticClass()));
 		DetailLayout.HideProperty(DetailLayout.GetProperty(GET_MEMBER_NAME_CHECKED(UGroomAsset, MinLOD), UGroomAsset::StaticClass()));
 		DetailLayout.HideProperty(DetailLayout.GetProperty(GET_MEMBER_NAME_CHECKED(UGroomAsset, DisableBelowMinLodStripping), UGroomAsset::StaticClass()));
-		DetailLayout.HideProperty(DetailLayout.GetProperty(GET_MEMBER_NAME_CHECKED(UGroomAsset, DeformedSkeletalMesh), UGroomAsset::StaticClass()));
+		DetailLayout.HideProperty(DetailLayout.GetProperty(GET_MEMBER_NAME_CHECKED(UGroomAsset, RiggedSkeletalMesh), UGroomAsset::StaticClass()));
 	}
 	break;
 	case EMaterialPanelType::Interpolation:
@@ -463,7 +463,7 @@ void FGroomRenderingDetails::CustomizeStrandsGroupProperties(IDetailLayoutBuilde
 //		DetailLayout.HideProperty(DetailLayout.GetProperty(GET_MEMBER_NAME_CHECKED(UGroomAsset, HairInterpolationType), UGroomAsset::StaticClass()));
 		DetailLayout.HideProperty(DetailLayout.GetProperty(GET_MEMBER_NAME_CHECKED(UGroomAsset, MinLOD), UGroomAsset::StaticClass()));
 		DetailLayout.HideProperty(DetailLayout.GetProperty(GET_MEMBER_NAME_CHECKED(UGroomAsset, DisableBelowMinLodStripping), UGroomAsset::StaticClass()));
-		//DetailLayout.HideProperty(DetailLayout.GetProperty(GET_MEMBER_NAME_CHECKED(UGroomAsset, DeformedSkeletalMesh), UGroomAsset::StaticClass()));
+		//DetailLayout.HideProperty(DetailLayout.GetProperty(GET_MEMBER_NAME_CHECKED(UGroomAsset, RiggedSkeletalMesh), UGroomAsset::StaticClass()));
 	}
 	break;
 	case EMaterialPanelType::LODs:
@@ -480,7 +480,7 @@ void FGroomRenderingDetails::CustomizeStrandsGroupProperties(IDetailLayoutBuilde
 		DetailLayout.HideProperty(DetailLayout.GetProperty(GET_MEMBER_NAME_CHECKED(UGroomAsset, HairInterpolationType), UGroomAsset::StaticClass()));
 //		DetailLayout.HideProperty(DetailLayout.GetProperty(GET_MEMBER_NAME_CHECKED(UGroomAsset, MinLOD), UGroomAsset::StaticClass()));
 //		DetailLayout.HideProperty(DetailLayout.GetProperty(GET_MEMBER_NAME_CHECKED(UGroomAsset, DisableBelowMinLodStripping), UGroomAsset::StaticClass()));
-		DetailLayout.HideProperty(DetailLayout.GetProperty(GET_MEMBER_NAME_CHECKED(UGroomAsset, DeformedSkeletalMesh), UGroomAsset::StaticClass()));
+		DetailLayout.HideProperty(DetailLayout.GetProperty(GET_MEMBER_NAME_CHECKED(UGroomAsset, RiggedSkeletalMesh), UGroomAsset::StaticClass()));
 	}
 	break;
 	}
@@ -865,9 +865,9 @@ bool FGroomRenderingDetails::CommonResetToDefault(TSharedPtr<IPropertyHandle> Ch
 
 		{
 			FHairDeformationSettings Default;
-			HAIR_RESET1(HairGroupsInterpolation, FHairDeformationSettings, DeformationSettings, NumCurves);
-			HAIR_RESET1(HairGroupsInterpolation, FHairDeformationSettings, DeformationSettings, NumPoints);
-			HAIR_RESET1(HairGroupsInterpolation, FHairDeformationSettings, DeformationSettings, bEnableDeformation);
+			HAIR_RESET1(HairGroupsInterpolation, FHairDeformationSettings, RiggingSettings, NumCurves);
+			HAIR_RESET1(HairGroupsInterpolation, FHairDeformationSettings, RiggingSettings, NumPoints);
+			HAIR_RESET1(HairGroupsInterpolation, FHairDeformationSettings, RiggingSettings, bEnableRigging);
 		}
 	}
 
@@ -1477,7 +1477,7 @@ void FGroomRenderingDetails::OnGenerateElementForHairGroup(TSharedRef<IPropertyH
 			{
 				ExpandStruct(ChildHandle, ChildrenBuilder, GroupIndex, -1, true);
 			}
-			else
+			else 
 			{
 				ChildrenBuilder.AddProperty(ChildHandle.ToSharedRef());
 			}
