Groom: Fixed playback syncing issue between GroomCache and GeometryCache when rendering in MRQ.

diff --git a/Engine/Plugins/Runtime/HairStrands/Source/HairStrandsCore/Private/GroomCache.cpp b/Engine/Plugins/Runtime/HairStrands/Source/HairStrandsCore/Private/GroomCache.cpp
--- a/Engine/Plugins/Runtime/HairStrands/Source/HairStrandsCore/Private/GroomCache.cpp
+++ b/Engine/Plugins/Runtime/HairStrands/Source/HairStrandsCore/Private/GroomCache.cpp
@@ -95,7 +95,7 @@ void UGroomCache::GetFrameIndicesAtTime(float Time, bool bLooping, bool bIsPlayi
 	}
 
 	OutFrameIndex = GetFrameIndexAtTime(Time, bLooping);
-	OutNextFrameIndex = FMath::Min(OutFrameIndex + 1, NumFrames);
+	OutNextFrameIndex = FMath::Min(OutFrameIndex + 1, NumFrames - 1); // -1 since index is 0-based
 
 	const float FrameDuration = GroomCacheInfo.AnimationInfo.SecondsPerFrame;
 	if (FMath::IsNearlyZero(FrameDuration))
diff --git a/Engine/Plugins/Runtime/HairStrands/Source/HairStrandsCore/Private/GroomComponent.cpp b/Engine/Plugins/Runtime/HairStrands/Source/HairStrandsCore/Private/GroomComponent.cpp
--- a/Engine/Plugins/Runtime/HairStrands/Source/HairStrandsCore/Private/GroomComponent.cpp
+++ b/Engine/Plugins/Runtime/HairStrands/Source/HairStrandsCore/Private/GroomComponent.cpp
@@ -33,6 +33,7 @@
 #include "GroomCache.h"
 #include "GroomCacheStreamingManager.h"
 #include "GroomPluginSettings.h"
+#include "Async/Async.h"
 #include "Async/ParallelFor.h"
 #include "PrimitiveSceneInfo.h"
 
@@ -1228,6 +1229,8 @@ public:
 	{
 		TRACE_CPUPROFILER_EVENT_SCOPE(FGroomCacheStreamedBuffers::UpdateBuffersAtTime);
 
+		FScopeLock Lock(GetCriticalSection());
+
 		// Find the frame indices and interpolation factor to interpolate between
 		int32 FrameIndexA = 0;
 		int32 FrameIndexB = 0;
@@ -1241,13 +1244,15 @@ public:
 			bComputeInterpolation = true;
 			if (FrameIndexA == NextFrameIndex)
 			{
-				// At this point, the NextFrame is already mapped so we know the pointer is valid
-				// It is mapped to increment its ref count
+				// NextFrame is mapped to increment its ref count, but this could fail if the frame was evicted
+				// This can happen if the time jumps around as is the case when Sequencer generates thumbnails
 				const FGroomCacheAnimationData* DataPtr = IGroomCacheStreamingManager::Get().MapAnimationData(GroomCache, NextFrameIndex);
-				IGroomCacheStreamingManager::Get().UnmapAnimationData(GroomCache, CurrentFrameIndex);
-
-				CurrentFramePtr = NextFramePtr;
-				CurrentFrameIndex = NextFrameIndex;
+				if (DataPtr)
+				{
+					IGroomCacheStreamingManager::Get().UnmapAnimationData(GroomCache, CurrentFrameIndex);
+					CurrentFramePtr = NextFramePtr;
+					CurrentFrameIndex = NextFrameIndex;
+				}
 			}
 			else
 			{
@@ -3335,7 +3340,21 @@ void UGroomComponent::UpdateGroomCache(float Time)
 		Buffers->UpdateBuffersAtTime(Time, bLooping);
 
 		// Trigger an update of the bounds so that it follows the GroomCache
-		MarkRenderTransformDirty();
+		if (IsInGameThread())
+		{
+			MarkRenderTransformDirty();
+		}
+		else
+		{
+			FWeakObjectPtr WeakComponent(this);
+			AsyncTask(ENamedThreads::GameThread, [WeakComponent]()
+			{
+				if (UGroomComponent* Component = Cast<UGroomComponent>(WeakComponent.Get()))
+				{
+					Component->MarkRenderTransformDirty();
+				}
+			});
+		}
 	}
 }
 
@@ -3396,17 +3415,28 @@ void UGroomComponent::TickAtThisTime(const float Time, bool bInIsRunning, bool b
 	if (GroomCache && bRunning && bManualTick)
 	{
 		float DeltaTime = Time - ElapsedTime;
-		ElapsedTime = Time;
 		if (GUseGroomCacheStreaming)
 		{
 			// Scrubbing forward (or backward) can induce large (or negative) delta time, so force a prefetch
 			if ((DeltaTime > GetDefault<UGroomPluginSettings>()->GroomCacheLookAheadBuffer) ||
 				(DeltaTime < 0))
 			{
+				ElapsedTime = Time;
 				IGroomCacheStreamingManager::Get().PrefetchData(this);
+
+				// Update the buffers with the prefetched data right away
+				UpdateGroomCache(Time);
+				return;
 			}
 		}
-		UpdateGroomCache(Time);
+
+		// Queue the update for the render thread to sync it with GeometryCache rendering
+		ENQUEUE_RENDER_COMMAND(FGroomCacheUpdate)(
+		[this, Time](FRHICommandList& RHICmdList)
+		{
+			ElapsedTime = Time;
+			UpdateGroomCache(Time);
+		});
 	}
 }
 
