Add attributes to the hair niagara DI in order to have the cache working

diff --git a/Engine/Plugins/Runtime/HairStrands/Source/HairStrandsCore/Private/Niagara/NiagaraDataInterfaceHairStrands.cpp b/Engine/Plugins/Runtime/HairStrands/Source/HairStrandsCore/Private/Niagara/NiagaraDataInterfaceHairStrands.cpp
--- a/Engine/Plugins/Runtime/HairStrands/Source/HairStrandsCore/Private/Niagara/NiagaraDataInterfaceHairStrands.cpp
+++ b/Engine/Plugins/Runtime/HairStrands/Source/HairStrandsCore/Private/Niagara/NiagaraDataInterfaceHairStrands.cpp
@@ -8,6 +8,7 @@
 #include "NiagaraSimStageData.h"
 #include "NiagaraShaderParametersBuilder.h"
 #include "NiagaraSystemInstance.h"
+#include "NiagaraConstants.h"
 
 #include "Components/SkeletalMeshComponent.h"
 #include "ShaderParameterUtils.h"
@@ -915,7 +916,7 @@ static void InterpolateGroomGuides(FRDGBuilder& GraphBuilder, FNiagaraDataBuffer
 	const bool bIsRestValid = bIsHairValid && HairStrandsBuffer->SourceRestResources && HairStrandsBuffer->SourceRestResources->IsInitialized();
 	const bool bIsDeformedValid = bIsHairValid && HairStrandsBuffer->SourceDeformedResources && HairStrandsBuffer->SourceDeformedResources->IsInitialized();
 
-	if(bIsRestValid && bIsDeformedValid && bHasValidGeometry)
+	if(bIsRestValid && bIsDeformedValid && bHasValidGeometry && HairStrandsBuffer->PointsCurveBuffer.IsValid() && HairStrandsBuffer->CurvesOffsetsBuffer.IsValid())
 	{
 		bool bIsRootValid = HairStrandsBuffer->SourceDeformedRootResources && HairStrandsBuffer->SourceDeformedRootResources->IsInitialized() &&
 								HairStrandsBuffer->SourceRestRootResources && HairStrandsBuffer->SourceRestRootResources->IsInitialized() && bHasSkinningBinding;
@@ -985,17 +986,18 @@ void UNiagaraDataInterfaceHairStrands::SimCachePostReadFrame(void* OptionalPerIn
 	{
 		FNiagaraDataSet& EmitterDatas  = EmitterInst->GetData();
 
-		const FNiagaraVariableBase NodePositionVariable(FNiagaraDataSetAccessorTypeInfo<FVector3f>::GetFloatType(), FName("NodePosition"));
+		const FNiagaraVariableBase NodePositionVariable(FNiagaraTypeDefinition::GetVec3Def(), FName("NodePosition"));
 		const int32 NodePositionIndex = EmitterDatas.GetCompiledData().Variables.IndexOfByKey(NodePositionVariable);
 
-		const FNiagaraVariableBase RestPositionVariable(FNiagaraDataSetAccessorTypeInfo<FVector3f>::GetFloatType(), FName("RestPosition"));
+		const FNiagaraVariableBase RestPositionVariable(FNiagaraTypeDefinition::GetVec3Def(), FName("RestPosition"));
 		const int32 RestPositionIndex = EmitterDatas.GetCompiledData().Variables.IndexOfByKey(RestPositionVariable);
 				
 		if (RestPositionIndex != INDEX_NONE && NodePositionIndex != INDEX_NONE)
 		{
-			const FNDIHairStrandsData* InstanceData = static_cast<FNDIHairStrandsData*>(OptionalPerInstanceData);
+			FNDIHairStrandsData* InstanceData = static_cast<FNDIHairStrandsData*>(OptionalPerInstanceData);
+			
 			if( EmitterDatas.GetCompiledData().VariableLayouts.IsValidIndex(NodePositionIndex) &&
-				EmitterDatas.GetCompiledData().VariableLayouts.IsValidIndex(RestPositionIndex) )
+				EmitterDatas.GetCompiledData().VariableLayouts.IsValidIndex(RestPositionIndex) && InstanceData)
 			{
 				const uint32 NodePositionComponent = EmitterDatas.GetCompiledData().VariableLayouts[NodePositionIndex].FloatComponentStart;
 				const uint32 RestPositionComponent = EmitterDatas.GetCompiledData().VariableLayouts[RestPositionIndex].FloatComponentStart;
@@ -1003,24 +1005,42 @@ void UNiagaraDataInterfaceHairStrands::SimCachePostReadFrame(void* OptionalPerIn
 				FNiagaraDataBuffer* ParticlesBuffer = EmitterDatas.GetCurrentData();
 				FNDIHairStrandsBuffer* HairstrandsBuffer = InstanceData->HairStrandsBuffer;
 
-				const uint32 StrandsSize = InstanceData->StrandsSize;
-				const bool bHasSkinningBinding = InstanceData->HairGroupInstance->BindingType == EHairBindingType::Skinning;
-				const bool bHasValidGeometry = InstanceData->HairGroupInstance->GeometryType != EHairGeometryType::NoneGeometry;
+				if(HairstrandsBuffer && InstanceData->HairGroupInstance)
+				{
+					const uint32 StrandsSize = InstanceData->StrandsSize;
+					const bool bHasSkinningBinding = InstanceData->HairGroupInstance->BindingType == EHairBindingType::Skinning;
+					const bool bHasValidGeometry = InstanceData->HairGroupInstance->GeometryType != EHairGeometryType::NoneGeometry;
 				
-				ENQUEUE_RENDER_COMMAND(NiagaraInterpolateGroomSimCache)(
-					[ParticlesBuffer, HairstrandsBuffer, NodePositionComponent, RestPositionComponent, StrandsSize, bHasSkinningBinding, bHasValidGeometry](FRHICommandListImmediate& RHICmdList)
-					{
-						FMemMark MemMark(FMemStack::Get());
-						FRDGBuilder GraphBuilder(RHICmdList);
-						InterpolateGroomGuides(GraphBuilder, ParticlesBuffer, NodePositionComponent, RestPositionComponent, HairstrandsBuffer, StrandsSize, bHasSkinningBinding, bHasValidGeometry);
-						GraphBuilder.Execute();
-					}
-				);
+					ENQUEUE_RENDER_COMMAND(NiagaraInterpolateGroomSimCache)(
+						[ParticlesBuffer, HairstrandsBuffer, NodePositionComponent, RestPositionComponent, StrandsSize, bHasSkinningBinding, bHasValidGeometry](FRHICommandListImmediate& RHICmdList)
+						{
+							FMemMark MemMark(FMemStack::Get());
+							FRDGBuilder GraphBuilder(RHICmdList);
+							InterpolateGroomGuides(GraphBuilder, ParticlesBuffer, NodePositionComponent, RestPositionComponent, HairstrandsBuffer, StrandsSize, bHasSkinningBinding, bHasValidGeometry);
+							GraphBuilder.Execute();
+
+							HairstrandsBuffer->PointsCurveBuffer.EndGraphUsage();
+							HairstrandsBuffer->CurvesOffsetsBuffer.EndGraphUsage();
+						}
+					);
+				}
 			}
 		}
 	}
 }
 
+TArray<FNiagaraVariableBase> UNiagaraDataInterfaceHairStrands::GetSimCacheRendererAttributes(UObject* UsageContext) const
+{
+	TArray<FNiagaraVariableBase> HairStrandsCachedVariables;
+	if (const UNiagaraEmitter* UsageEmitter = Cast<UNiagaraEmitter>(UsageContext))
+	{
+		HairStrandsCachedVariables.Emplace(FNiagaraTypeDefinition::GetVec3Def(), FName(UsageEmitter->GetUniqueEmitterName() + TEXT(".") + FNiagaraConstants::ParticleAttributeNamespaceString + TEXT(".") + TEXT("NodePosition")));
+		HairStrandsCachedVariables.Emplace(FNiagaraTypeDefinition::GetVec3Def(), FName(UsageEmitter->GetUniqueEmitterName() + TEXT(".") + FNiagaraConstants::ParticleAttributeNamespaceString + TEXT(".") + TEXT("RestPosition")));
+	}
+
+	return HairStrandsCachedVariables;
+}
+
 bool UNiagaraDataInterfaceHairStrands::CopyToInternal(UNiagaraDataInterface* Destination) const
 {
 	if (!Super::CopyToInternal(Destination))
@@ -3361,7 +3381,7 @@ void UNiagaraDataInterfaceHairStrands::SetShaderParameters(const FNiagaraDataInt
 	const bool bIsHairValid = ProxyData != nullptr && ProxyData->HairStrandsBuffer && ProxyData->HairStrandsBuffer->IsInitialized();
 	const bool bIsHairGroupInstValid = ProxyData != nullptr && ProxyData->HairGroupInstSource != nullptr && ProxyData->HairGroupInstSource->ContainsGroupInstance(ProxyData->HairGroupInstance);
 	const bool bHasSkinningBinding = bIsHairValid && bIsHairGroupInstValid && ProxyData->HairGroupInstance->BindingType == EHairBindingType::Skinning;
-	const bool bIsRootValid = bIsHairValid && ProxyData->HairStrandsBuffer->SourceDeformedRootResources && ProxyData->HairStrandsBuffer->SourceDeformedRootResources->IsInitialized() && bHasSkinningBinding;
+	const bool bIsRootValid = bIsHairValid && ProxyData->HairStrandsBuffer->SourceDeformedRootResources&& ProxyData->HairStrandsBuffer->SourceDeformedRootResources->IsInitialized() && bHasSkinningBinding;
 	const bool bIsRestValid = bIsHairValid && ProxyData->HairStrandsBuffer->SourceRestResources && ProxyData->HairStrandsBuffer->SourceRestResources->IsInitialized()&&
 		// TEMP: These check are only temporary for avoiding crashes while we find the bottom of the issue.
 		ProxyData->HairStrandsBuffer->CurvesOffsetsBuffer.IsValid() && ProxyData->HairStrandsBuffer->ParamsScaleBuffer.IsValid() && ProxyData->HairStrandsBuffer->BoundingBoxBuffer.IsValid();
diff --git a/Engine/Plugins/Runtime/HairStrands/Source/HairStrandsCore/Public/Niagara/NiagaraDataInterfaceHairStrands.h b/Engine/Plugins/Runtime/HairStrands/Source/HairStrandsCore/Public/Niagara/NiagaraDataInterfaceHairStrands.h
--- a/Engine/Plugins/Runtime/HairStrands/Source/HairStrandsCore/Public/Niagara/NiagaraDataInterfaceHairStrands.h
+++ b/Engine/Plugins/Runtime/HairStrands/Source/HairStrandsCore/Public/Niagara/NiagaraDataInterfaceHairStrands.h
@@ -416,6 +416,7 @@ public:
 	virtual bool HasTickGroupPrereqs() const override { return true; }
 	virtual ETickingGroup CalculateTickGroup(const void* PerInstanceData) const override;
 	virtual void SimCachePostReadFrame(void* OptionalPerInstanceData, FNiagaraSystemInstance* SystemInstance) override;
+	virtual TArray<FNiagaraVariableBase> GetSimCacheRendererAttributes(UObject* UsageContext) const override;
 
 	/** GPU simulation  functionality */
 #if WITH_EDITORONLY_DATA
