Moved the GroomCache streamer registration to its own runtime module since HairStrandsCore uses the PostConfigInit, which happens before the RHI is initialized, but FStreamingManagerCollection has a dependency on the RHI. This is needed to support remote MRQ, which uses -game.

diff --git a/Engine/Plugins/Runtime/HairStrands/HairStrands.uplugin b/Engine/Plugins/Runtime/HairStrands/HairStrands.uplugin
--- a/Engine/Plugins/Runtime/HairStrands/HairStrands.uplugin
+++ b/Engine/Plugins/Runtime/HairStrands/HairStrands.uplugin
@@ -20,6 +20,11 @@
 			"Type": "Runtime",
 			"LoadingPhase": "PostConfigInit"
 		},
+		{
+			"Name": "HairStrandsRuntime",
+			"Type": "Runtime",
+			"LoadingPhase": "Default"
+		},
 		{
 			"Name": "HairStrandsEditor",
 			"Type": "Editor",
diff --git a/Engine/Plugins/Runtime/HairStrands/Source/HairStrandsEditor/Private/HairStrandsEditor.cpp b/Engine/Plugins/Runtime/HairStrands/Source/HairStrandsEditor/Private/HairStrandsEditor.cpp
--- a/Engine/Plugins/Runtime/HairStrands/Source/HairStrandsEditor/Private/HairStrandsEditor.cpp
+++ b/Engine/Plugins/Runtime/HairStrands/Source/HairStrandsEditor/Private/HairStrandsEditor.cpp
@@ -18,7 +18,6 @@
 #include "GroomBindingDetailsCustomization.h"
 #include "GroomCacheImportOptions.h"
 #include "GroomCacheImportSettingsCustomization.h"
-#include "GroomCacheStreamingManager.h"
 #include "GroomCacheTrackEditor.h"
 #include "GroomComponentDetailsCustomization.h"
 #include "GroomCreateBindingOptions.h"
@@ -149,14 +148,10 @@ void FGroomEditor::StartupModule()
 			GetMutableDefault<UGroomPluginSettings>()
 		);
 	}
-
-	IGroomCacheStreamingManager::Register();
 }
 
 void FGroomEditor::ShutdownModule()
 {
-	IGroomCacheStreamingManager::Unregister();
-
 	ISettingsModule* SettingsModule = FModuleManager::GetModulePtr<ISettingsModule>("Settings");
 	if (SettingsModule != nullptr)
 	{
diff --git a/Engine/Plugins/Runtime/HairStrands/Source/HairStrandsRuntime/HairStrandsRuntime.Build.cs b/Engine/Plugins/Runtime/HairStrands/Source/HairStrandsRuntime/HairStrandsRuntime.Build.cs
new file mode 100644
--- /dev/null
+++ b/Engine/Plugins/Runtime/HairStrands/Source/HairStrandsRuntime/HairStrandsRuntime.Build.cs
@@ -0,0 +1,18 @@
+// Copyright Epic Games, Inc. All Rights Reserved.
+
+namespace UnrealBuildTool.Rules
+{
+	public class HairStrandsRuntime : ModuleRules
+	{
+		public HairStrandsRuntime(ReadOnlyTargetRules Target) : base(Target)
+		{
+			PrivateDependencyModuleNames.AddRange(
+				new string[]
+				{
+					"Core",
+					"Engine",
+					"HairStrandsCore"
+				});
+		}
+	}
+}
diff --git a/Engine/Plugins/Runtime/HairStrands/Source/HairStrandsRuntime/Private/HairStrandsRuntime.cpp b/Engine/Plugins/Runtime/HairStrands/Source/HairStrandsRuntime/Private/HairStrandsRuntime.cpp
new file mode 100644
--- /dev/null
+++ b/Engine/Plugins/Runtime/HairStrands/Source/HairStrandsRuntime/Private/HairStrandsRuntime.cpp
@@ -0,0 +1,30 @@
+// Copyright Epic Games, Inc. All Rights Reserved.
+
+#include "GroomCacheStreamingManager.h"
+#include "Modules/ModuleInterface.h"
+#include "Modules/ModuleManager.h"
+
+/*
+ * The HairStrandsRuntime module exists solely because it needs to be initialized after the RHI
+ * but the HairStrandsCore module has to be at PostConfigInit (before RHI is initialized)
+ * because it contains shaders that need that loading phase.
+ */
+class FHairStrandsRuntime : public IModuleInterface
+{
+public:
+	//~ IModuleInterface interface
+	virtual void StartupModule() override
+	{
+		// The GroomCache streaming manager has a RHI dependency because of 
+		// FStreamingManagerCollection::AddOrRemoveTextureStreamingManagerIfNeeded
+		// which is called when FStreamingManagerCollection is created via IStreamingManager::Get()
+		IGroomCacheStreamingManager::Register();
+	}
+
+	virtual void ShutdownModule() override
+	{
+		IGroomCacheStreamingManager::Unregister();
+	}
+};
+
+IMPLEMENT_MODULE(FHairStrandsRuntime, HairStrandsRuntime);
