Add support for PSO precaching from GroomAssets if IsResourcePSOPrecachingEnabled is set

#rnx

diff --git a/Engine/Plugins/Runtime/HairStrands/Source/HairStrandsCore/Private/GroomAsset.cpp b/Engine/Plugins/Runtime/HairStrands/Source/HairStrandsCore/Private/GroomAsset.cpp
--- a/Engine/Plugins/Runtime/HairStrands/Source/HairStrandsCore/Private/GroomAsset.cpp
+++ b/Engine/Plugins/Runtime/HairStrands/Source/HairStrandsCore/Private/GroomAsset.cpp
@@ -32,6 +32,8 @@
 #include "UObject/ObjectSaveContext.h"
 #include "GroomBindingAsset.h"
 #include "GroomDeformerBuilder.h"
+#include "HairCardsVertexFactory.h"
+#include "PSOPrecache.h"
 
 #include UE_INLINE_GENERATED_CPP_BY_NAME(GroomAsset)
 
@@ -783,6 +785,89 @@ bool UGroomAsset::HasGeometryType(EGroomGeometryType Type) const
 	return false;
 }
 
+TArray<FHairVertexFactoryTypesPerMaterialData> UGroomAsset::CollectVertexFactoryTypesPerMaterialData(EShaderPlatform ShaderPlatform)
+{
+	TArray<FHairVertexFactoryTypesPerMaterialData> VFsPerMaterials;
+
+	const auto AddVFPerMaterial = [&VFsPerMaterials](int32 InMaterialIndex, EHairGeometryType InHairGeometryType, const FVertexFactoryType* InVFType)
+	{
+		if (InMaterialIndex < 0)
+		{
+			return;
+		}
+
+		FHairVertexFactoryTypesPerMaterialData* VFsPerMaterial = VFsPerMaterials.FindByPredicate([InMaterialIndex, InHairGeometryType](const FHairVertexFactoryTypesPerMaterialData& Other) { return (Other.MaterialIndex == InMaterialIndex && Other.HairGeometryType == InHairGeometryType); });
+		if (VFsPerMaterial == nullptr)
+		{
+			VFsPerMaterial = &VFsPerMaterials.AddDefaulted_GetRef();
+			VFsPerMaterial->MaterialIndex = InMaterialIndex;
+			VFsPerMaterial->HairGeometryType = InHairGeometryType;
+		}
+		VFsPerMaterial->VertexFactoryTypes.AddUnique(InVFType);
+	};
+
+	const int32 GroupCount = GetNumHairGroups();
+	for (int32 GroupIt = 0; GroupIt < GroupCount; GroupIt++)
+	{
+		const FHairGroupData& InGroupData = HairGroupsData[GroupIt];
+
+		// Strands
+		{
+			const int32 MaterialIndex = GetMaterialIndex(HairGroupsRendering[GroupIt].MaterialSlotName);
+			AddVFPerMaterial(MaterialIndex, EHairGeometryType::Strands, &FHairStrandsVertexFactory::StaticType);
+		}
+
+		// Cards
+		if (IsHairStrandsEnabled(EHairStrandsShaderType::Cards, ShaderPlatform))
+		{
+			uint32 CardsLODIndex = 0;
+			for (const FHairGroupData::FCards::FLOD& LOD : InGroupData.Cards.LODs)
+			{
+				if (LOD.IsValid())
+				{
+					// Material
+					int32 MaterialIndex = INDEX_NONE;
+					for (const FHairGroupsCardsSourceDescription& Desc : HairGroupsCards)
+					{
+						if (Desc.GroupIndex == GroupIt && Desc.LODIndex == CardsLODIndex)
+						{
+							MaterialIndex = GetMaterialIndex(Desc.MaterialSlotName);
+							break;
+						}
+					}
+					AddVFPerMaterial(MaterialIndex, EHairGeometryType::Cards, &FHairCardsVertexFactory::StaticType);
+				}
+				++CardsLODIndex;
+			}
+		}
+
+		// Meshes
+		{
+			uint32 MeshesLODIndex = 0;
+			for (const FHairGroupData::FMeshes::FLOD& LOD : InGroupData.Meshes.LODs)
+			{
+				if (LOD.IsValid())
+				{
+					// Material
+					int32 MaterialIndex = INDEX_NONE;
+					for (const FHairGroupsMeshesSourceDescription& Desc : HairGroupsMeshes)
+					{
+						if (Desc.GroupIndex == GroupIt && Desc.LODIndex == MeshesLODIndex)
+						{
+							MaterialIndex = GetMaterialIndex(Desc.MaterialSlotName);
+							break;
+						}
+					}
+					AddVFPerMaterial(MaterialIndex, EHairGeometryType::Meshes, &FHairCardsVertexFactory::StaticType);
+				}
+				++MeshesLODIndex;
+			}
+		}
+	}
+
+	return VFsPerMaterials;
+}
+
 enum EGroomAssetChangeType
 {
 	GroomChangeType_Interpolation = 1,
@@ -1188,6 +1273,24 @@ void UGroomAsset::PostLoad()
 #if WITH_EDITORONLY_DATA
 	UpdateCachedSettings();
 #endif // #if WITH_EDITORONLY_DATA
+
+	if (IsResourcePSOPrecachingEnabled())
+	{
+		ERHIFeatureLevel::Type FeatureLevel = GetWorld() ? GetWorld()->FeatureLevel.GetValue() : GMaxRHIFeatureLevel;
+		EShaderPlatform ShaderPlatform = GetFeatureLevelShaderPlatform(FeatureLevel);
+
+		TArray<FHairVertexFactoryTypesPerMaterialData> VFsPerMaterials = CollectVertexFactoryTypesPerMaterialData(ShaderPlatform);
+
+		FPSOPrecacheParams PrecachePSOParams;
+		for (FHairVertexFactoryTypesPerMaterialData& VFsPerMaterial : VFsPerMaterials)
+		{
+			UMaterialInterface* MaterialInterface = HairGroupsMaterials[VFsPerMaterial.MaterialIndex].Material;
+			if (MaterialInterface)
+			{
+				MaterialInterface->PrecachePSOs(VFsPerMaterial.VertexFactoryTypes, PrecachePSOParams);
+			}
+		}
+	}
 }
 
 void UGroomAsset::PreSave(const class ITargetPlatform* TargetPlatform)
diff --git a/Engine/Plugins/Runtime/HairStrands/Source/HairStrandsCore/Private/GroomComponent.cpp b/Engine/Plugins/Runtime/HairStrands/Source/HairStrandsCore/Private/GroomComponent.cpp
--- a/Engine/Plugins/Runtime/HairStrands/Source/HairStrandsCore/Private/GroomComponent.cpp
+++ b/Engine/Plugins/Runtime/HairStrands/Source/HairStrandsCore/Private/GroomComponent.cpp
@@ -36,6 +36,7 @@
 #include "Async/Async.h"
 #include "Async/ParallelFor.h"
 #include "PrimitiveSceneInfo.h"
+#include "PSOPrecache.h"
 
 #include UE_INLINE_GENERATED_CPP_BY_NAME(GroomComponent)
 LLM_DECLARE_TAG(Groom);
@@ -3112,100 +3113,17 @@ void UGroomComponent::PostLoad()
 
 void UGroomComponent::PrecachePSOs()
 {
-	if (GroomAsset == nullptr || !FApp::CanEverRender() || !PipelineStateCache::IsPSOPrecachingEnabled())
+	if (!IsComponentPSOPrecachingEnabled() || GroomAsset == nullptr)
 	{
 		return;
 	}
 
-	struct VFsPerMaterialData
-	{
-		int16 MaterialIndex;
-		EHairGeometryType HairGeometryType;
-		TArray<const FVertexFactoryType*, TInlineAllocator<2>> VertexFactoryTypes;
-	};
-	TArray<VFsPerMaterialData, TInlineAllocator<2>> VFsPerMaterials;
-
-	const auto AddVFPerMaterial = [&VFsPerMaterials](int32 InMaterialIndex, EHairGeometryType InHairGeometryType, const FVertexFactoryType* InVFType)
-	{
-		VFsPerMaterialData* VFsPerMaterial = VFsPerMaterials.FindByPredicate([InMaterialIndex, InHairGeometryType](const VFsPerMaterialData& Other) { return (Other.MaterialIndex == InMaterialIndex && Other.HairGeometryType == InHairGeometryType); });
-		if (VFsPerMaterial == nullptr)
-		{
-			VFsPerMaterial = &VFsPerMaterials.AddDefaulted_GetRef();
-			VFsPerMaterial->MaterialIndex = InMaterialIndex;
-			VFsPerMaterial->HairGeometryType = InHairGeometryType;
-		}
-		VFsPerMaterial->VertexFactoryTypes.AddUnique(InVFType);
-	};
-
-	const EShaderPlatform ShaderPlatform = GetScene() ? GetScene()->GetShaderPlatform() : GMaxRHIShaderPlatform;
-	const int32 GroupCount = GroomAsset->GetNumHairGroups();
-	check(GroomAsset->HairGroupsData.Num() == HairGroupInstances.Num());
-	for (int32 GroupIt = 0; GroupIt < GroupCount; GroupIt++)
-	{
-		const FHairGroupData& InGroupData = GroomAsset->HairGroupsData[GroupIt];
-		FHairGroupInstance* HairInstance = HairGroupInstances[GroupIt];
-		check(HairInstance->HairGroupPublicData);
-
-		// Strands
-		if (HairInstance->Strands.IsValid())
-		{
-			const int32 MaterialIndex = GroomAsset->GetMaterialIndex(GroomAsset->HairGroupsRendering[GroupIt].MaterialSlotName);
-			AddVFPerMaterial(MaterialIndex, EHairGeometryType::Strands, &FHairStrandsVertexFactory::StaticType);
-		}
-
-		// Cards
-		if (HairInstance->Cards.IsValid() && IsHairStrandsEnabled(EHairStrandsShaderType::Cards, ShaderPlatform))
-		{
-			uint32 CardsLODIndex = 0;
-			for (const FHairGroupData::FCards::FLOD& LOD : InGroupData.Cards.LODs)
-			{
-				if (LOD.IsValid())
-				{
-					// Material
-					int32 MaterialIndex = INDEX_NONE;
-					for (const FHairGroupsCardsSourceDescription& Desc : GroomAsset->HairGroupsCards)
-					{
-						if (Desc.GroupIndex == GroupIt && Desc.LODIndex == CardsLODIndex)
-						{
-							MaterialIndex = GroomAsset->GetMaterialIndex(Desc.MaterialSlotName);
-							break;
-						}
-					}
-					AddVFPerMaterial(MaterialIndex, EHairGeometryType::Cards, &FHairCardsVertexFactory::StaticType);
-				}
-				++CardsLODIndex;
-			}
-		}
-
-		// Meshes
-		if (HairInstance->Meshes.IsValid())
-		{
-			uint32 MeshesLODIndex = 0;
-			for (const FHairGroupData::FMeshes::FLOD& LOD : InGroupData.Meshes.LODs)
-			{
-				if (LOD.IsValid())
-				{
-					// Material
-					int32 MaterialIndex = INDEX_NONE;
-					for (const FHairGroupsMeshesSourceDescription& Desc : GroomAsset->HairGroupsMeshes)
-					{
-						if (Desc.GroupIndex == GroupIt && Desc.LODIndex == MeshesLODIndex)
-						{
-							MaterialIndex = GroomAsset->GetMaterialIndex(Desc.MaterialSlotName);
-							break;
-						}
-					}
-					AddVFPerMaterial(MaterialIndex, EHairGeometryType::Meshes, &FHairCardsVertexFactory::StaticType);
-				}
-				++MeshesLODIndex;
-			}
-		}
-	}
+	TArray<FHairVertexFactoryTypesPerMaterialData> VFsPerMaterials = GroomAsset->CollectVertexFactoryTypesPerMaterialData(GMaxRHIShaderPlatform);
 
 	FPSOPrecacheParams PrecachePSOParams;
 	SetupPrecachePSOParams(PrecachePSOParams);
 
-	for (VFsPerMaterialData& VFsPerMaterial : VFsPerMaterials)
+	for (FHairVertexFactoryTypesPerMaterialData& VFsPerMaterial : VFsPerMaterials)
 	{
 		UMaterialInterface* MaterialInterface = GetMaterial(GetMaterialIndexWithFallback(VFsPerMaterial.MaterialIndex), VFsPerMaterial.HairGeometryType, true);
 		if (MaterialInterface)
diff --git a/Engine/Plugins/Runtime/HairStrands/Source/HairStrandsCore/Public/GroomAsset.h b/Engine/Plugins/Runtime/HairStrands/Source/HairStrandsCore/Public/GroomAsset.h
--- a/Engine/Plugins/Runtime/HairStrands/Source/HairStrandsCore/Public/GroomAsset.h
+++ b/Engine/Plugins/Runtime/HairStrands/Source/HairStrandsCore/Public/GroomAsset.h
@@ -403,6 +403,13 @@ enum class EHairAtlasTextureType : uint8
 	Material
 };
 
+struct FHairVertexFactoryTypesPerMaterialData
+{
+	int16 MaterialIndex;
+	EHairGeometryType HairGeometryType;
+	TArray<const FVertexFactoryType*, TInlineAllocator<2>> VertexFactoryTypes;
+};
+
 /**
  * Implements an asset that can be used to store hair strands
  */
@@ -572,6 +579,10 @@ public:
 	void UpdateHairGroupsInfo();
 	bool HasGeometryType(EGroomGeometryType Type) const;
 	bool HasGeometryType(uint32 GroupIndex, EGroomGeometryType Type) const;
+
+	/** Used for PSO precaching of used materials and vertex factories */
+	TArray<FHairVertexFactoryTypesPerMaterialData> CollectVertexFactoryTypesPerMaterialData(EShaderPlatform ShaderPlatform);
+
 //private :
 #if WITH_EDITOR
 	FOnGroomAssetChanged OnGroomAssetChanged;
