Fix missing inverse world matrix in groom cache interpolation

diff --git a/Engine/Plugins/Runtime/HairStrands/Shaders/Private/NiagaraInterpolateGroomGuides.usf b/Engine/Plugins/Runtime/HairStrands/Shaders/Private/NiagaraInterpolateGroomGuides.usf
--- a/Engine/Plugins/Runtime/HairStrands/Shaders/Private/NiagaraInterpolateGroomGuides.usf
+++ b/Engine/Plugins/Runtime/HairStrands/Shaders/Private/NiagaraInterpolateGroomGuides.usf
@@ -32,6 +32,7 @@ Buffer<uint4>		RestPositionBuffer;
 Buffer<uint>		CurvesOffsetsBuffer;
 Buffer<float4>		DeformedPositionOffset;
 float3				RestPositionOffset;
+float4x4			WorldToLocal;
 
 int					StrandsSize;
 int					NumPoints;
@@ -90,23 +91,23 @@ void MainCS(uint3 DispatchThreadId : SV_DispatchThreadID)
 		const int NodePositionYOffset = NodePositionXOffset + NiagaraFloatStride;
 		const int NodePositionZOffset = NodePositionYOffset + NiagaraFloatStride;
 		
-		const float3 WorldPrevNodePosition = float3(NiagaraFloatBuffer[NodePositionXOffset+NodePrev],
+		const float3 LocalPrevNodePosition = mul(float4(NiagaraFloatBuffer[NodePositionXOffset+NodePrev],
 		                                   NiagaraFloatBuffer[NodePositionYOffset+NodePrev],
-		                                   NiagaraFloatBuffer[NodePositionZOffset+NodePrev]);
+		                                   NiagaraFloatBuffer[NodePositionZOffset+NodePrev],1), WorldToLocal).xyz;
 
-		const float3 WorldNextNodePosition = float3(NiagaraFloatBuffer[NodePositionXOffset+NodeNext],
+		const float3 LocalNextNodePosition = mul(float4(NiagaraFloatBuffer[NodePositionXOffset+NodeNext],
 										   NiagaraFloatBuffer[NodePositionYOffset+NodeNext],
-										   NiagaraFloatBuffer[NodePositionZOffset+NodeNext]);
+										   NiagaraFloatBuffer[NodePositionZOffset+NodeNext],1), WorldToLocal).xyz;
 
 		const int RestPositionXOffset = RestPositionComponent * NiagaraFloatStride + StrandIndex * StrandsSize;
 		const int RestPositionYOffset = RestPositionXOffset + NiagaraFloatStride;
 		const int RestPositionZOffset = RestPositionYOffset + NiagaraFloatStride;
 		
-		const float3 WorldPrevRestPosition = float3(NiagaraFloatBuffer[RestPositionXOffset+NodePrev],
+		const float3 LocalPrevRestPosition = float3(NiagaraFloatBuffer[RestPositionXOffset+NodePrev],
 										   NiagaraFloatBuffer[RestPositionYOffset+NodePrev],
 										   NiagaraFloatBuffer[RestPositionZOffset+NodePrev]);
 
-		const float3 WorldNextRestPosition = float3(NiagaraFloatBuffer[RestPositionXOffset+NodeNext],
+		const float3 LocalNextRestPosition = float3(NiagaraFloatBuffer[RestPositionXOffset+NodeNext],
 										   NiagaraFloatBuffer[RestPositionYOffset+NodeNext],
 										   NiagaraFloatBuffer[RestPositionZOffset+NodeNext]);
 
@@ -121,13 +122,13 @@ void MainCS(uint3 DispatchThreadId : SV_DispatchThreadID)
 		DIHairStrands_BuildRestTriangle(ProjectionUV, TriangleIndex, StrandIndex, RestTrianglePosition, RestTriangleOrientation);
 		DIHairStrands_BuildDeformedTriangle(ProjectionUV, TriangleIndex, StrandIndex, DeformedTrianglePosition, DeformedTriangleOrientation);
 		
-		const float3 LocalPrevRestPosition = DIHairStrands_TriangleLocalPosition(RestTrianglePosition, RestTriangleOrientation, WorldPrevRestPosition);
-		const float3 LocalPrevNodePosition = DIHairStrands_TriangleLocalPosition(DeformedTrianglePosition, DeformedTriangleOrientation, WorldPrevNodePosition);
+		const float3 TrianglePrevRestPosition = DIHairStrands_TriangleLocalPosition(RestTrianglePosition, RestTriangleOrientation, LocalPrevRestPosition);
+		const float3 TrianglePrevNodePosition = DIHairStrands_TriangleLocalPosition(DeformedTrianglePosition, DeformedTriangleOrientation, LocalPrevNodePosition);
 			
-		const float3 LocalNextRestPosition = DIHairStrands_TriangleLocalPosition(RestTrianglePosition, RestTriangleOrientation, WorldNextRestPosition);
-		const float3 LocalNextNodePosition = DIHairStrands_TriangleLocalPosition(DeformedTrianglePosition, DeformedTriangleOrientation, WorldNextNodePosition);
+		const float3 TriangleNextRestPosition = DIHairStrands_TriangleLocalPosition(RestTrianglePosition, RestTriangleOrientation, LocalNextRestPosition);
+		const float3 TriangleNextNodePosition = DIHairStrands_TriangleLocalPosition(DeformedTrianglePosition, DeformedTriangleOrientation, LocalNextNodePosition);
 
-		const float3 LocalPointDisplace = (LocalPrevNodePosition - LocalPrevRestPosition)  * (1.0 - PointAlpha) + (LocalNextNodePosition - LocalNextRestPosition) * PointAlpha;
+		const float3 LocalPointDisplace = (TrianglePrevNodePosition - TrianglePrevRestPosition)  * (1.0 - PointAlpha) + (TriangleNextNodePosition - TriangleNextRestPosition) * PointAlpha;
 		
 		const float3 LocalTrianglePosition = DIHairStrands_TriangleLocalPosition(RestTrianglePosition, RestTriangleOrientation, 
 					f16tof32(RestPositionBuffer[DispatchThreadId.x].xyz) + RestPositionOffset) + LocalPointDisplace;
@@ -139,10 +140,10 @@ void MainCS(uint3 DispatchThreadId : SV_DispatchThreadID)
 		
 #if PERMUTATION_INTERPOLATION == WORLD_INTERPOLATION
 		
-		const float3 WorldPointDisplace = (WorldPrevNodePosition - WorldPrevRestPosition)  * (1.0 - PointAlpha) + (WorldNextNodePosition - WorldNextRestPosition) * PointAlpha;
+		const float3 LocalPointDisplace = (LocalPrevNodePosition - LocalPrevRestPosition)  * (1.0 - PointAlpha) + (LocalNextNodePosition - LocalNextRestPosition) * PointAlpha;
 
-		const float3 WorldComponentPosition = f16tof32(RestPositionBuffer[DispatchThreadId.x].xyz) + WorldPointDisplace + RestPositionOffset - DeformedPositionOffset[0].xyz;
-		DeformedPositionBuffer[DispatchThreadId.x] = uint4(f32tof16(WorldComponentPosition), RestPositionBuffer[DispatchThreadId.x].w);
+		const float3 LocalComponentPosition = f16tof32(RestPositionBuffer[DispatchThreadId.x].xyz) + LocalPointDisplace + RestPositionOffset - DeformedPositionOffset[0].xyz;
+		DeformedPositionBuffer[DispatchThreadId.x] = uint4(f32tof16(LocalComponentPosition), RestPositionBuffer[DispatchThreadId.x].w);
 		
 #endif
 	}
diff --git a/Engine/Plugins/Runtime/HairStrands/Source/HairStrandsCore/Private/Niagara/NiagaraDataInterfaceHairStrands.cpp b/Engine/Plugins/Runtime/HairStrands/Source/HairStrandsCore/Private/Niagara/NiagaraDataInterfaceHairStrands.cpp
--- a/Engine/Plugins/Runtime/HairStrands/Source/HairStrandsCore/Private/Niagara/NiagaraDataInterfaceHairStrands.cpp
+++ b/Engine/Plugins/Runtime/HairStrands/Source/HairStrandsCore/Private/Niagara/NiagaraDataInterfaceHairStrands.cpp
@@ -881,6 +881,7 @@ class FInterpolateGroomGuidesCS : public FGlobalShader
 		SHADER_PARAMETER_RDG_BUFFER_SRV(Buffer<uint>,		CurvesOffsetsBuffer)
 		SHADER_PARAMETER_RDG_BUFFER_SRV(Buffer<float4>,	DeformedPositionOffset)
 		SHADER_PARAMETER(FVector3f,				RestPositionOffset)
+		SHADER_PARAMETER(FMatrix44f,			WorldToLocal)
 
 		SHADER_PARAMETER(int,					StrandsSize)
 		SHADER_PARAMETER(int,					NumPoints)
@@ -908,9 +909,36 @@ IMPLEMENT_SHADER_TYPE(, FInterpolateGroomGuidesCS, TEXT("/Plugin/Runtime/HairStr
 
 //------------------------------------------------------------------------------------------------------------
 
+FMatrix44f ComputeBoneTransform(const FNDIHairStrandsData* InstanceData)
+{
+	const FMatrix44f RigidTransformFloat = FMatrix44f(InstanceData->HairGroupInstance ? InstanceData->HairGroupInstance->Debug.RigidCurrentLocalToWorld.ToMatrixWithScale() : InstanceData->WorldTransform.ToMatrixWithScale());
+	return FMatrix44f(InstanceData->BoneTransform.ToMatrixWithScale()) * RigidTransformFloat;
+}
+
+FMatrix44f ComputeWorldTransform(const FNDIHairStrandsData* InstanceData)
+{
+	const FMatrix44f BoneTransformFloat = ComputeBoneTransform(InstanceData);
+	
+	FMatrix44f WorldTransformFloat = FMatrix44f(InstanceData->HairGroupInstance ? InstanceData->HairGroupInstance->GetCurrentLocalToWorld().ToMatrixWithScale() : InstanceData->WorldTransform.ToMatrixWithScale());
+	if (InstanceData->LocalSimulation)
+	{
+		const FMatrix44d WorldTransformDouble(WorldTransformFloat);
+		const FMatrix44d BoneTransformDouble(BoneTransformFloat);
+
+		// Due to large world coordinate we store the relative world transform in double precision 
+		WorldTransformFloat = FMatrix44f(WorldTransformDouble * BoneTransformDouble.Inverse());
+	}
+		
+	if (WorldTransformFloat.ContainsNaN())
+	{
+		WorldTransformFloat = FMatrix44f(InstanceData->WorldTransform.ToMatrixWithScale());
+	}
+	return WorldTransformFloat;
+}
+
 static void InterpolateGroomGuides(FRDGBuilder& GraphBuilder, FNiagaraDataBuffer* ParticlesBuffer,
 	const uint32 NodePositionComponent,  const uint32 RestPositionComponent, FNDIHairStrandsBuffer* HairStrandsBuffer,
-	const uint32 StrandsSize, const bool bHasSkinningBinding, const bool bHasValidGeometry)
+	const uint32 StrandsSize, const bool bHasSkinningBinding, const bool bHasValidGeometry, const FMatrix44f& WorldToLocal)
 {
 	const bool bIsHairValid = HairStrandsBuffer && HairStrandsBuffer->IsInitialized();
 	const bool bIsRestValid = bIsHairValid && HairStrandsBuffer->SourceRestResources && HairStrandsBuffer->SourceRestResources->IsInitialized();
@@ -924,6 +952,7 @@ static void InterpolateGroomGuides(FRDGBuilder& GraphBuilder, FNiagaraDataBuffer
 		const uint32 NumPoints = HairStrandsBuffer->SourceRestResources->GetVertexCount();
 		const int32 MeshLODIndex = bIsRootValid ? HairStrandsBuffer->SourceDeformedRootResources->MeshLODIndex : -1;
 		bIsRootValid = bIsRootValid && HairStrandsBuffer->SourceDeformedRootResources->IsValid(MeshLODIndex);
+		
 
 		FInterpolateGroomGuidesCS::FPermutationDomain InterpolationDomain;
 		InterpolationDomain.Set<FInterpolateGroomGuidesCS::FInterpolationType>(!bIsRootValid);
@@ -938,6 +967,7 @@ static void InterpolateGroomGuides(FRDGBuilder& GraphBuilder, FNiagaraDataBuffer
 		PassParameters->NiagaraFloatStride = ParticlesBuffer->GetFloatStride() / sizeof(float);
 		PassParameters->NodePositionComponent = NodePositionComponent;
 		PassParameters->RestPositionComponent = RestPositionComponent;
+		PassParameters->WorldToLocal = WorldToLocal;
 
 		const int32 BufferIndex = HairStrandsBuffer->SourceDeformedResources->GetIndex(FHairStrandsDeformedResource::Current);
 
@@ -1004,19 +1034,21 @@ void UNiagaraDataInterfaceHairStrands::SimCachePostReadFrame(void* OptionalPerIn
 				
 				FNiagaraDataBuffer* ParticlesBuffer = EmitterDatas.GetCurrentData();
 				FNDIHairStrandsBuffer* HairstrandsBuffer = InstanceData->HairStrandsBuffer;
-
+				
 				if(HairstrandsBuffer && InstanceData->HairGroupInstance)
 				{
+					const FMatrix44f WorldToLocal = ComputeWorldTransform(InstanceData).Inverse();
+					
 					const uint32 StrandsSize = InstanceData->StrandsSize;
 					const bool bHasSkinningBinding = InstanceData->HairGroupInstance->BindingType == EHairBindingType::Skinning;
 					const bool bHasValidGeometry = InstanceData->HairGroupInstance->GeometryType != EHairGeometryType::NoneGeometry;
 				
 					ENQUEUE_RENDER_COMMAND(NiagaraInterpolateGroomSimCache)(
-						[ParticlesBuffer, HairstrandsBuffer, NodePositionComponent, RestPositionComponent, StrandsSize, bHasSkinningBinding, bHasValidGeometry](FRHICommandListImmediate& RHICmdList)
+						[ParticlesBuffer, HairstrandsBuffer, NodePositionComponent, RestPositionComponent, StrandsSize, bHasSkinningBinding, bHasValidGeometry, WorldToLocal](FRHICommandListImmediate& RHICmdList)
 						{
 							FMemMark MemMark(FMemStack::Get());
 							FRDGBuilder GraphBuilder(RHICmdList);
-							InterpolateGroomGuides(GraphBuilder, ParticlesBuffer, NodePositionComponent, RestPositionComponent, HairstrandsBuffer, StrandsSize, bHasSkinningBinding, bHasValidGeometry);
+							InterpolateGroomGuides(GraphBuilder, ParticlesBuffer, NodePositionComponent, RestPositionComponent, HairstrandsBuffer, StrandsSize, bHasSkinningBinding, bHasValidGeometry, WorldToLocal);
 							GraphBuilder.Execute();
 
 							HairstrandsBuffer->PointsCurveBuffer.EndGraphUsage();
@@ -3459,19 +3491,9 @@ void UNiagaraDataInterfaceHairStrands::SetShaderParameters(const FNiagaraDataInt
 		// Offsets / Transforms
 		FVector3f RestPositionOffsetValue = (FVector3f)ProxyData->HairStrandsBuffer->SourceRestResources->GetPositionOffset();
 
-		FMatrix44f RigidTransformFloat = FMatrix44f(ProxyData->HairGroupInstance ? ProxyData->HairGroupInstance->Debug.RigidCurrentLocalToWorld.ToMatrixWithScale() : ProxyData->WorldTransform.ToMatrixWithScale());
-		FMatrix44f WorldTransformFloat = FMatrix44f(ProxyData->HairGroupInstance ? ProxyData->HairGroupInstance->GetCurrentLocalToWorld().ToMatrixWithScale() : ProxyData->WorldTransform.ToMatrixWithScale());
-		FMatrix44f BoneTransformFloat = FMatrix44f(ProxyData->BoneTransform.ToMatrixWithScale()) * RigidTransformFloat;
-
-		if (ProxyData->LocalSimulation)
-		{
-			const FMatrix44d WorldTransformDouble(WorldTransformFloat);
-			const FMatrix44d BoneTransformDouble(BoneTransformFloat);
-
-			// Due to large world coordinate we store the relative world transform in double precision 
-			WorldTransformFloat = FMatrix44f(WorldTransformDouble * BoneTransformDouble.Inverse());
-		}
-
+		const FMatrix44f WorldTransformFloat = ComputeWorldTransform(ProxyData);
+		const FMatrix44f BoneTransformFloat = ComputeBoneTransform(ProxyData);
+		
 		if (ProxyData->BoneLinearVelocity.ContainsNaN() || ProxyData->BoneAngularVelocity.ContainsNaN() || ProxyData->BoneLinearAcceleration.ContainsNaN() || ProxyData->BoneAngularAcceleration.ContainsNaN())
 		{
 			UE_LOG(LogHairStrands, Log, TEXT("Bad bones state"));
@@ -3483,11 +3505,6 @@ void UNiagaraDataInterfaceHairStrands::SetShaderParameters(const FNiagaraDataInt
 				*ProxyData->HairGroupInstance->Debug.GroomAssetName, *ProxyData->HairGroupInstance->Debug.MeshComponentName);
 		}
 
-		if (WorldTransformFloat.ContainsNaN())
-		{
-			WorldTransformFloat = FMatrix44f(ProxyData->WorldTransform.ToMatrixWithScale());
-		}
-
 		if (bHasSkinnedInterpolation && (!DeformedMeshProjection || !RestMeshProjection))
 		{
 			UE_LOG(LogHairStrands, Log, TEXT("Bad valid projection"));
