Use RHICmdList flush functions rather than calling GDynamicRHI directly.
 - GDynamicRHI is not the public API, and calling it directly can skip an RHI thread flush.

#rnx

diff --git a/Engine/Plugins/Runtime/HairStrands/Source/HairStrandsCore/Private/HairCardsBuilder.cpp b/Engine/Plugins/Runtime/HairStrands/Source/HairStrandsCore/Private/HairCardsBuilder.cpp
--- a/Engine/Plugins/Runtime/HairStrands/Source/HairStrandsCore/Private/HairCardsBuilder.cpp
+++ b/Engine/Plugins/Runtime/HairStrands/Source/HairStrandsCore/Private/HairCardsBuilder.cpp
@@ -4177,8 +4177,8 @@ static void AddCardsTextureReadbackPass(
 			CopyInfo);
 
 		// Flush, to ensure that all texture generation is done
-		GDynamicRHI->RHISubmitCommandsAndFlushGPU();
-		GDynamicRHI->RHIBlockUntilGPUIdle();
+		RHICmdList.SubmitCommandsAndFlushGPU();
+		RHICmdList.BlockUntilGPUIdle();
 
 		// don't think we can mark package as a dirty in the package build
 #if WITH_EDITORONLY_DATA
