* Remove attribute patching permutation from the hair interpolation pass, and move it into a separate pass.
* Fix cluster debug view on non-deformed groom

diff --git a/Engine/Plugins/Runtime/HairStrands/Source/HairStrandsCore/Private/HairStrandsRendering.cpp b/Engine/Plugins/Runtime/HairStrands/Source/HairStrandsCore/Private/HairStrandsRendering.cpp
--- a/Engine/Plugins/Runtime/HairStrands/Source/HairStrandsCore/Private/HairStrandsRendering.cpp
+++ b/Engine/Plugins/Runtime/HairStrands/Source/HairStrandsCore/Private/HairStrandsRendering.cpp
@@ -479,19 +479,17 @@ class FHairInterpolationCS : public FGlobalShader
 	SHADER_USE_PARAMETER_STRUCT(FHairInterpolationCS, FGlobalShader);
 
 	class FGroupSize : SHADER_PERMUTATION_SPARSE_INT("PERMUTATION_GROUP_SIZE", 32, 64);
-	class FDebug : SHADER_PERMUTATION_BOOL("PERMUTATION_DEBUG");
 	class FDynamicGeometry : SHADER_PERMUTATION_INT("PERMUTATION_DYNAMIC_GEOMETRY", 5);
 	class FSimulation : SHADER_PERMUTATION_BOOL("PERMUTATION_SIMULATION");
 	class FSingleGuide : SHADER_PERMUTATION_BOOL("PERMUTATION_USE_SINGLE_GUIDE");
 	class FCulling : SHADER_PERMUTATION_BOOL("PERMUTATION_CULLING");
-	using FPermutationDomain = TShaderPermutationDomain<FGroupSize, FDebug, FDynamicGeometry, FSimulation, FSingleGuide, FCulling>;
+	using FPermutationDomain = TShaderPermutationDomain<FGroupSize, FDynamicGeometry, FSimulation, FSingleGuide, FCulling>;
 
 	BEGIN_SHADER_PARAMETER_STRUCT(FParameters, )
 		SHADER_PARAMETER_STRUCT_INCLUDE(ShaderPrint::FShaderParameters, ShaderDrawParameters)
 		SHADER_PARAMETER(uint32, VertexStart)
 		SHADER_PARAMETER(uint32, VertexCount)
 		SHADER_PARAMETER(uint32, DispatchCountX)
-		SHADER_PARAMETER(uint32, HairDebugMode)
 		SHADER_PARAMETER(float, HairLengthScale)
 		SHADER_PARAMETER(FVector3f, InRenderHairPositionOffset)
 		SHADER_PARAMETER(FVector3f, InSimHairPositionOffset)
@@ -501,8 +499,6 @@ class FHairInterpolationCS : public FGlobalShader
 		SHADER_PARAMETER_RDG_BUFFER_SRV(Buffer, RenderRestPosePositionBuffer)
 		SHADER_PARAMETER_RDG_BUFFER_UAV(RWBuffer, OutRenderDeformedPositionBuffer)
 
-		SHADER_PARAMETER_RDG_BUFFER_SRV(Buffer, VertexToClusterIdBuffer)
-
 		SHADER_PARAMETER_RDG_BUFFER_SRV(Buffer, SimRestPosePositionBuffer)
 		SHADER_PARAMETER_RDG_BUFFER_SRV(Buffer, DeformedSimPositionBuffer)
 
@@ -510,9 +506,6 @@ class FHairInterpolationCS : public FGlobalShader
 		SHADER_PARAMETER_RDG_BUFFER_SRV(Buffer, Interpolation0Buffer)
 		SHADER_PARAMETER_RDG_BUFFER_SRV(Buffer, Interpolation1Buffer)
 
-		SHADER_PARAMETER_RDG_BUFFER_SRV(Buffer<float4>, AttributeBuffer)
-		SHADER_PARAMETER_RDG_BUFFER_SRV(Buffer<float4>, SimAttributeBuffer)
-		SHADER_PARAMETER_RDG_BUFFER_UAV(RWBuffer<float4>, OutRenderAttributeBuffer)
 		SHADER_PARAMETER_RDG_BUFFER_SRV(Buffer<float4>, RestPosition0Buffer)
 		SHADER_PARAMETER_RDG_BUFFER_SRV(Buffer<float4>, RestPosition1Buffer)
 		SHADER_PARAMETER_RDG_BUFFER_SRV(Buffer<float4>, RestPosition2Buffer)
@@ -577,7 +570,6 @@ static void AddHairStrandsInterpolationPass(
 	const uint32 VertexCount,
 	const int32 MeshLODIndex,
 	const float HairLengthScale,
-	const bool bPatchedAttributeBuffer,
 	const EHairInterpolationType HairInterpolationType,
 	const EHairGeometryType InstanceGeometryType,
 	const FRDGHairStrandsCullingData& CullingData,
@@ -590,18 +582,14 @@ static void AddHairStrandsInterpolationPass(
 	const FHairStrandsDeformedRootResource* RenDeformedRootResources,
 	const FHairStrandsDeformedRootResource* SimDeformedRootResources,
 	const FRDGBufferSRVRef& RenderRestPosePositionBuffer,
-	const FRDGBufferSRVRef& RenderAttributeBuffer,
 	const bool bUseSingleGuide,
 	const FRDGBufferSRVRef& InterpolationBuffer,
 	const FRDGBufferSRVRef& Interpolation0Buffer,
 	const FRDGBufferSRVRef& Interpolation1Buffer,
 	const FRDGBufferSRVRef& SimRestPosePositionBuffer,
 	const FRDGBufferSRVRef& SimDeformedPositionBuffer,
-	const FRDGBufferSRVRef& SimAttributeBuffer,
 	FRDGImportedBuffer& OutRenderPositionBuffer,
 	FRDGImportedBuffer* OutRenderPrevPositionBuffer,
-	FRDGImportedBuffer* OutRenderAttributeBuffer,
-	const FRDGBufferSRVRef& VertexToClusterIdBuffer,
 	const FRDGBufferSRVRef& SimRootPointIndexBuffer,
 	const FHairStrandsDeformedRootResource::FLOD::EFrameType DeformedFrame)
 {
@@ -634,8 +622,6 @@ static void AddHairStrandsInterpolationPass(
 
 	Parameters->DispatchCountX = DispatchCount.X;
 	Parameters->SimRootPointIndexBuffer = SimRootPointIndexBuffer;
-
-	Parameters->AttributeBuffer = RenderAttributeBuffer;
 	
 	const bool bIsVertexToCurveBuffersValid_Sim =
 		SimRestRootResources &&
@@ -652,36 +638,10 @@ static void AddHairStrandsInterpolationPass(
 	{
 		Parameters->RenVertexToRootIndexBuffer = RegisterAsSRV(GraphBuilder, RenRestRootResources->VertexToCurveIndexBuffer);
 	}
-
-	Parameters->VertexToClusterIdBuffer = VertexToClusterIdBuffer;
 	
 	Parameters->LocalToWorldMatrix = FMatrix44f(Instance->LocalToWorld.ToMatrixWithScale());		// LWC_TODO: Precision loss
 	Parameters->HairLengthScale = HairLengthScale;
 
-	// Debug rendering
-	// 1: Patch attributes for displaying guides influences onto the render strands
-	// 2: Show the visible cluster
-	Parameters->HairDebugMode = 0;
-	{
-		// Patch Attribute0 (RooUV/NormalizedLength/Seed)
-		if (bPatchedAttributeBuffer && OutRenderAttributeBuffer)
-		{
-			Parameters->HairDebugMode = 1;
-			check(SimAttributeBuffer != nullptr);
-			check(OutRenderAttributeBuffer != nullptr);
-		}
-		else if (GetHairStrandsDebugStrandsMode() == EHairStrandsDebugMode::RenderVisCluster)
-		{
-			Parameters->HairDebugMode = 2;
-		}
-
-		if (Parameters->HairDebugMode > 0 && OutRenderAttributeBuffer)
-		{
-			Parameters->SimAttributeBuffer = SimAttributeBuffer;
-			Parameters->OutRenderAttributeBuffer = OutRenderAttributeBuffer->UAV;
-		}
-	}
-
 	// Only needed if DEBUG_ENABLE is manually enabledin HairStrandsInterpolation.usf. This is for manual debug purpose only
 	#if 0
 	if (ShaderPrintData && ShaderPrintData->IsValid())
@@ -764,7 +724,6 @@ static void AddHairStrandsInterpolationPass(
 
 	FHairInterpolationCS::FPermutationDomain PermutationVector;
 	PermutationVector.Set<FHairInterpolationCS::FGroupSize>(GroupSize);
-	PermutationVector.Set<FHairInterpolationCS::FDebug>(Parameters->HairDebugMode > 0);
 	PermutationVector.Set<FHairInterpolationCS::FDynamicGeometry>(DynamicGeometryType);
 	PermutationVector.Set<FHairInterpolationCS::FSimulation>(bHasLocalDeformation);
 	PermutationVector.Set<FHairInterpolationCS::FSingleGuide>(bUseSingleGuide);
@@ -798,11 +757,6 @@ static void AddHairStrandsInterpolationPass(
 	}
 
 	GraphBuilder.SetBufferAccessFinal(OutRenderPositionBuffer.Buffer, ERHIAccess::SRVMask);
-	if (OutRenderAttributeBuffer)
-	{
-		GraphBuilder.SetBufferAccessFinal(OutRenderAttributeBuffer->Buffer, ERHIAccess::SRVMask);
-	}
-
 	if (!(IsHairVisibilityComputeRasterContinuousLODEnabled() || IsHairVisibilityComputeRasterTemporalLayeringEnabled()) && Instance->HairGroupPublicData->VFInput.bHasLODSwitch && OutRenderPrevPositionBuffer && GHairStrandsTransferPositionOnLODChange > 0)
 	{
 		AddTransferPositionPass(GraphBuilder, ShaderMap, VertexCount, OutRenderPositionBuffer.SRV, OutRenderPrevPositionBuffer->UAV);
@@ -810,6 +764,96 @@ static void AddHairStrandsInterpolationPass(
 	}
 }
 
+///////////////////////////////////////////////////////////////////////////////////////////////////
+// Patch hair attribute (For debug visualization)
+
+class FHairPatchAttributeCS : public FGlobalShader
+{
+	DECLARE_GLOBAL_SHADER(FHairPatchAttributeCS);
+	SHADER_USE_PARAMETER_STRUCT(FHairPatchAttributeCS, FGlobalShader);
+
+	class FSimulation : SHADER_PERMUTATION_INT("PERMUTATION_SIMULATION", 3); // No simulation, Simlation with 1 guide, Simulation with 3 guides
+	using FPermutationDomain = TShaderPermutationDomain<FSimulation>;
+
+	BEGIN_SHADER_PARAMETER_STRUCT(FParameters, )
+		SHADER_PARAMETER(uint32, VertexCount)
+
+		SHADER_PARAMETER_RDG_BUFFER_SRV(Buffer<float4>, SimAttribute0Buffer)
+		SHADER_PARAMETER_RDG_BUFFER_SRV(Buffer, InterpolationBuffer)
+		SHADER_PARAMETER_RDG_BUFFER_SRV(Buffer, Interpolation0Buffer)
+		SHADER_PARAMETER_RDG_BUFFER_SRV(Buffer, Interpolation1Buffer)
+
+		SHADER_PARAMETER_RDG_BUFFER_SRV(Buffer<float4>, RenAttribute0Buffer)
+		SHADER_PARAMETER_RDG_BUFFER_SRV(Buffer, RenVertexToClusterIdBuffer)
+		SHADER_PARAMETER_RDG_BUFFER_UAV(RWBuffer<float4>, OutRenAttribute0Buffer)
+	END_SHADER_PARAMETER_STRUCT()
+
+public:
+	static bool ShouldCompilePermutation(const FGlobalShaderPermutationParameters& Parameters) { return IsHairStrandsSupported(EHairStrandsShaderType::All, Parameters.Platform); }
+	static void ModifyCompilationEnvironment(const FGlobalShaderPermutationParameters& Parameters, FShaderCompilerEnvironment& OutEnvironment)
+	{
+		FGlobalShader::ModifyCompilationEnvironment(Parameters, OutEnvironment);
+		OutEnvironment.SetDefine(TEXT("SHADER_PATCHATTRIBUTE"), 1);
+	}
+};
+
+IMPLEMENT_GLOBAL_SHADER(FHairPatchAttributeCS, "/Engine/Private/HairStrands/HairStrandsInterpolation.usf", "MainCS", SF_Compute);
+
+enum class EHairPatchAttribute : uint8
+{
+	None,
+	GuideInflucence,
+	ClusterInfluence
+};
+
+static void AddPatchAttributePass(
+	FRDGBuilder& GraphBuilder,
+	FGlobalShaderMap* ShaderMap,
+	const uint32 VertexCount,
+	const EHairPatchAttribute Mode,
+	const bool bSimulation,
+	const bool bUseSingleGuide,
+	const FRDGBufferSRVRef& RenAttribute0Buffer,
+	const FRDGBufferSRVRef& RenVertexToClusterIdBuffer,
+	const FRDGBufferSRVRef& SimAttribute0Buffer,
+	const FRDGBufferSRVRef& InterpolationBuffer,
+	const FRDGBufferSRVRef& Interpolation0Buffer,
+	const FRDGBufferSRVRef& Interpolation1Buffer,
+	FRDGImportedBuffer& OutRenAttribute0Buffer)
+{
+	// Sanity check
+	check(Mode == EHairPatchAttribute::ClusterInfluence || Mode == EHairPatchAttribute::GuideInflucence)
+	if (Mode == EHairPatchAttribute::GuideInflucence && !bSimulation) return;
+
+	const FIntVector DispatchCount = FIntVector(FMath::DivideAndRoundUp(VertexCount, 1024u), 1, 1);
+
+	FHairPatchAttributeCS::FParameters* Parameters = GraphBuilder.AllocParameters<FHairPatchAttributeCS::FParameters>();
+	Parameters->VertexCount = VertexCount;
+	Parameters->RenAttribute0Buffer = RenAttribute0Buffer;
+	Parameters->RenVertexToClusterIdBuffer = RenVertexToClusterIdBuffer;
+	if (bSimulation)
+	{
+		Parameters->SimAttribute0Buffer  = SimAttribute0Buffer;
+		Parameters->InterpolationBuffer  = InterpolationBuffer;
+		Parameters->Interpolation0Buffer = Interpolation0Buffer;
+		Parameters->Interpolation1Buffer = Interpolation1Buffer;
+	}
+	Parameters->OutRenAttribute0Buffer = OutRenAttribute0Buffer.UAV;
+
+	FHairPatchAttributeCS::FPermutationDomain PermutationVector;
+	PermutationVector.Set<FHairPatchAttributeCS::FSimulation>(Mode == EHairPatchAttribute::ClusterInfluence ? 0 : (bUseSingleGuide ? 1 : 2));
+
+	TShaderMapRef<FHairPatchAttributeCS> ComputeShader(ShaderMap, PermutationVector);
+	FComputeShaderUtils::AddPass(
+		GraphBuilder,
+		RDG_EVENT_NAME("HairStrands::PatchAttribute"),
+		ComputeShader,
+		Parameters,
+		DispatchCount);
+
+	GraphBuilder.SetBufferAccessFinal(OutRenAttribute0Buffer.Buffer, ERHIAccess::SRVMask);
+}
+
 ///////////////////////////////////////////////////////////////////////////////////////////////////
 
 class FHairClusterAABBCS : public FGlobalShader
@@ -1572,8 +1616,7 @@ EHairStrandsDebugMode GetHairStrandsGeometryDebugMode(const FHairGroupInstance*
 
 	// 2. If groom does not have guides (e.g., because it does not have simulation/RBF data), then we fallback to default view mode
 	if ((DebugMode == EHairStrandsDebugMode::SimHairStrands ||
-		 DebugMode == EHairStrandsDebugMode::RenderHairStrands ||
-		 DebugMode == EHairStrandsDebugMode::RenderVisCluster)
+		 DebugMode == EHairStrandsDebugMode::RenderHairStrands)
 		&& Instance->Guides.Data == nullptr)
 	{
 		DebugMode = EHairStrandsDebugMode::NoneDebug;
@@ -1581,11 +1624,6 @@ EHairStrandsDebugMode GetHairStrandsGeometryDebugMode(const FHairGroupInstance*
 	return DebugMode;
 }
 
-bool NeedsPatchAttributeBuffer(EHairStrandsDebugMode DebugMode)
-{
-	return DebugMode == EHairStrandsDebugMode::RenderHairStrands || DebugMode == EHairStrandsDebugMode::RenderVisCluster;
-}
-
 static void ConvertHairStrandsVFParameters(
 	FRDGBuilder* GraphBuilder,
 	FRDGImportedBuffer& OutBuffer,
@@ -1610,8 +1648,6 @@ static FHairGroupPublicData::FVertexFactoryInput InternalComputeHairStrandsVerte
 		return OutVFInput;
 
 	const EHairStrandsDebugMode DebugMode = GetHairStrandsGeometryDebugMode(Instance);
-	const bool bDebugModePatchedAttributeBuffer = NeedsPatchAttributeBuffer(DebugMode);
-
 	const bool bSupportDeformation = Instance->Strands.DeformedResource != nullptr;
 
 	// 1. Guides
@@ -1657,7 +1693,7 @@ static FHairGroupPublicData::FVertexFactoryInput InternalComputeHairStrandsVerte
 			OutVFInput.Strands.PrevPositionOffset = Instance->Strands.DeformedResource->GetPositionOffset(FHairStrandsDeformedResource::EFrameType::Current);
 		}
 		CONVERT_HAIRSSTRANDS_VF_PARAMETERS(OutVFInput.Strands.TangentBuffer,			Instance->Strands.DeformedResource->TangentBuffer);
-		CONVERT_HAIRSSTRANDS_VF_PARAMETERS(OutVFInput.Strands.Attribute0Buffer,			bDebugModePatchedAttributeBuffer ? Instance->Strands.DebugAttributeBuffer : Instance->Strands.RestResource->Attribute0Buffer); // TODO
+		CONVERT_HAIRSSTRANDS_VF_PARAMETERS(OutVFInput.Strands.Attribute0Buffer,			(DebugMode == EHairStrandsDebugMode::RenderHairStrands || DebugMode == EHairStrandsDebugMode::RenderVisCluster) ? Instance->Strands.DebugAttributeBuffer : Instance->Strands.RestResource->Attribute0Buffer);
 		CONVERT_HAIRSSTRANDS_VF_PARAMETERS(OutVFInput.Strands.Attribute1Buffer,			Instance->Strands.RestResource->Attribute1Buffer);
 		CONVERT_HAIRSSTRANDS_VF_PARAMETERS(OutVFInput.Strands.MaterialBuffer,			Instance->Strands.RestResource->MaterialBuffer);
 
@@ -1669,7 +1705,7 @@ static FHairGroupPublicData::FVertexFactoryInput InternalComputeHairStrandsVerte
 		CONVERT_HAIRSSTRANDS_VF_PARAMETERS(OutVFInput.Strands.PositionBuffer,			Instance->Strands.RestResource->PositionBuffer);
 		CONVERT_HAIRSSTRANDS_VF_PARAMETERS(OutVFInput.Strands.PrevPositionBuffer,		Instance->Strands.RestResource->PositionBuffer);
 		CONVERT_HAIRSSTRANDS_VF_PARAMETERS(OutVFInput.Strands.TangentBuffer,			Instance->Strands.RestResource->TangentBuffer);
-		CONVERT_HAIRSSTRANDS_VF_PARAMETERS(OutVFInput.Strands.Attribute0Buffer,			Instance->Strands.RestResource->Attribute0Buffer);
+		CONVERT_HAIRSSTRANDS_VF_PARAMETERS(OutVFInput.Strands.Attribute0Buffer,			(DebugMode == EHairStrandsDebugMode::RenderVisCluster) ? Instance->Strands.DebugAttributeBuffer : Instance->Strands.RestResource->Attribute0Buffer);
 		CONVERT_HAIRSSTRANDS_VF_PARAMETERS(OutVFInput.Strands.Attribute1Buffer,			Instance->Strands.RestResource->Attribute1Buffer);
 		CONVERT_HAIRSSTRANDS_VF_PARAMETERS(OutVFInput.Strands.MaterialBuffer,			Instance->Strands.RestResource->MaterialBuffer);
 		CONVERT_HAIRSSTRANDS_VF_PARAMETERS(OutVFInput.Strands.PositionOffsetBuffer,		Instance->Strands.RestResource->PositionOffsetBuffer);
@@ -1741,7 +1777,6 @@ void ComputeHairStrandsInterpolation(
 		// * Sim	: Show sim strands
 		// * Render : Show rendering strands with sim color influence
 		const EHairStrandsDebugMode DebugMode = GetHairStrandsGeometryDebugMode(Instance);
-		const bool bDebugModePatchedAttributeBuffer = NeedsPatchAttributeBuffer(DebugMode);
 
 		if (DebugMode == EHairStrandsDebugMode::SimHairStrands)
 		{
@@ -1777,6 +1812,13 @@ void ComputeHairStrandsInterpolation(
 					CullingData.GroupAABBBuffer);
 			}
 
+			// "WITH_EDITOR && PatchMode == EHairPatchAttribute::GuideInflucence" is a special path when visualizing groom within the groom editor, so that we can diplay guides even when there is no simulation or global interpolation enabled
+			EHairPatchAttribute PatchMode = EHairPatchAttribute::None;
+			#if WITH_EDITOR
+			if (DebugMode == EHairStrandsDebugMode::RenderHairStrands)	PatchMode = EHairPatchAttribute::GuideInflucence;
+			if (DebugMode == EHairStrandsDebugMode::RenderVisCluster)	PatchMode = EHairPatchAttribute::ClusterInfluence;
+			#endif
+
 			// 2. Deform hair if needed (e.g.: skinning, simulation, RBF) and recompute tangent
 			FRDGBufferSRVRef Strands_PositionSRV = nullptr;
 			FRDGBufferSRVRef Strands_PositionOffsetSRV = nullptr;
@@ -1788,32 +1830,15 @@ void ComputeHairStrandsInterpolation(
 				FRDGImportedBuffer Strands_DeformedPosition = Register(GraphBuilder, Instance->Strands.DeformedResource->GetBuffer(FHairStrandsDeformedResource::Current), ERDGImportedBufferFlags::CreateViews);
 				FRDGImportedBuffer Strands_DeformedPrevPosition = Register(GraphBuilder, Instance->Strands.DeformedResource->GetBuffer(FHairStrandsDeformedResource::Previous), ERDGImportedBufferFlags::CreateViews);
 				FRDGImportedBuffer Strands_DeformedTangent = Register(GraphBuilder, Instance->Strands.DeformedResource->TangentBuffer, ERDGImportedBufferFlags::CreateViews);
-				FRDGImportedBuffer* Strands_DebugAttribute = nullptr;
 
 				// Trach on which view the position has been update, so that we can ensure motion vector are coherent
 				Instance->Strands.DeformedResource->GetUniqueViewID(FHairStrandsDeformedResource::Current) = ViewUniqueID;
 
-			#if WITH_EDITOR
-				FRDGImportedBuffer Strands_DebugAttributeReg;
-				if (bDebugModePatchedAttributeBuffer)
-				{
-					// Create an debug buffer for storing cluster visalization data. This is only used for debug purpose, hence only enable in editor build.
-					// Special case for debug mode were the attribute buffer is patch with some custom data to show hair properties (strands belonging to the same cluster, ...)
-					if (Instance->Strands.DebugAttributeBuffer.Buffer == nullptr)
-					{
-						CreateHairStrandsDebugAttributeBuffer(GraphBuilder, &Instance->Strands.DebugAttributeBuffer, Instance->Strands.Data->GetNumPoints());
-					}
-
-					Strands_DebugAttributeReg = Register(GraphBuilder, Instance->Strands.DebugAttributeBuffer, ERDGImportedBufferFlags::CreateUAV);
-					Strands_DebugAttribute = &Strands_DebugAttributeReg;
-				}
-			#endif
 				Strands_PositionSRV = Strands_DeformedPosition.SRV;
 				Strands_PositionOffsetSRV = RegisterAsSRV(GraphBuilder, Instance->Strands.DeformedResource->GetPositionOffsetBuffer(FHairStrandsDeformedResource::EFrameType::Current));
 				Strands_TangentSRV = Strands_DeformedTangent.SRV;
 
-				// "WITH_EDITOR && bDebugModePatchedAttributeBuffer" is a special path when visualizing groom within the groom editor, so that we can diplay guides even when there is no simulation or global interpolation enabled
-				const bool bValidGuide = Instance->Guides.bIsSimulationEnable || Instance->Guides.bHasGlobalInterpolation || Instance->Guides.bIsDeformationEnable || (WITH_EDITOR && bDebugModePatchedAttributeBuffer);
+				const bool bValidGuide = Instance->Guides.bIsSimulationEnable || Instance->Guides.bHasGlobalInterpolation || Instance->Guides.bIsDeformationEnable || (WITH_EDITOR && PatchMode == EHairPatchAttribute::GuideInflucence);
 				const bool bHasSkinning = Instance->BindingType == EHairBindingType::Skinning;
 
 				bool bHasValidGroomCacheBuffers = Instance->Debug.GroomCacheBuffers.IsValid() && Instance->Debug.GroomCacheBuffers->GetInterpolatedFrameBuffer().GroupsData.IsValidIndex(Instance->Debug.GroupIndex);
@@ -1903,7 +1928,6 @@ void ComputeHairStrandsInterpolation(
 						VertexCount,
 						MeshLODIndex,
 						Instance->Strands.Modifier.HairLengthScale,
-						bDebugModePatchedAttributeBuffer,
 						Instance->Strands.HairInterpolationType,
 						InstanceGeometryType,
 						CullingData,
@@ -1916,63 +1940,54 @@ void ComputeHairStrandsInterpolation(
 						bHasSkinning ? Instance->Strands.DeformedRootResource : nullptr,
 						bHasSkinning && bValidGuide ? Instance->Guides.DeformedRootResource : nullptr,
 						RegisterAsSRV(GraphBuilder, Instance->Strands.RestResource->PositionBuffer),
-						RegisterAsSRV(GraphBuilder, Instance->Strands.RestResource->Attribute0Buffer),
 						bUseSingleGuide,
 						bValidGuide &&  bUseSingleGuide ? RegisterAsSRV(GraphBuilder, Instance->Strands.InterpolationResource->InterpolationBuffer) : nullptr,
 						bValidGuide && !bUseSingleGuide ? RegisterAsSRV(GraphBuilder, Instance->Strands.InterpolationResource->Interpolation0Buffer) : nullptr,
 						bValidGuide && !bUseSingleGuide ? RegisterAsSRV(GraphBuilder, Instance->Strands.InterpolationResource->Interpolation1Buffer) : nullptr,
 						bValidGuide ? RegisterAsSRV(GraphBuilder, Instance->Guides.RestResource->PositionBuffer) : nullptr,
 						bValidGuide ? RegisterAsSRV(GraphBuilder, Instance->Guides.DeformedResource->GetBuffer(FHairStrandsDeformedResource::Current)) : nullptr,
-						bValidGuide ? RegisterAsSRV(GraphBuilder, Instance->Guides.RestResource->Attribute0Buffer) : nullptr,
 						Strands_DeformedPosition,
 						&Strands_DeformedPrevPosition,
-						Strands_DebugAttribute,
-						RegisterAsSRV(GraphBuilder, Instance->Strands.ClusterCullingResource->VertexToClusterIdBuffer),
 						bValidGuide ? RegisterAsSRV(GraphBuilder, Instance->Strands.InterpolationResource->SimRootPointIndexBuffer) : nullptr, 
 						FHairStrandsDeformedRootResource::FLOD::Current);
 					
-						// second interpolation pass to calculate previous positions if needed
-						if (IsHairVisibilityComputeRasterContinuousLODEnabled() || IsHairVisibilityComputeRasterTemporalLayeringEnabled())
-						{
-							Strands_PrevPositionOffsetSRV = RegisterAsSRV(GraphBuilder, Instance->Strands.DeformedResource->GetPositionOffsetBuffer(FHairStrandsDeformedResource::EFrameType::Previous));
-
-							AddHairStrandsInterpolationPass(
-								GraphBuilder,
-								ShaderMap,
-								ShaderPrintData,
-								Instance,
-								VertexStart,
-								VertexCount,
-								MeshLODIndex,
-								Instance->Strands.Modifier.HairLengthScale,
-								bDebugModePatchedAttributeBuffer,
-								Instance->Strands.HairInterpolationType,
-								InstanceGeometryType,
-								CullingData,
-								Instance->Strands.RestResource->GetPositionOffset(),
-								bValidGuide ? Instance->Guides.RestResource->GetPositionOffset() : FVector::ZeroVector,
-								Strands_PrevPositionOffsetSRV,
-								bValidGuide ? RegisterAsSRV(GraphBuilder, Instance->Guides.DeformedResource->GetPositionOffsetBuffer(FHairStrandsDeformedResource::Previous)) : nullptr,
-								bHasSkinning ? Instance->Strands.RestRootResource : nullptr,
-								bHasSkinning && bValidGuide ? Instance->Guides.RestRootResource : nullptr,
-								bHasSkinning ? Instance->Strands.DeformedRootResource : nullptr,
-								bHasSkinning && bValidGuide ? Instance->Guides.DeformedRootResource : nullptr,
-								RegisterAsSRV(GraphBuilder, Instance->Strands.RestResource->PositionBuffer),
-								RegisterAsSRV(GraphBuilder, Instance->Strands.RestResource->Attribute0Buffer),
-								bUseSingleGuide,
-								bValidGuide&& bUseSingleGuide ? RegisterAsSRV(GraphBuilder, Instance->Strands.InterpolationResource->InterpolationBuffer) : nullptr,
-								bValidGuide && !bUseSingleGuide ? RegisterAsSRV(GraphBuilder, Instance->Strands.InterpolationResource->Interpolation0Buffer) : nullptr,
-								bValidGuide && !bUseSingleGuide ? RegisterAsSRV(GraphBuilder, Instance->Strands.InterpolationResource->Interpolation1Buffer) : nullptr,
-								bValidGuide ? RegisterAsSRV(GraphBuilder, Instance->Guides.RestResource->PositionBuffer) : nullptr,
-								bValidGuide ? RegisterAsSRV(GraphBuilder, Instance->Guides.DeformedResource->GetBuffer(FHairStrandsDeformedResource::Previous)) : nullptr,
-								bValidGuide ? RegisterAsSRV(GraphBuilder, Instance->Guides.RestResource->Attribute0Buffer) : nullptr,
-								Strands_DeformedPrevPosition,
-								nullptr,
-								nullptr,//Strands_DebugAttribute,
-								RegisterAsSRV(GraphBuilder, Instance->Strands.ClusterCullingResource->VertexToClusterIdBuffer),
-								bValidGuide ? RegisterAsSRV(GraphBuilder, Instance->Strands.InterpolationResource->SimRootPointIndexBuffer) : nullptr,
-								FHairStrandsDeformedRootResource::FLOD::Previous);
-						}
+					// second interpolation pass to calculate previous positions if needed
+					if (IsHairVisibilityComputeRasterContinuousLODEnabled() || IsHairVisibilityComputeRasterTemporalLayeringEnabled())
+					{
+						Strands_PrevPositionOffsetSRV = RegisterAsSRV(GraphBuilder, Instance->Strands.DeformedResource->GetPositionOffsetBuffer(FHairStrandsDeformedResource::EFrameType::Previous));
+
+						AddHairStrandsInterpolationPass(
+							GraphBuilder,
+							ShaderMap,
+							ShaderPrintData,
+							Instance,
+							VertexStart,
+							VertexCount,
+							MeshLODIndex,
+							Instance->Strands.Modifier.HairLengthScale,
+							Instance->Strands.HairInterpolationType,
+							InstanceGeometryType,
+							CullingData,
+							Instance->Strands.RestResource->GetPositionOffset(),
+							bValidGuide ? Instance->Guides.RestResource->GetPositionOffset() : FVector::ZeroVector,
+							Strands_PrevPositionOffsetSRV,
+							bValidGuide ? RegisterAsSRV(GraphBuilder, Instance->Guides.DeformedResource->GetPositionOffsetBuffer(FHairStrandsDeformedResource::Previous)) : nullptr,
+							bHasSkinning ? Instance->Strands.RestRootResource : nullptr,
+							bHasSkinning && bValidGuide ? Instance->Guides.RestRootResource : nullptr,
+							bHasSkinning ? Instance->Strands.DeformedRootResource : nullptr,
+							bHasSkinning && bValidGuide ? Instance->Guides.DeformedRootResource : nullptr,
+							RegisterAsSRV(GraphBuilder, Instance->Strands.RestResource->PositionBuffer),
+							bUseSingleGuide,
+							bValidGuide &&  bUseSingleGuide ? RegisterAsSRV(GraphBuilder, Instance->Strands.InterpolationResource->InterpolationBuffer) : nullptr,
+							bValidGuide && !bUseSingleGuide ? RegisterAsSRV(GraphBuilder, Instance->Strands.InterpolationResource->Interpolation0Buffer) : nullptr,
+							bValidGuide && !bUseSingleGuide ? RegisterAsSRV(GraphBuilder, Instance->Strands.InterpolationResource->Interpolation1Buffer) : nullptr,
+							bValidGuide ? RegisterAsSRV(GraphBuilder, Instance->Guides.RestResource->PositionBuffer) : nullptr,
+							bValidGuide ? RegisterAsSRV(GraphBuilder, Instance->Guides.DeformedResource->GetBuffer(FHairStrandsDeformedResource::Previous)) : nullptr,
+							Strands_DeformedPrevPosition,
+							nullptr,
+							bValidGuide ? RegisterAsSRV(GraphBuilder, Instance->Strands.InterpolationResource->SimRootPointIndexBuffer) : nullptr,
+							FHairStrandsDeformedRootResource::FLOD::Previous);
+					}
 				}
 				else if (bHasValidGroomCacheBuffers)
 				{
@@ -2024,7 +2039,38 @@ void ComputeHairStrandsInterpolation(
 				Strands_TangentSRV = RegisterAsSRV(GraphBuilder, Instance->Strands.RestResource->TangentBuffer);
 			}
 
-			// 2.1 Update the VF input with the update resources
+			// 2.1 Patch attribute for debug visualization (guide influence or clusters visualization)
+			#if WITH_EDITOR
+			if (PatchMode != EHairPatchAttribute::None)
+			{
+				// Create an debug buffer for storing cluster visualization data. This is only used for debug purpose, hence only enable in editor build.
+				// Special case for debug mode were the attribute buffer is patch with some custom data to show hair properties (strands belonging to the same cluster, ...)
+				if (Instance->Strands.DebugAttributeBuffer.Buffer == nullptr)
+				{
+					CreateHairStrandsDebugAttributeBuffer(GraphBuilder, &Instance->Strands.DebugAttributeBuffer, Instance->Strands.Data->GetNumPoints());
+				}
+				FRDGImportedBuffer OutRenAttribute0Buffer = Register(GraphBuilder, Instance->Strands.DebugAttributeBuffer, ERDGImportedBufferFlags::CreateUAV);
+
+				const bool bValidGuide = PatchMode == EHairPatchAttribute::GuideInflucence;
+				const bool bUseSingleGuide = bValidGuide && Instance->Strands.InterpolationResource->UseSingleGuide();
+				AddPatchAttributePass(
+					GraphBuilder,
+					ShaderMap,
+					Instance->Strands.RestResource->GetVertexCount(),
+					PatchMode,
+					bValidGuide,
+					bUseSingleGuide,
+					RegisterAsSRV(GraphBuilder, Instance->Strands.RestResource->Attribute0Buffer),
+					RegisterAsSRV(GraphBuilder, Instance->Strands.ClusterCullingResource->VertexToClusterIdBuffer),
+					bValidGuide ? RegisterAsSRV(GraphBuilder, Instance->Guides.RestResource->Attribute0Buffer) : nullptr,
+					bValidGuide &&  bUseSingleGuide ? RegisterAsSRV(GraphBuilder, Instance->Strands.InterpolationResource->InterpolationBuffer) : nullptr,
+					bValidGuide && !bUseSingleGuide ? RegisterAsSRV(GraphBuilder, Instance->Strands.InterpolationResource->Interpolation0Buffer) : nullptr,
+					bValidGuide && !bUseSingleGuide ? RegisterAsSRV(GraphBuilder, Instance->Strands.InterpolationResource->Interpolation1Buffer) : nullptr,
+					OutRenAttribute0Buffer);
+			}
+			#endif
+
+			// 2.2 Update the VF input with the update resources
 			Instance->HairGroupPublicData->VFInput = InternalComputeHairStrandsVertexInputData(&GraphBuilder, Instance);
 
 			// 3. Compute cluster AABBs (used for LODing and voxelization)
@@ -2201,7 +2247,6 @@ void ComputeHairStrandsInterpolation(
 						LOD.Guides.RestResource->GetVertexCount(),
 						MeshLODIndex,
 						1.0f,
-						false,
 						LOD.Guides.HairInterpolationType,
 						InstanceGeometryType,
 						CullingData,
@@ -2214,18 +2259,14 @@ void ComputeHairStrandsInterpolation(
 						bHasSkinning ? LOD.Guides.DeformedRootResource : nullptr,
 						bHasSkinning && bValidGuide ? Instance->Guides.DeformedRootResource : nullptr,
 						RegisterAsSRV(GraphBuilder, LOD.Guides.RestResource->PositionBuffer),
-						RegisterAsSRV(GraphBuilder, LOD.Guides.RestResource->Attribute0Buffer),
 						bUseSingleGuide,
 						bValidGuide &&  bUseSingleGuide ? RegisterAsSRV(GraphBuilder, LOD.Guides.InterpolationResource->InterpolationBuffer) : nullptr,
 						bValidGuide && !bUseSingleGuide ? RegisterAsSRV(GraphBuilder, LOD.Guides.InterpolationResource->Interpolation0Buffer) : nullptr,
 						bValidGuide && !bUseSingleGuide ? RegisterAsSRV(GraphBuilder, LOD.Guides.InterpolationResource->Interpolation1Buffer) : nullptr,
 						bValidGuide ? RegisterAsSRV(GraphBuilder, Instance->Guides.RestResource->PositionBuffer) : nullptr,
 						bValidGuide ? RegisterAsSRV(GraphBuilder, Instance->Guides.DeformedResource->GetBuffer(FHairStrandsDeformedResource::Current)) : nullptr,
-						bValidGuide ? RegisterAsSRV(GraphBuilder, Instance->Guides.RestResource->Attribute0Buffer) : nullptr,
 						Guides_DeformedPositionBuffer,
 						nullptr,
-						nullptr,
-						nullptr,
 						RegisterAsSRV(GraphBuilder, LOD.Guides.InterpolationResource->SimRootPointIndexBuffer),
 						FHairStrandsDeformedRootResource::FLOD::Current); // <- this should be optional
 
diff --git a/Engine/Shaders/Private/HairStrands/HairStrandsInterpolation.usf b/Engine/Shaders/Private/HairStrands/HairStrandsInterpolation.usf
--- a/Engine/Shaders/Private/HairStrands/HairStrandsInterpolation.usf
+++ b/Engine/Shaders/Private/HairStrands/HairStrandsInterpolation.usf
@@ -21,7 +21,44 @@
 #endif
 
 ///////////////////////////////////////////////////////////////////////////////////////////////////
-// Common
+// Guides
+struct FGuideData
+{
+	uint3  VertexIndices;
+	float3 VertexLerps;
+	float3 VertexWeights;
+	float3 OutPositionOffset;
+};
+
+FGuideData GetGuideData(uint InVertexIndex, uint InInterpolationData, float3 OutPositionOffset)
+{
+	FGuideData Out = (FGuideData)0;
+	Out.OutPositionOffset	= OutPositionOffset;
+	Out.VertexWeights		= float3(1, 0, 0);
+	Out.VertexIndices		= 0;
+	Out.VertexLerps			= 0;
+	Out.VertexIndices.x		= InInterpolationData & 0xFFFFFF;
+	Out.VertexLerps.x		= saturate(((InInterpolationData >> 24) & 0xFF) / 255.f);
+	return Out;
+}
+
+FGuideData GetGuideData(uint InVertexIndex, uint4 InInterpolationData0, uint4 InInterpolationData1, float3 OutPositionOffset)
+{
+	FGuideData Out = (FGuideData)0;
+	Out.OutPositionOffset = OutPositionOffset;
+
+	Out.VertexIndices.x = InInterpolationData0.x | (( InInterpolationData1.x & 0xFF) << 16);
+	Out.VertexIndices.y = InInterpolationData0.y | (((InInterpolationData1.x >> 8) & 0xFF) << 16);
+	Out.VertexIndices.z = InInterpolationData0.z | (( InInterpolationData1.y & 0xFF) << 16);
+
+	Out.VertexWeights	= float3(saturate((InInterpolationData0.w & 0xFF) / 255.f), saturate(((InInterpolationData0.w >> 8) & 0xFF) / 255.f), 0);
+	Out.VertexWeights.z = saturate(1.0f - (Out.VertexWeights.x + Out.VertexWeights.y));
+
+	Out.VertexLerps.x	= saturate(((InInterpolationData1.y >> 8) & 0xFF) / 255.f);
+	Out.VertexLerps.y	= saturate(( InInterpolationData1.z & 0xFF) / 255.f);
+	Out.VertexLerps.z	= saturate(((InInterpolationData1.z >> 8) & 0xFF) / 255.f);
+	return Out;
+}
 
 ///////////////////////////////////////////////////////////////////////////////////////////////////
 // Hair interpolation 
@@ -68,8 +105,6 @@ Buffer<float>	HairStrandsVF_CullingRadiusScaleBuffer;
 Buffer<uint>	HairStrandsVF_CullingIndirectBufferArgs;
 #endif
 
-Buffer<float4>  AttributeBuffer;
-
 #if PERMUTATION_DYNAMIC_GEOMETRY >= 1 || PERMUTATION_DEBUG == 1 
 Buffer<uint>   RenVertexToRootIndexBuffer;
 Buffer<uint>   SimVertexToRootIndexBuffer;
@@ -103,12 +138,6 @@ Buffer<uint>	SimRootBarycentricBuffer;
 Buffer<uint>	SimRootToUniqueTriangleIndexBuffer;
 #endif
 
-#if PERMUTATION_DEBUG == 1
-Buffer<int>		 VertexToClusterIdBuffer;
-Buffer<float4>	 SimAttributeBuffer;
-RWBuffer<float4> OutRenderAttributeBuffer; 
-#endif
-
 Buffer<float4> OutSimHairPositionOffsetBuffer;
 Buffer<float4> OutRenHairPositionOffsetBuffer;
 
@@ -233,47 +262,18 @@ FRootTriangleData GetRootTriangleData(uint VertexIndex, bool bSim)
 
 ///////////////////////////////////////////////////////////////////////////////////////////////////
 // Guides
-struct FGuideData
-{
-	uint3  VertexIndices;
-	float3 VertexLerps;
-	float3 VertexWeights;
-	float3 OutPositionOffset;
-};
 
 FGuideData GetGuideData(uint InVertexIndex)
 {
-	FGuideData Out = (FGuideData)0;
-#if PERMUTATION_SIMULATION == 1 || PERMUTATION_DEBUG == 1
-	{
-		Out.OutPositionOffset = OutSimHairPositionOffsetBuffer[0].xyz;
-
-	#if PERMUTATION_USE_SINGLE_GUIDE == 1
-		Out.VertexWeights = float3(1, 0, 0);
-		Out.VertexIndices = 0;
-		Out.VertexLerps = 0;
-
-		const uint InterpolationData = InterpolationBuffer[InVertexIndex];
-		Out.VertexIndices.x = InterpolationData & 0xFFFFFF;
-		Out.VertexLerps.x = saturate(((InterpolationData >> 24) & 0xFF) / 255.f);
-	#else
-		const uint4 InterpolationData0 = Interpolation0Buffer[InVertexIndex];
-		const uint4 InterpolationData1 = Interpolation1Buffer[InVertexIndex];
-
-		Out.VertexIndices.x = InterpolationData0.x | ((InterpolationData1.x & 0xFF) << 16);
-		Out.VertexIndices.y = InterpolationData0.y | (((InterpolationData1.x >> 8) & 0xFF) << 16);
-		Out.VertexIndices.z = InterpolationData0.z | ((InterpolationData1.y & 0xFF) << 16);
-
-		Out.VertexWeights = float3(saturate((InterpolationData0.w & 0xFF) / 255.f), saturate(((InterpolationData0.w >> 8) & 0xFF) / 255.f), 0);
-		Out.VertexWeights.z = saturate(1.0f - (Out.VertexWeights.x + Out.VertexWeights.y));
-
-		Out.VertexLerps.x = saturate(((InterpolationData1.y >> 8) & 0xFF) / 255.f);
-		Out.VertexLerps.y = saturate((InterpolationData1.z & 0xFF) / 255.f);
-		Out.VertexLerps.z = saturate(((InterpolationData1.z >> 8) & 0xFF) / 255.f);
-	#endif
-	}
+#if PERMUTATION_SIMULATION == 1
+  #if PERMUTATION_USE_SINGLE_GUIDE
+	return GetGuideData(InVertexIndex, InterpolationBuffer[InVertexIndex], OutSimHairPositionOffsetBuffer[0].xyz);
+  #else
+	return GetGuideData(InVertexIndex, Interpolation0Buffer[InVertexIndex], Interpolation1Buffer[InVertexIndex], OutSimHairPositionOffsetBuffer[0].xyz);
+  #endif
+#else
+	return (FGuideData)0;
 #endif
-	return Out;
 }
 
 ///////////////////////////////////////////////////////////////////////////////////////////////////
@@ -430,34 +430,6 @@ float3 ApplySkinningWithTranslationAndRotation(const FGuideData GuideData, const
 	return CurrOffset;
 }
 
-///////////////////////////////////////////////////////////////////////////////////////////////////
-// Patch attribute with debug data
-
-float4 ApplyDebugOutput(const uint InMode, const FGuideData GuideData,  const uint InClusterId)
-{
-	float4 Out = 0;
-	#if PERMUTATION_DEBUG == 1
-	if (InMode == 1)
-	{
-		const uint GuideIndex = GuideData.VertexIndices[0]; // Take the closest guide
-		Out = SimAttributeBuffer[GuideIndex];
-	}
-	else if (InMode == 2)
-	{
-		const uint ClusterIdLo8bits = InClusterId & 0xFF;
-		const uint ClusterIdHi8btis = (InClusterId >> 8) & 0xFF;
-		// Using FNV1 hash to break linearity of ClusterId (generated by linearly parsing cluster grid)
-		uint Hash = 0;
-		Hash = Hash * 17;
-		Hash = Hash ^ ClusterIdLo8bits;
-		Hash = Hash * 17;
-		Hash = Hash ^ ClusterIdHi8btis;
-		Out = asfloat(uint4(0, (Hash & 0xFF) << 8, 0, 0));
-	}
-	#endif
-	return Out;
-}
-
 [numthreads(GROUP_SIZE, 1, 1)]
 void MainCS(uint2 DispatchThreadId : SV_DispatchThreadID)
 {
@@ -530,17 +502,71 @@ void MainCS(uint2 DispatchThreadId : SV_DispatchThreadID)
 			}
 		}
 		OutRenderDeformedPositionBuffer[VertexIndex] = PackHairControlPointPosition(PackedRestPosition, OutPosition, OutRenHairPositionOffsetBuffer[0].xyz);
+	}
+}
+#endif // SHADER_HAIRINTERPOLATION
+
+///////////////////////////////////////////////////////////////////////////////////////////////////
+// Patch attribute (for debug visualization)
+#if SHADER_PATCHATTRIBUTE
+
+uint VertexCount;
+
+Buffer<float4>	 RenAttribute0Buffer;
+Buffer<int>		 RenVertexToClusterIdBuffer;
 
-		// 4. Optionally patch groom attribute for debug rendering
-		#if PERMUTATION_DEBUG == 1
-		if (HairDebugMode == 1 || HairDebugMode == 2)
+Buffer<float4>	 SimAttribute0Buffer;
+#if PERMUTATION_SIMULATION == 1
+Buffer<uint>	InterpolationBuffer;
+#elif PERMUTATION_SIMULATION == 2
+Buffer<uint4>	Interpolation0Buffer;
+Buffer<uint4>	Interpolation1Buffer;
+#endif
+RWBuffer<float4> OutRenAttribute0Buffer;
+
+
+FGuideData GetGuideData(uint InVertexIndex)
+{
+#if PERMUTATION_SIMULATION == 1
+	return GetGuideData(InVertexIndex, InterpolationBuffer[InVertexIndex], 0 /*OutPositionOffset*/);
+#elif PERMUTATION_SIMULATION == 2
+	return GetGuideData(InVertexIndex, Interpolation0Buffer[InVertexIndex], Interpolation1Buffer[InVertexIndex], 0 /*OutPositionOffset*/);
+#else
+	return (FGuideData)0;
+#endif
+}
+
+[numthreads(1024, 1, 1)]
+void MainCS(uint2 DispatchThreadId : SV_DispatchThreadID)
+{
+	const uint VertexIndex = DispatchThreadId.x;
+	if (VertexIndex < VertexCount)
+	{
+		float4 Out = 0;
+		#if PERMUTATION_SIMULATION == 0
 		{
-			OutRenderAttributeBuffer[VertexIndex] = ApplyDebugOutput(HairDebugMode, GuideData, VertexToClusterIdBuffer[VertexIndex]);
+			const uint ClusterId		= RenVertexToClusterIdBuffer[VertexIndex];
+			const uint ClusterIdLo8bits = ClusterId & 0xFF;
+			const uint ClusterIdHi8btis = (ClusterId >> 8) & 0xFF;
+			// Using FNV1 hash to break linearity of ClusterId (generated by linearly parsing cluster grid)
+			uint Hash = 0;
+			Hash = Hash * 17;
+			Hash = Hash ^ ClusterIdLo8bits;
+			Hash = Hash * 17;
+			Hash = Hash ^ ClusterIdHi8btis;
+			Out = float4(0, (Hash & 0xFF) / 255.f, 0, 0);
 		}
-		#endif
+		#else // if (PERMUTATION_SIMULATION > 0)
+		{
+			const FGuideData GuideData = GetGuideData(VertexIndex);
+			const uint GuideIndex = GuideData.VertexIndices[0]; // Take the closest guide
+			Out = SimAttribute0Buffer[GuideIndex];
+		}
+		#endif // PERMUTATION_SIMULATION
+		OutRenAttribute0Buffer[VertexIndex] = Out;
 	}
 }
-#endif // SHADER_HAIRINTERPOLATION
+#endif // SHADER_PATCHATTRIBUTE
 
 ///////////////////////////////////////////////////////////////////////////////////////////////////
 // Hair clusters update 
