* Move previous position logic outside of hair interoplation pass.
* Rename for consistency hair interoplation buffer.

diff --git a/Engine/Plugins/Runtime/HairStrands/Source/HairStrandsCore/Private/HairStrandsRendering.cpp b/Engine/Plugins/Runtime/HairStrands/Source/HairStrandsCore/Private/HairStrandsRendering.cpp
--- a/Engine/Plugins/Runtime/HairStrands/Source/HairStrandsCore/Private/HairStrandsRendering.cpp
+++ b/Engine/Plugins/Runtime/HairStrands/Source/HairStrandsCore/Private/HairStrandsRendering.cpp
@@ -35,7 +35,6 @@ static FAutoConsoleVariableRef CVarHairRaytracingRadiusScale(TEXT("r.HairStrands
 static int GHairRaytracingProceduralSplits = 4;
 static FAutoConsoleVariableRef CVarHairRaytracingProceduralScale(TEXT("r.HairStrands.RaytracingProceduralSplits"), GHairRaytracingProceduralSplits, TEXT("Change how many AABBs are used per hair segment to balance between BVH build cost and ray tracing performance. (default: 4)"));
 
-
 static float GStrandHairWidth = 0.0f;
 static FAutoConsoleVariableRef CVarStrandHairWidth(TEXT("r.HairStrands.StrandWidth"), GStrandHairWidth, TEXT("Width of hair strand"));
 
@@ -491,28 +490,29 @@ class FHairInterpolationCS : public FGlobalShader
 		SHADER_PARAMETER(uint32, VertexCount)
 		SHADER_PARAMETER(uint32, DispatchCountX)
 		SHADER_PARAMETER(float, HairLengthScale)
-		SHADER_PARAMETER(FVector3f, InRenderHairPositionOffset)
+		SHADER_PARAMETER(FVector3f, InRenHairPositionOffset)
 		SHADER_PARAMETER(FVector3f, InSimHairPositionOffset)
 		SHADER_PARAMETER(uint32,  HairStrandsVF_bIsCullingEnable)
 		SHADER_PARAMETER(FMatrix44f, LocalToWorldMatrix)
 
-		SHADER_PARAMETER_RDG_BUFFER_SRV(Buffer, RenderRestPosePositionBuffer)
-		SHADER_PARAMETER_RDG_BUFFER_UAV(RWBuffer, OutRenderDeformedPositionBuffer)
+		SHADER_PARAMETER_RDG_BUFFER_SRV(Buffer, RenRestPosePositionBuffer)
+		SHADER_PARAMETER_RDG_BUFFER_SRV(Buffer, RenDeformerPositionBuffer)
+		SHADER_PARAMETER_RDG_BUFFER_UAV(RWBuffer, OutRenDeformedPositionBuffer)
 
 		SHADER_PARAMETER_RDG_BUFFER_SRV(Buffer, SimRestPosePositionBuffer)
-		SHADER_PARAMETER_RDG_BUFFER_SRV(Buffer, DeformedSimPositionBuffer)
+		SHADER_PARAMETER_RDG_BUFFER_SRV(Buffer, SimDeformedPositionBuffer)
 
 		SHADER_PARAMETER_RDG_BUFFER_SRV(Buffer, InterpolationBuffer)
 		SHADER_PARAMETER_RDG_BUFFER_SRV(Buffer, Interpolation0Buffer)
 		SHADER_PARAMETER_RDG_BUFFER_SRV(Buffer, Interpolation1Buffer)
 
-		SHADER_PARAMETER_RDG_BUFFER_SRV(Buffer<float4>, RestPosition0Buffer)
-		SHADER_PARAMETER_RDG_BUFFER_SRV(Buffer<float4>, RestPosition1Buffer)
-		SHADER_PARAMETER_RDG_BUFFER_SRV(Buffer<float4>, RestPosition2Buffer)
+		SHADER_PARAMETER_RDG_BUFFER_SRV(Buffer<float4>, RenRestPosition0Buffer)
+		SHADER_PARAMETER_RDG_BUFFER_SRV(Buffer<float4>, RenRestPosition1Buffer)
+		SHADER_PARAMETER_RDG_BUFFER_SRV(Buffer<float4>, RenRestPosition2Buffer)
 
-		SHADER_PARAMETER_RDG_BUFFER_SRV(Buffer<float4>, DeformedPosition0Buffer)
-		SHADER_PARAMETER_RDG_BUFFER_SRV(Buffer<float4>, DeformedPosition1Buffer)
-		SHADER_PARAMETER_RDG_BUFFER_SRV(Buffer<float4>, DeformedPosition2Buffer)
+		SHADER_PARAMETER_RDG_BUFFER_SRV(Buffer<float4>, RenDeformedPosition0Buffer)
+		SHADER_PARAMETER_RDG_BUFFER_SRV(Buffer<float4>, RenDeformedPosition1Buffer)
+		SHADER_PARAMETER_RDG_BUFFER_SRV(Buffer<float4>, RenDeformedPosition2Buffer)
 
 		SHADER_PARAMETER_RDG_BUFFER_SRV(Buffer<uint>, RenRootBarycentricBuffer)
 		SHADER_PARAMETER_RDG_BUFFER_SRV(Buffer<uint>, RenVertexToRootIndexBuffer)
@@ -573,7 +573,7 @@ static void AddHairStrandsInterpolationPass(
 	const EHairInterpolationType HairInterpolationType,
 	const EHairGeometryType InstanceGeometryType,
 	const FRDGHairStrandsCullingData& CullingData,
-	const FVector& InRenderHairWorldOffset,
+	const FVector& InRenHairWorldOffset,
 	const FVector& InSimHairWorldOffset,
 	const FRDGBufferSRVRef& OutRenHairPositionOffsetBuffer,
 	const FRDGBufferSRVRef& OutSimHairPositionOffsetBuffer,
@@ -581,15 +581,14 @@ static void AddHairStrandsInterpolationPass(
 	const FHairStrandsRestRootResource* SimRestRootResources,
 	const FHairStrandsDeformedRootResource* RenDeformedRootResources,
 	const FHairStrandsDeformedRootResource* SimDeformedRootResources,
-	const FRDGBufferSRVRef& RenderRestPosePositionBuffer,
+	const FRDGBufferSRVRef& RenRestPosePositionBuffer,
 	const bool bUseSingleGuide,
 	const FRDGBufferSRVRef& InterpolationBuffer,
 	const FRDGBufferSRVRef& Interpolation0Buffer,
 	const FRDGBufferSRVRef& Interpolation1Buffer,
 	const FRDGBufferSRVRef& SimRestPosePositionBuffer,
 	const FRDGBufferSRVRef& SimDeformedPositionBuffer,
-	FRDGImportedBuffer& OutRenderPositionBuffer,
-	FRDGImportedBuffer* OutRenderPrevPositionBuffer,
+	FRDGImportedBuffer& OutRenPositionBuffer,
 	const FRDGBufferSRVRef& SimRootPointIndexBuffer,
 	const FHairStrandsDeformedRootResource::FLOD::EFrameType DeformedFrame)
 {
@@ -598,9 +597,9 @@ static void AddHairStrandsInterpolationPass(
 
 
 	FHairInterpolationCS::FParameters* Parameters = GraphBuilder.AllocParameters<FHairInterpolationCS::FParameters>();
-	Parameters->RenderRestPosePositionBuffer = RenderRestPosePositionBuffer;
+	Parameters->RenRestPosePositionBuffer = RenRestPosePositionBuffer;
 	Parameters->SimRestPosePositionBuffer = SimRestPosePositionBuffer;
-	Parameters->DeformedSimPositionBuffer = SimDeformedPositionBuffer;
+	Parameters->SimDeformedPositionBuffer = SimDeformedPositionBuffer;
 	if (bUseSingleGuide)
 	{
 		Parameters->InterpolationBuffer = InterpolationBuffer;
@@ -610,11 +609,11 @@ static void AddHairStrandsInterpolationPass(
 		Parameters->Interpolation0Buffer = Interpolation0Buffer;
 		Parameters->Interpolation1Buffer = Interpolation1Buffer;
 	}
-	Parameters->OutRenderDeformedPositionBuffer = OutRenderPositionBuffer.UAV;
+	Parameters->OutRenDeformedPositionBuffer = OutRenPositionBuffer.UAV;
 
 	Parameters->VertexStart = VertexStart;
 	Parameters->VertexCount = VertexCount;
-	Parameters->InRenderHairPositionOffset = (FVector3f)InRenderHairWorldOffset;
+	Parameters->InRenHairPositionOffset = (FVector3f)InRenHairWorldOffset;
 	Parameters->InSimHairPositionOffset = (FVector3f)InSimHairWorldOffset;
 
 	Parameters->OutSimHairPositionOffsetBuffer = OutSimHairPositionOffsetBuffer;
@@ -688,16 +687,16 @@ static void AddHairStrandsInterpolationPass(
 		const FHairStrandsRestRootResource::FLOD& Ren_RestLODDatas = RenRestRootResources->LODs[MeshLODIndex];
 		const FHairStrandsDeformedRootResource::FLOD& Ren_DeformedLODDatas = RenDeformedRootResources->LODs[MeshLODIndex];
 		{
-			Parameters->RestPosition0Buffer = RegisterAsSRV(GraphBuilder, Ren_RestLODDatas.RestUniqueTrianglePosition0Buffer);
-			Parameters->RestPosition1Buffer = RegisterAsSRV(GraphBuilder, Ren_RestLODDatas.RestUniqueTrianglePosition1Buffer);
-			Parameters->RestPosition2Buffer = RegisterAsSRV(GraphBuilder, Ren_RestLODDatas.RestUniqueTrianglePosition2Buffer);
+			Parameters->RenRestPosition0Buffer = RegisterAsSRV(GraphBuilder, Ren_RestLODDatas.RestUniqueTrianglePosition0Buffer);
+			Parameters->RenRestPosition1Buffer = RegisterAsSRV(GraphBuilder, Ren_RestLODDatas.RestUniqueTrianglePosition1Buffer);
+			Parameters->RenRestPosition2Buffer = RegisterAsSRV(GraphBuilder, Ren_RestLODDatas.RestUniqueTrianglePosition2Buffer);
 			Parameters->RenRootBarycentricBuffer = RegisterAsSRV(GraphBuilder, Ren_RestLODDatas.RootBarycentricBuffer);
 			Parameters->RenRootToUniqueTriangleIndexBuffer = RegisterAsSRV(GraphBuilder, Ren_RestLODDatas.RootToUniqueTriangleIndexBuffer);
 		}
 		{
-			Parameters->DeformedPosition0Buffer = RegisterAsSRV(GraphBuilder, Ren_DeformedLODDatas.GetDeformedUniqueTrianglePosition0Buffer(DeformedFrame));
-			Parameters->DeformedPosition1Buffer = RegisterAsSRV(GraphBuilder, Ren_DeformedLODDatas.GetDeformedUniqueTrianglePosition1Buffer(DeformedFrame));
-			Parameters->DeformedPosition2Buffer = RegisterAsSRV(GraphBuilder, Ren_DeformedLODDatas.GetDeformedUniqueTrianglePosition2Buffer(DeformedFrame));
+			Parameters->RenDeformedPosition0Buffer = RegisterAsSRV(GraphBuilder, Ren_DeformedLODDatas.GetDeformedUniqueTrianglePosition0Buffer(DeformedFrame));
+			Parameters->RenDeformedPosition1Buffer = RegisterAsSRV(GraphBuilder, Ren_DeformedLODDatas.GetDeformedUniqueTrianglePosition1Buffer(DeformedFrame));
+			Parameters->RenDeformedPosition2Buffer = RegisterAsSRV(GraphBuilder, Ren_DeformedLODDatas.GetDeformedUniqueTrianglePosition2Buffer(DeformedFrame));
 		}
 	}
 
@@ -756,12 +755,7 @@ static void AddHairStrandsInterpolationPass(
 			DispatchCount);
 	}
 
-	GraphBuilder.SetBufferAccessFinal(OutRenderPositionBuffer.Buffer, ERHIAccess::SRVMask);
-	if (!(IsHairVisibilityComputeRasterContinuousLODEnabled() || IsHairVisibilityComputeRasterTemporalLayeringEnabled()) && Instance->HairGroupPublicData->VFInput.bHasLODSwitch && OutRenderPrevPositionBuffer && GHairStrandsTransferPositionOnLODChange > 0)
-	{
-		AddTransferPositionPass(GraphBuilder, ShaderMap, VertexCount, OutRenderPositionBuffer.SRV, OutRenderPrevPositionBuffer->UAV);
-		GraphBuilder.SetBufferAccessFinal(OutRenderPrevPositionBuffer->Buffer, ERHIAccess::SRVMask);
-	}
+	GraphBuilder.SetBufferAccessFinal(OutRenPositionBuffer.Buffer, ERHIAccess::SRVMask);
 }
 
 ///////////////////////////////////////////////////////////////////////////////////////////////////
@@ -1918,6 +1912,7 @@ void ComputeHairStrandsInterpolation(
 					VertexStart = Instance->HairGroupPublicData->GetActiveStrandsVertexStart(Instance->Strands.RestResource->GetVertexCount());
 					VertexCount = Instance->HairGroupPublicData->GetActiveStrandsVertexCount(Instance->Strands.RestResource->GetVertexCount(), Instance->HairGroupPublicData->MaxScreenSize);
 
+					// 2.1.1 Current Position
 					const bool bUseSingleGuide = bValidGuide && Instance->Strands.InterpolationResource->UseSingleGuide();
 					AddHairStrandsInterpolationPass(
 						GraphBuilder,
@@ -1947,12 +1942,23 @@ void ComputeHairStrandsInterpolation(
 						bValidGuide ? RegisterAsSRV(GraphBuilder, Instance->Guides.RestResource->PositionBuffer) : nullptr,
 						bValidGuide ? RegisterAsSRV(GraphBuilder, Instance->Guides.DeformedResource->GetBuffer(FHairStrandsDeformedResource::Current)) : nullptr,
 						Strands_DeformedPosition,
-						&Strands_DeformedPrevPosition,
 						bValidGuide ? RegisterAsSRV(GraphBuilder, Instance->Strands.InterpolationResource->SimRootPointIndexBuffer) : nullptr, 
 						FHairStrandsDeformedRootResource::FLOD::Current);
 					
-					// second interpolation pass to calculate previous positions if needed
-					if (IsHairVisibilityComputeRasterContinuousLODEnabled() || IsHairVisibilityComputeRasterTemporalLayeringEnabled())
+					// 2.1.2 Previous Position
+					// * Transfer prev. position
+					// * Or recompute interpolation pass for previous positions if needed
+					const bool bRecomputePrevPosition = IsHairVisibilityComputeRasterContinuousLODEnabled() || IsHairVisibilityComputeRasterTemporalLayeringEnabled();
+					const bool bTransferPrevPosition = 
+						!bRecomputePrevPosition &&
+						Instance->HairGroupPublicData->VFInput.bHasLODSwitch && 
+						GHairStrandsTransferPositionOnLODChange > 0;
+					if (bTransferPrevPosition)
+					{
+						AddTransferPositionPass(GraphBuilder, ShaderMap, VertexCount, Strands_DeformedPosition.SRV, Strands_DeformedPrevPosition.UAV);
+						GraphBuilder.SetBufferAccessFinal(Strands_DeformedPrevPosition.Buffer, ERHIAccess::SRVMask);
+					}
+					else if (bRecomputePrevPosition)
 					{
 						Strands_PrevPositionOffsetSRV = RegisterAsSRV(GraphBuilder, Instance->Strands.DeformedResource->GetPositionOffsetBuffer(FHairStrandsDeformedResource::EFrameType::Previous));
 
@@ -1984,7 +1990,6 @@ void ComputeHairStrandsInterpolation(
 							bValidGuide ? RegisterAsSRV(GraphBuilder, Instance->Guides.RestResource->PositionBuffer) : nullptr,
 							bValidGuide ? RegisterAsSRV(GraphBuilder, Instance->Guides.DeformedResource->GetBuffer(FHairStrandsDeformedResource::Previous)) : nullptr,
 							Strands_DeformedPrevPosition,
-							nullptr,
 							bValidGuide ? RegisterAsSRV(GraphBuilder, Instance->Strands.InterpolationResource->SimRootPointIndexBuffer) : nullptr,
 							FHairStrandsDeformedRootResource::FLOD::Previous);
 					}
@@ -2266,7 +2271,6 @@ void ComputeHairStrandsInterpolation(
 						bValidGuide ? RegisterAsSRV(GraphBuilder, Instance->Guides.RestResource->PositionBuffer) : nullptr,
 						bValidGuide ? RegisterAsSRV(GraphBuilder, Instance->Guides.DeformedResource->GetBuffer(FHairStrandsDeformedResource::Current)) : nullptr,
 						Guides_DeformedPositionBuffer,
-						nullptr,
 						RegisterAsSRV(GraphBuilder, LOD.Guides.InterpolationResource->SimRootPointIndexBuffer),
 						FHairStrandsDeformedRootResource::FLOD::Current); // <- this should be optional
 
diff --git a/Engine/Shaders/Private/HairStrands/HairStrandsInterpolation.usf b/Engine/Shaders/Private/HairStrands/HairStrandsInterpolation.usf
--- a/Engine/Shaders/Private/HairStrands/HairStrandsInterpolation.usf
+++ b/Engine/Shaders/Private/HairStrands/HairStrandsInterpolation.usf
@@ -76,13 +76,13 @@ uint VertexCount;
 uint DispatchCountX;
 uint HairDebugMode;
 float HairLengthScale;
-float3 InRenderHairPositionOffset;
+float3 InRenHairPositionOffset;
 float3 InSimHairPositionOffset;
 float4x4 LocalToWorldMatrix;
 
-Buffer<uint4>  RenderRestPosePositionBuffer;
+Buffer<uint4>  RenRestPosePositionBuffer;
 Buffer<uint4>  SimRestPosePositionBuffer;
-Buffer<uint4>  DeformedSimPositionBuffer;
+Buffer<uint4>  SimDeformedPositionBuffer;
 
 // Hair interpolation support either 
 // * A single guide, 
@@ -110,18 +110,18 @@ Buffer<uint>   RenVertexToRootIndexBuffer;
 Buffer<uint>   SimVertexToRootIndexBuffer;
 #endif
 
-RWBuffer<uint4> OutRenderDeformedPositionBuffer; 
+RWBuffer<uint4> OutRenDeformedPositionBuffer;
 
 #if PERMUTATION_DYNAMIC_GEOMETRY >= 1 
 
 // Compact all these buffers into 2 buffers: translation + quaternion
-Buffer<float4>  RestPosition0Buffer;
-Buffer<float4>  RestPosition1Buffer;
-Buffer<float4>  RestPosition2Buffer;
+Buffer<float4>  RenRestPosition0Buffer;
+Buffer<float4>  RenRestPosition1Buffer;
+Buffer<float4>  RenRestPosition2Buffer;
 
-Buffer<float4>  DeformedPosition0Buffer;
-Buffer<float4>  DeformedPosition1Buffer;
-Buffer<float4>  DeformedPosition2Buffer;
+Buffer<float4>  RenDeformedPosition0Buffer;
+Buffer<float4>  RenDeformedPosition1Buffer;
+Buffer<float4>  RenDeformedPosition2Buffer;
 
 Buffer<uint>	RenRootBarycentricBuffer;
 Buffer<uint>	RenRootToUniqueTriangleIndexBuffer;
@@ -159,7 +159,7 @@ float3 ComputeStaticGeometryOffset(
 {
 	const float3 RestGuidePoint = f16tof32(SimRestPosePositionBuffer[GuideIndex].xyz) + InSimHairPositionOffset;
 
-	const float3 DeformedGuidePoint = f16tof32(DeformedSimPositionBuffer[GuideIndex].xyz) + OutSimHairPositionOffset;
+	const float3 DeformedGuidePoint = f16tof32(SimDeformedPositionBuffer[GuideIndex].xyz) + OutSimHairPositionOffset;
 
 	return (DeformedGuidePoint - RestGuidePoint) * GuideVertexWeight;
 }
@@ -172,7 +172,7 @@ float3 ComputeLocalGeometryOffset(
 
 	const float3 RestGuidePoint = f16tof32(SimRestPosePositionBuffer[GuideIndex].xyz) - f16tof32(SimRestPosePositionBuffer[RootIndex].xyz) ;
 
-	const float3 DeformedGuidePoint = f16tof32(DeformedSimPositionBuffer[GuideIndex].xyz) - f16tof32(DeformedSimPositionBuffer[RootIndex].xyz) ;
+	const float3 DeformedGuidePoint = f16tof32(SimDeformedPositionBuffer[GuideIndex].xyz) - f16tof32(SimDeformedPositionBuffer[RootIndex].xyz) ;
 
 	return (DeformedGuidePoint - RestGuidePoint) * GuideVertexWeight;
 }
@@ -188,7 +188,7 @@ float3 ComputeDynamicGeometryOffset(
 	const float3 RestGuidePoint = f16tof32(SimRestPosePositionBuffer[GuideIndex].xyz) + InSimHairPositionOffset;
 	const float3 LocalRestGuidePoint = ToTriangle(RestGuidePoint, RootBarycentric, RestTriangle);
 
-	const float3 DeformedGuidePoint = f16tof32(DeformedSimPositionBuffer[GuideIndex].xyz) + OutSimHairPositionOffset;
+	const float3 DeformedGuidePoint = f16tof32(SimDeformedPositionBuffer[GuideIndex].xyz) + OutSimHairPositionOffset;
 	const float3 LocalDeformedGuidePoint = ToTriangle(DeformedGuidePoint, RootBarycentric, DeformedTriangle);
 
 	return (LocalDeformedGuidePoint - LocalRestGuidePoint) * GuideVertexWeight;
@@ -202,7 +202,7 @@ float3 ComputeTranslateGeometryOffset(
 
 	const float3 RestGuidePoint = f16tof32(SimRestPosePositionBuffer[GuideIndex].xyz) - f16tof32(SimRestPosePositionBuffer[RootIndex].xyz) ;
 
-	const float3 DeformedGuidePoint = f16tof32(DeformedSimPositionBuffer[GuideIndex].xyz) - f16tof32(DeformedSimPositionBuffer[RootIndex].xyz) ;
+	const float3 DeformedGuidePoint = f16tof32(SimDeformedPositionBuffer[GuideIndex].xyz) - f16tof32(SimDeformedPositionBuffer[RootIndex].xyz) ;
 
 	return (DeformedGuidePoint - RestGuidePoint) * GuideVertexWeight;
 }
@@ -218,7 +218,7 @@ float3 ComputeRotateGeometryOffset(
 	const float3 RestGuidePoint = f16tof32(SimRestPosePositionBuffer[GuideIndex].xyz) - f16tof32(SimRestPosePositionBuffer[RootIndex].xyz) ;
 	const float3 LocalRestGuidePoint = VectorToTriangle(RestGuidePoint, RestTriangle);
 
-	const float3 DeformedGuidePoint = f16tof32(DeformedSimPositionBuffer[GuideIndex].xyz) - f16tof32(DeformedSimPositionBuffer[RootIndex].xyz) ;
+	const float3 DeformedGuidePoint = f16tof32(SimDeformedPositionBuffer[GuideIndex].xyz) - f16tof32(SimDeformedPositionBuffer[RootIndex].xyz) ;
 	const float3 LocalDeformedGuidePoint = VectorToTriangle(DeformedGuidePoint, DeformedTriangle);
 
 	return (LocalDeformedGuidePoint - LocalRestGuidePoint) * GuideVertexWeight;
@@ -253,8 +253,8 @@ FRootTriangleData GetRootTriangleData(uint VertexIndex, bool bSim)
 		Out.RootIndex		 = RenVertexToRootIndexBuffer[VertexIndex];
 		Out.TriangleIndex	 = RenRootToUniqueTriangleIndexBuffer[Out.RootIndex];
 		Out.RootBarycentric  = DecodeBarycentrics(RenRootBarycentricBuffer[Out.RootIndex]);
-		Out.RestTriangle	 = GetTriangleTransformation(Out.TriangleIndex, RestPosition0Buffer, RestPosition1Buffer, RestPosition2Buffer);
-		Out.DeformedTriangle = GetTriangleTransformation(Out.TriangleIndex, DeformedPosition0Buffer, DeformedPosition1Buffer, DeformedPosition2Buffer);
+		Out.RestTriangle	 = GetTriangleTransformation(Out.TriangleIndex, RenRestPosition0Buffer, RenRestPosition1Buffer, RenRestPosition2Buffer);
+		Out.DeformedTriangle = GetTriangleTransformation(Out.TriangleIndex, RenDeformedPosition0Buffer, RenDeformedPosition1Buffer, RenDeformedPosition2Buffer);
 	}
 	#endif
 	return Out;
@@ -430,6 +430,9 @@ float3 ApplySkinningWithTranslationAndRotation(const FGuideData GuideData, const
 	return CurrOffset;
 }
 
+///////////////////////////////////////////////////////////////////////////////////////////////////
+// Interpolation
+
 [numthreads(GROUP_SIZE, 1, 1)]
 void MainCS(uint2 DispatchThreadId : SV_DispatchThreadID)
 {
@@ -456,8 +459,8 @@ void MainCS(uint2 DispatchThreadId : SV_DispatchThreadID)
 		const FRootTriangleData RenTri = GetRootTriangleData(VertexIndex, false);
 
 		// Manual decoding of the rest position
-		const uint4 PackedRestPosition = RenderRestPosePositionBuffer[VertexIndex];
-		const float3 RestPosition = f16tof32(PackedRestPosition.xyz) + InRenderHairPositionOffset;
+		const uint4 PackedRestPosition = RenRestPosePositionBuffer[VertexIndex];
+		const float3 RestPosition = f16tof32(PackedRestPosition.xyz) + InRenHairPositionOffset;
 
 		float3 OutPosition = 0.0f;
 
@@ -501,7 +504,7 @@ void MainCS(uint2 DispatchThreadId : SV_DispatchThreadID)
 				OutPosition = INFINITE_FLOAT;
 			}
 		}
-		OutRenderDeformedPositionBuffer[VertexIndex] = PackHairControlPointPosition(PackedRestPosition, OutPosition, OutRenHairPositionOffsetBuffer[0].xyz);
+		OutRenDeformedPositionBuffer[VertexIndex] = PackHairControlPointPosition(PackedRestPosition, OutPosition, OutRenHairPositionOffsetBuffer[0].xyz);
 	}
 }
 #endif // SHADER_HAIRINTERPOLATION
