Add MeshDeformer support for Groom.

This is the initial iteration which allows to assign a mesh deformer to a groom instance. The mesh deformer access has access to most hair control point attributes, and can write out modified position (moved/trimmed/...).

diff --git a/Engine/Plugins/Runtime/HairStrands/HairStrands.uplugin b/Engine/Plugins/Runtime/HairStrands/HairStrands.uplugin
--- a/Engine/Plugins/Runtime/HairStrands/HairStrands.uplugin
+++ b/Engine/Plugins/Runtime/HairStrands/HairStrands.uplugin
@@ -44,6 +44,10 @@
 		{
 			"Name": "GeometryCache",
 			"Enabled": true
+		},
+		{
+			"Name": "DeformerGraph",
+			"Enabled": true
 		}
 	]
 }
diff --git a/Engine/Plugins/Runtime/HairStrands/Shaders/Private/DeformerDataInterfaceGroom.ush b/Engine/Plugins/Runtime/HairStrands/Shaders/Private/DeformerDataInterfaceGroom.ush
new file mode 100644
--- /dev/null
+++ b/Engine/Plugins/Runtime/HairStrands/Shaders/Private/DeformerDataInterfaceGroom.ush
@@ -0,0 +1,92 @@
+// Copyright Epic Games, Inc. All Rights Reserved.
+
+#include "/Plugin/ComputeFramework/Private/ComputeKernelCommon.ush"
+#include "/Engine/Private/HairStrands/HairStrandsVertexFactoryCommon.ush"
+
+///////////////////////////////////////////////////////////////////////////////////////////////////
+// Bindings
+
+uint  {DataInterfaceName}_NumControlPoints;
+uint  {DataInterfaceName}_NumCurves;
+float {DataInterfaceName}_VF_HairRadius;
+float {DataInterfaceName}_VF_HairLength;
+Buffer<uint4>  {DataInterfaceName}_PositionBuffer;
+Buffer<float2> {DataInterfaceName}_Attribute0Buffer;
+Buffer<uint>   {DataInterfaceName}_Attribute1Buffer;
+
+///////////////////////////////////////////////////////////////////////////////////////////////////
+// Helpers
+
+FHairControlPoint ReadControlPointData(uint VertexIndex)
+{
+	const uint MaxVertices = {DataInterfaceName}_NumControlPoints;
+	if (VertexIndex >= MaxVertices)
+	{
+		return (FHairControlPoint)0;
+	}
+
+	// Hair deformer work in local space so no need for position offset
+	return UnpackHairControlPoint(
+		{DataInterfaceName}_PositionBuffer[VertexIndex], 
+		float3(0, 0, 0) /*InVF_PositionOffset*/,
+		{DataInterfaceName}_VF_HairRadius, 
+		1.0f /*RootScale*/, 
+		1.0f /*TipScale*/);
+}
+
+FHairAttributes ReadControlPointAttributeData(uint VertexIndex)
+{
+	const uint MaxVertices = {DataInterfaceName}_NumControlPoints;
+	if (VertexIndex >= MaxVertices)
+	{
+		return (FHairAttributes)0;
+	}
+
+	return DecodeHairAttribute(
+		{DataInterfaceName}_Attribute0Buffer[VertexIndex],
+		{DataInterfaceName}_Attribute1Buffer[VertexIndex],
+		{DataInterfaceName}_VF_HairLength);
+}
+
+///////////////////////////////////////////////////////////////////////////////////////////////////
+// Interface
+
+uint {DataInterfaceName}_ReadNumControlPoints()
+{
+	return {DataInterfaceName}_NumControlPoints;
+}
+
+uint {DataInterfaceName}_ReadNumCurves()
+{
+	return {DataInterfaceName}_NumCurves;
+}
+
+float3 ReadPosition_{DataInterfaceName}(uint VertexIndex)
+{
+	return ReadControlPointData(VertexIndex).Position;
+}
+
+float ReadRadius_{DataInterfaceName}(uint VertexIndex)
+{
+	return ReadControlPointData(VertexIndex).WorldRadius;
+}
+
+float ReadCoordU_{DataInterfaceName}(uint VertexIndex)
+{
+	return ReadControlPointData(VertexIndex).UCoord;
+}
+
+float ReadLength_{DataInterfaceName}(uint VertexIndex)
+{
+	return ReadControlPointAttributeData(VertexIndex).WorldLength;
+}
+
+float2 ReadRootUV_{DataInterfaceName}(uint VertexIndex)
+{
+	return ReadControlPointAttributeData(VertexIndex).RootUV;
+}
+
+float ReadSeed_{DataInterfaceName}(uint VertexIndex)
+{
+	return ReadControlPointAttributeData(VertexIndex).Seed;
+}
diff --git a/Engine/Plugins/Runtime/HairStrands/Shaders/Private/DeformerDataInterfaceGroomExec.ush b/Engine/Plugins/Runtime/HairStrands/Shaders/Private/DeformerDataInterfaceGroomExec.ush
new file mode 100644
--- /dev/null
+++ b/Engine/Plugins/Runtime/HairStrands/Shaders/Private/DeformerDataInterfaceGroomExec.ush
@@ -0,0 +1,10 @@
+// Copyright Epic Games, Inc. All Rights Reserved.
+
+#include "/Plugin/ComputeFramework/Private/ComputeKernelCommon.ush"
+
+int3 {DataInterfaceName}_NumThreads;
+
+int3 ReadNumThreads_{DataInterfaceName}()
+{
+	return {DataInterfaceName}_NumThreads;
+}
diff --git a/Engine/Plugins/Runtime/HairStrands/Shaders/Private/DeformerDataInterfaceGroomWrite.ush b/Engine/Plugins/Runtime/HairStrands/Shaders/Private/DeformerDataInterfaceGroomWrite.ush
new file mode 100644
--- /dev/null
+++ b/Engine/Plugins/Runtime/HairStrands/Shaders/Private/DeformerDataInterfaceGroomWrite.ush
@@ -0,0 +1,33 @@
+// Copyright Epic Games, Inc. All Rights Reserved.
+
+#include "/Plugin/ComputeFramework/Private/ComputeKernelCommon.ush"
+#include "/Engine/Private/HairStrands/HairStrandsVertexFactoryCommon.ush"
+
+uint {DataInterfaceName}_NumControlPoints;
+uint {DataInterfaceName}_NumCurves;
+uint {DataInterfaceName}_OutputStreamStart;
+RWBuffer<uint4> {DataInterfaceName}_PositionBufferUAV;
+
+uint ReadNumControlPoints_{DataInterfaceName}()
+{
+	return {DataInterfaceName}_NumControlPoints;
+}
+
+uint ReadNumCurves_{DataInterfaceName}()
+{
+	return {DataInterfaceName}_NumCurves;
+}
+
+void WritePosition_{DataInterfaceName}(uint ControlPointIndex, float3 Position)
+{
+	// TODO: pipe the other data of the control point
+	FHairControlPoint CP;
+	CP.Position = Position;
+	CP.WorldRadius = 1.f;
+	CP.UCoord = 0;
+	CP.Type = HAIR_CONTROLPOINT_INSIDE;
+
+	const uint4 PackedCP = PackHairControlPoint(CP, 0 /*InVF_PositionOffset*/, 1.0f /*InVF_Radius*/); // TODO
+	const uint BufferIndex = {DataInterfaceName}_OutputStreamStart + ControlPointIndex;
+	{DataInterfaceName}_PositionBufferUAV[BufferIndex] = PackedCP;
+}
\ No newline at end of file
diff --git a/Engine/Plugins/Runtime/HairStrands/Source/HairStrandsCore/HairStrandsCore.Build.cs b/Engine/Plugins/Runtime/HairStrands/Source/HairStrandsCore/HairStrandsCore.Build.cs
--- a/Engine/Plugins/Runtime/HairStrands/Source/HairStrandsCore/HairStrandsCore.Build.cs
+++ b/Engine/Plugins/Runtime/HairStrands/Source/HairStrandsCore/HairStrandsCore.Build.cs
@@ -21,8 +21,8 @@ namespace UnrealBuildTool.Rules
 					"Core",
 					"CoreUObject",
 					"Engine",
-                    "GeometryCache",
-                    "Projects",
+					"GeometryCache",
+					"Projects",
 					"MeshDescription",
 					"MovieScene",
 					"NiagaraCore",
@@ -32,7 +32,9 @@ namespace UnrealBuildTool.Rules
 					"Renderer",
 					"VectorVM",
 					"RHI",
-					"StaticMeshDescription"
+					"StaticMeshDescription",
+					"ComputeFramework",
+					"OptimusCore",
 				});
 
 			if (Target.bBuildEditor == true)
diff --git a/Engine/Plugins/Runtime/HairStrands/Source/HairStrandsCore/Private/Deformer/DeformerDataInterfaceGroom.cpp b/Engine/Plugins/Runtime/HairStrands/Source/HairStrandsCore/Private/Deformer/DeformerDataInterfaceGroom.cpp
new file mode 100644
--- /dev/null
+++ b/Engine/Plugins/Runtime/HairStrands/Source/HairStrandsCore/Private/Deformer/DeformerDataInterfaceGroom.cpp
@@ -0,0 +1,194 @@
+// Copyright Epic Games, Inc. All Rights Reserved.
+
+#include "DeformerDataInterfaceGroom.h"
+
+#include "ComputeFramework/ShaderParamTypeDefinition.h"
+#include "ComputeFramework/ComputeKernelPermutationVector.h"
+#include "OptimusDataDomain.h"
+#include "ShaderParameterMetadataBuilder.h"
+#include "GroomComponent.h"
+#include "GroomInstance.h"
+#include "RenderGraphBuilder.h"
+
+FString UOptimusGroomDataInterface::GetDisplayName() const
+{
+	return TEXT("Groom");
+}
+
+TArray<FOptimusCDIPinDefinition> UOptimusGroomDataInterface::GetPinDefinitions() const
+{
+	TArray<FOptimusCDIPinDefinition> Defs;
+	Defs.Add({"NumControlPoints", "ReadNumControlPoints"});
+	Defs.Add({"NumCurves", "ReadNumCurves" });
+	Defs.Add({"Position",  "ReadPosition", Optimus::DomainName::Vertex, "ReadPosition"});
+	Defs.Add({"Radius",    "ReadRadius",   Optimus::DomainName::Vertex, "ReadRadius" });
+	Defs.Add({"CoordU",    "ReadCoordU",   Optimus::DomainName::Vertex, "ReadCoordU" });
+	Defs.Add({"Length",    "ReadLength",   Optimus::DomainName::Vertex, "ReadLength" });
+	Defs.Add({"RootUV",    "ReadRootUV",   Optimus::DomainName::Vertex, "ReadRootUV" });
+	Defs.Add({"Seed",      "ReadSeed",     Optimus::DomainName::Vertex, "ReadSeed"   });
+	return Defs;
+}
+ 
+TSubclassOf<UActorComponent> UOptimusGroomDataInterface::GetRequiredComponentClass() const
+{
+	return UGroomComponent::StaticClass();
+}
+
+void UOptimusGroomDataInterface::GetSupportedInputs(TArray<FShaderFunctionDefinition>& OutFunctions) const
+{
+	OutFunctions.AddDefaulted_GetRef()
+		.SetName(TEXT("ReadNumControlPoints"))
+		.AddReturnType(EShaderFundamentalType::Uint);
+
+	OutFunctions.AddDefaulted_GetRef()
+		.SetName(TEXT("ReadNumCurves"))
+		.AddReturnType(EShaderFundamentalType::Uint);
+
+	OutFunctions.AddDefaulted_GetRef()
+		.SetName(TEXT("ReadPosition"))
+		.AddReturnType(EShaderFundamentalType::Float, 3)
+		.AddParam(EShaderFundamentalType::Uint);
+
+	OutFunctions.AddDefaulted_GetRef()
+		.SetName(TEXT("ReadRadius"))
+		.AddReturnType(EShaderFundamentalType::Float)
+		.AddParam(EShaderFundamentalType::Uint);
+
+	OutFunctions.AddDefaulted_GetRef()
+		.SetName(TEXT("ReadCoordU"))
+		.AddReturnType(EShaderFundamentalType::Float)
+		.AddParam(EShaderFundamentalType::Uint);
+
+	OutFunctions.AddDefaulted_GetRef()
+		.SetName(TEXT("ReadLength"))
+		.AddReturnType(EShaderFundamentalType::Float)
+		.AddParam(EShaderFundamentalType::Uint);
+
+	OutFunctions.AddDefaulted_GetRef()
+		.SetName(TEXT("ReadRootUV"))
+		.AddReturnType(EShaderFundamentalType::Float, 2)
+		.AddParam(EShaderFundamentalType::Uint);
+
+	OutFunctions.AddDefaulted_GetRef()
+		.SetName(TEXT("ReadSeed"))
+		.AddReturnType(EShaderFundamentalType::Float)
+		.AddParam(EShaderFundamentalType::Uint);
+}
+
+BEGIN_SHADER_PARAMETER_STRUCT(FGroomDataInterfaceParameters, )
+	SHADER_PARAMETER(uint32, NumControlPoints)
+	SHADER_PARAMETER(uint32, NumCurves)
+	SHADER_PARAMETER(float, VF_HairRadius)
+	SHADER_PARAMETER(float, VF_HairLength)
+	SHADER_PARAMETER_RDG_BUFFER_SRV(Buffer<uint4>,  PositionBuffer)
+	SHADER_PARAMETER_RDG_BUFFER_SRV(Buffer<float2>, Attribute0Buffer)
+	SHADER_PARAMETER_RDG_BUFFER_SRV(Buffer<uint>,   Attribute1Buffer)
+END_SHADER_PARAMETER_STRUCT()
+
+void UOptimusGroomDataInterface::GetShaderParameters(TCHAR const* UID, FShaderParametersMetadataBuilder& InOutBuilder, FShaderParametersMetadataAllocations& InOutAllocations) const
+{
+	InOutBuilder.AddNestedStruct<FGroomDataInterfaceParameters>(UID);
+}
+
+void UOptimusGroomDataInterface::GetPermutations(FComputeKernelPermutationVector& OutPermutationVector) const
+{
+}
+
+void UOptimusGroomDataInterface::GetShaderHash(FString& InOutKey) const
+{
+	GetShaderFileHash(TEXT("/Plugin/Runtime/HairStrands/Private/DeformerDataInterfaceGroom.ush"), EShaderPlatform::SP_PCD3D_SM5).AppendString(InOutKey);
+}
+
+void UOptimusGroomDataInterface::GetHLSL(FString& OutHLSL, FString const& InDataInterfaceName) const
+{
+	TMap<FString, FStringFormatArg> TemplateArgs =
+	{
+		{ TEXT("DataInterfaceName"), InDataInterfaceName },
+	};
+
+	FString TemplateFile;
+	LoadShaderSourceFile(TEXT("/Plugin/Runtime/HairStrands/Private/DeformerDataInterfaceGroom.ush"), EShaderPlatform::SP_PCD3D_SM5, &TemplateFile, nullptr);
+	OutHLSL += FString::Format(*TemplateFile, TemplateArgs);
+}
+
+UComputeDataProvider* UOptimusGroomDataInterface::CreateDataProvider(TObjectPtr<UObject> InBinding, uint64 InInputMask, uint64 InOutputMask) const
+{
+	UOptimusGroomDataProvider* Provider = NewObject<UOptimusGroomDataProvider>();
+	Provider->Groom = Cast<UGroomComponent>(InBinding);
+	return Provider;
+}
+
+bool UOptimusGroomDataProvider::IsValid() const
+{
+	return Groom != nullptr;
+}
+
+FComputeDataProviderRenderProxy* UOptimusGroomDataProvider::GetRenderProxy()
+{
+	return new FOptimusGroomDataProviderProxy(Groom);
+}
+
+FOptimusGroomDataProviderProxy::FOptimusGroomDataProviderProxy(UGroomComponent* InGroomComponent)
+{
+	GroomComponent = InGroomComponent;
+}
+
+void FOptimusGroomDataProviderProxy::AllocateResources(FRDGBuilder& GraphBuilder)
+{
+	Resources.Empty();
+	const uint32 InstanceCount = GroomComponent ? GroomComponent->GetGroupCount() : 0;
+	for (uint32 Index =0; Index <InstanceCount;++Index)
+	{
+		if (FHairGroupInstance* Instance = GroomComponent->GetGroupInstance(Index))
+		{
+			FResources& R = Resources.AddDefaulted_GetRef();
+			R.PositionSRV = Register(GraphBuilder, Instance->Strands.RestResource->PositionBuffer, ERDGImportedBufferFlags::CreateSRV).SRV;
+			R.Attribute0SRV = Register(GraphBuilder, Instance->Strands.RestResource->Attribute0Buffer, ERDGImportedBufferFlags::CreateSRV).SRV;
+			R.Attribute1SRV = Register(GraphBuilder, Instance->Strands.RestResource->Attribute1Buffer, ERDGImportedBufferFlags::CreateSRV).SRV;
+			R.FallbackSRV = GraphBuilder.CreateSRV(GraphBuilder.RegisterExternalBuffer(GWhiteVertexBufferWithRDG->Buffer), PF_R16G16B16A16_UINT);
+		}
+	}
+}
+
+FHairGroupPublicData::FVertexFactoryInput ComputeHairStrandsVertexInputData(const FHairGroupInstance* Instance);
+void FOptimusGroomDataProviderProxy::GatherDispatchData(FDispatchSetup const& InDispatchSetup, FCollectedDispatchData& InOutDispatchData)
+{
+	if (!ensure(InDispatchSetup.ParameterStructSizeForValidation == sizeof(FGroomDataInterfaceParameters)))
+	{
+		return;
+	}
+
+	const uint32 InstanceCount = GroomComponent ? GroomComponent->GetGroupCount() : 0;
+	check(InDispatchSetup.NumInvocations == InstanceCount);
+	check(uint32(Resources.Num()) == InstanceCount);
+	for (int32 InvocationIndex = 0; InvocationIndex < InDispatchSetup.NumInvocations; ++InvocationIndex)
+	{
+		if (FHairGroupInstance* Instance = GroomComponent->GetGroupInstance(InvocationIndex))
+		{
+			const bool bIsSRVValid = Resources[InvocationIndex].PositionSRV != nullptr;
+			const int32 NumControlPoints = bIsSRVValid ? Instance->Strands.Data->PointCount : 0;
+			const int32 NumCurves = bIsSRVValid ? Instance->Strands.Data->CurveCount : 0;
+
+			const FHairGroupPublicData::FVertexFactoryInput VFInput = ComputeHairStrandsVertexInputData(Instance);
+			
+			FGroomDataInterfaceParameters* Parameters = (FGroomDataInterfaceParameters*)(InOutDispatchData.ParameterBuffer + InDispatchSetup.ParameterBufferOffset + InDispatchSetup.ParameterBufferStride * InvocationIndex);
+			Parameters->NumControlPoints = NumControlPoints;
+			Parameters->NumCurves = NumCurves;
+			Parameters->VF_HairRadius = VFInput.Strands.HairRadius;
+			Parameters->VF_HairLength = VFInput.Strands.HairLength;
+
+			if (bIsSRVValid)
+			{
+				Parameters->PositionBuffer = Resources[InvocationIndex].PositionSRV;
+				Parameters->Attribute0Buffer = Resources[InvocationIndex].Attribute0SRV;
+				Parameters->Attribute1Buffer = Resources[InvocationIndex].Attribute1SRV;
+			}
+			else
+			{
+				Parameters->PositionBuffer   = Resources[InvocationIndex].FallbackSRV;
+				Parameters->Attribute0Buffer = Resources[InvocationIndex].FallbackSRV;
+				Parameters->Attribute1Buffer = Resources[InvocationIndex].FallbackSRV;
+			}
+		}
+	}
+}
\ No newline at end of file
diff --git a/Engine/Plugins/Runtime/HairStrands/Source/HairStrandsCore/Private/Deformer/DeformerDataInterfaceGroom.h b/Engine/Plugins/Runtime/HairStrands/Source/HairStrandsCore/Private/Deformer/DeformerDataInterfaceGroom.h
new file mode 100644
--- /dev/null
+++ b/Engine/Plugins/Runtime/HairStrands/Source/HairStrandsCore/Private/Deformer/DeformerDataInterfaceGroom.h
@@ -0,0 +1,73 @@
+// Copyright Epic Games, Inc. All Rights Reserved.
+
+#pragma once
+
+#include "RenderGraphDefinitions.h"
+#include "OptimusComputeDataInterface.h"
+#include "ComputeFramework/ComputeDataProvider.h"
+#include "DeformerDataInterfaceGroom.generated.h"
+
+class UGroomComponent;
+
+/** Compute Framework Data Interface for reading groom. */
+UCLASS(Category = ComputeFramework)
+class HAIRSTRANDSCORE_API UOptimusGroomDataInterface : public UOptimusComputeDataInterface
+{
+	GENERATED_BODY()
+
+public:
+	//~ Begin UOptimusComputeDataInterface Interface
+	FString GetDisplayName() const override;
+	TArray<FOptimusCDIPinDefinition> GetPinDefinitions() const override;
+	TSubclassOf<UActorComponent> GetRequiredComponentClass() const override;
+	//~ End UOptimusComputeDataInterface Interface
+	
+	//~ Begin UComputeDataInterface Interface
+	TCHAR const* GetClassName() const override { return TEXT("Groom"); }
+	void GetSupportedInputs(TArray<FShaderFunctionDefinition>& OutFunctions) const override;
+	void GetShaderParameters(TCHAR const* UID, FShaderParametersMetadataBuilder& InOutBuilder, FShaderParametersMetadataAllocations& InOutAllocations) const override;
+	void GetPermutations(FComputeKernelPermutationVector& OutPermutationVector) const override;
+	void GetShaderHash(FString& InOutKey) const override;
+	void GetHLSL(FString& OutHLSL, FString const& InDataInterfaceName) const override;
+	UComputeDataProvider* CreateDataProvider(TObjectPtr<UObject> InBinding, uint64 InInputMask, uint64 InOutputMask) const override;
+	//~ End UComputeDataInterface Interface
+};
+
+/** Compute Framework Data Provider for reading groom. */
+UCLASS(BlueprintType, editinlinenew, Category = ComputeFramework)
+class UOptimusGroomDataProvider : public UComputeDataProvider
+{
+	GENERATED_BODY()
+
+public:
+	UPROPERTY(EditAnywhere, BlueprintReadWrite, Category = Binding)
+	TObjectPtr<UGroomComponent> Groom = nullptr;
+
+	//~ Begin UComputeDataProvider Interface
+	bool IsValid() const override;
+	FComputeDataProviderRenderProxy* GetRenderProxy() override;
+	//~ End UComputeDataProvider Interface
+};
+
+class FOptimusGroomDataProviderProxy : public FComputeDataProviderRenderProxy
+{
+public:
+	FOptimusGroomDataProviderProxy(UGroomComponent* InGroomComponent);
+
+	//~ Begin FComputeDataProviderRenderProxy Interface
+	void AllocateResources(FRDGBuilder& GraphBuilder) override;
+	void GatherDispatchData(FDispatchSetup const& InDispatchSetup, FCollectedDispatchData& InOutDispatchData) override;
+	//~ End FComputeDataProviderRenderProxy Interface
+
+private:
+
+	UGroomComponent* GroomComponent = nullptr; // Should it be HairInstance instead?
+	struct FResources
+	{
+		FRDGBufferSRVRef PositionSRV = nullptr;
+		FRDGBufferSRVRef Attribute0SRV = nullptr;
+		FRDGBufferSRVRef Attribute1SRV = nullptr;
+		FRDGBufferSRVRef FallbackSRV = nullptr;
+	};
+	TArray<FResources> Resources;
+};
diff --git a/Engine/Plugins/Runtime/HairStrands/Source/HairStrandsCore/Private/Deformer/DeformerDataInterfaceGroomExec.cpp b/Engine/Plugins/Runtime/HairStrands/Source/HairStrandsCore/Private/Deformer/DeformerDataInterfaceGroomExec.cpp
new file mode 100644
--- /dev/null
+++ b/Engine/Plugins/Runtime/HairStrands/Source/HairStrandsCore/Private/Deformer/DeformerDataInterfaceGroomExec.cpp
@@ -0,0 +1,134 @@
+// Copyright Epic Games, Inc. All Rights Reserved.
+
+#include "DeformerDataInterfaceGroomExec.h"
+
+#include "GroomComponent.h"
+#include "ComputeFramework/ShaderParamTypeDefinition.h"
+#include "ShaderParameterMetadataBuilder.h"
+#include "SkeletalRenderPublic.h"
+
+FString UOptimusGroomExecDataInterface::GetDisplayName() const
+{
+	return TEXT("Execute Groom");
+}
+
+FName UOptimusGroomExecDataInterface::GetCategory() const
+{
+	return CategoryName::ExecutionDataInterfaces;
+}
+
+TArray<FOptimusCDIPinDefinition> UOptimusGroomExecDataInterface::GetPinDefinitions() const
+{
+	TArray<FOptimusCDIPinDefinition> Defs;
+	Defs.Add({ "NumThreads", "ReadNumThreads" });
+	return Defs;
+}
+
+
+TSubclassOf<UActorComponent> UOptimusGroomExecDataInterface::GetRequiredComponentClass() const
+{
+	return UGroomComponent::StaticClass();
+}
+
+
+void UOptimusGroomExecDataInterface::GetSupportedInputs(TArray<FShaderFunctionDefinition>& OutFunctions) const
+{
+	OutFunctions.AddDefaulted_GetRef()
+		.SetName(TEXT("ReadNumThreads"))
+		.AddReturnType(EShaderFundamentalType::Int, 3);
+}
+
+
+BEGIN_SHADER_PARAMETER_STRUCT(FGroomExecDataInterfaceParameters, )
+	SHADER_PARAMETER(FIntVector, NumThreads)
+END_SHADER_PARAMETER_STRUCT()
+
+void UOptimusGroomExecDataInterface::GetShaderParameters(TCHAR const* UID, FShaderParametersMetadataBuilder& InOutBuilder, FShaderParametersMetadataAllocations& InOutAllocations) const
+{
+	InOutBuilder.AddNestedStruct<FGroomExecDataInterfaceParameters>(UID);
+}
+
+void UOptimusGroomExecDataInterface::GetShaderHash(FString& InOutKey) const
+{
+	GetShaderFileHash(TEXT("/Plugin/Runtime/HairStrands/Private/DeformerDataInterfaceGroomExec.ush"), EShaderPlatform::SP_PCD3D_SM5).AppendString(InOutKey);
+}
+
+void UOptimusGroomExecDataInterface::GetHLSL(FString& OutHLSL, FString const& InDataInterfaceName) const
+{
+	TMap<FString, FStringFormatArg> TemplateArgs =
+	{
+		{ TEXT("DataInterfaceName"), InDataInterfaceName },
+	};
+
+	FString TemplateFile;
+	LoadShaderSourceFile(TEXT("/Plugin/Runtime/HairStrands/Private/DeformerDataInterfaceGroomExec.ush"), EShaderPlatform::SP_PCD3D_SM5, &TemplateFile, nullptr);
+	OutHLSL += FString::Format(*TemplateFile, TemplateArgs);
+}
+
+UComputeDataProvider* UOptimusGroomExecDataInterface::CreateDataProvider(TObjectPtr<UObject> InBinding, uint64 InInputMask, uint64 InOutputMask) const
+{
+	UOptimusGroomExecDataProvider* Provider = NewObject<UOptimusGroomExecDataProvider>();
+	Provider->GroomComponent = Cast<UGroomComponent>(InBinding);
+	Provider->Domain = Domain;
+	return Provider;
+}
+
+bool UOptimusGroomExecDataProvider::IsValid() const
+{
+	return GroomComponent != nullptr && GroomComponent->GetGroupCount() > 0;
+}
+
+FComputeDataProviderRenderProxy* UOptimusGroomExecDataProvider::GetRenderProxy()
+{
+	return new FOptimusGroomExecDataProviderProxy(GroomComponent, Domain);
+}
+
+FOptimusGroomExecDataProviderProxy::FOptimusGroomExecDataProviderProxy(UGroomComponent* InGroomComponent, EOptimusGroomExecDomain InDomain)
+{
+	GroomComponent = InGroomComponent;
+	Domain = InDomain;
+}
+
+int32 FOptimusGroomExecDataProviderProxy::GetDispatchThreadCount(TArray<FIntVector>& ThreadCounts) const
+{
+	const int32 NumInvocations = GroomComponent ? GroomComponent->GetGroupCount() : 0;
+
+	ThreadCounts.Reset();
+	ThreadCounts.Reserve(NumInvocations);
+	for (int32 InvocationIndex = 0; InvocationIndex < NumInvocations; ++InvocationIndex)
+	{		
+		if (FHairGroupInstance* Instance = GroomComponent->GetGroupInstance(InvocationIndex))
+		{
+			const int32 NumControlPoints = Instance->Strands.Data->PointCount;
+			const int32 NumCurves = Instance->Strands.Data->CurveCount;
+			const int32 NumThreads = Domain == EOptimusGroomExecDomain::ControlPoint ? NumControlPoints : NumCurves;
+			ThreadCounts.Add(FIntVector(NumThreads, 1, 1));
+		}
+	}
+	return NumInvocations;
+}
+
+void FOptimusGroomExecDataProviderProxy::GatherDispatchData(FDispatchSetup const& InDispatchSetup, FCollectedDispatchData& InOutDispatchData)
+{
+	if (!ensure(InDispatchSetup.ParameterStructSizeForValidation == sizeof(FGroomExecDataInterfaceParameters)))
+	{
+		return;
+	}
+
+	const int32 NumInvocations = GroomComponent ? GroomComponent->GetGroupCount() : 0;
+	if (!ensure(NumInvocations == InDispatchSetup.NumInvocations))
+	{
+		return;
+	}
+
+	for (int32 InvocationIndex = 0; InvocationIndex < NumInvocations; ++InvocationIndex)
+	{
+		FHairGroupInstance* Instance = GroomComponent->GetGroupInstance(InvocationIndex);
+		const int32 NumControlPoints = Instance->Strands.Data->PointCount;
+		const int32 NumCurves = Instance->Strands.Data->CurveCount;
+		const int32 NumThreads = Domain == EOptimusGroomExecDomain::ControlPoint ? NumControlPoints : NumCurves;
+
+		FGroomExecDataInterfaceParameters* Parameters = (FGroomExecDataInterfaceParameters*)(InOutDispatchData.ParameterBuffer + InDispatchSetup.ParameterBufferOffset + InDispatchSetup.ParameterBufferStride * InvocationIndex);
+		Parameters->NumThreads = FIntVector(NumThreads, 1, 1);
+	}
+}
diff --git a/Engine/Plugins/Runtime/HairStrands/Source/HairStrandsCore/Private/Deformer/DeformerDataInterfaceGroomExec.h b/Engine/Plugins/Runtime/HairStrands/Source/HairStrandsCore/Private/Deformer/DeformerDataInterfaceGroomExec.h
new file mode 100644
--- /dev/null
+++ b/Engine/Plugins/Runtime/HairStrands/Source/HairStrandsCore/Private/Deformer/DeformerDataInterfaceGroomExec.h
@@ -0,0 +1,81 @@
+// Copyright Epic Games, Inc. All Rights Reserved.
+
+#pragma once
+
+#include "OptimusComputeDataInterface.h"
+#include "ComputeFramework/ComputeDataProvider.h"
+#include "DeformerDataInterfaceGroomExec.generated.h"
+
+class UGroomComponent;
+
+UENUM()
+enum class EOptimusGroomExecDomain : uint8
+{
+	None = 0 UMETA(Hidden),
+	/** Run kernel with one thread per control point. */
+	ControlPoint = 1,
+	/** Run kernel with one thread per curve. */
+	Curve,
+};
+
+/** Compute Framework Data Interface for executing kernels over a skinned mesh domain. */
+UCLASS(Category = ComputeFramework)
+class HAIRSTRANDSCORE_API UOptimusGroomExecDataInterface : public UOptimusComputeDataInterface
+{
+	GENERATED_BODY()
+
+public:
+	//~ Begin UOptimusComputeDataInterface Interface
+	FString GetDisplayName() const override;
+	FName GetCategory() const override;
+	TArray<FOptimusCDIPinDefinition> GetPinDefinitions() const override;
+	TSubclassOf<UActorComponent> GetRequiredComponentClass() const override;
+	//~ End UOptimusComputeDataInterface Interface
+	
+	//~ Begin UComputeDataInterface Interface
+	TCHAR const* GetClassName() const override { return TEXT("GroomExec"); }
+	bool IsExecutionInterface() const override { return true; }
+	void GetSupportedInputs(TArray<FShaderFunctionDefinition>& OutFunctions) const override;
+	void GetShaderParameters(TCHAR const* UID, FShaderParametersMetadataBuilder& InOutBuilder, FShaderParametersMetadataAllocations& InOutAllocations) const override;
+	void GetShaderHash(FString& InOutKey) const override;
+	void GetHLSL(FString& OutHLSL, FString const& InDataInterfaceName) const override;
+	UComputeDataProvider* CreateDataProvider(TObjectPtr<UObject> InBinding, uint64 InInputMask, uint64 InOutputMask) const override;
+	//~ End UComputeDataInterface Interface
+
+	UPROPERTY(EditAnywhere, Category = Execution)
+	EOptimusGroomExecDomain Domain = EOptimusGroomExecDomain::ControlPoint;
+};
+
+/** Compute Framework Data Provider for executing kernels over a skinned mesh domain. */
+UCLASS(BlueprintType, editinlinenew, Category = ComputeFramework)
+class UOptimusGroomExecDataProvider : public UComputeDataProvider
+{
+	GENERATED_BODY()
+
+public:
+	UPROPERTY()
+	TObjectPtr<UGroomComponent> GroomComponent = nullptr;
+
+	UPROPERTY()
+	EOptimusGroomExecDomain Domain = EOptimusGroomExecDomain::ControlPoint;
+
+	//~ Begin UComputeDataProvider Interface
+	bool IsValid() const override;
+	FComputeDataProviderRenderProxy* GetRenderProxy() override;
+	//~ End UComputeDataProvider Interface
+};
+
+class FOptimusGroomExecDataProviderProxy : public FComputeDataProviderRenderProxy
+{
+public:
+	FOptimusGroomExecDataProviderProxy(UGroomComponent* InGroomComponent, EOptimusGroomExecDomain InDomain);
+
+	//~ Begin FComputeDataProviderRenderProxy Interface
+	int32 GetDispatchThreadCount(TArray<FIntVector>& ThreadCounts) const override;
+	void GatherDispatchData(FDispatchSetup const& InDispatchSetup, FCollectedDispatchData& InOutDispatchData) override;
+	//~ End FComputeDataProviderRenderProxy Interface
+
+private:
+	UGroomComponent* GroomComponent = nullptr;
+	EOptimusGroomExecDomain Domain = EOptimusGroomExecDomain::ControlPoint;
+};
diff --git a/Engine/Plugins/Runtime/HairStrands/Source/HairStrandsCore/Private/Deformer/DeformerDataInterfaceGroomWrite.cpp b/Engine/Plugins/Runtime/HairStrands/Source/HairStrandsCore/Private/Deformer/DeformerDataInterfaceGroomWrite.cpp
new file mode 100644
--- /dev/null
+++ b/Engine/Plugins/Runtime/HairStrands/Source/HairStrandsCore/Private/Deformer/DeformerDataInterfaceGroomWrite.cpp
@@ -0,0 +1,164 @@
+// Copyright Epic Games, Inc. All Rights Reserved.
+
+#include "DeformerDataInterfaceGroomWrite.h"
+
+#include "ComputeFramework/ShaderParamTypeDefinition.h"
+#include "OptimusDataDomain.h"
+#include "RenderGraphBuilder.h"
+#include "Rendering/SkeletalMeshLODRenderData.h"
+#include "Rendering/SkeletalMeshRenderData.h"
+#include "ShaderParameterMetadataBuilder.h"
+#include "SkeletalMeshDeformerHelpers.h"
+#include "SkeletalRenderPublic.h"
+#include "GroomComponent.h"
+#include "GroomInstance.h"
+
+FString UOptimusGroomWriteDataInterface::GetDisplayName() const
+{
+	return TEXT("Write Groom");
+}
+
+FName UOptimusGroomWriteDataInterface::GetCategory() const
+{
+	return CategoryName::OutputDataInterfaces;
+}
+
+TArray<FOptimusCDIPinDefinition> UOptimusGroomWriteDataInterface::GetPinDefinitions() const
+{
+	TArray<FOptimusCDIPinDefinition> Defs;
+	Defs.Add({ "Position", "WritePosition", Optimus::DomainName::Vertex, "ReadNumControlPoints" }); // ControlPoint domain?
+	return Defs;
+}
+
+
+TSubclassOf<UActorComponent> UOptimusGroomWriteDataInterface::GetRequiredComponentClass() const
+{
+	return UGroomComponent::StaticClass();
+}
+
+
+void UOptimusGroomWriteDataInterface::GetSupportedInputs(TArray<FShaderFunctionDefinition>& OutFunctions) const
+{
+	OutFunctions.AddDefaulted_GetRef()
+		.SetName(TEXT("ReadNumControlPoints"))
+		.AddReturnType(EShaderFundamentalType::Uint);
+}
+
+void UOptimusGroomWriteDataInterface::GetSupportedOutputs(TArray<FShaderFunctionDefinition>& OutFunctions) const
+{
+	OutFunctions.AddDefaulted_GetRef()
+		.SetName(TEXT("WritePosition"))
+		.AddParam(EShaderFundamentalType::Uint)
+		.AddParam(EShaderFundamentalType::Float, 3);
+}
+
+BEGIN_SHADER_PARAMETER_STRUCT(FGroomWriteDataInterfaceParameters, )
+	SHADER_PARAMETER(uint32, NumControlPoints)
+	SHADER_PARAMETER(uint32, NumCurves)
+	SHADER_PARAMETER(uint32, OutputStreamStart)
+	SHADER_PARAMETER_RDG_BUFFER_UAV(RWBuffer<uint4>, PositionBufferUAV)
+END_SHADER_PARAMETER_STRUCT()
+
+void UOptimusGroomWriteDataInterface::GetShaderParameters(TCHAR const* UID, FShaderParametersMetadataBuilder& InOutBuilder, FShaderParametersMetadataAllocations& InOutAllocations) const
+{
+	InOutBuilder.AddNestedStruct<FGroomWriteDataInterfaceParameters>(UID);
+}
+
+void UOptimusGroomWriteDataInterface::GetShaderHash(FString& InOutKey) const
+{
+	GetShaderFileHash(TEXT("/Plugin/Runtime/HairStrands/Private/DeformerDataInterfaceGroomWrite.ush"), EShaderPlatform::SP_PCD3D_SM5).AppendString(InOutKey);
+}
+
+void UOptimusGroomWriteDataInterface::GetHLSL(FString& OutHLSL, FString const& InDataInterfaceName) const
+{
+	TMap<FString, FStringFormatArg> TemplateArgs =
+	{
+		{ TEXT("DataInterfaceName"), InDataInterfaceName },
+	};
+
+	FString TemplateFile;
+	LoadShaderSourceFile(TEXT("/Plugin/Runtime/HairStrands/Private/DeformerDataInterfaceGroomWrite.ush"), EShaderPlatform::SP_PCD3D_SM5, &TemplateFile, nullptr);
+	OutHLSL += FString::Format(*TemplateFile, TemplateArgs);
+}
+
+UComputeDataProvider* UOptimusGroomWriteDataInterface::CreateDataProvider(TObjectPtr<UObject> InBinding, uint64 InInputMask, uint64 InOutputMask) const
+{
+	UOptimusGroomWriteDataProvider* Provider = NewObject<UOptimusGroomWriteDataProvider>();
+	Provider->GroomComponent = Cast<UGroomComponent>(InBinding);
+	Provider->OutputMask = InOutputMask;
+	return Provider;
+}
+
+
+bool UOptimusGroomWriteDataProvider::IsValid() const
+{
+	return GroomComponent != nullptr && GroomComponent->GetGroupCount() > 0;
+}
+
+FComputeDataProviderRenderProxy* UOptimusGroomWriteDataProvider::GetRenderProxy()
+{
+	return new FOptimusGroomWriteDataProviderProxy(GroomComponent, OutputMask);
+}
+
+
+FOptimusGroomWriteDataProviderProxy::FOptimusGroomWriteDataProviderProxy(UGroomComponent* InGroomComponent, uint64 InOutputMask)
+{
+	GroomComponent = InGroomComponent;
+	OutputMask = InOutputMask;
+}
+
+void FOptimusGroomWriteDataProviderProxy::AllocateResources(FRDGBuilder& GraphBuilder)
+{
+	const uint32 InstanceCount = GroomComponent ? GroomComponent->GetGroupCount() : 0;
+	for (uint32 Index=0; Index < InstanceCount; ++Index)
+	{
+		if (FHairGroupInstance* GroomInstance = GroomComponent->GetGroupInstance(Index))
+		{
+			// Allocate required buffers
+			const int32 NumControlPoints = GroomInstance->Strands.Data->PointCount;
+			const int32 NumCurves = GroomInstance->Strands.Data->CurveCount;
+
+			FResources& R = Resources.AddDefaulted_GetRef();
+			if (OutputMask & 1)
+			{
+				R.PositionBufferUAV = Register(GraphBuilder, GroomInstance->Strands.DeformedResource->GetDeformerBuffer(GraphBuilder), ERDGImportedBufferFlags::CreateUAV).UAV;
+				
+				FRDGBufferRef PositionBuffer_fallback = GraphBuilder.CreateBuffer(FRDGBufferDesc::CreateBufferDesc(8, 1), TEXT("Groom.DeformedPositionBuffer"), ERDGBufferFlags::None);
+				R.PositionBufferUAV_fallback = GraphBuilder.CreateUAV(PositionBuffer_fallback, PF_R16G16B16A16_UINT, ERDGUnorderedAccessViewFlags::SkipBarrier);
+			}
+			else
+			{
+				R.PositionBufferUAV = GraphBuilder.CreateUAV(GraphBuilder.RegisterExternalBuffer(GWhiteVertexBufferWithRDG->Buffer), PF_R16G16B16A16_UINT);
+			}
+		}	
+	}
+}
+
+void FOptimusGroomWriteDataProviderProxy::GatherDispatchData(FDispatchSetup const& InDispatchSetup, FCollectedDispatchData& InOutDispatchData)
+{
+	if (!ensure(InDispatchSetup.ParameterStructSizeForValidation == sizeof(FGroomWriteDataInterfaceParameters)))
+	{
+		return;
+	}
+
+	const uint32 InstanceCount = GroomComponent ? GroomComponent->GetGroupCount() : 0;
+	check(InDispatchSetup.NumInvocations == InstanceCount);
+
+	for (int32 InvocationIndex = 0; InvocationIndex < InDispatchSetup.NumInvocations; ++InvocationIndex)
+	{
+		if (FHairGroupInstance* GroomInstance = GroomComponent->GetGroupInstance(InvocationIndex))
+		{
+			FResources& Resource = Resources[InvocationIndex];
+
+			const bool bValid = Resource.PositionBufferUAV != nullptr;
+			const int32 NumCurves = bValid ? GroomInstance->Strands.Data->CurveCount : 0;
+			const int32 NumControlPoints = bValid ? GroomInstance->Strands.Data->PointCount : 0;
+
+			FGroomWriteDataInterfaceParameters* Parameters = (FGroomWriteDataInterfaceParameters*)(InOutDispatchData.ParameterBuffer + InDispatchSetup.ParameterBufferOffset + InDispatchSetup.ParameterBufferStride * InvocationIndex);
+			Parameters->NumControlPoints = NumControlPoints;
+			Parameters->NumCurves = NumCurves;
+			Parameters->OutputStreamStart = 0;
+			Parameters->PositionBufferUAV = bValid ? Resource.PositionBufferUAV : Resource.PositionBufferUAV_fallback;
+		}
+	}
+}
\ No newline at end of file
diff --git a/Engine/Plugins/Runtime/HairStrands/Source/HairStrandsCore/Private/Deformer/DeformerDataInterfaceGroomWrite.h b/Engine/Plugins/Runtime/HairStrands/Source/HairStrandsCore/Private/Deformer/DeformerDataInterfaceGroomWrite.h
new file mode 100644
--- /dev/null
+++ b/Engine/Plugins/Runtime/HairStrands/Source/HairStrandsCore/Private/Deformer/DeformerDataInterfaceGroomWrite.h
@@ -0,0 +1,76 @@
+// Copyright Epic Games, Inc. All Rights Reserved.
+
+#pragma once
+
+#include "OptimusComputeDataInterface.h"
+#include "ComputeFramework/ComputeDataProvider.h"
+#include "DeformerDataInterfaceGroomWrite.generated.h"
+
+class UGroomComponent;
+class FRDGBuffer;
+class FRDGBufferUAV;
+
+/** Compute Framework Data Interface for writing skinned mesh. */
+UCLASS(Category = ComputeFramework)
+class HAIRSTRANDSCORE_API UOptimusGroomWriteDataInterface : public UOptimusComputeDataInterface
+{
+	GENERATED_BODY()
+
+public:
+	//~ Begin UOptimusComputeDataInterface Interface
+	FString GetDisplayName() const override;
+	FName GetCategory() const override;
+	TArray<FOptimusCDIPinDefinition> GetPinDefinitions() const override;
+	TSubclassOf<UActorComponent> GetRequiredComponentClass() const override;
+	//~ End UOptimusComputeDataInterface Interface
+	
+	//~ Begin UComputeDataInterface Interface
+	TCHAR const* GetClassName() const override { return TEXT("GroomWrite"); }
+	void GetSupportedInputs(TArray<FShaderFunctionDefinition>& OutFunctions) const override;
+	void GetSupportedOutputs(TArray<FShaderFunctionDefinition>& OutFunctions) const override;
+	void GetShaderParameters(TCHAR const* UID, FShaderParametersMetadataBuilder& InOutBuilder, FShaderParametersMetadataAllocations& InOutAllocations) const override;
+	void GetShaderHash(FString& InOutKey) const override;
+	void GetHLSL(FString& OutHLSL, FString const& InDataInterfaceName) const override;
+	UComputeDataProvider* CreateDataProvider(TObjectPtr<UObject> InBinding, uint64 InInputMask, uint64 InOutputMask) const override;
+	//~ End UComputeDataInterface Interface
+};
+
+/** Compute Framework Data Provider for writing skinned mesh. */
+UCLASS(BlueprintType, editinlinenew, Category = ComputeFramework)
+class UOptimusGroomWriteDataProvider : public UComputeDataProvider
+{
+	GENERATED_BODY()
+
+public:
+	UPROPERTY(EditAnywhere, BlueprintReadWrite, Category = Binding)
+	TObjectPtr<UGroomComponent> GroomComponent = nullptr;
+
+	uint64 OutputMask;
+
+	//~ Begin UComputeDataProvider Interface
+	bool IsValid() const override;
+	FComputeDataProviderRenderProxy* GetRenderProxy() override;
+	//~ End UComputeDataProvider Interface
+};
+
+class FOptimusGroomWriteDataProviderProxy : public FComputeDataProviderRenderProxy
+{
+public:
+	FOptimusGroomWriteDataProviderProxy(UGroomComponent* InGroomComponent, uint64 InOutputMask);
+
+	//~ Begin FComputeDataProviderRenderProxy Interface
+	void AllocateResources(FRDGBuilder& GraphBuilder) override;
+	void GatherDispatchData(FDispatchSetup const& InDispatchSetup, FCollectedDispatchData& InOutDispatchData) override;
+	//~ End FComputeDataProviderRenderProxy Interface
+
+private:
+	UGroomComponent* GroomComponent = nullptr;
+	uint64 OutputMask;
+
+	struct FResources
+	{
+		FRDGBufferUAV* PositionBufferUAV = nullptr;
+		FRDGBufferUAV* PositionBufferUAV_fallback = nullptr;
+	};
+	TArray<FResources> Resources;
+};
\ No newline at end of file
diff --git a/Engine/Plugins/Runtime/HairStrands/Source/HairStrandsCore/Private/Deformer/DeformerGroomComponentSource.cpp b/Engine/Plugins/Runtime/HairStrands/Source/HairStrandsCore/Private/Deformer/DeformerGroomComponentSource.cpp
new file mode 100644
--- /dev/null
+++ b/Engine/Plugins/Runtime/HairStrands/Source/HairStrandsCore/Private/Deformer/DeformerGroomComponentSource.cpp
@@ -0,0 +1,89 @@
+// Copyright Epic Games, Inc. All Rights Reserved.
+
+#include "DeformerGroomComponentSource.h"
+
+#include "GroomComponent.h"
+
+#define LOCTEXT_NAMESPACE "DeformersGroomComponentSource"
+
+static const FHairGroupInstance* GetGroomInstance(const UGroomComponent* In)
+{
+	// HAIR_TODO: apply to all hair group, or only selected one
+	if (In && In->GetGroupCount() > 0)
+	{
+		return In->GetGroupInstance(0);
+	}
+	return nullptr;
+}
+
+FName UOptimusGroomComponentSource::Domains::ControlPoint("ControlPoint");
+FName UOptimusGroomComponentSource::Domains::Curve("Curve");
+
+FText UOptimusGroomComponentSource::GetDisplayName() const
+{
+	return LOCTEXT("GroomComponent", "Groom Component");
+}
+
+TSubclassOf<UActorComponent> UOptimusGroomComponentSource::GetComponentClass() const
+{
+	return UGroomComponent::StaticClass();
+}
+
+bool UOptimusGroomComponentSource::IsUsableAsPrimarySource() const
+{
+	return GetComponentClass()->IsChildOf<UGroomComponent>();
+}
+
+TArray<FName> UOptimusGroomComponentSource::GetExecutionDomains() const
+{
+	return { Domains::ControlPoint, Domains::Curve };
+}
+
+int32 UOptimusGroomComponentSource::GetLodIndex(const UActorComponent* InComponent) const
+{
+	return 0;
+}
+
+bool UOptimusGroomComponentSource::GetComponentElementCountsForExecutionDomain(
+	FName InDomainName,
+	const UActorComponent* InComponent,
+	int32 InLodIndex,
+	TArray<int32>& OutInvocationElementCounts
+) const
+{
+	const UGroomComponent* GroomComponent = Cast<UGroomComponent>(InComponent);
+	if (!GroomComponent)
+	{
+		return false;
+	}
+
+	const FHairGroupInstance* Instance = GetGroomInstance(GroomComponent);
+	if (!Instance)
+	{
+		return 0;
+	}
+
+	const int32 NumVertices = Instance->Strands.Data->PointCount;
+	const int32 NumCurves = Instance->Strands.Data->CurveCount;
+
+	OutInvocationElementCounts.Reset();
+
+	if (InDomainName == Domains::ControlPoint || InDomainName == Domains::Curve)
+	{
+		const int32 NumInvocations = 1; //TODO: num hair group?
+
+		OutInvocationElementCounts.Reset();
+		OutInvocationElementCounts.Reserve(NumInvocations);
+		for (int32 InvocationIndex = 0; InvocationIndex < NumInvocations; ++InvocationIndex)
+		{
+			const int32 NumThreads = InDomainName == Domains::ControlPoint ? NumVertices : NumCurves;
+			OutInvocationElementCounts.Add(NumThreads);
+		}
+
+		return true;
+	}
+
+	// Unknown execution domain.
+	return false;
+}
+#undef LOCTEXT_NAMESPACE
diff --git a/Engine/Plugins/Runtime/HairStrands/Source/HairStrandsCore/Private/Deformer/DeformerGroomComponentSource.h b/Engine/Plugins/Runtime/HairStrands/Source/HairStrandsCore/Private/Deformer/DeformerGroomComponentSource.h
new file mode 100644
--- /dev/null
+++ b/Engine/Plugins/Runtime/HairStrands/Source/HairStrandsCore/Private/Deformer/DeformerGroomComponentSource.h
@@ -0,0 +1,28 @@
+// Copyright Epic Games, Inc. All Rights Reserved.
+
+#pragma once
+
+#include "OptimusComponentSource.h"
+#include "DeformerGroomComponentSource.generated.h"
+
+UCLASS()
+class UOptimusGroomComponentSource : public UOptimusComponentSource
+{
+	GENERATED_BODY()
+public:
+
+	struct Domains
+	{
+		static FName ControlPoint;
+		static FName Curve;
+	};
+
+	// UOptimusComponentSource implementations
+	FText GetDisplayName() const override;
+	FName GetBindingName() const override { return FName("Groom"); }
+	TSubclassOf<UActorComponent> GetComponentClass() const override;
+	TArray<FName> GetExecutionDomains() const override;
+	int32 GetLodIndex(const UActorComponent* InComponent) const override;
+	bool GetComponentElementCountsForExecutionDomain(FName InDomainName, const UActorComponent* InComponent, int32 InLodIndex, TArray<int32>& OutInvocationElementCounts) const override;
+	bool IsUsableAsPrimarySource() const override;
+};
diff --git a/Engine/Plugins/Runtime/HairStrands/Source/HairStrandsCore/Private/GroomComponent.cpp b/Engine/Plugins/Runtime/HairStrands/Source/HairStrandsCore/Private/GroomComponent.cpp
--- a/Engine/Plugins/Runtime/HairStrands/Source/HairStrandsCore/Private/GroomComponent.cpp
+++ b/Engine/Plugins/Runtime/HairStrands/Source/HairStrandsCore/Private/GroomComponent.cpp
@@ -21,6 +21,8 @@
 #include "Rendering/SkeletalMeshRenderData.h"
 #include "Engine/RendererSettings.h"
 #include "Animation/AnimationSettings.h"
+#include "Animation/MeshDeformerInstance.h"
+#include "Animation/MeshDeformer.h"
 #include "Logging/MessageLog.h"
 #include "Misc/UObjectToken.h"
 #include "Misc/MapErrors.h"
@@ -2743,6 +2745,11 @@ void UGroomComponent::InitResources(bool bIsBindingReloading)
 				{
 					bNeedDynamicResources = true;
 				}
+				// Mesh deformer requires to dynamic resources
+				if (MeshDeformer)
+				{
+					bNeedDynamicResources = true;
+				}
 			}
 
 			#if RHI_RAYTRACING
@@ -3175,6 +3182,8 @@ void UGroomComponent::OnRegister()
 		UpdateGroomCache(ElapsedTime);
 	}
 	UpdateHairSimulation();
+
+	MeshDeformerInstance = (MeshDeformer != nullptr) ? MeshDeformer->CreateInstance(this, MeshDeformerInstanceSettings) : nullptr;
 }
 
 void UGroomComponent::OnUnregister()
@@ -3192,6 +3201,8 @@ void UGroomComponent::OnUnregister()
 		}
 		IGroomCacheStreamingManager::Get().UnregisterComponent(this);
 	}
+
+	MeshDeformerInstance = nullptr;
 }
 
 void UGroomComponent::BeginDestroy()
@@ -3579,6 +3590,51 @@ void UGroomComponent::SendRenderDynamicData_Concurrent()
 			}
 		});
 	}
+
+	if (MeshDeformerInstance != nullptr && GroomAsset != nullptr)
+	{
+		if (MeshDeformerInstance->IsActive())
+		{
+			MeshDeformerInstance->EnqueueWork(GetScene(), UMeshDeformerInstance::WorkLoad_Update, GroomAsset->GetFName());
+		}
+		else
+		{
+			// TODO: Reset so groom appears rest pose
+		}
+	}
+}
+
+void UGroomComponent::CreateRenderState_Concurrent(FRegisterComponentContext* Context)
+{
+	if (MeshDeformerInstance)
+	{
+		MeshDeformerInstance->AllocateResources();
+	}
+
+	Super::CreateRenderState_Concurrent(Context);
+}
+
+void UGroomComponent::DestroyRenderState_Concurrent()
+{
+	Super::DestroyRenderState_Concurrent();
+
+	if (MeshDeformerInstance)
+	{
+		MeshDeformerInstance->ReleaseResources();
+	}
+}
+
+bool UGroomComponent::RequiresGameThreadEndOfFrameRecreate() const
+{
+	return Super::RequiresGameThreadEndOfFrameRecreate();
+}
+
+void UGroomComponent::SetMeshDeformer(UMeshDeformer* InMeshDeformer)
+{
+	MeshDeformer = InMeshDeformer;
+	MeshDeformerInstanceSettings = MeshDeformer->CreateSettingsInstance(this);
+	MeshDeformerInstance = (MeshDeformer != nullptr) ? MeshDeformer->CreateInstance(this, MeshDeformerInstanceSettings) : nullptr;
+	MarkRenderDynamicDataDirty();
 }
 
 void UGroomComponent::GetUsedMaterials(TArray<UMaterialInterface*>& OutMaterials, bool bGetDebugMaterials) const
@@ -3647,6 +3703,7 @@ void UGroomComponent::PostEditChangeProperty(FPropertyChangedEvent& PropertyChan
 	const bool bIsBindingCompatible = UGroomBindingAsset::IsCompatible(GroomAsset, BindingAsset, bValidationEnable);
 	const bool bEnableSolverChanged = PropertyName == GET_MEMBER_NAME_CHECKED(FHairSimulationSolver, bEnableSimulation);
 	const bool bGroomCacheChanged = PropertyName == GET_MEMBER_NAME_CHECKED(UGroomComponent, GroomCache);
+	const bool bMeshDeformerChanged = PropertyName == GET_MEMBER_NAME_CHECKED(UGroomComponent, MeshDeformer);
 	const bool bEnableLengthScaleChanged = PropertyName == GET_MEMBER_NAME_CHECKED(FHairGroupDesc, HairLengthScale);
 	if (!bIsBindingCompatible || !UGroomBindingAsset::IsBindingAssetValid(BindingAsset, false, bValidationEnable))
 	{
@@ -3703,7 +3760,7 @@ void UGroomComponent::PostEditChangeProperty(FPropertyChangedEvent& PropertyChan
 	// If material is assigned to the groom from the viewport (i.e., drag&drop a material from the content brown onto the groom geometry, it results into a unknown property). There is other case 
 	const bool bIsUnknown = PropertyThatChanged == nullptr;
 
-	const bool bRecreateResources = bAssetChanged || bBindingAssetChanged || bGroomCacheChanged || bEnableLengthScaleOverrideChanged || bIsUnknown || bSourceSkeletalMeshChanged || bRayTracingGeometryChanged || bEnableSolverChanged;
+	const bool bRecreateResources = bAssetChanged || bBindingAssetChanged || bGroomCacheChanged || bEnableLengthScaleOverrideChanged || bIsUnknown || bSourceSkeletalMeshChanged || bRayTracingGeometryChanged || bEnableSolverChanged|| bMeshDeformerChanged;
 	if (bRecreateResources)
 	{
 		// Release the resources before Super::PostEditChangeProperty so that they get
diff --git a/Engine/Plugins/Runtime/HairStrands/Source/HairStrandsCore/Private/GroomResources.cpp b/Engine/Plugins/Runtime/HairStrands/Source/HairStrandsCore/Private/GroomResources.cpp
--- a/Engine/Plugins/Runtime/HairStrands/Source/HairStrandsCore/Private/GroomResources.cpp
+++ b/Engine/Plugins/Runtime/HairStrands/Source/HairStrandsCore/Private/GroomResources.cpp
@@ -884,6 +884,16 @@ void FHairStrandsDeformedResource::InternalAllocate(FRDGBuilder& GraphBuilder)
 	InternalCreateVertexBufferRDG<FHairStrandsPositionOffsetFormat>(GraphBuilder, DefaultOffsets, DeformedOffsetBuffer[0], ToHairResourceDebugName(HAIRSTRANDS_RESOUCE_NAME(CurveType, Hair.StrandsDeformed_DeformedOffsetBuffer0), ResourceName), EHairResourceUsageType::Dynamic, ERDGInitialDataFlags::None);
 	InternalCreateVertexBufferRDG<FHairStrandsPositionOffsetFormat>(GraphBuilder, DefaultOffsets, DeformedOffsetBuffer[1], ToHairResourceDebugName(HAIRSTRANDS_RESOUCE_NAME(CurveType, Hair.StrandsDeformed_DeformedOffsetBuffer1), ResourceName), EHairResourceUsageType::Dynamic, ERDGInitialDataFlags::None);
 
+	// Note: DeformerBuffer is optionally/lazily allocated by a mesh-deformer graph
+}
+
+FRDGExternalBuffer& FHairStrandsDeformedResource::GetDeformerBuffer(FRDGBuilder& GraphBuilder)
+{
+	if (DeformerBuffer.Buffer == nullptr)
+	{
+		InternalCreateVertexBufferRDG<FHairStrandsPositionFormat>(GraphBuilder, BulkData.PointCount, DeformerBuffer, ToHairResourceDebugName(HAIRSTRANDS_RESOUCE_NAME(CurveType, Hair.StrandsDeformed_DeformerBuffer), ResourceName), EHairResourceUsageType::Dynamic);
+	}
+	return DeformerBuffer;
 }
 
 void FHairStrandsDeformedResource::InternalRelease()
@@ -891,6 +901,7 @@ void FHairStrandsDeformedResource::InternalRelease()
 	DeformedPositionBuffer[0].Release();
 	DeformedPositionBuffer[1].Release();
 	TangentBuffer.Release();
+	DeformerBuffer.Release();
 
 	DeformedOffsetBuffer[0].Release();
 	DeformedOffsetBuffer[1].Release();
diff --git a/Engine/Plugins/Runtime/HairStrands/Source/HairStrandsCore/Private/HairStrandsRendering.cpp b/Engine/Plugins/Runtime/HairStrands/Source/HairStrandsCore/Private/HairStrandsRendering.cpp
--- a/Engine/Plugins/Runtime/HairStrands/Source/HairStrandsCore/Private/HairStrandsRendering.cpp
+++ b/Engine/Plugins/Runtime/HairStrands/Source/HairStrandsCore/Private/HairStrandsRendering.cpp
@@ -482,7 +482,8 @@ class FHairInterpolationCS : public FGlobalShader
 	class FSimulation : SHADER_PERMUTATION_BOOL("PERMUTATION_SIMULATION");
 	class FSingleGuide : SHADER_PERMUTATION_BOOL("PERMUTATION_USE_SINGLE_GUIDE");
 	class FCulling : SHADER_PERMUTATION_BOOL("PERMUTATION_CULLING");
-	using FPermutationDomain = TShaderPermutationDomain<FGroupSize, FDynamicGeometry, FSimulation, FSingleGuide, FCulling>;
+	class FDeformer : SHADER_PERMUTATION_BOOL("PERMUTATION_DEFORMER");
+	using FPermutationDomain = TShaderPermutationDomain<FGroupSize, FDynamicGeometry, FSimulation, FSingleGuide, FCulling, FDeformer>;
 
 	BEGIN_SHADER_PARAMETER_STRUCT(FParameters, )
 		SHADER_PARAMETER_STRUCT_INCLUDE(ShaderPrint::FShaderParameters, ShaderDrawParameters)
@@ -588,8 +589,9 @@ static void AddHairStrandsInterpolationPass(
 	const FRDGBufferSRVRef& Interpolation1Buffer,
 	const FRDGBufferSRVRef& SimRestPosePositionBuffer,
 	const FRDGBufferSRVRef& SimDeformedPositionBuffer,
-	FRDGImportedBuffer& OutRenPositionBuffer,
 	const FRDGBufferSRVRef& SimRootPointIndexBuffer,
+	const FRDGBufferSRVRef& RenDeformerPositionBuffer,
+	FRDGImportedBuffer& OutRenPositionBuffer,
 	const FHairStrandsDeformedRootResource::FLOD::EFrameType DeformedFrame)
 {
 	const uint32 GroupSize = ComputeGroupSize();
@@ -700,6 +702,12 @@ static void AddHairStrandsInterpolationPass(
 		}
 	}
 
+	const bool bSupportDeformer = RenDeformerPositionBuffer != nullptr;
+	if (bSupportDeformer)
+	{
+		Parameters->RenDeformerPositionBuffer = RenDeformerPositionBuffer;
+	}
+
 	const bool bHasLocalDeformation = Instance->Guides.bIsSimulationEnable || bSupportGlobalInterpolation || Instance->Guides.bIsDeformationEnable;
 	const bool bCullingEnable = IsHairStrandContinuousDecimationReorderingEnabled() ? false : (InstanceGeometryType == EHairGeometryType::Strands && CullingData.bCullingResultAvailable); 	// TODO: improve reordering so that culling can be used effectively
 	Parameters->HairStrandsVF_bIsCullingEnable = bCullingEnable ? 1 : 0;
@@ -727,6 +735,7 @@ static void AddHairStrandsInterpolationPass(
 	PermutationVector.Set<FHairInterpolationCS::FSimulation>(bHasLocalDeformation);
 	PermutationVector.Set<FHairInterpolationCS::FSingleGuide>(bUseSingleGuide);
 	PermutationVector.Set<FHairInterpolationCS::FCulling>(bCullingEnable);
+	PermutationVector.Set<FHairInterpolationCS::FDeformer>(bSupportDeformer);
 
 	TShaderMapRef<FHairInterpolationCS> ComputeShader(ShaderMap, PermutationVector);
 
@@ -1824,6 +1833,7 @@ void ComputeHairStrandsInterpolation(
 				FRDGImportedBuffer Strands_DeformedPosition = Register(GraphBuilder, Instance->Strands.DeformedResource->GetBuffer(FHairStrandsDeformedResource::Current), ERDGImportedBufferFlags::CreateViews);
 				FRDGImportedBuffer Strands_DeformedPrevPosition = Register(GraphBuilder, Instance->Strands.DeformedResource->GetBuffer(FHairStrandsDeformedResource::Previous), ERDGImportedBufferFlags::CreateViews);
 				FRDGImportedBuffer Strands_DeformedTangent = Register(GraphBuilder, Instance->Strands.DeformedResource->TangentBuffer, ERDGImportedBufferFlags::CreateViews);
+				FRDGBufferSRVRef Strands_DeformerPositionSRV = Register(GraphBuilder, Instance->Strands.DeformedResource->DeformerBuffer, ERDGImportedBufferFlags::CreateSRV).SRV;
 
 				// Trach on which view the position has been update, so that we can ensure motion vector are coherent
 				Instance->Strands.DeformedResource->GetUniqueViewID(FHairStrandsDeformedResource::Current) = ViewUniqueID;
@@ -1941,8 +1951,9 @@ void ComputeHairStrandsInterpolation(
 						bValidGuide && !bUseSingleGuide ? RegisterAsSRV(GraphBuilder, Instance->Strands.InterpolationResource->Interpolation1Buffer) : nullptr,
 						bValidGuide ? RegisterAsSRV(GraphBuilder, Instance->Guides.RestResource->PositionBuffer) : nullptr,
 						bValidGuide ? RegisterAsSRV(GraphBuilder, Instance->Guides.DeformedResource->GetBuffer(FHairStrandsDeformedResource::Current)) : nullptr,
+						bValidGuide ? RegisterAsSRV(GraphBuilder, Instance->Strands.InterpolationResource->SimRootPointIndexBuffer) : nullptr,
+						Strands_DeformerPositionSRV,
 						Strands_DeformedPosition,
-						bValidGuide ? RegisterAsSRV(GraphBuilder, Instance->Strands.InterpolationResource->SimRootPointIndexBuffer) : nullptr, 
 						FHairStrandsDeformedRootResource::FLOD::Current);
 					
 					// 2.1.2 Previous Position
@@ -1989,8 +2000,9 @@ void ComputeHairStrandsInterpolation(
 							bValidGuide && !bUseSingleGuide ? RegisterAsSRV(GraphBuilder, Instance->Strands.InterpolationResource->Interpolation1Buffer) : nullptr,
 							bValidGuide ? RegisterAsSRV(GraphBuilder, Instance->Guides.RestResource->PositionBuffer) : nullptr,
 							bValidGuide ? RegisterAsSRV(GraphBuilder, Instance->Guides.DeformedResource->GetBuffer(FHairStrandsDeformedResource::Previous)) : nullptr,
-							Strands_DeformedPrevPosition,
 							bValidGuide ? RegisterAsSRV(GraphBuilder, Instance->Strands.InterpolationResource->SimRootPointIndexBuffer) : nullptr,
+							Strands_DeformerPositionSRV,
+							Strands_DeformedPrevPosition,
 							FHairStrandsDeformedRootResource::FLOD::Previous);
 					}
 				}
@@ -2270,8 +2282,9 @@ void ComputeHairStrandsInterpolation(
 						bValidGuide && !bUseSingleGuide ? RegisterAsSRV(GraphBuilder, LOD.Guides.InterpolationResource->Interpolation1Buffer) : nullptr,
 						bValidGuide ? RegisterAsSRV(GraphBuilder, Instance->Guides.RestResource->PositionBuffer) : nullptr,
 						bValidGuide ? RegisterAsSRV(GraphBuilder, Instance->Guides.DeformedResource->GetBuffer(FHairStrandsDeformedResource::Current)) : nullptr,
-						Guides_DeformedPositionBuffer,
 						RegisterAsSRV(GraphBuilder, LOD.Guides.InterpolationResource->SimRootPointIndexBuffer),
+						nullptr,
+						Guides_DeformedPositionBuffer,
 						FHairStrandsDeformedRootResource::FLOD::Current); // <- this should be optional
 
 					AddHairCardsDeformationPass(
diff --git a/Engine/Plugins/Runtime/HairStrands/Source/HairStrandsCore/Public/GroomComponent.h b/Engine/Plugins/Runtime/HairStrands/Source/HairStrandsCore/Public/GroomComponent.h
--- a/Engine/Plugins/Runtime/HairStrands/Source/HairStrandsCore/Public/GroomComponent.h
+++ b/Engine/Plugins/Runtime/HairStrands/Source/HairStrandsCore/Public/GroomComponent.h
@@ -11,15 +11,19 @@
 #include "RHIDefinitions.h"
 #include "GroomDesc.h"
 #include "LODSyncInterface.h"
+#include "MeshDeformerInterface.h"
 #include "GroomInstance.h"
 #include "NiagaraDataInterfacePhysicsAsset.h"
 
 #include "GroomComponent.generated.h"
 
 class UGroomCache;
+class UMeshDeformer;
+class UMeshDeformerInstance;
+class UMeshDeformerInstanceSettings;
 
 UCLASS(HideCategories = (Object, Physics, Activation, Mobility, "Components|Activation"), editinlinenew, meta = (BlueprintSpawnableComponent), ClassGroup = Rendering)
-class HAIRSTRANDSCORE_API UGroomComponent : public UMeshComponent, public ILODSyncInterface, public INiagaraPhysicsAssetDICollectorInterface
+class HAIRSTRANDSCORE_API UGroomComponent : public UMeshComponent, public ILODSyncInterface, public INiagaraPhysicsAssetDICollectorInterface, public IMeshDeformerInterface
 {
 	GENERATED_UCLASS_BODY()
 
@@ -55,6 +59,17 @@ public:
 	UPROPERTY(EditAnywhere, BlueprintReadWrite, interp, Category = "Simulation")
 	FHairSimulationSettings SimulationSettings;
 
+	/** If set the MeshDeformer will be applied on groonm instance for deformation. */
+	UPROPERTY(EditAnywhere, BlueprintReadOnly, Category = "Deformer")
+	TObjectPtr<UMeshDeformer> MeshDeformer;
+
+	/** Object containing state for the bound MeshDeformer. */
+	UPROPERTY(Transient, BlueprintReadOnly, Category = "Deformer")
+	TObjectPtr<UMeshDeformerInstance> MeshDeformerInstance;
+
+	UPROPERTY(VisibleAnywhere, BlueprintReadOnly, DisplayName = "Settings", Category = "Deformer", meta = (EditInline, EditCondition = "MeshDeformerInstanceSettings!=nullptr", HideEditConditionToggle, EditConditionHides))
+	TObjectPtr<UMeshDeformerInstanceSettings> MeshDeformerInstanceSettings;
+
 	/* Reference of the default/debug materials for each geometric representation */
 	UPROPERTY()
 	TObjectPtr<UMaterialInterface> Strands_DebugMaterial;
@@ -122,7 +137,10 @@ public:
 	virtual void DetachFromComponent(const FDetachmentTransformRules& DetachmentRules) override;
 	virtual void TickComponent(float DeltaTime, enum ELevelTick TickType, FActorComponentTickFunction *ThisTickFunction) override;
 	virtual void SendRenderTransform_Concurrent() override;
+	virtual void CreateRenderState_Concurrent(FRegisterComponentContext* Context) override;
 	virtual void SendRenderDynamicData_Concurrent() override;
+	virtual void DestroyRenderState_Concurrent() override;
+	virtual bool RequiresGameThreadEndOfFrameRecreate() const override;
 	//~ End UActorComponent Interface.
 
 	//~ Begin USceneComponent Interface.
@@ -174,6 +192,10 @@ public:
 	UFUNCTION(BlueprintCallable, Category = "Simulation")
 	void SetPhysicsAsset(UPhysicsAsset* InPhysicsAsset);
 
+	/* Change the MeshDeformer that is used for this Component. */
+	UFUNCTION(BlueprintCallable, Category = "Components|SkinnedMesh")
+	void SetMeshDeformer(UMeshDeformer* InMeshDeformer);
+
 	/* Add a skeletal mesh to the collision components */
 	UFUNCTION(BlueprintCallable, Category = "Simulation")
 	void AddCollisionComponent(USkeletalMeshComponent* SkeletalMeshComponent);
@@ -220,6 +242,10 @@ public:
 	virtual int32 GetCurrentSyncLOD() const override;
 	//~ End ILODSyncInterface
 
+	///~ Begin IMeshDeformerInterface Interface.
+	virtual UMeshDeformerInstance const* GetMeshDeformerInstance() const override { return MeshDeformerInstance; }
+	//~ End IMeshDeformerInterface
+
 	int32 GetNumLODs() const;
 	int32 GetForcedLOD() const;
 	void  SetForcedLOD(int32 LODIndex);
diff --git a/Engine/Plugins/Runtime/HairStrands/Source/HairStrandsCore/Public/GroomResources.h b/Engine/Plugins/Runtime/HairStrands/Source/HairStrandsCore/Public/GroomResources.h
--- a/Engine/Plugins/Runtime/HairStrands/Source/HairStrandsCore/Public/GroomResources.h
+++ b/Engine/Plugins/Runtime/HairStrands/Source/HairStrandsCore/Public/GroomResources.h
@@ -372,6 +372,7 @@ struct FHairStrandsDeformedResource : public FHairCommonResource
 		Total += GetBufferTotalNumBytes(DeformedPositionBuffer[0]);
 		Total += GetBufferTotalNumBytes(DeformedPositionBuffer[1]);
 		Total += GetBufferTotalNumBytes(TangentBuffer);
+		Total += GetBufferTotalNumBytes(DeformerBuffer);
 		return Total;
 	}
 
@@ -384,6 +385,9 @@ struct FHairStrandsDeformedResource : public FHairCommonResource
 	/* Strand hair tangent buffer */
 	FRDGExternalBuffer TangentBuffer;
 
+	/* Strand hair deformer buffer. This buffer is optionally created & filled in by a mesh deformer */
+	FRDGExternalBuffer DeformerBuffer;
+
 	/* Position offset as the deformed positions are expressed in relative coordinate (16bits) */
 	FVector PositionOffset[2] = {FVector::ZeroVector, FVector::ZeroVector};
 
@@ -419,6 +423,9 @@ struct FHairStrandsDeformedResource : public FHairCommonResource
 	inline uint32& GetUniqueViewID(EFrameType T)				{ return UniqueViewIDs[GetIndex(T)]; }
 	inline uint32 GetUniqueViewID(EFrameType T) const			{ return UniqueViewIDs[GetIndex(T)]; }
 	//bool NeedsToUpdateTangent();
+
+	// Return deformer buffer
+	FRDGExternalBuffer& GetDeformerBuffer(FRDGBuilder& GraphBuilder);
 };
 
 struct FHairStrandsClusterCullingResource : public FHairCommonResource
diff --git a/Engine/Shaders/Private/HairStrands/HairStrandsInterpolation.usf b/Engine/Shaders/Private/HairStrands/HairStrandsInterpolation.usf
--- a/Engine/Shaders/Private/HairStrands/HairStrandsInterpolation.usf
+++ b/Engine/Shaders/Private/HairStrands/HairStrandsInterpolation.usf
@@ -110,7 +110,10 @@ Buffer<uint>   RenVertexToRootIndexBuffer;
 Buffer<uint>   SimVertexToRootIndexBuffer;
 #endif
 
-RWBuffer<uint4> OutRenDeformedPositionBuffer;
+#if PERMUTATION_DEFORMER
+Buffer<uint4>  RenDeformerPositionBuffer;
+#endif
+RWBuffer<uint4> OutRenDeformedPositionBuffer; 
 
 #if PERMUTATION_DYNAMIC_GEOMETRY >= 1 
 
@@ -430,6 +433,26 @@ float3 ApplySkinningWithTranslationAndRotation(const FGuideData GuideData, const
 	return CurrOffset;
 }
 
+///////////////////////////////////////////////////////////////////////////////////////////////////
+// Deformer transformation
+float3 ApplyDeformer(const FRootTriangleData RenTri, float3 InRestPostion, float3 InDeformerPosition, float3 InDeformedPosition)
+{
+	float3 OutPosition = InDeformedPosition;
+#if PERMUTATION_DEFORMER
+	#if PERMUTATION_DYNAMIC_GEOMETRY == 0
+		OutPosition = InDeformedPosition + (InDeformerPosition - InRestPostion);
+	#else
+		const float3 LocalRestPoint     = ToTriangle(InRestPostion, RenTri.RootBarycentric, RenTri.RestTriangle);
+		const float3 LocalDeformerPoint = ToTriangle(InDeformerPosition, RenTri.RootBarycentric, RenTri.RestTriangle);
+
+		const float3 LocalOffset = LocalDeformerPoint - LocalRestPoint;
+		const float3 LocalDeformedPoint = ToTriangle(InDeformedPosition, RenTri.RootBarycentric, RenTri.DeformedTriangle);
+		OutPosition = ToWorld(LocalDeformedPoint + LocalOffset, RenTri.RootBarycentric, RenTri.DeformedTriangle);
+	#endif
+#endif
+	return OutPosition;
+}
+
 ///////////////////////////////////////////////////////////////////////////////////////////////////
 // Interpolation
 
@@ -460,7 +483,7 @@ void MainCS(uint2 DispatchThreadId : SV_DispatchThreadID)
 
 		// Manual decoding of the rest position
 		const uint4 PackedRestPosition = RenRestPosePositionBuffer[VertexIndex];
-		const float3 RestPosition = f16tof32(PackedRestPosition.xyz) + InRenHairPositionOffset;
+		const float3 RestPosition = UnpackHairControlPointPosition(PackedRestPosition, InRenHairPositionOffset);
 
 		float3 OutPosition = 0.0f;
 
@@ -494,7 +517,17 @@ void MainCS(uint2 DispatchThreadId : SV_DispatchThreadID)
 		}
 		#endif
 
-		// 3. Write out the final position
+		// 3. Apply deformer
+		#if PERMUTATION_DEFORMER
+		//if (VertexIndex > 10000000)
+		{
+			const uint4 PackedDeformerPosition = RenDeformerPositionBuffer[VertexIndex];
+			const float3 DeformerPosition = UnpackHairControlPointPosition(PackedDeformerPosition, InRenHairPositionOffset);
+			OutPosition = ApplyDeformer(RenTri, RestPosition, DeformerPosition, OutPosition);
+		}
+		#endif
+
+		// 4. Write out the final position
 		// Optionally trim hair
 		if (HairLengthScale < 1.0f)
 		{
