Update HairDescriptionBulkData to use FEditorBulkData.
* Avoiding to load descriptor data when loading groom asset
* Allow groom asset to be virtualized

diff --git a/Engine/Plugins/Runtime/HairStrands/Source/HairStrandsCore/Private/HairDescription.cpp b/Engine/Plugins/Runtime/HairStrands/Source/HairStrandsCore/Private/HairDescription.cpp
--- a/Engine/Plugins/Runtime/HairStrands/Source/HairStrandsCore/Private/HairDescription.cpp
+++ b/Engine/Plugins/Runtime/HairStrands/Source/HairStrandsCore/Private/HairDescription.cpp
@@ -4,10 +4,15 @@
 #include "Misc/SecureHash.h"
 #include "Serialization/BulkDataReader.h"
 #include "Serialization/BulkDataWriter.h"
+#include "Serialization/EditorBulkDataReader.h"
+#include "Serialization/EditorBulkDataWriter.h"
 
 const FStrandID FStrandID::Invalid(TNumericLimits<uint32>::Max());
 const FGroomID FGroomID::Invalid(TNumericLimits<uint32>::Max());
 
+///////////////////////////////////////////////////////////////////////////////////////////////////
+// FHairDescription
+
 FHairDescription::FHairDescription()
 	: NumVertices(0)
 	, NumStrands(0)
@@ -76,8 +81,10 @@ void FHairDescription::Serialize(FArchive& Ar)
 	Ar << GroomAttributesSet;
 }
 
-#if WITH_EDITORONLY_DATA
+///////////////////////////////////////////////////////////////////////////////////////////////////
+// FHairDescriptionBulkData
 
+#if WITH_EDITORONLY_DATA
 void FHairDescriptionBulkData::Serialize(FArchive& Ar, UObject* Owner)
 {
 	Ar.UsingCustomVersion(FEditorObjectVersion::GUID);
@@ -105,9 +112,21 @@ void FHairDescriptionBulkData::Serialize(FArchive& Ar, UObject* Owner)
 		}
 	}
 
-	BulkData.Serialize(Ar, Owner);
+	// Convert legacy hair description bulk data
+	if (Ar.IsLoading() && Ar.CustomVer(FUE5MainStreamObjectVersion::GUID) < FUE5MainStreamObjectVersion::UpdateHairDescriptionBulkData)
+	{
+		FByteBulkData TempBulkData;
+		TempBulkData.Serialize(Ar, Owner);
+
+		FGuid TempGuid;
+		Ar << TempGuid;
 
-	Ar << Guid;
+		BulkData.CreateFromBulkData(TempBulkData, TempGuid, Owner);
+	}
+	else
+	{
+		BulkData.Serialize(Ar, Owner);
+	}
 
 	if (!Ar.IsTransacting() && Ar.IsLoading())
 	{
@@ -122,20 +141,17 @@ void FHairDescriptionBulkData::Serialize(FArchive& Ar, UObject* Owner)
 
 void FHairDescriptionBulkData::SaveHairDescription(FHairDescription& HairDescription)
 {
-	BulkData.RemoveBulkData();
+	BulkData.Reset();
 
 	if (HairDescription.IsValid())
 	{
-		FBulkDataWriter Ar(BulkData, /*bIsPersistent*/ true);
+		UE::Serialization::FEditorBulkDataWriter Ar(BulkData);
 		HairDescription.Serialize(Ar);
 
 		// Preserve CustomVersions at save time so we can reuse the same ones when reloading direct from memory
 		CustomVersions = Ar.GetCustomVersions();
 	}
 
-	// Use bulk data hash instead of guid to identify content to improve DDC cache hit
-	ComputeGuidFromHash();
-
 	// Mark the HairDescriptionBulkData as having been updated.
 	// This means we know that its version is up-to-date.
 	bBulkDataUpdated = true;
@@ -148,7 +164,7 @@ void FHairDescriptionBulkData::LoadHairDescription(FHairDescription& HairDescrip
 
 	if (!IsEmpty())
 	{
-		FBulkDataReader Ar(BulkData, /*bIsPersistent*/ true);
+		UE::Serialization::FEditorBulkDataReader Ar(BulkData);
 
 		// Propagate the custom version information from the package to the bulk data, so that the HairDescription
 		// is serialized with the same versioning
@@ -160,28 +176,14 @@ void FHairDescriptionBulkData::LoadHairDescription(FHairDescription& HairDescrip
 
 void FHairDescriptionBulkData::Empty()
 {
-	BulkData.RemoveBulkData();
+	BulkData.Reset();
 }
 
 FString FHairDescriptionBulkData::GetIdString() const
-{
-	FString GuidString = Guid.ToString();
+{		
+	FString GuidString = UE::Serialization::IoHashToGuid(BulkData.GetPayloadId()).ToString();
 	GuidString += TEXT("X");
 	return GuidString;
 }
 
-void FHairDescriptionBulkData::ComputeGuidFromHash()
-{
-	uint32 Hash[5] = {};
-
-	if (BulkData.GetBulkDataSize() > 0)
-	{
-		void* Buffer = BulkData.Lock(LOCK_READ_ONLY);
-		FSHA1::HashBuffer(Buffer, BulkData.GetBulkDataSize(), (uint8*)Hash);
-		BulkData.Unlock();
-	}
-
-	Guid = FGuid(Hash[0] ^ Hash[4], Hash[1], Hash[2], Hash[3]);
-}
-
 #endif // #if WITH_EDITORONLY_DATA
diff --git a/Engine/Plugins/Runtime/HairStrands/Source/HairStrandsCore/Public/HairDescription.h b/Engine/Plugins/Runtime/HairStrands/Source/HairStrandsCore/Public/HairDescription.h
--- a/Engine/Plugins/Runtime/HairStrands/Source/HairStrandsCore/Public/HairDescription.h
+++ b/Engine/Plugins/Runtime/HairStrands/Source/HairStrandsCore/Public/HairDescription.h
@@ -9,6 +9,7 @@
 #include "MeshDescription.h" // for TVertexAttributesRef
 #include "MeshTypes.h" // for FElementID, FVertexID
 #include "Serialization/BulkData.h"
+#include "Serialization/EditorBulkData.h"
 
 struct FStrandID : public FElementID
 {
@@ -114,10 +115,7 @@ private:
 struct HAIRSTRANDSCORE_API FHairDescriptionBulkData
 {
 public:
-	FHairDescriptionBulkData()
-	{
-		BulkData.SetBulkDataFlags(BULKDATA_SerializeCompressed | BULKDATA_SerializeCompressedZLIB);
-	}
+	FHairDescriptionBulkData() {}
 
 #if WITH_EDITORONLY_DATA
 	void Serialize(FArchive& Ar, UObject* Owner);
@@ -132,28 +130,21 @@ public:
 	void Empty();
 
 	/** Returns true if there is nothing in the bulk data */
-	bool IsEmpty() const { return BulkData.GetBulkDataSize() == 0; }
+	bool IsEmpty() const { return !BulkData.HasPayloadData(); }
 
 	/** Returns unique ID string for this bulk data */
 	FString GetIdString() const;
 
-private:
-	/** Computes a GUID from the hash of the bulk data, useful to prevent recomputing content already in cache. */
-	void ComputeGuidFromHash();
-#endif
-
 private:
 	/** Internally store bulk data as bytes */
-	FByteBulkData BulkData;
-
-	/** GUID associated with the stored bulk data */
-	FGuid Guid;
+	UE::Serialization::FEditorBulkData BulkData;
 
 	/** Custom version to propagate to archive when serializing the bulk data */
 	FCustomVersionContainer CustomVersions;
 
 	/** Whether the bulk data has been written via SaveHairDescription */
 	bool bBulkDataUpdated = false;
+#endif
 };
 
 template <typename AttributeType> using TStrandAttributesRef = TMeshAttributesRef<FStrandID, AttributeType>;
diff --git a/Engine/Source/Runtime/Core/Public/UObject/UE5MainStreamObjectVersion.h b/Engine/Source/Runtime/Core/Public/UObject/UE5MainStreamObjectVersion.h
--- a/Engine/Source/Runtime/Core/Public/UObject/UE5MainStreamObjectVersion.h
+++ b/Engine/Source/Runtime/Core/Public/UObject/UE5MainStreamObjectVersion.h
@@ -289,6 +289,9 @@ struct CORE_API FUE5MainStreamObjectVersion
 		// bIsUsedWithVolumetricCloud flag auto conversion
 		MaterialHasIsUsedWithVolumetricCloudFlag,
 
+		// bIsUsedWithVolumetricCloud flag auto conversion
+		UpdateHairDescriptionBulkData,
+
 		// -----<new versions can be added above this line>-------------------------------------------------
 		VersionPlusOne,
 		LatestVersion = VersionPlusOne - 1
