Update to Groom to DDC 2.0 API.

diff --git a/Engine/Plugins/Runtime/HairStrands/Source/HairStrandsCore/HairStrandsCore.Build.cs b/Engine/Plugins/Runtime/HairStrands/Source/HairStrandsCore/HairStrandsCore.Build.cs
--- a/Engine/Plugins/Runtime/HairStrands/Source/HairStrandsCore/HairStrandsCore.Build.cs
+++ b/Engine/Plugins/Runtime/HairStrands/Source/HairStrandsCore/HairStrandsCore.Build.cs
@@ -43,6 +43,7 @@ namespace UnrealBuildTool.Rules
 					new string[]
 					{
 						"Eigen",
+						"DerivedDataCache",
 					});
 
 				PrivateIncludePathModuleNames.AddRange(
diff --git a/Engine/Plugins/Runtime/HairStrands/Source/HairStrandsCore/Private/GroomAsset.cpp b/Engine/Plugins/Runtime/HairStrands/Source/HairStrandsCore/Private/GroomAsset.cpp
--- a/Engine/Plugins/Runtime/HairStrands/Source/HairStrandsCore/Private/GroomAsset.cpp
+++ b/Engine/Plugins/Runtime/HairStrands/Source/HairStrandsCore/Private/GroomAsset.cpp
@@ -35,6 +35,11 @@
 #include "HairCardsVertexFactory.h"
 #include "PSOPrecache.h"
 
+#if WITH_EDITORONLY_DATA
+#include "DerivedDataCache.h"
+#include "DerivedDataRequestOwner.h"
+#endif
+
 #include UE_INLINE_GENERATED_CPP_BY_NAME(GroomAsset)
 
 
@@ -1872,7 +1877,7 @@ void UGroomAsset::SetHairWidth(float Width)
 // differences, etc.) replace the version GUID below with a new one.
 // In case of merge conflicts with DDC versions, you MUST generate a new GUID
 // and set this new GUID as the version.
-#define GROOM_DERIVED_DATA_VERSION TEXT("FC0D84AFC0F444A9A988C651BEED113A")
+#define GROOM_DERIVED_DATA_VERSION TEXT("282985F4A254469A87884E2F1538CCEB")
 
 #if WITH_EDITORONLY_DATA
 
@@ -2276,29 +2281,31 @@ bool UGroomAsset::CacheStrandsData(uint32 GroupIndex, FString& OutDerivedDataKey
 		return false;
 	}
 
+	using namespace UE::DerivedData;
+
 	const FString DerivedDataKey = UGroomAsset::GetDerivedDataKeyForStrands(GroupIndex);
+	const FCacheKey Key = ConvertLegacyCacheKey(DerivedDataKey);
+	const FSharedString Name = MakeStringView(GetPathName());
+	FSharedBuffer Data;
+	{
+		FRequestOwner Owner(EPriority::Blocking);
+		GetCache().GetValue({ {Name, Key} }, Owner, [&Data](FCacheGetValueResponse&& Response)
+		{
+			Data = Response.Value.GetData().Decompress();
+		});
+		Owner.Wait();
+	}
 
 	bool bSuccess = true;
-	TArray<uint8> DerivedData;
-	if (GetDerivedDataCacheRef().GetSynchronous(*DerivedDataKey, DerivedData, GetPathName()))
+	if (Data)
 	{
 		UE_CLOG(IsHairStrandsDDCLogEnable(), LogHairStrands, Log, TEXT("[Groom/DDC] Strands - Found (Groom:%s Group6:%d)."), *GetName(), GroupIndex);
 
-		FMemoryReader Ar(DerivedData, /*bIsPersistent=*/ true);
-
-		int64 UncompressedSize = 0;
-		Ar << UncompressedSize;
-
-		uint8* DecompressionBuffer = reinterpret_cast<uint8*>(FMemory::Malloc(UncompressedSize));
-		Ar.SerializeCompressed(DecompressionBuffer, UncompressedSize, NAME_Zlib);
-
-		FHairGroupData& HairGroupData = HairGroupsData[GroupIndex];
-
 		// Reset hair group data to ensure previously loaded bulk data are cleared/cleaned priori to load new data.
-		HairGroupData = FHairGroupData();
+		HairGroupsData[GroupIndex] = FHairGroupData();
 
-		FLargeMemoryReader LargeMemReader(DecompressionBuffer, UncompressedSize, ELargeMemoryReaderFlags::Persistent | ELargeMemoryReaderFlags::TakeOwnership);
-		InternalSerialize(LargeMemReader, this, HairGroupData);
+		FMemoryReaderView Ar(Data, /*bIsPersistent*/ true);
+		InternalSerialize(Ar, this, HairGroupsData[GroupIndex]);
 	}
 	else
 	{
@@ -2333,32 +2340,12 @@ bool UGroomAsset::CacheStrandsData(uint32 GroupIndex, FString& OutDerivedDataKey
 		
 		if (bSuccess)
 		{
-			// Using a LargeMemoryWriter for serialization since the data can be bigger than 2 GB
-			FLargeMemoryWriter LargeMemWriter(0, /*bIsPersistent=*/ true);
-			InternalSerialize(LargeMemWriter, this, HairGroupsData[GroupIndex]);
-			int64 UncompressedSize = LargeMemWriter.TotalSize();
-
-			// Then the content of the LargeMemWriter is compressed into a MemoryWriter
-			// Compression ratio can reach about 5:2 depending on the data
-			// Since the DDC doesn't support data bigger than 2 GB
-			// we can compute a size threshold to skip the caching when
-			// the uncompressed size exceeds the threshold
-			static constexpr const int64 SizeThreshold = (int64)MAX_int32 * 2.5;
-			const bool bIsCacheable = UncompressedSize < SizeThreshold;
-			if (bIsCacheable)
-			{
-				FMemoryWriter CompressedArchive(DerivedData, true);
-
-				CompressedArchive << UncompressedSize; // needed for allocating decompression buffer
-				CompressedArchive.SerializeCompressed(LargeMemWriter.GetData(), UncompressedSize, NAME_Zlib);
-
-				GetDerivedDataCacheRef().Put(*DerivedDataKey, DerivedData, GetPathName());
-			}
-			else
-			{
-				UE_LOG(LogHairStrands, Error, TEXT("[Groom/DDC] Strands - The groom asset is too large to be cached into the Derived Data Cache. (Groom:%s, Group:%d, Size %d bytes, DDC Limit: SizeThreshold %d bytes)) ."), *GetName(), GroupIndex, UncompressedSize, SizeThreshold);
-				return false;
-			}
+			TArray<uint8> WriteData;
+			FMemoryWriter Ar(WriteData, /*bIsPersistent*/ true);
+			InternalSerialize(Ar, this, HairGroupsData[GroupIndex]);
+			FRequestOwner AsyncOwner(EPriority::Normal);
+			GetCache().PutValue({ {Name, Key, FValue::Compress(MakeSharedBufferFromArray(MoveTemp(WriteData)))} }, AsyncOwner);
+			AsyncOwner.KeepAlive();
 		}
 	}
 
@@ -2461,12 +2448,24 @@ bool UGroomAsset::CacheCardsGeometry(uint32 GroupIndex, const FString& StrandsKe
 		}
 	}
 
-	TArray<uint8> DerivedData;
-	if (bDataCanBeBuilt && GetDerivedDataCacheRef().GetSynchronous(*DerivedDataKey, DerivedData, GetPathName()))
+	using namespace UE::DerivedData;
+	const FCacheKey Key = ConvertLegacyCacheKey(DerivedDataKey);
+	const FSharedString Name = MakeStringView(GetPathName());
+	FSharedBuffer Data;
+	{
+		FRequestOwner Owner(EPriority::Blocking);
+		GetCache().GetValue({ {Name, Key} }, Owner, [&Data](FCacheGetValueResponse&& Response)
+		{
+			Data = Response.Value.GetData().Decompress();
+		});
+		Owner.Wait();
+	}
+
+	if (Data)
 	{
 		UE_CLOG(IsHairStrandsDDCLogEnable(), LogHairStrands, Log, TEXT("[Groom/DDC] Cards - Found (Groom:%s Group:%d)."), *GetName(), GroupIndex);
 		
-		FMemoryReader Ar(DerivedData, true);
+		FMemoryReaderView Ar(Data, /*bIsPersistent*/ true);
 		InternalSerialize(Ar, this, HairGroupData.Cards.LODs);
 	}
 	else
@@ -2479,10 +2478,12 @@ bool UGroomAsset::CacheCardsGeometry(uint32 GroupIndex, const FString& StrandsKe
 			return false;
 		}
 
-		FMemoryWriter Ar(DerivedData, true);
+		TArray<uint8> WriteData;
+		FMemoryWriter Ar(WriteData, /*bIsPersistent*/ true);
 		InternalSerialize(Ar, this, HairGroupData.Cards.LODs);
-
-		GetDerivedDataCacheRef().Put(*DerivedDataKey, DerivedData, GetPathName());
+		FRequestOwner AsyncOwner(EPriority::Normal);
+		GetCache().PutValue({ {Name, Key, FValue::Compress(MakeSharedBufferFromArray(MoveTemp(WriteData)))} }, AsyncOwner);
+		AsyncOwner.KeepAlive();
 	}
 
 	// Handle the case where the cards data is already cached in the DDC
@@ -2557,12 +2558,25 @@ bool UGroomAsset::CacheMeshesGeometry(uint32 GroupIndex)
 		}
 	}
 
+	using namespace UE::DerivedData;
+	const FCacheKey Key = ConvertLegacyCacheKey(DerivedDataKey);
+	const FSharedString Name = MakeStringView(GetPathName());
+	FSharedBuffer Data;
+	{
+		FRequestOwner Owner(EPriority::Blocking);
+		GetCache().GetValue({ {Name, Key} }, Owner, [&Data](FCacheGetValueResponse&& Response)
+		{
+			Data = Response.Value.GetData().Decompress();
+		});
+		Owner.Wait();
+	}
+
 	TArray<uint8> DerivedData;
 	if (bDataCanBeBuilt && GetDerivedDataCacheRef().GetSynchronous(*DerivedDataKey, DerivedData, GetPathName()))
 	{
 		UE_CLOG(IsHairStrandsDDCLogEnable(), LogHairStrands, Log, TEXT("[Groom/DDC] Meshes - Found (Groom:%s Group:%d)."), *GetName(), GroupIndex);
 
-		FMemoryReader Ar(DerivedData, true);
+		FMemoryReaderView Ar(Data, /*bIsPersistent*/ true);
 		InternalSerialize(Ar, this, HairGroupData.Meshes.LODs);
 	}
 	else
@@ -2575,10 +2589,12 @@ bool UGroomAsset::CacheMeshesGeometry(uint32 GroupIndex)
 			return false;
 		}
 
-		FMemoryWriter Ar(DerivedData, true);
+		TArray<uint8> WriteData;
+		FMemoryWriter Ar(WriteData, /*bIsPersistent*/ true);
 		InternalSerialize(Ar, this, HairGroupData.Meshes.LODs);
-
-		GetDerivedDataCacheRef().Put(*DerivedDataKey, DerivedData, GetPathName());
+		FRequestOwner AsyncOwner(EPriority::Normal);
+		GetCache().PutValue({ {Name, Key, FValue::Compress(MakeSharedBufferFromArray(MoveTemp(WriteData)))} }, AsyncOwner);
+		AsyncOwner.KeepAlive();
 	}
 
 	MeshesDerivedDataKey[GroupIndex] = DerivedDataKey;
