Add curve buffer (point offset/count) to hair rest resources.

diff --git a/Engine/Plugins/Runtime/HairStrands/Shaders/Private/NiagaraDataInterfaceHairStrandsTemplate.ush b/Engine/Plugins/Runtime/HairStrands/Shaders/Private/NiagaraDataInterfaceHairStrandsTemplate.ush
--- a/Engine/Plugins/Runtime/HairStrands/Shaders/Private/NiagaraDataInterfaceHairStrandsTemplate.ush
+++ b/Engine/Plugins/Runtime/HairStrands/Shaders/Private/NiagaraDataInterfaceHairStrandsTemplate.ush
@@ -112,6 +112,26 @@ void BuildBoundingBox_{ParameterName}(in float3 NodePosition, out bool FunctionS
 	DeviceMemoryBarrier();
 }
 
+/* -----------------------------------------------------------------
+ * Utilities regarding strands curve
+ * -----------------------------------------------------------------
+ */
+struct FHairCurve
+{
+	uint PointIndex;
+	uint PointCount;
+};
+
+FHairCurve UnpackCurve_{ParameterName}(in uint InCurveIndex)
+{
+	const uint In = {ParameterName}_CurvesOffsetsBuffer[InCurveIndex];
+
+	FHairCurve Out = (FHairCurve)0;
+	Out.PointIndex = In & 0xFFFFFF;
+	Out.PointCount = In>>24;
+	return Out;
+}
+
 /* -----------------------------------------------------------------
  * Utilities regarding strands indexing
  * -----------------------------------------------------------------
@@ -134,8 +154,9 @@ float ComputeStrandLength_{ParameterName}(in int StrandIndex)
 	float StrandLength = 0.0;
 	if( StrandIndex < {ParameterName}_NumStrands )
 	{
-		int PointOffset = {ParameterName}_CurvesOffsetsBuffer[StrandIndex];
-		const int EdgeCount = {ParameterName}_CurvesOffsetsBuffer[StrandIndex+1]-PointOffset-1;
+		FHairCurve Curve = UnpackCurve_{ParameterName}(StrandIndex);
+		int PointOffset = Curve.PointIndex;
+		const int EdgeCount = Curve.PointCount-1;
 
 		float3 PointNext = f16tof32({ParameterName}_RestPositionBuffer[PointOffset++].xyz);
 		float3 PointPrev = PointNext;
@@ -233,10 +254,12 @@ void ComputeNodePosition_{ParameterName}(out float3 OutNodePosition)
 	const bool RootNode = (LocalIndex == 0);
 	LocalIndex = max(0,LocalIndex - 1);
 
-	const int PointOffset = {ParameterName}_CurvesOffsetsBuffer[StrandIndex];
+	const FHairCurve Curve = UnpackCurve_{ParameterName}(StrandIndex);
+
+	const int PointOffset = Curve.PointIndex;
 
 	const float NodeCoord = (float)(LocalIndex) * CoordScale;
-	const float PointCount = {ParameterName}_CurvesOffsetsBuffer[StrandIndex+1]-PointOffset-1;
+	const float PointCount = Curve.PointCount-1;
 
 	const float NodePoint = NodeCoord * PointCount;
 	const int PointPrev = (LocalIndex==0) ? 0 : (LocalIndex==(ReducedSize-1)) ? PointCount-1 : floor(NodePoint);
@@ -551,9 +574,9 @@ void ResetPointPosition_{ParameterName}(out bool ReportStatus)
 
 	if(LocalIndex == 0)
 	{
-		const int PointBegin = {ParameterName}_CurvesOffsetsBuffer[StrandIndex];
-		const int PointEnd = {ParameterName}_CurvesOffsetsBuffer[StrandIndex+1]; 
-		const int PointCount = PointEnd - PointBegin;
+		const FHairCurve Curve = UnpackCurve_{ParameterName}(StrandIndex);
+		const int PointBegin = Curve.PointIndex;
+		const int PointCount = Curve.PointCount;
 
 		const uint4 PositionOffset = uint4( f32tof16({ParameterName}_RestPositionOffset - {ParameterName}_DeformedPositionOffset[0].xyz), 0 );
 
@@ -686,9 +709,9 @@ void UpdatePointPosition_{ParameterName}(in float3 NodePosition, in float3 RestP
 	UpdateTriangleDisplace_{ParameterName}(StrandIndex, NodePosition, RestPosition, 
 			RestTrianglePosition, RestTriangleOrientation, DeformedTrianglePosition, DeformedTriangleOrientation);
 	
-	const int PointBegin = {ParameterName}_CurvesOffsetsBuffer[StrandIndex];
-	const int PointEnd = {ParameterName}_CurvesOffsetsBuffer[StrandIndex + 1];
-	const int PointCount = PointEnd - PointBegin;
+	const FHairCurve Curve = UnpackCurve_{ParameterName}(StrandIndex);
+	const int PointBegin = Curve.PointIndex;
+	const int PointCount = Curve.PointCount;
 
 
 	for (int PointIndex = 0; PointIndex < PointCount; ++PointIndex)
diff --git a/Engine/Plugins/Runtime/HairStrands/Source/HairStrandsCore/Private/GroomBuilder.cpp b/Engine/Plugins/Runtime/HairStrands/Source/HairStrandsCore/Private/GroomBuilder.cpp
--- a/Engine/Plugins/Runtime/HairStrands/Source/HairStrandsCore/Private/GroomBuilder.cpp
+++ b/Engine/Plugins/Runtime/HairStrands/Source/HairStrandsCore/Private/GroomBuilder.cpp
@@ -41,9 +41,9 @@ FString FGroomBuilder::GetVersion()
 {
 	// Important to update the version when groom building changes
 	if (IsHairStrandContinuousDecimationReorderingEnabled())
-		return TEXT("3r0");
+		return TEXT("v4");
 	else
-		return TEXT("2r2");
+		return TEXT("v3");
 }
 
 namespace FHairStrandsDecimation
@@ -267,13 +267,13 @@ namespace HairStrandsBuilder
 		TArray<FHairStrandsAttribute0Format::Type> OutPackedAttributes0;
 		TArray<FHairStrandsAttribute1Format::Type> OutPackedAttributes1;
 		TArray<FHairStrandsMaterialFormat::Type> OutPackedMaterials;
-		TArray<FHairStrandsRootIndexFormat::Type> OutCurveOffsets;
+		TArray<FHairStrandsCurveFormat::Type> OutPackedCurves;
 
 		OutPackedPositions.SetNum(NumPoints * FHairStrandsPositionFormat::ComponentCount);
 		OutPackedAttributes0.SetNum(NumPoints * FHairStrandsAttribute0Format::ComponentCount);
 		OutPackedAttributes1.SetNum(NumPoints * FHairStrandsAttribute1Format::ComponentCount);
 		OutPackedMaterials.SetNum(NumPoints * FHairStrandsMaterialFormat::ComponentCount);
-		OutCurveOffsets = HairStrands.StrandsCurves.CurvesOffset;
+		OutPackedCurves.SetNum(NumCurves * FHairStrandsCurveFormat::ComponentCount);
 
 		const FVector HairBoxCenter = HairStrands.BoundingBox.GetCenter();
 
@@ -358,6 +358,15 @@ namespace HairStrandsBuilder
 					}
 				}
 			}
+
+			// Curves
+			{
+				UE_CLOG(PointCount > 0xFF, LogGroomBuilder, Warning, TEXT("Curve point count is > 256"));
+				UE_CLOG(IndexOffset > 0xFFFFFF, LogGroomBuilder, Warning, TEXT("Curve point offset is > 24M"));
+				FHairStrandsCurveFormat::Type& Curve = OutPackedCurves[CurveIndex];
+				Curve.PointCount  = FMath::Min(uint32(PointCount), 0xFFu);
+				Curve.PointOffset = FMath::Min(uint32(IndexOffset), 0xFFFFFFu);
+			}
 		}
 
 		OutBulkData.BoundingBox = HairStrands.BoundingBox;
@@ -377,7 +386,7 @@ namespace HairStrandsBuilder
 		{
 			CopyToBulkData<FHairStrandsMaterialFormat>(OutBulkData.Materials, OutPackedMaterials);
 		}
-		CopyToBulkData<FHairStrandsRootIndexFormat>(OutBulkData.CurveOffsets, OutCurveOffsets);
+		CopyToBulkData<FHairStrandsCurveFormat>(OutBulkData.Curves, OutPackedCurves);
 	}
 
 	void BuildRenderData(const FHairStrandsDatas& HairStrands, FHairStrandsBulkData& OutBulkData)
diff --git a/Engine/Plugins/Runtime/HairStrands/Source/HairStrandsCore/Private/GroomResources.cpp b/Engine/Plugins/Runtime/HairStrands/Source/HairStrandsCore/Private/GroomResources.cpp
--- a/Engine/Plugins/Runtime/HairStrands/Source/HairStrandsCore/Private/GroomResources.cpp
+++ b/Engine/Plugins/Runtime/HairStrands/Source/HairStrandsCore/Private/GroomResources.cpp
@@ -793,6 +793,8 @@ bool FHairStrandsRestResource::InternalIsDataLoaded()
 			Batch.Read(BulkData.Materials);
 		}
 
+		Batch.Read(BulkData.Curves);
+
 		Batch.Issue(BulkDataRequest);
 	}
 	
@@ -804,12 +806,14 @@ void FHairStrandsRestResource::InternalAllocate(FRDGBuilder& GraphBuilder)
 	BulkDataRequest = FBulkDataBatchRequest();
 
 	const uint32 PointCount = BulkData.PointCount;
+	const uint32 CurveCount = BulkData.CurveCount;
 
 	// 1. Lock data, which force the loading data from files (on non-editor build/cooked data). These data are then uploaded to the GPU
 	// 2. A local copy is done by the buffer uploader. This copy is discarded once the uploading is done.
 	InternalCreateVertexBufferRDG_FromBulkData<FHairStrandsPositionFormat>(GraphBuilder, BulkData.Positions, PointCount, PositionBuffer, ToHairResourceDebugName(HAIRSTRANDS_RESOUCE_NAME(CurveType, Hair.StrandsRest_PositionBuffer), ResourceName), EHairResourceUsageType::Static);
 	InternalCreateVertexBufferRDG_FromBulkData<FHairStrandsAttribute0Format>(GraphBuilder, BulkData.Attributes0, PointCount, Attribute0Buffer, ToHairResourceDebugName(HAIRSTRANDS_RESOUCE_NAME(CurveType, Hair.StrandsRest_Attribute0Buffer), ResourceName), EHairResourceUsageType::Static);
 	InternalCreateVertexBufferRDG_FromBulkData<FHairStrandsAttribute1Format>(GraphBuilder, BulkData.Attributes1, PointCount, Attribute1Buffer, ToHairResourceDebugName(HAIRSTRANDS_RESOUCE_NAME(CurveType, Hair.StrandsRest_Attribute1Buffer), ResourceName), EHairResourceUsageType::Static);
+	InternalCreateVertexBufferRDG_FromBulkData<FHairStrandsCurveFormat>(GraphBuilder, BulkData.Curves, CurveCount, CurveBuffer, ToHairResourceDebugName(HAIRSTRANDS_RESOUCE_NAME(CurveType, Hair.StrandsRest_CurveBuffer), ResourceName), EHairResourceUsageType::Static);
 
 	if (!!(BulkData.Flags & FHairStrandsBulkData::DataFlags_HasMaterialData))
 	{
@@ -859,6 +863,7 @@ void FHairStrandsRestResource::InternalRelease()
 	Attribute1Buffer.Release();
 	MaterialBuffer.Release();
 	TangentBuffer.Release();
+	CurveBuffer.Release();
 }
 
 /////////////////////////////////////////////////////////////////////////////////////////
@@ -889,6 +894,7 @@ void FHairStrandsDeformedResource::InternalAllocate(FRDGBuilder& GraphBuilder)
 
 FRDGExternalBuffer& FHairStrandsDeformedResource::GetDeformerBuffer(FRDGBuilder& GraphBuilder)
 {
+	// Lazy allocation and update
 	if (DeformerBuffer.Buffer == nullptr)
 	{
 		InternalCreateVertexBufferRDG<FHairStrandsPositionFormat>(GraphBuilder, BulkData.PointCount, DeformerBuffer, ToHairResourceDebugName(HAIRSTRANDS_RESOUCE_NAME(CurveType, Hair.StrandsDeformed_DeformerBuffer), ResourceName), EHairResourceUsageType::Dynamic);
diff --git a/Engine/Plugins/Runtime/HairStrands/Source/HairStrandsCore/Private/HairStrandsDatas.cpp b/Engine/Plugins/Runtime/HairStrands/Source/HairStrandsCore/Private/HairStrandsDatas.cpp
--- a/Engine/Plugins/Runtime/HairStrands/Source/HairStrandsCore/Private/HairStrandsDatas.cpp
+++ b/Engine/Plugins/Runtime/HairStrands/Source/HairStrandsCore/Private/HairStrandsDatas.cpp
@@ -230,7 +230,7 @@ void FHairStrandsBulkData::Serialize(FArchive& Ar, UObject* Owner)
 		Attributes0.SetBulkDataFlags(BulkFlags);
 		Attributes1.SetBulkDataFlags(BulkFlags);
 		Materials.SetBulkDataFlags(BulkFlags);
-		CurveOffsets.SetBulkDataFlags(BulkFlags);
+		Curves.SetBulkDataFlags(BulkFlags);
 	}
 
 	if (!!(Flags & DataFlags_HasData))
@@ -246,7 +246,7 @@ void FHairStrandsBulkData::Serialize(FArchive& Ar, UObject* Owner)
 		{
 			Materials.Serialize(Ar, Owner, ChunkIndex, bAttemptFileMapping);
 		}
-		CurveOffsets.Serialize(Ar, Owner, ChunkIndex, bAttemptFileMapping);
+		Curves.Serialize(Ar, Owner, ChunkIndex, bAttemptFileMapping);
 	}
 }
 
@@ -264,14 +264,14 @@ void FHairStrandsBulkData::Reset()
 	Attributes0.RemoveBulkData();
 	Attributes1.RemoveBulkData();
 	Materials.RemoveBulkData();
-	CurveOffsets.RemoveBulkData();
+	Curves.RemoveBulkData();
 
 	// Reset the bulk byte buffer to ensure the (serialize) data size is reset to 0
-	Positions 	 = FByteBulkData();
-	Attributes0  = FByteBulkData();
-	Attributes1  = FByteBulkData();
-	Materials 	 = FByteBulkData();
-	CurveOffsets = FByteBulkData();
+	Positions 	= FByteBulkData();
+	Attributes0 = FByteBulkData();
+	Attributes1 = FByteBulkData();
+	Materials 	= FByteBulkData();
+	Curves		= FByteBulkData();
 }
 
 void FHairStrandsDatas::Reset()
diff --git a/Engine/Plugins/Runtime/HairStrands/Source/HairStrandsCore/Private/Niagara/NiagaraDataInterfaceHairStrands.cpp b/Engine/Plugins/Runtime/HairStrands/Source/HairStrandsCore/Private/Niagara/NiagaraDataInterfaceHairStrands.cpp
--- a/Engine/Plugins/Runtime/HairStrands/Source/HairStrandsCore/Private/Niagara/NiagaraDataInterfaceHairStrands.cpp
+++ b/Engine/Plugins/Runtime/HairStrands/Source/HairStrandsCore/Private/Niagara/NiagaraDataInterfaceHairStrands.cpp
@@ -302,28 +302,15 @@ void FNDIHairStrandsBuffer::InitRHI()
 		FMemMark MemMark(FMemStack::Get());
 		FRDGBuilder GraphBuilder(FRHICommandListExecutor::GetImmediateCommandList());
 		FHairStrandsBulkData* SourceDatas = &SourceRestResources->BulkData; // This could be released by that time depending on how the initialization order is
-		{
-			const uint32 OffsetCount = SourceDatas->GetNumCurves() + 1;
-			const uint32 OffsetBytes = sizeof(uint32)*OffsetCount;
-
-			const FHairStrandsRootIndexFormat::Type* SourceData = (const FHairStrandsRootIndexFormat::Type*)SourceDatas->CurveOffsets.Lock(LOCK_READ_ONLY);
-			
-			CurvesOffsetsBuffer.Initialize(GraphBuilder, TEXT("CurvesOffsetsBuffer"), EPixelFormat::PF_R32_UINT, sizeof(uint32), OffsetCount);
-			GraphBuilder.QueueBufferUpload(CurvesOffsetsBuffer.GetOrCreateBuffer(GraphBuilder), SourceData, OffsetBytes);
-			CurvesOffsetsBuffer.EndGraphUsage();
-
-			SourceDatas->CurveOffsets.Unlock();
-		}
 		{
 			// First retrieve the curves offsets from the bulk datas
-			TArray<FHairStrandsRootIndexFormat::Type> CurvesOffsets;
-			const uint32 InDataSize = SourceDatas->CurveOffsets.GetBulkDataSize();
+			TArray<FHairStrandsCurveFormat::Type> Curves;
+			const uint32 InDataSize = SourceDatas->Curves.GetBulkDataSize();
 			const uint32 ElementCount = InDataSize / sizeof(FHairStrandsRootIndexFormat::BulkType);
-			CurvesOffsets.SetNum(ElementCount);
-
-			const uint8* InData = (const uint8*)SourceDatas->CurveOffsets.LockReadOnly();
-			FMemory::Memcpy(CurvesOffsets.GetData(), InData, InDataSize);
-			SourceDatas->CurveOffsets.Unlock();
+			Curves.SetNum(ElementCount);
+			const uint8* InData = (const uint8*)SourceDatas->Curves.LockReadOnly();
+			FMemory::Memcpy(Curves.GetData(), InData, InDataSize);
+			SourceDatas->Curves.Unlock();
 			
 			TArray<uint32> PointsCurve;
 			PointsCurve.SetNum(SourceDatas->GetNumPoints());
@@ -331,7 +318,8 @@ void FNDIHairStrandsBuffer::InitRHI()
 			// Build the curve index array for each points using the curve offsets array
 			for (uint32 CurveIndex = 0; CurveIndex < SourceDatas->GetNumCurves(); ++CurveIndex)
 			{
-				for (uint32 PointIndex = CurvesOffsets[CurveIndex]; PointIndex < CurvesOffsets[CurveIndex+1]; ++PointIndex)
+				const FHairStrandsCurveFormat::Type& Curve = Curves[CurveIndex];				
+				for (uint32 PointIndex = Curve.PointOffset; PointIndex < Curve.PointCount; ++PointIndex)
 				{
 					PointsCurve[PointIndex] = CurveIndex;
 				}
@@ -382,7 +370,6 @@ void FNDIHairStrandsBuffer::ReleaseRHI()
 	//delete ReadbackBuffer;
 	//ReadbackBuffer = nullptr;
 	
-	CurvesOffsetsBuffer.Release();
 	DeformedPositionBuffer.SafeRelease();
 	BoundingBoxBuffer.Release();
 	ParamsScaleBuffer.Release();
@@ -944,7 +931,7 @@ static void InterpolateGroomGuides(FRDGBuilder& GraphBuilder, FNiagaraDataBuffer
 	const bool bIsRestValid = bIsHairValid && HairStrandsBuffer->SourceRestResources && HairStrandsBuffer->SourceRestResources->IsInitialized();
 	const bool bIsDeformedValid = bIsHairValid && HairStrandsBuffer->SourceDeformedResources && HairStrandsBuffer->SourceDeformedResources->IsInitialized();
 
-	if(bIsRestValid && bIsDeformedValid && bHasValidGeometry && HairStrandsBuffer->PointsCurveBuffer.IsValid() && HairStrandsBuffer->CurvesOffsetsBuffer.IsValid())
+	if(bIsRestValid && bIsDeformedValid && bHasValidGeometry && HairStrandsBuffer->PointsCurveBuffer.IsValid())
 	{
 		bool bIsRootValid = HairStrandsBuffer->SourceDeformedRootResources && HairStrandsBuffer->SourceDeformedRootResources->IsInitialized() &&
 								HairStrandsBuffer->SourceRestRootResources && HairStrandsBuffer->SourceRestRootResources->IsInitialized() && bHasSkinningBinding;
@@ -977,7 +964,7 @@ static void InterpolateGroomGuides(FRDGBuilder& GraphBuilder, FNiagaraDataBuffer
 		PassParameters->RestPositionBuffer = RegisterAsSRV(GraphBuilder, HairStrandsBuffer->SourceRestResources->PositionBuffer);
 		
 		PassParameters->VertexToCurveIndexBuffer = HairStrandsBuffer->PointsCurveBuffer.GetOrCreateSRV(GraphBuilder);
-		PassParameters->CurvesOffsetsBuffer = HairStrandsBuffer->CurvesOffsetsBuffer.GetOrCreateSRV(GraphBuilder);
+		PassParameters->CurvesOffsetsBuffer = RegisterAsSRV(GraphBuilder, HairStrandsBuffer->SourceRestResources->CurveBuffer);
 		if(bIsRootValid)
 		{
 			const FHairStrandsRestRootResource::FLOD* RestMeshProjection = &(HairStrandsBuffer->SourceRestRootResources->LODs[MeshLODIndex]) ;
@@ -1052,7 +1039,6 @@ void UNiagaraDataInterfaceHairStrands::SimCachePostReadFrame(void* OptionalPerIn
 							GraphBuilder.Execute();
 
 							HairstrandsBuffer->PointsCurveBuffer.EndGraphUsage();
-							HairstrandsBuffer->CurvesOffsetsBuffer.EndGraphUsage();
 						}
 					);
 				}
@@ -3416,7 +3402,7 @@ void UNiagaraDataInterfaceHairStrands::SetShaderParameters(const FNiagaraDataInt
 	const bool bIsRootValid = bIsHairValid && ProxyData->HairStrandsBuffer->SourceDeformedRootResources&& ProxyData->HairStrandsBuffer->SourceDeformedRootResources->IsInitialized() && bHasSkinningBinding;
 	const bool bIsRestValid = bIsHairValid && ProxyData->HairStrandsBuffer->SourceRestResources && ProxyData->HairStrandsBuffer->SourceRestResources->IsInitialized()&&
 		// TEMP: These check are only temporary for avoiding crashes while we find the bottom of the issue.
-		ProxyData->HairStrandsBuffer->CurvesOffsetsBuffer.IsValid() && ProxyData->HairStrandsBuffer->ParamsScaleBuffer.IsValid() && ProxyData->HairStrandsBuffer->BoundingBoxBuffer.IsValid();
+		ProxyData->HairStrandsBuffer->ParamsScaleBuffer.IsValid() && ProxyData->HairStrandsBuffer->BoundingBoxBuffer.IsValid();
 
 	const bool bIsGeometryValid = bIsHairValid && (!bIsHairGroupInstValid || (bIsHairGroupInstValid && (ProxyData->HairGroupInstance->GeometryType != EHairGeometryType::NoneGeometry)));
 	const bool bIsDeformedValid = bIsHairValid && ProxyData->HairStrandsBuffer->SourceDeformedResources && ProxyData->HairStrandsBuffer->SourceDeformedResources->IsInitialized();
@@ -3483,7 +3469,7 @@ void UNiagaraDataInterfaceHairStrands::SetShaderParameters(const FNiagaraDataInt
 		// Simulation setup (we update the rest configuration based on the deformed positions 
 		// if in restupdate mode or if we are resetting the sim and using RBF transfer since the rest positions are not matrching the physics asset)
 		const int32 NeedResetValue = (ProxyData->TickCount <= GHairSimulationMaxDelay) || !HairStrandsBuffer->bValidGeometryType;
-		const int32 RestUpdateValue = GHairSimulationRestUpdate || (NeedResetValue && ProxyData->bSkinningTransfer);
+		const int32 RestUpdateValue = GHairSimulationRestUpdate || (NeedResetValue && ProxyData->bSkinningTransfer) ? 1 : 0;
 		const int32 LocalSimulationValue = ProxyData->LocalSimulation;
 
 		HairStrandsBuffer->bValidGeometryType = true;
@@ -3538,7 +3524,7 @@ void UNiagaraDataInterfaceHairStrands::SetShaderParameters(const FNiagaraDataInt
 		ShaderParameters->BoundingBoxBuffer = HairStrandsBuffer->BoundingBoxBuffer.GetOrCreateUAV(GraphBuilder);
 
 		// Set Shader SRV
-		ShaderParameters->CurvesOffsetsBuffer = HairStrandsBuffer->CurvesOffsetsBuffer.GetOrCreateSRV(GraphBuilder);
+		ShaderParameters->CurvesOffsetsBuffer = RegisterAsSRV(GraphBuilder, HairStrandsBuffer->SourceRestResources->CurveBuffer);
 		ShaderParameters->ParamsScaleBuffer = HairStrandsBuffer->ParamsScaleBuffer.GetOrCreateSRV(GraphBuilder);
 		ShaderParameters->RestPositionBuffer = RegisterAsSRV(GraphBuilder,HairStrandsBuffer->SourceRestResources->PositionBuffer); //-OPT: If constant across all dispatches can create in PreStage once
 		ShaderParameters->DeformedPositionOffset = DeformedPositionOffsetSRV;
@@ -3587,7 +3573,7 @@ void UNiagaraDataInterfaceHairStrands::SetShaderParameters(const FNiagaraDataInt
 		ShaderParameters->BoundingBoxBuffer = Context.GetComputeDispatchInterface().GetEmptyBufferUAV(GraphBuilder, PF_R32_UINT);
 
 		// Set Shader SRV
-		ShaderParameters->CurvesOffsetsBuffer = Context.GetComputeDispatchInterface().GetEmptyBufferSRV(GraphBuilder, PF_R32_UINT);;
+		ShaderParameters->CurvesOffsetsBuffer = Context.GetComputeDispatchInterface().GetEmptyBufferSRV(GraphBuilder, FHairStrandsCurveFormat::Format);
 		ShaderParameters->RestPositionBuffer = Context.GetComputeDispatchInterface().GetEmptyBufferSRV(GraphBuilder, FHairStrandsPositionFormat::Format);
 		ShaderParameters->DeformedPositionOffset = Context.GetComputeDispatchInterface().GetEmptyBufferSRV(GraphBuilder, FHairStrandsPositionOffsetFormat::Format);
 		ShaderParameters->RestTrianglePositionABuffer = Context.GetComputeDispatchInterface().GetEmptyBufferSRV(GraphBuilder, FHairStrandsMeshTrianglePositionFormat::Format);
@@ -3656,7 +3642,6 @@ void FNDIHairStrandsProxy::PostSimulate(const FNDIGpuComputePostSimulateContext&
 	if (bIsHairValid && Context.IsFinalPostSimulate())
 	{
 		ProxyData->HairStrandsBuffer->BoundingBoxBuffer.EndGraphUsage();
-		ProxyData->HairStrandsBuffer->CurvesOffsetsBuffer.EndGraphUsage();
 		ProxyData->HairStrandsBuffer->ParamsScaleBuffer.EndGraphUsage();
 		ProxyData->HairStrandsBuffer->PointsCurveBuffer.EndGraphUsage();
 	}
diff --git a/Engine/Plugins/Runtime/HairStrands/Source/HairStrandsCore/Public/GroomAsset.h b/Engine/Plugins/Runtime/HairStrands/Source/HairStrandsCore/Public/GroomAsset.h
--- a/Engine/Plugins/Runtime/HairStrands/Source/HairStrandsCore/Public/GroomAsset.h
+++ b/Engine/Plugins/Runtime/HairStrands/Source/HairStrandsCore/Public/GroomAsset.h
@@ -93,11 +93,11 @@ struct HAIRSTRANDSCORE_API FHairGroupData
 		uint32 GetDataSize() const
 		{
 			uint32 Total = 0;
-			Total += BulkData.Positions.IsBulkDataLoaded() 		? BulkData.Positions.GetBulkDataSize()   : 0;
-			Total += BulkData.Attributes0.IsBulkDataLoaded() 	? BulkData.Attributes0.GetBulkDataSize() : 0;
-			Total += BulkData.Attributes1.IsBulkDataLoaded() 	? BulkData.Attributes1.GetBulkDataSize() : 0;
-			Total += BulkData.Materials.IsBulkDataLoaded() 		? BulkData.Materials.GetBulkDataSize()   : 0;
-			Total += BulkData.CurveOffsets.IsBulkDataLoaded() 	? BulkData.CurveOffsets.GetBulkDataSize(): 0;
+			Total += BulkData.Positions.IsBulkDataLoaded() 	? BulkData.Positions.GetBulkDataSize()   : 0;
+			Total += BulkData.Attributes0.IsBulkDataLoaded()? BulkData.Attributes0.GetBulkDataSize() : 0;
+			Total += BulkData.Attributes1.IsBulkDataLoaded()? BulkData.Attributes1.GetBulkDataSize() : 0;
+			Total += BulkData.Materials.IsBulkDataLoaded() 	? BulkData.Materials.GetBulkDataSize()   : 0;
+			Total += BulkData.Curves.IsBulkDataLoaded() 	? BulkData.Curves.GetBulkDataSize(): 0;
 			return Total;
 		}
 
diff --git a/Engine/Plugins/Runtime/HairStrands/Source/HairStrandsCore/Public/GroomResources.h b/Engine/Plugins/Runtime/HairStrands/Source/HairStrandsCore/Public/GroomResources.h
--- a/Engine/Plugins/Runtime/HairStrands/Source/HairStrandsCore/Public/GroomResources.h
+++ b/Engine/Plugins/Runtime/HairStrands/Source/HairStrandsCore/Public/GroomResources.h
@@ -320,6 +320,7 @@ struct FHairStrandsRestResource : public FHairCommonResource
 		Total += GetBufferTotalNumBytes(Attribute0Buffer);
 		Total += GetBufferTotalNumBytes(Attribute1Buffer);
 		Total += GetBufferTotalNumBytes(MaterialBuffer);
+		Total += GetBufferTotalNumBytes(CurveBuffer);
 		return Total;
 	}
 
@@ -341,6 +342,9 @@ struct FHairStrandsRestResource : public FHairCommonResource
 	/* Strand hair material buffer */
 	FRDGExternalBuffer MaterialBuffer;
 
+	/* Strand hair curves buffer (contains curves' points offset and count */
+	FRDGExternalBuffer CurveBuffer;
+
 	/* Reference to the hair strands render data */
 	FHairStrandsBulkData& BulkData;
 
diff --git a/Engine/Plugins/Runtime/HairStrands/Source/HairStrandsCore/Public/HairStrandsDatas.h b/Engine/Plugins/Runtime/HairStrands/Source/HairStrandsCore/Public/HairStrandsDatas.h
--- a/Engine/Plugins/Runtime/HairStrands/Source/HairStrandsCore/Public/HairStrandsDatas.h
+++ b/Engine/Plugins/Runtime/HairStrands/Source/HairStrandsCore/Public/HairStrandsDatas.h
@@ -85,6 +85,14 @@ struct FHairInterpolation1Vertex
 	uint8 Pad1;
 };
 
+struct FPackedHairCurve
+{
+	typedef uint32 BulkType;
+
+	uint32 PointOffset : 24;
+	uint32 PointCount  : 8;
+};
+
 struct FVector4_16
 {
 	FFloat16 X;
@@ -198,6 +206,17 @@ struct FHairStrandsRootIndexFormat
 	static const EPixelFormat Format = PF_R32_UINT;
 };
 
+struct FHairStrandsCurveFormat
+{
+	typedef FPackedHairCurve Type;
+	typedef FPackedHairCurve::BulkType BulkType;
+
+	static const uint32 ComponentCount = 1;
+	static const uint32 SizeInByte = sizeof(Type);
+	static const EVertexElementType VertexElementType = VET_UInt;
+	static const EPixelFormat Format = PF_R32_UINT;
+};
+
 struct FHairStrandsRaytracingFormat
 {
 	typedef FVector4f Type;
@@ -469,7 +488,7 @@ struct HAIRSTRANDSCORE_API FHairStrandsBulkData
 	FByteBulkData Attributes0;	// Size = PointCount
 	FByteBulkData Attributes1;	// Size = PointCount
 	FByteBulkData Materials;	// Size = PointCount
-	FByteBulkData CurveOffsets;	// Size = CurveCount+1 - Store the root point index for the curve
+	FByteBulkData Curves;		// Size = CurveCount
 };
 
 /** Hair strands debug data */
diff --git a/Engine/Plugins/Runtime/HairStrands/Source/HairStrandsCore/Public/Niagara/NiagaraDataInterfaceHairStrands.h b/Engine/Plugins/Runtime/HairStrands/Source/HairStrandsCore/Public/Niagara/NiagaraDataInterfaceHairStrands.h
--- a/Engine/Plugins/Runtime/HairStrands/Source/HairStrandsCore/Public/Niagara/NiagaraDataInterfaceHairStrands.h
+++ b/Engine/Plugins/Runtime/HairStrands/Source/HairStrandsCore/Public/Niagara/NiagaraDataInterfaceHairStrands.h
@@ -50,9 +50,6 @@ struct FNDIHairStrandsBuffer : public FRenderResource
 	/** Get the resource name */
 	virtual FString GetFriendlyName() const override { return TEXT("FNDIHairStrandsBuffer"); }
 
-	/** Strand curves point offset buffer */
-	FNiagaraPooledRWBuffer CurvesOffsetsBuffer;
-
 	/** Deformed position buffer in case no resource are there */
 	TRefCountPtr<FRDGPooledBuffer> DeformedPositionBuffer;
 
diff --git a/Engine/Shaders/Private/HairStrands/HairStrandsAttributeCommon.ush b/Engine/Shaders/Private/HairStrands/HairStrandsAttributeCommon.ush
--- a/Engine/Shaders/Private/HairStrands/HairStrandsAttributeCommon.ush
+++ b/Engine/Shaders/Private/HairStrands/HairStrandsAttributeCommon.ush
@@ -35,8 +35,8 @@ FHairVertexAttributes GetHairStrandsAttributes(uint HairPrimitiveId, float2 Vert
 	const uint VertexId1 = HairPrimitiveId+1;
 
 	// Instead of doing the actual interpolation we could rely on dithering as an optimization
-	const FHairAttributes Attributes0 = DecodeHairAttribute(HairStrandsVF.Attribute0Buffer[VertexId0], HairStrandsVF.Attribute1Buffer[VertexId0], HairStrandsVF.Length);
-	const FHairAttributes Attributes1 = DecodeHairAttribute(HairStrandsVF.Attribute0Buffer[VertexId1], HairStrandsVF.Attribute1Buffer[VertexId1], HairStrandsVF.Length);
+	const FHairAttributes Attributes0 = UnpackHairAttribute(HairStrandsVF.Attribute0Buffer[VertexId0], HairStrandsVF.Attribute1Buffer[VertexId0], HairStrandsVF.Length);
+	const FHairAttributes Attributes1 = UnpackHairAttribute(HairStrandsVF.Attribute0Buffer[VertexId1], HairStrandsVF.Attribute1Buffer[VertexId1], HairStrandsVF.Length);
 
 	const FHairControlPoint ControlPoint0 = UnpackHairControlPoint(HairStrandsVF.PositionBuffer[VertexId0], float3(0, 0 ,0), HairStrandsVF.Radius, HairStrandsVF.RootScale, HairStrandsVF.TipScale);
 	const FHairControlPoint ControlPoint1 = UnpackHairControlPoint(HairStrandsVF.PositionBuffer[VertexId1], float3(0, 0, 0), HairStrandsVF.Radius, HairStrandsVF.RootScale, HairStrandsVF.TipScale);
diff --git a/Engine/Shaders/Private/HairStrands/HairStrandsTexturesGeneration.usf b/Engine/Shaders/Private/HairStrands/HairStrandsTexturesGeneration.usf
--- a/Engine/Shaders/Private/HairStrands/HairStrandsTexturesGeneration.usf
+++ b/Engine/Shaders/Private/HairStrands/HairStrandsTexturesGeneration.usf
@@ -306,8 +306,8 @@ bool Trace(
 							OutTangent = Tangent;
 							OutDepth = Distance;
 
-							const FHairAttributes Attributes0 = DecodeHairAttribute(InVF_Attribute0Buffer[P0], InVF_Attribute1Buffer[P0], InVF_Length);
-							const FHairAttributes Attributes1 = DecodeHairAttribute(InVF_Attribute0Buffer[P1], InVF_Attribute1Buffer[P0], InVF_Length);
+							const FHairAttributes Attributes0 = UnpackHairAttribute(InVF_Attribute0Buffer[P0], InVF_Attribute1Buffer[P0], InVF_Length);
+							const FHairAttributes Attributes1 = UnpackHairAttribute(InVF_Attribute0Buffer[P1], InVF_Attribute1Buffer[P0], InVF_Length);
 
 							OutSeed		= Attributes0.Seed;
 							OutRootUV	= Attributes0.RootUV;
diff --git a/Engine/Shaders/Private/HairStrands/HairStrandsVertexFactoryCommon.ush b/Engine/Shaders/Private/HairStrands/HairStrandsVertexFactoryCommon.ush
--- a/Engine/Shaders/Private/HairStrands/HairStrandsVertexFactoryCommon.ush
+++ b/Engine/Shaders/Private/HairStrands/HairStrandsVertexFactoryCommon.ush
@@ -81,7 +81,7 @@ struct FHairAttributes
 	float WorldLength;
 };
 
-FHairAttributes DecodeHairAttribute(float2 Attributes0, uint Attributes1, float InVF_Length)
+FHairAttributes UnpackHairAttribute(float2 Attributes0, uint Attributes1, float InVF_Length)
 {
 	float2 RootUV, RootIndex;
 	RootUV.x		= ((Attributes1    ) & 0x7FF) / 2047.f;	// Coord are encoded on 11 bits
@@ -99,6 +99,23 @@ FHairAttributes DecodeHairAttribute(float2 Attributes0, uint Attributes1, float
 	return Out;
 }
 
+////////////////////////////////////////////////////////////////////////////////////////////////////////////////////////////////
+// Hair Curve
+
+struct FHairCurve
+{
+	uint PointIndex;
+	uint PointCount;
+};
+
+FHairCurve UnpackHairCurve(uint In)
+{
+	FHairCurve Out = (FHairCurve)0;
+	Out.PointIndex = In & 0xFFFFFF;
+	Out.PointCount = In >> 24;
+	return Out;
+}
+
 ///////////////////////////////////////////////////////////////////////////////////////////////////
 // Hair strands instance parameters
 
