Add curve data to hair deformer to have all information to alter curve shape coherently.

diff --git a/Engine/Plugins/Runtime/HairStrands/Shaders/Private/DeformerDataInterfaceGroom.ush b/Engine/Plugins/Runtime/HairStrands/Shaders/Private/DeformerDataInterfaceGroom.ush
--- a/Engine/Plugins/Runtime/HairStrands/Shaders/Private/DeformerDataInterfaceGroom.ush
+++ b/Engine/Plugins/Runtime/HairStrands/Shaders/Private/DeformerDataInterfaceGroom.ush
@@ -13,6 +13,7 @@ float {DataInterfaceName}_VF_HairLength;
 Buffer<uint4>  {DataInterfaceName}_PositionBuffer;
 Buffer<float2> {DataInterfaceName}_Attribute0Buffer;
 Buffer<uint>   {DataInterfaceName}_Attribute1Buffer;
+Buffer<uint>   {DataInterfaceName}_CurveBuffer;
 
 ///////////////////////////////////////////////////////////////////////////////////////////////////
 // Helpers
@@ -42,21 +43,32 @@ FHairAttributes ReadControlPointAttributeData(uint VertexIndex)
 		return (FHairAttributes)0;
 	}
 
-	return DecodeHairAttribute(
+	return UnpackHairAttribute(
 		{DataInterfaceName}_Attribute0Buffer[VertexIndex],
 		{DataInterfaceName}_Attribute1Buffer[VertexIndex],
 		{DataInterfaceName}_VF_HairLength);
 }
 
+FHairCurve ReadCurveData(uint CurveIndex)
+{
+	const uint MaxCurve = {DataInterfaceName}_NumCurves;
+	if (CurveIndex >= MaxCurve)
+	{
+		return (FHairCurve)0;
+	}
+
+	return UnpackHairCurve({DataInterfaceName}_CurveBuffer[CurveIndex]);
+}
+
 ///////////////////////////////////////////////////////////////////////////////////////////////////
 // Interface
 
-uint {DataInterfaceName}_ReadNumControlPoints()
+uint ReadNumControlPoints_{DataInterfaceName}()
 {
 	return {DataInterfaceName}_NumControlPoints;
 }
 
-uint {DataInterfaceName}_ReadNumCurves()
+uint ReadNumCurves_{DataInterfaceName}()
 {
 	return {DataInterfaceName}_NumCurves;
 }
@@ -90,3 +102,14 @@ float ReadSeed_{DataInterfaceName}(uint VertexIndex)
 {
 	return ReadControlPointAttributeData(VertexIndex).Seed;
 }
+
+float ReadCurveOffsetPoint_{DataInterfaceName}(uint CurveIndex)
+{
+	return ReadCurveData(CurveIndex).PointIndex;
+}
+
+float ReadCurveNumPoint_{DataInterfaceName}(uint CurveIndex)
+{
+	return ReadCurveData(CurveIndex).PointCount;
+}
+
diff --git a/Engine/Plugins/Runtime/HairStrands/Shaders/Private/DeformerDataInterfaceGroomWrite.ush b/Engine/Plugins/Runtime/HairStrands/Shaders/Private/DeformerDataInterfaceGroomWrite.ush
--- a/Engine/Plugins/Runtime/HairStrands/Shaders/Private/DeformerDataInterfaceGroomWrite.ush
+++ b/Engine/Plugins/Runtime/HairStrands/Shaders/Private/DeformerDataInterfaceGroomWrite.ush
@@ -6,6 +6,7 @@
 uint {DataInterfaceName}_NumControlPoints;
 uint {DataInterfaceName}_NumCurves;
 uint {DataInterfaceName}_OutputStreamStart;
+Buffer<uint4>   {DataInterfaceName}_PositionBufferSRV;
 RWBuffer<uint4> {DataInterfaceName}_PositionBufferUAV;
 
 uint ReadNumControlPoints_{DataInterfaceName}()
@@ -20,14 +21,17 @@ uint ReadNumCurves_{DataInterfaceName}()
 
 void WritePosition_{DataInterfaceName}(uint ControlPointIndex, float3 Position)
 {
-	// TODO: pipe the other data of the control point
 	FHairControlPoint CP;
 	CP.Position = Position;
 	CP.WorldRadius = 1.f;
 	CP.UCoord = 0;
 	CP.Type = HAIR_CONTROLPOINT_INSIDE;
 
-	const uint4 PackedCP = PackHairControlPoint(CP, 0 /*InVF_PositionOffset*/, 1.0f /*InVF_Radius*/); // TODO
-	const uint BufferIndex = {DataInterfaceName}_OutputStreamStart + ControlPointIndex;
-	{DataInterfaceName}_PositionBufferUAV[BufferIndex] = PackedCP;
+	const uint WriteIndex = {DataInterfaceName}_OutputStreamStart + ControlPointIndex;
+	if (WriteIndex < {DataInterfaceName}_NumControlPoints)
+	{
+		const uint4 PackedCP = {DataInterfaceName}_PositionBufferSRV[WriteIndex];
+		const uint4 NewPackedCP = PackHairControlPointPosition(PackedCP, Position, 0 /*InVF_PositionOffset*/);
+		{DataInterfaceName}_PositionBufferUAV[WriteIndex] = NewPackedCP;
+	}
 }
\ No newline at end of file
diff --git a/Engine/Plugins/Runtime/HairStrands/Source/HairStrandsCore/Private/Deformer/DeformerDataInterfaceGroom.cpp b/Engine/Plugins/Runtime/HairStrands/Source/HairStrandsCore/Private/Deformer/DeformerDataInterfaceGroom.cpp
--- a/Engine/Plugins/Runtime/HairStrands/Source/HairStrandsCore/Private/Deformer/DeformerDataInterfaceGroom.cpp
+++ b/Engine/Plugins/Runtime/HairStrands/Source/HairStrandsCore/Private/Deformer/DeformerDataInterfaceGroom.cpp
@@ -19,13 +19,15 @@ TArray<FOptimusCDIPinDefinition> UOptimusGroomDataInterface::GetPinDefinitions()
 {
 	TArray<FOptimusCDIPinDefinition> Defs;
 	Defs.Add({"NumControlPoints", "ReadNumControlPoints"});
-	Defs.Add({"NumCurves", "ReadNumCurves" });
-	Defs.Add({"Position",  "ReadPosition", Optimus::DomainName::Vertex, "ReadPosition"});
-	Defs.Add({"Radius",    "ReadRadius",   Optimus::DomainName::Vertex, "ReadRadius" });
-	Defs.Add({"CoordU",    "ReadCoordU",   Optimus::DomainName::Vertex, "ReadCoordU" });
-	Defs.Add({"Length",    "ReadLength",   Optimus::DomainName::Vertex, "ReadLength" });
-	Defs.Add({"RootUV",    "ReadRootUV",   Optimus::DomainName::Vertex, "ReadRootUV" });
-	Defs.Add({"Seed",      "ReadSeed",     Optimus::DomainName::Vertex, "ReadSeed"   });
+	Defs.Add({"NumCurves",        "ReadNumCurves" });
+	Defs.Add({"Position",         "ReadPosition",          Optimus::DomainName::Vertex,   "ReadPosition"});
+	Defs.Add({"Radius",           "ReadRadius",            Optimus::DomainName::Vertex,   "ReadRadius" });
+	Defs.Add({"CoordU",           "ReadCoordU",            Optimus::DomainName::Vertex,   "ReadCoordU" });
+	Defs.Add({"Length",           "ReadLength",            Optimus::DomainName::Vertex,   "ReadLength" });
+	Defs.Add({"RootUV",           "ReadRootUV",            Optimus::DomainName::Vertex,   "ReadRootUV" });
+	Defs.Add({"Seed",             "ReadSeed",              Optimus::DomainName::Vertex,   "ReadSeed"   });
+	Defs.Add({"CurveOffsetPoint", "ReadCurveOffsetPoint",  Optimus::DomainName::Triangle, "ReadCurveOffsetPoint" });
+	Defs.Add({"CurveNumPoint",    "ReadCurveNumPoint",     Optimus::DomainName::Triangle, "ReadCurveNumPoint" });
 	return Defs;
 }
  
@@ -73,6 +75,16 @@ void UOptimusGroomDataInterface::GetSupportedInputs(TArray<FShaderFunctionDefini
 		.SetName(TEXT("ReadSeed"))
 		.AddReturnType(EShaderFundamentalType::Float)
 		.AddParam(EShaderFundamentalType::Uint);
+
+	OutFunctions.AddDefaulted_GetRef()
+		.SetName(TEXT("ReadCurveOffsetPoint"))
+		.AddReturnType(EShaderFundamentalType::Uint)
+		.AddParam(EShaderFundamentalType::Uint);
+
+	OutFunctions.AddDefaulted_GetRef()
+		.SetName(TEXT("ReadCurveNumPoint"))
+		.AddReturnType(EShaderFundamentalType::Uint)
+		.AddParam(EShaderFundamentalType::Uint);
 }
 
 BEGIN_SHADER_PARAMETER_STRUCT(FGroomDataInterfaceParameters, )
@@ -83,6 +95,7 @@ BEGIN_SHADER_PARAMETER_STRUCT(FGroomDataInterfaceParameters, )
 	SHADER_PARAMETER_RDG_BUFFER_SRV(Buffer<uint4>,  PositionBuffer)
 	SHADER_PARAMETER_RDG_BUFFER_SRV(Buffer<float2>, Attribute0Buffer)
 	SHADER_PARAMETER_RDG_BUFFER_SRV(Buffer<uint>,   Attribute1Buffer)
+	SHADER_PARAMETER_RDG_BUFFER_SRV(Buffer<uint>,   CurveBuffer)
 END_SHADER_PARAMETER_STRUCT()
 
 void UOptimusGroomDataInterface::GetShaderParameters(TCHAR const* UID, FShaderParametersMetadataBuilder& InOutBuilder, FShaderParametersMetadataAllocations& InOutAllocations) const
@@ -142,10 +155,11 @@ void FOptimusGroomDataProviderProxy::AllocateResources(FRDGBuilder& GraphBuilder
 		if (FHairGroupInstance* Instance = GroomComponent->GetGroupInstance(Index))
 		{
 			FResources& R = Resources.AddDefaulted_GetRef();
-			R.PositionSRV = Register(GraphBuilder, Instance->Strands.RestResource->PositionBuffer, ERDGImportedBufferFlags::CreateSRV).SRV;
+			R.PositionSRV	= Register(GraphBuilder, Instance->Strands.RestResource->PositionBuffer, ERDGImportedBufferFlags::CreateSRV).SRV;
 			R.Attribute0SRV = Register(GraphBuilder, Instance->Strands.RestResource->Attribute0Buffer, ERDGImportedBufferFlags::CreateSRV).SRV;
 			R.Attribute1SRV = Register(GraphBuilder, Instance->Strands.RestResource->Attribute1Buffer, ERDGImportedBufferFlags::CreateSRV).SRV;
-			R.FallbackSRV = GraphBuilder.CreateSRV(GraphBuilder.RegisterExternalBuffer(GWhiteVertexBufferWithRDG->Buffer), PF_R16G16B16A16_UINT);
+			R.CurveSRV		= Register(GraphBuilder, Instance->Strands.RestResource->CurveBuffer, ERDGImportedBufferFlags::CreateSRV).SRV;
+			R.FallbackSRV	= GraphBuilder.CreateSRV(GraphBuilder.RegisterExternalBuffer(GWhiteVertexBufferWithRDG->Buffer), PF_R16G16B16A16_UINT);
 		}
 	}
 }
@@ -182,12 +196,14 @@ void FOptimusGroomDataProviderProxy::GatherDispatchData(FDispatchSetup const& In
 				Parameters->PositionBuffer = Resources[InvocationIndex].PositionSRV;
 				Parameters->Attribute0Buffer = Resources[InvocationIndex].Attribute0SRV;
 				Parameters->Attribute1Buffer = Resources[InvocationIndex].Attribute1SRV;
+				Parameters->CurveBuffer = Resources[InvocationIndex].CurveSRV;
 			}
 			else
 			{
 				Parameters->PositionBuffer   = Resources[InvocationIndex].FallbackSRV;
 				Parameters->Attribute0Buffer = Resources[InvocationIndex].FallbackSRV;
 				Parameters->Attribute1Buffer = Resources[InvocationIndex].FallbackSRV;
+				Parameters->CurveBuffer		 = Resources[InvocationIndex].FallbackSRV;
 			}
 		}
 	}
diff --git a/Engine/Plugins/Runtime/HairStrands/Source/HairStrandsCore/Private/Deformer/DeformerDataInterfaceGroom.h b/Engine/Plugins/Runtime/HairStrands/Source/HairStrandsCore/Private/Deformer/DeformerDataInterfaceGroom.h
--- a/Engine/Plugins/Runtime/HairStrands/Source/HairStrandsCore/Private/Deformer/DeformerDataInterfaceGroom.h
+++ b/Engine/Plugins/Runtime/HairStrands/Source/HairStrandsCore/Private/Deformer/DeformerDataInterfaceGroom.h
@@ -68,6 +68,7 @@ private:
 		FRDGBufferSRVRef Attribute0SRV = nullptr;
 		FRDGBufferSRVRef Attribute1SRV = nullptr;
 		FRDGBufferSRVRef FallbackSRV = nullptr;
+		FRDGBufferSRVRef CurveSRV = nullptr;
 	};
 	TArray<FResources> Resources;
 };
diff --git a/Engine/Plugins/Runtime/HairStrands/Source/HairStrandsCore/Private/Deformer/DeformerDataInterfaceGroomWrite.cpp b/Engine/Plugins/Runtime/HairStrands/Source/HairStrandsCore/Private/Deformer/DeformerDataInterfaceGroomWrite.cpp
--- a/Engine/Plugins/Runtime/HairStrands/Source/HairStrandsCore/Private/Deformer/DeformerDataInterfaceGroomWrite.cpp
+++ b/Engine/Plugins/Runtime/HairStrands/Source/HairStrandsCore/Private/Deformer/DeformerDataInterfaceGroomWrite.cpp
@@ -56,6 +56,7 @@ BEGIN_SHADER_PARAMETER_STRUCT(FGroomWriteDataInterfaceParameters, )
 	SHADER_PARAMETER(uint32, NumControlPoints)
 	SHADER_PARAMETER(uint32, NumCurves)
 	SHADER_PARAMETER(uint32, OutputStreamStart)
+	SHADER_PARAMETER_RDG_BUFFER_SRV(Buffer<uint4>, PositionBufferSRV)
 	SHADER_PARAMETER_RDG_BUFFER_UAV(RWBuffer<uint4>, PositionBufferUAV)
 END_SHADER_PARAMETER_STRUCT()
 
@@ -121,13 +122,16 @@ void FOptimusGroomWriteDataProviderProxy::AllocateResources(FRDGBuilder& GraphBu
 			FResources& R = Resources.AddDefaulted_GetRef();
 			if (OutputMask & 1)
 			{
+				R.PositionBufferSRV = Register(GraphBuilder, GroomInstance->Strands.RestResource->PositionBuffer, ERDGImportedBufferFlags::CreateSRV).SRV;
 				R.PositionBufferUAV = Register(GraphBuilder, GroomInstance->Strands.DeformedResource->GetDeformerBuffer(GraphBuilder), ERDGImportedBufferFlags::CreateUAV).UAV;
 				
 				FRDGBufferRef PositionBuffer_fallback = GraphBuilder.CreateBuffer(FRDGBufferDesc::CreateBufferDesc(8, 1), TEXT("Groom.DeformedPositionBuffer"), ERDGBufferFlags::None);
+				R.PositionBufferSRV_fallback = GraphBuilder.CreateSRV(PositionBuffer_fallback, PF_R16G16B16A16_UINT);
 				R.PositionBufferUAV_fallback = GraphBuilder.CreateUAV(PositionBuffer_fallback, PF_R16G16B16A16_UINT, ERDGUnorderedAccessViewFlags::SkipBarrier);
 			}
 			else
 			{
+				R.PositionBufferSRV = GraphBuilder.CreateSRV(GraphBuilder.RegisterExternalBuffer(GWhiteVertexBufferWithRDG->Buffer), PF_R16G16B16A16_UINT);
 				R.PositionBufferUAV = GraphBuilder.CreateUAV(GraphBuilder.RegisterExternalBuffer(GWhiteVertexBufferWithRDG->Buffer), PF_R16G16B16A16_UINT);
 			}
 		}	
@@ -158,6 +162,7 @@ void FOptimusGroomWriteDataProviderProxy::GatherDispatchData(FDispatchSetup cons
 			Parameters->NumControlPoints = NumControlPoints;
 			Parameters->NumCurves = NumCurves;
 			Parameters->OutputStreamStart = 0;
+			Parameters->PositionBufferSRV = bValid ? Resource.PositionBufferSRV : Resource.PositionBufferSRV_fallback;
 			Parameters->PositionBufferUAV = bValid ? Resource.PositionBufferUAV : Resource.PositionBufferUAV_fallback;
 		}
 	}
diff --git a/Engine/Plugins/Runtime/HairStrands/Source/HairStrandsCore/Private/Deformer/DeformerDataInterfaceGroomWrite.h b/Engine/Plugins/Runtime/HairStrands/Source/HairStrandsCore/Private/Deformer/DeformerDataInterfaceGroomWrite.h
--- a/Engine/Plugins/Runtime/HairStrands/Source/HairStrandsCore/Private/Deformer/DeformerDataInterfaceGroomWrite.h
+++ b/Engine/Plugins/Runtime/HairStrands/Source/HairStrandsCore/Private/Deformer/DeformerDataInterfaceGroomWrite.h
@@ -9,6 +9,7 @@
 class UGroomComponent;
 class FRDGBuffer;
 class FRDGBufferUAV;
+class FRDGBufferSRV;
 
 /** Compute Framework Data Interface for writing skinned mesh. */
 UCLASS(Category = ComputeFramework)
@@ -69,7 +70,9 @@ private:
 
 	struct FResources
 	{
+		FRDGBufferSRV* PositionBufferSRV = nullptr;
 		FRDGBufferUAV* PositionBufferUAV = nullptr;
+		FRDGBufferSRV* PositionBufferSRV_fallback = nullptr;
 		FRDGBufferUAV* PositionBufferUAV_fallback = nullptr;
 	};
 	TArray<FResources> Resources;
