* Change Groom deformer to use an existing SRV resources to ensure it is correctly initialized from RDG point of view.
* Change GroomWriter proxy to only access rendering object (instead of component.

diff --git a/Engine/Plugins/Runtime/HairStrands/Source/HairStrandsCore/Private/Deformer/DeformerDataInterfaceGroomWrite.cpp b/Engine/Plugins/Runtime/HairStrands/Source/HairStrandsCore/Private/Deformer/DeformerDataInterfaceGroomWrite.cpp
--- a/Engine/Plugins/Runtime/HairStrands/Source/HairStrandsCore/Private/Deformer/DeformerDataInterfaceGroomWrite.cpp
+++ b/Engine/Plugins/Runtime/HairStrands/Source/HairStrandsCore/Private/Deformer/DeformerDataInterfaceGroomWrite.cpp
@@ -104,16 +104,18 @@ FComputeDataProviderRenderProxy* UOptimusGroomWriteDataProvider::GetRenderProxy(
 
 FOptimusGroomWriteDataProviderProxy::FOptimusGroomWriteDataProviderProxy(UGroomComponent* InGroomComponent, uint64 InOutputMask)
 {
-	GroomComponent = InGroomComponent;
-	OutputMask = InOutputMask;
+	const uint32 InstanceCount = InGroomComponent ? InGroomComponent->GetGroupCount() : 0;
+	for (uint32 Index = 0; Index < InstanceCount; ++Index)
+	{
+		Instances.Add(InGroomComponent->GetGroupInstance(Index));
+	}
 }
 
 void FOptimusGroomWriteDataProviderProxy::AllocateResources(FRDGBuilder& GraphBuilder)
 {
-	const uint32 InstanceCount = GroomComponent ? GroomComponent->GetGroupCount() : 0;
-	for (uint32 Index=0; Index < InstanceCount; ++Index)
+	for (FHairGroupInstance* GroomInstance : Instances)
 	{
-		if (FHairGroupInstance* GroomInstance = GroomComponent->GetGroupInstance(Index))
+		if (GroomInstance)
 		{
 			// Allocate required buffers
 			const int32 NumControlPoints = GroomInstance->Strands.Data->PointCount;
@@ -126,7 +128,7 @@ void FOptimusGroomWriteDataProviderProxy::AllocateResources(FRDGBuilder& GraphBu
 				R.PositionBufferUAV = Register(GraphBuilder, GroomInstance->Strands.DeformedResource->GetDeformerBuffer(GraphBuilder), ERDGImportedBufferFlags::CreateUAV).UAV;
 				
 				FRDGBufferRef PositionBuffer_fallback = GraphBuilder.CreateBuffer(FRDGBufferDesc::CreateBufferDesc(8, 1), TEXT("Groom.DeformedPositionBuffer"), ERDGBufferFlags::None);
-				R.PositionBufferSRV_fallback = GraphBuilder.CreateSRV(PositionBuffer_fallback, PF_R16G16B16A16_UINT);
+				R.PositionBufferSRV_fallback = GraphBuilder.CreateSRV(GraphBuilder.RegisterExternalBuffer(GWhiteVertexBufferWithRDG->Buffer), PF_R16G16B16A16_UINT);
 				R.PositionBufferUAV_fallback = GraphBuilder.CreateUAV(PositionBuffer_fallback, PF_R16G16B16A16_UINT, ERDGUnorderedAccessViewFlags::SkipBarrier);
 			}
 			else
@@ -145,16 +147,15 @@ void FOptimusGroomWriteDataProviderProxy::GatherDispatchData(FDispatchSetup cons
 		return;
 	}
 
-	const uint32 InstanceCount = GroomComponent ? GroomComponent->GetGroupCount() : 0;
-	check(InDispatchSetup.NumInvocations == InstanceCount);
+	check(InDispatchSetup.NumInvocations == Instances.Num());
 
 	for (int32 InvocationIndex = 0; InvocationIndex < InDispatchSetup.NumInvocations; ++InvocationIndex)
 	{
-		if (FHairGroupInstance* GroomInstance = GroomComponent->GetGroupInstance(InvocationIndex))
+		if (FHairGroupInstance* GroomInstance = Instances[InvocationIndex])
 		{
 			FResources& Resource = Resources[InvocationIndex];
 
-			const bool bValid = Resource.PositionBufferUAV != nullptr;
+			const bool bValid = Resource.PositionBufferUAV != nullptr && Resource.PositionBufferSRV != nullptr;
 			const int32 NumCurves = bValid ? GroomInstance->Strands.Data->CurveCount : 0;
 			const int32 NumControlPoints = bValid ? GroomInstance->Strands.Data->PointCount : 0;
 
diff --git a/Engine/Plugins/Runtime/HairStrands/Source/HairStrandsCore/Private/Deformer/DeformerDataInterfaceGroomWrite.h b/Engine/Plugins/Runtime/HairStrands/Source/HairStrandsCore/Private/Deformer/DeformerDataInterfaceGroomWrite.h
--- a/Engine/Plugins/Runtime/HairStrands/Source/HairStrandsCore/Private/Deformer/DeformerDataInterfaceGroomWrite.h
+++ b/Engine/Plugins/Runtime/HairStrands/Source/HairStrandsCore/Private/Deformer/DeformerDataInterfaceGroomWrite.h
@@ -10,6 +10,7 @@ class UGroomComponent;
 class FRDGBuffer;
 class FRDGBufferUAV;
 class FRDGBufferSRV;
+struct FHairGroupInstance;
 
 /** Compute Framework Data Interface for writing skinned mesh. */
 UCLASS(Category = ComputeFramework)
@@ -65,7 +66,7 @@ public:
 	//~ End FComputeDataProviderRenderProxy Interface
 
 private:
-	UGroomComponent* GroomComponent = nullptr;
+	TArray<FHairGroupInstance*> Instances;
 	uint64 OutputMask;
 
 	struct FResources
