Change GroomComponent MeshDeformerInstanceSettings UPROPERTY settings to match recent similar changes in SkinnedMeshComponent.

diff --git a/Engine/Plugins/Runtime/HairStrands/Source/HairStrandsCore/Public/GroomComponent.h b/Engine/Plugins/Runtime/HairStrands/Source/HairStrandsCore/Public/GroomComponent.h
--- a/Engine/Plugins/Runtime/HairStrands/Source/HairStrandsCore/Public/GroomComponent.h
+++ b/Engine/Plugins/Runtime/HairStrands/Source/HairStrandsCore/Public/GroomComponent.h
@@ -59,7 +59,7 @@ public:
 	UPROPERTY(EditAnywhere, BlueprintReadWrite, interp, Category = "Simulation")
 	FHairSimulationSettings SimulationSettings;
 
-	/** If set the MeshDeformer will be applied on groonm instance for deformation. */
+	/** If set the MeshDeformer will be applied on groom instance for deformation. */
 	UPROPERTY(EditAnywhere, BlueprintReadOnly, Category = "Deformer")
 	TObjectPtr<UMeshDeformer> MeshDeformer;
 
@@ -67,7 +67,8 @@ public:
 	UPROPERTY(Transient, BlueprintReadOnly, Category = "Deformer")
 	TObjectPtr<UMeshDeformerInstance> MeshDeformerInstance;
 
-	UPROPERTY(VisibleAnywhere, BlueprintReadOnly, DisplayName = "Settings", Category = "Deformer", meta = (EditInline, EditCondition = "MeshDeformerInstanceSettings!=nullptr", HideEditConditionToggle, EditConditionHides))
+	/** Object containing instance settings for the bound MeshDeformer. */
+	UPROPERTY(VisibleAnywhere, BlueprintReadOnly, Instanced, Category = "Deformer", meta = (DisplayName = "Deformer Settings", EditCondition = "MeshDeformerInstanceSettings!=nullptr", ShowOnlyInnerProperties))
 	TObjectPtr<UMeshDeformerInstanceSettings> MeshDeformerInstanceSettings;
 
 	/* Reference of the default/debug materials for each geometric representation */
