Fix incorrect DDC caching of groom with several groups and having both strands & cards data for same LOD.

diff --git a/Engine/Plugins/Runtime/HairStrands/Source/HairStrandsCore/Private/GroomAsset.cpp b/Engine/Plugins/Runtime/HairStrands/Source/HairStrandsCore/Private/GroomAsset.cpp
--- a/Engine/Plugins/Runtime/HairStrands/Source/HairStrandsCore/Private/GroomAsset.cpp
+++ b/Engine/Plugins/Runtime/HairStrands/Source/HairStrandsCore/Private/GroomAsset.cpp
@@ -1881,7 +1881,7 @@ void UGroomAsset::SetHairWidth(float Width)
 // differences, etc.) replace the version GUID below with a new one.
 // In case of merge conflicts with DDC versions, you MUST generate a new GUID
 // and set this new GUID as the version.
-#define GROOM_DERIVED_DATA_VERSION TEXT("282985F4A254469A87884E2F1538CCEB")
+#define GROOM_DERIVED_DATA_VERSION TEXT("05A37379F8A049D5986824E9AAA83F41")
 
 #if WITH_EDITORONLY_DATA
 
@@ -2239,6 +2239,7 @@ bool UGroomAsset::CacheDerivedDatas()
 		if (!bSucceed)
 			return false;
 	}
+	UpdateCachedSettings();
 	UpdateHairGroupsInfo();
 
 	if (bIsGameThread)
@@ -2272,8 +2273,6 @@ bool UGroomAsset::CacheDerivedData(uint32 GroupIndex)
 	bSuccess |= CacheCardsGeometry(GroupIndex, DerivedDataKey);
 	bSuccess |= CacheMeshesGeometry(GroupIndex);
 
-	UpdateCachedSettings();
-
 	return bSuccess;
 }
 
