HairDescription: Store and load the UEVersion and LicenseeVersion in addition to CustomVersions for the serialization of types in BulkData.
#rnx

diff --git a/Engine/Plugins/Runtime/HairStrands/Source/HairStrandsCore/Private/HairDescription.cpp b/Engine/Plugins/Runtime/HairStrands/Source/HairStrandsCore/Private/HairDescription.cpp
--- a/Engine/Plugins/Runtime/HairStrands/Source/HairStrandsCore/Private/HairDescription.cpp
+++ b/Engine/Plugins/Runtime/HairStrands/Source/HairStrandsCore/Private/HairDescription.cpp
@@ -81,6 +81,38 @@ void FHairDescription::Serialize(FArchive& Ar)
 	Ar << GroomAttributesSet;
 }
 
+FArchive& operator<<(FArchive& Ar, FHairDescriptionVersion& Version)
+{
+	Ar << Version.bIsValid;
+	if (Version.bIsValid)
+	{
+		Ar << Version.UEVersion;
+		Ar << Version.LicenseeVersion;
+		Version.CustomVersions.Serialize(Ar);
+	}
+	else if (Ar.IsLoading())
+	{
+		Version = FHairDescriptionVersion();
+	}
+	return Ar;
+}
+
+void FHairDescriptionVersion::CopyVersionsFromArchive(const FArchive& Ar)
+{
+	bIsValid = true;
+	UEVersion = Ar.UEVer();
+	LicenseeVersion = Ar.LicenseeUEVer();
+	CustomVersions = Ar.GetCustomVersions();
+}
+
+void FHairDescriptionVersion::CopyVersionsToArchive(FArchive& Ar) const
+{
+	check(bIsValid);
+	Ar.SetUEVer(UEVersion);
+	Ar.SetLicenseeUEVer(LicenseeVersion);
+	Ar.SetCustomVersions(CustomVersions);
+}
+
 ///////////////////////////////////////////////////////////////////////////////////////////////////
 // FHairDescriptionBulkData
 
@@ -94,7 +126,7 @@ void FHairDescriptionBulkData::Serialize(FArchive& Ar, UObject* Owner)
 	if (Ar.IsTransacting())
 	{
 		// If transacting, keep these members alive the other side of an undo, otherwise their values will get lost
-		CustomVersions.Serialize(Ar);
+		Ar << BulkDataVersion;
 		Ar << bBulkDataUpdated;
 	}
 	else
@@ -125,17 +157,15 @@ void FHairDescriptionBulkData::Serialize(FArchive& Ar, UObject* Owner)
 	}
 	else
 	{
-		BulkData.Serialize(Ar, Owner);
+	BulkData.Serialize(Ar, Owner);
 	}
 
 	if (!Ar.IsTransacting() && Ar.IsLoading())
 	{
 		// If loading, take the package custom version so it can be applied to the bulk data archive
 		// when unpacking HairDescription from it
-		// TODO: Save the UEVersion and LicenseeVersion as well
-		FPackageFileVersion UEVersion;
-		int32 LicenseeUEVersion;
-		BulkData.GetBulkDataVersions(Ar, UEVersion, LicenseeUEVersion, CustomVersions);
+		BulkDataVersion.bIsValid = true;
+		BulkData.GetBulkDataVersions(Ar, BulkDataVersion.UEVersion, BulkDataVersion.LicenseeVersion, BulkDataVersion.CustomVersions);
 	}
 }
 
@@ -149,7 +179,7 @@ void FHairDescriptionBulkData::SaveHairDescription(FHairDescription& HairDescrip
 		HairDescription.Serialize(Ar);
 
 		// Preserve CustomVersions at save time so we can reuse the same ones when reloading direct from memory
-		CustomVersions = Ar.GetCustomVersions();
+		BulkDataVersion.CopyVersionsFromArchive(Ar);
 	}
 
 	// Mark the HairDescriptionBulkData as having been updated.
@@ -157,7 +187,6 @@ void FHairDescriptionBulkData::SaveHairDescription(FHairDescription& HairDescrip
 	bBulkDataUpdated = true;
 }
 
-
 void FHairDescriptionBulkData::LoadHairDescription(FHairDescription& HairDescription)
 {
 	HairDescription.Reset();
@@ -168,7 +197,8 @@ void FHairDescriptionBulkData::LoadHairDescription(FHairDescription& HairDescrip
 
 		// Propagate the custom version information from the package to the bulk data, so that the HairDescription
 		// is serialized with the same versioning
-		Ar.SetCustomVersions(CustomVersions);
+		checkf(BulkDataVersion.IsValid(), TEXT("If BulkData is non-empty, we should have set the BulkDataVersion when we populated it."));
+		BulkDataVersion.CopyVersionsToArchive(Ar);
 
 		HairDescription.Serialize(Ar);
 	}
@@ -180,7 +210,7 @@ void FHairDescriptionBulkData::Empty()
 }
 
 FString FHairDescriptionBulkData::GetIdString() const
-{		
+{
 	FString GuidString = UE::Serialization::IoHashToGuid(BulkData.GetPayloadId()).ToString();
 	GuidString += TEXT("X");
 	return GuidString;
diff --git a/Engine/Plugins/Runtime/HairStrands/Source/HairStrandsCore/Public/HairDescription.h b/Engine/Plugins/Runtime/HairStrands/Source/HairStrandsCore/Public/HairDescription.h
--- a/Engine/Plugins/Runtime/HairStrands/Source/HairStrandsCore/Public/HairDescription.h
+++ b/Engine/Plugins/Runtime/HairStrands/Source/HairStrandsCore/Public/HairDescription.h
@@ -109,6 +109,18 @@ private:
 	int32 NumStrands;
 };
 
+struct FHairDescriptionVersion
+{
+	FCustomVersionContainer CustomVersions;
+	FPackageFileVersion UEVersion;
+	int32 LicenseeVersion = 0;
+	bool bIsValid = false;
+	bool IsValid() const { return bIsValid; }
+	void CopyVersionsFromArchive(const FArchive& Ar);
+	void CopyVersionsToArchive(FArchive& Ar) const;
+};
+FArchive& operator<<(FArchive& Ar, FHairDescriptionVersion& Version);
+
 /**
  * Bulk data storage for FHairDescription
  */
@@ -139,8 +151,8 @@ private:
 	/** Internally store bulk data as bytes */
 	UE::Serialization::FEditorBulkData BulkData;
 
-	/** Custom version to propagate to archive when serializing the bulk data */
-	FCustomVersionContainer CustomVersions;
+	/** UE/Custom/Licensee version to propagate to archive when serializing the bulk data */
+	FHairDescriptionVersion BulkDataVersion;
 
 	/** Whether the bulk data has been written via SaveHairDescription */
 	bool bBulkDataUpdated = false;
