Fix Groom bulk data request size.

diff --git a/Engine/Plugins/Runtime/HairStrands/Source/HairStrandsCore/Private/GroomResources.cpp b/Engine/Plugins/Runtime/HairStrands/Source/HairStrandsCore/Private/GroomResources.cpp
--- a/Engine/Plugins/Runtime/HairStrands/Source/HairStrandsCore/Private/GroomResources.cpp
+++ b/Engine/Plugins/Runtime/HairStrands/Source/HairStrandsCore/Private/GroomResources.cpp
@@ -786,7 +786,7 @@ bool FHairStrandsRestResource::InternalIsDataLoaded()
 {
 	if (BulkDataRequest.IsNone())
 	{
-		FBulkDataBatchRequest::FBatchBuilder Batch = FBulkDataBatchRequest::NewBatch(4);
+		FBulkDataBatchRequest::FBatchBuilder Batch = FBulkDataBatchRequest::NewBatch(5);
 		Batch.Read(BulkData.Positions);
 		Batch.Read(BulkData.Attributes0);
 		Batch.Read(BulkData.Attributes1);
