Change groom asset and groom binding asset to use the unified guid system.

diff --git a/Engine/Plugins/Runtime/HairStrands/Source/HairStrandsCore/Private/GroomAsset.cpp b/Engine/Plugins/Runtime/HairStrands/Source/HairStrandsCore/Private/GroomAsset.cpp
--- a/Engine/Plugins/Runtime/HairStrands/Source/HairStrandsCore/Private/GroomAsset.cpp
+++ b/Engine/Plugins/Runtime/HairStrands/Source/HairStrandsCore/Private/GroomAsset.cpp
@@ -38,6 +38,7 @@
 #include "HairCardsVertexFactory.h"
 #include "PSOPrecache.h"
 #include "UObject/UObjectIterator.h"
+#include "UObject/DevObjectVersion.h"
 
 #if WITH_EDITORONLY_DATA
 #include "DerivedDataCache.h"
@@ -1881,7 +1882,10 @@ void UGroomAsset::SetHairWidth(float Width)
 // differences, etc.) replace the version GUID below with a new one.
 // In case of merge conflicts with DDC versions, you MUST generate a new GUID
 // and set this new GUID as the version.
-#define GROOM_DERIVED_DATA_VERSION TEXT("05A37379F8A049D5986824E9AAA83F41")
+// DDC Guid needs to be bumped in:
+// * Main    : Engine\Source\Runtime\Core\Private\UObject\DevObjectVersion.cpp
+// * Release : Engine\Source\Runtime\Core\Private\UObject\UE5ReleaseStreamObjectVersion.cpp
+// * ...
 
 #if WITH_EDITORONLY_DATA
 
@@ -1889,7 +1893,7 @@ namespace GroomDerivedDataCacheUtils
 {
 	const FString& GetGroomDerivedDataVersion()
 	{
-		static FString CachedVersionString(GROOM_DERIVED_DATA_VERSION);
+		static FString CachedVersionString = FDevSystemGuids::GetSystemGuid(FDevSystemGuids::Get().GROOM_DERIVED_DATA_VERSION).ToString();
 		return CachedVersionString;
 	}
 
diff --git a/Engine/Plugins/Runtime/HairStrands/Source/HairStrandsCore/Private/GroomBindingAsset.cpp b/Engine/Plugins/Runtime/HairStrands/Source/HairStrandsCore/Private/GroomBindingAsset.cpp
--- a/Engine/Plugins/Runtime/HairStrands/Source/HairStrandsCore/Private/GroomBindingAsset.cpp
+++ b/Engine/Plugins/Runtime/HairStrands/Source/HairStrandsCore/Private/GroomBindingAsset.cpp
@@ -10,6 +10,7 @@
 #include "HAL/LowLevelMemTracker.h"
 #include "UObject/AnimObjectVersion.h"
 #include "UObject/ObjectSaveContext.h"
+#include "UObject/DevObjectVersion.h"
 
 #include UE_INLINE_GENERATED_CPP_BY_NAME(GroomBindingAsset)
 
@@ -662,13 +663,16 @@ void UGroomBindingAsset::PostEditChangeProperty(FPropertyChangedEvent& PropertyC
 // differences, etc.) replace the version GUID below with a new one.
 // In case of merge conflicts with DDC versions, you MUST generate a new GUID
 // and set this new GUID as the version.
-#define GROOM_BINDING_DERIVED_DATA_VERSION TEXT("30769E530C574C7BA15C56F224A64E32")
+// DDC Guid needs to be bumped in:
+// * Main    : Engine\Source\Runtime\Core\Private\UObject\DevObjectVersion.cpp
+// * Release : Engine\Source\Runtime\Core\Private\UObject\UE5ReleaseStreamObjectVersion.cpp
+// * ...
 
 namespace GroomBindingDerivedDataCacheUtils
 {
 	const FString& GetGroomBindingDerivedDataVersion()
 	{
-		static FString CachedVersionString(GROOM_BINDING_DERIVED_DATA_VERSION);
+		static FString CachedVersionString = FDevSystemGuids::GetSystemGuid(FDevSystemGuids::Get().GROOM_BINDING_DERIVED_DATA_VERSION).ToString();
 		return CachedVersionString;
 	}
 
diff --git a/Engine/Source/Runtime/Core/Private/UObject/DevObjectVersion.cpp b/Engine/Source/Runtime/Core/Private/UObject/DevObjectVersion.cpp
--- a/Engine/Source/Runtime/Core/Private/UObject/DevObjectVersion.cpp
+++ b/Engine/Source/Runtime/Core/Private/UObject/DevObjectVersion.cpp
@@ -118,14 +118,16 @@ FGuid FDevSystemGuidRegistry::GetSystemGuid(FGuid System)
 
 FDevSystemGuids::FDevSystemGuids()
 	: GLOBALSHADERMAP_DERIVEDDATA_VER(0x7912A706, 0x52B8450A, 0x9CD66FA7, 0xEFBBFA0A)
+	, GROOM_BINDING_DERIVED_DATA_VERSION(0x30769E53, 0x0C574C7B, 0xA15C56F2, 0x24A64E32)
+	, GROOM_DERIVED_DATA_VERSION(0x05A37379, 0xF8A049D5, 0x986824E9, 0xAAA83F41)
 	, LANDSCAPE_MOBILE_COOK_VERSION(0x0E9ADF72, 0xD6B64E0D, 0x81C4A92B, 0x081A37AB)
 	, MATERIALSHADERMAP_DERIVEDDATA_VER(0x2579AAFE, 0x8F1D4E4F, 0xB04278C6, 0x35917C17)
 	, NANITE_DERIVEDDATA_VER(0xBEB0226A, 0x070E4ECA, 0x972C1E7D, 0xD8599E68)
 	, NIAGARASHADERMAP_DERIVEDDATA_VER(0x7BBD9913, 0xC1554D20, 0xADAE9F17, 0xB006299E)
 	, Niagara_LatestScriptCompileVersion(0x6D32B8EE, 0x909FCA7E, 0xA5CE4F17, 0x066A5F25)
+	, POSESEARCHDB_DERIVEDDATA_VER(0x389117E4, 0x807A4CC0, 0x9F37C2E6, 0xD808A78D)
 	, SkeletalMeshDerivedDataVersion(0x9B5F4544, 0x76D7481C, 0x9AD3F614, 0xA6C07904)
 	, STATICMESH_DERIVEDDATA_VER(0x5B7A05A7, 0x7371440C, 0xA314621A, 0x9B1E08BC)
-	, POSESEARCHDB_DERIVEDDATA_VER(0x389117E4, 0x807A4CC0, 0x9F37C2E6, 0xD808A78D)
 {
 }
 
diff --git a/Engine/Source/Runtime/Core/Public/UObject/DevObjectVersion.h b/Engine/Source/Runtime/Core/Public/UObject/DevObjectVersion.h
--- a/Engine/Source/Runtime/Core/Public/UObject/DevObjectVersion.h
+++ b/Engine/Source/Runtime/Core/Public/UObject/DevObjectVersion.h
@@ -19,14 +19,16 @@ struct CORE_API FDevSystemGuids
 	
 	FDevSystemGuids();
 	const FGuid GLOBALSHADERMAP_DERIVEDDATA_VER;
+	const FGuid GROOM_BINDING_DERIVED_DATA_VERSION;
+	const FGuid GROOM_DERIVED_DATA_VERSION;
 	const FGuid LANDSCAPE_MOBILE_COOK_VERSION;
 	const FGuid MATERIALSHADERMAP_DERIVEDDATA_VER;
 	const FGuid NANITE_DERIVEDDATA_VER;
 	const FGuid NIAGARASHADERMAP_DERIVEDDATA_VER;
 	const FGuid Niagara_LatestScriptCompileVersion;
+	const FGuid POSESEARCHDB_DERIVEDDATA_VER;
 	const FGuid SkeletalMeshDerivedDataVersion;
 	const FGuid STATICMESH_DERIVEDDATA_VER;
-	const FGuid POSESEARCHDB_DERIVEDDATA_VER;
 };
 
 class CORE_API FDevVersionRegistration :  public FCustomVersionRegistration
