[HairStrandsCore]
* Fixed c++20 compile error

diff --git a/Engine/Plugins/Runtime/HairStrands/Source/HairStrandsCore/Private/Niagara/NiagaraDataInterfaceHairStrands.cpp b/Engine/Plugins/Runtime/HairStrands/Source/HairStrandsCore/Private/Niagara/NiagaraDataInterfaceHairStrands.cpp
--- a/Engine/Plugins/Runtime/HairStrands/Source/HairStrandsCore/Private/Niagara/NiagaraDataInterfaceHairStrands.cpp
+++ b/Engine/Plugins/Runtime/HairStrands/Source/HairStrandsCore/Private/Niagara/NiagaraDataInterfaceHairStrands.cpp
@@ -1004,10 +1004,10 @@ void UNiagaraDataInterfaceHairStrands::SimCachePostReadFrame(void* OptionalPerIn
 	{
 		FNiagaraDataSet& EmitterDatas  = EmitterInst->GetData();
 
-		const FNiagaraVariableBase NodePositionVariable(FNiagaraTypeDefinition::GetVec3Def(), FName("NodePosition"));
+		const FNiagaraVariable NodePositionVariable(FNiagaraTypeDefinition::GetVec3Def(), FName("NodePosition"));
 		const int32 NodePositionIndex = EmitterDatas.GetCompiledData().Variables.IndexOfByKey(NodePositionVariable);
 
-		const FNiagaraVariableBase RestPositionVariable(FNiagaraTypeDefinition::GetVec3Def(), FName("RestPosition"));
+		const FNiagaraVariable RestPositionVariable(FNiagaraTypeDefinition::GetVec3Def(), FName("RestPosition"));
 		const int32 RestPositionIndex = EmitterDatas.GetCompiledData().Variables.IndexOfByKey(RestPositionVariable);
 				
 		if (RestPositionIndex != INDEX_NONE && NodePositionIndex != INDEX_NONE)
