When a groom component is attached to a skinned mesh, the skinned mesh is forced to update its dynamic data immediately, instead of (optionally) deferred them.

This makes GDME optimization compatible with groom simulation.

diff --git a/Engine/Plugins/Runtime/HairStrands/Source/HairStrandsCore/Private/GroomComponent.cpp b/Engine/Plugins/Runtime/HairStrands/Source/HairStrandsCore/Private/GroomComponent.cpp
--- a/Engine/Plugins/Runtime/HairStrands/Source/HairStrandsCore/Private/GroomComponent.cpp
+++ b/Engine/Plugins/Runtime/HairStrands/Source/HairStrandsCore/Private/GroomComponent.cpp
@@ -2539,6 +2539,13 @@ void UGroomComponent::InitResources(bool bIsBindingReloading)
 				LocalBindingAsset = nullptr;
 			}
 		}
+
+		// When a groom is attached to a skinned mesh, skin dynamic data needs to be update immediately and not deferred 
+		// until drawing. This ensures that skin data are ready/available for simulation.
+		if (USkinnedMeshComponent* SkinnedMesh = Cast<USkinnedMeshComponent>(RegisteredMeshComponent))
+		{
+			SkinnedMesh->SetForceUpdateDynamicDataImmediately(true);
+		}
 	}
 
 	// Update requirement based on binding asset availability
