Clean up ray tracing masks for groom to prepare for path tracing and raytracing mask separation.

diff --git a/Engine/Plugins/Runtime/HairStrands/Source/HairStrandsCore/Private/GroomComponent.cpp b/Engine/Plugins/Runtime/HairStrands/Source/HairStrandsCore/Private/GroomComponent.cpp
--- a/Engine/Plugins/Runtime/HairStrands/Source/HairStrandsCore/Private/GroomComponent.cpp
+++ b/Engine/Plugins/Runtime/HairStrands/Source/HairStrandsCore/Private/GroomComponent.cpp
@@ -556,7 +556,7 @@ public:
 				{
 					RTGeometry = Instance->Strands.RenRaytracingResource;
 					InstanceViewRayTracingMask = Instance->Strands.ViewRayTracingMask;
-					RayTracingMask = RAY_TRACING_MASK_THIN_SHADOW;
+					RayTracingMask = RAY_TRACING_MASK_HAIR_STRANDS;
 					break;
 				}
 				case EHairGeometryType::Cards:
@@ -592,27 +592,6 @@ public:
 					RayTracingInstance.InstanceTransforms.Add(OverrideLocalToWorld);
 					RayTracingInstance.BuildInstanceMaskAndFlags(GetScene().GetFeatureLevel(), ERayTracingInstanceLayer::NearField, RayTracingMask);
 
-					// If thin shadow is requested, we ensure that regular shadow mask is not added.
-					// TODO: can this logic live in BuildInstanceMaskAndFlags instead?
-					if (RayTracingMask == RAY_TRACING_MASK_THIN_SHADOW)
-					{
-						if (RayTracingInstance.Mask & RAY_TRACING_MASK_SHADOW)
-						{
-							// geometry casts shadows, make sure it is in only one of the shadow casting groups, so it can be treated seperately if desired
-							RayTracingInstance.Mask &= ~RAY_TRACING_MASK_SHADOW;
-							RayTracingInstance.Mask |= RAY_TRACING_MASK_THIN_SHADOW;
-						}
-						else
-						{
-							// if the geometry does not cast shadows, remove this flag
-							RayTracingInstance.Mask &= ~RAY_TRACING_MASK_THIN_SHADOW;
-						}
-						// make sure geometry is only in the hair group
-						RayTracingInstance.Mask &= ~RAY_TRACING_MASK_OPAQUE;
-						RayTracingInstance.Mask &= ~RAY_TRACING_MASK_TRANSLUCENT;
-						RayTracingInstance.Mask |= RAY_TRACING_MASK_HAIR_STRANDS;
-					}
-
 					OutRayTracingInstances.Add(RayTracingInstance);
 				}
 			}
diff --git a/Engine/Source/Runtime/Engine/Private/RayTracingInstance.cpp b/Engine/Source/Runtime/Engine/Private/RayTracingInstance.cpp
--- a/Engine/Source/Runtime/Engine/Private/RayTracingInstance.cpp
+++ b/Engine/Source/Runtime/Engine/Private/RayTracingInstance.cpp
@@ -55,6 +55,19 @@ FRayTracingMaskAndFlags BuildRayTracingInstanceMaskAndFlags(TArrayView<const FMe
 
 	Result.Mask |= bAnySegmentsCastShadow ? RAY_TRACING_MASK_SHADOW : 0;
 
+	if (Result.Mask & RAY_TRACING_MASK_HAIR_STRANDS)
+	{
+		// For hair strands, opaque/translucent mask should be cleared to make sure geometry is only in the hair group. 
+		// If any segment receives shadow, it should receive only thin shadow instead of shadow.
+
+		Result.Mask &= (~RAY_TRACING_MASK_SHADOW | ~RAY_TRACING_MASK_THIN_SHADOW | ~RAY_TRACING_MASK_TRANSLUCENT | ~RAY_TRACING_MASK_OPAQUE);
+
+		if (bAnySegmentsCastShadow)
+		{
+			Result.Mask |= RAY_TRACING_MASK_THIN_SHADOW;
+		}
+	}
+
 	if (InstanceLayer == ERayTracingInstanceLayer::FarField)
 	{
 		// if far field, set that flag exclusively
