GroomCache:
* Added support for importing guides-only groom cache.
* Added groom cache import options to import strands or guides cache and override groom asset conversion settings (for scripting).
* Added groom cache compatibility with groom asset detection in UI and groom cache importer (for scripting).

diff --git a/Engine/Plugins/Runtime/HairStrands/Source/HairStrandsCore/Public/GroomCacheImportOptions.h b/Engine/Plugins/Runtime/HairStrands/Source/HairStrandsCore/Public/GroomCacheImportOptions.h
--- a/Engine/Plugins/Runtime/HairStrands/Source/HairStrandsCore/Public/GroomCacheImportOptions.h
+++ b/Engine/Plugins/Runtime/HairStrands/Source/HairStrandsCore/Public/GroomCacheImportOptions.h
@@ -3,9 +3,21 @@
 #pragma once
 
 #include "EditorFramework/AssetImportData.h"
+#include "GroomSettings.h"
 #include "UObject/SoftObjectPath.h"
 #include "GroomCacheImportOptions.generated.h"
 
+UENUM(BlueprintType)
+enum class EGroomCacheImportType : uint8
+{
+	None = 0x00 UMETA(Hidden),
+	Strands = 0x01,
+	Guides = 0x02,
+	All = Strands | Guides
+};
+
+ENUM_CLASS_FLAGS(EGroomCacheImportType);
+
 USTRUCT(BlueprintType)
 struct HAIRSTRANDSCORE_API FGroomCacheImportSettings
 {
@@ -15,6 +27,10 @@ struct HAIRSTRANDSCORE_API FGroomCacheImportSettings
 	UPROPERTY(EditAnywhere, BlueprintReadWrite, Category = GroomCache)
 	bool bImportGroomCache = true;
 
+	/** Groom Cache types to import */
+	UPROPERTY(EditAnywhere, BlueprintReadWrite, Category = GroomCache)
+	EGroomCacheImportType ImportType = EGroomCacheImportType::All;
+
 	/** Starting index to start sampling the animation from */
 	UPROPERTY(EditAnywhere, BlueprintReadWrite, Category = Sampling)
 	int32 FrameStart = 0;
@@ -34,6 +50,14 @@ struct HAIRSTRANDSCORE_API FGroomCacheImportSettings
 	/** The groom asset the groom cache will be built from (must be compatible) */
 	UPROPERTY(EditAnywhere, BlueprintReadWrite, Category = GroomCache, meta = (MetaClass = "/Script/HairStrandsCore.GroomAsset"))
 	FSoftObjectPath GroomAsset;
+
+	/** Set to true to override the groom conversion settings. Otherwise, use the settings from the groom import options */
+	UPROPERTY(BlueprintReadWrite, Category = GroomCache)
+	bool bOverrideConversionSettings = false;
+
+	/** Conversion settings to apply to the groom cache import when override is enabled */
+	UPROPERTY(BlueprintReadWrite, Category = GroomCache)
+	FGroomConversionSettings ConversionSettings;
 };
 
 UCLASS(BlueprintType)
diff --git a/Engine/Plugins/Runtime/HairStrands/Source/HairStrandsEditor/Private/GroomCacheImportSettingsCustomization.cpp b/Engine/Plugins/Runtime/HairStrands/Source/HairStrandsEditor/Private/GroomCacheImportSettingsCustomization.cpp
--- a/Engine/Plugins/Runtime/HairStrands/Source/HairStrandsEditor/Private/GroomCacheImportSettingsCustomization.cpp
+++ b/Engine/Plugins/Runtime/HairStrands/Source/HairStrandsEditor/Private/GroomCacheImportSettingsCustomization.cpp
@@ -31,7 +31,8 @@ void FGroomCacheImportSettingsCustomization::CustomizeChildren(TSharedRef<IPrope
 			// Visibility depends on other settings
 			VisibleType = 1;
 		}
-		else if	(ChildHandle->GetProperty()->GetFName() == GET_MEMBER_NAME_CHECKED(FGroomCacheImportSettings, FrameStart) ||
+		else if	(ChildHandle->GetProperty()->GetFName() == GET_MEMBER_NAME_CHECKED(FGroomCacheImportSettings, ImportType) ||
+			ChildHandle->GetProperty()->GetFName() == GET_MEMBER_NAME_CHECKED(FGroomCacheImportSettings, FrameStart) ||
 			ChildHandle->GetProperty()->GetFName() == GET_MEMBER_NAME_CHECKED(FGroomCacheImportSettings, FrameEnd) ||
 			ChildHandle->GetProperty()->GetFName() == GET_MEMBER_NAME_CHECKED(FGroomCacheImportSettings, bSkipEmptyFrames))
 		{
diff --git a/Engine/Plugins/Runtime/HairStrands/Source/HairStrandsEditor/Private/GroomCacheImporter.cpp b/Engine/Plugins/Runtime/HairStrands/Source/HairStrandsEditor/Private/GroomCacheImporter.cpp
--- a/Engine/Plugins/Runtime/HairStrands/Source/HairStrandsEditor/Private/GroomCacheImporter.cpp
+++ b/Engine/Plugins/Runtime/HairStrands/Source/HairStrandsEditor/Private/GroomCacheImporter.cpp
@@ -93,6 +93,8 @@ static UGroomCache* CreateGroomCache(EGroomCacheType Type, UObject*& InParent, c
 
 UGroomCache* FGroomCacheImporter::ProcessToGroomCache(FGroomCacheProcessor& Processor, const FGroomAnimationInfo& AnimInfo, FHairImportContext& ImportContext, const FString& ObjectNameSuffix)
 {
+	TRACE_CPUPROFILER_EVENT_SCOPE(FGroomCacheImporter::ProcessToGroomCache);
+
 	if (UGroomCache* GroomCache = CreateGroomCache(Processor.GetType(), ImportContext.Parent, ObjectNameSuffix, ImportContext.Flags))
 	{
 		Processor.TransferChunks(GroomCache);
@@ -105,9 +107,14 @@ UGroomCache* FGroomCacheImporter::ProcessToGroomCache(FGroomCacheProcessor& Proc
 	return nullptr;
 }
 
-TArray<UGroomCache*> FGroomCacheImporter::ImportGroomCache(const FString& SourceFilename, TSharedPtr<IGroomTranslator> Translator, const FGroomAnimationInfo& AnimInfo, FHairImportContext& HairImportContext, UGroomAsset* GroomAssetForCache)
+TArray<UGroomCache*> FGroomCacheImporter::ImportGroomCache(const FString& SourceFilename, TSharedPtr<IGroomTranslator> Translator, const FGroomAnimationInfo& AnimInfo, FHairImportContext& HairImportContext, UGroomAsset* GroomAssetForCache, EGroomCacheImportType ImportType)
 {
 	bool bSuccess = true;
+	bool bGuidesOnly = false;
+
+	const bool bImportStrandsCache = (EnumHasAnyFlags(ImportType, EGroomCacheImportType::Strands));
+	const bool bImportGuidesCache = (EnumHasAnyFlags(ImportType, EGroomCacheImportType::Guides));
+
 	FGroomCacheProcessor StrandsProcessor(EGroomCacheType::Strands, AnimInfo.Attributes);
 	FGroomCacheProcessor GuidesProcessor(EGroomCacheType::Guides, AnimInfo.Attributes);
 	if (Translator->BeginTranslation(SourceFilename))
@@ -122,6 +129,8 @@ TArray<UGroomCache*> FGroomCacheImporter::ImportGroomCache(const FString& Source
 		// Each frame is translated into a HairDescription and processed into HairGroupData
 		for (int32 FrameIndex = AnimInfo.StartFrame; FrameIndex < AnimInfo.EndFrame + 1; ++FrameIndex)
 		{
+			TRACE_CPUPROFILER_EVENT_SCOPE(FGroomCacheImporter::ImportGroomCache::OneFrame);
+
 			const uint32 CurrentFrame = FrameIndex - AnimInfo.StartFrame;
 
 			FTextBuilder TextBuilder;
@@ -140,6 +149,14 @@ TArray<UGroomCache*> FGroomCacheImporter::ImportGroomCache(const FString& Source
 
 				const uint32 GroupCount = HairDescriptionGroups.HairGroups.Num();
 
+				if (GroupCount != GroomHairGroupsData.Num())
+				{
+					bSuccess = false;
+					UE_LOG(LogGroomCacheImporter, Warning, TEXT("GroomCache does not have the same number of groups as the static groom (%d instead of %d). Aborting GroomCache import."),
+						GroupCount, GroomHairGroupsData.Num());
+					break;
+				}
+
 				TArray<FHairGroupInfoWithVisibility> HairGroupsInfo = GroomAssetForCache->HairGroupsInfo;
 				TArray<FHairDescriptionGroup> HairGroupsData;
 				HairGroupsData.SetNum(GroupCount);
@@ -151,24 +168,40 @@ TArray<UGroomCache*> FGroomCacheImporter::ImportGroomCache(const FString& Source
 				}
 
 				// Validate that the GroomCache has the same topology as the static groom
-				if (HairGroupsData.Num() == GroomHairGroupsData.Num())
+				for (uint32 GroupIndex = 0; GroupIndex < GroupCount; ++GroupIndex)
 				{
-					for (uint32 GroupIndex = 0; GroupIndex < GroupCount; ++GroupIndex)
+					if (bImportStrandsCache)
 					{
-						if (HairGroupsData[GroupIndex].Strands.GetNumPoints() != GroomHairGroupsData[GroupIndex].Strands.BulkData.GetNumPoints())
+						if (HairGroupsData[GroupIndex].Strands.GetNumCurves() != GroomHairGroupsData[GroupIndex].Strands.BulkData.GetNumCurves() ||
+							HairGroupsData[GroupIndex].Strands.GetNumPoints() != GroomHairGroupsData[GroupIndex].Strands.BulkData.GetNumPoints())
 						{
 							bSuccess = false;
-							UE_LOG(LogGroomCacheImporter, Warning, TEXT("GroomCache frame %d does not have the same number of vertices as the static groom (%u instead of %u). Aborting GroomCache import."),
-								FrameIndex, HairGroupsData[GroupIndex].Strands.GetNumPoints(), GroomHairGroupsData[GroupIndex].Strands.BulkData.GetNumPoints());
+							UE_LOG(LogGroomCacheImporter, Warning, TEXT("GroomCache frame %d does not have the same number of curves (%u) \
+								or vertices (%u) for the strands as the static groom (%u and %u respectively). Aborting GroomCache import."),
+								FrameIndex, HairGroupsData[GroupIndex].Strands.GetNumCurves(), HairGroupsData[GroupIndex].Strands.GetNumPoints(),
+								GroomHairGroupsData[GroupIndex].Strands.BulkData.GetNumCurves(), GroomHairGroupsData[GroupIndex].Strands.BulkData.GetNumPoints());
+							break;
+						}
+					}
+
+					if (bImportGuidesCache)
+					{
+						if (HairGroupsData[GroupIndex].Strands.GetNumPoints() == 0)
+						{
+							bGuidesOnly = true;
+						}
+
+						if (HairGroupsData[GroupIndex].Guides.GetNumCurves() != GroomHairGroupsData[GroupIndex].Guides.BulkData.GetNumCurves() ||
+							HairGroupsData[GroupIndex].Guides.GetNumPoints() != GroomHairGroupsData[GroupIndex].Guides.BulkData.GetNumPoints())
+						{
+							bSuccess = false;
+							UE_LOG(LogGroomCacheImporter, Warning, TEXT("GroomCache frame %d does not have the same number of curves (%u) \
+								or vertices (%u) for the guides as the static groom (%u and %u respectively). Aborting GroomCache import."),
+								FrameIndex, HairGroupsData[GroupIndex].Guides.GetNumCurves(), HairGroupsData[GroupIndex].Guides.GetNumPoints(),
+								GroomHairGroupsData[GroupIndex].Guides.BulkData.GetNumCurves(), GroomHairGroupsData[GroupIndex].Guides.BulkData.GetNumPoints());
 							break;
 						}
 					}
-				}
-				else
-				{
-					bSuccess = false;
-					UE_LOG(LogGroomCacheImporter, Warning, TEXT("GroomCache does not have the same number of groups as the static groom (%d instead of %d). Aborting GroomCache import."),
-						HairGroupsData.Num(), GroomHairGroupsData.Num());
 				}
 
 				if (!bSuccess)
@@ -177,8 +210,15 @@ TArray<UGroomCache*> FGroomCacheImporter::ImportGroomCache(const FString& Source
 				}
 
 				// The HairGroupData is converted into animated groom data by the GroomCacheProcessor
-				StrandsProcessor.AddGroomSample(MoveTemp(HairGroupsData));
-				GuidesProcessor.AddGroomSample(MoveTemp(HairGroupsData));
+				if (bImportStrandsCache && !bGuidesOnly)
+				{
+					StrandsProcessor.AddGroomSample(MoveTemp(HairGroupsData));
+				}
+
+				if (bImportGuidesCache)
+				{
+					GuidesProcessor.AddGroomSample(MoveTemp(HairGroupsData));
+				}
 			}
 		}
 	}
@@ -192,16 +232,22 @@ TArray<UGroomCache*> FGroomCacheImporter::ImportGroomCache(const FString& Source
 	if (bSuccess)
 	{
 		// Once the processing has completed successfully, the data is transferred to the GroomCache
-		UGroomCache* GroomCache = ProcessToGroomCache(StrandsProcessor, AnimInfo, HairImportContext, "strands_cache");
-		if (GroomCache)
+		if (bImportStrandsCache && !bGuidesOnly)
 		{
-			GroomCaches.Add(GroomCache);
+			UGroomCache* GroomCache = ProcessToGroomCache(StrandsProcessor, AnimInfo, HairImportContext, "strands_cache");
+			if (GroomCache)
+			{
+				GroomCaches.Add(GroomCache);
+			}
 		}
 
-		GroomCache = ProcessToGroomCache(GuidesProcessor, AnimInfo, HairImportContext, "guides_cache");
-		if (GroomCache)
+		if (bImportGuidesCache)
 		{
-			GroomCaches.Add(GroomCache);
+			UGroomCache* GroomCache = ProcessToGroomCache(GuidesProcessor, AnimInfo, HairImportContext, "guides_cache");
+			if (GroomCache)
+			{
+				GroomCaches.Add(GroomCache);
+			}
 		}
 	}
 	return GroomCaches;
diff --git a/Engine/Plugins/Runtime/HairStrands/Source/HairStrandsEditor/Private/GroomImportOptionsWindow.cpp b/Engine/Plugins/Runtime/HairStrands/Source/HairStrandsEditor/Private/GroomImportOptionsWindow.cpp
--- a/Engine/Plugins/Runtime/HairStrands/Source/HairStrandsEditor/Private/GroomImportOptionsWindow.cpp
+++ b/Engine/Plugins/Runtime/HairStrands/Source/HairStrandsEditor/Private/GroomImportOptionsWindow.cpp
@@ -21,31 +21,165 @@
 
 enum class EHairDescriptionStatus
 {
+	Unset,
 	Valid,
 	NoGroup,
 	NoCurve,
+	GroomCache, // groom cache with unspecified groom asset
+	GroomCacheCompatible,
+	GroomCacheIncompatible,
+	GuidesOnly, // guides-only with unspecified groom asset
+	GuidesOnlyCompatible,
+	GuidesOnlyIncompatible,
 	Unknown
 };
 
-static EHairDescriptionStatus GetStatus(UGroomHairGroupsPreview* Description)
+void SGroomImportOptionsWindow::UpdateStatus(UGroomHairGroupsPreview* Description) const
 {
-	if (Description)
+	if (!Description)
 	{
-		for (const FGroomHairGroupPreview& Group : Description->Groups)
+		CurrentStatus = EHairDescriptionStatus::Unknown;
+		return;
+	}
+
+	const bool bImportGroomAsset = !GroomCacheImportOptions || GroomCacheImportOptions->ImportSettings.bImportGroomAsset;
+	const bool bImportGroomCache = GroomCacheImportOptions && GroomCacheImportOptions->ImportSettings.bImportGroomCache;
+	if (!bImportGroomAsset && !bImportGroomCache)
+	{
+		CurrentStatus = EHairDescriptionStatus::Unset;
+		return;
+	}
+
+	if (Description->Groups.Num() == 0)
+	{
+		CurrentStatus = EHairDescriptionStatus::NoGroup;
+		return;
+	}
+
+	// Check the validity of the groom to import
+	CurrentStatus = EHairDescriptionStatus::Valid;
+
+	bool bGuidesOnly = false;
+	for (const FGroomHairGroupPreview& Group : Description->Groups)
+	{
+		if (Group.CurveCount == 0)
+		{
+			CurrentStatus = EHairDescriptionStatus::NoCurve;
+			if (Group.GuideCount > 0)
+			{
+				bGuidesOnly = true;
+			}
+			break;
+		}
+	}
+
+	if (!bImportGroomCache)
+	{
+		return;
+	}
+
+	// Update the states of the properties being monitored
+	bImportGroomAssetState = GroomCacheImportOptions->ImportSettings.bImportGroomAsset;
+	bImportGroomCacheState = GroomCacheImportOptions->ImportSettings.bImportGroomCache;
+	GroomAsset = GroomCacheImportOptions->ImportSettings.GroomAsset;
+
+	if (!GroomCacheImportOptions->ImportSettings.bImportGroomAsset)
+	{
+		// When importing a groom cache with a provided groom asset, check their compatibility
+		UGroomAsset* GroomAssetForCache = Cast<UGroomAsset>(GroomCacheImportOptions->ImportSettings.GroomAsset.TryLoad());
+		if (!GroomAssetForCache)
 		{
-			if (Group.CurveCount == 0)
-				return EHairDescriptionStatus::NoCurve;
+			// No groom asset provided or loaded but one is needed with this setting
+			CurrentStatus = bGuidesOnly ? EHairDescriptionStatus::GuidesOnly : EHairDescriptionStatus::GroomCache;
+			return;
 		}
 
-		if (Description->Groups.Num() == 0)
+		const TArray<FHairGroupData>& GroomHairGroupsData = GroomAssetForCache->HairGroupsData;
+		if (GroomHairGroupsData.Num() != Description->Groups.Num())
 		{
-			return EHairDescriptionStatus::NoGroup;
+			CurrentStatus = bGuidesOnly ? EHairDescriptionStatus::GuidesOnlyIncompatible : EHairDescriptionStatus::GroomCacheIncompatible;
+			return;
+		}
+
+		CurrentStatus = bGuidesOnly ? EHairDescriptionStatus::GuidesOnlyCompatible : EHairDescriptionStatus::GroomCacheCompatible;
+		for (int32 Index = 0; Index < GroomHairGroupsData.Num(); ++Index)
+		{
+			// Check the strands compatibility
+			if (!bGuidesOnly && Description->Groups[Index].CurveCount != GroomHairGroupsData[Index].Strands.BulkData.GetNumCurves())
+			{
+				CurrentStatus = EHairDescriptionStatus::GroomCacheIncompatible;
+				break;
+			}
+
+			// Check the guides compatibility if there were strands tagged as guides
+			// Otherwise, guides will be generated according to the groom asset interpolation settings
+			// and compatibility cannot be determined here
+			if (Description->Groups[Index].GuideCount > 0 &&
+				Description->Groups[Index].GuideCount != GroomHairGroupsData[Index].Guides.BulkData.GetNumCurves())
+			{
+				CurrentStatus = bGuidesOnly ? EHairDescriptionStatus::GuidesOnlyIncompatible : EHairDescriptionStatus::GroomCacheIncompatible;
+				break;
+			}
 		}
-		return EHairDescriptionStatus::Valid;
 	}
 	else
 	{
-		return EHairDescriptionStatus::Unknown;
+		// A guides-only groom cannot be imported as asset, but otherwise the imported groom asset
+		// is always compatible with the groom cache since they are from the same file
+		CurrentStatus = bGuidesOnly ? EHairDescriptionStatus::GuidesOnly : EHairDescriptionStatus::Valid;
+	}
+}
+
+FText SGroomImportOptionsWindow::GetStatusText() const
+{
+	switch (CurrentStatus)
+	{
+		case EHairDescriptionStatus::Valid:
+			return LOCTEXT("GroomOptionsWindow_ValidationText0", "Valid");
+		case EHairDescriptionStatus::NoCurve:
+			return LOCTEXT("GroomOptionsWindow_ValidationText1", "Invalid. Some groups have 0 curves.");
+		case EHairDescriptionStatus::NoGroup:
+			return LOCTEXT("GroomOptionsWindow_ValidationText2", "Invalid. The groom does not contain any group.");
+		case EHairDescriptionStatus::GroomCache:
+			return LOCTEXT("GroomOptionsWindow_ValidationText4", "A compatible groom asset must be provided to import the groom cache.");
+		case EHairDescriptionStatus::GroomCacheCompatible:
+			return LOCTEXT("GroomOptionsWindow_ValidationText5", "The groom cache is compatible with the groom asset provided .");
+		case EHairDescriptionStatus::GroomCacheIncompatible:
+			return LOCTEXT("GroomOptionsWindow_ValidationText6", "The groom cache is incompatible with the groom asset provided .");
+		case EHairDescriptionStatus::GuidesOnly:
+			return LOCTEXT("GroomOptionsWindow_ValidationText7", "Only guides were detected. A compatible groom asset must be provided.");
+		case EHairDescriptionStatus::GuidesOnlyCompatible:
+			return LOCTEXT("GroomOptionsWindow_ValidationText8", "Only guides were detected. The groom asset provided is compatible.");
+		case EHairDescriptionStatus::GuidesOnlyIncompatible:
+			return LOCTEXT("GroomOptionsWindow_ValidationText9", "Only guides were detected. The groom asset provided is incompatible.");
+		case EHairDescriptionStatus::Unset:
+		case EHairDescriptionStatus::Unknown:
+		default:
+			return LOCTEXT("GroomOptionsWindow_ValidationText3", "Unknown");
+	}
+}
+
+FSlateColor SGroomImportOptionsWindow::GetStatusColor() const
+{
+	switch (CurrentStatus)
+	{
+		case EHairDescriptionStatus::Valid:
+		case EHairDescriptionStatus::GroomCacheCompatible:
+		case EHairDescriptionStatus::GuidesOnlyCompatible:
+			return FLinearColor(0, 0.80f, 0, 1);
+		case EHairDescriptionStatus::NoCurve:
+		case EHairDescriptionStatus::GroomCacheIncompatible:
+		case EHairDescriptionStatus::GuidesOnlyIncompatible:
+			return FLinearColor(0.80f, 0, 0, 1);
+		case EHairDescriptionStatus::NoGroup:
+			return FLinearColor(1, 0, 0, 1);
+		case EHairDescriptionStatus::GroomCache:
+		case EHairDescriptionStatus::GuidesOnly:
+			return FLinearColor(0.80f, 0.80f, 0, 1);
+		case EHairDescriptionStatus::Unset:
+		case EHairDescriptionStatus::Unknown:
+		default:
+			return FLinearColor(1, 1, 1);
 	}
 }
 
@@ -70,29 +204,8 @@ void SGroomImportOptionsWindow::Construct(const FArguments& InArgs)
 	GroomCacheDetailsView = PropertyEditorModule.CreateDetailView(DetailsViewArgs);
 	GroomCacheDetailsView->SetObject(GroomCacheImportOptions);
 
-	const EHairDescriptionStatus Status = GetStatus(GroupsPreview);
-
-	FText ValidationText;
-	FLinearColor ValidationColor(1,1,1);
-	if (Status == EHairDescriptionStatus::Valid)
-	{
-		ValidationText = LOCTEXT("GroomOptionsWindow_ValidationText0", "Valid");
-		ValidationColor = FLinearColor(0, 0.80f, 0, 1);
-	}
-	else if (Status == EHairDescriptionStatus::NoCurve)
-	{
-		ValidationText = LOCTEXT("GroomOptionsWindow_ValidationText1", "Invalid. Some groups have 0 curves.");
-		ValidationColor = FLinearColor(0.80f, 0, 0, 1);
-	}
-	else if (Status == EHairDescriptionStatus::NoGroup)
-	{
-		ValidationText = LOCTEXT("GroomOptionsWindow_ValidationText2", "Invalid. The groom does not contain any group.");
-		ValidationColor = FLinearColor(1, 0, 0, 1);
-	}
-	else
-	{
-		ValidationText = LOCTEXT("GroomOptionsWindow_ValidationText3", "Unknown");
-	}
+	CurrentStatus = EHairDescriptionStatus::Unset;
+	UpdateStatus(GroupsPreview);
 
 	const FSlateFontInfo AttributeFont = FAppStyle::GetFontStyle("CurveEd.InfoFont");
 	const FSlateFontInfo AttributeResultFont = FAppStyle::GetFontStyle("CurveEd.InfoFont");
@@ -171,8 +284,8 @@ void SGroomImportOptionsWindow::Construct(const FArguments& InArgs)
 				[
 					SNew(STextBlock)
 					.Font(FAppStyle::GetFontStyle("CurveEd.InfoFont"))
-					.Text(ValidationText)
-					.ColorAndOpacity(ValidationColor)
+					.Text(TAttribute<FText>::Create(TAttribute<FText>::FGetter::CreateSP(this, &SGroomImportOptionsWindow::GetStatusText)))
+					.ColorAndOpacity(TAttribute<FSlateColor>::Create(TAttribute<FSlateColor>::FGetter::CreateSP(this, &SGroomImportOptionsWindow::GetStatusColor)))
 				]
 			]
 		]
@@ -345,31 +458,38 @@ void SGroomImportOptionsWindow::Construct(const FArguments& InArgs)
 	];
 }
 
-bool SGroomImportOptionsWindow::CanImport()  const
+bool SGroomImportOptionsWindow::CanImport() const
 {
-	if (GroupsPreview)
+	bool bNeedUpdate = CurrentStatus == EHairDescriptionStatus::Unset;
+	if (GroomCacheImportOptions)
 	{
-		for (const FGroomHairGroupPreview& Group : GroupsPreview->Groups)
-		{
-			if (Group.CurveCount == 0)
-				return false;
-		}
+		bNeedUpdate |= bImportGroomAssetState != GroomCacheImportOptions->ImportSettings.bImportGroomAsset;
+		bNeedUpdate |= bImportGroomCacheState != GroomCacheImportOptions->ImportSettings.bImportGroomCache;
+		bNeedUpdate |= GroomAsset != GroomCacheImportOptions->ImportSettings.GroomAsset;
+	}
 
-		if (GroupsPreview->Groups.Num() == 0)
-		{
-			return false;
-		}
+	if (bNeedUpdate)
+	{
+		UpdateStatus(GroupsPreview);
 	}
 
-	// For the GroomCache import, if we don't import the groom asset, a compatible replacement must be specified
-	// If neither the asset nor the cache is imported, then there's nothing to import
-	if (GroomCacheImportOptions && !GroomCacheImportOptions->ImportSettings.bImportGroomAsset &&
-		(!GroomCacheImportOptions->ImportSettings.GroomAsset.IsValid() || !GroomCacheImportOptions->ImportSettings.bImportGroomCache))
+	switch (CurrentStatus)
 	{
-		// TODO: Test for compatibility between cache and selected groom asset
-		return false;
+		case EHairDescriptionStatus::Valid:
+		case EHairDescriptionStatus::GroomCacheCompatible:
+		case EHairDescriptionStatus::GuidesOnlyCompatible:
+			return true;
+		case EHairDescriptionStatus::Unset:
+		case EHairDescriptionStatus::NoGroup:
+		case EHairDescriptionStatus::NoCurve:
+		case EHairDescriptionStatus::GroomCache:
+		case EHairDescriptionStatus::GroomCacheIncompatible:
+		case EHairDescriptionStatus::GuidesOnly:
+		case EHairDescriptionStatus::GuidesOnlyIncompatible:
+		case EHairDescriptionStatus::Unknown:
+		default:
+			return false;
 	}
-	return true;
 }
 
 enum class EGroomOptionsVisibility : uint8
diff --git a/Engine/Plugins/Runtime/HairStrands/Source/HairStrandsEditor/Private/GroomImportOptionsWindow.h b/Engine/Plugins/Runtime/HairStrands/Source/HairStrandsEditor/Private/GroomImportOptionsWindow.h
--- a/Engine/Plugins/Runtime/HairStrands/Source/HairStrandsEditor/Private/GroomImportOptionsWindow.h
+++ b/Engine/Plugins/Runtime/HairStrands/Source/HairStrandsEditor/Private/GroomImportOptionsWindow.h
@@ -15,6 +15,8 @@ class UGroomCacheImportOptions;
 class UGroomHairGroupsPreview;
 struct FHairGroupInfo;
 
+enum class EHairDescriptionStatus;
+
 class SGroomImportOptionsWindow : public SCompoundWidget
 {
 public:
@@ -95,6 +97,15 @@ public:
 private:
 
 	bool CanImport() const;
+	void UpdateStatus(UGroomHairGroupsPreview* Description) const;
+	FText GetStatusText() const;
+	FSlateColor GetStatusColor() const;
+
+	// Properties that are monitored for changes
+	mutable EHairDescriptionStatus CurrentStatus;
+	mutable FSoftObjectPath GroomAsset;
+	mutable bool bImportGroomAssetState;
+	mutable bool bImportGroomCacheState;
 
 private:
 	UGroomImportOptions* ImportOptions;
diff --git a/Engine/Plugins/Runtime/HairStrands/Source/HairStrandsEditor/Private/HairStrandsFactory.cpp b/Engine/Plugins/Runtime/HairStrands/Source/HairStrandsEditor/Private/HairStrandsFactory.cpp
--- a/Engine/Plugins/Runtime/HairStrands/Source/HairStrandsEditor/Private/HairStrandsFactory.cpp
+++ b/Engine/Plugins/Runtime/HairStrands/Source/HairStrandsEditor/Private/HairStrandsFactory.cpp
@@ -72,6 +72,8 @@ void UHairStrandsFactory::GetSupportedFileExtensions(TArray<FString>& OutExtensi
 UObject* UHairStrandsFactory::FactoryCreateFile(UClass* InClass, UObject* InParent, FName InName, EObjectFlags Flags, 
 	const FString& Filename, const TCHAR* Parms, FFeedbackContext* Warn, bool& bOutOperationCanceled) 
 {
+	TRACE_CPUPROFILER_EVENT_SCOPE(UHairStrandsFactory::FactoryCreateFile);
+
 	bOutOperationCanceled = false;
 
 	const bool bIsUnattended = (IsAutomatedImport()
@@ -104,6 +106,7 @@ UObject* UHairStrandsFactory::FactoryCreateFile(UClass* InClass, UObject* InPare
 		}
 	}
 
+	bool bGuidesOnly = false;
 	FGroomAnimationInfo AnimInfo;
 	{
 		// Load the alembic file upfront to preview & report any potential issue
@@ -151,6 +154,7 @@ UObject* UHairStrandsFactory::FactoryCreateFile(UClass* InClass, UObject* InPare
 				OutGroup.bHasPrecomputedWeights = bHasPrecomputedWeights;
 				OutGroup.bHasColorAttributes = bHasColorAttributes;
 				OutGroup.bHasRoughnessAttributes = bHasRoughnessAttributes;
+				bGuidesOnly |= (OutGroup.CurveCount == 0 && OutGroup.GuideCount > 0);
 
 				if (OutGroup.GroupID < OutDescription.HairGroups.Num())
 				{				
@@ -161,6 +165,15 @@ UObject* UHairStrandsFactory::FactoryCreateFile(UClass* InClass, UObject* InPare
 
 		FGroomCacheImporter::SetupImportSettings(GroomCacheImportOptions->ImportSettings, AnimInfo);
 
+		if (bGuidesOnly)
+		{
+			// Guides-only groom cache cannot import groom asset and must specify an existing one
+			GroomCacheImportOptions->ImportSettings.bImportGroomAsset = false;
+
+			// Hint since All will import guides too
+			GroomCacheImportOptions->ImportSettings.ImportType = EGroomCacheImportType::Guides;
+		}
+
 		// Don't bother saving the options coming from script
 		if (!bIsUnattended)
 		{
@@ -198,7 +211,7 @@ UObject* UHairStrandsFactory::FactoryCreateFile(UClass* InClass, UObject* InPare
 	UObject* CurrentAsset = nullptr;
 	UGroomAsset* GroomAssetForCache = nullptr;
 	FHairImportContext HairImportContext(ImportOptions, InParent, InClass, InName, Flags);
-	if (GroomCacheImportOptions->ImportSettings.bImportGroomAsset)
+	if (GroomCacheImportOptions->ImportSettings.bImportGroomAsset && !bGuidesOnly)
 	{
 		// Might try to import the same file in the same folder, so if an asset already exists there, reuse and update it
 		// Since we are importing (not reimporting) we reset the object completely. All previous settings will be lost.
@@ -232,7 +245,11 @@ UObject* UHairStrandsFactory::FactoryCreateFile(UClass* InClass, UObject* InPare
 
 	if (GroomCacheImportOptions->ImportSettings.bImportGroomCache && GroomAssetForCache)
 	{
-		TArray<UGroomCache*> GroomCaches = FGroomCacheImporter::ImportGroomCache(Filename, SelectedTranslator, AnimInfo, HairImportContext, GroomAssetForCache);
+		if (GroomCacheImportOptions->ImportSettings.bOverrideConversionSettings)
+		{
+			HairImportContext.ImportOptions->ConversionSettings = GroomCacheImportOptions->ImportSettings.ConversionSettings;
+		}
+		TArray<UGroomCache*> GroomCaches = FGroomCacheImporter::ImportGroomCache(Filename, SelectedTranslator, AnimInfo, HairImportContext, GroomAssetForCache, GroomCacheImportOptions->ImportSettings.ImportType);
 
 		// Setup asset import data
 		for (UGroomCache* GroomCache : GroomCaches)
diff --git a/Engine/Plugins/Runtime/HairStrands/Source/HairStrandsEditor/Private/ReimportHairStrandsFactory.cpp b/Engine/Plugins/Runtime/HairStrands/Source/HairStrandsEditor/Private/ReimportHairStrandsFactory.cpp
--- a/Engine/Plugins/Runtime/HairStrands/Source/HairStrandsEditor/Private/ReimportHairStrandsFactory.cpp
+++ b/Engine/Plugins/Runtime/HairStrands/Source/HairStrandsEditor/Private/ReimportHairStrandsFactory.cpp
@@ -228,7 +228,11 @@ namespace ReimportGroomCacheHelpers
 			{
 				HairImportContext.Parent = SourceGroomCache;
 			}
-			TArray<UGroomCache*> GroomCaches = FGroomCacheImporter::ImportGroomCache(SourceFilename, Translator, AnimInfo, HairImportContext, GroomAsset);
+			if (Settings.bOverrideConversionSettings)
+			{
+				HairImportContext.ImportOptions->ConversionSettings = Settings.ConversionSettings;
+			}
+			TArray<UGroomCache*> GroomCaches = FGroomCacheImporter::ImportGroomCache(SourceFilename, Translator, AnimInfo, HairImportContext, GroomAsset, Settings.ImportType);
 
 			// Update asset import data
 			for (UGroomCache* GroomCache : GroomCaches)
diff --git a/Engine/Plugins/Runtime/HairStrands/Source/HairStrandsEditor/Public/GroomCacheImporter.h b/Engine/Plugins/Runtime/HairStrands/Source/HairStrandsEditor/Public/GroomCacheImporter.h
--- a/Engine/Plugins/Runtime/HairStrands/Source/HairStrandsEditor/Public/GroomCacheImporter.h
+++ b/Engine/Plugins/Runtime/HairStrands/Source/HairStrandsEditor/Public/GroomCacheImporter.h
@@ -11,9 +11,11 @@ class IGroomTranslator;
 class UGroomAsset;
 class UGroomCache;
 
+enum class EGroomCacheImportType : uint8;
+
 struct HAIRSTRANDSEDITOR_API FGroomCacheImporter
 {
-	static TArray<UGroomCache*> ImportGroomCache(const FString& SourceFilename, TSharedPtr<IGroomTranslator> Translator, const FGroomAnimationInfo& InAnimInfo, FHairImportContext& HairImportContext, UGroomAsset* GroomAssetForCache);
+	static TArray<UGroomCache*> ImportGroomCache(const FString& SourceFilename, TSharedPtr<IGroomTranslator> Translator, const FGroomAnimationInfo& InAnimInfo, FHairImportContext& HairImportContext, UGroomAsset* GroomAssetForCache, EGroomCacheImportType ImportType);
 	static void SetupImportSettings(struct FGroomCacheImportSettings& ImportSettings, const FGroomAnimationInfo& AnimInfo);
 	static void ApplyImportSettings(struct FGroomCacheImportSettings& ImportSettings, FGroomAnimationInfo& AnimInfo);
 	static UGroomCache* ProcessToGroomCache(FGroomCacheProcessor& Processor, const FGroomAnimationInfo& AnimInfo, FHairImportContext& ImportContext, const FString& ObjectNameSuffix);
