Change Groom binding to use DDC2 API.

diff --git a/Engine/Plugins/Runtime/HairStrands/Source/HairStrandsCore/Private/GroomBindingAsset.cpp b/Engine/Plugins/Runtime/HairStrands/Source/HairStrandsCore/Private/GroomBindingAsset.cpp
--- a/Engine/Plugins/Runtime/HairStrands/Source/HairStrandsCore/Private/GroomBindingAsset.cpp
+++ b/Engine/Plugins/Runtime/HairStrands/Source/HairStrandsCore/Private/GroomBindingAsset.cpp
@@ -16,6 +16,8 @@
 #include UE_INLINE_GENERATED_CPP_BY_NAME(GroomBindingAsset)
 
 #if WITH_EDITORONLY_DATA
+#include "DerivedDataCache.h"
+#include "DerivedDataRequestOwner.h"
 #include "DerivedDataCacheInterface.h"
 #include "Serialization/LargeMemoryReader.h"
 #include "Serialization/LargeMemoryWriter.h"
@@ -727,56 +729,46 @@ void UGroomBindingAsset::CacheDerivedDatas()
 
 	if (DerivedDataKey != CachedDerivedDataKey)
 	{
-		FGroomComponentRecreateRenderStateContext RecreateRenderContext(Groom);
+		using namespace UE::DerivedData;
 
-		TArray<uint8> DerivedData;
-		if (GetDerivedDataCacheRef().GetSynchronous(*DerivedDataKey, DerivedData, GetPathName()))
+		const FCacheKey Key = ConvertLegacyCacheKey(DerivedDataKey);
+		const FSharedString Name = MakeStringView(GetPathName());
+		FSharedBuffer Data;
 		{
-			if (IsHairStrandsDDCLogEnable())
+			FRequestOwner Owner(EPriority::Blocking);
+			GetCache().GetValue({ {Name, Key} }, Owner, [&Data](FCacheGetValueResponse&& Response)
 			{
-				UE_LOG(LogHairStrands, Log, TEXT("[GroomBinding/DDC] Found (GroomiBinding:%s)."), *GetName());
-			}
-
-			FMemoryReader Ar(DerivedData, /*bIsPersistent=*/ true);
+				Data = Response.Value.GetData().Decompress();
+			});
+			Owner.Wait();
+		}
 
-			int64 UncompressedSize = 0;
-			Ar << UncompressedSize;
+		FGroomComponentRecreateRenderStateContext RecreateRenderContext(Groom);
 
-			uint8* DecompressionBuffer = reinterpret_cast<uint8*>(FMemory::Malloc(UncompressedSize));
-			Ar.SerializeCompressed(DecompressionBuffer, UncompressedSize, NAME_Zlib);
+		if (Data)
+		{
+			UE_CLOG(IsHairStrandsDDCLogEnable(), LogHairStrands, Log, TEXT("[GroomBinding/DDC] Found (GroomBinding:%s)."), *GetName());
 
-			FLargeMemoryReader LargeMemReader(DecompressionBuffer, UncompressedSize, ELargeMemoryReaderFlags::Persistent | ELargeMemoryReaderFlags::TakeOwnership);
-			InternalSerialize(LargeMemReader, this, HairGroupBulkDatas, 0);
+			FMemoryReaderView Ar(Data, /*bIsPersistent*/ true);
+			InternalSerialize(Ar, this, HairGroupBulkDatas, 0);
 
 			bIsValid = true;
 		}
 		else
 		{
-			if (IsHairStrandsDDCLogEnable())
-			{
-				UE_LOG(LogHairStrands, Log, TEXT("[GroomBinding/DDC] Not found (GroomiBinding:%s)."), *GetName());
-			}
+			UE_CLOG(IsHairStrandsDDCLogEnable(), LogHairStrands, Log, TEXT("[GroomBinding/DDC] Not found (GroomBinding:%s)."), *GetName());
 
 			// Build groom binding data
 			bIsValid = FGroomBindingBuilder::BuildBinding(this, false);
 			if (bIsValid)
 			{
-				// Using a LargeMemoryWriter for serialization since the data can be bigger than 2 GB
-				FLargeMemoryWriter LargeMemWriter(0, /*bIsPersistent=*/ true);
-				InternalSerialize(LargeMemWriter, this, HairGroupBulkDatas, 0);
+				TArray<uint8> WriteData;
+				FMemoryWriter Ar(WriteData, /*bIsPersistent*/ true);
+				InternalSerialize(Ar, this, HairGroupBulkDatas, 0);
 
-				int64 UncompressedSize = LargeMemWriter.TotalSize();
-
-				// Then the content of the LargeMemWriter is compressed into a MemoryWriter
-				// Compression ratio can reach about 5:2 depending on the data
-				{
-					FMemoryWriter CompressedArchive(DerivedData, true);
-
-					CompressedArchive << UncompressedSize; // needed for allocating decompression buffer
-					CompressedArchive.SerializeCompressed(LargeMemWriter.GetData(), UncompressedSize, NAME_Zlib);
-
-					GetDerivedDataCacheRef().Put(*DerivedDataKey, DerivedData, GetPathName());
-				}
+				FRequestOwner AsyncOwner(EPriority::Normal);
+				GetCache().PutValue({ {Name, Key, FValue::Compress(MakeSharedBufferFromArray(MoveTemp(WriteData)))} }, AsyncOwner);
+				AsyncOwner.KeepAlive();
 			}
 		}
 
@@ -789,12 +781,12 @@ void UGroomBindingAsset::CacheDerivedDatas()
 		for (uint32 GroupIt=0; GroupIt< GroupCount; ++GroupIt)
 		{
 			FGoomBindingGroupInfo& Info = GroupInfos[GroupIt];
-			const FHairGroupBulkData& Data = HairGroupBulkDatas[GroupIt];
+			const FHairGroupBulkData& BulkData = HairGroupBulkDatas[GroupIt];
 			{
-				Info.SimRootCount = Data.SimRootBulkData.RootCount;
-				Info.SimLODCount  = Data.SimRootBulkData.MeshProjectionLODs.Num();
-				Info.RenRootCount = Data.RenRootBulkData.RootCount;
-				Info.RenLODCount  = Data.RenRootBulkData.MeshProjectionLODs.Num();
+				Info.SimRootCount = BulkData.SimRootBulkData.RootCount;
+				Info.SimLODCount  = BulkData.SimRootBulkData.MeshProjectionLODs.Num();
+				Info.RenRootCount = BulkData.RenRootBulkData.RootCount;
+				Info.RenLODCount  = BulkData.RenRootBulkData.MeshProjectionLODs.Num();
 			}
 		}
 
diff --git a/Engine/Source/Runtime/Core/Private/UObject/UE5MainStreamObjectVersion.cpp b/Engine/Source/Runtime/Core/Private/UObject/UE5MainStreamObjectVersion.cpp
--- a/Engine/Source/Runtime/Core/Private/UObject/UE5MainStreamObjectVersion.cpp
+++ b/Engine/Source/Runtime/Core/Private/UObject/UE5MainStreamObjectVersion.cpp
@@ -18,6 +18,7 @@ TMap<FGuid, FGuid> FUE5MainStreamObjectVersion::GetSystemGuids()
 	SystemGuids.Add(DevGuids.SkeletalMeshDerivedDataVersion, FGuid("4084BE94509A8C3FA7CF410DDF8ABF04"));
 	SystemGuids.Add(DevGuids.STATICMESH_DERIVEDDATA_VER, FGuid("98F7D79A5811013825E75BBCC41ED3E9"));
 	SystemGuids.Add(DevGuids.POSESEARCHDB_DERIVEDDATA_VER, FGuid("4E595C2AC5E947D6BA9ABC874353E5BC"));
+	SystemGuids.Add(DevGuids.GROOM_BINDING_DERIVED_DATA_VERSION, FGuid("156678D4F9084D7CAE98FADC6AB93573"));
 
 	return SystemGuids;
 }
