First set of changes for fixing groom pipeline being non-deterministic.

Pipeline key hasn't been bumped yet, because of remaing issues.

diff --git a/Engine/Plugins/Runtime/HairStrands/Source/HairStrandsCore/Private/GroomBuilder.cpp b/Engine/Plugins/Runtime/HairStrands/Source/HairStrandsCore/Private/GroomBuilder.cpp
--- a/Engine/Plugins/Runtime/HairStrands/Source/HairStrandsCore/Private/GroomBuilder.cpp
+++ b/Engine/Plugins/Runtime/HairStrands/Source/HairStrandsCore/Private/GroomBuilder.cpp
@@ -157,6 +157,7 @@ namespace HairStrandsBuilder
 
 		HairStrands.BoundingBox.Min = {  FLT_MAX,  FLT_MAX ,  FLT_MAX };
 		HairStrands.BoundingBox.Max = { -FLT_MAX, -FLT_MAX , -FLT_MAX };
+		HairStrands.BoundingBox.IsValid = 0;
 
 		if (HairStrands.GetNumCurves() > 0 && HairStrands.GetNumPoints() > 0)
 		{
@@ -186,7 +187,7 @@ namespace HairStrandsBuilder
 				for (uint32 PointIndex = 0; PointIndex < StrandCount; ++PointIndex, ++PositionIterator, ++RadiusIterator, ++CoordUIterator)
 				{
 					HairStrands.BoundingBox += (FVector)*PositionIterator;
-
+					HairStrands.BoundingBox.IsValid = 1;
 					if (PointIndex > 0)
 					{
 						StrandLength += ((FVector)*PositionIterator - (FVector)PreviousPosition).Size();
@@ -1414,48 +1415,6 @@ namespace HairInterpolationBuilder
 	}
 }
 
-class FGroomDataRandomizer
-{
-public:
-	FGroomDataRandomizer(int Seed, int NumRenderCurves, int NumSimCurves)
-	{
-		Random.Initialize(Seed);
-
-		RenderCurveSeeds.SetNumUninitialized(NumRenderCurves);
-		for (int Index = 0; Index < NumRenderCurves; ++Index)
-		{
-			RenderCurveSeeds[Index] = Random.RandHelper(255);
-		}
-
-		SimCurveSeeds.SetNumUninitialized(NumSimCurves);
-		for (int Index = 0; Index < NumSimCurves; ++Index)
-		{
-			SimCurveSeeds[Index] = Random.RandHelper(255);
-		}
-
-		// This randomization makes certain strands being affected by 1, 2, or 3 guides
-		GuideIndices.SetNumUninitialized(NumRenderCurves);
-		for (int Index = 0; Index < NumRenderCurves; ++Index)
-		{
-			FIntVector& RandomIndices = GuideIndices[Index];
-			for (int GuideIndex = 0; GuideIndex < HairInterpolationBuilder::FMetrics::Count; ++GuideIndex)
-			{
-				RandomIndices[GuideIndex] = Random.RandRange(0, HairInterpolationBuilder::FMetrics::Count - 1);
-			}
-		}
-	}
-
-	const TArray<uint8>& GetRenderCurveSeeds() const { return RenderCurveSeeds; }
-	const TArray<uint8>& GetSimCurveSeeds() const { return SimCurveSeeds; }
-	const TArray<FIntVector>& GetRandomGuideIndices() const { return GuideIndices; }
-
-private:
-	FRandomStream Random;
-	TArray<uint8> RenderCurveSeeds;
-	TArray<uint8> SimCurveSeeds;
-	TArray<FIntVector> GuideIndices;
-};
-
 FORCEINLINE FIntVector VectorToIntVector(const FVector& Index)
 {
 	return FIntVector(Index.X, Index.Y, Index.Z);
@@ -2106,6 +2065,7 @@ void Decimate(
 
 	uint32 OutTotalPointCount = 0;
 	FRandomStream Random;
+	Random.Initialize(0xdeedbeed);
 
 	TArray<uint32> CurveRandomizedIndex;
 
@@ -2114,8 +2074,6 @@ void Decimate(
 	{
 		CurveRandomizedIndex.SetNum(InCurveCount);
 
-		Random.Initialize(0xdeedbeed);
-
 		//initialize value
 		for (uint32 CurveIndex = 0; CurveIndex < InCurveCount; CurveIndex++)
 		{
