Fix non-determinism in Groom builder.

This CL still does not bump the groom pipeline version until all deterministic issue are fixes.

diff --git a/Engine/Plugins/Runtime/HairStrands/Source/HairStrandsCore/Private/GroomBuilder.cpp b/Engine/Plugins/Runtime/HairStrands/Source/HairStrandsCore/Private/GroomBuilder.cpp
--- a/Engine/Plugins/Runtime/HairStrands/Source/HairStrandsCore/Private/GroomBuilder.cpp
+++ b/Engine/Plugins/Runtime/HairStrands/Source/HairStrandsCore/Private/GroomBuilder.cpp
@@ -2795,7 +2795,7 @@ static void BuildClusterData(
 	TArray<FHairClusterInfo>& LocalClusterInfos = Out.ClusterInfos;
 	TArray<FHairClusterLODInfo>& LocalClusterLODInfos = Out.ClusterLODInfos;
 	TArray<uint32>& LocalVertexToClusterIds = Out.VertexToClusterIds;
-#define USE_PARALLE_FOR 1
+#define USE_PARALLE_FOR 0
 #if USE_PARALLE_FOR
 	ParallelFor(Out.ClusterCount,
 		[
diff --git a/Engine/Plugins/Runtime/HairStrands/Source/HairStrandsCore/Private/HairCardsBuilder.cpp b/Engine/Plugins/Runtime/HairStrands/Source/HairStrandsCore/Private/HairCardsBuilder.cpp
--- a/Engine/Plugins/Runtime/HairStrands/Source/HairStrandsCore/Private/HairCardsBuilder.cpp
+++ b/Engine/Plugins/Runtime/HairStrands/Source/HairStrandsCore/Private/HairCardsBuilder.cpp
@@ -1347,6 +1347,7 @@ namespace HairCards
 
 		uint32 OutTotalPointCount = 0;
 		FRandomStream Random;
+		Random.Initialize(0xdeedbeed);
 		for (uint32 BucketIndex = 0; BucketIndex < OutCurveCount; BucketIndex++)
 		{
 			const uint32 CurveIndex = BucketIndex * BucketSize;// +BucketSize * Random.FRand();
@@ -4403,6 +4404,9 @@ void ImportGeometry(
 	const uint32 MeshLODIndex = 0;
 
 	// Note: if there are multiple section we only import the first one. Support for multiple section could be added later on. 
+	const FStaticMeshRenderData* MeshRenderData = StaticMesh->GetRenderData();
+	check(MeshRenderData != nullptr);
+	check(MeshRenderData->CurrentFirstLODIdx == MeshLODIndex);
 	const FStaticMeshLODResources& LODData = StaticMesh->GetLODForExport(MeshLODIndex);
 	const uint32 VertexCount = LODData.VertexBuffers.PositionVertexBuffer.GetNumVertices();
 	const uint32 IndexCount = LODData.IndexBuffer.GetNumIndices();
diff --git a/Engine/Plugins/Runtime/HairStrands/Source/HairStrandsCore/Public/HairCardsDatas.h b/Engine/Plugins/Runtime/HairStrands/Source/HairStrandsCore/Public/HairCardsDatas.h
--- a/Engine/Plugins/Runtime/HairStrands/Source/HairStrandsCore/Public/HairCardsDatas.h
+++ b/Engine/Plugins/Runtime/HairStrands/Source/HairStrandsCore/Public/HairCardsDatas.h
@@ -196,7 +196,7 @@ struct FHairCardsGeometry
 	TArray<uint32> IndexOffsets;
 	TArray<uint32> IndexCounts;
 
-	FBox3f BoundingBox;
+	FBox3f BoundingBox = FBox3f(EForceInit::ForceInit);
 
 	void Reset()
 	{
@@ -244,7 +244,6 @@ struct FHairCardsGeometry
 	}
 };
 
-struct FHairCardsBulkData;
 struct FHairCardsDatas
 {
 	bool IsValid() const { return Cards.Positions.Num() > 0; }
@@ -283,7 +282,7 @@ struct FHairCardsBulkData
 	TArray<FHairCardsUVFormat::Type> UVs;
 	TArray<FHairCardsMaterialFormat::Type> Materials;
 	TArray<FHairCardsIndexFormat::Type> Indices;
-	FBox BoundingBox;
+	FBox BoundingBox = FBox(EForceInit::ForceInit);
 
 	void Serialize(FArchive& Ar);
 };
@@ -418,7 +417,7 @@ struct FHairMeshes
 	TArray<FVector3f> Positions;
 	TArray<uint32>    Indices;
 
-	FBox3f BoundingBox;
+	FBox3f BoundingBox = FBox3f(EForceInit::ForceInit);
 
 	void SetNum(uint32 Count)
 	{
@@ -457,7 +456,7 @@ struct FHairMeshesBulkData
 	TArray<FHairCardsNormalFormat::Type> Normals;
 	TArray<FHairCardsUVFormat::Type> UVs;
 	TArray<FHairCardsIndexFormat::Type> Indices;
-	FBox BoundingBox;
+	FBox BoundingBox = FBox(EForceInit::ForceInit);
 
 	void Serialize(FArchive& Ar);
 };
diff --git a/Engine/Plugins/Runtime/HairStrands/Source/HairStrandsCore/Public/HairStrandsDatas.h b/Engine/Plugins/Runtime/HairStrands/Source/HairStrandsCore/Public/HairStrandsDatas.h
--- a/Engine/Plugins/Runtime/HairStrands/Source/HairStrandsCore/Public/HairStrandsDatas.h
+++ b/Engine/Plugins/Runtime/HairStrands/Source/HairStrandsCore/Public/HairStrandsDatas.h
@@ -454,7 +454,7 @@ struct HAIRSTRANDSCORE_API FHairStrandsDatas
 	float HairDensity = 1;
 
 	/* Strands bounding box */
-	FBox BoundingBox;
+	FBox BoundingBox = FBox(EForceInit::ForceInit);
 };
 
 struct HAIRSTRANDSCORE_API FHairStrandsBulkData
@@ -481,7 +481,7 @@ struct HAIRSTRANDSCORE_API FHairStrandsBulkData
 	uint32 PointCount = 0;
 	float MaxLength = 0;
 	float MaxRadius = 0;
-	FBox BoundingBox;
+	FBox BoundingBox = FBox(EForceInit::ForceInit);
 	uint32 Flags = 0;
 
 	FByteBulkData Positions;	// Size = PointCount
