Fix loading all binding asset when opening a groom asset.

diff --git a/Engine/Plugins/Runtime/HairStrands/Source/HairStrandsEditor/Private/GroomCustomAssetEditorToolkit.cpp b/Engine/Plugins/Runtime/HairStrands/Source/HairStrandsEditor/Private/GroomCustomAssetEditorToolkit.cpp
--- a/Engine/Plugins/Runtime/HairStrands/Source/HairStrandsEditor/Private/GroomCustomAssetEditorToolkit.cpp
+++ b/Engine/Plugins/Runtime/HairStrands/Source/HairStrandsEditor/Private/GroomCustomAssetEditorToolkit.cpp
@@ -398,8 +398,6 @@ void FGroomCustomAssetEditorToolkit::InitCustomAssetEditor(const EToolkitMode::T
 	ViewportTab = SNew(SGroomEditorViewport);
 	ThumbnailPool = MakeShared<FAssetThumbnailPool>(64);
 
-	GroomBindingAssetList = NewObject<UGroomBindingAssetList>(GetTransientPackage(), NAME_None, RF_Transient);
-	ListAllBindingAssets(InCustomAsset, GroomBindingAssetList);
 	// Automatically affect the first skelal mesh compatible with the groom asset
 	#if 0
 	if (GroomBindingAssetList->Bindings.Num() > 0)
@@ -841,12 +839,25 @@ TSharedRef<SDockTab> FGroomCustomAssetEditorToolkit::SpawnTab_BindingProperties(
 {
 	check(Args.GetTabId() == TabId_BindingProperties);
 
-	return SNew(SDockTab)
+	TSharedRef<SDockTab> DockTab = SNew(SDockTab)
 		.Label(LOCTEXT("BindingPropertiesTab", "Binding"))
 		.TabColorScale(GetTabColorScale())
 		[
 			DetailView_BindingProperties.ToSharedRef()
 		];
+
+	DockTab->SetOnTabActivated(SDockTab::FOnTabActivatedCallback::CreateLambda([this](TSharedRef<SDockTab> Input, ETabActivationCause)
+	{
+		UGroomAsset* LocalGroomAsset = GetCustomAsset();
+		if (GroomBindingAssetList == nullptr && LocalGroomAsset)
+		{
+			GroomBindingAssetList = NewObject<UGroomBindingAssetList>(GetTransientPackage(), NAME_None, RF_Transient);
+			ListAllBindingAssets(LocalGroomAsset, GroomBindingAssetList);
+			DetailView_BindingProperties->ForceRefresh();
+		}
+	}));
+
+	return DockTab;
 }
 
 UGroomComponent *FGroomCustomAssetEditorToolkit::GetPreview_GroomComponent() const
