Expose strands curve buffers to shaders.

diff --git a/Engine/Plugins/Runtime/HairStrands/Source/HairStrandsCore/Private/HairStrandsRendering.cpp b/Engine/Plugins/Runtime/HairStrands/Source/HairStrandsCore/Private/HairStrandsRendering.cpp
--- a/Engine/Plugins/Runtime/HairStrands/Source/HairStrandsCore/Private/HairStrandsRendering.cpp
+++ b/Engine/Plugins/Runtime/HairStrands/Source/HairStrandsCore/Private/HairStrandsRendering.cpp
@@ -1647,12 +1647,14 @@ static FHairGroupPublicData::FVertexFactoryInput InternalComputeHairStrandsVerte
 		CONVERT_HAIRSSTRANDS_VF_PARAMETERS(OutVFInput.Strands.Attribute0Buffer,			Instance->Guides.RestResource->Attribute0Buffer);
 		CONVERT_HAIRSSTRANDS_VF_PARAMETERS(OutVFInput.Strands.Attribute1Buffer,			Instance->Guides.RestResource->Attribute1Buffer);
 		CONVERT_HAIRSSTRANDS_VF_PARAMETERS(OutVFInput.Strands.MaterialBuffer,			Instance->Guides.RestResource->MaterialBuffer);
+		CONVERT_HAIRSSTRANDS_VF_PARAMETERS(OutVFInput.Strands.CurveBuffer,				Instance->Guides.RestResource->CurveBuffer);
 		CONVERT_HAIRSSTRANDS_VF_PARAMETERS(OutVFInput.Strands.PositionOffsetBuffer,		Instance->Guides.DeformedResource->GetPositionOffsetBuffer(FHairStrandsDeformedResource::EFrameType::Current));
 		CONVERT_HAIRSSTRANDS_VF_PARAMETERS(OutVFInput.Strands.PrevPositionOffsetBuffer,	Instance->Guides.DeformedResource->GetPositionOffsetBuffer(FHairStrandsDeformedResource::EFrameType::Previous));
 
 		OutVFInput.Strands.PositionOffset = Instance->Guides.DeformedResource->GetPositionOffset(FHairStrandsDeformedResource::EFrameType::Current);
 		OutVFInput.Strands.PrevPositionOffset = Instance->Guides.DeformedResource->GetPositionOffset(FHairStrandsDeformedResource::EFrameType::Previous);
-		OutVFInput.Strands.VertexCount = Instance->Guides.RestResource->GetVertexCount();		
+		OutVFInput.Strands.VertexCount = Instance->Guides.RestResource->GetVertexCount();
+		OutVFInput.Strands.CurveCount = Instance->Guides.RestResource->GetCurveCount();
 	}
 	// 2. Render Strands with deformation
 	else if (bSupportDeformation)
@@ -1684,8 +1686,10 @@ static FHairGroupPublicData::FVertexFactoryInput InternalComputeHairStrandsVerte
 		CONVERT_HAIRSSTRANDS_VF_PARAMETERS(OutVFInput.Strands.Attribute0Buffer,			(ViewMode == EGroomViewMode::RenderHairStrands || ViewMode == EGroomViewMode::RenderVisCluster || ViewMode == EGroomViewMode::RenderVisClusterAABB) ? Instance->Strands.DebugAttributeBuffer : Instance->Strands.RestResource->Attribute0Buffer);
 		CONVERT_HAIRSSTRANDS_VF_PARAMETERS(OutVFInput.Strands.Attribute1Buffer,			Instance->Strands.RestResource->Attribute1Buffer);
 		CONVERT_HAIRSSTRANDS_VF_PARAMETERS(OutVFInput.Strands.MaterialBuffer,			Instance->Strands.RestResource->MaterialBuffer);
+		CONVERT_HAIRSSTRANDS_VF_PARAMETERS(OutVFInput.Strands.CurveBuffer,				Instance->Strands.RestResource->CurveBuffer);
 
 		OutVFInput.Strands.VertexCount = Instance->Strands.RestResource->GetVertexCount();
+		OutVFInput.Strands.CurveCount = Instance->Strands.RestResource->GetCurveCount();
 	}
 	// 3. Render Strands in rest position: used when there are no skinning (rigid binding or no binding), no simulation, and no RBF
 	else
@@ -1696,12 +1700,14 @@ static FHairGroupPublicData::FVertexFactoryInput InternalComputeHairStrandsVerte
 		CONVERT_HAIRSSTRANDS_VF_PARAMETERS(OutVFInput.Strands.Attribute0Buffer,			(ViewMode == EGroomViewMode::RenderVisCluster || ViewMode == EGroomViewMode::RenderVisClusterAABB) ? Instance->Strands.DebugAttributeBuffer : Instance->Strands.RestResource->Attribute0Buffer);
 		CONVERT_HAIRSSTRANDS_VF_PARAMETERS(OutVFInput.Strands.Attribute1Buffer,			Instance->Strands.RestResource->Attribute1Buffer);
 		CONVERT_HAIRSSTRANDS_VF_PARAMETERS(OutVFInput.Strands.MaterialBuffer,			Instance->Strands.RestResource->MaterialBuffer);
+		CONVERT_HAIRSSTRANDS_VF_PARAMETERS(OutVFInput.Strands.CurveBuffer,				Instance->Strands.RestResource->CurveBuffer);
 		CONVERT_HAIRSSTRANDS_VF_PARAMETERS(OutVFInput.Strands.PositionOffsetBuffer,		Instance->Strands.RestResource->PositionOffsetBuffer);
 		CONVERT_HAIRSSTRANDS_VF_PARAMETERS(OutVFInput.Strands.PrevPositionOffsetBuffer, Instance->Strands.RestResource->PositionOffsetBuffer);
 
 		OutVFInput.Strands.PositionOffset = Instance->Strands.RestResource->GetPositionOffset();
 		OutVFInput.Strands.PrevPositionOffset = Instance->Strands.RestResource->GetPositionOffset();
 		OutVFInput.Strands.VertexCount = Instance->Strands.RestResource->GetVertexCount();
+		OutVFInput.Strands.CurveCount = Instance->Strands.RestResource->GetCurveCount();
 	}
 
 	OutVFInput.Strands.HairRadius = (GStrandHairWidth > 0 ? GStrandHairWidth : Instance->Strands.Modifier.HairWidth) * 0.5f;
diff --git a/Engine/Plugins/Runtime/HairStrands/Source/HairStrandsCore/Public/GroomResources.h b/Engine/Plugins/Runtime/HairStrands/Source/HairStrandsCore/Public/GroomResources.h
--- a/Engine/Plugins/Runtime/HairStrands/Source/HairStrandsCore/Public/GroomResources.h
+++ b/Engine/Plugins/Runtime/HairStrands/Source/HairStrandsCore/Public/GroomResources.h
@@ -356,6 +356,7 @@ struct FHairStrandsRestResource : public FHairCommonResource
 	const EHairStrandsResourcesType CurveType;
 
 	inline uint32 GetVertexCount() const { return BulkData.PointCount; }
+	inline uint32 GetCurveCount() const { return BulkData.CurveCount; }
 };
 
 struct FHairStrandsDeformedResource : public FHairCommonResource
diff --git a/Engine/Shaders/Private/HairStrands/HairStrandsVertexFactoryCommon.ush b/Engine/Shaders/Private/HairStrands/HairStrandsVertexFactoryCommon.ush
--- a/Engine/Shaders/Private/HairStrands/HairStrandsVertexFactoryCommon.ush
+++ b/Engine/Shaders/Private/HairStrands/HairStrandsVertexFactoryCommon.ush
@@ -130,10 +130,12 @@ float			HairStrandsVF_Density;
 uint			HairStrandsVF_bUseStableRasterization;
 float3			HairStrandsVF_PositionOffset;
 uint			HairStrandsVF_VertexCount;
+uint			HairStrandsVF_CurveCount;
 float4x4		HairStrandsVF_LocalToWorldPrimitiveTransform;
 float4x4		HairStrandsVF_LocalToTranslatedWorldPrimitiveTransform;
 Buffer<uint4>	HairStrandsVF_PositionBuffer;
 Buffer<float4>	HairStrandsVF_PositionOffsetBuffer;
+Buffer<uint>	HairStrandsVF_CurveBuffer;
 
 #if defined(PERMUTATION_CULLING) && PERMUTATION_CULLING == 1
 Buffer<uint>	HairStrandsVF_CullingIndirectBuffer;
diff --git a/Engine/Source/Runtime/Renderer/Private/HairStrands/HairStrandsData.h b/Engine/Source/Runtime/Renderer/Private/HairStrands/HairStrandsData.h
--- a/Engine/Source/Runtime/Renderer/Private/HairStrands/HairStrandsData.h
+++ b/Engine/Source/Runtime/Renderer/Private/HairStrands/HairStrandsData.h
@@ -58,11 +58,13 @@ BEGIN_SHADER_PARAMETER_STRUCT(FHairStrandsInstanceParameters, )
 	SHADER_PARAMETER(float, HairStrandsVF_Length)
 	SHADER_PARAMETER(uint32, HairStrandsVF_bUseStableRasterization)
 	SHADER_PARAMETER(uint32, HairStrandsVF_VertexCount)
+	SHADER_PARAMETER(uint32, HairStrandsVF_CurveCount)
 	SHADER_PARAMETER(FVector3f, HairStrandsVF_PositionOffset)
 	SHADER_PARAMETER(FMatrix44f, HairStrandsVF_LocalToWorldPrimitiveTransform)
 	SHADER_PARAMETER(FMatrix44f, HairStrandsVF_LocalToTranslatedWorldPrimitiveTransform)
 	SHADER_PARAMETER_RDG_BUFFER_SRV(Buffer, HairStrandsVF_PositionBuffer)
 	SHADER_PARAMETER_RDG_BUFFER_SRV(Buffer, HairStrandsVF_PositionOffsetBuffer)
+	SHADER_PARAMETER_RDG_BUFFER_SRV(Buffer, HairStrandsVF_CurveBuffer)
 	SHADER_PARAMETER_RDG_BUFFER_SRV(Buffer, HairStrandsVF_CullingIndirectBuffer)
 	SHADER_PARAMETER_RDG_BUFFER_SRV(Buffer, HairStrandsVF_CullingIndexBuffer)
 	SHADER_PARAMETER_RDG_BUFFER_SRV(Buffer, HairStrandsVF_CullingRadiusScaleBuffer)
@@ -139,6 +141,9 @@ struct FHairStrandsVisibilityData
 	uint32			NodeGroupSize = 0;
 	FVector4f		HairOnlyDepthHZBParameters = FVector4f::Zero();
 
+	FRDGBufferSRVRef	ControlPointsSRV = nullptr;
+	FRDGTextureRef		ControlPointCount = nullptr;
+
 	FHairStrandsTiles TileData;
 
 	const static EPixelFormat NodeCoordFormat = PF_R16G16_UINT;
diff --git a/Engine/Source/Runtime/Renderer/Public/HairStrandsInterface.h b/Engine/Source/Runtime/Renderer/Public/HairStrandsInterface.h
--- a/Engine/Source/Runtime/Renderer/Public/HairStrandsInterface.h
+++ b/Engine/Source/Runtime/Renderer/Public/HairStrandsInterface.h
@@ -214,6 +214,7 @@ public:
 			FRDGImportedBuffer Attribute1Buffer;
 			FRDGImportedBuffer PositionOffsetBuffer;
 			FRDGImportedBuffer PrevPositionOffsetBuffer;
+			FRDGImportedBuffer CurveBuffer;
 
 			FRDGExternalBuffer PositionBufferExternal;
 			FRDGExternalBuffer PrevPositionBufferExternal;
@@ -223,6 +224,7 @@ public:
 			FRDGExternalBuffer Attribute1BufferExternal;
 			FRDGExternalBuffer PositionOffsetBufferExternal;
 			FRDGExternalBuffer PrevPositionOffsetBufferExternal;
+			FRDGExternalBuffer CurveBufferExternal;
 
 			FShaderResourceViewRHIRef PositionBufferRHISRV				= nullptr;
 			FShaderResourceViewRHIRef PrevPositionBufferRHISRV			= nullptr;
@@ -232,10 +234,12 @@ public:
 			FShaderResourceViewRHIRef Attribute1BufferRHISRV			= nullptr;
 			FShaderResourceViewRHIRef PositionOffsetBufferRHISRV		= nullptr;
 			FShaderResourceViewRHIRef PrevPositionOffsetBufferRHISRV	= nullptr;
+			FShaderResourceViewRHIRef CurveBufferRHISRV					= nullptr;
 
 			FVector PositionOffset = FVector::ZeroVector;
 			FVector PrevPositionOffset = FVector::ZeroVector;
 
+			uint32 CurveCount = 0;
 			uint32 VertexCount = 0;
 			float HairRadius = 0;
 			float HairRootScale = 0;
