Do not flush rendering commands unless we have something to invalidate

#rnx

diff --git a/Engine/Plugins/Runtime/HairStrands/Source/HairStrandsCore/Private/GroomComponent.cpp b/Engine/Plugins/Runtime/HairStrands/Source/HairStrandsCore/Private/GroomComponent.cpp
--- a/Engine/Plugins/Runtime/HairStrands/Source/HairStrandsCore/Private/GroomComponent.cpp
+++ b/Engine/Plugins/Runtime/HairStrands/Source/HairStrandsCore/Private/GroomComponent.cpp
@@ -4191,8 +4191,11 @@ FGroomComponentRecreateRenderStateContext::FGroomComponentRecreateRenderStateCon
 		}
 	}
 
-	// Flush the rendering commands generated by the detachments.
-	FlushRenderingCommands();
+	if (GroomComponents.Num())
+	{
+		// Flush the rendering commands generated by the detachments.
+		FlushRenderingCommands();
+	}
 }
 
 FGroomComponentRecreateRenderStateContext::~FGroomComponentRecreateRenderStateContext()
