Change hair strands attribute storage.

* All attributes are now stored into a single ByteAddress buffer, which ease adding new/optional attributes. All attributes can now be optional which save memory.
* Attributes which are at per-curve frequency are now stored once per curve instead of once per vertex. This saves memory with an indirection cost.
* Move vertex to curve mapping into the groom resource rather than the groom resources.

GroomAsset size shrink by 5/6% and binding asset size shrink by 45% in simple case.

diff --git a/Engine/Plugins/Runtime/HairStrands/Shaders/Private/DeformerDataInterfaceGroom.ush b/Engine/Plugins/Runtime/HairStrands/Shaders/Private/DeformerDataInterfaceGroom.ush
--- a/Engine/Plugins/Runtime/HairStrands/Shaders/Private/DeformerDataInterfaceGroom.ush
+++ b/Engine/Plugins/Runtime/HairStrands/Shaders/Private/DeformerDataInterfaceGroom.ush
@@ -2,18 +2,22 @@
 
 #include "/Plugin/ComputeFramework/Private/ComputeKernelCommon.ush"
 #include "/Engine/Private/HairStrands/HairStrandsVertexFactoryCommon.ush"
-
+#include "/Engine/Private/HairStrands/HairStrandsAttributeCommon.ush"
 ///////////////////////////////////////////////////////////////////////////////////////////////////
 // Bindings
 
-uint  {DataInterfaceName}_NumControlPoints;
-uint  {DataInterfaceName}_NumCurves;
-float {DataInterfaceName}_VF_HairRadius;
-float {DataInterfaceName}_VF_HairLength;
-Buffer<uint4>  {DataInterfaceName}_PositionBuffer;
-Buffer<float2> {DataInterfaceName}_Attribute0Buffer;
-Buffer<uint>   {DataInterfaceName}_Attribute1Buffer;
-Buffer<uint>   {DataInterfaceName}_CurveBuffer;
+uint				{DataInterfaceName}_NumControlPoints;
+uint				{DataInterfaceName}_NumCurves;
+float				{DataInterfaceName}_VF_HairRadius;
+float				{DataInterfaceName}_VF_HairLength;
+float				{DataInterfaceName}_VF_RootScale;
+float				{DataInterfaceName}_VF_TipScale;
+uint				{DataInterfaceName}_VF_GroupIndex;
+uint4				{DataInterfaceName}_VF_AttributeOffets[HAIR_ATTRIBUTE_OFFSET_COUNT];
+Buffer<uint4>		{DataInterfaceName}_PositionBuffer;
+ByteAddressBuffer	{DataInterfaceName}_AttributeBuffer;
+Buffer<uint>		{DataInterfaceName}_VertexToCurveBuffer;
+Buffer<uint>		{DataInterfaceName}_CurveBuffer;
 
 ///////////////////////////////////////////////////////////////////////////////////////////////////
 // Helpers
@@ -35,20 +39,6 @@ FHairControlPoint ReadControlPointData(uint VertexIndex)
 		1.0f /*TipScale*/);
 }
 
-FHairAttributes ReadControlPointAttributeData(uint VertexIndex)
-{
-	const uint MaxVertices = {DataInterfaceName}_NumControlPoints;
-	if (VertexIndex >= MaxVertices)
-	{
-		return (FHairAttributes)0;
-	}
-
-	return UnpackHairAttribute(
-		{DataInterfaceName}_Attribute0Buffer[VertexIndex],
-		{DataInterfaceName}_Attribute1Buffer[VertexIndex],
-		{DataInterfaceName}_VF_HairLength);
-}
-
 FHairCurve ReadCurveData(uint CurveIndex)
 {
 	const uint MaxCurve = {DataInterfaceName}_NumCurves;
@@ -90,17 +80,19 @@ float ReadCoordU_{DataInterfaceName}(uint VertexIndex)
 
 float ReadLength_{DataInterfaceName}(uint VertexIndex)
 {
-	return ReadControlPointAttributeData(VertexIndex).WorldLength;
+	const float CoordU      = ReadControlPointData(VertexIndex).UCoord;
+	const float CurveLength = UnpackHairLength(InternalGetHairStrandsAttributesPerCurve(VertexIndex, HAIR_ATTRIBUTE_OFFSET_LENGTH({DataInterfaceName}_VF_AttributeOffets), HAIR_ATTRIBUTE_STRIDE_LENGTH, {DataInterfaceName}_AttributeBuffer, {DataInterfaceName}_VertexToCurveBuffer));
+	return CurveLength * CoordU;
 }
 
 float2 ReadRootUV_{DataInterfaceName}(uint VertexIndex)
 {
-	return ReadControlPointAttributeData(VertexIndex).RootUV;
+	return UnpackHairRootUV(InternalGetHairStrandsAttributesPerCurve(VertexIndex, HAIR_ATTRIBUTE_OFFSET_ROOTUV({DataInterfaceName}_VF_AttributeOffets), HAIR_ATTRIBUTE_STRIDE_ROOTUV, {DataInterfaceName}_AttributeBuffer, {DataInterfaceName}_VertexToCurveBuffer));
 }
 
 float ReadSeed_{DataInterfaceName}(uint VertexIndex)
 {
-	return ReadControlPointAttributeData(VertexIndex).Seed;
+	return UnpackHairSeed(InternalGetHairStrandsAttributesPerCurve(VertexIndex, HAIR_ATTRIBUTE_OFFSET_SEED({DataInterfaceName}_VF_AttributeOffets), HAIR_ATTRIBUTE_STRIDE_SEED, {DataInterfaceName}_AttributeBuffer, {DataInterfaceName}_VertexToCurveBuffer));
 }
 
 float ReadCurveOffsetPoint_{DataInterfaceName}(uint CurveIndex)
diff --git a/Engine/Plugins/Runtime/HairStrands/Shaders/Private/DeformerDataInterfaceGroomWrite.ush b/Engine/Plugins/Runtime/HairStrands/Shaders/Private/DeformerDataInterfaceGroomWrite.ush
--- a/Engine/Plugins/Runtime/HairStrands/Shaders/Private/DeformerDataInterfaceGroomWrite.ush
+++ b/Engine/Plugins/Runtime/HairStrands/Shaders/Private/DeformerDataInterfaceGroomWrite.ush
@@ -21,12 +21,6 @@ uint ReadNumCurves_{DataInterfaceName}()
 
 void WritePosition_{DataInterfaceName}(uint ControlPointIndex, float3 Position)
 {
-	FHairControlPoint CP;
-	CP.Position = Position;
-	CP.WorldRadius = 1.f;
-	CP.UCoord = 0;
-	CP.Type = HAIR_CONTROLPOINT_INSIDE;
-
 	const uint WriteIndex = {DataInterfaceName}_OutputStreamStart + ControlPointIndex;
 	if (WriteIndex < {DataInterfaceName}_NumControlPoints)
 	{
diff --git a/Engine/Plugins/Runtime/HairStrands/Source/HairStrandsCore/Private/Deformer/DeformerDataInterfaceGroom.cpp b/Engine/Plugins/Runtime/HairStrands/Source/HairStrandsCore/Private/Deformer/DeformerDataInterfaceGroom.cpp
--- a/Engine/Plugins/Runtime/HairStrands/Source/HairStrandsCore/Private/Deformer/DeformerDataInterfaceGroom.cpp
+++ b/Engine/Plugins/Runtime/HairStrands/Source/HairStrandsCore/Private/Deformer/DeformerDataInterfaceGroom.cpp
@@ -11,6 +11,7 @@
 #include "GroomInstance.h"
 #include "RenderGraphBuilder.h"
 #include "GlobalRenderResources.h"
+#include "HairStrandsDefinitions.h"
 
 FString UOptimusGroomDataInterface::GetDisplayName() const
 {
@@ -94,10 +95,14 @@ BEGIN_SHADER_PARAMETER_STRUCT(FGroomDataInterfaceParameters, )
 	SHADER_PARAMETER(uint32, NumCurves)
 	SHADER_PARAMETER(float, VF_HairRadius)
 	SHADER_PARAMETER(float, VF_HairLength)
-	SHADER_PARAMETER_RDG_BUFFER_SRV(Buffer<uint4>,  PositionBuffer)
-	SHADER_PARAMETER_RDG_BUFFER_SRV(Buffer<float2>, Attribute0Buffer)
-	SHADER_PARAMETER_RDG_BUFFER_SRV(Buffer<uint>,   Attribute1Buffer)
-	SHADER_PARAMETER_RDG_BUFFER_SRV(Buffer<uint>,   CurveBuffer)
+	SHADER_PARAMETER(float, VF_RootScale)
+	SHADER_PARAMETER(float, VF_TipScale)
+	SHADER_PARAMETER(uint32, VF_GroupIndex)
+	SHADER_PARAMETER_ARRAY(FUintVector4, VF_AttributeOffsets, [HAIR_ATTRIBUTE_OFFSET_COUNT])
+	SHADER_PARAMETER_RDG_BUFFER_SRV(Buffer<uint4>,     PositionBuffer)
+	SHADER_PARAMETER_RDG_BUFFER_SRV(ByteAddressBuffer, AttributeBuffer)
+	SHADER_PARAMETER_RDG_BUFFER_SRV(Buffer<uint>,      VertexToCurveBuffer)
+	SHADER_PARAMETER_RDG_BUFFER_SRV(Buffer<uint>,      CurveBuffer)
 END_SHADER_PARAMETER_STRUCT()
 
 void UOptimusGroomDataInterface::GetShaderParameters(TCHAR const* UID, FShaderParametersMetadataBuilder& InOutBuilder, FShaderParametersMetadataAllocations& InOutAllocations) const
@@ -172,11 +177,11 @@ void FOptimusGroomDataProviderProxy::AllocateResources(FRDGBuilder& GraphBuilder
 		if (FHairGroupInstance* Instance = GroomComponent->GetGroupInstance(Index))
 		{
 			FResources& R = Resources.AddDefaulted_GetRef();
-			R.PositionSRV	= Register(GraphBuilder, Instance->Strands.RestResource->PositionBuffer, ERDGImportedBufferFlags::CreateSRV).SRV;
-			R.Attribute0SRV = Register(GraphBuilder, Instance->Strands.RestResource->Attribute0Buffer, ERDGImportedBufferFlags::CreateSRV).SRV;
-			R.Attribute1SRV = Register(GraphBuilder, Instance->Strands.RestResource->Attribute1Buffer, ERDGImportedBufferFlags::CreateSRV).SRV;
-			R.CurveSRV		= Register(GraphBuilder, Instance->Strands.RestResource->CurveBuffer, ERDGImportedBufferFlags::CreateSRV).SRV;
-			R.FallbackSRV	= GraphBuilder.CreateSRV(GraphBuilder.RegisterExternalBuffer(GWhiteVertexBufferWithRDG->Buffer), PF_R16G16B16A16_UINT);
+			R.PositionSRV		= Register(GraphBuilder, Instance->Strands.RestResource->PositionBuffer, ERDGImportedBufferFlags::CreateSRV).SRV;
+			R.AttributeSRV		= Register(GraphBuilder, Instance->Strands.RestResource->AttributeBuffer, ERDGImportedBufferFlags::CreateSRV).SRV;
+			R.VertexToCurveSRV	= Register(GraphBuilder, Instance->Strands.RestResource->VertexToCurveBuffer, ERDGImportedBufferFlags::CreateSRV).SRV;
+			R.CurveSRV			= Register(GraphBuilder, Instance->Strands.RestResource->CurveBuffer, ERDGImportedBufferFlags::CreateSRV).SRV;
+			R.FallbackSRV		= GraphBuilder.CreateSRV(GraphBuilder.RegisterExternalBuffer(GWhiteVertexBufferWithRDG->Buffer), PF_R16G16B16A16_UINT);
 		}
 	}
 }
@@ -204,20 +209,24 @@ void FOptimusGroomDataProviderProxy::GatherDispatchData(FDispatchData const& InD
 			Parameters.NumCurves = NumCurves;
 			Parameters.VF_HairRadius = VFInput.Strands.HairRadius;
 			Parameters.VF_HairLength = VFInput.Strands.HairLength;
+			Parameters.VF_RootScale = VFInput.Strands.HairRootScale;
+			Parameters.VF_TipScale = VFInput.Strands.HairTipScale;
+			Parameters.VF_GroupIndex = Instance->Debug.GroupIndex;
+			PACK_HAIR_ATTRIBUTE_OFFSETS(Parameters.VF_AttributeOffsets, VFInput.Strands.AttributeOffsets);
 
 			if (bIsSRVValid)
 			{
-				Parameters.PositionBuffer = Resources[InvocationIndex].PositionSRV;
-				Parameters.Attribute0Buffer = Resources[InvocationIndex].Attribute0SRV;
-				Parameters.Attribute1Buffer = Resources[InvocationIndex].Attribute1SRV;
-				Parameters.CurveBuffer = Resources[InvocationIndex].CurveSRV;
+				Parameters.PositionBuffer		= Resources[InvocationIndex].PositionSRV;
+				Parameters.AttributeBuffer		= Resources[InvocationIndex].AttributeSRV;
+				Parameters.VertexToCurveBuffer	= Resources[InvocationIndex].VertexToCurveSRV;
+				Parameters.CurveBuffer			= Resources[InvocationIndex].CurveSRV;
 			}
 			else
 			{
-				Parameters.PositionBuffer   = Resources[InvocationIndex].FallbackSRV;
-				Parameters.Attribute0Buffer = Resources[InvocationIndex].FallbackSRV;
-				Parameters.Attribute1Buffer = Resources[InvocationIndex].FallbackSRV;
-				Parameters.CurveBuffer		 = Resources[InvocationIndex].FallbackSRV;
+				Parameters.PositionBuffer		= Resources[InvocationIndex].FallbackSRV;
+				Parameters.AttributeBuffer		= Resources[InvocationIndex].FallbackSRV;
+				Parameters.VertexToCurveBuffer	= Resources[InvocationIndex].FallbackSRV;
+				Parameters.CurveBuffer			= Resources[InvocationIndex].FallbackSRV;
 			}
 		}
 	}
diff --git a/Engine/Plugins/Runtime/HairStrands/Source/HairStrandsCore/Private/Deformer/DeformerDataInterfaceGroom.h b/Engine/Plugins/Runtime/HairStrands/Source/HairStrandsCore/Private/Deformer/DeformerDataInterfaceGroom.h
--- a/Engine/Plugins/Runtime/HairStrands/Source/HairStrandsCore/Private/Deformer/DeformerDataInterfaceGroom.h
+++ b/Engine/Plugins/Runtime/HairStrands/Source/HairStrandsCore/Private/Deformer/DeformerDataInterfaceGroom.h
@@ -67,8 +67,8 @@ private:
 	struct FResources
 	{
 		FRDGBufferSRVRef PositionSRV = nullptr;
-		FRDGBufferSRVRef Attribute0SRV = nullptr;
-		FRDGBufferSRVRef Attribute1SRV = nullptr;
+		FRDGBufferSRVRef AttributeSRV = nullptr;
+		FRDGBufferSRVRef VertexToCurveSRV = nullptr;
 		FRDGBufferSRVRef FallbackSRV = nullptr;
 		FRDGBufferSRVRef CurveSRV = nullptr;
 	};
diff --git a/Engine/Plugins/Runtime/HairStrands/Source/HairStrandsCore/Private/GroomBindingBuilder.cpp b/Engine/Plugins/Runtime/HairStrands/Source/HairStrandsCore/Private/GroomBindingBuilder.cpp
--- a/Engine/Plugins/Runtime/HairStrands/Source/HairStrandsCore/Private/GroomBindingBuilder.cpp
+++ b/Engine/Plugins/Runtime/HairStrands/Source/HairStrandsCore/Private/GroomBindingBuilder.cpp
@@ -1841,19 +1841,6 @@ static void InitHairStrandsRootData(FHairStrandsRootData& Out, const FHairStrand
 	Out.RootCount = HairStrandsDatas->GetNumCurves();
 	Out.PointCount = HairStrandsDatas->GetNumPoints();
 
-	const uint32 CurveCount = HairStrandsDatas->GetNumCurves();
-	Out.VertexToCurveIndexBuffer.SetNum(HairStrandsDatas->GetNumPoints());
-	
-	for (uint32 CurveIndex = 0; CurveIndex < CurveCount; ++CurveIndex)
-	{
-		const uint32 RootIndex = HairStrandsDatas->StrandsCurves.CurvesOffset[CurveIndex];
-		const uint32 PointCount = HairStrandsDatas->StrandsCurves.CurvesCount[CurveIndex];
-		for (uint32 PointIndex = 0; PointIndex < PointCount; ++PointIndex)
-		{
-			Out.VertexToCurveIndexBuffer[RootIndex + PointIndex] = CurveIndex; // RootIndex;
-		}
-	}
-
 	Out.MeshProjectionLODs.SetNum(LODCount);
 	uint32 LODIndex = 0;
 	for (FHairStrandsRootData::FMeshProjectionLOD& MeshProjectionLOD : Out.MeshProjectionLODs)
@@ -1905,8 +1892,6 @@ static void BuildRootBulkData(
 	Out.RootCount = In.RootCount;
 	Out.PointCount = In.PointCount;
 
-	CopyToBulkData<FHairStrandsIndexFormat>(Out.VertexToCurveIndexBuffer, In.VertexToCurveIndexBuffer);
-
 	const uint32 MeshLODCount = In.MeshProjectionLODs.Num();
 	Out.MeshProjectionLODs.SetNum(MeshLODCount);
 	for (uint32 MeshLODIt = 0; MeshLODIt < MeshLODCount; ++MeshLODIt)
@@ -1953,8 +1938,6 @@ static void BuildRootData(
 	Out.RootCount = In.RootCount;
 	Out.PointCount = In.PointCount;
 
-	CopyFromBulkData<FHairStrandsIndexFormat>(Out.VertexToCurveIndexBuffer, In.VertexToCurveIndexBuffer);
-
 	const uint32 MeshLODCount = In.MeshProjectionLODs.Num();
 	Out.MeshProjectionLODs.SetNum(MeshLODCount);
 	for (uint32 MeshLODIt = 0; MeshLODIt < MeshLODCount; ++MeshLODIt)
diff --git a/Engine/Plugins/Runtime/HairStrands/Source/HairStrandsCore/Private/GroomBuilder.cpp b/Engine/Plugins/Runtime/HairStrands/Source/HairStrandsCore/Private/GroomBuilder.cpp
--- a/Engine/Plugins/Runtime/HairStrands/Source/HairStrandsCore/Private/GroomBuilder.cpp
+++ b/Engine/Plugins/Runtime/HairStrands/Source/HairStrandsCore/Private/GroomBuilder.cpp
@@ -42,9 +42,9 @@ FString FGroomBuilder::GetVersion()
 {
 	// Important to update the version when groom building changes
 	if (IsHairStrandContinuousDecimationReorderingEnabled())
-		return TEXT("v4");
+		return TEXT("v4g");
 	else
-		return TEXT("v3");
+		return TEXT("v3g");
 }
 
 namespace FHairStrandsDecimation
@@ -255,6 +255,18 @@ namespace HairStrandsBuilder
 		Out.Unlock();
 	}
 
+	template<typename T>
+	bool HasValues(const TArray<T>& In, T Ref)
+	{
+		for (const T& Value : In)
+		{
+			if (Value != Ref)
+			{
+				return true;
+			}
+		}
+		return false;
+	}
 	/** Build the packed datas for gpu rendering/simulation */
 	void BuildRenderData(const FHairStrandsDatas& HairStrands, const TArray<uint8>& RandomSeeds, FHairStrandsBulkData& OutBulkData)
 	{
@@ -266,25 +278,43 @@ namespace HairStrandsBuilder
 			return;
 
 		TArray<FHairStrandsPositionFormat::Type> OutPackedPositions;
-		TArray<FHairStrandsAttribute0Format::Type> OutPackedAttributes0;
-		TArray<FHairStrandsAttribute1Format::Type> OutPackedAttributes1;
-		TArray<FHairStrandsMaterialFormat::Type> OutPackedMaterials;
 		TArray<FHairStrandsCurveFormat::Type> OutPackedCurves;
 
 		OutPackedPositions.SetNum(NumPoints * FHairStrandsPositionFormat::ComponentCount);
-		OutPackedAttributes0.SetNum(NumPoints * FHairStrandsAttribute0Format::ComponentCount);
-		OutPackedAttributes1.SetNum(NumPoints * FHairStrandsAttribute1Format::ComponentCount);
-		OutPackedMaterials.SetNum(NumPoints * FHairStrandsMaterialFormat::ComponentCount);
 		OutPackedCurves.SetNum(NumCurves * FHairStrandsCurveFormat::ComponentCount);
 
+		const bool bUse16bitsCurveIndex = NumCurves < 65536;
+		TArray<FHairStrandsVertexToCurveFormat16::Type> OutVertexToCurve16;
+		TArray<FHairStrandsVertexToCurveFormat32::Type> OutVertexToCurve32;
+		if (bUse16bitsCurveIndex)
+		{
+			OutVertexToCurve16.SetNum(NumPoints);
+		}
+		else
+		{
+			OutVertexToCurve32.SetNum(NumPoints);
+		}
+
+		const bool bHasBaseColor = HasValues(HairStrands.StrandsPoints.PointsBaseColor, FLinearColor::Transparent);
+		const bool bHasRoughness = HasValues(HairStrands.StrandsPoints.PointsRoughness, 0.f);
+
+		TArray<uint32> AttributeRootUV;
+		TArray<uint32> AttributeSeed;
+		TArray<uint32> AttributeLength;
+		TArray<uint32> AttributeBaseColor;
+		TArray<uint32> AttributeRoughness;
+
+		AttributeRootUV.SetNum(NumCurves);									// 4 bytes encoding - Per-Curve
+		AttributeSeed.SetNum(FMath::DivideAndRoundUp(NumCurves, 4u));		// 1 bytes encoding - Per-Curve
+		AttributeLength.SetNum(FMath::DivideAndRoundUp(NumCurves, 2u));		// 2 bytes encoding - Per-Curve
+		AttributeBaseColor.SetNum(NumPoints);								// 4 bytes encoding - Per-Vertex
+		AttributeRoughness.SetNum(FMath::DivideAndRoundUp(NumPoints, 2u));	// 1 bytes encoding - Per-Vertex
+
 		const FVector HairBoxCenter = HairStrands.BoundingBox.GetCenter();
 
 		const FHairStrandsCurves& Curves = HairStrands.StrandsCurves;
 		const FHairStrandsPoints& Points = HairStrands.StrandsPoints;
 
-		// Track what features/data the bulk data contains
-		bool bHasMaterialData = false;
-
 		struct FPackedRadiusAndType
 		{
 			union
@@ -301,7 +331,6 @@ namespace HairStrandsBuilder
 		const bool bSeedValid = RandomSeeds.Num() > 0;
 		for (uint32 CurveIndex = 0; CurveIndex < NumCurves; ++CurveIndex)
 		{
-			const uint8 CurveSeed = bSeedValid ? RandomSeeds[CurveIndex] : 0;
 			const int32 IndexOffset = Curves.CurvesOffset[CurveIndex];
 			const uint16& PointCount = Curves.CurvesCount[CurveIndex];
 			for (int32 PointIndex = 0; PointIndex < PointCount; ++PointIndex)
@@ -310,23 +339,65 @@ namespace HairStrandsBuilder
 				const uint32 NextIndex = FMath::Min(PointCount + 1, PointCount - 1);
 				const FVector& PointPosition = (FVector)Points.PointsPosition[PointIndex + IndexOffset];
 
-				const float CoordU = Points.PointsCoordU[PointIndex + IndexOffset];
-				const float NormalizedRadius = Points.PointsRadius[PointIndex + IndexOffset];
-				const float NormalizedLength = CoordU * Curves.CurvesLength[CurveIndex];
+				// Packed Position|CoordU|NormalizedRadius
+				{
+					const float CoordU = Points.PointsCoordU[PointIndex + IndexOffset];
+					const float NormalizedRadius = Points.PointsRadius[PointIndex + IndexOffset];
+
+					FHairStrandsPositionFormat::Type& PackedPosition = OutPackedPositions[PointIndex + IndexOffset];
+					CopyVectorToPosition(PointPosition - HairBoxCenter, PackedPosition);
+					FPackedRadiusAndType PackedRadiusAndType;
+					PackedRadiusAndType.Data.ControlPointType = (PointIndex == 0) ? 1u : (PointIndex == (PointCount - 1) ? 2u : 0u);
+					PackedRadiusAndType.Data.NormalizedRadius = uint8(FMath::Clamp(NormalizedRadius * 63.f, 0.f, 63.f));
+					PackedPosition.PackedRadiusAndType = PackedRadiusAndType.Packed;
+					PackedPosition.UCoord = uint8(FMath::Clamp(CoordU * 255.f, 0.f, 255.f));
+				}
 
-				FHairStrandsPositionFormat::Type& PackedPosition = OutPackedPositions[PointIndex + IndexOffset];
-				CopyVectorToPosition(PointPosition - HairBoxCenter, PackedPosition);
-				FPackedRadiusAndType PackedRadiusAndType;
-				PackedRadiusAndType.Data.ControlPointType = (PointIndex == 0) ? 1u : (PointIndex == (PointCount - 1) ? 2u : 0u);
-				PackedRadiusAndType.Data.NormalizedRadius = uint8(FMath::Clamp(NormalizedRadius * 63.f, 0.f, 63.f));
-				PackedPosition.PackedRadiusAndType = PackedRadiusAndType.Packed;
-				PackedPosition.UCoord = uint8(FMath::Clamp(CoordU * 255.f, 0.f, 255.f));
+				// Vertex to Curve
+				if (bUse16bitsCurveIndex)
+				{
+					OutVertexToCurve16[PointIndex + IndexOffset] = CurveIndex;
+				}
+				else
+				{
+					OutVertexToCurve32[PointIndex + IndexOffset] = CurveIndex;
+				}
+
+				// Per-Vertex BaseColor
+				if (bHasBaseColor)
+				{
+					uint32& Packed = AttributeBaseColor[PointIndex + IndexOffset];
+
+					// Cheap sRGB encoding
+					Packed= FMath::Clamp(uint32(FMath::Sqrt(Points.PointsBaseColor[PointIndex + IndexOffset].R) * 0xFF), 0u, 0xFFu)|
+							FMath::Clamp(uint32(FMath::Sqrt(Points.PointsBaseColor[PointIndex + IndexOffset].G) * 0xFF), 0u, 0xFFu)|
+							FMath::Clamp(uint32(FMath::Sqrt(Points.PointsBaseColor[PointIndex + IndexOffset].B) * 0xFF), 0u, 0xFFu);
+				}
+
+				// Per-Vertex Roughness
+				if (bHasRoughness)
+				{
+					uint8* Packed = (uint8*)AttributeRoughness.GetData();
+					Packed[PointIndex + IndexOffset] = FMath::Clamp(uint32(Points.PointsRoughness[PointIndex + IndexOffset] * 0xFF), 0u, 0xFFu);
+				}
+			}
+
+			// Per-Curve Length
+			{
+				const float Length = Curves.CurvesLength[CurveIndex] * Curves.MaxLength;
+				FFloat16* Packed = (FFloat16*)AttributeLength.GetData();
+				Packed[CurveIndex] = Length;
+			}
 
-				FVector2D RootUV = FVector2D(Curves.CurvesRootUV[CurveIndex]); 
-				FHairStrandsAttribute0Format::Type& PackedAttributes0 = OutPackedAttributes0[PointIndex + IndexOffset];
-				PackedAttributes0.NormalizedLength = uint8(FMath::Clamp(NormalizedLength * 255.f, 0.f, 255.f));
-				PackedAttributes0.Seed = CurveSeed;
+			// Per-Curve Seed
+			{
+				uint8* Packed = (uint8*)AttributeSeed.GetData();
+				Packed[CurveIndex] = bSeedValid ? RandomSeeds[CurveIndex] : 0;
+			}
 
+			// Per-Curve Root UVs
+			{
+				FVector2D RootUV = FVector2D(Curves.CurvesRootUV[CurveIndex]);
 				// Root UV support UDIM texture coordinate but limit the spans of the UDIM to be in 256x256 instead of 9999x9999.
 				// The internal UV coords are also limited to 8bits, which means if sampling need to be super precise, this is no enough.
 				// Specal case for UV == 1.0f, as we don't need UDIM data in this case, so force the value to be in [0..1)
@@ -337,30 +408,15 @@ namespace HairStrandsBuilder
 				const FVector2D TextureIndexUV = RootUV - TextureRootUV;
 
 				// UDIM
-				FHairStrandsAttribute1Format::Type& PackedAttributes1 = OutPackedAttributes1[PointIndex + IndexOffset];
-				PackedAttributes1.Packed = 0;
-				PackedAttributes1.Packed |= (uint32(FMath::Clamp(TextureRootUV.X * 2047.f, 0.f, 2047.f)) & 0x7FFu);
-				PackedAttributes1.Packed |= (uint32(FMath::Clamp(TextureRootUV.Y * 2047.f, 0.f, 2047.f)) & 0x7FFu) << 11u;
-				PackedAttributes1.Packed |= (uint32(FMath::Clamp(TextureIndexUV.X, 0.f, 31.f)) & 0x1F) << 22u;
-				PackedAttributes1.Packed |= (uint32(FMath::Clamp(TextureIndexUV.Y, 0.f, 31.f)) & 0x1F) << 27u;
-
-				// Material
-				{
-					FHairStrandsMaterialFormat::Type& Material = OutPackedMaterials[PointIndex + IndexOffset];
-					// Cheap sRGB encoding instead of PointsBaseColor.ToFColor(true), as this makes the decompression 
-					// cheaper on GPU (since R8G8B8A8 sRGB format used/exposed not exposed)
-					Material.BaseColorR = FMath::Clamp(uint32(FMath::Sqrt(Points.PointsBaseColor[PointIndex + IndexOffset].R) * 0xFF), 0u, 0xFFu);
-					Material.BaseColorG = FMath::Clamp(uint32(FMath::Sqrt(Points.PointsBaseColor[PointIndex + IndexOffset].G) * 0xFF), 0u, 0xFFu);
-					Material.BaseColorB = FMath::Clamp(uint32(FMath::Sqrt(Points.PointsBaseColor[PointIndex + IndexOffset].B) * 0xFF), 0u, 0xFFu);
-					Material.Roughness  = FMath::Clamp(uint32(Points.PointsRoughness[PointIndex + IndexOffset] * 0xFF), 0u, 0xFFu);
-
-					if (Material.BaseColorR != 0 || Material.BaseColorG != 0 || Material.BaseColorB != 0 || Material.Roughness != 0)
-					{
-						bHasMaterialData = true;
-					}
-				}
+				uint32& Packed = AttributeRootUV[CurveIndex];
+				Packed = 0;
+				Packed |= (uint32(FMath::Clamp(TextureRootUV.X * 2047.f, 0.f, 2047.f)) & 0x7FFu);
+				Packed |= (uint32(FMath::Clamp(TextureRootUV.Y * 2047.f, 0.f, 2047.f)) & 0x7FFu) << 11u;
+				Packed |= (uint32(FMath::Clamp(TextureIndexUV.X, 0.f, 31.f)) & 0x1F) << 22u;
+				Packed |= (uint32(FMath::Clamp(TextureIndexUV.Y, 0.f, 31.f)) & 0x1F) << 27u;
 			}
 
+
 			// Curves
 			{
 				UE_CLOG(PointCount > 0xFF, LogGroomBuilder, Warning, TEXT("Curve point count is > 256"));
@@ -378,17 +434,63 @@ namespace HairStrandsBuilder
 		OutBulkData.MaxRadius = HairStrands.StrandsCurves.MaxRadius;
 		OutBulkData.Flags = FHairStrandsBulkData::DataFlags_HasData;
 
-		if (bHasMaterialData)	{ OutBulkData.Flags |= FHairStrandsBulkData::DataFlags_HasMaterialData; }	else { OutPackedMaterials.Empty(); }
+		// Concatenate all attributes
+		TArray<FHairStrandsAttributeFormat::Type> OutPackedAttributes;
+		{
+			OutPackedAttributes.Reserve(
+				AttributeRootUV.Num() + 
+				AttributeSeed.Num() +
+				AttributeLength.Num() +
+				AttributeBaseColor.Num() +
+				AttributeRoughness.Num());
 
-		CopyToBulkData<FHairStrandsPositionFormat>(OutBulkData.Positions, OutPackedPositions);
-		CopyToBulkData<FHairStrandsAttribute0Format>(OutBulkData.Attributes0, OutPackedAttributes0);
-		CopyToBulkData<FHairStrandsAttribute1Format>(OutBulkData.Attributes1, OutPackedAttributes1);
+			const uint32 UintToByte = 4;
 
-		if (bHasMaterialData)
-		{
-			CopyToBulkData<FHairStrandsMaterialFormat>(OutBulkData.Materials, OutPackedMaterials);
+			for (uint32 AttributeIt=0; AttributeIt< HAIR_ATTRIBUTE_COUNT; ++AttributeIt)
+			{
+				OutBulkData.AttributeOffsets[AttributeIt] = 0xFFFFFFFF;
+			}
+
+			{
+				OutBulkData.AttributeOffsets[HAIR_ATTRIBUTE_ROOTUV] = OutPackedAttributes.Num() * UintToByte;
+				OutPackedAttributes.Append(AttributeRootUV);
+			}
+
+			{
+				OutBulkData.AttributeOffsets[HAIR_ATTRIBUTE_SEED]   = OutPackedAttributes.Num() * UintToByte;
+				OutPackedAttributes.Append(AttributeSeed);
+			}
+
+			{
+				OutBulkData.AttributeOffsets[HAIR_ATTRIBUTE_LENGTH] = OutPackedAttributes.Num() * UintToByte;
+				OutPackedAttributes.Append(AttributeLength);
+			}
+
+			if (bHasBaseColor)
+			{
+				OutBulkData.AttributeOffsets[HAIR_ATTRIBUTE_BASECOLOR] = OutPackedAttributes.Num() * UintToByte;
+				OutPackedAttributes.Append(AttributeBaseColor);
+			}
+
+			if (bHasRoughness)
+			{
+				OutBulkData.AttributeOffsets[HAIR_ATTRIBUTE_ROUGHNESS] = OutPackedAttributes.Num() * UintToByte;
+				OutPackedAttributes.Append(AttributeRoughness);
+			}
 		}
+
+		CopyToBulkData<FHairStrandsPositionFormat>(OutBulkData.Positions, OutPackedPositions);
+		CopyToBulkData<FHairStrandsAttributeFormat>(OutBulkData.Attributes, OutPackedAttributes);
 		CopyToBulkData<FHairStrandsCurveFormat>(OutBulkData.Curves, OutPackedCurves);
+		if (bUse16bitsCurveIndex)
+		{
+			OutBulkData.Flags |= FHairStrandsBulkData::DataFlags_Has16bitsCurveIndex;
+			CopyToBulkData<FHairStrandsVertexToCurveFormat16>(OutBulkData.VertexToCurve, OutVertexToCurve16);
+		}
+		else
+		{
+			CopyToBulkData<FHairStrandsVertexToCurveFormat32>(OutBulkData.VertexToCurve, OutVertexToCurve32);
+		}
 	}
 
 	void BuildRenderData(const FHairStrandsDatas& HairStrands, FHairStrandsBulkData& OutBulkData)
@@ -1484,20 +1586,6 @@ bool FGroomBuilder::BuildHairDescriptionGroups(const FHairDescription& HairDescr
 		GroomHairWidth = GroomHairWidthAttribute[GroomID];
 	}
 
-	TGroomAttributesConstRef<FVector3f> GroomHairColorAttribute = HairDescription.GroomAttributes().GetAttributesRef<FVector3f>(HairAttribute::Groom::Color);
-	TOptional<FVector> GroomHairColor;
-	if (GroomHairColorAttribute.IsValid())
-	{
-		GroomHairColor = (FVector)GroomHairColorAttribute[GroomID];
-	}
-
-	TGroomAttributesConstRef<float> GroomHairRoughnessAttribute = HairDescription.GroomAttributes().GetAttributesRef<float>(HairAttribute::Groom::Roughness);
-	TOptional<float> GroomHairRoughness;
-	if (GroomHairRoughnessAttribute.IsValid())
-	{
-		GroomHairRoughness = GroomHairRoughnessAttribute[GroomID];
-	}
-
 	TVertexAttributesConstRef<FVector3f> VertexPositions	= HairDescription.VertexAttributes().GetAttributesRef<FVector3f>(HairAttribute::Vertex::Position);
 	TVertexAttributesConstRef<FVector3f> VertexBaseColor	= HairDescription.VertexAttributes().GetAttributesRef<FVector3f>(HairAttribute::Vertex::Color);
 	TVertexAttributesConstRef<float> VertexRoughness	= HairDescription.VertexAttributes().GetAttributesRef<float>(HairAttribute::Vertex::Roughness);
@@ -1672,8 +1760,15 @@ bool FGroomBuilder::BuildHairDescriptionGroups(const FHairDescription& HairDescr
 			FVertexID VertexID(GlobalVertexIndex);
 
 			CurrentHairStrandsDatas->StrandsPoints.PointsPosition.Add(VertexPositions[VertexID]);
-			CurrentHairStrandsDatas->StrandsPoints.PointsBaseColor.Add(bHasBaseColorAttribute ? FLinearColor(VertexBaseColor[VertexID]) : (GroomHairColor     ? FLinearColor(GroomHairColor.GetValue()) : FLinearColor::Black));
-			CurrentHairStrandsDatas->StrandsPoints.PointsRoughness.Add(bHasRoughnessAttribute ? VertexRoughness[VertexID] : (GroomHairRoughness ? GroomHairRoughness.GetValue() : 0.f));
+
+			if (bHasBaseColorAttribute)
+			{
+				CurrentHairStrandsDatas->StrandsPoints.PointsBaseColor.Add(FLinearColor(VertexBaseColor[VertexID]));
+			}
+			if (bHasRoughnessAttribute)
+			{
+				CurrentHairStrandsDatas->StrandsPoints.PointsRoughness.Add(VertexRoughness[VertexID]);
+			}
 
 			// Vertex
 			float VertexWidth = 0.f;
diff --git a/Engine/Plugins/Runtime/HairStrands/Source/HairStrandsCore/Private/GroomRBFDeformer.cpp b/Engine/Plugins/Runtime/HairStrands/Source/HairStrandsCore/Private/GroomRBFDeformer.cpp
--- a/Engine/Plugins/Runtime/HairStrands/Source/HairStrandsCore/Private/GroomRBFDeformer.cpp
+++ b/Engine/Plugins/Runtime/HairStrands/Source/HairStrandsCore/Private/GroomRBFDeformer.cpp
@@ -520,6 +520,23 @@ void FGroomRBFDeformer::GetRBFDeformedGroomAsset(const UGroomAsset* InGroomAsset
 			check(InGroomAsset->HairGroupsInterpolation[GroupIt].DecimationSettings.CurveDecimation == 1);
 		}
 
+		// Build Curve index for every vertices
+		auto BuildVertexToCurveMapping = [](const FHairStrandsDatas& In, TArray<uint32>& Out)
+		{
+			const uint32 CurveCount = In.GetNumCurves();
+			Out.SetNum(In.GetNumPoints());
+
+			for (uint32 CurveIndex = 0; CurveIndex < CurveCount; ++CurveIndex)
+			{
+				const uint32 RootIndex = In.StrandsCurves.CurvesOffset[CurveIndex];
+				const uint32 PointCount = In.StrandsCurves.CurvesCount[CurveIndex];
+				for (uint32 PointIndex = 0; PointIndex < PointCount; ++PointIndex)
+				{
+					Out[RootIndex + PointIndex] = CurveIndex; // RootIndex;
+				}
+			}
+		};
+
 		// Note that the GroupID from the HairGroups cannot be used as the GroupIndex since 
 		// the former may not be strictly increasing nor consecutive
 		// but the ordering of the groups does represent the GroupIndex		
@@ -534,6 +551,11 @@ void FGroomRBFDeformer::GetRBFDeformedGroomAsset(const UGroomAsset* InGroomAsset
 			FHairGroupInfo DummyInfo;
 			FGroomBuilder::BuildData(Group, InGroomAsset->HairGroupsInterpolation[GroupIndex], DummyInfo, StrandsData, GuidesData);
 
+			TArray<uint32> SimRootDataVertexToCurveIndexBuffer;
+			TArray<uint32> RenRootDataVertexToCurveIndexBuffer;
+			BuildVertexToCurveMapping(GuidesData, SimRootDataVertexToCurveIndexBuffer);
+			BuildVertexToCurveMapping(StrandsData, RenRootDataVertexToCurveIndexBuffer);
+
 			// Get deformed guides
 			// If the groom override the value, we output dummy value for the guides, since they won't be used
 			if (InGroomAsset->HairGroupsInterpolation[GroupIndex].InterpolationSettings.bOverrideGuides)
@@ -551,7 +573,7 @@ void FGroomRBFDeformer::GetRBFDeformedGroomAsset(const UGroomAsset* InGroomAsset
 					GuidesData,
 					MeshLODIndex,
 					SkeletalMeshData_Target,
-					SimRootData.VertexToCurveIndexBuffer,
+					SimRootDataVertexToCurveIndexBuffer,
 					SimRootData.MeshProjectionLODs[MeshLODIndex]);
 			}
 
@@ -565,7 +587,7 @@ void FGroomRBFDeformer::GetRBFDeformedGroomAsset(const UGroomAsset* InGroomAsset
 					StrandsData,
 					MeshLODIndex,
 					SkeletalMeshData_Target,
-					RenRootData.VertexToCurveIndexBuffer,
+					RenRootDataVertexToCurveIndexBuffer,
 					RenRootData.MeshProjectionLODs[MeshLODIndex]);
 			}
 		}
diff --git a/Engine/Plugins/Runtime/HairStrands/Source/HairStrandsCore/Private/GroomResources.cpp b/Engine/Plugins/Runtime/HairStrands/Source/HairStrandsCore/Private/GroomResources.cpp
--- a/Engine/Plugins/Runtime/HairStrands/Source/HairStrandsCore/Private/GroomResources.cpp
+++ b/Engine/Plugins/Runtime/HairStrands/Source/HairStrandsCore/Private/GroomResources.cpp
@@ -163,9 +163,13 @@ static FRDGBufferRef InternalCreateVertexBuffer(
 	const FRDGBufferDesc& Desc,
 	const void* InitialData,
 	uint64 InitialDataSize,
-	ERDGInitialDataFlags InitialDataFlags)
+	ERDGInitialDataFlags InitialDataFlags,
+	bool bIsByteAddressBuffer = false)
 {
-	checkf(EnumHasAnyFlags(Desc.Usage, EBufferUsageFlags::VertexBuffer), TEXT("CreateVertexBuffer called with an FRDGBufferDesc underlying type that is not 'VertexBuffer'. Buffer: %s"), Name);
+	if (!bIsByteAddressBuffer)
+	{
+		checkf(EnumHasAnyFlags(Desc.Usage, EBufferUsageFlags::VertexBuffer), TEXT("CreateVertexBuffer called with an FRDGBufferDesc underlying type that is not 'VertexBuffer'. Buffer: %s"), Name);
+	}
 	FRDGBufferRef Buffer = GraphBuilder.CreateBuffer(Desc, Name, ERDGBufferFlags::MultiFrame);
 	if (InitialData && InitialDataSize)
 	{
@@ -230,6 +234,61 @@ void InternalCreateVertexBufferRDG_FromBulkData(FRDGBuilder& GraphBuilder, FByte
 	ConvertToExternalBufferWithViews(GraphBuilder, Buffer, Out, FormatType::Format);
 }
 
+void InternalCreateByteAddressBufferRDG_FromBulkData(FRDGBuilder& GraphBuilder, FByteBulkData& InBulkData, FRDGExternalBuffer& Out, const TCHAR* DebugName, EHairResourceUsageType UsageType)
+{
+	const uint32 DataSizeInBytes = InBulkData.GetBulkDataSize();
+	if (DataSizeInBytes == 0)
+	{
+		Out.Buffer = nullptr;
+		return;
+	}
+
+	FRDGBufferDesc Desc = FRDGBufferDesc::CreateByteAddressDesc(DataSizeInBytes);
+	if (UsageType != EHairResourceUsageType::Dynamic)
+	{
+		Desc.Usage = Desc.Usage & (~BUF_UnorderedAccess);
+	}
+
+	// Unloading of the bulk data is only supported on cooked build, as we can reload the data from the file/archieve
+#if !WITH_EDITORONLY_DATA
+	const bool bReleaseCPUData = GHairStrandsBulkData_ReleaseAfterUse > 0;
+	if (bReleaseCPUData)
+	{
+		InBulkData.SetBulkDataFlags(BULKDATA_SingleUse);
+	}
+#endif
+
+	const uint8* BulkData = (const uint8*)InBulkData.Lock(LOCK_READ_ONLY);
+	FRDGBufferRef Buffer = nullptr;
+#if !WITH_EDITORONLY_DATA
+	if (bReleaseCPUData)
+	{		
+		Buffer = GraphBuilder.CreateBuffer(Desc, DebugName, ERDGBufferFlags::MultiFrame);
+		if (BulkData && DataSizeInBytes)
+		{
+			GraphBuilder.QueueBufferUpload(Buffer, BulkData, DataSizeInBytes, [&InBulkData](const void* Ptr) { InBulkData.Unlock(); });
+		}
+	}
+	else
+#endif
+	{
+		Buffer = InternalCreateVertexBuffer(
+			GraphBuilder,
+			DebugName,
+			Desc,
+			BulkData,
+			DataSizeInBytes,
+			ERDGInitialDataFlags::None, 
+			true /*bIsByteAddressBuffer*/); // Copy data internally
+		InBulkData.Unlock();
+	}
+
+	ConvertToExternalBufferWithViews(GraphBuilder, Buffer, Out, PF_R32_UINT); // TODO?
+}
+
+
+
+
 template<typename FormatType>
 void InternalCreateVertexBufferRDG(FRDGBuilder& GraphBuilder, const typename FormatType::Type* InData, uint32 InDataCount, FRDGExternalBuffer& Out, const TCHAR* DebugName, EHairResourceUsageType UsageType, ERDGInitialDataFlags InitialDataFlags)
 {
@@ -455,11 +514,15 @@ static UTexture2D* CreateCardTexture(FIntPoint Resolution)
 }
 
 /////////////////////////////////////////////////////////////////////////////////////////
-void CreateHairStrandsDebugAttributeBuffer(FRDGBuilder& GraphBuilder, FRDGExternalBuffer* DebugAttributeBuffer, uint32 VertexCount)
+void CreateHairStrandsDebugAttributeBuffer(FRDGBuilder& GraphBuilder, FRDGExternalBuffer* DebugAttributeBuffer, uint32 SizeInBytes)
 {
-	if (VertexCount == 0 || !DebugAttributeBuffer)
+	if (SizeInBytes == 0 || !DebugAttributeBuffer)
+	{
 		return;
-	InternalCreateVertexBufferRDG<FHairStrandsAttribute0Format>(GraphBuilder, VertexCount, *DebugAttributeBuffer, TEXT("Hair.Strands_DebugAttributeBuffer"), EHairResourceUsageType::Dynamic);
+	}
+
+	FRDGBufferRef Buffer = GraphBuilder.CreateBuffer(FRDGBufferDesc::CreateByteAddressDesc(SizeInBytes), TEXT("Hair.Strands_DebugAttributeBuffer"), ERDGBufferFlags::MultiFrame);
+	ConvertToExternalBufferWithViews(GraphBuilder, Buffer, *DebugAttributeBuffer, PF_R32_UINT);
 }
 
 /////////////////////////////////////////////////////////////////////////////////////////
@@ -777,7 +840,7 @@ void FHairMeshesDeformedResource::InternalRelease()
 
 FHairStrandsRestResource::FHairStrandsRestResource(FHairStrandsBulkData& InBulkData, EHairStrandsResourcesType InCurveType, const FHairResourceName& InResourceName) :
 	FHairCommonResource(EHairStrandsAllocationType::Deferred, InResourceName),
-	PositionBuffer(), Attribute0Buffer(), Attribute1Buffer(), MaterialBuffer(), BulkData(InBulkData), CurveType(InCurveType)
+	PositionBuffer(), AttributeBuffer(), VertexToCurveBuffer(), BulkData(InBulkData), CurveType(InCurveType)
 {
 	// Sanity check
 	check(!!(BulkData.Flags & FHairStrandsBulkData::DataFlags_HasData));
@@ -789,14 +852,8 @@ bool FHairStrandsRestResource::InternalIsDataLoaded()
 	{
 		FBulkDataBatchRequest::FBatchBuilder Batch = FBulkDataBatchRequest::NewBatch(5);
 		Batch.Read(BulkData.Positions);
-		Batch.Read(BulkData.Attributes0);
-		Batch.Read(BulkData.Attributes1);
-
-		if (!!(BulkData.Flags & FHairStrandsBulkData::DataFlags_HasMaterialData))
-		{
-			Batch.Read(BulkData.Materials);
-		}
-
+		Batch.Read(BulkData.Attributes);
+		Batch.Read(BulkData.VertexToCurve);
 		Batch.Read(BulkData.Curves);
 
 		Batch.Issue(BulkDataRequest);
@@ -815,13 +872,15 @@ void FHairStrandsRestResource::InternalAllocate(FRDGBuilder& GraphBuilder)
 	// 1. Lock data, which force the loading data from files (on non-editor build/cooked data). These data are then uploaded to the GPU
 	// 2. A local copy is done by the buffer uploader. This copy is discarded once the uploading is done.
 	InternalCreateVertexBufferRDG_FromBulkData<FHairStrandsPositionFormat>(GraphBuilder, BulkData.Positions, PointCount, PositionBuffer, ToHairResourceDebugName(HAIRSTRANDS_RESOUCE_NAME(CurveType, Hair.StrandsRest_PositionBuffer), ResourceName), EHairResourceUsageType::Static);
-	InternalCreateVertexBufferRDG_FromBulkData<FHairStrandsAttribute0Format>(GraphBuilder, BulkData.Attributes0, PointCount, Attribute0Buffer, ToHairResourceDebugName(HAIRSTRANDS_RESOUCE_NAME(CurveType, Hair.StrandsRest_Attribute0Buffer), ResourceName), EHairResourceUsageType::Static);
-	InternalCreateVertexBufferRDG_FromBulkData<FHairStrandsAttribute1Format>(GraphBuilder, BulkData.Attributes1, PointCount, Attribute1Buffer, ToHairResourceDebugName(HAIRSTRANDS_RESOUCE_NAME(CurveType, Hair.StrandsRest_Attribute1Buffer), ResourceName), EHairResourceUsageType::Static);
+	InternalCreateByteAddressBufferRDG_FromBulkData(GraphBuilder, BulkData.Attributes, AttributeBuffer, ToHairResourceDebugName(HAIRSTRANDS_RESOUCE_NAME(CurveType, Hair.StrandsRest_AttributeBuffer), ResourceName), EHairResourceUsageType::Static);
 	InternalCreateVertexBufferRDG_FromBulkData<FHairStrandsCurveFormat>(GraphBuilder, BulkData.Curves, CurveCount, CurveBuffer, ToHairResourceDebugName(HAIRSTRANDS_RESOUCE_NAME(CurveType, Hair.StrandsRest_CurveBuffer), ResourceName), EHairResourceUsageType::Static);
-
-	if (!!(BulkData.Flags & FHairStrandsBulkData::DataFlags_HasMaterialData))
+	if (!!(BulkData.Flags & BulkData.DataFlags_Has16bitsCurveIndex))
 	{
-		InternalCreateVertexBufferRDG_FromBulkData<FHairStrandsMaterialFormat>(GraphBuilder, BulkData.Materials, PointCount, MaterialBuffer, ToHairResourceDebugName(HAIRSTRANDS_RESOUCE_NAME(CurveType, Hair.StrandsRest_MaterialBuffer), ResourceName), EHairResourceUsageType::Static);
+		InternalCreateVertexBufferRDG_FromBulkData<FHairStrandsVertexToCurveFormat16>(GraphBuilder, BulkData.VertexToCurve, PointCount, VertexToCurveBuffer, ToHairResourceDebugName(HAIRSTRANDS_RESOUCE_NAME(CurveType, Hair.StrandsRest_VertexToCurveBuffer), ResourceName), EHairResourceUsageType::Static);
+	}
+	else
+	{
+		InternalCreateVertexBufferRDG_FromBulkData<FHairStrandsVertexToCurveFormat32>(GraphBuilder, BulkData.VertexToCurve, PointCount, VertexToCurveBuffer, ToHairResourceDebugName(HAIRSTRANDS_RESOUCE_NAME(CurveType, Hair.StrandsRest_VertexToCurveBuffer), ResourceName), EHairResourceUsageType::Static);
 	}
 
 	TArray<FVector4f> RestOffset;
@@ -863,9 +922,8 @@ void FHairStrandsRestResource::InternalRelease()
 {
 	PositionBuffer.Release();
 	PositionOffsetBuffer.Release();
-	Attribute0Buffer.Release();
-	Attribute1Buffer.Release();
-	MaterialBuffer.Release();
+	AttributeBuffer.Release();
+	VertexToCurveBuffer.Release();
 	TangentBuffer.Release();
 	CurveBuffer.Release();
 }
@@ -1104,20 +1162,8 @@ void FHairStrandsRestRootResource::PopulateFromRootData()
 }
 
 bool FHairStrandsRestRootResource::InternalIsDataLoaded()
-{
-	if (BulkData.PointCount == 0)
-	{
-		return true;
-	}
-
-	if (BulkDataRequest.IsNone())
-	{
-		FBulkDataBatchRequest::NewBatch(1)
-			.Read(BulkData.VertexToCurveIndexBuffer)
-			.Issue(BulkDataRequest);
-	}
-
-	return BulkDataRequest.IsCompleted();
+{	
+	return true;
 }
 
 bool FHairStrandsRestRootResource::InternalIsLODDataLoaded(int32 LODIndex)
@@ -1172,11 +1218,6 @@ void FHairStrandsRestRootResource::InternalAllocate(FRDGBuilder& GraphBuilder)
 	{
 		PopulateFromRootData();
 	}
-
-	if (BulkData.PointCount > 0)
-	{
-		InternalCreateVertexBufferRDG_FromBulkData<FHairStrandsIndexFormat>(GraphBuilder, BulkData.VertexToCurveIndexBuffer, BulkData.PointCount, VertexToCurveIndexBuffer, HAIRSTRANDS_RESOUCE_NAME(CurveType, Hair.StrandsRestRoot_VertexToCurveIndexBuffer), EHairResourceUsageType::Static);
-	}
 }
 
 void FHairStrandsRestRootResource::InternalAllocateLOD(FRDGBuilder& GraphBuilder, int32 LODIndex)
@@ -1245,8 +1286,6 @@ void FHairStrandsRestRootResource::InternalAllocateLOD(FRDGBuilder& GraphBuilder
 
 void FHairStrandsRestRootResource::InternalRelease()
 {
-	VertexToCurveIndexBuffer.Release();
-	
 	for (FLOD& GPUData : LODs)
 	{
 		GPUData.Status = FLOD::EStatus::Invalid;
@@ -1435,16 +1474,10 @@ void FHairStrandsRootBulkData::Serialize(FArchive& Ar, UObject* Owner)
 	const int32 ChunkIndex = 0;
 	bool bAttemptFileMapping = false;
 
-	{
-		const uint32 BulkFlags = BULKDATA_Force_NOT_InlinePayload;
-		VertexToCurveIndexBuffer.SetBulkDataFlags(BulkFlags);
-	}
-
 	if (!Ar.IsObjectReferenceCollector())
 	{
 		Ar << RootCount;
 		Ar << PointCount;
-		VertexToCurveIndexBuffer.Serialize(Ar, Owner, ChunkIndex, bAttemptFileMapping);
 		InternalSerialize(Ar, Owner, MeshProjectionLODs);
 	}
 }
@@ -1452,8 +1485,6 @@ void FHairStrandsRootBulkData::Serialize(FArchive& Ar, UObject* Owner)
 void FHairStrandsRootBulkData::Reset()
 {
 	RootCount = 0;
-	VertexToCurveIndexBuffer.RemoveBulkData();
-	VertexToCurveIndexBuffer = FByteBulkData();
 	for (FMeshProjectionLOD& LOD : MeshProjectionLODs)
 	{
 		LOD.RootBarycentricBuffer.RemoveBulkData();
@@ -1514,7 +1545,6 @@ bool FHairStrandsRootData::HasProjectionData() const
 void FHairStrandsRootData::Reset()
 {
 	RootCount = 0;
-	VertexToCurveIndexBuffer.Empty();
 	MeshProjectionLODs.Empty();
 }
 
diff --git a/Engine/Plugins/Runtime/HairStrands/Source/HairStrandsCore/Private/GroomTextureBuilder.cpp b/Engine/Plugins/Runtime/HairStrands/Source/HairStrandsCore/Private/GroomTextureBuilder.cpp
--- a/Engine/Plugins/Runtime/HairStrands/Source/HairStrandsCore/Private/GroomTextureBuilder.cpp
+++ b/Engine/Plugins/Runtime/HairStrands/Source/HairStrandsCore/Private/GroomTextureBuilder.cpp
@@ -18,6 +18,7 @@
 #include "Engine/StaticMesh.h"
 #include "TextureResource.h"
 #include "UObject/Package.h"
+#include "HairStrandsDefinitions.h"
 
 static int32 GHairStrandsTextureDilationPassCount = 8;
 static FAutoConsoleVariableRef CVarHairStrandsTextureDilationPassCount(TEXT("r.HairStrands.Textures.DilationCount"), GHairStrandsTextureDilationPassCount, TEXT("Number of dilation pass run onto the generated hair strands textures (Default:8)."));
@@ -621,19 +622,16 @@ class FHairStrandsTexturePS : public FGlobalShader
 		SHADER_PARAMETER(uint32, VertexCount)
 		SHADER_PARAMETER(float, MaxDistance)
 		SHADER_PARAMETER(int32, TracingDirection)
-		SHADER_PARAMETER(uint32, bHasMaterialData)
+		SHADER_PARAMETER_ARRAY(FUintVector4, InVF_AttributeOffsets, [HAIR_ATTRIBUTE_OFFSET_COUNT])
 
 		SHADER_PARAMETER(uint32, UVsChannelIndex)
 		SHADER_PARAMETER(uint32, UVsChannelCount)
 
 		SHADER_PARAMETER(float, InVF_Radius)
-		SHADER_PARAMETER(float, InVF_Length)
 		SHADER_PARAMETER(FVector3f, InVF_PositionOffset)
 		SHADER_PARAMETER_SRV(Buffer, InVF_PositionBuffer)
-		SHADER_PARAMETER_SRV(Buffer, InVF_Attribute0Buffer)
-		SHADER_PARAMETER_SRV(Buffer, InVF_Attribute1Buffer)
-		SHADER_PARAMETER_SRV(Buffer, InVF_MaterialBuffer)
-		SHADER_PARAMETER(uint32, InVF_ControlPointCount)
+		SHADER_PARAMETER_SRV(ByteAddressBuffer, InVF_AttributeBuffer)
+		SHADER_PARAMETER_SRV(Buffer, InVF_VertexToCurveBuffer)
 		SHADER_PARAMETER(uint32, InVF_GroupIndex)
 
 		SHADER_PARAMETER(FVector3f, Voxel_MinBound)
@@ -760,19 +758,17 @@ static void InternalGenerateHairStrandsTextures(
 	FRDGBufferRef VoxelData,
 	
 	FRHIShaderResourceView* InHairStrands_PositionBuffer,
-	FRHIShaderResourceView* InHairStrands_Attribute0Buffer,
-	FRHIShaderResourceView* InHairStrands_Attribute1Buffer,
-	FRHIShaderResourceView* InHairStrands_MaterialBuffer,
+	FRHIShaderResourceView* InHairStrands_AttributeBuffer,
+	FRHIShaderResourceView* InHairStrands_VertexToCurveBuffer,
 	const FVector& InHairStrands_PositionOffset,
 	float InHairStrands_Radius,
 	float InHairStrands_Length,
 	uint32 InHairStrands_ControlPointCount,
 	uint32 InHairStrands_GroupIndex,
-
+	const uint32* InHairStrands_AttributeOffsets,
 	FHairStrandsRDGTextures& Out)
 {
 	const FIntPoint OutputResolution = Out.DepthTexture->Desc.Extent;
-	const bool bHasMaterialData = InHairStrands_MaterialBuffer != nullptr;
 	const FIntPoint OutTileCoord = Out.TileCoord;
 
 	FHairStrandsTexturePS::FParameters* ParametersPS = GraphBuilder.AllocParameters<FHairStrandsTexturePS::FParameters>();
@@ -783,21 +779,19 @@ static void InternalGenerateHairStrandsTextures(
 	ParametersPS->NormalsBuffer = InMeshNormalsBuffer;
 	ParametersPS->MaxDistance = FMath::Abs(InMaxDistance);
 	ParametersPS->TracingDirection = InTracingDirection;
-	ParametersPS->bHasMaterialData = bHasMaterialData ? 1u : 0u;
 
 	ParametersPS->UVsChannelIndex = UVsChannelIndex;
 	ParametersPS->UVsChannelCount = UVsChannelCount;
 
 	ParametersPS->InVF_PositionBuffer = InHairStrands_PositionBuffer;
 	ParametersPS->InVF_PositionOffset = (FVector3f)InHairStrands_PositionOffset;
-	ParametersPS->InVF_Attribute0Buffer = InHairStrands_Attribute0Buffer;
-	ParametersPS->InVF_Attribute1Buffer = InHairStrands_Attribute1Buffer;
-	ParametersPS->InVF_MaterialBuffer = bHasMaterialData ? InHairStrands_MaterialBuffer : InHairStrands_Attribute0Buffer;
+	ParametersPS->InVF_AttributeBuffer = InHairStrands_AttributeBuffer;
+	ParametersPS->InVF_VertexToCurveBuffer = InHairStrands_VertexToCurveBuffer;
+	PACK_HAIR_ATTRIBUTE_OFFSETS(ParametersPS->InVF_AttributeOffsets, InHairStrands_AttributeOffsets);
+
 	ParametersPS->InVF_Radius = InHairStrands_Radius;
-	ParametersPS->InVF_Length = InHairStrands_Length;
-	ParametersPS->InVF_ControlPointCount = InHairStrands_ControlPointCount;
 	ParametersPS->InVF_GroupIndex = InHairStrands_GroupIndex;
-
+	
 	ParametersPS->Voxel_MinBound = VoxelMinBound;
 	ParametersPS->Voxel_MaxBound = VoxelMaxBound;
 	ParametersPS->Voxel_Resolution = VoxelResolution;
@@ -1094,14 +1088,14 @@ static bool TraceTextures(
 				VoxelData,
 
 				GroupData.Strands.RestResource->PositionBuffer.SRV,
-				GroupData.Strands.RestResource->Attribute0Buffer.SRV,
-				GroupData.Strands.RestResource->Attribute1Buffer.SRV,
-				GroupData.Strands.RestResource->MaterialBuffer.SRV,
+				GroupData.Strands.RestResource->AttributeBuffer.SRV,
+				GroupData.Strands.RestResource->VertexToCurveBuffer.SRV,
 				GroupData.Strands.RestResource->GetPositionOffset(),
 				RenderingData.GeometrySettings.HairWidth * 0.5f,
 				GroupData.Strands.RestResource->BulkData.MaxLength,
 				GroupData.Strands.RestResource->BulkData.PointCount,
 				GroupIndex,
+				GroupData.Strands.RestResource->BulkData.AttributeOffsets,
 
 				Output);
 
diff --git a/Engine/Plugins/Runtime/HairStrands/Source/HairStrandsCore/Private/HairStrandsDatas.cpp b/Engine/Plugins/Runtime/HairStrands/Source/HairStrandsCore/Private/HairStrandsDatas.cpp
--- a/Engine/Plugins/Runtime/HairStrands/Source/HairStrandsCore/Private/HairStrandsDatas.cpp
+++ b/Engine/Plugins/Runtime/HairStrands/Source/HairStrandsCore/Private/HairStrandsDatas.cpp
@@ -19,7 +19,7 @@ void FHairStrandsCurves::SetNum(const uint32 NumCurves)
 	CurvesLength.SetNum(NumCurves);
 	CurvesRootUV.SetNum(NumCurves);
 
-	// Curves clossest guide are not initialized to track if the curves 
+	// Curves closest guide are not initialized to track if the curves 
 	// data contains or not precomputed weights
 	// CurvesClosestGuideIDs.SetNum(0);
 	// CurvesClosestGuideWeights.SetNum(0);
@@ -79,16 +79,6 @@ FArchive& operator<<(FArchive& Ar, FPackedHairVertex& Vertex)
 	return Ar;
 }
 
-FArchive& operator<<(FArchive& Ar, FHairMaterialVertex& Vertex)
-{
-	Ar << Vertex.BaseColorR;
-	Ar << Vertex.BaseColorG;
-	Ar << Vertex.BaseColorB;
-	Ar << Vertex.Roughness;
-
-	return Ar;
-}
-
 FArchive& operator<<(FArchive& Ar, FHairInterpolation0Vertex& Vertex)
 {
 	Ar << Vertex.Index0;
@@ -206,9 +196,9 @@ void FHairStrandsInterpolationBulkData::Serialize(FArchive& Ar, UObject* Owner)
 void FHairStrandsBulkData::Serialize(FArchive& Ar, UObject* Owner)
 {
 	static_assert(sizeof(FHairStrandsPositionFormat::BulkType) == sizeof(FHairStrandsPositionFormat::Type));
-	static_assert(sizeof(FHairStrandsAttribute0Format::BulkType) == sizeof(FHairStrandsAttribute0Format::Type));
-	static_assert(sizeof(FHairStrandsAttribute1Format::BulkType) == sizeof(FHairStrandsAttribute1Format::Type));
-	static_assert(sizeof(FHairStrandsMaterialFormat::BulkType) == sizeof(FHairStrandsMaterialFormat::Type));
+	static_assert(sizeof(FHairStrandsAttributeFormat::BulkType) == sizeof(FHairStrandsAttributeFormat::Type));
+	static_assert(sizeof(FHairStrandsVertexToCurveFormat16::BulkType) == sizeof(FHairStrandsVertexToCurveFormat16::Type));
+	static_assert(sizeof(FHairStrandsVertexToCurveFormat32::BulkType) == sizeof(FHairStrandsVertexToCurveFormat32::Type));
 	static_assert(sizeof(FHairStrandsRootIndexFormat::BulkType) == sizeof(FHairStrandsRootIndexFormat::Type)); 
 
 	Ar.UsingCustomVersion(FReleaseObjectVersion::GUID);
@@ -220,6 +210,10 @@ void FHairStrandsBulkData::Serialize(FArchive& Ar, UObject* Owner)
 	Ar << MaxRadius;
 	Ar << BoundingBox;
 	Ar << Flags;
+	for (uint8 AttributeIt = 0; AttributeIt < HAIR_ATTRIBUTE_COUNT; ++AttributeIt)
+	{
+		Ar << AttributeOffsets[AttributeIt];
+	}
 
 	// Forced not inline means the bulk data won't automatically be loaded when we deserialize
 	// but only when we explicitly take action to load it
@@ -227,9 +221,8 @@ void FHairStrandsBulkData::Serialize(FArchive& Ar, UObject* Owner)
 	{
 		const uint32 BulkFlags = BULKDATA_Force_NOT_InlinePayload;
 		Positions.SetBulkDataFlags(BulkFlags);
-		Attributes0.SetBulkDataFlags(BulkFlags);
-		Attributes1.SetBulkDataFlags(BulkFlags);
-		Materials.SetBulkDataFlags(BulkFlags);
+		Attributes.SetBulkDataFlags(BulkFlags);
+		VertexToCurve.SetBulkDataFlags(BulkFlags);
 		Curves.SetBulkDataFlags(BulkFlags);
 	}
 
@@ -239,13 +232,8 @@ void FHairStrandsBulkData::Serialize(FArchive& Ar, UObject* Owner)
 		bool bAttemptFileMapping = false;
 
 		Positions.Serialize(Ar, Owner, ChunkIndex, bAttemptFileMapping);
-		Attributes0.Serialize(Ar, Owner, ChunkIndex, bAttemptFileMapping);
-		Attributes1.Serialize(Ar, Owner, ChunkIndex, bAttemptFileMapping);
-
-		if (!!(Flags & DataFlags_HasMaterialData))
-		{
-			Materials.Serialize(Ar, Owner, ChunkIndex, bAttemptFileMapping);
-		}
+		Attributes.Serialize(Ar, Owner, ChunkIndex, bAttemptFileMapping);
+		VertexToCurve.Serialize(Ar, Owner, ChunkIndex, bAttemptFileMapping);
 		Curves.Serialize(Ar, Owner, ChunkIndex, bAttemptFileMapping);
 	}
 }
@@ -258,20 +246,21 @@ void FHairStrandsBulkData::Reset()
 	MaxRadius = 0;
 	BoundingBox = FBox(EForceInit::ForceInit);
 	Flags = 0;
-
+	for (uint8 AttributeIt = 0; AttributeIt < HAIR_ATTRIBUTE_COUNT; ++AttributeIt)
+	{
+		AttributeOffsets[AttributeIt] = 0xFFFFFFFF;
+	}
 	// Deallocate memory if needed
 	Positions.RemoveBulkData();
-	Attributes0.RemoveBulkData();
-	Attributes1.RemoveBulkData();
-	Materials.RemoveBulkData();
+	Attributes.RemoveBulkData();
+	VertexToCurve.RemoveBulkData();
 	Curves.RemoveBulkData();
 
 	// Reset the bulk byte buffer to ensure the (serialize) data size is reset to 0
-	Positions 	= FByteBulkData();
-	Attributes0 = FByteBulkData();
-	Attributes1 = FByteBulkData();
-	Materials 	= FByteBulkData();
-	Curves		= FByteBulkData();
+	Positions 		= FByteBulkData();
+	Attributes		= FByteBulkData();
+	VertexToCurve	= FByteBulkData();
+	Curves			= FByteBulkData();
 }
 
 void FHairStrandsDatas::Reset()
diff --git a/Engine/Plugins/Runtime/HairStrands/Source/HairStrandsCore/Private/HairStrandsRendering.cpp b/Engine/Plugins/Runtime/HairStrands/Source/HairStrandsCore/Private/HairStrandsRendering.cpp
--- a/Engine/Plugins/Runtime/HairStrands/Source/HairStrandsCore/Private/HairStrandsRendering.cpp
+++ b/Engine/Plugins/Runtime/HairStrands/Source/HairStrandsCore/Private/HairStrandsRendering.cpp
@@ -299,7 +299,7 @@ class FDeformGuideCS : public FGlobalShader
 		SHADER_PARAMETER_RDG_BUFFER_SRV(Buffer<float4>, SimDeformedPosition2Buffer)
 
 		SHADER_PARAMETER_RDG_BUFFER_SRV(Buffer<uint>, SimRootBarycentricBuffer)
-		SHADER_PARAMETER_RDG_BUFFER_SRV(Buffer<uint>, SimVertexToRootIndexBuffer)
+		SHADER_PARAMETER_RDG_BUFFER_SRV(Buffer<uint>, SimVertexToCurveBuffer)
 		SHADER_PARAMETER_RDG_BUFFER_SRV(Buffer<uint>, SimRootToUniqueTriangleIndexBuffer)
 
 		SHADER_PARAMETER_RDG_BUFFER_SRV(Buffer, SimDeformedOffsetBuffer)
@@ -328,6 +328,7 @@ static void AddDeformSimHairStrandsPass(
 	FHairStrandsRestRootResource* SimRestRootResources,
 	FHairStrandsDeformedRootResource* SimDeformedRootResources,
 	FRDGBufferSRVRef SimRestPosePositionBuffer,
+	FRDGBufferSRVRef SimVertexToCurveBuffer,
 	FRDGImportedBuffer OutSimDeformedPositionBuffer,
 	const FVector& SimRestOffset,
 	FRDGBufferSRVRef SimDeformedOffsetBuffer,
@@ -368,10 +369,10 @@ static void AddDeformSimHairStrandsPass(
 
 	if (DeformationType == EDeformationType::OffsetGuide)
 	{
-		const bool bIsVertexToCurveBuffersValid = SimRestRootResources && SimRestRootResources->VertexToCurveIndexBuffer.Buffer != nullptr;
+		const bool bIsVertexToCurveBuffersValid = SimVertexToCurveBuffer != nullptr;
 		if (bIsVertexToCurveBuffersValid)
 		{
-			Parameters->SimVertexToRootIndexBuffer = RegisterAsSRV(GraphBuilder, SimRestRootResources->VertexToCurveIndexBuffer);
+			Parameters->SimVertexToCurveBuffer = SimVertexToCurveBuffer;
 		}
 
 		const uint32 RootCount = SimRestRootResources ? SimRestRootResources->BulkData.RootCount : 0;
@@ -516,7 +517,7 @@ class FHairInterpolationCS : public FGlobalShader
 		SHADER_PARAMETER_RDG_BUFFER_SRV(Buffer<float4>, RenDeformedPosition2Buffer)
 
 		SHADER_PARAMETER_RDG_BUFFER_SRV(Buffer<uint>, RenRootBarycentricBuffer)
-		SHADER_PARAMETER_RDG_BUFFER_SRV(Buffer<uint>, RenVertexToRootIndexBuffer)
+		SHADER_PARAMETER_RDG_BUFFER_SRV(Buffer<uint>, RenVertexToCurveBuffer)
 		SHADER_PARAMETER_RDG_BUFFER_SRV(Buffer<uint>, RenRootToUniqueTriangleIndexBuffer)
 
 		SHADER_PARAMETER_RDG_BUFFER_SRV(Buffer<float4>, SimRestPosition0Buffer)
@@ -528,7 +529,7 @@ class FHairInterpolationCS : public FGlobalShader
 		SHADER_PARAMETER_RDG_BUFFER_SRV(Buffer<float4>, SimDeformedPosition2Buffer)
 
 		SHADER_PARAMETER_RDG_BUFFER_SRV(Buffer<uint>, SimRootBarycentricBuffer)
-		SHADER_PARAMETER_RDG_BUFFER_SRV(Buffer<uint>, SimVertexToRootIndexBuffer)
+		SHADER_PARAMETER_RDG_BUFFER_SRV(Buffer<uint>, SimVertexToCurveBuffer)
 		SHADER_PARAMETER_RDG_BUFFER_SRV(Buffer<uint>, SimRootToUniqueTriangleIndexBuffer)
 
 		SHADER_PARAMETER_RDG_BUFFER_SRV(Buffer<uint>, SimRootPointIndexBuffer)
@@ -583,6 +584,7 @@ static void AddHairStrandsInterpolationPass(
 	const FHairStrandsDeformedRootResource* RenDeformedRootResources,
 	const FHairStrandsDeformedRootResource* SimDeformedRootResources,
 	const FRDGBufferSRVRef& RenRestPosePositionBuffer,
+	const FRDGBufferSRVRef& RenVertexToCurveBuffer,
 	const bool bUseSingleGuide,
 	const FRDGBufferSRVRef& InterpolationBuffer,
 	const FRDGBufferSRVRef& Interpolation0Buffer,
@@ -590,6 +592,7 @@ static void AddHairStrandsInterpolationPass(
 	const FRDGBufferSRVRef& SimRestPosePositionBuffer,
 	const FRDGBufferSRVRef& SimDeformedPositionBuffer,
 	const FRDGBufferSRVRef& SimRootPointIndexBuffer,
+	const FRDGBufferSRVRef& SimVertexToCurveBuffer,
 	const FRDGBufferSRVRef& RenDeformerPositionBuffer,
 	FRDGImportedBuffer& OutRenPositionBuffer,
 	const FHairStrandsDeformedRootResource::FLOD::EFrameType DeformedFrame)
@@ -623,22 +626,8 @@ static void AddHairStrandsInterpolationPass(
 
 	Parameters->DispatchCountX = DispatchCount.X;
 	Parameters->SimRootPointIndexBuffer = SimRootPointIndexBuffer;
-	
-	const bool bIsVertexToCurveBuffersValid_Sim =
-		SimRestRootResources &&
-		SimRestRootResources->VertexToCurveIndexBuffer.Buffer != nullptr;
-	if (bIsVertexToCurveBuffersValid_Sim)
-	{
-		Parameters->SimVertexToRootIndexBuffer = RegisterAsSRV(GraphBuilder, SimRestRootResources->VertexToCurveIndexBuffer);
-	}
-
-	const bool bIsVertexToCurveBuffersValid_Ren =
-		RenRestRootResources &&
-		RenRestRootResources->VertexToCurveIndexBuffer.Buffer != nullptr;
-	if (bIsVertexToCurveBuffersValid_Ren)
-	{
-		Parameters->RenVertexToRootIndexBuffer = RegisterAsSRV(GraphBuilder, RenRestRootResources->VertexToCurveIndexBuffer);
-	}
+	Parameters->SimVertexToCurveBuffer = SimVertexToCurveBuffer;
+	Parameters->RenVertexToCurveBuffer = RenVertexToCurveBuffer;
 	
 	Parameters->LocalToWorldMatrix = FMatrix44f(Instance->LocalToWorld.ToMatrixWithScale());		// LWC_TODO: Precision loss
 	Parameters->HairLengthScale = HairLengthScale;
@@ -652,7 +641,6 @@ static void AddHairStrandsInterpolationPass(
 	#endif
 
 	const bool bSupportDynamicMesh = 
-		bIsVertexToCurveBuffersValid_Ren &&
 		RenRestRootResources &&
 		RenRestRootResources->BulkData.RootCount > 0 && 
 		MeshLODIndex >= 0 && 
@@ -779,16 +767,16 @@ class FHairPatchAttributeCS : public FGlobalShader
 	using FPermutationDomain = TShaderPermutationDomain<FSimulation>;
 
 	BEGIN_SHADER_PARAMETER_STRUCT(FParameters, )
-		SHADER_PARAMETER(uint32, VertexCount)
+		SHADER_PARAMETER(uint32, CurveCount)
+		SHADER_PARAMETER_ARRAY(FUintVector4, RenAttributeOffsets, [HAIR_ATTRIBUTE_OFFSET_COUNT])
+		SHADER_PARAMETER_RDG_BUFFER_SRV(Buffer, RenCurveBuffer)
 
-		SHADER_PARAMETER_RDG_BUFFER_SRV(Buffer<float4>, SimAttribute0Buffer)
 		SHADER_PARAMETER_RDG_BUFFER_SRV(Buffer, InterpolationBuffer)
 		SHADER_PARAMETER_RDG_BUFFER_SRV(Buffer, Interpolation0Buffer)
 		SHADER_PARAMETER_RDG_BUFFER_SRV(Buffer, Interpolation1Buffer)
 
-		SHADER_PARAMETER_RDG_BUFFER_SRV(Buffer<float4>, RenAttribute0Buffer)
 		SHADER_PARAMETER_RDG_BUFFER_SRV(Buffer, RenVertexToClusterIdBuffer)
-		SHADER_PARAMETER_RDG_BUFFER_UAV(RWBuffer<float4>, OutRenAttribute0Buffer)
+		SHADER_PARAMETER_RDG_BUFFER_UAV(RWByteAddressBuffer, OutRenAttributeBuffer)
 	END_SHADER_PARAMETER_STRUCT()
 
 public:
@@ -812,36 +800,42 @@ enum class EHairPatchAttribute : uint8
 static void AddPatchAttributePass(
 	FRDGBuilder& GraphBuilder,
 	FGlobalShaderMap* ShaderMap,
-	const uint32 VertexCount,
+	const uint32 CurveCount,
 	const EHairPatchAttribute Mode,
 	const bool bSimulation,
 	const bool bUseSingleGuide,
-	const FRDGBufferSRVRef& RenAttribute0Buffer,
+	const uint32* RenAttributeOffests,
+	const FRDGBufferRef& RenAttributeBuffer,
+	const FRDGBufferSRVRef& RenCurveBuffer,
 	const FRDGBufferSRVRef& RenVertexToClusterIdBuffer,
-	const FRDGBufferSRVRef& SimAttribute0Buffer,
 	const FRDGBufferSRVRef& InterpolationBuffer,
 	const FRDGBufferSRVRef& Interpolation0Buffer,
 	const FRDGBufferSRVRef& Interpolation1Buffer,
-	FRDGImportedBuffer& OutRenAttribute0Buffer)
+	FRDGImportedBuffer& OutRenAttributeBuffer)
 {
 	// Sanity check
 	check(Mode == EHairPatchAttribute::ClusterInfluence || Mode == EHairPatchAttribute::GuideInflucence)
 	if (Mode == EHairPatchAttribute::GuideInflucence && !bSimulation) return;
 
-	const FIntVector DispatchCount = FIntVector(FMath::DivideAndRoundUp(VertexCount, 1024u), 1, 1);
+	const FIntVector DispatchCount = FIntVector(FMath::DivideAndRoundUp(CurveCount, 1024u), 1, 1);
+
+	// First copy all the rendering attributes into the new bufffer, before overriding some of them
+	AddCopyBufferPass(GraphBuilder, RenAttributeBuffer, OutRenAttributeBuffer.Buffer);
 
 	FHairPatchAttributeCS::FParameters* Parameters = GraphBuilder.AllocParameters<FHairPatchAttributeCS::FParameters>();
-	Parameters->VertexCount = VertexCount;
-	Parameters->RenAttribute0Buffer = RenAttribute0Buffer;
+	Parameters->CurveCount = CurveCount;
+	Parameters->RenCurveBuffer = RenCurveBuffer;
 	Parameters->RenVertexToClusterIdBuffer = RenVertexToClusterIdBuffer;
+	PACK_HAIR_ATTRIBUTE_OFFSETS(Parameters->RenAttributeOffsets, RenAttributeOffests)
+
 	if (bSimulation)
 	{
-		Parameters->SimAttribute0Buffer  = SimAttribute0Buffer;
 		Parameters->InterpolationBuffer  = InterpolationBuffer;
 		Parameters->Interpolation0Buffer = Interpolation0Buffer;
 		Parameters->Interpolation1Buffer = Interpolation1Buffer;
 	}
-	Parameters->OutRenAttribute0Buffer = OutRenAttribute0Buffer.UAV;
+
+	Parameters->OutRenAttributeBuffer = OutRenAttributeBuffer.UAV;
 
 	FHairPatchAttributeCS::FPermutationDomain PermutationVector;
 	PermutationVector.Set<FHairPatchAttributeCS::FSimulation>(Mode == EHairPatchAttribute::ClusterInfluence ? 0 : (bUseSingleGuide ? 1 : 2));
@@ -854,7 +848,7 @@ static void AddPatchAttributePass(
 		Parameters,
 		DispatchCount);
 
-	GraphBuilder.SetBufferAccessFinal(OutRenAttribute0Buffer.Buffer, ERHIAccess::SRVMask);
+	GraphBuilder.SetBufferAccessFinal(OutRenAttributeBuffer.Buffer, ERHIAccess::SRVMask);
 }
 
 ///////////////////////////////////////////////////////////////////////////////////////////////////
@@ -1106,13 +1100,12 @@ static void AddHairCardsDeformationPass(
 
 	Parameters->CardsInterpolationBuffer	= RegisterAsSRV(GraphBuilder, LOD.InterpolationResource->InterpolationBuffer);
 
+	const FHairStrandsRestResource* RestResources = LOD.Guides.RestResource;
 	const FHairStrandsRestRootResource* RestRootResources = LOD.Guides.RestRootResource;
 	const FHairStrandsDeformedRootResource* DeformedRootResources = LOD.Guides.DeformedRootResource;
 
-	const bool bIsVertexToCurveBuffersValid = RestRootResources && RestRootResources->VertexToCurveIndexBuffer.Buffer != nullptr;
 	const bool bSupportDynamicMesh =
 		Instance->BindingType == EHairBindingType::Skinning &&
-		bIsVertexToCurveBuffersValid &&
 		RestRootResources &&
 		RestRootResources->BulkData.RootCount > 0 &&
 		MeshLODIndex >= 0 &&
@@ -1122,7 +1115,7 @@ static void AddHairCardsDeformationPass(
 		DeformedRootResources->LODs[MeshLODIndex].IsValid();
 	if (bSupportDynamicMesh)
 	{
-		Parameters->GuideVertexToRootIndexBuffer = RegisterAsSRV(GraphBuilder, RestRootResources->VertexToCurveIndexBuffer);
+		Parameters->GuideVertexToRootIndexBuffer = RegisterAsSRV(GraphBuilder, RestResources->VertexToCurveBuffer);
 		const FHairStrandsRestRootResource::FLOD& RestLODDatas = RestRootResources->LODs[MeshLODIndex];
 		const FHairStrandsDeformedRootResource::FLOD& DeformedLODDatas = DeformedRootResources->LODs[MeshLODIndex];
 
@@ -1633,10 +1626,14 @@ static FHairGroupPublicData::FVertexFactoryInput InternalComputeHairStrandsVerte
 {
 	FHairGroupPublicData::FVertexFactoryInput OutVFInput;
 	if (!Instance || Instance->GeometryType != EHairGeometryType::Strands)
+	{
+		OutVFInput.Strands.AttributeOffsets.Init(0xFFFFFFFF, HAIR_ATTRIBUTE_COUNT);
 		return OutVFInput;
+	}
 
 	const bool bSupportHasGuides = Instance->Guides.IsValid();
 	const bool bSupportDeformation = Instance->Strands.DeformedResource != nullptr;
+	bool bUseGuideAttributeOffsets = false;
 
 	// 1. Guides
 	if (ViewMode == EGroomViewMode::SimHairStrands && bSupportHasGuides)
@@ -1644,10 +1641,9 @@ static FHairGroupPublicData::FVertexFactoryInput InternalComputeHairStrandsVerte
 		CONVERT_HAIRSSTRANDS_VF_PARAMETERS(OutVFInput.Strands.PositionBuffer,			Instance->Guides.DeformedResource->GetBuffer(FHairStrandsDeformedResource::EFrameType::Current));
 		CONVERT_HAIRSSTRANDS_VF_PARAMETERS(OutVFInput.Strands.PrevPositionBuffer,		Instance->Guides.DeformedResource->GetBuffer(FHairStrandsDeformedResource::EFrameType::Previous));
 		CONVERT_HAIRSSTRANDS_VF_PARAMETERS(OutVFInput.Strands.TangentBuffer,			Instance->Guides.DeformedResource->TangentBuffer);
-		CONVERT_HAIRSSTRANDS_VF_PARAMETERS(OutVFInput.Strands.Attribute0Buffer,			Instance->Guides.RestResource->Attribute0Buffer);
-		CONVERT_HAIRSSTRANDS_VF_PARAMETERS(OutVFInput.Strands.Attribute1Buffer,			Instance->Guides.RestResource->Attribute1Buffer);
-		CONVERT_HAIRSSTRANDS_VF_PARAMETERS(OutVFInput.Strands.MaterialBuffer,			Instance->Guides.RestResource->MaterialBuffer);
+		CONVERT_HAIRSSTRANDS_VF_PARAMETERS(OutVFInput.Strands.AttributeBuffer,			Instance->Guides.RestResource->AttributeBuffer);
 		CONVERT_HAIRSSTRANDS_VF_PARAMETERS(OutVFInput.Strands.CurveBuffer,				Instance->Guides.RestResource->CurveBuffer);
+		CONVERT_HAIRSSTRANDS_VF_PARAMETERS(OutVFInput.Strands.VertexToCurveBuffer,		Instance->Guides.RestResource->VertexToCurveBuffer);
 		CONVERT_HAIRSSTRANDS_VF_PARAMETERS(OutVFInput.Strands.PositionOffsetBuffer,		Instance->Guides.DeformedResource->GetPositionOffsetBuffer(FHairStrandsDeformedResource::EFrameType::Current));
 		CONVERT_HAIRSSTRANDS_VF_PARAMETERS(OutVFInput.Strands.PrevPositionOffsetBuffer,	Instance->Guides.DeformedResource->GetPositionOffsetBuffer(FHairStrandsDeformedResource::EFrameType::Previous));
 
@@ -1655,6 +1651,8 @@ static FHairGroupPublicData::FVertexFactoryInput InternalComputeHairStrandsVerte
 		OutVFInput.Strands.PrevPositionOffset = Instance->Guides.DeformedResource->GetPositionOffset(FHairStrandsDeformedResource::EFrameType::Previous);
 		OutVFInput.Strands.VertexCount = Instance->Guides.RestResource->GetVertexCount();
 		OutVFInput.Strands.CurveCount = Instance->Guides.RestResource->GetCurveCount();
+
+		bUseGuideAttributeOffsets = true;
 	}
 	// 2. Render Strands with deformation
 	else if (bSupportDeformation)
@@ -1683,9 +1681,8 @@ static FHairGroupPublicData::FVertexFactoryInput InternalComputeHairStrandsVerte
 			OutVFInput.Strands.PrevPositionOffset = Instance->Strands.DeformedResource->GetPositionOffset(FHairStrandsDeformedResource::EFrameType::Current);
 		}
 		CONVERT_HAIRSSTRANDS_VF_PARAMETERS(OutVFInput.Strands.TangentBuffer,			Instance->Strands.DeformedResource->TangentBuffer);
-		CONVERT_HAIRSSTRANDS_VF_PARAMETERS(OutVFInput.Strands.Attribute0Buffer,			(ViewMode == EGroomViewMode::RenderHairStrands || ViewMode == EGroomViewMode::RenderVisCluster || ViewMode == EGroomViewMode::RenderVisClusterAABB) ? Instance->Strands.DebugAttributeBuffer : Instance->Strands.RestResource->Attribute0Buffer);
-		CONVERT_HAIRSSTRANDS_VF_PARAMETERS(OutVFInput.Strands.Attribute1Buffer,			Instance->Strands.RestResource->Attribute1Buffer);
-		CONVERT_HAIRSSTRANDS_VF_PARAMETERS(OutVFInput.Strands.MaterialBuffer,			Instance->Strands.RestResource->MaterialBuffer);
+		CONVERT_HAIRSSTRANDS_VF_PARAMETERS(OutVFInput.Strands.AttributeBuffer,			(ViewMode == EGroomViewMode::RenderHairStrands || ViewMode == EGroomViewMode::RenderVisCluster || ViewMode == EGroomViewMode::RenderVisClusterAABB) ? Instance->Strands.DebugAttributeBuffer : Instance->Strands.RestResource->AttributeBuffer);
+		CONVERT_HAIRSSTRANDS_VF_PARAMETERS(OutVFInput.Strands.VertexToCurveBuffer,		Instance->Strands.RestResource->VertexToCurveBuffer);
 		CONVERT_HAIRSSTRANDS_VF_PARAMETERS(OutVFInput.Strands.CurveBuffer,				Instance->Strands.RestResource->CurveBuffer);
 
 		OutVFInput.Strands.VertexCount = Instance->Strands.RestResource->GetVertexCount();
@@ -1697,9 +1694,8 @@ static FHairGroupPublicData::FVertexFactoryInput InternalComputeHairStrandsVerte
 		CONVERT_HAIRSSTRANDS_VF_PARAMETERS(OutVFInput.Strands.PositionBuffer,			Instance->Strands.RestResource->PositionBuffer);
 		CONVERT_HAIRSSTRANDS_VF_PARAMETERS(OutVFInput.Strands.PrevPositionBuffer,		Instance->Strands.RestResource->PositionBuffer);
 		CONVERT_HAIRSSTRANDS_VF_PARAMETERS(OutVFInput.Strands.TangentBuffer,			Instance->Strands.RestResource->TangentBuffer);
-		CONVERT_HAIRSSTRANDS_VF_PARAMETERS(OutVFInput.Strands.Attribute0Buffer,			(ViewMode == EGroomViewMode::RenderVisCluster || ViewMode == EGroomViewMode::RenderVisClusterAABB) ? Instance->Strands.DebugAttributeBuffer : Instance->Strands.RestResource->Attribute0Buffer);
-		CONVERT_HAIRSSTRANDS_VF_PARAMETERS(OutVFInput.Strands.Attribute1Buffer,			Instance->Strands.RestResource->Attribute1Buffer);
-		CONVERT_HAIRSSTRANDS_VF_PARAMETERS(OutVFInput.Strands.MaterialBuffer,			Instance->Strands.RestResource->MaterialBuffer);
+		CONVERT_HAIRSSTRANDS_VF_PARAMETERS(OutVFInput.Strands.AttributeBuffer,			(ViewMode == EGroomViewMode::RenderVisCluster || ViewMode == EGroomViewMode::RenderVisClusterAABB) ? Instance->Strands.DebugAttributeBuffer : Instance->Strands.RestResource->AttributeBuffer);
+		CONVERT_HAIRSSTRANDS_VF_PARAMETERS(OutVFInput.Strands.VertexToCurveBuffer,		Instance->Strands.RestResource->VertexToCurveBuffer);
 		CONVERT_HAIRSSTRANDS_VF_PARAMETERS(OutVFInput.Strands.CurveBuffer,				Instance->Strands.RestResource->CurveBuffer);
 		CONVERT_HAIRSSTRANDS_VF_PARAMETERS(OutVFInput.Strands.PositionOffsetBuffer,		Instance->Strands.RestResource->PositionOffsetBuffer);
 		CONVERT_HAIRSSTRANDS_VF_PARAMETERS(OutVFInput.Strands.PrevPositionOffsetBuffer, Instance->Strands.RestResource->PositionOffsetBuffer);
@@ -1719,6 +1715,11 @@ static FHairGroupPublicData::FVertexFactoryInput InternalComputeHairStrandsVerte
 	OutVFInput.Strands.bScatterSceneLighting = Instance->Strands.Modifier.bScatterSceneLighting;
 	OutVFInput.Strands.bUseStableRasterization = Instance->Strands.Modifier.bUseStableRasterization;
 	OutVFInput.Strands.bUseRaytracingGeometry = false;
+
+	OutVFInput.Strands.AttributeOffsets.SetNum(0, false);
+	OutVFInput.Strands.AttributeOffsets.Reserve(HAIR_ATTRIBUTE_COUNT);
+	OutVFInput.Strands.AttributeOffsets.Append(bUseGuideAttributeOffsets ? Instance->Guides.Data->AttributeOffsets : Instance->Strands.Data->AttributeOffsets);
+
 #if RHI_RAYTRACING
 	// Flag bUseRaytracingGeometry only if RT geometry has been allocated for RayTracing view (not for PathTracing view). 
 	// This flag is used later for selecting if voxelization needs to flags voxel has shadow caster or if shadow casting is handled 
@@ -1894,6 +1895,7 @@ void ComputeHairStrandsInterpolation(
 									Instance->Guides.RestRootResource,
 									Instance->Guides.DeformedRootResource,
 									RegisterAsSRV(GraphBuilder, Instance->Guides.RestResource->PositionBuffer),
+									RegisterAsSRV(GraphBuilder, Instance->Guides.RestResource->VertexToCurveBuffer),
 									Register(GraphBuilder, Instance->Guides.DeformedResource->GetBuffer(FHairStrandsDeformedResource::Current), ERDGImportedBufferFlags::CreateUAV),
 									Instance->Guides.RestResource->GetPositionOffset(),
 									RegisterAsSRV(GraphBuilder, Instance->Guides.DeformedResource->GetPositionOffsetBuffer(FHairStrandsDeformedResource::Current)),
@@ -1935,6 +1937,7 @@ void ComputeHairStrandsInterpolation(
 						bHasSkinning ? Instance->Strands.DeformedRootResource : nullptr,
 						bHasSkinning && bValidGuide ? Instance->Guides.DeformedRootResource : nullptr,
 						RegisterAsSRV(GraphBuilder, Instance->Strands.RestResource->PositionBuffer),
+						RegisterAsSRV(GraphBuilder, Instance->Strands.RestResource->VertexToCurveBuffer),
 						bUseSingleGuide,
 						bValidGuide &&  bUseSingleGuide ? RegisterAsSRV(GraphBuilder, Instance->Strands.InterpolationResource->InterpolationBuffer) : nullptr,
 						bValidGuide && !bUseSingleGuide ? RegisterAsSRV(GraphBuilder, Instance->Strands.InterpolationResource->Interpolation0Buffer) : nullptr,
@@ -1942,6 +1945,7 @@ void ComputeHairStrandsInterpolation(
 						bValidGuide ? RegisterAsSRV(GraphBuilder, Instance->Guides.RestResource->PositionBuffer) : nullptr,
 						bValidGuide ? RegisterAsSRV(GraphBuilder, Instance->Guides.DeformedResource->GetBuffer(FHairStrandsDeformedResource::Current)) : nullptr,
 						bValidGuide ? RegisterAsSRV(GraphBuilder, Instance->Strands.InterpolationResource->SimRootPointIndexBuffer) : nullptr,
+						bValidGuide ? RegisterAsSRV(GraphBuilder, Instance->Guides.RestResource->VertexToCurveBuffer) : nullptr,
 						Strands_DeformerPositionSRV,
 						Strands_DeformedPosition,
 						FHairStrandsDeformedRootResource::FLOD::Current);
@@ -1984,6 +1988,7 @@ void ComputeHairStrandsInterpolation(
 							bHasSkinning ? Instance->Strands.DeformedRootResource : nullptr,
 							bHasSkinning && bValidGuide ? Instance->Guides.DeformedRootResource : nullptr,
 							RegisterAsSRV(GraphBuilder, Instance->Strands.RestResource->PositionBuffer),
+							RegisterAsSRV(GraphBuilder, Instance->Strands.RestResource->VertexToCurveBuffer),
 							bUseSingleGuide,
 							bValidGuide &&  bUseSingleGuide ? RegisterAsSRV(GraphBuilder, Instance->Strands.InterpolationResource->InterpolationBuffer) : nullptr,
 							bValidGuide && !bUseSingleGuide ? RegisterAsSRV(GraphBuilder, Instance->Strands.InterpolationResource->Interpolation0Buffer) : nullptr,
@@ -1991,6 +1996,7 @@ void ComputeHairStrandsInterpolation(
 							bValidGuide ? RegisterAsSRV(GraphBuilder, Instance->Guides.RestResource->PositionBuffer) : nullptr,
 							bValidGuide ? RegisterAsSRV(GraphBuilder, Instance->Guides.DeformedResource->GetBuffer(FHairStrandsDeformedResource::Previous)) : nullptr,
 							bValidGuide ? RegisterAsSRV(GraphBuilder, Instance->Strands.InterpolationResource->SimRootPointIndexBuffer) : nullptr,
+							bValidGuide ? RegisterAsSRV(GraphBuilder, Instance->Guides.RestResource->VertexToCurveBuffer) : nullptr,
 							Strands_DeformerPositionSRV,
 							Strands_DeformedPrevPosition,
 							FHairStrandsDeformedRootResource::FLOD::Previous);
@@ -2054,24 +2060,25 @@ void ComputeHairStrandsInterpolation(
 				// Special case for debug mode were the attribute buffer is patch with some custom data to show hair properties (strands belonging to the same cluster, ...)
 				if (Instance->Strands.DebugAttributeBuffer.Buffer == nullptr)
 				{
-					CreateHairStrandsDebugAttributeBuffer(GraphBuilder, &Instance->Strands.DebugAttributeBuffer, Instance->Strands.Data->GetNumPoints());
+					CreateHairStrandsDebugAttributeBuffer(GraphBuilder, &Instance->Strands.DebugAttributeBuffer, Instance->Strands.Data->Attributes.GetBulkDataSize());
 				}
-				FRDGImportedBuffer OutRenAttribute0Buffer = Register(GraphBuilder, Instance->Strands.DebugAttributeBuffer, ERDGImportedBufferFlags::CreateUAV);
+				FRDGImportedBuffer OutRenAttributeBuffer = Register(GraphBuilder, Instance->Strands.DebugAttributeBuffer, ERDGImportedBufferFlags::CreateUAV);
 
 				AddPatchAttributePass(
 					GraphBuilder,
 					ShaderMap,
-					Instance->Strands.RestResource->GetVertexCount(),
+					Instance->Strands.RestResource->GetCurveCount(),
 					PatchMode,
 					bValidGuide,
 					bUseSingleGuide,
-					RegisterAsSRV(GraphBuilder, Instance->Strands.RestResource->Attribute0Buffer),
+					Instance->Strands.Data->AttributeOffsets,
+					Register(GraphBuilder, Instance->Strands.RestResource->AttributeBuffer, ERDGImportedBufferFlags::CreateSRV).Buffer,
+					RegisterAsSRV(GraphBuilder, Instance->Strands.RestResource->CurveBuffer),
 					RegisterAsSRV(GraphBuilder, Instance->Strands.ClusterCullingResource->VertexToClusterIdBuffer),
-					bValidGuide ? RegisterAsSRV(GraphBuilder, Instance->Guides.RestResource->Attribute0Buffer) : nullptr,
 					bValidGuide &&  bUseSingleGuide ? RegisterAsSRV(GraphBuilder, Instance->Strands.InterpolationResource->InterpolationBuffer) : nullptr,
 					bValidGuide && !bUseSingleGuide ? RegisterAsSRV(GraphBuilder, Instance->Strands.InterpolationResource->Interpolation0Buffer) : nullptr,
 					bValidGuide && !bUseSingleGuide ? RegisterAsSRV(GraphBuilder, Instance->Strands.InterpolationResource->Interpolation1Buffer) : nullptr,
-					OutRenAttribute0Buffer);
+					OutRenAttributeBuffer);
 			}
 			#endif
 
@@ -2264,6 +2271,7 @@ void ComputeHairStrandsInterpolation(
 						bHasSkinning ? LOD.Guides.DeformedRootResource : nullptr,
 						bHasSkinning && bValidGuide ? Instance->Guides.DeformedRootResource : nullptr,
 						RegisterAsSRV(GraphBuilder, LOD.Guides.RestResource->PositionBuffer),
+						RegisterAsSRV(GraphBuilder, LOD.Guides.RestResource->VertexToCurveBuffer),
 						bUseSingleGuide,
 						bValidGuide &&  bUseSingleGuide ? RegisterAsSRV(GraphBuilder, LOD.Guides.InterpolationResource->InterpolationBuffer) : nullptr,
 						bValidGuide && !bUseSingleGuide ? RegisterAsSRV(GraphBuilder, LOD.Guides.InterpolationResource->Interpolation0Buffer) : nullptr,
@@ -2271,6 +2279,7 @@ void ComputeHairStrandsInterpolation(
 						bValidGuide ? RegisterAsSRV(GraphBuilder, Instance->Guides.RestResource->PositionBuffer) : nullptr,
 						bValidGuide ? RegisterAsSRV(GraphBuilder, Instance->Guides.DeformedResource->GetBuffer(FHairStrandsDeformedResource::Current)) : nullptr,
 						RegisterAsSRV(GraphBuilder, LOD.Guides.InterpolationResource->SimRootPointIndexBuffer),
+						bValidGuide ? RegisterAsSRV(GraphBuilder, Instance->Guides.RestResource->VertexToCurveBuffer) : nullptr,
 						nullptr,
 						Guides_DeformedPositionBuffer,
 						FHairStrandsDeformedRootResource::FLOD::Current); // <- this should be optional
@@ -2446,6 +2455,7 @@ void ResetHairStrandsInterpolation(
 		Instance->Guides.RestRootResource,
 		Instance->Guides.DeformedRootResource,
 		RegisterAsSRV(GraphBuilder, Instance->Guides.RestResource->PositionBuffer),
+		RegisterAsSRV(GraphBuilder, Instance->Guides.RestResource->VertexToCurveBuffer),
 		DeformedPositionBuffer,
 		Instance->Guides.RestResource->GetPositionOffset(),
 		RegisterAsSRV(GraphBuilder, Instance->Guides.DeformedResource->GetPositionOffsetBuffer(FHairStrandsDeformedResource::Current)),
diff --git a/Engine/Plugins/Runtime/HairStrands/Source/HairStrandsCore/Private/HairStrandsVertexFactory.cpp b/Engine/Plugins/Runtime/HairStrands/Source/HairStrandsCore/Private/HairStrandsVertexFactory.cpp
--- a/Engine/Plugins/Runtime/HairStrands/Source/HairStrandsCore/Private/HairStrandsVertexFactory.cpp
+++ b/Engine/Plugins/Runtime/HairStrands/Source/HairStrandsCore/Private/HairStrandsVertexFactory.cpp
@@ -38,23 +38,41 @@ public:
 	FShaderResourceViewRHIRef SRVRGBA;
 	FShaderResourceViewRHIRef SRVRGBA_Uint;
 
+	FShaderResourceViewRHIRef SRVByteAddress;
+	FBufferRHIRef ByteAddressBufferRHI;
+
 	virtual void InitRHI() override
 	{
-		FRHIResourceCreateInfo CreateInfo(TEXT("FDummyCulledDispatchVertexIdsBuffer"));
-		uint32 NumBytes = sizeof(uint32) * 4;
-		VertexBufferRHI = RHICreateBuffer(NumBytes, BUF_Static | BUF_VertexBuffer | BUF_ShaderResource, 0, ERHIAccess::VertexOrIndexBuffer | ERHIAccess::SRVMask, CreateInfo);
-		uint32* DummyContents = (uint32*)RHILockBuffer(VertexBufferRHI, 0, NumBytes, RLM_WriteOnly);
-		DummyContents[0] = DummyContents[1] = DummyContents[2] = DummyContents[3] = 0;
-		RHIUnlockBuffer(VertexBufferRHI);
+		{
+			FRHIResourceCreateInfo CreateInfo(TEXT("FDummyCulledDispatchVertexIdsBuffer"));
+			uint32 NumBytes = sizeof(uint32) * 4;
+			VertexBufferRHI = RHICreateBuffer(NumBytes, BUF_Static | BUF_VertexBuffer | BUF_ShaderResource, 0, ERHIAccess::VertexOrIndexBuffer | ERHIAccess::SRVMask, CreateInfo);
+			uint32* DummyContents = (uint32*)RHILockBuffer(VertexBufferRHI, 0, NumBytes, RLM_WriteOnly);
+			DummyContents[0] = DummyContents[1] = DummyContents[2] = DummyContents[3] = 0;
+			RHIUnlockBuffer(VertexBufferRHI);
+		}
+
+		{
+			FRHIResourceCreateInfo CreateInfo(TEXT("FDummyByteAddressBufferBuffer"));
+			uint32 NumBytes = sizeof(uint32) * 4;
+			ByteAddressBufferRHI = RHICreateBuffer(NumBytes, BUF_Static | BUF_ShaderResource |BUF_ByteAddressBuffer, 0, ERHIAccess::SRVMask, CreateInfo);
+			uint32* DummyContents = (uint32*)RHILockBuffer(ByteAddressBufferRHI, 0, NumBytes, RLM_WriteOnly);
+			DummyContents[0] = DummyContents[1] = DummyContents[2] = DummyContents[3] = 0;
+			RHIUnlockBuffer(ByteAddressBufferRHI);
+		}
 
 		SRVUint = RHICreateShaderResourceView(VertexBufferRHI, sizeof(uint32), PF_R32_UINT);
 		SRVFloat = RHICreateShaderResourceView(VertexBufferRHI, sizeof(uint32), PF_R32_FLOAT);
 		SRVRGBA = RHICreateShaderResourceView(VertexBufferRHI, sizeof(uint32), PF_R8G8B8A8);
 		SRVRGBA_Uint = RHICreateShaderResourceView(VertexBufferRHI, sizeof(uint32), PF_R8G8B8A8_UINT);
+		SRVByteAddress = RHICreateShaderResourceView(ByteAddressBufferRHI, sizeof(uint32), PF_R32_UINT);
 	}
 
 	virtual void ReleaseRHI() override
 	{
+		ByteAddressBufferRHI.SafeRelease();
+		SRVByteAddress.SafeRelease();
+
 		VertexBufferRHI.SafeRelease();
 		SRVUint.SafeRelease();
 		SRVFloat.SafeRelease();
@@ -85,10 +103,8 @@ FHairStrandsVertexFactoryUniformShaderParameters FHairGroupInstance::GetHairStan
 	Out.ScatterSceneLighing			= VFInput.Strands.bScatterSceneLighting;
 	Out.PositionBuffer				= VFInput.Strands.PositionBufferRHISRV;
 	Out.PreviousPositionBuffer		= VFInput.Strands.PrevPositionBufferRHISRV;
-	Out.Attribute0Buffer			= VFInput.Strands.Attribute0BufferRHISRV;
-	Out.Attribute1Buffer			= VFInput.Strands.Attribute1BufferRHISRV;
-	Out.MaterialBuffer 				= VFInput.Strands.MaterialBufferRHISRV;
-	Out.HasMaterial 				= Out.MaterialBuffer != nullptr;
+	Out.AttributeBuffer				= VFInput.Strands.AttributeBufferRHISRV;
+	Out.VertexToCurveBuffer			= VFInput.Strands.VertexToCurveBufferRHISRV;
 	Out.TangentBuffer 				= VFInput.Strands.TangentBufferRHISRV;
 	Out.PositionOffsetBuffer 		= VFInput.Strands.PositionOffsetBufferRHISRV;
 	Out.PreviousPositionOffsetBuffer= VFInput.Strands.PrevPositionOffsetBufferRHISRV;
@@ -96,9 +112,8 @@ FHairStrandsVertexFactoryUniformShaderParameters FHairGroupInstance::GetHairStan
 	// swap in some default data for those buffers that are not valid yet
 	if (!Out.PositionBuffer) 				{ Out.PositionBuffer = GDummyCulledDispatchVertexIdsBuffer.SRVFloat; }
 	if (!Out.PreviousPositionBuffer) 		{ Out.PreviousPositionBuffer = GDummyCulledDispatchVertexIdsBuffer.SRVFloat; }
-	if (!Out.Attribute0Buffer) 				{ Out.Attribute0Buffer = GDummyCulledDispatchVertexIdsBuffer.SRVFloat; }
-	if (!Out.Attribute1Buffer) 				{ Out.Attribute1Buffer = GDummyCulledDispatchVertexIdsBuffer.SRVFloat; }
-	if (!Out.MaterialBuffer) 				{ Out.MaterialBuffer = GDummyCulledDispatchVertexIdsBuffer.SRVRGBA; }
+	if (!Out.AttributeBuffer) 				{ Out.AttributeBuffer = GDummyCulledDispatchVertexIdsBuffer.SRVByteAddress; }
+	if (!Out.VertexToCurveBuffer) 			{ Out.VertexToCurveBuffer = GDummyCulledDispatchVertexIdsBuffer.SRVUint; }
 	if (!Out.TangentBuffer) 				{ Out.TangentBuffer = GDummyCulledDispatchVertexIdsBuffer.SRVFloat; }
 	if (!Out.PositionOffsetBuffer) 			{ Out.PositionOffsetBuffer = GDummyCulledDispatchVertexIdsBuffer.SRVFloat; }
 	if (!Out.PreviousPositionOffsetBuffer) 	{ Out.PreviousPositionOffsetBuffer = GDummyCulledDispatchVertexIdsBuffer.SRVFloat; }
@@ -114,6 +129,8 @@ FHairStrandsVertexFactoryUniformShaderParameters FHairGroupInstance::GetHairStan
 		Out.CulledVertexIdsBuffer = GDummyCulledDispatchVertexIdsBuffer.SRVUint;
 		Out.CulledVertexRadiusScaleBuffer = GDummyCulledDispatchVertexIdsBuffer.SRVFloat;
 	}
+
+	PACK_HAIR_ATTRIBUTE_OFFSETS(Out.AttributeOffsets, VFInput.Strands.AttributeOffsets);	
 	return Out;
 }
 
diff --git a/Engine/Plugins/Runtime/HairStrands/Source/HairStrandsCore/Public/GroomAsset.h b/Engine/Plugins/Runtime/HairStrands/Source/HairStrandsCore/Public/GroomAsset.h
--- a/Engine/Plugins/Runtime/HairStrands/Source/HairStrandsCore/Public/GroomAsset.h
+++ b/Engine/Plugins/Runtime/HairStrands/Source/HairStrandsCore/Public/GroomAsset.h
@@ -93,11 +93,10 @@ struct HAIRSTRANDSCORE_API FHairGroupData
 		uint32 GetDataSize() const
 		{
 			uint32 Total = 0;
-			Total += BulkData.Positions.IsBulkDataLoaded() 	? BulkData.Positions.GetBulkDataSize()   : 0;
-			Total += BulkData.Attributes0.IsBulkDataLoaded()? BulkData.Attributes0.GetBulkDataSize() : 0;
-			Total += BulkData.Attributes1.IsBulkDataLoaded()? BulkData.Attributes1.GetBulkDataSize() : 0;
-			Total += BulkData.Materials.IsBulkDataLoaded() 	? BulkData.Materials.GetBulkDataSize()   : 0;
-			Total += BulkData.Curves.IsBulkDataLoaded() 	? BulkData.Curves.GetBulkDataSize(): 0;
+			Total += BulkData.Positions.IsBulkDataLoaded() 		? BulkData.Positions.GetBulkDataSize()   : 0;
+			Total += BulkData.Attributes.IsBulkDataLoaded()		? BulkData.Attributes.GetBulkDataSize() : 0;
+			Total += BulkData.Curves.IsBulkDataLoaded() 		? BulkData.Curves.GetBulkDataSize(): 0;
+			Total += BulkData.VertexToCurve.IsBulkDataLoaded()	? BulkData.VertexToCurve.GetBulkDataSize() : 0;
 			return Total;
 		}
 
diff --git a/Engine/Plugins/Runtime/HairStrands/Source/HairStrandsCore/Public/GroomInstance.h b/Engine/Plugins/Runtime/HairStrands/Source/HairStrandsCore/Public/GroomInstance.h
--- a/Engine/Plugins/Runtime/HairStrands/Source/HairStrandsCore/Public/GroomInstance.h
+++ b/Engine/Plugins/Runtime/HairStrands/Source/HairStrandsCore/Public/GroomInstance.h
@@ -8,6 +8,7 @@
 #include "GroomAsset.h"
 #include "GroomBindingAsset.h"
 #include "HairStrandsInterface.h"
+#include "HairStrandsDefinitions.h"
 
 class UMeshComponent;
 class UGroomComponent;
@@ -58,6 +59,7 @@ BEGIN_GLOBAL_SHADER_PARAMETER_STRUCT(FHairStrandsVertexFactoryUniformShaderParam
 	SHADER_PARAMETER(uint32, ScatterSceneLighing)
 	SHADER_PARAMETER(uint32, RaytracingProceduralSplits)
 	SHADER_PARAMETER(float, GroupIndex)
+	SHADER_PARAMETER_ARRAY(FUintVector4, AttributeOffsets, [HAIR_ATTRIBUTE_OFFSET_COUNT])
 
 	SHADER_PARAMETER_SRV(Buffer<float4>, PositionOffsetBuffer)
 	SHADER_PARAMETER_SRV(Buffer<float4>, PreviousPositionOffsetBuffer)
@@ -65,9 +67,8 @@ BEGIN_GLOBAL_SHADER_PARAMETER_STRUCT(FHairStrandsVertexFactoryUniformShaderParam
 	SHADER_PARAMETER_SRV(Buffer<uint4>, PositionBuffer)
 	SHADER_PARAMETER_SRV(Buffer<uint4>, PreviousPositionBuffer)
 
-	SHADER_PARAMETER_SRV(Buffer<float2>, Attribute0Buffer)
-	SHADER_PARAMETER_SRV(Buffer<uint>,   Attribute1Buffer)
-	SHADER_PARAMETER_SRV(Buffer<float4>, MaterialBuffer)
+	SHADER_PARAMETER_SRV(ByteAddressBuffer, AttributeBuffer)
+	SHADER_PARAMETER_SRV(Buffer<uint>, VertexToCurveBuffer)
 	SHADER_PARAMETER_SRV(Buffer<float4>, TangentBuffer)
 
 	SHADER_PARAMETER_SRV(Buffer<uint>, CulledVertexIdsBuffer)
diff --git a/Engine/Plugins/Runtime/HairStrands/Source/HairStrandsCore/Public/GroomResources.h b/Engine/Plugins/Runtime/HairStrands/Source/HairStrandsCore/Public/GroomResources.h
--- a/Engine/Plugins/Runtime/HairStrands/Source/HairStrandsCore/Public/GroomResources.h
+++ b/Engine/Plugins/Runtime/HairStrands/Source/HairStrandsCore/Public/GroomResources.h
@@ -129,7 +129,6 @@ struct FHairStrandsRestRootResource : public FHairCommonResource
 	uint32 GetResourcesSize() const 
 	{
 		uint32 Total = 0;
-		Total += GetBufferTotalNumBytes(VertexToCurveIndexBuffer);
 		for (const FLOD& LOD : LODs)
 		{
 			Total += GetBufferTotalNumBytes(LOD.UniqueTriangleIndexBuffer);
@@ -145,8 +144,6 @@ struct FHairStrandsRestRootResource : public FHairCommonResource
 		return Total;
 	}
 
-	FRDGExternalBuffer VertexToCurveIndexBuffer;
-
 	struct FLOD
 	{
 		enum class EStatus { Invalid, Initialized, Completed };
@@ -184,9 +181,6 @@ struct FHairStrandsRestRootResource : public FHairCommonResource
 	/* Store CPU data for root info & root binding */
 	FHairStrandsRootBulkData& BulkData;
 
-	/* Bulk data request handle */
-	FBulkDataBatchRequest BulkDataRequest;
-
 	/* Type of curves */
 	const EHairStrandsResourcesType CurveType;
 };
@@ -318,9 +312,8 @@ struct FHairStrandsRestResource : public FHairCommonResource
 		Total += GetBufferTotalNumBytes(PositionBuffer);
 		Total += GetBufferTotalNumBytes(PositionOffsetBuffer);
 		Total += GetBufferTotalNumBytes(TangentBuffer);
-		Total += GetBufferTotalNumBytes(Attribute0Buffer);
-		Total += GetBufferTotalNumBytes(Attribute1Buffer);
-		Total += GetBufferTotalNumBytes(MaterialBuffer);
+		Total += GetBufferTotalNumBytes(AttributeBuffer);
+		Total += GetBufferTotalNumBytes(VertexToCurveBuffer);
 		Total += GetBufferTotalNumBytes(CurveBuffer);
 		return Total;
 	}
@@ -335,15 +328,12 @@ struct FHairStrandsRestResource : public FHairCommonResource
 	FRDGExternalBuffer TangentBuffer;
 
 	/* Strand hair attribute buffer */
-	FRDGExternalBuffer Attribute0Buffer;
-
-	/* Strand hair attribute buffer */
-	FRDGExternalBuffer Attribute1Buffer;
+	FRDGExternalBuffer AttributeBuffer;
 
-	/* Strand hair material buffer */
-	FRDGExternalBuffer MaterialBuffer;
+	/* Strand hair vertex to curve index mapping */
+	FRDGExternalBuffer VertexToCurveBuffer;
 
-	/* Strand hair curves buffer (contains curves' points offset and count */
+	/* Strand hair curves buffer (contains curves' points offset and count) */
 	FRDGExternalBuffer CurveBuffer;
 
 	/* Reference to the hair strands render data */
diff --git a/Engine/Plugins/Runtime/HairStrands/Source/HairStrandsCore/Public/HairStrandsDatas.h b/Engine/Plugins/Runtime/HairStrands/Source/HairStrandsCore/Public/HairStrandsDatas.h
--- a/Engine/Plugins/Runtime/HairStrands/Source/HairStrandsCore/Public/HairStrandsDatas.h
+++ b/Engine/Plugins/Runtime/HairStrands/Source/HairStrandsCore/Public/HairStrandsDatas.h
@@ -8,6 +8,7 @@
 #include "PackedNormal.h"
 #include "RenderGraphResources.h"
 #include "Serialization/BulkData.h"
+#include "HairStrandsDefinitions.h"
 
 DECLARE_LOG_CATEGORY_EXTERN(LogHairStrands, Log, All);
 
@@ -28,25 +29,6 @@ struct FPackedHairAttribute0Vertex
 	uint8 Seed;
 };
 
-struct FPackedHairAttribute1Vertex
-{
-	typedef uint32 BulkType;
-
-	uint32 Packed;
-};
-
-struct FHairMaterialVertex
-{
-	typedef uint32 BulkType;
-
-	// sRGB color space
-	uint8 BaseColorR;
-	uint8 BaseColorG;
-	uint8 BaseColorB;
-
-	uint8 Roughness;
-};
-
 struct FHairInterpolationVertex
 {
 	typedef uint32 BulkType;
@@ -121,34 +103,34 @@ struct FHairStrandsPositionOffsetFormat
 	static const EPixelFormat Format = PF_A32B32G32R32F;
 };
 
-struct FHairStrandsAttribute0Format
+struct FHairStrandsAttributeFormat
 {
-	typedef FPackedHairAttribute0Vertex Type;
-	typedef FPackedHairAttribute0Vertex::BulkType BulkType;
+	typedef uint32 Type;
+	typedef uint32 BulkType;
 	static const uint32 ComponentCount = 1;
 	static const uint32 SizeInByte = sizeof(Type);
-	static const EVertexElementType VertexElementType = VET_UByte4;
-	static const EPixelFormat Format = PF_R8G8;
+	static const EVertexElementType VertexElementType = VET_UInt;
+	static const EPixelFormat Format = PF_R32_UINT;
 };
 
-struct FHairStrandsAttribute1Format
+struct FHairStrandsVertexToCurveFormat16
 {
-	typedef FPackedHairAttribute1Vertex Type;
-	typedef FPackedHairAttribute1Vertex::BulkType BulkType;
+	typedef uint16 Type;
+	typedef uint16 BulkType;
 	static const uint32 ComponentCount = 1;
 	static const uint32 SizeInByte = sizeof(Type);
-	static const EVertexElementType VertexElementType = VET_UInt;
-	static const EPixelFormat Format = PF_R32_UINT;
+	static const EVertexElementType VertexElementType = VET_MAX;
+	static const EPixelFormat Format = PF_R16_UINT;
 };
 
-struct FHairStrandsMaterialFormat
+struct FHairStrandsVertexToCurveFormat32
 {
-	typedef FHairMaterialVertex Type;
-	typedef FHairMaterialVertex::BulkType BulkType;
+	typedef uint32 Type;
+	typedef uint32 BulkType;
 	static const uint32 ComponentCount = 1;
 	static const uint32 SizeInByte = sizeof(Type);
-	static const EVertexElementType VertexElementType = VET_UByte4;
-	static const EPixelFormat Format = PF_R8G8B8A8;
+	static const EVertexElementType VertexElementType = VET_UInt;
+	static const EPixelFormat Format = PF_R32_UINT;
 };
 
 struct FHairStrandsTangentFormat
@@ -461,8 +443,8 @@ struct HAIRSTRANDSCORE_API FHairStrandsBulkData
 {
 	enum EDataFlags
 	{
-		DataFlags_HasData = 1,			// Contains valid data. Otherwise: Position, Attributes, ... are all empty
-		DataFlags_HasMaterialData = 2	// Contains material data (albedo and/or roughness)
+		DataFlags_HasData = 1,				// Contains valid data. Otherwise: Position, Attributes, ... are all empty
+		DataFlags_Has16bitsCurveIndex = 2,	// Use 16bits index for vertex to curve mapping
 	};
 
 	void Serialize(FArchive& Ar, UObject* Owner);
@@ -483,11 +465,11 @@ struct HAIRSTRANDSCORE_API FHairStrandsBulkData
 	float MaxRadius = 0;
 	FBox BoundingBox = FBox(EForceInit::ForceInit);
 	uint32 Flags = 0;
+	uint32 AttributeOffsets[HAIR_ATTRIBUTE_COUNT];
 
 	FByteBulkData Positions;	// Size = PointCount
-	FByteBulkData Attributes0;	// Size = PointCount
-	FByteBulkData Attributes1;	// Size = PointCount
-	FByteBulkData Materials;	// Size = PointCount
+	FByteBulkData Attributes;	// Size = x*PointCount + y*CurveCount (depends on the stored attributes, which could be per-vertex(x) or per-curve(y))
+	FByteBulkData VertexToCurve;// Size = PointCount
 	FByteBulkData Curves;		// Size = CurveCount
 };
 
@@ -544,7 +526,6 @@ struct FHairStrandsRootBulkData
 	uint32 GetDataSize() const
 	{
 		uint32 Total = 0;
-		Total += VertexToCurveIndexBuffer.IsBulkDataLoaded() ? VertexToCurveIndexBuffer.GetBulkDataSize() : 0u;
 		for (const FMeshProjectionLOD& LOD : MeshProjectionLODs)
 		{
 			Total += LOD.UniqueTriangleIndexBuffer.IsBulkDataLoaded() ?			LOD.UniqueTriangleIndexBuffer.GetBulkDataSize() : 0u;
@@ -601,9 +582,6 @@ struct FHairStrandsRootBulkData
 	/* Number of control points */
 	uint32 PointCount = 0;
 
-	/* Curve index for every vertices */
-	FByteBulkData VertexToCurveIndexBuffer;
-
 	/* Store the hair projection information for each mesh LOD */
 	TArray<FMeshProjectionLOD> MeshProjectionLODs;
 };
@@ -654,9 +632,6 @@ struct FHairStrandsRootData
 	/* Number of control points */
 	uint32 PointCount = 0;
 
-	/* Curve index for every vertices */
-	TArray<FHairStrandsIndexFormat::Type> VertexToCurveIndexBuffer;
-
 	/* Store the hair projection information for each mesh LOD */
 	TArray<FMeshProjectionLOD> MeshProjectionLODs;
 };
diff --git a/Engine/Shaders/Private/HairStrands/HairStrandsAttributeCommon.ush b/Engine/Shaders/Private/HairStrands/HairStrandsAttributeCommon.ush
--- a/Engine/Shaders/Private/HairStrands/HairStrandsAttributeCommon.ush
+++ b/Engine/Shaders/Private/HairStrands/HairStrandsAttributeCommon.ush
@@ -2,66 +2,79 @@
 
 #pragma once 
 
-#if HAIR_STRAND_MESH_FACTORY
-
 #include "/Engine/Private/HairStrands/HairStrandsVertexFactoryCommon.ush"
 
 ////////////////////////////////////////////////////////////////////////////////
+// Data stride in bytes
+#define HAIR_ATTRIBUTE_STRIDE_ROOTUV		4
+#define HAIR_ATTRIBUTE_STRIDE_SEED			1
+#define HAIR_ATTRIBUTE_STRIDE_LENGTH		2
+#define HAIR_ATTRIBUTE_STRIDE_BASECOLOR		4
+#define HAIR_ATTRIBUTE_STRIDE_ROUGHNESS		1
+
+// Data offset in bytes. The order needs to match HairStrandsDefinitions.h
+#define HAIR_ATTRIBUTE_OFFSET_ROOTUV(InOffsets)		(InOffsets[0].x)
+#define HAIR_ATTRIBUTE_OFFSET_SEED(InOffsets)		(InOffsets[0].y)
+#define HAIR_ATTRIBUTE_OFFSET_LENGTH(InOffsets)		(InOffsets[0].z)
+#define HAIR_ATTRIBUTE_OFFSET_BASECOLOR(InOffsets)	(InOffsets[0].w)
+#define HAIR_ATTRIBUTE_OFFSET_ROUGHNESS(InOffsets)	(InOffsets[1].x)
 
-
-struct FHairVertexAttributes
+////////////////////////////////////////////////////////////////////////////////
+// Attribute access functions
+uint InternalGetHairStrandsAttributes(uint Index, uint AttributeOffsetInBytes, uint AttributeStrideInBytes, ByteAddressBuffer InAttributeBuffer)
 {
-	float2 HairUV;
-	float2 HairRootUV;
-	float2 HairDimensions;
-	float3 HairBaseColor;
-	float  HairSeed;
-	float  HairRoughness;
-};
+	// Ensure the reading address is 4-byte aligned
+	const uint Address = AttributeOffsetInBytes + Index * AttributeStrideInBytes;
+	const uint Address_4BytesAligned = Address & (~0x3);
+	const uint Packed = InAttributeBuffer.Load(Address_4BytesAligned);
+
+	// If the data has a stride < 4bytes, compute the relative offset for reading the data
+	const uint ElementIndex = Address - Address_4BytesAligned;
+	return Packed >> (8u * ElementIndex);
+}
 
-float4 GetHairStrandsOptionalMaterial(uint InVertexId)
+uint InternalGetHairStrandsAttributesPerCurve(uint HairPrimitiveId, uint AttributeOffsetInBytes, uint AttributeStrideInBytes, ByteAddressBuffer InAttributeBuffer, Buffer<uint> InVertexToCurveBuffer)
 {
-	float4 Out = 0;
-	if (HairStrandsVF.HasMaterial > 0)
-	{
-		Out = HairStrandsVF.MaterialBuffer[InVertexId];
-	}
-	return Out;
+	const uint CurveIndex = InVertexToCurveBuffer[HairPrimitiveId];
+	return InternalGetHairStrandsAttributes(CurveIndex, AttributeOffsetInBytes, AttributeStrideInBytes, InAttributeBuffer);
 }
 
-FHairVertexAttributes GetHairStrandsAttributes(uint HairPrimitiveId, float2 VertexUV)
+uint InternalGetHairStrandsAttributesPerVertex(uint HairPrimitiveId, uint AttributeOffsetInBytes, uint AttributeStrideInBytes, ByteAddressBuffer InAttributeBuffer)
 {
-	const uint VertexId0 = HairPrimitiveId;
-	const uint VertexId1 = HairPrimitiveId+1;
-
-	// Instead of doing the actual interpolation we could rely on dithering as an optimization
-	const FHairAttributes Attributes0 = UnpackHairAttribute(HairStrandsVF.Attribute0Buffer[VertexId0], HairStrandsVF.Attribute1Buffer[VertexId0], HairStrandsVF.Length);
-	const FHairAttributes Attributes1 = UnpackHairAttribute(HairStrandsVF.Attribute0Buffer[VertexId1], HairStrandsVF.Attribute1Buffer[VertexId1], HairStrandsVF.Length);
+	return InternalGetHairStrandsAttributes(HairPrimitiveId, AttributeOffsetInBytes, AttributeStrideInBytes, InAttributeBuffer);
+}
 
-	const FHairControlPoint ControlPoint0 = UnpackHairControlPoint(HairStrandsVF.PositionBuffer[VertexId0], float3(0, 0 ,0), HairStrandsVF.Radius, HairStrandsVF.RootScale, HairStrandsVF.TipScale);
-	const FHairControlPoint ControlPoint1 = UnpackHairControlPoint(HairStrandsVF.PositionBuffer[VertexId1], float3(0, 0, 0), HairStrandsVF.Radius, HairStrandsVF.RootScale, HairStrandsVF.TipScale);
-	const float WorldLength = lerp(Attributes0.WorldLength, Attributes1.WorldLength, VertexUV.x);
-	const float WorldRadius = lerp(ControlPoint0.WorldRadius, ControlPoint1.WorldRadius, VertexUV.x);
+float InternalGetHairStrandsCoordU(Buffer<uint4> InPositionBuffer, uint VertexId, float InRadius, float InRootScale, float InTipScale)
+{
+	return UnpackHairControlPoint(InPositionBuffer[VertexId], float3(0, 0, 0), InRadius, InRootScale, InTipScale).UCoord;
+}
 
-	FHairVertexAttributes Out;
-	// Constant along the strand
-	Out.HairDimensions	= float2(WorldLength, WorldRadius);
-	Out.HairSeed		= Attributes0.Seed;
-	Out.HairRootUV		= Attributes0.RootUV;
+float InternalGetHairStrandsWorldRadius(Buffer<uint4> InPositionBuffer, uint VertexId, float InRadius, float InRootScale, float InTipScale)
+{
+	return UnpackHairControlPoint(InPositionBuffer[VertexId], float3(0, 0, 0), InRadius, InRootScale, InTipScale).WorldRadius;
+}
 
-	// Varying along the strand
-	Out.HairUV.x		= lerp(ControlPoint0.UCoord, ControlPoint1.UCoord, VertexUV.x);
-	Out.HairUV.y		= VertexUV.y;
+bool IsAttributeValid(uint InOffset)
+{
+	return InOffset != 0xFFFFFFFF;
+}
 
-	// Instead of doing the actual interpolation we could rely on dithering as an optimization
-	const float4 Material0 = GetHairStrandsOptionalMaterial(VertexId0);
-	const float4 Material1 = GetHairStrandsOptionalMaterial(VertexId1);
+////////////////////////////////////////////////////////////////////////////////
+// Attribute access with explicit resources
+#if HAIR_STRAND_MESH_FACTORY
 
-	// Cheap/manual sRGB decoding as the underlying format is linear quantization
-	Out.HairBaseColor	= lerp(Material0.xyz*Material0.xyz, Material1.xyz*Material1.xyz, VertexUV.x);
-	Out.HairRoughness	= lerp(Material0.w, Material1.w, VertexUV.x);
+float InternalGetHairStrandsWorldRadius(uint HairPrimitiveId, float2 VertexUV)
+{
+	const float Radius0 = InternalGetHairStrandsWorldRadius(HairStrandsVF.PositionBuffer, HairPrimitiveId,   HairStrandsVF.Radius, HairStrandsVF.RootScale, HairStrandsVF.TipScale);
+	const float Radius1 = InternalGetHairStrandsWorldRadius(HairStrandsVF.PositionBuffer, HairPrimitiveId+1, HairStrandsVF.Radius, HairStrandsVF.RootScale, HairStrandsVF.TipScale);
+	return lerp(Radius0, Radius1, VertexUV.x);
+}
 
-	return Out; 
+float InternalGetHairStrandsCoordU(uint HairPrimitiveId, float2 VertexUV)
+{
+	const float UCoord0 = InternalGetHairStrandsCoordU(HairStrandsVF.PositionBuffer, HairPrimitiveId,   HairStrandsVF.Radius, HairStrandsVF.RootScale, HairStrandsVF.TipScale);
+	const float UCoord1 = InternalGetHairStrandsCoordU(HairStrandsVF.PositionBuffer, HairPrimitiveId+1, HairStrandsVF.Radius, HairStrandsVF.RootScale, HairStrandsVF.TipScale);
+	return lerp(UCoord0, UCoord1, VertexUV.x);
 }
 
 float GetHairStrandsDepth(float2 InUV)
@@ -76,32 +89,59 @@ float GetHairStrandsCoverage(uint HairPrimitiveId, float2 HairPrimitiveUV)
 
 float2 GetHairStrandsUV(uint HairPrimitiveId, float2 HairPrimitiveUV)
 {
-	return GetHairStrandsAttributes(HairPrimitiveId, HairPrimitiveUV).HairUV;
+	float2 Out;
+	Out.x = InternalGetHairStrandsCoordU(HairPrimitiveId, HairPrimitiveUV);
+	Out.y = HairPrimitiveUV.y;
+	return Out;
 }
 
 float2 GetHairStrandsDimensions(uint HairPrimitiveId, float2 HairPrimitiveUV)
 {
-	return GetHairStrandsAttributes(HairPrimitiveId, HairPrimitiveUV).HairDimensions;
+	const float CurveLength = UnpackHairLength(InternalGetHairStrandsAttributesPerCurve(HairPrimitiveId, HAIR_ATTRIBUTE_OFFSET_LENGTH(HairStrandsVF.AttributeOffsets), HAIR_ATTRIBUTE_STRIDE_LENGTH, HairStrandsVF.AttributeBuffer, HairStrandsVF.VertexToCurveBuffer));
+	const float CoordU = InternalGetHairStrandsCoordU(HairPrimitiveId, HairPrimitiveUV);
+	const float Radius = InternalGetHairStrandsWorldRadius(HairPrimitiveId, HairPrimitiveUV);
+	return float2(CurveLength * CoordU, Radius);
 }
 
 float2 GetHairStrandsRootUV(uint HairPrimitiveId, float2 HairPrimitiveUV)
 {
-	return GetHairStrandsAttributes(HairPrimitiveId, HairPrimitiveUV).HairRootUV;
+	// Encoding need to match GroomBuilder.cpp - 
+	return UnpackHairRootUV(InternalGetHairStrandsAttributesPerCurve(HairPrimitiveId, HAIR_ATTRIBUTE_OFFSET_ROOTUV(HairStrandsVF.AttributeOffsets), HAIR_ATTRIBUTE_STRIDE_ROOTUV, HairStrandsVF.AttributeBuffer, HairStrandsVF.VertexToCurveBuffer));
 }
 
-float  GetHairStrandsSeed(uint HairPrimitiveId, float2 HairPrimitiveUV)
+float GetHairStrandsSeed(uint HairPrimitiveId, float2 HairPrimitiveUV)
 {
-	return GetHairStrandsAttributes(HairPrimitiveId, HairPrimitiveUV).HairSeed;
+	return UnpackHairSeed(InternalGetHairStrandsAttributesPerCurve(HairPrimitiveId, HAIR_ATTRIBUTE_OFFSET_SEED(HairStrandsVF.AttributeOffsets), HAIR_ATTRIBUTE_STRIDE_SEED, HairStrandsVF.AttributeBuffer, HairStrandsVF.VertexToCurveBuffer));
 }
 
 float3 GetHairStrandsBaseColor(uint HairPrimitiveId, float2 HairPrimitiveUV)
 {
-	return GetHairStrandsAttributes(HairPrimitiveId, HairPrimitiveUV).HairBaseColor;
+	const uint BaseColorOffset = HAIR_ATTRIBUTE_OFFSET_BASECOLOR(HairStrandsVF.AttributeOffsets);
+	if (IsAttributeValid(BaseColorOffset))
+	{
+		const float3 Color0 = UnpackHairBaseColor(InternalGetHairStrandsAttributesPerVertex(HairPrimitiveId,   BaseColorOffset, HAIR_ATTRIBUTE_STRIDE_BASECOLOR, HairStrandsVF.AttributeBuffer));
+		const float3 Color1 = UnpackHairBaseColor(InternalGetHairStrandsAttributesPerVertex(HairPrimitiveId+1, BaseColorOffset, HAIR_ATTRIBUTE_STRIDE_BASECOLOR, HairStrandsVF.AttributeBuffer));
+		return lerp(Color0, Color1, HairPrimitiveUV.x);
+	}
+	else
+	{
+		return 0;
+	}
 }
 
 float GetHairStrandsRoughness(uint HairPrimitiveId, float2 HairPrimitiveUV)
 {
-	return GetHairStrandsAttributes(HairPrimitiveId, HairPrimitiveUV).HairRoughness;
+	const uint RoughnessOffset = HAIR_ATTRIBUTE_OFFSET_ROUGHNESS(HairStrandsVF.AttributeOffsets);
+	if (IsAttributeValid(RoughnessOffset))
+	{
+		const float Roughness0 = UnpackHairRoughness(InternalGetHairStrandsAttributesPerVertex(HairPrimitiveId,   RoughnessOffset, HAIR_ATTRIBUTE_STRIDE_ROUGHNESS, HairStrandsVF.AttributeBuffer));
+		const float Roughness1 = UnpackHairRoughness(InternalGetHairStrandsAttributesPerVertex(HairPrimitiveId+1, RoughnessOffset, HAIR_ATTRIBUTE_STRIDE_ROUGHNESS, HairStrandsVF.AttributeBuffer));
+		return lerp(Roughness0, Roughness1, HairPrimitiveUV.x);
+	}
+	else
+	{
+		return 0;
+	}
 }
 
 float GetHairStrandsDepth(uint HairPrimitiveId, float2 InUV, float InDeviceZ)
@@ -113,6 +153,7 @@ float GetHairStrandsGroupIndex(float2 InAtlasUV)
 {
 	return HairStrandsVF.GroupIndex;
 }
+
 ////////////////////////////////////////////////////////////////////////////////
 
 #endif // HAIR_STRAND_MESH_FACTORY
\ No newline at end of file
diff --git a/Engine/Shaders/Private/HairStrands/HairStrandsGuideDeform.usf b/Engine/Shaders/Private/HairStrands/HairStrandsGuideDeform.usf
--- a/Engine/Shaders/Private/HairStrands/HairStrandsGuideDeform.usf
+++ b/Engine/Shaders/Private/HairStrands/HairStrandsGuideDeform.usf
@@ -37,7 +37,7 @@ Buffer<float4>  SimDeformedPosition1Buffer;
 Buffer<float4>  SimDeformedPosition2Buffer;
 
 Buffer<uint>	SimRootBarycentricBuffer;
-Buffer<uint>    SimVertexToRootIndexBuffer;
+Buffer<uint>    SimVertexToCurveBuffer;
 Buffer<uint>	SimRootToUniqueTriangleIndexBuffer;
 
 // Return world position of a guide vertex deformed by a skinned triangle
@@ -83,7 +83,7 @@ float3 DisplacePosition(uint VertexIndex, float3 Pos)
 #endif
 
 #if PERMUTATION_DEFORMATION == DEFORMATION_OFFSET_GUIDE_DYNAMIC
-	const uint RootIndex = SimVertexToRootIndexBuffer[VertexIndex];
+	const uint RootIndex = SimVertexToCurveBuffer[VertexIndex];
 	const float3 RootBarycentric = DecodeBarycentrics(SimRootBarycentricBuffer[RootIndex]);
 
 	const uint TriangleIndex = SimRootToUniqueTriangleIndexBuffer[RootIndex];
diff --git a/Engine/Shaders/Private/HairStrands/HairStrandsInterpolation.usf b/Engine/Shaders/Private/HairStrands/HairStrandsInterpolation.usf
--- a/Engine/Shaders/Private/HairStrands/HairStrandsInterpolation.usf
+++ b/Engine/Shaders/Private/HairStrands/HairStrandsInterpolation.usf
@@ -106,8 +106,8 @@ Buffer<uint>	HairStrandsVF_CullingIndirectBufferArgs;
 #endif
 
 #if PERMUTATION_DYNAMIC_GEOMETRY >= 1 || PERMUTATION_DEBUG == 1 
-Buffer<uint>   RenVertexToRootIndexBuffer;
-Buffer<uint>   SimVertexToRootIndexBuffer;
+Buffer<uint>   RenVertexToCurveBuffer;
+Buffer<uint>   SimVertexToCurveBuffer;
 #endif
 
 #if PERMUTATION_DEFORMER
@@ -245,7 +245,7 @@ FRootTriangleData GetRootTriangleData(uint VertexIndex, bool bSim)
 	#if PERMUTATION_DYNAMIC_GEOMETRY >= 1 
 	if (bSim)
 	{
-		Out.RootIndex		 = SimVertexToRootIndexBuffer[VertexIndex];
+		Out.RootIndex		 = SimVertexToCurveBuffer[VertexIndex];
 		Out.TriangleIndex	 = SimRootToUniqueTriangleIndexBuffer[Out.RootIndex];
 		Out.RootBarycentric  = DecodeBarycentrics(SimRootBarycentricBuffer[Out.RootIndex]);
 		Out.RestTriangle	 = GetTriangleTransformation(Out.TriangleIndex, SimRestPosition0Buffer, SimRestPosition1Buffer, SimRestPosition2Buffer);
@@ -253,7 +253,7 @@ FRootTriangleData GetRootTriangleData(uint VertexIndex, bool bSim)
 	}
 	else
 	{
-		Out.RootIndex		 = RenVertexToRootIndexBuffer[VertexIndex];
+		Out.RootIndex		 = RenVertexToCurveBuffer[VertexIndex];
 		Out.TriangleIndex	 = RenRootToUniqueTriangleIndexBuffer[Out.RootIndex];
 		Out.RootBarycentric  = DecodeBarycentrics(RenRootBarycentricBuffer[Out.RootIndex]);
 		Out.RestTriangle	 = GetTriangleTransformation(Out.TriangleIndex, RenRestPosition0Buffer, RenRestPosition1Buffer, RenRestPosition2Buffer);
@@ -546,19 +546,22 @@ void MainCS(uint2 DispatchThreadId : SV_DispatchThreadID)
 // Patch attribute (for debug visualization)
 #if SHADER_PATCHATTRIBUTE
 
-uint VertexCount;
+#include "HairStrandsVertexFactoryCommon.ush"
+#include "HairStrandsAttributeCommon.ush"
 
-Buffer<float4>	 RenAttribute0Buffer;
-Buffer<int>		 RenVertexToClusterIdBuffer;
+uint				CurveCount;
+
+Buffer<uint>		RenCurveBuffer;
+uint4				RenAttributeOffsets[HAIR_ATTRIBUTE_OFFSET_COUNT];
+Buffer<int>			RenVertexToClusterIdBuffer;
 
-Buffer<float4>	 SimAttribute0Buffer;
 #if PERMUTATION_SIMULATION == 1
-Buffer<uint>	InterpolationBuffer;
+Buffer<uint>		InterpolationBuffer;
 #elif PERMUTATION_SIMULATION == 2
-Buffer<uint4>	Interpolation0Buffer;
-Buffer<uint4>	Interpolation1Buffer;
+Buffer<uint4>		Interpolation0Buffer;
+Buffer<uint4>		Interpolation1Buffer;
 #endif
-RWBuffer<float4> OutRenAttribute0Buffer;
+RWByteAddressBuffer	OutRenAttributeBuffer;
 
 
 FGuideData GetGuideData(uint InVertexIndex)
@@ -571,35 +574,61 @@ FGuideData GetGuideData(uint InVertexIndex)
 	return (FGuideData)0;
 #endif
 }
+groupshared uint SeedValues[1024];
+
+uint ToSeedHash(uint In)
+{
+	const uint InLo8bits =  In & 0xFF;
+	const uint InHi8btis = (In >> 8) & 0xFF;
+
+	// Using FNV1 hash to break linearity of ClusterId (generated by linearly parsing cluster grid)
+	uint Hash = 0;
+	Hash = Hash * 17;
+	Hash = Hash ^ InLo8bits;
+	Hash = Hash * 17;
+	Hash = Hash ^ InHi8btis;
+	return (Hash & 0xFF);
+}
 
 [numthreads(1024, 1, 1)]
-void MainCS(uint2 DispatchThreadId : SV_DispatchThreadID)
+void MainCS(uint2 DispatchThreadId : SV_DispatchThreadID, uint GroupThread1D : SV_GroupIndex)
 {
-	const uint VertexIndex = DispatchThreadId.x;
-	if (VertexIndex < VertexCount)
+	SeedValues[GroupThread1D] = 0;
+
+	const uint CurveIndex = DispatchThreadId.x;
+	const bool bValid = CurveIndex < CurveCount;
+	if (bValid)
 	{
-		float4 Out = 0;
+		uint Out = 0;
 		#if PERMUTATION_SIMULATION == 0
 		{
-			const uint ClusterId		= RenVertexToClusterIdBuffer[VertexIndex];
-			const uint ClusterIdLo8bits = ClusterId & 0xFF;
-			const uint ClusterIdHi8btis = (ClusterId >> 8) & 0xFF;
-			// Using FNV1 hash to break linearity of ClusterId (generated by linearly parsing cluster grid)
-			uint Hash = 0;
-			Hash = Hash * 17;
-			Hash = Hash ^ ClusterIdLo8bits;
-			Hash = Hash * 17;
-			Hash = Hash ^ ClusterIdHi8btis;
-			Out = float4(0, (Hash & 0xFF) / 255.f, 0, 0);
+			const uint VertexIndex0		= UnpackHairCurve(RenCurveBuffer[CurveIndex]).PointIndex;
+			const uint ClusterId		= RenVertexToClusterIdBuffer[VertexIndex0];
+			SeedValues[GroupThread1D]	= ToSeedHash(ClusterId);
 		}
 		#else // if (PERMUTATION_SIMULATION > 0)
 		{
-			const FGuideData GuideData = GetGuideData(VertexIndex);
-			const uint GuideIndex = GuideData.VertexIndices[0]; // Take the closest guide
-			Out = SimAttribute0Buffer[GuideIndex];
+			const uint VertexIndex0		= UnpackHairCurve(RenCurveBuffer[CurveIndex]).PointIndex;
+			const FGuideData GuideData	= GetGuideData(VertexIndex0);
+			const uint GuideIndex		= GuideData.VertexIndices[0]; // Take the closest guide
+			SeedValues[GroupThread1D]	= ToSeedHash(GuideIndex);
 		}
 		#endif // PERMUTATION_SIMULATION
-		OutRenAttribute0Buffer[VertexIndex] = Out;
+	}
+
+	GroupMemoryBarrierWithGroupSync();
+
+	// Write 4 values at a type since seed are encoded into 8bits, and data are written out as uint
+	if ((CurveIndex % 4) == 0)
+	{	
+		uint Out = 0;
+		Out |= (SeedValues[GroupThread1D]   & 0xFF);
+		Out |= (SeedValues[GroupThread1D+1] & 0xFF) << 8;
+		Out |= (SeedValues[GroupThread1D+2] & 0xFF) << 16;
+		Out |= (SeedValues[GroupThread1D+3] & 0xFF) << 24;
+
+		const uint OffsetInBytes = HAIR_ATTRIBUTE_OFFSET_SEED(RenAttributeOffsets);
+		OutRenAttributeBuffer.Store(OffsetInBytes + CurveIndex, Out);
 	}
 }
 #endif // SHADER_PATCHATTRIBUTE
diff --git a/Engine/Shaders/Private/HairStrands/HairStrandsTexturesGeneration.usf b/Engine/Shaders/Private/HairStrands/HairStrandsTexturesGeneration.usf
--- a/Engine/Shaders/Private/HairStrands/HairStrandsTexturesGeneration.usf
+++ b/Engine/Shaders/Private/HairStrands/HairStrandsTexturesGeneration.usf
@@ -59,7 +59,9 @@ void MainVS(
 
 #if SHADER_PIXEL
 
+#define HAIR_STRANDS_PARAMETERS 0
 #include "HairStrandsVertexFactoryCommon.ush"
+#include "HairStrandsAttributeCommon.ush"
 
 float3 ComputeTangent(
 	const FHairControlPoint CP0,
@@ -83,15 +85,13 @@ uint VertexCount;
 uint PrimitiveCount;
 uint bHasMaterialData;
 
-Buffer<uint4> InVF_PositionBuffer;
-Buffer<float2> InVF_Attribute0Buffer;
-Buffer<uint> InVF_Attribute1Buffer;
-Buffer<float4> InVF_MaterialBuffer;
+Buffer<uint4>		InVF_PositionBuffer;
+ByteAddressBuffer	InVF_AttributeBuffer;
+Buffer<uint>		InVF_VertexToCurveBuffer;
+uint4				InVF_AttributeOffsets[HAIR_ATTRIBUTE_OFFSET_COUNT];
 
 float3 InVF_PositionOffset;
 float InVF_Radius;
-float InVF_Length;
-uint InVF_ControlPointCount;
 uint InVF_GroupIndex;
 
 Buffer<float> VertexBuffer;
@@ -271,14 +271,6 @@ bool Trace(
 						const uint4 PackedCP_P0 = InVF_PositionBuffer[P0];
 						const uint4 PackedCP_P1 = InVF_PositionBuffer[P1];
 
-						float4 Material0 = 0;
-						float4 Material1 = 0;
-						if (bHasMaterialData)
-						{
-							Material0 = InVF_MaterialBuffer[P0];
-							Material1 = InVF_MaterialBuffer[P1];
-						}
-
 						const FHairControlPoint CP_Curr = UnpackHairControlPoint(PackedCP_P0, InVF_PositionOffset, InVF_Radius, InVF_RootScale, InVF_TipScale);
 						const FHairControlPoint CP_Next = UnpackHairControlPoint(PackedCP_P1, InVF_PositionOffset, InVF_Radius, InVF_RootScale, InVF_TipScale);
 
@@ -306,13 +298,27 @@ bool Trace(
 							OutTangent = Tangent;
 							OutDepth = Distance;
 
-							const FHairAttributes Attributes0 = UnpackHairAttribute(InVF_Attribute0Buffer[P0], InVF_Attribute1Buffer[P0], InVF_Length);
-							const FHairAttributes Attributes1 = UnpackHairAttribute(InVF_Attribute0Buffer[P1], InVF_Attribute1Buffer[P0], InVF_Length);
+							OutSeed         = UnpackHairSeed(InternalGetHairStrandsAttributesPerCurve(P0, HAIR_ATTRIBUTE_OFFSET_SEED(InVF_AttributeOffsets), HAIR_ATTRIBUTE_STRIDE_SEED, InVF_AttributeBuffer, InVF_VertexToCurveBuffer));
+							OutRootUV       = UnpackHairRootUV(InternalGetHairStrandsAttributesPerCurve(P0, HAIR_ATTRIBUTE_OFFSET_ROOTUV(InVF_AttributeOffsets), HAIR_ATTRIBUTE_STRIDE_ROOTUV, InVF_AttributeBuffer, InVF_VertexToCurveBuffer));
+							OutU			= lerp(CP_Curr.UCoord, CP_Next.UCoord, VertexU);
+							OutMaterial.xyz = 0;
+							OutMaterial.w   = 0;
+
+							const uint BaseColorOffset = HAIR_ATTRIBUTE_OFFSET_BASECOLOR(InVF_AttributeOffsets);
+							if (IsAttributeValid(BaseColorOffset))
+							{
+								const float3 C0 = UnpackHairBaseColor(InternalGetHairStrandsAttributesPerVertex(P0, BaseColorOffset, HAIR_ATTRIBUTE_STRIDE_BASECOLOR, InVF_AttributeBuffer));
+								const float3 C1 = UnpackHairBaseColor(InternalGetHairStrandsAttributesPerVertex(P1, BaseColorOffset, HAIR_ATTRIBUTE_STRIDE_BASECOLOR, InVF_AttributeBuffer));
+								OutMaterial.xyz = lerp(C0, C1, VertexU);
+							}
 
-							OutSeed		= Attributes0.Seed;
-							OutRootUV	= Attributes0.RootUV;
-							OutU		= lerp(CP_Curr.UCoord, CP_Next.UCoord, VertexU);
-							OutMaterial = lerp(Material0, Material1, VertexU);
+							const uint RoughnessOffset = HAIR_ATTRIBUTE_OFFSET_ROUGHNESS(InVF_AttributeOffsets);
+							if (IsAttributeValid(RoughnessOffset))
+							{
+								const float R0 = UnpackHairRoughness(InternalGetHairStrandsAttributesPerVertex(P0, RoughnessOffset, HAIR_ATTRIBUTE_STRIDE_ROUGHNESS, InVF_AttributeBuffer));
+								const float R1 = UnpackHairRoughness(InternalGetHairStrandsAttributesPerVertex(P1, RoughnessOffset, HAIR_ATTRIBUTE_STRIDE_ROUGHNESS, InVF_AttributeBuffer));
+								OutMaterial.w = lerp(R0, R1, VertexU);
+							}
 
 							bIsValid = true;
 
diff --git a/Engine/Shaders/Private/HairStrands/HairStrandsVertexFactoryCommon.ush b/Engine/Shaders/Private/HairStrands/HairStrandsVertexFactoryCommon.ush
--- a/Engine/Shaders/Private/HairStrands/HairStrandsVertexFactoryCommon.ush
+++ b/Engine/Shaders/Private/HairStrands/HairStrandsVertexFactoryCommon.ush
@@ -2,6 +2,8 @@
 
 #pragma once
 
+#include "/Engine/Shared/HairStrandsDefinitions.h"
+
 #define HAIR_CONTROLPOINT_INSIDE 0
 #define HAIR_CONTROLPOINT_START	1
 #define HAIR_CONTROLPOINT_END	2
@@ -74,33 +76,46 @@ bool IsValidHairStrandsSegment(const FHairControlPoint CP0, const FHairControlPo
 ////////////////////////////////////////////////////////////////////////////////////////////////////////////////////////////////
 // Hair attributes
 
-struct FHairAttributes
-{
-	float2 RootUV;
-	float Seed;
-	float WorldLength;
-};
-
-FHairAttributes UnpackHairAttribute(float2 Attributes0, uint Attributes1, float InVF_Length)
+float2 UnpackHairRootUV(uint In)
 {
 	float2 RootUV, RootIndex;
-	RootUV.x		= ((Attributes1    ) & 0x7FF) / 2047.f;	// Coord are encoded on 11 bits
-	RootUV.y		= ((Attributes1>>11) & 0x7FF) / 2047.f;	// Coord are encoded on 11 bits 
-	RootIndex.x		= ((Attributes1>>22) & 0x1F);			// UDIM tile are encoded on 5bits
-	RootIndex.y		= ((Attributes1>>27) & 0x1F);			// UDIM tile are encoded on 5bits
+	RootUV.x	= ((In) & 0x7FF) / 2047.f;			// Coord are encoded on 11 bits
+	RootUV.y	= ((In >> 11) & 0x7FF) / 2047.f;	// Coord are encoded on 11 bits 
+	RootIndex.x = ((In >> 22) & 0x1F);				// UDIM tile are encoded on 5bits
+	RootIndex.y = ((In >> 27) & 0x1F);				// UDIM tile are encoded on 5bits
 
 	// Invert V to compensate image origin flip. Similar to GetHairStrandsRootUV in HairCardsAttributeCommon.ush
 	RootUV.y = 1.0f - RootUV.y;
 
-	FHairAttributes Out;
-	Out.RootUV		= RootUV + RootIndex;
-	Out.WorldLength = Attributes0.x * InVF_Length;
-	Out.Seed		= Attributes0.y;
-	return Out;
+	return RootUV = RootUV + RootIndex;
+}
+
+float UnpackHairLength(uint In)
+{
+	return f16tof32(In & 0xFFFF);
+}
+
+float UnpackHairSeed(uint In)
+{
+	return (In & 0xFF) / 255.f;
+}
+
+float3 UnpackHairBaseColor(uint In)
+{
+	float3 Out;
+	Out.x = (In       & 0xFF) / 255.f;
+	Out.y = ((In>>8)  & 0xFF) / 255.f;
+	Out.z = ((In>>16) & 0xFF) / 255.f;
+	return Out * Out; // Cheap sRGB encoding
+}
+
+float UnpackHairRoughness(uint In)
+{
+	return (In & 0xFF) / 255.f;
 }
 
 ////////////////////////////////////////////////////////////////////////////////////////////////////////////////////////////////
-// Hair Curve
+// Hair curve
 
 struct FHairCurve
 {
@@ -136,6 +151,7 @@ float4x4		HairStrandsVF_LocalToTranslatedWorldPrimitiveTransform;
 Buffer<uint4>	HairStrandsVF_PositionBuffer;
 Buffer<float4>	HairStrandsVF_PositionOffsetBuffer;
 Buffer<uint>	HairStrandsVF_CurveBuffer;
+uint4			HairStrandsVF_AttributeOffsets[HAIR_ATTRIBUTE_OFFSET_COUNT];
 
 #if defined(PERMUTATION_CULLING) && PERMUTATION_CULLING == 1
 Buffer<uint>	HairStrandsVF_CullingIndirectBuffer;
diff --git a/Engine/Shaders/Shared/HairStrandsDefinitions.h b/Engine/Shaders/Shared/HairStrandsDefinitions.h
new file mode 100644
--- /dev/null
+++ b/Engine/Shaders/Shared/HairStrandsDefinitions.h
@@ -0,0 +1,36 @@
+// Copyright Epic Games, Inc. All Rights Reserved.
+
+/*================================================================================================
+	HairStrandsDefinitions.ush: used in ray tracing shaders and C++ code to define common constants
+	!!! Changing this file requires recompilation of the engine !!!
+=================================================================================================*/
+
+#pragma once
+
+// Change this to force recompilation of all strata dependent shaders (use https://www.random.org/cgi-bin/randbyte?nbytes=4&format=h)
+#define HAIRSTRANDS_SHADER_VERSION 0x3ba5af0e
+
+// Attribute index
+#define HAIR_ATTRIBUTE_ROOTUV 0
+#define HAIR_ATTRIBUTE_SEED 1
+#define HAIR_ATTRIBUTE_LENGTH 2
+#define HAIR_ATTRIBUTE_BASECOLOR 3
+#define HAIR_ATTRIBUTE_ROUGHNESS 4
+#define HAIR_ATTRIBUTE_COUNT 5
+
+// HAIR_ATTRIBUTE_MAX rounded to 4
+#define HAIR_ATTRIBUTE_OFFSET_COUNT ((HAIR_ATTRIBUTE_COUNT + 3) / 4)
+
+// Pack all offset into uint4
+#define PACK_HAIR_ATTRIBUTE_OFFSETS(Out, In) \
+	for (uint32 AttributeIt4 = 0; AttributeIt4 < HAIR_ATTRIBUTE_OFFSET_COUNT; ++AttributeIt4) \
+	{ \
+		Out[AttributeIt4] = FUintVector4(0xFFFFFFFF, 0xFFFFFFFF, 0xFFFFFFFF, 0xFFFFFFFF); \
+	} \
+	for (uint32 AttributeIt = 0; AttributeIt < HAIR_ATTRIBUTE_COUNT; ++AttributeIt) \
+	{ \
+		const uint32 Index4 = AttributeIt & (~0x3); \
+		const uint32 SubIndex = AttributeIt - Index4; \
+		const uint32 IndexDiv4 = Index4 >> 2u; \
+		Out[IndexDiv4][SubIndex] = In[AttributeIt]; \
+	}
\ No newline at end of file
diff --git a/Engine/Source/Runtime/Renderer/Private/HairStrands/HairStrandsData.h b/Engine/Source/Runtime/Renderer/Private/HairStrands/HairStrandsData.h
--- a/Engine/Source/Runtime/Renderer/Private/HairStrands/HairStrandsData.h
+++ b/Engine/Source/Runtime/Renderer/Private/HairStrands/HairStrandsData.h
@@ -12,6 +12,7 @@
 #include "RHIGPUReadback.h"
 #include "Shader.h"
 #include "ConvexVolume.h"
+#include "HairStrandsDefinitions.h"
 
 class FLightSceneInfo;
 class FPrimitiveSceneProxy;
@@ -59,12 +60,15 @@ BEGIN_SHADER_PARAMETER_STRUCT(FHairStrandsInstanceParameters, )
 	SHADER_PARAMETER(uint32, HairStrandsVF_bUseStableRasterization)
 	SHADER_PARAMETER(uint32, HairStrandsVF_VertexCount)
 	SHADER_PARAMETER(uint32, HairStrandsVF_CurveCount)
+	SHADER_PARAMETER_ARRAY(FUintVector4, HairStrandsVF_AttributeOffsets, [HAIR_ATTRIBUTE_OFFSET_COUNT])
 	SHADER_PARAMETER(FVector3f, HairStrandsVF_PositionOffset)
 	SHADER_PARAMETER(FMatrix44f, HairStrandsVF_LocalToWorldPrimitiveTransform)
 	SHADER_PARAMETER(FMatrix44f, HairStrandsVF_LocalToTranslatedWorldPrimitiveTransform)
 	SHADER_PARAMETER_RDG_BUFFER_SRV(Buffer, HairStrandsVF_PositionBuffer)
 	SHADER_PARAMETER_RDG_BUFFER_SRV(Buffer, HairStrandsVF_PositionOffsetBuffer)
 	SHADER_PARAMETER_RDG_BUFFER_SRV(Buffer, HairStrandsVF_CurveBuffer)
+	SHADER_PARAMETER_RDG_BUFFER_SRV(Buffer, HairStrandsVF_VertexToCurveBuffer)
+	SHADER_PARAMETER_RDG_BUFFER_SRV(ByteAddressBuffer, HairStrandsVF_AttributeBuffer)
 	SHADER_PARAMETER_RDG_BUFFER_SRV(Buffer, HairStrandsVF_CullingIndirectBuffer)
 	SHADER_PARAMETER_RDG_BUFFER_SRV(Buffer, HairStrandsVF_CullingIndexBuffer)
 	SHADER_PARAMETER_RDG_BUFFER_SRV(Buffer, HairStrandsVF_CullingRadiusScaleBuffer)
diff --git a/Engine/Source/Runtime/Renderer/Private/HairStrands/HairStrandsVisibility.cpp b/Engine/Source/Runtime/Renderer/Private/HairStrands/HairStrandsVisibility.cpp
--- a/Engine/Source/Runtime/Renderer/Private/HairStrands/HairStrandsVisibility.cpp
+++ b/Engine/Source/Runtime/Renderer/Private/HairStrands/HairStrandsVisibility.cpp
@@ -2721,12 +2721,8 @@ void AddMeshDrawTransitionPass(
 			ExternalAccessQueue.Add(VFInput.Strands.PositionBuffer.Buffer);
 			ExternalAccessQueue.Add(VFInput.Strands.PrevPositionBuffer.Buffer);
 			ExternalAccessQueue.Add(VFInput.Strands.TangentBuffer.Buffer);
-			ExternalAccessQueue.Add(VFInput.Strands.Attribute0Buffer.Buffer);
-			ExternalAccessQueue.Add(VFInput.Strands.Attribute1Buffer.Buffer);
-			if (VFInput.Strands.MaterialBuffer.Buffer != nullptr)
-			{
-				ExternalAccessQueue.Add(VFInput.Strands.MaterialBuffer.Buffer);
-			}
+			ExternalAccessQueue.Add(VFInput.Strands.AttributeBuffer.Buffer);
+			ExternalAccessQueue.Add(VFInput.Strands.VertexToCurveBuffer.Buffer);
 			ExternalAccessQueue.Add(VFInput.Strands.PositionOffsetBuffer.Buffer);
 			ExternalAccessQueue.Add(VFInput.Strands.PrevPositionOffsetBuffer.Buffer);
 
@@ -2743,9 +2739,8 @@ void AddMeshDrawTransitionPass(
 				VFInput.Strands.PositionBuffer				= FRDGImportedBuffer();
 				VFInput.Strands.PrevPositionBuffer			= FRDGImportedBuffer();
 				VFInput.Strands.TangentBuffer				= FRDGImportedBuffer();
-				VFInput.Strands.Attribute0Buffer			= FRDGImportedBuffer();
-				VFInput.Strands.Attribute1Buffer			= FRDGImportedBuffer();
-				VFInput.Strands.MaterialBuffer				= FRDGImportedBuffer();
+				VFInput.Strands.AttributeBuffer				= FRDGImportedBuffer();
+				VFInput.Strands.VertexToCurveBuffer			= FRDGImportedBuffer();
 				VFInput.Strands.PositionOffsetBuffer		= FRDGImportedBuffer();
 				VFInput.Strands.PrevPositionOffsetBuffer	= FRDGImportedBuffer();
 			}
@@ -2817,27 +2812,32 @@ FHairStrandsInstanceParameters GetHairStrandsInstanceParameters(FRDGBuilder& Gra
 	{
 		Out.HairStrandsVF_PositionBuffer		= Register(GraphBuilder, VFInput.Strands.PositionBufferExternal, ERDGImportedBufferFlags::CreateSRV).SRV;;
 		Out.HairStrandsVF_PositionOffsetBuffer	= Register(GraphBuilder, VFInput.Strands.PositionOffsetBufferExternal, ERDGImportedBufferFlags::CreateSRV).SRV;;
-		Out.HairStrandsVF_CurveBuffer = Register(GraphBuilder, VFInput.Strands.CurveBufferExternal, ERDGImportedBufferFlags::CreateSRV).SRV;
+		Out.HairStrandsVF_CurveBuffer			= Register(GraphBuilder, VFInput.Strands.CurveBufferExternal, ERDGImportedBufferFlags::CreateSRV).SRV;
+		Out.HairStrandsVF_VertexToCurveBuffer	= Register(GraphBuilder, VFInput.Strands.VertexToCurveBufferExternal, ERDGImportedBufferFlags::CreateSRV).SRV;
+		Out.HairStrandsVF_AttributeBuffer		= Register(GraphBuilder, VFInput.Strands.AttributeBufferExternal, ERDGImportedBufferFlags::CreateSRV).SRV;
 	}
 	else
 	{
 		Out.HairStrandsVF_PositionBuffer		= VFInput.Strands.PositionBuffer.SRV;
 		Out.HairStrandsVF_PositionOffsetBuffer	= VFInput.Strands.PositionOffsetBuffer.SRV;
-		Out.HairStrandsVF_CurveBuffer = VFInput.Strands.CurveBuffer.SRV;
+		Out.HairStrandsVF_CurveBuffer			= VFInput.Strands.CurveBuffer.SRV;
+		Out.HairStrandsVF_VertexToCurveBuffer	= VFInput.Strands.VertexToCurveBuffer.SRV;
+		Out.HairStrandsVF_AttributeBuffer		= VFInput.Strands.AttributeBuffer.SRV;
 	}
 
-
-	Out.HairStrandsVF_VertexCount						= VFInput.Strands.VertexCount;
-	Out.HairStrandsVF_CurveCount						= VFInput.Strands.CurveCount;
-	Out.HairStrandsVF_Radius							= VFInput.Strands.HairRadius;
-	Out.HairStrandsVF_RootScale							= VFInput.Strands.HairRootScale;
-	Out.HairStrandsVF_TipScale							= VFInput.Strands.HairTipScale;
-	Out.HairStrandsVF_Length							= VFInput.Strands.HairLength;
-	Out.HairStrandsVF_bUseStableRasterization			= VFInput.Strands.bUseStableRasterization ? 1 : 0;
-	Out.HairStrandsVF_Density							= VFInput.Strands.HairDensity;
-	Out.HairStrandsVF_bIsCullingEnable					= 0;
-	Out.HairStrandsVF_bHasRaytracedGeometry				= VFInput.Strands.bUseRaytracingGeometry ? 1u : 0u;
-	Out.HairStrandsVF_PositionOffset					= (FVector3f)VFInput.Strands.PositionOffset;
+	Out.HairStrandsVF_VertexCount				= VFInput.Strands.VertexCount;
+	Out.HairStrandsVF_CurveCount				= VFInput.Strands.CurveCount;
+	Out.HairStrandsVF_Radius					= VFInput.Strands.HairRadius;
+	Out.HairStrandsVF_RootScale					= VFInput.Strands.HairRootScale;
+	Out.HairStrandsVF_TipScale					= VFInput.Strands.HairTipScale;
+	Out.HairStrandsVF_Length					= VFInput.Strands.HairLength;
+	Out.HairStrandsVF_bUseStableRasterization	= VFInput.Strands.bUseStableRasterization ? 1 : 0;
+	Out.HairStrandsVF_Density					= VFInput.Strands.HairDensity;
+	Out.HairStrandsVF_bIsCullingEnable			= 0;
+	Out.HairStrandsVF_bHasRaytracedGeometry		= VFInput.Strands.bUseRaytracingGeometry ? 1u : 0u;
+	Out.HairStrandsVF_PositionOffset			= (FVector3f)VFInput.Strands.PositionOffset;
+
+	PACK_HAIR_ATTRIBUTE_OFFSETS(Out.HairStrandsVF_AttributeOffsets, VFInput.Strands.AttributeOffsets);
 
 	// Absolute local to world
 	Out.HairStrandsVF_LocalToWorldPrimitiveTransform	= FMatrix44f(VFInput.LocalToWorldTransform.ToMatrixWithScale()); // LWC_TODO: Precision loss
diff --git a/Engine/Source/Runtime/Renderer/Public/HairStrandsInterface.h b/Engine/Source/Runtime/Renderer/Public/HairStrandsInterface.h
--- a/Engine/Source/Runtime/Renderer/Public/HairStrandsInterface.h
+++ b/Engine/Source/Runtime/Renderer/Public/HairStrandsInterface.h
@@ -209,9 +209,8 @@ public:
 			FRDGImportedBuffer PositionBuffer;
 			FRDGImportedBuffer PrevPositionBuffer;
 			FRDGImportedBuffer TangentBuffer;
-			FRDGImportedBuffer MaterialBuffer;
-			FRDGImportedBuffer Attribute0Buffer;
-			FRDGImportedBuffer Attribute1Buffer;
+			FRDGImportedBuffer AttributeBuffer;
+			FRDGImportedBuffer VertexToCurveBuffer;
 			FRDGImportedBuffer PositionOffsetBuffer;
 			FRDGImportedBuffer PrevPositionOffsetBuffer;
 			FRDGImportedBuffer CurveBuffer;
@@ -219,9 +218,8 @@ public:
 			FRDGExternalBuffer PositionBufferExternal;
 			FRDGExternalBuffer PrevPositionBufferExternal;
 			FRDGExternalBuffer TangentBufferExternal;
-			FRDGExternalBuffer MaterialBufferExternal;
-			FRDGExternalBuffer Attribute0BufferExternal;
-			FRDGExternalBuffer Attribute1BufferExternal;
+			FRDGExternalBuffer AttributeBufferExternal;
+			FRDGExternalBuffer VertexToCurveBufferExternal;
 			FRDGExternalBuffer PositionOffsetBufferExternal;
 			FRDGExternalBuffer PrevPositionOffsetBufferExternal;
 			FRDGExternalBuffer CurveBufferExternal;
@@ -229,15 +227,15 @@ public:
 			FShaderResourceViewRHIRef PositionBufferRHISRV				= nullptr;
 			FShaderResourceViewRHIRef PrevPositionBufferRHISRV			= nullptr;
 			FShaderResourceViewRHIRef TangentBufferRHISRV				= nullptr;
-			FShaderResourceViewRHIRef MaterialBufferRHISRV				= nullptr;
-			FShaderResourceViewRHIRef Attribute0BufferRHISRV			= nullptr;
-			FShaderResourceViewRHIRef Attribute1BufferRHISRV			= nullptr;
+			FShaderResourceViewRHIRef AttributeBufferRHISRV				= nullptr;
+			FShaderResourceViewRHIRef VertexToCurveBufferRHISRV			= nullptr;
 			FShaderResourceViewRHIRef PositionOffsetBufferRHISRV		= nullptr;
 			FShaderResourceViewRHIRef PrevPositionOffsetBufferRHISRV	= nullptr;
 			FShaderResourceViewRHIRef CurveBufferRHISRV					= nullptr;
 
 			FVector PositionOffset = FVector::ZeroVector;
 			FVector PrevPositionOffset = FVector::ZeroVector;
+			TArray<uint32> AttributeOffsets;
 
 			uint32 CurveCount = 0;
 			uint32 VertexCount = 0;
