Replace Niagara hair PointsCurveBuffer buffer with rest resources VertexToCurveBuffer.

diff --git a/Engine/Plugins/Runtime/HairStrands/Source/HairStrandsCore/Private/Niagara/NiagaraDataInterfaceHairStrands.cpp b/Engine/Plugins/Runtime/HairStrands/Source/HairStrandsCore/Private/Niagara/NiagaraDataInterfaceHairStrands.cpp
--- a/Engine/Plugins/Runtime/HairStrands/Source/HairStrandsCore/Private/Niagara/NiagaraDataInterfaceHairStrands.cpp
+++ b/Engine/Plugins/Runtime/HairStrands/Source/HairStrandsCore/Private/Niagara/NiagaraDataInterfaceHairStrands.cpp
@@ -303,36 +303,6 @@ void FNDIHairStrandsBuffer::InitRHI()
 		FMemMark MemMark(FMemStack::Get());
 		FRDGBuilder GraphBuilder(FRHICommandListExecutor::GetImmediateCommandList());
 		FHairStrandsBulkData* SourceDatas = &SourceRestResources->BulkData; // This could be released by that time depending on how the initialization order is
-		{
-			// First retrieve the curves offsets from the bulk datas
-			TArray<FHairStrandsCurveFormat::Type> Curves;
-			const uint32 InDataSize = SourceDatas->Curves.GetBulkDataSize();
-			const uint32 ElementCount = InDataSize / sizeof(FHairStrandsRootIndexFormat::BulkType);
-			Curves.SetNum(ElementCount);
-			const uint8* InData = (const uint8*)SourceDatas->Curves.LockReadOnly();
-			FMemory::Memcpy(Curves.GetData(), InData, InDataSize);
-			SourceDatas->Curves.Unlock();
-			
-			TArray<uint32> PointsCurve;
-			PointsCurve.SetNum(SourceDatas->GetNumPoints());
-
-			// Build the curve index array for each points using the curve offsets array
-			for (uint32 CurveIndex = 0; CurveIndex < SourceDatas->GetNumCurves(); ++CurveIndex)
-			{
-				const FHairStrandsCurveFormat::Type& Curve = Curves[CurveIndex];				
-				for (uint32 PointIndex = Curve.PointOffset; PointIndex < Curve.PointCount; ++PointIndex)
-				{
-					PointsCurve[PointIndex] = CurveIndex;
-				}
-			}
-			const uint32 OffsetCount = SourceDatas->GetNumPoints();
-			const uint32 OffsetBytes = sizeof(uint32)*OffsetCount;
-
-			// Copy the points buffer data to gpu for interpolation
-			PointsCurveBuffer.Initialize(GraphBuilder, TEXT("PointsCurveBuffer"), EPixelFormat::PF_R32_UINT, sizeof(uint32), OffsetCount);
-			GraphBuilder.QueueBufferUpload(PointsCurveBuffer.GetOrCreateBuffer(GraphBuilder), PointsCurve.GetData(), OffsetBytes);
-			PointsCurveBuffer.EndGraphUsage();
-		}
 		{
 			static const uint32 ZeroData[] = {
 				UINT_MAX,UINT_MAX,UINT_MAX,0,0,0,
@@ -374,7 +344,6 @@ void FNDIHairStrandsBuffer::ReleaseRHI()
 	DeformedPositionBuffer.SafeRelease();
 	BoundingBoxBuffer.Release();
 	ParamsScaleBuffer.Release();
-	PointsCurveBuffer.Release();
 }
 
 //------------------------------------------------------------------------------------------------------------
@@ -932,7 +901,7 @@ static void InterpolateGroomGuides(FRDGBuilder& GraphBuilder, FNiagaraDataBuffer
 	const bool bIsRestValid = bIsHairValid && HairStrandsBuffer->SourceRestResources && HairStrandsBuffer->SourceRestResources->IsInitialized();
 	const bool bIsDeformedValid = bIsHairValid && HairStrandsBuffer->SourceDeformedResources && HairStrandsBuffer->SourceDeformedResources->IsInitialized();
 
-	if(bIsRestValid && bIsDeformedValid && bHasValidGeometry && HairStrandsBuffer->PointsCurveBuffer.IsValid())
+	if(bIsRestValid && bIsDeformedValid && bHasValidGeometry)
 	{
 		bool bIsRootValid = HairStrandsBuffer->SourceDeformedRootResources && HairStrandsBuffer->SourceDeformedRootResources->IsInitialized() &&
 								HairStrandsBuffer->SourceRestRootResources && HairStrandsBuffer->SourceRestRootResources->IsInitialized() && bHasSkinningBinding;
@@ -964,7 +933,7 @@ static void InterpolateGroomGuides(FRDGBuilder& GraphBuilder, FNiagaraDataBuffer
 		PassParameters->DeformedPositionBuffer = RegisterAsUAV(GraphBuilder, HairStrandsBuffer->SourceDeformedResources->DeformedPositionBuffer[BufferIndex]);
 		PassParameters->RestPositionBuffer = RegisterAsSRV(GraphBuilder, HairStrandsBuffer->SourceRestResources->PositionBuffer);
 		
-		PassParameters->VertexToCurveIndexBuffer = HairStrandsBuffer->PointsCurveBuffer.GetOrCreateSRV(GraphBuilder);
+		PassParameters->VertexToCurveIndexBuffer = RegisterAsSRV(GraphBuilder, HairStrandsBuffer->SourceRestResources->VertexToCurveBuffer);
 		PassParameters->CurvesOffsetsBuffer = RegisterAsSRV(GraphBuilder, HairStrandsBuffer->SourceRestResources->CurveBuffer);
 		if(bIsRootValid)
 		{
@@ -1038,8 +1007,6 @@ void UNiagaraDataInterfaceHairStrands::SimCachePostReadFrame(void* OptionalPerIn
 							FRDGBuilder GraphBuilder(RHICmdList);
 							InterpolateGroomGuides(GraphBuilder, ParticlesBuffer, NodePositionComponent, RestPositionComponent, HairstrandsBuffer, StrandsSize, bHasSkinningBinding, bHasValidGeometry, WorldToLocal);
 							GraphBuilder.Execute();
-
-							HairstrandsBuffer->PointsCurveBuffer.EndGraphUsage();
 						}
 					);
 				}
@@ -3644,7 +3611,6 @@ void FNDIHairStrandsProxy::PostSimulate(const FNDIGpuComputePostSimulateContext&
 	{
 		ProxyData->HairStrandsBuffer->BoundingBoxBuffer.EndGraphUsage();
 		ProxyData->HairStrandsBuffer->ParamsScaleBuffer.EndGraphUsage();
-		ProxyData->HairStrandsBuffer->PointsCurveBuffer.EndGraphUsage();
 	}
 }
 
diff --git a/Engine/Plugins/Runtime/HairStrands/Source/HairStrandsCore/Public/Niagara/NiagaraDataInterfaceHairStrands.h b/Engine/Plugins/Runtime/HairStrands/Source/HairStrandsCore/Public/Niagara/NiagaraDataInterfaceHairStrands.h
--- a/Engine/Plugins/Runtime/HairStrands/Source/HairStrandsCore/Public/Niagara/NiagaraDataInterfaceHairStrands.h
+++ b/Engine/Plugins/Runtime/HairStrands/Source/HairStrandsCore/Public/Niagara/NiagaraDataInterfaceHairStrands.h
@@ -59,9 +59,6 @@ struct FNDIHairStrandsBuffer : public FRenderResource
 	/** Params scale buffer */
 	FNiagaraPooledRWBuffer ParamsScaleBuffer;
 
-	/** Points curve index for fast query */
-	FNiagaraPooledRWBuffer PointsCurveBuffer;
-
 	/** The strand asset resource from which to sample */
 	const FHairStrandsRestResource* SourceRestResources;
 
