Fix CIS error.

diff --git a/Engine/Plugins/Runtime/HairStrands/Source/HairStrandsCore/Private/HairStrandsRendering.cpp b/Engine/Plugins/Runtime/HairStrands/Source/HairStrandsCore/Private/HairStrandsRendering.cpp
--- a/Engine/Plugins/Runtime/HairStrands/Source/HairStrandsCore/Private/HairStrandsRendering.cpp
+++ b/Engine/Plugins/Runtime/HairStrands/Source/HairStrandsCore/Private/HairStrandsRendering.cpp
@@ -1718,7 +1718,10 @@ static FHairGroupPublicData::FVertexFactoryInput InternalComputeHairStrandsVerte
 
 	OutVFInput.Strands.AttributeOffsets.SetNum(0, false);
 	OutVFInput.Strands.AttributeOffsets.Reserve(HAIR_ATTRIBUTE_COUNT);
-	OutVFInput.Strands.AttributeOffsets.Append(bUseGuideAttributeOffsets ? Instance->Guides.Data->AttributeOffsets : Instance->Strands.Data->AttributeOffsets);
+	for (uint32 AttributeIt=0; AttributeIt< HAIR_ATTRIBUTE_COUNT; ++AttributeIt)
+	{
+		OutVFInput.Strands.AttributeOffsets.Add(bUseGuideAttributeOffsets ? Instance->Guides.Data->AttributeOffsets[AttributeIt] : Instance->Strands.Data->AttributeOffsets[AttributeIt]);
+	}
 
 #if RHI_RAYTRACING
 	// Flag bUseRaytracingGeometry only if RT geometry has been allocated for RayTracing view (not for PathTracing view). 
