Add ClumpID & Per-Vertex bake AO attribute support for groom geometry.

diff --git a/Engine/Plugins/Runtime/HairStrands/Source/HairStrandsCore/Private/GroomBuilder.cpp b/Engine/Plugins/Runtime/HairStrands/Source/HairStrandsCore/Private/GroomBuilder.cpp
--- a/Engine/Plugins/Runtime/HairStrands/Source/HairStrandsCore/Private/GroomBuilder.cpp
+++ b/Engine/Plugins/Runtime/HairStrands/Source/HairStrandsCore/Private/GroomBuilder.cpp
@@ -42,9 +42,9 @@ FString FGroomBuilder::GetVersion()
 {
 	// Important to update the version when groom building changes
 	if (IsHairStrandContinuousDecimationReorderingEnabled())
-		return TEXT("v4g");
+		return TEXT("v4i");
 	else
-		return TEXT("v3g");
+		return TEXT("v3i");
 }
 
 namespace FHairStrandsDecimation
@@ -295,20 +295,26 @@ namespace HairStrandsBuilder
 			OutVertexToCurve32.SetNum(NumPoints);
 		}
 
-		const bool bHasBaseColor = HasValues(HairStrands.StrandsPoints.PointsBaseColor, FLinearColor::Transparent);
+		const bool bHasClumpIDs = HasValues(HairStrands.StrandsCurves.ClumpIDs, uint16(0));
+		const bool bHasColor = HasValues(HairStrands.StrandsPoints.PointsBaseColor, FLinearColor::Transparent);
 		const bool bHasRoughness = HasValues(HairStrands.StrandsPoints.PointsRoughness, 0.f);
+		const bool bHasAO = HasValues(HairStrands.StrandsPoints.PointsAO, 1.f);
 
 		TArray<uint32> AttributeRootUV;
 		TArray<uint32> AttributeSeed;
 		TArray<uint32> AttributeLength;
+		TArray<uint32> AttributeClumpIDs;
 		TArray<uint32> AttributeBaseColor;
 		TArray<uint32> AttributeRoughness;
+		TArray<uint32> AttributeAO;
 
 		AttributeRootUV.SetNum(NumCurves);									// 4 bytes encoding - Per-Curve
 		AttributeSeed.SetNum(FMath::DivideAndRoundUp(NumCurves, 4u));		// 1 bytes encoding - Per-Curve
 		AttributeLength.SetNum(FMath::DivideAndRoundUp(NumCurves, 2u));		// 2 bytes encoding - Per-Curve
+		AttributeClumpIDs.SetNum(FMath::DivideAndRoundUp(NumCurves, 2u));	// 2 bytes encoding - Per-Curve
 		AttributeBaseColor.SetNum(NumPoints);								// 4 bytes encoding - Per-Vertex
 		AttributeRoughness.SetNum(FMath::DivideAndRoundUp(NumPoints, 2u));	// 1 bytes encoding - Per-Vertex
+		AttributeAO.SetNum(FMath::DivideAndRoundUp(NumPoints, 4u));			// 1 bytes encoding - Per-Vertex
 
 		const FVector HairBoxCenter = HairStrands.BoundingBox.GetCenter();
 
@@ -364,7 +370,7 @@ namespace HairStrandsBuilder
 				}
 
 				// Per-Vertex BaseColor
-				if (bHasBaseColor)
+				if (bHasColor)
 				{
 					uint32& Packed = AttributeBaseColor[PointIndex + IndexOffset];
 
@@ -380,6 +386,13 @@ namespace HairStrandsBuilder
 					uint8* Packed = (uint8*)AttributeRoughness.GetData();
 					Packed[PointIndex + IndexOffset] = FMath::Clamp(uint32(Points.PointsRoughness[PointIndex + IndexOffset] * 0xFF), 0u, 0xFFu);
 				}
+
+				// Per-Vertex AO
+				if (bHasAO)
+				{
+					uint8* Packed = (uint8*)AttributeAO.GetData();
+					Packed[PointIndex + IndexOffset] = FMath::Clamp(uint32(Points.PointsAO[PointIndex + IndexOffset] * 0xFF), 0u, 0xFFu);
+				}
 			}
 
 			// Per-Curve Length
@@ -416,6 +429,12 @@ namespace HairStrandsBuilder
 				Packed |= (uint32(FMath::Clamp(TextureIndexUV.Y, 0.f, 31.f)) & 0x1F) << 27u;
 			}
 
+			// Per-Curve Clump ID
+			if (bHasClumpIDs)
+			{
+				uint8* Packed = (uint8*)AttributeClumpIDs.GetData();
+				Packed[CurveIndex] = Curves.ClumpIDs[CurveIndex];
+			}
 
 			// Curves
 			{
@@ -441,8 +460,10 @@ namespace HairStrandsBuilder
 				AttributeRootUV.Num() + 
 				AttributeSeed.Num() +
 				AttributeLength.Num() +
+				AttributeClumpIDs.Num() +
 				AttributeBaseColor.Num() +
-				AttributeRoughness.Num());
+				AttributeRoughness.Num() +
+				AttributeAO.Num());
 
 			const uint32 UintToByte = 4;
 
@@ -466,7 +487,13 @@ namespace HairStrandsBuilder
 				OutPackedAttributes.Append(AttributeLength);
 			}
 
-			if (bHasBaseColor)
+			if (bHasClumpIDs)
+			{
+				OutBulkData.AttributeOffsets[HAIR_ATTRIBUTE_CLUMPID] = OutPackedAttributes.Num() * UintToByte;
+				OutPackedAttributes.Append(AttributeClumpIDs);
+			}
+
+			if (bHasColor)
 			{
 				OutBulkData.AttributeOffsets[HAIR_ATTRIBUTE_BASECOLOR] = OutPackedAttributes.Num() * UintToByte;
 				OutPackedAttributes.Append(AttributeBaseColor);
@@ -477,6 +504,12 @@ namespace HairStrandsBuilder
 				OutBulkData.AttributeOffsets[HAIR_ATTRIBUTE_ROUGHNESS] = OutPackedAttributes.Num() * UintToByte;
 				OutPackedAttributes.Append(AttributeRoughness);
 			}
+
+			if (bHasAO)
+			{
+				OutBulkData.AttributeOffsets[HAIR_ATTRIBUTE_AO] = OutPackedAttributes.Num() * UintToByte;
+				OutPackedAttributes.Append(AttributeAO);
+			}
 		}
 
 		CopyToBulkData<FHairStrandsPositionFormat>(OutBulkData.Positions, OutPackedPositions);
@@ -1528,6 +1561,11 @@ bool FHairDescription::HasRootUV() const
 	return StrandAttributes().GetAttributesRef<FVector2f>(HairAttribute::Strand::RootUV).IsValid();
 }
 
+bool FHairDescription::HasClumpID() const
+{
+	return StrandAttributes().GetAttributesRef<FVector2f>(HairAttribute::Strand::ClumpID).IsValid();
+}
+
 bool FHairDescription::HasGuideWeights() const
 {
 	return
@@ -1542,18 +1580,17 @@ bool FHairDescription::HasGuideWeights() const
 
 bool FHairDescription::HasColorAttributes() const
 {
-	return
-		VertexAttributes().GetAttributesRef<FVector3f>(HairAttribute::Vertex::Color).IsValid() ||
-		StrandAttributes().GetAttributesRef<FVector3f>(HairAttribute::Strand::Color).IsValid() ||
-		GroomAttributes().GetAttributesRef<FVector3f>(HairAttribute::Groom::Color).IsValid();
+	return VertexAttributes().GetAttributesRef<FVector3f>(HairAttribute::Vertex::Color).IsValid();
 }
 
 bool FHairDescription::HasRoughnessAttributes() const
 {
-	return
-		VertexAttributes().GetAttributesRef<float>(HairAttribute::Vertex::Roughness).IsValid() ||
-		StrandAttributes().GetAttributesRef<float>(HairAttribute::Strand::Roughness).IsValid() ||
-		GroomAttributes().GetAttributesRef<float>(HairAttribute::Groom::Roughness).IsValid();
+	return VertexAttributes().GetAttributesRef<float>(HairAttribute::Vertex::Roughness).IsValid();
+}
+
+bool FHairDescription::HasAOAttributes() const
+{
+	return VertexAttributes().GetAttributesRef<float>(HairAttribute::Vertex::AO).IsValid();
 }
 
 bool FGroomBuilder::BuildHairDescriptionGroups(const FHairDescription& HairDescription, FHairDescriptionGroups& Out)
@@ -1588,8 +1625,10 @@ bool FGroomBuilder::BuildHairDescriptionGroups(const FHairDescription& HairDescr
 
 	TVertexAttributesConstRef<FVector3f> VertexPositions	= HairDescription.VertexAttributes().GetAttributesRef<FVector3f>(HairAttribute::Vertex::Position);
 	TVertexAttributesConstRef<FVector3f> VertexBaseColor	= HairDescription.VertexAttributes().GetAttributesRef<FVector3f>(HairAttribute::Vertex::Color);
-	TVertexAttributesConstRef<float> VertexRoughness	= HairDescription.VertexAttributes().GetAttributesRef<float>(HairAttribute::Vertex::Roughness);
-	TStrandAttributesConstRef<int> StrandNumVertices	= HairDescription.StrandAttributes().GetAttributesRef<int>(HairAttribute::Strand::VertexCount);
+	TVertexAttributesConstRef<float> VertexRoughness		= HairDescription.VertexAttributes().GetAttributesRef<float>(HairAttribute::Vertex::Roughness);
+	TVertexAttributesConstRef<float> VertexAO				= HairDescription.VertexAttributes().GetAttributesRef<float>(HairAttribute::Vertex::AO);
+	TStrandAttributesConstRef<int> StrandNumVertices		= HairDescription.StrandAttributes().GetAttributesRef<int>(HairAttribute::Strand::VertexCount);
+	TStrandAttributesConstRef<int> ClumpIDs					= HairDescription.StrandAttributes().GetAttributesRef<int>(HairAttribute::Strand::ClumpID);
 
 	if (!VertexPositions.IsValid() || !StrandNumVertices.IsValid())
 	{
@@ -1599,6 +1638,8 @@ bool FGroomBuilder::BuildHairDescriptionGroups(const FHairDescription& HairDescr
 
 	const bool bHasBaseColorAttribute = VertexBaseColor.IsValid();
 	const bool bHasRoughnessAttribute = VertexRoughness.IsValid();
+	const bool bHasAOAttribute = VertexAO.IsValid();
+	const bool bHasClumpIDs = ClumpIDs.IsValid();
 
 	TVertexAttributesConstRef<float> VertexWidths = HairDescription.VertexAttributes().GetAttributesRef<float>(HairAttribute::Vertex::Width);
 	TStrandAttributesConstRef<float> StrandWidths = HairDescription.StrandAttributes().GetAttributesRef<float>(HairAttribute::Strand::Width);
@@ -1729,6 +1770,11 @@ bool FGroomBuilder::BuildHairDescriptionGroups(const FHairDescription& HairDescr
 			}
 		}
 
+		if (bHasClumpIDs)
+		{
+			CurrentHairStrandsDatas->StrandsCurves.ClumpIDs.Add(ClumpIDs[StrandID]);
+		}
+
 		if (bHasUVData)
 		{
 			CurrentHairStrandsDatas->StrandsCurves.CurvesRootUV.Add(StrandRootUV[StrandID]);
@@ -1769,6 +1815,10 @@ bool FGroomBuilder::BuildHairDescriptionGroups(const FHairDescription& HairDescr
 			{
 				CurrentHairStrandsDatas->StrandsPoints.PointsRoughness.Add(VertexRoughness[VertexID]);
 			}
+			if (bHasAOAttribute)
+			{
+				CurrentHairStrandsDatas->StrandsPoints.PointsAO.Add(VertexAO[VertexID]);
+			}
 
 			// Vertex
 			float VertexWidth = 0.f;
@@ -2092,7 +2142,12 @@ void Decimate(
 	OutData.HairDensity = InData.HairDensity;
 
 	const bool bHasPrecomputedWeights = InData.StrandsCurves.CurvesClosestGuideIDs.Num() > 0 && InData.StrandsCurves.CurvesClosestGuideWeights.Num() > 0;
-	
+	const bool bHasColor = InData.StrandsPoints.PointsBaseColor.Num() > 0;
+	const bool bHasRoughness = InData.StrandsPoints.PointsRoughness.Num() > 0;
+	const bool bHasAO = InData.StrandsPoints.PointsAO.Num() > 0;
+	const bool bHasClumpIDs = InData.StrandsCurves.ClumpIDs.Num() > 0;
+	const bool bHasStrandIDs = InData.StrandsCurves.StrandIDs.Num() > 0;
+
 	PointCount = 0;
 	for(uint32 CurveIndex = 0; CurveIndex < CurveCount; ++CurveIndex)
 	{
@@ -2107,6 +2162,14 @@ void Decimate(
 		OutData.StrandsCurves.CurvesRootUV[CurveIndex] = InData.StrandsCurves.CurvesRootUV[GuideIndex];
 		OutData.StrandsCurves.CurvesOffset[CurveIndex] = PointCount;
 		OutData.StrandsCurves.CurvesLength[CurveIndex] = InData.StrandsCurves.CurvesLength[GuideIndex];
+		if (bHasStrandIDs)
+		{
+			OutData.StrandsCurves.StrandIDs[CurveIndex] = InData.StrandsCurves.StrandIDs[GuideIndex];
+		}
+		if (bHasClumpIDs)
+		{
+			OutData.StrandsCurves.ClumpIDs[CurveIndex] = InData.StrandsCurves.ClumpIDs[GuideIndex];
+		}
 		if (bHasPrecomputedWeights)
 		{
 			OutData.StrandsCurves.CurvesClosestGuideIDs[CurveIndex] = InData.StrandsCurves.CurvesClosestGuideIDs[GuideIndex];
@@ -2120,20 +2183,15 @@ void Decimate(
 			const float PointSpace = GuideOffset + PointIndex * BoneSpace;
 	
 			const int32 PointBegin = FMath::Min(FMath::FloorToInt32(PointSpace), GuideOffset + GuideCount-2);
-			const int32 PointEnd = PointBegin+1;
-	
+			const int32 PointEnd = PointBegin+1;	
 			const float PointAlpha = PointSpace - PointBegin;
 			
-			OutData.StrandsPoints.PointsPosition[PointCount] = InData.StrandsPoints.PointsPosition[PointBegin] * (1.0-PointAlpha) +
-				InData.StrandsPoints.PointsPosition[PointEnd] * PointAlpha;
-			OutData.StrandsPoints.PointsCoordU[PointCount] = InData.StrandsPoints.PointsCoordU[PointBegin] * (1.0-PointAlpha) +
-				InData.StrandsPoints.PointsCoordU[PointEnd] * PointAlpha;
-			OutData.StrandsPoints.PointsRadius[PointCount] = InData.StrandsPoints.PointsRadius[PointBegin] * (1.0-PointAlpha) +
-				InData.StrandsPoints.PointsRadius[PointEnd] * PointAlpha;
-			OutData.StrandsPoints.PointsBaseColor[PointCount] = InData.StrandsPoints.PointsBaseColor[PointBegin] * (1.0-PointAlpha) +
-				InData.StrandsPoints.PointsBaseColor[PointEnd] * PointAlpha;
-			OutData.StrandsPoints.PointsRoughness[PointCount] = InData.StrandsPoints.PointsRoughness[PointBegin] * (1.0-PointAlpha) +
-				InData.StrandsPoints.PointsRoughness[PointEnd] * PointAlpha;
+			OutData.StrandsPoints.PointsPosition[PointCount]	= FMath::Lerp(InData.StrandsPoints.PointsPosition[PointBegin],	InData.StrandsPoints.PointsPosition[PointEnd],	PointAlpha);
+			OutData.StrandsPoints.PointsCoordU[PointCount]		= FMath::Lerp(InData.StrandsPoints.PointsCoordU[PointBegin],	InData.StrandsPoints.PointsCoordU[PointEnd],	PointAlpha);
+			OutData.StrandsPoints.PointsRadius[PointCount]		= FMath::Lerp(InData.StrandsPoints.PointsRadius[PointBegin],	InData.StrandsPoints.PointsRadius[PointEnd],	PointAlpha);
+			if (bHasColor)		{ OutData.StrandsPoints.PointsBaseColor[PointCount]	= FMath::Lerp(InData.StrandsPoints.PointsBaseColor[PointBegin], InData.StrandsPoints.PointsBaseColor[PointEnd], PointAlpha); }
+			if (bHasRoughness)	{ OutData.StrandsPoints.PointsRoughness[PointCount]	= FMath::Lerp(InData.StrandsPoints.PointsRoughness[PointBegin], InData.StrandsPoints.PointsRoughness[PointEnd], PointAlpha); }
+			if (bHasAO)			{ OutData.StrandsPoints.PointsAO[PointCount]		= FMath::Lerp(InData.StrandsPoints.PointsAO[PointBegin],		InData.StrandsPoints.PointsAO[PointEnd],		PointAlpha); }
 		}
 	}
 	HairStrandsBuilder::BuildInternalData(OutData, false);
@@ -2154,6 +2212,11 @@ void Decimate(
 	const uint32 InCurveCount = InData.StrandsCurves.Num();
 	const uint32 OutCurveCount = FMath::Clamp(uint32(InCurveCount * CurveDecimationPercentage), 1u, InCurveCount);
 	const bool bHasPrecomputedWeights = InData.StrandsCurves.CurvesClosestGuideIDs.Num() > 0 && InData.StrandsCurves.CurvesClosestGuideWeights.Num() > 0;
+	const bool bHasColor = InData.StrandsPoints.PointsBaseColor.Num() > 0;
+	const bool bHasRoughness = InData.StrandsPoints.PointsRoughness.Num() > 0;
+	const bool bHasAO = InData.StrandsPoints.PointsAO.Num() > 0;
+	const bool bHasClumpIDs = InData.StrandsCurves.ClumpIDs.Num() > 0;
+	const bool bHasStrandIDs = InData.StrandsCurves.StrandIDs.Num() > 0;
 
 	TArray<uint32> CurveIndices;
 	CurveIndices.SetNum(OutCurveCount);
@@ -2207,6 +2270,17 @@ void Decimate(
 	OutData.StrandsPoints.SetNum(OutTotalPointCount);
 	OutData.HairDensity = InData.HairDensity;
 
+	if (bHasColor)		{ OutData.StrandsPoints.PointsBaseColor.SetNum(OutTotalPointCount); }
+	if (bHasRoughness)	{ OutData.StrandsPoints.PointsRoughness.SetNum(OutTotalPointCount); }
+	if (bHasAO)			{ OutData.StrandsPoints.PointsAO.SetNum(OutTotalPointCount); }
+	if (bHasStrandIDs)	{ OutData.StrandsCurves.StrandIDs.SetNum(OutCurveCount); }
+	if (bHasClumpIDs)	{ OutData.StrandsCurves.ClumpIDs.SetNum(OutCurveCount); }
+	if (bHasPrecomputedWeights)
+	{
+		OutData.StrandsCurves.CurvesClosestGuideIDs.SetNum(OutCurveCount);
+		OutData.StrandsCurves.CurvesClosestGuideWeights.SetNum(OutCurveCount);
+	}
+
 	uint32 OutPointOffset = 0; 
 	for (uint32 OutCurveIndex = 0; OutCurveIndex < OutCurveCount; ++OutCurveIndex)
 	{
@@ -2237,7 +2311,6 @@ void Decimate(
 		}
 
 		const uint32 OutPointCount = OutPointIndices.Num();
-
 		for (uint32 OutPointIndex = 0; OutPointIndex < OutPointCount; ++OutPointIndex)
 		{
 			const uint32 InPointIndex = OutPointIndices[OutPointIndex];
@@ -2245,8 +2318,9 @@ void Decimate(
 			OutData.StrandsPoints.PointsPosition [OutPointIndex + OutPointOffset] = InData.StrandsPoints.PointsPosition	[InPointIndex + InPointOffset];
 			OutData.StrandsPoints.PointsCoordU	 [OutPointIndex + OutPointOffset] = InData.StrandsPoints.PointsCoordU	[InPointIndex + InPointOffset];
 			OutData.StrandsPoints.PointsRadius	 [OutPointIndex + OutPointOffset] = InData.StrandsPoints.PointsRadius	[InPointIndex + InPointOffset];
-			OutData.StrandsPoints.PointsBaseColor[OutPointIndex + OutPointOffset] = InData.StrandsPoints.PointsBaseColor[InPointIndex + InPointOffset];
-			OutData.StrandsPoints.PointsRoughness[OutPointIndex + OutPointOffset] = InData.StrandsPoints.PointsRoughness[InPointIndex + InPointOffset];
+			if (bHasColor)		{ OutData.StrandsPoints.PointsBaseColor[OutPointIndex + OutPointOffset] = InData.StrandsPoints.PointsBaseColor[InPointIndex + InPointOffset]; }
+			if (bHasRoughness)	{ OutData.StrandsPoints.PointsRoughness[OutPointIndex + OutPointOffset] = InData.StrandsPoints.PointsRoughness[InPointIndex + InPointOffset]; }
+			if (bHasAO)			{ OutData.StrandsPoints.PointsAO[OutPointIndex + OutPointOffset]		= InData.StrandsPoints.PointsAO[InPointIndex + InPointOffset]; }
 		}
 	#else
 		// Decimation using uniform/interleaved removal
@@ -2268,8 +2342,9 @@ void Decimate(
 			OutData.StrandsPoints.PointsPosition [OutPointIndex + OutPointOffset] = InData.StrandsPoints.PointsPosition	[InPointIndex + InPointOffset];
 			OutData.StrandsPoints.PointsCoordU	 [OutPointIndex + OutPointOffset] = InData.StrandsPoints.PointsCoordU	[InPointIndex + InPointOffset];
 			OutData.StrandsPoints.PointsRadius	 [OutPointIndex + OutPointOffset] = InData.StrandsPoints.PointsRadius	[InPointIndex + InPointOffset];
-			OutData.StrandsPoints.PointsBaseColor[OutPointIndex + OutPointOffset] = FLinearColor::Black;
-			OutData.StrandsPoints.PointsRoughness[OutPointIndex + OutPointOffset] = 0;
+			if (bHasColor)		{ OutData.StrandsPoints.PointsBaseColor[OutPointIndex + OutPointOffset] = InData.StrandsPoints.PointsBaseColor[InPointIndex + InPointOffset]; }
+			if (bHasRoughness)	{ OutData.StrandsPoints.PointsRoughness[OutPointIndex + OutPointOffset] = InData.StrandsPoints.PointsRoughness[InPointIndex + InPointOffset]; }
+			if (bHasAO)			{ OutData.StrandsPoints.PointsAO[OutPointIndex + OutPointOffset]		= InData.StrandsPoints.PointsAO[InPointIndex + InPointOffset]; }
 		}
 	#endif
 
@@ -2282,6 +2357,14 @@ void Decimate(
 			OutData.StrandsCurves.CurvesClosestGuideIDs[OutCurveIndex] = InData.StrandsCurves.CurvesClosestGuideIDs[EffectiveCurveIndex];
 			OutData.StrandsCurves.CurvesClosestGuideWeights[OutCurveIndex] = InData.StrandsCurves.CurvesClosestGuideWeights[EffectiveCurveIndex];
 		}
+		if (bHasStrandIDs)
+		{
+			OutData.StrandsCurves.StrandIDs[OutCurveIndex] = InData.StrandsCurves.StrandIDs[EffectiveCurveIndex];
+		}
+		if (bHasClumpIDs)
+		{
+			OutData.StrandsCurves.ClumpIDs[OutCurveIndex] = InData.StrandsCurves.ClumpIDs[EffectiveCurveIndex];
+		}
 		OutData.StrandsCurves.MaxLength = InData.StrandsCurves.MaxLength;
 		OutData.StrandsCurves.MaxRadius = InData.StrandsCurves.MaxRadius;
 		OutPointOffset += OutPointCount;
@@ -2399,6 +2482,9 @@ static void Voxelize(const FHairDescriptionGroups& InGroups, FHairGrid& Out)
 		FHairStrandsDatas In = Group.Strands;
 		HairStrandsBuilder::BuildInternalData(In, false);
 
+		const bool bHasColor     = In.StrandsPoints.PointsBaseColor.Num() > 0;
+		const bool bHasRoughness = In.StrandsPoints.PointsRoughness.Num() > 0;
+
 		// Fill in voxel (TODO: make it parallel)
 		for (uint32 CurveIt = 0, CurveCount = In.GetNumCurves(); CurveIt < CurveCount; ++CurveIt)
 		{
@@ -2414,10 +2500,10 @@ static void Voxelize(const FHairDescriptionGroups& InGroups, FHairGrid& Out)
 				const FVector3f& P1 = In.StrandsPoints.PointsPosition[Index1];
 				const float R0 = In.StrandsPoints.PointsRadius[Index0];
 				const float R1 = In.StrandsPoints.PointsRadius[Index1];
-				const FVector3f C0(In.StrandsPoints.PointsBaseColor[Index0].R, In.StrandsPoints.PointsBaseColor[Index0].G, In.StrandsPoints.PointsBaseColor[Index0].B);
-				const FVector3f C1(In.StrandsPoints.PointsBaseColor[Index1].R, In.StrandsPoints.PointsBaseColor[Index1].G, In.StrandsPoints.PointsBaseColor[Index1].B);
-				const float Rough0 = In.StrandsPoints.PointsRoughness[Index0];
-				const float Rough1 = In.StrandsPoints.PointsRoughness[Index1];
+				const FVector3f C0 = bHasColor ? FVector3f(In.StrandsPoints.PointsBaseColor[Index0].R, In.StrandsPoints.PointsBaseColor[Index0].G, In.StrandsPoints.PointsBaseColor[Index0].B) : FVector3f(0);
+				const FVector3f C1 = bHasColor ? FVector3f(In.StrandsPoints.PointsBaseColor[Index1].R, In.StrandsPoints.PointsBaseColor[Index1].G, In.StrandsPoints.PointsBaseColor[Index1].B) : FVector3f(0);
+				const float Rough0 = bHasRoughness ? In.StrandsPoints.PointsRoughness[Index0] : 0;
+				const float Rough1 = bHasRoughness ? In.StrandsPoints.PointsRoughness[Index1] : 0;
 				const FVector3f Segment = P1 - P0;
 
 				// This is a coarse/non-conservative voxelization, by ray-marching segment, instead of walking intersected voxels
diff --git a/Engine/Plugins/Runtime/HairStrands/Source/HairStrandsCore/Private/GroomEdit.cpp b/Engine/Plugins/Runtime/HairStrands/Source/HairStrandsCore/Private/GroomEdit.cpp
--- a/Engine/Plugins/Runtime/HairStrands/Source/HairStrandsCore/Private/GroomEdit.cpp
+++ b/Engine/Plugins/Runtime/HairStrands/Source/HairStrandsCore/Private/GroomEdit.cpp
@@ -37,6 +37,7 @@ static void HairStrandsDataToEditableHairGuides(const FHairStrandsDatas& In, TAr
 static void HairStrandsDataToEditableHairStrands(const FHairStrandsDatas& In, TArray<FEditableHairStrand>& Out)
 {
 	const bool bHasValidStrandsIDs = In.StrandsCurves.StrandIDs.Num() > 0;
+	const bool bHasValidClumpIDs = In.StrandsCurves.ClumpIDs.Num() > 0;	
 	const bool bHasValidClosestGuide = In.StrandsCurves.CurvesClosestGuideIDs.Num() > 0;
 	const uint32 StrandsCount = In.GetNumCurves();
 	Out.SetNum(StrandsCount);
@@ -46,6 +47,7 @@ static void HairStrandsDataToEditableHairStrands(const FHairStrandsDatas& In, TA
 		const FVector3f GuideWeights = bHasValidClosestGuide ? (FVector3f)In.StrandsCurves.CurvesClosestGuideIDs[StrandsIt] : FVector3f::ZeroVector;
 
 		Out[StrandsIt].StrandID		= bHasValidStrandsIDs ? In.StrandsCurves.StrandIDs[StrandsIt] : StrandsIt;
+		Out[StrandsIt].ClumpID		= bHasValidClumpIDs ? In.StrandsCurves.ClumpIDs[StrandsIt] : 0;
 		Out[StrandsIt].RootUV		= In.StrandsCurves.CurvesRootUV[StrandsIt];
 		Out[StrandsIt].GuideIDs[0]  = GuideIDs.X;
 		Out[StrandsIt].GuideIDs[1]  = GuideIDs.Y;
@@ -66,6 +68,7 @@ static void HairStrandsDataToEditableHairStrands(const FHairStrandsDatas& In, TA
 			Out[StrandsIt].ControlPoints[PointIt].Radius	= In.StrandsPoints.PointsRadius[EffectiveIndex] * In.StrandsCurves.MaxRadius;
 			Out[StrandsIt].ControlPoints[PointIt].BaseColor	= In.StrandsPoints.PointsBaseColor[EffectiveIndex];
 			Out[StrandsIt].ControlPoints[PointIt].Roughness	= In.StrandsPoints.PointsRoughness[EffectiveIndex];
+			Out[StrandsIt].ControlPoints[PointIt].AO		= In.StrandsPoints.PointsAO[EffectiveIndex];
 		}
 	}
 }
@@ -112,6 +115,7 @@ static void EditableGroomGroupToHairDescription(FHairDescription& OutHairDescrip
 		SetHairStrandAttribute(OutHairDescription, StrandID, HairAttribute::Strand::VertexCount, (int32)PointCount);
 		SetHairStrandAttribute(OutHairDescription, StrandID, HairAttribute::Strand::RootUV,      Strands.RootUV);
 		SetHairStrandAttribute(OutHairDescription, StrandID, HairAttribute::Strand::ID,          int(Strands.StrandID));
+		SetHairStrandAttribute(OutHairDescription, StrandID, HairAttribute::Strand::ClumpID,     int(Strands.ClumpID));
 		SetHairStrandAttribute(OutHairDescription, StrandID, HairAttribute::Strand::GroupID,     int(GroupID));
 		SetHairStrandAttribute(OutHairDescription, StrandID, HairAttribute::Strand::GroupName,   GroupName);
 		SetHairStrandAttribute(OutHairDescription, StrandID, HairAttribute::Strand::Guide,       0);
@@ -138,6 +142,7 @@ static void EditableGroomGroupToHairDescription(FHairDescription& OutHairDescrip
 			SetHairVertexAttribute(OutHairDescription, VertexID, HairAttribute::Vertex::Width,     Point.Radius * 2.f);
 			SetHairVertexAttribute(OutHairDescription, VertexID, HairAttribute::Vertex::Color,     FVector3f(Point.BaseColor.R, Point.BaseColor.G, Point.BaseColor.B));
 			SetHairVertexAttribute(OutHairDescription, VertexID, HairAttribute::Vertex::Roughness, Point.Roughness);
+			SetHairVertexAttribute(OutHairDescription, VertexID, HairAttribute::Vertex::AO,        Point.AO);
 		}
 	}
 }
diff --git a/Engine/Plugins/Runtime/HairStrands/Source/HairStrandsCore/Private/HairAttributes.cpp b/Engine/Plugins/Runtime/HairStrands/Source/HairStrandsCore/Private/HairAttributes.cpp
--- a/Engine/Plugins/Runtime/HairStrands/Source/HairStrandsCore/Private/HairAttributes.cpp
+++ b/Engine/Plugins/Runtime/HairStrands/Source/HairStrandsCore/Private/HairAttributes.cpp
@@ -6,6 +6,7 @@ namespace HairAttribute
 {
 	const FName Vertex::Color("groom_color");
 	const FName Vertex::Roughness("groom_roughness");
+	const FName Vertex::AO("groom_ao");
 	const FName Vertex::Position("position");
 	const FName Vertex::Width("groom_width");
 
@@ -14,6 +15,7 @@ namespace HairAttribute
 	const FName Strand::GroupID("groom_group_id");
 	const FName Strand::Guide("groom_guide");
 	const FName Strand::ID("groom_id");
+	const FName Strand::ClumpID("groom_clumpid");
 	const FName Strand::RootUV("groom_root_uv");
 	const FName Strand::VertexCount("vertexcount");
 	const FName Strand::Width("groom_width");
diff --git a/Engine/Plugins/Runtime/HairStrands/Source/HairStrandsCore/Private/HairCardsBuilder.cpp b/Engine/Plugins/Runtime/HairStrands/Source/HairStrandsCore/Private/HairCardsBuilder.cpp
--- a/Engine/Plugins/Runtime/HairStrands/Source/HairStrandsCore/Private/HairCardsBuilder.cpp
+++ b/Engine/Plugins/Runtime/HairStrands/Source/HairStrandsCore/Private/HairCardsBuilder.cpp
@@ -1341,6 +1341,7 @@ namespace HairCards
 		const uint32 OutCurveCount = FMath::Clamp(uint32(CurveCount * DecimationPercentage), 1u, CurveCount);
 
 		const uint32 BucketSize = CurveCount / OutCurveCount;
+		const bool bHasClumpID = InData.StrandsCurves.ClumpIDs.Num() > 0;
 
 		TArray<uint32> CurveIndices;
 		CurveIndices.SetNum(OutCurveCount);
@@ -1371,14 +1372,16 @@ namespace HairCards
 			OutData.StrandsCurves.CurvesLength[OutCurveIndex] = InData.StrandsCurves.CurvesLength[InCurveIndex] * InData.StrandsCurves.MaxLength;
 			OutData.StrandsCurves.MaxLength = InData.StrandsCurves.MaxLength;
 			OutData.StrandsCurves.MaxRadius = InData.StrandsCurves.MaxRadius;
+			if (bHasClumpID)
+			{
+				OutData.StrandsCurves.CurvesOffset[OutCurveIndex] = InData.StrandsCurves.ClumpIDs[InCurveIndex];
+			}
 
 			for (uint32 PointIndex = 0; PointIndex < PointCount; ++PointIndex)
 			{
 				OutData.StrandsPoints.PointsPosition[PointIndex + OutPointOffset] = InData.StrandsPoints.PointsPosition[PointIndex + InPointOffset];
 				OutData.StrandsPoints.PointsCoordU[PointIndex + OutPointOffset] = InData.StrandsPoints.PointsCoordU[PointIndex + InPointOffset];
 				OutData.StrandsPoints.PointsRadius[PointIndex + OutPointOffset] = InData.StrandsPoints.PointsRadius[PointIndex + InPointOffset];
-				OutData.StrandsPoints.PointsBaseColor[PointIndex + OutPointOffset] = FLinearColor::Black;
-				OutData.StrandsPoints.PointsRoughness[PointIndex + OutPointOffset] = 0;
 			}
 			OutPointOffset += PointCount;
 		}
@@ -2836,8 +2839,6 @@ namespace HairCards
 		OutGuides.StrandsPoints.PointsPosition.Reserve(MaxGuidePoints);
 		OutGuides.StrandsPoints.PointsRadius.Reserve(MaxGuidePoints);
 		OutGuides.StrandsPoints.PointsCoordU.Reserve(MaxGuidePoints);
-		OutGuides.StrandsPoints.PointsBaseColor.Reserve(MaxGuidePoints);
-		OutGuides.StrandsPoints.PointsRoughness.Reserve(MaxGuidePoints);
 		OutGuides.StrandsCurves.SetNum(NumCards);
 		OutGuides.BoundingBox.Init();
 
@@ -3090,8 +3091,6 @@ namespace HairCards
 				OutGuides.BoundingBox += (FVector)P0;
 
 				OutGuides.StrandsPoints.PointsPosition.Add(P0);
-				OutGuides.StrandsPoints.PointsBaseColor.Add(FLinearColor(FVector3f::ZeroVector));
-				OutGuides.StrandsPoints.PointsRoughness.Add(0);
 				OutGuides.StrandsPoints.PointsCoordU.Add(FMath::Clamp(CurrentLength / TotalLength, 0.f, 1.f));
 				OutGuides.StrandsPoints.PointsRadius.Add(1);
 
diff --git a/Engine/Plugins/Runtime/HairStrands/Source/HairStrandsCore/Private/HairStrandsDatas.cpp b/Engine/Plugins/Runtime/HairStrands/Source/HairStrandsCore/Private/HairStrandsDatas.cpp
--- a/Engine/Plugins/Runtime/HairStrands/Source/HairStrandsCore/Private/HairStrandsDatas.cpp
+++ b/Engine/Plugins/Runtime/HairStrands/Source/HairStrandsCore/Private/HairStrandsDatas.cpp
@@ -19,8 +19,8 @@ void FHairStrandsCurves::SetNum(const uint32 NumCurves)
 	CurvesLength.SetNum(NumCurves);
 	CurvesRootUV.SetNum(NumCurves);
 
-	// Curves closest guide are not initialized to track if the curves 
-	// data contains or not precomputed weights
+	// Not initialized to track if the data are available
+	// ClumpIDs.SetNum(0);
 	// CurvesClosestGuideIDs.SetNum(0);
 	// CurvesClosestGuideWeights.SetNum(0);
 }
@@ -30,8 +30,11 @@ void FHairStrandsPoints::SetNum(const uint32 NumPoints)
 	PointsPosition.SetNum(NumPoints);
 	PointsRadius.SetNum(NumPoints);
 	PointsCoordU.SetNum(NumPoints);
-	PointsBaseColor.SetNum(NumPoints);
-	PointsRoughness.SetNum(NumPoints);
+
+	// Not initialized to track if the data are available
+	// PointsBaseColor.SetNum(0);
+	// PointsRoughness.SetNum(0);
+	// PointsAO.SetNum(0);
 }
 
 void FHairStrandsInterpolationDatas::Reset()
@@ -48,6 +51,7 @@ void FHairStrandsCurves::Reset()
 	CurvesCount.Reset();
 	CurvesLength.Reset();
 	CurvesRootUV.Reset();
+	ClumpIDs.Reset();
 	CurvesClosestGuideIDs.Reset();
 	CurvesClosestGuideWeights.Reset();
 }
@@ -57,6 +61,9 @@ void FHairStrandsPoints::Reset()
 	PointsPosition.Reset();
 	PointsRadius.Reset();
 	PointsCoordU.Reset();
+	PointsBaseColor.Reset();
+	PointsRoughness.Reset();
+	PointsAO.Reset();
 }
 
 FArchive& operator<<(FArchive& Ar, FVector4_16& Vertex)
diff --git a/Engine/Plugins/Runtime/HairStrands/Source/HairStrandsCore/Public/GroomEdit.h b/Engine/Plugins/Runtime/HairStrands/Source/HairStrandsCore/Public/GroomEdit.h
--- a/Engine/Plugins/Runtime/HairStrands/Source/HairStrandsCore/Public/GroomEdit.h
+++ b/Engine/Plugins/Runtime/HairStrands/Source/HairStrandsCore/Public/GroomEdit.h
@@ -17,6 +17,7 @@ struct FEditableHairStrandControlPoint
 	float U;
 	FLinearColor BaseColor;
 	float Roughness;
+	float AO;
 };
 
 struct FEditableHairStrand
@@ -24,6 +25,7 @@ struct FEditableHairStrand
 	TArray<FEditableHairStrandControlPoint> ControlPoints;
 
 	uint32 StrandID;
+	uint32 ClumpID;
 	FVector2f RootUV;
 
 	// Pre-computed Simulation
diff --git a/Engine/Plugins/Runtime/HairStrands/Source/HairStrandsCore/Public/GroomImportOptions.h b/Engine/Plugins/Runtime/HairStrands/Source/HairStrandsCore/Public/GroomImportOptions.h
--- a/Engine/Plugins/Runtime/HairStrands/Source/HairStrandsCore/Public/GroomImportOptions.h
+++ b/Engine/Plugins/Runtime/HairStrands/Source/HairStrandsCore/Public/GroomImportOptions.h
@@ -32,8 +32,10 @@ struct HAIRSTRANDSCORE_API FGroomHairGroupPreview
 	, CurveCount(0)
 	, GuideCount(0)
 	, bHasRootUV(false)
-	, bHasColorAttributes(false)
-	, bHasRoughnessAttributes(false)
+	, bHasClumpID(false)
+	, bHasColor(false)
+	, bHasRoughness(false)
+	, bHasAO(false)
 	, bHasPrecomputedWeights(false)
 	, InterpolationSettings()
 	{}
@@ -54,10 +56,16 @@ struct HAIRSTRANDSCORE_API FGroomHairGroupPreview
 	bool bHasRootUV;
 
 	UPROPERTY(VisibleAnywhere, BlueprintReadWrite, Category = Preview)
-	bool bHasColorAttributes;
+	bool bHasClumpID;
 
 	UPROPERTY(VisibleAnywhere, BlueprintReadWrite, Category = Preview)
-	bool bHasRoughnessAttributes;
+	bool bHasColor;
+
+	UPROPERTY(VisibleAnywhere, BlueprintReadWrite, Category = Preview)
+	bool bHasRoughness;
+
+	UPROPERTY(VisibleAnywhere, BlueprintReadWrite, Category = Preview)
+	bool bHasAO;
 
 	UPROPERTY(VisibleAnywhere, BlueprintReadWrite, Category = Preview)
 	bool bHasPrecomputedWeights;
diff --git a/Engine/Plugins/Runtime/HairStrands/Source/HairStrandsCore/Public/HairAttributes.h b/Engine/Plugins/Runtime/HairStrands/Source/HairStrandsCore/Public/HairAttributes.h
--- a/Engine/Plugins/Runtime/HairStrands/Source/HairStrandsCore/Public/HairAttributes.h
+++ b/Engine/Plugins/Runtime/HairStrands/Source/HairStrandsCore/Public/HairAttributes.h
@@ -10,6 +10,7 @@ namespace HairAttribute
 	{
 		extern HAIRSTRANDSCORE_API const FName Color;			// FVector
 		extern HAIRSTRANDSCORE_API const FName Roughness;		// float
+		extern HAIRSTRANDSCORE_API const FName AO;				// float
 		extern HAIRSTRANDSCORE_API const FName Position;		// FVector
 		extern HAIRSTRANDSCORE_API const FName Width;			// float
 	}
@@ -21,6 +22,7 @@ namespace HairAttribute
 		extern HAIRSTRANDSCORE_API const FName GroupID;			// int
 		extern HAIRSTRANDSCORE_API const FName Guide;			// int
 		extern HAIRSTRANDSCORE_API const FName ID;				// int
+		extern HAIRSTRANDSCORE_API const FName ClumpID;			// int
 		extern HAIRSTRANDSCORE_API const FName RootUV;			// FVector2f
 		extern HAIRSTRANDSCORE_API const FName VertexCount;		// int
 		extern HAIRSTRANDSCORE_API const FName Width;			// float
diff --git a/Engine/Plugins/Runtime/HairStrands/Source/HairStrandsCore/Public/HairDescription.h b/Engine/Plugins/Runtime/HairStrands/Source/HairStrandsCore/Public/HairDescription.h
--- a/Engine/Plugins/Runtime/HairStrands/Source/HairStrandsCore/Public/HairDescription.h
+++ b/Engine/Plugins/Runtime/HairStrands/Source/HairStrandsCore/Public/HairDescription.h
@@ -96,9 +96,11 @@ public:
 	// #ueent_todo: Expose "remove" operations to allow editing HairDescription
 
 	bool HasRootUV() const;
+	bool HasClumpID() const;
 	bool HasGuideWeights() const;
 	bool HasColorAttributes() const;
 	bool HasRoughnessAttributes() const;
+	bool HasAOAttributes() const;
 
 private:
 	TAttributesSet<FVertexID> VertexAttributesSet;
diff --git a/Engine/Plugins/Runtime/HairStrands/Source/HairStrandsCore/Public/HairStrandsDatas.h b/Engine/Plugins/Runtime/HairStrands/Source/HairStrandsCore/Public/HairStrandsDatas.h
--- a/Engine/Plugins/Runtime/HairStrands/Source/HairStrandsCore/Public/HairStrandsDatas.h
+++ b/Engine/Plugins/Runtime/HairStrands/Source/HairStrandsCore/Public/HairStrandsDatas.h
@@ -359,11 +359,14 @@ struct HAIRSTRANDSCORE_API FHairStrandsPoints
 	/** Normalized length */
 	TArray<float> PointsCoordU; // [0..1]
 
-	/** Material base color */
+	/** Material per-vertex 'baked' base color (optional) */
 	TArray<FLinearColor> PointsBaseColor; // [0..1]
 
-	/** Material roughness */
+	/** Material per-vertex 'baked' roughness (optional) */
 	TArray<float> PointsRoughness; // [0..1]
+
+	/** Material per-vertex 'baked' AO (optional) */
+	TArray<float> PointsAO; // [0..1]	
 };
 
 /** Hair strands Curves attribute */
@@ -395,13 +398,16 @@ struct HAIRSTRANDSCORE_API FHairStrandsCurves
 	/** Strand ID associated with each curve */
 	TArray<int> StrandIDs;
 
+	/** Clump ID associated with each curve (optional) */
+	TArray<uint16> ClumpIDs;
+
 	/** Mapping of imported Groom ID to index */
 	TMap<int, int> GroomIDToIndex;
 
-	/** Custom guide IDs (indexed with StrandID) */
+	/** Custom guide IDs (indexed with StrandID) (optional) */
 	TArray<FIntVector> CurvesClosestGuideIDs;
 
-	/** Custom guid weights (indexed with StrandID) */
+	/** Custom guid weights (indexed with StrandID) (optional) */
 	TArray<FVector> CurvesClosestGuideWeights;
 
 	/** Max strands Curves length */
diff --git a/Engine/Plugins/Runtime/HairStrands/Source/HairStrandsEditor/Private/GroomActions.cpp b/Engine/Plugins/Runtime/HairStrands/Source/HairStrandsEditor/Private/GroomActions.cpp
--- a/Engine/Plugins/Runtime/HairStrands/Source/HairStrandsEditor/Private/GroomActions.cpp
+++ b/Engine/Plugins/Runtime/HairStrands/Source/HairStrandsEditor/Private/GroomActions.cpp
@@ -184,10 +184,6 @@ void FGroomActions::ExecuteRebuild(TArray<TWeakObjectPtr<UGroomAsset>> Objects)
 						OutGroup.GroupName	= GroomAsset->HairGroupsInfo[GroupIndex].GroupName;
 						OutGroup.CurveCount = GroomAsset->HairGroupsData[GroupIndex].Strands.BulkData.GetNumCurves();
 						OutGroup.GuideCount = GroomAsset->HairGroupsData[GroupIndex].Guides.BulkData.GetNumCurves();
-						OutGroup.bHasRootUV = false;
-						OutGroup.bHasColorAttributes = false;
-						OutGroup.bHasRoughnessAttributes = false;
-						OutGroup.bHasPrecomputedWeights = false;
 						OutGroup.InterpolationSettings = GroomAsset->HairGroupsInterpolation[GroupIndex];
 						OutGroup.InterpolationSettings.RiggingSettings.bCanEditRigging = true;
 					}
diff --git a/Engine/Plugins/Runtime/HairStrands/Source/HairStrandsEditor/Private/GroomImportOptionsWindow.cpp b/Engine/Plugins/Runtime/HairStrands/Source/HairStrandsEditor/Private/GroomImportOptionsWindow.cpp
--- a/Engine/Plugins/Runtime/HairStrands/Source/HairStrandsEditor/Private/GroomImportOptionsWindow.cpp
+++ b/Engine/Plugins/Runtime/HairStrands/Source/HairStrandsEditor/Private/GroomImportOptionsWindow.cpp
@@ -214,8 +214,10 @@ void SGroomImportOptionsWindow::Construct(const FArguments& InArgs)
 	const FText FalseText = LOCTEXT("GroomOptionsWindow_AttributeFalse", "False");
 
 	FText HasRootUVText = FalseText;
-	FText HasColorAttributeText = FalseText;
-	FText HasRoughnessAttributeText = FalseText;
+	FText HasClumpIDText = FalseText;
+	FText HasColorText = FalseText;
+	FText HasRoughnessText = FalseText;
+	FText HasAOText = FalseText;
 	FText HasGuideWeightsText = FalseText;
 
 	if (GroupsPreview)
@@ -223,8 +225,10 @@ void SGroomImportOptionsWindow::Construct(const FArguments& InArgs)
 		for (const FGroomHairGroupPreview& Group : GroupsPreview->Groups)
 		{
 			if (Group.bHasRootUV)				{ HasRootUVText = TrueText; }
-			if (Group.bHasColorAttributes)		{ HasColorAttributeText = TrueText; }
-			if (Group.bHasRoughnessAttributes)	{ HasRoughnessAttributeText = TrueText; }
+			if (Group.bHasClumpID)				{ HasClumpIDText = TrueText; }
+			if (Group.bHasColor)				{ HasColorText = TrueText; }
+			if (Group.bHasRoughness)			{ HasRoughnessText = TrueText; }
+			if (Group.bHasAO)					{ HasAOText = TrueText; }
 			if (Group.bHasPrecomputedWeights)	{ HasGuideWeightsText = TrueText; }
 		}
 	}
@@ -320,6 +324,35 @@ void SGroomImportOptionsWindow::Construct(const FArguments& InArgs)
 			]
 		]
 
+		// Clump ID
+		+ SVerticalBox::Slot()
+		.AutoHeight()
+		.Padding(2)
+		[
+			SNew(SBorder)
+			.Padding(FMargin(3))
+			.BorderImage(FAppStyle::GetBrush("ToolPanel.GroupBorder"))
+			[
+				SNew(SHorizontalBox)
+				+ SHorizontalBox::Slot()
+				.AutoWidth()
+				[
+					SNew(STextBlock)
+					.Font(AttributeFont)
+					.Text(LOCTEXT("GroomOptionsWindow_HasClumpID", "Has Clump ID: "))
+				]
+				+ SHorizontalBox::Slot()
+				.Padding(5, 0, 0, 0)
+				.AutoWidth()
+				.VAlign(VAlign_Center)
+				[
+					SNew(STextBlock)
+					.Font(AttributeResultFont)
+					.Text(HasClumpIDText)
+					.ColorAndOpacity(AttributeColor)
+				]
+			]
+		]
 
 		// Color attributes
 		+ SVerticalBox::Slot()
@@ -336,7 +369,7 @@ void SGroomImportOptionsWindow::Construct(const FArguments& InArgs)
 				[
 					SNew(STextBlock)
 					.Font(AttributeFont)
-					.Text(LOCTEXT("GroomOptionsWindow_HasColor", "Has Color Attributes: "))
+					.Text(LOCTEXT("GroomOptionsWindow_HasColor", "Has Color: "))
 				]
 				+ SHorizontalBox::Slot()
 				.Padding(5, 0, 0, 0)
@@ -345,7 +378,7 @@ void SGroomImportOptionsWindow::Construct(const FArguments& InArgs)
 				[
 					SNew(STextBlock)
 					.Font(AttributeResultFont)
-					.Text(HasColorAttributeText)
+					.Text(HasColorText)
 					.ColorAndOpacity(AttributeColor)
 				]
 			]
@@ -366,7 +399,37 @@ void SGroomImportOptionsWindow::Construct(const FArguments& InArgs)
 				[
 					SNew(STextBlock)
 					.Font(AttributeFont)
-					.Text(LOCTEXT("GroomOptionsWindow_HasRoughness", "Has Roughness Attributes: "))
+					.Text(LOCTEXT("GroomOptionsWindow_HasRoughness", "Has Roughness: "))
+				]
+				+ SHorizontalBox::Slot()
+				.Padding(5, 0, 0, 0)
+				.AutoWidth()
+				.VAlign(VAlign_Center)
+				[
+					SNew(STextBlock)
+					.Font(AttributeResultFont)
+					.Text(HasRoughnessText)
+					.ColorAndOpacity(AttributeColor)
+				]
+			]
+		]
+
+		// AO attributes
+		+ SVerticalBox::Slot()
+		.AutoHeight()
+		.Padding(2)
+		[
+			SNew(SBorder)
+			.Padding(FMargin(3))
+			.BorderImage(FAppStyle::GetBrush("ToolPanel.GroupBorder"))
+			[
+				SNew(SHorizontalBox)
+				+ SHorizontalBox::Slot()
+				.AutoWidth()
+				[
+					SNew(STextBlock)
+					.Font(AttributeFont)
+					.Text(LOCTEXT("GroomOptionsWindow_HasAO", "Has AO: "))
 				]
 				+ SHorizontalBox::Slot()
 				.Padding(5, 0, 0, 0)
@@ -375,7 +438,7 @@ void SGroomImportOptionsWindow::Construct(const FArguments& InArgs)
 				[
 					SNew(STextBlock)
 					.Font(AttributeResultFont)
-					.Text(HasRoughnessAttributeText)
+					.Text(HasAOText)
 					.ColorAndOpacity(AttributeColor)
 				]
 			]
diff --git a/Engine/Plugins/Runtime/HairStrands/Source/HairStrandsEditor/Private/HairStrandsFactory.cpp b/Engine/Plugins/Runtime/HairStrands/Source/HairStrandsEditor/Private/HairStrandsFactory.cpp
--- a/Engine/Plugins/Runtime/HairStrands/Source/HairStrandsEditor/Private/HairStrandsFactory.cpp
+++ b/Engine/Plugins/Runtime/HairStrands/Source/HairStrandsEditor/Private/HairStrandsFactory.cpp
@@ -112,9 +112,11 @@ UObject* UHairStrandsFactory::FactoryCreateFile(UClass* InClass, UObject* InPare
 		// Load the alembic file upfront to preview & report any potential issue
 		FHairDescriptionGroups OutDescription;
 		bool bHasRootUV = false;
+		bool bHasClumpID = false;
 		bool bHasPrecomputedWeights = false;
-		bool bHasColorAttributes = false;
-		bool bHasRoughnessAttributes = false;
+		bool bHasColor= false;
+		bool bHasRoughness = false;
+		bool bHasAO = false;
 		{
 			FScopedSlowTask Progress((float)1, LOCTEXT("ImportHairAssetForPreview", "Importing hair asset for preview..."), true);
 			Progress.MakeDialog(true);
@@ -127,9 +129,11 @@ UObject* UHairStrandsFactory::FactoryCreateFile(UClass* InClass, UObject* InPare
 
 			FGroomBuilder::BuildHairDescriptionGroups(HairDescription, OutDescription);
 			bHasRootUV = HairDescription.HasRootUV();
+			bHasClumpID = HairDescription.HasClumpID();
 			bHasPrecomputedWeights = HairDescription.HasGuideWeights();
-			bHasColorAttributes = HairDescription.HasColorAttributes();
-			bHasRoughnessAttributes = HairDescription.HasRoughnessAttributes();
+			bHasColor = HairDescription.HasColorAttributes();
+			bHasRoughness = HairDescription.HasRoughnessAttributes();
+			bHasAO = HairDescription.HasAOAttributes();
 
 			// Populate the interpolation settings based on the group count, as this is used later during the ImportHair() to define 
 			// the exact number of group to create
@@ -151,9 +155,11 @@ UObject* UHairStrandsFactory::FactoryCreateFile(UClass* InClass, UObject* InPare
 				OutGroup.CurveCount = Group.Info.NumCurves;
 				OutGroup.GuideCount = Group.Info.NumGuides;
 				OutGroup.bHasRootUV = bHasRootUV;
+				OutGroup.bHasClumpID = bHasClumpID;
 				OutGroup.bHasPrecomputedWeights = bHasPrecomputedWeights;
-				OutGroup.bHasColorAttributes = bHasColorAttributes;
-				OutGroup.bHasRoughnessAttributes = bHasRoughnessAttributes;
+				OutGroup.bHasColor = bHasColor;
+				OutGroup.bHasRoughness = bHasRoughness;
+				OutGroup.bHasAO = bHasAO;
 				bGuidesOnly |= (OutGroup.CurveCount == 0 && OutGroup.GuideCount > 0);
 
 				if (OutGroup.GroupID < OutDescription.HairGroups.Num())
diff --git a/Engine/Shaders/Private/HairStrands/HairCardsAttributeCommon.ush b/Engine/Shaders/Private/HairStrands/HairCardsAttributeCommon.ush
--- a/Engine/Shaders/Private/HairStrands/HairCardsAttributeCommon.ush
+++ b/Engine/Shaders/Private/HairStrands/HairCardsAttributeCommon.ush
@@ -111,6 +111,11 @@ float  GetHairStrandsSeed(float2 HairPrimitiveUV)
 	return GetHairStrandsAttributes(HairPrimitiveUV).HairSeed;
 }
 
+uint GetHairStrandsClumpID(float2 HairPrimitiveUV)
+{
+	return 0;
+}
+
 float3 GetHairStrandsBaseColor(float2 InAtlasUV, float3 InBaseColor)
 {
 	if (HairCardsVF.Flags & EHairCardsFVFlags_TextureBaseColor)
@@ -127,6 +132,11 @@ float GetHairStrandsRoughness(float2 HairPrimitiveUV, float InRoughness)
 	return (HairCardsVF.Flags & EHairCardsFVFlags_TextureRoughness) ? 0 : InRoughness;
 }
 
+float GetHairStrandsAO(float2 HairPrimitiveUV, float InAO)
+{
+	return 1.0f;
+}
+
 float GetHairStrandsGroupIndex(float2 InAtlasUV, float InHairPrimitiveGroupIndex)
 {
 	// * For hair cards,  group index is stored onto vertices
diff --git a/Engine/Shaders/Private/HairStrands/HairStrandsAttributeCommon.ush b/Engine/Shaders/Private/HairStrands/HairStrandsAttributeCommon.ush
--- a/Engine/Shaders/Private/HairStrands/HairStrandsAttributeCommon.ush
+++ b/Engine/Shaders/Private/HairStrands/HairStrandsAttributeCommon.ush
@@ -9,15 +9,19 @@
 #define HAIR_ATTRIBUTE_STRIDE_ROOTUV		4
 #define HAIR_ATTRIBUTE_STRIDE_SEED			1
 #define HAIR_ATTRIBUTE_STRIDE_LENGTH		2
+#define HAIR_ATTRIBUTE_STRIDE_CLUMPID		2
 #define HAIR_ATTRIBUTE_STRIDE_BASECOLOR		4
 #define HAIR_ATTRIBUTE_STRIDE_ROUGHNESS		1
+#define HAIR_ATTRIBUTE_STRIDE_AO			1
 
 // Data offset in bytes. The order needs to match HairStrandsDefinitions.h
 #define HAIR_ATTRIBUTE_OFFSET_ROOTUV(InOffsets)		(InOffsets[0].x)
 #define HAIR_ATTRIBUTE_OFFSET_SEED(InOffsets)		(InOffsets[0].y)
 #define HAIR_ATTRIBUTE_OFFSET_LENGTH(InOffsets)		(InOffsets[0].z)
-#define HAIR_ATTRIBUTE_OFFSET_BASECOLOR(InOffsets)	(InOffsets[0].w)
-#define HAIR_ATTRIBUTE_OFFSET_ROUGHNESS(InOffsets)	(InOffsets[1].x)
+#define HAIR_ATTRIBUTE_OFFSET_CLUMPID(InOffsets)	(InOffsets[0].w)
+#define HAIR_ATTRIBUTE_OFFSET_BASECOLOR(InOffsets)	(InOffsets[1].x)
+#define HAIR_ATTRIBUTE_OFFSET_ROUGHNESS(InOffsets)	(InOffsets[1].y)
+#define HAIR_ATTRIBUTE_OFFSET_AO(InOffsets)			(InOffsets[1].z)
 
 ////////////////////////////////////////////////////////////////////////////////
 // Attribute access functions
@@ -114,6 +118,19 @@ float GetHairStrandsSeed(uint HairPrimitiveId, float2 HairPrimitiveUV)
 	return UnpackHairSeed(InternalGetHairStrandsAttributesPerCurve(HairPrimitiveId, HAIR_ATTRIBUTE_OFFSET_SEED(HairStrandsVF.AttributeOffsets), HAIR_ATTRIBUTE_STRIDE_SEED, HairStrandsVF.AttributeBuffer, HairStrandsVF.VertexToCurveBuffer));
 }
 
+uint GetHairStrandsClumpID(uint HairPrimitiveId, float2 HairPrimitiveUV)
+{
+	const uint ClumpIdOffset = HAIR_ATTRIBUTE_OFFSET_CLUMPID(HairStrandsVF.AttributeOffsets);
+	if (IsAttributeValid(ClumpIdOffset))
+	{
+		return UnpackHairClumpID(InternalGetHairStrandsAttributesPerCurve(HairPrimitiveId, ClumpIdOffset, HAIR_ATTRIBUTE_STRIDE_CLUMPID, HairStrandsVF.AttributeBuffer, HairStrandsVF.VertexToCurveBuffer));
+	}
+	else
+	{
+		return 0;
+	}
+}
+
 float3 GetHairStrandsBaseColor(uint HairPrimitiveId, float2 HairPrimitiveUV)
 {
 	const uint BaseColorOffset = HAIR_ATTRIBUTE_OFFSET_BASECOLOR(HairStrandsVF.AttributeOffsets);
@@ -144,6 +161,21 @@ float GetHairStrandsRoughness(uint HairPrimitiveId, float2 HairPrimitiveUV)
 	}
 }
 
+float GetHairStrandsAO(uint HairPrimitiveId, float2 HairPrimitiveUV)
+{
+	const uint AoOffset = HAIR_ATTRIBUTE_OFFSET_AO(HairStrandsVF.AttributeOffsets);
+	if (IsAttributeValid(AoOffset))
+	{
+		const float AO0 = UnpackHairRoughness(InternalGetHairStrandsAttributesPerVertex(HairPrimitiveId, AoOffset, HAIR_ATTRIBUTE_STRIDE_AO, HairStrandsVF.AttributeBuffer));
+		const float AO1 = UnpackHairRoughness(InternalGetHairStrandsAttributesPerVertex(HairPrimitiveId+1, AoOffset, HAIR_ATTRIBUTE_STRIDE_AO, HairStrandsVF.AttributeBuffer));
+		return lerp(AO0, AO1, HairPrimitiveUV.x);
+	}
+	else
+	{
+		return 0;
+	}
+}
+
 float GetHairStrandsDepth(uint HairPrimitiveId, float2 InUV, float InDeviceZ)
 {
 	return 0.f;
diff --git a/Engine/Shaders/Private/HairStrands/HairStrandsVertexFactoryCommon.ush b/Engine/Shaders/Private/HairStrands/HairStrandsVertexFactoryCommon.ush
--- a/Engine/Shaders/Private/HairStrands/HairStrandsVertexFactoryCommon.ush
+++ b/Engine/Shaders/Private/HairStrands/HairStrandsVertexFactoryCommon.ush
@@ -100,6 +100,11 @@ float UnpackHairSeed(uint In)
 	return (In & 0xFF) / 255.f;
 }
 
+uint UnpackHairClumpID(uint In)
+{
+	return (In & 0xFFFF);
+}
+
 float3 UnpackHairBaseColor(uint In)
 {
 	float3 Out;
diff --git a/Engine/Shaders/Private/MaterialTemplate.ush b/Engine/Shaders/Private/MaterialTemplate.ush
--- a/Engine/Shaders/Private/MaterialTemplate.ush
+++ b/Engine/Shaders/Private/MaterialTemplate.ush
@@ -409,9 +409,9 @@ struct FMaterialPixelParameters
 #if HAIR_CARD_MESH_FACTORY
 	float2	HairPrimitiveUV;		// AtlasUV
 	float2  HairPrimitiveRootUV;	// RootUV
-	float4	HairPrimitiveMaterial;	// Card material
-	float   HairPrimitiveLength;	// Card length
-	float	HairPrimitiveGroupIndex;// Card group index
+	half4	HairPrimitiveMaterial;	// Card material
+	half    HairPrimitiveLength;	// Card length
+	half	HairPrimitiveGroupIndex;// Card group index
 #endif
 
 #if HAS_INSTANCE_LOCAL_TO_WORLD_PS
@@ -1492,6 +1492,17 @@ float MaterialExpressionGetHairSeed(FMaterialPixelParameters Parameters)
 #endif
 }
 
+uint MaterialExpressionGetHairClumpId(FMaterialPixelParameters Parameters)
+{
+#if HAIR_STRAND_MESH_FACTORY
+	return GetHairStrandsClumpID(Parameters.HairPrimitiveId, Parameters.HairPrimitiveUV);
+#elif HAIR_CARD_MESH_FACTORY
+	return GetHairStrandsClumpID(Parameters.HairPrimitiveUV);
+#else
+	return 0;
+#endif
+}
+
 float3 MaterialExpressionGetHairBaseColor(FMaterialPixelParameters Parameters)
 {
 #if HAIR_STRAND_MESH_FACTORY
@@ -1514,6 +1525,17 @@ float MaterialExpressionGetHairRoughness(FMaterialPixelParameters Parameters)
 #endif
 }
 
+float MaterialExpressionGetHairAO(FMaterialPixelParameters Parameters)
+{
+#if HAIR_STRAND_MESH_FACTORY
+	return GetHairStrandsAO(Parameters.HairPrimitiveId, Parameters.HairPrimitiveUV);
+#elif HAIR_CARD_MESH_FACTORY
+	return GetHairStrandsAO(Parameters.HairPrimitiveUV, 1.0f);
+#else
+	return 0;
+#endif
+}
+
 float MaterialExpressionGetHairDepth(FMaterialVertexParameters Parameters)
 {
 	return 0;
diff --git a/Engine/Shaders/Shared/HairStrandsDefinitions.h b/Engine/Shaders/Shared/HairStrandsDefinitions.h
--- a/Engine/Shaders/Shared/HairStrandsDefinitions.h
+++ b/Engine/Shaders/Shared/HairStrandsDefinitions.h
@@ -14,9 +14,11 @@
 #define HAIR_ATTRIBUTE_ROOTUV 0
 #define HAIR_ATTRIBUTE_SEED 1
 #define HAIR_ATTRIBUTE_LENGTH 2
-#define HAIR_ATTRIBUTE_BASECOLOR 3
-#define HAIR_ATTRIBUTE_ROUGHNESS 4
-#define HAIR_ATTRIBUTE_COUNT 5
+#define HAIR_ATTRIBUTE_CLUMPID 3
+#define HAIR_ATTRIBUTE_BASECOLOR 4
+#define HAIR_ATTRIBUTE_ROUGHNESS 5
+#define HAIR_ATTRIBUTE_AO 6
+#define HAIR_ATTRIBUTE_COUNT 7
 
 // HAIR_ATTRIBUTE_MAX rounded to 4
 #define HAIR_ATTRIBUTE_OFFSET_COUNT ((HAIR_ATTRIBUTE_COUNT + 3) / 4)
diff --git a/Engine/Source/Runtime/Engine/Private/Materials/HLSLMaterialTranslator.cpp b/Engine/Source/Runtime/Engine/Private/Materials/HLSLMaterialTranslator.cpp
--- a/Engine/Source/Runtime/Engine/Private/Materials/HLSLMaterialTranslator.cpp
+++ b/Engine/Source/Runtime/Engine/Private/Materials/HLSLMaterialTranslator.cpp
@@ -9385,6 +9385,11 @@ int32 FHLSLMaterialTranslator::GetHairSeed()
 	return AddCodeChunk(MCT_Float1, TEXT("MaterialExpressionGetHairSeed(Parameters)"));
 }
 
+int32 FHLSLMaterialTranslator::GetHairClumpID()
+{
+	return AddCodeChunk(MCT_UInt1, TEXT("MaterialExpressionGetHairClumpID(Parameters)"));
+}
+
 int32 FHLSLMaterialTranslator::GetHairTangent(bool bUseTangentSpace)
 {
 	return AddCodeChunk(MCT_Float3, TEXT("MaterialExpressionGetHairTangent(Parameters, %s)"), bUseTangentSpace ? TEXT("true") : TEXT("false"));
@@ -9405,6 +9410,11 @@ int32 FHLSLMaterialTranslator::GetHairRoughness()
 	return AddCodeChunk(MCT_Float1, TEXT("MaterialExpressionGetHairRoughness(Parameters)"));
 }
 
+int32 FHLSLMaterialTranslator::GetHairAO()
+{
+	return AddCodeChunk(MCT_Float1, TEXT("MaterialExpressionGetHairAO(Parameters)"));
+}
+
 int32 FHLSLMaterialTranslator::GetHairDepth()
 {
 	return AddCodeChunk(MCT_Float1, TEXT("MaterialExpressionGetHairDepth(Parameters)"));
diff --git a/Engine/Source/Runtime/Engine/Private/Materials/HLSLMaterialTranslator.h b/Engine/Source/Runtime/Engine/Private/Materials/HLSLMaterialTranslator.h
--- a/Engine/Source/Runtime/Engine/Private/Materials/HLSLMaterialTranslator.h
+++ b/Engine/Source/Runtime/Engine/Private/Materials/HLSLMaterialTranslator.h
@@ -1007,9 +1007,11 @@ protected:
 	virtual int32 GetHairUV() override;
 	virtual int32 GetHairDimensions() override;
 	virtual int32 GetHairSeed() override;
+	virtual int32 GetHairClumpID() override;
 	virtual int32 GetHairTangent(bool bUseTangentSpace) override;
 	virtual int32 GetHairRootUV() override;
 	virtual int32 GetHairBaseColor() override;
+	virtual int32 GetHairAO() override;
 	virtual int32 GetHairRoughness() override;
 	virtual int32 GetHairDepth() override;
 	virtual int32 GetHairCoverage() override;
diff --git a/Engine/Source/Runtime/Engine/Private/Materials/MaterialExpressionHLSL.cpp b/Engine/Source/Runtime/Engine/Private/Materials/MaterialExpressionHLSL.cpp
--- a/Engine/Source/Runtime/Engine/Private/Materials/MaterialExpressionHLSL.cpp
+++ b/Engine/Source/Runtime/Engine/Private/Materials/MaterialExpressionHLSL.cpp
@@ -2318,6 +2318,8 @@ bool UMaterialExpressionHairAttributes::GenerateHLSLExpression(FMaterialHLSLGene
 	case 11: OutExpression = Generator.GetTree().NewExpression<FExpressionInlineCustomHLSL>(EValueType::Float4, TEXT("MaterialExpressionGetHairAuxilaryData(Parameters)")); break;
 	case 12: OutExpression = Generator.GetTree().NewExpression<FExpressionInlineCustomHLSL>(EValueType::Float2, TEXT("MaterialExpressionGetAtlasUVs(Parameters)")); break;
 	case 13: OutExpression = Generator.GetTree().NewExpression<FExpressionInlineCustomHLSL>(EValueType::Float1, TEXT("MaterialExpressionGetHairGroupIndex(Parameters)")); break;
+	case 14: OutExpression = Generator.GetTree().NewExpression<FExpressionInlineCustomHLSL>(EValueType::Float1, TEXT("MaterialExpressionGetHairAO(Parameters)")); break;
+	case 15: OutExpression = Generator.GetTree().NewExpression<FExpressionInlineCustomHLSL>(EValueType::Int1, TEXT("MaterialExpressionGetHairClumpID(Parameters)")); break;
 	default: return Generator.Error(TEXT("Invalid output"));
 	}
 	return OutExpression != nullptr;
diff --git a/Engine/Source/Runtime/Engine/Private/Materials/MaterialExpressions.cpp b/Engine/Source/Runtime/Engine/Private/Materials/MaterialExpressions.cpp
--- a/Engine/Source/Runtime/Engine/Private/Materials/MaterialExpressions.cpp
+++ b/Engine/Source/Runtime/Engine/Private/Materials/MaterialExpressions.cpp
@@ -21065,6 +21065,8 @@ UMaterialExpressionHairAttributes::UMaterialExpressionHairAttributes(const FObje
 	Outputs.Add(FExpressionOutput(TEXT("AuxilaryData"), 1, 1, 1, 1, 1));
 	Outputs.Add(FExpressionOutput(TEXT("AtlasUVs"), 1, 1, 1, 0, 0));
 	Outputs.Add(FExpressionOutput(TEXT("Group Index"), 1, 1, 0, 0, 0));
+	Outputs.Add(FExpressionOutput(TEXT("AO"), 1, 1, 0, 0, 0));
+	Outputs.Add(FExpressionOutput(TEXT("Clump ID"), 1, 1, 0, 0, 0));
 #endif
 }
 
@@ -21119,6 +21121,14 @@ int32 UMaterialExpressionHairAttributes::Compile(class FMaterialCompiler* Compil
 	{
 		return Compiler->GetHairGroupIndex();
 	}
+	else if (OutputIndex == 14)
+	{
+		return Compiler->GetHairAO();
+	}
+	else if (OutputIndex == 15)
+	{
+		return Compiler->GetHairClumpID();
+	}
 
 	return Compiler->Errorf(TEXT("Invalid input parameter"));
 }
diff --git a/Engine/Source/Runtime/Engine/Public/MaterialCompiler.h b/Engine/Source/Runtime/Engine/Public/MaterialCompiler.h
--- a/Engine/Source/Runtime/Engine/Public/MaterialCompiler.h
+++ b/Engine/Source/Runtime/Engine/Public/MaterialCompiler.h
@@ -605,10 +605,12 @@ public:
 	virtual int32 GetHairUV() = 0;
 	virtual int32 GetHairDimensions() = 0;
 	virtual int32 GetHairSeed() = 0;
+	virtual int32 GetHairClumpID() = 0;
 	virtual int32 GetHairTangent(bool bUseTangentSpace) = 0;
 	virtual int32 GetHairRootUV() = 0;
 	virtual int32 GetHairBaseColor() = 0;
 	virtual int32 GetHairRoughness() = 0;
+	virtual int32 GetHairAO() = 0;
 	virtual int32 GetHairDepth() = 0;
 	virtual int32 GetHairCoverage() = 0;
 	virtual int32 GetHairAuxilaryData() = 0;
@@ -1031,10 +1033,12 @@ public:
 	virtual int32 GetHairUV() override { return Compiler->GetHairUV(); }
 	virtual int32 GetHairDimensions() override { return Compiler->GetHairDimensions(); }
 	virtual int32 GetHairSeed() override { return Compiler->GetHairSeed(); }
+	virtual int32 GetHairClumpID() override { return Compiler->GetHairClumpID(); }
 	virtual int32 GetHairTangent(bool bUseTangentSpace) override { return Compiler->GetHairTangent(bUseTangentSpace); }
 	virtual int32 GetHairRootUV() override { return Compiler->GetHairRootUV(); }
 	virtual int32 GetHairBaseColor() override { return Compiler->GetHairBaseColor(); }
 	virtual int32 GetHairRoughness() override { return Compiler->GetHairRoughness(); }
+	virtual int32 GetHairAO() override { return Compiler->GetHairAO(); }
 	virtual int32 GetHairDepth() override { return Compiler->GetHairDepth(); }
 	virtual int32 GetHairCoverage() override { return Compiler->GetHairCoverage(); }
 	virtual int32 GetHairAuxilaryData() override { return Compiler->GetHairAuxilaryData(); }
