* Fix invalid strands radius computation during groom building when curve reordering or curve decimation is enabled.
* Enabled curves reordering by default.

diff --git a/Engine/Plugins/Runtime/HairStrands/Source/HairStrandsCore/Private/GroomBuilder.cpp b/Engine/Plugins/Runtime/HairStrands/Source/HairStrandsCore/Private/GroomBuilder.cpp
--- a/Engine/Plugins/Runtime/HairStrands/Source/HairStrandsCore/Private/GroomBuilder.cpp
+++ b/Engine/Plugins/Runtime/HairStrands/Source/HairStrandsCore/Private/GroomBuilder.cpp
@@ -40,11 +40,7 @@ static FAutoConsoleVariableRef CVarHairGroupIndexBuilder_MaxVoxelResolution(TEXT
 
 FString FGroomBuilder::GetVersion()
 {
-	// Important to update the version when groom building changes
-	if (IsHairStrandContinuousDecimationReorderingEnabled())
-		return TEXT("v4i");
-	else
-		return TEXT("v3i");
+	return TEXT("v5");
 }
 
 namespace FHairStrandsDecimation
@@ -1921,7 +1917,7 @@ void FGroomBuilder::BuildData(
 	const float CurveDecimation		= FMath::Clamp(InSettings.DecimationSettings.CurveDecimation, 0.f, 1.f);
 	const float VertexDecimation	= FMath::Clamp(InSettings.DecimationSettings.VertexDecimation, 0.f, 1.f);
 	const float HairToGuideDensity	= FMath::Clamp(InSettings.InterpolationSettings.HairToGuideDensity, 0.f, 1.f);	
-	const bool bContinuousDecimationReordering = IsHairStrandContinuousDecimationReorderingEnabled(); //this is a read only project setting which will modify the platform data that is stored in the DDC
+	const bool bCurveReordering 	= true;
 
 	{
 		OutGroupInfo.GroupID = InHairDescriptionGroup.Info.GroupID;
@@ -1935,11 +1931,11 @@ void FGroomBuilder::BuildData(
 			HairStrandsBuilder::BuildInternalData(OutRen, !InHairDescriptionGroup.bHasUVData);
 
 			// Decimate
-			if (CurveDecimation < 1 || VertexDecimation < 1 || bContinuousDecimationReordering)
+			if (CurveDecimation < 1 || VertexDecimation < 1 || bCurveReordering)
 			{
 				FHairStrandsDatas FullData = OutRen;
 				OutRen.Reset();
-				FHairStrandsDecimation::Decimate(FullData, CurveDecimation, VertexDecimation, bContinuousDecimationReordering, OutRen);
+				FHairStrandsDecimation::Decimate(FullData, CurveDecimation, VertexDecimation, bCurveReordering, OutRen);
 			}
 			OutGroupInfo.NumCurves			= OutRen.GetNumCurves();
 			OutGroupInfo.NumCurveVertices	= OutRen.GetNumPoints();
@@ -2147,6 +2143,7 @@ void Decimate(
 	const bool bHasAO = InData.StrandsPoints.PointsAO.Num() > 0;
 	const bool bHasClumpIDs = InData.StrandsCurves.ClumpIDs.Num() > 0;
 	const bool bHasStrandIDs = InData.StrandsCurves.StrandIDs.Num() > 0;
+	check(InData.StrandsCurves.MaxRadius > 0);
 
 	PointCount = 0;
 	for(uint32 CurveIndex = 0; CurveIndex < CurveCount; ++CurveIndex)
@@ -2188,7 +2185,8 @@ void Decimate(
 			
 			OutData.StrandsPoints.PointsPosition[PointCount]	= FMath::Lerp(InData.StrandsPoints.PointsPosition[PointBegin],	InData.StrandsPoints.PointsPosition[PointEnd],	PointAlpha);
 			OutData.StrandsPoints.PointsCoordU[PointCount]		= FMath::Lerp(InData.StrandsPoints.PointsCoordU[PointBegin],	InData.StrandsPoints.PointsCoordU[PointEnd],	PointAlpha);
-			OutData.StrandsPoints.PointsRadius[PointCount]		= FMath::Lerp(InData.StrandsPoints.PointsRadius[PointBegin],	InData.StrandsPoints.PointsRadius[PointEnd],	PointAlpha);
+			// Multiply by MaxRadius as HairStrandsBuilder::BuildInternalData will recompute MaxRadius based on the actual Radius not the not the normalized radius.
+			OutData.StrandsPoints.PointsRadius[PointCount]		= FMath::Lerp(InData.StrandsPoints.PointsRadius[PointBegin],	InData.StrandsPoints.PointsRadius[PointEnd],	PointAlpha) * InData.StrandsCurves.MaxRadius; 
 			if (bHasColor)		{ OutData.StrandsPoints.PointsBaseColor[PointCount]	= FMath::Lerp(InData.StrandsPoints.PointsBaseColor[PointBegin], InData.StrandsPoints.PointsBaseColor[PointEnd], PointAlpha); }
 			if (bHasRoughness)	{ OutData.StrandsPoints.PointsRoughness[PointCount]	= FMath::Lerp(InData.StrandsPoints.PointsRoughness[PointBegin], InData.StrandsPoints.PointsRoughness[PointEnd], PointAlpha); }
 			if (bHasAO)			{ OutData.StrandsPoints.PointsAO[PointCount]		= FMath::Lerp(InData.StrandsPoints.PointsAO[PointBegin],		InData.StrandsPoints.PointsAO[PointEnd],		PointAlpha); }
@@ -2217,11 +2215,8 @@ void Decimate(
 	const bool bHasAO = InData.StrandsPoints.PointsAO.Num() > 0;
 	const bool bHasClumpIDs = InData.StrandsCurves.ClumpIDs.Num() > 0;
 	const bool bHasStrandIDs = InData.StrandsCurves.StrandIDs.Num() > 0;
+	check(InData.StrandsCurves.MaxRadius > 0);
 
-	TArray<uint32> CurveIndices;
-	CurveIndices.SetNum(OutCurveCount);
-
-	uint32 OutTotalPointCount = 0;
 	FRandomStream Random;
 	Random.Initialize(0xdeedbeed);
 
@@ -2245,24 +2240,32 @@ void Decimate(
 		}
 	}
 
-	const float CurveBucketSize = float(InCurveCount) / float(OutCurveCount);
-	int32 LastCurveIndex = -1;
-	for (uint32 BucketIndex = 0; BucketIndex < OutCurveCount; BucketIndex++)
+	uint32 OutTotalPointCount = InData.GetNumPoints();
+	TArray<uint32> CurveIndices;
+	const bool bNeedDecimation = CurveDecimationPercentage < 1 || VertexDecimationPercentage < 1;
+	if (bNeedDecimation)
 	{
-		const float MinBucket = FMath::Max(BucketIndex  * CurveBucketSize, float(LastCurveIndex+1));
-		const float MaxBucket = (BucketIndex+1) * CurveBucketSize;
-		const float AdjustedBucketSize = MaxBucket - MinBucket;
-		if (AdjustedBucketSize > 0)
+		CurveIndices.SetNum(OutCurveCount);
+		OutTotalPointCount = 0;
+		const float CurveBucketSize = float(InCurveCount) / float(OutCurveCount);
+		int32 LastCurveIndex = -1;
+		for (uint32 BucketIndex = 0; BucketIndex < OutCurveCount; BucketIndex++)
 		{
-			const uint32 CurveIndex = Random.RandRange(MinBucket, FMath::FloorToInt(MinBucket + AdjustedBucketSize)-1);
-			CurveIndices[BucketIndex] = CurveIndex;
-			LastCurveIndex = CurveIndex;
-			
-			const uint32 EffectiveCurveIndex = bContinuousDecimationReordering ? CurveRandomizedIndex[CurveIndex] : CurveIndex;
+			const float MinBucket = FMath::Max(BucketIndex * CurveBucketSize, float(LastCurveIndex + 1));
+			const float MaxBucket = (BucketIndex + 1) * CurveBucketSize;
+			const float AdjustedBucketSize = MaxBucket - MinBucket;
+			if (AdjustedBucketSize > 0)
+			{
+				const uint32 CurveIndex = Random.RandRange(MinBucket, FMath::FloorToInt(MinBucket + AdjustedBucketSize) - 1);
+				CurveIndices[BucketIndex] = CurveIndex;
+				LastCurveIndex = CurveIndex;
 
-			const uint32 InPointCount = InData.StrandsCurves.CurvesCount[EffectiveCurveIndex];
-			const uint32 OutPointCount = DecimatePointCount(InPointCount, VertexDecimationPercentage);
-			OutTotalPointCount += OutPointCount;
+				const uint32 EffectiveCurveIndex = bContinuousDecimationReordering ? CurveRandomizedIndex[CurveIndex] : CurveIndex;
+
+				const uint32 InPointCount = InData.StrandsCurves.CurvesCount[EffectiveCurveIndex];
+				const uint32 OutPointCount = DecimatePointCount(InPointCount, VertexDecimationPercentage);
+				OutTotalPointCount += OutPointCount;
+			}
 		}
 	}
 
@@ -2284,7 +2287,7 @@ void Decimate(
 	uint32 OutPointOffset = 0; 
 	for (uint32 OutCurveIndex = 0; OutCurveIndex < OutCurveCount; ++OutCurveIndex)
 	{
-		const uint32 InCurveIndex	= CurveIndices[OutCurveIndex];
+		const uint32 InCurveIndex		 = bNeedDecimation ? CurveIndices[OutCurveIndex] : OutCurveIndex;
 		const uint32 EffectiveCurveIndex = bContinuousDecimationReordering ? CurveRandomizedIndex[InCurveIndex] : InCurveIndex;
 
 		const uint32 InPointOffset	= InData.StrandsCurves.CurvesOffset[EffectiveCurveIndex];
@@ -2293,31 +2296,35 @@ void Decimate(
 		// Decimation using area metric
 	#if 1
 		TArray<uint32> OutPointIndices;
-		// Don't need to DecimateCurve if there are only 2 control vertices
-		if (InPointCount > 2)
-		{
-			DecimateCurve(
-				InData.StrandsPoints.PointsPosition,
-				InPointOffset,
-				InPointCount,
-				VertexDecimationPercentage,
-				OutPointIndices);
-		}
-		else
+		if (bNeedDecimation)
 		{
-			// Just pass the start and end point of the curve, no decimation needed
-			OutPointIndices.Add(0);
-			OutPointIndices.Add(1);
+			// Don't need to DecimateCurve if there are only 2 control vertices
+			if (InPointCount > 2)
+			{
+				DecimateCurve(
+					InData.StrandsPoints.PointsPosition,
+					InPointOffset,
+					InPointCount,
+					VertexDecimationPercentage,
+					OutPointIndices);
+			}
+			else
+			{
+				// Just pass the start and end point of the curve, no decimation needed
+				OutPointIndices.Add(0);
+				OutPointIndices.Add(1);
+			}
 		}
 
-		const uint32 OutPointCount = OutPointIndices.Num();
+		const uint32 OutPointCount = bNeedDecimation ? OutPointIndices.Num() : InPointCount;
 		for (uint32 OutPointIndex = 0; OutPointIndex < OutPointCount; ++OutPointIndex)
 		{
-			const uint32 InPointIndex = OutPointIndices[OutPointIndex];
+			const uint32 InPointIndex = bNeedDecimation ? OutPointIndices[OutPointIndex] : OutPointIndex;
 
 			OutData.StrandsPoints.PointsPosition [OutPointIndex + OutPointOffset] = InData.StrandsPoints.PointsPosition	[InPointIndex + InPointOffset];
 			OutData.StrandsPoints.PointsCoordU	 [OutPointIndex + OutPointOffset] = InData.StrandsPoints.PointsCoordU	[InPointIndex + InPointOffset];
-			OutData.StrandsPoints.PointsRadius	 [OutPointIndex + OutPointOffset] = InData.StrandsPoints.PointsRadius	[InPointIndex + InPointOffset];
+			// Multiply by MaxRadius as HairStrandsBuilder::BuildInternalData will recompute MaxRadius based on the actual Radius not the not the normalized radius.
+			OutData.StrandsPoints.PointsRadius	 [OutPointIndex + OutPointOffset] = InData.StrandsPoints.PointsRadius	[InPointIndex + InPointOffset] * InData.StrandsCurves.MaxRadius;
 			if (bHasColor)		{ OutData.StrandsPoints.PointsBaseColor[OutPointIndex + OutPointOffset] = InData.StrandsPoints.PointsBaseColor[InPointIndex + InPointOffset]; }
 			if (bHasRoughness)	{ OutData.StrandsPoints.PointsRoughness[OutPointIndex + OutPointOffset] = InData.StrandsPoints.PointsRoughness[InPointIndex + InPointOffset]; }
 			if (bHasAO)			{ OutData.StrandsPoints.PointsAO[OutPointIndex + OutPointOffset]		= InData.StrandsPoints.PointsAO[InPointIndex + InPointOffset]; }
@@ -2349,8 +2356,8 @@ void Decimate(
 	#endif
 
 		OutData.StrandsCurves.CurvesCount[OutCurveIndex]  = OutPointCount;
-		OutData.StrandsCurves.CurvesRootUV[OutCurveIndex] = InData.StrandsCurves.CurvesRootUV[EffectiveCurveIndex];
 		OutData.StrandsCurves.CurvesOffset[OutCurveIndex] = OutPointOffset;
+		OutData.StrandsCurves.CurvesRootUV[OutCurveIndex] = InData.StrandsCurves.CurvesRootUV[EffectiveCurveIndex];
 		OutData.StrandsCurves.CurvesLength[OutCurveIndex] = InData.StrandsCurves.CurvesLength[EffectiveCurveIndex];
 		if (bHasPrecomputedWeights)
 		{
diff --git a/Engine/Plugins/Runtime/HairStrands/Source/HairStrandsCore/Private/HairStrandsDatas.cpp b/Engine/Plugins/Runtime/HairStrands/Source/HairStrandsCore/Private/HairStrandsDatas.cpp
--- a/Engine/Plugins/Runtime/HairStrands/Source/HairStrandsCore/Private/HairStrandsDatas.cpp
+++ b/Engine/Plugins/Runtime/HairStrands/Source/HairStrandsCore/Private/HairStrandsDatas.cpp
@@ -54,6 +54,8 @@ void FHairStrandsCurves::Reset()
 	ClumpIDs.Reset();
 	CurvesClosestGuideIDs.Reset();
 	CurvesClosestGuideWeights.Reset();
+	MaxLength = 0;
+	MaxRadius = 0;
 }
 
 void FHairStrandsPoints::Reset()
@@ -274,4 +276,6 @@ void FHairStrandsDatas::Reset()
 {
 	StrandsCurves.Reset();
 	StrandsPoints.Reset();
+	HairDensity = 1;
+	BoundingBox = FBox(EForceInit::ForceInit);
 }
