- Fix missing implementation for refactor of GroomComponent PSO precache collection
- Fix missing include

#rnx

diff --git a/Engine/Plugins/Runtime/HairStrands/Source/HairStrandsCore/Private/GroomAsset.cpp b/Engine/Plugins/Runtime/HairStrands/Source/HairStrandsCore/Private/GroomAsset.cpp
--- a/Engine/Plugins/Runtime/HairStrands/Source/HairStrandsCore/Private/GroomAsset.cpp
+++ b/Engine/Plugins/Runtime/HairStrands/Source/HairStrandsCore/Private/GroomAsset.cpp
@@ -814,7 +814,7 @@ TArray<FHairVertexFactoryTypesPerMaterialData> UGroomAsset::CollectVertexFactory
 			VFsPerMaterial->MaterialIndex = InMaterialIndex;
 			VFsPerMaterial->HairGeometryType = InHairGeometryType;
 		}
-		VFsPerMaterial->VertexFactoryTypes.AddUnique(InVFType);
+		VFsPerMaterial->VertexFactoryDataList.AddUnique(FPSOPrecacheVertexFactoryData(InVFType));
 	};
 
 	const int32 GroupCount = GetNumHairGroups();
@@ -1292,13 +1292,14 @@ void UGroomAsset::PostLoad()
 
 		TArray<FHairVertexFactoryTypesPerMaterialData> VFsPerMaterials = CollectVertexFactoryTypesPerMaterialData(ShaderPlatform);
 
+		TArray<FMaterialPSOPrecacheRequestID> RequestIDs;
 		FPSOPrecacheParams PrecachePSOParams;
 		for (FHairVertexFactoryTypesPerMaterialData& VFsPerMaterial : VFsPerMaterials)
 		{
 			UMaterialInterface* MaterialInterface = HairGroupsMaterials[VFsPerMaterial.MaterialIndex].Material;
 			if (MaterialInterface)
 			{
-				MaterialInterface->PrecachePSOs(VFsPerMaterial.VertexFactoryTypes, PrecachePSOParams);
+				MaterialInterface->PrecachePSOs(VFsPerMaterial.VertexFactoryDataList, PrecachePSOParams, EPSOPrecachePriority::Medium, RequestIDs);
 			}
 		}
 	}
diff --git a/Engine/Plugins/Runtime/HairStrands/Source/HairStrandsCore/Private/GroomComponent.cpp b/Engine/Plugins/Runtime/HairStrands/Source/HairStrandsCore/Private/GroomComponent.cpp
--- a/Engine/Plugins/Runtime/HairStrands/Source/HairStrandsCore/Private/GroomComponent.cpp
+++ b/Engine/Plugins/Runtime/HairStrands/Source/HairStrandsCore/Private/GroomComponent.cpp
@@ -3106,24 +3106,26 @@ void UGroomComponent::PostLoad()
 #endif
 }
 
-void UGroomComponent::PrecachePSOs()
+void UGroomComponent::CollectPSOPrecacheData(const FPSOPrecacheParams& BasePrecachePSOParams, FComponentPSOPrecacheParamsList& OutParams) 
 {
-	if (!IsComponentPSOPrecachingEnabled() || GroomAsset == nullptr)
+	if (GroomAsset == nullptr)
 	{
 		return;
 	}
 
 	TArray<FHairVertexFactoryTypesPerMaterialData> VFsPerMaterials = GroomAsset->CollectVertexFactoryTypesPerMaterialData(GMaxRHIShaderPlatform);
 
-	FPSOPrecacheParams PrecachePSOParams;
-	SetupPrecachePSOParams(PrecachePSOParams);
+	FPSOPrecacheParams PrecachePSOParams = BasePrecachePSOParams;
 
 	for (FHairVertexFactoryTypesPerMaterialData& VFsPerMaterial : VFsPerMaterials)
 	{
 		UMaterialInterface* MaterialInterface = GetMaterial(GetMaterialIndexWithFallback(VFsPerMaterial.MaterialIndex), VFsPerMaterial.HairGeometryType, true);
 		if (MaterialInterface)
 		{
-			MaterialInterface->PrecachePSOs(VFsPerMaterial.VertexFactoryTypes, PrecachePSOParams);
+			FComponentPSOPrecacheParams& ComponentParams = OutParams[OutParams.AddDefaulted()];
+			ComponentParams.MaterialInterface = MaterialInterface;
+			ComponentParams.VertexFactoryDataList = VFsPerMaterial.VertexFactoryDataList;
+			ComponentParams.PSOPrecacheParams = PrecachePSOParams;
 		}
 	}
 
diff --git a/Engine/Plugins/Runtime/HairStrands/Source/HairStrandsCore/Public/GroomAsset.h b/Engine/Plugins/Runtime/HairStrands/Source/HairStrandsCore/Public/GroomAsset.h
--- a/Engine/Plugins/Runtime/HairStrands/Source/HairStrandsCore/Public/GroomAsset.h
+++ b/Engine/Plugins/Runtime/HairStrands/Source/HairStrandsCore/Public/GroomAsset.h
@@ -406,7 +406,7 @@ struct FHairVertexFactoryTypesPerMaterialData
 {
 	int16 MaterialIndex;
 	EHairGeometryType HairGeometryType;
-	TArray<const FVertexFactoryType*, TInlineAllocator<2>> VertexFactoryTypes;
+	FPSOPrecacheVertexFactoryDataList VertexFactoryDataList;
 };
 
 /**
diff --git a/Engine/Source/Runtime/Engine/Classes/Engine/SkinnedAsset.h b/Engine/Source/Runtime/Engine/Classes/Engine/SkinnedAsset.h
--- a/Engine/Source/Runtime/Engine/Classes/Engine/SkinnedAsset.h
+++ b/Engine/Source/Runtime/Engine/Classes/Engine/SkinnedAsset.h
@@ -12,6 +12,7 @@
 #include "Interfaces/Interface_AsyncCompilation.h"
 #include "ReferenceSkeleton.h"
 #include "PerPlatformProperties.h"
+#include "PSOPrecache.h"
 #if UE_ENABLE_INCLUDE_ORDER_DEPRECATED_IN_5_2
 #include "SkeletalMeshTypes.h"
 #include "SkinnedAssetAsyncCompileUtils.h"
