Add RHI resource owner name tracking to RDG buffers used by hair plugin.

diff --git a/Engine/Plugins/Runtime/HairStrands/Source/HairStrandsCore/Private/GroomAsset.cpp b/Engine/Plugins/Runtime/HairStrands/Source/HairStrandsCore/Private/GroomAsset.cpp
--- a/Engine/Plugins/Runtime/HairStrands/Source/HairStrandsCore/Private/GroomAsset.cpp
+++ b/Engine/Plugins/Runtime/HairStrands/Source/HairStrandsCore/Private/GroomAsset.cpp
@@ -950,7 +950,7 @@ void UGroomAsset::UpdateResource()
 					HairGroupsInfo,
 					HairGroupsData);
 
-				GroupData.Strands.ClusterCullingResource = new FHairStrandsClusterCullingResource(GroupData.Strands.ClusterCullingBulkData, FHairResourceName(GetFName(), GroupIndex));
+				GroupData.Strands.ClusterCullingResource = new FHairStrandsClusterCullingResource(GroupData.Strands.ClusterCullingBulkData, FHairResourceName(GetFName(), GroupIndex), GetAssetPathName());
 			}
 			else
 			{
@@ -2757,8 +2757,9 @@ bool UGroomAsset::BuildCardsGeometry(uint32 GroupIndex)
 			if (bInitResources)
 			{
 				FHairResourceName ResourceName(GetFName(), GroupIndex, LODIt);
+				const FName OwnerName = GetAssetPathName(LODIt);
 
-				LOD.RestResource = new FHairCardsRestResource(LOD.BulkData, ResourceName);
+				LOD.RestResource = new FHairCardsRestResource(LOD.BulkData, ResourceName, OwnerName);
 				BeginInitResource(LOD.RestResource); // Immediate allocation, as needed for the vertex factory, input stream building
 
 				// 2.1 Load atlas textures
@@ -2771,7 +2772,7 @@ bool UGroomAsset::BuildCardsGeometry(uint32 GroupIndex)
 				LOD.RestResource->bInvertUV = Desc->SourceType == EHairCardsSourceType::Procedural;
 				
 				// 2.2 Load interoplatino resources
-				LOD.InterpolationResource = new FHairCardsInterpolationResource(LOD.InterpolationBulkData, ResourceName);
+				LOD.InterpolationResource = new FHairCardsInterpolationResource(LOD.InterpolationBulkData, ResourceName, OwnerName);
 
 				// Create own interpolation settings for cards.
 				// Force closest guides as this is the most relevant matching metric for cards, due to their coarse geometry
@@ -2815,9 +2816,9 @@ bool UGroomAsset::BuildCardsGeometry(uint32 GroupIndex)
 					GroupData.Guides.BulkData  = LOD.Guides.BulkData;
 				}
 
-				LOD.Guides.RestResource = new FHairStrandsRestResource(LOD.Guides.BulkData, EHairStrandsResourcesType::Cards, ResourceName);
+				LOD.Guides.RestResource = new FHairStrandsRestResource(LOD.Guides.BulkData, EHairStrandsResourcesType::Cards, ResourceName, OwnerName);
 
-				LOD.Guides.InterpolationResource = new FHairStrandsInterpolationResource(LOD.Guides.InterpolationBulkData, ResourceName);
+				LOD.Guides.InterpolationResource = new FHairStrandsInterpolationResource(LOD.Guides.InterpolationBulkData, ResourceName, OwnerName);
 
 				// Update card stats to display
 				Desc->CardsInfo.NumCardVertices = LOD.BulkData.GetNumVertices();
@@ -3029,7 +3030,7 @@ FHairStrandsRestResource* UGroomAsset::AllocateGuidesResources(uint32 GroupIndex
 		{
 			if (GroupData.Guides.RestResource == nullptr)
 			{
-				GroupData.Guides.RestResource = new FHairStrandsRestResource(GroupData.Guides.BulkData, EHairStrandsResourcesType::Guides, FHairResourceName(GetFName(), GroupIndex));
+				GroupData.Guides.RestResource = new FHairStrandsRestResource(GroupData.Guides.BulkData, EHairStrandsResourcesType::Guides, FHairResourceName(GetFName(), GroupIndex), GetAssetPathName());
 			}
 			return GroupData.Guides.RestResource;
 		}
@@ -3045,7 +3046,7 @@ FHairStrandsInterpolationResource* UGroomAsset::AllocateInterpolationResources(u
 		check(GroupData.Guides.IsValid());
 		if (GroupData.Strands.InterpolationResource == nullptr)
 		{
-			GroupData.Strands.InterpolationResource = new FHairStrandsInterpolationResource(GroupData.Strands.InterpolationBulkData, FHairResourceName(GetFName(), GroupIndex));
+			GroupData.Strands.InterpolationResource = new FHairStrandsInterpolationResource(GroupData.Strands.InterpolationBulkData, FHairResourceName(GetFName(), GroupIndex), GetAssetPathName());
 		}
 		return GroupData.Strands.InterpolationResource;
 	}
@@ -3063,7 +3064,7 @@ FHairStrandsRaytracingResource* UGroomAsset::AllocateCardsRaytracingResources(ui
 
 		if (LOD.RaytracingResource == nullptr)
 		{
-			LOD.RaytracingResource = new FHairStrandsRaytracingResource(LOD.BulkData, FHairResourceName(GetFName(), GroupIndex, LODIndex));
+			LOD.RaytracingResource = new FHairStrandsRaytracingResource(LOD.BulkData, FHairResourceName(GetFName(), GroupIndex, LODIndex), GetAssetPathName(LODIndex));
 		}
 		return LOD.RaytracingResource;
 	}
@@ -3080,7 +3081,7 @@ FHairStrandsRaytracingResource* UGroomAsset::AllocateMeshesRaytracingResources(u
 
 		if (LOD.RaytracingResource == nullptr)
 		{
-			LOD.RaytracingResource = new FHairStrandsRaytracingResource(LOD.BulkData, FHairResourceName(GetFName(), GroupIndex, LODIndex));
+			LOD.RaytracingResource = new FHairStrandsRaytracingResource(LOD.BulkData, FHairResourceName(GetFName(), GroupIndex, LODIndex), GetAssetPathName(LODIndex));
 		}
 		return LOD.RaytracingResource;
 	}
@@ -3096,7 +3097,7 @@ FHairStrandsRaytracingResource* UGroomAsset::AllocateStrandsRaytracingResources(
 
 		if (GroupData.Strands.RaytracingResource == nullptr)
 		{
-			GroupData.Strands.RaytracingResource = new FHairStrandsRaytracingResource(GroupData.Strands.BulkData, FHairResourceName(GetFName(), GroupIndex));
+			GroupData.Strands.RaytracingResource = new FHairStrandsRaytracingResource(GroupData.Strands.BulkData, FHairResourceName(GetFName(), GroupIndex), GetAssetPathName());
 		}
 		return GroupData.Strands.RaytracingResource;
 	}
@@ -3118,6 +3119,8 @@ void UGroomAsset::InitStrandsResources()
 		return;
 	}
 
+	const FName OwnerName = GetAssetPathName();
+
 	for (uint32 GroupIndex = 0, GroupCount = GetNumHairGroups(); GroupIndex < GroupCount; ++GroupIndex)
 	{
 		FHairGroupData& GroupData = HairGroupsData[GroupIndex];
@@ -3125,8 +3128,7 @@ void UGroomAsset::InitStrandsResources()
 		if (GroupData.Strands.HasValidData())
 		{
 			FHairResourceName ResourceName(GetFName(), GroupIndex);
-
-			GroupData.Strands.RestResource = new FHairStrandsRestResource(GroupData.Strands.BulkData, EHairStrandsResourcesType::Strands, ResourceName);
+			GroupData.Strands.RestResource = new FHairStrandsRestResource(GroupData.Strands.BulkData, EHairStrandsResourcesType::Strands, ResourceName, OwnerName);
 
 			if (GroupData.Strands.ClusterCullingBulkData.IsValid())
 			{
@@ -3147,7 +3149,7 @@ void UGroomAsset::InitStrandsResources()
 					}
 				}
 
-				GroupData.Strands.ClusterCullingResource = new FHairStrandsClusterCullingResource(GroupData.Strands.ClusterCullingBulkData, ResourceName);
+				GroupData.Strands.ClusterCullingResource = new FHairStrandsClusterCullingResource(GroupData.Strands.ClusterCullingBulkData, ResourceName, OwnerName);
 			}
 
 			// Interpolation are lazy allocated, as these resources are only used when RBF/simulation is enabled by a groom component
@@ -3232,15 +3234,16 @@ void UGroomAsset::InitCardsResources()
 				LOD.HasValidData())
 			{
 				FHairResourceName ResourceName(GetFName(), GroupIndex, LODIt);
+				const FName OwnerName = GetAssetPathName(LODIt);
 
-				LOD.RestResource = new FHairCardsRestResource(LOD.BulkData, ResourceName);
+				LOD.RestResource = new FHairCardsRestResource(LOD.BulkData, ResourceName, OwnerName);
 				BeginInitResource(LOD.RestResource); // Immediate allocation, as needed for the vertex factory, input stream building
 
-				LOD.InterpolationResource = new FHairCardsInterpolationResource(LOD.InterpolationBulkData, ResourceName);
+				LOD.InterpolationResource = new FHairCardsInterpolationResource(LOD.InterpolationBulkData, ResourceName, OwnerName);
 
-				LOD.Guides.RestResource = new FHairStrandsRestResource(LOD.Guides.BulkData, EHairStrandsResourcesType::Cards, ResourceName);
+				LOD.Guides.RestResource = new FHairStrandsRestResource(LOD.Guides.BulkData, EHairStrandsResourcesType::Cards, ResourceName, OwnerName);
 
-				LOD.Guides.InterpolationResource = new FHairStrandsInterpolationResource(LOD.Guides.InterpolationBulkData, ResourceName);
+				LOD.Guides.InterpolationResource = new FHairStrandsInterpolationResource(LOD.Guides.InterpolationBulkData, ResourceName, OwnerName);
 
 				if (Desc)
 				{
@@ -3291,7 +3294,7 @@ void UGroomAsset::InitMeshesResources()
 
 			if (LOD.HasValidData())
 			{
-				LOD.RestResource = new FHairMeshesRestResource(LOD.BulkData, FHairResourceName(GetFName(), GroupIndex, LODIt));
+				LOD.RestResource = new FHairMeshesRestResource(LOD.BulkData, FHairResourceName(GetFName(), GroupIndex, LODIt), GetAssetPathName(LODIt));
 				BeginInitResource(LOD.RestResource); // Immediate allocation, as needed for the vertex factory, input stream building
 
 
@@ -3719,10 +3722,11 @@ void UGroomAsset::SaveProceduralCards(uint32 DescIndex)
 		Q.Textures = &Desc->Textures;
 	}
 
+	const FName OwnerName = GetAssetPathName(LODIndex);
 	// 3. Create resources and enqueue texture generation (GPU, kicked by the render thread) 
-	Q.Resources = new FHairCardsRestResource(Q.BulkData, FHairResourceName(GetFName(), GroupIndex, LODIndex));
+	Q.Resources = new FHairCardsRestResource(Q.BulkData, FHairResourceName(GetFName(), GroupIndex, LODIndex), OwnerName);
 	BeginInitResource(Q.Resources); // Immediate allocation, as needed for the vertex factory, input stream building
-	Q.ProceduralResources = new FHairCardsProceduralResource(Q.ProceduralData.RenderData, Q.ProceduralData.Atlas.Resolution, Q.ProceduralData.Voxels);
+	Q.ProceduralResources = new FHairCardsProceduralResource(Q.ProceduralData.RenderData, Q.ProceduralData.Atlas.Resolution, Q.ProceduralData.Voxels, OwnerName);
 	BeginInitResource(Q.ProceduralResources); // Immediate allocation, as needed for the vertex factory, input stream building
 
 	FHairCardsBuilder::BuildTextureAtlas(&Q.ProceduralData, Q.Resources, Q.ProceduralResources, Q.Textures);
@@ -3880,5 +3884,21 @@ void UGroomAsset::GetResourceSizeEx(FResourceSizeEx& CumulativeResourceSize)
 	}
 }
 
+FName UGroomAsset::GetAssetPathName(int32 LODIndex)
+{
+#if RHI_ENABLE_RESOURCE_INFO
+	if (LODIndex > -1)
+	{
+		return FName(FString::Printf(TEXT("%s [LOD%d]"), *GetPathName(), LODIndex));
+	}
+	else
+	{
+		return FName(GetPathName());
+	}
+#else
+	return NAME_None;
+#endif
+}
+
 #undef LOCTEXT_NAMESPACE
 
diff --git a/Engine/Plugins/Runtime/HairStrands/Source/HairStrandsCore/Private/GroomBindingAsset.cpp b/Engine/Plugins/Runtime/HairStrands/Source/HairStrandsCore/Private/GroomBindingAsset.cpp
--- a/Engine/Plugins/Runtime/HairStrands/Source/HairStrandsCore/Private/GroomBindingAsset.cpp
+++ b/Engine/Plugins/Runtime/HairStrands/Source/HairStrandsCore/Private/GroomBindingAsset.cpp
@@ -111,12 +111,13 @@ void UGroomBindingAsset::InitResource()
 		FHairGroupResource& Resource = HairGroupResources.AddDefaulted_GetRef();
 
 		FHairResourceName ResourceName(GetFName(), GroupIndex);
+		const FName OwnerName = GetAssetPathName();
 
 		// Guides
 		Resource.SimRootResources = nullptr;
 		if (BulkData.SimRootBulkData.IsValid())
 		{
-			Resource.SimRootResources = new FHairStrandsRestRootResource(BulkData.SimRootBulkData, EHairStrandsResourcesType::Guides, ResourceName);
+			Resource.SimRootResources = new FHairStrandsRestRootResource(BulkData.SimRootBulkData, EHairStrandsResourcesType::Guides, ResourceName, OwnerName);
 			BeginInitResource(Resource.SimRootResources);
 		}
 
@@ -124,7 +125,7 @@ void UGroomBindingAsset::InitResource()
 		Resource.RenRootResources = nullptr;
 		if (IsHairStrandsEnabled(EHairStrandsShaderType::Strands) && BulkData.RenRootBulkData.IsValid())
 		{
-			Resource.RenRootResources = new FHairStrandsRestRootResource(BulkData.RenRootBulkData, EHairStrandsResourcesType::Strands, ResourceName);
+			Resource.RenRootResources = new FHairStrandsRestRootResource(BulkData.RenRootBulkData, EHairStrandsResourcesType::Strands, ResourceName, OwnerName);
 			BeginInitResource(Resource.RenRootResources);
 		}
 
@@ -138,7 +139,7 @@ void UGroomBindingAsset::InitResource()
 				Resource.CardsRootResources[CardsLODIt] = nullptr;
 				if (BulkData.CardsRootBulkData[CardsLODIt].IsValid())
 				{
-					Resource.CardsRootResources[CardsLODIt] = new FHairStrandsRestRootResource(BulkData.CardsRootBulkData[CardsLODIt], EHairStrandsResourcesType::Cards, FHairResourceName(GetFName(), GroupIndex, CardsLODIt));
+					Resource.CardsRootResources[CardsLODIt] = new FHairStrandsRestRootResource(BulkData.CardsRootBulkData[CardsLODIt], EHairStrandsResourcesType::Cards, FHairResourceName(GetFName(), GroupIndex, CardsLODIt), GetAssetPathName(CardsLODIt));
 					BeginInitResource(Resource.CardsRootResources[CardsLODIt]);
 				}
 			}
@@ -818,3 +819,18 @@ void UGroomBindingAsset::GetResourceSizeEx(FResourceSizeEx& CumulativeResourceSi
 	}
 }
 
+FName UGroomBindingAsset::GetAssetPathName(int32 LODIndex)
+{
+#if RHI_ENABLE_RESOURCE_INFO
+	if (LODIndex > -1)
+	{
+		return FName(FString::Printf(TEXT("%s [LOD%d]"), *GetPathName(), LODIndex));
+	}
+	else
+	{
+		return FName(GetPathName());
+	}
+#else
+	return NAME_None;
+#endif
+}
\ No newline at end of file
diff --git a/Engine/Plugins/Runtime/HairStrands/Source/HairStrandsCore/Private/GroomComponent.cpp b/Engine/Plugins/Runtime/HairStrands/Source/HairStrandsCore/Private/GroomComponent.cpp
--- a/Engine/Plugins/Runtime/HairStrands/Source/HairStrandsCore/Private/GroomComponent.cpp
+++ b/Engine/Plugins/Runtime/HairStrands/Source/HairStrandsCore/Private/GroomComponent.cpp
@@ -2586,10 +2586,12 @@ void UGroomComponent::InitResources(bool bIsBindingReloading)
 
 		const EHairInterpolationType HairInterpolationType = ToHairInterpolationType(GroomAsset->HairInterpolationType);
 
+		const FName OwnerName = GroomAsset->GetAssetPathName();
+
 		// Initialize LOD screen size & visibility
 		bool bNeedStrandsData = false;
 		{
-			HairGroupInstance->HairGroupPublicData = new FHairGroupPublicData(GroupIt);
+			HairGroupInstance->HairGroupPublicData = new FHairGroupPublicData(GroupIt, OwnerName);
 			HairGroupInstance->HairGroupPublicData->Instance = HairGroupInstance;
 			HairGroupInstance->HairGroupPublicData->bDebugDrawLODInfo = bPreviewMode;
 			HairGroupInstance->HairGroupPublicData->DebugScreenSize = 0.f;
@@ -2657,7 +2659,7 @@ void UGroomComponent::InitResources(bool bIsBindingReloading)
 				}
 
 				HairGroupInstance->Guides.RestRootResource = LocalBindingAsset->HairGroupResources[GroupIt].SimRootResources;
-				HairGroupInstance->Guides.DeformedRootResource = new FHairStrandsDeformedRootResource(HairGroupInstance->Guides.RestRootResource, EHairStrandsResourcesType::Guides, ResourceName);
+				HairGroupInstance->Guides.DeformedRootResource = new FHairStrandsDeformedRootResource(HairGroupInstance->Guides.RestRootResource, EHairStrandsResourcesType::Guides, ResourceName, OwnerName);
 			}
 
 			// Lazy allocation of the guide resources
@@ -2665,7 +2667,7 @@ void UGroomComponent::InitResources(bool bIsBindingReloading)
 			check(GroupData.Guides.RestResource);
 
 			// If guides are allocated, deformed resources are always needs since they are either used with simulation, or RBF deformation. Both are dynamics, and require deformed positions
-			HairGroupInstance->Guides.DeformedResource = new FHairStrandsDeformedResource(GroupData.Guides.BulkData, EHairStrandsResourcesType::Guides, ResourceName);
+			HairGroupInstance->Guides.DeformedResource = new FHairStrandsDeformedResource(GroupData.Guides.BulkData, EHairStrandsResourcesType::Guides, ResourceName, OwnerName);
 
 			// Initialize the simulation and the global deformation to its default behavior by setting it with LODIndex = -1
 			const int32 LODIndex = -1;
@@ -2747,7 +2749,7 @@ void UGroomComponent::InitResources(bool bIsBindingReloading)
 				if (bNeedDynamicResources)
 				{
 					// Allocate dynamic raytracing resources (owned by the groom component/instance)
-					HairGroupInstance->Strands.RenRaytracingResource = new FHairStrandsRaytracingResource(GroupData.Strands.BulkData, ResourceName);
+					HairGroupInstance->Strands.RenRaytracingResource = new FHairStrandsRaytracingResource(GroupData.Strands.BulkData, ResourceName, OwnerName);
 					HairGroupInstance->Strands.RenRaytracingResourceOwned = true;
 				}
 				else
@@ -2771,13 +2773,13 @@ void UGroomComponent::InitResources(bool bIsBindingReloading)
 				}
 
 				HairGroupInstance->Strands.RestRootResource = LocalBindingAsset->HairGroupResources[GroupIt].RenRootResources;
-				HairGroupInstance->Strands.DeformedRootResource = new FHairStrandsDeformedRootResource(HairGroupInstance->Strands.RestRootResource, EHairStrandsResourcesType::Strands, ResourceName);
+				HairGroupInstance->Strands.DeformedRootResource = new FHairStrandsDeformedRootResource(HairGroupInstance->Strands.RestRootResource, EHairStrandsResourcesType::Strands, ResourceName, OwnerName);
 			}
 
 			HairGroupInstance->Strands.RestResource = GroupData.Strands.RestResource;
 			if (bNeedDynamicResources)
 			{
-				HairGroupInstance->Strands.DeformedResource = new FHairStrandsDeformedResource(GroupData.Strands.BulkData, EHairStrandsResourcesType::Strands, ResourceName);
+				HairGroupInstance->Strands.DeformedResource = new FHairStrandsDeformedResource(GroupData.Strands.BulkData, EHairStrandsResourcesType::Strands, ResourceName, OwnerName);
 			} 
 
 			// An empty groom doesn't have a ClusterCullingResource
@@ -2817,6 +2819,7 @@ void UGroomComponent::InitResources(bool bIsBindingReloading)
 			if (LOD.IsValid())
 			{
 				FHairResourceName LODResourceName(GetFName(), GroupIt, CardsLODIndex);
+				const FName LODOwnerName = GroomAsset->GetAssetPathName(CardsLODIndex);
 
 				const EHairBindingType BindingType	= HairGroupInstance->HairGroupPublicData->GetBindingType(CardsLODIndex);
 				const bool bHasSimulation			= HairGroupInstance->HairGroupPublicData->IsSimulationEnable(CardsLODIndex);
@@ -2837,7 +2840,7 @@ void UGroomComponent::InitResources(bool bIsBindingReloading)
 
 				if (bNeedDeformedPositions)
 				{
-					InstanceLOD.DeformedResource = new FHairCardsDeformedResource(LOD.BulkData, false, LODResourceName);
+					InstanceLOD.DeformedResource = new FHairCardsDeformedResource(LOD.BulkData, false, LODResourceName, LODOwnerName);
 				}
 
 				#if RHI_RAYTRACING
@@ -2846,7 +2849,7 @@ void UGroomComponent::InitResources(bool bIsBindingReloading)
 					if (bNeedDeformedPositions)
 					{
 						// Allocate dynamic raytracing resources (owned by the groom component/instance)
-						InstanceLOD.RaytracingResource = new FHairStrandsRaytracingResource(*InstanceLOD.Data, LODResourceName);
+						InstanceLOD.RaytracingResource = new FHairStrandsRaytracingResource(*InstanceLOD.Data, LODResourceName, LODOwnerName);
 						InstanceLOD.RaytracingResourceOwned = true;
 					}
 					else
@@ -2875,12 +2878,12 @@ void UGroomComponent::InitResources(bool bIsBindingReloading)
 						}
 
 						InstanceLOD.Guides.RestRootResource = LocalBindingAsset->HairGroupResources[GroupIt].CardsRootResources[CardsLODIndex];
-						InstanceLOD.Guides.DeformedRootResource = new FHairStrandsDeformedRootResource(InstanceLOD.Guides.RestRootResource, EHairStrandsResourcesType::Cards, LODResourceName);
+						InstanceLOD.Guides.DeformedRootResource = new FHairStrandsDeformedRootResource(InstanceLOD.Guides.RestRootResource, EHairStrandsResourcesType::Cards, LODResourceName, LODOwnerName);
 					}
 
 					InstanceLOD.Guides.RestResource = LOD.Guides.RestResource;
 					{
-						InstanceLOD.Guides.DeformedResource = new FHairStrandsDeformedResource(LOD.Guides.BulkData, EHairStrandsResourcesType::Cards, LODResourceName);
+						InstanceLOD.Guides.DeformedResource = new FHairStrandsDeformedResource(LOD.Guides.BulkData, EHairStrandsResourcesType::Cards, LODResourceName, LODOwnerName);
 					}
 
 					InstanceLOD.Guides.HairInterpolationType = HairInterpolationType;
@@ -2896,6 +2899,7 @@ void UGroomComponent::InitResources(bool bIsBindingReloading)
 			if (LOD.IsValid())
 			{
 				FHairResourceName LODResourceName(GetFName(), GroupIt, MeshLODIndex);
+				const FName LODOwnerName = GroomAsset->GetAssetPathName(MeshLODIndex);
 
 				const EHairBindingType BindingType = HairGroupInstance->HairGroupPublicData->GetBindingType(MeshLODIndex);
 				const bool bHasGlobalDeformation   = HairGroupInstance->HairGroupPublicData->IsGlobalInterpolationEnable(MeshLODIndex);
@@ -2905,7 +2909,7 @@ void UGroomComponent::InitResources(bool bIsBindingReloading)
 				InstanceLOD.RestResource = LOD.RestResource;
 				if (bNeedDeformedPositions)
 				{
-					InstanceLOD.DeformedResource = new FHairMeshesDeformedResource(LOD.BulkData, true, LODResourceName);
+					InstanceLOD.DeformedResource = new FHairMeshesDeformedResource(LOD.BulkData, true, LODResourceName, LODOwnerName);
 				}
 
 				#if RHI_RAYTRACING
@@ -2914,7 +2918,7 @@ void UGroomComponent::InitResources(bool bIsBindingReloading)
 					if (bNeedDeformedPositions)
 					{
 						// Allocate dynamic raytracing resources (owned by the groom component/instance)
-						InstanceLOD.RaytracingResource = new FHairStrandsRaytracingResource(*InstanceLOD.Data, LODResourceName);
+						InstanceLOD.RaytracingResource = new FHairStrandsRaytracingResource(*InstanceLOD.Data, LODResourceName, LODOwnerName);
 						InstanceLOD.RaytracingResourceOwned = true;
 					}
 					else
diff --git a/Engine/Plugins/Runtime/HairStrands/Source/HairStrandsCore/Private/GroomManager.cpp b/Engine/Plugins/Runtime/HairStrands/Source/HairStrandsCore/Private/GroomManager.cpp
--- a/Engine/Plugins/Runtime/HairStrands/Source/HairStrandsCore/Private/GroomManager.cpp
+++ b/Engine/Plugins/Runtime/HairStrands/Source/HairStrandsCore/Private/GroomManager.cpp
@@ -578,9 +578,15 @@ static void AllocateRaytracingResources(FHairGroupInstance* Instance)
 	{
 		check(Instance->Strands.Data);
 
+#if RHI_ENABLE_RESOURCE_INFO
+		FName OwnerName(FString::Printf(TEXT("%s [LOD%d]"), *Instance->Debug.MeshComponentName, Instance->Debug.MeshLODIndex));
+#else
+		FName OwnerName = NAME_None;
+#endif
+
 		// Allocate dynamic raytracing resources (owned by the groom component/instance)
 		FHairResourceName ResourceName(FName(Instance->Debug.GroomAssetName), Instance->Debug.GroupIndex);
-		Instance->Strands.RenRaytracingResource		 = new FHairStrandsRaytracingResource(*Instance->Strands.Data, ResourceName);
+		Instance->Strands.RenRaytracingResource		 = new FHairStrandsRaytracingResource(*Instance->Strands.Data, ResourceName, OwnerName);
 		Instance->Strands.RenRaytracingResourceOwned = true;
 	}
 	Instance->Strands.ViewRayTracingMask |= EHairViewRayTracingMask::PathTracing;
diff --git a/Engine/Plugins/Runtime/HairStrands/Source/HairStrandsCore/Private/GroomResources.cpp b/Engine/Plugins/Runtime/HairStrands/Source/HairStrandsCore/Private/GroomResources.cpp
--- a/Engine/Plugins/Runtime/HairStrands/Source/HairStrandsCore/Private/GroomResources.cpp
+++ b/Engine/Plugins/Runtime/HairStrands/Source/HairStrandsCore/Private/GroomResources.cpp
@@ -136,7 +136,7 @@ void CreateBuffer(uint32 InVertexCount, FRWBuffer& OutBuffer, const TCHAR* Debug
 }
 
 template<typename FormatType>
-void CreateBuffer(const TArray<typename FormatType::Type>& InData, FHairCardsVertexBuffer& OutBuffer, const TCHAR* DebugName, ERHIAccess InitialAccess = ERHIAccess::SRVMask)
+void CreateBuffer(const TArray<typename FormatType::Type>& InData, FHairCardsVertexBuffer& OutBuffer, const TCHAR* DebugName, const FName& OwnerName, ERHIAccess InitialAccess = ERHIAccess::SRVMask)
 {
 	const uint32 DataCount = InData.Num();
 	const uint32 DataSizeInBytes = FormatType::SizeInByte * DataCount;
@@ -147,6 +147,7 @@ void CreateBuffer(const TArray<typename FormatType::Type>& InData, FHairCardsVer
 	CreateInfo.ResourceArray = nullptr;
 
 	OutBuffer.VertexBufferRHI = RHICreateVertexBuffer(DataSizeInBytes, BUF_Static | BUF_ShaderResource, InitialAccess, CreateInfo);
+	OutBuffer.VertexBufferRHI->SetOwnerName(OwnerName);
 
 	void* BufferData = RHILockBuffer(OutBuffer.VertexBufferRHI, 0, DataSizeInBytes, RLM_WriteOnly);
 	FMemory::Memcpy(BufferData, InData.GetData(), DataSizeInBytes);
@@ -164,6 +165,7 @@ static FRDGBufferRef InternalCreateVertexBuffer(
 	const void* InitialData,
 	uint64 InitialDataSize,
 	ERDGInitialDataFlags InitialDataFlags,
+	const FName& OwnerName,
 	bool bIsByteAddressBuffer = false)
 {
 	if (!bIsByteAddressBuffer)
@@ -171,6 +173,7 @@ static FRDGBufferRef InternalCreateVertexBuffer(
 		checkf(EnumHasAnyFlags(Desc.Usage, EBufferUsageFlags::VertexBuffer), TEXT("CreateVertexBuffer called with an FRDGBufferDesc underlying type that is not 'VertexBuffer'. Buffer: %s"), Name);
 	}
 	FRDGBufferRef Buffer = GraphBuilder.CreateBuffer(Desc, Name, ERDGBufferFlags::MultiFrame);
+	Buffer->SetOwnerName(OwnerName);
 	if (InitialData && InitialDataSize)
 	{
 		GraphBuilder.QueueBufferUpload(Buffer, InitialData, InitialDataSize, InitialDataFlags);
@@ -179,7 +182,7 @@ static FRDGBufferRef InternalCreateVertexBuffer(
 }
 
 template<typename FormatType>
-void InternalCreateVertexBufferRDG_FromBulkData(FRDGBuilder& GraphBuilder, FByteBulkData& InBulkData, uint32 InDataCount, FRDGExternalBuffer& Out, const TCHAR* DebugName, EHairResourceUsageType UsageType)
+void InternalCreateVertexBufferRDG_FromBulkData(FRDGBuilder& GraphBuilder, FByteBulkData& InBulkData, uint32 InDataCount, FRDGExternalBuffer& Out, const TCHAR* DebugName, const FName& OwnerName, EHairResourceUsageType UsageType)
 {
 	const uint32 InDataCount_Check = InBulkData.GetBulkDataSize() / sizeof(typename FormatType::BulkType);
 	check(InDataCount_Check == InDataCount);
@@ -213,6 +216,7 @@ void InternalCreateVertexBufferRDG_FromBulkData(FRDGBuilder& GraphBuilder, FByte
 	{
 		checkf(EnumHasAnyFlags(Desc.Usage, EBufferUsageFlags::VertexBuffer), TEXT("CreateVertexBuffer called with an FRDGBufferDesc underlying type that is not 'VertexBuffer'. Buffer: %s"), DebugName);
 		Buffer = GraphBuilder.CreateBuffer(Desc, DebugName, ERDGBufferFlags::MultiFrame);
+		Buffer->SetOwnerName(OwnerName);
 		if (BulkData && DataSizeInBytes)
 		{
 			GraphBuilder.QueueBufferUpload(Buffer, BulkData, DataSizeInBytes, [&InBulkData](const void* Ptr) { InBulkData.Unlock(); });
@@ -227,14 +231,15 @@ void InternalCreateVertexBufferRDG_FromBulkData(FRDGBuilder& GraphBuilder, FByte
 			Desc,
 			BulkData,
 			DataSizeInBytes,
-			ERDGInitialDataFlags::None); // Copy data internally
+			ERDGInitialDataFlags::None, // Copy data internally
+			OwnerName);
 		InBulkData.Unlock();
 	}
 
 	ConvertToExternalBufferWithViews(GraphBuilder, Buffer, Out, FormatType::Format);
 }
 
-void InternalCreateByteAddressBufferRDG_FromBulkData(FRDGBuilder& GraphBuilder, FByteBulkData& InBulkData, FRDGExternalBuffer& Out, const TCHAR* DebugName, EHairResourceUsageType UsageType)
+void InternalCreateByteAddressBufferRDG_FromBulkData(FRDGBuilder& GraphBuilder, FByteBulkData& InBulkData, FRDGExternalBuffer& Out, const TCHAR* DebugName, const FName& OwnerName, EHairResourceUsageType UsageType)
 {
 	const uint32 DataSizeInBytes = InBulkData.GetBulkDataSize();
 	if (DataSizeInBytes == 0)
@@ -264,6 +269,7 @@ void InternalCreateByteAddressBufferRDG_FromBulkData(FRDGBuilder& GraphBuilder,
 	if (bReleaseCPUData)
 	{		
 		Buffer = GraphBuilder.CreateBuffer(Desc, DebugName, ERDGBufferFlags::MultiFrame);
+		Buffer->SetOwnerName(OwnerName);
 		if (BulkData && DataSizeInBytes)
 		{
 			GraphBuilder.QueueBufferUpload(Buffer, BulkData, DataSizeInBytes, [&InBulkData](const void* Ptr) { InBulkData.Unlock(); });
@@ -279,6 +285,7 @@ void InternalCreateByteAddressBufferRDG_FromBulkData(FRDGBuilder& GraphBuilder,
 			BulkData,
 			DataSizeInBytes,
 			ERDGInitialDataFlags::None, 
+			OwnerName,
 			true /*bIsByteAddressBuffer*/); // Copy data internally
 		InBulkData.Unlock();
 	}
@@ -290,7 +297,7 @@ void InternalCreateByteAddressBufferRDG_FromBulkData(FRDGBuilder& GraphBuilder,
 
 
 template<typename FormatType>
-void InternalCreateVertexBufferRDG(FRDGBuilder& GraphBuilder, const typename FormatType::Type* InData, uint32 InDataCount, FRDGExternalBuffer& Out, const TCHAR* DebugName, EHairResourceUsageType UsageType, ERDGInitialDataFlags InitialDataFlags)
+void InternalCreateVertexBufferRDG(FRDGBuilder& GraphBuilder, const typename FormatType::Type* InData, uint32 InDataCount, FRDGExternalBuffer& Out, const TCHAR* DebugName, const FName& OwnerName, EHairResourceUsageType UsageType, ERDGInitialDataFlags InitialDataFlags)
 {
 	FRDGBufferRef Buffer = nullptr;
 
@@ -312,13 +319,14 @@ void InternalCreateVertexBufferRDG(FRDGBuilder& GraphBuilder, const typename For
 		Desc,
 		InData,
 		DataSizeInBytes,
-		InitialDataFlags);
+		InitialDataFlags,
+		OwnerName);
 
 	ConvertToExternalBufferWithViews(GraphBuilder, Buffer, Out, FormatType::Format);
 }
 
 template<typename FormatType>
-void InternalCreateVertexBufferRDG(FRDGBuilder& GraphBuilder, const TArray<typename FormatType::Type>& InData, FRDGExternalBuffer& Out, const TCHAR* DebugName, EHairResourceUsageType UsageType, ERDGInitialDataFlags InitialDataFlags= ERDGInitialDataFlags::NoCopy)
+void InternalCreateVertexBufferRDG(FRDGBuilder& GraphBuilder, const TArray<typename FormatType::Type>& InData, FRDGExternalBuffer& Out, const TCHAR* DebugName, const FName& OwnerName, EHairResourceUsageType UsageType, ERDGInitialDataFlags InitialDataFlags= ERDGInitialDataFlags::NoCopy)
 {
 	FRDGBufferRef Buffer = nullptr;
 
@@ -341,13 +349,14 @@ void InternalCreateVertexBufferRDG(FRDGBuilder& GraphBuilder, const TArray<typen
 		Desc,
 		InData.GetData(),
 		DataSizeInBytes,
-		InitialDataFlags);
+		InitialDataFlags,
+		OwnerName);
 
 	ConvertToExternalBufferWithViews(GraphBuilder, Buffer, Out, FormatType::Format);
 }
 
 template<typename DataType>
-void InternalCreateVertexBufferRDG(FRDGBuilder& GraphBuilder, const TArray<DataType>& InData, EPixelFormat Format, FRDGExternalBuffer& Out, const TCHAR* DebugName, EHairResourceUsageType UsageType)
+void InternalCreateVertexBufferRDG(FRDGBuilder& GraphBuilder, const TArray<DataType>& InData, EPixelFormat Format, FRDGExternalBuffer& Out, const TCHAR* DebugName, const FName& OwnerName, EHairResourceUsageType UsageType)
 {
 	FRDGBufferRef Buffer = nullptr;
 
@@ -370,13 +379,14 @@ void InternalCreateVertexBufferRDG(FRDGBuilder& GraphBuilder, const TArray<DataT
 		Desc,
 		InData.GetData(),
 		DataSizeInBytes,
-		ERDGInitialDataFlags::NoCopy);
+		ERDGInitialDataFlags::NoCopy,
+		OwnerName);
 
 	ConvertToExternalBufferWithViews(GraphBuilder, Buffer, Out, Format);
 }
 
 template<typename FormatType>
-void InternalCreateVertexBufferRDG(FRDGBuilder& GraphBuilder, uint32 InVertexCount, FRDGExternalBuffer& Out, const TCHAR* DebugName, EHairResourceUsageType Usage)
+void InternalCreateVertexBufferRDG(FRDGBuilder& GraphBuilder, uint32 InVertexCount, FRDGExternalBuffer& Out, const TCHAR* DebugName, const FName& OwnerName, EHairResourceUsageType Usage)
 {
 	// Sanity check
 	check(Usage == EHairResourceUsageType::Dynamic);
@@ -391,6 +401,7 @@ void InternalCreateVertexBufferRDG(FRDGBuilder& GraphBuilder, uint32 InVertexCou
 
 	FRDGBufferDesc Desc = FRDGBufferDesc::CreateBufferDesc(FormatType::SizeInByte, InVertexCount);
 	FRDGBufferRef Buffer = GraphBuilder.CreateBuffer(Desc, DebugName, ERDGBufferFlags::MultiFrame);
+	Buffer->SetOwnerName(OwnerName);
 
 	if (IsFloatFormat(FormatType::Format) || IsUnormFormat(FormatType::Format) || IsSnormFormat(FormatType::Format))
 	{
@@ -405,7 +416,7 @@ void InternalCreateVertexBufferRDG(FRDGBuilder& GraphBuilder, uint32 InVertexCou
 }
 
 template<typename FormatType>
-void InternalCreateStructuredBufferRDG(FRDGBuilder& GraphBuilder, uint32 DataCount, FRDGExternalBuffer& Out, const TCHAR* DebugName, EHairResourceUsageType Usage)
+void InternalCreateStructuredBufferRDG(FRDGBuilder& GraphBuilder, uint32 DataCount, FRDGExternalBuffer& Out, const TCHAR* DebugName, const FName& OwnerName, EHairResourceUsageType Usage)
 {
 	// Sanity check
 	check(Usage == EHairResourceUsageType::Dynamic);
@@ -419,13 +430,14 @@ void InternalCreateStructuredBufferRDG(FRDGBuilder& GraphBuilder, uint32 DataCou
 
 	FRDGBufferDesc Desc = FRDGBufferDesc::CreateStructuredDesc(FormatType::SizeInByte, DataCount);
 	FRDGBufferRef Buffer = GraphBuilder.CreateBuffer(Desc, DebugName, ERDGBufferFlags::MultiFrame);
+	Buffer->SetOwnerName(OwnerName);
 	AddClearUAVPass(GraphBuilder, GraphBuilder.CreateUAV(Buffer, FormatType::Format), 0u);
 	ConvertToExternalBufferWithViews(GraphBuilder, Buffer, Out, FormatType::Format);
 }
 
 
 template<typename FormatType>
-void InternalCreateStructuredBufferRDG_FromBulkData(FRDGBuilder& GraphBuilder, FByteBulkData& InBulkData, uint32 InDataCount, FRDGExternalBuffer& Out, const TCHAR* DebugName, EHairResourceUsageType UsageType)
+void InternalCreateStructuredBufferRDG_FromBulkData(FRDGBuilder& GraphBuilder, FByteBulkData& InBulkData, uint32 InDataCount, FRDGExternalBuffer& Out, const TCHAR* DebugName, const FName& OwnerName, EHairResourceUsageType UsageType)
 {
 	const uint32 InSizeInByte = sizeof(typename FormatType::Type);
 	const uint32 DataSizeInBytes = InSizeInByte * InDataCount;
@@ -454,6 +466,7 @@ void InternalCreateStructuredBufferRDG_FromBulkData(FRDGBuilder& GraphBuilder, F
 
 	const typename FormatType::Type* Data = (const typename FormatType::Type*)InBulkData.Lock(LOCK_READ_ONLY);
 	FRDGBufferRef Buffer = GraphBuilder.CreateBuffer(Desc, DebugName, ERDGBufferFlags::MultiFrame);
+	Buffer->SetOwnerName(OwnerName);
 	if (Data && DataSizeInBytes)
 	{
 		#if !WITH_EDITORONLY_DATA
@@ -514,7 +527,7 @@ static UTexture2D* CreateCardTexture(FIntPoint Resolution)
 }
 
 /////////////////////////////////////////////////////////////////////////////////////////
-void CreateHairStrandsDebugAttributeBuffer(FRDGBuilder& GraphBuilder, FRDGExternalBuffer* DebugAttributeBuffer, uint32 SizeInBytes)
+void CreateHairStrandsDebugAttributeBuffer(FRDGBuilder& GraphBuilder, FRDGExternalBuffer* DebugAttributeBuffer, uint32 SizeInBytes, const FName& OwnerName)
 {
 	if (SizeInBytes == 0 || !DebugAttributeBuffer)
 	{
@@ -522,16 +535,18 @@ void CreateHairStrandsDebugAttributeBuffer(FRDGBuilder& GraphBuilder, FRDGExtern
 	}
 
 	FRDGBufferRef Buffer = GraphBuilder.CreateBuffer(FRDGBufferDesc::CreateByteAddressDesc(SizeInBytes), TEXT("Hair.Strands_DebugAttributeBuffer"), ERDGBufferFlags::MultiFrame);
+	Buffer->SetOwnerName(OwnerName);
 	ConvertToExternalBufferWithViews(GraphBuilder, Buffer, *DebugAttributeBuffer, PF_R32_UINT);
 }
 
 /////////////////////////////////////////////////////////////////////////////////////////
 
-FHairCommonResource::FHairCommonResource(EHairStrandsAllocationType InAllocationType, const FHairResourceName& InResourceName, bool bInUseRenderGraph):
+FHairCommonResource::FHairCommonResource(EHairStrandsAllocationType InAllocationType, const FHairResourceName& InResourceName, const FName& InOwnerName, bool bInUseRenderGraph):
 bUseRenderGraph(bInUseRenderGraph),
 bIsInitialized(false),
 AllocationType(InAllocationType),
-ResourceName(InResourceName)
+ResourceName(InResourceName),
+OwnerName(InOwnerName)
 {
 }
 
@@ -610,6 +625,12 @@ void FHairCommonResource::StreamInLODData(int32 LODIndex)
 
 /////////////////////////////////////////////////////////////////////////////////////////
 
+FHairCardIndexBuffer::FHairCardIndexBuffer(const TArray<FHairCardsIndexFormat::Type>& InIndices, const FName& InOwnerName)
+	:Indices(InIndices) 
+{ 
+	SetOwnerName(InOwnerName); 
+}
+
 void FHairCardIndexBuffer::InitRHI()
 {
 	const uint32 DataSizeInBytes = FHairCardsIndexFormat::SizeInByte * Indices.Num();
@@ -619,12 +640,13 @@ void FHairCardIndexBuffer::InitRHI()
 	void* Buffer = RHILockBuffer(IndexBufferRHI, 0, DataSizeInBytes, RLM_WriteOnly);
 	FMemory::Memcpy(Buffer, Indices.GetData(), DataSizeInBytes);
 	RHIUnlockBuffer(IndexBufferRHI);
+	IndexBufferRHI->SetOwnerName(GetOwnerName());
 }
 
-FHairCardsRestResource::FHairCardsRestResource(const FHairCardsBulkData& InBulkData, const FHairResourceName& InResourceName) :
-	FHairCommonResource(EHairStrandsAllocationType::Immediate, InResourceName, false),
+FHairCardsRestResource::FHairCardsRestResource(const FHairCardsBulkData& InBulkData, const FHairResourceName& InResourceName, const FName& InOwnerName) :
+	FHairCommonResource(EHairStrandsAllocationType::Immediate, InResourceName, InOwnerName, false),
 	RestPositionBuffer(),
-	RestIndexBuffer(InBulkData.Indices),
+	RestIndexBuffer(InBulkData.Indices, InOwnerName),
 	NormalsBuffer(),
 	UVsBuffer(),
 	MaterialsBuffer(),
@@ -636,10 +658,10 @@ FHairCardsRestResource::FHairCardsRestResource(const FHairCardsBulkData& InBulkD
 void FHairCardsRestResource::InternalAllocate()
 {
 	// These resources are kept as regular (i.e., non-RDG resources) as they need to be bound at the input assembly stage by the Vertex declaraction which requires FVertexBuffer type
-	CreateBuffer<FHairCardsPositionFormat>(BulkData.Positions, RestPositionBuffer, ToHairResourceDebugName(TEXT("Hair.CardsRest_PositionBuffer"), ResourceName));
-	CreateBuffer<FHairCardsNormalFormat>(BulkData.Normals, NormalsBuffer, ToHairResourceDebugName(TEXT("Hair.CardsRest_NormalBuffer"), ResourceName));
-	CreateBuffer<FHairCardsUVFormat>(BulkData.UVs, UVsBuffer, ToHairResourceDebugName(TEXT("Hair.CardsRest_UVBuffer"), ResourceName));
-	CreateBuffer<FHairCardsMaterialFormat>(BulkData.Materials, MaterialsBuffer, ToHairResourceDebugName(TEXT("Hair.CardsRest_MaterialBuffer"), ResourceName));
+	CreateBuffer<FHairCardsPositionFormat>(BulkData.Positions, RestPositionBuffer, ToHairResourceDebugName(TEXT("Hair.CardsRest_PositionBuffer"), ResourceName), OwnerName);
+	CreateBuffer<FHairCardsNormalFormat>(BulkData.Normals, NormalsBuffer, ToHairResourceDebugName(TEXT("Hair.CardsRest_NormalBuffer"), ResourceName), OwnerName);
+	CreateBuffer<FHairCardsUVFormat>(BulkData.UVs, UVsBuffer, ToHairResourceDebugName(TEXT("Hair.CardsRest_UVBuffer"), ResourceName), OwnerName);
+	CreateBuffer<FHairCardsMaterialFormat>(BulkData.Materials, MaterialsBuffer, ToHairResourceDebugName(TEXT("Hair.CardsRest_MaterialBuffer"), ResourceName), OwnerName);
 
 	FSamplerStateRHIRef DefaultSampler = TStaticSamplerState<SF_Bilinear, AM_Clamp, AM_Clamp, AM_Clamp>::GetRHI();
 	DepthSampler = DefaultSampler;
@@ -674,8 +696,8 @@ void FHairCardsRestResource::ReleaseResource()
 }
 
 /////////////////////////////////////////////////////////////////////////////////////////
-FHairCardsProceduralResource::FHairCardsProceduralResource(const FHairCardsProceduralDatas::FRenderData& InRenderData, const FIntPoint& InAtlasResolution, const FHairCardsVoxel& InVoxel):
-	FHairCommonResource(EHairStrandsAllocationType::Immediate, FHairResourceName()),
+FHairCardsProceduralResource::FHairCardsProceduralResource(const FHairCardsProceduralDatas::FRenderData& InRenderData, const FIntPoint& InAtlasResolution, const FHairCardsVoxel& InVoxel, const FName& InOwnerName):
+	FHairCommonResource(EHairStrandsAllocationType::Immediate, FHairResourceName(), InOwnerName),
 	CardBoundCount(InRenderData.ClusterBounds.Num()),
 	AtlasResolution(InAtlasResolution),
 	AtlasRectBuffer(),
@@ -692,19 +714,19 @@ FHairCardsProceduralResource::FHairCardsProceduralResource(const FHairCardsProce
 
 void FHairCardsProceduralResource::InternalAllocate(FRDGBuilder& GraphBuilder)
 {
-	InternalCreateVertexBufferRDG<FHairCardsAtlasRectFormat>(GraphBuilder, RenderData.CardsRect, AtlasRectBuffer, ToHairResourceDebugName(TEXT("Hair.CardsProcedural_AtlasRectBuffer"), ResourceName), EHairResourceUsageType::Static);
-	InternalCreateVertexBufferRDG<FHairCardsDimensionFormat>(GraphBuilder, RenderData.CardsLengths, LengthBuffer, ToHairResourceDebugName(TEXT("Hair.CardsProcedural_LengthBuffer"), ResourceName), EHairResourceUsageType::Static);
+	InternalCreateVertexBufferRDG<FHairCardsAtlasRectFormat>(GraphBuilder, RenderData.CardsRect, AtlasRectBuffer, ToHairResourceDebugName(TEXT("Hair.CardsProcedural_AtlasRectBuffer"), ResourceName), OwnerName, EHairResourceUsageType::Static);
+	InternalCreateVertexBufferRDG<FHairCardsDimensionFormat>(GraphBuilder, RenderData.CardsLengths, LengthBuffer, ToHairResourceDebugName(TEXT("Hair.CardsProcedural_LengthBuffer"), ResourceName), OwnerName, EHairResourceUsageType::Static);
 
-	InternalCreateVertexBufferRDG<FHairCardsOffsetAndCount>(GraphBuilder, RenderData.CardItToCluster, CardItToClusterBuffer, ToHairResourceDebugName(TEXT("Hair.CardsProcedural_CardItToClusterBuffer"), ResourceName), EHairResourceUsageType::Static);
-	InternalCreateVertexBufferRDG<FHairCardsOffsetAndCount>(GraphBuilder, RenderData.ClusterIdToVertices, ClusterIdToVerticesBuffer, ToHairResourceDebugName(TEXT("Hair.CardsProcedural_ClusterIdToVerticesBuffer"), ResourceName), EHairResourceUsageType::Static);
-	InternalCreateVertexBufferRDG<FHairCardsBoundsFormat>(GraphBuilder, RenderData.ClusterBounds, ClusterBoundBuffer, ToHairResourceDebugName(TEXT("Hair.CardsProcedural_ClusterBoundBuffer"), ResourceName), EHairResourceUsageType::Static);
+	InternalCreateVertexBufferRDG<FHairCardsOffsetAndCount>(GraphBuilder, RenderData.CardItToCluster, CardItToClusterBuffer, ToHairResourceDebugName(TEXT("Hair.CardsProcedural_CardItToClusterBuffer"), ResourceName), OwnerName, EHairResourceUsageType::Static);
+	InternalCreateVertexBufferRDG<FHairCardsOffsetAndCount>(GraphBuilder, RenderData.ClusterIdToVertices, ClusterIdToVerticesBuffer, ToHairResourceDebugName(TEXT("Hair.CardsProcedural_ClusterIdToVerticesBuffer"), ResourceName), OwnerName, EHairResourceUsageType::Static);
+	InternalCreateVertexBufferRDG<FHairCardsBoundsFormat>(GraphBuilder, RenderData.ClusterBounds, ClusterBoundBuffer, ToHairResourceDebugName(TEXT("Hair.CardsProcedural_ClusterBoundBuffer"), ResourceName), OwnerName, EHairResourceUsageType::Static);
 
-	InternalCreateVertexBufferRDG<FHairCardsVoxelDensityFormat>(GraphBuilder, RenderData.VoxelDensity, CardVoxel.DensityBuffer, ToHairResourceDebugName(TEXT("Hair.CardsProcedural_VoxelDensityBuffer"), ResourceName), EHairResourceUsageType::Static);
-	InternalCreateVertexBufferRDG<FHairCardsVoxelTangentFormat>(GraphBuilder, RenderData.VoxelTangent, CardVoxel.TangentBuffer, ToHairResourceDebugName(TEXT("Hair.CardsProcedural_VoxelTangentBuffer"), ResourceName), EHairResourceUsageType::Static);
-	InternalCreateVertexBufferRDG<FHairCardsVoxelTangentFormat>(GraphBuilder, RenderData.VoxelNormal, CardVoxel.NormalBuffer, ToHairResourceDebugName(TEXT("Hair.CardsProcedural_VoxelNormalBuffer"), ResourceName), EHairResourceUsageType::Static);
+	InternalCreateVertexBufferRDG<FHairCardsVoxelDensityFormat>(GraphBuilder, RenderData.VoxelDensity, CardVoxel.DensityBuffer, ToHairResourceDebugName(TEXT("Hair.CardsProcedural_VoxelDensityBuffer"), ResourceName), OwnerName, EHairResourceUsageType::Static);
+	InternalCreateVertexBufferRDG<FHairCardsVoxelTangentFormat>(GraphBuilder, RenderData.VoxelTangent, CardVoxel.TangentBuffer, ToHairResourceDebugName(TEXT("Hair.CardsProcedural_VoxelTangentBuffer"), ResourceName), OwnerName, EHairResourceUsageType::Static);
+	InternalCreateVertexBufferRDG<FHairCardsVoxelTangentFormat>(GraphBuilder, RenderData.VoxelNormal, CardVoxel.NormalBuffer, ToHairResourceDebugName(TEXT("Hair.CardsProcedural_VoxelNormalBuffer"), ResourceName), OwnerName, EHairResourceUsageType::Static);
 
-	InternalCreateVertexBufferRDG<FHairCardsStrandsPositionFormat>(GraphBuilder, RenderData.CardsStrandsPositions, CardsStrandsPositions, ToHairResourceDebugName(TEXT("Hair.CardsProcedural_CardsStrandsPositions"), ResourceName), EHairResourceUsageType::Static);
-	InternalCreateVertexBufferRDG<FHairCardsStrandsAttributeFormat>(GraphBuilder, RenderData.CardsStrandsAttributes, CardsStrandsAttributes, ToHairResourceDebugName(TEXT("Hair.CardsProcedural_CardsStrandsAttributes"), ResourceName), EHairResourceUsageType::Static);
+	InternalCreateVertexBufferRDG<FHairCardsStrandsPositionFormat>(GraphBuilder, RenderData.CardsStrandsPositions, CardsStrandsPositions, ToHairResourceDebugName(TEXT("Hair.CardsProcedural_CardsStrandsPositions"), ResourceName), OwnerName, EHairResourceUsageType::Static);
+	InternalCreateVertexBufferRDG<FHairCardsStrandsAttributeFormat>(GraphBuilder, RenderData.CardsStrandsAttributes, CardsStrandsAttributes, ToHairResourceDebugName(TEXT("Hair.CardsProcedural_CardsStrandsAttributes"), ResourceName), OwnerName, EHairResourceUsageType::Static);
 }
 
 void FHairCardsProceduralResource::InternalRelease()
@@ -725,8 +747,8 @@ void FHairCardsProceduralResource::InternalRelease()
 
 /////////////////////////////////////////////////////////////////////////////////////////
 
-FHairCardsDeformedResource::FHairCardsDeformedResource(const FHairCardsBulkData& InBulkData, bool bInInitializedData, const FHairResourceName& InResourceName) :
-	FHairCommonResource(EHairStrandsAllocationType::Deferred, InResourceName),
+FHairCardsDeformedResource::FHairCardsDeformedResource(const FHairCardsBulkData& InBulkData, bool bInInitializedData, const FHairResourceName& InResourceName, const FName& InOwnerName) :
+	FHairCommonResource(EHairStrandsAllocationType::Deferred, InResourceName, InOwnerName),
 	BulkData(InBulkData), bInitializedData(bInInitializedData)
 {}
 
@@ -734,17 +756,17 @@ void FHairCardsDeformedResource::InternalAllocate(FRDGBuilder& GraphBuilder)
 {
 	if (bInitializedData)
 	{
-		InternalCreateVertexBufferRDG<FHairCardsPositionFormat>(GraphBuilder, BulkData.Positions, DeformedPositionBuffer[0], ToHairResourceDebugName(TEXT("Hair.CardsDeformedPosition(Current)"), ResourceName), EHairResourceUsageType::Dynamic);
-		InternalCreateVertexBufferRDG<FHairCardsPositionFormat>(GraphBuilder, BulkData.Positions, DeformedPositionBuffer[1], ToHairResourceDebugName(TEXT("Hair.CardsDeformedPosition(Previous)"), ResourceName), EHairResourceUsageType::Dynamic);
+		InternalCreateVertexBufferRDG<FHairCardsPositionFormat>(GraphBuilder, BulkData.Positions, DeformedPositionBuffer[0], ToHairResourceDebugName(TEXT("Hair.CardsDeformedPosition(Current)"), ResourceName), OwnerName, EHairResourceUsageType::Dynamic);
+		InternalCreateVertexBufferRDG<FHairCardsPositionFormat>(GraphBuilder, BulkData.Positions, DeformedPositionBuffer[1], ToHairResourceDebugName(TEXT("Hair.CardsDeformedPosition(Previous)"), ResourceName), OwnerName, EHairResourceUsageType::Dynamic);
 
-		InternalCreateVertexBufferRDG<FHairCardsNormalFormat>(GraphBuilder, BulkData.Normals, DeformedNormalBuffer, ToHairResourceDebugName(TEXT("Hair.CardsDeformedNormal"), ResourceName), EHairResourceUsageType::Dynamic);
+		InternalCreateVertexBufferRDG<FHairCardsNormalFormat>(GraphBuilder, BulkData.Normals, DeformedNormalBuffer, ToHairResourceDebugName(TEXT("Hair.CardsDeformedNormal"), ResourceName), OwnerName, EHairResourceUsageType::Dynamic);
 	}
 	else
 	{
-		InternalCreateVertexBufferRDG<FHairCardsPositionFormat>(GraphBuilder, BulkData.GetNumVertices(), DeformedPositionBuffer[0], ToHairResourceDebugName(TEXT("Hair.CardsDeformedPosition(Current)"), ResourceName), EHairResourceUsageType::Dynamic);
-		InternalCreateVertexBufferRDG<FHairCardsPositionFormat>(GraphBuilder, BulkData.GetNumVertices(), DeformedPositionBuffer[1], ToHairResourceDebugName(TEXT("Hair.CardsDeformedPosition(Previous)"), ResourceName), EHairResourceUsageType::Dynamic);
+		InternalCreateVertexBufferRDG<FHairCardsPositionFormat>(GraphBuilder, BulkData.GetNumVertices(), DeformedPositionBuffer[0], ToHairResourceDebugName(TEXT("Hair.CardsDeformedPosition(Current)"), ResourceName), OwnerName, EHairResourceUsageType::Dynamic);
+		InternalCreateVertexBufferRDG<FHairCardsPositionFormat>(GraphBuilder, BulkData.GetNumVertices(), DeformedPositionBuffer[1], ToHairResourceDebugName(TEXT("Hair.CardsDeformedPosition(Previous)"), ResourceName), OwnerName, EHairResourceUsageType::Dynamic);
 
-		InternalCreateVertexBufferRDG<FHairCardsNormalFormat>(GraphBuilder, BulkData.GetNumVertices() * FHairCardsNormalFormat::ComponentCount, DeformedNormalBuffer, ToHairResourceDebugName(TEXT("Hair.CardsDeformedNormal"), ResourceName), EHairResourceUsageType::Dynamic);
+		InternalCreateVertexBufferRDG<FHairCardsNormalFormat>(GraphBuilder, BulkData.GetNumVertices() * FHairCardsNormalFormat::ComponentCount, DeformedNormalBuffer, ToHairResourceDebugName(TEXT("Hair.CardsDeformedNormal"), ResourceName), OwnerName, EHairResourceUsageType::Dynamic);
 
 		// Manually transit to SRVs, in case of the cards are rendered not visible, but still rendered (in shadows for instance). In such a case, the cards deformation pass is not called, and thus the 
 		// buffers are never transit from UAV (clear) to SRV for rasterization. 
@@ -765,10 +787,10 @@ void FHairCardsDeformedResource::InternalRelease()
 
 /////////////////////////////////////////////////////////////////////////////////////////
 
-FHairMeshesRestResource::FHairMeshesRestResource(const FHairMeshesBulkData& InBulkData, const FHairResourceName& InResourceName) :
-	FHairCommonResource(EHairStrandsAllocationType::Immediate, InResourceName, false),
+FHairMeshesRestResource::FHairMeshesRestResource(const FHairMeshesBulkData& InBulkData, const FHairResourceName& InResourceName, const FName& InOwnerName) :
+	FHairCommonResource(EHairStrandsAllocationType::Immediate, InResourceName, InOwnerName, false),
 	RestPositionBuffer(),
-	IndexBuffer(InBulkData.Indices),
+	IndexBuffer(InBulkData.Indices, InOwnerName),
 	NormalsBuffer(),
 	UVsBuffer(),
 	BulkData(InBulkData)
@@ -780,9 +802,9 @@ FHairMeshesRestResource::FHairMeshesRestResource(const FHairMeshesBulkData& InBu
 void FHairMeshesRestResource::InternalAllocate()
 {
 	// These resources are kept as regular (i.e., non-RDG resources) as they need to be bound at the input assembly stage by the Vertex declaraction which requires FVertexBuffer type
-	CreateBuffer<FHairCardsPositionFormat>(BulkData.Positions, RestPositionBuffer, ToHairResourceDebugName(TEXT("Hair.MeshesRest_Positions"), ResourceName));
-	CreateBuffer<FHairCardsNormalFormat>(BulkData.Normals, NormalsBuffer, ToHairResourceDebugName(TEXT("Hair.MeshesRest_Normals"), ResourceName));
-	CreateBuffer<FHairCardsUVFormat>(BulkData.UVs, UVsBuffer, ToHairResourceDebugName(TEXT("Hair.MeshesRest_UVs"), ResourceName));
+	CreateBuffer<FHairCardsPositionFormat>(BulkData.Positions, RestPositionBuffer, ToHairResourceDebugName(TEXT("Hair.MeshesRest_Positions"), ResourceName), OwnerName);
+	CreateBuffer<FHairCardsNormalFormat>(BulkData.Normals, NormalsBuffer, ToHairResourceDebugName(TEXT("Hair.MeshesRest_Normals"), ResourceName), OwnerName);
+	CreateBuffer<FHairCardsUVFormat>(BulkData.UVs, UVsBuffer, ToHairResourceDebugName(TEXT("Hair.MeshesRest_UVs"), ResourceName), OwnerName);
 }
 
 void FHairMeshesRestResource::InternalRelease()
@@ -811,8 +833,8 @@ void FHairMeshesRestResource::ReleaseResource()
 
 /////////////////////////////////////////////////////////////////////////////////////////
 
-FHairMeshesDeformedResource::FHairMeshesDeformedResource(const FHairMeshesBulkData& InBulkData, bool bInInitializedData, const FHairResourceName& InResourceName) :
-	FHairCommonResource(EHairStrandsAllocationType::Deferred, InResourceName),
+FHairMeshesDeformedResource::FHairMeshesDeformedResource(const FHairMeshesBulkData& InBulkData, bool bInInitializedData, const FHairResourceName& InResourceName, const FName& InOwnerName) :
+	FHairCommonResource(EHairStrandsAllocationType::Deferred, InResourceName, InOwnerName),
 	BulkData(InBulkData), bInitializedData(bInInitializedData)
 {}
 
@@ -820,13 +842,13 @@ void FHairMeshesDeformedResource::InternalAllocate(FRDGBuilder& GraphBuilder)
 {
 	if (bInitializedData)
 	{
-		InternalCreateVertexBufferRDG<FHairCardsPositionFormat>(GraphBuilder, BulkData.Positions, DeformedPositionBuffer[0], ToHairResourceDebugName(TEXT("Hair.MeshesDeformed(Current)"), ResourceName), EHairResourceUsageType::Dynamic);
-		InternalCreateVertexBufferRDG<FHairCardsPositionFormat>(GraphBuilder, BulkData.Positions, DeformedPositionBuffer[1], ToHairResourceDebugName(TEXT("Hair.MeshesDeformed(Previous)"), ResourceName), EHairResourceUsageType::Dynamic);
+		InternalCreateVertexBufferRDG<FHairCardsPositionFormat>(GraphBuilder, BulkData.Positions, DeformedPositionBuffer[0], ToHairResourceDebugName(TEXT("Hair.MeshesDeformed(Current)"), ResourceName), OwnerName, EHairResourceUsageType::Dynamic);
+		InternalCreateVertexBufferRDG<FHairCardsPositionFormat>(GraphBuilder, BulkData.Positions, DeformedPositionBuffer[1], ToHairResourceDebugName(TEXT("Hair.MeshesDeformed(Previous)"), ResourceName), OwnerName, EHairResourceUsageType::Dynamic);
 	}
 	else
 	{
-		InternalCreateVertexBufferRDG<FHairCardsPositionFormat>(GraphBuilder, BulkData.GetNumVertices(), DeformedPositionBuffer[0], ToHairResourceDebugName(TEXT("Hair.MeshesDeformed(Current)"), ResourceName), EHairResourceUsageType::Dynamic);
-		InternalCreateVertexBufferRDG<FHairCardsPositionFormat>(GraphBuilder, BulkData.GetNumVertices(), DeformedPositionBuffer[1], ToHairResourceDebugName(TEXT("Hair.MeshesDeformed(Previous)"), ResourceName), EHairResourceUsageType::Dynamic);
+		InternalCreateVertexBufferRDG<FHairCardsPositionFormat>(GraphBuilder, BulkData.GetNumVertices(), DeformedPositionBuffer[0], ToHairResourceDebugName(TEXT("Hair.MeshesDeformed(Current)"), ResourceName), OwnerName, EHairResourceUsageType::Dynamic);
+		InternalCreateVertexBufferRDG<FHairCardsPositionFormat>(GraphBuilder, BulkData.GetNumVertices(), DeformedPositionBuffer[1], ToHairResourceDebugName(TEXT("Hair.MeshesDeformed(Previous)"), ResourceName), OwnerName, EHairResourceUsageType::Dynamic);
 	}
 }
 
@@ -838,8 +860,8 @@ void FHairMeshesDeformedResource::InternalRelease()
 
 /////////////////////////////////////////////////////////////////////////////////////////
 
-FHairStrandsRestResource::FHairStrandsRestResource(FHairStrandsBulkData& InBulkData, EHairStrandsResourcesType InCurveType, const FHairResourceName& InResourceName) :
-	FHairCommonResource(EHairStrandsAllocationType::Deferred, InResourceName),
+FHairStrandsRestResource::FHairStrandsRestResource(FHairStrandsBulkData& InBulkData, EHairStrandsResourcesType InCurveType, const FHairResourceName& InResourceName, const FName& InOwnerName) :
+	FHairCommonResource(EHairStrandsAllocationType::Deferred, InResourceName, InOwnerName),
 	PositionBuffer(), AttributeBuffer(), VertexToCurveBuffer(), BulkData(InBulkData), CurveType(InCurveType)
 {
 	// Sanity check
@@ -871,21 +893,21 @@ void FHairStrandsRestResource::InternalAllocate(FRDGBuilder& GraphBuilder)
 
 	// 1. Lock data, which force the loading data from files (on non-editor build/cooked data). These data are then uploaded to the GPU
 	// 2. A local copy is done by the buffer uploader. This copy is discarded once the uploading is done.
-	InternalCreateVertexBufferRDG_FromBulkData<FHairStrandsPositionFormat>(GraphBuilder, BulkData.Positions, PointCount, PositionBuffer, ToHairResourceDebugName(HAIRSTRANDS_RESOUCE_NAME(CurveType, Hair.StrandsRest_PositionBuffer), ResourceName), EHairResourceUsageType::Static);
-	InternalCreateByteAddressBufferRDG_FromBulkData(GraphBuilder, BulkData.Attributes, AttributeBuffer, ToHairResourceDebugName(HAIRSTRANDS_RESOUCE_NAME(CurveType, Hair.StrandsRest_AttributeBuffer), ResourceName), EHairResourceUsageType::Static);
-	InternalCreateVertexBufferRDG_FromBulkData<FHairStrandsCurveFormat>(GraphBuilder, BulkData.Curves, CurveCount, CurveBuffer, ToHairResourceDebugName(HAIRSTRANDS_RESOUCE_NAME(CurveType, Hair.StrandsRest_CurveBuffer), ResourceName), EHairResourceUsageType::Static);
+	InternalCreateVertexBufferRDG_FromBulkData<FHairStrandsPositionFormat>(GraphBuilder, BulkData.Positions, PointCount, PositionBuffer, ToHairResourceDebugName(HAIRSTRANDS_RESOUCE_NAME(CurveType, Hair.StrandsRest_PositionBuffer), ResourceName), OwnerName, EHairResourceUsageType::Static);
+	InternalCreateByteAddressBufferRDG_FromBulkData(GraphBuilder, BulkData.Attributes, AttributeBuffer, ToHairResourceDebugName(HAIRSTRANDS_RESOUCE_NAME(CurveType, Hair.StrandsRest_AttributeBuffer), ResourceName), OwnerName, EHairResourceUsageType::Static);
+	InternalCreateVertexBufferRDG_FromBulkData<FHairStrandsCurveFormat>(GraphBuilder, BulkData.Curves, CurveCount, CurveBuffer, ToHairResourceDebugName(HAIRSTRANDS_RESOUCE_NAME(CurveType, Hair.StrandsRest_CurveBuffer), ResourceName), OwnerName, EHairResourceUsageType::Static);
 	if (!!(BulkData.Flags & BulkData.DataFlags_Has16bitsCurveIndex))
 	{
-		InternalCreateVertexBufferRDG_FromBulkData<FHairStrandsVertexToCurveFormat16>(GraphBuilder, BulkData.VertexToCurve, PointCount, VertexToCurveBuffer, ToHairResourceDebugName(HAIRSTRANDS_RESOUCE_NAME(CurveType, Hair.StrandsRest_VertexToCurveBuffer), ResourceName), EHairResourceUsageType::Static);
+		InternalCreateVertexBufferRDG_FromBulkData<FHairStrandsVertexToCurveFormat16>(GraphBuilder, BulkData.VertexToCurve, PointCount, VertexToCurveBuffer, ToHairResourceDebugName(HAIRSTRANDS_RESOUCE_NAME(CurveType, Hair.StrandsRest_VertexToCurveBuffer), ResourceName), OwnerName, EHairResourceUsageType::Static);
 	}
 	else
 	{
-		InternalCreateVertexBufferRDG_FromBulkData<FHairStrandsVertexToCurveFormat32>(GraphBuilder, BulkData.VertexToCurve, PointCount, VertexToCurveBuffer, ToHairResourceDebugName(HAIRSTRANDS_RESOUCE_NAME(CurveType, Hair.StrandsRest_VertexToCurveBuffer), ResourceName), EHairResourceUsageType::Static);
+		InternalCreateVertexBufferRDG_FromBulkData<FHairStrandsVertexToCurveFormat32>(GraphBuilder, BulkData.VertexToCurve, PointCount, VertexToCurveBuffer, ToHairResourceDebugName(HAIRSTRANDS_RESOUCE_NAME(CurveType, Hair.StrandsRest_VertexToCurveBuffer), ResourceName), OwnerName, EHairResourceUsageType::Static);
 	}
 
 	TArray<FVector4f> RestOffset;
 	RestOffset.Add((FVector3f)BulkData.GetPositionOffset());// LWC_TODO: precision loss
-	InternalCreateVertexBufferRDG<FHairStrandsPositionOffsetFormat>(GraphBuilder, RestOffset, PositionOffsetBuffer, ToHairResourceDebugName(HAIRSTRANDS_RESOUCE_NAME(CurveType, Hair.StrandsRest_PositionOffsetBuffer), ResourceName), EHairResourceUsageType::Static, ERDGInitialDataFlags::None);
+	InternalCreateVertexBufferRDG<FHairStrandsPositionOffsetFormat>(GraphBuilder, RestOffset, PositionOffsetBuffer, ToHairResourceDebugName(HAIRSTRANDS_RESOUCE_NAME(CurveType, Hair.StrandsRest_PositionOffsetBuffer), ResourceName), OwnerName, EHairResourceUsageType::Static, ERDGInitialDataFlags::None);
 	GraphBuilder.UseExternalAccessMode(Register(GraphBuilder, PositionOffsetBuffer, ERDGImportedBufferFlags::CreateSRV).Buffer, ERHIAccess::SRVMask);
 }
 
@@ -903,7 +925,7 @@ FRDGExternalBuffer FHairStrandsRestResource::GetTangentBuffer(FRDGBuilder& Graph
 	// Lazy allocation and update
 	if (TangentBuffer.Buffer == nullptr)
 	{
-		InternalCreateVertexBufferRDG<FHairStrandsTangentFormat>(GraphBuilder, BulkData.PointCount * FHairStrandsTangentFormat::ComponentCount, TangentBuffer, ToHairResourceDebugName(TEXT("Hair.StrandsRest_TangentBuffer"), ResourceName), EHairResourceUsageType::Dynamic);
+		InternalCreateVertexBufferRDG<FHairStrandsTangentFormat>(GraphBuilder, BulkData.PointCount * FHairStrandsTangentFormat::ComponentCount, TangentBuffer, ToHairResourceDebugName(TEXT("Hair.StrandsRest_TangentBuffer"), ResourceName), OwnerName, EHairResourceUsageType::Dynamic);
 
 		AddHairTangentPass(
 			GraphBuilder,
@@ -930,8 +952,8 @@ void FHairStrandsRestResource::InternalRelease()
 
 /////////////////////////////////////////////////////////////////////////////////////////
 
-FHairStrandsDeformedResource::FHairStrandsDeformedResource(FHairStrandsBulkData& InBulkData, EHairStrandsResourcesType InCurveType, const FHairResourceName& InResourceName) :
-	FHairCommonResource(EHairStrandsAllocationType::Deferred, InResourceName),
+FHairStrandsDeformedResource::FHairStrandsDeformedResource(FHairStrandsBulkData& InBulkData, EHairStrandsResourcesType InCurveType, const FHairResourceName& InResourceName, const FName& InOwnerName) :
+	FHairCommonResource(EHairStrandsAllocationType::Deferred, InResourceName, InOwnerName),
 	BulkData(InBulkData), CurveType(InCurveType)
 {
 	GetPositionOffset(FHairStrandsDeformedResource::EFrameType::Current)  = BulkData.GetPositionOffset();
@@ -942,14 +964,14 @@ void FHairStrandsDeformedResource::InternalAllocate(FRDGBuilder& GraphBuilder)
 {
 	const uint32 PointCount = BulkData.PointCount;
 
-	InternalCreateVertexBufferRDG<FHairStrandsPositionFormat>(GraphBuilder, BulkData.PointCount, DeformedPositionBuffer[0], ToHairResourceDebugName(HAIRSTRANDS_RESOUCE_NAME(CurveType, Hair.StrandsDeformed_DeformedPositionBuffer0), ResourceName), EHairResourceUsageType::Dynamic);
-	InternalCreateVertexBufferRDG<FHairStrandsPositionFormat>(GraphBuilder, BulkData.PointCount, DeformedPositionBuffer[1], ToHairResourceDebugName(HAIRSTRANDS_RESOUCE_NAME(CurveType, Hair.StrandsDeformed_DeformedPositionBuffer1), ResourceName), EHairResourceUsageType::Dynamic);
-	InternalCreateVertexBufferRDG<FHairStrandsTangentFormat>(GraphBuilder, BulkData.PointCount * FHairStrandsTangentFormat::ComponentCount, TangentBuffer, ToHairResourceDebugName(HAIRSTRANDS_RESOUCE_NAME(CurveType, Hair.StrandsDeformed_TangentBuffer), ResourceName), EHairResourceUsageType::Dynamic);
+	InternalCreateVertexBufferRDG<FHairStrandsPositionFormat>(GraphBuilder, BulkData.PointCount, DeformedPositionBuffer[0], ToHairResourceDebugName(HAIRSTRANDS_RESOUCE_NAME(CurveType, Hair.StrandsDeformed_DeformedPositionBuffer0), ResourceName), OwnerName, EHairResourceUsageType::Dynamic);
+	InternalCreateVertexBufferRDG<FHairStrandsPositionFormat>(GraphBuilder, BulkData.PointCount, DeformedPositionBuffer[1], ToHairResourceDebugName(HAIRSTRANDS_RESOUCE_NAME(CurveType, Hair.StrandsDeformed_DeformedPositionBuffer1), ResourceName), OwnerName, EHairResourceUsageType::Dynamic);
+	InternalCreateVertexBufferRDG<FHairStrandsTangentFormat>(GraphBuilder, BulkData.PointCount * FHairStrandsTangentFormat::ComponentCount, TangentBuffer, ToHairResourceDebugName(HAIRSTRANDS_RESOUCE_NAME(CurveType, Hair.StrandsDeformed_TangentBuffer), ResourceName), OwnerName, EHairResourceUsageType::Dynamic);
 
 	TArray<FVector4f> DefaultOffsets;
 	DefaultOffsets.Add((FVector3f)BulkData.GetPositionOffset()); // LWC_TODO: precision loss
-	InternalCreateVertexBufferRDG<FHairStrandsPositionOffsetFormat>(GraphBuilder, DefaultOffsets, DeformedOffsetBuffer[0], ToHairResourceDebugName(HAIRSTRANDS_RESOUCE_NAME(CurveType, Hair.StrandsDeformed_DeformedOffsetBuffer0), ResourceName), EHairResourceUsageType::Dynamic, ERDGInitialDataFlags::None);
-	InternalCreateVertexBufferRDG<FHairStrandsPositionOffsetFormat>(GraphBuilder, DefaultOffsets, DeformedOffsetBuffer[1], ToHairResourceDebugName(HAIRSTRANDS_RESOUCE_NAME(CurveType, Hair.StrandsDeformed_DeformedOffsetBuffer1), ResourceName), EHairResourceUsageType::Dynamic, ERDGInitialDataFlags::None);
+	InternalCreateVertexBufferRDG<FHairStrandsPositionOffsetFormat>(GraphBuilder, DefaultOffsets, DeformedOffsetBuffer[0], ToHairResourceDebugName(HAIRSTRANDS_RESOUCE_NAME(CurveType, Hair.StrandsDeformed_DeformedOffsetBuffer0), ResourceName), OwnerName, EHairResourceUsageType::Dynamic, ERDGInitialDataFlags::None);
+	InternalCreateVertexBufferRDG<FHairStrandsPositionOffsetFormat>(GraphBuilder, DefaultOffsets, DeformedOffsetBuffer[1], ToHairResourceDebugName(HAIRSTRANDS_RESOUCE_NAME(CurveType, Hair.StrandsDeformed_DeformedOffsetBuffer1), ResourceName), OwnerName, EHairResourceUsageType::Dynamic, ERDGInitialDataFlags::None);
 
 	// Note: DeformerBuffer is optionally/lazily allocated by a mesh-deformer graph
 }
@@ -959,7 +981,7 @@ FRDGExternalBuffer& FHairStrandsDeformedResource::GetDeformerBuffer(FRDGBuilder&
 	// Lazy allocation and update
 	if (DeformerBuffer.Buffer == nullptr)
 	{
-		InternalCreateVertexBufferRDG<FHairStrandsPositionFormat>(GraphBuilder, BulkData.PointCount, DeformerBuffer, ToHairResourceDebugName(HAIRSTRANDS_RESOUCE_NAME(CurveType, Hair.StrandsDeformed_DeformerBuffer), ResourceName), EHairResourceUsageType::Dynamic);
+		InternalCreateVertexBufferRDG<FHairStrandsPositionFormat>(GraphBuilder, BulkData.PointCount, DeformerBuffer, ToHairResourceDebugName(HAIRSTRANDS_RESOUCE_NAME(CurveType, Hair.StrandsDeformed_DeformerBuffer), ResourceName), OwnerName, EHairResourceUsageType::Dynamic);
 	}
 	return DeformerBuffer;
 }
@@ -1090,8 +1112,8 @@ void FHairStrandsClusterCullingBulkData::Validate(bool bIsSaving)
 
 /////////////////////////////////////////////////////////////////////////////////////////
 // Cluster culling resources
-FHairStrandsClusterCullingResource::FHairStrandsClusterCullingResource(FHairStrandsClusterCullingBulkData& InBulkData, const FHairResourceName& InResourceName):
-	FHairCommonResource(EHairStrandsAllocationType::Deferred, InResourceName),
+FHairStrandsClusterCullingResource::FHairStrandsClusterCullingResource(FHairStrandsClusterCullingBulkData& InBulkData, const FHairResourceName& InResourceName, const FName& InOwnerName):
+	FHairCommonResource(EHairStrandsAllocationType::Deferred, InResourceName, InOwnerName),
 	BulkData(InBulkData)
 {
 
@@ -1121,11 +1143,11 @@ void FHairStrandsClusterCullingResource::InternalAllocate(FRDGBuilder& GraphBuil
 		BulkData.Validate(false);
 	}
 
-	InternalCreateStructuredBufferRDG_FromBulkData<FHairClusterInfoFormat>(GraphBuilder, BulkData.PackedClusterInfos, BulkData.ClusterCount, ClusterInfoBuffer, ToHairResourceDebugName(TEXT("Hair.StrandsClusterCulling_ClusterInfoBuffer"), ResourceName), EHairResourceUsageType::Static);
-	InternalCreateStructuredBufferRDG_FromBulkData<FHairClusterLODInfoFormat>(GraphBuilder, BulkData.ClusterLODInfos, BulkData.ClusterLODCount, ClusterLODInfoBuffer, ToHairResourceDebugName(TEXT("Hair.StrandsClusterCulling_ClusterLODInfoBuffer"), ResourceName), EHairResourceUsageType::Static);
+	InternalCreateStructuredBufferRDG_FromBulkData<FHairClusterInfoFormat>(GraphBuilder, BulkData.PackedClusterInfos, BulkData.ClusterCount, ClusterInfoBuffer, ToHairResourceDebugName(TEXT("Hair.StrandsClusterCulling_ClusterInfoBuffer"), ResourceName), OwnerName, EHairResourceUsageType::Static);
+	InternalCreateStructuredBufferRDG_FromBulkData<FHairClusterLODInfoFormat>(GraphBuilder, BulkData.ClusterLODInfos, BulkData.ClusterLODCount, ClusterLODInfoBuffer, ToHairResourceDebugName(TEXT("Hair.StrandsClusterCulling_ClusterLODInfoBuffer"), ResourceName), OwnerName, EHairResourceUsageType::Static);
 
-	InternalCreateVertexBufferRDG_FromBulkData<FHairClusterIndexFormat>(GraphBuilder, BulkData.VertexToClusterIds, BulkData.VertexCount, VertexToClusterIdBuffer, ToHairResourceDebugName(TEXT("Hair.StrandsClusterCulling_VertexToClusterIds"), ResourceName), EHairResourceUsageType::Static);
-	InternalCreateVertexBufferRDG_FromBulkData<FHairClusterIndexFormat>(GraphBuilder, BulkData.ClusterVertexIds, BulkData.VertexLODCount, ClusterVertexIdBuffer, ToHairResourceDebugName(TEXT("Hair.StrandsClusterCulling_ClusterVertexIds"), ResourceName), EHairResourceUsageType::Static);
+	InternalCreateVertexBufferRDG_FromBulkData<FHairClusterIndexFormat>(GraphBuilder, BulkData.VertexToClusterIds, BulkData.VertexCount, VertexToClusterIdBuffer, ToHairResourceDebugName(TEXT("Hair.StrandsClusterCulling_VertexToClusterIds"), ResourceName), OwnerName, EHairResourceUsageType::Static);
+	InternalCreateVertexBufferRDG_FromBulkData<FHairClusterIndexFormat>(GraphBuilder, BulkData.ClusterVertexIds, BulkData.VertexLODCount, ClusterVertexIdBuffer, ToHairResourceDebugName(TEXT("Hair.StrandsClusterCulling_ClusterVertexIds"), ResourceName), OwnerName, EHairResourceUsageType::Static);
 }
 
 void FHairStrandsClusterCullingResource::InternalRelease()
@@ -1139,8 +1161,8 @@ void FHairStrandsClusterCullingResource::InternalRelease()
 
 /////////////////////////////////////////////////////////////////////////////////////////
 
-FHairStrandsRestRootResource::FHairStrandsRestRootResource(FHairStrandsRootBulkData& InBulkData, EHairStrandsResourcesType InCurveType, const FHairResourceName& InResourceName) :
-	FHairCommonResource(EHairStrandsAllocationType::Deferred, InResourceName),
+FHairStrandsRestRootResource::FHairStrandsRestRootResource(FHairStrandsRootBulkData& InBulkData, EHairStrandsResourcesType InCurveType, const FHairResourceName& InResourceName, const FName& InOwnerName) :
+	FHairCommonResource(EHairStrandsAllocationType::Deferred, InResourceName, InOwnerName),
 	BulkData(InBulkData), CurveType(InCurveType)
 {
 	PopulateFromRootData();
@@ -1243,26 +1265,26 @@ void FHairStrandsRestRootResource::InternalAllocateLOD(FRDGBuilder& GraphBuilder
 		{
 			GPUData.Status = FLOD::EStatus::Completed;
 			
-			InternalCreateVertexBufferRDG_FromBulkData<FHairStrandsRootBarycentricFormat>(GraphBuilder, CPUData.RootBarycentricBuffer, BulkData.RootCount, GPUData.RootBarycentricBuffer, ToHairResourceDebugName(HAIRSTRANDS_RESOUCE_NAME(CurveType, Hair.StrandsRestRoot_RootTriangleBarycentricBuffer), ResourceName), EHairResourceUsageType::Static);
-			InternalCreateVertexBufferRDG_FromBulkData<FHairStrandsRootToUniqueTriangleIndexFormat>(GraphBuilder, CPUData.RootToUniqueTriangleIndexBuffer, BulkData.RootCount, GPUData.RootToUniqueTriangleIndexBuffer, ToHairResourceDebugName(HAIRSTRANDS_RESOUCE_NAME(CurveType, Hair.StrandsRestRoot_RootToUniqueTriangleIndexBuffer), ResourceName), EHairResourceUsageType::Static);
-			InternalCreateVertexBufferRDG_FromBulkData<FHairStrandsUniqueTriangleIndexFormat>(GraphBuilder, CPUData.UniqueTriangleIndexBuffer, CPUData.UniqueTriangleCount, GPUData.UniqueTriangleIndexBuffer, ToHairResourceDebugName(HAIRSTRANDS_RESOUCE_NAME(CurveType, Hair.StrandsRestRoot_UniqueTriangleIndexBuffer), ResourceName), EHairResourceUsageType::Static);
+			InternalCreateVertexBufferRDG_FromBulkData<FHairStrandsRootBarycentricFormat>(GraphBuilder, CPUData.RootBarycentricBuffer, BulkData.RootCount, GPUData.RootBarycentricBuffer, ToHairResourceDebugName(HAIRSTRANDS_RESOUCE_NAME(CurveType, Hair.StrandsRestRoot_RootTriangleBarycentricBuffer), ResourceName), OwnerName, EHairResourceUsageType::Static);
+			InternalCreateVertexBufferRDG_FromBulkData<FHairStrandsRootToUniqueTriangleIndexFormat>(GraphBuilder, CPUData.RootToUniqueTriangleIndexBuffer, BulkData.RootCount, GPUData.RootToUniqueTriangleIndexBuffer, ToHairResourceDebugName(HAIRSTRANDS_RESOUCE_NAME(CurveType, Hair.StrandsRestRoot_RootToUniqueTriangleIndexBuffer), ResourceName), OwnerName, EHairResourceUsageType::Static);
+			InternalCreateVertexBufferRDG_FromBulkData<FHairStrandsUniqueTriangleIndexFormat>(GraphBuilder, CPUData.UniqueTriangleIndexBuffer, CPUData.UniqueTriangleCount, GPUData.UniqueTriangleIndexBuffer, ToHairResourceDebugName(HAIRSTRANDS_RESOUCE_NAME(CurveType, Hair.StrandsRestRoot_UniqueTriangleIndexBuffer), ResourceName), OwnerName, EHairResourceUsageType::Static);
 
-			InternalCreateVertexBufferRDG_FromBulkData<FHairStrandsMeshTrianglePositionFormat>(GraphBuilder, CPUData.RestUniqueTrianglePosition0Buffer, CPUData.UniqueTriangleCount, GPUData.RestUniqueTrianglePosition0Buffer, ToHairResourceDebugName(HAIRSTRANDS_RESOUCE_NAME(CurveType, Hair.StrandsRestRoot_RestUniqueTrianglePosition0Buffer), ResourceName), EHairResourceUsageType::Static);
-			InternalCreateVertexBufferRDG_FromBulkData<FHairStrandsMeshTrianglePositionFormat>(GraphBuilder, CPUData.RestUniqueTrianglePosition1Buffer, CPUData.UniqueTriangleCount, GPUData.RestUniqueTrianglePosition1Buffer, ToHairResourceDebugName(HAIRSTRANDS_RESOUCE_NAME(CurveType, Hair.StrandsRestRoot_RestUniqueTrianglePosition1Buffer), ResourceName), EHairResourceUsageType::Static);
-			InternalCreateVertexBufferRDG_FromBulkData<FHairStrandsMeshTrianglePositionFormat>(GraphBuilder, CPUData.RestUniqueTrianglePosition2Buffer, CPUData.UniqueTriangleCount, GPUData.RestUniqueTrianglePosition2Buffer, ToHairResourceDebugName(HAIRSTRANDS_RESOUCE_NAME(CurveType, Hair.StrandsRestRoot_RestUniqueTrianglePosition2Buffer), ResourceName), EHairResourceUsageType::Static);
+			InternalCreateVertexBufferRDG_FromBulkData<FHairStrandsMeshTrianglePositionFormat>(GraphBuilder, CPUData.RestUniqueTrianglePosition0Buffer, CPUData.UniqueTriangleCount, GPUData.RestUniqueTrianglePosition0Buffer, ToHairResourceDebugName(HAIRSTRANDS_RESOUCE_NAME(CurveType, Hair.StrandsRestRoot_RestUniqueTrianglePosition0Buffer), ResourceName), OwnerName, EHairResourceUsageType::Static);
+			InternalCreateVertexBufferRDG_FromBulkData<FHairStrandsMeshTrianglePositionFormat>(GraphBuilder, CPUData.RestUniqueTrianglePosition1Buffer, CPUData.UniqueTriangleCount, GPUData.RestUniqueTrianglePosition1Buffer, ToHairResourceDebugName(HAIRSTRANDS_RESOUCE_NAME(CurveType, Hair.StrandsRestRoot_RestUniqueTrianglePosition1Buffer), ResourceName), OwnerName, EHairResourceUsageType::Static);
+			InternalCreateVertexBufferRDG_FromBulkData<FHairStrandsMeshTrianglePositionFormat>(GraphBuilder, CPUData.RestUniqueTrianglePosition2Buffer, CPUData.UniqueTriangleCount, GPUData.RestUniqueTrianglePosition2Buffer, ToHairResourceDebugName(HAIRSTRANDS_RESOUCE_NAME(CurveType, Hair.StrandsRestRoot_RestUniqueTrianglePosition2Buffer), ResourceName), OwnerName, EHairResourceUsageType::Static);
 		}
 		else
 		{
 			GPUData.Status = FLOD::EStatus::Initialized;
 			
-			InternalCreateVertexBufferRDG<FHairStrandsRootBarycentricFormat>(GraphBuilder, BulkData.RootCount, GPUData.RootBarycentricBuffer, ToHairResourceDebugName(HAIRSTRANDS_RESOUCE_NAME(CurveType, Hair.StrandsRestRoot_RootBarycentricBuffer), ResourceName), EHairResourceUsageType::Dynamic);
-			InternalCreateVertexBufferRDG<FHairStrandsRootToUniqueTriangleIndexFormat>(GraphBuilder, BulkData.RootCount, GPUData.RootToUniqueTriangleIndexBuffer, ToHairResourceDebugName(HAIRSTRANDS_RESOUCE_NAME(CurveType, Hair.StrandsRestRoot_RootToUniqueTriangleIndexBuffer), ResourceName), EHairResourceUsageType::Dynamic);
-			InternalCreateVertexBufferRDG<FHairStrandsUniqueTriangleIndexFormat>(GraphBuilder, CPUData.UniqueTriangleCount, GPUData.UniqueTriangleIndexBuffer, ToHairResourceDebugName(HAIRSTRANDS_RESOUCE_NAME(CurveType, Hair.StrandsRestRoot_UniqueTriangleIndexBuffer), ResourceName), EHairResourceUsageType::Dynamic);
+			InternalCreateVertexBufferRDG<FHairStrandsRootBarycentricFormat>(GraphBuilder, BulkData.RootCount, GPUData.RootBarycentricBuffer, ToHairResourceDebugName(HAIRSTRANDS_RESOUCE_NAME(CurveType, Hair.StrandsRestRoot_RootBarycentricBuffer), ResourceName), OwnerName, EHairResourceUsageType::Dynamic);
+			InternalCreateVertexBufferRDG<FHairStrandsRootToUniqueTriangleIndexFormat>(GraphBuilder, BulkData.RootCount, GPUData.RootToUniqueTriangleIndexBuffer, ToHairResourceDebugName(HAIRSTRANDS_RESOUCE_NAME(CurveType, Hair.StrandsRestRoot_RootToUniqueTriangleIndexBuffer), ResourceName), OwnerName, EHairResourceUsageType::Dynamic);
+			InternalCreateVertexBufferRDG<FHairStrandsUniqueTriangleIndexFormat>(GraphBuilder, CPUData.UniqueTriangleCount, GPUData.UniqueTriangleIndexBuffer, ToHairResourceDebugName(HAIRSTRANDS_RESOUCE_NAME(CurveType, Hair.StrandsRestRoot_UniqueTriangleIndexBuffer), ResourceName), OwnerName, EHairResourceUsageType::Dynamic);
 			
 			// Create buffers. Initialization will be done by render passes
-			InternalCreateVertexBufferRDG<FHairStrandsMeshTrianglePositionFormat>(GraphBuilder, CPUData.UniqueTriangleCount, GPUData.RestUniqueTrianglePosition0Buffer, ToHairResourceDebugName(HAIRSTRANDS_RESOUCE_NAME(CurveType, Hair.StrandsRestRoot_RestUniqueTrianglePosition0Buffer), ResourceName), EHairResourceUsageType::Dynamic);
-			InternalCreateVertexBufferRDG<FHairStrandsMeshTrianglePositionFormat>(GraphBuilder, CPUData.UniqueTriangleCount, GPUData.RestUniqueTrianglePosition1Buffer, ToHairResourceDebugName(HAIRSTRANDS_RESOUCE_NAME(CurveType, Hair.StrandsRestRoot_RestUniqueTrianglePosition1Buffer), ResourceName), EHairResourceUsageType::Dynamic);
-			InternalCreateVertexBufferRDG<FHairStrandsMeshTrianglePositionFormat>(GraphBuilder, CPUData.UniqueTriangleCount, GPUData.RestUniqueTrianglePosition2Buffer, ToHairResourceDebugName(HAIRSTRANDS_RESOUCE_NAME(CurveType, Hair.StrandsRestRoot_RestUniqueTrianglePosition2Buffer), ResourceName), EHairResourceUsageType::Dynamic);
+			InternalCreateVertexBufferRDG<FHairStrandsMeshTrianglePositionFormat>(GraphBuilder, CPUData.UniqueTriangleCount, GPUData.RestUniqueTrianglePosition0Buffer, ToHairResourceDebugName(HAIRSTRANDS_RESOUCE_NAME(CurveType, Hair.StrandsRestRoot_RestUniqueTrianglePosition0Buffer), ResourceName), OwnerName, EHairResourceUsageType::Dynamic);
+			InternalCreateVertexBufferRDG<FHairStrandsMeshTrianglePositionFormat>(GraphBuilder, CPUData.UniqueTriangleCount, GPUData.RestUniqueTrianglePosition1Buffer, ToHairResourceDebugName(HAIRSTRANDS_RESOUCE_NAME(CurveType, Hair.StrandsRestRoot_RestUniqueTrianglePosition1Buffer), ResourceName), OwnerName, EHairResourceUsageType::Dynamic);
+			InternalCreateVertexBufferRDG<FHairStrandsMeshTrianglePositionFormat>(GraphBuilder, CPUData.UniqueTriangleCount, GPUData.RestUniqueTrianglePosition2Buffer, ToHairResourceDebugName(HAIRSTRANDS_RESOUCE_NAME(CurveType, Hair.StrandsRestRoot_RestUniqueTrianglePosition2Buffer), ResourceName), OwnerName, EHairResourceUsageType::Dynamic);
 		}
 
 		GPUData.SampleCount = CPUData.SampleCount;
@@ -1270,16 +1292,16 @@ void FHairStrandsRestRootResource::InternalAllocateLOD(FRDGBuilder& GraphBuilder
 		if (bHasValidCPUWeights)
 		{
 			const uint32 InteroplationWeightCount = CPUData.MeshInterpolationWeightsBuffer.GetBulkDataSize() / sizeof(FHairStrandsWeightFormat::Type);
-			InternalCreateVertexBufferRDG_FromBulkData<FHairStrandsWeightFormat>(GraphBuilder, CPUData.MeshInterpolationWeightsBuffer, InteroplationWeightCount, GPUData.MeshInterpolationWeightsBuffer, ToHairResourceDebugName(HAIRSTRANDS_RESOUCE_NAME(CurveType, Hair.StrandsRestRoot_MeshInterpolationWeightsBuffer), ResourceName), EHairResourceUsageType::Static);
-			InternalCreateVertexBufferRDG_FromBulkData<FHairStrandsIndexFormat>(GraphBuilder, CPUData.MeshSampleIndicesBuffer, CPUData.SampleCount, GPUData.MeshSampleIndicesBuffer, ToHairResourceDebugName(HAIRSTRANDS_RESOUCE_NAME(CurveType, Hair.StrandsRestRoot_MeshSampleIndicesBuffer), ResourceName), EHairResourceUsageType::Static);
-			InternalCreateVertexBufferRDG_FromBulkData<FHairStrandsMeshTrianglePositionFormat>(GraphBuilder, CPUData.RestSamplePositionsBuffer, CPUData.SampleCount, GPUData.RestSamplePositionsBuffer, ToHairResourceDebugName(HAIRSTRANDS_RESOUCE_NAME(CurveType, Hair.StrandsRestRoot_RestSamplePositionsBuffer), ResourceName), EHairResourceUsageType::Static);
+			InternalCreateVertexBufferRDG_FromBulkData<FHairStrandsWeightFormat>(GraphBuilder, CPUData.MeshInterpolationWeightsBuffer, InteroplationWeightCount, GPUData.MeshInterpolationWeightsBuffer, ToHairResourceDebugName(HAIRSTRANDS_RESOUCE_NAME(CurveType, Hair.StrandsRestRoot_MeshInterpolationWeightsBuffer), ResourceName), OwnerName, EHairResourceUsageType::Static);
+			InternalCreateVertexBufferRDG_FromBulkData<FHairStrandsIndexFormat>(GraphBuilder, CPUData.MeshSampleIndicesBuffer, CPUData.SampleCount, GPUData.MeshSampleIndicesBuffer, ToHairResourceDebugName(HAIRSTRANDS_RESOUCE_NAME(CurveType, Hair.StrandsRestRoot_MeshSampleIndicesBuffer), ResourceName), OwnerName, EHairResourceUsageType::Static);
+			InternalCreateVertexBufferRDG_FromBulkData<FHairStrandsMeshTrianglePositionFormat>(GraphBuilder, CPUData.RestSamplePositionsBuffer, CPUData.SampleCount, GPUData.RestSamplePositionsBuffer, ToHairResourceDebugName(HAIRSTRANDS_RESOUCE_NAME(CurveType, Hair.StrandsRestRoot_RestSamplePositionsBuffer), ResourceName), OwnerName, EHairResourceUsageType::Static);
 		}
 		else
 		{
 			// TODO: do not allocate these resources, since they won't be used
-			InternalCreateVertexBufferRDG<FHairStrandsWeightFormat>(GraphBuilder, (CPUData.SampleCount+4) * (CPUData.SampleCount+4), GPUData.MeshInterpolationWeightsBuffer, ToHairResourceDebugName(HAIRSTRANDS_RESOUCE_NAME(CurveType, Hair.StrandsRestRoot_MeshInterpolationWeightsBuffer), ResourceName), EHairResourceUsageType::Dynamic);
-			InternalCreateVertexBufferRDG<FHairStrandsIndexFormat>(GraphBuilder, CPUData.SampleCount, GPUData.MeshSampleIndicesBuffer, ToHairResourceDebugName(HAIRSTRANDS_RESOUCE_NAME(CurveType, Hair.StrandsRestRoot_MeshSampleIndicesBuffer), ResourceName), EHairResourceUsageType::Dynamic);
-			InternalCreateVertexBufferRDG<FHairStrandsMeshTrianglePositionFormat>(GraphBuilder, CPUData.SampleCount, GPUData.RestSamplePositionsBuffer, ToHairResourceDebugName(HAIRSTRANDS_RESOUCE_NAME(CurveType, Hair.StrandsRestRoot_RestSamplePositionsBuffer), ResourceName), EHairResourceUsageType::Dynamic);
+			InternalCreateVertexBufferRDG<FHairStrandsWeightFormat>(GraphBuilder, (CPUData.SampleCount+4) * (CPUData.SampleCount+4), GPUData.MeshInterpolationWeightsBuffer, ToHairResourceDebugName(HAIRSTRANDS_RESOUCE_NAME(CurveType, Hair.StrandsRestRoot_MeshInterpolationWeightsBuffer), ResourceName), OwnerName, EHairResourceUsageType::Dynamic);
+			InternalCreateVertexBufferRDG<FHairStrandsIndexFormat>(GraphBuilder, CPUData.SampleCount, GPUData.MeshSampleIndicesBuffer, ToHairResourceDebugName(HAIRSTRANDS_RESOUCE_NAME(CurveType, Hair.StrandsRestRoot_MeshSampleIndicesBuffer), ResourceName), OwnerName, EHairResourceUsageType::Dynamic);
+			InternalCreateVertexBufferRDG<FHairStrandsMeshTrianglePositionFormat>(GraphBuilder, CPUData.SampleCount, GPUData.RestSamplePositionsBuffer, ToHairResourceDebugName(HAIRSTRANDS_RESOUCE_NAME(CurveType, Hair.StrandsRestRoot_RestSamplePositionsBuffer), ResourceName), OwnerName, EHairResourceUsageType::Dynamic);
 		}
 	}
 }
@@ -1306,15 +1328,15 @@ void FHairStrandsRestRootResource::InternalRelease()
 
 /////////////////////////////////////////////////////////////////////////////////////////
 
-FHairStrandsDeformedRootResource::FHairStrandsDeformedRootResource(EHairStrandsResourcesType InCurveType, const FHairResourceName& InResourceName):
-	FHairCommonResource(EHairStrandsAllocationType::Deferred, InResourceName),
+FHairStrandsDeformedRootResource::FHairStrandsDeformedRootResource(EHairStrandsResourcesType InCurveType, const FHairResourceName& InResourceName, const FName& InOwnerName):
+	FHairCommonResource(EHairStrandsAllocationType::Deferred, InResourceName, InOwnerName),
 	CurveType(InCurveType)
 {
 
 }
 
-FHairStrandsDeformedRootResource::FHairStrandsDeformedRootResource(const FHairStrandsRestRootResource* InRestResources, EHairStrandsResourcesType InCurveType, const FHairResourceName& InResourceName):
-	FHairCommonResource(EHairStrandsAllocationType::Deferred, InResourceName),
+FHairStrandsDeformedRootResource::FHairStrandsDeformedRootResource(const FHairStrandsRestRootResource* InRestResources, EHairStrandsResourcesType InCurveType, const FHairResourceName& InResourceName, const FName& InOwnerName):
+	FHairCommonResource(EHairStrandsAllocationType::Deferred, InResourceName, InOwnerName),
 	CurveType(InCurveType)
 {
 	check(InRestResources);
@@ -1341,27 +1363,27 @@ void FHairStrandsDeformedRootResource::InternalAllocateLOD(FRDGBuilder& GraphBui
 			LOD.Status = FLOD::EStatus::Initialized;
 			if (LOD.SampleCount > 0)
 			{
-				InternalCreateVertexBufferRDG<FHairStrandsMeshTrianglePositionFormat>(GraphBuilder, LOD.SampleCount, LOD.DeformedSamplePositionsBuffer[0], ToHairResourceDebugName(HAIRSTRANDS_RESOUCE_NAME(CurveType, Hair.StrandsRootDeformed_DeformedSamplePositionsBuffer0), ResourceName), EHairResourceUsageType::Dynamic);
-				InternalCreateVertexBufferRDG<FHairStrandsMeshTrianglePositionFormat>(GraphBuilder, LOD.SampleCount + 4, LOD.MeshSampleWeightsBuffer[0], ToHairResourceDebugName(HAIRSTRANDS_RESOUCE_NAME(CurveType, Hair.StrandsRootDeformed_MeshSampleWeightsBuffer0), ResourceName), EHairResourceUsageType::Dynamic);
+				InternalCreateVertexBufferRDG<FHairStrandsMeshTrianglePositionFormat>(GraphBuilder, LOD.SampleCount, LOD.DeformedSamplePositionsBuffer[0], ToHairResourceDebugName(HAIRSTRANDS_RESOUCE_NAME(CurveType, Hair.StrandsRootDeformed_DeformedSamplePositionsBuffer0), ResourceName), OwnerName, EHairResourceUsageType::Dynamic);
+				InternalCreateVertexBufferRDG<FHairStrandsMeshTrianglePositionFormat>(GraphBuilder, LOD.SampleCount + 4, LOD.MeshSampleWeightsBuffer[0], ToHairResourceDebugName(HAIRSTRANDS_RESOUCE_NAME(CurveType, Hair.StrandsRootDeformed_MeshSampleWeightsBuffer0), ResourceName), OwnerName, EHairResourceUsageType::Dynamic);
 
 				// Double buffering is disabled by default unless the read-only cvar r.HairStrands.ContinuousDecimationReordering is set
 				if (IsHairStrandContinuousDecimationReorderingEnabled())
 				{
-					InternalCreateVertexBufferRDG<FHairStrandsMeshTrianglePositionFormat>(GraphBuilder, LOD.SampleCount, LOD.DeformedSamplePositionsBuffer[1], ToHairResourceDebugName(HAIRSTRANDS_RESOUCE_NAME(CurveType, Hair.StrandsRootDeformed_DeformedSamplePositionsBuffer1), ResourceName), EHairResourceUsageType::Dynamic);
-					InternalCreateVertexBufferRDG<FHairStrandsMeshTrianglePositionFormat>(GraphBuilder, LOD.SampleCount + 4, LOD.MeshSampleWeightsBuffer[1], ToHairResourceDebugName(HAIRSTRANDS_RESOUCE_NAME(CurveType, Hair.StrandsRootDeformed_MeshSampleWeightsBuffer1), ResourceName), EHairResourceUsageType::Dynamic);
+					InternalCreateVertexBufferRDG<FHairStrandsMeshTrianglePositionFormat>(GraphBuilder, LOD.SampleCount, LOD.DeformedSamplePositionsBuffer[1], ToHairResourceDebugName(HAIRSTRANDS_RESOUCE_NAME(CurveType, Hair.StrandsRootDeformed_DeformedSamplePositionsBuffer1), ResourceName), OwnerName, EHairResourceUsageType::Dynamic);
+					InternalCreateVertexBufferRDG<FHairStrandsMeshTrianglePositionFormat>(GraphBuilder, LOD.SampleCount + 4, LOD.MeshSampleWeightsBuffer[1], ToHairResourceDebugName(HAIRSTRANDS_RESOUCE_NAME(CurveType, Hair.StrandsRootDeformed_MeshSampleWeightsBuffer1), ResourceName), OwnerName, EHairResourceUsageType::Dynamic);
 				}
 			}
 
-			InternalCreateVertexBufferRDG<FHairStrandsMeshTrianglePositionFormat>(GraphBuilder, RootCount, LOD.DeformedUniqueTrianglePosition0Buffer[0], ToHairResourceDebugName(HAIRSTRANDS_RESOUCE_NAME(CurveType, Hair.StrandsRootDeformed_DeformedUniqueTrianglePosition0Buffer0), ResourceName), EHairResourceUsageType::Dynamic);
-			InternalCreateVertexBufferRDG<FHairStrandsMeshTrianglePositionFormat>(GraphBuilder, RootCount, LOD.DeformedUniqueTrianglePosition1Buffer[0], ToHairResourceDebugName(HAIRSTRANDS_RESOUCE_NAME(CurveType, Hair.StrandsRootDeformed_DeformedUniqueTrianglePosition1Buffer0), ResourceName), EHairResourceUsageType::Dynamic);
-			InternalCreateVertexBufferRDG<FHairStrandsMeshTrianglePositionFormat>(GraphBuilder, RootCount, LOD.DeformedUniqueTrianglePosition2Buffer[0], ToHairResourceDebugName(HAIRSTRANDS_RESOUCE_NAME(CurveType, Hair.StrandsRootDeformed_DeformedUniqueTrianglePosition2Buffer0), ResourceName), EHairResourceUsageType::Dynamic);
+			InternalCreateVertexBufferRDG<FHairStrandsMeshTrianglePositionFormat>(GraphBuilder, RootCount, LOD.DeformedUniqueTrianglePosition0Buffer[0], ToHairResourceDebugName(HAIRSTRANDS_RESOUCE_NAME(CurveType, Hair.StrandsRootDeformed_DeformedUniqueTrianglePosition0Buffer0), ResourceName), OwnerName, EHairResourceUsageType::Dynamic);
+			InternalCreateVertexBufferRDG<FHairStrandsMeshTrianglePositionFormat>(GraphBuilder, RootCount, LOD.DeformedUniqueTrianglePosition1Buffer[0], ToHairResourceDebugName(HAIRSTRANDS_RESOUCE_NAME(CurveType, Hair.StrandsRootDeformed_DeformedUniqueTrianglePosition1Buffer0), ResourceName), OwnerName, EHairResourceUsageType::Dynamic);
+			InternalCreateVertexBufferRDG<FHairStrandsMeshTrianglePositionFormat>(GraphBuilder, RootCount, LOD.DeformedUniqueTrianglePosition2Buffer[0], ToHairResourceDebugName(HAIRSTRANDS_RESOUCE_NAME(CurveType, Hair.StrandsRootDeformed_DeformedUniqueTrianglePosition2Buffer0), ResourceName), OwnerName, EHairResourceUsageType::Dynamic);
 
 			// Double buffering is disabled by default unless the read-only cvar r.HairStrands.ContinuousDecimationReordering is set
 			if (IsHairStrandContinuousDecimationReorderingEnabled())
 			{
-				InternalCreateVertexBufferRDG<FHairStrandsMeshTrianglePositionFormat>(GraphBuilder, RootCount, LOD.DeformedUniqueTrianglePosition0Buffer[1], ToHairResourceDebugName(HAIRSTRANDS_RESOUCE_NAME(CurveType, Hair.StrandsRootDeformed_DeformedUniqueTrianglePosition0Buffer1), ResourceName), EHairResourceUsageType::Dynamic);
-				InternalCreateVertexBufferRDG<FHairStrandsMeshTrianglePositionFormat>(GraphBuilder, RootCount, LOD.DeformedUniqueTrianglePosition1Buffer[1], ToHairResourceDebugName(HAIRSTRANDS_RESOUCE_NAME(CurveType, Hair.StrandsRootDeformed_DeformedUniqueTrianglePosition1Buffer1), ResourceName), EHairResourceUsageType::Dynamic);
-				InternalCreateVertexBufferRDG<FHairStrandsMeshTrianglePositionFormat>(GraphBuilder, RootCount, LOD.DeformedUniqueTrianglePosition2Buffer[1], ToHairResourceDebugName(HAIRSTRANDS_RESOUCE_NAME(CurveType, Hair.StrandsRootDeformed_DeformedUniqueTrianglePosition2Buffer1), ResourceName), EHairResourceUsageType::Dynamic);
+				InternalCreateVertexBufferRDG<FHairStrandsMeshTrianglePositionFormat>(GraphBuilder, RootCount, LOD.DeformedUniqueTrianglePosition0Buffer[1], ToHairResourceDebugName(HAIRSTRANDS_RESOUCE_NAME(CurveType, Hair.StrandsRootDeformed_DeformedUniqueTrianglePosition0Buffer1), ResourceName), OwnerName, EHairResourceUsageType::Dynamic);
+				InternalCreateVertexBufferRDG<FHairStrandsMeshTrianglePositionFormat>(GraphBuilder, RootCount, LOD.DeformedUniqueTrianglePosition1Buffer[1], ToHairResourceDebugName(HAIRSTRANDS_RESOUCE_NAME(CurveType, Hair.StrandsRootDeformed_DeformedUniqueTrianglePosition1Buffer1), ResourceName), OwnerName, EHairResourceUsageType::Dynamic);
+				InternalCreateVertexBufferRDG<FHairStrandsMeshTrianglePositionFormat>(GraphBuilder, RootCount, LOD.DeformedUniqueTrianglePosition2Buffer[1], ToHairResourceDebugName(HAIRSTRANDS_RESOUCE_NAME(CurveType, Hair.StrandsRootDeformed_DeformedUniqueTrianglePosition2Buffer1), ResourceName), OwnerName, EHairResourceUsageType::Dynamic);
 			}
 		}
 	}
@@ -1550,8 +1572,8 @@ void FHairStrandsRootData::Reset()
 
 /////////////////////////////////////////////////////////////////////////////////////////
 
-FHairStrandsInterpolationResource::FHairStrandsInterpolationResource(FHairStrandsInterpolationBulkData& InBulkData, const FHairResourceName& InResourceName) :
-	FHairCommonResource(EHairStrandsAllocationType::Deferred, InResourceName),
+FHairStrandsInterpolationResource::FHairStrandsInterpolationResource(FHairStrandsInterpolationBulkData& InBulkData, const FHairResourceName& InResourceName, const FName& InOwnerName) :
+	FHairCommonResource(EHairStrandsAllocationType::Deferred, InResourceName, InOwnerName),
 	InterpolationBuffer(), Interpolation0Buffer(), Interpolation1Buffer(), BulkData(InBulkData)
 {
 	// Sanity check
@@ -1588,14 +1610,14 @@ void FHairStrandsInterpolationResource::InternalAllocate(FRDGBuilder& GraphBuild
 	const bool bUseSingleGuide = !!(BulkData.Flags & FHairStrandsInterpolationBulkData::DataFlags_HasSingleGuideData);
 	if (bUseSingleGuide)
 	{
-		InternalCreateVertexBufferRDG_FromBulkData<FHairStrandsInterpolationFormat>(GraphBuilder, BulkData.Interpolation, BulkData.PointCount, InterpolationBuffer, ToHairResourceDebugName(TEXT("Hair.StrandsInterpolation_InterpolationBuffer"), ResourceName), EHairResourceUsageType::Static);
+		InternalCreateVertexBufferRDG_FromBulkData<FHairStrandsInterpolationFormat>(GraphBuilder, BulkData.Interpolation, BulkData.PointCount, InterpolationBuffer, ToHairResourceDebugName(TEXT("Hair.StrandsInterpolation_InterpolationBuffer"), ResourceName), OwnerName, EHairResourceUsageType::Static);
 	}
 	else
 	{
-		InternalCreateVertexBufferRDG_FromBulkData<FHairStrandsInterpolation0Format>(GraphBuilder, BulkData.Interpolation0, BulkData.PointCount, Interpolation0Buffer, ToHairResourceDebugName(TEXT("Hair.StrandsInterpolation_Interpolation0Buffer"), ResourceName), EHairResourceUsageType::Static);
-		InternalCreateVertexBufferRDG_FromBulkData<FHairStrandsInterpolation1Format>(GraphBuilder, BulkData.Interpolation1, BulkData.PointCount, Interpolation1Buffer, ToHairResourceDebugName(TEXT("Hair.StrandsInterpolation_Interpolation1Buffer"), ResourceName), EHairResourceUsageType::Static);
+		InternalCreateVertexBufferRDG_FromBulkData<FHairStrandsInterpolation0Format>(GraphBuilder, BulkData.Interpolation0, BulkData.PointCount, Interpolation0Buffer, ToHairResourceDebugName(TEXT("Hair.StrandsInterpolation_Interpolation0Buffer"), ResourceName), OwnerName, EHairResourceUsageType::Static);
+		InternalCreateVertexBufferRDG_FromBulkData<FHairStrandsInterpolation1Format>(GraphBuilder, BulkData.Interpolation1, BulkData.PointCount, Interpolation1Buffer, ToHairResourceDebugName(TEXT("Hair.StrandsInterpolation_Interpolation1Buffer"), ResourceName), OwnerName, EHairResourceUsageType::Static);
 	}
-	InternalCreateVertexBufferRDG_FromBulkData<FHairStrandsRootIndexFormat>(GraphBuilder, BulkData.SimRootPointIndex, BulkData.SimPointCount, SimRootPointIndexBuffer, ToHairResourceDebugName(TEXT("Hair.StrandsInterpolation_SimRootPointIndex"), ResourceName), EHairResourceUsageType::Static);
+	InternalCreateVertexBufferRDG_FromBulkData<FHairStrandsRootIndexFormat>(GraphBuilder, BulkData.SimRootPointIndex, BulkData.SimPointCount, SimRootPointIndexBuffer, ToHairResourceDebugName(TEXT("Hair.StrandsInterpolation_SimRootPointIndex"), ResourceName), OwnerName, EHairResourceUsageType::Static);
 }
 
 void FHairStrandsInterpolationResource::InternalRelease()
@@ -1627,15 +1649,15 @@ void FHairCardsInterpolationBulkData::Serialize(FArchive& Ar)
 	Ar << Interpolation;
 }
 
-FHairCardsInterpolationResource::FHairCardsInterpolationResource(FHairCardsInterpolationBulkData& InBulkData, const FHairResourceName& InResourceName) :
-	FHairCommonResource(EHairStrandsAllocationType::Deferred, InResourceName),
+FHairCardsInterpolationResource::FHairCardsInterpolationResource(FHairCardsInterpolationBulkData& InBulkData, const FHairResourceName& InResourceName, const FName& InOwnerName) :
+	FHairCommonResource(EHairStrandsAllocationType::Deferred, InResourceName, InOwnerName),
 	InterpolationBuffer(), BulkData(InBulkData)
 {
 }
 
 void FHairCardsInterpolationResource::InternalAllocate(FRDGBuilder& GraphBuilder)
 {
-	InternalCreateVertexBufferRDG<FHairCardsInterpolationFormat>(GraphBuilder, BulkData.Interpolation, InterpolationBuffer, ToHairResourceDebugName(TEXT("Hair.CardsInterpolation_InterpolationBuffer"), ResourceName), EHairResourceUsageType::Static);
+	InternalCreateVertexBufferRDG<FHairCardsInterpolationFormat>(GraphBuilder, BulkData.Interpolation, InterpolationBuffer, ToHairResourceDebugName(TEXT("Hair.CardsInterpolation_InterpolationBuffer"), ResourceName), OwnerName, EHairResourceUsageType::Static);
 }
 
 void FHairCardsInterpolationResource::InternalRelease()
@@ -1656,8 +1678,8 @@ void FHairCardsInterpolationResource::InternalRelease()
 // NOTE: the vertex buffer is a float4 because it is registered as a UAV for the compute shader to work
 // TODO: use a plain float vertex buffer with 3x the entries instead to save memory? (float3 UAVs are not allowed)
 bool GetSupportHairStrandsProceduralPrimitive(EShaderPlatform InShaderPlatform);
-FHairStrandsRaytracingResource::FHairStrandsRaytracingResource(const FHairStrandsBulkData& InData, const FHairResourceName& ResourceName) :
-	FHairCommonResource(EHairStrandsAllocationType::Deferred, ResourceName), bOwnBuffers(true)
+FHairStrandsRaytracingResource::FHairStrandsRaytracingResource(const FHairStrandsBulkData& InData, const FHairResourceName& ResourceName, const FName& InOwnerName) :
+	FHairCommonResource(EHairStrandsAllocationType::Deferred, ResourceName, InOwnerName), bOwnBuffers(true)
 {
 	bProceduralPrimitive = GetSupportHairStrandsProceduralPrimitive(GMaxRHIShaderPlatform);
 	if (bProceduralPrimitive)
@@ -1672,13 +1694,13 @@ FHairStrandsRaytracingResource::FHairStrandsRaytracingResource(const FHairStrand
 	}
 }
 
-FHairStrandsRaytracingResource::FHairStrandsRaytracingResource(const FHairCardsBulkData& InData, const FHairResourceName& ResourceName) :
-	FHairCommonResource(EHairStrandsAllocationType::Deferred, ResourceName),
+FHairStrandsRaytracingResource::FHairStrandsRaytracingResource(const FHairCardsBulkData& InData, const FHairResourceName& ResourceName, const FName& InOwnerName) :
+	FHairCommonResource(EHairStrandsAllocationType::Deferred, ResourceName, InOwnerName),
 	VertexCount(InData.GetNumVertices()), bOwnBuffers(false)
 {}
 
-FHairStrandsRaytracingResource::FHairStrandsRaytracingResource(const FHairMeshesBulkData& InData, const FHairResourceName& ResourceName) :
-	FHairCommonResource(EHairStrandsAllocationType::Deferred, ResourceName),
+FHairStrandsRaytracingResource::FHairStrandsRaytracingResource(const FHairMeshesBulkData& InData, const FHairResourceName& ResourceName, const FName& InOwnerName) :
+	FHairCommonResource(EHairStrandsAllocationType::Deferred, ResourceName, InOwnerName),
 	VertexCount(InData.GetNumVertices()), bOwnBuffers(false)
 {}
 
@@ -1686,8 +1708,8 @@ void FHairStrandsRaytracingResource::InternalAllocate(FRDGBuilder& GraphBuilder)
 {
 	if (bOwnBuffers)
 	{
-		InternalCreateVertexBufferRDG<FHairStrandsRaytracingFormat>(GraphBuilder, VertexCount, PositionBuffer, ToHairResourceDebugName(TEXT("Hair.StrandsRaytracing_PositionBuffer"), ResourceName), EHairResourceUsageType::Dynamic);
-		InternalCreateStructuredBufferRDG<FHairStrandsIndexFormat>(GraphBuilder, IndexCount, IndexBuffer, ToHairResourceDebugName(TEXT("Hair.StrandsRaytracing_IndexBuffer"), ResourceName), EHairResourceUsageType::Dynamic);
+		InternalCreateVertexBufferRDG<FHairStrandsRaytracingFormat>(GraphBuilder, VertexCount, PositionBuffer, ToHairResourceDebugName(TEXT("Hair.StrandsRaytracing_PositionBuffer"), ResourceName), OwnerName, EHairResourceUsageType::Dynamic);
+		InternalCreateStructuredBufferRDG<FHairStrandsIndexFormat>(GraphBuilder, IndexCount, IndexBuffer, ToHairResourceDebugName(TEXT("Hair.StrandsRaytracing_IndexBuffer"), ResourceName), OwnerName, EHairResourceUsageType::Dynamic);
 	}
 }
 
diff --git a/Engine/Plugins/Runtime/HairStrands/Source/HairStrandsCore/Private/HairStrandsRendering.cpp b/Engine/Plugins/Runtime/HairStrands/Source/HairStrandsCore/Private/HairStrandsRendering.cpp
--- a/Engine/Plugins/Runtime/HairStrands/Source/HairStrandsCore/Private/HairStrandsRendering.cpp
+++ b/Engine/Plugins/Runtime/HairStrands/Source/HairStrandsCore/Private/HairStrandsRendering.cpp
@@ -1442,6 +1442,7 @@ static void BuildHairAccelerationStructure_Strands(
 	FBufferRHIRef& IndexBuffer,
 	FRayTracingGeometry* OutRayTracingGeometry,
 	const FString& AssetDebugName,
+	const FString& MeshComponentName,
 	uint32 LODIndex,
 	bool bProceduralPrimitive,
 	int ProceduralSplits)
@@ -1455,7 +1456,10 @@ static void BuildHairAccelerationStructure_Strands(
 	static int32 DebugNumber = 0;
 	Initializer.DebugName = FDebugName(DebugName, DebugNumber++);
 #endif
-
+#if RHI_ENABLE_RESOURCE_INFO
+	const FName OwnerName(FString::Printf(TEXT("%s [LOD%d]"), *MeshComponentName, LODIndex));
+	Initializer.OwnerName = OwnerName;
+#endif
 	if (bProceduralPrimitive)
 	{
 		Initializer.GeometryType = RTGT_Procedural;
@@ -1498,6 +1502,7 @@ static void BuildHairAccelerationStructure_Cards(FRHICommandList& RHICmdList,
 	FHairCardsDeformedResource* DeformedResource,
 	FRayTracingGeometry* OutRayTracingGeometry,
 	const FString& AssetDebugName,
+	const FString& MeshComponentName,
 	uint32 LODIndex)
 {
 	FRayTracingGeometryInitializer Initializer;
@@ -1508,6 +1513,10 @@ static void BuildHairAccelerationStructure_Cards(FRHICommandList& RHICmdList,
 	static const FName DebugName("AS_HairCards");
 	static int32 DebugNumber = 0;
 	Initializer.DebugName = FDebugName(DebugName, DebugNumber++);
+#endif
+#if RHI_ENABLE_RESOURCE_INFO
+	const FName OwnerName(FString::Printf(TEXT("%s [LOD%d]"), *MeshComponentName, LODIndex));
+	Initializer.OwnerName = OwnerName;
 #endif
 	Initializer.IndexBuffer = RestResource->RestIndexBuffer.IndexBufferRHI;
 	Initializer.IndexBufferOffset = 0;
@@ -1738,7 +1747,7 @@ FHairGroupPublicData::FVertexFactoryInput ComputeHairStrandsVertexInputData(cons
 	return InternalComputeHairStrandsVertexInputData(nullptr, Instance, ViewMode);
 }
 
-void CreateHairStrandsDebugAttributeBuffer(FRDGBuilder& GraphBuilder, FRDGExternalBuffer* DebugAttributeBuffer, uint32 VertexCount);
+void CreateHairStrandsDebugAttributeBuffer(FRDGBuilder& GraphBuilder, FRDGExternalBuffer* DebugAttributeBuffer, uint32 VertexCount, const FName& OwnerName);
 
 void ComputeHairStrandsInterpolation(
 	FRDGBuilder& GraphBuilder,
@@ -2063,7 +2072,7 @@ void ComputeHairStrandsInterpolation(
 				// Special case for debug mode were the attribute buffer is patch with some custom data to show hair properties (strands belonging to the same cluster, ...)
 				if (Instance->Strands.DebugAttributeBuffer.Buffer == nullptr)
 				{
-					CreateHairStrandsDebugAttributeBuffer(GraphBuilder, &Instance->Strands.DebugAttributeBuffer, Instance->Strands.Data->Attributes.GetBulkDataSize());
+					CreateHairStrandsDebugAttributeBuffer(GraphBuilder, &Instance->Strands.DebugAttributeBuffer, Instance->Strands.Data->Attributes.GetBulkDataSize(), FName(Instance->Debug.MeshComponentName));
 				}
 				FRDGImportedBuffer OutRenAttributeBuffer = Register(GraphBuilder, Instance->Strands.DebugAttributeBuffer, ERDGImportedBufferFlags::CreateUAV);
 
@@ -2198,6 +2207,7 @@ void ComputeHairStrandsInterpolation(
 								IndexBuffer,
 								&Instance->Strands.RenRaytracingResource->RayTracingGeometry,
 								Instance->Debug.GroomAssetName,
+								Instance->Debug.MeshComponentName,
 								HairLODIndex,
 								bProceduralPrimitive,
 								ProceduralSplits
@@ -2340,7 +2350,7 @@ void ComputeHairStrandsInterpolation(
 						const bool bLocalNeedBuild = !LocalLOD.RaytracingResource->bIsRTGeometryInitialized;
 						if (bLocalNeedBuild)
 						{
-							BuildHairAccelerationStructure_Cards(RHICmdList, LocalLOD.RestResource, LocalLOD.DeformedResource, &LocalLOD.RaytracingResource->RayTracingGeometry, Instance->Debug.GroomAssetName, HairLODIndex);
+							BuildHairAccelerationStructure_Cards(RHICmdList, LocalLOD.RestResource, LocalLOD.DeformedResource, &LocalLOD.RaytracingResource->RayTracingGeometry, Instance->Debug.GroomAssetName, Instance->Debug.MeshComponentName, HairLODIndex);
 							LocalLOD.RaytracingResource->bIsRTGeometryInitialized = true;
 						}
 						else if (bNeedUpdate)
diff --git a/Engine/Plugins/Runtime/HairStrands/Source/HairStrandsCore/Public/GroomAsset.h b/Engine/Plugins/Runtime/HairStrands/Source/HairStrandsCore/Public/GroomAsset.h
--- a/Engine/Plugins/Runtime/HairStrands/Source/HairStrandsCore/Public/GroomAsset.h
+++ b/Engine/Plugins/Runtime/HairStrands/Source/HairStrandsCore/Public/GroomAsset.h
@@ -582,6 +582,9 @@ public:
 	/** Used for PSO precaching of used materials and vertex factories */
 	TArray<FHairVertexFactoryTypesPerMaterialData> CollectVertexFactoryTypesPerMaterialData(EShaderPlatform ShaderPlatform);
 
+	/** Helper function to return the asset path name, optionally joined with the LOD index if LODIndex > -1. */
+	FName GetAssetPathName(int32 LODIndex = -1);
+
 //private :
 #if WITH_EDITOR
 	FOnGroomAssetChanged OnGroomAssetChanged;
diff --git a/Engine/Plugins/Runtime/HairStrands/Source/HairStrandsCore/Public/GroomBindingAsset.h b/Engine/Plugins/Runtime/HairStrands/Source/HairStrandsCore/Public/GroomBindingAsset.h
--- a/Engine/Plugins/Runtime/HairStrands/Source/HairStrandsCore/Public/GroomBindingAsset.h
+++ b/Engine/Plugins/Runtime/HairStrands/Source/HairStrandsCore/Public/GroomBindingAsset.h
@@ -142,6 +142,9 @@ public:
 	/** Returns true if the target is not null and matches the binding type */ 
 	bool HasValidTarget() const;
 
+	/** Helper function to return the asset path name, optionally joined with the LOD index if LODIndex > -1. */
+	FName GetAssetPathName(int32 LODIndex = -1);
+
 #if WITH_EDITOR
 	FOnGroomBindingAssetChanged& GetOnGroomBindingAssetChanged() { return OnGroomBindingAssetChanged; }
 
diff --git a/Engine/Plugins/Runtime/HairStrands/Source/HairStrandsCore/Public/GroomResources.h b/Engine/Plugins/Runtime/HairStrands/Source/HairStrandsCore/Public/GroomResources.h
--- a/Engine/Plugins/Runtime/HairStrands/Source/HairStrandsCore/Public/GroomResources.h
+++ b/Engine/Plugins/Runtime/HairStrands/Source/HairStrandsCore/Public/GroomResources.h
@@ -77,7 +77,7 @@ EHairResourceLoadingType GetHairResourceLoadingType(EHairGeometryType InGeometry
 struct FHairCommonResource : public FRenderResource
 {
 	/** Build the hair strands resource */
-	FHairCommonResource(EHairStrandsAllocationType AllocationType, const FHairResourceName& InResourceName, bool bUseRenderGraph=true);
+	FHairCommonResource(EHairStrandsAllocationType AllocationType, const FHairResourceName& InResourceName, const FName& InOwnerName, bool bUseRenderGraph=true);
 
 	/* Init/Release buffers (FRenderResource) */
 	virtual void InitRHI() override;
@@ -104,13 +104,14 @@ struct FHairCommonResource : public FRenderResource
 
 	/* Store (precis) debug name */
 	FHairResourceName ResourceName;
+	FName OwnerName;
 };
 
 /* Render buffers for root deformation for dynamic meshes */
 struct FHairStrandsRestRootResource : public FHairCommonResource
 {
 	/** Build the hair strands resource */
-	FHairStrandsRestRootResource(FHairStrandsRootBulkData& BulkData, EHairStrandsResourcesType CurveType, const FHairResourceName& ResourceName);
+	FHairStrandsRestRootResource(FHairStrandsRootBulkData& BulkData, EHairStrandsResourcesType CurveType, const FHairResourceName& ResourceName, const FName& OwnerName);
 
 	/* Init/Release buffers */
 	virtual void InternalAllocate(FRDGBuilder& GraphBuilder) override;
@@ -189,8 +190,8 @@ struct FHairStrandsRestRootResource : public FHairCommonResource
 struct FHairStrandsDeformedRootResource : public FHairCommonResource
 {
 	/** Build the hair strands resource */
-	FHairStrandsDeformedRootResource(EHairStrandsResourcesType CurveType, const FHairResourceName& ResourceName);
-	FHairStrandsDeformedRootResource(const FHairStrandsRestRootResource* InRestResources, EHairStrandsResourcesType CurveType, const FHairResourceName& ResourceName);
+	FHairStrandsDeformedRootResource(EHairStrandsResourcesType CurveType, const FHairResourceName& ResourceName, const FName& OwnerName);
+	FHairStrandsDeformedRootResource(const FHairStrandsRestRootResource* InRestResources, EHairStrandsResourcesType CurveType, const FHairResourceName& ResourceName, const FName& OwnerName);
 
 	/* Init/Release buffers */
 	virtual void InternalAllocateLOD(FRDGBuilder& GraphBuilder, int32 LODIndex) override;
@@ -291,7 +292,7 @@ struct FHairStrandsDeformedRootResource : public FHairCommonResource
 struct FHairStrandsRestResource : public FHairCommonResource
 {
 	/** Build the hair strands resource */
-	FHairStrandsRestResource(FHairStrandsBulkData& InBulkData, EHairStrandsResourcesType CurveType, const FHairResourceName& ResourceName);
+	FHairStrandsRestResource(FHairStrandsBulkData& InBulkData, EHairStrandsResourcesType CurveType, const FHairResourceName& ResourceName, const FName& OwnerName);
 
 	/* Init/Release buffers */
 	virtual void InternalAllocate(FRDGBuilder& GraphBuilder) override;
@@ -352,7 +353,7 @@ struct FHairStrandsRestResource : public FHairCommonResource
 struct FHairStrandsDeformedResource : public FHairCommonResource
 {
 	/** Build the hair strands resource */
-	FHairStrandsDeformedResource(FHairStrandsBulkData& BulkData, EHairStrandsResourcesType CurveType, const FHairResourceName& ResourceName);
+	FHairStrandsDeformedResource(FHairStrandsBulkData& BulkData, EHairStrandsResourcesType CurveType, const FHairResourceName& ResourceName, const FName& OwnerName);
 
 	/* Init/Release buffers */
 	virtual void InternalAllocate(FRDGBuilder& GraphBuilder) override;
@@ -426,7 +427,7 @@ struct FHairStrandsDeformedResource : public FHairCommonResource
 
 struct FHairStrandsClusterCullingResource : public FHairCommonResource
 {
-	FHairStrandsClusterCullingResource(FHairStrandsClusterCullingBulkData& Data, const FHairResourceName& ResourceName);
+	FHairStrandsClusterCullingResource(FHairStrandsClusterCullingBulkData& Data, const FHairResourceName& ResourceName, const FName& OwnerName);
 
 	/* Init/Release buffers */
 	virtual void InternalAllocate(FRDGBuilder& GraphBuilder) override;
@@ -466,7 +467,7 @@ struct FHairStrandsClusterCullingResource : public FHairCommonResource
 struct FHairStrandsInterpolationResource : public FHairCommonResource
 {
 	/** Build the hair strands resource */
-	FHairStrandsInterpolationResource(FHairStrandsInterpolationBulkData& InBulkData, const FHairResourceName& ResourceName);
+	FHairStrandsInterpolationResource(FHairStrandsInterpolationBulkData& InBulkData, const FHairResourceName& ResourceName, const FName& OwnerName);
 	
 	/* Init/Release buffers */
 	virtual void InternalAllocate(FRDGBuilder& GraphBuilder) override;
@@ -505,9 +506,9 @@ struct FHairStrandsInterpolationResource : public FHairCommonResource
 struct FHairStrandsRaytracingResource : public FHairCommonResource
 {
 	/** Build the hair strands resource */
-	FHairStrandsRaytracingResource(const FHairStrandsBulkData& InData, const FHairResourceName& ResourceName);
-	FHairStrandsRaytracingResource(const FHairCardsBulkData& InData, const FHairResourceName& ResourceName);
-	FHairStrandsRaytracingResource(const FHairMeshesBulkData& InData, const FHairResourceName& ResourceName);
+	FHairStrandsRaytracingResource(const FHairStrandsBulkData& InData, const FHairResourceName& ResourceName, const FName& OwnerName);
+	FHairStrandsRaytracingResource(const FHairCardsBulkData& InData, const FHairResourceName& ResourceName, const FName& OwnerName);
+	FHairStrandsRaytracingResource(const FHairMeshesBulkData& InData, const FHairResourceName& ResourceName, const FName& OwnerName);
 
 	/* Init/Release buffers */
 	virtual void InternalAllocate(FRDGBuilder& GraphBuilder) override;
@@ -549,7 +550,7 @@ class FHairCardIndexBuffer : public FIndexBuffer
 {
 public:
 	const TArray<FHairCardsIndexFormat::Type>& Indices;
-	FHairCardIndexBuffer(const TArray<FHairCardsIndexFormat::Type>& InIndices) :Indices(InIndices) {}
+	FHairCardIndexBuffer(const TArray<FHairCardsIndexFormat::Type>& InIndices, const FName& InOwnerName);
 	virtual void InitRHI() override;
 };
 
@@ -559,7 +560,7 @@ struct FHairCardsBulkData;
 struct FHairCardsRestResource : public FHairCommonResource
 {
 	/** Build the hair strands resource */
-	FHairCardsRestResource(const FHairCardsBulkData& InBulkData, const FHairResourceName& ResourceName);
+	FHairCardsRestResource(const FHairCardsBulkData& InBulkData, const FHairResourceName& ResourceName, const FName& OwnerName);
 
 	virtual void InitResource() override;
 	virtual void ReleaseResource() override;
@@ -617,7 +618,7 @@ struct FHairCardsRestResource : public FHairCommonResource
 struct FHairCardsProceduralResource : public FHairCommonResource
 {
 	/** Build the hair strands resource */
-	FHairCardsProceduralResource(const FHairCardsProceduralDatas::FRenderData& HairCardsRenderData, const FIntPoint& AtlasResolution, const FHairCardsVoxel& InVoxel);
+	FHairCardsProceduralResource(const FHairCardsProceduralDatas::FRenderData& HairCardsRenderData, const FIntPoint& AtlasResolution, const FHairCardsVoxel& InVoxel, const FName& OwnerName);
 
 	/* Init/release buffers */
 	virtual void InternalAllocate(FRDGBuilder& GraphBuilder) override;
@@ -656,7 +657,7 @@ struct FHairCardsProceduralResource : public FHairCommonResource
 struct FHairCardsDeformedResource : public FHairCommonResource
 {
 	/** Build the hair strands resource */
-	FHairCardsDeformedResource(const FHairCardsBulkData& BulkData, bool bInitializeData, const FHairResourceName& ResourceName);
+	FHairCardsDeformedResource(const FHairCardsBulkData& BulkData, bool bInitializeData, const FHairResourceName& ResourceName, const FName& OwnerName);
 
 	/* Init/release buffers */
 	virtual void InternalAllocate(FRDGBuilder& GraphBuilder) override;
@@ -729,7 +730,7 @@ struct HAIRSTRANDSCORE_API FHairCardsInterpolationBulkData
 struct FHairCardsInterpolationResource : public FHairCommonResource
 {
 	/** Build the hair strands resource */
-	FHairCardsInterpolationResource(FHairCardsInterpolationBulkData& InBulkData, const FHairResourceName& ResourceName);
+	FHairCardsInterpolationResource(FHairCardsInterpolationBulkData& InBulkData, const FHairResourceName& ResourceName, const FName& OwnerName);
 
 	/* Init/release buffers */
 	virtual void InternalAllocate(FRDGBuilder& GraphBuilder) override;
@@ -759,7 +760,7 @@ struct FHairCardsInterpolationResource : public FHairCommonResource
 struct FHairMeshesRestResource : public FHairCommonResource
 {
 	/** Build the hair strands resource */
-	FHairMeshesRestResource(const FHairMeshesBulkData& BulkData, const FHairResourceName& ResourceName);
+	FHairMeshesRestResource(const FHairMeshesBulkData& BulkData, const FHairResourceName& ResourceName, const FName& OwnerName);
 
 	virtual void InitResource() override;
 	virtual void ReleaseResource() override;
@@ -815,7 +816,7 @@ struct FHairMeshesRestResource : public FHairCommonResource
 struct FHairMeshesDeformedResource : public FHairCommonResource
 {
 	/** Build the hair strands resource */
-	FHairMeshesDeformedResource(const FHairMeshesBulkData& InBulkData, bool bInInitializedData, const FHairResourceName& ResourceName);
+	FHairMeshesDeformedResource(const FHairMeshesBulkData& InBulkData, bool bInInitializedData, const FHairResourceName& ResourceName, const FName& OwnerName);
 
 	/* Init/release buffers */
 	virtual void InternalAllocate(FRDGBuilder& GraphBuilder) override;
diff --git a/Engine/Source/Runtime/RenderCore/Private/RenderGraphBuilder.cpp b/Engine/Source/Runtime/RenderCore/Private/RenderGraphBuilder.cpp
--- a/Engine/Source/Runtime/RenderCore/Private/RenderGraphBuilder.cpp
+++ b/Engine/Source/Runtime/RenderCore/Private/RenderGraphBuilder.cpp
@@ -3546,6 +3546,13 @@ void FRDGBuilder::BeginResourceRHI(FRDGPassHandle PassHandle, FRDGBufferRef Buff
 
 		SetRHI(Buffer, GRenderGraphResourcePool.FindFreeBuffer(Buffer->Desc, Buffer->Name, Alignment), PassHandle);
 	}
+
+#if RHI_ENABLE_RESOURCE_INFO
+	if (Buffer->HasRHI())
+	{
+		Buffer->ResourceRHI->SetOwnerName(Buffer->OwnerName);
+	}
+#endif
 }
 
 void FRDGBuilder::InitRHI(FRDGBufferSRVRef SRV)
diff --git a/Engine/Source/Runtime/RenderCore/Public/RenderGraphResources.h b/Engine/Source/Runtime/RenderCore/Public/RenderGraphResources.h
--- a/Engine/Source/Runtime/RenderCore/Public/RenderGraphResources.h
+++ b/Engine/Source/Runtime/RenderCore/Public/RenderGraphResources.h
@@ -136,6 +136,13 @@ public:
 		return ResourceRHI;
 	}
 
+	void SetOwnerName(const FName& InOwnerName)
+	{
+#if RHI_ENABLE_RESOURCE_INFO
+		OwnerName = InOwnerName;
+#endif
+	}
+
 	//////////////////////////////////////////////////////////////////////////
 
 protected:
@@ -165,6 +172,10 @@ private:
 	FRDGResourceDebugData& GetDebugData() const;
 #endif
 
+#if RHI_ENABLE_RESOURCE_INFO
+	FName OwnerName;	// For RHI resource tracking
+#endif
+
 	friend FRDGBuilder;
 	friend FRDGUserValidation;
 	friend FRDGBarrierValidation;
diff --git a/Engine/Source/Runtime/Renderer/Private/HairStrands/HairStrandsInterface.cpp b/Engine/Source/Runtime/Renderer/Private/HairStrands/HairStrandsInterface.cpp
--- a/Engine/Source/Runtime/Renderer/Private/HairStrands/HairStrandsInterface.cpp
+++ b/Engine/Source/Runtime/Renderer/Private/HairStrands/HairStrandsInterface.cpp
@@ -273,16 +273,17 @@ void ConvertToExternalBufferWithViews(FRDGBuilder& GraphBuilder, FRDGBufferRef&
 	OutBuffer.Format = Format;
 }
 
-void InternalCreateIndirectBufferRDG(FRDGBuilder& GraphBuilder, FRDGExternalBuffer& Out, const TCHAR* DebugName)
+void InternalCreateIndirectBufferRDG(FRDGBuilder& GraphBuilder, FRDGExternalBuffer& Out, const TCHAR* DebugName, const FName& OwnerName)
 {
 	FRDGBufferDesc Desc = FRDGBufferDesc::CreateBufferDesc(4, 4);
 	Desc.Usage |= BUF_DrawIndirect;
-	FRDGBufferRef Buffer = GraphBuilder.CreateBuffer(Desc, DebugName);
+	FRDGBufferRef Buffer = GraphBuilder.CreateBuffer(Desc, DebugName, ERDGBufferFlags::None);
+	Buffer->SetOwnerName(OwnerName);
 	AddClearUAVPass(GraphBuilder, GraphBuilder.CreateUAV(Buffer, PF_R32_UINT), 0u);
 	ConvertToExternalBufferWithViews(GraphBuilder, Buffer, Out, PF_R32_UINT);
 }
 
-void InternalCreateVertexBufferRDG(FRDGBuilder& GraphBuilder, uint32 ElementSizeInBytes, uint32 ElementCount, EPixelFormat Format, FRDGExternalBuffer& Out, const TCHAR* DebugName, bool bClearFloat=false)
+void InternalCreateVertexBufferRDG(FRDGBuilder& GraphBuilder, uint32 ElementSizeInBytes, uint32 ElementCount, EPixelFormat Format, FRDGExternalBuffer& Out, const TCHAR* DebugName, const FName& OwnerName, bool bClearFloat=false)
 {
 	FRDGBufferRef Buffer = nullptr;
 
@@ -297,6 +298,7 @@ void InternalCreateVertexBufferRDG(FRDGBuilder& GraphBuilder, uint32 ElementSize
 	// #hair_todo: Create this with a create+clear pass instead?
 	const FRDGBufferDesc Desc = FRDGBufferDesc::CreateBufferDesc(ElementSizeInBytes, ElementCount);
 	Buffer = GraphBuilder.CreateBuffer(Desc, DebugName, ERDGBufferFlags::MultiFrame);
+	Buffer->SetOwnerName(OwnerName);
 	if (bClearFloat)
 	{
 		AddClearUAVFloatPass(GraphBuilder, GraphBuilder.CreateUAV(Buffer, Format), 0.f);
@@ -308,8 +310,9 @@ void InternalCreateVertexBufferRDG(FRDGBuilder& GraphBuilder, uint32 ElementSize
 	ConvertToExternalBufferWithViews(GraphBuilder, Buffer, Out, Format);
 }
 
-FHairGroupPublicData::FHairGroupPublicData(uint32 InGroupIndex)
+FHairGroupPublicData::FHairGroupPublicData(uint32 InGroupIndex, const FName& InOwnerName)
 {
+	SetOwnerName(InOwnerName);
 	GroupIndex = InGroupIndex;
 	GroupControlTriangleStripVertexCount = 0;
 	ClusterCount = 0;
@@ -357,14 +360,14 @@ void FHairGroupPublicData::Allocate(FRDGBuilder& GraphBuilder)
 	
 	if (GUsingNullRHI || !bHasStrands) { return; }
 
-	InternalCreateIndirectBufferRDG(GraphBuilder, DrawIndirectBuffer, TEXT("Hair.Cluster_DrawIndirectBuffer"));
-	InternalCreateIndirectBufferRDG(GraphBuilder, DrawIndirectRasterComputeBuffer, TEXT("Hair.Cluster_DrawIndirectRasterComputeBuffer"));
+	InternalCreateIndirectBufferRDG(GraphBuilder, DrawIndirectBuffer, TEXT("Hair.Cluster_DrawIndirectBuffer"), GetOwnerName());
+	InternalCreateIndirectBufferRDG(GraphBuilder, DrawIndirectRasterComputeBuffer, TEXT("Hair.Cluster_DrawIndirectRasterComputeBuffer"), GetOwnerName());
 
-	InternalCreateVertexBufferRDG(GraphBuilder, sizeof(int32), ClusterCount * 6, EPixelFormat::PF_R32_SINT, ClusterAABBBuffer, TEXT("Hair.Cluster_ClusterAABBBuffer"));
-	InternalCreateVertexBufferRDG(GraphBuilder, sizeof(int32), 6, EPixelFormat::PF_R32_SINT, GroupAABBBuffer, TEXT("Hair.Cluster_GroupAABBBuffer"));
+	InternalCreateVertexBufferRDG(GraphBuilder, sizeof(int32), ClusterCount * 6, EPixelFormat::PF_R32_SINT, ClusterAABBBuffer, TEXT("Hair.Cluster_ClusterAABBBuffer"), GetOwnerName());
+	InternalCreateVertexBufferRDG(GraphBuilder, sizeof(int32), 6, EPixelFormat::PF_R32_SINT, GroupAABBBuffer, TEXT("Hair.Cluster_GroupAABBBuffer"), GetOwnerName());
 
-	InternalCreateVertexBufferRDG(GraphBuilder, sizeof(int32), VertexCount, EPixelFormat::PF_R32_UINT, CulledVertexIdBuffer, TEXT("Hair.Cluster_CulledVertexIdBuffer"));
-	InternalCreateVertexBufferRDG(GraphBuilder, sizeof(float), VertexCount, EPixelFormat::PF_R32_FLOAT, CulledVertexRadiusScaleBuffer, TEXT("Hair.Cluster_CulledVertexRadiusScaleBuffer"), true);
+	InternalCreateVertexBufferRDG(GraphBuilder, sizeof(int32), VertexCount, EPixelFormat::PF_R32_UINT, CulledVertexIdBuffer, TEXT("Hair.Cluster_CulledVertexIdBuffer"), GetOwnerName());
+	InternalCreateVertexBufferRDG(GraphBuilder, sizeof(float), VertexCount, EPixelFormat::PF_R32_FLOAT, CulledVertexRadiusScaleBuffer, TEXT("Hair.Cluster_CulledVertexRadiusScaleBuffer"), GetOwnerName(), true);
 
 	GraphBuilder.SetBufferAccessFinal(Register(GraphBuilder, DrawIndirectBuffer, ERDGImportedBufferFlags::None).Buffer, ERHIAccess::IndirectArgs);
 
diff --git a/Engine/Source/Runtime/Renderer/Public/HairStrandsInterface.h b/Engine/Source/Runtime/Renderer/Public/HairStrandsInterface.h
--- a/Engine/Source/Runtime/Renderer/Public/HairStrandsInterface.h
+++ b/Engine/Source/Runtime/Renderer/Public/HairStrandsInterface.h
@@ -107,7 +107,7 @@ typedef TArray<FHairStrandsInstance*> FHairStrandsInstances;
 class RENDERER_API FHairGroupPublicData : public FRenderResource
 {
 public:
-	FHairGroupPublicData(uint32 InGroupIndex);
+	FHairGroupPublicData(uint32 InGroupIndex, const FName& OwnerName);
 	void SetClusters(uint32 InClusterCount, uint32 InVertexCount);
 	
 	virtual void InitRHI() override;
