Fix uninitialized data in hair strands to fix indeterminism in cooked build

#rnx

diff --git a/Engine/Plugins/Runtime/HairStrands/Source/HairStrandsCore/Public/HairStrandsDatas.h b/Engine/Plugins/Runtime/HairStrands/Source/HairStrandsCore/Public/HairStrandsDatas.h
--- a/Engine/Plugins/Runtime/HairStrands/Source/HairStrandsCore/Public/HairStrandsDatas.h
+++ b/Engine/Plugins/Runtime/HairStrands/Source/HairStrandsCore/Public/HairStrandsDatas.h
@@ -471,7 +471,7 @@ struct HAIRSTRANDSCORE_API FHairStrandsBulkData
 	float MaxRadius = 0;
 	FBox BoundingBox = FBox(EForceInit::ForceInit);
 	uint32 Flags = 0;
-	uint32 AttributeOffsets[HAIR_ATTRIBUTE_COUNT];
+	uint32 AttributeOffsets[HAIR_ATTRIBUTE_COUNT] = {0};
 
 	FByteBulkData Positions;	// Size = PointCount
 	FByteBulkData Attributes;	// Size = x*PointCount + y*CurveCount (depends on the stored attributes, which could be per-vertex(x) or per-curve(y))
