HairStrands: Implement GetShaderVirtualPath() for all optimus data interfaces that can support it.

diff --git a/Engine/Plugins/Runtime/HairStrands/Source/HairStrandsCore/Private/Deformer/DeformerDataInterfaceGroom.cpp b/Engine/Plugins/Runtime/HairStrands/Source/HairStrandsCore/Private/Deformer/DeformerDataInterfaceGroom.cpp
--- a/Engine/Plugins/Runtime/HairStrands/Source/HairStrandsCore/Private/Deformer/DeformerDataInterfaceGroom.cpp
+++ b/Engine/Plugins/Runtime/HairStrands/Source/HairStrandsCore/Private/Deformer/DeformerDataInterfaceGroom.cpp
@@ -110,13 +110,16 @@ void UOptimusGroomDataInterface::GetShaderParameters(TCHAR const* UID, FShaderPa
 	InOutBuilder.AddNestedStruct<FGroomDataInterfaceParameters>(UID);
 }
 
-void UOptimusGroomDataInterface::GetPermutations(FComputeKernelPermutationVector& OutPermutationVector) const
+TCHAR const* UOptimusGroomDataInterface::TemplateFilePath = TEXT("/Plugin/Runtime/HairStrands/Private/DeformerDataInterfaceGroom.ush");
+
+TCHAR const* UOptimusGroomDataInterface::GetShaderVirtualPath() const
 {
+	return TemplateFilePath;
 }
 
 void UOptimusGroomDataInterface::GetShaderHash(FString& InOutKey) const
 {
-	GetShaderFileHash(TEXT("/Plugin/Runtime/HairStrands/Private/DeformerDataInterfaceGroom.ush"), EShaderPlatform::SP_PCD3D_SM5).AppendString(InOutKey);
+	GetShaderFileHash(TemplateFilePath, EShaderPlatform::SP_PCD3D_SM5).AppendString(InOutKey);
 }
 
 void UOptimusGroomDataInterface::GetHLSL(FString& OutHLSL, FString const& InDataInterfaceName) const
@@ -127,7 +130,7 @@ void UOptimusGroomDataInterface::GetHLSL(FString& OutHLSL, FString const& InData
 	};
 
 	FString TemplateFile;
-	LoadShaderSourceFile(TEXT("/Plugin/Runtime/HairStrands/Private/DeformerDataInterfaceGroom.ush"), EShaderPlatform::SP_PCD3D_SM5, &TemplateFile, nullptr);
+	LoadShaderSourceFile(TemplateFilePath, EShaderPlatform::SP_PCD3D_SM5, &TemplateFile, nullptr);
 	OutHLSL += FString::Format(*TemplateFile, TemplateArgs);
 }
 
diff --git a/Engine/Plugins/Runtime/HairStrands/Source/HairStrandsCore/Private/Deformer/DeformerDataInterfaceGroom.h b/Engine/Plugins/Runtime/HairStrands/Source/HairStrandsCore/Private/Deformer/DeformerDataInterfaceGroom.h
--- a/Engine/Plugins/Runtime/HairStrands/Source/HairStrandsCore/Private/Deformer/DeformerDataInterfaceGroom.h
+++ b/Engine/Plugins/Runtime/HairStrands/Source/HairStrandsCore/Private/Deformer/DeformerDataInterfaceGroom.h
@@ -27,11 +27,14 @@ public:
 	TCHAR const* GetClassName() const override { return TEXT("Groom"); }
 	void GetSupportedInputs(TArray<FShaderFunctionDefinition>& OutFunctions) const override;
 	void GetShaderParameters(TCHAR const* UID, FShaderParametersMetadataBuilder& InOutBuilder, FShaderParametersMetadataAllocations& InOutAllocations) const override;
-	void GetPermutations(FComputeKernelPermutationVector& OutPermutationVector) const override;
+	TCHAR const* GetShaderVirtualPath() const override;
 	void GetShaderHash(FString& InOutKey) const override;
 	void GetHLSL(FString& OutHLSL, FString const& InDataInterfaceName) const override;
 	UComputeDataProvider* CreateDataProvider(TObjectPtr<UObject> InBinding, uint64 InInputMask, uint64 InOutputMask) const override;
 	//~ End UComputeDataInterface Interface
+
+private:
+	static TCHAR const* TemplateFilePath;
 };
 
 /** Compute Framework Data Provider for reading groom. */
diff --git a/Engine/Plugins/Runtime/HairStrands/Source/HairStrandsCore/Private/Deformer/DeformerDataInterfaceGroomExec.cpp b/Engine/Plugins/Runtime/HairStrands/Source/HairStrandsCore/Private/Deformer/DeformerDataInterfaceGroomExec.cpp
--- a/Engine/Plugins/Runtime/HairStrands/Source/HairStrandsCore/Private/Deformer/DeformerDataInterfaceGroomExec.cpp
+++ b/Engine/Plugins/Runtime/HairStrands/Source/HairStrandsCore/Private/Deformer/DeformerDataInterfaceGroomExec.cpp
@@ -49,9 +49,16 @@ void UOptimusGroomExecDataInterface::GetShaderParameters(TCHAR const* UID, FShad
 	InOutBuilder.AddNestedStruct<FGroomExecDataInterfaceParameters>(UID);
 }
 
+TCHAR const* UOptimusGroomExecDataInterface::TemplateFilePath = TEXT("/Plugin/Runtime/HairStrands/Private/DeformerDataInterfaceGroomExec.ush");
+
+TCHAR const* UOptimusGroomExecDataInterface::GetShaderVirtualPath() const
+{
+	return TemplateFilePath;
+}
+
 void UOptimusGroomExecDataInterface::GetShaderHash(FString& InOutKey) const
 {
-	GetShaderFileHash(TEXT("/Plugin/Runtime/HairStrands/Private/DeformerDataInterfaceGroomExec.ush"), EShaderPlatform::SP_PCD3D_SM5).AppendString(InOutKey);
+	GetShaderFileHash(TemplateFilePath, EShaderPlatform::SP_PCD3D_SM5).AppendString(InOutKey);
 }
 
 void UOptimusGroomExecDataInterface::GetHLSL(FString& OutHLSL, FString const& InDataInterfaceName) const
@@ -62,7 +69,7 @@ void UOptimusGroomExecDataInterface::GetHLSL(FString& OutHLSL, FString const& In
 	};
 
 	FString TemplateFile;
-	LoadShaderSourceFile(TEXT("/Plugin/Runtime/HairStrands/Private/DeformerDataInterfaceGroomExec.ush"), EShaderPlatform::SP_PCD3D_SM5, &TemplateFile, nullptr);
+	LoadShaderSourceFile(TemplateFilePath, EShaderPlatform::SP_PCD3D_SM5, &TemplateFile, nullptr);
 	OutHLSL += FString::Format(*TemplateFile, TemplateArgs);
 }
 
diff --git a/Engine/Plugins/Runtime/HairStrands/Source/HairStrandsCore/Private/Deformer/DeformerDataInterfaceGroomExec.h b/Engine/Plugins/Runtime/HairStrands/Source/HairStrandsCore/Private/Deformer/DeformerDataInterfaceGroomExec.h
--- a/Engine/Plugins/Runtime/HairStrands/Source/HairStrandsCore/Private/Deformer/DeformerDataInterfaceGroomExec.h
+++ b/Engine/Plugins/Runtime/HairStrands/Source/HairStrandsCore/Private/Deformer/DeformerDataInterfaceGroomExec.h
@@ -38,6 +38,7 @@ public:
 	bool IsExecutionInterface() const override { return true; }
 	void GetSupportedInputs(TArray<FShaderFunctionDefinition>& OutFunctions) const override;
 	void GetShaderParameters(TCHAR const* UID, FShaderParametersMetadataBuilder& InOutBuilder, FShaderParametersMetadataAllocations& InOutAllocations) const override;
+	TCHAR const* GetShaderVirtualPath() const override;
 	void GetShaderHash(FString& InOutKey) const override;
 	void GetHLSL(FString& OutHLSL, FString const& InDataInterfaceName) const override;
 	UComputeDataProvider* CreateDataProvider(TObjectPtr<UObject> InBinding, uint64 InInputMask, uint64 InOutputMask) const override;
@@ -45,6 +46,9 @@ public:
 
 	UPROPERTY(EditAnywhere, Category = Execution)
 	EOptimusGroomExecDomain Domain = EOptimusGroomExecDomain::ControlPoint;
+
+private:
+	static TCHAR const* TemplateFilePath;
 };
 
 /** Compute Framework Data Provider for executing kernels over a skinned mesh domain. */
diff --git a/Engine/Plugins/Runtime/HairStrands/Source/HairStrandsCore/Private/Deformer/DeformerDataInterfaceGroomWrite.cpp b/Engine/Plugins/Runtime/HairStrands/Source/HairStrandsCore/Private/Deformer/DeformerDataInterfaceGroomWrite.cpp
--- a/Engine/Plugins/Runtime/HairStrands/Source/HairStrandsCore/Private/Deformer/DeformerDataInterfaceGroomWrite.cpp
+++ b/Engine/Plugins/Runtime/HairStrands/Source/HairStrandsCore/Private/Deformer/DeformerDataInterfaceGroomWrite.cpp
@@ -65,9 +65,16 @@ void UOptimusGroomWriteDataInterface::GetShaderParameters(TCHAR const* UID, FSha
 	InOutBuilder.AddNestedStruct<FGroomWriteDataInterfaceParameters>(UID);
 }
 
+TCHAR const* UOptimusGroomWriteDataInterface::TemplateFilePath = TEXT("/Plugin/Runtime/HairStrands/Private/DeformerDataInterfaceGroomWrite.ush");
+
+TCHAR const* UOptimusGroomWriteDataInterface::GetShaderVirtualPath() const
+{
+	return TemplateFilePath;
+}
+
 void UOptimusGroomWriteDataInterface::GetShaderHash(FString& InOutKey) const
 {
-	GetShaderFileHash(TEXT("/Plugin/Runtime/HairStrands/Private/DeformerDataInterfaceGroomWrite.ush"), EShaderPlatform::SP_PCD3D_SM5).AppendString(InOutKey);
+	GetShaderFileHash(TemplateFilePath, EShaderPlatform::SP_PCD3D_SM5).AppendString(InOutKey);
 }
 
 void UOptimusGroomWriteDataInterface::GetHLSL(FString& OutHLSL, FString const& InDataInterfaceName) const
@@ -78,7 +85,7 @@ void UOptimusGroomWriteDataInterface::GetHLSL(FString& OutHLSL, FString const& I
 	};
 
 	FString TemplateFile;
-	LoadShaderSourceFile(TEXT("/Plugin/Runtime/HairStrands/Private/DeformerDataInterfaceGroomWrite.ush"), EShaderPlatform::SP_PCD3D_SM5, &TemplateFile, nullptr);
+	LoadShaderSourceFile(TemplateFilePath, EShaderPlatform::SP_PCD3D_SM5, &TemplateFile, nullptr);
 	OutHLSL += FString::Format(*TemplateFile, TemplateArgs);
 }
 
diff --git a/Engine/Plugins/Runtime/HairStrands/Source/HairStrandsCore/Private/Deformer/DeformerDataInterfaceGroomWrite.h b/Engine/Plugins/Runtime/HairStrands/Source/HairStrandsCore/Private/Deformer/DeformerDataInterfaceGroomWrite.h
--- a/Engine/Plugins/Runtime/HairStrands/Source/HairStrandsCore/Private/Deformer/DeformerDataInterfaceGroomWrite.h
+++ b/Engine/Plugins/Runtime/HairStrands/Source/HairStrandsCore/Private/Deformer/DeformerDataInterfaceGroomWrite.h
@@ -32,10 +32,14 @@ public:
 	void GetSupportedInputs(TArray<FShaderFunctionDefinition>& OutFunctions) const override;
 	void GetSupportedOutputs(TArray<FShaderFunctionDefinition>& OutFunctions) const override;
 	void GetShaderParameters(TCHAR const* UID, FShaderParametersMetadataBuilder& InOutBuilder, FShaderParametersMetadataAllocations& InOutAllocations) const override;
+	TCHAR const* GetShaderVirtualPath() const override;
 	void GetShaderHash(FString& InOutKey) const override;
 	void GetHLSL(FString& OutHLSL, FString const& InDataInterfaceName) const override;
 	UComputeDataProvider* CreateDataProvider(TObjectPtr<UObject> InBinding, uint64 InInputMask, uint64 InOutputMask) const override;
 	//~ End UComputeDataInterface Interface
+
+private:
+	static TCHAR const* TemplateFilePath;
 };
 
 /** Compute Framework Data Provider for writing skinned mesh. */
