Fix ClumpID translation/building and debug visualization.

diff --git a/Engine/Plugins/Runtime/HairStrands/Source/HairStrandsCore/Private/GroomBuilder.cpp b/Engine/Plugins/Runtime/HairStrands/Source/HairStrandsCore/Private/GroomBuilder.cpp
--- a/Engine/Plugins/Runtime/HairStrands/Source/HairStrandsCore/Private/GroomBuilder.cpp
+++ b/Engine/Plugins/Runtime/HairStrands/Source/HairStrandsCore/Private/GroomBuilder.cpp
@@ -40,7 +40,7 @@ static FAutoConsoleVariableRef CVarHairGroupIndexBuilder_MaxVoxelResolution(TEXT
 
 FString FGroomBuilder::GetVersion()
 {
-	return TEXT("v5");
+	return TEXT("v6a");
 }
 
 namespace FHairStrandsDecimation
@@ -425,11 +425,11 @@ namespace HairStrandsBuilder
 				Packed |= (uint32(FMath::Clamp(TextureIndexUV.Y, 0.f, 31.f)) & 0x1F) << 27u;
 			}
 
-			// Per-Curve Clump ID
+			// Per-Curve Clump ID (16 bits max)
 			if (bHasClumpIDs)
 			{
-				uint8* Packed = (uint8*)AttributeClumpIDs.GetData();
-				Packed[CurveIndex] = Curves.ClumpIDs[CurveIndex];
+				uint16* Packed = (uint16*)AttributeClumpIDs.GetData();
+				Packed[CurveIndex] = uint16(FMath::Clamp(Curves.ClumpIDs[CurveIndex], 0, 65535));
 			}
 
 			// Curves
@@ -1559,7 +1559,7 @@ bool FHairDescription::HasRootUV() const
 
 bool FHairDescription::HasClumpID() const
 {
-	return StrandAttributes().GetAttributesRef<FVector2f>(HairAttribute::Strand::ClumpID).IsValid();
+	return StrandAttributes().GetAttributesRef<int>(HairAttribute::Strand::ClumpID).IsValid();
 }
 
 bool FHairDescription::HasGuideWeights() const
