Fixed a few issues on enabling GPUScene on mobile

diff --git a/Engine/Plugins/Runtime/HairStrands/Source/HairStrandsCore/Private/HairStrandsVertexFactory.cpp b/Engine/Plugins/Runtime/HairStrands/Source/HairStrandsCore/Private/HairStrandsVertexFactory.cpp
--- a/Engine/Plugins/Runtime/HairStrands/Source/HairStrandsCore/Private/HairStrandsVertexFactory.cpp
+++ b/Engine/Plugins/Runtime/HairStrands/Source/HairStrandsCore/Private/HairStrandsVertexFactory.cpp
@@ -194,6 +194,7 @@ void FHairStrandsVertexFactory::ValidateCompiledResult(const FVertexFactoryType*
 #if VF_STRANDS_SUPPORT_GPU_SCENE
 	if (Type->SupportsPrimitiveIdStream()
 		&& UseGPUScene(Platform, GetMaxSupportedFeatureLevel(Platform)) 
+		&& !IsMobilePlatform(Platform) // On mobile VS may use PrimtiveUB while GPUScene is enabled
 		&& ParameterMap.ContainsParameterAllocation(FPrimitiveUniformShaderParameters::StaticStructMetadata.GetShaderVariableName()))
 	{
 		OutErrors.AddUnique(*FString::Printf(TEXT("Shader attempted to bind the Primitive uniform buffer even though Vertex Factory %s computes a PrimitiveId per-instance.  This will break auto-instancing.  Shaders should use GetPrimitiveData(PrimitiveId).Member instead of Primitive.Member."), Type->GetName()));
