Fix incorrect loading of the bindings assets within the groom editor in certain cases.

diff --git a/Engine/Plugins/Runtime/HairStrands/Source/HairStrandsEditor/Private/GroomCustomAssetEditorToolkit.cpp b/Engine/Plugins/Runtime/HairStrands/Source/HairStrandsEditor/Private/GroomCustomAssetEditorToolkit.cpp
--- a/Engine/Plugins/Runtime/HairStrands/Source/HairStrandsEditor/Private/GroomCustomAssetEditorToolkit.cpp
+++ b/Engine/Plugins/Runtime/HairStrands/Source/HairStrandsEditor/Private/GroomCustomAssetEditorToolkit.cpp
@@ -491,6 +491,19 @@ void FGroomCustomAssetEditorToolkit::InitCustomAssetEditor(const EToolkitMode::T
 		bCreateDefaultToolbar,
 		(UObject*)InCustomAsset);
 	
+	bIsTabManagerInitialized = true;
+
+	// Initialize the content of the binding tab/
+	// This initialization happens after the tab spawning to check it the binding tab is actually active/visible. 
+	// This avoids to loading the binding assets if the tab is not visibile.
+	{
+		TSharedPtr<SDockTab> BindingTab = TabManager->FindExistingLiveTab(TabId_BindingProperties);
+		if (BindingTab && BindingTab->IsActive())
+		{
+			InitializeBindingAssetTabContent();
+		}
+	}
+
 	FProperty* P0 = FindFProperty<FProperty>(GroomAsset->GetClass(), GET_MEMBER_NAME_CHECKED(UGroomAsset, HairGroupsInterpolation));
 	FProperty* P1 = FindFProperty<FProperty>(GroomAsset->GetClass(), GET_MEMBER_NAME_CHECKED(UGroomAsset, HairGroupsRendering));
 	FProperty* P2 = FindFProperty<FProperty>(GroomAsset->GetClass(), GET_MEMBER_NAME_CHECKED(UGroomAsset, HairGroupsPhysics));
@@ -835,6 +848,18 @@ TSharedRef<SDockTab> FGroomCustomAssetEditorToolkit::SpawnTab_PreviewGroomCompon
 		];
 }
 
+void FGroomCustomAssetEditorToolkit::InitializeBindingAssetTabContent()
+{
+	UGroomAsset* LocalGroomAsset = GetCustomAsset();
+	if (GroomBindingAssetList == nullptr && LocalGroomAsset)
+	{
+		GroomBindingAssetList = NewObject<UGroomBindingAssetList>(GetTransientPackage(), NAME_None, RF_Transient);
+		ListAllBindingAssets(LocalGroomAsset, GroomBindingAssetList);
+		DetailView_BindingProperties->SetObject(Cast<UObject>(GroomBindingAssetList));
+		DetailView_BindingProperties->ForceRefresh();
+	}
+}
+
 TSharedRef<SDockTab> FGroomCustomAssetEditorToolkit::SpawnTab_BindingProperties(const FSpawnTabArgs& Args)
 {
 	check(Args.GetTabId() == TabId_BindingProperties);
@@ -846,16 +871,17 @@ TSharedRef<SDockTab> FGroomCustomAssetEditorToolkit::SpawnTab_BindingProperties(
 			DetailView_BindingProperties.ToSharedRef()
 		];
 
-	DockTab->SetOnTabActivated(SDockTab::FOnTabActivatedCallback::CreateLambda([this](TSharedRef<SDockTab> Input, ETabActivationCause)
+	if (bIsTabManagerInitialized)
 	{
-		UGroomAsset* LocalGroomAsset = GetCustomAsset();
-		if (GroomBindingAssetList == nullptr && LocalGroomAsset)
+		InitializeBindingAssetTabContent();
+	}
+	else
+	{	
+		DockTab->SetOnTabActivated(SDockTab::FOnTabActivatedCallback::CreateLambda([this](TSharedRef<SDockTab> Input, ETabActivationCause)
 		{
-			GroomBindingAssetList = NewObject<UGroomBindingAssetList>(GetTransientPackage(), NAME_None, RF_Transient);
-			ListAllBindingAssets(LocalGroomAsset, GroomBindingAssetList);
-			DetailView_BindingProperties->ForceRefresh();
-		}
-	}));
+			InitializeBindingAssetTabContent();
+		}));
+	}
 
 	return DockTab;
 }
diff --git a/Engine/Plugins/Runtime/HairStrands/Source/HairStrandsEditor/Public/GroomCustomAssetEditorToolkit.h b/Engine/Plugins/Runtime/HairStrands/Source/HairStrandsEditor/Public/GroomCustomAssetEditorToolkit.h
--- a/Engine/Plugins/Runtime/HairStrands/Source/HairStrandsEditor/Public/GroomCustomAssetEditorToolkit.h
+++ b/Engine/Plugins/Runtime/HairStrands/Source/HairStrandsEditor/Public/GroomCustomAssetEditorToolkit.h
@@ -123,6 +123,9 @@ private:
 	// create the custom components we need
 	void InitPreviewComponents();
 
+	// Initialized the content of the binding tab
+	void InitializeBindingAssetTabContent();
+
 	// return a pointer to the groom preview component
 	UGroomComponent*		GetPreview_GroomComponent() const;
 	USkeletalMeshComponent*	GetPreview_SkeletalMeshComponent() const;
@@ -168,6 +171,7 @@ private:
 	static const FName TabId_BindingProperties;
 
 
+	bool bIsTabManagerInitialized = false;
 	int32 ActiveGroomBindingIndex = -1;
 	FDelegateHandle PropertyListenDelegate;
 	TWeakObjectPtr<UGroomAsset> GroomAsset;
