* Change groom binding UI to disable groom transfer when the groom asset does not have RootUV data embedded.
* Change optional attribute handling to be enum based, which ease extensibility in several places.

diff --git a/Engine/Plugins/Runtime/HairStrands/Source/HairStrandsCore/Private/GroomBindingBuilder.cpp b/Engine/Plugins/Runtime/HairStrands/Source/HairStrandsCore/Private/GroomBindingBuilder.cpp
--- a/Engine/Plugins/Runtime/HairStrands/Source/HairStrandsCore/Private/GroomBindingBuilder.cpp
+++ b/Engine/Plugins/Runtime/HairStrands/Source/HairStrandsCore/Private/GroomBindingBuilder.cpp
@@ -2079,6 +2079,7 @@ static bool InternalBuildBinding_CPU(UGroomBindingAsset* BindingAsset, bool bIni
 			++GroupIndex;
 		}
 
+		// Transfer requires root UV embedded into the groom asset. It is not possible to read safely hair description here to extra this data.
 		const bool bNeedTransferPosition = SourceMeshData->IsValid();
 
 		// Create mapping between the source & target using their UV
diff --git a/Engine/Plugins/Runtime/HairStrands/Source/HairStrandsCore/Private/GroomBuilder.cpp b/Engine/Plugins/Runtime/HairStrands/Source/HairStrandsCore/Private/GroomBuilder.cpp
--- a/Engine/Plugins/Runtime/HairStrands/Source/HairStrandsCore/Private/GroomBuilder.cpp
+++ b/Engine/Plugins/Runtime/HairStrands/Source/HairStrandsCore/Private/GroomBuilder.cpp
@@ -1551,44 +1551,6 @@ FORCEINLINE FIntVector VectorToIntVector(const FVector& Index)
 	return FIntVector(Index.X, Index.Y, Index.Z);
 }
 
-// These accessors are defined here to ease mirror logic with BuildHairDescriptionGroups, if any type/attributes changes
-bool FHairDescription::HasRootUV() const
-{
-	return StrandAttributes().GetAttributesRef<FVector2f>(HairAttribute::Strand::RootUV).IsValid();
-}
-
-bool FHairDescription::HasClumpID() const
-{
-	return StrandAttributes().GetAttributesRef<int>(HairAttribute::Strand::ClumpID).IsValid();
-}
-
-bool FHairDescription::HasGuideWeights() const
-{
-	return
-		// Single
-		(StrandAttributes().GetAttributesRef<int>(HairAttribute::Strand::ClosestGuides).IsValid() && 
-		 StrandAttributes().GetAttributesRef<float>(HairAttribute::Strand::GuideWeights).IsValid())
-		||
-		// Triplet
-		(StrandAttributes().GetAttributesRef<FVector3f>(HairAttribute::Strand::ClosestGuides).IsValid() &&
-		 StrandAttributes().GetAttributesRef<FVector3f>(HairAttribute::Strand::GuideWeights).IsValid());
-}
-
-bool FHairDescription::HasColorAttributes() const
-{
-	return VertexAttributes().GetAttributesRef<FVector3f>(HairAttribute::Vertex::Color).IsValid();
-}
-
-bool FHairDescription::HasRoughnessAttributes() const
-{
-	return VertexAttributes().GetAttributesRef<float>(HairAttribute::Vertex::Roughness).IsValid();
-}
-
-bool FHairDescription::HasAOAttributes() const
-{
-	return VertexAttributes().GetAttributesRef<float>(HairAttribute::Vertex::AO).IsValid();
-}
-
 bool FGroomBuilder::BuildHairDescriptionGroups(const FHairDescription& HairDescription, FHairDescriptionGroups& Out)
 {
 	TRACE_CPUPROFILER_EVENT_SCOPE(FGroomBuilder::BuildHairDescriptionGroups);
@@ -1637,11 +1599,18 @@ bool FGroomBuilder::BuildHairDescriptionGroups(const FHairDescription& HairDescr
 	const bool bHasAOAttribute = VertexAO.IsValid();
 	const bool bHasClumpIDs = ClumpIDs.IsValid();
 
+	// Sanity check
+	check(bHasBaseColorAttribute == HairDescription.HasAttribute(EHairAttribute::Color));
+	check(bHasRoughnessAttribute == HairDescription.HasAttribute(EHairAttribute::Roughness));
+	check(bHasAOAttribute == HairDescription.HasAttribute(EHairAttribute::AO));
+	check(bHasClumpIDs == HairDescription.HasAttribute(EHairAttribute::ClumpID)); 
+
 	TVertexAttributesConstRef<float> VertexWidths = HairDescription.VertexAttributes().GetAttributesRef<float>(HairAttribute::Vertex::Width);
 	TStrandAttributesConstRef<float> StrandWidths = HairDescription.StrandAttributes().GetAttributesRef<float>(HairAttribute::Strand::Width);
 
 	TStrandAttributesConstRef<FVector2f> StrandRootUV = HairDescription.StrandAttributes().GetAttributesRef<FVector2f>(HairAttribute::Strand::RootUV);
 	const bool bHasUVData = StrandRootUV.IsValid();
+	check (bHasUVData == HairDescription.HasAttribute(EHairAttribute::RootUV)); // Sanity check
 
 	TStrandAttributesConstRef<int> StrandGuides = HairDescription.StrandAttributes().GetAttributesRef<int>(HairAttribute::Strand::Guide);
 	TStrandAttributesConstRef<int> GroupIDs = HairDescription.StrandAttributes().GetAttributesRef<int>(HairAttribute::Strand::GroupID);
@@ -1663,6 +1632,7 @@ bool FGroomBuilder::BuildHairDescriptionGroups(const FHairDescription& HairDescr
 	const bool bPrecomputedWeight1 = ClosestGuide.IsValid()  && GuideWeight.IsValid();
 	const bool bPrecomputedWeight3 = ClosestGuides.IsValid() && GuideWeights.IsValid();
 	const bool bCanUseClosestGuidesAndWeights = bImportGuides && StrandIDs.IsValid() && (bPrecomputedWeight1 || bPrecomputedWeight3);
+	check(bCanUseClosestGuidesAndWeights == HairDescription.HasAttribute(EHairAttribute::PrecomputedGuideWeights)); // Sanity check
 
 	auto FindOrAdd = [&Out](int32 GroupID, FName GroupName) -> FHairDescriptionGroup&
 	{
@@ -1849,8 +1819,16 @@ bool FGroomBuilder::BuildHairDescriptionGroups(const FHairDescription& HairDescr
 		for (FHairDescriptionGroup& Group : Out.HairGroups)
 		{
 			Group.Info.GroupID = GroupIndex++;
-			Group.bCanUseClosestGuidesAndWeights = bCanUseClosestGuidesAndWeights;
-			Group.bHasUVData = bHasUVData;
+
+			// Propagate valid optional attributes
+			for (uint32 AttributeIt=0; AttributeIt< uint32(EHairAttribute::Count); ++AttributeIt)
+			{
+				const EHairAttribute HairAttribute = (EHairAttribute)AttributeIt;
+				if (HairDescription.HasAttribute(HairAttribute)) 
+				{ 
+					SetHairAttribute(Group.Attributes, HairAttribute); 
+				}
+			}
 		}
 	}
 
@@ -1928,7 +1906,7 @@ void FGroomBuilder::BuildData(
 		{
 			OutRen = InHairDescriptionGroup.Strands;
 
-			HairStrandsBuilder::BuildInternalData(OutRen, !InHairDescriptionGroup.bHasUVData);
+			HairStrandsBuilder::BuildInternalData(OutRen, !HasHairAttribute(InHairDescriptionGroup.Attributes, EHairAttribute::RootUV));
 
 			// Decimate
 			if (CurveDecimation < 1 || VertexDecimation < 1 || bCurveReordering)
diff --git a/Engine/Plugins/Runtime/HairStrands/Source/HairStrandsCore/Private/HairAttributes.cpp b/Engine/Plugins/Runtime/HairStrands/Source/HairStrandsCore/Private/HairAttributes.cpp
--- a/Engine/Plugins/Runtime/HairStrands/Source/HairStrandsCore/Private/HairAttributes.cpp
+++ b/Engine/Plugins/Runtime/HairStrands/Source/HairStrandsCore/Private/HairAttributes.cpp
@@ -35,3 +35,14 @@ namespace HairAttribute
 	const FName Groom::Tool("groom_tool");
 	const FName Groom::Properties("groom_properties");
 }
+
+static_assert(uint32(EHairAttribute::Count) < 32u);
+bool HasHairAttribute(uint32 In, EHairAttribute InAttribute)
+{ 
+	return (In & (1u << uint32(InAttribute))) != 0; 
+}
+
+void SetHairAttribute(uint32& Out, EHairAttribute InAttribute) 
+{ 
+	Out |= (1u << uint32(InAttribute)); 
+}
\ No newline at end of file
diff --git a/Engine/Plugins/Runtime/HairStrands/Source/HairStrandsCore/Private/HairDescription.cpp b/Engine/Plugins/Runtime/HairStrands/Source/HairStrandsCore/Private/HairDescription.cpp
--- a/Engine/Plugins/Runtime/HairStrands/Source/HairStrandsCore/Private/HairDescription.cpp
+++ b/Engine/Plugins/Runtime/HairStrands/Source/HairStrandsCore/Private/HairDescription.cpp
@@ -66,6 +66,41 @@ bool FHairDescription::IsValid() const
 	return (NumStrands > 0) && (NumVertices > 0);
 }
 
+bool FHairDescription::HasAttribute(EHairAttribute InAttribute) const
+{
+	// If a new optional attribute is added, please add its UI/text description here
+	static_assert(uint32(EHairAttribute::Count) == 6);
+
+	switch (InAttribute)
+	{
+	case EHairAttribute::RootUV:		return StrandAttributes().GetAttributesRef<FVector2f>(HairAttribute::Strand::RootUV).IsValid();
+	case EHairAttribute::ClumpID:		return StrandAttributes().GetAttributesRef<int>(HairAttribute::Strand::ClumpID).IsValid();
+	case EHairAttribute::Color:			return VertexAttributes().GetAttributesRef<FVector3f>(HairAttribute::Vertex::Color).IsValid();
+	case EHairAttribute::Roughness:		return VertexAttributes().GetAttributesRef<float>(HairAttribute::Vertex::Roughness).IsValid();
+	case EHairAttribute::AO:			return VertexAttributes().GetAttributesRef<float>(HairAttribute::Vertex::AO).IsValid();
+	case EHairAttribute::PrecomputedGuideWeights:
+	{
+		return
+			// Single
+			(StrandAttributes().GetAttributesRef<int>(HairAttribute::Strand::ClosestGuides).IsValid() &&
+				StrandAttributes().GetAttributesRef<float>(HairAttribute::Strand::GuideWeights).IsValid())
+			||
+			// Triplet
+			(StrandAttributes().GetAttributesRef<FVector3f>(HairAttribute::Strand::ClosestGuides).IsValid() &&
+				StrandAttributes().GetAttributesRef<FVector3f>(HairAttribute::Strand::GuideWeights).IsValid());
+	}
+	}
+	return false;
+}
+
+// Deprecated named accessors. Use HasAttribute() instead
+// These accessors are defined here to ease mirror logic with BuildHairDescriptionGroups, if any type/attributes changes
+bool FHairDescription::HasRootUV() const { return HasAttribute(EHairAttribute::RootUV); }
+bool FHairDescription::HasClumpID() const { return HasAttribute(EHairAttribute::ClumpID); }
+bool FHairDescription::HasGuideWeights() const { return HasAttribute(EHairAttribute::PrecomputedGuideWeights); }
+bool FHairDescription::HasColorAttributes() const { return HasAttribute(EHairAttribute::Color); }
+bool FHairDescription::HasRoughnessAttributes() const { return HasAttribute(EHairAttribute::Roughness); }
+bool FHairDescription::HasAOAttributes() const { return HasAttribute(EHairAttribute::AO); }
 
 void FHairDescription::Serialize(FArchive& Ar)
 {
diff --git a/Engine/Plugins/Runtime/HairStrands/Source/HairStrandsCore/Public/GroomAsset.h b/Engine/Plugins/Runtime/HairStrands/Source/HairStrandsCore/Public/GroomAsset.h
--- a/Engine/Plugins/Runtime/HairStrands/Source/HairStrandsCore/Public/GroomAsset.h
+++ b/Engine/Plugins/Runtime/HairStrands/Source/HairStrandsCore/Public/GroomAsset.h
@@ -369,9 +369,7 @@ struct FHairDescriptionGroup
 	FHairGroupInfo		Info;
 	FHairStrandsDatas	Strands;
 	FHairStrandsDatas	Guides;
-
-	bool  bCanUseClosestGuidesAndWeights = false;
-	bool  bHasUVData = false;
+	uint32				Attributes = 0;
 };
 
 struct FHairDescriptionGroups
diff --git a/Engine/Plugins/Runtime/HairStrands/Source/HairStrandsCore/Public/GroomCreateBindingOptions.h b/Engine/Plugins/Runtime/HairStrands/Source/HairStrandsCore/Public/GroomCreateBindingOptions.h
--- a/Engine/Plugins/Runtime/HairStrands/Source/HairStrandsCore/Public/GroomCreateBindingOptions.h
+++ b/Engine/Plugins/Runtime/HairStrands/Source/HairStrandsCore/Public/GroomCreateBindingOptions.h
@@ -12,6 +12,10 @@ class HAIRSTRANDSCORE_API UGroomCreateBindingOptions : public UObject
 	GENERATED_UCLASS_BODY()
 
 public:
+	/** Groom asset */
+	UPROPERTY()
+	TWeakObjectPtr<UGroomAsset> GroomAsset;
+
 	/** Type of mesh to create groom binding for */
 	UPROPERTY(EditAnywhere, BlueprintReadWrite, Category = BuildSettings)
 	EGroomBindingMeshType GroomBindingType;
diff --git a/Engine/Plugins/Runtime/HairStrands/Source/HairStrandsCore/Public/GroomImportOptions.h b/Engine/Plugins/Runtime/HairStrands/Source/HairStrandsCore/Public/GroomImportOptions.h
--- a/Engine/Plugins/Runtime/HairStrands/Source/HairStrandsCore/Public/GroomImportOptions.h
+++ b/Engine/Plugins/Runtime/HairStrands/Source/HairStrandsCore/Public/GroomImportOptions.h
@@ -31,12 +31,7 @@ struct HAIRSTRANDSCORE_API FGroomHairGroupPreview
 	: GroupID(0)
 	, CurveCount(0)
 	, GuideCount(0)
-	, bHasRootUV(false)
-	, bHasClumpID(false)
-	, bHasColor(false)
-	, bHasRoughness(false)
-	, bHasAO(false)
-	, bHasPrecomputedWeights(false)
+	, Attributes(0)
 	, InterpolationSettings()
 	{}
 
@@ -52,23 +47,8 @@ struct HAIRSTRANDSCORE_API FGroomHairGroupPreview
 	UPROPERTY(VisibleAnywhere, BlueprintReadWrite, Category = Preview)
 	int32 GuideCount;
 
-	UPROPERTY(VisibleAnywhere, BlueprintReadWrite, Category = Preview)
-	bool bHasRootUV;
-
-	UPROPERTY(VisibleAnywhere, BlueprintReadWrite, Category = Preview)
-	bool bHasClumpID;
-
-	UPROPERTY(VisibleAnywhere, BlueprintReadWrite, Category = Preview)
-	bool bHasColor;
-
-	UPROPERTY(VisibleAnywhere, BlueprintReadWrite, Category = Preview)
-	bool bHasRoughness;
-
-	UPROPERTY(VisibleAnywhere, BlueprintReadWrite, Category = Preview)
-	bool bHasAO;
-
-	UPROPERTY(VisibleAnywhere, BlueprintReadWrite, Category = Preview)
-	bool bHasPrecomputedWeights;
+	UPROPERTY()
+	uint32 Attributes;
 
 	UPROPERTY(EditAnywhere, BlueprintReadWrite, Category = Preview)
 	FHairGroupsInterpolation InterpolationSettings;
diff --git a/Engine/Plugins/Runtime/HairStrands/Source/HairStrandsCore/Public/HairAttributes.h b/Engine/Plugins/Runtime/HairStrands/Source/HairStrandsCore/Public/HairAttributes.h
--- a/Engine/Plugins/Runtime/HairStrands/Source/HairStrandsCore/Public/HairAttributes.h
+++ b/Engine/Plugins/Runtime/HairStrands/Source/HairStrandsCore/Public/HairAttributes.h
@@ -46,3 +46,17 @@ namespace HairAttribute
 		extern HAIRSTRANDSCORE_API const FName Width;			// float
 	}
 }
+
+enum class EHairAttribute : uint8
+{
+	RootUV,
+	ClumpID,
+	PrecomputedGuideWeights,
+	Color,
+	Roughness,
+	AO,
+	Count
+};
+
+HAIRSTRANDSCORE_API bool HasHairAttribute(uint32 In, EHairAttribute InAttribute);
+HAIRSTRANDSCORE_API void SetHairAttribute(uint32& Out, EHairAttribute InAttribute);
\ No newline at end of file
diff --git a/Engine/Plugins/Runtime/HairStrands/Source/HairStrandsCore/Public/HairDescription.h b/Engine/Plugins/Runtime/HairStrands/Source/HairStrandsCore/Public/HairDescription.h
--- a/Engine/Plugins/Runtime/HairStrands/Source/HairStrandsCore/Public/HairDescription.h
+++ b/Engine/Plugins/Runtime/HairStrands/Source/HairStrandsCore/Public/HairDescription.h
@@ -94,12 +94,19 @@ public:
 	void Serialize(FArchive& Ar);
 
 	// #ueent_todo: Expose "remove" operations to allow editing HairDescription
+	bool HasAttribute(EHairAttribute InAttribute) const;
 
+	UE_DEPRECATED(5.3, "HasRootUV accessor is deprecated. Use HasAttribute instead")
 	bool HasRootUV() const;
+	UE_DEPRECATED(5.3, "HasClumpID accessor is deprecated. Use HasAttribute instead")
 	bool HasClumpID() const;
+	UE_DEPRECATED(5.3, "HasGuideWeights accessor is deprecated. Use HasAttribute instead")
 	bool HasGuideWeights() const;
+	UE_DEPRECATED(5.3, "HasColorAttributes accessor is deprecated. Use HasAttribute instead")
 	bool HasColorAttributes() const;
+	UE_DEPRECATED(5.3, "HasRoughnessAttributes accessor is deprecated. Use HasAttribute instead")
 	bool HasRoughnessAttributes() const;
+	UE_DEPRECATED(5.3, "HasAOAttributes accessor is deprecated. Use HasAttribute instead")
 	bool HasAOAttributes() const;
 
 private:
diff --git a/Engine/Plugins/Runtime/HairStrands/Source/HairStrandsEditor/Private/GroomActions.cpp b/Engine/Plugins/Runtime/HairStrands/Source/HairStrandsEditor/Private/GroomActions.cpp
--- a/Engine/Plugins/Runtime/HairStrands/Source/HairStrandsEditor/Private/GroomActions.cpp
+++ b/Engine/Plugins/Runtime/HairStrands/Source/HairStrandsEditor/Private/GroomActions.cpp
@@ -176,7 +176,11 @@ void FGroomActions::ExecuteRebuild(TArray<TWeakObjectPtr<UGroomAsset>> Objects)
 			
 				const uint32 GroupCount = GroomAsset->GetNumHairGroups();
 				UGroomHairGroupsPreview* GroupsPreview = NewObject<UGroomHairGroupsPreview>();
-				{					
+				{
+					FHairDescription HairDescription = GroomAsset->GetHairDescription();
+					FHairDescriptionGroups HairDescriptionGroups;
+					FGroomBuilder::BuildHairDescriptionGroups(HairDescription, HairDescriptionGroups);
+
 					for (uint32 GroupIndex = 0; GroupIndex < GroupCount; GroupIndex++)
 					{
 						FGroomHairGroupPreview& OutGroup = GroupsPreview->Groups.AddDefaulted_GetRef();
@@ -184,6 +188,7 @@ void FGroomActions::ExecuteRebuild(TArray<TWeakObjectPtr<UGroomAsset>> Objects)
 						OutGroup.GroupName	= GroomAsset->HairGroupsInfo[GroupIndex].GroupName;
 						OutGroup.CurveCount = GroomAsset->HairGroupsData[GroupIndex].Strands.BulkData.GetNumCurves();
 						OutGroup.GuideCount = GroomAsset->HairGroupsData[GroupIndex].Guides.BulkData.GetNumCurves();
+						OutGroup.Attributes = HairDescriptionGroups.HairGroups[GroupIndex].Attributes;
 						OutGroup.InterpolationSettings = GroomAsset->HairGroupsInterpolation[GroupIndex];
 						OutGroup.InterpolationSettings.RiggingSettings.bCanEditRigging = true;
 					}
@@ -249,6 +254,7 @@ void FGroomActions::ExecuteCreateBindingAsset(TArray<TWeakObjectPtr<UGroomAsset>
 		{
 			// Duplicate the options to prevent dirtying the asset when they are modified but the rebuild is cancelled
 			UGroomCreateBindingOptions* CurrentOptions = NewObject<UGroomCreateBindingOptions>();
+			CurrentOptions->GroomAsset = GroomAsset;
 			TSharedPtr<SGroomCreateBindingOptionsWindow> GroomOptionWindow = SGroomCreateBindingOptionsWindow::DisplayCreateBindingOptions(CurrentOptions);
 
 			if (!GroomOptionWindow->ShouldCreate())
diff --git a/Engine/Plugins/Runtime/HairStrands/Source/HairStrandsEditor/Private/GroomCreateBindingOptionsWindow.cpp b/Engine/Plugins/Runtime/HairStrands/Source/HairStrandsEditor/Private/GroomCreateBindingOptionsWindow.cpp
--- a/Engine/Plugins/Runtime/HairStrands/Source/HairStrandsEditor/Private/GroomCreateBindingOptionsWindow.cpp
+++ b/Engine/Plugins/Runtime/HairStrands/Source/HairStrandsEditor/Private/GroomCreateBindingOptionsWindow.cpp
@@ -19,6 +19,7 @@
 #include "Widgets/Text/STextBlock.h"
 #include "Editor.h"
 #include "GroomBindingAsset.h"
+#include "GroomAsset.h"
 
 #define LOCTEXT_NAMESPACE "GroomCreateBindingOptionsWindow"
 
@@ -27,6 +28,36 @@ void SGroomCreateBindingOptionsWindow::Construct(const FArguments& InArgs)
 	BindingOptions = InArgs._BindingOptions;
 	WidgetWindow = InArgs._WidgetWindow;
 
+	const FSlateFontInfo AttributeFont = FAppStyle::GetFontStyle("CurveEd.InfoFont");
+	const FSlateFontInfo AttributeResultFont = FAppStyle::GetFontStyle("CurveEd.InfoFont");
+	const FLinearColor AttributeColor(0.0f, 0.80f, 0.0f);
+	const FText TrueText = LOCTEXT("GroomBindingOptionsWindow_AttributeTrue", "True");
+	const FText FalseText = LOCTEXT("GroomBindingOptionsWindow_AttributeFalse", "False");
+
+	bool bHasRootUV = false;
+	FText HasRootUVText = FalseText;
+	FLinearColor ValidationColor(1, 1, 1);
+
+	// Check if UV transfer are possible
+	BindingOptions->GroomAsset->ConditionalPostLoad();
+	if (UGroomAsset* GroomAsset = BindingOptions->GroomAsset.Get())
+	{
+		const FHairDescription HairDescription = GroomAsset->GetHairDescription();
+		bHasRootUV = HairDescription.HasAttribute(EHairAttribute::RootUV);
+	}
+
+	if (bHasRootUV)
+	{
+		HasRootUVText = TrueText;
+		if (FProperty* TransferProperty = FindFProperty<FProperty>(BindingOptions->GetClass(), GET_MEMBER_NAME_CHECKED(UGroomCreateBindingOptions, SourceSkeletalMesh))) { TransferProperty->SetMetaData(TEXT("Category"), TEXT("Conversion")); }
+		if (FProperty* TransferProperty = FindFProperty<FProperty>(BindingOptions->GetClass(), GET_MEMBER_NAME_CHECKED(UGroomCreateBindingOptions, MatchingSection))) { TransferProperty->SetMetaData(TEXT("Category"), TEXT("Conversion")); }
+	}
+	else
+	{
+		if (FProperty* TransferProperty = FindFProperty<FProperty>(BindingOptions->GetClass(), GET_MEMBER_NAME_CHECKED(UGroomCreateBindingOptions, SourceSkeletalMesh))) { TransferProperty->SetMetaData(TEXT("Category"), TEXT("Hidden")); }
+		if (FProperty* TransferProperty = FindFProperty<FProperty>(BindingOptions->GetClass(), GET_MEMBER_NAME_CHECKED(UGroomCreateBindingOptions, MatchingSection)))   { TransferProperty->SetMetaData(TEXT("Category"), TEXT("Hidden")); }
+	}
+
 	FPropertyEditorModule& PropertyEditorModule = FModuleManager::GetModuleChecked<FPropertyEditorModule>("PropertyEditor");
 	FDetailsViewArgs DetailsViewArgs;
 	DetailsViewArgs.bAllowSearch = false;
@@ -38,6 +69,36 @@ void SGroomCreateBindingOptionsWindow::Construct(const FArguments& InArgs)
 	[
 		SNew(SVerticalBox)
 
+		// Root UV
+		+ SVerticalBox::Slot()
+		.AutoHeight()
+		.Padding(2)
+		[
+			SNew(SBorder)
+			.Padding(FMargin(3))
+			.BorderImage(FAppStyle::GetBrush("ToolPanel.GroupBorder"))
+			[
+				SNew(SHorizontalBox)
+				+ SHorizontalBox::Slot()
+				.AutoWidth()
+				[
+					SNew(STextBlock)
+					.Font(AttributeFont)
+					.Text(LOCTEXT("GroomOptionsWindow_HasRootUV", "Root UV: "))
+				]
+				+ SHorizontalBox::Slot()
+				.Padding(5, 0, 0, 0)
+				.AutoWidth()
+				.VAlign(VAlign_Center)
+				[
+					SNew(STextBlock)
+					.Font(AttributeResultFont)
+					.Text(HasRootUVText)
+					.ColorAndOpacity(AttributeColor)
+				]
+			]
+		]
+
 		+ SVerticalBox::Slot()
 		.AutoHeight()
 		.Padding(2)
diff --git a/Engine/Plugins/Runtime/HairStrands/Source/HairStrandsEditor/Private/GroomImportOptionsWindow.cpp b/Engine/Plugins/Runtime/HairStrands/Source/HairStrandsEditor/Private/GroomImportOptionsWindow.cpp
--- a/Engine/Plugins/Runtime/HairStrands/Source/HairStrandsEditor/Private/GroomImportOptionsWindow.cpp
+++ b/Engine/Plugins/Runtime/HairStrands/Source/HairStrandsEditor/Private/GroomImportOptionsWindow.cpp
@@ -183,6 +183,51 @@ FSlateColor SGroomImportOptionsWindow::GetStatusColor() const
 	}
 }
 
+static void AddAttribute(SVerticalBox::FScopedWidgetSlotArguments& Slot, FText AttributeLegend)
+{
+	const FLinearColor AttributeColor(0.72f, 0.72f, 0.20f);
+	const FSlateFontInfo AttributeFont = FAppStyle::GetFontStyle("CurveEd.InfoFont");
+	const FSlateFontInfo AttributeResultFont = FAppStyle::GetFontStyle("CurveEd.InfoFont");
+
+	Slot
+	.AutoHeight()
+	.Padding(2)
+	[
+		SNew(SBorder)
+		.Padding(FMargin(3))
+		.BorderImage(FAppStyle::GetBrush("ToolPanel.GroupBorder"))
+		[
+			SNew(SHorizontalBox)
+			+ SHorizontalBox::Slot()
+			.AutoWidth()
+			.Padding(10, 0, 0, 0)
+			[
+				SNew(STextBlock)
+				.Font(AttributeFont)
+				.Text(AttributeLegend)
+				.ColorAndOpacity(AttributeColor)
+			]				
+		]
+	];
+}
+
+FText GetHairAttributeText(EHairAttribute In)
+{
+	// If a new optional attribute is added, please add its UI/text description here
+	static_assert(uint32(EHairAttribute::Count) == 6);
+
+	switch (In)
+	{
+	case EHairAttribute::RootUV:					return LOCTEXT("GroomOptionsWindow_HasRootUV", "Root UV");
+	case EHairAttribute::ClumpID:					return LOCTEXT("GroomOptionsWindow_HasRootUV", "Clump ID");
+	case EHairAttribute::PrecomputedGuideWeights:	return LOCTEXT("GroomOptionsWindow_HasPercomputedGuideWeights", "Pre-Computed Guide Weights");
+	case EHairAttribute::Color:						return LOCTEXT("GroomOptionsWindow_HasColor", "Has Color");
+	case EHairAttribute::Roughness:					return LOCTEXT("GroomOptionsWindow_HasRoughness", "Roughness");
+	case EHairAttribute::AO:						return LOCTEXT("GroomOptionsWindow_HasAO", "AO");
+	}
+	return FText::GetEmpty();
+}
+
 void SGroomImportOptionsWindow::Construct(const FArguments& InArgs)
 {
 	ImportOptions = InArgs._ImportOptions;
@@ -207,36 +252,22 @@ void SGroomImportOptionsWindow::Construct(const FArguments& InArgs)
 	CurrentStatus = EHairDescriptionStatus::Unset;
 	UpdateStatus(GroupsPreview);
 
-	const FSlateFontInfo AttributeFont = FAppStyle::GetFontStyle("CurveEd.InfoFont");
-	const FSlateFontInfo AttributeResultFont = FAppStyle::GetFontStyle("CurveEd.InfoFont");
-	const FLinearColor AttributeColor(0.80f, 0.80f, 0.80f);
-	const FText TrueText  = LOCTEXT("GroomOptionsWindow_AttributeTrue", "True");
-	const FText FalseText = LOCTEXT("GroomOptionsWindow_AttributeFalse", "False");
-
-	FText HasRootUVText = FalseText;
-	FText HasClumpIDText = FalseText;
-	FText HasColorText = FalseText;
-	FText HasRoughnessText = FalseText;
-	FText HasAOText = FalseText;
-	FText HasGuideWeightsText = FalseText;
-
-	if (GroupsPreview)
+	// Aggregate attributes from all groups (ideally we should display each group attribute separately, to check if one groom is not missing data)
+	uint32 Attributes = 0;
+	for (const FGroomHairGroupPreview& Group : GroupsPreview->Groups)
 	{
-		for (const FGroomHairGroupPreview& Group : GroupsPreview->Groups)
-		{
-			if (Group.bHasRootUV)				{ HasRootUVText = TrueText; }
-			if (Group.bHasClumpID)				{ HasClumpIDText = TrueText; }
-			if (Group.bHasColor)				{ HasColorText = TrueText; }
-			if (Group.bHasRoughness)			{ HasRoughnessText = TrueText; }
-			if (Group.bHasAO)					{ HasAOText = TrueText; }
-			if (Group.bHasPrecomputedWeights)	{ HasGuideWeightsText = TrueText; }
-		}
+		Attributes |= Group.Attributes;
 	}
 
-	this->ChildSlot
-	[
-		SNew(SVerticalBox)
+	FText bHasAttributeText = LOCTEXT("GroomOptionsWindow_HasAttributeNone", "None");
+	FLinearColor bHasAttributeColor = FLinearColor(0.80f, 0, 0, 1);
+	if (Attributes != 0)
+	{
+		bHasAttributeText = LOCTEXT("GroomOptionsWindow_HasAttributeValid", "Valid");
+		bHasAttributeColor = FLinearColor(0, 0.80f, 0, 1);
+	}
 
+	auto VerticalSlot = SNew(SVerticalBox)
 		+ SVerticalBox::Slot()
 		.AutoHeight()
 		.Padding(2)
@@ -294,128 +325,9 @@ void SGroomImportOptionsWindow::Construct(const FArguments& InArgs)
 			]
 		]
 
-		// Root UV
-		+ SVerticalBox::Slot()
-		.AutoHeight()
-		.Padding(2)
-		[
-			SNew(SBorder)
-			.Padding(FMargin(3))
-			.BorderImage(FAppStyle::GetBrush("ToolPanel.GroupBorder"))
-			[
-				SNew(SHorizontalBox)
-				+ SHorizontalBox::Slot()
-				.AutoWidth()
-				[
-					SNew(STextBlock)
-					.Font(AttributeFont)
-					.Text(LOCTEXT("GroomOptionsWindow_HasRootUV", "Has Root UV: "))
-				]
-				+ SHorizontalBox::Slot()
-				.Padding(5, 0, 0, 0)
-				.AutoWidth()
-				.VAlign(VAlign_Center)
-				[
-					SNew(STextBlock)
-					.Font(AttributeResultFont)
-					.Text(HasRootUVText)
-					.ColorAndOpacity(AttributeColor)
-				]
-			]
-		]
-
-		// Clump ID
-		+ SVerticalBox::Slot()
-		.AutoHeight()
-		.Padding(2)
-		[
-			SNew(SBorder)
-			.Padding(FMargin(3))
-			.BorderImage(FAppStyle::GetBrush("ToolPanel.GroupBorder"))
-			[
-				SNew(SHorizontalBox)
-				+ SHorizontalBox::Slot()
-				.AutoWidth()
-				[
-					SNew(STextBlock)
-					.Font(AttributeFont)
-					.Text(LOCTEXT("GroomOptionsWindow_HasClumpID", "Has Clump ID: "))
-				]
-				+ SHorizontalBox::Slot()
-				.Padding(5, 0, 0, 0)
-				.AutoWidth()
-				.VAlign(VAlign_Center)
-				[
-					SNew(STextBlock)
-					.Font(AttributeResultFont)
-					.Text(HasClumpIDText)
-					.ColorAndOpacity(AttributeColor)
-				]
-			]
-		]
-
-		// Color attributes
-		+ SVerticalBox::Slot()
-		.AutoHeight()
-		.Padding(2)
-		[
-			SNew(SBorder)
-			.Padding(FMargin(3))
-			.BorderImage(FAppStyle::GetBrush("ToolPanel.GroupBorder"))
-			[
-				SNew(SHorizontalBox)
-				+ SHorizontalBox::Slot()
-				.AutoWidth()
-				[
-					SNew(STextBlock)
-					.Font(AttributeFont)
-					.Text(LOCTEXT("GroomOptionsWindow_HasColor", "Has Color: "))
-				]
-				+ SHorizontalBox::Slot()
-				.Padding(5, 0, 0, 0)
-				.AutoWidth()
-				.VAlign(VAlign_Center)
-				[
-					SNew(STextBlock)
-					.Font(AttributeResultFont)
-					.Text(HasColorText)
-					.ColorAndOpacity(AttributeColor)
-				]
-			]
-		]
-
-		// Roughness attributes
-		+ SVerticalBox::Slot()
-		.AutoHeight()
-		.Padding(2)
-		[
-			SNew(SBorder)
-			.Padding(FMargin(3))
-			.BorderImage(FAppStyle::GetBrush("ToolPanel.GroupBorder"))
-			[
-				SNew(SHorizontalBox)
-				+ SHorizontalBox::Slot()
-				.AutoWidth()
-				[
-					SNew(STextBlock)
-					.Font(AttributeFont)
-					.Text(LOCTEXT("GroomOptionsWindow_HasRoughness", "Has Roughness: "))
-				]
-				+ SHorizontalBox::Slot()
-				.Padding(5, 0, 0, 0)
-				.AutoWidth()
-				.VAlign(VAlign_Center)
-				[
-					SNew(STextBlock)
-					.Font(AttributeResultFont)
-					.Text(HasRoughnessText)
-					.ColorAndOpacity(AttributeColor)
-				]
-			]
-		]
-
-		// AO attributes
-		+ SVerticalBox::Slot()
+		
+		// Insert title of for the attributes
+		+SVerticalBox::Slot()
 		.AutoHeight()
 		.Padding(2)
 		[
@@ -428,8 +340,8 @@ void SGroomImportOptionsWindow::Construct(const FArguments& InArgs)
 				.AutoWidth()
 				[
 					SNew(STextBlock)
-					.Font(AttributeFont)
-					.Text(LOCTEXT("GroomOptionsWindow_HasAO", "Has AO: "))
+					.Font(FAppStyle::GetFontStyle("CurveEd.LabelFont"))
+					.Text(LOCTEXT("GroomOptionsWindow_Attribute", "Attributes: "))
 				]
 				+ SHorizontalBox::Slot()
 				.Padding(5, 0, 0, 0)
@@ -437,42 +349,15 @@ void SGroomImportOptionsWindow::Construct(const FArguments& InArgs)
 				.VAlign(VAlign_Center)
 				[
 					SNew(STextBlock)
-					.Font(AttributeResultFont)
-					.Text(HasAOText)
-					.ColorAndOpacity(AttributeColor)
+					.Font(FAppStyle::GetFontStyle("CurveEd.InfoFont"))
+					.Text(bHasAttributeText)
+					.ColorAndOpacity(bHasAttributeColor)
 				]
 			]
 		]
 
-		// Guide weights
-		+ SVerticalBox::Slot()
-		.AutoHeight()
-		.Padding(2)
-		[
-			SNew(SBorder)
-			.Padding(FMargin(3))
-			.BorderImage(FAppStyle::GetBrush("ToolPanel.GroupBorder"))
-			[
-				SNew(SHorizontalBox)
-				+ SHorizontalBox::Slot()
-				.AutoWidth()
-				[
-					SNew(STextBlock)
-					.Font(AttributeFont)
-					.Text(LOCTEXT("GroomOptionsWindow_HasGuideWeights", "Has Pre-Computed Guides Weights: "))
-				]
-				+ SHorizontalBox::Slot()
-				.Padding(5, 0, 0, 0)
-				.AutoWidth()
-				.VAlign(VAlign_Center)
-				[
-					SNew(STextBlock)
-					.Font(AttributeResultFont)
-					.Text(HasGuideWeightsText)
-					.ColorAndOpacity(AttributeColor)
-				]
-			]
-		]
+		// All optional attribute will be inserted here
+		// The widget are inserted at the end of this function
 
 		+ SVerticalBox::Slot()
 		.Padding(2)
@@ -517,7 +402,23 @@ void SGroomImportOptionsWindow::Construct(const FArguments& InArgs)
 				.Text(LOCTEXT("Cancel", "Cancel"))
 				.OnClicked(this, &SGroomImportOptionsWindow::OnCancel)
 			]
-		]
+		];
+
+	// Insert all the optional attributes
+	uint32 AttributeSlotIndex = 3;
+	for (uint32 AttributeIt = 0; AttributeIt < uint32(EHairAttribute::Count); ++AttributeIt)
+	{
+		const EHairAttribute AttributeType = (EHairAttribute)AttributeIt;
+		if (HasHairAttribute(Attributes, AttributeType))
+		{
+			SVerticalBox::FScopedWidgetSlotArguments SlotArg = VerticalSlot->InsertSlot(AttributeSlotIndex++);
+			AddAttribute(SlotArg, GetHairAttributeText(AttributeType));
+		}
+	}
+
+	this->ChildSlot
+	[
+		VerticalSlot
 	];
 }
 
diff --git a/Engine/Plugins/Runtime/HairStrands/Source/HairStrandsEditor/Private/HairStrandsFactory.cpp b/Engine/Plugins/Runtime/HairStrands/Source/HairStrandsEditor/Private/HairStrandsFactory.cpp
--- a/Engine/Plugins/Runtime/HairStrands/Source/HairStrandsEditor/Private/HairStrandsFactory.cpp
+++ b/Engine/Plugins/Runtime/HairStrands/Source/HairStrandsEditor/Private/HairStrandsFactory.cpp
@@ -111,12 +111,6 @@ UObject* UHairStrandsFactory::FactoryCreateFile(UClass* InClass, UObject* InPare
 	{
 		// Load the alembic file upfront to preview & report any potential issue
 		FHairDescriptionGroups OutDescription;
-		bool bHasRootUV = false;
-		bool bHasClumpID = false;
-		bool bHasPrecomputedWeights = false;
-		bool bHasColor= false;
-		bool bHasRoughness = false;
-		bool bHasAO = false;
 		{
 			FScopedSlowTask Progress((float)1, LOCTEXT("ImportHairAssetForPreview", "Importing hair asset for preview..."), true);
 			Progress.MakeDialog(true);
@@ -128,12 +122,6 @@ UObject* UHairStrandsFactory::FactoryCreateFile(UClass* InClass, UObject* InPare
 			}
 
 			FGroomBuilder::BuildHairDescriptionGroups(HairDescription, OutDescription);
-			bHasRootUV = HairDescription.HasRootUV();
-			bHasClumpID = HairDescription.HasClumpID();
-			bHasPrecomputedWeights = HairDescription.HasGuideWeights();
-			bHasColor = HairDescription.HasColorAttributes();
-			bHasRoughness = HairDescription.HasRoughnessAttributes();
-			bHasAO = HairDescription.HasAOAttributes();
 
 			// Populate the interpolation settings based on the group count, as this is used later during the ImportHair() to define 
 			// the exact number of group to create
@@ -154,12 +142,7 @@ UObject* UHairStrandsFactory::FactoryCreateFile(UClass* InClass, UObject* InPare
 				OutGroup.GroupID	= Group.Info.GroupID;
 				OutGroup.CurveCount = Group.Info.NumCurves;
 				OutGroup.GuideCount = Group.Info.NumGuides;
-				OutGroup.bHasRootUV = bHasRootUV;
-				OutGroup.bHasClumpID = bHasClumpID;
-				OutGroup.bHasPrecomputedWeights = bHasPrecomputedWeights;
-				OutGroup.bHasColor = bHasColor;
-				OutGroup.bHasRoughness = bHasRoughness;
-				OutGroup.bHasAO = bHasAO;
+				OutGroup.Attributes = Group.Attributes;
 				bGuidesOnly |= (OutGroup.CurveCount == 0 && OutGroup.GuideCount > 0);
 
 				if (OutGroup.GroupID < OutDescription.HairGroups.Num())
diff --git a/Engine/Plugins/Runtime/HairStrands/Source/HairStrandsEditor/Private/ReimportHairStrandsFactory.cpp b/Engine/Plugins/Runtime/HairStrands/Source/HairStrandsEditor/Private/ReimportHairStrandsFactory.cpp
--- a/Engine/Plugins/Runtime/HairStrands/Source/HairStrandsEditor/Private/ReimportHairStrandsFactory.cpp
+++ b/Engine/Plugins/Runtime/HairStrands/Source/HairStrandsEditor/Private/ReimportHairStrandsFactory.cpp
@@ -384,7 +384,7 @@ EReimportResult::Type UReimportHairStrandsFactory::Reimport(UObject* Obj)
 				OutGroup.GroupID	= Group.Info.GroupID;
 				OutGroup.CurveCount = Group.Info.NumCurves;
 				OutGroup.GuideCount = Group.Info.NumGuides;
-				OutGroup.bHasPrecomputedWeights = Group.Strands.StrandsCurves.HasPrecomputedWeights();
+				OutGroup.Attributes = Group.Attributes;
 
 				if (OutGroup.GroupID < OutDescription.HairGroups.Num())
 				{
