Fix invalid 'too old' asset when loading groom asset containing only mesh and/or cards data.

diff --git a/Engine/Plugins/Runtime/HairStrands/Source/HairStrandsCore/Private/GroomAsset.cpp b/Engine/Plugins/Runtime/HairStrands/Source/HairStrandsCore/Private/GroomAsset.cpp
--- a/Engine/Plugins/Runtime/HairStrands/Source/HairStrandsCore/Private/GroomAsset.cpp
+++ b/Engine/Plugins/Runtime/HairStrands/Source/HairStrandsCore/Private/GroomAsset.cpp
@@ -2286,7 +2286,21 @@ bool UGroomAsset::CacheStrandsData(uint32 GroupIndex, FString& OutDerivedDataKey
 {
 	if (!HairDescriptionBulkData[HairDescriptionType])
 	{
-		UE_LOG(LogHairStrands, Error, TEXT("[Groom] The groom asset %s is too old. Please reimported the groom from its original source file."), *GetName());
+		// Groom having only cards/meshes (i.e., not containing imported strands) don't have hair description
+		bool bHasStrands = false;
+		for (const FHairLODSettings& LOD : HairGroupsLOD[GroupIndex].LODs)
+		{
+			if (LOD.GeometryType == EGroomGeometryType::Strands)
+			{
+				bHasStrands = true;
+				break;
+			}
+		}
+
+		if (bHasStrands)
+		{
+			UE_LOG(LogHairStrands, Error, TEXT("[Groom] The groom asset %s is too old. Please reimported the groom from its original source file."), *GetName());
+		}
 		return false;
 	}
 
