Adding missing forward declaration.

diff --git a/Engine/Plugins/Runtime/HairStrands/Source/HairStrandsCore/Public/HairStrandsRendering.h b/Engine/Plugins/Runtime/HairStrands/Source/HairStrandsCore/Public/HairStrandsRendering.h
--- a/Engine/Plugins/Runtime/HairStrands/Source/HairStrandsCore/Public/HairStrandsRendering.h
+++ b/Engine/Plugins/Runtime/HairStrands/Source/HairStrandsCore/Public/HairStrandsRendering.h
@@ -26,6 +26,8 @@ struct FHairStrandClusterData;
 struct FRWBuffer;
 struct FShaderPrintData;
 
+enum class EGroomViewMode : uint8;
+
 // Reset the interpolation data. This needs to be called prior to ComputeHairStrandsInterpolation 
 // and prior to the actual hair simulation in order to insure that:
 //  1) when hair simulation is enabled, the first frame is correct
