Fix CIS issue.

diff --git a/Engine/Plugins/Runtime/HairStrands/Source/HairStrandsEditor/Private/GroomImportOptionsWindow.cpp b/Engine/Plugins/Runtime/HairStrands/Source/HairStrandsEditor/Private/GroomImportOptionsWindow.cpp
--- a/Engine/Plugins/Runtime/HairStrands/Source/HairStrandsEditor/Private/GroomImportOptionsWindow.cpp
+++ b/Engine/Plugins/Runtime/HairStrands/Source/HairStrandsEditor/Private/GroomImportOptionsWindow.cpp
@@ -219,7 +219,7 @@ FText GetHairAttributeText(EHairAttribute In)
 	switch (In)
 	{
 	case EHairAttribute::RootUV:					return LOCTEXT("GroomOptionsWindow_HasRootUV", "Root UV");
-	case EHairAttribute::ClumpID:					return LOCTEXT("GroomOptionsWindow_HasRootUV", "Clump ID");
+	case EHairAttribute::ClumpID:					return LOCTEXT("GroomOptionsWindow_HasClumpID", "Clump ID");
 	case EHairAttribute::PrecomputedGuideWeights:	return LOCTEXT("GroomOptionsWindow_HasPercomputedGuideWeights", "Pre-Computed Guide Weights");
 	case EHairAttribute::Color:						return LOCTEXT("GroomOptionsWindow_HasColor", "Has Color");
 	case EHairAttribute::Roughness:					return LOCTEXT("GroomOptionsWindow_HasRoughness", "Roughness");
