Simplify/uniformize group size accros hair passes.

This addresses a few render hair passes. More a conversion are coming.

diff --git a/Engine/Plugins/Runtime/HairStrands/Source/HairStrandsCore/Private/HairStrandsRendering.cpp b/Engine/Plugins/Runtime/HairStrands/Source/HairStrandsCore/Private/HairStrandsRendering.cpp
--- a/Engine/Plugins/Runtime/HairStrands/Source/HairStrandsCore/Private/HairStrandsRendering.cpp
+++ b/Engine/Plugins/Runtime/HairStrands/Source/HairStrandsCore/Private/HairStrandsRendering.cpp
@@ -88,18 +88,6 @@ static FIntVector ComputeDispatchCount(uint32 ItemCount, uint32 GroupSize)
 	return FIntVector(DispatchCountX, DispatchCountY, 1);
 }
 
-// Same as above but the group count is what matters and is preserved
-static FIntVector ComputeDispatchCount(uint32 GroupCount)
-{
-	const uint32 DispatchCountX = FMath::FloorToInt(FMath::Sqrt(static_cast<float>(GroupCount)));
-	const uint32 DispatchCountY = DispatchCountX + FMath::DivideAndRoundUp(GroupCount - DispatchCountX * DispatchCountX, DispatchCountX);
-
-	check(DispatchCountX <= 65535);
-	check(DispatchCountY <= 65535);
-	check(GroupCount <= DispatchCountX * DispatchCountY);
-	return FIntVector(DispatchCountX, DispatchCountY, 1);
-}
-
 inline uint32 ComputeGroupSize()
 {
 	const uint32 GroupSize = IsRHIDeviceAMD() ? 64 : (IsRHIDeviceNVIDIA() ? 32 : 64);
@@ -128,11 +116,13 @@ class FTransferVelocityPassCS : public FGlobalShader
 	END_SHADER_PARAMETER_STRUCT()
 
 public:
+	static uint32 GetGroupSize() { return HAIR_VERTEXCOUNT_GROUP_SIZE; }
 	static bool ShouldCompilePermutation(const FGlobalShaderPermutationParameters& Parameters) { return IsHairStrandsSupported(EHairStrandsShaderType::All, Parameters.Platform); }
 	static void ModifyCompilationEnvironment(const FGlobalShaderPermutationParameters& Parameters, FShaderCompilerEnvironment& OutEnvironment)
 	{
 		FGlobalShader::ModifyCompilationEnvironment(Parameters, OutEnvironment);
 		OutEnvironment.SetDefine(TEXT("SHADER_HAIRTRANSFER_PREV_POSITION"), 1);
+		OutEnvironment.SetDefine(TEXT("GROUP_SIZE"), GetGroupSize());
 	}
 };
 
@@ -147,23 +137,19 @@ static void AddTransferPositionPass(
 {
 	if (ElementCount == 0) return;
 
-	const uint32 GroupSize = 64;
-	const uint32 DispatchCount = FMath::DivideAndRoundUp(ElementCount, GroupSize);
-	const uint32 DispatchCountX = 128;
-	const uint32 DispatchCountY = FMath::DivideAndRoundUp(DispatchCount, DispatchCountX);
-
 	FTransferVelocityPassCS::FParameters* Parameters = GraphBuilder.AllocParameters<FTransferVelocityPassCS::FParameters>();
 	Parameters->ElementCount = ElementCount;
 	Parameters->InBuffer = InBuffer;
 	Parameters->OutBuffer = OutBuffer;
 
+	const FIntVector DispatchCount(FMath::DivideAndRoundUp(ElementCount, FTransferVelocityPassCS::GetGroupSize()), 1, 1);
 	TShaderMapRef<FTransferVelocityPassCS> ComputeShader(ShaderMap);
 	FComputeShaderUtils::AddPass(
 		GraphBuilder,
 		RDG_EVENT_NAME("HairStrands::TransferPositionForVelocity"),
 		ComputeShader,
 		Parameters,
-		FIntVector(DispatchCountX, DispatchCountY, 1));
+		DispatchCount);
 }
 
 ////////////////////////////////////////////////////////////////////////////////////////////////////////////////////////
@@ -175,21 +161,22 @@ class FGroomCacheUpdatePassCS : public FGlobalShader
 
 	BEGIN_SHADER_PARAMETER_STRUCT(FParameters, )
 		SHADER_PARAMETER(uint32, ElementCount)
-		SHADER_PARAMETER(uint32, DispatchCountX)
 		SHADER_PARAMETER(int, bHasRadiusData)
 		SHADER_PARAMETER_RDG_BUFFER_SRV(Buffer, InAnimatedBuffer)
 		SHADER_PARAMETER_RDG_BUFFER_SRV(Buffer, InRadiusBuffer)
 		SHADER_PARAMETER_RDG_BUFFER_SRV(Buffer, InRestPoseBuffer)
 		SHADER_PARAMETER_RDG_BUFFER_SRV(Buffer, InDeformedOffsetBuffer)
 		SHADER_PARAMETER_RDG_BUFFER_UAV(RWBuffer, OutDeformedBuffer)
-		END_SHADER_PARAMETER_STRUCT()
+	END_SHADER_PARAMETER_STRUCT()
 
 public:
+	static uint32 GetGroupSize() { return HAIR_VERTEXCOUNT_GROUP_SIZE; }
 	static bool ShouldCompilePermutation(const FGlobalShaderPermutationParameters& Parameters) { return IsHairStrandsSupported(EHairStrandsShaderType::All, Parameters.Platform); }
 	static void ModifyCompilationEnvironment(const FGlobalShaderPermutationParameters& Parameters, FShaderCompilerEnvironment& OutEnvironment)
 	{
 		FGlobalShader::ModifyCompilationEnvironment(Parameters, OutEnvironment);
 		OutEnvironment.SetDefine(TEXT("SHADER_GROOMCACHE_UPDATE"), 1);
+		OutEnvironment.SetDefine(TEXT("GROUP_SIZE"), GetGroupSize());
 	}
 };
 
@@ -209,9 +196,6 @@ static void AddGroomCacheUpdatePass(
 
 	if (ElementCount == 0) return;
 
-	const uint32 GroupSize = 64;
-	const FIntVector DispatchCount = ComputeDispatchCount(ElementCount, GroupSize);
-
 	FRDGBufferRef VertexBuffer = nullptr;
 
 	const uint32 DataCount = GroomCacheData.VertexData.PointsPosition.Num();
@@ -236,7 +220,6 @@ static void AddGroomCacheUpdatePass(
 	}
 
 	FGroomCacheUpdatePassCS::FParameters* Parameters = GraphBuilder.AllocParameters<FGroomCacheUpdatePassCS::FParameters>();
-	Parameters->DispatchCountX = DispatchCount.X;
 	Parameters->ElementCount = ElementCount;
 	Parameters->InAnimatedBuffer = GraphBuilder.CreateSRV(VertexBuffer, PF_R32_FLOAT);
 	Parameters->InRestPoseBuffer = InBuffer;
@@ -265,6 +248,7 @@ static void AddGroomCacheUpdatePass(
 	}
 	Parameters->bHasRadiusData = bHasRadiusData ? 1 : 0;
 
+	const FIntVector DispatchCount = FIntVector(FMath::DivideAndRoundUp(ElementCount, FGroomCacheUpdatePassCS::GetGroupSize()), 0, 0);
 	TShaderMapRef<FGroomCacheUpdatePassCS> ComputeShader(ShaderMap);
 	FComputeShaderUtils::AddPass(
 		GraphBuilder,
@@ -281,14 +265,12 @@ class FDeformGuideCS : public FGlobalShader
 	DECLARE_GLOBAL_SHADER(FDeformGuideCS);
 	SHADER_USE_PARAMETER_STRUCT(FDeformGuideCS, FGlobalShader);
 
-	class FGroupSize : SHADER_PERMUTATION_SPARSE_INT("PERMUTATION_GROUP_SIZE", 32, 64);
 	class FDeformationType : SHADER_PERMUTATION_INT("PERMUTATION_DEFORMATION", 5);
-	using FPermutationDomain = TShaderPermutationDomain<FGroupSize, FDeformationType>;
+	using FPermutationDomain = TShaderPermutationDomain<FDeformationType>;
 
 	BEGIN_SHADER_PARAMETER_STRUCT(FParameters, )
 		SHADER_PARAMETER(uint32, VertexCount)
 		SHADER_PARAMETER(FVector3f, SimRestOffset)
-		SHADER_PARAMETER(uint32, DispatchCountX)
 
 		SHADER_PARAMETER_RDG_BUFFER_SRV(Buffer<float4>, SimRestPosition0Buffer)
 		SHADER_PARAMETER_RDG_BUFFER_SRV(Buffer<float4>, SimRestPosition1Buffer)
@@ -314,7 +296,14 @@ class FDeformGuideCS : public FGlobalShader
 	END_SHADER_PARAMETER_STRUCT()
 
 public:
+	static uint32 GetGroupSize() { return HAIR_VERTEXCOUNT_GROUP_SIZE; }
 	static bool ShouldCompilePermutation(const FGlobalShaderPermutationParameters& Parameters) { return IsHairStrandsSupported(EHairStrandsShaderType::All, Parameters.Platform); }
+	static void ModifyCompilationEnvironment(const FGlobalShaderPermutationParameters& Parameters, FShaderCompilerEnvironment& OutEnvironment)
+	{
+		FGlobalShader::ModifyCompilationEnvironment(Parameters, OutEnvironment);
+		OutEnvironment.SetDefine(TEXT("SHADER_DEFORM_GUIDE"), 1);
+		OutEnvironment.SetDefine(TEXT("GROUP_SIZE"), GetGroupSize());
+	}
 };
 
 IMPLEMENT_GLOBAL_SHADER(FDeformGuideCS, "/Engine/Private/HairStrands/HairStrandsGuideDeform.usf", "MainCS", SF_Compute);
@@ -354,18 +343,12 @@ static void AddDeformSimHairStrandsPass(
 
 	if (InternalDeformationType == InternalDeformationTypeCount) return;
 
-	const uint32 GroupSize = ComputeGroupSize();
-	const uint32 DispatchCount = FMath::DivideAndRoundUp(VertexCount, GroupSize);
-	const uint32 DispatchCountX = 16;
-	const uint32 DispatchCountY = FMath::DivideAndRoundUp(DispatchCount, DispatchCountX);
-
 	FDeformGuideCS::FParameters* Parameters = GraphBuilder.AllocParameters<FDeformGuideCS::FParameters>();
 	Parameters->SimRestPosePositionBuffer = SimRestPosePositionBuffer;
 	Parameters->OutSimDeformedPositionBuffer = OutSimDeformedPositionBuffer.UAV;
 	Parameters->VertexCount = VertexCount;
 	Parameters->SimDeformedOffsetBuffer = SimDeformedOffsetBuffer;
 	Parameters->SimRestOffset = (FVector3f)SimRestOffset;
-	Parameters->DispatchCountX = DispatchCountX;
 
 	if (DeformationType == EDeformationType::OffsetGuide)
 	{
@@ -426,16 +409,16 @@ static void AddDeformSimHairStrandsPass(
 	}
 
 	FDeformGuideCS::FPermutationDomain PermutationVector;
-	PermutationVector.Set<FDeformGuideCS::FGroupSize>(GroupSize);
 	PermutationVector.Set<FDeformGuideCS::FDeformationType>(InternalDeformationType);
 
+	const FIntVector DispatchCount(FMath::DivideAndRoundUp(VertexCount,FDeformGuideCS::GetGroupSize()), 1, 1);
 	TShaderMapRef<FDeformGuideCS> ComputeShader(ShaderMap, PermutationVector);
 	FComputeShaderUtils::AddPass(
 		GraphBuilder,
 		RDG_EVENT_NAME("HairStrands::DeformSimHairStrands"),
 		ComputeShader,
 		Parameters,
-		FIntVector(DispatchCountX, DispatchCountY, 1));
+		DispatchCount);
 
 	GraphBuilder.SetBufferAccessFinal(OutSimDeformedPositionBuffer.Buffer, ERHIAccess::SRVMask);
 }
@@ -1009,9 +992,8 @@ class FHairCardsDeformationCS : public FGlobalShader
 	DECLARE_GLOBAL_SHADER(FHairCardsDeformationCS);
 	SHADER_USE_PARAMETER_STRUCT(FHairCardsDeformationCS, FGlobalShader);
 
-	class FGroupSize : SHADER_PERMUTATION_SPARSE_INT("PERMUTATION_GROUP_SIZE", 32, 64);
 	class FDynamicGeometry : SHADER_PERMUTATION_INT("PERMUTATION_DYNAMIC_GEOMETRY", 2);
-	using FPermutationDomain = TShaderPermutationDomain<FGroupSize, FDynamicGeometry>;
+	using FPermutationDomain = TShaderPermutationDomain<FDynamicGeometry>;
 
 	BEGIN_SHADER_PARAMETER_STRUCT(FParameters, )
 		SHADER_PARAMETER(FMatrix44f, LocalToWorld)
@@ -1044,7 +1026,14 @@ class FHairCardsDeformationCS : public FGlobalShader
 	END_SHADER_PARAMETER_STRUCT()
 
 public:
+	static uint32 GetGroupSize() { return 256u; }
 	static bool ShouldCompilePermutation(const FGlobalShaderPermutationParameters& Parameters) { return IsHairStrandsSupported(EHairStrandsShaderType::Cards, Parameters.Platform); }
+	static void ModifyCompilationEnvironment(const FGlobalShaderPermutationParameters& Parameters, FShaderCompilerEnvironment& OutEnvironment)
+	{
+		FGlobalShader::ModifyCompilationEnvironment(Parameters, OutEnvironment);
+		OutEnvironment.SetDefine(TEXT("SHADER_DEFORM_CARDS"), 1);
+		OutEnvironment.SetDefine(TEXT("GROUP_SIZE"), GetGroupSize());
+	}
 };
 
 IMPLEMENT_GLOBAL_SHADER(FHairCardsDeformationCS, "/Engine/Private/HairStrands/HairCardsDeformation.usf", "MainCS", SF_Compute);
@@ -1131,20 +1120,17 @@ static void AddHairCardsDeformationPass(
 		ShaderPrint::SetParameters(GraphBuilder, *ShaderPrintData, Parameters->ShaderDrawParameters);
 	}
 
-	const uint32 GroupSize = ComputeGroupSize();
 	FHairCardsDeformationCS::FPermutationDomain PermutationVector;
 	PermutationVector.Set<FHairCardsDeformationCS::FDynamicGeometry>(bSupportDynamicMesh ? 1 : 0);
-	PermutationVector.Set<FHairCardsDeformationCS::FGroupSize>(GroupSize);
-
 	TShaderMapRef<FHairCardsDeformationCS> ComputeShader(ShaderMap, PermutationVector);
 
-	const int32 DispatchCountX = FMath::DivideAndRoundUp(Parameters->CardsVertexCount, GroupSize);
+	const FIntVector DispatchCount = FIntVector(FMath::DivideAndRoundUp(Parameters->CardsVertexCount, FHairCardsDeformationCS::GetGroupSize()), 1, 1);
 	FComputeShaderUtils::AddPass(
 		GraphBuilder,
 		RDG_EVENT_NAME("HairStrands::CardsDeformation(%s)", bSupportDynamicMesh ? TEXT("Dynamic") : TEXT("Static")),
 		ComputeShader,
 		Parameters,
-		FIntVector(DispatchCountX,1,1));
+		DispatchCount);
 
 	// If LOD has switched, copy the current buffer, so that we don't get incorrect motion vector
 	if (bHasLODSwitch)
@@ -1244,14 +1230,12 @@ class FHairRaytracingGeometryCS : public FGlobalShader
 	DECLARE_GLOBAL_SHADER(FHairRaytracingGeometryCS);
 	SHADER_USE_PARAMETER_STRUCT(FHairRaytracingGeometryCS, FGlobalShader);
 
-	class FGroupSize : SHADER_PERMUTATION_SPARSE_INT("PERMUTATION_GROUP_SIZE", 32, 64);
 	class FCulling : SHADER_PERMUTATION_BOOL("PERMUTATION_CULLING");
 	class FProceduralPrimitive : SHADER_PERMUTATION_BOOL("PERMUTATION_PROCEDURAL_PRIMITIVE"); 
-	using FPermutationDomain = TShaderPermutationDomain<FGroupSize, FCulling, FProceduralPrimitive>;
+	using FPermutationDomain = TShaderPermutationDomain<FCulling, FProceduralPrimitive>;
 
 	BEGIN_SHADER_PARAMETER_STRUCT(FParameters, )
 		SHADER_PARAMETER(uint32, VertexCount)
-		SHADER_PARAMETER(uint32, DispatchCountX)
 		SHADER_PARAMETER(float, HairStrandsVF_HairRadius)
 		SHADER_PARAMETER(float, HairStrandsVF_HairRootScale)
 		SHADER_PARAMETER(float, HairStrandsVF_HairTipScale)
@@ -1272,11 +1256,13 @@ class FHairRaytracingGeometryCS : public FGlobalShader
 	END_SHADER_PARAMETER_STRUCT()
 
 public:
+	static uint32 GetGroupSize() { return HAIR_VERTEXCOUNT_GROUP_SIZE; }
 	static bool ShouldCompilePermutation(const FGlobalShaderPermutationParameters& Parameters) { return IsHairStrandsSupported(EHairStrandsShaderType::Strands, Parameters.Platform); }
 	static void ModifyCompilationEnvironment(const FGlobalShaderPermutationParameters& Parameters, FShaderCompilerEnvironment& OutEnvironment)
 	{
 		FGlobalShader::ModifyCompilationEnvironment(Parameters, OutEnvironment);
 		OutEnvironment.SetDefine(TEXT("SHADER_RT_GEOMETRY"), 1);
+		OutEnvironment.SetDefine(TEXT("GROUP_SIZE"), GetGroupSize());
 	}
 };
 
@@ -1308,12 +1294,8 @@ static void AddGenerateRaytracingGeometryPass(
 	const FRDGBufferUAVRef& OutPositionBuffer,
 	const FRDGBufferUAVRef& OutIndexBuffer)
 {
-	const uint32 GroupSize = ComputeGroupSize();
-	const FIntVector DispatchCount = ComputeDispatchCount(VertexCount, GroupSize);
-
 	FHairRaytracingGeometryCS::FParameters* Parameters = GraphBuilder.AllocParameters<FHairRaytracingGeometryCS::FParameters>();
 	Parameters->VertexCount = VertexCount;
-	Parameters->DispatchCountX = DispatchCount.X;
 	Parameters->PositionOffsetBuffer = HairWorldOffsetBuffer;
 	Parameters->HairStrandsVF_HairRadius = HairRadius;
 	Parameters->HairStrandsVF_HairRootScale = RootScale;
@@ -1340,10 +1322,11 @@ static void AddGenerateRaytracingGeometryPass(
 	}
 
 	FHairRaytracingGeometryCS::FPermutationDomain PermutationVector;
-	PermutationVector.Set<FHairRaytracingGeometryCS::FGroupSize>(GroupSize);
 	PermutationVector.Set<FHairRaytracingGeometryCS::FCulling>(bCullingEnable);
 	PermutationVector.Set<FHairRaytracingGeometryCS::FProceduralPrimitive>(bProceduralPrimitive);
 	TShaderMapRef<FHairRaytracingGeometryCS> ComputeShader(ShaderMap, PermutationVector);
+
+	const FIntVector DispatchCount(FMath::DivideAndRoundUp(VertexCount, FHairRaytracingGeometryCS::GetGroupSize()), 1, 1);
 	FComputeShaderUtils::AddPass(
 		GraphBuilder,
 		RDG_EVENT_NAME("HairStrands::GenerateRaytracingGeometry(bCulling:%d, Procedural:%d)", bCullingEnable, bProceduralPrimitive),
diff --git a/Engine/Shaders/Private/HairStrands/HairCardsDeformation.usf b/Engine/Shaders/Private/HairStrands/HairCardsDeformation.usf
--- a/Engine/Shaders/Private/HairStrands/HairCardsDeformation.usf
+++ b/Engine/Shaders/Private/HairStrands/HairCardsDeformation.usf
@@ -10,6 +10,7 @@
 #include "../ShaderPrint.ush"?
 #endif
 
+#if SHADER_DEFORM_CARDS
 float4x4 LocalToWorld;
 
 uint CardsVertexCount;
@@ -40,16 +41,6 @@ Buffer<uint>	GuideRootBarycentricBuffer;
 Buffer<uint>    GuideVertexToRootIndexBuffer;
 Buffer<uint>    GuideRootToUniqueTriangleIndexBuffer;
 
-
-#if PERMUTATION_GROUP_SIZE == 64
-	#define GROUP_SIZE 64
-#elif PERMUTATION_GROUP_SIZE == 32
-	#define GROUP_SIZE 32
-#else
-	#error Unknown group size	
-#endif
-
-
 #if PERMUTATION_DYNAMIC_GEOMETRY == 1
 float3 ComputeDynamicGeometryOffset(
 	uint GuideIndex,
@@ -87,9 +78,9 @@ void CopyTangent(uint InVertexId, uint InOffset)
 }
 
 [numthreads(GROUP_SIZE, 1, 1)]
-void MainCS(uint2 DispatchThreadId : SV_DispatchThreadID)
+void MainCS(uint DispatchThreadId : SV_DispatchThreadID)
 {
-	const uint VertexId = DispatchThreadId.x;
+	const uint VertexId = DispatchThreadId;
 	if (VertexId >= CardsVertexCount)
 		return;
 
@@ -187,4 +178,5 @@ void MainCS(uint2 DispatchThreadId : SV_DispatchThreadID)
 	TransformTangent(VertexId, 0, RestTriangle, DeformedTriangle);
 	TransformTangent(VertexId, 1, RestTriangle, DeformedTriangle);
 #endif // PERMUTATION_DYNAMIC_GEOMETRY
-}
\ No newline at end of file
+}
+#endif // SHADER_DEFORM_CARDS
\ No newline at end of file
diff --git a/Engine/Shaders/Private/HairStrands/HairStrandsGuideDeform.usf b/Engine/Shaders/Private/HairStrands/HairStrandsGuideDeform.usf
--- a/Engine/Shaders/Private/HairStrands/HairStrandsGuideDeform.usf
+++ b/Engine/Shaders/Private/HairStrands/HairStrandsGuideDeform.usf
@@ -4,21 +4,15 @@
 #include "HairStrandsClusterCommon.ush"
 #include "HairStrandsMeshProjectionCommon.ush"
 
+#if SHADER_DEFORM_GUIDE
+
 float3 SimRestOffset;
 uint VertexCount;
-uint DispatchCountX; 
 
 Buffer<float4> SimDeformedOffsetBuffer;
 Buffer<uint4> SimRestPosePositionBuffer;
 RWBuffer<uint4> OutSimDeformedPositionBuffer;
 
-#if PERMUTATION_GROUP_SIZE == 64
-	#define GROUP_SIZE 64
-#elif PERMUTATION_GROUP_SIZE == 32
-	#define GROUP_SIZE 32
-#else
-	#error Unknown group size	
-#endif
 
 #define DEFORMATION_BYPASS 0
 #define DEFORMATION_OFFSET_GUIDE 1
@@ -109,9 +103,9 @@ float3 DisplacePosition(uint VertexIndex, float3 Pos)
 }
 
 [numthreads(GROUP_SIZE, 1, 1)]
-void MainCS(uint2 DispatchThreadId : SV_DispatchThreadID)
+void MainCS(uint DispatchThreadId : SV_DispatchThreadID)
 {
-	const uint VertexIndex = GetHairStrandsVertexFetchIndex(DispatchThreadId, GROUP_SIZE, DispatchCountX);
+	const uint VertexIndex = DispatchThreadId;
 	
 	if (VertexIndex < VertexCount)
 	{
@@ -121,3 +115,5 @@ void MainCS(uint2 DispatchThreadId : SV_DispatchThreadID)
 		OutSimDeformedPositionBuffer[VertexIndex] = uint4(f32tof16(DisplacedPosition), Data.w);
 	}
 }
+
+#endif // SHADER_DEFORM_GUIDE
\ No newline at end of file
diff --git a/Engine/Shaders/Private/HairStrands/HairStrandsInterpolation.usf b/Engine/Shaders/Private/HairStrands/HairStrandsInterpolation.usf
--- a/Engine/Shaders/Private/HairStrands/HairStrandsInterpolation.usf
+++ b/Engine/Shaders/Private/HairStrands/HairStrandsInterpolation.usf
@@ -820,15 +820,13 @@ void ClusterAABBEvaluationCS(uint3 DispatchThreadId : SV_DispatchThreadID, uint3
 #if SHADER_HAIRTRANSFER_PREV_POSITION
 
 uint ElementCount;
-uint DispatchCountX;
 Buffer<uint4> InBuffer;
 RWBuffer<uint4> OutBuffer;
 
-#define GROUP_SIZE 64
 [numthreads(GROUP_SIZE, 1, 1)]
-void MainCS(uint2 DispatchThreadId : SV_DispatchThreadID)
+void MainCS(uint DispatchThreadId : SV_DispatchThreadID)
 {
-	const uint Index = GetHairStrandsVertexFetchIndex(DispatchThreadId, GROUP_SIZE, DispatchCountX);
+	const uint Index = DispatchThreadId;
 	if (Index >= ElementCount)
 	{
 		return;
@@ -844,7 +842,6 @@ void MainCS(uint2 DispatchThreadId : SV_DispatchThreadID)
 #if SHADER_GROOMCACHE_UPDATE
 
 uint ElementCount;
-uint DispatchCountX;
 uint bHasRadiusData;
 
 Buffer<float> InRadiusBuffer;
@@ -853,11 +850,10 @@ Buffer<uint4> InRestPoseBuffer;
 Buffer<float4> InDeformedOffsetBuffer;
 RWBuffer<uint4> OutDeformedBuffer;
 
-#define GROUP_SIZE 64
 [numthreads(GROUP_SIZE, 1, 1)]
-void MainCS(uint2 DispatchThreadId : SV_DispatchThreadID)
+void MainCS(uint DispatchThreadId : SV_DispatchThreadID)
 {
-	const uint Index = GetHairStrandsVertexFetchIndex(DispatchThreadId, GROUP_SIZE, DispatchCountX);
+	const uint Index = DispatchThreadId;
 	if (Index >= ElementCount)
 	{
 		return;
diff --git a/Engine/Shaders/Private/HairStrands/HairStrandsRaytracingGeometry.usf b/Engine/Shaders/Private/HairStrands/HairStrandsRaytracingGeometry.usf
--- a/Engine/Shaders/Private/HairStrands/HairStrandsRaytracingGeometry.usf
+++ b/Engine/Shaders/Private/HairStrands/HairStrandsRaytracingGeometry.usf
@@ -9,7 +9,6 @@
 #if SHADER_RT_GEOMETRY
 
 uint VertexCount;
-uint DispatchCountX;
 
 float HairStrandsVF_HairRadius;
 float HairStrandsVF_HairRootScale;
@@ -33,14 +32,6 @@ Buffer<float>	HairStrandsVF_CullingRadiusScaleBuffer;
 Buffer<uint>	HairStrandsVF_CullingIndirectBufferArgs;
 #endif
 
-#if PERMUTATION_GROUP_SIZE == 64
-	#define GROUP_SIZE 64
-#elif PERMUTATION_GROUP_SIZE == 32
-	#define GROUP_SIZE 32
-#else
-	#error Unknown group size	
-#endif
-
 FHairControlPoint GetControlPoint(uint ControlPointId)
 {
 	FHairControlPoint Out = (FHairControlPoint)0;
@@ -163,10 +154,18 @@ bool IsValidPosition(FHairControlPoint P)
 }
 
 [numthreads(GROUP_SIZE, 1, 1)]
-void MainCS(uint2 DispatchThreadId : SV_DispatchThreadID)
+void MainCS(uint DispatchThreadId : SV_DispatchThreadID)
 {
-	uint VertexIndex0 = GetHairStrandsVertexFetchIndex(DispatchThreadId, GROUP_SIZE, DispatchCountX);
+	uint VertexIndex0 = DispatchThreadId;
 	uint VertexIndex1 = VertexIndex0 + 1;
+
+	const bool bIsValid = VertexIndex0 < VertexCount && VertexIndex1 < VertexCount;
+	if (!bIsValid)
+	{
+		SetInvalidSegment(VertexIndex0);
+		return;
+	}
+
 	float LodRadiusScale0 = 1;
 	float LodRadiusScale1 = 1;
 
@@ -187,7 +186,6 @@ void MainCS(uint2 DispatchThreadId : SV_DispatchThreadID)
 		SourceIndex1 = HairStrandsVF_CullingIndexBuffer[VertexIndex1];
 		LodRadiusScale1 = HairStrandsVF_CullingRadiusScaleBuffer[VertexIndex1];
 
-
 		if (LodRadiusScale0 <= 0.0f)
 		{
 			SetInvalidSegment(VertexIndex0);
@@ -295,9 +293,12 @@ void MainCS(uint2 DispatchThreadId : SV_DispatchThreadID)
 #include "../ShaderPrint.ush"
 #include "HairStrandsAABBCommon.ush"
 
+#if GROUP_SIZE != 1024
+#error GROUP_SIZE has changed, please update the reduction code.
+#endif
+
 uint VertexCount;
 uint bDrawInvalidElement;
-uint DispatchCountX;
 float PositionThreshold2;
 
 uint HairStrandsVF_bIsCullingEnable;
@@ -309,7 +310,6 @@ Buffer<int> 	 GroupAABBBuffer;
 RWBuffer<uint> 	 InvalidationBuffer;
 RWBuffer<uint> 	 InvalidationPrintCounter;
 
-#define GROUP_SIZE 64
 groupshared uint s_InvalidMask[GROUP_SIZE];
 
 [numthreads(GROUP_SIZE, 1, 1)]
@@ -322,19 +322,21 @@ void MainCS(uint2 DispatchThreadId : SV_DispatchThreadID, uint LinearIndex : SV_
 	if (!bIsAlreadyInvalid)
 	{
 		// Fetch vertex index if culling is enabled
-		uint VertexIndex0 = GetHairStrandsVertexFetchIndex(DispatchThreadId, GROUP_SIZE, DispatchCountX);
-		if (HairStrandsVF_bIsCullingEnable)
+		uint VertexIndex0 = DispatchThreadId;
+		bool bIsValid = VertexIndex0 < VertexCount;
+		if (bIsValid && HairStrandsVF_bIsCullingEnable)
 		{
 			const uint VertexCountAfterCulling = HairStrandsVF_CullingIndirectBuffer[3];
-			uint FetchIndex0 = VertexIndex0;
-			if (FetchIndex0 < VertexCountAfterCulling)
+			const uint FetchIndex0 = DispatchThreadId;
+			bIsValid = FetchIndex0 < VertexCountAfterCulling;
+			if (bIsValid)
 			{
 				VertexIndex0 = HairStrandsVF_CullingIndexBuffer[FetchIndex0];
 			}
 		}
 
 		// Position comparison
-		if (VertexIndex0 < VertexCount)
+		if (bIsValid)
 		{
 			const FHairControlPoint Curr = UnpackHairControlPoint(CurrPositionBuffer[VertexIndex0], float3(0, 0, 0), 1, 1, 1);
 			const FHairControlPoint Prev = UnpackHairControlPoint(PrevPositionBuffer[VertexIndex0], float3(0, 0, 0), 1, 1, 1);
@@ -357,6 +359,25 @@ void MainCS(uint2 DispatchThreadId : SV_DispatchThreadID, uint LinearIndex : SV_
 	}
 
 	// Reduction
+	if (LinearIndex < 512)
+	{
+		s_InvalidMask[LinearIndex] = s_InvalidMask[LinearIndex] | s_InvalidMask[LinearIndex + 512];
+	}
+	GroupMemoryBarrierWithGroupSync();
+	if (LinearIndex < 256)
+	{
+		s_InvalidMask[LinearIndex] = s_InvalidMask[LinearIndex] | s_InvalidMask[LinearIndex + 256];
+	}
+	GroupMemoryBarrierWithGroupSync();
+	if (LinearIndex < 128)
+	{
+		s_InvalidMask[LinearIndex] = s_InvalidMask[LinearIndex] | s_InvalidMask[LinearIndex + 128];
+	}
+	GroupMemoryBarrierWithGroupSync();
+	if (LinearIndex < 64)
+	{
+		s_InvalidMask[LinearIndex] = s_InvalidMask[LinearIndex] | s_InvalidMask[LinearIndex + 64];
+	}
 	GroupMemoryBarrierWithGroupSync();
 	if (LinearIndex < 32)
 	{
diff --git a/Engine/Shaders/Shared/HairStrandsDefinitions.h b/Engine/Shaders/Shared/HairStrandsDefinitions.h
--- a/Engine/Shaders/Shared/HairStrandsDefinitions.h
+++ b/Engine/Shaders/Shared/HairStrandsDefinitions.h
@@ -20,6 +20,15 @@
 #define HAIR_ATTRIBUTE_AO 6
 #define HAIR_ATTRIBUTE_COUNT 7
 
+// Max number of discrete LOD that a hair group can have
+#define MAX_HAIR_LOD 8
+
+// Group size for dispatching based on a groom vertex/curve/cluster count. 
+// This defines ensures the group size is consistent across the hair pipeline
+#define HAIR_VERTEXCOUNT_GROUP_SIZE  1024u
+#define HAIR_CURVECOUNT_GROUP_SIZE   1024u
+#define HAIR_CLUSTERCOUNT_GROUP_SIZE 1024u
+
 // HAIR_ATTRIBUTE_MAX rounded to 4
 #define HAIR_ATTRIBUTE_OFFSET_COUNT ((HAIR_ATTRIBUTE_COUNT + 3) / 4)
 
diff --git a/Engine/Source/Runtime/Renderer/Private/HairStrands/HairStrandsUtils.cpp b/Engine/Source/Runtime/Renderer/Private/HairStrands/HairStrandsUtils.cpp
--- a/Engine/Source/Runtime/Renderer/Private/HairStrands/HairStrandsUtils.cpp
+++ b/Engine/Source/Runtime/Renderer/Private/HairStrands/HairStrandsUtils.cpp
@@ -335,17 +335,6 @@ FIntVector ComputeDispatchCount(uint32 ItemCount, uint32 GroupSize)
 	return FIntVector(DispatchCountX, DispatchCountY, 1);
 }
 
-FIntVector ComputeDispatchCount(uint32 GroupCount)
-{
-	const uint32 DispatchCountX = FMath::FloorToInt(FMath::Sqrt(static_cast<float>(GroupCount)));
-	const uint32 DispatchCountY = DispatchCountX + FMath::DivideAndRoundUp(GroupCount - DispatchCountX * DispatchCountX, DispatchCountX);
-
-	check(DispatchCountX <= 65535);
-	check(DispatchCountY <= 65535);
-	check(GroupCount <= DispatchCountX * DispatchCountY);
-	return FIntVector(DispatchCountX, DispatchCountY, 1);
-}
-
 FVector4f PackHairRenderInfo(
 	float PrimaryRadiusAtDepth1,
 	float StableRadiusAtDepth1,
diff --git a/Engine/Source/Runtime/Renderer/Private/HairStrands/HairStrandsUtils.h b/Engine/Source/Runtime/Renderer/Private/HairStrands/HairStrandsUtils.h
--- a/Engine/Source/Runtime/Renderer/Private/HairStrands/HairStrandsUtils.h
+++ b/Engine/Source/Runtime/Renderer/Private/HairStrands/HairStrandsUtils.h
@@ -55,7 +55,6 @@ uint32 GetVendorOptimalGroupSize1D();
 FIntPoint GetVendorOptimalGroupSize2D();
 
 FIntVector ComputeDispatchCount(uint32 ItemCount, uint32 GroupSize);
-FIntVector ComputeDispatchCount(uint32 GroupCount);
 
 enum class  EHairStrandsCompositionType : uint8
 {
diff --git a/Engine/Source/Runtime/Renderer/Private/HairStrands/HairStrandsVisibility.cpp b/Engine/Source/Runtime/Renderer/Private/HairStrands/HairStrandsVisibility.cpp
--- a/Engine/Source/Runtime/Renderer/Private/HairStrands/HairStrandsVisibility.cpp
+++ b/Engine/Source/Runtime/Renderer/Private/HairStrands/HairStrandsVisibility.cpp
@@ -3880,10 +3880,12 @@ class FHairStrandsPositionChangedCS : public FGlobalShader
 		SHADER_PARAMETER_STRUCT_INCLUDE(ShaderPrint::FShaderParameters, ShaderPrintParameters)
 	END_SHADER_PARAMETER_STRUCT()
 public:
+	static uint32 GetGroupSize() { return HAIR_VERTEXCOUNT_GROUP_SIZE; }
 	static bool ShouldCompilePermutation(const FGlobalShaderPermutationParameters& Parameters) { return IsHairStrandsSupported(EHairStrandsShaderType::Strands, Parameters.Platform); }
 	static void ModifyCompilationEnvironment(const FGlobalShaderPermutationParameters& Parameters, FShaderCompilerEnvironment& OutEnvironment)
 	{
 		OutEnvironment.SetDefine(TEXT("SHADER_POSITION_CHANGED"), TEXT("1"));
+		OutEnvironment.SetDefine(TEXT("GROUP_SIZE"), GetGroupSize());
 	}
 };
 
@@ -3896,8 +3898,6 @@ static void AddHairStrandsHasPositionChangedPass(
 	FRDGBufferUAVRef InvalidationBuffer)
 {
 	const uint32 VertexCount = HairGroupPublicData->VertexCount;
-	const uint32 GroupSize = 64;
-	const FIntVector DispatchCount = ComputeDispatchCount(VertexCount, GroupSize);
 
 	FRDGBufferRef InvalidationPrintCounter = GraphBuilder.CreateBuffer(FRDGBufferDesc::CreateBufferDesc(sizeof(uint32), 1), TEXT("Hair.InvalidationPrintCounter"));
 	FRDGBufferUAVRef InvalidationPrintCounterUAV = GraphBuilder.CreateUAV(InvalidationPrintCounter, PF_R32_UINT);
@@ -3905,7 +3905,6 @@ static void AddHairStrandsHasPositionChangedPass(
 
 	FHairStrandsPositionChangedCS::FParameters* Parameters = GraphBuilder.AllocParameters<FHairStrandsPositionChangedCS::FParameters>();
 	Parameters->VertexCount = VertexCount;
-	Parameters->DispatchCountX = DispatchCount.X;
 	Parameters->PositionThreshold2 = FMath::Square(GHairStrands_InvalidationPosition_Threshold);
 	Parameters->bDrawInvalidElement = GHairStrands_InvalidationPosition_Debug > 0 ? 1u : 0u;
 	Parameters->HairStrandsVF_bIsCullingEnable = 0u;
@@ -3924,6 +3923,7 @@ static void AddHairStrandsHasPositionChangedPass(
 		Parameters->HairStrandsVF_bIsCullingEnable = 1;
 	}
 
+	const FIntVector DispatchCount(FMath::DivideAndRoundUp(VertexCount, FHairStrandsPositionChangedCS::GetGroupSize()), 1, 1);
 	TShaderMapRef<FHairStrandsPositionChangedCS> ComputeShader(View->ShaderMap);
 	FComputeShaderUtils::AddPass(
 		GraphBuilder,
