Fix incorrect group size check.

diff --git a/Engine/Plugins/Runtime/HairStrands/Source/HairStrandsCore/Private/HairStrandsMeshProjection.cpp b/Engine/Plugins/Runtime/HairStrands/Source/HairStrandsCore/Private/HairStrandsMeshProjection.cpp
--- a/Engine/Plugins/Runtime/HairStrands/Source/HairStrandsCore/Private/HairStrandsMeshProjection.cpp
+++ b/Engine/Plugins/Runtime/HairStrands/Source/HairStrandsCore/Private/HairStrandsMeshProjection.cpp
@@ -150,7 +150,7 @@ void AddSkinUpdatePass(
 	PermutationVector.Set<FSkinUpdateCS::FPrevious>(bPrevPosition);
 
 	const FIntVector DispatchGroupCount = FComputeShaderUtils::GetGroupCount(NumVertexToProcess, 64);
-	check(DispatchGroupCount.X < 65536);
+	check(DispatchGroupCount.X <= GRHIMaxDispatchThreadGroupsPerDimension.X);
 	TShaderMapRef<FSkinUpdateCS> ComputeShader(ShaderMap, PermutationVector);
 	FComputeShaderUtils::AddPass(
 		GraphBuilder,
@@ -529,7 +529,7 @@ void AddHairStrandUpdateMeshTrianglesPass(
 		PermutationVector.Set<FHairUpdateMeshTriangleCS::FPositionType>(bUseRDGPositionBuffer ? 1 : 0);
 
 		const FIntVector DispatchGroupCount = FComputeShaderUtils::GetGroupCount(RootCount, 128);
-		check(DispatchGroupCount.X < 65536);
+		check(DispatchGroupCount.X <= GRHIMaxDispatchThreadGroupsPerDimension.X);
 		TShaderMapRef<FHairUpdateMeshTriangleCS> ComputeShader(ShaderMap, PermutationVector);
 		FComputeShaderUtils::AddPass(
 			GraphBuilder,
@@ -619,7 +619,7 @@ void InternalAddHairRBFInterpolationPass(
 	Parameters->MeshSampleWeightsBuffer		= RegisterAsSRV(GraphBuilder, DeformedLODData.GetMeshSampleWeightsBuffer(FHairStrandsDeformedRootResource::FLOD::Current));
 
 	const FIntVector DispatchGroupCount = FComputeShaderUtils::GetGroupCount(VertexCount, 128);
-	check(DispatchGroupCount.X < 65536);
+	check(DispatchGroupCount.X <= GRHIMaxDispatchThreadGroupsPerDimension.X);
 	TShaderMapRef<FHairMeshesInterpolateCS> ComputeShader(ShaderMap);
 	FComputeShaderUtils::AddPass(
 		GraphBuilder,
@@ -930,7 +930,7 @@ void AddHairStrandInitMeshSamplesPass(
 			PermutationVector.Set<FHairInitMeshSamplesCS::FPositionType>(bUseRDGPositionBuffer ? 1 : 0);
 
 			const FIntVector DispatchGroupCount = FComputeShaderUtils::GetGroupCount(RestLODData.SampleCount, 128);
-			check(DispatchGroupCount.X < 65536);
+			check(DispatchGroupCount.X <= GRHIMaxDispatchThreadGroupsPerDimension.X);
 			TShaderMapRef<FHairInitMeshSamplesCS> ComputeShader(ShaderMap, PermutationVector);
 
 			FComputeShaderUtils::AddPass(
@@ -1008,7 +1008,7 @@ void AddHairStrandUpdateMeshSamplesPass(
 		Parameters->OutSampleDeformationsBuffer		= OutWeightsBuffer.UAV;
 
 		const FIntVector DispatchGroupCount = FComputeShaderUtils::GetGroupCount(RestLODData.SampleCount+4, 128);
-		check(DispatchGroupCount.X < 65536);
+		check(DispatchGroupCount.X <= GRHIMaxDispatchThreadGroupsPerDimension.X);
 		TShaderMapRef<FHairUpdateMeshSamplesCS> ComputeShader(ShaderMap);
 		FComputeShaderUtils::AddPass(
 			GraphBuilder,
diff --git a/Engine/Plugins/Runtime/HairStrands/Source/HairStrandsCore/Private/HairStrandsRendering.cpp b/Engine/Plugins/Runtime/HairStrands/Source/HairStrandsCore/Private/HairStrandsRendering.cpp
--- a/Engine/Plugins/Runtime/HairStrands/Source/HairStrandsCore/Private/HairStrandsRendering.cpp
+++ b/Engine/Plugins/Runtime/HairStrands/Source/HairStrandsCore/Private/HairStrandsRendering.cpp
@@ -82,8 +82,8 @@ static FIntVector ComputeDispatchCount(uint32 ItemCount, uint32 GroupSize)
 	const uint32 DispatchCountX = FMath::FloorToInt(FMath::Sqrt(static_cast<float>(GroupCount)));
 	const uint32 DispatchCountY = DispatchCountX + FMath::DivideAndRoundUp(GroupCount - DispatchCountX * DispatchCountX, DispatchCountX);
 
-	check(DispatchCountX <= 65535);
-	check(DispatchCountY <= 65535);
+	check(DispatchCountX <= uint32(GRHIMaxDispatchThreadGroupsPerDimension.X));
+	check(DispatchCountY <= uint32(GRHIMaxDispatchThreadGroupsPerDimension.Y));
 	check(GroupCount <= DispatchCountX * DispatchCountY);
 	return FIntVector(DispatchCountX, DispatchCountY, 1);
 }
diff --git a/Engine/Source/Runtime/Renderer/Private/HairStrands/HairStrandsClusters.cpp b/Engine/Source/Runtime/Renderer/Private/HairStrands/HairStrandsClusters.cpp
--- a/Engine/Source/Runtime/Renderer/Private/HairStrands/HairStrandsClusters.cpp
+++ b/Engine/Source/Runtime/Renderer/Private/HairStrands/HairStrandsClusters.cpp
@@ -407,7 +407,7 @@ static void AddClusterCullingPass(
 			ComputeShader,
 			Parameters,
 			DispatchIndirectParametersClusterCountDiv512, 0); // FIX ME, this could get over 65535
-		check(ClusterData.ClusterCount / 512 <= 65535);
+		check(ClusterData.ClusterCount / 512u <= uint32(GRHIMaxDispatchThreadGroupsPerDimension.X));
 	}
 
 	/// Prefix sum on the total index count per block of 512
@@ -425,7 +425,7 @@ static void AddClusterCullingPass(
 			ComputeShader,
 			Parameters,
 			DispatchIndirectParametersClusterCountDiv512Div512, 0); // FIX ME, this could get over 65535
-		check(ClusterData.ClusterCount / (512*512) <= 65535);
+		check(ClusterData.ClusterCount / (512u*512u) <= uint32(GRHIMaxDispatchThreadGroupsPerDimension.X));
 	}
 
 	/// Compact to VertexId buffer using hierarchical binary search or splatting
diff --git a/Engine/Source/Runtime/Renderer/Private/HairStrands/HairStrandsDebug.cpp b/Engine/Source/Runtime/Renderer/Private/HairStrands/HairStrandsDebug.cpp
--- a/Engine/Source/Runtime/Renderer/Private/HairStrands/HairStrandsDebug.cpp
+++ b/Engine/Source/Runtime/Renderer/Private/HairStrands/HairStrandsDebug.cpp
@@ -1169,7 +1169,7 @@ static void AddDrawDebugClusterPass(
 					}
 					ShaderPrint::SetParameters(GraphBuilder, View.ShaderPrintData, Parameters->ShaderPrintParameters);
 
-					check(Parameters->ClusterCount / 64 <= 65535);
+					check(Parameters->ClusterCount / 64u <= uint32(GRHIMaxDispatchThreadGroupsPerDimension.X));
 					const FIntVector DispatchCount = DispatchCount.DivideAndRoundUp(FIntVector(Parameters->ClusterCount, 1, 1), FIntVector(64, 1, 1));// FIX ME, this could get over 65535
 					FComputeShaderUtils::AddPass(
 						GraphBuilder, RDG_EVENT_NAME("DrawDebugClusterAABB"),
diff --git a/Engine/Source/Runtime/Renderer/Private/HairStrands/HairStrandsUtils.cpp b/Engine/Source/Runtime/Renderer/Private/HairStrands/HairStrandsUtils.cpp
--- a/Engine/Source/Runtime/Renderer/Private/HairStrands/HairStrandsUtils.cpp
+++ b/Engine/Source/Runtime/Renderer/Private/HairStrands/HairStrandsUtils.cpp
@@ -329,8 +329,8 @@ FIntVector ComputeDispatchCount(uint32 ItemCount, uint32 GroupSize)
 	const uint32 DispatchCountX = FMath::FloorToInt(FMath::Sqrt(static_cast<float>(GroupCount)));
 	const uint32 DispatchCountY = DispatchCountX + FMath::DivideAndRoundUp(GroupCount - DispatchCountX * DispatchCountX, DispatchCountX);
 
-	check(DispatchCountX <= 65535);
-	check(DispatchCountY <= 65535);
+	check(DispatchCountX <= uint32(GRHIMaxDispatchThreadGroupsPerDimension.X));
+	check(DispatchCountY <= uint32(GRHIMaxDispatchThreadGroupsPerDimension.Y));
 	check(GroupCount <= DispatchCountX * DispatchCountY);
 	return FIntVector(DispatchCountX, DispatchCountY, 1);
 }
diff --git a/Engine/Source/Runtime/Renderer/Private/HairStrands/HairStrandsVoxelization.cpp b/Engine/Source/Runtime/Renderer/Private/HairStrands/HairStrandsVoxelization.cpp
--- a/Engine/Source/Runtime/Renderer/Private/HairStrands/HairStrandsVoxelization.cpp
+++ b/Engine/Source/Runtime/Renderer/Private/HairStrands/HairStrandsVoxelization.cpp
@@ -763,7 +763,7 @@ static void AddAllocateVoxelPagesPass(
 						FMath::DivideAndRoundUp(CPUAllocationDesc.PageIndexResolution.Y, 8),
 						FMath::DivideAndRoundUp(CPUAllocationDesc.PageIndexResolution.Z, 8));
 
-				check(DispatchCount.X < 65535);
+				check(DispatchCount.X <= GRHIMaxDispatchThreadGroupsPerDimension.X);
 				FVoxelMarkValidPageIndex_PrepareCS::FPermutationDomain PermutationVector;
 				PermutationVector.Set<FVoxelMarkValidPageIndex_PrepareCS::FUseCluster>(bUseClusterAABB ? 1 : 0);
 				TShaderMapRef<FVoxelMarkValidPageIndex_PrepareCS> ComputeShader(View.ShaderMap, PermutationVector);
@@ -838,7 +838,7 @@ static void AddAllocateVoxelPagesPass(
 				PermutationVector.Set<FVoxelMarkValidPageIndexCS::FGPUDriven>(bIsGPUDriven ? 1 : 0);
 
 				FIntVector DispatchCount((Parameters->MaxClusterCount + GroupSize - 1) / GroupSize, 1, 1);
-				check(DispatchCount.X < 65535);
+				check(DispatchCount.X <= GRHIMaxDispatchThreadGroupsPerDimension.X);
 				TShaderMapRef<FVoxelMarkValidPageIndexCS> ComputeShader(View.ShaderMap, PermutationVector);
 				FComputeShaderUtils::AddPass(
 					GraphBuilder,
@@ -912,7 +912,7 @@ static void AddAllocateVoxelPagesPass(
 			else
 			{
 				const FIntVector DispatchCount((CPUAllocationDesc.PageIndexCount + GroupSize - 1) / GroupSize, 1, 1);
-				check(DispatchCount.X < 65535);
+				check(DispatchCount.X <= GRHIMaxDispatchThreadGroupsPerDimension.X);
 				FComputeShaderUtils::AddPass(GraphBuilder, RDG_EVENT_NAME("HairStrands::AllocateVoxelPage"), ComputeShader, Parameters, DispatchCount);
 			}
 		}
