Change hair culling output threading layout to be 1D.

This increase the group size of most of the post-culling consumer passes (raster, voxlization raster, tangent, interpolation, ...). This simplify the overall handling and and also reduces the permutation count for a few hair shaders. Performance remains unchanged

diff --git a/Engine/Plugins/Runtime/HairStrands/Source/HairStrandsCore/Private/HairStrandsRendering.cpp b/Engine/Plugins/Runtime/HairStrands/Source/HairStrandsCore/Private/HairStrandsRendering.cpp
--- a/Engine/Plugins/Runtime/HairStrands/Source/HairStrandsCore/Private/HairStrandsRendering.cpp
+++ b/Engine/Plugins/Runtime/HairStrands/Source/HairStrandsCore/Private/HairStrandsRendering.cpp
@@ -76,18 +76,6 @@ int GetHairRaytracingProceduralSplits()
 	return FMath::Clamp(GHairRaytracingProceduralSplits, 1, STRANDS_PROCEDURAL_INTERSECTOR_MAX_SPLITS);
 }
 
-static FIntVector ComputeDispatchCount(uint32 ItemCount, uint32 GroupSize)
-{
-	const uint32 GroupCount = FMath::DivideAndRoundUp(ItemCount, GroupSize);
-	const uint32 DispatchCountX = FMath::FloorToInt(FMath::Sqrt(static_cast<float>(GroupCount)));
-	const uint32 DispatchCountY = DispatchCountX + FMath::DivideAndRoundUp(GroupCount - DispatchCountX * DispatchCountX, DispatchCountX);
-
-	check(DispatchCountX <= uint32(GRHIMaxDispatchThreadGroupsPerDimension.X));
-	check(DispatchCountY <= uint32(GRHIMaxDispatchThreadGroupsPerDimension.Y));
-	check(GroupCount <= DispatchCountX * DispatchCountY);
-	return FIntVector(DispatchCountX, DispatchCountY, 1);
-}
-
 inline uint32 ComputeGroupSize()
 {
 	const uint32 GroupSize = IsRHIDeviceAMD() ? 64 : (IsRHIDeviceNVIDIA() ? 32 : 64);
@@ -110,7 +98,6 @@ class FTransferVelocityPassCS : public FGlobalShader
 
 	BEGIN_SHADER_PARAMETER_STRUCT(FParameters, )
 		SHADER_PARAMETER(uint32, ElementCount)
-		SHADER_PARAMETER(uint32, DispatchCountX)
 		SHADER_PARAMETER_RDG_BUFFER_SRV(Buffer, InBuffer)
 		SHADER_PARAMETER_RDG_BUFFER_UAV(RWBuffer, OutBuffer)
 	END_SHADER_PARAMETER_STRUCT()
@@ -461,18 +448,16 @@ class FHairInterpolationCS : public FGlobalShader
 	DECLARE_GLOBAL_SHADER(FHairInterpolationCS);
 	SHADER_USE_PARAMETER_STRUCT(FHairInterpolationCS, FGlobalShader);
 
-	class FGroupSize : SHADER_PERMUTATION_SPARSE_INT("PERMUTATION_GROUP_SIZE", 32, 64);
 	class FDynamicGeometry : SHADER_PERMUTATION_INT("PERMUTATION_DYNAMIC_GEOMETRY", 5);
 	class FSimulation : SHADER_PERMUTATION_BOOL("PERMUTATION_SIMULATION");
 	class FSingleGuide : SHADER_PERMUTATION_BOOL("PERMUTATION_USE_SINGLE_GUIDE");
 	class FCulling : SHADER_PERMUTATION_BOOL("PERMUTATION_CULLING");
 	class FDeformer : SHADER_PERMUTATION_BOOL("PERMUTATION_DEFORMER");
-	using FPermutationDomain = TShaderPermutationDomain<FGroupSize, FDynamicGeometry, FSimulation, FSingleGuide, FCulling, FDeformer>;
+	using FPermutationDomain = TShaderPermutationDomain<FDynamicGeometry, FSimulation, FSingleGuide, FCulling, FDeformer>;
 
 	BEGIN_SHADER_PARAMETER_STRUCT(FParameters, )
 		SHADER_PARAMETER_STRUCT_INCLUDE(ShaderPrint::FShaderParameters, ShaderDrawParameters)
 		SHADER_PARAMETER(uint32, VertexCount)
-		SHADER_PARAMETER(uint32, DispatchCountX)
 		SHADER_PARAMETER(float, HairLengthScale)
 		SHADER_PARAMETER(FVector3f, InRenHairPositionOffset)
 		SHADER_PARAMETER(FVector3f, InSimHairPositionOffset)
@@ -527,11 +512,13 @@ class FHairInterpolationCS : public FGlobalShader
 		END_SHADER_PARAMETER_STRUCT()
 
 public:
+	static uint32 GetGroupSize() { return HAIR_VERTEXCOUNT_GROUP_SIZE; }
 	static bool ShouldCompilePermutation(const FGlobalShaderPermutationParameters& Parameters) { return IsHairStrandsSupported(EHairStrandsShaderType::All, Parameters.Platform); }
 	static void ModifyCompilationEnvironment(const FGlobalShaderPermutationParameters& Parameters, FShaderCompilerEnvironment& OutEnvironment)
 	{
 		FGlobalShader::ModifyCompilationEnvironment(Parameters, OutEnvironment);
 		OutEnvironment.SetDefine(TEXT("SHADER_HAIRINTERPOLATION"), 1);
+		OutEnvironment.SetDefine(TEXT("GROUP_SIZE"), GetGroupSize());
 	}
 };
 
@@ -578,10 +565,6 @@ static void AddHairStrandsInterpolationPass(
 	FRDGImportedBuffer& OutRenPositionBuffer,
 	const FHairStrandsDeformedRootResource::FLOD::EFrameType DeformedFrame)
 {
-	const uint32 GroupSize = ComputeGroupSize();
-	const FIntVector DispatchCount = ComputeDispatchCount(VertexCount, GroupSize);
-
-
 	FHairInterpolationCS::FParameters* Parameters = GraphBuilder.AllocParameters<FHairInterpolationCS::FParameters>();
 	Parameters->RenRestPosePositionBuffer = RenRestPosePositionBuffer;
 	Parameters->SimRestPosePositionBuffer = SimRestPosePositionBuffer;
@@ -604,7 +587,6 @@ static void AddHairStrandsInterpolationPass(
 	Parameters->OutSimHairPositionOffsetBuffer = OutSimHairPositionOffsetBuffer;
 	Parameters->OutRenHairPositionOffsetBuffer = OutRenHairPositionOffsetBuffer;
 
-	Parameters->DispatchCountX = DispatchCount.X;
 	Parameters->SimRootPointIndexBuffer = SimRootPointIndexBuffer;
 	Parameters->SimVertexToCurveBuffer = SimVertexToCurveBuffer;
 	Parameters->RenVertexToCurveBuffer = RenVertexToCurveBuffer;
@@ -698,7 +680,6 @@ static void AddHairStrandsInterpolationPass(
 	}
 
 	FHairInterpolationCS::FPermutationDomain PermutationVector;
-	PermutationVector.Set<FHairInterpolationCS::FGroupSize>(GroupSize);
 	PermutationVector.Set<FHairInterpolationCS::FDynamicGeometry>(DynamicGeometryType);
 	PermutationVector.Set<FHairInterpolationCS::FSimulation>(bHasLocalDeformation);
 	PermutationVector.Set<FHairInterpolationCS::FSingleGuide>(bUseSingleGuide);
@@ -724,6 +705,7 @@ static void AddHairStrandsInterpolationPass(
 	}
 	else
 	{
+		const FIntVector DispatchCount(FMath::DivideAndRoundUp(VertexCount, FHairInterpolationCS::GetGroupSize()), 1, 1);
 		FComputeShaderUtils::AddPass(
 			GraphBuilder,
 			RDG_EVENT_NAME("HairStrands::Interpolation(culling=off)"),
@@ -1152,13 +1134,11 @@ class FHairTangentCS : public FGlobalShader
 	DECLARE_GLOBAL_SHADER(FHairTangentCS);
 	SHADER_USE_PARAMETER_STRUCT(FHairTangentCS, FGlobalShader);
 
-	class FGroupSize : SHADER_PERMUTATION_SPARSE_INT("PERMUTATION_GROUP_SIZE", 32, 64);
-	class FCulling : SHADER_PERMUTATION_INT("PERMUTATION_CULLING", 2);
-	using FPermutationDomain = TShaderPermutationDomain<FGroupSize, FCulling>;
+	class FCulling : SHADER_PERMUTATION_BOOL("PERMUTATION_CULLING");
+	using FPermutationDomain = TShaderPermutationDomain<FCulling>;
 
 	BEGIN_SHADER_PARAMETER_STRUCT(FParameters, )
 		SHADER_PARAMETER(uint32, VertexCount)
-		SHADER_PARAMETER(uint32, DispatchCountX)
 		SHADER_PARAMETER(uint32, HairStrandsVF_bIsCullingEnable)
 		SHADER_PARAMETER_RDG_BUFFER_SRV(Buffer, PositionBuffer)
 		SHADER_PARAMETER_RDG_BUFFER_SRV(Buffer<uint>,	HairStrandsVF_CullingIndirectBuffer)
@@ -1168,7 +1148,14 @@ class FHairTangentCS : public FGlobalShader
 	END_SHADER_PARAMETER_STRUCT()
 
 public:
+	static uint32 GetGroupSize() { return HAIR_VERTEXCOUNT_GROUP_SIZE; }
 	static bool ShouldCompilePermutation(const FGlobalShaderPermutationParameters& Parameters) { return IsHairStrandsSupported(EHairStrandsShaderType::Strands, Parameters.Platform); }
+	static void ModifyCompilationEnvironment(const FGlobalShaderPermutationParameters& Parameters, FShaderCompilerEnvironment& OutEnvironment)
+	{
+		FGlobalShader::ModifyCompilationEnvironment(Parameters, OutEnvironment);
+		OutEnvironment.SetDefine(TEXT("SHADER_TANGENT"), 1);
+		OutEnvironment.SetDefine(TEXT("GROUP_SIZE"), GetGroupSize());
+	}
 };
 
 IMPLEMENT_GLOBAL_SHADER(FHairTangentCS, "/Engine/Private/HairStrands/HairStrandsTangent.usf", "MainCS", SF_Compute);
@@ -1181,20 +1168,16 @@ void AddHairTangentPass(
 	FRDGBufferSRVRef PositionBuffer,
 	FRDGImportedBuffer OutTangentBuffer)
 {
-	const uint32 GroupSize = ComputeGroupSize();
-	const FIntVector DispatchCount = ComputeDispatchCount(VertexCount, GroupSize);
 	const bool bCullingEnable = IsHairStrandContinuousDecimationReorderingEnabled() ? false : (HairGroupPublicData && HairGroupPublicData->GetCullingResultAvailable());
 
 	FHairTangentCS::FParameters* Parameters = GraphBuilder.AllocParameters<FHairTangentCS::FParameters>();
 	Parameters->PositionBuffer = PositionBuffer;
 	Parameters->OutputTangentBuffer = OutTangentBuffer.UAV;
 	Parameters->VertexCount = VertexCount;
-	Parameters->DispatchCountX = DispatchCount.X;
 	Parameters->HairStrandsVF_bIsCullingEnable = bCullingEnable ? 1 : 0;
 
 	FHairTangentCS::FPermutationDomain PermutationVector;
-	PermutationVector.Set<FHairTangentCS::FGroupSize>(GroupSize);
-	PermutationVector.Set<FHairTangentCS::FCulling>(bCullingEnable ? 1 : 0);
+	PermutationVector.Set<FHairTangentCS::FCulling>(bCullingEnable);
 
 	TShaderMapRef<FHairTangentCS> ComputeShader(ShaderMap, PermutationVector);
 
@@ -1214,6 +1197,7 @@ void AddHairTangentPass(
 	}
 	else
 	{
+		const FIntVector DispatchCount(FMath::DivideAndRoundUp(VertexCount, FHairTangentCS::GetGroupSize()), 1, 1);
 		FComputeShaderUtils::AddPass(
 			GraphBuilder,
 			RDG_EVENT_NAME("HairStrands::Tangent(culling=off)"),
diff --git a/Engine/Shaders/Private/HairStrands/HairStrandsClusterCulling.usf b/Engine/Shaders/Private/HairStrands/HairStrandsClusterCulling.usf
--- a/Engine/Shaders/Private/HairStrands/HairStrandsClusterCulling.usf
+++ b/Engine/Shaders/Private/HairStrands/HairStrandsClusterCulling.usf
@@ -433,14 +433,6 @@ void MainClusterCullingCS(uint3 DispatchThreadId : SV_DispatchThreadID)
 RWBuffer<uint> DrawIndirectBuffer; 
 RWBuffer<uint> DispatchIndirectBuffer; 
 
-#if PERMUTATION_GROUP_SIZE == 64
-	#define GROUP_SIZE 64
-#elif PERMUTATION_GROUP_SIZE == 32
-	#define GROUP_SIZE 32
-#else
-	#error Unknown group size
-#endif
-
 [numthreads(1, 1, 1)]
 void MainClusterCullingPrepareIndirectDispatchCS(
 	uint3 DispatchThreadId : SV_DispatchThreadID, 
@@ -451,18 +443,12 @@ void MainClusterCullingPrepareIndirectDispatchCS(
 	{
 		// DrawIndirectBuffer contains the total number of vertex count needed for drawing quad triangle strips, so 6 vertices per quad / control point
 		const uint VertexCount    = DrawIndirectBuffer[0];
-#if USE_HAIR_TRIANGLE_STRIP
-		const uint ControlPoint   = VertexCount / 2;
-#else
-		const uint ControlPoint	  = VertexCount / 6;
-#endif
-		const uint DispatchCount  = DivideAndRoundUp(ControlPoint, GROUP_SIZE);
-		const uint DispatchCountX = DispatchCount >= INDIRECT_GROUP_COUNT_X ? INDIRECT_GROUP_COUNT_X : (DispatchCount % INDIRECT_GROUP_COUNT_X);
-		const uint DispatchCountY = DivideAndRoundUp(DispatchCount, INDIRECT_GROUP_COUNT_X);
+		const uint ControlPoint   = VertexCount / (USE_HAIR_TRIANGLE_STRIP ? 2u : 6u);
+		const uint DispatchCount  = DivideAndRoundUp(ControlPoint, HAIR_VERTEXCOUNT_GROUP_SIZE);
 
-		DispatchIndirectBuffer[0] = DispatchCountX;
-		DispatchIndirectBuffer[1] = DispatchCountY;
-		DispatchIndirectBuffer[2] = 1.0;
+		DispatchIndirectBuffer[0] = DispatchCount;
+		DispatchIndirectBuffer[1] = 1;
+		DispatchIndirectBuffer[2] = 1;
 		DispatchIndirectBuffer[3] = ControlPoint;
 	}
 }
diff --git a/Engine/Shaders/Private/HairStrands/HairStrandsInterpolation.usf b/Engine/Shaders/Private/HairStrands/HairStrandsInterpolation.usf
--- a/Engine/Shaders/Private/HairStrands/HairStrandsInterpolation.usf
+++ b/Engine/Shaders/Private/HairStrands/HairStrandsInterpolation.usf
@@ -72,7 +72,6 @@ FGuideData GetGuideData(uint InVertexIndex, uint4 InInterpolationData0, uint4 In
 #if SHADER_HAIRINTERPOLATION
 
 uint VertexCount;
-uint DispatchCountX;
 uint HairDebugMode;
 float HairLengthScale;
 float3 InRenHairPositionOffset;
@@ -459,11 +458,11 @@ float3 ApplyDeformer(const FRootTriangleData RenTri, float3 InRestPostion, float
 void MainCS(uint2 DispatchThreadId : SV_DispatchThreadID)
 {
 	// 1. Find vertex index
-	uint VertexIndex = GetHairStrandsVertexFetchIndex(DispatchThreadId, GROUP_SIZE, DispatchCountX);
+	uint VertexIndex = DispatchThreadId;
 #if PERMUTATION_CULLING
 	if (HairStrandsVF_bIsCullingEnable)
 	{
-		const uint FetchIndex = GetHairStrandsVertexFetchIndex_Culled(DispatchThreadId, GROUP_SIZE);
+		const uint FetchIndex = DispatchThreadId;
 		const uint VertexCountAfterCulling = HairStrandsVF_CullingIndirectBuffer[3];
 		if (FetchIndex >= VertexCountAfterCulling)
 		{
diff --git a/Engine/Shaders/Private/HairStrands/HairStrandsTangent.usf b/Engine/Shaders/Private/HairStrands/HairStrandsTangent.usf
--- a/Engine/Shaders/Private/HairStrands/HairStrandsTangent.usf
+++ b/Engine/Shaders/Private/HairStrands/HairStrandsTangent.usf
@@ -4,8 +4,8 @@
 #include "../MonteCarlo.ush"
 #include "HairStrandsClusterCommon.ush"
 
+#if SHADER_TANGENT
 uint VertexCount;
-uint DispatchCountX;
 
 Buffer<uint4>  	 PositionBuffer;
 RWBuffer<float4> OutputTangentBuffer;
@@ -18,31 +18,23 @@ Buffer<uint>	HairStrandsVF_CullingIndexBuffer;
 Buffer<uint>	IndirectBufferArgs;
 #endif
 
-#if PERMUTATION_GROUP_SIZE == 64
-	#define GROUP_SIZE 64
-#elif PERMUTATION_GROUP_SIZE == 32
-	#define GROUP_SIZE 32
-#else
-	#error Unknown group size	
-#endif
-
 float4 PackNormal(float3 V, uint A)
 {
 	return float4(clamp(V,float(-1).xxx,float(1).xxx), 1.0f);
 }
 
 [numthreads(GROUP_SIZE, 1, 1)]
-void MainCS(uint2 DispatchThreadId : SV_DispatchThreadID)
+void MainCS(uint DispatchThreadId : SV_DispatchThreadID)
 {
 	const int MaxVertexCount = int(VertexCount);
-	int IndexCurr = GetHairStrandsVertexFetchIndex(DispatchThreadId, GROUP_SIZE, DispatchCountX);
+	int IndexCurr = DispatchThreadId;
 	int IndexPrev = clamp(IndexCurr-1, 0, MaxVertexCount-1);
 	int IndexNext = clamp(IndexCurr+1, 0, MaxVertexCount-1);
 	bool bValid = IndexCurr < MaxVertexCount;
 #if PERMUTATION_CULLING
 	if (HairStrandsVF_bIsCullingEnable)
 	{
-		const int FetchIndex = GetHairStrandsVertexFetchIndex_Culled(DispatchThreadId, GROUP_SIZE);
+		const int FetchIndex = DispatchThreadId;
 		const int VertexCountAfterCulling = clamp((int)HairStrandsVF_CullingIndirectBuffer[3], 0, MaxVertexCount);
 		bValid = false;
 		if (FetchIndex < VertexCountAfterCulling)
@@ -97,4 +89,5 @@ void MainCS(uint2 DispatchThreadId : SV_DispatchThreadID)
 		OutputTangentBuffer[IndexCurr*2]   = PackNormal(Basis[0], 0);
 		OutputTangentBuffer[IndexCurr*2+1] = PackNormal(Basis[2], 127);
 	}
-}
\ No newline at end of file
+}
+#endif // SHADER_TANGENT
\ No newline at end of file
diff --git a/Engine/Shaders/Private/HairStrands/HairStrandsVoxelRasterCompute.usf b/Engine/Shaders/Private/HairStrands/HairStrandsVoxelRasterCompute.usf
--- a/Engine/Shaders/Private/HairStrands/HairStrandsVoxelRasterCompute.usf
+++ b/Engine/Shaders/Private/HairStrands/HairStrandsVoxelRasterCompute.usf
@@ -807,13 +807,10 @@ float3 GetHairVoxelJitter(uint2 PixelCoord, uint Seed)
 [numthreads(GROUP_SIZE, 1, 1)]
 void MainCS(uint2 DispatchThreadID : SV_DispatchThreadID)
 {
-#if PERMUTATION_CULLING == 1
-	uint VertexIndex0 = GetHairStrandsVertexFetchIndex_Culled(DispatchThreadID, GROUP_SIZE);
-#else
 	uint VertexIndex0 = DispatchThreadID.x;
-#endif
-	uint VertexIndex1 = min(VertexIndex0 + 1, VertexCount - 1);
-	if (VertexIndex0 >= VertexCount)
+	uint VertexIndex1 = VertexIndex0 + 1;
+	bool bIsValid = VertexIndex0 < VertexCount && VertexIndex1 < VertexCount;
+	if (!bIsValid)
 		return;
 
 	float LodRadiusScale0 = 1;
@@ -824,7 +821,8 @@ void MainCS(uint2 DispatchThreadID : SV_DispatchThreadID)
 		const uint VertexCountAfterCulling = HairStrandsVF_CullingIndirectBuffer[3];
 		uint FetchIndex0 = VertexIndex0;
 		uint FetchIndex1 = VertexIndex1;
-		if (FetchIndex0 >= VertexCountAfterCulling)
+		bIsValid = FetchIndex0 < VertexCountAfterCulling&& FetchIndex1 < VertexCountAfterCulling;
+		if (!bIsValid)
 		{
 			return;
 		}
diff --git a/Engine/Source/Runtime/Renderer/Private/HairStrands/HairStrandsClusters.cpp b/Engine/Source/Runtime/Renderer/Private/HairStrands/HairStrandsClusters.cpp
--- a/Engine/Source/Runtime/Renderer/Private/HairStrands/HairStrandsClusters.cpp
+++ b/Engine/Source/Runtime/Renderer/Private/HairStrands/HairStrandsClusters.cpp
@@ -130,8 +130,7 @@ class FMainClusterCullingPrepareIndirectDispatchCS : public FGlobalShader
 	DECLARE_GLOBAL_SHADER(FMainClusterCullingPrepareIndirectDispatchCS);
 	SHADER_USE_PARAMETER_STRUCT(FMainClusterCullingPrepareIndirectDispatchCS, FGlobalShader);
 
-	class FGroupSize : SHADER_PERMUTATION_SPARSE_INT("PERMUTATION_GROUP_SIZE", 32, 64);
-	using FPermutationDomain = TShaderPermutationDomain<FGroupSize>;
+	using FPermutationDomain = TShaderPermutationDomain<>;
 
 	BEGIN_SHADER_PARAMETER_STRUCT(FParameters, )
 		SHADER_PARAMETER_RDG_BUFFER_UAV(RWBuffer, DrawIndirectBuffer)
@@ -484,9 +483,7 @@ static void AddClusterCullingPass(
 		Parameters->DrawIndirectBuffer = DrawIndirectParametersBuffer.UAV;
 		Parameters->DispatchIndirectBuffer = DrawIndirectParametersRasterComputeBuffer.UAV;
 		
-		FMainClusterCullingPrepareIndirectDispatchCS::FPermutationDomain PermutationVector;
-		PermutationVector.Set<FMainClusterCullingPrepareIndirectDispatchCS::FGroupSize>(GetVendorOptimalGroupSize1D());
-		TShaderMapRef<FMainClusterCullingPrepareIndirectDispatchCS> ComputeShader(ShaderMap, PermutationVector);
+		TShaderMapRef<FMainClusterCullingPrepareIndirectDispatchCS> ComputeShader(ShaderMap);
 		FComputeShaderUtils::AddPass(
 			GraphBuilder,
 			RDG_EVENT_NAME("PrepareIndirectDispatchCS"),
diff --git a/Engine/Source/Runtime/Renderer/Private/HairStrands/HairStrandsUtils.cpp b/Engine/Source/Runtime/Renderer/Private/HairStrands/HairStrandsUtils.cpp
--- a/Engine/Source/Runtime/Renderer/Private/HairStrands/HairStrandsUtils.cpp
+++ b/Engine/Source/Runtime/Renderer/Private/HairStrands/HairStrandsUtils.cpp
@@ -323,18 +323,6 @@ FIntPoint GetVendorOptimalGroupSize2D()
 	}
 }
 
-FIntVector ComputeDispatchCount(uint32 ItemCount, uint32 GroupSize)
-{
-	const uint32 GroupCount = FMath::DivideAndRoundUp(ItemCount, GroupSize);
-	const uint32 DispatchCountX = FMath::FloorToInt(FMath::Sqrt(static_cast<float>(GroupCount)));
-	const uint32 DispatchCountY = DispatchCountX + FMath::DivideAndRoundUp(GroupCount - DispatchCountX * DispatchCountX, DispatchCountX);
-
-	check(DispatchCountX <= uint32(GRHIMaxDispatchThreadGroupsPerDimension.X));
-	check(DispatchCountY <= uint32(GRHIMaxDispatchThreadGroupsPerDimension.Y));
-	check(GroupCount <= DispatchCountX * DispatchCountY);
-	return FIntVector(DispatchCountX, DispatchCountY, 1);
-}
-
 FVector4f PackHairRenderInfo(
 	float PrimaryRadiusAtDepth1,
 	float StableRadiusAtDepth1,
diff --git a/Engine/Source/Runtime/Renderer/Private/HairStrands/HairStrandsUtils.h b/Engine/Source/Runtime/Renderer/Private/HairStrands/HairStrandsUtils.h
--- a/Engine/Source/Runtime/Renderer/Private/HairStrands/HairStrandsUtils.h
+++ b/Engine/Source/Runtime/Renderer/Private/HairStrands/HairStrandsUtils.h
@@ -54,8 +54,6 @@ bool IsHairStrandsViewRectOptimEnable();
 uint32 GetVendorOptimalGroupSize1D();
 FIntPoint GetVendorOptimalGroupSize2D();
 
-FIntVector ComputeDispatchCount(uint32 ItemCount, uint32 GroupSize);
-
 enum class  EHairStrandsCompositionType : uint8
 {
 	BeforeTranslucent,
diff --git a/Engine/Source/Runtime/Renderer/Private/HairStrands/HairStrandsVisibility.cpp b/Engine/Source/Runtime/Renderer/Private/HairStrands/HairStrandsVisibility.cpp
--- a/Engine/Source/Runtime/Renderer/Private/HairStrands/HairStrandsVisibility.cpp
+++ b/Engine/Source/Runtime/Renderer/Private/HairStrands/HairStrandsVisibility.cpp
@@ -2907,7 +2907,6 @@ class FVisiblityRasterComputeBinningCS : public FGlobalShader
 	BEGIN_SHADER_PARAMETER_STRUCT(FParameters, )
 		SHADER_PARAMETER_STRUCT_INCLUDE(FRasterComputeCommonParameters, Common)
 		SHADER_PARAMETER(uint32, MacroGroupId)
-		SHADER_PARAMETER(uint32, DispatchCountX)
 		SHADER_PARAMETER(uint32, HairMaterialId)
 		SHADER_PARAMETER_STRUCT_INCLUDE(FHairStrandsInstanceParameters, HairInstance)
 		SHADER_PARAMETER_RDG_TEXTURE(Texture2D, SceneDepthTexture)
@@ -2980,7 +2979,6 @@ class FVisiblityRasterComputeRasterizeCS : public FGlobalShader
 	BEGIN_SHADER_PARAMETER_STRUCT(FParameters, )
 		SHADER_PARAMETER_STRUCT_INCLUDE(FRasterComputeCommonParameters, Common)
 		SHADER_PARAMETER(uint32, MacroGroupId)
-		SHADER_PARAMETER(uint32, DispatchCountX)
 		SHADER_PARAMETER(uint32, HairMaterialId)
 		SHADER_PARAMETER_STRUCT_INCLUDE(FHairStrandsInstanceParameters, HairInstance)
 		SHADER_PARAMETER_RDG_TEXTURE(Texture2D, SceneDepthTexture)
@@ -3019,8 +3017,7 @@ class FVisiblityRasterComputeRasterizeMultiSampleCS : public FGlobalShader
 	BEGIN_SHADER_PARAMETER_STRUCT(FParameters, )
 		SHADER_PARAMETER_STRUCT_INCLUDE(FRasterComputeCommonParameters, Common)
 		SHADER_PARAMETER(uint32, MacroGroupId)
-		SHADER_PARAMETER(uint32, DispatchCountX)
-		SHADER_PARAMETER(uint32, HairMaterialId)
+			SHADER_PARAMETER(uint32, HairMaterialId)
 		SHADER_PARAMETER_STRUCT_INCLUDE(FHairStrandsInstanceParameters, HairInstance)
 		SHADER_PARAMETER_RDG_TEXTURE(Texture2D, SceneDepthTexture)
 		SHADER_PARAMETER(float, SampleWeight)
diff --git a/Engine/Source/Runtime/Renderer/Private/HairStrands/HairStrandsVoxelization.cpp b/Engine/Source/Runtime/Renderer/Private/HairStrands/HairStrandsVoxelization.cpp
--- a/Engine/Source/Runtime/Renderer/Private/HairStrands/HairStrandsVoxelization.cpp
+++ b/Engine/Source/Runtime/Renderer/Private/HairStrands/HairStrandsVoxelization.cpp
@@ -1161,9 +1161,8 @@ class FVoxelRasterComputeCS : public FGlobalShader
 	DECLARE_GLOBAL_SHADER(FVoxelRasterComputeCS);
 	SHADER_USE_PARAMETER_STRUCT(FVoxelRasterComputeCS, FGlobalShader);
 
-	class FGroupSize : SHADER_PERMUTATION_SPARSE_INT("PERMUTATION_GROUP_SIZE", 32, 64);
-	class FCulling : SHADER_PERMUTATION_INT("PERMUTATION_CULLING", 2);
-	using FPermutationDomain = TShaderPermutationDomain<FGroupSize, FCulling>;
+	class FCulling : SHADER_PERMUTATION_BOOL("PERMUTATION_CULLING");
+	using FPermutationDomain = TShaderPermutationDomain<FCulling>;
 
 	BEGIN_SHADER_PARAMETER_STRUCT(FParameters, )
 		SHADER_PARAMETER_STRUCT(FHairStrandsVoxelCommonParameters, VirtualVoxelParams)
@@ -1184,11 +1183,7 @@ public:
 	{
 		FGlobalShader::ModifyCompilationEnvironment(Parameters, OutEnvironment);
 		OutEnvironment.SetDefine(TEXT("SHADER_RASTERCOMPUTE"), 1);
-
-		FPermutationDomain PermutationVector(Parameters.PermutationId);
-		const bool bUseCulling = PermutationVector.Get<FCulling>() != 0;
-		const uint32 GroupSize = PermutationVector.Get<FGroupSize>();		
-		OutEnvironment.SetDefine(TEXT("GROUP_SIZE"), bUseCulling ? GroupSize : GetGroupSize());
+		OutEnvironment.SetDefine(TEXT("GROUP_SIZE"), GetGroupSize());
 	}
 };
 
@@ -1216,15 +1211,12 @@ static void AddVirtualVoxelizationComputeRasterPass(
 		FRDGTextureUAVRef PageTextureUAV = GraphBuilder.CreateUAV(VoxelResources.PageTexture);
 
 		const uint32 FrameIdMode8 = ViewInfo && ViewInfo->ViewState ? (ViewInfo->ViewState->GetFrameIndex() % 8) : 0;
-		const uint32 GroupSize = GetVendorOptimalGroupSize1D(); // TODO when updating culling
 		const FVector& TranslatedWorldOffset = ViewInfo->ViewMatrices.GetPreViewTranslation();
 
 		FVoxelRasterComputeCS::FPermutationDomain PermutationVector_Off;
 		FVoxelRasterComputeCS::FPermutationDomain PermutationVector_On;
-		PermutationVector_Off.Set<FVoxelRasterComputeCS::FCulling>(0);
-		PermutationVector_Off.Set<FVoxelRasterComputeCS::FGroupSize>(GroupSize);
-		PermutationVector_On.Set<FVoxelRasterComputeCS::FCulling>(1);
-		PermutationVector_On.Set<FVoxelRasterComputeCS::FGroupSize>(GroupSize);
+		PermutationVector_Off.Set<FVoxelRasterComputeCS::FCulling>(false);
+		PermutationVector_On.Set<FVoxelRasterComputeCS::FCulling>(true);
 
 		TShaderMapRef<FVoxelRasterComputeCS> ComputeShader_CullingOff(ViewInfo->ShaderMap, PermutationVector_Off);
 		TShaderMapRef<FVoxelRasterComputeCS> ComputeShader_CullingOn(ViewInfo->ShaderMap, PermutationVector_On);
