Fix Vulkan crash when looking at hair in the Mesh LOD Coloration debug view mode.

Passing a pixel format when creating a view for a structured or byte address buffer confuses some RHIs and creates the wrong kind of view. We need to pass PF_Unknown in these cases. We're working on fixing this oddity in the RHI, but in the mean time this is an easy fix for the crash.

#rnx

diff --git a/Engine/Plugins/Runtime/HairStrands/Source/HairStrandsCore/Private/GroomResources.cpp b/Engine/Plugins/Runtime/HairStrands/Source/HairStrandsCore/Private/GroomResources.cpp
--- a/Engine/Plugins/Runtime/HairStrands/Source/HairStrandsCore/Private/GroomResources.cpp
+++ b/Engine/Plugins/Runtime/HairStrands/Source/HairStrandsCore/Private/GroomResources.cpp
@@ -290,7 +290,7 @@ void InternalCreateByteAddressBufferRDG_FromBulkData(FRDGBuilder& GraphBuilder,
 		InBulkData.Unlock();
 	}
 
-	ConvertToExternalBufferWithViews(GraphBuilder, Buffer, Out, PF_R32_UINT); // TODO?
+	ConvertToExternalBufferWithViews(GraphBuilder, Buffer, Out, PF_Unknown);
 }
 
 
