Rename hair strands Vertex -> Point for sematic clarity

diff --git a/Engine/Plugins/Runtime/HairStrands/Source/HairStrandsCore/Private/Deformer/DeformerDataInterfaceGroom.cpp b/Engine/Plugins/Runtime/HairStrands/Source/HairStrandsCore/Private/Deformer/DeformerDataInterfaceGroom.cpp
--- a/Engine/Plugins/Runtime/HairStrands/Source/HairStrandsCore/Private/Deformer/DeformerDataInterfaceGroom.cpp
+++ b/Engine/Plugins/Runtime/HairStrands/Source/HairStrandsCore/Private/Deformer/DeformerDataInterfaceGroom.cpp
@@ -182,7 +182,7 @@ void FOptimusGroomDataProviderProxy::AllocateResources(FRDGBuilder& GraphBuilder
 			FResources& R = Resources.AddDefaulted_GetRef();
 			R.PositionSRV		= Register(GraphBuilder, Instance->Strands.RestResource->PositionBuffer, ERDGImportedBufferFlags::CreateSRV).SRV;
 			R.AttributeSRV		= Register(GraphBuilder, Instance->Strands.RestResource->AttributeBuffer, ERDGImportedBufferFlags::CreateSRV).SRV;
-			R.VertexToCurveSRV	= Register(GraphBuilder, Instance->Strands.RestResource->VertexToCurveBuffer, ERDGImportedBufferFlags::CreateSRV).SRV;
+			R.VertexToCurveSRV	= Register(GraphBuilder, Instance->Strands.RestResource->PointToCurveBuffer, ERDGImportedBufferFlags::CreateSRV).SRV;
 			R.CurveSRV			= Register(GraphBuilder, Instance->Strands.RestResource->CurveBuffer, ERDGImportedBufferFlags::CreateSRV).SRV;
 			R.FallbackSRV		= GraphBuilder.CreateSRV(GraphBuilder.RegisterExternalBuffer(GWhiteVertexBufferWithRDG->Buffer), PF_R16G16B16A16_UINT);
 		}
diff --git a/Engine/Plugins/Runtime/HairStrands/Source/HairStrandsCore/Private/GroomComponent.cpp b/Engine/Plugins/Runtime/HairStrands/Source/HairStrandsCore/Private/GroomComponent.cpp
--- a/Engine/Plugins/Runtime/HairStrands/Source/HairStrandsCore/Private/GroomComponent.cpp
+++ b/Engine/Plugins/Runtime/HairStrands/Source/HairStrandsCore/Private/GroomComponent.cpp
@@ -800,7 +800,7 @@ public:
 		else // if (GeometryType == EHairGeometryType::Strands)
 		{
 			VertexFactory = (FVertexFactory*)Instance->Strands.VertexFactory;
-			HairVertexCount = Instance->Strands.RestResource->GetVertexCount();
+			HairVertexCount = Instance->Strands.RestResource->GetPointCount();
 			MaxVertexIndex = HairVertexCount * HAIR_POINT_TO_VERTEX;
 			bUseCulling = Instance->Strands.bIsCullingEnabled;
 			NumPrimitive = bUseCulling ? 0 : HairVertexCount * HAIR_POINT_TO_TRIANGLE;
@@ -2679,6 +2679,9 @@ void UGroomComponent::InitResources(bool bIsBindingReloading)
 			check(GroupIt < GroomGroupsDesc.Num());
 
 			HairGroupInstance->Strands.Data = &GroupData.Strands.BulkData;
+			HairGroupInstance->HairGroupPublicData->RestPointCount = GroupData.Strands.BulkData.GetNumPoints();
+			HairGroupInstance->HairGroupPublicData->RestCurveCount = GroupData.Strands.BulkData.GetNumCurves();
+			HairGroupInstance->HairGroupPublicData->ClusterCount = 0u;
 
 			// (Lazy) Allocate interpolation resources, only if guides are required
 			if (bNeedGuides)
@@ -2797,7 +2800,7 @@ void UGroomComponent::InitResources(bool bIsBindingReloading)
 						check(GroupData.Strands.ClusterCullingResource->BulkData.LODVisibility[LODIt] == LODVisibilities[LODIt]);
 					}
 				}
-				HairGroupInstance->HairGroupPublicData->SetClusters(HairGroupInstance->Strands.ClusterCullingResource->BulkData.ClusterCount, GroupData.Strands.BulkData.GetNumPoints());
+				HairGroupInstance->HairGroupPublicData->ClusterCount = HairGroupInstance->Strands.ClusterCullingResource->BulkData.ClusterCount;
 				BeginInitResource(HairGroupInstance->HairGroupPublicData);
 			}
 
diff --git a/Engine/Plugins/Runtime/HairStrands/Source/HairStrandsCore/Private/GroomDebug.cpp b/Engine/Plugins/Runtime/HairStrands/Source/HairStrandsCore/Private/GroomDebug.cpp
--- a/Engine/Plugins/Runtime/HairStrands/Source/HairStrandsCore/Private/GroomDebug.cpp
+++ b/Engine/Plugins/Runtime/HairStrands/Source/HairStrandsCore/Private/GroomDebug.cpp
@@ -600,13 +600,13 @@ static void AddDrawDebugStrandsCVsPass(
 	Parameters->DepthTexture = DepthTexture;
 	Parameters->LinearSampler = TStaticSamplerState<SF_Bilinear, AM_Clamp, AM_Clamp, AM_Clamp>::GetRHI();
 
-	const uint32 VertexCount = Instance->HairGroupPublicData->VFInput.Strands.VertexCount;
+	const uint32 PointCount = Instance->HairGroupPublicData->VFInput.Strands.PointCount;
 	FComputeShaderUtils::AddPass(
 		GraphBuilder, 
 		RDG_EVENT_NAME("HairStrands::DrawCVs"), 
 		ComputeShader, 
 		Parameters,
-		FIntVector::DivideAndRoundUp(FIntVector(VertexCount, 1, 1), FIntVector(256, 1, 1)));
+		FIntVector::DivideAndRoundUp(FIntVector(PointCount, 1, 1), FIntVector(256, 1, 1)));
 }
 
 ///////////////////////////////////////////////////////////////////////////////////////////////////
@@ -684,8 +684,8 @@ static void AddDrawDebugCardsGuidesPass(
 	}
 	
 	const uint32 MaxCount = FMath::Max(
-		Instance->Guides.RestResource ? Instance->Guides.RestResource->GetVertexCount() * 2 : 0,
-		LOD.Guides.RestResource ? LOD.Guides.RestResource->GetVertexCount() * 2 : 0);
+		Instance->Guides.RestResource ? Instance->Guides.RestResource->GetPointCount() * 2 : 0,
+		LOD.Guides.RestResource ? LOD.Guides.RestResource->GetPointCount() * 2 : 0);
 	ShaderPrint::RequestSpaceForLines(MaxCount);
 
 	TShaderMapRef<FDrawDebugCardGuidesCS> ComputeShader(ShaderMap);
@@ -717,7 +717,7 @@ static void AddDrawDebugCardsGuidesPass(
 
 	if (bRen)
 	{
-		Parameters->RenVertexCount = LOD.Guides.RestResource->GetVertexCount();
+		Parameters->RenVertexCount = LOD.Guides.RestResource->GetPointCount();
 		Parameters->RenRestOffset = (FVector3f)LOD.Guides.RestResource->GetPositionOffset();
 		Parameters->RenRestPosition = RegisterAsSRV(GraphBuilder, LOD.Guides.RestResource->PositionBuffer);
 		if (bDeformed)
@@ -729,7 +729,7 @@ static void AddDrawDebugCardsGuidesPass(
 
 	if (!bRen)
 	{
-		Parameters->SimVertexCount = Instance->Guides.RestResource->GetVertexCount();
+		Parameters->SimVertexCount = Instance->Guides.RestResource->GetPointCount();
 		Parameters->SimRestOffset = (FVector3f)Instance->Guides.RestResource->GetPositionOffset();
 		Parameters->SimRestPosition = RegisterAsSRV(GraphBuilder, Instance->Guides.RestResource->PositionBuffer);
 		if (bDeformed)
diff --git a/Engine/Plugins/Runtime/HairStrands/Source/HairStrandsCore/Private/GroomResources.cpp b/Engine/Plugins/Runtime/HairStrands/Source/HairStrandsCore/Private/GroomResources.cpp
--- a/Engine/Plugins/Runtime/HairStrands/Source/HairStrandsCore/Private/GroomResources.cpp
+++ b/Engine/Plugins/Runtime/HairStrands/Source/HairStrandsCore/Private/GroomResources.cpp
@@ -862,7 +862,7 @@ void FHairMeshesDeformedResource::InternalRelease()
 
 FHairStrandsRestResource::FHairStrandsRestResource(FHairStrandsBulkData& InBulkData, EHairStrandsResourcesType InCurveType, const FHairResourceName& InResourceName, const FName& InOwnerName) :
 	FHairCommonResource(EHairStrandsAllocationType::Deferred, InResourceName, InOwnerName),
-	PositionBuffer(), AttributeBuffer(), VertexToCurveBuffer(), BulkData(InBulkData), CurveType(InCurveType)
+	PositionBuffer(), AttributeBuffer(), PointToCurveBuffer(), BulkData(InBulkData), CurveType(InCurveType)
 {
 	// Sanity check
 	check(!!(BulkData.Flags & FHairStrandsBulkData::DataFlags_HasData));
@@ -898,11 +898,11 @@ void FHairStrandsRestResource::InternalAllocate(FRDGBuilder& GraphBuilder)
 	InternalCreateVertexBufferRDG_FromBulkData<FHairStrandsCurveFormat>(GraphBuilder, BulkData.Curves, CurveCount, CurveBuffer, ToHairResourceDebugName(HAIRSTRANDS_RESOUCE_NAME(CurveType, Hair.StrandsRest_CurveBuffer), ResourceName), OwnerName, EHairResourceUsageType::Static);
 	if (!!(BulkData.Flags & BulkData.DataFlags_Has16bitsCurveIndex))
 	{
-		InternalCreateVertexBufferRDG_FromBulkData<FHairStrandsVertexToCurveFormat16>(GraphBuilder, BulkData.VertexToCurve, PointCount, VertexToCurveBuffer, ToHairResourceDebugName(HAIRSTRANDS_RESOUCE_NAME(CurveType, Hair.StrandsRest_VertexToCurveBuffer), ResourceName), OwnerName, EHairResourceUsageType::Static);
+		InternalCreateVertexBufferRDG_FromBulkData<FHairStrandsVertexToCurveFormat16>(GraphBuilder, BulkData.VertexToCurve, PointCount, PointToCurveBuffer, ToHairResourceDebugName(HAIRSTRANDS_RESOUCE_NAME(CurveType, Hair.StrandsRest_PointToCurveBuffer), ResourceName), OwnerName, EHairResourceUsageType::Static);
 	}
 	else
 	{
-		InternalCreateVertexBufferRDG_FromBulkData<FHairStrandsVertexToCurveFormat32>(GraphBuilder, BulkData.VertexToCurve, PointCount, VertexToCurveBuffer, ToHairResourceDebugName(HAIRSTRANDS_RESOUCE_NAME(CurveType, Hair.StrandsRest_VertexToCurveBuffer), ResourceName), OwnerName, EHairResourceUsageType::Static);
+		InternalCreateVertexBufferRDG_FromBulkData<FHairStrandsVertexToCurveFormat32>(GraphBuilder, BulkData.VertexToCurve, PointCount, PointToCurveBuffer, ToHairResourceDebugName(HAIRSTRANDS_RESOUCE_NAME(CurveType, Hair.StrandsRest_PointToCurveBuffer), ResourceName), OwnerName, EHairResourceUsageType::Static);
 	}
 
 	TArray<FVector4f> RestOffset;
@@ -943,7 +943,7 @@ void FHairStrandsRestResource::InternalRelease()
 	PositionBuffer.Release();
 	PositionOffsetBuffer.Release();
 	AttributeBuffer.Release();
-	VertexToCurveBuffer.Release();
+	PointToCurveBuffer.Release();
 	TangentBuffer.Release();
 	CurveBuffer.Release();
 }
diff --git a/Engine/Plugins/Runtime/HairStrands/Source/HairStrandsCore/Private/GroomTextureBuilder.cpp b/Engine/Plugins/Runtime/HairStrands/Source/HairStrandsCore/Private/GroomTextureBuilder.cpp
--- a/Engine/Plugins/Runtime/HairStrands/Source/HairStrandsCore/Private/GroomTextureBuilder.cpp
+++ b/Engine/Plugins/Runtime/HairStrands/Source/HairStrandsCore/Private/GroomTextureBuilder.cpp
@@ -1089,7 +1089,7 @@ static bool TraceTextures(
 
 				GroupData.Strands.RestResource->PositionBuffer.SRV,
 				GroupData.Strands.RestResource->AttributeBuffer.SRV,
-				GroupData.Strands.RestResource->VertexToCurveBuffer.SRV,
+				GroupData.Strands.RestResource->PointToCurveBuffer.SRV,
 				GroupData.Strands.RestResource->GetPositionOffset(),
 				RenderingData.GeometrySettings.HairWidth * 0.5f,
 				GroupData.Strands.RestResource->BulkData.MaxLength,
diff --git a/Engine/Plugins/Runtime/HairStrands/Source/HairStrandsCore/Private/HairStrandsRendering.cpp b/Engine/Plugins/Runtime/HairStrands/Source/HairStrandsCore/Private/HairStrandsRendering.cpp
--- a/Engine/Plugins/Runtime/HairStrands/Source/HairStrandsCore/Private/HairStrandsRendering.cpp
+++ b/Engine/Plugins/Runtime/HairStrands/Source/HairStrandsCore/Private/HairStrandsRendering.cpp
@@ -898,7 +898,7 @@ static void AddHairClusterAABBPass(
 	Parameters->RenderDeformedOffsetBuffer = RenderDeformedOffsetBuffer;
 	Parameters->TotalClusterCount = 1;
 
-	Parameters->VertexCount = Instance->HairGroupPublicData->GetActiveStrandsVertexCount(Instance->Strands.RestResource->GetVertexCount(), Instance->HairGroupPublicData->MaxScreenSize);
+	Parameters->VertexCount = Instance->HairGroupPublicData->GetActiveStrandsPointCount(Instance->Strands.RestResource->GetPointCount(), Instance->HairGroupPublicData->MaxScreenSize);
 	Parameters->RcpSampleWeight = 1.0 / Instance->HairGroupPublicData->GetActiveStrandsSampleWeight(true, Instance->HairGroupPublicData->MaxScreenSize);
 
 	if (ShaderPrintData)
@@ -1052,7 +1052,7 @@ static void AddHairCardsDeformationPass(
 	
 	FHairCardsDeformationCS::FParameters* Parameters = GraphBuilder.AllocParameters<FHairCardsDeformationCS::FParameters>();
 	Parameters->LocalToWorld				= FMatrix44f(Instance->GetCurrentLocalToWorld().ToMatrixWithScale());		// LWC_TODO: Precision loss
-	Parameters->GuideVertexCount			= LOD.Guides.RestResource->GetVertexCount();
+	Parameters->GuideVertexCount			= LOD.Guides.RestResource->GetPointCount();
 	Parameters->GuideRestPositionOffset		= (FVector3f)LOD.Guides.RestResource->GetPositionOffset();
 	Parameters->GuideRestPositionBuffer		= RegisterAsSRV(GraphBuilder, LOD.Guides.RestResource->PositionBuffer);
 	Parameters->GuideDeformedPositionBuffer = RegisterAsSRV(GraphBuilder, LOD.Guides.DeformedResource->GetBuffer(FHairStrandsDeformedResource::Current));
@@ -1081,7 +1081,7 @@ static void AddHairCardsDeformationPass(
 		DeformedRootResources->LODs[MeshLODIndex].IsValid();
 	if (bSupportDynamicMesh)
 	{
-		Parameters->GuideVertexToRootIndexBuffer = RegisterAsSRV(GraphBuilder, RestResources->VertexToCurveBuffer);
+		Parameters->GuideVertexToRootIndexBuffer = RegisterAsSRV(GraphBuilder, RestResources->PointToCurveBuffer);
 		const FHairStrandsRestRootResource::FLOD& RestLODDatas = RestRootResources->LODs[MeshLODIndex];
 		const FHairStrandsDeformedRootResource::FLOD& DeformedLODDatas = DeformedRootResources->LODs[MeshLODIndex];
 
@@ -1551,7 +1551,7 @@ void RegisterClusterData(FHairGroupInstance* Instance, FHairStrandClusterData* I
 	const int32 ClusterDataGroupIndex = InClusterData->HairGroups.Num();
 	FHairStrandClusterData::FHairGroup& HairGroupCluster = InClusterData->HairGroups.Emplace_GetRef();
 	HairGroupCluster.ClusterCount = Instance->HairGroupPublicData->GetClusterCount();
-	HairGroupCluster.VertexCount = Instance->HairGroupPublicData->GetGroupInstanceVertexCount();
+	HairGroupCluster.VertexCount = Instance->HairGroupPublicData->RestPointCount * HAIR_POINT_TO_VERTEX;
 	HairGroupCluster.GroupAABBBuffer = &Instance->HairGroupPublicData->GetGroupAABBBuffer();
 	HairGroupCluster.ClusterAABBBuffer = &Instance->HairGroupPublicData->GetClusterAABBBuffer();
 
@@ -1611,13 +1611,13 @@ static FHairGroupPublicData::FVertexFactoryInput InternalComputeHairStrandsVerte
 		CONVERT_HAIRSSTRANDS_VF_PARAMETERS(OutVFInput.Strands.TangentBuffer,			Instance->Guides.DeformedResource->TangentBuffer);
 		CONVERT_HAIRSSTRANDS_VF_PARAMETERS(OutVFInput.Strands.AttributeBuffer,			Instance->Guides.RestResource->AttributeBuffer);
 		CONVERT_HAIRSSTRANDS_VF_PARAMETERS(OutVFInput.Strands.CurveBuffer,				Instance->Guides.RestResource->CurveBuffer);
-		CONVERT_HAIRSSTRANDS_VF_PARAMETERS(OutVFInput.Strands.VertexToCurveBuffer,		Instance->Guides.RestResource->VertexToCurveBuffer);
+		CONVERT_HAIRSSTRANDS_VF_PARAMETERS(OutVFInput.Strands.PointToCurveBuffer,		Instance->Guides.RestResource->PointToCurveBuffer);
 		CONVERT_HAIRSSTRANDS_VF_PARAMETERS(OutVFInput.Strands.PositionOffsetBuffer,		Instance->Guides.DeformedResource->GetPositionOffsetBuffer(FHairStrandsDeformedResource::EFrameType::Current));
 		CONVERT_HAIRSSTRANDS_VF_PARAMETERS(OutVFInput.Strands.PrevPositionOffsetBuffer,	Instance->Guides.DeformedResource->GetPositionOffsetBuffer(FHairStrandsDeformedResource::EFrameType::Previous));
 
 		OutVFInput.Strands.PositionOffset = Instance->Guides.DeformedResource->GetPositionOffset(FHairStrandsDeformedResource::EFrameType::Current);
 		OutVFInput.Strands.PrevPositionOffset = Instance->Guides.DeformedResource->GetPositionOffset(FHairStrandsDeformedResource::EFrameType::Previous);
-		OutVFInput.Strands.VertexCount = Instance->Guides.RestResource->GetVertexCount();
+		OutVFInput.Strands.PointCount = Instance->Guides.RestResource->GetPointCount();
 		OutVFInput.Strands.CurveCount = Instance->Guides.RestResource->GetCurveCount();
 
 		bUseGuideAttributeOffsets = true;
@@ -1650,10 +1650,10 @@ static FHairGroupPublicData::FVertexFactoryInput InternalComputeHairStrandsVerte
 		}
 		CONVERT_HAIRSSTRANDS_VF_PARAMETERS(OutVFInput.Strands.TangentBuffer,			Instance->Strands.DeformedResource->TangentBuffer);
 		CONVERT_HAIRSSTRANDS_VF_PARAMETERS(OutVFInput.Strands.AttributeBuffer,			(ViewMode == EGroomViewMode::RenderHairStrands || ViewMode == EGroomViewMode::Cluster || ViewMode == EGroomViewMode::ClusterAABB) ? Instance->Strands.DebugAttributeBuffer : Instance->Strands.RestResource->AttributeBuffer);
-		CONVERT_HAIRSSTRANDS_VF_PARAMETERS(OutVFInput.Strands.VertexToCurveBuffer,		Instance->Strands.RestResource->VertexToCurveBuffer);
+		CONVERT_HAIRSSTRANDS_VF_PARAMETERS(OutVFInput.Strands.PointToCurveBuffer,		Instance->Strands.RestResource->PointToCurveBuffer);
 		CONVERT_HAIRSSTRANDS_VF_PARAMETERS(OutVFInput.Strands.CurveBuffer,				Instance->Strands.RestResource->CurveBuffer);
 
-		OutVFInput.Strands.VertexCount = Instance->Strands.RestResource->GetVertexCount();
+		OutVFInput.Strands.PointCount = Instance->Strands.RestResource->GetPointCount();
 		OutVFInput.Strands.CurveCount = Instance->Strands.RestResource->GetCurveCount();
 	}
 	// 3. Render Strands in rest position: used when there are no skinning (rigid binding or no binding), no simulation, and no RBF
@@ -1663,14 +1663,14 @@ static FHairGroupPublicData::FVertexFactoryInput InternalComputeHairStrandsVerte
 		CONVERT_HAIRSSTRANDS_VF_PARAMETERS(OutVFInput.Strands.PrevPositionBuffer,		Instance->Strands.RestResource->PositionBuffer);
 		CONVERT_HAIRSSTRANDS_VF_PARAMETERS(OutVFInput.Strands.TangentBuffer,			Instance->Strands.RestResource->TangentBuffer);
 		CONVERT_HAIRSSTRANDS_VF_PARAMETERS(OutVFInput.Strands.AttributeBuffer,			(ViewMode == EGroomViewMode::Cluster || ViewMode == EGroomViewMode::ClusterAABB) ? Instance->Strands.DebugAttributeBuffer : Instance->Strands.RestResource->AttributeBuffer);
-		CONVERT_HAIRSSTRANDS_VF_PARAMETERS(OutVFInput.Strands.VertexToCurveBuffer,		Instance->Strands.RestResource->VertexToCurveBuffer);
+		CONVERT_HAIRSSTRANDS_VF_PARAMETERS(OutVFInput.Strands.PointToCurveBuffer,		Instance->Strands.RestResource->PointToCurveBuffer);
 		CONVERT_HAIRSSTRANDS_VF_PARAMETERS(OutVFInput.Strands.CurveBuffer,				Instance->Strands.RestResource->CurveBuffer);
 		CONVERT_HAIRSSTRANDS_VF_PARAMETERS(OutVFInput.Strands.PositionOffsetBuffer,		Instance->Strands.RestResource->PositionOffsetBuffer);
 		CONVERT_HAIRSSTRANDS_VF_PARAMETERS(OutVFInput.Strands.PrevPositionOffsetBuffer, Instance->Strands.RestResource->PositionOffsetBuffer);
 
 		OutVFInput.Strands.PositionOffset = Instance->Strands.RestResource->GetPositionOffset();
 		OutVFInput.Strands.PrevPositionOffset = Instance->Strands.RestResource->GetPositionOffset();
-		OutVFInput.Strands.VertexCount = Instance->Strands.RestResource->GetVertexCount();
+		OutVFInput.Strands.PointCount = Instance->Strands.RestResource->GetPointCount();
 		OutVFInput.Strands.CurveCount = Instance->Strands.RestResource->GetCurveCount();
 	}
 
@@ -1751,7 +1751,7 @@ void ComputeHairStrandsInterpolation(
 			AddHairTangentPass(
 				GraphBuilder,
 				ShaderMap,
-				Instance->Guides.RestResource->GetVertexCount(),
+				Instance->Guides.RestResource->GetPointCount(),
 				Instance->HairGroupPublicData,
 				RegisterAsSRV(GraphBuilder, Instance->Guides.DeformedResource->GetBuffer(FHairStrandsDeformedResource::Current)),
 				Register(GraphBuilder, Instance->Guides.DeformedResource->TangentBuffer, ERDGImportedBufferFlags::CreateUAV));
@@ -1836,7 +1836,7 @@ void ComputeHairStrandsInterpolation(
 					AddGroomCacheUpdatePass(
 						GraphBuilder,
 						ShaderMap,
-						Instance->Guides.RestResource->GetVertexCount(),
+						Instance->Guides.RestResource->GetPointCount(),
 						GroomCacheGroupData,
 						RegisterAsSRV(GraphBuilder, Instance->Guides.RestResource->PositionBuffer),
 						RegisterAsSRV(GraphBuilder, Instance->Guides.DeformedResource->GetPositionOffsetBuffer(FHairStrandsDeformedResource::EFrameType::Current)),
@@ -1861,11 +1861,11 @@ void ComputeHairStrandsInterpolation(
 									ShaderMap,
 									EDeformationType::OffsetGuide,
 									MeshLODIndex,
-									Instance->Guides.RestResource->GetVertexCount(),
+									Instance->Guides.RestResource->GetPointCount(),
 									Instance->Guides.RestRootResource,
 									Instance->Guides.DeformedRootResource,
 									RegisterAsSRV(GraphBuilder, Instance->Guides.RestResource->PositionBuffer),
-									RegisterAsSRV(GraphBuilder, Instance->Guides.RestResource->VertexToCurveBuffer),
+									RegisterAsSRV(GraphBuilder, Instance->Guides.RestResource->PointToCurveBuffer),
 									Register(GraphBuilder, Instance->Guides.DeformedResource->GetBuffer(FHairStrandsDeformedResource::Current), ERDGImportedBufferFlags::CreateUAV),
 									Instance->Guides.RestResource->GetPositionOffset(),
 									RegisterAsSRV(GraphBuilder, Instance->Guides.DeformedResource->GetPositionOffsetBuffer(FHairStrandsDeformedResource::Current)),
@@ -1876,12 +1876,12 @@ void ComputeHairStrandsInterpolation(
 					}
 				}
 
-				uint32 VertexCount = Instance->Strands.RestResource->GetVertexCount();
+				uint32 PointCount = Instance->Strands.RestResource->GetPointCount();
 
 				// 2.1 Compute deformation position based on simulation/skinning/RBF
 				if (Instance->Debug.GroomCacheType != EGroomCacheType::Strands)
 				{
-					VertexCount = Instance->HairGroupPublicData->GetActiveStrandsVertexCount(Instance->Strands.RestResource->GetVertexCount(), Instance->HairGroupPublicData->MaxScreenSize);
+					PointCount = Instance->HairGroupPublicData->GetActiveStrandsPointCount(Instance->Strands.RestResource->GetPointCount(), Instance->HairGroupPublicData->MaxScreenSize);
 
 					// 2.1.1 Current Position
 					AddHairStrandsInterpolationPass(
@@ -1889,7 +1889,7 @@ void ComputeHairStrandsInterpolation(
 						ShaderMap,
 						ShaderPrintData,
 						Instance,
-						VertexCount,
+						PointCount,
 						MeshLODIndex,
 						Instance->Strands.Modifier.HairLengthScale,
 						Instance->Strands.HairInterpolationType,
@@ -1904,7 +1904,7 @@ void ComputeHairStrandsInterpolation(
 						bHasSkinning ? Instance->Strands.DeformedRootResource : nullptr,
 						bHasSkinning && bValidGuide ? Instance->Guides.DeformedRootResource : nullptr,
 						RegisterAsSRV(GraphBuilder, Instance->Strands.RestResource->PositionBuffer),
-						RegisterAsSRV(GraphBuilder, Instance->Strands.RestResource->VertexToCurveBuffer),
+						RegisterAsSRV(GraphBuilder, Instance->Strands.RestResource->PointToCurveBuffer),
 						bUseSingleGuide,
 						bValidGuide &&  bUseSingleGuide ? RegisterAsSRV(GraphBuilder, Instance->Strands.InterpolationResource->InterpolationBuffer) : nullptr,
 						bValidGuide && !bUseSingleGuide ? RegisterAsSRV(GraphBuilder, Instance->Strands.InterpolationResource->Interpolation0Buffer) : nullptr,
@@ -1912,7 +1912,7 @@ void ComputeHairStrandsInterpolation(
 						bValidGuide ? RegisterAsSRV(GraphBuilder, Instance->Guides.RestResource->PositionBuffer) : nullptr,
 						bValidGuide ? RegisterAsSRV(GraphBuilder, Instance->Guides.DeformedResource->GetBuffer(FHairStrandsDeformedResource::Current)) : nullptr,
 						bValidGuide ? RegisterAsSRV(GraphBuilder, Instance->Strands.InterpolationResource->SimRootPointIndexBuffer) : nullptr,
-						bValidGuide ? RegisterAsSRV(GraphBuilder, Instance->Guides.RestResource->VertexToCurveBuffer) : nullptr,
+						bValidGuide ? RegisterAsSRV(GraphBuilder, Instance->Guides.RestResource->PointToCurveBuffer) : nullptr,
 						Strands_DeformerPositionSRV,
 						Strands_DeformedPosition,
 						FHairStrandsDeformedRootResource::FLOD::Current);
@@ -1927,7 +1927,7 @@ void ComputeHairStrandsInterpolation(
 						GHairStrandsTransferPositionOnLODChange > 0;
 					if (bTransferPrevPosition)
 					{
-						AddTransferPositionPass(GraphBuilder, ShaderMap, VertexCount, Strands_DeformedPosition.SRV, Strands_DeformedPrevPosition.UAV);
+						AddTransferPositionPass(GraphBuilder, ShaderMap, PointCount, Strands_DeformedPosition.SRV, Strands_DeformedPrevPosition.UAV);
 						GraphBuilder.SetBufferAccessFinal(Strands_DeformedPrevPosition.Buffer, ERHIAccess::SRVMask);
 					}
 					else if (bRecomputePrevPosition)
@@ -1939,7 +1939,7 @@ void ComputeHairStrandsInterpolation(
 							ShaderMap,
 							ShaderPrintData,
 							Instance,
-							VertexCount,
+							PointCount,
 							MeshLODIndex,
 							Instance->Strands.Modifier.HairLengthScale,
 							Instance->Strands.HairInterpolationType,
@@ -1954,7 +1954,7 @@ void ComputeHairStrandsInterpolation(
 							bHasSkinning ? Instance->Strands.DeformedRootResource : nullptr,
 							bHasSkinning && bValidGuide ? Instance->Guides.DeformedRootResource : nullptr,
 							RegisterAsSRV(GraphBuilder, Instance->Strands.RestResource->PositionBuffer),
-							RegisterAsSRV(GraphBuilder, Instance->Strands.RestResource->VertexToCurveBuffer),
+							RegisterAsSRV(GraphBuilder, Instance->Strands.RestResource->PointToCurveBuffer),
 							bUseSingleGuide,
 							bValidGuide &&  bUseSingleGuide ? RegisterAsSRV(GraphBuilder, Instance->Strands.InterpolationResource->InterpolationBuffer) : nullptr,
 							bValidGuide && !bUseSingleGuide ? RegisterAsSRV(GraphBuilder, Instance->Strands.InterpolationResource->Interpolation0Buffer) : nullptr,
@@ -1962,7 +1962,7 @@ void ComputeHairStrandsInterpolation(
 							bValidGuide ? RegisterAsSRV(GraphBuilder, Instance->Guides.RestResource->PositionBuffer) : nullptr,
 							bValidGuide ? RegisterAsSRV(GraphBuilder, Instance->Guides.DeformedResource->GetBuffer(FHairStrandsDeformedResource::Previous)) : nullptr,
 							bValidGuide ? RegisterAsSRV(GraphBuilder, Instance->Strands.InterpolationResource->SimRootPointIndexBuffer) : nullptr,
-							bValidGuide ? RegisterAsSRV(GraphBuilder, Instance->Guides.RestResource->VertexToCurveBuffer) : nullptr,
+							bValidGuide ? RegisterAsSRV(GraphBuilder, Instance->Guides.RestResource->PointToCurveBuffer) : nullptr,
 							Strands_DeformerPositionSRV,
 							Strands_DeformedPrevPosition,
 							FHairStrandsDeformedRootResource::FLOD::Previous);
@@ -1986,7 +1986,7 @@ void ComputeHairStrandsInterpolation(
 					AddGroomCacheUpdatePass(
 						GraphBuilder,
 						ShaderMap,
-						Instance->Strands.RestResource->GetVertexCount(),
+						Instance->Strands.RestResource->GetPointCount(),
 						GroomCacheGroupData,
 						RegisterAsSRV(GraphBuilder, Instance->Strands.RestResource->PositionBuffer),
 						RegisterAsSRV(GraphBuilder, Instance->Strands.DeformedResource->GetPositionOffsetBuffer(FHairStrandsDeformedResource::EFrameType::Current)),
@@ -2003,7 +2003,7 @@ void ComputeHairStrandsInterpolation(
 				AddHairTangentPass(
 					GraphBuilder,
 					ShaderMap,
-					VertexCount,
+					PointCount,
 					Instance->HairGroupPublicData,
 					Strands_DeformedPosition.SRV,
 					Strands_DeformedTangent);
@@ -2116,7 +2116,7 @@ void ComputeHairStrandsInterpolation(
 						GraphBuilder,
 						ShaderMap,
 						ShaderPrintData,
-						Instance->Strands.RestResource->GetVertexCount(),
+						Instance->Strands.RestResource->GetPointCount(),
 						bProceduralPrimitive,
 						ProceduralSplits,
 						HairRadiusRT,
@@ -2221,7 +2221,7 @@ void ComputeHairStrandsInterpolation(
 						ShaderMap,
 						ShaderPrintData,
 						Instance,
-						LOD.Guides.RestResource->GetVertexCount(),
+						LOD.Guides.RestResource->GetPointCount(),
 						MeshLODIndex,
 						1.0f,
 						LOD.Guides.HairInterpolationType,
@@ -2236,7 +2236,7 @@ void ComputeHairStrandsInterpolation(
 						bHasSkinning ? LOD.Guides.DeformedRootResource : nullptr,
 						bHasSkinning && bValidGuide ? Instance->Guides.DeformedRootResource : nullptr,
 						RegisterAsSRV(GraphBuilder, LOD.Guides.RestResource->PositionBuffer),
-						RegisterAsSRV(GraphBuilder, LOD.Guides.RestResource->VertexToCurveBuffer),
+						RegisterAsSRV(GraphBuilder, LOD.Guides.RestResource->PointToCurveBuffer),
 						bUseSingleGuide,
 						bValidGuide &&  bUseSingleGuide ? RegisterAsSRV(GraphBuilder, LOD.Guides.InterpolationResource->InterpolationBuffer) : nullptr,
 						bValidGuide && !bUseSingleGuide ? RegisterAsSRV(GraphBuilder, LOD.Guides.InterpolationResource->Interpolation0Buffer) : nullptr,
@@ -2244,7 +2244,7 @@ void ComputeHairStrandsInterpolation(
 						bValidGuide ? RegisterAsSRV(GraphBuilder, Instance->Guides.RestResource->PositionBuffer) : nullptr,
 						bValidGuide ? RegisterAsSRV(GraphBuilder, Instance->Guides.DeformedResource->GetBuffer(FHairStrandsDeformedResource::Current)) : nullptr,
 						RegisterAsSRV(GraphBuilder, LOD.Guides.InterpolationResource->SimRootPointIndexBuffer),
-						bValidGuide ? RegisterAsSRV(GraphBuilder, Instance->Guides.RestResource->VertexToCurveBuffer) : nullptr,
+						bValidGuide ? RegisterAsSRV(GraphBuilder, Instance->Guides.RestResource->PointToCurveBuffer) : nullptr,
 						nullptr,
 						Guides_DeformedPositionBuffer,
 						FHairStrandsDeformedRootResource::FLOD::Current); // <- this should be optional
@@ -2416,11 +2416,11 @@ void ResetHairStrandsInterpolation(
 		ShaderMap,
 		EDeformationType::OffsetGuide,
 		MeshLODIndex,
-		Instance->Guides.RestResource->GetVertexCount(),
+		Instance->Guides.RestResource->GetPointCount(),
 		Instance->Guides.RestRootResource,
 		Instance->Guides.DeformedRootResource,
 		RegisterAsSRV(GraphBuilder, Instance->Guides.RestResource->PositionBuffer),
-		RegisterAsSRV(GraphBuilder, Instance->Guides.RestResource->VertexToCurveBuffer),
+		RegisterAsSRV(GraphBuilder, Instance->Guides.RestResource->PointToCurveBuffer),
 		DeformedPositionBuffer,
 		Instance->Guides.RestResource->GetPositionOffset(),
 		RegisterAsSRV(GraphBuilder, Instance->Guides.DeformedResource->GetPositionOffsetBuffer(FHairStrandsDeformedResource::Current)),
diff --git a/Engine/Plugins/Runtime/HairStrands/Source/HairStrandsCore/Private/HairStrandsVertexFactory.cpp b/Engine/Plugins/Runtime/HairStrands/Source/HairStrandsCore/Private/HairStrandsVertexFactory.cpp
--- a/Engine/Plugins/Runtime/HairStrands/Source/HairStrandsCore/Private/HairStrandsVertexFactory.cpp
+++ b/Engine/Plugins/Runtime/HairStrands/Source/HairStrandsCore/Private/HairStrandsVertexFactory.cpp
@@ -104,7 +104,7 @@ FHairStrandsVertexFactoryUniformShaderParameters FHairGroupInstance::GetHairStan
 	Out.PositionBuffer				= VFInput.Strands.PositionBufferRHISRV;
 	Out.PreviousPositionBuffer		= VFInput.Strands.PrevPositionBufferRHISRV;
 	Out.AttributeBuffer				= VFInput.Strands.AttributeBufferRHISRV;
-	Out.VertexToCurveBuffer			= VFInput.Strands.VertexToCurveBufferRHISRV;
+	Out.PointToCurveBuffer			= VFInput.Strands.PointToCurveBufferRHISRV;
 	Out.TangentBuffer 				= VFInput.Strands.TangentBufferRHISRV;
 	Out.PositionOffsetBuffer 		= VFInput.Strands.PositionOffsetBufferRHISRV;
 	Out.PreviousPositionOffsetBuffer= VFInput.Strands.PrevPositionOffsetBufferRHISRV;
@@ -113,7 +113,7 @@ FHairStrandsVertexFactoryUniformShaderParameters FHairGroupInstance::GetHairStan
 	if (!Out.PositionBuffer) 				{ Out.PositionBuffer = GDummyCulledDispatchVertexIdsBuffer.SRVFloat; }
 	if (!Out.PreviousPositionBuffer) 		{ Out.PreviousPositionBuffer = GDummyCulledDispatchVertexIdsBuffer.SRVFloat; }
 	if (!Out.AttributeBuffer) 				{ Out.AttributeBuffer = GDummyCulledDispatchVertexIdsBuffer.SRVByteAddress; }
-	if (!Out.VertexToCurveBuffer) 			{ Out.VertexToCurveBuffer = GDummyCulledDispatchVertexIdsBuffer.SRVUint; }
+	if (!Out.PointToCurveBuffer) 			{ Out.PointToCurveBuffer = GDummyCulledDispatchVertexIdsBuffer.SRVUint; }
 	if (!Out.TangentBuffer) 				{ Out.TangentBuffer = GDummyCulledDispatchVertexIdsBuffer.SRVFloat; }
 	if (!Out.PositionOffsetBuffer) 			{ Out.PositionOffsetBuffer = GDummyCulledDispatchVertexIdsBuffer.SRVFloat; }
 	if (!Out.PreviousPositionOffsetBuffer) 	{ Out.PreviousPositionOffsetBuffer = GDummyCulledDispatchVertexIdsBuffer.SRVFloat; }
diff --git a/Engine/Plugins/Runtime/HairStrands/Source/HairStrandsCore/Private/Niagara/NiagaraDataInterfaceHairStrands.cpp b/Engine/Plugins/Runtime/HairStrands/Source/HairStrandsCore/Private/Niagara/NiagaraDataInterfaceHairStrands.cpp
--- a/Engine/Plugins/Runtime/HairStrands/Source/HairStrandsCore/Private/Niagara/NiagaraDataInterfaceHairStrands.cpp
+++ b/Engine/Plugins/Runtime/HairStrands/Source/HairStrandsCore/Private/Niagara/NiagaraDataInterfaceHairStrands.cpp
@@ -906,7 +906,7 @@ static void InterpolateGroomGuides(FRDGBuilder& GraphBuilder, FNiagaraDataBuffer
 		bool bIsRootValid = HairStrandsBuffer->SourceDeformedRootResources && HairStrandsBuffer->SourceDeformedRootResources->IsInitialized() &&
 								HairStrandsBuffer->SourceRestRootResources && HairStrandsBuffer->SourceRestRootResources->IsInitialized() && bHasSkinningBinding;
 
-		const uint32 NumPoints = HairStrandsBuffer->SourceRestResources->GetVertexCount();
+		const uint32 NumPoints = HairStrandsBuffer->SourceRestResources->GetPointCount();
 		const int32 MeshLODIndex = bIsRootValid ? HairStrandsBuffer->SourceDeformedRootResources->MeshLODIndex : -1;
 		bIsRootValid = bIsRootValid && HairStrandsBuffer->SourceDeformedRootResources->IsValid(MeshLODIndex);
 		
@@ -933,7 +933,7 @@ static void InterpolateGroomGuides(FRDGBuilder& GraphBuilder, FNiagaraDataBuffer
 		PassParameters->DeformedPositionBuffer = RegisterAsUAV(GraphBuilder, HairStrandsBuffer->SourceDeformedResources->DeformedPositionBuffer[BufferIndex]);
 		PassParameters->RestPositionBuffer = RegisterAsSRV(GraphBuilder, HairStrandsBuffer->SourceRestResources->PositionBuffer);
 		
-		PassParameters->VertexToCurveIndexBuffer = RegisterAsSRV(GraphBuilder, HairStrandsBuffer->SourceRestResources->VertexToCurveBuffer);
+		PassParameters->VertexToCurveIndexBuffer = RegisterAsSRV(GraphBuilder, HairStrandsBuffer->SourceRestResources->PointToCurveBuffer);
 		PassParameters->CurvesOffsetsBuffer = RegisterAsSRV(GraphBuilder, HairStrandsBuffer->SourceRestResources->CurveBuffer);
 		if(bIsRootValid)
 		{
diff --git a/Engine/Plugins/Runtime/HairStrands/Source/HairStrandsCore/Public/GroomInstance.h b/Engine/Plugins/Runtime/HairStrands/Source/HairStrandsCore/Public/GroomInstance.h
--- a/Engine/Plugins/Runtime/HairStrands/Source/HairStrandsCore/Public/GroomInstance.h
+++ b/Engine/Plugins/Runtime/HairStrands/Source/HairStrandsCore/Public/GroomInstance.h
@@ -69,7 +69,7 @@ BEGIN_GLOBAL_SHADER_PARAMETER_STRUCT(FHairStrandsVertexFactoryUniformShaderParam
 	SHADER_PARAMETER_SRV(Buffer<uint4>, PreviousPositionBuffer)
 
 	SHADER_PARAMETER_SRV(ByteAddressBuffer, AttributeBuffer)
-	SHADER_PARAMETER_SRV(Buffer<uint>, VertexToCurveBuffer)
+	SHADER_PARAMETER_SRV(Buffer<uint>, PointToCurveBuffer)
 	SHADER_PARAMETER_SRV(Buffer<float4>, TangentBuffer)
 
 	SHADER_PARAMETER_SRV(Buffer<uint>, CulledVertexIdsBuffer)
diff --git a/Engine/Plugins/Runtime/HairStrands/Source/HairStrandsCore/Public/GroomResources.h b/Engine/Plugins/Runtime/HairStrands/Source/HairStrandsCore/Public/GroomResources.h
--- a/Engine/Plugins/Runtime/HairStrands/Source/HairStrandsCore/Public/GroomResources.h
+++ b/Engine/Plugins/Runtime/HairStrands/Source/HairStrandsCore/Public/GroomResources.h
@@ -314,7 +314,7 @@ struct FHairStrandsRestResource : public FHairCommonResource
 		Total += GetBufferTotalNumBytes(PositionOffsetBuffer);
 		Total += GetBufferTotalNumBytes(TangentBuffer);
 		Total += GetBufferTotalNumBytes(AttributeBuffer);
-		Total += GetBufferTotalNumBytes(VertexToCurveBuffer);
+		Total += GetBufferTotalNumBytes(PointToCurveBuffer);
 		Total += GetBufferTotalNumBytes(CurveBuffer);
 		return Total;
 	}
@@ -332,7 +332,7 @@ struct FHairStrandsRestResource : public FHairCommonResource
 	FRDGExternalBuffer AttributeBuffer;
 
 	/* Strand hair vertex to curve index mapping */
-	FRDGExternalBuffer VertexToCurveBuffer;
+	FRDGExternalBuffer PointToCurveBuffer;
 
 	/* Strand hair curves buffer (contains curves' points offset and count) */
 	FRDGExternalBuffer CurveBuffer;
@@ -346,7 +346,10 @@ struct FHairStrandsRestResource : public FHairCommonResource
 	/* Type of curves */
 	const EHairStrandsResourcesType CurveType;
 
-	inline uint32 GetVertexCount() const { return BulkData.PointCount; }
+	/* Curve: Point offset & count */
+	TArray<FPackedHairCurve> CurveData;
+
+	inline uint32 GetPointCount() const { return BulkData.PointCount; }
 	inline uint32 GetCurveCount() const { return BulkData.CurveCount; }
 };
 
diff --git a/Engine/Shaders/Private/HairStrands/HairStrandsVertexFactoryCommon.ush b/Engine/Shaders/Private/HairStrands/HairStrandsVertexFactoryCommon.ush
--- a/Engine/Shaders/Private/HairStrands/HairStrandsVertexFactoryCommon.ush
+++ b/Engine/Shaders/Private/HairStrands/HairStrandsVertexFactoryCommon.ush
@@ -143,7 +143,7 @@ float			HairStrandsVF_Length;
 float			HairStrandsVF_Density;
 uint			HairStrandsVF_bUseStableRasterization;
 float3			HairStrandsVF_PositionOffset;
-uint			HairStrandsVF_VertexCount;
+uint			HairStrandsVF_PointCount;
 uint			HairStrandsVF_CurveCount;
 float4x4		HairStrandsVF_LocalToWorldPrimitiveTransform;
 float4x4		HairStrandsVF_LocalToTranslatedWorldPrimitiveTransform;
diff --git a/Engine/Source/Runtime/Renderer/Private/HairStrands/HairStrandsClusters.cpp b/Engine/Source/Runtime/Renderer/Private/HairStrands/HairStrandsClusters.cpp
--- a/Engine/Source/Runtime/Renderer/Private/HairStrands/HairStrandsClusters.cpp
+++ b/Engine/Source/Runtime/Renderer/Private/HairStrands/HairStrandsClusters.cpp
@@ -521,7 +521,7 @@ static void AddClusterResetLod0(
 	// Initialise indirect buffers to entire lod 0 dispatch
 	FHairIndBufferClearCS::FParameters* Parameters = GraphBuilder.AllocParameters<FHairIndBufferClearCS::FParameters>();
 	Parameters->DrawIndirectParameters = IndirectBuffer.UAV;
-	Parameters->VertexCountPerInstance = ClusterData.HairGroupPublicPtr->GetGroupInstanceVertexCount();
+	Parameters->VertexCountPerInstance = ClusterData.HairGroupPublicPtr->RestPointCount * HAIR_POINT_TO_VERTEX;
 
 	FHairIndBufferClearCS::FPermutationDomain Permutation;
 	Permutation.Set<FHairIndBufferClearCS::FSetIndirectDraw>(true);
diff --git a/Engine/Source/Runtime/Renderer/Private/HairStrands/HairStrandsData.h b/Engine/Source/Runtime/Renderer/Private/HairStrands/HairStrandsData.h
--- a/Engine/Source/Runtime/Renderer/Private/HairStrands/HairStrandsData.h
+++ b/Engine/Source/Runtime/Renderer/Private/HairStrands/HairStrandsData.h
@@ -58,7 +58,7 @@ BEGIN_SHADER_PARAMETER_STRUCT(FHairStrandsInstanceParameters, )
 	SHADER_PARAMETER(float, HairStrandsVF_TipScale)
 	SHADER_PARAMETER(float, HairStrandsVF_Length)
 	SHADER_PARAMETER(uint32, HairStrandsVF_bUseStableRasterization)
-	SHADER_PARAMETER(uint32, HairStrandsVF_VertexCount)
+	SHADER_PARAMETER(uint32, HairStrandsVF_PointCount)
 	SHADER_PARAMETER(uint32, HairStrandsVF_CurveCount)
 	SHADER_PARAMETER_ARRAY(FUintVector4, HairStrandsVF_AttributeOffsets, [HAIR_ATTRIBUTE_OFFSET_COUNT])
 	SHADER_PARAMETER(FVector3f, HairStrandsVF_PositionOffset)
@@ -67,7 +67,7 @@ BEGIN_SHADER_PARAMETER_STRUCT(FHairStrandsInstanceParameters, )
 	SHADER_PARAMETER_RDG_BUFFER_SRV(Buffer, HairStrandsVF_PositionBuffer)
 	SHADER_PARAMETER_RDG_BUFFER_SRV(Buffer, HairStrandsVF_PositionOffsetBuffer)
 	SHADER_PARAMETER_RDG_BUFFER_SRV(Buffer, HairStrandsVF_CurveBuffer)
-	SHADER_PARAMETER_RDG_BUFFER_SRV(Buffer, HairStrandsVF_VertexToCurveBuffer)
+	SHADER_PARAMETER_RDG_BUFFER_SRV(Buffer, HairStrandsVF_PointToCurveBuffer)
 	SHADER_PARAMETER_RDG_BUFFER_SRV(ByteAddressBuffer, HairStrandsVF_AttributeBuffer)
 	SHADER_PARAMETER_RDG_BUFFER_SRV(Buffer, HairStrandsVF_CullingIndirectBuffer)
 	SHADER_PARAMETER_RDG_BUFFER_SRV(Buffer, HairStrandsVF_CullingIndexBuffer)
diff --git a/Engine/Source/Runtime/Renderer/Private/HairStrands/HairStrandsForwardRaster.cpp b/Engine/Source/Runtime/Renderer/Private/HairStrands/HairStrandsForwardRaster.cpp
--- a/Engine/Source/Runtime/Renderer/Private/HairStrands/HairStrandsForwardRaster.cpp
+++ b/Engine/Source/Runtime/Renderer/Private/HairStrands/HairStrandsForwardRaster.cpp
@@ -305,7 +305,7 @@ FRasterForwardCullingOutput AddHairStrandsForwardCullingPass(
 
 			const FHairGroupPublicData::FVertexFactoryInput& VFInput = HairGroupPublicData->VFInput;
 
-			uint32 VertexCount = VFInput.Strands.VertexCount;
+			uint32 PointCount = VFInput.Strands.PointCount;
 			const bool bCullingPossible = bSupportCulling && GHairVisibilityComputeRaster_Culling;
 			if (!bCullingPossible)
 			{
@@ -313,10 +313,10 @@ FRasterForwardCullingOutput AddHairStrandsForwardCullingPass(
 				const FSphere BoundsSphere = HairGroupPublicData->ContinuousLODBounds.GetSphere();
 				const float ScreenSize = ComputeBoundsScreenSize(FVector4(BoundsSphere.Center, 1), BoundsSphere.W, ViewInfo);
 
-				VertexCount = HairGroupPublicData->GetActiveStrandsVertexCount(VFInput.Strands.VertexCount, ScreenSize);
+				PointCount = HairGroupPublicData->GetActiveStrandsPointCount(VFInput.Strands.PointCount, ScreenSize);
 			}
 
-			MaxNumPrimIDs += VertexCount;
+			MaxNumPrimIDs += PointCount;
 		}
 	}
 
@@ -373,15 +373,15 @@ FRasterForwardCullingOutput AddHairStrandsForwardCullingPass(
 			// calculate current view screen size - which can result in fewer strands rasterized in current view
 			const FSphere BoundsSphere = HairGroupPublicData->ContinuousLODBounds.GetSphere();
 			const float ScreenSize = ComputeBoundsScreenSize(FVector4(BoundsSphere.Center, 1), BoundsSphere.W, ViewInfo);
-			const uint32 CurveCount = VFInput.Strands.CurveCount; //HairGroupPublicData->GetActiveStrandsVertexCount(VFInput.Strands.VertexCount, ScreenSize);
-			const uint32 VertexCount = VFInput.Strands.VertexCount; //HairGroupPublicData->GetActiveStrandsVertexCount(VFInput.Strands.VertexCount, ScreenSize);
+			const uint32 CurveCount = VFInput.Strands.CurveCount;
+			const uint32 PointCount = VFInput.Strands.PointCount;
 
 			// Curve version
 			#if 1
 			FHairCullSegmentCS::FParameters* Parameters = GraphBuilder.AllocParameters<FHairCullSegmentCS::FParameters>();
 			Parameters->OutputResolution		= Out.Resolution;
 			Parameters->HairMaterialId			= PrimitiveInfo.MaterialId;
-			Parameters->ControlPointCount		= VertexCount;
+			Parameters->ControlPointCount		= PointCount;
 			Parameters->CurveCount				= CurveCount;
 			Parameters->RWHairVis				= GraphBuilder.CreateUAV(Out.NodeVis);
 			Parameters->RWCoord					= GraphBuilder.CreateUAV(Out.NodeCoord, PF_R16G16_UINT);
diff --git a/Engine/Source/Runtime/Renderer/Private/HairStrands/HairStrandsInterface.cpp b/Engine/Source/Runtime/Renderer/Private/HairStrands/HairStrandsInterface.cpp
--- a/Engine/Source/Runtime/Renderer/Private/HairStrands/HairStrandsInterface.cpp
+++ b/Engine/Source/Runtime/Renderer/Private/HairStrands/HairStrandsInterface.cpp
@@ -314,16 +314,6 @@ FHairGroupPublicData::FHairGroupPublicData(uint32 InGroupIndex, const FName& InO
 {
 	SetOwnerName(InOwnerName);
 	GroupIndex = InGroupIndex;
-	GroupControlTriangleStripVertexCount = 0;
-	ClusterCount = 0;
-	VertexCount = 0;
-}
-
-void FHairGroupPublicData::SetClusters(uint32 InClusterCount, uint32 InVertexCount)
-{
-	GroupControlTriangleStripVertexCount = InVertexCount * 6; // 6 vertex per point for a quad
-	ClusterCount = InClusterCount;
-	VertexCount = InVertexCount; // Control points
 }
 
 void FHairGroupPublicData::InitRHI()
@@ -366,8 +356,8 @@ void FHairGroupPublicData::Allocate(FRDGBuilder& GraphBuilder)
 	InternalCreateVertexBufferRDG(GraphBuilder, sizeof(int32), ClusterCount * 6, EPixelFormat::PF_R32_SINT, ClusterAABBBuffer, TEXT("Hair.Cluster_ClusterAABBBuffer"), GetOwnerName());
 	InternalCreateVertexBufferRDG(GraphBuilder, sizeof(int32), 6, EPixelFormat::PF_R32_SINT, GroupAABBBuffer, TEXT("Hair.Cluster_GroupAABBBuffer"), GetOwnerName());
 
-	InternalCreateVertexBufferRDG(GraphBuilder, sizeof(int32), VertexCount, EPixelFormat::PF_R32_UINT, CulledVertexIdBuffer, TEXT("Hair.Cluster_CulledVertexIdBuffer"), GetOwnerName());
-	InternalCreateVertexBufferRDG(GraphBuilder, sizeof(float), VertexCount, EPixelFormat::PF_R32_FLOAT, CulledVertexRadiusScaleBuffer, TEXT("Hair.Cluster_CulledVertexRadiusScaleBuffer"), GetOwnerName(), true);
+	InternalCreateVertexBufferRDG(GraphBuilder, sizeof(int32), RestPointCount, EPixelFormat::PF_R32_UINT, CulledVertexIdBuffer, TEXT("Hair.Cluster_CulledVertexIdBuffer"), GetOwnerName());
+	InternalCreateVertexBufferRDG(GraphBuilder, sizeof(float), RestPointCount, EPixelFormat::PF_R32_FLOAT, CulledVertexRadiusScaleBuffer, TEXT("Hair.Cluster_CulledVertexRadiusScaleBuffer"), GetOwnerName(), true);
 
 	GraphBuilder.SetBufferAccessFinal(Register(GraphBuilder, DrawIndirectBuffer, ERDGImportedBufferFlags::None).Buffer, ERHIAccess::IndirectArgs);
 
@@ -550,21 +540,21 @@ uint32 GetHairVisibilityComputeRasterVertexStart(uint32 TemporalIndex, uint32 In
 	return VertexStart;
 }
 
-uint32 GetHairVisibilityComputeRasterVertexCount(float ScreenSize, uint32 InVertexCount)
+uint32 GetHairVisibilityComputeRasterPointCount(float ScreenSize, uint32 InPointCount)
 {
-	uint32 VertexCount = InVertexCount;
+	uint32 PointCount = InPointCount;
 
 	if (IsHairVisibilityComputeRasterTemporalLayeringEnabled())
 	{
-		VertexCount /= GetHairVisibilityComputeRasterTemporalLayerCount();
+		PointCount /= GetHairVisibilityComputeRasterTemporalLayerCount();
 	}
 
 	if (IsHairVisibilityComputeRasterContinuousLODEnabled())
 	{
-		VertexCount *= GetHairVisibilityComputeRasterContinuousLODScale(ScreenSize);
+		PointCount *= GetHairVisibilityComputeRasterContinuousLODScale(ScreenSize);
 	}
 
-	return VertexCount;
+	return PointCount;
 }
 
 float GetHairVisibilityComputeRasterSampleWeight(float ScreenSize, bool bUseTemporalWeight)
@@ -586,9 +576,9 @@ float GetHairVisibilityComputeRasterSampleWeight(float ScreenSize, bool bUseTemp
 	return SampleWeight;
 }
 
-uint32 FHairGroupPublicData::GetActiveStrandsVertexCount(uint32 InVertexCount, float ScreenSize) const
+uint32 FHairGroupPublicData::GetActiveStrandsPointCount(uint32 InPointCount, float ScreenSize) const
 {
-	return GetHairVisibilityComputeRasterVertexCount(FMath::Min(MaxScreenSize, ScreenSize), InVertexCount);
+	return GetHairVisibilityComputeRasterPointCount(FMath::Min(MaxScreenSize, ScreenSize), InPointCount);
 }
 
 float FHairGroupPublicData::GetActiveStrandsSampleWeight(bool bUseTemporalWeight, float ScreenSize) const
diff --git a/Engine/Source/Runtime/Renderer/Private/HairStrands/HairStrandsVisibility.cpp b/Engine/Source/Runtime/Renderer/Private/HairStrands/HairStrandsVisibility.cpp
--- a/Engine/Source/Runtime/Renderer/Private/HairStrands/HairStrandsVisibility.cpp
+++ b/Engine/Source/Runtime/Renderer/Private/HairStrands/HairStrandsVisibility.cpp
@@ -2725,7 +2725,7 @@ void AddMeshDrawTransitionPass(
 			ExternalAccessQueue.Add(VFInput.Strands.PrevPositionBuffer.Buffer);
 			ExternalAccessQueue.Add(VFInput.Strands.TangentBuffer.Buffer);
 			ExternalAccessQueue.Add(VFInput.Strands.AttributeBuffer.Buffer);
-			ExternalAccessQueue.Add(VFInput.Strands.VertexToCurveBuffer.Buffer);
+			ExternalAccessQueue.Add(VFInput.Strands.PointToCurveBuffer.Buffer);
 			ExternalAccessQueue.Add(VFInput.Strands.PositionOffsetBuffer.Buffer);
 			ExternalAccessQueue.Add(VFInput.Strands.PrevPositionOffsetBuffer.Buffer);
 
@@ -2743,7 +2743,7 @@ void AddMeshDrawTransitionPass(
 				VFInput.Strands.PrevPositionBuffer			= FRDGImportedBuffer();
 				VFInput.Strands.TangentBuffer				= FRDGImportedBuffer();
 				VFInput.Strands.AttributeBuffer				= FRDGImportedBuffer();
-				VFInput.Strands.VertexToCurveBuffer			= FRDGImportedBuffer();
+				VFInput.Strands.PointToCurveBuffer			= FRDGImportedBuffer();
 				VFInput.Strands.PositionOffsetBuffer		= FRDGImportedBuffer();
 				VFInput.Strands.PrevPositionOffsetBuffer	= FRDGImportedBuffer();
 			}
@@ -2816,7 +2816,7 @@ FHairStrandsInstanceParameters GetHairStrandsInstanceParameters(FRDGBuilder& Gra
 		Out.HairStrandsVF_PositionBuffer		= Register(GraphBuilder, VFInput.Strands.PositionBufferExternal, ERDGImportedBufferFlags::CreateSRV).SRV;;
 		Out.HairStrandsVF_PositionOffsetBuffer	= Register(GraphBuilder, VFInput.Strands.PositionOffsetBufferExternal, ERDGImportedBufferFlags::CreateSRV).SRV;;
 		Out.HairStrandsVF_CurveBuffer			= Register(GraphBuilder, VFInput.Strands.CurveBufferExternal, ERDGImportedBufferFlags::CreateSRV).SRV;
-		Out.HairStrandsVF_VertexToCurveBuffer	= Register(GraphBuilder, VFInput.Strands.VertexToCurveBufferExternal, ERDGImportedBufferFlags::CreateSRV).SRV;
+		Out.HairStrandsVF_PointToCurveBuffer	= Register(GraphBuilder, VFInput.Strands.PointToCurveBufferExternal, ERDGImportedBufferFlags::CreateSRV).SRV;
 		Out.HairStrandsVF_AttributeBuffer		= Register(GraphBuilder, VFInput.Strands.AttributeBufferExternal, ERDGImportedBufferFlags::CreateSRV).SRV;
 	}
 	else
@@ -2824,11 +2824,11 @@ FHairStrandsInstanceParameters GetHairStrandsInstanceParameters(FRDGBuilder& Gra
 		Out.HairStrandsVF_PositionBuffer		= VFInput.Strands.PositionBuffer.SRV;
 		Out.HairStrandsVF_PositionOffsetBuffer	= VFInput.Strands.PositionOffsetBuffer.SRV;
 		Out.HairStrandsVF_CurveBuffer			= VFInput.Strands.CurveBuffer.SRV;
-		Out.HairStrandsVF_VertexToCurveBuffer	= VFInput.Strands.VertexToCurveBuffer.SRV;
+		Out.HairStrandsVF_PointToCurveBuffer	= VFInput.Strands.PointToCurveBuffer.SRV;
 		Out.HairStrandsVF_AttributeBuffer		= VFInput.Strands.AttributeBuffer.SRV;
 	}
 
-	Out.HairStrandsVF_VertexCount				= VFInput.Strands.VertexCount;
+	Out.HairStrandsVF_PointCount				= VFInput.Strands.PointCount;
 	Out.HairStrandsVF_CurveCount				= VFInput.Strands.CurveCount;
 	Out.HairStrandsVF_Radius					= VFInput.Strands.HairRadius;
 	Out.HairStrandsVF_RootScale					= VFInput.Strands.HairRootScale;
@@ -3240,7 +3240,7 @@ static FRasterComputeOutput AddVisibilityComputeRasterPass(
 
 			const FHairGroupPublicData::FVertexFactoryInput& VFInput = HairGroupPublicData->VFInput;
 
-			uint32 VertexCount = VFInput.Strands.VertexCount;
+			uint32 PointCount = VFInput.Strands.PointCount;
 			const bool bCullingPossible = bSupportCulling && GHairVisibilityComputeRaster_Culling;
 			if (!bCullingPossible)
 			{
@@ -3248,10 +3248,10 @@ static FRasterComputeOutput AddVisibilityComputeRasterPass(
 				const FSphere BoundsSphere = HairGroupPublicData->ContinuousLODBounds.GetSphere();
 				const float ScreenSize = ComputeBoundsScreenSize(FVector4(BoundsSphere.Center, 1), BoundsSphere.W, ViewInfo);
 
-				VertexCount = HairGroupPublicData->GetActiveStrandsVertexCount(VFInput.Strands.VertexCount, ScreenSize);
+				PointCount = HairGroupPublicData->GetActiveStrandsPointCount(VFInput.Strands.PointCount, ScreenSize);
 			}
 
-			MaxNumPrimIDs = FMath::Max(MaxNumPrimIDs, VertexCount);
+			MaxNumPrimIDs = FMath::Max(MaxNumPrimIDs, PointCount);
 		}
 	}
 
@@ -3434,7 +3434,7 @@ static FRasterComputeOutput AddVisibilityComputeRasterPass(
 			const FSphere BoundsSphere = HairGroupPublicData->ContinuousLODBounds.GetSphere();
 			const float ScreenSize = ComputeBoundsScreenSize(FVector4(BoundsSphere.Center, 1), BoundsSphere.W, ViewInfo);
 
-			const uint32 VertexCount = HairGroupPublicData->GetActiveStrandsVertexCount(VFInput.Strands.VertexCount, ScreenSize);
+			const uint32 PointCount = HairGroupPublicData->GetActiveStrandsPointCount(VFInput.Strands.PointCount, ScreenSize);
 			const float SampleWeight = HairGroupPublicData->GetActiveStrandsSampleWeight(false, ScreenSize);
 
 			// HW/SW classification
@@ -3450,7 +3450,7 @@ static FRasterComputeOutput AddVisibilityComputeRasterPass(
 
 				FVisiblityRasterClassificationCS::FParameters* ClassificationParameters = GraphBuilder.AllocParameters<FVisiblityRasterClassificationCS::FParameters>();
 				ClassificationParameters->Common = Common;
-				ClassificationParameters->ControlPointCount = VertexCount;
+				ClassificationParameters->ControlPointCount = PointCount;
 				ClassificationParameters->PrimIDsBufferSize = MaxNumPrimIDs;
 				ClassificationParameters->NumWorkGroups = NumWorkGroups;
 				ClassificationParameters->HairInstance = GetHairStrandsInstanceParameters(GraphBuilder, ViewInfo, HairGroupPublicData, bCullingEnable, bForceRegister);
@@ -3496,7 +3496,7 @@ static FRasterComputeOutput AddVisibilityComputeRasterPass(
 					FExclusiveDepthStencil::DepthRead_StencilNop);
 
 				TShaderMapRef<FVisiblityRasterHWVS> VertexShaderRaster = bCullingEnableInRasterizers ? VertexShaderRaster_CullingOn : VertexShaderRaster_CullingOff;
-				const uint32 NumInstances = bCullingEnableInRasterizers ? VFInput.Strands.VertexCount : VertexCount;
+				const uint32 NumInstances = bCullingEnableInRasterizers ? VFInput.Strands.PointCount : PointCount;
 
 				GraphBuilder.AddPass(
 					RDG_EVENT_NAME("HairStrands::VisibilityRasterHW"),
@@ -3540,7 +3540,7 @@ static FRasterComputeOutput AddVisibilityComputeRasterPass(
 					FVisiblityRasterComputeNaiveCS::FParameters* RasterParameters = GraphBuilder.AllocParameters<FVisiblityRasterComputeNaiveCS::FParameters>();
 					RasterParameters->OutputResolution = Common.OutputResolutionf;
 					RasterParameters->HairMaterialId = PrimitiveInfo.MaterialId;
-					RasterParameters->ControlPointCount = VertexCount;
+					RasterParameters->ControlPointCount = PointCount;
 					RasterParameters->SampleWeight = SampleWeight;
 					RasterParameters->NumWorkGroups = NumWorkGroups;
 					RasterParameters->HairInstance = GetHairStrandsInstanceParameters(GraphBuilder, ViewInfo, HairGroupPublicData, bCullingEnableInRasterizers, bForceRegister);
@@ -3586,7 +3586,7 @@ static FRasterComputeOutput AddVisibilityComputeRasterPass(
 						BinningParameters->OutVisTileArgs = VisTileArgsUAV;
 						BinningParameters->OutVisTileData = GraphBuilder.CreateUAV(VisTileData);
 						BinningParameters->HairInstance = GetHairStrandsInstanceParameters(GraphBuilder, ViewInfo, HairGroupPublicData, bCullingEnableInRasterizers, bForceRegister);
-						BinningParameters->VertexCount = VertexCount;
+						BinningParameters->VertexCount = PointCount;
 						BinningParameters->VisTileBinningGridTex = VisTileBinningGrid;
 						BinningParameters->IndirectPrimIDCount = bClassification ? SWRasterPrimIDCountSRV : nullptr;
 						BinningParameters->IndirectPrimIDs = bClassification ? RasterizerPrimIDsSRV : nullptr;
@@ -3863,7 +3863,6 @@ class FHairStrandsPositionChangedCS : public FGlobalShader
 
 	BEGIN_SHADER_PARAMETER_STRUCT(FParameters, )
 		SHADER_PARAMETER(uint32, VertexCount)
-		SHADER_PARAMETER(uint32, DispatchCountX)
 		SHADER_PARAMETER(float, PositionThreshold2)
 		SHADER_PARAMETER(uint32, HairStrandsVF_bIsCullingEnable)
 		SHADER_PARAMETER(uint32, bDrawInvalidElement)
@@ -3894,14 +3893,14 @@ static void AddHairStrandsHasPositionChangedPass(
 	const FHairGroupPublicData* HairGroupPublicData,
 	FRDGBufferUAVRef InvalidationBuffer)
 {
-	const uint32 VertexCount = HairGroupPublicData->VertexCount;
+	const uint32 PointCount = HairGroupPublicData->RestPointCount;
 
 	FRDGBufferRef InvalidationPrintCounter = GraphBuilder.CreateBuffer(FRDGBufferDesc::CreateBufferDesc(sizeof(uint32), 1), TEXT("Hair.InvalidationPrintCounter"));
 	FRDGBufferUAVRef InvalidationPrintCounterUAV = GraphBuilder.CreateUAV(InvalidationPrintCounter, PF_R32_UINT);
 	AddClearUAVPass(GraphBuilder, InvalidationPrintCounterUAV, 0u);
 
 	FHairStrandsPositionChangedCS::FParameters* Parameters = GraphBuilder.AllocParameters<FHairStrandsPositionChangedCS::FParameters>();
-	Parameters->VertexCount = VertexCount;
+	Parameters->VertexCount = PointCount;
 	Parameters->PositionThreshold2 = FMath::Square(GHairStrands_InvalidationPosition_Threshold);
 	Parameters->bDrawInvalidElement = GHairStrands_InvalidationPosition_Debug > 0 ? 1u : 0u;
 	Parameters->HairStrandsVF_bIsCullingEnable = 0u;
@@ -3920,7 +3919,7 @@ static void AddHairStrandsHasPositionChangedPass(
 		Parameters->HairStrandsVF_bIsCullingEnable = 1;
 	}
 
-	const FIntVector DispatchCount(FMath::DivideAndRoundUp(VertexCount, FHairStrandsPositionChangedCS::GetGroupSize()), 1, 1);
+	const FIntVector DispatchCount(FMath::DivideAndRoundUp(PointCount, FHairStrandsPositionChangedCS::GetGroupSize()), 1, 1);
 	TShaderMapRef<FHairStrandsPositionChangedCS> ComputeShader(View->ShaderMap);
 	FComputeShaderUtils::AddPass(
 		GraphBuilder,
diff --git a/Engine/Source/Runtime/Renderer/Private/HairStrands/HairStrandsVoxelization.cpp b/Engine/Source/Runtime/Renderer/Private/HairStrands/HairStrandsVoxelization.cpp
--- a/Engine/Source/Runtime/Renderer/Private/HairStrands/HairStrandsVoxelization.cpp
+++ b/Engine/Source/Runtime/Renderer/Private/HairStrands/HairStrandsVoxelization.cpp
@@ -1233,7 +1233,7 @@ static void AddVirtualVoxelizationComputeRasterPass(
 					continue;
 				}
 
-				const uint32 VertexCount = HairGroupPublicData->GetActiveStrandsVertexCount(VFInput.Strands.VertexCount, HairGroupPublicData->MaxScreenSize);
+				const uint32 PointCount = HairGroupPublicData->GetActiveStrandsPointCount(VFInput.Strands.PointCount, HairGroupPublicData->MaxScreenSize);
 				const float SampleWeight = HairGroupPublicData->GetActiveStrandsSampleWeight(false, HairGroupPublicData->MaxScreenSize);
 
 				FTransform LocalToTranslatedWorldTransform = VFInput.LocalToWorldTransform;
@@ -1247,7 +1247,7 @@ static void AddVirtualVoxelizationComputeRasterPass(
 				PassParameters->OutPageTexture = PageTextureUAV;
 				PassParameters->FrameIdMod8 = FrameIdMode8;
 
-				PassParameters->VertexCount = VertexCount;
+				PassParameters->VertexCount = PointCount;
 				PassParameters->SampleWeight = SampleWeight;
 			
 				const bool bCullingEnable = !IsHairStrandContinuousDecimationReorderingEnabled() && HairGroupPublicData->GetCullingResultAvailable();
diff --git a/Engine/Source/Runtime/Renderer/Public/HairStrandsInterface.h b/Engine/Source/Runtime/Renderer/Public/HairStrandsInterface.h
--- a/Engine/Source/Runtime/Renderer/Public/HairStrandsInterface.h
+++ b/Engine/Source/Runtime/Renderer/Public/HairStrandsInterface.h
@@ -15,6 +15,7 @@
 #include "RenderGraphResources.h"
 #include "ShaderPrintParameters.h"
 #include "GroomVisualizationData.h"
+#include "HairStrandsDefinitions.h"
 
 class UTexture2D;
 
@@ -108,7 +109,6 @@ class RENDERER_API FHairGroupPublicData : public FRenderResource
 {
 public:
 	FHairGroupPublicData(uint32 InGroupIndex, const FName& OwnerName);
-	void SetClusters(uint32 InClusterCount, uint32 InVertexCount);
 	
 	virtual void InitRHI() override;
 	virtual void ReleaseRHI() override;
@@ -117,10 +117,6 @@ public:
 	void Release();
 	uint32 GetResourcesSize() const;
 
-	// The primitive count when no culling and neither lod happens
-	uint32 GetGroupInstanceVertexCount() const { return GroupControlTriangleStripVertexCount; }
-	uint32 GetGroupControlPointCount() const { return VertexCount; }
-
 	uint32 GetGroupIndex() const { return GroupIndex; }
 
 	FRDGExternalBuffer& GetDrawIndirectRasterComputeBuffer() { return DrawIndirectRasterComputeBuffer; }
@@ -196,9 +192,8 @@ public:
 
 	uint32 GetClusterCount() const { return ClusterCount;  }
 
-	uint32 GetActiveStrandsVertexCount(uint32 InVertexCount, float ScreenSize) const;
+	uint32 GetActiveStrandsPointCount(uint32 InPointCount, float ScreenSize) const;
 	float GetActiveStrandsSampleWeight(bool bUseTemporalWeight, float ScreenSize) const;
-
 	void UpdateTemporalIndex();
 
 	struct FVertexFactoryInput 
@@ -209,7 +204,7 @@ public:
 			FRDGImportedBuffer PrevPositionBuffer;
 			FRDGImportedBuffer TangentBuffer;
 			FRDGImportedBuffer AttributeBuffer;
-			FRDGImportedBuffer VertexToCurveBuffer;
+			FRDGImportedBuffer PointToCurveBuffer;
 			FRDGImportedBuffer PositionOffsetBuffer;
 			FRDGImportedBuffer PrevPositionOffsetBuffer;
 			FRDGImportedBuffer CurveBuffer;
@@ -218,7 +213,7 @@ public:
 			FRDGExternalBuffer PrevPositionBufferExternal;
 			FRDGExternalBuffer TangentBufferExternal;
 			FRDGExternalBuffer AttributeBufferExternal;
-			FRDGExternalBuffer VertexToCurveBufferExternal;
+			FRDGExternalBuffer PointToCurveBufferExternal;
 			FRDGExternalBuffer PositionOffsetBufferExternal;
 			FRDGExternalBuffer PrevPositionOffsetBufferExternal;
 			FRDGExternalBuffer CurveBufferExternal;
@@ -227,7 +222,7 @@ public:
 			FShaderResourceViewRHIRef PrevPositionBufferRHISRV			= nullptr;
 			FShaderResourceViewRHIRef TangentBufferRHISRV				= nullptr;
 			FShaderResourceViewRHIRef AttributeBufferRHISRV				= nullptr;
-			FShaderResourceViewRHIRef VertexToCurveBufferRHISRV			= nullptr;
+			FShaderResourceViewRHIRef PointToCurveBufferRHISRV			= nullptr;
 			FShaderResourceViewRHIRef PositionOffsetBufferRHISRV		= nullptr;
 			FShaderResourceViewRHIRef PrevPositionOffsetBufferRHISRV	= nullptr;
 			FShaderResourceViewRHIRef CurveBufferRHISRV					= nullptr;
@@ -237,7 +232,7 @@ public:
 			TArray<uint32> AttributeOffsets;
 
 			uint32 CurveCount = 0;
-			uint32 VertexCount = 0;
+			uint32 PointCount = 0;
 			float HairRadius = 0;
 			float HairRootScale = 0;
 			float HairTipScale = 0;
@@ -273,10 +268,10 @@ public:
 	FVertexFactoryInput VFInput;
 	uint32 ClusterDataIndex = ~0; // #hair_todo: move this into instance data, or remove FHairStrandClusterData
 
-	uint32 GroupControlTriangleStripVertexCount;
-	uint32 GroupIndex;
-	uint32 ClusterCount;
-	uint32 VertexCount;
+	uint32 GroupIndex = 0;
+	uint32 ClusterCount = 0;
+	uint32 RestPointCount = 0;
+	uint32 RestCurveCount = 0;
 
 	/* Indirect draw buffer to draw everything or the result of the culling per pass */
 	FRDGExternalBuffer DrawIndirectBuffer;
