Centralize active/culled hair strands  point/curve computation/query.

diff --git a/Engine/Plugins/Runtime/HairStrands/Source/HairStrandsCore/Private/GroomComponent.cpp b/Engine/Plugins/Runtime/HairStrands/Source/HairStrandsCore/Private/GroomComponent.cpp
--- a/Engine/Plugins/Runtime/HairStrands/Source/HairStrandsCore/Private/GroomComponent.cpp
+++ b/Engine/Plugins/Runtime/HairStrands/Source/HairStrandsCore/Private/GroomComponent.cpp
@@ -800,7 +800,7 @@ public:
 		else // if (GeometryType == EHairGeometryType::Strands)
 		{
 			VertexFactory = (FVertexFactory*)Instance->Strands.VertexFactory;
-			HairVertexCount = Instance->Strands.RestResource->GetPointCount();
+			HairVertexCount = Instance->HairGroupPublicData->GetActiveStrandsPointCount();
 			MaxVertexIndex = HairVertexCount * HAIR_POINT_TO_VERTEX;
 			bUseCulling = Instance->Strands.bIsCullingEnabled;
 			NumPrimitive = bUseCulling ? 0 : HairVertexCount * HAIR_POINT_TO_TRIANGLE;
diff --git a/Engine/Plugins/Runtime/HairStrands/Source/HairStrandsCore/Private/GroomManager.cpp b/Engine/Plugins/Runtime/HairStrands/Source/HairStrandsCore/Private/GroomManager.cpp
--- a/Engine/Plugins/Runtime/HairStrands/Source/HairStrandsCore/Private/GroomManager.cpp
+++ b/Engine/Plugins/Runtime/HairStrands/Source/HairStrandsCore/Private/GroomManager.cpp
@@ -688,6 +688,17 @@ void AddHairStreamingRequest(FHairGroupInstance* Instance, int32 InLODIndex)
 	}
 }
 
+static uint32 GetHairVisibilityComputeRasterPointCount(float ScreenSize, uint32 InPointCount)
+{
+	if (IsHairVisibilityComputeRasterContinuousLODEnabled())
+	{
+		const float Scale = FMath::Clamp(FMath::Pow(FMath::Clamp(ScreenSize, 1.0f / 16.0f, 1.0f), 1.0f), 0.f, 1.f);
+		InPointCount *= Scale;
+	}
+
+	return InPointCount;
+}
+
 static void RunHairLODSelection(
 	FRDGBuilder& GraphBuilder, 
 	const FHairStrandsInstances& Instances, 
@@ -734,7 +745,7 @@ static void RunHairLODSelection(
 
 		// 0. Initial LOD index picking
 		// Insure that MinLOD is necessary taken into account if a force LOD is request (i.e., LODIndex>=0). If a Force LOD 
-		// is not resquested (i.e., LODIndex<0), the MinLOD is applied after ViewLODIndex has been determined in the codeblock below
+		// is not requested (i.e., LODIndex<0), the MinLOD is applied after ViewLODIndex has been determined in the codeblock below
 		const int32 LODCount = Instance->HairGroupPublicData->GetLODVisibilities().Num();
 		const float MinLOD = FMath::Max(0, GHairStrandsMinLOD);
 		
@@ -751,7 +762,7 @@ static void RunHairLODSelection(
 				const float CurrLODViewIndex = FMath::Max(MinLOD, GetHairInstanceLODIndex(Instance->HairGroupPublicData->GetLODScreenSizes(), ScreenSize, LODBias));
 				MaxScreenSize = FMath::Max(MaxScreenSize, ScreenSize);
 
-				// Select highest LOD accross all views
+				// Select highest LOD across all views
 				LODViewIndex = LODViewIndex < 0 ? CurrLODViewIndex : FMath::Min(LODViewIndex, CurrLODViewIndex);
 			}
 
@@ -765,18 +776,16 @@ static void RunHairLODSelection(
 			// Feedback game thread with LOD selection 
 			Instance->Debug.LODPredictedIndex = LODViewIndex;
 			Instance->HairGroupPublicData->DebugScreenSize = MaxScreenSize;
+			Instance->HairGroupPublicData->ContinuousLODPointCount = Instance->HairGroupPublicData->RestPointCount;
+			Instance->HairGroupPublicData->ContinuousLODCurveCount = Instance->HairGroupPublicData->RestCurveCount;
+			Instance->HairGroupPublicData->MaxScreenSize = 1.0;
 
-			if (IsHairStrandContinuousDecimationReorderingEnabled())
-			{
-				if (IsHairVisibilityComputeRasterContinuousLODEnabled())
-				{
-					Instance->HairGroupPublicData->ContinuousLODBounds = SphereBound;
-					Instance->HairGroupPublicData->MaxScreenSize = MaxScreenSize;
-				}
-			}
-			else
+			if (IsHairStrandContinuousDecimationReorderingEnabled() && IsHairVisibilityComputeRasterContinuousLODEnabled())
 			{
-				Instance->HairGroupPublicData->MaxScreenSize = 1.0;
+				Instance->HairGroupPublicData->ContinuousLODBounds = SphereBound;
+				Instance->HairGroupPublicData->MaxScreenSize = MaxScreenSize;
+				Instance->HairGroupPublicData->ContinuousLODPointCount = GetHairVisibilityComputeRasterPointCount(MaxScreenSize, Instance->HairGroupPublicData->RestPointCount);
+				Instance->HairGroupPublicData->ContinuousLODCurveCount = Instance->HairGroupPublicData->RestCurveCount;
 			}
 		}
 
diff --git a/Engine/Plugins/Runtime/HairStrands/Source/HairStrandsCore/Private/HairStrandsRendering.cpp b/Engine/Plugins/Runtime/HairStrands/Source/HairStrandsCore/Private/HairStrandsRendering.cpp
--- a/Engine/Plugins/Runtime/HairStrands/Source/HairStrandsCore/Private/HairStrandsRendering.cpp
+++ b/Engine/Plugins/Runtime/HairStrands/Source/HairStrandsCore/Private/HairStrandsRendering.cpp
@@ -897,7 +897,7 @@ static void AddHairClusterAABBPass(
 	Parameters->RenderDeformedOffsetBuffer = RenderDeformedOffsetBuffer;
 	Parameters->TotalClusterCount = 1;
 
-	Parameters->PointCount = Instance->HairGroupPublicData->GetActiveStrandsPointCount(Instance->Strands.RestResource->GetPointCount(), Instance->HairGroupPublicData->MaxScreenSize);
+	Parameters->PointCount = Instance->HairGroupPublicData->GetActiveStrandsPointCount();
 
 	if (ShaderPrintData)
 	{
@@ -1549,7 +1549,7 @@ void RegisterClusterData(FHairGroupInstance* Instance, FHairStrandClusterData* I
 	const int32 ClusterDataGroupIndex = InClusterData->HairGroups.Num();
 	FHairStrandClusterData::FHairGroup& HairGroupCluster = InClusterData->HairGroups.Emplace_GetRef();
 	HairGroupCluster.ClusterCount = Instance->HairGroupPublicData->GetClusterCount();
-	HairGroupCluster.VertexCount = Instance->HairGroupPublicData->RestPointCount * HAIR_POINT_TO_VERTEX;
+	HairGroupCluster.VertexCount = Instance->HairGroupPublicData->RestPointCount * HAIR_POINT_TO_VERTEX; // Instance->HairGroupPublicData->GetActiveStrandsPointCount() ?
 	HairGroupCluster.GroupAABBBuffer = &Instance->HairGroupPublicData->GetGroupAABBBuffer();
 	HairGroupCluster.ClusterAABBBuffer = &Instance->HairGroupPublicData->GetClusterAABBBuffer();
 
@@ -1651,8 +1651,8 @@ static FHairGroupPublicData::FVertexFactoryInput InternalComputeHairStrandsVerte
 		CONVERT_HAIRSSTRANDS_VF_PARAMETERS(OutVFInput.Strands.PointToCurveBuffer,		Instance->Strands.RestResource->PointToCurveBuffer);
 		CONVERT_HAIRSSTRANDS_VF_PARAMETERS(OutVFInput.Strands.CurveBuffer,				Instance->Strands.RestResource->CurveBuffer);
 
-		OutVFInput.Strands.PointCount = Instance->Strands.RestResource->GetPointCount();
-		OutVFInput.Strands.CurveCount = Instance->Strands.RestResource->GetCurveCount();
+		OutVFInput.Strands.PointCount = Instance->HairGroupPublicData->GetActiveStrandsPointCount();
+		OutVFInput.Strands.CurveCount = Instance->HairGroupPublicData->GetActiveStrandsCurveCount();
 	}
 	// 3. Render Strands in rest position: used when there are no skinning (rigid binding or no binding), no simulation, and no RBF
 	else
@@ -1668,8 +1668,8 @@ static FHairGroupPublicData::FVertexFactoryInput InternalComputeHairStrandsVerte
 
 		OutVFInput.Strands.PositionOffset = Instance->Strands.RestResource->GetPositionOffset();
 		OutVFInput.Strands.PrevPositionOffset = Instance->Strands.RestResource->GetPositionOffset();
-		OutVFInput.Strands.PointCount = Instance->Strands.RestResource->GetPointCount();
-		OutVFInput.Strands.CurveCount = Instance->Strands.RestResource->GetCurveCount();
+		OutVFInput.Strands.PointCount = Instance->HairGroupPublicData->GetActiveStrandsPointCount();
+		OutVFInput.Strands.CurveCount = Instance->HairGroupPublicData->GetActiveStrandsCurveCount();
 	}
 
 	OutVFInput.Strands.HairRadius = (GStrandHairWidth > 0 ? GStrandHairWidth : Instance->Strands.Modifier.HairWidth) * 0.5f;
@@ -1790,6 +1790,9 @@ void ComputeHairStrandsInterpolation(
 			const bool bUseSingleGuide	= bNeedDeformation && bValidGuide && Instance->Strands.InterpolationResource->UseSingleGuide();
 			const bool bHasSkinning		= bNeedDeformation && Instance->BindingType == EHairBindingType::Skinning;
 
+			const uint32 ActivePointCount = Instance->HairGroupPublicData->GetActiveStrandsPointCount();
+			const uint32 ActiveCurveCount = Instance->HairGroupPublicData->GetActiveStrandsCurveCount();
+
 			if (bNeedDeformation)
 			{
 				FRDGImportedBuffer Strands_DeformedPosition = Register(GraphBuilder, Instance->Strands.DeformedResource->GetBuffer(FHairStrandsDeformedResource::Current), ERDGImportedBufferFlags::CreateViews);
@@ -1874,20 +1877,16 @@ void ComputeHairStrandsInterpolation(
 					}
 				}
 
-				uint32 PointCount = Instance->Strands.RestResource->GetPointCount();
-
 				// 2.1 Compute deformation position based on simulation/skinning/RBF
 				if (Instance->Debug.GroomCacheType != EGroomCacheType::Strands)
 				{
-					PointCount = Instance->HairGroupPublicData->GetActiveStrandsPointCount(Instance->Strands.RestResource->GetPointCount(), Instance->HairGroupPublicData->MaxScreenSize);
-
 					// 2.1.1 Current Position
 					AddHairStrandsInterpolationPass(
 						GraphBuilder,
 						ShaderMap,
 						ShaderPrintData,
 						Instance,
-						PointCount,
+						ActivePointCount,
 						MeshLODIndex,
 						Instance->Strands.Modifier.HairLengthScale,
 						Instance->Strands.HairInterpolationType,
@@ -1925,7 +1924,7 @@ void ComputeHairStrandsInterpolation(
 						GHairStrandsTransferPositionOnLODChange > 0;
 					if (bTransferPrevPosition)
 					{
-						AddTransferPositionPass(GraphBuilder, ShaderMap, PointCount, Strands_DeformedPosition.SRV, Strands_DeformedPrevPosition.UAV);
+						AddTransferPositionPass(GraphBuilder, ShaderMap, ActivePointCount, Strands_DeformedPosition.SRV, Strands_DeformedPrevPosition.UAV);
 						GraphBuilder.SetBufferAccessFinal(Strands_DeformedPrevPosition.Buffer, ERHIAccess::SRVMask);
 					}
 					else if (bRecomputePrevPosition)
@@ -1937,7 +1936,7 @@ void ComputeHairStrandsInterpolation(
 							ShaderMap,
 							ShaderPrintData,
 							Instance,
-							PointCount,
+							ActivePointCount,
 							MeshLODIndex,
 							Instance->Strands.Modifier.HairLengthScale,
 							Instance->Strands.HairInterpolationType,
@@ -1984,7 +1983,7 @@ void ComputeHairStrandsInterpolation(
 					AddGroomCacheUpdatePass(
 						GraphBuilder,
 						ShaderMap,
-						Instance->Strands.RestResource->GetPointCount(),
+						ActivePointCount,
 						GroomCacheGroupData,
 						RegisterAsSRV(GraphBuilder, Instance->Strands.RestResource->PositionBuffer),
 						RegisterAsSRV(GraphBuilder, Instance->Strands.DeformedResource->GetPositionOffsetBuffer(FHairStrandsDeformedResource::EFrameType::Current)),
@@ -2001,7 +2000,7 @@ void ComputeHairStrandsInterpolation(
 				AddHairTangentPass(
 					GraphBuilder,
 					ShaderMap,
-					PointCount,
+					ActivePointCount,
 					Instance->HairGroupPublicData,
 					Strands_DeformedPosition.SRV,
 					Strands_DeformedTangent);
@@ -2030,7 +2029,7 @@ void ComputeHairStrandsInterpolation(
 				AddPatchAttributePass(
 					GraphBuilder,
 					ShaderMap,
-					Instance->Strands.RestResource->GetCurveCount(),
+					ActiveCurveCount,
 					PatchMode,
 					bValidGuide,
 					bUseSingleGuide,
@@ -2114,7 +2113,7 @@ void ComputeHairStrandsInterpolation(
 						GraphBuilder,
 						ShaderMap,
 						ShaderPrintData,
-						Instance->Strands.RestResource->GetPointCount(),
+						ActivePointCount,
 						bProceduralPrimitive,
 						ProceduralSplits,
 						HairRadiusRT,
@@ -2152,7 +2151,7 @@ void ComputeHairStrandsInterpolation(
 							FBufferRHIRef IndexBuffer(bProceduralPrimitive ? nullptr : Instance->Strands.RenRaytracingResource->IndexBuffer.Buffer->GetRHI());
 							Instance->Strands.CachedProceduralSplits = ProceduralSplits;
 							BuildHairAccelerationStructure_Strands(RHICmdList,
-								Instance->Strands.RenRaytracingResource->VertexCount,
+								Instance->Strands.RenRaytracingResource->VertexCount, // HAIR_TODO
 								Instance->Strands.RenRaytracingResource->IndexCount,
 								PositionBuffer,
 								IndexBuffer,
diff --git a/Engine/Source/Runtime/Renderer/Private/HairStrands/HairStrandsClusters.cpp b/Engine/Source/Runtime/Renderer/Private/HairStrands/HairStrandsClusters.cpp
--- a/Engine/Source/Runtime/Renderer/Private/HairStrands/HairStrandsClusters.cpp
+++ b/Engine/Source/Runtime/Renderer/Private/HairStrands/HairStrandsClusters.cpp
@@ -521,7 +521,7 @@ static void AddClusterResetLod0(
 	// Initialise indirect buffers to entire lod 0 dispatch
 	FHairIndBufferClearCS::FParameters* Parameters = GraphBuilder.AllocParameters<FHairIndBufferClearCS::FParameters>();
 	Parameters->DrawIndirectParameters = IndirectBuffer.UAV;
-	Parameters->VertexCountPerInstance = ClusterData.HairGroupPublicPtr->RestPointCount * HAIR_POINT_TO_VERTEX;
+	Parameters->VertexCountPerInstance = ClusterData.HairGroupPublicPtr->GetActiveStrandsPointCount() * HAIR_POINT_TO_VERTEX;
 
 	FHairIndBufferClearCS::FPermutationDomain Permutation;
 	Permutation.Set<FHairIndBufferClearCS::FSetIndirectDraw>(true);
diff --git a/Engine/Source/Runtime/Renderer/Private/HairStrands/HairStrandsForwardRaster.cpp b/Engine/Source/Runtime/Renderer/Private/HairStrands/HairStrandsForwardRaster.cpp
--- a/Engine/Source/Runtime/Renderer/Private/HairStrands/HairStrandsForwardRaster.cpp
+++ b/Engine/Source/Runtime/Renderer/Private/HairStrands/HairStrandsForwardRaster.cpp
@@ -303,18 +303,9 @@ FRasterForwardCullingOutput AddHairStrandsForwardCullingPass(
 			const FHairGroupPublicData* HairGroupPublicData = reinterpret_cast<const FHairGroupPublicData*>(PrimitiveInfo.Mesh->Elements[0].VertexFactoryUserData);
 			check(HairGroupPublicData);
 
-			const FHairGroupPublicData::FVertexFactoryInput& VFInput = HairGroupPublicData->VFInput;
-
-			uint32 PointCount = VFInput.Strands.PointCount;
-			const bool bCullingPossible = bSupportCulling && GHairVisibilityComputeRaster_Culling;
-			if (!bCullingPossible)
-			{
-				// calculate current view screen size - which can result in fewer strands rasterized in current view
-				const FSphere BoundsSphere = HairGroupPublicData->ContinuousLODBounds.GetSphere();
-				const float ScreenSize = ComputeBoundsScreenSize(FVector4(BoundsSphere.Center, 1), BoundsSphere.W, ViewInfo);
-
-				PointCount = HairGroupPublicData->GetActiveStrandsPointCount(VFInput.Strands.PointCount, ScreenSize);
-			}
+			const uint32 PointCount = HairGroupPublicData->GetActiveStrandsPointCount();
+			// Sanity check
+			check(PointCount == HairGroupPublicData->VFInput.Strands.PointCount);
 
 			MaxNumPrimIDs += PointCount;
 		}
@@ -371,10 +362,8 @@ FRasterForwardCullingOutput AddHairStrandsForwardCullingPass(
 			const bool bCullingEnable = false; //bSupportCulling && GHairVisibilityComputeRaster_Culling ? HairGroupPublicData->GetCullingResultAvailable() : false;
 			
 			// calculate current view screen size - which can result in fewer strands rasterized in current view
-			const FSphere BoundsSphere = HairGroupPublicData->ContinuousLODBounds.GetSphere();
-			const float ScreenSize = ComputeBoundsScreenSize(FVector4(BoundsSphere.Center, 1), BoundsSphere.W, ViewInfo);
-			const uint32 CurveCount = VFInput.Strands.CurveCount;
-			const uint32 PointCount = VFInput.Strands.PointCount;
+			const uint32 CurveCount  = HairGroupPublicData->GetActiveStrandsCurveCount();
+			const uint32 PointCount = HairGroupPublicData->GetActiveStrandsPointCount();
 
 			// Curve version
 			#if 1
diff --git a/Engine/Source/Runtime/Renderer/Private/HairStrands/HairStrandsInterface.cpp b/Engine/Source/Runtime/Renderer/Private/HairStrands/HairStrandsInterface.cpp
--- a/Engine/Source/Runtime/Renderer/Private/HairStrands/HairStrandsInterface.cpp
+++ b/Engine/Source/Runtime/Renderer/Private/HairStrands/HairStrandsInterface.cpp
@@ -491,32 +491,20 @@ bool IsHairVisibilityComputeRasterContinuousLODEnabled()
 	return IsHairStrandContinuousDecimationReorderingEnabled() && IsHairVisibilityComputeRasterEnabled() && (CVarHairStrandsVisibilityComputeRaster_ContinuousLOD.GetValueOnAnyThread() > 0);
 }
 
-float GetHairVisibilityComputeRasterContinuousLODScale(float ScreenSize)
+uint32 FHairGroupPublicData::GetActiveStrandsPointCount() const
 {
-	//TODO: make values in this calculation customizable per groom 
-	return FMath::Pow(FMath::Clamp(ScreenSize, 1.0f / 16.0f, 1.0f), 1.0f);
+	return ContinuousLODPointCount;
 }
 
-static uint32 GetHairVisibilityComputeRasterPointCount(float ScreenSize, uint32 InPointCount)
+uint32 FHairGroupPublicData::GetActiveStrandsCurveCount() const
 {
-	uint32 PointCount = InPointCount;
-
-	if (IsHairVisibilityComputeRasterContinuousLODEnabled())
-	{
-		PointCount *= GetHairVisibilityComputeRasterContinuousLODScale(ScreenSize);
-	}
-
-	return PointCount;
-}
-
-uint32 FHairGroupPublicData::GetActiveStrandsPointCount(uint32 InPointCount, float ScreenSize) const
-{
-	return GetHairVisibilityComputeRasterPointCount(FMath::Min(MaxScreenSize, ScreenSize), InPointCount);
+	return ContinuousLODCurveCount;
 }
 
 float FHairGroupPublicData::GetActiveStrandsCoverageScale() const
 {
-	return 1.f;
+	const float CurveRatio = ContinuousLODCurveCount / float(RestCurveCount);
+	return 1.f/FMath::Max(CurveRatio, 0.01f);
 }
 
 ///////////////////////////////////////////////////////////////////////////////////////////////////
diff --git a/Engine/Source/Runtime/Renderer/Private/HairStrands/HairStrandsVisibility.cpp b/Engine/Source/Runtime/Renderer/Private/HairStrands/HairStrandsVisibility.cpp
--- a/Engine/Source/Runtime/Renderer/Private/HairStrands/HairStrandsVisibility.cpp
+++ b/Engine/Source/Runtime/Renderer/Private/HairStrands/HairStrandsVisibility.cpp
@@ -3255,18 +3255,9 @@ static FRasterComputeOutput AddVisibilityComputeRasterPass(
 			const FHairGroupPublicData* HairGroupPublicData = reinterpret_cast<const FHairGroupPublicData*>(PrimitiveInfo.Mesh->Elements[0].VertexFactoryUserData);
 			check(HairGroupPublicData);
 
-			const FHairGroupPublicData::FVertexFactoryInput& VFInput = HairGroupPublicData->VFInput;
-
-			uint32 PointCount = VFInput.Strands.PointCount;
-			const bool bCullingPossible = bSupportCulling && GHairVisibilityComputeRaster_Culling;
-			if (!bCullingPossible)
-			{
-				// calculate current view screen size - which can result in fewer strands rasterized in current view
-				const FSphere BoundsSphere = HairGroupPublicData->ContinuousLODBounds.GetSphere();
-				const float ScreenSize = ComputeBoundsScreenSize(FVector4(BoundsSphere.Center, 1), BoundsSphere.W, ViewInfo);
-
-				PointCount = HairGroupPublicData->GetActiveStrandsPointCount(VFInput.Strands.PointCount, ScreenSize);
-			}
+			const uint32 PointCount = HairGroupPublicData->GetActiveStrandsPointCount();
+			// Sanity check
+			check(HairGroupPublicData->VFInput.Strands.PointCount == PointCount);
 
 			MaxNumPrimIDs = FMath::Max(MaxNumPrimIDs, PointCount);
 		}
@@ -3446,12 +3437,9 @@ static FRasterComputeOutput AddVisibilityComputeRasterPass(
 
 			const bool bCullingEnable = bSupportCulling && GHairVisibilityComputeRaster_Culling ? HairGroupPublicData->GetCullingResultAvailable() : false;
 			const bool bCullingEnableInRasterizers = bCullingEnable && !bClassification;
-			
-			// calculate current view screen size - which can result in fewer strands rasterized in current view
-			const FSphere BoundsSphere = HairGroupPublicData->ContinuousLODBounds.GetSphere();
-			const float ScreenSize = ComputeBoundsScreenSize(FVector4(BoundsSphere.Center, 1), BoundsSphere.W, ViewInfo);
 
-			const uint32 PointCount = HairGroupPublicData->GetActiveStrandsPointCount(VFInput.Strands.PointCount, ScreenSize);
+			const uint32 PointCount = HairGroupPublicData->GetActiveStrandsPointCount();
+			check(PointCount == HairGroupPublicData->VFInput.Strands.PointCount);
 			const float CoverageScale = HairGroupPublicData->GetActiveStrandsCoverageScale();
 
 			// HW/SW classification
@@ -3513,13 +3501,12 @@ static FRasterComputeOutput AddVisibilityComputeRasterPass(
 					FExclusiveDepthStencil::DepthRead_StencilNop);
 
 				TShaderMapRef<FVisiblityRasterHWVS> VertexShaderRaster = bCullingEnableInRasterizers ? VertexShaderRaster_CullingOn : VertexShaderRaster_CullingOff;
-				const uint32 NumInstances = bCullingEnableInRasterizers ? VFInput.Strands.PointCount : PointCount;
 
 				GraphBuilder.AddPass(
 					RDG_EVENT_NAME("HairStrands::VisibilityRasterHW"),
 					Parameters,
 					ERDGPassFlags::Raster,
-					[Parameters, VertexShaderRaster, bClassification, PixelShaderRaster, Viewport, Resolution, NumInstances, DrawIndexedIndirectArgs](FRHICommandList &RHICmdList)
+					[Parameters, VertexShaderRaster, bClassification, PixelShaderRaster, Viewport, Resolution, PointCount, DrawIndexedIndirectArgs](FRHICommandList &RHICmdList)
 					{
 						FGraphicsPipelineStateInitializer GraphicsPSOInit;
 						RHICmdList.ApplyCachedRenderTargets(GraphicsPSOInit);
@@ -3543,6 +3530,7 @@ static FRasterComputeOutput AddVisibilityComputeRasterPass(
 						}
 						else
 						{
+							const uint32 NumInstances = PointCount;
 							RHICmdList.DrawIndexedPrimitive(GTwoTrianglesIndexBuffer.IndexBufferRHI, 0, 0, 6, 0, 2, NumInstances);
 						}
 					});
@@ -3910,7 +3898,7 @@ static void AddHairStrandsHasPositionChangedPass(
 	const FHairGroupPublicData* HairGroupPublicData,
 	FRDGBufferUAVRef InvalidationBuffer)
 {
-	const uint32 PointCount = HairGroupPublicData->RestPointCount;
+	const uint32 PointCount = HairGroupPublicData->GetActiveStrandsPointCount();
 
 	FRDGBufferRef InvalidationPrintCounter = GraphBuilder.CreateBuffer(FRDGBufferDesc::CreateBufferDesc(sizeof(uint32), 1), TEXT("Hair.InvalidationPrintCounter"));
 	FRDGBufferUAVRef InvalidationPrintCounterUAV = GraphBuilder.CreateUAV(InvalidationPrintCounter, PF_R32_UINT);
diff --git a/Engine/Source/Runtime/Renderer/Private/HairStrands/HairStrandsVoxelization.cpp b/Engine/Source/Runtime/Renderer/Private/HairStrands/HairStrandsVoxelization.cpp
--- a/Engine/Source/Runtime/Renderer/Private/HairStrands/HairStrandsVoxelization.cpp
+++ b/Engine/Source/Runtime/Renderer/Private/HairStrands/HairStrandsVoxelization.cpp
@@ -1233,7 +1233,7 @@ static void AddVirtualVoxelizationComputeRasterPass(
 					continue;
 				}
 
-				const uint32 PointCount = HairGroupPublicData->GetActiveStrandsPointCount(VFInput.Strands.PointCount, HairGroupPublicData->MaxScreenSize);
+				const uint32 PointCount = HairGroupPublicData->GetActiveStrandsPointCount();
 				const float CoverageScale = HairGroupPublicData->GetActiveStrandsCoverageScale();
 
 				FTransform LocalToTranslatedWorldTransform = VFInput.LocalToWorldTransform;
diff --git a/Engine/Source/Runtime/Renderer/Public/HairStrandsInterface.h b/Engine/Source/Runtime/Renderer/Public/HairStrandsInterface.h
--- a/Engine/Source/Runtime/Renderer/Public/HairStrandsInterface.h
+++ b/Engine/Source/Runtime/Renderer/Public/HairStrandsInterface.h
@@ -192,7 +192,9 @@ public:
 
 	uint32 GetClusterCount() const { return ClusterCount;  }
 
-	uint32 GetActiveStrandsPointCount(uint32 InPointCount, float ScreenSize) const;
+	// Return the number of active point/curve for strand geometry
+	uint32 GetActiveStrandsPointCount() const;
+	uint32 GetActiveStrandsCurveCount() const;
 	float  GetActiveStrandsCoverageScale() const;
 
 	struct FVertexFactoryInput 
@@ -309,7 +311,10 @@ public:
 
 	FBoxSphereBounds ContinuousLODBounds; 	//used by Continuous LOD
 	float MaxScreenSize = 0.f; 				//used by Continuous LOD
-	uint32 TemporalIndex = 0; 				//used by Temporal Layering
+
+	// Active/used point/curved based on select continuous LOD
+	uint32 ContinuousLODPointCount = 0;
+	uint32 ContinuousLODCurveCount = 0;
 
 	// Debug
 	bool  bDebugDrawLODInfo = false; // Enable/disable hair LOD info
