* Add strands culling based on curve data.
* Add debug visualziation active curve/point

diff --git a/Engine/Plugins/Runtime/HairStrands/Source/HairStrandsCore/Private/GroomDebug.cpp b/Engine/Plugins/Runtime/HairStrands/Source/HairStrandsCore/Private/GroomDebug.cpp
--- a/Engine/Plugins/Runtime/HairStrands/Source/HairStrandsCore/Private/GroomDebug.cpp
+++ b/Engine/Plugins/Runtime/HairStrands/Source/HairStrandsCore/Private/GroomDebug.cpp
@@ -813,6 +813,7 @@ class FHairDebugPrintInstanceCS : public FGlobalShader
 		SHADER_PARAMETER(uint32, InstanceCount_StrandsShadowView)
 		SHADER_PARAMETER(uint32, InstanceCount_CardsOrMeshesPrimaryView)
 		SHADER_PARAMETER(uint32, InstanceCount_CardsOrMeshesShadowView)
+		SHADER_PARAMETER(FVector4f, InstanceScreenSphereBound)
 		SHADER_PARAMETER_RDG_BUFFER_SRV(StructuredBuffer<uint2>, NameInfos)
 		SHADER_PARAMETER_RDG_BUFFER_SRV(Buffer<uint8>, Names)
 		SHADER_PARAMETER_RDG_BUFFER_SRV(StructuredBuffer<uint4>, Infos)
@@ -844,6 +845,7 @@ struct FHairDebugNameInfo
 static void AddHairDebugPrintInstancePass(
 	FRDGBuilder& GraphBuilder, 
 	FGlobalShaderMap* ShaderMap,
+	const FSceneView& View,
 	const FShaderPrintData* ShaderPrintData,
 	const FHairStrandsInstances& Instances,
 	const FUintVector4& InstanceCountPerType)
@@ -935,6 +937,9 @@ static void AddHairDebugPrintInstancePass(
 					Data2 = FUintVector4(0);
 					Data2.X |= InstanceIndex < InstanceCountPerType[HairInstanceCount_StrandsPrimaryView] ? 0x1u : 0u;
 					Data2.X |= InstanceIndex < InstanceCountPerType[HairInstanceCount_StrandsShadowView]  ? 0x2u : 0u;
+					Data2.X |= uint32(FFloat16(Instance->HairGroupPublicData->ContinuousLODScreenSize).Encoded) << 16u;
+					Data2.Y = Instance->HairGroupPublicData->GetActiveStrandsPointCount();
+					Data2.Z = Instance->HairGroupPublicData->GetActiveStrandsCurveCount();
 				}
 			}
 			break;
@@ -1021,8 +1026,12 @@ static void AddHairDebugPrintInstancePass(
 
 		if (Instance->GeometryType == EHairGeometryType::Strands)
 		{
+			const float MaxRectSizeInPixels = FMath::Min(View.UnscaledViewRect.Height(), View.UnscaledViewRect.Width());
+			const float ContinousLODRadius = Instance->HairGroupPublicData->ContinuousLODScreenSize * MaxRectSizeInPixels * 0.5f; // Diameter->Radius
+
 			FHairDebugPrintInstanceCS::FParameters* Parameters = GraphBuilder.AllocParameters<FHairDebugPrintInstanceCS::FParameters>();
 			Parameters->InstanceAABB = Register(GraphBuilder, Instance->HairGroupPublicData->GroupAABBBuffer, ERDGImportedBufferFlags::CreateSRV).SRV;
+			Parameters->InstanceScreenSphereBound = FVector4f(Instance->HairGroupPublicData->ContinuousLODScreenPos.X, Instance->HairGroupPublicData->ContinuousLODScreenPos.Y, 0.f, ContinousLODRadius);
 			ShaderPrint::SetParameters(GraphBuilder, *ShaderPrintData, Parameters->ShaderPrintUniformBuffer);
 			ClearUnusedGraphResources(ComputeShader, Parameters);
 
@@ -1052,7 +1061,7 @@ void RunHairStrandsDebug(
 
 	if (ViewMode == EGroomViewMode::MacroGroups)
 	{
-		AddHairDebugPrintInstancePass(GraphBuilder, ShaderMap, ShaderPrintData, Instances, InstanceCountPerType);
+		AddHairDebugPrintInstancePass(GraphBuilder, ShaderMap, View, ShaderPrintData, Instances, InstanceCountPerType);
 	}
 
 	if (ViewMode == EGroomViewMode::MeshProjection)
diff --git a/Engine/Plugins/Runtime/HairStrands/Source/HairStrandsCore/Private/GroomManager.cpp b/Engine/Plugins/Runtime/HairStrands/Source/HairStrandsCore/Private/GroomManager.cpp
--- a/Engine/Plugins/Runtime/HairStrands/Source/HairStrandsCore/Private/GroomManager.cpp
+++ b/Engine/Plugins/Runtime/HairStrands/Source/HairStrandsCore/Private/GroomManager.cpp
@@ -37,6 +37,9 @@ static FAutoConsoleVariableRef CVarGHairStrands_ManualSkinCache(TEXT("r.HairStra
 static int32 GHairStrands_InterpolationFrustumCullingEnable = 1;
 static FAutoConsoleVariableRef CVarHairStrands_InterpolationFrustumCullingEnable(TEXT("r.HairStrands.Interoplation.FrustumCulling"), GHairStrands_InterpolationFrustumCullingEnable, TEXT("Swap rendering buffer at the end of frame. This is an experimental toggle. Default:1"));
 
+static int32 GHairStrands_ContinousLOD = 0;
+static FAutoConsoleVariableRef CVarHairStrands_ContinousLOD(TEXT("r.HairStrands.ContinuousLOD"), GHairStrands_ContinousLOD, TEXT("Continuous LODing toggle."));
+
 EHairBufferSwapType GetHairSwapBufferType()
 {
 	switch (GHairStrands_SwapBufferType)
@@ -699,6 +702,44 @@ static uint32 GetHairVisibilityComputeRasterPointCount(float ScreenSize, uint32
 	return InPointCount;
 }
 
+static FVector2f ComputeProjectedScreenPos(const FVector& InWorldPos, const FSceneView& View)
+{
+	// Compute the MinP/MaxP in pixel coord, relative to View.ViewRect.Min
+	const FMatrix& WorldToView = View.ViewMatrices.GetViewMatrix();
+	const FMatrix& ViewToProj = View.ViewMatrices.GetProjectionMatrix();
+	const float NearClippingDistance = View.NearClippingDistance + SMALL_NUMBER;
+	const FIntRect ViewRect = View.UnconstrainedViewRect;
+
+	// Clamp position on the near plane to get valid rect even if bounds' points are behind the camera
+	FPlane P_View = WorldToView.TransformFVector4(FVector4(InWorldPos, 1.f));
+	if (P_View.Z <= NearClippingDistance)
+	{
+		P_View.Z = NearClippingDistance;
+	}
+
+	// Project from view to projective space
+	FVector2D MinP(FLT_MAX, FLT_MAX);
+	FVector2D MaxP(-FLT_MAX, -FLT_MAX);
+	FVector2D ScreenPos;
+	const bool bIsValid = FSceneView::ProjectWorldToScreen(P_View, ViewRect, ViewToProj, ScreenPos);
+
+	// Clamp to pixel border
+	ScreenPos = FIntPoint(FMath::FloorToInt(ScreenPos.X), FMath::FloorToInt(ScreenPos.Y));
+
+	// Clamp to screen rect
+	ScreenPos.X = FMath::Clamp(ScreenPos.X, ViewRect.Min.X, ViewRect.Max.X);
+	ScreenPos.Y = FMath::Clamp(ScreenPos.Y, ViewRect.Min.Y, ViewRect.Max.Y);
+
+	return FVector2f(ScreenPos.X, ScreenPos.Y);
+}
+
+static uint32 ComputeActiveCurveCount(float InScreenSize, uint32 InCurveCount)
+{
+	const float Power = 1.8f; // Per Groom settings?
+	const uint32 OutCurveCount = InCurveCount * FMath::Pow(FMath::Clamp(InScreenSize, 0.f, 1.0f), Power);
+	return FMath::Clamp(OutCurveCount, 1, InCurveCount);
+}
+
 static void RunHairLODSelection(
 	FRDGBuilder& GraphBuilder, 
 	const FHairStrandsInstances& Instances, 
@@ -751,16 +792,26 @@ static void RunHairLODSelection(
 		
 		// If continuous LOD is enabled, we bypass all other type of geometric representation, and only use LOD0
 		float LODIndex = IsHairVisibilityComputeRasterContinuousLODEnabled() ? 0.0 : (Instance->Debug.LODForcedIndex >= 0 ? FMath::Max(Instance->Debug.LODForcedIndex, MinLOD) : -1.0f);
-		float LODViewIndex = -1;
 		{
-			float MaxScreenSize = 0.f;
+			FVector2f MaxContinuousLODScreenPos = FVector2f(0.f, 0.f);
+			float LODViewIndex = -1;
+			float MaxScreenSize_RestBound = 0.f;
+			float MaxScreenSize_Bound = 0.f;
 			const FSphere SphereBound = Instance->GetBounds().GetSphere();
 			for (const FSceneView* View : Views)
 			{
-				const float ScreenSize = ComputeBoundsScreenSize(FVector4(SphereBound.Center, 1), SphereBound.W, *View);
-				const float LODBias = Instance->Strands.Modifier.LODBias;
-				const float CurrLODViewIndex = FMath::Max(MinLOD, GetHairInstanceLODIndex(Instance->HairGroupPublicData->GetLODScreenSizes(), ScreenSize, LODBias));
-				MaxScreenSize = FMath::Max(MaxScreenSize, ScreenSize);
+				const FVector3d BoundScale = Instance->LocalToWorld.GetScale3D();
+				const FVector3f BoundExtent = FVector3f(Instance->Strands.Data->BoundingBox.GetExtent());
+				const float BoundRadius = FMath::Max3(BoundExtent.X, BoundExtent.Y, BoundExtent.Z) * FMath::Max3(BoundScale.X, BoundScale.Y, BoundScale.Z);
+
+				const float ScreenSize_RestBound = FMath::Clamp(ComputeBoundsScreenSize(FVector4(SphereBound.Center, 1), BoundRadius, *View), 0.f, 1.0f);
+				const float ScreenSize_Bound = FMath::Clamp(ComputeBoundsScreenSize(FVector4(SphereBound.Center, 1), SphereBound.W, *View), 0.f, 1.0f);
+
+				const float CurrLODViewIndex = FMath::Max(MinLOD, GetHairInstanceLODIndex(Instance->HairGroupPublicData->GetLODScreenSizes(), ScreenSize_Bound, Instance->Strands.Modifier.LODBias));
+
+				MaxScreenSize_Bound = FMath::Max(MaxScreenSize_Bound, ScreenSize_Bound);
+				MaxScreenSize_RestBound = FMath::Max(MaxScreenSize_RestBound, ScreenSize_RestBound);
+				MaxContinuousLODScreenPos = ComputeProjectedScreenPos(SphereBound.Center, *View);
 
 				// Select highest LOD across all views
 				LODViewIndex = LODViewIndex < 0 ? CurrLODViewIndex : FMath::Min(LODViewIndex, CurrLODViewIndex);
@@ -770,21 +821,26 @@ static void RunHairLODSelection(
 			{
 				LODIndex = LODViewIndex;
 			}
-
 			LODIndex = FMath::Clamp(LODIndex, 0.f, float(LODCount - 1));
 
 			// Feedback game thread with LOD selection 
 			Instance->Debug.LODPredictedIndex = LODViewIndex;
-			Instance->HairGroupPublicData->DebugScreenSize = MaxScreenSize;
+			Instance->HairGroupPublicData->DebugScreenSize = MaxScreenSize_Bound;
 			Instance->HairGroupPublicData->ContinuousLODPointCount = Instance->HairGroupPublicData->RestPointCount;
 			Instance->HairGroupPublicData->ContinuousLODCurveCount = Instance->HairGroupPublicData->RestCurveCount;
-			Instance->HairGroupPublicData->MaxScreenSize = 1.0;
+			Instance->HairGroupPublicData->ContinuousLODScreenSize = MaxScreenSize_RestBound;
+			Instance->HairGroupPublicData->ContinuousLODScreenPos = MaxContinuousLODScreenPos;
+			Instance->HairGroupPublicData->ContinuousLODBounds = SphereBound;
 
-			if (IsHairStrandContinuousDecimationReorderingEnabled() && IsHairVisibilityComputeRasterContinuousLODEnabled())
+			if (GHairStrands_ContinousLOD > 0)
+			{
+				// These values are set later, once the resources are loaded
+				Instance->HairGroupPublicData->ContinuousLODCurveCount = 0;
+				Instance->HairGroupPublicData->ContinuousLODPointCount = 0;
+			}
+			else if (IsHairStrandContinuousDecimationReorderingEnabled() && IsHairVisibilityComputeRasterContinuousLODEnabled())
 			{
-				Instance->HairGroupPublicData->ContinuousLODBounds = SphereBound;
-				Instance->HairGroupPublicData->MaxScreenSize = MaxScreenSize;
-				Instance->HairGroupPublicData->ContinuousLODPointCount = GetHairVisibilityComputeRasterPointCount(MaxScreenSize, Instance->HairGroupPublicData->RestPointCount);
+				Instance->HairGroupPublicData->ContinuousLODPointCount = GetHairVisibilityComputeRasterPointCount(MaxScreenSize_Bound, Instance->HairGroupPublicData->RestPointCount);
 				Instance->HairGroupPublicData->ContinuousLODCurveCount = Instance->HairGroupPublicData->RestCurveCount;
 			}
 		}
@@ -952,6 +1008,19 @@ static void RunHairLODSelection(
 			{
 				return false;
 			}
+
+			// Adapt the number curve/point based on Bound/CPU screen-size
+			if (GeometryType == EHairGeometryType::Strands && GHairStrands_ContinousLOD)
+			{
+				check(bIsLODDataReady);
+
+				const uint32 ActiveCurveCount = ComputeActiveCurveCount(Instance->HairGroupPublicData->ContinuousLODScreenSize, Instance->HairGroupPublicData->RestCurveCount);
+				const FPackedHairCurve Curve = Instance->Strands.RestResource->CurveData[ActiveCurveCount - 1];
+				const uint32 ActiveVertexCount = ActiveCurveCount < Instance->HairGroupPublicData->RestCurveCount ? (Curve.PointOffset + Curve.PointCount) : Instance->HairGroupPublicData->RestPointCount;
+
+				Instance->HairGroupPublicData->ContinuousLODCurveCount = ActiveCurveCount;
+				Instance->HairGroupPublicData->ContinuousLODPointCount = ActiveVertexCount;
+			}
 			return true;
 		};
 
diff --git a/Engine/Plugins/Runtime/HairStrands/Source/HairStrandsCore/Private/GroomResources.cpp b/Engine/Plugins/Runtime/HairStrands/Source/HairStrandsCore/Private/GroomResources.cpp
--- a/Engine/Plugins/Runtime/HairStrands/Source/HairStrandsCore/Private/GroomResources.cpp
+++ b/Engine/Plugins/Runtime/HairStrands/Source/HairStrandsCore/Private/GroomResources.cpp
@@ -909,6 +909,16 @@ void FHairStrandsRestResource::InternalAllocate(FRDGBuilder& GraphBuilder)
 	RestOffset.Add((FVector3f)BulkData.GetPositionOffset());// LWC_TODO: precision loss
 	InternalCreateVertexBufferRDG<FHairStrandsPositionOffsetFormat>(GraphBuilder, RestOffset, PositionOffsetBuffer, ToHairResourceDebugName(HAIRSTRANDS_RESOUCE_NAME(CurveType, Hair.StrandsRest_PositionOffsetBuffer), ResourceName), OwnerName, EHairResourceUsageType::Static, ERDGInitialDataFlags::None);
 	GraphBuilder.UseExternalAccessMode(Register(GraphBuilder, PositionOffsetBuffer, ERDGImportedBufferFlags::CreateSRV).Buffer, ERHIAccess::SRVMask);
+
+	// Copy curve data for CPU lookup
+	{
+		CurveData.SetNum(CurveCount);
+		uint32* Data = (uint32*)CurveData.GetData();
+
+		const uint32* ReadOnlyData = (const uint32*)BulkData.Curves.Lock(LOCK_READ_ONLY);
+		memcpy(Data, ReadOnlyData, sizeof(uint32) * CurveCount);
+		BulkData.Curves.Unlock();
+	}
 }
 
 void AddHairTangentPass(
diff --git a/Engine/Shaders/Private/HairStrands/HairStrandsDebug.usf b/Engine/Shaders/Private/HairStrands/HairStrandsDebug.usf
--- a/Engine/Shaders/Private/HairStrands/HairStrandsDebug.usf
+++ b/Engine/Shaders/Private/HairStrands/HairStrandsDebug.usf
@@ -1077,6 +1077,8 @@ StructuredBuffer<uint2> NameInfos;
 Buffer<uint>			Names;
 Buffer<int>				InstanceAABB;
 
+float4					InstanceScreenSphereBound;
+
 struct FInstanceInfo
 {
 	uint  GroupIndex;
@@ -1092,6 +1094,9 @@ struct FInstanceInfo
 	// Strands data
 	uint StrandsCount;
 	uint VertexCount;
+	float ScreenSize;
+	uint ActivePointCount;
+	uint ActiveCurveCount;
 
 	// Binding
 	uint UniqueTriangle;
@@ -1148,6 +1153,9 @@ FInstanceInfo ReadInstanceInfo(uint Index)
 
 	Out.bIsVisibleInPrimaryView = Data2.x & 0x1;
 	Out.bIsVisibleInShadowView	= Data2.x & 0x2;
+	Out.ScreenSize				= f16tof32((Data2.x>>16) & 0xFFFF);
+	Out.ActivePointCount		= Data2.y;
+	Out.ActiveCurveCount		= Data2.z;
 
 	return Out;
 }
@@ -1259,7 +1267,7 @@ void MainCS(uint GroupIndex : SV_GroupIndex, uint3 DispatchThreadId : SV_Dispatc
 	Newline(Context);
 
 	// Draw instance details
-	Print(Context, TEXT("   Group  LOD     Geometry Binding Sim   RBF   Clip    Strands Vertex  Bind.Triangle Bind.Section Primary Shadow Name"));
+	Print(Context, TEXT("   Group  LOD     Geometry Binding Sim   RBF   Clip Screen Strands          Vertex           Bind.Tri/Sect Primary Shadow Name"));
 	Newline(Context);
 
 	for (uint InstanceIndex = 0; InstanceIndex < InstanceCount; ++InstanceIndex)
@@ -1309,20 +1317,29 @@ void MainCS(uint GroupIndex : SV_GroupIndex, uint3 DispatchThreadId : SV_Dispatc
 
 		// Clip
 		if (Info.ClipValue >= 0)
-			Print(Context,Info.ClipValue, FontGreen, 7);
+			Print(Context,Info.ClipValue, FontGreen, 5,3);
 		else 
-			Print(Context,0u, FontSilver, 7);
-		Print(Context, TEXT(" "));
+			Print(Context,0u, FontSilver, 5, 3);
+
+		// Strands effective
+		Print(Context, Info.ScreenSize, FontOrange, 4, 2);
+		Print(Context, TEXT("   "));
 
 		// Strands/Vertex sCount
-		Print(Context,Info.StrandsCount, FontOrange, 8);
-		Print(Context,Info.VertexCount,  FontSilver, 8);
+		const uint EffectiveStrandsCount = ceil(Info.StrandsCount * Info.ScreenSize);
+		Print(Context, Info.ActiveCurveCount, FontOrange, 8);
+		Print(Context, TEXT("/"));
+		Print(Context,Info.StrandsCount, FontSilver, 8);
+
+		Print(Context, Info.ActivePointCount, FontOrange, 8);
+		Print(Context, TEXT("/"));
+		Print(Context,Info.VertexCount, FontSilver, 8);
 
 		// Binding data
 		Print(Context,Info.UniqueTriangle, FontOrange, 8);
-		Print(Context, TEXT("      "));
-		Print(Context,Info.UniqueSection, FontSilver, 8);
-		Print(Context, TEXT("     "));
+		Print(Context, TEXT(" "));
+		Print(Context,Info.UniqueSection, FontSilver, 3);
+		Print(Context, TEXT("  "));
 
 		// Primary view
 		if (Info.bIsVisibleInPrimaryView)
@@ -1337,7 +1354,7 @@ void MainCS(uint GroupIndex : SV_GroupIndex, uint3 DispatchThreadId : SV_Dispatc
 			Print(Context, TEXT("True "), FontGreen);
 		else
 			Print(Context, TEXT("False"), FontRed);
-		Print(Context, TEXT(" "), FontWhite);
+		Print(Context, TEXT("  "), FontWhite);
 
 		// Name
 		PrintInstanceName(Context, InstanceIndex, FontWhite);
@@ -1349,8 +1366,13 @@ void MainCS(uint GroupIndex : SV_GroupIndex, uint3 DispatchThreadId : SV_Dispatc
 {
 	if (all(DispatchThreadId == 0))
 	{
+		// Bound
 		const FHairAABB AABB = ReadHairAABB(0, InstanceAABB);
 		AddAABBTWS(AABB.Min, AABB.Max, ColorDarkGreen);
+
+		// Screen size
+		FShaderPrintContext Context = InitShaderPrintContext(true, uint2(50, 50));
+		AddCircleSS(Context, InstanceScreenSphereBound.xy, InstanceScreenSphereBound.w, ColorGreen);
 	}
 }
 #endif // PERMUTATION_OUTPUT_TYPE
diff --git a/Engine/Source/Runtime/Renderer/Public/HairStrandsInterface.h b/Engine/Source/Runtime/Renderer/Public/HairStrandsInterface.h
--- a/Engine/Source/Runtime/Renderer/Public/HairStrandsInterface.h
+++ b/Engine/Source/Runtime/Renderer/Public/HairStrandsInterface.h
@@ -309,12 +309,12 @@ public:
 	float LODBias = 0;			// Current LOD bias
 	bool bLODVisibility = true; // Enable/disable hair rendering for this component
 
-	FBoxSphereBounds ContinuousLODBounds; 	//used by Continuous LOD
-	float MaxScreenSize = 0.f; 				//used by Continuous LOD
-
 	// Active/used point/curved based on select continuous LOD
 	uint32 ContinuousLODPointCount = 0;
 	uint32 ContinuousLODCurveCount = 0;
+	float ContinuousLODScreenSize = 1.f;
+	FVector2f ContinuousLODScreenPos = FVector2f(0,0);
+	FBoxSphereBounds ContinuousLODBounds; 	//used by Continuous LOD
 
 	// Debug
 	bool  bDebugDrawLODInfo = false; // Enable/disable hair LOD info
