Fix CIS issue.

diff --git a/Engine/Plugins/Runtime/HairStrands/Source/HairStrandsCore/Private/GroomManager.cpp b/Engine/Plugins/Runtime/HairStrands/Source/HairStrandsCore/Private/GroomManager.cpp
--- a/Engine/Plugins/Runtime/HairStrands/Source/HairStrandsCore/Private/GroomManager.cpp
+++ b/Engine/Plugins/Runtime/HairStrands/Source/HairStrandsCore/Private/GroomManager.cpp
@@ -1013,6 +1013,7 @@ static void RunHairLODSelection(
 			if (GeometryType == EHairGeometryType::Strands && GHairStrands_ContinousLOD)
 			{
 				check(bIsLODDataReady);
+				check(Instance->Strands.RestResource);
 
 				const uint32 ActiveCurveCount = ComputeActiveCurveCount(Instance->HairGroupPublicData->ContinuousLODScreenSize, Instance->HairGroupPublicData->RestCurveCount);
 				const FPackedHairCurve Curve = Instance->Strands.RestResource->CurveData[ActiveCurveCount - 1];
