Fix CIS issue.

diff --git a/Engine/Plugins/Runtime/HairStrands/Source/HairStrandsEditor/Private/GroomActions.cpp b/Engine/Plugins/Runtime/HairStrands/Source/HairStrandsEditor/Private/GroomActions.cpp
--- a/Engine/Plugins/Runtime/HairStrands/Source/HairStrandsEditor/Private/GroomActions.cpp
+++ b/Engine/Plugins/Runtime/HairStrands/Source/HairStrandsEditor/Private/GroomActions.cpp
@@ -254,7 +254,10 @@ void FGroomActions::ExecuteCreateBindingAsset(TArray<TWeakObjectPtr<UGroomAsset>
 		{
 			// Duplicate the options to prevent dirtying the asset when they are modified but the rebuild is cancelled
 			UGroomCreateBindingOptions* CurrentOptions = NewObject<UGroomCreateBindingOptions>();
-			CurrentOptions->GroomAsset = GroomAsset;
+			if (CurrentOptions)
+			{
+				CurrentOptions->GroomAsset = GroomAsset;
+			}
 			TSharedPtr<SGroomCreateBindingOptionsWindow> GroomOptionWindow = SGroomCreateBindingOptionsWindow::DisplayCreateBindingOptions(CurrentOptions);
 
 			if (!GroomOptionWindow->ShouldCreate())
