Fix crash in groom LOD selection, when running with groom not having strands data.

diff --git a/Engine/Plugins/Runtime/HairStrands/Source/HairStrandsCore/Private/GroomManager.cpp b/Engine/Plugins/Runtime/HairStrands/Source/HairStrandsCore/Private/GroomManager.cpp
--- a/Engine/Plugins/Runtime/HairStrands/Source/HairStrandsCore/Private/GroomManager.cpp
+++ b/Engine/Plugins/Runtime/HairStrands/Source/HairStrandsCore/Private/GroomManager.cpp
@@ -801,7 +801,7 @@ static void RunHairLODSelection(
 			for (const FSceneView* View : Views)
 			{
 				const FVector3d BoundScale = Instance->LocalToWorld.GetScale3D();
-				const FVector3f BoundExtent = FVector3f(Instance->Strands.Data->BoundingBox.GetExtent());
+				const FVector3f BoundExtent = Instance->Strands.Data ? FVector3f(Instance->Strands.Data->BoundingBox.GetExtent()) : FVector3f(0,0,0);
 				const float BoundRadius = FMath::Max3(BoundExtent.X, BoundExtent.Y, BoundExtent.Z) * FMath::Max3(BoundScale.X, BoundScale.Y, BoundScale.Z);
 
 				const float ScreenSize_RestBound = FMath::Clamp(ComputeBoundsScreenSize(FVector4(SphereBound.Center, 1), BoundRadius, *View), 0.f, 1.0f);
