Disable cluster culling when hair CLOD is enabled.

diff --git a/Engine/Plugins/Runtime/HairStrands/Source/HairStrandsCore/Private/GroomManager.cpp b/Engine/Plugins/Runtime/HairStrands/Source/HairStrandsCore/Private/GroomManager.cpp
--- a/Engine/Plugins/Runtime/HairStrands/Source/HairStrandsCore/Private/GroomManager.cpp
+++ b/Engine/Plugins/Runtime/HairStrands/Source/HairStrandsCore/Private/GroomManager.cpp
@@ -448,6 +448,11 @@ static void RunHairStrandsInterpolation_Strands(
 		ClusterData);
 }
 
+bool IsHairStrandsContinousLODEnabled()
+{
+	return GHairStrands_ContinousLOD > 0;
+}
+
 static void RunHairStrandsGatherCluster(
 	const FHairStrandsInstances& Instances,
 	FHairStrandClusterData* ClusterData)
@@ -455,11 +460,7 @@ static void RunHairStrandsGatherCluster(
 	for (FHairStrandsInstance* AbstractInstance : Instances)
 	{
 		FHairGroupInstance* Instance = static_cast<FHairGroupInstance*>(AbstractInstance);
-
-		if (Instance->GeometryType != EHairGeometryType::Strands)
-			continue;
-
-		if (Instance->Strands.IsValid())
+		if (Instance->GeometryType == EHairGeometryType::Strands && Instance->Strands.IsValid())
 		{
 			if (Instance->Strands.bIsCullingEnabled)
 			{
@@ -467,7 +468,7 @@ static void RunHairStrandsGatherCluster(
 			}
 			else
 			{
-				Instance->HairGroupPublicData->bCullingResultAvailable = false;
+				Instance->HairGroupPublicData->SetCullingResultAvailable(false);
 			}
 		}
 	}
@@ -885,7 +886,7 @@ static void RunHairLODSelection(
 				{
 					const bool bNeedLODing		= LODIndex > 0;
 					const bool bNeedDeformation = Instance->Strands.DeformedResource != nullptr;
-					bCullingEnable = bNeedLODing || bNeedDeformation;
+					bCullingEnable = (bNeedLODing || bNeedDeformation) && !IsHairStrandsContinousLODEnabled();
 				}
 			}
 
@@ -894,10 +895,10 @@ static void RunHairLODSelection(
 				GeometryType = EHairGeometryType::NoneGeometry;
 			}
 
-			const bool bSimulationEnable = Instance->HairGroupPublicData->IsSimulationEnable(IntLODIndex);
+			const bool bSimulationEnable			= Instance->HairGroupPublicData->IsSimulationEnable(IntLODIndex);
 			const bool bDeformationEnable			= Instance->HairGroupPublicData->bIsDeformationEnable;
-			const bool bGlobalInterpolationEnable =  Instance->HairGroupPublicData->IsGlobalInterpolationEnable(IntLODIndex);
-			const bool bLODNeedsGuides = bSimulationEnable || bDeformationEnable || bGlobalInterpolationEnable;
+			const bool bGlobalInterpolationEnable	= Instance->HairGroupPublicData->IsGlobalInterpolationEnable(IntLODIndex);
+			const bool bLODNeedsGuides				= bSimulationEnable || bDeformationEnable || bGlobalInterpolationEnable;
 
 			const EHairResourceLoadingType LoadingType = GetHairResourceLoadingType(GeometryType, IntLODIndex);
 			EHairResourceStatus ResourceStatus = EHairResourceStatus::None;
diff --git a/Engine/Plugins/Runtime/HairStrands/Source/HairStrandsCore/Private/HairStrandsRendering.cpp b/Engine/Plugins/Runtime/HairStrands/Source/HairStrandsCore/Private/HairStrandsRendering.cpp
--- a/Engine/Plugins/Runtime/HairStrands/Source/HairStrandsCore/Private/HairStrandsRendering.cpp
+++ b/Engine/Plugins/Runtime/HairStrands/Source/HairStrandsCore/Private/HairStrandsRendering.cpp
@@ -76,6 +76,8 @@ int GetHairRaytracingProceduralSplits()
 	return FMath::Clamp(GHairRaytracingProceduralSplits, 1, STRANDS_PROCEDURAL_INTERSECTOR_MAX_SPLITS);
 }
 
+bool IsHairStrandsContinousLODEnabled();
+
 inline uint32 ComputeGroupSize()
 {
 	const uint32 GroupSize = IsRHIDeviceAMD() ? 64 : (IsRHIDeviceNVIDIA() ? 32 : 64);
@@ -427,7 +429,7 @@ struct FRDGHairStrandsCullingData
 FRDGHairStrandsCullingData ImportCullingData(FRDGBuilder& GraphBuilder, FHairGroupPublicData* In)
 {
 	FRDGHairStrandsCullingData Out;
-	Out.bCullingResultAvailable					= In->GetCullingResultAvailable();
+	Out.bCullingResultAvailable = In->GetCullingResultAvailable();
 
 	if (Out.bCullingResultAvailable)
 	{
@@ -2050,7 +2052,7 @@ void ComputeHairStrandsInterpolation(
 			// 3. Compute cluster AABBs (used for LODing and voxelization)
 			{
 				// Optim: If an instance does not voxelize it's data, then there is no need for having valid AABB
-				bool bNeedAABB = Instance->Strands.Modifier.bSupportVoxelization && Instance->bCastShadow;
+				bool bNeedGPUAABB = Instance->Strands.Modifier.bSupportVoxelization && Instance->bCastShadow;
 
 				FHairStrandClusterData::FHairGroup* HairGroupCluster = nullptr;
 				if (CullingData.bCullingResultAvailable)
@@ -2060,13 +2062,20 @@ void ComputeHairStrandsInterpolation(
 					HairGroupCluster = &InClusterData->HairGroups[Instance->HairGroupPublicData->ClusterDataIndex];
 					if (!HairGroupCluster->bVisible)
 					{
-						bNeedAABB = false;
+						bNeedGPUAABB = false;
 					}
 				}
+
+				// When CLOD is enabled, GPU AABB are not supported yet
+				if (IsHairStrandsContinousLODEnabled())
+				{
+					bNeedGPUAABB = false;
+				}
 				Instance->HairGroupPublicData->bClusterAABBValid = false;
 				Instance->HairGroupPublicData->bGroupAABBValid = false;
 
-				if (bNeedAABB)
+
+				if (bNeedGPUAABB)
 				{
 					FRDGImportedBuffer Strands_CulledVertexCount = Register(GraphBuilder, Instance->HairGroupPublicData->GetDrawIndirectRasterComputeBuffer(), ERDGImportedBufferFlags::CreateSRV);
 					AddHairClusterAABBPass(
diff --git a/Engine/Source/Runtime/Renderer/Private/HairStrands/HairStrandsClusters.cpp b/Engine/Source/Runtime/Renderer/Private/HairStrands/HairStrandsClusters.cpp
--- a/Engine/Source/Runtime/Renderer/Private/HairStrands/HairStrandsClusters.cpp
+++ b/Engine/Source/Runtime/Renderer/Private/HairStrands/HairStrandsClusters.cpp
@@ -550,7 +550,6 @@ void ComputeHairStrandsClustersCulling(
 	TRACE_CPUPROFILER_EVENT_SCOPE(ComputeHairStrandsClustersCulling);
 	RDG_GPU_STAT_SCOPE(GraphBuilder, HairStrandsClusterCulling);
 
-	const bool bClusterCulling = IsHairStrandsClusterCullingEnable();
 	for (const FViewInfo& View : Views)
 	{
 		// TODO use compute overlap (will need to split AddClusterCullingPass)
@@ -559,17 +558,14 @@ void ComputeHairStrandsClustersCulling(
 			AddClusterResetLod0(GraphBuilder, &ShaderMap, ClusterData);
 		}
 
-		if (bClusterCulling)
+		for (FHairStrandClusterData::FHairGroup& ClusterData : ClusterDatas.HairGroups)		
 		{
-			for (FHairStrandClusterData::FHairGroup& ClusterData : ClusterDatas.HairGroups)		
-			{
-				AddClusterCullingPass(
-					GraphBuilder, 
-					&ShaderMap, 
-					View, 
-					CullingParameters,
-					ClusterData);
-			}
+			AddClusterCullingPass(
+				GraphBuilder,
+				&ShaderMap,
+				View,
+				CullingParameters,
+				ClusterData);
 		}
 	}
 }
diff --git a/Engine/Source/Runtime/Renderer/Private/HairStrands/HairStrandsDebug.cpp b/Engine/Source/Runtime/Renderer/Private/HairStrands/HairStrandsDebug.cpp
--- a/Engine/Source/Runtime/Renderer/Private/HairStrands/HairStrandsDebug.cpp
+++ b/Engine/Source/Runtime/Renderer/Private/HairStrands/HairStrandsDebug.cpp
@@ -1097,8 +1097,6 @@ public:
 
 IMPLEMENT_GLOBAL_SHADER(FDrawDebugClusterAABBCS, "/Engine/Private/HairStrands/HairStrandsClusterCulling.usf", "MainDrawDebugAABBCS", SF_Compute);
 
-bool IsHairStrandsClusterCullingEnable();
-
 static void AddDrawDebugClusterPass(
 	FRDGBuilder& GraphBuilder,
 	const FHairStrandClusterData& HairClusterData,
@@ -1106,8 +1104,7 @@ static void AddDrawDebugClusterPass(
 {
 	const EGroomViewMode ViewMode = GetGroomViewMode(View);
 	const bool bDebugEnable = ViewMode == EGroomViewMode::ClusterAABB || ViewMode == EGroomViewMode::Cluster;
-	const bool bCullingEnable = IsHairStrandsClusterCullingEnable();
-	if (!bDebugEnable || !bCullingEnable)
+	if (!bDebugEnable)
 	{
 		return;
 	}
