Fix misspelling.

diff --git a/Engine/Plugins/Runtime/HairStrands/Source/HairStrandsEditor/Private/GroomImportOptionsWindow.cpp b/Engine/Plugins/Runtime/HairStrands/Source/HairStrandsEditor/Private/GroomImportOptionsWindow.cpp
--- a/Engine/Plugins/Runtime/HairStrands/Source/HairStrandsEditor/Private/GroomImportOptionsWindow.cpp
+++ b/Engine/Plugins/Runtime/HairStrands/Source/HairStrandsEditor/Private/GroomImportOptionsWindow.cpp
@@ -221,7 +221,7 @@ FText GetHairAttributeText(EHairAttribute In)
 	case EHairAttribute::RootUV:					return LOCTEXT("GroomOptionsWindow_HasRootUV", "Root UV");
 	case EHairAttribute::ClumpID:					return LOCTEXT("GroomOptionsWindow_HasClumpID", "Clump ID");
 	case EHairAttribute::PrecomputedGuideWeights:	return LOCTEXT("GroomOptionsWindow_HasPercomputedGuideWeights", "Pre-Computed Guide Weights");
-	case EHairAttribute::Color:						return LOCTEXT("GroomOptionsWindow_HasColor", "Has Color");
+	case EHairAttribute::Color:						return LOCTEXT("GroomOptionsWindow_HasColor", "Color");
 	case EHairAttribute::Roughness:					return LOCTEXT("GroomOptionsWindow_HasRoughness", "Roughness");
 	case EHairAttribute::AO:						return LOCTEXT("GroomOptionsWindow_HasAO", "AO");
 	}
