* Fix groom cache (guide/strands) rendering
* Fix guide cache debug visualization

diff --git a/Engine/Plugins/Runtime/HairStrands/Source/HairStrandsCore/Private/GroomBuilder.cpp b/Engine/Plugins/Runtime/HairStrands/Source/HairStrandsCore/Private/GroomBuilder.cpp
--- a/Engine/Plugins/Runtime/HairStrands/Source/HairStrandsCore/Private/GroomBuilder.cpp
+++ b/Engine/Plugins/Runtime/HairStrands/Source/HairStrandsCore/Private/GroomBuilder.cpp
@@ -1903,6 +1903,7 @@ void FGroomBuilder::BuildData(
 		OutGroupInfo.MaxImportedWidth = InHairDescriptionGroup.Info.MaxImportedWidth;
 
 		// Rendering data
+		if (InHairDescriptionGroup.Strands.IsValid())
 		{
 			OutRen = InHairDescriptionGroup.Strands;
 
@@ -1926,6 +1927,7 @@ void FGroomBuilder::BuildData(
 		}
 
 		// Simulation data
+		if (InHairDescriptionGroup.Guides.IsValid())
 		{
 			OutSim = InHairDescriptionGroup.Guides;
 			if (InHairDescriptionGroup.Info.NumGuides > 0 && !InSettings.InterpolationSettings.bOverrideGuides)
diff --git a/Engine/Plugins/Runtime/HairStrands/Source/HairStrandsCore/Private/HairStrandsRendering.cpp b/Engine/Plugins/Runtime/HairStrands/Source/HairStrandsCore/Private/HairStrandsRendering.cpp
--- a/Engine/Plugins/Runtime/HairStrands/Source/HairStrandsCore/Private/HairStrandsRendering.cpp
+++ b/Engine/Plugins/Runtime/HairStrands/Source/HairStrandsCore/Private/HairStrandsRendering.cpp
@@ -237,7 +237,7 @@ static void AddGroomCacheUpdatePass(
 	}
 	Parameters->bHasRadiusData = bHasRadiusData ? 1 : 0;
 
-	const FIntVector DispatchCount = FIntVector(FMath::DivideAndRoundUp(ElementCount, FGroomCacheUpdatePassCS::GetGroupSize()), 0, 0);
+	const FIntVector DispatchCount = FIntVector(FMath::DivideAndRoundUp(ElementCount, FGroomCacheUpdatePassCS::GetGroupSize()), 1, 1);
 	TShaderMapRef<FGroomCacheUpdatePassCS> ComputeShader(ShaderMap);
 	FComputeShaderUtils::AddPass(
 		GraphBuilder,
@@ -1739,6 +1739,8 @@ void ComputeHairStrandsInterpolation(
 		RDG_EVENT_SCOPE(GraphBuilder, "HairInterpolation(Strands)");
 		RDG_GPU_STAT_SCOPE(GraphBuilder, HairStrandsInterpolation);
 
+		const bool bHasValidGroomCacheBuffers = Instance->Debug.GroomCacheBuffers.IsValid() && Instance->Debug.GroomCacheBuffers->GetInterpolatedFrameBuffer().GroupsData.IsValidIndex(Instance->Debug.GroupIndex);
+
 		// Debug mode:
 		// * None	: Display hair normally
 		// * Sim	: Show sim strands
@@ -1748,6 +1750,22 @@ void ComputeHairStrandsInterpolation(
 			// Disable culling when drawing only guides, as the culling output has been computed for the strands, not for the guides.
 			Instance->HairGroupPublicData->SetCullingResultAvailable(false);
 
+			// If groom guide cache, is enabled, used the deform position for debug visualization
+			if (bHasValidGroomCacheBuffers && Instance->Debug.GroomCacheType == EGroomCacheType::Guides)
+			{
+				FScopeLock Lock(Instance->Debug.GroomCacheBuffers->GetCriticalSection());
+				const FGroomCacheGroupData& GroomCacheGroupData = Instance->Debug.GroomCacheBuffers->GetInterpolatedFrameBuffer().GroupsData[Instance->Debug.GroupIndex];
+
+				Instance->Guides.DeformedResource->SetPositionOffset(FHairStrandsDeformedResource::EFrameType::Current, GroomCacheGroupData.BoundingBox.GetCenter());
+
+				AddHairStrandUpdatePositionOffsetPass(
+					GraphBuilder,
+					ShaderMap,
+					MeshLODIndex,
+					Instance->Guides.DeformedRootResource,
+					Instance->Guides.DeformedResource);
+			}
+
 			AddHairTangentPass(
 				GraphBuilder,
 				ShaderMap,
@@ -1809,8 +1827,6 @@ void ComputeHairStrandsInterpolation(
 				Strands_PositionOffsetSRV = RegisterAsSRV(GraphBuilder, Instance->Strands.DeformedResource->GetPositionOffsetBuffer(FHairStrandsDeformedResource::EFrameType::Current));
 				Strands_TangentSRV = Strands_DeformedTangent.SRV;
 
-
-				bool bHasValidGroomCacheBuffers = Instance->Debug.GroomCacheBuffers.IsValid() && Instance->Debug.GroomCacheBuffers->GetInterpolatedFrameBuffer().GroupsData.IsValidIndex(Instance->Debug.GroupIndex);
 				if (bHasValidGroomCacheBuffers && Instance->Debug.GroomCacheType == EGroomCacheType::Guides)
 				{
 					FScopeLock Lock(Instance->Debug.GroomCacheBuffers->GetCriticalSection());
