Convert all groom actions to the new API.

diff --git a/Engine/Plugins/Runtime/HairStrands/Source/HairStrandsEditor/HairStrandsEditor.Build.cs b/Engine/Plugins/Runtime/HairStrands/Source/HairStrandsEditor/HairStrandsEditor.Build.cs
--- a/Engine/Plugins/Runtime/HairStrands/Source/HairStrandsEditor/HairStrandsEditor.Build.cs
+++ b/Engine/Plugins/Runtime/HairStrands/Source/HairStrandsEditor/HairStrandsEditor.Build.cs
@@ -23,6 +23,7 @@ namespace UnrealBuildTool.Rules
 					"ToolMenus",
 					"ContentBrowser",
 					"UnrealEd",
+					"AssetDefinition",
 					"AssetTools",
 					"EditorInteractiveToolsFramework",
 					"AdvancedPreviewScene",
@@ -49,4 +50,4 @@ namespace UnrealBuildTool.Rules
 			);
 		}
 	}
-}
+}
\ No newline at end of file
diff --git a/Engine/Plugins/Runtime/HairStrands/Source/HairStrandsEditor/Private/GroomActions.cpp b/Engine/Plugins/Runtime/HairStrands/Source/HairStrandsEditor/Private/GroomActions.cpp
--- a/Engine/Plugins/Runtime/HairStrands/Source/HairStrandsEditor/Private/GroomActions.cpp
+++ b/Engine/Plugins/Runtime/HairStrands/Source/HairStrandsEditor/Private/GroomActions.cpp
@@ -34,79 +34,33 @@
 #include "Widgets/Notifications/SNotificationList.h"
 #include "Framework/Notifications/NotificationManager.h"
 
-#define LOCTEXT_NAMESPACE "AssetTypeActions"
-
-///////////////////////////////////////////////////////////////////////////////////////////////////
-// Actions
+#include "AssetToolsModule.h"
+#include "ContentBrowserMenuContexts.h"
+#include "ToolMenus.h"
 
-FGroomActions::FGroomActions()
-{}
-
-bool FGroomActions::CanFilter()
-{
-	return true;
-}
+#define LOCTEXT_NAMESPACE "AssetTypeActions"
 
-void FGroomActions::GetActions(const TArray<UObject*>& InObjects, FToolMenuSection& Section)
+FLinearColor UAssetDefinition_GroomAsset::GetAssetColor() const
 {
-	TArray<TWeakObjectPtr<UGroomAsset>> GroomAssets = GetTypedWeakObjectPtrs<UGroomAsset>(InObjects);
-
-	Section.AddMenuEntry(
-		"RebuildGroom",
-		LOCTEXT("RebuildGroom", "Rebuild"),
-		LOCTEXT("RebuildGroomTooltip", "Rebuild the groom with new build settings"),
-		FSlateIcon(FAppStyle::GetAppStyleSetName(), "ContentBrowser.AssetActions"),
-		FUIAction(
-			FExecuteAction::CreateSP(this, &FGroomActions::ExecuteRebuild, GroomAssets),
-			FCanExecuteAction::CreateSP(this, &FGroomActions::CanRebuild, GroomAssets)
-		)
-	);
-
-	Section.AddMenuEntry(
-		"CreateBindingAsset",
-		LOCTEXT("CreateBindingAsset", "Create Binding"),
-		LOCTEXT("CreateBindingAssetTooltip", "Create a binding asset between a skeletal mesh and a groom asset"),
-		FSlateIcon(FAppStyle::GetAppStyleSetName(), "ContentBrowser.AssetActions"),
-		FUIAction(
-			FExecuteAction::CreateSP(this, &FGroomActions::ExecuteCreateBindingAsset, GroomAssets),
-			FCanExecuteAction::CreateSP(this, &FGroomActions::CanCreateBindingAsset, GroomAssets)
-		)
-	);
-
-	Section.AddMenuEntry(
-		"CreateFollicleTexture",
-		LOCTEXT("CreateFollicleTexture", "Create Follicle Texture"),
-		LOCTEXT("CreateFollicleTextureTooltip", "Create a follicle texture for the selected groom assets"),
-		FSlateIcon(FAppStyle::GetAppStyleSetName(), "ContentBrowser.AssetActions"),
-		FUIAction(
-			FExecuteAction::CreateSP(this, &FGroomActions::ExecuteCreateFollicleTexture, GroomAssets),
-			FCanExecuteAction::CreateSP(this, &FGroomActions::CanCreateFollicleTexture, GroomAssets)
-			)
-		);
-
-	Section.AddMenuEntry(
-		"CreateStrandsTextures",
-		LOCTEXT("CreateStrandsTextures", "Create Strands Textures"),
-		LOCTEXT("CreateStrandsTexturesTooltip", "Create projected strands textures onto meshes"),
-		FSlateIcon(FAppStyle::GetAppStyleSetName(), "ContentBrowser.AssetActions"),
-		FUIAction(
-			FExecuteAction::CreateSP(this, &FGroomActions::ExecuteCreateStrandsTextures, GroomAssets),
-			FCanExecuteAction::CreateSP(this, &FGroomActions::CanCreateStrandsTextures, GroomAssets)
-		)
-	);
+	return FColor::White;
 }
 
-uint32 FGroomActions::GetCategories()
+EAssetCommandResult UAssetDefinition_GroomAsset::OpenAssets(const FAssetOpenArgs& OpenArgs) const
 {
-	return EAssetTypeCategories::Misc;
-}
+	// #ueent_todo: Will need a custom editor at some point, for now just use the Properties editor
+	for (UGroomAsset* GroomAsset : OpenArgs.LoadObjects<UGroomAsset>())
+	{
+		if (GroomAsset != nullptr)
+		{
+			TSharedRef<FGroomCustomAssetEditorToolkit> NewCustomAssetEditor(new FGroomCustomAssetEditorToolkit());
+			NewCustomAssetEditor->InitCustomAssetEditor(OpenArgs.GetToolkitMode(), OpenArgs.ToolkitHost, GroomAsset);
+		}
+	}
 
-FText FGroomActions::GetName() const
-{
-	return NSLOCTEXT("AssetTypeActions", "AssetTypeActions_Groom", "Groom");
+	return EAssetCommandResult::Handled;
 }
 
-void FGroomActions::GetResolvedSourceFilePaths(const TArray<UObject*>& TypeAssets, TArray<FString>& OutSourceFilePaths) const
+void UAssetDefinition_GroomAsset::GetResolvedSourceFilePaths(const TArray<UObject*>& TypeAssets, TArray<FString>& OutSourceFilePaths) const
 {
 	for (UObject* Asset : TypeAssets)
 	{
@@ -118,41 +72,17 @@ void FGroomActions::GetResolvedSourceFilePaths(const TArray<UObject*>& TypeAsset
 	}
 }
 
-UClass* FGroomActions::GetSupportedClass() const
-{
-	return UGroomAsset::StaticClass();
-}
-
-FColor FGroomActions::GetTypeColor() const
+namespace MenuExtension_GroomAsset
 {
-	return FColor::White;
-}
-
-void FGroomActions::OpenAssetEditor(const TArray<UObject*>& InObjects, TSharedPtr<IToolkitHost> EditWithinLevelEditor)
-{
-	// #ueent_todo: Will need a custom editor at some point, for now just use the Properties editor
-	for (auto ObjIt = InObjects.CreateConstIterator(); ObjIt; ++ObjIt)
-	{
-		auto GroomAsset = Cast<UGroomAsset>(*ObjIt);
-		if (GroomAsset != nullptr)
-		{
-			// Make sure the groom asset has a document 
-			//FHairLabDataIO::GetDocumentForAsset(GroomAsset);
-
-			TSharedRef<FGroomCustomAssetEditorToolkit> NewCustomAssetEditor(new FGroomCustomAssetEditorToolkit());
-			NewCustomAssetEditor->InitCustomAssetEditor(EToolkitMode::Standalone, EditWithinLevelEditor, GroomAsset);
-		}
-	}
-}
-
 ////////////////////////////////////////////////////////////////////////////////////////////////////////////////////////
 // Groom build/rebuild
 
-bool FGroomActions::CanRebuild(TArray<TWeakObjectPtr<UGroomAsset>> Objects) const
+bool CanRebuild(const FToolMenuContext& InContext)
 {
-	for (TWeakObjectPtr<UGroomAsset> GroomAsset : Objects)
+	const UContentBrowserAssetContextMenuContext* CBContext = UContentBrowserAssetContextMenuContext::FindContextWithAssets(InContext);
+	for (UGroomAsset* GroomAsset : CBContext->LoadSelectedObjects<UGroomAsset>())
 	{
-		if (GroomAsset.IsValid() && GroomAsset->CanRebuildFromDescription())
+		if (GroomAsset && GroomAsset->IsValid() && GroomAsset->CanRebuildFromDescription())
 		{
 			return true;
 		}
@@ -160,11 +90,12 @@ bool FGroomActions::CanRebuild(TArray<TWeakObjectPtr<UGroomAsset>> Objects) cons
 	return false;
 }
 
-void FGroomActions::ExecuteRebuild(TArray<TWeakObjectPtr<UGroomAsset>> Objects) const
+void ExecuteRebuild(const FToolMenuContext& InContext)
 {
-	for (TWeakObjectPtr<UGroomAsset> GroomAsset : Objects)
+	const UContentBrowserAssetContextMenuContext* CBContext = UContentBrowserAssetContextMenuContext::FindContextWithAssets(InContext);
+	for (UGroomAsset* GroomAsset : CBContext->LoadSelectedObjects<UGroomAsset>())
 	{
-		if (GroomAsset.IsValid() && GroomAsset->CanRebuildFromDescription() && GroomAsset->AssetImportData)
+		if (GroomAsset && GroomAsset->IsValid() && GroomAsset->CanRebuildFromDescription() && GroomAsset->AssetImportData)
 		{
 			UGroomAssetImportData* GroomAssetImportData = Cast<UGroomAssetImportData>(GroomAsset->AssetImportData);
 			if (GroomAssetImportData && GroomAssetImportData->ImportOptions)
@@ -215,7 +146,7 @@ void FGroomActions::ExecuteRebuild(TArray<TWeakObjectPtr<UGroomAsset>> Objects)
 				{
 					if(bEnableRigging)
 					{
-						GroomAsset->RiggedSkeletalMesh = FGroomDeformerBuilder::CreateSkeletalMesh(GroomAsset.Get());
+						GroomAsset->RiggedSkeletalMesh = FGroomDeformerBuilder::CreateSkeletalMesh(GroomAsset);
 					}
 					// Move the transient ImportOptions to the asset package and set it on the GroomAssetImportData for serialization
 					CurrentOptions->Rename(nullptr, GroomAssetImportData);
@@ -234,11 +165,12 @@ void FGroomActions::ExecuteRebuild(TArray<TWeakObjectPtr<UGroomAsset>> Objects)
 ////////////////////////////////////////////////////////////////////////////////////////////////////////////////////////
 // Binding
 
-bool FGroomActions::CanCreateBindingAsset(TArray<TWeakObjectPtr<UGroomAsset>> Objects) const
+bool CanCreateBindingAsset(const FToolMenuContext& InContext)
 {
-	for (TWeakObjectPtr<UGroomAsset> GroomAsset : Objects)
+	const UContentBrowserAssetContextMenuContext* CBContext = UContentBrowserAssetContextMenuContext::FindContextWithAssets(InContext);
+	for (const FAssetData& SelectedAsset : CBContext->SelectedAssets)
 	{
-		if (GroomAsset.IsValid())
+		if (SelectedAsset.IsValid())
 		{
 			return true;
 		}
@@ -246,11 +178,12 @@ bool FGroomActions::CanCreateBindingAsset(TArray<TWeakObjectPtr<UGroomAsset>> Ob
 	return false;
 }
 
-void FGroomActions::ExecuteCreateBindingAsset(TArray<TWeakObjectPtr<UGroomAsset>> Objects) const
+void ExecuteCreateBindingAsset(const FToolMenuContext& InContext)
 {
-	for (TWeakObjectPtr<UGroomAsset> GroomAsset : Objects)
+	const UContentBrowserAssetContextMenuContext* CBContext = UContentBrowserAssetContextMenuContext::FindContextWithAssets(InContext);
+	for (UGroomAsset* GroomAsset : CBContext->LoadSelectedObjects<UGroomAsset>())
 	{
-		if (GroomAsset.IsValid())
+		if (GroomAsset && GroomAsset->IsValid())
 		{
 			// Duplicate the options to prevent dirtying the asset when they are modified but the rebuild is cancelled
 			UGroomCreateBindingOptions* CurrentOptions = NewObject<UGroomCreateBindingOptions>();
@@ -264,7 +197,7 @@ void FGroomActions::ExecuteCreateBindingAsset(TArray<TWeakObjectPtr<UGroomAsset>
 			{
 				continue;
 			}
-			else if (GroomAsset.Get() && CurrentOptions && 
+			else if (CurrentOptions && 
 				    ((CurrentOptions->GroomBindingType == EGroomBindingMeshType::SkeletalMesh && CurrentOptions->TargetSkeletalMesh) ||
 					(CurrentOptions->GroomBindingType == EGroomBindingMeshType::GeometryCache && CurrentOptions->TargetGeometryCache)))
 			{
@@ -278,7 +211,7 @@ void FGroomActions::ExecuteCreateBindingAsset(TArray<TWeakObjectPtr<UGroomAsset>
 					{
 						CurrentOptions->SourceSkeletalMesh->ConditionalPostLoad();
 					}
-					BindingAsset = FHairStrandsCore::CreateGroomBindingAsset(CurrentOptions->GroomBindingType, GroomAsset.Get(), CurrentOptions->SourceSkeletalMesh, CurrentOptions->TargetSkeletalMesh, CurrentOptions->NumInterpolationPoints, CurrentOptions->MatchingSection);
+					BindingAsset = FHairStrandsCore::CreateGroomBindingAsset(CurrentOptions->GroomBindingType, GroomAsset, CurrentOptions->SourceSkeletalMesh, CurrentOptions->TargetSkeletalMesh, CurrentOptions->NumInterpolationPoints, CurrentOptions->MatchingSection);
 				}
 				else
 				{
@@ -287,7 +220,7 @@ void FGroomActions::ExecuteCreateBindingAsset(TArray<TWeakObjectPtr<UGroomAsset>
 					{
 						CurrentOptions->SourceGeometryCache->ConditionalPostLoad();
 					}
-					BindingAsset = FHairStrandsCore::CreateGroomBindingAsset(CurrentOptions->GroomBindingType, GroomAsset.Get(), CurrentOptions->SourceGeometryCache, CurrentOptions->TargetGeometryCache, CurrentOptions->NumInterpolationPoints, CurrentOptions->MatchingSection);
+					BindingAsset = FHairStrandsCore::CreateGroomBindingAsset(CurrentOptions->GroomBindingType, GroomAsset, CurrentOptions->SourceGeometryCache, CurrentOptions->TargetGeometryCache, CurrentOptions->NumInterpolationPoints, CurrentOptions->MatchingSection);
 				}
 
 				if (BindingAsset)
@@ -300,9 +233,9 @@ void FGroomActions::ExecuteCreateBindingAsset(TArray<TWeakObjectPtr<UGroomAsset>
 
 						FContentBrowserModule& ContentBrowserModule = FModuleManager::LoadModuleChecked<FContentBrowserModule>("ContentBrowser");
 						ContentBrowserModule.Get().SyncBrowserToAssets(CreatedObjects);
-#if WITH_EDITOR
+					#if WITH_EDITOR
 						GEditor->GetEditorSubsystem<UAssetEditorSubsystem>()->OpenEditorForAssets(CreatedObjects);
-#endif
+					#endif
 					}
 					else
 					{
@@ -324,11 +257,12 @@ void FGroomActions::ExecuteCreateBindingAsset(TArray<TWeakObjectPtr<UGroomAsset>
 ////////////////////////////////////////////////////////////////////////////////////////////////////////////////////////
 // Follicle
 
-bool FGroomActions::CanCreateFollicleTexture(TArray<TWeakObjectPtr<UGroomAsset>> Objects) const
+bool CanCreateFollicleTexture(const FToolMenuContext& InContext)
 {
-	for (TWeakObjectPtr<UGroomAsset> GroomAsset : Objects)
+	const UContentBrowserAssetContextMenuContext* CBContext = UContentBrowserAssetContextMenuContext::FindContextWithAssets(InContext);
+	for (const FAssetData& SelectedAsset : CBContext->SelectedAssets)
 	{
-		if (GroomAsset.IsValid())
+		if (SelectedAsset.IsValid())
 		{
 			return true;
 		}
@@ -336,9 +270,11 @@ bool FGroomActions::CanCreateFollicleTexture(TArray<TWeakObjectPtr<UGroomAsset>>
 	return false;
 }
 
-void FGroomActions::ExecuteCreateFollicleTexture(TArray<TWeakObjectPtr<UGroomAsset>> Objects) const
+void ExecuteCreateFollicleTexture(const FToolMenuContext& InContext)
 {
-	if (Objects.Num() == 0)
+	const UContentBrowserAssetContextMenuContext* CBContext = UContentBrowserAssetContextMenuContext::FindContextWithAssets(InContext);
+	TArray<UGroomAsset*> GroomAssets = CBContext->LoadSelectedObjects<UGroomAsset>();
+	if (GroomAssets.Num() == 0)
 	{
 		return;
 	}
@@ -350,12 +286,12 @@ void FGroomActions::ExecuteCreateFollicleTexture(TArray<TWeakObjectPtr<UGroomAss
 		return;
 	}
 
-	for (TWeakObjectPtr<UGroomAsset>& GroomAsset : Objects)
+	for (UGroomAsset* GroomAsset : GroomAssets)
 	{
-		if (GroomAsset.IsValid())
+		if (GroomAsset && GroomAsset->IsValid())
 		{
 			FFollicleMaskOptions& Items = CurrentOptions->Grooms.AddDefaulted_GetRef();;
-			Items.Groom   = GroomAsset.Get();
+			Items.Groom   = GroomAsset;
 			Items.Channel = EFollicleMaskChannel::R;
 		}
 	}
@@ -399,11 +335,12 @@ void FGroomActions::ExecuteCreateFollicleTexture(TArray<TWeakObjectPtr<UGroomAss
 ////////////////////////////////////////////////////////////////////////////////////////////////////////////////////////
 // Strands Textures
 
-bool FGroomActions::CanCreateStrandsTextures(TArray<TWeakObjectPtr<UGroomAsset>> Objects) const
+bool CanCreateStrandsTextures(const FToolMenuContext& InContext)
 {
-	for (TWeakObjectPtr<UGroomAsset> GroomAsset : Objects)
+	const UContentBrowserAssetContextMenuContext* CBContext = UContentBrowserAssetContextMenuContext::FindContextWithAssets(InContext);
+	for (const FAssetData& SelectedAsset : CBContext->SelectedAssets)
 	{
-		if (GroomAsset.IsValid())
+		if (SelectedAsset.IsValid())
 		{
 			return true;
 		}
@@ -411,18 +348,20 @@ bool FGroomActions::CanCreateStrandsTextures(TArray<TWeakObjectPtr<UGroomAsset>>
 	return false;
 }
 
-void FGroomActions::ExecuteCreateStrandsTextures(TArray<TWeakObjectPtr<UGroomAsset>> Objects) const
+void ExecuteCreateStrandsTextures(const FToolMenuContext& InContext)
 {
-	if (Objects.Num() == 0)
+	const UContentBrowserAssetContextMenuContext* CBContext = UContentBrowserAssetContextMenuContext::FindContextWithAssets(InContext);
+	TArray<UGroomAsset*> GroomAssets = CBContext->LoadSelectedObjects<UGroomAsset>();
+	if (GroomAssets.Num() == 0)
 	{
 		return;
 	}
 
 	// Duplicate the options to prevent dirtying the asset when they are modified but the rebuild is cancelled
 	UGroomCreateStrandsTexturesOptions* CurrentOptions = NewObject<UGroomCreateStrandsTexturesOptions>();
-	for (TWeakObjectPtr<UGroomAsset>& GroomAsset : Objects)
+	for (UGroomAsset* GroomAsset : GroomAssets)
 	{
-		if (GroomAsset.IsValid() && CurrentOptions)
+		if (GroomAsset && GroomAsset->IsValid() && CurrentOptions)
 		{
 			TSharedPtr<SGroomCreateStrandsTexturesOptionsWindow> GroomOptionWindow = SGroomCreateStrandsTexturesOptionsWindow::DisplayCreateStrandsTexturesOptions(CurrentOptions);
 			if (!GroomOptionWindow->ShouldCreate())
@@ -461,7 +400,7 @@ void FGroomActions::ExecuteCreateStrandsTextures(TArray<TWeakObjectPtr<UGroomAss
 				}
 				
 				FStrandsTexturesInfo Info;
-				Info.GroomAsset   = GroomAsset.Get();
+				Info.GroomAsset   = GroomAsset;
 				Info.TracingDirection = SignDirection;
 				Info.MaxTracingDistance = MaxDistance;
 				Info.Resolution = FMath::RoundUpToPowerOfTwo(FMath::Max(256, CurrentOptions->Resolution));
@@ -481,7 +420,7 @@ void FGroomActions::ExecuteCreateStrandsTextures(TArray<TWeakObjectPtr<UGroomAss
 						Info.GroupIndices.Add(GroupIndex);
 					}
 				}
-				FStrandsTexturesOutput Output = FGroomTextureBuilder::CreateGroomStrandsTexturesTexture(GroomAsset.Get(), Info.Resolution);
+				FStrandsTexturesOutput Output = FGroomTextureBuilder::CreateGroomStrandsTexturesTexture(GroomAsset, Info.Resolution);
 				if (Output.IsValid())
 				{
 					FGroomTextureBuilder::BuildStrandsTextures(Info, Output);
@@ -489,8 +428,67 @@ void FGroomActions::ExecuteCreateStrandsTextures(TArray<TWeakObjectPtr<UGroomAss
 			}
 		}
 	}
-
 }
 
 
+////////////////////////////////////////////////////////////////////////////////////////////////////////////////////////
+// Actions registration
+
+static FDelayedAutoRegisterHelper DelayedAutoRegister(EDelayedRegisterRunPhase::EndOfEngineInit, []{ 
+	UToolMenus::RegisterStartupCallback(FSimpleMulticastDelegate::FDelegate::CreateLambda([]()
+	{
+		FToolMenuOwnerScoped OwnerScoped(UE_MODULE_NAME);
+		UToolMenu* Menu = UE::ContentBrowser::ExtendToolMenu_AssetContextMenu(UGroomAsset::StaticClass());
+		
+		FToolMenuSection& Section = Menu->FindOrAddSection("GetAssetActions");
+		Section.AddDynamicEntry(NAME_None, FNewToolMenuSectionDelegate::CreateLambda([](FToolMenuSection& InSection)
+		{
+			{
+				const TAttribute<FText> Label = LOCTEXT("RebuildGroom", "Rebuild");
+				const TAttribute<FText> ToolTip = LOCTEXT("RebuildGroomTooltip", "Rebuild the groom with new build settings");
+				const FSlateIcon Icon = FSlateIcon(FAppStyle::GetAppStyleSetName(), "ContentBrowser.AssetActions");
+
+				FToolUIAction UIAction;
+				UIAction.ExecuteAction = FToolMenuExecuteAction::CreateStatic(&ExecuteRebuild);
+				UIAction.CanExecuteAction = FToolMenuCanExecuteAction::CreateStatic(&CanRebuild);
+				InSection.AddMenuEntry("GroomAsset_RebuildGroom", Label, ToolTip, Icon, UIAction);					
+			}
+
+			{
+				const TAttribute<FText> Label = LOCTEXT("CreateBindingAsset", "Create Binding");
+				const TAttribute<FText> ToolTip = LOCTEXT("CreateBindingAssetTooltip", "Create a binding asset between a skeletal mesh and a groom asset");
+				const FSlateIcon Icon = FSlateIcon(FAppStyle::GetAppStyleSetName(), "ContentBrowser.AssetActions");
+
+				FToolUIAction UIAction;
+				UIAction.ExecuteAction = FToolMenuExecuteAction::CreateStatic(&ExecuteCreateBindingAsset);
+				UIAction.CanExecuteAction = FToolMenuCanExecuteAction::CreateStatic(&CanCreateBindingAsset);
+				InSection.AddMenuEntry("GroomAsset_CreateBindingAsset", Label, ToolTip, Icon, UIAction);
+			}
+
+			{
+				const TAttribute<FText> Label = LOCTEXT("CreateFollicleTexture", "Create Follicle Texture");
+				const TAttribute<FText> ToolTip = LOCTEXT("CreateFollicleTextureTooltip", "Create a follicle texture for the selected groom assets");
+				const FSlateIcon Icon = FSlateIcon(FAppStyle::GetAppStyleSetName(), "ContentBrowser.AssetActions");
+
+				FToolUIAction UIAction;
+				UIAction.ExecuteAction = FToolMenuExecuteAction::CreateStatic(&ExecuteCreateFollicleTexture);
+				UIAction.CanExecuteAction = FToolMenuCanExecuteAction::CreateStatic(&CanCreateFollicleTexture);
+				InSection.AddMenuEntry("GroomAsset_CreateFollicleTexture", Label, ToolTip, Icon, UIAction);
+			}
+
+			{
+				const TAttribute<FText> Label = LOCTEXT("CreateStrandsTextures", "Create Strands Textures");
+				const TAttribute<FText> ToolTip = LOCTEXT("CreateStrandsTexturesTooltip", "Create projected strands textures onto meshes");
+				const FSlateIcon Icon = FSlateIcon(FAppStyle::GetAppStyleSetName(), "ContentBrowser.AssetActions");
+
+				FToolUIAction UIAction;
+				UIAction.ExecuteAction = FToolMenuExecuteAction::CreateStatic(&ExecuteCreateStrandsTextures);
+				UIAction.CanExecuteAction = FToolMenuCanExecuteAction::CreateStatic(&CanCreateStrandsTextures);
+				InSection.AddMenuEntry("GroomAsset_CreateStrandsTextures", Label, ToolTip, Icon, UIAction);
+			}
+		}));
+	}));
+});
+
+} // namespace MenuExtension_GroomAsset
 #undef LOCTEXT_NAMESPACE
diff --git a/Engine/Plugins/Runtime/HairStrands/Source/HairStrandsEditor/Private/GroomBindingActions.cpp b/Engine/Plugins/Runtime/HairStrands/Source/HairStrandsEditor/Private/GroomBindingActions.cpp
--- a/Engine/Plugins/Runtime/HairStrands/Source/HairStrandsEditor/Private/GroomBindingActions.cpp
+++ b/Engine/Plugins/Runtime/HairStrands/Source/HairStrandsEditor/Private/GroomBindingActions.cpp
@@ -4,6 +4,7 @@
 
 #include "ContentBrowserMenuContexts.h"
 #include "GroomAsset.h"
+#include "GroomCustomAssetEditorToolkit.h"
 
 #include "GeometryCache.h"
 #include "Toolkits/SimpleAssetEditor.h"
@@ -13,52 +14,17 @@
 
 #define LOCTEXT_NAMESPACE "AssetTypeActions"
 
-uint32 FGroomBindingActions::GetCategories()
-{
-	return EAssetTypeCategories::Misc;
-}
-
-FText FGroomBindingActions::GetName() const
-{
-	return NSLOCTEXT("AssetTypeActions", "AssetTypeActions_GroomBinding", "GroomBinding");
-}
-
-UClass* FGroomBindingActions::GetSupportedClass() const
-{
-	return UGroomBindingAsset::StaticClass();
-}
-
-FColor FGroomBindingActions::GetTypeColor() const
+FLinearColor UAssetDefinition_GroomBindingAsset::GetAssetColor() const
 {
 	return FColor::White;
 }
 
-void FGroomBindingActions::OpenAssetEditor(const TArray<UObject*>& InObjects, TSharedPtr<IToolkitHost> EditWithinLevelEditor)
+namespace MenuExtension_GroomBindingAsset
 {
-	// #ueent_todo: Will need a custom editor at some point, for now just use the Properties editor
-	FSimpleAssetEditor::CreateEditor(EToolkitMode::Standalone, EditWithinLevelEditor, InObjects);
-}
 
-void FGroomBindingActions::RegisterMenus()
+void ExecuteRebuildBindingAsset(const FToolMenuContext& InContext)
 {
-	FToolMenuOwnerScoped MenuOwner(UE_MODULE_NAME);
-	UToolMenu* Menu = UToolMenus::Get()->ExtendMenu("ContentBrowser.AssetContextMenu.GroomBinding");
-	FToolMenuSection& Section = Menu->FindOrAddSection("GetAssetActions");
-
-	Section.AddDynamicEntry("GroomBinding", FNewToolMenuSectionDelegate::CreateLambda([](FToolMenuSection& InSection)
-	{
-		const TAttribute<FText> Label = LOCTEXT("RebuildGroomBinding", "Rebuild");
-		const TAttribute<FText> ToolTip = LOCTEXT("RebuildGroomBindingTooltip", "Rebuild the groom binding");
-		const FSlateIcon Icon = FSlateIcon(FAppStyle::GetAppStyleSetName(), "ContentBrowser.AssetActions");
-		const FToolMenuExecuteAction UIAction = FToolMenuExecuteAction::CreateStatic(&FGroomBindingActions::ExecuteRebuildBindingAsset);
-
-		InSection.AddMenuEntry("GroomBinding_Rebuild", Label, ToolTip, Icon, UIAction);
-	}));
-}
-
-void FGroomBindingActions::ExecuteRebuildBindingAsset(const FToolMenuContext& MenuContext)
-{
-	if (const UContentBrowserAssetContextMenuContext* Context = UContentBrowserAssetContextMenuContext::FindContextWithAssets(MenuContext))
+	if (const UContentBrowserAssetContextMenuContext* Context = UContentBrowserAssetContextMenuContext::FindContextWithAssets(InContext))
 	{
 		for (UGroomBindingAsset* BindingAsset : Context->LoadSelectedObjects<UGroomBindingAsset>())
 		{
@@ -88,4 +54,34 @@ void FGroomBindingActions::ExecuteRebuildBindingAsset(const FToolMenuContext& Me
 	}
 }
 
+void GroomBindingRebuildAction(FToolMenuSection& InSection)
+{
+	const TAttribute<FText> Label = LOCTEXT("RebuildGroomBinding", "Rebuild");
+	const TAttribute<FText> ToolTip = LOCTEXT("RebuildGroomBindingTooltip", "Rebuild the groom binding");
+	const FSlateIcon Icon = FSlateIcon(FAppStyle::GetAppStyleSetName(), "ContentBrowser.AssetActions");
+
+	FToolUIAction UIAction;
+	UIAction.ExecuteAction = FToolMenuExecuteAction::CreateStatic(&ExecuteRebuildBindingAsset);
+	InSection.AddMenuEntry("GroomBinding_Rebuild", Label, ToolTip, Icon, UIAction);
+}
+
+////////////////////////////////////////////////////////////////////////////////////////////////////////////////////////
+// Actions registration
+
+static FDelayedAutoRegisterHelper DelayedAutoRegister(EDelayedRegisterRunPhase::EndOfEngineInit, []{ 
+	UToolMenus::RegisterStartupCallback(FSimpleMulticastDelegate::FDelegate::CreateLambda([]()
+	{
+		FToolMenuOwnerScoped OwnerScoped(UE_MODULE_NAME);
+		UToolMenu* Menu = UE::ContentBrowser::ExtendToolMenu_AssetContextMenu(UGroomBindingAsset::StaticClass());
+		
+		FToolMenuSection& Section = Menu->FindOrAddSection("GetAssetActions");
+		Section.AddDynamicEntry(NAME_None, FNewToolMenuSectionDelegate::CreateLambda([](FToolMenuSection& InSection)
+		{
+			GroomBindingRebuildAction(InSection);
+		}));
+	}));
+});
+
+} // namespace MenuExtension_GroomBindingAsset
+
 #undef LOCTEXT_NAMESPACE
\ No newline at end of file
diff --git a/Engine/Plugins/Runtime/HairStrands/Source/HairStrandsEditor/Private/GroomCacheActions.cpp b/Engine/Plugins/Runtime/HairStrands/Source/HairStrandsEditor/Private/GroomCacheActions.cpp
--- a/Engine/Plugins/Runtime/HairStrands/Source/HairStrandsEditor/Private/GroomCacheActions.cpp
+++ b/Engine/Plugins/Runtime/HairStrands/Source/HairStrandsEditor/Private/GroomCacheActions.cpp
@@ -2,29 +2,14 @@
 
 #include "GroomCacheActions.h"
 #include "EditorFramework/AssetImportData.h"
-#include "GroomCache.h"
+#include "GroomCustomAssetEditorToolkit.h"
 
-FText FGroomCacheActions::GetName() const
-{
-	return NSLOCTEXT("AssetTypeActions", "AssetTypeActions_GroomCache", "Groom Cache");
-}
-
-UClass* FGroomCacheActions::GetSupportedClass() const
-{
-	return UGroomCache::StaticClass();
-}
-
-FColor FGroomCacheActions::GetTypeColor() const
+FLinearColor UAssetDefinition_GroomCacheAsset::GetAssetColor() const
 {
 	return FColor::White;
 }
 
-uint32 FGroomCacheActions::GetCategories()
-{
-	return EAssetTypeCategories::Misc;
-}
-
-void FGroomCacheActions::GetResolvedSourceFilePaths(const TArray<UObject*>& TypeAssets, TArray<FString>& OutSourceFilePaths) const
+void UAssetDefinition_GroomCacheAsset::GetResolvedSourceFilePaths(const TArray<UObject*>& TypeAssets, TArray<FString>& OutSourceFilePaths) const
 {
 	for (UObject* Asset : TypeAssets)
 	{
diff --git a/Engine/Plugins/Runtime/HairStrands/Source/HairStrandsEditor/Private/HairStrandsEditor.cpp b/Engine/Plugins/Runtime/HairStrands/Source/HairStrandsEditor/Private/HairStrandsEditor.cpp
--- a/Engine/Plugins/Runtime/HairStrands/Source/HairStrandsEditor/Private/HairStrandsEditor.cpp
+++ b/Engine/Plugins/Runtime/HairStrands/Source/HairStrandsEditor/Private/HairStrandsEditor.cpp
@@ -76,15 +76,6 @@ void FGroomEditor::StartupModule()
 	UToolMenus::RegisterStartupCallback(FSimpleMulticastDelegate::FDelegate::CreateRaw(this, &FGroomEditor::RegisterMenus));
 
 	IAssetTools& AssetTools = FModuleManager::LoadModuleChecked<FAssetToolsModule>("AssetTools").Get();
-	TSharedRef<IAssetTypeActions> GroomAssetActions = MakeShared<FGroomActions>();
-	TSharedRef<IAssetTypeActions> BindingAssetActions = MakeShared<FGroomBindingActions>();
-	TSharedRef<IAssetTypeActions> GroomCacheActions = MakeShared<FGroomCacheActions>();
-
-	AssetTools.RegisterAssetTypeActions(GroomAssetActions);
-	AssetTools.RegisterAssetTypeActions(BindingAssetActions);
-	AssetTools.RegisterAssetTypeActions(GroomCacheActions);
-	RegisteredAssetTypeActions.Add(GroomAssetActions);
-	RegisteredAssetTypeActions.Add(BindingAssetActions);
 
 	// Only register once
 	if (!StyleSet.IsValid())
@@ -180,19 +171,6 @@ void FGroomEditor::ShutdownModule()
 		}
 	}
 
-	// #ueent_todo: Unregister the translators
-	FAssetToolsModule* AssetToolsModule = FModuleManager::GetModulePtr<FAssetToolsModule>("AssetTools");
-
-	if (AssetToolsModule != nullptr)
-	{
-		IAssetTools& AssetTools = AssetToolsModule->Get();
-
-		for (auto Action : RegisteredAssetTypeActions)
-		{
-			AssetTools.UnregisterAssetTypeActions(Action);
-		}
-	}
-
 	if (StyleSet.IsValid())
 	{
 		FSlateStyleRegistry::UnRegisterSlateStyle(*StyleSet.Get());
@@ -203,7 +181,7 @@ void FGroomEditor::ShutdownModule()
 
 void FGroomEditor::RegisterMenus()
 {
-	FGroomBindingActions::RegisterMenus();
+
 }
 
 TArray<TSharedPtr<IGroomTranslator>> FGroomEditor::GetHairTranslators()
diff --git a/Engine/Plugins/Runtime/HairStrands/Source/HairStrandsEditor/Public/GroomActions.h b/Engine/Plugins/Runtime/HairStrands/Source/HairStrandsEditor/Public/GroomActions.h
--- a/Engine/Plugins/Runtime/HairStrands/Source/HairStrandsEditor/Public/GroomActions.h
+++ b/Engine/Plugins/Runtime/HairStrands/Source/HairStrandsEditor/Public/GroomActions.h
@@ -2,55 +2,30 @@
 
 #pragma once
 
-#include "AssetTypeActions_Base.h"
-#include "Templates/SharedPointer.h"
+#include "CoreMinimal.h"
+#include "GroomAsset.h"
+#include "AssetDefinitionDefault.h"
 
-class ISlateStyle;
-class UGroomAsset;
+#include "GroomActions.generated.h"
 
-/**
- * Implements an action for groom assets.
- */
-class FGroomActions : public FAssetTypeActions_Base
+UCLASS()
+class UAssetDefinition_GroomAsset : public UAssetDefinitionDefault
 {
-public:
-
-	/**
-	 * Creates and initializes a new instance.
-	 *
-	 * @param InStyle The style set to use for asset editor toolkits.
-	 */
-	FGroomActions();
+	GENERATED_BODY()
 
 public:
-
-	//~ FAssetTypeActions_Base overrides
-
-	virtual bool CanFilter() override;
-	virtual void GetActions(const TArray<UObject*>& InObjects, FToolMenuSection& Section) override;
-	virtual uint32 GetCategories() override;
-	virtual FText GetName() const override;
-	virtual void GetResolvedSourceFilePaths(const TArray<UObject*>& TypeAssets, TArray<FString>& OutSourceFilePaths) const override;
-	virtual UClass* GetSupportedClass() const override;
-	virtual FColor GetTypeColor() const override;
-	virtual bool IsImportedAsset() const override { return true; }
-	virtual void OpenAssetEditor(const TArray<UObject*>& InObjects, TSharedPtr<IToolkitHost> EditWithinLevelEditor = TSharedPtr<IToolkitHost>()) override;
-
-private:
-
-	/** Callback for Rebuild groom action */
-	bool CanRebuild(TArray<TWeakObjectPtr<UGroomAsset>> Objects) const;
-	void ExecuteRebuild(TArray<TWeakObjectPtr<UGroomAsset>> Objects) const;
-
-	/** Callback for creating binding asset action */
-	bool CanCreateBindingAsset(TArray<TWeakObjectPtr<UGroomAsset>> Objects) const;
-	void ExecuteCreateBindingAsset(TArray<TWeakObjectPtr<UGroomAsset>> Objects) const;
-
-	/** Callback for creating follicle texture action */
-	bool CanCreateFollicleTexture(TArray<TWeakObjectPtr<UGroomAsset>> Objects) const;
-	void ExecuteCreateFollicleTexture(TArray<TWeakObjectPtr<UGroomAsset>> Objects) const;
-
-	/** Callback for creating strands texture action */
-	bool CanCreateStrandsTextures(TArray<TWeakObjectPtr<UGroomAsset>> Objects) const;
-	void ExecuteCreateStrandsTextures(TArray<TWeakObjectPtr<UGroomAsset>> Objects) const;
+	// UAssetDefinition Begin
+	virtual FText GetAssetDisplayName() const override { return NSLOCTEXT("AssetTypeActions", "AssetTypeActions_Groom", "Groom"); }
+	virtual FLinearColor GetAssetColor() const override;
+	virtual TSoftClassPtr<UObject> GetAssetClass() const override { return UGroomAsset::StaticClass(); }
+	virtual TConstArrayView<FAssetCategoryPath> GetAssetCategories() const override
+	{
+		static const auto Categories = { EAssetCategoryPaths::Misc };
+		return Categories;
+	}
+
+	virtual EAssetCommandResult OpenAssets(const FAssetOpenArgs& OpenArgs) const override;
+	// UAssetDefinition End	
+
+	void GetResolvedSourceFilePaths(const TArray<UObject*>& TypeAssets, TArray<FString>& OutSourceFilePaths) const;
 };
diff --git a/Engine/Plugins/Runtime/HairStrands/Source/HairStrandsEditor/Public/GroomBindingActions.h b/Engine/Plugins/Runtime/HairStrands/Source/HairStrandsEditor/Public/GroomBindingActions.h
--- a/Engine/Plugins/Runtime/HairStrands/Source/HairStrandsEditor/Public/GroomBindingActions.h
+++ b/Engine/Plugins/Runtime/HairStrands/Source/HairStrandsEditor/Public/GroomBindingActions.h
@@ -2,41 +2,26 @@
 
 #pragma once
 
-#include "AssetTypeActions_Base.h"
-#include "Templates/SharedPointer.h"
+#include "CoreMinimal.h"
 #include "GroomBindingAsset.h"
+#include "AssetDefinitionDefault.h"
 
-class ISlateStyle;
-struct FToolMenuContext;
+#include "GroomBindingActions.generated.h"
 
-/**
- * Implements an action for groom binding assets.
- */
-class FGroomBindingActions : public FAssetTypeActions_Base
+UCLASS()
+class UAssetDefinition_GroomBindingAsset : public UAssetDefinitionDefault
 {
-public:
-
-	/**
-	 * Creates and initializes a new instance.
-	 *
-	 * @param InStyle The style set to use for asset editor toolkits.
-	 */
-	FGroomBindingActions() { }
+	GENERATED_BODY()
 
 public:
-
-	//~ FAssetTypeActions_Base overrides
-
-	virtual bool CanFilter() override { return true; }
-	virtual uint32 GetCategories() override;
-	virtual FText GetName() const override;
-	virtual UClass* GetSupportedClass() const override;
-	virtual FColor GetTypeColor() const override;
-	virtual void OpenAssetEditor(const TArray<UObject*>& InObjects, TSharedPtr<IToolkitHost> EditWithinLevelEditor = TSharedPtr<IToolkitHost>()) override;
-
-public:
-	static void RegisterMenus();
-
-protected:
-	static void ExecuteRebuildBindingAsset(const FToolMenuContext& MenuContext);
+	// UAssetDefinition Begin
+	virtual FText GetAssetDisplayName() const override { return NSLOCTEXT("AssetTypeActions", "AssetTypeActions_GroomBinding", "GroomBinding"); }
+	virtual FLinearColor GetAssetColor() const override;
+	virtual TSoftClassPtr<UObject> GetAssetClass() const override { return UGroomBindingAsset::StaticClass(); }
+	virtual TConstArrayView<FAssetCategoryPath> GetAssetCategories() const override
+	{
+		static const auto Categories = { EAssetCategoryPaths::Misc };
+		return Categories;
+	}
+	// UAssetDefinition End
 };
diff --git a/Engine/Plugins/Runtime/HairStrands/Source/HairStrandsEditor/Public/GroomCacheActions.h b/Engine/Plugins/Runtime/HairStrands/Source/HairStrandsEditor/Public/GroomCacheActions.h
--- a/Engine/Plugins/Runtime/HairStrands/Source/HairStrandsEditor/Public/GroomCacheActions.h
+++ b/Engine/Plugins/Runtime/HairStrands/Source/HairStrandsEditor/Public/GroomCacheActions.h
@@ -2,19 +2,28 @@
 
 #pragma once
 
-#include "AssetTypeActions_Base.h"
+#include "CoreMinimal.h"
+#include "GroomCache.h"
+#include "AssetDefinitionDefault.h"
 
-class FGroomCacheActions : public FAssetTypeActions_Base
+#include "GroomCacheActions.generated.h"
+
+UCLASS()
+class UAssetDefinition_GroomCacheAsset : public UAssetDefinitionDefault
 {
-public:
+	GENERATED_BODY()
 
-	//~ IAssetTypeActions Interface
-	virtual FText GetName() const override;
-	virtual UClass* GetSupportedClass() const override;
-	virtual FColor GetTypeColor() const override;
-	virtual uint32 GetCategories() override;
+public:
+	// UAssetDefinition Begin
+	virtual FText GetAssetDisplayName() const override { return NSLOCTEXT("AssetTypeActions", "AssetTypeActions_GroomCache", "Groom Cache"); }
+	virtual FLinearColor GetAssetColor() const override;
+	virtual TSoftClassPtr<UObject> GetAssetClass() const override { return UGroomCache::StaticClass(); }
+	virtual TConstArrayView<FAssetCategoryPath> GetAssetCategories() const override
+	{
+		static const auto Categories = { EAssetCategoryPaths::Misc };
+		return Categories;
+	}
+	// UAssetDefinition End	
 
-	//~ FAssetTypeActions_Base overrides
-	virtual bool IsImportedAsset() const override { return true; }
-	virtual void GetResolvedSourceFilePaths(const TArray<UObject*>& TypeAssets, TArray<FString>& OutSourceFilePaths) const override;
+	void GetResolvedSourceFilePaths(const TArray<UObject*>& TypeAssets, TArray<FString>& OutSourceFilePaths) const;
 };
diff --git a/Engine/Plugins/Runtime/HairStrands/Source/HairStrandsEditor/Public/HairStrandsEditor.h b/Engine/Plugins/Runtime/HairStrands/Source/HairStrandsEditor/Public/HairStrandsEditor.h
--- a/Engine/Plugins/Runtime/HairStrands/Source/HairStrandsEditor/Public/HairStrandsEditor.h
+++ b/Engine/Plugins/Runtime/HairStrands/Source/HairStrandsEditor/Public/HairStrandsEditor.h
@@ -51,9 +51,6 @@ public:
 private:
 	void RegisterMenus();
 
-	/** The collection of registered asset type actions. */
-	TArray<TSharedRef<class IAssetTypeActions>> RegisteredAssetTypeActions;
-
 	TArray<TFunction<TSharedPtr<IGroomTranslator>()>> TranslatorSpawners;
 
 	TSharedPtr<FSlateStyleSet> StyleSet;
