Fix crash with instance visualization.

diff --git a/Engine/Plugins/Runtime/HairStrands/Source/HairStrandsCore/Private/GroomDebug.cpp b/Engine/Plugins/Runtime/HairStrands/Source/HairStrandsCore/Private/GroomDebug.cpp
--- a/Engine/Plugins/Runtime/HairStrands/Source/HairStrandsCore/Private/GroomDebug.cpp
+++ b/Engine/Plugins/Runtime/HairStrands/Source/HairStrandsCore/Private/GroomDebug.cpp
@@ -925,7 +925,7 @@ static void AddHairDebugPrintInstancePass(
 				Data0.Z = Instance->Strands.Data->GetNumCurves(); // Change this later on for having dynamic value
 				Data0.W = Instance->Strands.Data->GetNumPoints(); // Change this later on for having dynamic value
 				const int32 MeshLODIndex = Instance->HairGroupPublicData->MeshLODIndex;
-				if (MeshLODIndex>=0)
+				if (MeshLODIndex>=0 && Instance->Strands.RestRootResource)
 				{
 					Data1.X = Instance->Strands.RestRootResource->BulkData.MeshProjectionLODs[MeshLODIndex].UniqueSectionIndices.Num();
 					Data1.Y = Instance->Strands.RestRootResource->BulkData.MeshProjectionLODs[MeshLODIndex].UniqueTriangleCount;
