Add strands details to instance listing for quicker debugging.

diff --git a/Engine/Plugins/Runtime/HairStrands/Source/HairStrandsCore/Private/GroomDebug.cpp b/Engine/Plugins/Runtime/HairStrands/Source/HairStrandsCore/Private/GroomDebug.cpp
--- a/Engine/Plugins/Runtime/HairStrands/Source/HairStrandsCore/Private/GroomDebug.cpp
+++ b/Engine/Plugins/Runtime/HairStrands/Source/HairStrandsCore/Private/GroomDebug.cpp
@@ -807,15 +807,23 @@ class FHairDebugPrintInstanceCS : public FGlobalShader
 
 	BEGIN_SHADER_PARAMETER_STRUCT(FParameters, )
 		SHADER_PARAMETER(uint32, InstanceCount)
-		SHADER_PARAMETER(uint32, NameInfoCount)
-		SHADER_PARAMETER(uint32, NameCharacterCount)
+
 		SHADER_PARAMETER(uint32, InstanceCount_StrandsPrimaryView)
 		SHADER_PARAMETER(uint32, InstanceCount_StrandsShadowView)
 		SHADER_PARAMETER(uint32, InstanceCount_CardsOrMeshesPrimaryView)
 		SHADER_PARAMETER(uint32, InstanceCount_CardsOrMeshesShadowView)
 		SHADER_PARAMETER(FVector4f, InstanceScreenSphereBound)
-		SHADER_PARAMETER_RDG_BUFFER_SRV(StructuredBuffer<uint2>, NameInfos)
-		SHADER_PARAMETER_RDG_BUFFER_SRV(Buffer<uint8>, Names)
+		// Instance names
+		SHADER_PARAMETER(uint32, InstanceNameInfoCount)
+		SHADER_PARAMETER(uint32, InstanceNameCharacterCount)
+		SHADER_PARAMETER_RDG_BUFFER_SRV(StructuredBuffer<uint2>, InstanceNameInfos)
+		SHADER_PARAMETER_RDG_BUFFER_SRV(Buffer<uint8>, InstanceNames)
+		// Attribute names
+		SHADER_PARAMETER(uint32, AttributeNameInfoCount)
+		SHADER_PARAMETER(uint32, AttributeNameCharacterCount)
+		SHADER_PARAMETER_RDG_BUFFER_SRV(StructuredBuffer<uint2>, AttributeNameInfos)
+		SHADER_PARAMETER_RDG_BUFFER_SRV(Buffer<uint8>, AttributeNames)
+
 		SHADER_PARAMETER_RDG_BUFFER_SRV(StructuredBuffer<uint4>, Infos)
 		SHADER_PARAMETER_RDG_BUFFER_SRV(Buffer<int>, InstanceAABB)
 		SHADER_PARAMETER_STRUCT_INCLUDE(ShaderPrint::FShaderParameters, ShaderPrintUniformBuffer)
@@ -842,6 +850,8 @@ struct FHairDebugNameInfo
 	uint8  Pad0;
 };
 
+const TCHAR* GetHairAttributeText(EHairAttribute In);
+
 static void AddHairDebugPrintInstancePass(
 	FRDGBuilder& GraphBuilder, 
 	FGlobalShaderMap* ShaderMap,
@@ -863,15 +873,16 @@ static void AddHairDebugPrintInstancePass(
 	const uint32 MaxPrimitiveNameCount = 128u;
 	check(sizeof(FHairDebugNameInfo) == 8);
 
-	TArray<FHairDebugNameInfo> NameInfos;
-	TArray<uint8> Names;
-	Names.Reserve(MaxPrimitiveNameCount * 30u);
+	TArray<FHairDebugNameInfo> InstanceNameInfos;
+	TArray<uint8> InstanceNames;
+	InstanceNames.Reserve(MaxPrimitiveNameCount * 30u);
 
 	struct FInstanceInfos
 	{
 		FUintVector4 Data0;
 		FUintVector4 Data1;
 		FUintVector4 Data2;
+		FUintVector4 Data3;
 	};
 	TArray<FInstanceInfos> Infos;
 	Infos.Reserve(InstanceCount);
@@ -880,18 +891,18 @@ static void AddHairDebugPrintInstancePass(
 		const FHairStrandsInstance* AbstractInstance = Instances[InstanceIndex];
 		const FHairGroupInstance* Instance = static_cast<const FHairGroupInstance*>(AbstractInstance);
 
-		// Collect Names
+		// Collect InstanceNames
 		if (InstanceIndex < MaxPrimitiveNameCount)
 		{
 			const FString Name = *Instance->Debug.GroomAssetName;
-			const uint32 NameOffset = Names.Num();
+			const uint32 NameOffset = InstanceNames.Num();
 			const uint32 NameLength = Name.Len();
 			for (TCHAR C : Name)
 			{
-				Names.Add(uint8(C));
+				InstanceNames.Add(uint8(C));
 			}
 
-			FHairDebugNameInfo& NameInfo = NameInfos.AddDefaulted_GetRef();
+			FHairDebugNameInfo& NameInfo = InstanceNameInfos.AddDefaulted_GetRef();
 			NameInfo.PrimitiveID = InstanceIndex;
 			NameInfo.Length = NameLength;
 			NameInfo.Offset = NameOffset;
@@ -904,6 +915,7 @@ static void AddHairDebugPrintInstancePass(
 		FUintVector4 Data0 = { 0,0,0,0 };
 		FUintVector4 Data1 = { 0,0,0,0 };
 		FUintVector4 Data2 = { 0,0,0,0 };
+		FUintVector4 Data3 = { 0,0,0,0 };
 		Data0.X =
 			((Instance->Debug.GroupIndex & 0xFF)) |
 			((Instance->Debug.GroupCount & 0xFF) << 8) |
@@ -935,12 +947,35 @@ static void AddHairDebugPrintInstancePass(
 
 				{
 					Data2 = FUintVector4(0);
-					Data2.X |= InstanceIndex < InstanceCountPerType[HairInstanceCount_StrandsPrimaryView] ? 0x1u : 0u;
-					Data2.X |= InstanceIndex < InstanceCountPerType[HairInstanceCount_StrandsShadowView]  ? 0x2u : 0u;
+					Data2.X |= InstanceIndex < InstanceCountPerType[HairInstanceCount_StrandsPrimaryView] ? 0x1u  : 0u;
+					Data2.X |= InstanceIndex < InstanceCountPerType[HairInstanceCount_StrandsShadowView]  ? 0x2u  : 0u;
+					Data2.X |= Instance->HairGroupPublicData->VFInput.Strands.bScatterSceneLighting       ? 0x4u  : 0u;
+					Data2.X |= Instance->HairGroupPublicData->VFInput.Strands.bUseRaytracingGeometry      ? 0x8u  : 0u;
+					Data2.X |= Instance->HairGroupPublicData->VFInput.Strands.bUseStableRasterization     ? 0x10u : 0u;
+					Data2.X |= Instance->HairGroupPublicData->bSupportVoxelization                        ? 0x20u : 0u;
 					Data2.X |= uint32(FFloat16(Instance->HairGroupPublicData->ContinuousLODScreenSize).Encoded) << 16u;
+
 					Data2.Y = Instance->HairGroupPublicData->GetActiveStrandsPointCount();
 					Data2.Z = Instance->HairGroupPublicData->GetActiveStrandsCurveCount();
+
+					if (Instance->Strands.Data)
+					{
+						for (uint32 AttributeIt=0; AttributeIt<uint32(EHairAttribute::Count); ++AttributeIt)
+						{
+							Data2.W |= Instance->Strands.Data->HasAttributes((EHairAttribute)AttributeIt) ? (1u<<AttributeIt) : 0u;
+						}
+					}
 				}
+				
+				Data3 = FUintVector4(0);
+				Data3.X |= FFloat16(Instance->HairGroupPublicData->VFInput.Strands.HairRadius).Encoded;
+				Data3.X |= FFloat16(Instance->HairGroupPublicData->VFInput.Strands.HairDensity).Encoded << 16u;
+
+				Data3.Y |= FFloat16(Instance->HairGroupPublicData->VFInput.Strands.HairRootScale).Encoded;
+				Data3.Y |= FFloat16(Instance->HairGroupPublicData->VFInput.Strands.HairTipScale).Encoded << 16u;
+
+				Data3.Z |= FFloat16(Instance->HairGroupPublicData->VFInput.Strands.HairLength).Encoded;
+				Data3.Z |= FFloat16(Instance->HairGroupPublicData->VFInput.Strands.HairLengthScale).Encoded << 16u;
 			}
 			break;
 		case EHairGeometryType::Cards:
@@ -966,27 +1001,49 @@ static void AddHairDebugPrintInstancePass(
 			}
 			break;
 		}
-		Infos.Add({Data0, Data1, Data2 });
+		Infos.Add({Data0, Data1, Data2, Data3 });
 	}
 
-	if (NameInfos.IsEmpty())
+	if (InstanceNameInfos.IsEmpty())
 	{
-		FHairDebugNameInfo& NameInfo = NameInfos.AddDefaulted_GetRef();
+		FHairDebugNameInfo& NameInfo = InstanceNameInfos.AddDefaulted_GetRef();
 		NameInfo.PrimitiveID = ~0;
 		NameInfo.Length = 4;
 		NameInfo.Offset = 0;
-		Names.Add(uint8('N'));
-		Names.Add(uint8('o'));
-		Names.Add(uint8('n'));
-		Names.Add(uint8('e'));
+		InstanceNames.Add(uint8('N'));
+		InstanceNames.Add(uint8('o'));
+		InstanceNames.Add(uint8('n'));
+		InstanceNames.Add(uint8('e'));
 	}	
 
+	TArray<FHairDebugNameInfo> AttributeNameInfos;
+	TArray<uint8> AttributeNames;
+	AttributeNames.Reserve(uint32(EHairAttribute::Count) * 30u);
+	for (uint32 AttributeIt = 0; AttributeIt < uint32(EHairAttribute::Count); ++AttributeIt)
+	{
+		const FString Name = GetHairAttributeText((EHairAttribute)AttributeIt);
+		FHairDebugNameInfo& NameInfo = AttributeNameInfos.AddDefaulted_GetRef();
+		NameInfo.PrimitiveID = ~0;
+		NameInfo.Length = Name.Len();
+		NameInfo.Offset = AttributeNames.Num();	
+		for (TCHAR C : Name)
+		{
+			AttributeNames.Add(uint8(C));
+		}
+	}
+
 	const uint32 InfoInBytes  = sizeof(FInstanceInfos);
 	const uint32 InfoIn4Bytes = sizeof(FInstanceInfos) / sizeof(uint8);
-	FRDGBufferRef NameBuffer = CreateVertexBuffer(GraphBuilder, TEXT("Hair.Debug.InstanceNames"), FRDGBufferDesc::CreateBufferDesc(1, Names.Num()), Names.GetData(), Names.Num());
-	FRDGBufferRef NameInfoBuffer = CreateStructuredBuffer(GraphBuilder, TEXT("Hair.Debug.InstanceNameInfos"), NameInfos);	
 	FRDGBufferRef InfoBuffer = CreateVertexBuffer(GraphBuilder, TEXT("Hair.Debug.InstanceInfos"), FRDGBufferDesc::CreateBufferDesc(4, InfoIn4Bytes * Infos.Num()), Infos.GetData(), InfoInBytes * Infos.Num());
 
+	// Resource for instance names
+	FRDGBufferRef InstanceNameBuffer = CreateVertexBuffer(GraphBuilder, TEXT("Hair.Debug.InstanceNames"), FRDGBufferDesc::CreateBufferDesc(1, InstanceNames.Num()), InstanceNames.GetData(), InstanceNames.Num());
+	FRDGBufferRef InstanceNameInfoBuffer = CreateStructuredBuffer(GraphBuilder, TEXT("Hair.Debug.InstanceNameInfos"), InstanceNameInfos);
+
+	// Resource for attribute names
+	FRDGBufferRef AttributeNameBuffer = CreateVertexBuffer(GraphBuilder, TEXT("Hair.Debug.AttributeNames"), FRDGBufferDesc::CreateBufferDesc(1, AttributeNames.Num()), AttributeNames.GetData(), AttributeNames.Num());
+	FRDGBufferRef AttributeNameInfoBuffer = CreateStructuredBuffer(GraphBuilder, TEXT("Hair.Debug.AttributeNameInfos"), AttributeNameInfos);
+
 	// Draw general information for all instances (one pass for all instances)
 	{
 		FHairDebugPrintInstanceCS::FParameters* Parameters = GraphBuilder.AllocParameters<FHairDebugPrintInstanceCS::FParameters>();
@@ -995,10 +1052,19 @@ static void AddHairDebugPrintInstancePass(
 		Parameters->InstanceCount_StrandsShadowView = InstanceCountPerType[HairInstanceCount_StrandsShadowView];
 		Parameters->InstanceCount_CardsOrMeshesPrimaryView = InstanceCountPerType[HairInstanceCount_CardsOrMeshesPrimaryView];
 		Parameters->InstanceCount_CardsOrMeshesShadowView = InstanceCountPerType[HairInstanceCount_CardsOrMeshesShadowView];
-		Parameters->NameInfoCount = NameInfos.Num();
-		Parameters->NameCharacterCount = Names.Num();
-		Parameters->Names = GraphBuilder.CreateSRV(NameBuffer, PF_R8_UINT);
-		Parameters->NameInfos = GraphBuilder.CreateSRV(NameInfoBuffer);
+
+		// Instance name
+		Parameters->InstanceNameInfoCount = InstanceNameInfos.Num();
+		Parameters->InstanceNameCharacterCount = InstanceNames.Num();
+		Parameters->InstanceNames = GraphBuilder.CreateSRV(InstanceNameBuffer, PF_R8_UINT);
+		Parameters->InstanceNameInfos = GraphBuilder.CreateSRV(InstanceNameInfoBuffer);
+
+		// Attribute name
+		Parameters->AttributeNameInfoCount = AttributeNameInfos.Num();
+		Parameters->AttributeNameCharacterCount = AttributeNames.Num();
+		Parameters->AttributeNames = GraphBuilder.CreateSRV(AttributeNameBuffer, PF_R8_UINT);
+		Parameters->AttributeNameInfos = GraphBuilder.CreateSRV(AttributeNameInfoBuffer);
+
 		Parameters->Infos = GraphBuilder.CreateSRV(InfoBuffer, PF_R32_UINT);
 		ShaderPrint::SetParameters(GraphBuilder, *ShaderPrintData, Parameters->ShaderPrintUniformBuffer);
 		FHairDebugPrintInstanceCS::FPermutationDomain PermutationVector;
diff --git a/Engine/Plugins/Runtime/HairStrands/Source/HairStrandsCore/Private/HairAttributes.cpp b/Engine/Plugins/Runtime/HairStrands/Source/HairStrandsCore/Private/HairAttributes.cpp
--- a/Engine/Plugins/Runtime/HairStrands/Source/HairStrandsCore/Private/HairAttributes.cpp
+++ b/Engine/Plugins/Runtime/HairStrands/Source/HairStrandsCore/Private/HairAttributes.cpp
@@ -45,4 +45,21 @@ bool HasHairAttribute(uint32 In, EHairAttribute InAttribute)
 void SetHairAttribute(uint32& Out, EHairAttribute InAttribute) 
 { 
 	Out |= (1u << uint32(InAttribute)); 
+}
+
+const TCHAR* GetHairAttributeText(EHairAttribute In)
+{
+	// If a new optional attribute is added, please add its UI/text description here
+	static_assert(uint32(EHairAttribute::Count) == 6);
+
+	switch (In)
+	{
+	case EHairAttribute::RootUV:					return TEXT("RootUV");
+	case EHairAttribute::ClumpID:					return TEXT("ClumpID");
+	case EHairAttribute::PrecomputedGuideWeights:	return TEXT("ImportedGuideWeights");
+	case EHairAttribute::Color:						return TEXT("Color");
+	case EHairAttribute::Roughness:					return TEXT("Roughness");
+	case EHairAttribute::AO:						return TEXT("AO");
+	}
+	return TEXT("UNKNOWN");
 }
\ No newline at end of file
diff --git a/Engine/Plugins/Runtime/HairStrands/Source/HairStrandsCore/Public/HairStrandsDatas.h b/Engine/Plugins/Runtime/HairStrands/Source/HairStrandsCore/Public/HairStrandsDatas.h
--- a/Engine/Plugins/Runtime/HairStrands/Source/HairStrandsCore/Public/HairStrandsDatas.h
+++ b/Engine/Plugins/Runtime/HairStrands/Source/HairStrandsCore/Public/HairStrandsDatas.h
@@ -9,6 +9,7 @@
 #include "RenderGraphResources.h"
 #include "Serialization/BulkData.h"
 #include "HairStrandsDefinitions.h"
+#include "HairAttributes.h"
 
 DECLARE_LOG_CATEGORY_EXTERN(LogHairStrands, Log, All);
 
@@ -464,6 +465,7 @@ struct HAIRSTRANDSCORE_API FHairStrandsBulkData
 	float  GetMaxRadius() const { return MaxRadius; }
 	FVector GetPositionOffset() const { return BoundingBox.GetCenter(); }
 	const FBox& GetBounds() const { return BoundingBox; }
+	bool HasAttributes(EHairAttribute In) const { return AttributeOffsets[uint32(In)] != 0xFFFFFFFF; }
 
 	uint32 CurveCount = 0;
 	uint32 PointCount = 0;
diff --git a/Engine/Plugins/Runtime/HairStrands/Source/HairStrandsEditor/Private/GroomImportOptionsWindow.cpp b/Engine/Plugins/Runtime/HairStrands/Source/HairStrandsEditor/Private/GroomImportOptionsWindow.cpp
--- a/Engine/Plugins/Runtime/HairStrands/Source/HairStrandsEditor/Private/GroomImportOptionsWindow.cpp
+++ b/Engine/Plugins/Runtime/HairStrands/Source/HairStrandsEditor/Private/GroomImportOptionsWindow.cpp
@@ -211,7 +211,7 @@ static void AddAttribute(SVerticalBox::FScopedWidgetSlotArguments& Slot, FText A
 	];
 }
 
-FText GetHairAttributeText(EHairAttribute In)
+static FText GetHairAttributeLocText(EHairAttribute In)
 {
 	// If a new optional attribute is added, please add its UI/text description here
 	static_assert(uint32(EHairAttribute::Count) == 6);
@@ -412,7 +412,7 @@ void SGroomImportOptionsWindow::Construct(const FArguments& InArgs)
 		if (HasHairAttribute(Attributes, AttributeType))
 		{
 			SVerticalBox::FScopedWidgetSlotArguments SlotArg = VerticalSlot->InsertSlot(AttributeSlotIndex++);
-			AddAttribute(SlotArg, GetHairAttributeText(AttributeType));
+			AddAttribute(SlotArg, GetHairAttributeLocText(AttributeType));
 		}
 	}
 
diff --git a/Engine/Shaders/Private/HairStrands/HairStrandsDebug.usf b/Engine/Shaders/Private/HairStrands/HairStrandsDebug.usf
--- a/Engine/Shaders/Private/HairStrands/HairStrandsDebug.usf
+++ b/Engine/Shaders/Private/HairStrands/HairStrandsDebug.usf
@@ -62,15 +62,20 @@ void PrintLargeNumber(inout FShaderPrintContext Context, uint Value, FFontColor
 
 void PrintPercentage(inout FShaderPrintContext Context, float Value, FFontColor InColor)
 {
-	const float fValue = clamp(Value * 100.f, 0.f, 99.9f);
-	if (fValue < 10.f)
+	const uint iValue = ceil(clamp(Value * 100.f, 0.f, 100.f));
+	if (iValue < 10)
+	{
+		Print(Context, TEXT("  "), InColor);
+		Print(Context, iValue, InColor, 1);
+	}
+	else if (iValue < 100)
 	{
 		Print(Context, TEXT(" "), InColor);
-		Print(Context, fValue, InColor, 3, 1);
+		Print(Context, iValue, InColor, 2);
 	}
 	else
 	{
-		Print(Context, fValue, InColor, 4, 1);
+		Print(Context, iValue, InColor, 3);
 	}
 	Print(Context, TEXT("% "), InColor);
 }
@@ -1147,14 +1152,19 @@ uint					InstanceCount_CardsOrMeshesShadowView;
 
 Buffer<uint>			Infos;
 
-uint					NameInfoCount;
-uint					NameCharacterCount;
-StructuredBuffer<uint2> NameInfos;
-Buffer<uint>			Names;
+uint					InstanceNameInfoCount;
+uint					InstanceNameCharacterCount;
+StructuredBuffer<uint2> InstanceNameInfos;
+Buffer<uint>			InstanceNames;
 Buffer<int>				InstanceAABB;
 
 float4					InstanceScreenSphereBound;
 
+uint					AttributeNameInfoCount;
+uint					AttributeNameCharacterCount;
+StructuredBuffer<uint2> AttributeNameInfos;
+Buffer<uint>			AttributeNames;
+
 struct FInstanceInfo
 {
 	uint  GroupIndex;
@@ -1173,6 +1183,17 @@ struct FInstanceInfo
 	float ScreenSize;
 	uint ActivePointCount;
 	uint ActiveCurveCount;
+	uint Attributes;
+	float Radius;
+	float Density;
+	float TipScale;
+	float RootScale;
+	float Length;
+	float LengthScale;
+	bool Voxel;
+	bool Scatter;
+	bool RTGeometry;
+	bool StableRaster;
 
 	// Binding
 	uint UniqueTriangle;
@@ -1186,25 +1207,31 @@ struct FInstanceInfo
 
 FInstanceInfo ReadInstanceInfo(uint Index)
 {
-	const uint Index12 = Index * 12;
+	const uint Index16 = Index * 16;
 
 	uint4 Data0 = 0;
-	Data0.x = Infos[Index12 + 0];
-	Data0.y = Infos[Index12 + 1];
-	Data0.z = Infos[Index12 + 2];
-	Data0.w = Infos[Index12 + 3];
+	Data0.x = Infos[Index16 + 0];
+	Data0.y = Infos[Index16 + 1];
+	Data0.z = Infos[Index16 + 2];
+	Data0.w = Infos[Index16 + 3];
 
 	uint4 Data1 = 0;
-	Data1.x = Infos[Index12 + 4];
-	Data1.y = Infos[Index12 + 5];
-	Data1.z = Infos[Index12 + 6];
-	Data1.w = Infos[Index12 + 7];
+	Data1.x = Infos[Index16 + 4];
+	Data1.y = Infos[Index16 + 5];
+	Data1.z = Infos[Index16 + 6];
+	Data1.w = Infos[Index16 + 7];
 
 	uint4 Data2 = 0;
-	Data2.x = Infos[Index12 + 8];
-	Data2.y = Infos[Index12 + 9];
-	Data2.z = Infos[Index12 + 10];
-	Data2.w = Infos[Index12 + 11];
+	Data2.x = Infos[Index16 + 8];
+	Data2.y = Infos[Index16 + 9];
+	Data2.z = Infos[Index16 + 10];
+	Data2.w = Infos[Index16 + 11];
+
+	uint4 Data3 = 0;
+	Data3.x = Infos[Index16 + 12];
+	Data3.y = Infos[Index16 + 13];
+	Data3.z = Infos[Index16 + 14];
+	Data3.w = Infos[Index16 + 15];
 
 	FInstanceInfo Out = (FInstanceInfo)0;
 	Out.GroupIndex = (Data0.x >> 0) & 0xFF;
@@ -1229,13 +1256,30 @@ FInstanceInfo ReadInstanceInfo(uint Index)
 
 	Out.bIsVisibleInPrimaryView = Data2.x & 0x1;
 	Out.bIsVisibleInShadowView	= Data2.x & 0x2;
+	Out.Scatter					= Data2.x & 0x4;
+	Out.RTGeometry				= Data2.x & 0x8;
+	Out.StableRaster			= Data2.x & 0x10;
+	Out.Voxel					= Data2.x & 0x20;
 	Out.ScreenSize				= f16tof32((Data2.x>>16) & 0xFFFF);
 	Out.ActivePointCount		= Data2.y;
 	Out.ActiveCurveCount		= Data2.z;
+	Out.Attributes				= Data2.w;
+
+	Out.Radius					= f16tof32(Data3.x & 0xFFFF);
+	Out.Density					= f16tof32(Data3.x >> 16u);
+	Out.RootScale				= f16tof32(Data3.y & 0xFFFF);
+	Out.TipScale				= f16tof32(Data3.y >> 16u);
+	Out.Length					= f16tof32(Data3.z & 0xFFFF);
+	Out.LengthScale				= f16tof32(Data3.z >> 16u);
 
 	return Out;
 }
 
+bool HasAttribute(FInstanceInfo In, uint InAttributeIt)
+{
+	return (In.Attributes & (1u<< InAttributeIt)) != 0;
+}
+
 struct FNameInfo
 {
 	uint PrimitiveID;
@@ -1259,11 +1303,11 @@ FNameInfo FindNameInfo(uint PrimitiveID)
 {
 	FNameInfo Out = (FNameInfo)0;
 	Out.PrimitiveID = ~0;
-	for (uint It = 0; It < NameInfoCount; ++It)
+	for (uint It = 0; It < InstanceNameInfoCount; ++It)
 	{
-		if (NameInfos[It].x == PrimitiveID)
+		if (InstanceNameInfos[It].x == PrimitiveID)
 		{
-			return UnpackDebugNameInfo(NameInfos[It]);
+			return UnpackDebugNameInfo(InstanceNameInfos[It]);
 		}
 	}
 	return Out;
@@ -1272,11 +1316,29 @@ FNameInfo FindNameInfo(uint PrimitiveID)
 void PrintInstanceName(inout FShaderPrintContext Context, uint PrimitiveID, FFontColor InColor)
 {
 	const FNameInfo Info = FindNameInfo(PrimitiveID);
-	if (Info.PrimitiveID != ~0 && (Info.Length + Info.Offset) <= NameCharacterCount)
+	if (Info.PrimitiveID != ~0 && (Info.Length + Info.Offset) <= InstanceNameCharacterCount)
+	{
+		for (uint It = 0; It < Info.Length; ++It)
+		{
+			const uint Char = InstanceNames[It + Info.Offset];
+			PrintSymbol(Context, Char, InColor);
+		}
+	}
+}
+
+void PrintAttributeName(inout FShaderPrintContext Context, uint AttributeIt, FFontColor InColor)
+{
+	if (AttributeIt >= AttributeNameInfoCount)
+	{
+		return;
+	}
+
+	const FNameInfo Info = UnpackDebugNameInfo(AttributeNameInfos[AttributeIt]);
+	if ((Info.Length + Info.Offset) <= AttributeNameCharacterCount)
 	{
 		for (uint It = 0; It < Info.Length; ++It)
 		{
-			const uint Char = Names[It + Info.Offset];
+			const uint Char = AttributeNames[It + Info.Offset];
 			PrintSymbol(Context, Char, InColor);
 		}
 	}
@@ -1293,7 +1355,7 @@ void MainCS(uint GroupIndex : SV_GroupIndex, uint3 DispatchThreadId : SV_Dispatc
 	}
 
 	FShaderPrintContext Context = InitShaderPrintContext(true, uint2(50, 100));
-	// Draw legend
+	bool bShowAdvancedStrandsData = false;
 
 	// Draw totals
 	const uint TotalGroup = InstanceCount;
@@ -1317,7 +1379,7 @@ void MainCS(uint GroupIndex : SV_GroupIndex, uint3 DispatchThreadId : SV_Dispatc
 		PrintLargeNumber(Context,TotalStrandsCount, FontOrange);
 		Newline(Context);
 
-		Print(Context, TEXT("Total Strand Vertex : "), FontWhite);
+		Print(Context, TEXT("Total Strand Points : "), FontWhite);
 		PrintLargeNumber(Context,TotalVertexCount, FontSilver);
 		Newline(Context);
 	}
@@ -1337,11 +1399,14 @@ void MainCS(uint GroupIndex : SV_GroupIndex, uint3 DispatchThreadId : SV_Dispatc
 		Print(Context, TEXT("| Shadow:"), FontSilver);
 		Print(Context, InstanceCount_CardsOrMeshesShadowView, FontYellow, 4);
 		Newline(Context);
+
+		bShowAdvancedStrandsData = AddCheckbox(Context, TEXT("Advanced"), false, FontSilver);
+		Newline(Context);
 	}
 	Newline(Context);
 
 	// Draw instance details
-	Print(Context, TEXT("   Group  LOD     Geometry Binding Sim RBF Clip Screen Strands             Vertex              Bind.Tri/Sect Vis. Shad. Name"));
+	Print(Context, TEXT("   Group  LOD     Geometry Binding Sim RBF Clip Screen Strands            Points             Bind.Tri/Sect Vis. Shad. Name"));
 	Newline(Context);
 
 	for (uint InstanceIndex = 0; InstanceIndex < InstanceCount; ++InstanceIndex)
@@ -1426,6 +1491,70 @@ void MainCS(uint GroupIndex : SV_GroupIndex, uint3 DispatchThreadId : SV_Dispatc
 
 		Newline(Context);
 	}
+	Newline(Context);
+
+	// Optional advanced debugging
+	if (bShowAdvancedStrandsData)
+	{
+		// Attribute
+		Print(Context, TEXT("   Group  "), FontWhite);
+		Print(Context, TEXT("Voxel "), FontWhite);
+		Print(Context, TEXT("RT  "), FontWhite);
+		Print(Context, TEXT("Scatter "), FontWhite);
+		Print(Context, TEXT("Radius "), FontWhite);
+		Print(Context, TEXT("Root  "), FontWhite);
+		Print(Context, TEXT("Tip   "), FontWhite);
+		for (uint AttributeIt=0; AttributeIt < AttributeNameInfoCount; ++AttributeIt)
+		{
+			PrintAttributeName(Context, AttributeIt, FontWhite);
+			Print(Context, TEXT(" "));
+		}
+		Newline(Context);
+
+		for (uint StrandInstanceIndex = 0; StrandInstanceIndex < InstanceCount; ++StrandInstanceIndex)
+		{
+			const FInstanceInfo Info = ReadInstanceInfo(StrandInstanceIndex);
+			if (Info.GeometryType > 0)
+			{
+				continue;
+			}
+
+			// Index
+			Print(Context, StrandInstanceIndex, FontSilver, 2);
+			Print(Context, TEXT(" "));
+
+			// Group
+			Print(Context, Info.GroupIndex, FontYellow, 1);
+			Print(Context, TEXT("/"));
+			Print(Context, Info.GroupCount, FontYellow, 1);
+			Print(Context, TEXT("    "));
+
+			PrintBool(Context, Info.Voxel); Print(Context, TEXT("  "));
+			PrintBool(Context, Info.RTGeometry); 
+			PrintBool(Context, Info.Scatter); Print(Context, TEXT("    "));
+			Print(Context, Info.Radius, FontYellow, 5, 3); Print(Context, TEXT("  "));
+			Print(Context, Info.RootScale, FontYellow, 5, 3); Print(Context, TEXT(" "));
+			Print(Context, Info.TipScale, FontYellow, 5, 3); Print(Context, TEXT(" "));
+
+			for (uint AttributeIt = 0; AttributeIt < AttributeNameInfoCount; ++AttributeIt)
+			{
+				PrintBool(Context, HasAttribute(Info, AttributeIt));
+				const FNameInfo Info = UnpackDebugNameInfo(AttributeNameInfos[AttributeIt]);
+				if (Info.Length > 3)
+				{
+					const uint SpaceCount = Info.Length - 3u;
+					for (uint SpaceIt = 0; SpaceIt < SpaceCount; ++SpaceIt)
+					{
+						Print(Context, TEXT(" "));
+					}
+				}
+			}
+
+			PrintInstanceName(Context, StrandInstanceIndex, FontWhite);
+
+			Newline(Context);
+		}
+	}
 }
 #else  // PERMUTATION_OUTPUT_TYPE
 {
