Add cache type to hair instance listing + fix attribute listings.

diff --git a/Engine/Plugins/Runtime/HairStrands/Source/HairStrandsCore/Private/GroomDebug.cpp b/Engine/Plugins/Runtime/HairStrands/Source/HairStrandsCore/Private/GroomDebug.cpp
--- a/Engine/Plugins/Runtime/HairStrands/Source/HairStrandsCore/Private/GroomDebug.cpp
+++ b/Engine/Plugins/Runtime/HairStrands/Source/HairStrandsCore/Private/GroomDebug.cpp
@@ -16,6 +16,7 @@
 #include "CanvasTypes.h"
 #include "ShaderCompilerCore.h"
 #include "GroomVisualizationData.h"
+#include "GroomCacheData.h"
 
 ///////////////////////////////////////////////////////////////////////////////////////////////////
 
@@ -851,6 +852,8 @@ struct FHairDebugNameInfo
 };
 
 const TCHAR* GetHairAttributeText(EHairAttribute In);
+uint32 GetHairAttributeIndex(EHairAttribute In);
+EGroomCacheType GetHairInstanceCacheType(const FHairGroupInstance* Instance);
 
 static void AddHairDebugPrintInstancePass(
 	FRDGBuilder& GraphBuilder, 
@@ -911,6 +914,7 @@ static void AddHairDebugPrintInstancePass(
 		const float LODIndex = Instance->HairGroupPublicData->LODIndex;
 		const uint32 IntLODIndex = Instance->HairGroupPublicData->LODIndex;
 		const uint32 LODCount = Instance->HairGroupPublicData->GetLODScreenSizes().Num();
+		const EGroomCacheType ActiveGroomCacheType = GetHairInstanceCacheType(Instance);
 
 		FUintVector4 Data0 = { 0,0,0,0 };
 		FUintVector4 Data1 = { 0,0,0,0 };
@@ -953,6 +957,8 @@ static void AddHairDebugPrintInstancePass(
 					Data2.X |= Instance->HairGroupPublicData->VFInput.Strands.bUseRaytracingGeometry      ? 0x8u  : 0u;
 					Data2.X |= Instance->HairGroupPublicData->VFInput.Strands.bUseStableRasterization     ? 0x10u : 0u;
 					Data2.X |= Instance->HairGroupPublicData->bSupportVoxelization                        ? 0x20u : 0u;
+					Data2.X |= ActiveGroomCacheType == EGroomCacheType::Guides                            ? 0x40u : 0u;
+					Data2.X |= ActiveGroomCacheType == EGroomCacheType::Strands                           ? 0x80u : 0u;
 					Data2.X |= uint32(FFloat16(Instance->HairGroupPublicData->ContinuousLODScreenSize).Encoded) << 16u;
 
 					Data2.Y = Instance->HairGroupPublicData->GetActiveStrandsPointCount();
@@ -962,7 +968,12 @@ static void AddHairDebugPrintInstancePass(
 					{
 						for (uint32 AttributeIt=0; AttributeIt<uint32(EHairAttribute::Count); ++AttributeIt)
 						{
-							Data2.W |= Instance->Strands.Data->HasAttributes((EHairAttribute)AttributeIt) ? (1u<<AttributeIt) : 0u;
+							const EHairAttribute Attribute = (EHairAttribute)AttributeIt;
+							const uint32 Index = GetHairAttributeIndex(Attribute);
+							if (Index < HAIR_ATTRIBUTE_COUNT)
+							{
+								Data2.W |= Instance->Strands.Data->AttributeOffsets[Index] != 0xFFFFFFFF ? (1u<<Index) : 0u;
+							}
 						}
 					}
 				}
@@ -1021,14 +1032,20 @@ static void AddHairDebugPrintInstancePass(
 	AttributeNames.Reserve(uint32(EHairAttribute::Count) * 30u);
 	for (uint32 AttributeIt = 0; AttributeIt < uint32(EHairAttribute::Count); ++AttributeIt)
 	{
-		const FString Name = GetHairAttributeText((EHairAttribute)AttributeIt);
-		FHairDebugNameInfo& NameInfo = AttributeNameInfos.AddDefaulted_GetRef();
-		NameInfo.PrimitiveID = ~0;
-		NameInfo.Length = Name.Len();
-		NameInfo.Offset = AttributeNames.Num();	
-		for (TCHAR C : Name)
+		// Only add valid optional attribute 
+		// HAIR_ATTRIBUTE_XXX and EHairAttribute don't have a 1:1 mapping
+		const EHairAttribute Attribute = (EHairAttribute)AttributeIt;
+		if (GetHairAttributeIndex(Attribute) < HAIR_ATTRIBUTE_COUNT)
 		{
-			AttributeNames.Add(uint8(C));
+			const FString Name = GetHairAttributeText(Attribute);
+			FHairDebugNameInfo& NameInfo = AttributeNameInfos.AddDefaulted_GetRef();
+			NameInfo.PrimitiveID = ~0;
+			NameInfo.Length = Name.Len();
+			NameInfo.Offset = AttributeNames.Num();	
+			for (TCHAR C : Name)
+			{
+				AttributeNames.Add(uint8(C));
+			}
 		}
 	}
 
diff --git a/Engine/Plugins/Runtime/HairStrands/Source/HairStrandsCore/Private/HairAttributes.cpp b/Engine/Plugins/Runtime/HairStrands/Source/HairStrandsCore/Private/HairAttributes.cpp
--- a/Engine/Plugins/Runtime/HairStrands/Source/HairStrandsCore/Private/HairAttributes.cpp
+++ b/Engine/Plugins/Runtime/HairStrands/Source/HairStrandsCore/Private/HairAttributes.cpp
@@ -1,6 +1,7 @@
 // Copyright Epic Games, Inc. All Rights Reserved.
 
 #include "HairAttributes.h"
+#include "HairStrandsDefinitions.h"
 
 namespace HairAttribute
 {
@@ -62,4 +63,18 @@ const TCHAR* GetHairAttributeText(EHairAttribute In)
 	case EHairAttribute::AO:						return TEXT("AO");
 	}
 	return TEXT("UNKNOWN");
+}
+
+uint32 GetHairAttributeIndex(EHairAttribute In)
+{
+	switch (In)
+	{
+	case EHairAttribute::RootUV:					return HAIR_ATTRIBUTE_ROOTUV;
+	case EHairAttribute::ClumpID:					return HAIR_ATTRIBUTE_CLUMPID;
+	case EHairAttribute::PrecomputedGuideWeights:	return HAIR_ATTRIBUTE_COUNT; // We don't store this information
+	case EHairAttribute::Color:						return HAIR_ATTRIBUTE_BASECOLOR;
+	case EHairAttribute::Roughness:					return HAIR_ATTRIBUTE_ROUGHNESS;
+	case EHairAttribute::AO:						return HAIR_ATTRIBUTE_AO;
+	}
+	return HAIR_ATTRIBUTE_COUNT;
 }
\ No newline at end of file
diff --git a/Engine/Plugins/Runtime/HairStrands/Source/HairStrandsCore/Private/HairStrandsRendering.cpp b/Engine/Plugins/Runtime/HairStrands/Source/HairStrandsCore/Private/HairStrandsRendering.cpp
--- a/Engine/Plugins/Runtime/HairStrands/Source/HairStrandsCore/Private/HairStrandsRendering.cpp
+++ b/Engine/Plugins/Runtime/HairStrands/Source/HairStrandsCore/Private/HairStrandsRendering.cpp
@@ -1708,6 +1708,12 @@ FHairGroupPublicData::FVertexFactoryInput ComputeHairStrandsVertexInputData(cons
 
 void CreateHairStrandsDebugAttributeBuffer(FRDGBuilder& GraphBuilder, FRDGExternalBuffer* DebugAttributeBuffer, uint32 VertexCount, const FName& OwnerName);
 
+EGroomCacheType GetHairInstanceCacheType(const FHairGroupInstance* Instance)
+{
+	const bool bHasValidGroomCacheBuffers = Instance->Debug.GroomCacheBuffers.IsValid() && Instance->Debug.GroomCacheBuffers->GetInterpolatedFrameBuffer().GroupsData.IsValidIndex(Instance->Debug.GroupIndex);
+	return bHasValidGroomCacheBuffers ? Instance->Debug.GroomCacheType : EGroomCacheType::None;
+}
+
 void ComputeHairStrandsInterpolation(
 	FRDGBuilder& GraphBuilder,
 	FGlobalShaderMap* ShaderMap,
@@ -1739,7 +1745,7 @@ void ComputeHairStrandsInterpolation(
 		RDG_EVENT_SCOPE(GraphBuilder, "HairInterpolation(Strands)");
 		RDG_GPU_STAT_SCOPE(GraphBuilder, HairStrandsInterpolation);
 
-		const bool bHasValidGroomCacheBuffers = Instance->Debug.GroomCacheBuffers.IsValid() && Instance->Debug.GroomCacheBuffers->GetInterpolatedFrameBuffer().GroupsData.IsValidIndex(Instance->Debug.GroupIndex);
+		const EGroomCacheType ActiveGroomCacheType = GetHairInstanceCacheType(Instance);
 
 		// Debug mode:
 		// * None	: Display hair normally
@@ -1751,7 +1757,7 @@ void ComputeHairStrandsInterpolation(
 			Instance->HairGroupPublicData->SetCullingResultAvailable(false);
 
 			// If groom guide cache, is enabled, used the deform position for debug visualization
-			if (bHasValidGroomCacheBuffers && Instance->Debug.GroomCacheType == EGroomCacheType::Guides)
+			if (ActiveGroomCacheType == EGroomCacheType::Guides)
 			{
 				FScopeLock Lock(Instance->Debug.GroomCacheBuffers->GetCriticalSection());
 				const FGroomCacheGroupData& GroomCacheGroupData = Instance->Debug.GroomCacheBuffers->GetInterpolatedFrameBuffer().GroupsData[Instance->Debug.GroupIndex];
@@ -1827,7 +1833,7 @@ void ComputeHairStrandsInterpolation(
 				Strands_PositionOffsetSRV = RegisterAsSRV(GraphBuilder, Instance->Strands.DeformedResource->GetPositionOffsetBuffer(FHairStrandsDeformedResource::EFrameType::Current));
 				Strands_TangentSRV = Strands_DeformedTangent.SRV;
 
-				if (bHasValidGroomCacheBuffers && Instance->Debug.GroomCacheType == EGroomCacheType::Guides)
+				if (ActiveGroomCacheType == EGroomCacheType::Guides)
 				{
 					FScopeLock Lock(Instance->Debug.GroomCacheBuffers->GetCriticalSection());
 					const FGroomCacheGroupData& GroomCacheGroupData = Instance->Debug.GroomCacheBuffers->GetInterpolatedFrameBuffer().GroupsData[Instance->Debug.GroupIndex];
@@ -1896,7 +1902,7 @@ void ComputeHairStrandsInterpolation(
 				}
 
 				// 2.1 Compute deformation position based on simulation/skinning/RBF
-				if (Instance->Debug.GroomCacheType != EGroomCacheType::Strands)
+				if (ActiveGroomCacheType != EGroomCacheType::Strands)
 				{
 					// 2.1.1 Current Position
 					AddHairStrandsInterpolationPass(
@@ -1983,7 +1989,7 @@ void ComputeHairStrandsInterpolation(
 							FHairStrandsDeformedRootResource::FLOD::Previous);
 					}
 				}
-				else if (bHasValidGroomCacheBuffers)
+				else if (ActiveGroomCacheType == EGroomCacheType::Strands)
 				{
 					FScopeLock Lock(Instance->Debug.GroomCacheBuffers->GetCriticalSection());
 					const FGroomCacheGroupData& GroomCacheGroupData = Instance->Debug.GroomCacheBuffers->GetInterpolatedFrameBuffer().GroupsData[Instance->Debug.GroupIndex];
diff --git a/Engine/Plugins/Runtime/HairStrands/Source/HairStrandsCore/Public/HairStrandsDatas.h b/Engine/Plugins/Runtime/HairStrands/Source/HairStrandsCore/Public/HairStrandsDatas.h
--- a/Engine/Plugins/Runtime/HairStrands/Source/HairStrandsCore/Public/HairStrandsDatas.h
+++ b/Engine/Plugins/Runtime/HairStrands/Source/HairStrandsCore/Public/HairStrandsDatas.h
@@ -9,7 +9,6 @@
 #include "RenderGraphResources.h"
 #include "Serialization/BulkData.h"
 #include "HairStrandsDefinitions.h"
-#include "HairAttributes.h"
 
 DECLARE_LOG_CATEGORY_EXTERN(LogHairStrands, Log, All);
 
@@ -465,7 +464,6 @@ struct HAIRSTRANDSCORE_API FHairStrandsBulkData
 	float  GetMaxRadius() const { return MaxRadius; }
 	FVector GetPositionOffset() const { return BoundingBox.GetCenter(); }
 	const FBox& GetBounds() const { return BoundingBox; }
-	bool HasAttributes(EHairAttribute In) const { return AttributeOffsets[uint32(In)] != 0xFFFFFFFF; }
 
 	uint32 CurveCount = 0;
 	uint32 PointCount = 0;
diff --git a/Engine/Shaders/Private/HairStrands/HairStrandsDebug.usf b/Engine/Shaders/Private/HairStrands/HairStrandsDebug.usf
--- a/Engine/Shaders/Private/HairStrands/HairStrandsDebug.usf
+++ b/Engine/Shaders/Private/HairStrands/HairStrandsDebug.usf
@@ -1194,6 +1194,8 @@ struct FInstanceInfo
 	bool Scatter;
 	bool RTGeometry;
 	bool StableRaster;
+	bool CacheRen;
+	bool CacheSim;
 
 	// Binding
 	uint UniqueTriangle;
@@ -1260,6 +1262,8 @@ FInstanceInfo ReadInstanceInfo(uint Index)
 	Out.RTGeometry				= Data2.x & 0x8;
 	Out.StableRaster			= Data2.x & 0x10;
 	Out.Voxel					= Data2.x & 0x20;
+	Out.CacheSim				= Data2.x & 0x40;
+	Out.CacheRen				= Data2.x & 0x80;
 	Out.ScreenSize				= f16tof32((Data2.x>>16) & 0xFFFF);
 	Out.ActivePointCount		= Data2.y;
 	Out.ActiveCurveCount		= Data2.z;
@@ -1406,7 +1410,7 @@ void MainCS(uint GroupIndex : SV_GroupIndex, uint3 DispatchThreadId : SV_Dispatc
 	Newline(Context);
 
 	// Draw instance details
-	Print(Context, TEXT("   Group  LOD     Geometry Binding Sim RBF Clip Screen Strands            Points             Bind.Tri/Sect Vis. Shad. Name"));
+	Print(Context, TEXT("   Group  LOD     Geometry Binding Cache Sim RBF Clip Screen Strands            Points             Bind.Tri/Sect Vis. Shad. Name"));
 	Newline(Context);
 
 	for (uint InstanceIndex = 0; InstanceIndex < InstanceCount; ++InstanceIndex)
@@ -1440,6 +1444,11 @@ void MainCS(uint GroupIndex : SV_GroupIndex, uint3 DispatchThreadId : SV_Dispatc
 		else if (Info.BindingType == 1)	{ Print(Context, TEXT("Rigid   "), FontTurquoise); }
 		else if (Info.BindingType == 2)	{ Print(Context, TEXT("None    "), FontTurquoise); }
 
+		// Cache
+		if (Info.CacheRen)      { Print(Context, TEXT("Ren  "), FontOrange); }
+		else if (Info.CacheSim) { Print(Context, TEXT("Sim  "), FontOrange); }
+		else                    { Print(Context, TEXT("None "), FontSilver); }
+		Print(Context, TEXT(" "));
 		// Sim
 		PrintBool(Context, Info.bHasSim);
 
