Fix CIS issue.

diff --git a/Engine/Plugins/Runtime/HairStrands/Source/HairStrandsCore/Private/GroomDebug.cpp b/Engine/Plugins/Runtime/HairStrands/Source/HairStrandsCore/Private/GroomDebug.cpp
--- a/Engine/Plugins/Runtime/HairStrands/Source/HairStrandsCore/Private/GroomDebug.cpp
+++ b/Engine/Plugins/Runtime/HairStrands/Source/HairStrandsCore/Private/GroomDebug.cpp
@@ -938,6 +938,7 @@ static void AddHairDebugPrintInstancePass(
 		case EHairGeometryType::Strands:
 			if (Instance->Strands.IsValid())
 			{
+				check(Instance->Strands.Data);
 				Data0.Z = Instance->Strands.Data->GetNumCurves(); // Change this later on for having dynamic value
 				Data0.W = Instance->Strands.Data->GetNumPoints(); // Change this later on for having dynamic value
 				const int32 MeshLODIndex = Instance->HairGroupPublicData->MeshLODIndex;
@@ -964,16 +965,13 @@ static void AddHairDebugPrintInstancePass(
 					Data2.Y = Instance->HairGroupPublicData->GetActiveStrandsPointCount();
 					Data2.Z = Instance->HairGroupPublicData->GetActiveStrandsCurveCount();
 
-					if (Instance->Strands.Data)
+					for (uint32 AttributeIt=0; AttributeIt<uint32(EHairAttribute::Count); ++AttributeIt)
 					{
-						for (uint32 AttributeIt=0; AttributeIt<uint32(EHairAttribute::Count); ++AttributeIt)
+						const EHairAttribute Attribute = (EHairAttribute)AttributeIt;
+						const uint32 Index = GetHairAttributeIndex(Attribute);
+						if (Index < HAIR_ATTRIBUTE_COUNT)
 						{
-							const EHairAttribute Attribute = (EHairAttribute)AttributeIt;
-							const uint32 Index = GetHairAttributeIndex(Attribute);
-							if (Index < HAIR_ATTRIBUTE_COUNT)
-							{
-								Data2.W |= Instance->Strands.Data->AttributeOffsets[Index] != 0xFFFFFFFF ? (1u<<Index) : 0u;
-							}
+							Data2.W |= Instance->Strands.Data->AttributeOffsets[Index] != 0xFFFFFFFF ? (1u<<Index) : 0u;
 						}
 					}
 				}
