Update groom deformer API so that read/writen position are expressed in groom local space, instead of translated local space.

diff --git a/Engine/Plugins/Runtime/HairStrands/Shaders/Private/DeformerDataInterfaceGroom.ush b/Engine/Plugins/Runtime/HairStrands/Shaders/Private/DeformerDataInterfaceGroom.ush
--- a/Engine/Plugins/Runtime/HairStrands/Shaders/Private/DeformerDataInterfaceGroom.ush
+++ b/Engine/Plugins/Runtime/HairStrands/Shaders/Private/DeformerDataInterfaceGroom.ush
@@ -14,6 +14,7 @@ float				{DataInterfaceName}_VF_RootScale;
 float				{DataInterfaceName}_VF_TipScale;
 uint				{DataInterfaceName}_VF_GroupIndex;
 uint4				{DataInterfaceName}_VF_AttributeOffsets[HAIR_ATTRIBUTE_OFFSET_COUNT];
+Buffer<float4>		{DataInterfaceName}_PositionOffsetBuffer;
 Buffer<uint4>		{DataInterfaceName}_PositionBuffer;
 ByteAddressBuffer	{DataInterfaceName}_AttributeBuffer;
 Buffer<uint>		{DataInterfaceName}_VertexToCurveBuffer;
@@ -33,7 +34,7 @@ FHairControlPoint ReadControlPointData(uint VertexIndex)
 	// Hair deformer work in local space so no need for position offset
 	return UnpackHairControlPoint(
 		{DataInterfaceName}_PositionBuffer[VertexIndex], 
-		float3(0, 0, 0) /*InVF_PositionOffset*/,
+		{DataInterfaceName}_PositionOffsetBuffer[0],
 		{DataInterfaceName}_VF_HairRadius, 
 		1.0f /*RootScale*/, 
 		1.0f /*TipScale*/);
diff --git a/Engine/Plugins/Runtime/HairStrands/Shaders/Private/DeformerDataInterfaceGroomWrite.ush b/Engine/Plugins/Runtime/HairStrands/Shaders/Private/DeformerDataInterfaceGroomWrite.ush
--- a/Engine/Plugins/Runtime/HairStrands/Shaders/Private/DeformerDataInterfaceGroomWrite.ush
+++ b/Engine/Plugins/Runtime/HairStrands/Shaders/Private/DeformerDataInterfaceGroomWrite.ush
@@ -6,6 +6,7 @@
 uint {DataInterfaceName}_NumControlPoints;
 uint {DataInterfaceName}_NumCurves;
 uint {DataInterfaceName}_OutputStreamStart;
+Buffer<float4>	{DataInterfaceName}_PositionOffsetBufferSRV;
 Buffer<uint4>   {DataInterfaceName}_PositionBufferSRV;
 RWBuffer<uint4> {DataInterfaceName}_PositionBufferUAV;
 
@@ -25,7 +26,7 @@ void WritePosition_{DataInterfaceName}(uint ControlPointIndex, float3 Position)
 	if (WriteIndex < {DataInterfaceName}_NumControlPoints)
 	{
 		const uint4 PackedCP = {DataInterfaceName}_PositionBufferSRV[WriteIndex];
-		const uint4 NewPackedCP = PackHairControlPointPosition(PackedCP, Position, 0 /*InVF_PositionOffset*/);
+		const uint4 NewPackedCP = PackHairControlPointPosition(PackedCP, Position, {DataInterfaceName}_PositionOffsetBufferSRV[0]);
 		{DataInterfaceName}_PositionBufferUAV[WriteIndex] = NewPackedCP;
 	}
 }
\ No newline at end of file
diff --git a/Engine/Plugins/Runtime/HairStrands/Source/HairStrandsCore/Private/Deformer/DeformerDataInterfaceGroom.cpp b/Engine/Plugins/Runtime/HairStrands/Source/HairStrandsCore/Private/Deformer/DeformerDataInterfaceGroom.cpp
--- a/Engine/Plugins/Runtime/HairStrands/Source/HairStrandsCore/Private/Deformer/DeformerDataInterfaceGroom.cpp
+++ b/Engine/Plugins/Runtime/HairStrands/Source/HairStrandsCore/Private/Deformer/DeformerDataInterfaceGroom.cpp
@@ -99,6 +99,7 @@ BEGIN_SHADER_PARAMETER_STRUCT(FGroomDataInterfaceParameters, )
 	SHADER_PARAMETER(float, VF_TipScale)
 	SHADER_PARAMETER(uint32, VF_GroupIndex)
 	SHADER_PARAMETER_ARRAY(FUintVector4, VF_AttributeOffsets, [HAIR_ATTRIBUTE_OFFSET_COUNT])
+	SHADER_PARAMETER_RDG_BUFFER_SRV(Buffer<float4>,    PositionOffsetBuffer)
 	SHADER_PARAMETER_RDG_BUFFER_SRV(Buffer<uint4>,     PositionBuffer)
 	SHADER_PARAMETER_RDG_BUFFER_SRV(ByteAddressBuffer, AttributeBuffer)
 	SHADER_PARAMETER_RDG_BUFFER_SRV(Buffer<uint>,      VertexToCurveBuffer)
@@ -180,6 +181,7 @@ void FOptimusGroomDataProviderProxy::AllocateResources(FRDGBuilder& GraphBuilder
 		if (FHairGroupInstance* Instance = GroomComponent->GetGroupInstance(Index))
 		{
 			FResources& R = Resources.AddDefaulted_GetRef();
+			R.PositionOffsetSRV = Register(GraphBuilder, Instance->Strands.RestResource->PositionOffsetBuffer, ERDGImportedBufferFlags::CreateSRV).SRV;
 			R.PositionSRV		= Register(GraphBuilder, Instance->Strands.RestResource->PositionBuffer, ERDGImportedBufferFlags::CreateSRV).SRV;
 			R.AttributeSRV		= Register(GraphBuilder, Instance->Strands.RestResource->AttributeBuffer, ERDGImportedBufferFlags::CreateSRV).SRV;
 			R.VertexToCurveSRV	= Register(GraphBuilder, Instance->Strands.RestResource->PointToCurveBuffer, ERDGImportedBufferFlags::CreateSRV).SRV;
@@ -219,6 +221,7 @@ void FOptimusGroomDataProviderProxy::GatherDispatchData(FDispatchData const& InD
 
 			if (bIsSRVValid)
 			{
+				Parameters.PositionOffsetBuffer = Resources[InvocationIndex].PositionOffsetSRV;
 				Parameters.PositionBuffer		= Resources[InvocationIndex].PositionSRV;
 				Parameters.AttributeBuffer		= Resources[InvocationIndex].AttributeSRV;
 				Parameters.VertexToCurveBuffer	= Resources[InvocationIndex].VertexToCurveSRV;
@@ -226,6 +229,7 @@ void FOptimusGroomDataProviderProxy::GatherDispatchData(FDispatchData const& InD
 			}
 			else
 			{
+				Parameters.PositionOffsetBuffer = Resources[InvocationIndex].FallbackSRV;
 				Parameters.PositionBuffer		= Resources[InvocationIndex].FallbackSRV;
 				Parameters.AttributeBuffer		= Resources[InvocationIndex].FallbackSRV;
 				Parameters.VertexToCurveBuffer	= Resources[InvocationIndex].FallbackSRV;
diff --git a/Engine/Plugins/Runtime/HairStrands/Source/HairStrandsCore/Private/Deformer/DeformerDataInterfaceGroom.h b/Engine/Plugins/Runtime/HairStrands/Source/HairStrandsCore/Private/Deformer/DeformerDataInterfaceGroom.h
--- a/Engine/Plugins/Runtime/HairStrands/Source/HairStrandsCore/Private/Deformer/DeformerDataInterfaceGroom.h
+++ b/Engine/Plugins/Runtime/HairStrands/Source/HairStrandsCore/Private/Deformer/DeformerDataInterfaceGroom.h
@@ -69,6 +69,7 @@ private:
 	UGroomComponent* GroomComponent = nullptr; // Should it be HairInstance instead?
 	struct FResources
 	{
+		FRDGBufferSRVRef PositionOffsetSRV = nullptr;
 		FRDGBufferSRVRef PositionSRV = nullptr;
 		FRDGBufferSRVRef AttributeSRV = nullptr;
 		FRDGBufferSRVRef VertexToCurveSRV = nullptr;
diff --git a/Engine/Plugins/Runtime/HairStrands/Source/HairStrandsCore/Private/Deformer/DeformerDataInterfaceGroomWrite.cpp b/Engine/Plugins/Runtime/HairStrands/Source/HairStrandsCore/Private/Deformer/DeformerDataInterfaceGroomWrite.cpp
--- a/Engine/Plugins/Runtime/HairStrands/Source/HairStrandsCore/Private/Deformer/DeformerDataInterfaceGroomWrite.cpp
+++ b/Engine/Plugins/Runtime/HairStrands/Source/HairStrandsCore/Private/Deformer/DeformerDataInterfaceGroomWrite.cpp
@@ -56,6 +56,7 @@ BEGIN_SHADER_PARAMETER_STRUCT(FGroomWriteDataInterfaceParameters, )
 	SHADER_PARAMETER(uint32, NumControlPoints)
 	SHADER_PARAMETER(uint32, NumCurves)
 	SHADER_PARAMETER(uint32, OutputStreamStart)
+	SHADER_PARAMETER_RDG_BUFFER_SRV(Buffer, PositionOffsetBufferSRV)
 	SHADER_PARAMETER_RDG_BUFFER_SRV(Buffer<uint4>, PositionBufferSRV)
 	SHADER_PARAMETER_RDG_BUFFER_UAV(RWBuffer<uint4>, PositionBufferUAV)
 END_SHADER_PARAMETER_STRUCT()
@@ -140,6 +141,7 @@ void FOptimusGroomWriteDataProviderProxy::AllocateResources(FRDGBuilder& GraphBu
 			FResources& R = Resources.AddDefaulted_GetRef();
 			if (OutputMask & 1)
 			{
+				R.PositionOffsetBufferSRV = Register(GraphBuilder, GroomInstance->Strands.RestResource->PositionOffsetBuffer, ERDGImportedBufferFlags::CreateSRV).SRV;
 				R.PositionBufferSRV = Register(GraphBuilder, GroomInstance->Strands.RestResource->PositionBuffer, ERDGImportedBufferFlags::CreateSRV).SRV;
 				R.PositionBufferUAV = Register(GraphBuilder, GroomInstance->Strands.DeformedResource->GetDeformerBuffer(GraphBuilder), ERDGImportedBufferFlags::CreateUAV).UAV;
 				
@@ -149,6 +151,7 @@ void FOptimusGroomWriteDataProviderProxy::AllocateResources(FRDGBuilder& GraphBu
 			}
 			else
 			{
+				R.PositionOffsetBufferSRV = GraphBuilder.CreateSRV(GraphBuilder.RegisterExternalBuffer(GWhiteVertexBufferWithRDG->Buffer), PF_A32B32G32R32F);
 				R.PositionBufferSRV = GraphBuilder.CreateSRV(GraphBuilder.RegisterExternalBuffer(GWhiteVertexBufferWithRDG->Buffer), PF_R16G16B16A16_UINT);
 				R.PositionBufferUAV = GraphBuilder.CreateUAV(GraphBuilder.RegisterExternalBuffer(GWhiteVertexBufferWithRDG->Buffer), PF_R16G16B16A16_UINT);
 			}
@@ -165,7 +168,7 @@ void FOptimusGroomWriteDataProviderProxy::GatherDispatchData(FDispatchData const
 		{
 			FResources& Resource = Resources[InvocationIndex];
 
-			const bool bValid = Resource.PositionBufferUAV != nullptr && Resource.PositionBufferSRV != nullptr;
+			const bool bValid = Resource.PositionBufferUAV != nullptr && Resource.PositionBufferSRV != nullptr && Resource.PositionOffsetBufferSRV != nullptr;
 			const int32 NumCurves = bValid ? GroomInstance->Strands.Data->CurveCount : 0;
 			const int32 NumControlPoints = bValid ? GroomInstance->Strands.Data->PointCount : 0;
 
@@ -173,6 +176,7 @@ void FOptimusGroomWriteDataProviderProxy::GatherDispatchData(FDispatchData const
 			Parameters.NumControlPoints = NumControlPoints;
 			Parameters.NumCurves = NumCurves;
 			Parameters.OutputStreamStart = 0;
+			Parameters.PositionOffsetBufferSRV = bValid ? Resource.PositionOffsetBufferSRV : Resource.PositionBufferSRV_fallback;
 			Parameters.PositionBufferSRV = bValid ? Resource.PositionBufferSRV : Resource.PositionBufferSRV_fallback;
 			Parameters.PositionBufferUAV = bValid ? Resource.PositionBufferUAV : Resource.PositionBufferUAV_fallback;
 		}
diff --git a/Engine/Plugins/Runtime/HairStrands/Source/HairStrandsCore/Private/Deformer/DeformerDataInterfaceGroomWrite.h b/Engine/Plugins/Runtime/HairStrands/Source/HairStrandsCore/Private/Deformer/DeformerDataInterfaceGroomWrite.h
--- a/Engine/Plugins/Runtime/HairStrands/Source/HairStrandsCore/Private/Deformer/DeformerDataInterfaceGroomWrite.h
+++ b/Engine/Plugins/Runtime/HairStrands/Source/HairStrandsCore/Private/Deformer/DeformerDataInterfaceGroomWrite.h
@@ -78,6 +78,7 @@ private:
 
 	struct FResources
 	{
+		FRDGBufferSRV* PositionOffsetBufferSRV = nullptr;
 		FRDGBufferSRV* PositionBufferSRV = nullptr;
 		FRDGBufferUAV* PositionBufferUAV = nullptr;
 		FRDGBufferSRV* PositionBufferSRV_fallback = nullptr;
