Fix incorrect simulation curve check, causing groom without guides, to prevent automatically guide generation.

diff --git a/Engine/Plugins/Runtime/HairStrands/Source/HairStrandsCore/Private/GroomBuilder.cpp b/Engine/Plugins/Runtime/HairStrands/Source/HairStrandsCore/Private/GroomBuilder.cpp
--- a/Engine/Plugins/Runtime/HairStrands/Source/HairStrandsCore/Private/GroomBuilder.cpp
+++ b/Engine/Plugins/Runtime/HairStrands/Source/HairStrandsCore/Private/GroomBuilder.cpp
@@ -1927,7 +1927,6 @@ void FGroomBuilder::BuildData(
 		}
 
 		// Simulation data
-		if (InHairDescriptionGroup.Guides.IsValid())
 		{
 			OutSim = InHairDescriptionGroup.Guides;
 			if (InHairDescriptionGroup.Info.NumGuides > 0 && !InSettings.InterpolationSettings.bOverrideGuides)
