Fix crash on citysample with cook data.

Disable curve loading when CLOD is disable .

diff --git a/Engine/Plugins/Runtime/HairStrands/Source/HairStrandsCore/Private/GroomManager.cpp b/Engine/Plugins/Runtime/HairStrands/Source/HairStrandsCore/Private/GroomManager.cpp
--- a/Engine/Plugins/Runtime/HairStrands/Source/HairStrandsCore/Private/GroomManager.cpp
+++ b/Engine/Plugins/Runtime/HairStrands/Source/HairStrandsCore/Private/GroomManager.cpp
@@ -833,7 +833,7 @@ static void RunHairLODSelection(
 			Instance->HairGroupPublicData->ContinuousLODScreenPos = MaxContinuousLODScreenPos;
 			Instance->HairGroupPublicData->ContinuousLODBounds = SphereBound;
 
-			if (GHairStrands_ContinousLOD > 0)
+			if (IsHairStrandsContinousLODEnabled())
 			{
 				// These values are set later, once the resources are loaded
 				Instance->HairGroupPublicData->ContinuousLODCurveCount = 0;
@@ -1011,7 +1011,7 @@ static void RunHairLODSelection(
 			}
 
 			// Adapt the number curve/point based on Bound/CPU screen-size
-			if (GeometryType == EHairGeometryType::Strands && GHairStrands_ContinousLOD)
+			if (GeometryType == EHairGeometryType::Strands && IsHairStrandsContinousLODEnabled())
 			{
 				check(bIsLODDataReady);
 				check(Instance->Strands.RestResource);
diff --git a/Engine/Plugins/Runtime/HairStrands/Source/HairStrandsCore/Private/GroomResources.cpp b/Engine/Plugins/Runtime/HairStrands/Source/HairStrandsCore/Private/GroomResources.cpp
--- a/Engine/Plugins/Runtime/HairStrands/Source/HairStrandsCore/Private/GroomResources.cpp
+++ b/Engine/Plugins/Runtime/HairStrands/Source/HairStrandsCore/Private/GroomResources.cpp
@@ -884,6 +884,7 @@ bool FHairStrandsRestResource::InternalIsDataLoaded()
 	return BulkDataRequest.IsCompleted();
 }
 
+bool IsHairStrandsContinousLODEnabled();
 void FHairStrandsRestResource::InternalAllocate(FRDGBuilder& GraphBuilder)
 {
 	BulkDataRequest = FBulkDataBatchRequest();
@@ -911,6 +912,7 @@ void FHairStrandsRestResource::InternalAllocate(FRDGBuilder& GraphBuilder)
 	GraphBuilder.UseExternalAccessMode(Register(GraphBuilder, PositionOffsetBuffer, ERDGImportedBufferFlags::CreateSRV).Buffer, ERHIAccess::SRVMask);
 
 	// Copy curve data for CPU lookup
+	if (IsHairStrandsContinousLODEnabled())
 	{
 		CurveData.SetNum(CurveCount);
 		uint32* Data = (uint32*)CurveData.GetData();
