Change how simulation cache are handle on groom to fix dynamic assignement.

* A groom asset now needs to opt-in to support groom asset. This opt-in ensure that interoplation data are property built for the asset. This opt-in is automatically enabled when assigning a groom cache in the property grid.
* A groom asset with simulation cache support enable, can now get a groom cache assign through BP at runtime.
* This avoids to spawn actual simuation when running with a simulation cache.

diff --git a/Engine/Plugins/Runtime/HairStrands/Source/HairStrandsCore/Private/GroomAsset.cpp b/Engine/Plugins/Runtime/HairStrands/Source/HairStrandsCore/Private/GroomAsset.cpp
--- a/Engine/Plugins/Runtime/HairStrands/Source/HairStrandsCore/Private/GroomAsset.cpp
+++ b/Engine/Plugins/Runtime/HairStrands/Source/HairStrandsCore/Private/GroomAsset.cpp
@@ -1435,8 +1435,9 @@ static bool IsStrandsInterpolationAttributes(const FName PropertyName)
 		|| PropertyName == GET_MEMBER_NAME_CHECKED(FHairDeformationSettings, NumPoints)
 		|| PropertyName == GET_MEMBER_NAME_CHECKED(FHairDeformationSettings, bEnableRigging)
 
-		// Add dependency on simulation and per LOD-simulation/global-interoplation to strip-out interoplation data if there are not needed
+		// Add dependency on simulation and per LOD-simulation/global-interpolation to strip-out interpolation data if there are not needed
 		|| PropertyName == GET_MEMBER_NAME_CHECKED(UGroomAsset, EnableGlobalInterpolation)
+		|| PropertyName == GET_MEMBER_NAME_CHECKED(UGroomAsset, EnableSimulationCache)
 		|| PropertyName == GET_MEMBER_NAME_CHECKED(FHairSolverSettings, EnableSimulation)
 		|| PropertyName == GET_MEMBER_NAME_CHECKED(FHairLODSettings, Simulation)
 		|| PropertyName == GET_MEMBER_NAME_CHECKED(FHairLODSettings, GlobalInterpolation)
@@ -2058,6 +2059,39 @@ namespace GroomDerivedDataCacheUtils
 		static FString MeshPrefixString(TEXT("MESHES_V") + FHairMeshesBuilder::GetVersion() + TEXT("_"));
 		return MeshPrefixString + Hash.ToString();
 	}
+
+	FString BuildStrandsDerivedDataKeySuffix(uint32 GroupIndex, const FHairGroupsInterpolation& InterpolationSettings, const FHairGroupsLOD& LODSettings, bool bNeedInterpolationData, const FHairDescriptionBulkData* HairBulkDescription)
+	{
+		// Serialize the build settings into a temporary array
+		// The archive is flagged as persistent so that machines of different endianness produce identical binary results.
+		TArray<uint8> TempBytes;
+		TempBytes.Reserve(64);
+		FMemoryWriter Ar(TempBytes, /*bIsPersistent=*/ true);
+
+		GroomDerivedDataCacheUtils::SerializeHairInterpolationSettingsForDDC(Ar, GroupIndex, const_cast<FHairGroupsInterpolation&>(InterpolationSettings), const_cast<FHairGroupsLOD&>(LODSettings), bNeedInterpolationData);
+
+		FString KeySuffix;
+		if (HairBulkDescription)
+		{
+			// Reserve twice the size of TempBytes because of ByteToHex below + 3 for "ID" and \0
+			KeySuffix.Reserve(HairBulkDescription->GetIdString().Len() + TempBytes.Num() * 2 + 3);
+			KeySuffix += TEXT("ID");
+			KeySuffix += HairBulkDescription->GetIdString();
+		}
+		else
+		{
+			KeySuffix.Reserve(TempBytes.Num() * 2 + 1);
+		}
+
+		// Now convert the raw bytes to a string
+		const uint8* SettingsAsBytes = TempBytes.GetData();
+		for (int32 ByteIndex = 0; ByteIndex < TempBytes.Num(); ++ByteIndex)
+		{
+			ByteToHex(SettingsAsBytes[ByteIndex], KeySuffix);
+		}
+
+		return KeySuffix;
+	}
 }
 
 bool UGroomAsset::IsFullyCached()
@@ -2108,42 +2142,6 @@ bool UGroomAsset::IsFullyCached()
 	return GetDerivedDataCacheRef().AllCachedDataProbablyExists(CacheKeys);
 }
 
-FString UGroomAsset::BuildDerivedDataKeySuffix(uint32 GroupIndex, const FHairGroupsInterpolation& InterpolationSettings, const FHairGroupsLOD& LODSettings) const
-{
-	// Serialize the build settings into a temporary array
-	// The archive is flagged as persistent so that machines of different endianness produce identical binary results.
-	TArray<uint8> TempBytes;
-	TempBytes.Reserve(64);
-	FMemoryWriter Ar(TempBytes, /*bIsPersistent=*/ true);
-	
-	// If simulation or global interpolation is enabled, then interoplation data are required. Otherwise they can be skipped
-	const bool bNeedInterpolationData = NeedsInterpolationData(GroupIndex);
-
-	GroomDerivedDataCacheUtils::SerializeHairInterpolationSettingsForDDC(Ar, GroupIndex, const_cast<FHairGroupsInterpolation&>(InterpolationSettings), const_cast<FHairGroupsLOD&>(LODSettings), bNeedInterpolationData);
-
-	FString KeySuffix;
-	if (HairDescriptionBulkData[HairDescriptionType])
-	{
-		// Reserve twice the size of TempBytes because of ByteToHex below + 3 for "ID" and \0
-		KeySuffix.Reserve(HairDescriptionBulkData[HairDescriptionType]->GetIdString().Len() + TempBytes.Num() * 2 + 3);
-		KeySuffix += TEXT("ID");
-		KeySuffix += HairDescriptionBulkData[HairDescriptionType]->GetIdString();
-	}
-	else
-	{
-		KeySuffix.Reserve(TempBytes.Num() * 2 + 1);
-	}
-
-	// Now convert the raw bytes to a string
-	const uint8* SettingsAsBytes = TempBytes.GetData();
-	for (int32 ByteIndex = 0; ByteIndex < TempBytes.Num(); ++ByteIndex)
-	{
-		ByteToHex(SettingsAsBytes[ByteIndex], KeySuffix);
-	}
-
-	return KeySuffix;
-}
-
 FString UGroomAsset::GetDerivedDataKey()
 {
 	FString Key;
@@ -2165,25 +2163,26 @@ FString UGroomAsset::GetDerivedDataKeyForCards(uint32 GroupIndex, const FString&
 	return DerivedDataKey;
 }
 
-FString UGroomAsset::GetDerivedDataKeyForStrands(uint32 GroupIndex)
+FString UGroomAsset::GetDerivedDataKeyForMeshes(uint32 GroupIndex)
 {
-	const FHairGroupsInterpolation& InterpolationSettings = HairGroupsInterpolation[GroupIndex];
-	const FHairGroupsLOD& LODSettings = HairGroupsLOD[GroupIndex];
-
-	const FString KeySuffix = BuildDerivedDataKeySuffix(GroupIndex, InterpolationSettings, LODSettings);
+	const FString KeySuffix = GroomDerivedDataCacheUtils::BuildMeshesDerivedDataKeySuffix(GroupIndex, HairGroupsLOD[GroupIndex].LODs, HairGroupsMeshes);
 	const FString DerivedDataKey = GroomDerivedDataCacheUtils::BuildGroomDerivedDataKey(KeySuffix);
-
 	return DerivedDataKey;
 }
 
-FString UGroomAsset::GetDerivedDataKeyForMeshes(uint32 GroupIndex)
+FString UGroomAsset::GetDerivedDataKeyForStrands(uint32 GroupIndex)
 {
-	const FString KeySuffix = GroomDerivedDataCacheUtils::BuildMeshesDerivedDataKeySuffix(GroupIndex, HairGroupsLOD[GroupIndex].LODs, HairGroupsMeshes);
+	const FHairGroupsInterpolation& InterpolationSettings = HairGroupsInterpolation[GroupIndex];
+	const FHairGroupsLOD& LODSettings = HairGroupsLOD[GroupIndex];
+
+	// If simulation or global interpolation is enabled, then interpolation data are required. Otherwise they can be skipped
+	const bool bNeedInterpolationData = NeedsInterpolationData(GroupIndex);
+	const FString KeySuffix = GroomDerivedDataCacheUtils::BuildStrandsDerivedDataKeySuffix(GroupIndex, InterpolationSettings, LODSettings, bNeedInterpolationData, HairDescriptionBulkData[HairDescriptionType].Get());
 	const FString DerivedDataKey = GroomDerivedDataCacheUtils::BuildGroomDerivedDataKey(KeySuffix);
+
 	return DerivedDataKey;
 }
 
-
 void UGroomAsset::CommitHairDescription(FHairDescription&& InHairDescription, EHairDescriptionType Type)
 {
 	HairDescriptionType = Type;
@@ -3450,7 +3449,7 @@ bool UGroomAsset::NeedsInterpolationData(int32 GroupIndex) const
 {
 	for (int32 LODIt = 0; LODIt < HairGroupsLOD.Num(); ++LODIt)
 	{
-		if (IsSimulationEnable(GroupIndex, LODIt) || IsGlobalInterpolationEnable(GroupIndex, LODIt))
+		if (IsSimulationEnable(GroupIndex, LODIt) || IsGlobalInterpolationEnable(GroupIndex, LODIt) || EnableSimulationCache)
 		{
 			return true;
 		}
diff --git a/Engine/Plugins/Runtime/HairStrands/Source/HairStrandsCore/Private/GroomComponent.cpp b/Engine/Plugins/Runtime/HairStrands/Source/HairStrandsCore/Private/GroomComponent.cpp
--- a/Engine/Plugins/Runtime/HairStrands/Source/HairStrandsCore/Private/GroomComponent.cpp
+++ b/Engine/Plugins/Runtime/HairStrands/Source/HairStrandsCore/Private/GroomComponent.cpp
@@ -378,7 +378,7 @@ public:
 
 			{
 				// If one of the group has simulation enable, then we enable velocity rendering for meshes/cards
-				if (IsHairStrandsSimulationEnable() && HairInstance->Guides.IsValid() && (HairInstance->Guides.bIsSimulationEnable || HairInstance->Guides.bIsDeformationEnable))
+				if (IsHairStrandsSimulationEnable() && HairInstance->Guides.IsValid() && (HairInstance->Guides.bIsSimulationEnable || HairInstance->Guides.bIsDeformationEnable || HairInstance->Guides.bIsSimulationCacheEnable))
 				{
 					bAlwaysHasVelocity = true;
 				}
@@ -2187,7 +2187,8 @@ void UGroomComponent::UpdateSimulatedGroups()
 					check(Instance->HairGroupPublicData);
 					Instance->Guides.bIsSimulationEnable	 = Instance->HairGroupPublicData->IsSimulationEnable(LODIndex);
 					Instance->Guides.bHasGlobalInterpolation = Instance->HairGroupPublicData->IsGlobalInterpolationEnable(LODIndex);
-					Instance->Guides.bIsDeformationEnable = Instance->HairGroupPublicData->bIsDeformationEnable;
+					Instance->Guides.bIsDeformationEnable	 = Instance->HairGroupPublicData->bIsDeformationEnable;
+					Instance->Guides.bIsSimulationCacheEnable= Instance->Debug.GroomCacheType == EGroomCacheType::Guides;
 				}
 				++GroupIt;
 			}
@@ -2347,6 +2348,23 @@ static EGroomGeometryType GetEffectiveGeometryType(EGroomGeometryType Type, bool
 	return Type == EGroomGeometryType::Strands && (!IsHairStrandsEnabled(EHairStrandsShaderType::Strands) || bUseCards) ? EGroomGeometryType::Cards : Type;
 }
 
+static EGroomCacheType GetEffectiveGroomCacheType(const UGroomCache* InCache, const UGroomAsset* InGroom)
+{
+	if (InCache && InGroom)
+	{
+		if (InCache->GetType() == EGroomCacheType::Guides)
+		{
+			// Ensure simulation cache is enable only if the groom asset as interpolation data enabled (i.e., EnableSimulationCache=true)
+			return InGroom->EnableSimulationCache ? EGroomCacheType::Guides : EGroomCacheType::None;
+		}
+		else
+		{
+			return InCache->GetType();
+		}
+	}
+	return EGroomCacheType::None;
+}
+
 void UGroomComponent::InitResources(bool bIsBindingReloading)
 {
 	LLM_SCOPE_BYTAG(Groom);
@@ -2555,7 +2573,7 @@ void UGroomComponent::InitResources(bool bIsBindingReloading)
 		HairGroupInstance->Debug.MeshComponent = RegisteredMeshComponent;
 		HairGroupInstance->Debug.GroomComponentForDebug = this;
 		HairGroupInstance->Debug.GroomBindingType = BindingAsset ? BindingAsset->GroomBindingType : EGroomBindingMeshType::SkeletalMesh;
-		HairGroupInstance->Debug.GroomCacheType = GroomCache ? GroomCache->GetType() : EGroomCacheType::None;
+		HairGroupInstance->Debug.GroomCacheType = GetEffectiveGroomCacheType(GroomCache, GroomAsset);
 		HairGroupInstance->Debug.GroomCacheBuffers = GroomCacheBuffers;
 		HairGroupInstance->Debug.LODForcedIndex = LODForcedIndex;
 		HairGroupInstance->Debug.LODPredictedIndex = LODPredictedIndex;
@@ -2621,6 +2639,7 @@ void UGroomComponent::InitResources(bool bIsBindingReloading)
 				HairGroupInstance->HairGroupPublicData->LODSimulations.Add(LODSimulation);
 				HairGroupInstance->HairGroupPublicData->LODGlobalInterpolations.Add(LODGlobalInterpolation);
 			}
+			HairGroupInstance->HairGroupPublicData->bIsSimulationCacheEnable = HairGroupInstance->Debug.GroomCacheType == EGroomCacheType::Guides;
 			HairGroupInstance->HairGroupPublicData->bIsDeformationEnable = GroomAsset->IsDeformationEnable(GroupIt);
 			HairGroupInstance->HairGroupPublicData->SetLODScreenSizes(CPULODScreenSize);
 			HairGroupInstance->HairGroupPublicData->SetLODVisibilities(LODVisibility);
@@ -2637,7 +2656,7 @@ void UGroomComponent::InitResources(bool bIsBindingReloading)
 		// * Physics simulation
 		// * RBF deformation.
 		// Therefore, even if simulation is disabled, we need to run partially the update if the binding system is enabled (skin deformation + RBF correction)
-		const bool bNeedGuides = (GroupData.Guides.HasValidData() && (bHasNeedSimulation[GroupIt] || bHasNeedGlobalDeformation[GroupIt] || bHasNeedDeformation[GroupIt])) || (HairGroupInstance->Debug.GroomCacheType == EGroomCacheType::Guides);
+		const bool bNeedGuides = (GroupData.Guides.HasValidData() && (bHasNeedSimulation[GroupIt] || bHasNeedGlobalDeformation[GroupIt] || bHasNeedDeformation[GroupIt])) || HairGroupInstance->HairGroupPublicData->bIsSimulationCacheEnable;
 		if (bNeedGuides)
 		{
 			HairGroupInstance->Guides.Data = &GroupData.Guides.BulkData;
@@ -2668,6 +2687,7 @@ void UGroomComponent::InitResources(bool bIsBindingReloading)
 			HairGroupInstance->Guides.bIsSimulationEnable = IsSimulationEnable(GroupIt,LODIndex);
 			HairGroupInstance->Guides.bHasGlobalInterpolation = LocalBindingAsset && GroomAsset->IsGlobalInterpolationEnable(GroupIt,LODIndex);
 			HairGroupInstance->Guides.bIsDeformationEnable = GroomAsset->IsDeformationEnable(GroupIt);
+			HairGroupInstance->Guides.bIsSimulationCacheEnable = HairGroupInstance->HairGroupPublicData->bIsSimulationCacheEnable;
 		}
 
 		// LODBias is in the Modifier which is needed for LOD selection regardless if the strands are there or not
@@ -2727,11 +2747,13 @@ void UGroomComponent::InitResources(bool bIsBindingReloading)
 				{
 					bNeedDynamicResources = true;
 				}
+
 				// the bones deformation needs to setup the deformed resources
 				if(HairGroupInstance->HairGroupPublicData->bIsDeformationEnable)
 				{
 					bNeedDynamicResources = true;
 				}
+
 				// Mesh deformer requires to dynamic resources
 				if (MeshDeformer)
 				{
@@ -3105,6 +3127,12 @@ void UGroomComponent::PostLoad()
 	// Do not validate the groom yet as the component count be loaded, but material/binding & co will be set later on
 	// ValidateMaterials(false);
 #endif
+
+	// Warning in case the groom is using a guide cache, but the EnableSimulationCache is disabled.
+	if (GroomAsset && GroomCache && GroomCache->GetType() == EGroomCacheType::Guides && !GroomAsset->EnableSimulationCache)
+	{
+		UE_LOG(LogHairStrands, Warning, TEXT("[Groom] %s - The groom instance tries to use a guide-cache, but 'Enable Guide-Cache Support' option is not enabled on the groom asset. Because of this, the Guide-Cache won't be used."), *GetPathName());
+	}
 }
 
 void UGroomComponent::CollectPSOPrecacheData(const FPSOPrecacheParams& BasePrecachePSOParams, FComponentPSOPrecacheParamsList& OutParams) 
@@ -3763,29 +3791,21 @@ void UGroomComponent::PostEditChangeProperty(FPropertyChangedEvent& PropertyChan
 		ResetAnimationTime();
 	}
 
-	// Check to see if simulation is required, which is needed when using with the guides groom cache,
+	// Check to see if simulation cache is enabled, which is needed when using with the guides groom cache,
 	// and enable it on all groom groups if it's not already
 	const EGroomCacheType GroomCacheType = GroomCache ? GroomCache->GetType() : EGroomCacheType::None;
-	const bool bSimulationRequired = GroomCacheType == EGroomCacheType::Guides;
-	if (bSimulationRequired && GroomAsset)
+	if (GroomCacheType == EGroomCacheType::Guides && GroomAsset)
 	{
-		bool bGroomAssetChanged = false;
 		for (int32 Index = 0; Index < GroomAsset->GetNumHairGroups(); ++Index)
 		{
-			if (!IsSimulationEnable(Index, -1))
+			if (!GroomAsset->EnableSimulationCache)
 			{
-				if (!bGroomAssetChanged)
-				{
-					GroomAsset->Modify();
-					bGroomAssetChanged = true;
-				}
-				GroomAsset->HairGroupsPhysics[Index].SolverSettings.EnableSimulation = true;
+				GroomAsset->Modify();
+				GroomAsset->EnableSimulationCache = true;
+				GroomAsset->CacheDerivedDatas();
+				break;
 			}
 		}
-		if (bGroomAssetChanged)
-		{
-			GroomAsset->CacheDerivedDatas();
-		}
 	}
 
 	if (bAssetChanged)
diff --git a/Engine/Plugins/Runtime/HairStrands/Source/HairStrandsCore/Private/GroomManager.cpp b/Engine/Plugins/Runtime/HairStrands/Source/HairStrandsCore/Private/GroomManager.cpp
--- a/Engine/Plugins/Runtime/HairStrands/Source/HairStrandsCore/Private/GroomManager.cpp
+++ b/Engine/Plugins/Runtime/HairStrands/Source/HairStrandsCore/Private/GroomManager.cpp
@@ -21,6 +21,7 @@
 #include "RenderGraphEvent.h"
 #include "RenderGraphUtils.h"
 #include "CachedGeometry.h"
+#include "GroomCacheData.h"
 
 static int32 GHairStrandsMinLOD = 0;
 static FAutoConsoleVariableRef CVarGHairStrandsMinLOD(TEXT("r.HairStrands.MinLOD"), GHairStrandsMinLOD, TEXT("Clamp the min hair LOD to this value, preventing to reach lower/high-quality LOD."), ECVF_Scalability);
@@ -652,7 +653,8 @@ void AddHairStreamingRequest(FHairGroupInstance* Instance, int32 InLODIndex)
 		const bool bSimulationEnable			= Instance->HairGroupPublicData->IsSimulationEnable(LODIndex);
 		const bool bDeformationEnable			= Instance->HairGroupPublicData->bIsDeformationEnable;
 		const bool bGlobalInterpolationEnable	= Instance->HairGroupPublicData->IsGlobalInterpolationEnable(LODIndex);
-		const bool bLODNeedsGuides				= bSimulationEnable || bDeformationEnable || bGlobalInterpolationEnable;
+		const bool bSimulationCacheEnable		= Instance->HairGroupPublicData->bIsSimulationCacheEnable;
+		const bool bLODNeedsGuides				= bSimulationEnable || bDeformationEnable || bGlobalInterpolationEnable || bSimulationCacheEnable;
 
 		const EHairResourceLoadingType LoadingType = GetHairResourceLoadingType(GeometryType, int32(LODIndex));
 		if (LoadingType != EHairResourceLoadingType::Async || GeometryType == EHairGeometryType::NoneGeometry)
@@ -898,7 +900,8 @@ static void RunHairLODSelection(
 			const bool bSimulationEnable			= Instance->HairGroupPublicData->IsSimulationEnable(IntLODIndex);
 			const bool bDeformationEnable			= Instance->HairGroupPublicData->bIsDeformationEnable;
 			const bool bGlobalInterpolationEnable	= Instance->HairGroupPublicData->IsGlobalInterpolationEnable(IntLODIndex);
-			const bool bLODNeedsGuides				= bSimulationEnable || bDeformationEnable || bGlobalInterpolationEnable;
+			const bool bSimulationCacheEnable		= Instance->HairGroupPublicData->bIsSimulationCacheEnable;
+			const bool bLODNeedsGuides				= bSimulationEnable || bDeformationEnable || bGlobalInterpolationEnable || bSimulationCacheEnable;
 
 			const EHairResourceLoadingType LoadingType = GetHairResourceLoadingType(GeometryType, IntLODIndex);
 			EHairResourceStatus ResourceStatus = EHairResourceStatus::None;
@@ -992,6 +995,7 @@ static void RunHairLODSelection(
 				Instance->Guides.bIsSimulationEnable = Instance->HairGroupPublicData->IsSimulationEnable(IntLODIndex);
 				Instance->Guides.bHasGlobalInterpolation = Instance->HairGroupPublicData->IsGlobalInterpolationEnable(IntLODIndex);
 				Instance->Guides.bIsDeformationEnable = Instance->HairGroupPublicData->bIsDeformationEnable;
+				Instance->Guides.bIsSimulationCacheEnable = Instance->HairGroupPublicData->bIsSimulationCacheEnable;
 				Instance->Strands.bIsCullingEnabled = bCullingEnable;
 			}
 
diff --git a/Engine/Plugins/Runtime/HairStrands/Source/HairStrandsCore/Private/HairStrandsRendering.cpp b/Engine/Plugins/Runtime/HairStrands/Source/HairStrandsCore/Private/HairStrandsRendering.cpp
--- a/Engine/Plugins/Runtime/HairStrands/Source/HairStrandsCore/Private/HairStrandsRendering.cpp
+++ b/Engine/Plugins/Runtime/HairStrands/Source/HairStrandsCore/Private/HairStrandsRendering.cpp
@@ -615,7 +615,7 @@ static void AddHairStrandsInterpolationPass(
 
 	// Guides
 	bool bSupportGlobalInterpolation = false;
-	if (bSupportDynamicMesh && (Instance->Guides.bIsSimulationEnable || Instance->Guides.bHasGlobalInterpolation || Instance->Guides.bIsDeformationEnable))
+	if (bSupportDynamicMesh && (Instance->Guides.bIsSimulationEnable || Instance->Guides.bHasGlobalInterpolation || Instance->Guides.bIsDeformationEnable)) // No need for Instance->Guides.bHasSimulationCache since it is incompatible with binding
 	{
 		const FHairStrandsRestRootResource::FLOD& Sim_RestLODDatas = SimRestRootResources->LODs[MeshLODIndex];
 		const FHairStrandsDeformedRootResource::FLOD& Sim_DeformedLODDatas = SimDeformedRootResources->LODs[MeshLODIndex];
@@ -660,7 +660,7 @@ static void AddHairStrandsInterpolationPass(
 		Parameters->RenDeformerPositionBuffer = RenDeformerPositionBuffer;
 	}
 
-	const bool bHasLocalDeformation = Instance->Guides.bIsSimulationEnable || bSupportGlobalInterpolation || Instance->Guides.bIsDeformationEnable;
+	const bool bHasLocalDeformation = Instance->Guides.bIsSimulationEnable || bSupportGlobalInterpolation || Instance->Guides.bIsDeformationEnable || Instance->Guides.bIsSimulationCacheEnable;
 	const bool bCullingEnable = IsHairStrandContinuousDecimationReorderingEnabled() ? false : (InstanceGeometryType == EHairGeometryType::Strands && CullingData.bCullingResultAvailable); 	// TODO: improve reordering so that culling can be used effectively
 	Parameters->HairStrandsVF_bIsCullingEnable = bCullingEnable ? 1 : 0;
 
@@ -1812,7 +1812,7 @@ void ComputeHairStrandsInterpolation(
 			FRDGBufferSRVRef Strands_PrevPositionOffsetSRV = nullptr;
 			FRDGBufferSRVRef Strands_TangentSRV = nullptr;
 
-			const bool bValidGuide		= bNeedDeformation && (Instance->Guides.bIsSimulationEnable || Instance->Guides.bHasGlobalInterpolation || Instance->Guides.bIsDeformationEnable);// || (WITH_EDITOR && PatchMode == EHairPatchAttribute::GuideInflucence);
+			const bool bValidGuide		= bNeedDeformation && (Instance->Guides.bIsSimulationEnable || Instance->Guides.bHasGlobalInterpolation || Instance->Guides.bIsDeformationEnable || Instance->Guides.bIsSimulationCacheEnable);// || (WITH_EDITOR && PatchMode == EHairPatchAttribute::GuideInflucence);
 			const bool bUseSingleGuide	= bNeedDeformation && bValidGuide && Instance->Strands.InterpolationResource->UseSingleGuide();
 			const bool bHasSkinning		= bNeedDeformation && Instance->BindingType == EHairBindingType::Skinning;
 
@@ -1867,7 +1867,7 @@ void ComputeHairStrandsInterpolation(
 						RegisterAsSRV(GraphBuilder, Instance->Guides.DeformedResource->GetPositionOffsetBuffer(FHairStrandsDeformedResource::EFrameType::Current)),
 						RegisterAsUAV(GraphBuilder, Instance->Guides.DeformedResource->GetBuffer(FHairStrandsDeformedResource::Current)));
 				}
-				else if( Instance->Guides.bIsDeformationEnable && Instance->Guides.DeformedResource && Instance->DeformedComponent && (Instance->DeformedSection != INDEX_NONE))
+				else if (Instance->Guides.bIsDeformationEnable && Instance->Guides.DeformedResource && Instance->DeformedComponent && (Instance->DeformedSection != INDEX_NONE))
 				{
 					if (USkeletalMeshComponent* SkeletalMeshComponent = Cast<USkeletalMeshComponent>(Instance->DeformedComponent))
 					{
@@ -2221,7 +2221,7 @@ void ComputeHairStrandsInterpolation(
 		const bool bIsCardsValid = Instance->Cards.IsValid(HairLODIndex);
 		if (bIsCardsValid)
 		{
-			const bool bValidGuide = Instance->Guides.bIsSimulationEnable || Instance->Guides.bHasGlobalInterpolation || Instance->Guides.bIsDeformationEnable;
+			const bool bValidGuide = Instance->Guides.bIsSimulationEnable || Instance->Guides.bHasGlobalInterpolation || Instance->Guides.bIsDeformationEnable || Instance->Guides.bIsSimulationCacheEnable;
 			const bool bHasSkinning = Instance->BindingType == EHairBindingType::Skinning && MeshLODIndex >= 0;
 			const bool bNeedDeformation = bValidGuide || bHasSkinning;
 
@@ -2428,8 +2428,8 @@ void ResetHairStrandsInterpolation(
 	int32 MeshLODIndex)
 {
 	if (!Instance || 
-		(Instance && (Instance->Guides.bIsSimulationEnable || Instance->Guides.bIsDeformationEnable)) ||
-		(Instance && !Instance->Guides.bHasGlobalInterpolation && !Instance->Guides.bIsSimulationEnable && !Instance->Guides.bIsDeformationEnable) ||
+		(Instance && (Instance->Guides.bIsSimulationEnable || Instance->Guides.bIsDeformationEnable || Instance->Guides.bIsSimulationCacheEnable)) ||
+		(Instance && !Instance->Guides.bHasGlobalInterpolation && !Instance->Guides.bIsSimulationEnable && !Instance->Guides.bIsDeformationEnable && !Instance->Guides.bIsSimulationCacheEnable) ||
 		!IsHairStrandsBindingEnable()) return;
 
 	DECLARE_GPU_STAT(HairStrandsGuideDeform);
diff --git a/Engine/Plugins/Runtime/HairStrands/Source/HairStrandsCore/Public/GroomAsset.h b/Engine/Plugins/Runtime/HairStrands/Source/HairStrandsCore/Public/GroomAsset.h
--- a/Engine/Plugins/Runtime/HairStrands/Source/HairStrandsCore/Public/GroomAsset.h
+++ b/Engine/Plugins/Runtime/HairStrands/Source/HairStrandsCore/Public/GroomAsset.h
@@ -458,6 +458,10 @@ public:
 	UPROPERTY(EditAnywhere, EditFixedSize, BlueprintReadWrite, Category = "HairInterpolation", meta = (ToolTip = "Enable radial basis function interpolation to be used instead of the local skin rigid transform (WIP)", DisplayName = "RBF Interpolation"))
 	bool EnableGlobalInterpolation = false;
 
+	/** Enable guide-cache support. This allows to attach a guide-cache dynamically at runtime */
+	UPROPERTY(EditAnywhere, EditFixedSize, BlueprintReadWrite, Category = "HairInterpolation", meta = (ToolTip = "Enable guide-cache support. This allows to attach a simulation-cache dynamically at runtime", DisplayName = "Enable Guide-Cache Support"))
+	bool EnableSimulationCache = false;
+
 	/** Type of interpolation used */
 	UPROPERTY(EditAnywhere, EditFixedSize, BlueprintReadWrite, Category = "HairInterpolation", meta = (ToolTip = "Type of interpolation used (WIP)"))
 	EGroomInterpolationType HairInterpolationType = EGroomInterpolationType::SmoothTransform;
@@ -674,7 +678,6 @@ public:
 	const FHairDescriptionGroups& GetHairDescriptionGroups();
 private:
 	
-	FString BuildDerivedDataKeySuffix(uint32 GroupIndex, const FHairGroupsInterpolation& InterpolationSettings, const FHairGroupsLOD& LODSettings) const;
 	bool IsFullyCached();
 	TUniquePtr<FHairDescriptionBulkData> HairDescriptionBulkData[EHairDescriptionType::Count];
 	EHairDescriptionType HairDescriptionType = EHairDescriptionType::Source;
diff --git a/Engine/Plugins/Runtime/HairStrands/Source/HairStrandsCore/Public/GroomInstance.h b/Engine/Plugins/Runtime/HairStrands/Source/HairStrandsCore/Public/GroomInstance.h
--- a/Engine/Plugins/Runtime/HairStrands/Source/HairStrandsCore/Public/GroomInstance.h
+++ b/Engine/Plugins/Runtime/HairStrands/Source/HairStrandsCore/Public/GroomInstance.h
@@ -137,6 +137,7 @@ struct HAIRSTRANDSCORE_API FHairGroupInstance : public FHairStrandsInstance
 		bool bIsSimulationEnable = false;
 		bool bIsDeformationEnable = false;
 		bool bHasGlobalInterpolation = false;
+		bool bIsSimulationCacheEnable = false;
 	} Guides;
 
 	//////////////////////////////////////////////////////////////////////////////////////////
diff --git a/Engine/Plugins/Runtime/HairStrands/Source/HairStrandsEditor/Private/GroomAssetDetails.cpp b/Engine/Plugins/Runtime/HairStrands/Source/HairStrandsEditor/Private/GroomAssetDetails.cpp
--- a/Engine/Plugins/Runtime/HairStrands/Source/HairStrandsEditor/Private/GroomAssetDetails.cpp
+++ b/Engine/Plugins/Runtime/HairStrands/Source/HairStrandsEditor/Private/GroomAssetDetails.cpp
@@ -392,6 +392,7 @@ void FGroomRenderingDetails::CustomizeStrandsGroupProperties(IDetailLayoutBuilde
 		DetailLayout.HideProperty(DetailLayout.GetProperty(GET_MEMBER_NAME_CHECKED(UGroomAsset, HairGroupsLOD), UGroomAsset::StaticClass()));
 		DetailLayout.HideProperty(DetailLayout.GetProperty(GET_MEMBER_NAME_CHECKED(UGroomAsset, HairGroupsInfo), UGroomAsset::StaticClass()));
 		DetailLayout.HideProperty(DetailLayout.GetProperty(GET_MEMBER_NAME_CHECKED(UGroomAsset, EnableGlobalInterpolation), UGroomAsset::StaticClass()));
+		DetailLayout.HideProperty(DetailLayout.GetProperty(GET_MEMBER_NAME_CHECKED(UGroomAsset, EnableSimulationCache), UGroomAsset::StaticClass()));
 		DetailLayout.HideProperty(DetailLayout.GetProperty(GET_MEMBER_NAME_CHECKED(UGroomAsset, HairInterpolationType), UGroomAsset::StaticClass()));
 		DetailLayout.HideProperty(DetailLayout.GetProperty(GET_MEMBER_NAME_CHECKED(UGroomAsset, MinLOD), UGroomAsset::StaticClass()));
 		DetailLayout.HideProperty(DetailLayout.GetProperty(GET_MEMBER_NAME_CHECKED(UGroomAsset, DisableBelowMinLodStripping), UGroomAsset::StaticClass()));
@@ -409,6 +410,7 @@ void FGroomRenderingDetails::CustomizeStrandsGroupProperties(IDetailLayoutBuilde
 		DetailLayout.HideProperty(DetailLayout.GetProperty(GET_MEMBER_NAME_CHECKED(UGroomAsset, HairGroupsLOD), UGroomAsset::StaticClass()));
 		DetailLayout.HideProperty(DetailLayout.GetProperty(GET_MEMBER_NAME_CHECKED(UGroomAsset, HairGroupsInfo), UGroomAsset::StaticClass()));
 		DetailLayout.HideProperty(DetailLayout.GetProperty(GET_MEMBER_NAME_CHECKED(UGroomAsset, EnableGlobalInterpolation), UGroomAsset::StaticClass()));
+		DetailLayout.HideProperty(DetailLayout.GetProperty(GET_MEMBER_NAME_CHECKED(UGroomAsset, EnableSimulationCache), UGroomAsset::StaticClass()));
 		DetailLayout.HideProperty(DetailLayout.GetProperty(GET_MEMBER_NAME_CHECKED(UGroomAsset, HairInterpolationType), UGroomAsset::StaticClass()));
 		DetailLayout.HideProperty(DetailLayout.GetProperty(GET_MEMBER_NAME_CHECKED(UGroomAsset, MinLOD), UGroomAsset::StaticClass()));
 		DetailLayout.HideProperty(DetailLayout.GetProperty(GET_MEMBER_NAME_CHECKED(UGroomAsset, DisableBelowMinLodStripping), UGroomAsset::StaticClass()));
@@ -426,6 +428,7 @@ void FGroomRenderingDetails::CustomizeStrandsGroupProperties(IDetailLayoutBuilde
 		DetailLayout.HideProperty(DetailLayout.GetProperty(GET_MEMBER_NAME_CHECKED(UGroomAsset, HairGroupsLOD), UGroomAsset::StaticClass()));
 		DetailLayout.HideProperty(DetailLayout.GetProperty(GET_MEMBER_NAME_CHECKED(UGroomAsset, HairGroupsInfo), UGroomAsset::StaticClass()));
 		DetailLayout.HideProperty(DetailLayout.GetProperty(GET_MEMBER_NAME_CHECKED(UGroomAsset, EnableGlobalInterpolation), UGroomAsset::StaticClass()));
+		DetailLayout.HideProperty(DetailLayout.GetProperty(GET_MEMBER_NAME_CHECKED(UGroomAsset, EnableSimulationCache), UGroomAsset::StaticClass()));
 		DetailLayout.HideProperty(DetailLayout.GetProperty(GET_MEMBER_NAME_CHECKED(UGroomAsset, HairInterpolationType), UGroomAsset::StaticClass()));
 		DetailLayout.HideProperty(DetailLayout.GetProperty(GET_MEMBER_NAME_CHECKED(UGroomAsset, MinLOD), UGroomAsset::StaticClass()));
 		DetailLayout.HideProperty(DetailLayout.GetProperty(GET_MEMBER_NAME_CHECKED(UGroomAsset, DisableBelowMinLodStripping), UGroomAsset::StaticClass()));
@@ -443,6 +446,7 @@ void FGroomRenderingDetails::CustomizeStrandsGroupProperties(IDetailLayoutBuilde
 		DetailLayout.HideProperty(DetailLayout.GetProperty(GET_MEMBER_NAME_CHECKED(UGroomAsset, HairGroupsLOD), UGroomAsset::StaticClass()));
 		DetailLayout.HideProperty(DetailLayout.GetProperty(GET_MEMBER_NAME_CHECKED(UGroomAsset, HairGroupsInfo), UGroomAsset::StaticClass()));
 		DetailLayout.HideProperty(DetailLayout.GetProperty(GET_MEMBER_NAME_CHECKED(UGroomAsset, EnableGlobalInterpolation), UGroomAsset::StaticClass()));
+		DetailLayout.HideProperty(DetailLayout.GetProperty(GET_MEMBER_NAME_CHECKED(UGroomAsset, EnableSimulationCache), UGroomAsset::StaticClass()));
 		DetailLayout.HideProperty(DetailLayout.GetProperty(GET_MEMBER_NAME_CHECKED(UGroomAsset, HairInterpolationType), UGroomAsset::StaticClass()));
 		DetailLayout.HideProperty(DetailLayout.GetProperty(GET_MEMBER_NAME_CHECKED(UGroomAsset, MinLOD), UGroomAsset::StaticClass()));
 		DetailLayout.HideProperty(DetailLayout.GetProperty(GET_MEMBER_NAME_CHECKED(UGroomAsset, DisableBelowMinLodStripping), UGroomAsset::StaticClass()));
@@ -460,6 +464,7 @@ void FGroomRenderingDetails::CustomizeStrandsGroupProperties(IDetailLayoutBuilde
 		DetailLayout.HideProperty(DetailLayout.GetProperty(GET_MEMBER_NAME_CHECKED(UGroomAsset, HairGroupsLOD), UGroomAsset::StaticClass()));
 		DetailLayout.HideProperty(DetailLayout.GetProperty(GET_MEMBER_NAME_CHECKED(UGroomAsset, HairGroupsInfo), UGroomAsset::StaticClass()));
 //		DetailLayout.HideProperty(DetailLayout.GetProperty(GET_MEMBER_NAME_CHECKED(UGroomAsset, EnableGlobalInterpolation), UGroomAsset::StaticClass()));
+//		DetailLayout.HideProperty(DetailLayout.GetProperty(GET_MEMBER_NAME_CHECKED(UGroomAsset, EnableSimulationCache), UGroomAsset::StaticClass()));
 //		DetailLayout.HideProperty(DetailLayout.GetProperty(GET_MEMBER_NAME_CHECKED(UGroomAsset, HairInterpolationType), UGroomAsset::StaticClass()));
 		DetailLayout.HideProperty(DetailLayout.GetProperty(GET_MEMBER_NAME_CHECKED(UGroomAsset, MinLOD), UGroomAsset::StaticClass()));
 		DetailLayout.HideProperty(DetailLayout.GetProperty(GET_MEMBER_NAME_CHECKED(UGroomAsset, DisableBelowMinLodStripping), UGroomAsset::StaticClass()));
@@ -477,6 +482,7 @@ void FGroomRenderingDetails::CustomizeStrandsGroupProperties(IDetailLayoutBuilde
 //		DetailLayout.HideProperty(DetailLayout.GetProperty(GET_MEMBER_NAME_CHECKED(UGroomAsset, HairGroupsLOD), UGroomAsset::StaticClass()));
 		DetailLayout.HideProperty(DetailLayout.GetProperty(GET_MEMBER_NAME_CHECKED(UGroomAsset, HairGroupsInfo), UGroomAsset::StaticClass()));
 		DetailLayout.HideProperty(DetailLayout.GetProperty(GET_MEMBER_NAME_CHECKED(UGroomAsset, EnableGlobalInterpolation), UGroomAsset::StaticClass()));
+		DetailLayout.HideProperty(DetailLayout.GetProperty(GET_MEMBER_NAME_CHECKED(UGroomAsset, EnableSimulationCache), UGroomAsset::StaticClass()));
 		DetailLayout.HideProperty(DetailLayout.GetProperty(GET_MEMBER_NAME_CHECKED(UGroomAsset, HairInterpolationType), UGroomAsset::StaticClass()));
 //		DetailLayout.HideProperty(DetailLayout.GetProperty(GET_MEMBER_NAME_CHECKED(UGroomAsset, MinLOD), UGroomAsset::StaticClass()));
 //		DetailLayout.HideProperty(DetailLayout.GetProperty(GET_MEMBER_NAME_CHECKED(UGroomAsset, DisableBelowMinLodStripping), UGroomAsset::StaticClass()));
diff --git a/Engine/Source/Runtime/Renderer/Public/HairStrandsInterface.h b/Engine/Source/Runtime/Renderer/Public/HairStrandsInterface.h
--- a/Engine/Source/Runtime/Renderer/Public/HairStrandsInterface.h
+++ b/Engine/Source/Runtime/Renderer/Public/HairStrandsInterface.h
@@ -298,8 +298,9 @@ public:
 	TArray<float>LODScreenSizes;
 	TArray<bool> LODSimulations;
 	TArray<bool> LODGlobalInterpolations;
-	bool bIsDeformationEnable;
 	TArray<EHairGeometryType> LODGeometryTypes;
+	bool bIsDeformationEnable = false;
+	bool bIsSimulationCacheEnable = false;
 
 	TArray<EHairBindingType> BindingTypes;
 
