Fix crash in groom builder when using groom rigging.

* Add StrandID attribute tracking
* Add point/curve copy function to reduce maintenance issue when adding new (optional) attribute

diff --git a/Engine/Plugins/Runtime/HairStrands/Source/HairStrandsCore/Private/GroomBuilder.cpp b/Engine/Plugins/Runtime/HairStrands/Source/HairStrandsCore/Private/GroomBuilder.cpp
--- a/Engine/Plugins/Runtime/HairStrands/Source/HairStrandsCore/Private/GroomBuilder.cpp
+++ b/Engine/Plugins/Runtime/HairStrands/Source/HairStrandsCore/Private/GroomBuilder.cpp
@@ -251,18 +251,6 @@ namespace HairStrandsBuilder
 		Out.Unlock();
 	}
 
-	template<typename T>
-	bool HasValues(const TArray<T>& In, T Ref)
-	{
-		for (const T& Value : In)
-		{
-			if (Value != Ref)
-			{
-				return true;
-			}
-		}
-		return false;
-	}
 	/** Build the packed datas for gpu rendering/simulation */
 	void BuildRenderData(const FHairStrandsDatas& HairStrands, const TArray<uint8>& RandomSeeds, FHairStrandsBulkData& OutBulkData)
 	{
@@ -291,10 +279,7 @@ namespace HairStrandsBuilder
 			OutVertexToCurve32.SetNum(NumPoints);
 		}
 
-		const bool bHasClumpIDs = HasValues(HairStrands.StrandsCurves.ClumpIDs, uint16(0));
-		const bool bHasColor = HasValues(HairStrands.StrandsPoints.PointsBaseColor, FLinearColor::Transparent);
-		const bool bHasRoughness = HasValues(HairStrands.StrandsPoints.PointsRoughness, 0.f);
-		const bool bHasAO = HasValues(HairStrands.StrandsPoints.PointsAO, 1.f);
+		const uint32 Attributes = HairStrands.GetAttributes();
 
 		TArray<uint32> AttributeRootUV;
 		TArray<uint32> AttributeSeed;
@@ -307,10 +292,10 @@ namespace HairStrandsBuilder
 		AttributeRootUV.SetNum(NumCurves);									// 4 bytes encoding - Per-Curve
 		AttributeSeed.SetNum(FMath::DivideAndRoundUp(NumCurves, 4u));		// 1 bytes encoding - Per-Curve
 		AttributeLength.SetNum(FMath::DivideAndRoundUp(NumCurves, 2u));		// 2 bytes encoding - Per-Curve
-		AttributeClumpIDs.SetNum(FMath::DivideAndRoundUp(NumCurves, 2u));	// 2 bytes encoding - Per-Curve
-		AttributeBaseColor.SetNum(NumPoints);								// 4 bytes encoding - Per-Vertex
-		AttributeRoughness.SetNum(FMath::DivideAndRoundUp(NumPoints, 2u));	// 1 bytes encoding - Per-Vertex
-		AttributeAO.SetNum(FMath::DivideAndRoundUp(NumPoints, 4u));			// 1 bytes encoding - Per-Vertex
+		if (HasHairAttribute(Attributes, EHairAttribute::ClumpID))	{ AttributeClumpIDs.SetNum(FMath::DivideAndRoundUp(NumCurves, 2u)); }	// 2 bytes encoding - Per-Curve
+		if (HasHairAttribute(Attributes, EHairAttribute::Color))	{ AttributeBaseColor.SetNum(NumPoints); }								// 4 bytes encoding - Per-Vertex
+		if (HasHairAttribute(Attributes, EHairAttribute::Roughness)){ AttributeRoughness.SetNum(FMath::DivideAndRoundUp(NumPoints, 2u)); }	// 1 bytes encoding - Per-Vertex
+		if (HasHairAttribute(Attributes, EHairAttribute::AO))		{ AttributeAO.SetNum(FMath::DivideAndRoundUp(NumPoints, 4u)); }			// 1 bytes encoding - Per-Vertex
 
 		const FVector HairBoxCenter = HairStrands.BoundingBox.GetCenter();
 
@@ -366,7 +351,7 @@ namespace HairStrandsBuilder
 				}
 
 				// Per-Vertex BaseColor
-				if (bHasColor)
+				if (HasHairAttribute(Attributes, EHairAttribute::Color))
 				{
 					uint32& Packed = AttributeBaseColor[PointIndex + IndexOffset];
 
@@ -377,14 +362,14 @@ namespace HairStrandsBuilder
 				}
 
 				// Per-Vertex Roughness
-				if (bHasRoughness)
+				if (HasHairAttribute(Attributes, EHairAttribute::Roughness))
 				{
 					uint8* Packed = (uint8*)AttributeRoughness.GetData();
 					Packed[PointIndex + IndexOffset] = FMath::Clamp(uint32(Points.PointsRoughness[PointIndex + IndexOffset] * 0xFF), 0u, 0xFFu);
 				}
 
 				// Per-Vertex AO
-				if (bHasAO)
+				if (HasHairAttribute(Attributes, EHairAttribute::AO))
 				{
 					uint8* Packed = (uint8*)AttributeAO.GetData();
 					Packed[PointIndex + IndexOffset] = FMath::Clamp(uint32(Points.PointsAO[PointIndex + IndexOffset] * 0xFF), 0u, 0xFFu);
@@ -426,12 +411,17 @@ namespace HairStrandsBuilder
 			}
 
 			// Per-Curve Clump ID (16 bits max)
-			if (bHasClumpIDs)
+			if (HasHairAttribute(Attributes, EHairAttribute::ClumpID))
 			{
 				uint16* Packed = (uint16*)AttributeClumpIDs.GetData();
 				Packed[CurveIndex] = uint16(FMath::Clamp(Curves.ClumpIDs[CurveIndex], 0, 65535));
 			}
 
+			// Per-Curve Strand ID
+			{
+				// At the moment we do not store this data within the bulk data
+			}
+
 			// Curves
 			{
 				UE_CLOG(PointCount > 0xFF, LogGroomBuilder, Warning, TEXT("Curve point count is > 256"));
@@ -483,25 +473,25 @@ namespace HairStrandsBuilder
 				OutPackedAttributes.Append(AttributeLength);
 			}
 
-			if (bHasClumpIDs)
+			if (HasHairAttribute(Attributes, EHairAttribute::ClumpID))
 			{
 				OutBulkData.AttributeOffsets[HAIR_ATTRIBUTE_CLUMPID] = OutPackedAttributes.Num() * UintToByte;
 				OutPackedAttributes.Append(AttributeClumpIDs);
 			}
 
-			if (bHasColor)
+			if (HasHairAttribute(Attributes, EHairAttribute::Color))
 			{
 				OutBulkData.AttributeOffsets[HAIR_ATTRIBUTE_BASECOLOR] = OutPackedAttributes.Num() * UintToByte;
 				OutPackedAttributes.Append(AttributeBaseColor);
 			}
 
-			if (bHasRoughness)
+			if (HasHairAttribute(Attributes, EHairAttribute::Roughness))
 			{
 				OutBulkData.AttributeOffsets[HAIR_ATTRIBUTE_ROUGHNESS] = OutPackedAttributes.Num() * UintToByte;
 				OutPackedAttributes.Append(AttributeRoughness);
 			}
 
-			if (bHasAO)
+			if (HasHairAttribute(Attributes, EHairAttribute::AO))
 			{
 				OutBulkData.AttributeOffsets[HAIR_ATTRIBUTE_AO] = OutPackedAttributes.Num() * UintToByte;
 				OutPackedAttributes.Append(AttributeAO);
@@ -1865,12 +1855,12 @@ bool FGroomBuilder::BuildHairDescriptionGroups(const FHairDescription& HairDescr
 		Group.Info.NumGuideVertices = Group.Guides.GetNumPoints();
 
 		// Prepare/Resize all attributes based on position/curves counts
-		Group.Strands.StrandsPoints.SetNum(Group.Strands.GetNumPoints());
-		Group.Strands.StrandsCurves.SetNum(Group.Strands.GetNumCurves());
+		Group.Strands.StrandsPoints.SetNum(Group.Strands.GetNumPoints(), 0u);
+		Group.Strands.StrandsCurves.SetNum(Group.Strands.GetNumCurves(), 0u);
 
 		// Prepare/Resize all attributes based on position/curves counts
-		Group.Guides.StrandsPoints.SetNum(Group.Guides.GetNumPoints());
-		Group.Guides.StrandsCurves.SetNum(Group.Guides.GetNumCurves());
+		Group.Guides.StrandsPoints.SetNum(Group.Guides.GetNumPoints(), 0u);
+		Group.Guides.StrandsCurves.SetNum(Group.Guides.GetNumCurves(), 0u);
 	}
 
 	return true;
@@ -2025,8 +2015,6 @@ void FGroomBuilder::BuildInterplationData(
 
 			HairInterpolationBuilder::BuildInterpolationData(OutInterpolationData, InSimData, InRenData, InInterpolationSettings, GuideIndices);
 		}
-
-
 	}
 }
 
@@ -2099,29 +2087,24 @@ void Decimate(
 	RootPositions.Init(FVector3f::ZeroVector, InData.GetNumCurves());
 	for(uint32 StrandIndex = 0; StrandIndex < InData.GetNumCurves(); ++StrandIndex)
 	{
-		RootPositions[StrandIndex] =
-			InData.StrandsPoints.PointsPosition[InData.StrandsCurves.CurvesOffset[StrandIndex]];
+		RootPositions[StrandIndex] = InData.StrandsPoints.PointsPosition[InData.StrandsCurves.CurvesOffset[StrandIndex]];
 	}
 				
 	TArray<bool> ValidPoints;
 	ValidPoints.Init(true, InData.GetNumCurves());
 
+	const uint32 InAttributes = InData.GetAttributes();
+
 	// Pick NumCurves strands from the guides
 	GroomBinding_RBFWeighting::FPointsSampler PointsSampler(ValidPoints, RootPositions.GetData(), OutCurveCount);
 
 	const uint32 CurveCount = PointsSampler.SampleIndices.Num();
 	uint32 PointCount = CurveCount * NumVertices;
 
-	OutData.StrandsCurves.SetNum(CurveCount);
-	OutData.StrandsPoints.SetNum(PointCount);
+	OutData.StrandsCurves.SetNum(CurveCount, InAttributes);
+	OutData.StrandsPoints.SetNum(PointCount, InAttributes);
 	OutData.HairDensity = InData.HairDensity;
 
-	const bool bHasPrecomputedWeights = InData.StrandsCurves.CurvesClosestGuideIDs.Num() > 0 && InData.StrandsCurves.CurvesClosestGuideWeights.Num() > 0;
-	const bool bHasColor = InData.StrandsPoints.PointsBaseColor.Num() > 0;
-	const bool bHasRoughness = InData.StrandsPoints.PointsRoughness.Num() > 0;
-	const bool bHasAO = InData.StrandsPoints.PointsAO.Num() > 0;
-	const bool bHasClumpIDs = InData.StrandsCurves.ClumpIDs.Num() > 0;
-	const bool bHasStrandIDs = InData.StrandsCurves.StrandIDs.Num() > 0;
 	check(InData.StrandsCurves.MaxRadius > 0);
 
 	PointCount = 0;
@@ -2135,24 +2118,9 @@ void Decimate(
 		const float BoneSpace = FMath::Max(1.0,float(GuideCount-1) / (MaxPoints-1));
 
 		OutData.StrandsCurves.CurvesCount[CurveIndex]  = MaxPoints;
-		OutData.StrandsCurves.CurvesRootUV[CurveIndex] = InData.StrandsCurves.CurvesRootUV[GuideIndex];
-		OutData.StrandsCurves.CurvesOffset[CurveIndex] = PointCount;
-		OutData.StrandsCurves.CurvesLength[CurveIndex] = InData.StrandsCurves.CurvesLength[GuideIndex];
-		if (bHasStrandIDs)
-		{
-			OutData.StrandsCurves.StrandIDs[CurveIndex] = InData.StrandsCurves.StrandIDs[GuideIndex];
-		}
-		if (bHasClumpIDs)
-		{
-			OutData.StrandsCurves.ClumpIDs[CurveIndex] = InData.StrandsCurves.ClumpIDs[GuideIndex];
-		}
-		if (bHasPrecomputedWeights)
-		{
-			OutData.StrandsCurves.CurvesClosestGuideIDs[CurveIndex] = InData.StrandsCurves.CurvesClosestGuideIDs[GuideIndex];
-			OutData.StrandsCurves.CurvesClosestGuideWeights[CurveIndex] = InData.StrandsCurves.CurvesClosestGuideWeights[GuideIndex];
-		}
 		OutData.StrandsCurves.MaxLength = InData.StrandsCurves.MaxLength;
 		OutData.StrandsCurves.MaxRadius = InData.StrandsCurves.MaxRadius;
+		FHairStrandsDatas::CopyCurve(InData, OutData, InAttributes, GuideIndex, CurveIndex);
 	
 		for(int32 PointIndex = 0; PointIndex < MaxPoints; ++PointIndex, ++PointCount)
 		{
@@ -2162,13 +2130,7 @@ void Decimate(
 			const int32 PointEnd = PointBegin+1;	
 			const float PointAlpha = PointSpace - PointBegin;
 			
-			OutData.StrandsPoints.PointsPosition[PointCount]	= FMath::Lerp(InData.StrandsPoints.PointsPosition[PointBegin],	InData.StrandsPoints.PointsPosition[PointEnd],	PointAlpha);
-			OutData.StrandsPoints.PointsCoordU[PointCount]		= FMath::Lerp(InData.StrandsPoints.PointsCoordU[PointBegin],	InData.StrandsPoints.PointsCoordU[PointEnd],	PointAlpha);
-			// Multiply by MaxRadius as HairStrandsBuilder::BuildInternalData will recompute MaxRadius based on the actual Radius not the not the normalized radius.
-			OutData.StrandsPoints.PointsRadius[PointCount]		= FMath::Lerp(InData.StrandsPoints.PointsRadius[PointBegin],	InData.StrandsPoints.PointsRadius[PointEnd],	PointAlpha) * InData.StrandsCurves.MaxRadius; 
-			if (bHasColor)		{ OutData.StrandsPoints.PointsBaseColor[PointCount]	= FMath::Lerp(InData.StrandsPoints.PointsBaseColor[PointBegin], InData.StrandsPoints.PointsBaseColor[PointEnd], PointAlpha); }
-			if (bHasRoughness)	{ OutData.StrandsPoints.PointsRoughness[PointCount]	= FMath::Lerp(InData.StrandsPoints.PointsRoughness[PointBegin], InData.StrandsPoints.PointsRoughness[PointEnd], PointAlpha); }
-			if (bHasAO)			{ OutData.StrandsPoints.PointsAO[PointCount]		= FMath::Lerp(InData.StrandsPoints.PointsAO[PointBegin],		InData.StrandsPoints.PointsAO[PointEnd],		PointAlpha); }
+			FHairStrandsDatas::CopyPointLerp(InData, OutData, InAttributes, PointBegin, PointEnd, PointAlpha, PointCount);
 		}
 	}
 	HairStrandsBuilder::BuildInternalData(OutData, false);
@@ -2188,12 +2150,7 @@ void Decimate(
 	// Divide strands in buckets and pick randomly one stand per bucket
 	const uint32 InCurveCount = InData.StrandsCurves.Num();
 	const uint32 OutCurveCount = FMath::Clamp(uint32(InCurveCount * CurveDecimationPercentage), 1u, InCurveCount);
-	const bool bHasPrecomputedWeights = InData.StrandsCurves.CurvesClosestGuideIDs.Num() > 0 && InData.StrandsCurves.CurvesClosestGuideWeights.Num() > 0;
-	const bool bHasColor = InData.StrandsPoints.PointsBaseColor.Num() > 0;
-	const bool bHasRoughness = InData.StrandsPoints.PointsRoughness.Num() > 0;
-	const bool bHasAO = InData.StrandsPoints.PointsAO.Num() > 0;
-	const bool bHasClumpIDs = InData.StrandsCurves.ClumpIDs.Num() > 0;
-	const bool bHasStrandIDs = InData.StrandsCurves.StrandIDs.Num() > 0;
+	const uint32 InAttributes = InData.GetAttributes();
 	check(InData.StrandsCurves.MaxRadius > 0);
 
 	FRandomStream Random;
@@ -2248,21 +2205,10 @@ void Decimate(
 		}
 	}
 
-	OutData.StrandsCurves.SetNum(OutCurveCount);
-	OutData.StrandsPoints.SetNum(OutTotalPointCount);
+	OutData.StrandsCurves.SetNum(OutCurveCount, InAttributes);
+	OutData.StrandsPoints.SetNum(OutTotalPointCount, InAttributes);
 	OutData.HairDensity = InData.HairDensity;
 
-	if (bHasColor)		{ OutData.StrandsPoints.PointsBaseColor.SetNum(OutTotalPointCount); }
-	if (bHasRoughness)	{ OutData.StrandsPoints.PointsRoughness.SetNum(OutTotalPointCount); }
-	if (bHasAO)			{ OutData.StrandsPoints.PointsAO.SetNum(OutTotalPointCount); }
-	if (bHasStrandIDs)	{ OutData.StrandsCurves.StrandIDs.SetNum(OutCurveCount); }
-	if (bHasClumpIDs)	{ OutData.StrandsCurves.ClumpIDs.SetNum(OutCurveCount); }
-	if (bHasPrecomputedWeights)
-	{
-		OutData.StrandsCurves.CurvesClosestGuideIDs.SetNum(OutCurveCount);
-		OutData.StrandsCurves.CurvesClosestGuideWeights.SetNum(OutCurveCount);
-	}
-
 	uint32 OutPointOffset = 0; 
 	for (uint32 OutCurveIndex = 0; OutCurveIndex < OutCurveCount; ++OutCurveIndex)
 	{
@@ -2273,7 +2219,6 @@ void Decimate(
 		const uint32 InPointCount	= InData.StrandsCurves.CurvesCount[EffectiveCurveIndex];
 
 		// Decimation using area metric
-	#if 1
 		TArray<uint32> OutPointIndices;
 		if (bNeedDecimation)
 		{
@@ -2300,59 +2245,15 @@ void Decimate(
 		{
 			const uint32 InPointIndex = bNeedDecimation ? OutPointIndices[OutPointIndex] : OutPointIndex;
 
-			OutData.StrandsPoints.PointsPosition [OutPointIndex + OutPointOffset] = InData.StrandsPoints.PointsPosition	[InPointIndex + InPointOffset];
-			OutData.StrandsPoints.PointsCoordU	 [OutPointIndex + OutPointOffset] = InData.StrandsPoints.PointsCoordU	[InPointIndex + InPointOffset];
-			// Multiply by MaxRadius as HairStrandsBuilder::BuildInternalData will recompute MaxRadius based on the actual Radius not the not the normalized radius.
-			OutData.StrandsPoints.PointsRadius	 [OutPointIndex + OutPointOffset] = InData.StrandsPoints.PointsRadius	[InPointIndex + InPointOffset] * InData.StrandsCurves.MaxRadius;
-			if (bHasColor)		{ OutData.StrandsPoints.PointsBaseColor[OutPointIndex + OutPointOffset] = InData.StrandsPoints.PointsBaseColor[InPointIndex + InPointOffset]; }
-			if (bHasRoughness)	{ OutData.StrandsPoints.PointsRoughness[OutPointIndex + OutPointOffset] = InData.StrandsPoints.PointsRoughness[InPointIndex + InPointOffset]; }
-			if (bHasAO)			{ OutData.StrandsPoints.PointsAO[OutPointIndex + OutPointOffset]		= InData.StrandsPoints.PointsAO[InPointIndex + InPointOffset]; }
-		}
-	#else
-		// Decimation using uniform/interleaved removal
-
-		// Insure always the original start/end points are part of the trimmed curve
-		const uint32 OutPointCount = FMath::Clamp(uint32(InPointCount * VertexDecimationPercentage), 2u, InPointCount);
-		const uint32 VertexBucketSize = InPointCount / OutPointCount;
-
-		// Take vertex in the middle of the bucket
-		const uint32 OffsetInsideBucket = FMath::FloorToInt(VertexBucketSize * 0.5f);
-		for (uint32 OutPointIndex = 0; OutPointIndex < OutPointCount; ++OutPointIndex)
-		{
-			uint32 InPointIndex = FMath::Clamp(uint32(VertexBucketSize * OutPointIndex + OffsetInsideBucket), 0u, InPointCount - 1);
-
-			// Insure first and last points map onto the root and tip vertices
-			InPointIndex = OutPointIndex == 0 ? 0 : InPointIndex;
-			InPointIndex = OutPointIndex == OutPointCount-1 ? InPointCount-1 : InPointIndex;
-
-			OutData.StrandsPoints.PointsPosition [OutPointIndex + OutPointOffset] = InData.StrandsPoints.PointsPosition	[InPointIndex + InPointOffset];
-			OutData.StrandsPoints.PointsCoordU	 [OutPointIndex + OutPointOffset] = InData.StrandsPoints.PointsCoordU	[InPointIndex + InPointOffset];
-			OutData.StrandsPoints.PointsRadius	 [OutPointIndex + OutPointOffset] = InData.StrandsPoints.PointsRadius	[InPointIndex + InPointOffset];
-			if (bHasColor)		{ OutData.StrandsPoints.PointsBaseColor[OutPointIndex + OutPointOffset] = InData.StrandsPoints.PointsBaseColor[InPointIndex + InPointOffset]; }
-			if (bHasRoughness)	{ OutData.StrandsPoints.PointsRoughness[OutPointIndex + OutPointOffset] = InData.StrandsPoints.PointsRoughness[InPointIndex + InPointOffset]; }
-			if (bHasAO)			{ OutData.StrandsPoints.PointsAO[OutPointIndex + OutPointOffset]		= InData.StrandsPoints.PointsAO[InPointIndex + InPointOffset]; }
+			FHairStrandsDatas::CopyPoint(InData, OutData, InAttributes, InPointIndex + InPointOffset, OutPointIndex + OutPointOffset);
 		}
-	#endif
 
 		OutData.StrandsCurves.CurvesCount[OutCurveIndex]  = OutPointCount;
 		OutData.StrandsCurves.CurvesOffset[OutCurveIndex] = OutPointOffset;
-		OutData.StrandsCurves.CurvesRootUV[OutCurveIndex] = InData.StrandsCurves.CurvesRootUV[EffectiveCurveIndex];
-		OutData.StrandsCurves.CurvesLength[OutCurveIndex] = InData.StrandsCurves.CurvesLength[EffectiveCurveIndex];
-		if (bHasPrecomputedWeights)
-		{
-			OutData.StrandsCurves.CurvesClosestGuideIDs[OutCurveIndex] = InData.StrandsCurves.CurvesClosestGuideIDs[EffectiveCurveIndex];
-			OutData.StrandsCurves.CurvesClosestGuideWeights[OutCurveIndex] = InData.StrandsCurves.CurvesClosestGuideWeights[EffectiveCurveIndex];
-		}
-		if (bHasStrandIDs)
-		{
-			OutData.StrandsCurves.StrandIDs[OutCurveIndex] = InData.StrandsCurves.StrandIDs[EffectiveCurveIndex];
-		}
-		if (bHasClumpIDs)
-		{
-			OutData.StrandsCurves.ClumpIDs[OutCurveIndex] = InData.StrandsCurves.ClumpIDs[EffectiveCurveIndex];
-		}
 		OutData.StrandsCurves.MaxLength = InData.StrandsCurves.MaxLength;
 		OutData.StrandsCurves.MaxRadius = InData.StrandsCurves.MaxRadius;
+		FHairStrandsDatas::CopyCurve(InData, OutData, InAttributes, EffectiveCurveIndex, OutCurveIndex);
+
 		OutPointOffset += OutPointCount;
 	}
 
@@ -2468,8 +2369,7 @@ static void Voxelize(const FHairDescriptionGroups& InGroups, FHairGrid& Out)
 		FHairStrandsDatas In = Group.Strands;
 		HairStrandsBuilder::BuildInternalData(In, false);
 
-		const bool bHasColor     = In.StrandsPoints.PointsBaseColor.Num() > 0;
-		const bool bHasRoughness = In.StrandsPoints.PointsRoughness.Num() > 0;
+		const uint32 Attributes = In.GetAttributes();
 
 		// Fill in voxel (TODO: make it parallel)
 		for (uint32 CurveIt = 0, CurveCount = In.GetNumCurves(); CurveIt < CurveCount; ++CurveIt)
@@ -2486,10 +2386,10 @@ static void Voxelize(const FHairDescriptionGroups& InGroups, FHairGrid& Out)
 				const FVector3f& P1 = In.StrandsPoints.PointsPosition[Index1];
 				const float R0 = In.StrandsPoints.PointsRadius[Index0];
 				const float R1 = In.StrandsPoints.PointsRadius[Index1];
-				const FVector3f C0 = bHasColor ? FVector3f(In.StrandsPoints.PointsBaseColor[Index0].R, In.StrandsPoints.PointsBaseColor[Index0].G, In.StrandsPoints.PointsBaseColor[Index0].B) : FVector3f(0);
-				const FVector3f C1 = bHasColor ? FVector3f(In.StrandsPoints.PointsBaseColor[Index1].R, In.StrandsPoints.PointsBaseColor[Index1].G, In.StrandsPoints.PointsBaseColor[Index1].B) : FVector3f(0);
-				const float Rough0 = bHasRoughness ? In.StrandsPoints.PointsRoughness[Index0] : 0;
-				const float Rough1 = bHasRoughness ? In.StrandsPoints.PointsRoughness[Index1] : 0;
+				const FVector3f C0 = HasHairAttribute(Attributes, EHairAttribute::Color) ? FVector3f(In.StrandsPoints.PointsBaseColor[Index0].R, In.StrandsPoints.PointsBaseColor[Index0].G, In.StrandsPoints.PointsBaseColor[Index0].B) : FVector3f(0);
+				const FVector3f C1 = HasHairAttribute(Attributes, EHairAttribute::Color) ? FVector3f(In.StrandsPoints.PointsBaseColor[Index1].R, In.StrandsPoints.PointsBaseColor[Index1].G, In.StrandsPoints.PointsBaseColor[Index1].B) : FVector3f(0);
+				const float Rough0 = HasHairAttribute(Attributes, EHairAttribute::Roughness) ? In.StrandsPoints.PointsRoughness[Index0] : 0;
+				const float Rough1 = HasHairAttribute(Attributes, EHairAttribute::Roughness) ? In.StrandsPoints.PointsRoughness[Index1] : 0;
 				const FVector3f Segment = P1 - P0;
 
 				// This is a coarse/non-conservative voxelization, by ray-marching segment, instead of walking intersected voxels
diff --git a/Engine/Plugins/Runtime/HairStrands/Source/HairStrandsCore/Private/HairAttributes.cpp b/Engine/Plugins/Runtime/HairStrands/Source/HairStrandsCore/Private/HairAttributes.cpp
--- a/Engine/Plugins/Runtime/HairStrands/Source/HairStrandsCore/Private/HairAttributes.cpp
+++ b/Engine/Plugins/Runtime/HairStrands/Source/HairStrandsCore/Private/HairAttributes.cpp
@@ -2,6 +2,7 @@
 
 #include "HairAttributes.h"
 #include "HairStrandsDefinitions.h"
+#include "HairStrandsDatas.h"
 
 namespace HairAttribute
 {
@@ -51,12 +52,13 @@ void SetHairAttribute(uint32& Out, EHairAttribute InAttribute)
 const TCHAR* GetHairAttributeText(EHairAttribute In)
 {
 	// If a new optional attribute is added, please add its UI/text description here
-	static_assert(uint32(EHairAttribute::Count) == 6);
+	static_assert(uint32(EHairAttribute::Count) == 7);
 
 	switch (In)
 	{
 	case EHairAttribute::RootUV:					return TEXT("RootUV");
 	case EHairAttribute::ClumpID:					return TEXT("ClumpID");
+	case EHairAttribute::StrandID:					return TEXT("StrandID");
 	case EHairAttribute::PrecomputedGuideWeights:	return TEXT("ImportedGuideWeights");
 	case EHairAttribute::Color:						return TEXT("Color");
 	case EHairAttribute::Roughness:					return TEXT("Roughness");
@@ -71,10 +73,85 @@ uint32 GetHairAttributeIndex(EHairAttribute In)
 	{
 	case EHairAttribute::RootUV:					return HAIR_ATTRIBUTE_ROOTUV;
 	case EHairAttribute::ClumpID:					return HAIR_ATTRIBUTE_CLUMPID;
-	case EHairAttribute::PrecomputedGuideWeights:	return HAIR_ATTRIBUTE_COUNT; // We don't store this information
+	case EHairAttribute::StrandID:					return HAIR_ATTRIBUTE_COUNT;	// We don't store this information (yet)
+	case EHairAttribute::PrecomputedGuideWeights:	return HAIR_ATTRIBUTE_COUNT;	// We don't store this information
 	case EHairAttribute::Color:						return HAIR_ATTRIBUTE_BASECOLOR;
 	case EHairAttribute::Roughness:					return HAIR_ATTRIBUTE_ROUGHNESS;
 	case EHairAttribute::AO:						return HAIR_ATTRIBUTE_AO;
 	}
 	return HAIR_ATTRIBUTE_COUNT;
+}
+
+bool FHairStrandsPoints::HasAttribute(EHairAttribute In) const
+{
+	if (In == EHairAttribute::Color)		{ return PointsBaseColor.Num() > 0; }
+	if (In == EHairAttribute::Roughness)	{ return PointsRoughness.Num() > 0; }
+	if (In == EHairAttribute::AO)			{ return PointsAO.Num() > 0; }
+	return false;
+}
+
+bool FHairStrandsCurves::HasAttribute(EHairAttribute In) const
+{
+	if (In == EHairAttribute::PrecomputedGuideWeights)	{ return CurvesClosestGuideIDs.Num() > 0 && CurvesClosestGuideWeights.Num() > 0; }
+	if (In == EHairAttribute::ClumpID)					{ return ClumpIDs.Num() > 0; }
+	if (In == EHairAttribute::StrandID)					{ return StrandIDs.Num() > 0; }
+	return false;
+}
+
+uint32 FHairStrandsDatas::GetAttributes() const
+{
+	uint32 Out = 0;
+	for (uint32 It=0;It<uint32(EHairAttribute::Count); ++It)
+	{	
+		const EHairAttribute Attribute = (EHairAttribute)It;
+		if (StrandsPoints.HasAttribute(Attribute) || StrandsCurves.HasAttribute(Attribute))
+		{
+			SetHairAttribute(Out, Attribute);
+		}
+	}
+	return Out;
+}
+
+void FHairStrandsDatas::CopyPoint(const FHairStrandsDatas& In, FHairStrandsDatas& Out, uint32 InAttributes, uint32 InIndex, uint32 OutIndex)
+{
+	// Multiply by MaxRadius as HairStrandsBuilder::BuildInternalData will recompute MaxRadius based on the actual Radius not the not the normalized radius.
+	Out.StrandsPoints.PointsPosition[OutIndex]	= In.StrandsPoints.PointsPosition[InIndex];
+	Out.StrandsPoints.PointsCoordU[OutIndex]	= In.StrandsPoints.PointsCoordU[InIndex];
+	Out.StrandsPoints.PointsRadius[OutIndex]	= In.StrandsPoints.PointsRadius[InIndex] * In.StrandsCurves.MaxRadius;
+
+	if (HasHairAttribute(InAttributes, EHairAttribute::Color))		{ Out.StrandsPoints.PointsBaseColor[OutIndex] = In.StrandsPoints.PointsBaseColor[InIndex]; }
+	if (HasHairAttribute(InAttributes, EHairAttribute::Roughness))	{ Out.StrandsPoints.PointsRoughness[OutIndex] = In.StrandsPoints.PointsRoughness[InIndex]; }
+	if (HasHairAttribute(InAttributes, EHairAttribute::AO))			{ Out.StrandsPoints.PointsAO[OutIndex] = In.StrandsPoints.PointsAO[InIndex]; }
+}
+
+
+void FHairStrandsDatas::CopyPointLerp(const FHairStrandsDatas& In, FHairStrandsDatas& Out, uint32 InAttributes, uint32 InIndex0, uint32 InIndex1, float InAlpha, uint32 OutIndex)
+{
+	// Multiply by MaxRadius as HairStrandsBuilder::BuildInternalData will recompute MaxRadius based on the actual Radius not the not the normalized radius.
+	Out.StrandsPoints.PointsPosition[OutIndex]	= FMath::Lerp(In.StrandsPoints.PointsPosition[InIndex0], In.StrandsPoints.PointsPosition[InIndex1], InAlpha);
+	Out.StrandsPoints.PointsCoordU[OutIndex]	= FMath::Lerp(In.StrandsPoints.PointsCoordU[InIndex0], In.StrandsPoints.PointsCoordU[InIndex1], InAlpha);
+	Out.StrandsPoints.PointsRadius[OutIndex]	= FMath::Lerp(In.StrandsPoints.PointsRadius[InIndex0] , In.StrandsPoints.PointsRadius[InIndex1], InAlpha) * In.StrandsCurves.MaxRadius;
+
+	if (HasHairAttribute(InAttributes, EHairAttribute::Color))		{ Out.StrandsPoints.PointsBaseColor[OutIndex]	= FMath::Lerp(In.StrandsPoints.PointsBaseColor[InIndex0], In.StrandsPoints.PointsBaseColor[InIndex1], InAlpha); }
+	if (HasHairAttribute(InAttributes, EHairAttribute::Roughness))	{ Out.StrandsPoints.PointsRoughness[OutIndex]	= FMath::Lerp(In.StrandsPoints.PointsRoughness[InIndex0], In.StrandsPoints.PointsRoughness[InIndex1], InAlpha); }
+	if (HasHairAttribute(InAttributes, EHairAttribute::AO))			{ Out.StrandsPoints.PointsAO[OutIndex]			= FMath::Lerp(In.StrandsPoints.PointsAO[InIndex0], In.StrandsPoints.PointsAO[InIndex1], InAlpha); }
+}
+
+void FHairStrandsDatas::CopyCurve(const FHairStrandsDatas& In, FHairStrandsDatas& Out, uint32 InAttributes, uint32 InIndex, uint32 OutIndex)
+{
+	Out.StrandsCurves.CurvesRootUV[OutIndex] = In.StrandsCurves.CurvesRootUV[InIndex];
+	Out.StrandsCurves.CurvesLength[OutIndex] = In.StrandsCurves.CurvesLength[InIndex];
+	if (HasHairAttribute(InAttributes, EHairAttribute::PrecomputedGuideWeights))
+	{
+		Out.StrandsCurves.CurvesClosestGuideIDs[OutIndex] = In.StrandsCurves.CurvesClosestGuideIDs[InIndex];
+		Out.StrandsCurves.CurvesClosestGuideWeights[OutIndex] = In.StrandsCurves.CurvesClosestGuideWeights[InIndex];
+	}
+	if (HasHairAttribute(InAttributes, EHairAttribute::StrandID))
+	{
+		Out.StrandsCurves.StrandIDs[OutIndex] = In.StrandsCurves.StrandIDs[InIndex];
+	}
+	if (HasHairAttribute(InAttributes, EHairAttribute::ClumpID))
+	{
+		Out.StrandsCurves.ClumpIDs[OutIndex] = In.StrandsCurves.ClumpIDs[InIndex];
+	}
 }
\ No newline at end of file
diff --git a/Engine/Plugins/Runtime/HairStrands/Source/HairStrandsCore/Private/HairCardsBuilder.cpp b/Engine/Plugins/Runtime/HairStrands/Source/HairStrandsCore/Private/HairCardsBuilder.cpp
--- a/Engine/Plugins/Runtime/HairStrands/Source/HairStrandsCore/Private/HairCardsBuilder.cpp
+++ b/Engine/Plugins/Runtime/HairStrands/Source/HairStrandsCore/Private/HairCardsBuilder.cpp
@@ -1356,8 +1356,8 @@ namespace HairCards
 			OutTotalPointCount += InData.StrandsCurves.CurvesCount[CurveIndex];
 		}
 
-		OutData.StrandsCurves.SetNum(OutCurveCount);
-		OutData.StrandsPoints.SetNum(OutTotalPointCount);
+		OutData.StrandsCurves.SetNum(OutCurveCount, 0u);
+		OutData.StrandsPoints.SetNum(OutTotalPointCount, 0u);
 		OutData.HairDensity = InData.HairDensity;
 
 		uint32 OutPointOffset = 0;
@@ -2770,8 +2770,8 @@ namespace HairCards
 		FHairStrandsDatas& OutGuides)
 	{
 		const uint32 ClusterCount = InClusters.GetNum();
-		OutGuides.StrandsPoints.SetNum(InClusters.GuidePoints.Num());
-		OutGuides.StrandsCurves.SetNum(InClusters.GetNum());
+		OutGuides.StrandsPoints.SetNum(InClusters.GuidePoints.Num(), 0u);
+		OutGuides.StrandsCurves.SetNum(InClusters.GetNum(), 0u);
 		OutGuides.BoundingBox.Init();
 
 		uint32 IndexIt = 0;
@@ -2839,7 +2839,7 @@ namespace HairCards
 		OutGuides.StrandsPoints.PointsPosition.Reserve(MaxGuidePoints);
 		OutGuides.StrandsPoints.PointsRadius.Reserve(MaxGuidePoints);
 		OutGuides.StrandsPoints.PointsCoordU.Reserve(MaxGuidePoints);
-		OutGuides.StrandsCurves.SetNum(NumCards);
+		OutGuides.StrandsCurves.SetNum(NumCards, 0u);
 		OutGuides.BoundingBox.Init();
 
 		InCards.CoordU.SetNum(InCards.Positions.Num());
diff --git a/Engine/Plugins/Runtime/HairStrands/Source/HairStrandsCore/Private/HairDescription.cpp b/Engine/Plugins/Runtime/HairStrands/Source/HairStrandsCore/Private/HairDescription.cpp
--- a/Engine/Plugins/Runtime/HairStrands/Source/HairStrandsCore/Private/HairDescription.cpp
+++ b/Engine/Plugins/Runtime/HairStrands/Source/HairStrandsCore/Private/HairDescription.cpp
@@ -69,12 +69,13 @@ bool FHairDescription::IsValid() const
 bool FHairDescription::HasAttribute(EHairAttribute InAttribute) const
 {
 	// If a new optional attribute is added, please add its UI/text description here
-	static_assert(uint32(EHairAttribute::Count) == 6);
+	static_assert(uint32(EHairAttribute::Count) == 7);
 
 	switch (InAttribute)
 	{
 	case EHairAttribute::RootUV:		return StrandAttributes().GetAttributesRef<FVector2f>(HairAttribute::Strand::RootUV).IsValid();
 	case EHairAttribute::ClumpID:		return StrandAttributes().GetAttributesRef<int>(HairAttribute::Strand::ClumpID).IsValid();
+	case EHairAttribute::StrandID:		return StrandAttributes().GetAttributesRef<int>(HairAttribute::Strand::ID).IsValid();
 	case EHairAttribute::Color:			return VertexAttributes().GetAttributesRef<FVector3f>(HairAttribute::Vertex::Color).IsValid();
 	case EHairAttribute::Roughness:		return VertexAttributes().GetAttributesRef<float>(HairAttribute::Vertex::Roughness).IsValid();
 	case EHairAttribute::AO:			return VertexAttributes().GetAttributesRef<float>(HairAttribute::Vertex::AO).IsValid();
@@ -83,11 +84,11 @@ bool FHairDescription::HasAttribute(EHairAttribute InAttribute) const
 		return
 			// Single
 			(StrandAttributes().GetAttributesRef<int>(HairAttribute::Strand::ClosestGuides).IsValid() &&
-				StrandAttributes().GetAttributesRef<float>(HairAttribute::Strand::GuideWeights).IsValid())
+			 StrandAttributes().GetAttributesRef<float>(HairAttribute::Strand::GuideWeights).IsValid())
 			||
 			// Triplet
 			(StrandAttributes().GetAttributesRef<FVector3f>(HairAttribute::Strand::ClosestGuides).IsValid() &&
-				StrandAttributes().GetAttributesRef<FVector3f>(HairAttribute::Strand::GuideWeights).IsValid());
+			 StrandAttributes().GetAttributesRef<FVector3f>(HairAttribute::Strand::GuideWeights).IsValid());
 	}
 	}
 	return false;
diff --git a/Engine/Plugins/Runtime/HairStrands/Source/HairStrandsCore/Private/HairStrandsDatas.cpp b/Engine/Plugins/Runtime/HairStrands/Source/HairStrandsCore/Private/HairStrandsDatas.cpp
--- a/Engine/Plugins/Runtime/HairStrands/Source/HairStrandsCore/Private/HairStrandsDatas.cpp
+++ b/Engine/Plugins/Runtime/HairStrands/Source/HairStrandsCore/Private/HairStrandsDatas.cpp
@@ -3,6 +3,7 @@
 #include "HairStrandsDatas.h"
 #include "UObject/ReleaseObjectVersion.h"
 #include "UObject/UE5ReleaseStreamObjectVersion.h"
+#include "HairAttributes.h"
 
 void FHairStrandsInterpolationDatas::SetNum(const uint32 NumCurves)
 {
@@ -12,7 +13,7 @@ void FHairStrandsInterpolationDatas::SetNum(const uint32 NumCurves)
 	PointsSimCurvesIndex.SetNum(NumCurves);
 }
 
-void FHairStrandsCurves::SetNum(const uint32 NumCurves)
+void FHairStrandsCurves::SetNum(const uint32 NumCurves, uint32 InAttributes)
 {
 	CurvesOffset.SetNum(NumCurves + 1);
 	CurvesCount.SetNum(NumCurves);
@@ -20,21 +21,25 @@ void FHairStrandsCurves::SetNum(const uint32 NumCurves)
 	CurvesRootUV.SetNum(NumCurves);
 
 	// Not initialized to track if the data are available
-	// ClumpIDs.SetNum(0);
-	// CurvesClosestGuideIDs.SetNum(0);
-	// CurvesClosestGuideWeights.SetNum(0);
+	if (HasHairAttribute(InAttributes, EHairAttribute::StrandID))	{ StrandIDs.SetNum(NumCurves); }
+	if (HasHairAttribute(InAttributes, EHairAttribute::ClumpID))	{ ClumpIDs.SetNum(NumCurves); }
+	if (HasHairAttribute(InAttributes, EHairAttribute::PrecomputedGuideWeights))
+	{
+		CurvesClosestGuideIDs.SetNum(NumCurves);
+		CurvesClosestGuideWeights.SetNum(NumCurves);
+	}
 }
 
-void FHairStrandsPoints::SetNum(const uint32 NumPoints)
+void FHairStrandsPoints::SetNum(const uint32 NumPoints, uint32 InAttributes)
 {
 	PointsPosition.SetNum(NumPoints);
 	PointsRadius.SetNum(NumPoints);
 	PointsCoordU.SetNum(NumPoints);
 
 	// Not initialized to track if the data are available
-	// PointsBaseColor.SetNum(0);
-	// PointsRoughness.SetNum(0);
-	// PointsAO.SetNum(0);
+	if (HasHairAttribute(InAttributes, EHairAttribute::Color))		{ PointsBaseColor.SetNum(NumPoints); }
+	if (HasHairAttribute(InAttributes, EHairAttribute::Roughness))	{ PointsRoughness.SetNum(NumPoints); }
+	if (HasHairAttribute(InAttributes, EHairAttribute::AO))			{ PointsAO.SetNum(NumPoints); }
 }
 
 void FHairStrandsInterpolationDatas::Reset()
diff --git a/Engine/Plugins/Runtime/HairStrands/Source/HairStrandsCore/Public/HairAttributes.h b/Engine/Plugins/Runtime/HairStrands/Source/HairStrandsCore/Public/HairAttributes.h
--- a/Engine/Plugins/Runtime/HairStrands/Source/HairStrandsCore/Public/HairAttributes.h
+++ b/Engine/Plugins/Runtime/HairStrands/Source/HairStrandsCore/Public/HairAttributes.h
@@ -51,6 +51,7 @@ enum class EHairAttribute : uint8
 {
 	RootUV,
 	ClumpID,
+	StrandID,
 	PrecomputedGuideWeights,
 	Color,
 	Roughness,
diff --git a/Engine/Plugins/Runtime/HairStrands/Source/HairStrandsCore/Public/HairStrandsDatas.h b/Engine/Plugins/Runtime/HairStrands/Source/HairStrandsCore/Public/HairStrandsDatas.h
--- a/Engine/Plugins/Runtime/HairStrands/Source/HairStrandsCore/Public/HairStrandsDatas.h
+++ b/Engine/Plugins/Runtime/HairStrands/Source/HairStrandsCore/Public/HairStrandsDatas.h
@@ -12,6 +12,8 @@
 
 DECLARE_LOG_CATEGORY_EXTERN(LogHairStrands, Log, All);
 
+enum class EHairAttribute : uint8;
+
 struct FPackedHairVertex
 {
 	typedef uint64 BulkType;
@@ -342,7 +344,7 @@ struct HAIRSTRANDSCORE_API FHairStrandsInterpolationBulkData
 struct HAIRSTRANDSCORE_API FHairStrandsPoints
 {
 	/** Set the number of points */
-	void SetNum(const uint32 NumPoints);
+	void SetNum(const uint32 NumPoints, uint32 InAttributes);
 
 	/** Reset the points to 0 */
 	void Reset();
@@ -350,6 +352,8 @@ struct HAIRSTRANDSCORE_API FHairStrandsPoints
 	/** Get the number of points */
 	uint32 Num() const { return PointsPosition.Num();  }
 
+	bool HasAttribute(EHairAttribute In) const;
+
 	/** Points position in local space */
 	TArray<FVector3f> PointsPosition;
 
@@ -373,7 +377,7 @@ struct HAIRSTRANDSCORE_API FHairStrandsPoints
 struct HAIRSTRANDSCORE_API FHairStrandsCurves
 {
 	/** Set the number of Curves */
-	void SetNum(const uint32 NumPoints);
+	void SetNum(const uint32 NumPoints, uint32 InAttributes);
 
 	/** Reset the curves to 0 */
 	void Reset();
@@ -383,6 +387,8 @@ struct HAIRSTRANDSCORE_API FHairStrandsCurves
 
 	bool HasPrecomputedWeights() const { return CurvesClosestGuideIDs.Num() > 0 && CurvesClosestGuideWeights.Num() > 0; }
 
+	bool HasAttribute(EHairAttribute In) const;
+
 	/** Number of points per rod */
 	TArray<uint16> CurvesCount;
 
@@ -428,10 +434,17 @@ struct HAIRSTRANDSCORE_API FHairStrandsDatas
 	/* Get the total number of Curves */
 	uint32 GetNumCurves() const { return StrandsCurves.Num(); }
 
+	uint32 GetAttributes() const;
+
 	void Reset();
 
 	bool IsValid() const { return StrandsCurves.Num() > 0 && StrandsPoints.Num() > 0; }
 
+	/* Copy a point or a curve from In to Out data */
+	static void CopyCurve(const FHairStrandsDatas& In, FHairStrandsDatas& Out, uint32 InAttributes, uint32 InIndex, uint32 OutIndex);
+	static void CopyPoint(const FHairStrandsDatas& In, FHairStrandsDatas& Out, uint32 InAttributes, uint32 InIndex, uint32 OutIndex);
+	static void CopyPointLerp(const FHairStrandsDatas& In, FHairStrandsDatas& Out, uint32 InAttributes, uint32 InIndex0, uint32 InIndex1, float InAlpha, uint32 OutIndex);
+
 	/** List of all the strands points */
 	FHairStrandsPoints StrandsPoints;
 
diff --git a/Engine/Plugins/Runtime/HairStrands/Source/HairStrandsEditor/Private/GroomImportOptionsWindow.cpp b/Engine/Plugins/Runtime/HairStrands/Source/HairStrandsEditor/Private/GroomImportOptionsWindow.cpp
--- a/Engine/Plugins/Runtime/HairStrands/Source/HairStrandsEditor/Private/GroomImportOptionsWindow.cpp
+++ b/Engine/Plugins/Runtime/HairStrands/Source/HairStrandsEditor/Private/GroomImportOptionsWindow.cpp
@@ -214,12 +214,13 @@ static void AddAttribute(SVerticalBox::FScopedWidgetSlotArguments& Slot, FText A
 static FText GetHairAttributeLocText(EHairAttribute In)
 {
 	// If a new optional attribute is added, please add its UI/text description here
-	static_assert(uint32(EHairAttribute::Count) == 6);
+	static_assert(uint32(EHairAttribute::Count) == 7);
 
 	switch (In)
 	{
 	case EHairAttribute::RootUV:					return LOCTEXT("GroomOptionsWindow_HasRootUV", "Root UV");
 	case EHairAttribute::ClumpID:					return LOCTEXT("GroomOptionsWindow_HasClumpID", "Clump ID");
+	case EHairAttribute::StrandID:					return LOCTEXT("GroomOptionsWindow_HasStrandID", "Strand ID");
 	case EHairAttribute::PrecomputedGuideWeights:	return LOCTEXT("GroomOptionsWindow_HasPercomputedGuideWeights", "Pre-Computed Guide Weights");
 	case EHairAttribute::Color:						return LOCTEXT("GroomOptionsWindow_HasColor", "Color");
 	case EHairAttribute::Roughness:					return LOCTEXT("GroomOptionsWindow_HasRoughness", "Roughness");
