MH - Fix crash in groom builder when using hair rigging feature.

diff --git a/Engine/Plugins/Runtime/HairStrands/Source/HairStrandsCore/Private/GroomBuilder.cpp b/Engine/Plugins/Runtime/HairStrands/Source/HairStrandsCore/Private/GroomBuilder.cpp
--- a/Engine/Plugins/Runtime/HairStrands/Source/HairStrandsCore/Private/GroomBuilder.cpp
+++ b/Engine/Plugins/Runtime/HairStrands/Source/HairStrandsCore/Private/GroomBuilder.cpp
@@ -2107,6 +2107,17 @@ void Decimate(
 
 	check(InData.StrandsCurves.MaxRadius > 0);
 
+	if (bHasColor)		{ OutData.StrandsPoints.PointsBaseColor.SetNum(PointCount); }
+	if (bHasRoughness)	{ OutData.StrandsPoints.PointsRoughness.SetNum(PointCount); }
+	if (bHasAO)			{ OutData.StrandsPoints.PointsAO.SetNum(PointCount); }
+	if (bHasStrandIDs)	{ OutData.StrandsCurves.StrandIDs.SetNum(CurveCount); }
+	if (bHasClumpIDs)	{ OutData.StrandsCurves.ClumpIDs.SetNum(CurveCount); }
+	if (bHasPrecomputedWeights)
+	{
+		OutData.StrandsCurves.CurvesClosestGuideIDs.SetNum(CurveCount);
+		OutData.StrandsCurves.CurvesClosestGuideWeights.SetNum(CurveCount);
+	}
+
 	PointCount = 0;
 	for(uint32 CurveIndex = 0; CurveIndex < CurveCount; ++CurveIndex)
 	{
