Fix invalid groom guide weight initialization when precomputed data are invalid.

diff --git a/Engine/Plugins/Runtime/HairStrands/Source/HairStrandsCore/Private/GroomBuilder.cpp b/Engine/Plugins/Runtime/HairStrands/Source/HairStrandsCore/Private/GroomBuilder.cpp
--- a/Engine/Plugins/Runtime/HairStrands/Source/HairStrandsCore/Private/GroomBuilder.cpp
+++ b/Engine/Plugins/Runtime/HairStrands/Source/HairStrandsCore/Private/GroomBuilder.cpp
@@ -40,7 +40,7 @@ static FAutoConsoleVariableRef CVarHairGroupIndexBuilder_MaxVoxelResolution(TEXT
 
 FString FGroomBuilder::GetVersion()
 {
-	return TEXT("v6a");
+	return TEXT("v6l");
 }
 
 namespace FHairStrandsDecimation
@@ -1503,6 +1503,7 @@ namespace HairInterpolationBuilder
 				const FVector& RenPointPosition = (FVector)RenData.StrandsPoints.PointsPosition[PointGlobalIndex];
 				const float RenPointDistance = RenData.StrandsPoints.PointsCoordU[PointGlobalIndex] * RenData.StrandsCurves.CurvesLength[CurveIndex] * RenData.StrandsCurves.MaxLength;
 
+				bool bHasValidGuide = false;
 				float TotalWeight = 0;
 				for (uint32 GuideIndex = 0; GuideIndex < FClosestGuides::Count; ++GuideIndex)
 				{
@@ -1520,6 +1521,26 @@ namespace HairInterpolationBuilder
 						OutInterpolation.PointsSimCurvesVertexLerp[PointGlobalIndex][GuideIndex] = Desc.T;
 						OutInterpolation.PointsSimCurvesVertexWeights[PointGlobalIndex][GuideIndex] = StrandGuideWeights[GuideIndex];
 						TotalWeight += OutInterpolation.PointsSimCurvesVertexWeights[PointGlobalIndex][GuideIndex];
+						bHasValidGuide = true;
+					}
+				}
+
+				// To not have invalid data, filled in interpolation data with first guide
+				if (!bHasValidGuide)
+				{
+					for (uint32 GuideIndex = 0; GuideIndex < FClosestGuides::Count; ++GuideIndex)
+					{
+						const uint32 SimCurveIndex = 0; // Take the first guide
+
+						// Fill the interpolation data using the ParametricDistance algorithm with a constant weight for all vertices along the strand
+						const uint32 SimOffset = SimData.StrandsCurves.CurvesOffset[SimCurveIndex];
+						const FVertexInterpolationDesc Desc = FindMatchingVertex(RenPointDistance, SimData, SimCurveIndex);
+
+						OutInterpolation.PointsSimCurvesIndex[PointGlobalIndex][GuideIndex] = SimCurveIndex;
+						OutInterpolation.PointsSimCurvesVertexIndex[PointGlobalIndex][GuideIndex] = Desc.Index0 + SimOffset;
+						OutInterpolation.PointsSimCurvesVertexLerp[PointGlobalIndex][GuideIndex] = Desc.T;
+						OutInterpolation.PointsSimCurvesVertexWeights[PointGlobalIndex][GuideIndex] = 1.0f;
+						TotalWeight += OutInterpolation.PointsSimCurvesVertexWeights[PointGlobalIndex][GuideIndex];
 					}
 				}
 
