Fix incorrect tool tip on GroomBinding asset.

diff --git a/Engine/Plugins/Runtime/HairStrands/Source/HairStrandsCore/Public/GroomBindingAsset.h b/Engine/Plugins/Runtime/HairStrands/Source/HairStrandsCore/Public/GroomBindingAsset.h
--- a/Engine/Plugins/Runtime/HairStrands/Source/HairStrandsCore/Public/GroomBindingAsset.h
+++ b/Engine/Plugins/Runtime/HairStrands/Source/HairStrandsCore/Public/GroomBindingAsset.h
@@ -96,7 +96,7 @@ public:
 	UPROPERTY(VisibleAnywhere, BlueprintReadWrite, Category = "BuildSettings")
 	int32 NumInterpolationPoints = 100;
 
-	/** Number of points used for the rbf interpolation */
+	/** Section to pick to transfer the positions */
 	UPROPERTY(VisibleAnywhere, BlueprintReadWrite, Category = "BuildSettings")
 	int32 MatchingSection = 0;
 
