* Fix incorrect encoding/decoding of per-CV groom color.
* Improve color precision from 8/8/8bits to 11/11/10bits.

diff --git a/Engine/Plugins/Runtime/HairStrands/Source/HairStrandsCore/Private/GroomBuilder.cpp b/Engine/Plugins/Runtime/HairStrands/Source/HairStrandsCore/Private/GroomBuilder.cpp
--- a/Engine/Plugins/Runtime/HairStrands/Source/HairStrandsCore/Private/GroomBuilder.cpp
+++ b/Engine/Plugins/Runtime/HairStrands/Source/HairStrandsCore/Private/GroomBuilder.cpp
@@ -40,7 +40,7 @@ static FAutoConsoleVariableRef CVarHairGroupIndexBuilder_MaxVoxelResolution(TEXT
 
 FString FGroomBuilder::GetVersion()
 {
-	return TEXT("v6l");
+	return TEXT("v7");
 }
 
 namespace FHairStrandsDecimation
@@ -350,15 +350,13 @@ namespace HairStrandsBuilder
 					OutVertexToCurve32[PointIndex + IndexOffset] = CurveIndex;
 				}
 
-				// Per-Vertex BaseColor
+				// Per-Vertex Color
 				if (HasHairAttribute(Attributes, EHairAttribute::Color))
 				{
 					uint32& Packed = AttributeBaseColor[PointIndex + IndexOffset];
-
-					// Cheap sRGB encoding
-					Packed= FMath::Clamp(uint32(FMath::Sqrt(Points.PointsBaseColor[PointIndex + IndexOffset].R) * 0xFF), 0u, 0xFFu)|
-							FMath::Clamp(uint32(FMath::Sqrt(Points.PointsBaseColor[PointIndex + IndexOffset].G) * 0xFF), 0u, 0xFFu)|
-							FMath::Clamp(uint32(FMath::Sqrt(Points.PointsBaseColor[PointIndex + IndexOffset].B) * 0xFF), 0u, 0xFFu);
+					Packed |= FMath::Clamp(uint32(FMath::Sqrt(Points.PointsBaseColor[PointIndex + IndexOffset].R) * 0x7FF), 0u, 0x7FFu);
+					Packed |= FMath::Clamp(uint32(FMath::Sqrt(Points.PointsBaseColor[PointIndex + IndexOffset].G) * 0x7FF), 0u, 0x7FFu)<<11;
+					Packed |= FMath::Clamp(uint32(FMath::Sqrt(Points.PointsBaseColor[PointIndex + IndexOffset].B) * 0x3FF), 0u, 0x3FFu)<<22;
 				}
 
 				// Per-Vertex Roughness
diff --git a/Engine/Shaders/Private/HairStrands/HairStrandsAttributeCommon.ush b/Engine/Shaders/Private/HairStrands/HairStrandsAttributeCommon.ush
--- a/Engine/Shaders/Private/HairStrands/HairStrandsAttributeCommon.ush
+++ b/Engine/Shaders/Private/HairStrands/HairStrandsAttributeCommon.ush
@@ -136,8 +136,8 @@ float3 GetHairStrandsBaseColor(uint HairPrimitiveId, float2 HairPrimitiveUV)
 	const uint BaseColorOffset = HAIR_ATTRIBUTE_OFFSET_BASECOLOR(HairStrandsVF.AttributeOffsets);
 	if (IsAttributeValid(BaseColorOffset))
 	{
-		const float3 Color0 = UnpackHairBaseColor(InternalGetHairStrandsAttributesPerVertex(HairPrimitiveId,   BaseColorOffset, HAIR_ATTRIBUTE_STRIDE_BASECOLOR, HairStrandsVF.AttributeBuffer));
-		const float3 Color1 = UnpackHairBaseColor(InternalGetHairStrandsAttributesPerVertex(HairPrimitiveId+1, BaseColorOffset, HAIR_ATTRIBUTE_STRIDE_BASECOLOR, HairStrandsVF.AttributeBuffer));
+		const float3 Color0 = UnpackHairColor(InternalGetHairStrandsAttributesPerVertex(HairPrimitiveId,   BaseColorOffset, HAIR_ATTRIBUTE_STRIDE_BASECOLOR, HairStrandsVF.AttributeBuffer));
+		const float3 Color1 = UnpackHairColor(InternalGetHairStrandsAttributesPerVertex(HairPrimitiveId+1, BaseColorOffset, HAIR_ATTRIBUTE_STRIDE_BASECOLOR, HairStrandsVF.AttributeBuffer));
 		return lerp(Color0, Color1, HairPrimitiveUV.x);
 	}
 	else
diff --git a/Engine/Shaders/Private/HairStrands/HairStrandsTexturesGeneration.usf b/Engine/Shaders/Private/HairStrands/HairStrandsTexturesGeneration.usf
--- a/Engine/Shaders/Private/HairStrands/HairStrandsTexturesGeneration.usf
+++ b/Engine/Shaders/Private/HairStrands/HairStrandsTexturesGeneration.usf
@@ -307,8 +307,8 @@ bool Trace(
 							const uint BaseColorOffset = HAIR_ATTRIBUTE_OFFSET_BASECOLOR(InVF_AttributeOffsets);
 							if (IsAttributeValid(BaseColorOffset))
 							{
-								const float3 C0 = UnpackHairBaseColor(InternalGetHairStrandsAttributesPerVertex(P0, BaseColorOffset, HAIR_ATTRIBUTE_STRIDE_BASECOLOR, InVF_AttributeBuffer));
-								const float3 C1 = UnpackHairBaseColor(InternalGetHairStrandsAttributesPerVertex(P1, BaseColorOffset, HAIR_ATTRIBUTE_STRIDE_BASECOLOR, InVF_AttributeBuffer));
+								const float3 C0 = UnpackHairColor(InternalGetHairStrandsAttributesPerVertex(P0, BaseColorOffset, HAIR_ATTRIBUTE_STRIDE_BASECOLOR, InVF_AttributeBuffer));
+								const float3 C1 = UnpackHairColor(InternalGetHairStrandsAttributesPerVertex(P1, BaseColorOffset, HAIR_ATTRIBUTE_STRIDE_BASECOLOR, InVF_AttributeBuffer));
 								OutMaterial.xyz = lerp(C0, C1, VertexU);
 							}
 
diff --git a/Engine/Shaders/Private/HairStrands/HairStrandsVertexFactoryCommon.ush b/Engine/Shaders/Private/HairStrands/HairStrandsVertexFactoryCommon.ush
--- a/Engine/Shaders/Private/HairStrands/HairStrandsVertexFactoryCommon.ush
+++ b/Engine/Shaders/Private/HairStrands/HairStrandsVertexFactoryCommon.ush
@@ -99,13 +99,13 @@ uint UnpackHairClumpID(uint In)
 	return (In & 0xFFFF);
 }
 
-float3 UnpackHairBaseColor(uint In)
+float3 UnpackHairColor(uint In)
 {
 	float3 Out;
-	Out.x = (In       & 0xFF) / 255.f;
-	Out.y = ((In>>8)  & 0xFF) / 255.f;
-	Out.z = ((In>>16) & 0xFF) / 255.f;
-	return Out * Out; // Cheap sRGB encoding
+	Out.x = (In       & 0x7FF) / 2047.f;
+	Out.y = ((In>>11) & 0x7FF) / 2047.f;
+	Out.z = ((In>>22) & 0x3FF) / 1023.f;
+	return Out;
 }
 
 float UnpackHairRoughness(uint In)
