Change Groom RootUV to be optional.

diff --git a/Engine/Plugins/Runtime/HairStrands/Source/HairStrandsCore/Private/GroomBuilder.cpp b/Engine/Plugins/Runtime/HairStrands/Source/HairStrandsCore/Private/GroomBuilder.cpp
--- a/Engine/Plugins/Runtime/HairStrands/Source/HairStrandsCore/Private/GroomBuilder.cpp
+++ b/Engine/Plugins/Runtime/HairStrands/Source/HairStrandsCore/Private/GroomBuilder.cpp
@@ -40,7 +40,7 @@ static FAutoConsoleVariableRef CVarHairGroupIndexBuilder_MaxVoxelResolution(TEXT
 
 FString FGroomBuilder::GetVersion()
 {
-	return TEXT("v7");
+	return TEXT("v8");
 }
 
 namespace FHairStrandsDecimation
@@ -61,90 +61,8 @@ namespace FHairStrandsDecimation
 
 namespace HairStrandsBuilder
 {
-	FVector2D SignNotZero(const FVector2D& v)
-	{
-		return FVector2D((v.X >= 0.0) ? +1.0 : -1.0, (v.Y >= 0.0) ? +1.0 : -1.0);
-	}
-
-	// A Survey of Efficient Representations for Independent Unit Vectors
-	// Reference: http://jcgt.org/published/0003/02/01/paper.pdf
-	// Assume normalized input. Output is on [-1, 1] for each component.
-	FVector2D SphericalToOctahedron(const FVector& v)
-	{
-		// Project the sphere onto the octahedron, and then onto the xy plane
-		FVector2D p = FVector2D(v.X, v.Y) * (1.0 / (abs(v.X) + abs(v.Y) + abs(v.Z)));
-		// Reflect the folds of the lower hemisphere over the diagonals
-		return (v.Z <= 0.0) ? ((FVector2D(1, 1) - FVector2D(abs(p.Y), abs(p.X))) * SignNotZero(p)) : p;
-	}
-
-	// Auto-generate Root UV data if not loaded
-	void ComputeRootUV(FHairStrandsCurves& Curves, FHairStrandsPoints& Points)
-	{
-		TRACE_CPUPROFILER_EVENT_SCOPE(HairStrandsBuilder::ComputeRootUV);
-
-		TArray<FVector> RootPoints;
-		const uint32 CurveCount = Curves.Num();
-		RootPoints.Reserve(Curves.Num());
-		FVector MinAABB(MAX_FLT, MAX_FLT, MAX_FLT);
-		FVector MaxAABB(-MAX_FLT, -MAX_FLT, -MAX_FLT);
-		FMatrix Rotation = FRotationMatrix::Make(FRotator(0, 0, -90));
-		for (uint32 CurveIndex=0; CurveIndex< CurveCount; ++CurveIndex)
-		{
-			const uint32 Offset = Curves.CurvesOffset[CurveIndex];
-			check(Offset < uint32(Points.PointsPosition.Num()));
-			const FVector P = Rotation.TransformPosition((FVector)Points.PointsPosition[Offset]);
-
-			RootPoints.Add(P);
-			MinAABB.X = FMath::Min(P.X, MinAABB.X);
-			MinAABB.Y = FMath::Min(P.Y, MinAABB.Y);
-			MinAABB.Z = FMath::Min(P.Z, MinAABB.Z);
-
-			MaxAABB.X = FMath::Max(P.X, MaxAABB.X);
-			MaxAABB.Y = FMath::Max(P.Y, MaxAABB.Y);
-			MaxAABB.Z = FMath::Max(P.Z, MaxAABB.Z);
-		}
-
-		// Compute sphere bound
-		const FVector Extent = MaxAABB - MinAABB;
-		FSphere SBound;
-		SBound.Center = (MaxAABB + MinAABB) * 0.5f;
-		SBound.W = FMath::Max(Extent.X, FMath::Max(Extent.Y, Extent.Z));
-
-		// Project root point onto the bounding sphere and map it onto 
-		// an octahedron, which is unfold onto the unit space [0,1]^2
-		TArray<FVector2D> RootUVs;
-		RootUVs.Reserve(Curves.Num());
-		FVector2D MinUV(MAX_FLT, MAX_FLT);
-		FVector2D MaxUV(-MAX_FLT, -MAX_FLT);
-		for (const FVector& RootP : RootPoints)
-		{
-			FVector D = RootP - SBound.Center;
-			D.Normalize();
-			FVector2D UV = SphericalToOctahedron(D);
-			UV += FVector2D(1, 1);
-			UV *= 0.5f;
-			RootUVs.Add(UV);
-
-			MinUV.X = FMath::Min(UV.X, MinUV.X);
-			MinUV.Y = FMath::Min(UV.Y, MinUV.Y);
-
-			MaxUV.X = FMath::Max(UV.X, MaxUV.X);
-			MaxUV.Y = FMath::Max(UV.Y, MaxUV.Y);
-		}
-
-		// Find the minimal UV space cover by root point, and 
-		// offsets/scales it to maximize UV space
-		const FVector2f UVScale(1 / (MaxUV.X - MinUV.X), 1 / (MaxUV.Y - MinUV.Y));
-		const FVector2f UVOffset(-MinUV.X, -MinUV.Y);
-		uint32 Index = 0;
-		for (FVector2f& RootUV : Curves.CurvesRootUV)
-		{
-			RootUV = (FVector2f(RootUVs[Index++]) + UVOffset) * UVScale;	// LWC_TODO: Precision loss
-		}
-	}
-
 	/** Build the internal points and curves data */
-	void BuildInternalData(FHairStrandsDatas& HairStrands, bool bComputeRootUV)
+	void BuildInternalData(FHairStrandsDatas& HairStrands)
 	{
 		TRACE_CPUPROFILER_EVENT_SCOPE(HairStrandsBuilder::BuildInternalData);
 
@@ -222,11 +140,6 @@ namespace HairStrandsBuilder
 					*LengthIterator /= Curves.MaxLength;
 				}
 			}
-
-			if (bComputeRootUV)
-			{
-				ComputeRootUV(Curves, Points);
-			}
 		}
 	}
 
@@ -388,6 +301,7 @@ namespace HairStrandsBuilder
 			}
 
 			// Per-Curve Root UVs
+			if (HasHairAttribute(Attributes, EHairAttribute::RootUV))
 			{
 				FVector2D RootUV = FVector2D(Curves.CurvesRootUV[CurveIndex]);
 				// Root UV support UDIM texture coordinate but limit the spans of the UDIM to be in 256x256 instead of 9999x9999.
@@ -422,8 +336,8 @@ namespace HairStrandsBuilder
 
 			// Curves
 			{
-				UE_CLOG(PointCount > 0xFF, LogGroomBuilder, Warning, TEXT("Curve point count is > 256"));
-				UE_CLOG(IndexOffset > 0xFFFFFF, LogGroomBuilder, Warning, TEXT("Curve point offset is > 24M"));
+				UE_CLOG(PointCount > 0xFF, LogGroomBuilder, Warning, TEXT("[Groom] Curve point count is > 256, which is not supported."));
+				UE_CLOG(IndexOffset > 0xFFFFFF, LogGroomBuilder, Warning, TEXT("[Groom] Curve point offset is > 24M, which is not supported."));
 				FHairStrandsCurveFormat::Type& Curve = OutPackedCurves[CurveIndex];
 				Curve.PointCount  = FMath::Min(uint32(PointCount), 0xFFu);
 				Curve.PointOffset = FMath::Min(uint32(IndexOffset), 0xFFFFFFu);
@@ -456,6 +370,7 @@ namespace HairStrandsBuilder
 				OutBulkData.AttributeOffsets[AttributeIt] = 0xFFFFFFFF;
 			}
 
+			if (HasHairAttribute(Attributes, EHairAttribute::RootUV))
 			{
 				OutBulkData.AttributeOffsets[HAIR_ATTRIBUTE_ROOTUV] = OutPackedAttributes.Num() * UintToByte;
 				OutPackedAttributes.Append(AttributeRootUV);
@@ -1887,7 +1802,7 @@ bool FGroomBuilder::BuildHairDescriptionGroups(const FHairDescription& HairDescr
 
 void FGroomBuilder::BuildData(FHairStrandsDatas& OutStrands)
 {
-	HairStrandsBuilder::BuildInternalData(OutStrands, false);
+	HairStrandsBuilder::BuildInternalData(OutStrands);
 }
 
 void FGroomBuilder::BuildData(
@@ -1916,7 +1831,7 @@ void FGroomBuilder::BuildData(
 		{
 			OutRen = InHairDescriptionGroup.Strands;
 
-			HairStrandsBuilder::BuildInternalData(OutRen, !HasHairAttribute(InHairDescriptionGroup.Attributes, EHairAttribute::RootUV));
+			HairStrandsBuilder::BuildInternalData(OutRen);
 
 			// Decimate
 			if (CurveDecimation < 1 || VertexDecimation < 1 || bCurveReordering)
@@ -1940,7 +1855,7 @@ void FGroomBuilder::BuildData(
 			OutSim = InHairDescriptionGroup.Guides;
 			if (InHairDescriptionGroup.Info.NumGuides > 0 && !InSettings.InterpolationSettings.bOverrideGuides)
 			{
-				HairStrandsBuilder::BuildInternalData(OutSim, true); // Imported guides don't currently have root UVs so force computing them
+				HairStrandsBuilder::BuildInternalData(OutSim);
 			}
 			else
 			{
@@ -1951,7 +1866,7 @@ void FGroomBuilder::BuildData(
 					{
 						// We pick the new guides among the imported ones
 						FHairStrandsDatas TempSim = InHairDescriptionGroup.Guides;
-						HairStrandsBuilder::BuildInternalData(TempSim, true);
+						HairStrandsBuilder::BuildInternalData(TempSim);
 						FHairStrandsDecimation::Decimate(TempSim, InSettings.RiggingSettings.NumCurves, InSettings.RiggingSettings.NumPoints, OutSim);
 					}
 					else
@@ -2152,7 +2067,7 @@ void Decimate(
 			FHairStrandsDatas::CopyPointLerp(InData, OutData, InAttributes, PointBegin, PointEnd, PointAlpha, PointCount);
 		}
 	}
-	HairStrandsBuilder::BuildInternalData(OutData, false);
+	HairStrandsBuilder::BuildInternalData(OutData);
 }
 	
 void Decimate(
@@ -2276,7 +2191,7 @@ void Decimate(
 		OutPointOffset += OutPointCount;
 	}
 
-	HairStrandsBuilder::BuildInternalData(OutData, false);
+	HairStrandsBuilder::BuildInternalData(OutData);
 }
 } // namespace FHairStrandsDecimation
 
@@ -2386,7 +2301,7 @@ static void Voxelize(const FHairDescriptionGroups& InGroups, FHairGrid& Out)
 	{
 		// Local copy of the hair strands data, and build the derived data so that curve offset are available
 		FHairStrandsDatas In = Group.Strands;
-		HairStrandsBuilder::BuildInternalData(In, false);
+		HairStrandsBuilder::BuildInternalData(In);
 
 		const uint32 Attributes = In.GetAttributes();
 
diff --git a/Engine/Plugins/Runtime/HairStrands/Source/HairStrandsCore/Private/GroomDebug.cpp b/Engine/Plugins/Runtime/HairStrands/Source/HairStrandsCore/Private/GroomDebug.cpp
--- a/Engine/Plugins/Runtime/HairStrands/Source/HairStrandsCore/Private/GroomDebug.cpp
+++ b/Engine/Plugins/Runtime/HairStrands/Source/HairStrandsCore/Private/GroomDebug.cpp
@@ -971,7 +971,7 @@ static void AddHairDebugPrintInstancePass(
 						const uint32 Index = GetHairAttributeIndex(Attribute);
 						if (Index < HAIR_ATTRIBUTE_COUNT)
 						{
-							Data2.W |= Instance->Strands.Data->AttributeOffsets[Index] != 0xFFFFFFFF ? (1u<<Index) : 0u;
+							Data2.W |= Instance->Strands.Data->AttributeOffsets[Index] != 0xFFFFFFFF ? (1u<<AttributeIt) : 0u;
 						}
 					}
 				}
@@ -1033,7 +1033,6 @@ static void AddHairDebugPrintInstancePass(
 		// Only add valid optional attribute 
 		// HAIR_ATTRIBUTE_XXX and EHairAttribute don't have a 1:1 mapping
 		const EHairAttribute Attribute = (EHairAttribute)AttributeIt;
-		if (GetHairAttributeIndex(Attribute) < HAIR_ATTRIBUTE_COUNT)
 		{
 			const FString Name = GetHairAttributeText(Attribute);
 			FHairDebugNameInfo& NameInfo = AttributeNameInfos.AddDefaulted_GetRef();
diff --git a/Engine/Plugins/Runtime/HairStrands/Source/HairStrandsCore/Private/GroomTextureBuilder.cpp b/Engine/Plugins/Runtime/HairStrands/Source/HairStrandsCore/Private/GroomTextureBuilder.cpp
--- a/Engine/Plugins/Runtime/HairStrands/Source/HairStrandsCore/Private/GroomTextureBuilder.cpp
+++ b/Engine/Plugins/Runtime/HairStrands/Source/HairStrandsCore/Private/GroomTextureBuilder.cpp
@@ -106,6 +106,12 @@ struct FPixel
 // CPU raster
 static void RasterToTexture(int32 Resolution, int32 KernelExtent, uint32 Channel, const FHairStrandsDatas& InStrandsData, FPixel* OutPixels)
 {
+	const uint32 Attribute = InStrandsData.GetAttributes();
+	if (!HasHairAttribute(Attribute, EHairAttribute::RootUV))
+	{
+		return;
+	}
+
 	const uint32 CurveCount = InStrandsData.GetNumCurves();
 	for (uint32 CurveIndex = 0; CurveIndex < CurveCount; ++CurveIndex)
 	{
@@ -312,7 +318,7 @@ static void InternalBuildFollicleTexture_GPU(
 		}
 
 		// Create root UVs buffer
-		if (Info.StrandsData)
+		if (Info.StrandsData && HasHairAttribute(Info.StrandsData->GetAttributes(), EHairAttribute::RootUV))
 		{
 			const uint32 DataCount = Info.StrandsData->StrandsCurves.CurvesRootUV.Num();
 			const uint32 DataSizeInBytes = sizeof(FVector2f) * DataCount;
@@ -330,7 +336,7 @@ static void InternalBuildFollicleTexture_GPU(
 		}
 	}
 
-	const EPixelFormat Format = bCopyDataBackToCPU ? PF_B8G8R8A8 : PF_R8G8B8A8;		 
+	const EPixelFormat Format = bCopyDataBackToCPU ? PF_B8G8R8A8 : PF_R8G8B8A8;
 	InternalGenerateFollicleTexture_GPU(GraphBuilder, ShaderMap, bCopyDataBackToCPU, Resolution, MipCount, Format, KernelSizeInPixels, RootUVBuffers[0], RootUVBuffers[1], RootUVBuffers[2], RootUVBuffers[3], OutTexture);
 }
 
diff --git a/Engine/Plugins/Runtime/HairStrands/Source/HairStrandsCore/Private/HairAttributes.cpp b/Engine/Plugins/Runtime/HairStrands/Source/HairStrandsCore/Private/HairAttributes.cpp
--- a/Engine/Plugins/Runtime/HairStrands/Source/HairStrandsCore/Private/HairAttributes.cpp
+++ b/Engine/Plugins/Runtime/HairStrands/Source/HairStrandsCore/Private/HairAttributes.cpp
@@ -92,6 +92,7 @@ bool FHairStrandsPoints::HasAttribute(EHairAttribute In) const
 
 bool FHairStrandsCurves::HasAttribute(EHairAttribute In) const
 {
+	if (In == EHairAttribute::RootUV)					{ return CurvesRootUV.Num() > 0; }
 	if (In == EHairAttribute::PrecomputedGuideWeights)	{ return CurvesClosestGuideIDs.Num() > 0 && CurvesClosestGuideWeights.Num() > 0; }
 	if (In == EHairAttribute::ClumpID)					{ return ClumpIDs.Num() > 0; }
 	if (In == EHairAttribute::StrandID)					{ return StrandIDs.Num() > 0; }
@@ -139,8 +140,11 @@ void FHairStrandsDatas::CopyPointLerp(const FHairStrandsDatas& In, FHairStrandsD
 
 void FHairStrandsDatas::CopyCurve(const FHairStrandsDatas& In, FHairStrandsDatas& Out, uint32 InAttributes, uint32 InIndex, uint32 OutIndex)
 {
-	Out.StrandsCurves.CurvesRootUV[OutIndex] = In.StrandsCurves.CurvesRootUV[InIndex];
 	Out.StrandsCurves.CurvesLength[OutIndex] = In.StrandsCurves.CurvesLength[InIndex];
+	if (HasHairAttribute(InAttributes, EHairAttribute::RootUV))
+	{
+		Out.StrandsCurves.CurvesRootUV[OutIndex] = In.StrandsCurves.CurvesRootUV[InIndex];
+	}
 	if (HasHairAttribute(InAttributes, EHairAttribute::PrecomputedGuideWeights))
 	{
 		Out.StrandsCurves.CurvesClosestGuideIDs[OutIndex] = In.StrandsCurves.CurvesClosestGuideIDs[InIndex];
diff --git a/Engine/Plugins/Runtime/HairStrands/Source/HairStrandsCore/Private/HairCardsBuilder.cpp b/Engine/Plugins/Runtime/HairStrands/Source/HairStrandsCore/Private/HairCardsBuilder.cpp
--- a/Engine/Plugins/Runtime/HairStrands/Source/HairStrandsCore/Private/HairCardsBuilder.cpp
+++ b/Engine/Plugins/Runtime/HairStrands/Source/HairStrandsCore/Private/HairCardsBuilder.cpp
@@ -1360,6 +1360,8 @@ namespace HairCards
 		OutData.StrandsPoints.SetNum(OutTotalPointCount, 0u);
 		OutData.HairDensity = InData.HairDensity;
 
+		const uint32 InAttribute = InData.GetAttributes();
+
 		uint32 OutPointOffset = 0;
 		for (uint32 OutCurveIndex = 0; OutCurveIndex < OutCurveCount; ++OutCurveIndex)
 		{
@@ -1367,21 +1369,14 @@ namespace HairCards
 			const uint32 InPointOffset = InData.StrandsCurves.CurvesOffset[InCurveIndex];
 			const uint32 PointCount = InData.StrandsCurves.CurvesCount[InCurveIndex];
 			OutData.StrandsCurves.CurvesCount[OutCurveIndex] = PointCount;
-			OutData.StrandsCurves.CurvesRootUV[OutCurveIndex] = InData.StrandsCurves.CurvesRootUV[InCurveIndex];
 			OutData.StrandsCurves.CurvesOffset[OutCurveIndex] = OutPointOffset;
-			OutData.StrandsCurves.CurvesLength[OutCurveIndex] = InData.StrandsCurves.CurvesLength[InCurveIndex] * InData.StrandsCurves.MaxLength;
 			OutData.StrandsCurves.MaxLength = InData.StrandsCurves.MaxLength;
 			OutData.StrandsCurves.MaxRadius = InData.StrandsCurves.MaxRadius;
-			if (bHasClumpID)
-			{
-				OutData.StrandsCurves.CurvesOffset[OutCurveIndex] = InData.StrandsCurves.ClumpIDs[InCurveIndex];
-			}
+			FHairStrandsDatas::CopyCurve(InData, OutData, InAttribute, InCurveIndex, OutCurveIndex);
 
 			for (uint32 PointIndex = 0; PointIndex < PointCount; ++PointIndex)
 			{
-				OutData.StrandsPoints.PointsPosition[PointIndex + OutPointOffset] = InData.StrandsPoints.PointsPosition[PointIndex + InPointOffset];
-				OutData.StrandsPoints.PointsCoordU[PointIndex + OutPointOffset] = InData.StrandsPoints.PointsCoordU[PointIndex + InPointOffset];
-				OutData.StrandsPoints.PointsRadius[PointIndex + OutPointOffset] = InData.StrandsPoints.PointsRadius[PointIndex + InPointOffset];
+				FHairStrandsDatas::CopyPoint(InData, OutData, InAttribute, PointIndex + InPointOffset, PointIndex + OutPointOffset);
 			}
 			OutPointOffset += PointCount;
 		}
@@ -1637,6 +1632,8 @@ namespace HairCards
 			}			
 		}
 
+		const uint32 InAttribute = SimStrandsData.GetAttributes();
+
 		// Store the cluster guide points
 		uint32 ClusterOffset = 0;
 		OutClusters.GuidePoints.SetNum(SimStrandsData.GetNumPoints());
@@ -1644,7 +1641,7 @@ namespace HairCards
 		{
 			const uint32 SimCount  = SimStrandsData.StrandsCurves.CurvesCount[GuideIndex];
 			const uint32 SimOffset = SimStrandsData.StrandsCurves.CurvesOffset[GuideIndex];
-			const FVector2f SimRootUV = SimStrandsData.StrandsCurves.CurvesRootUV[GuideIndex];
+			const FVector2f SimRootUV = HasHairAttribute(InAttribute, EHairAttribute::RootUV) ? SimStrandsData.StrandsCurves.CurvesRootUV[GuideIndex] : FVector2f(0,0);
 
 			const uint32 ClusterIt = GuideIndexToClusterIndex[GuideIndex];
 			if (GuideIndexToClusterIndex[GuideIndex] >= 0)
@@ -2786,7 +2783,7 @@ namespace HairCards
 			OutGuides.StrandsCurves.CurvesCount[ClusterIt] = PointCount;
 			OutGuides.StrandsCurves.CurvesOffset[ClusterIt] = PointOffset;
 			OutGuides.StrandsCurves.CurvesLength[ClusterIt] = TotalLength;
-			OutGuides.StrandsCurves.CurvesRootUV[ClusterIt] = FVector2f(0, 0);
+			//OutGuides.StrandsCurves.CurvesRootUV[ClusterIt] = FVector2f(0, 0);
 			//OutGuides.StrandsCurves.StrandIDs[ClusterIt] = ClusterIt;
 			//OutGuides.StrandsCurves.GroomIDToIndex[ClusterIt] = ;
 			OutGuides.StrandsCurves.MaxLength = FMath::Max(OutGuides.StrandsCurves.MaxLength, TotalLength);
@@ -2801,8 +2798,8 @@ namespace HairCards
 				OutGuides.BoundingBox += (FVector)P0;
 
 				OutGuides.StrandsPoints.PointsPosition[PointIndex] = P0;
-				OutGuides.StrandsPoints.PointsBaseColor[PointIndex] = FLinearColor(FVector3f::ZeroVector);
-				OutGuides.StrandsPoints.PointsRoughness[PointIndex] = 0;
+				//OutGuides.StrandsPoints.PointsBaseColor[PointIndex] = FLinearColor(FVector3f::ZeroVector);
+				//OutGuides.StrandsPoints.PointsRoughness[PointIndex] = 0;
 				OutGuides.StrandsPoints.PointsCoordU[PointIndex] = FMath::Clamp(CurrentLength / TotalLength, 0.f, 1.f);
 				OutGuides.StrandsPoints.PointsRadius[PointIndex] = 1;
 
@@ -3079,7 +3076,7 @@ namespace HairCards
 			OutGuides.StrandsCurves.CurvesCount[CardIt] = PointCount;
 			OutGuides.StrandsCurves.CurvesOffset[CardIt] = PointOffset;
 			OutGuides.StrandsCurves.CurvesLength[CardIt] = TotalLength;
-			OutGuides.StrandsCurves.CurvesRootUV[CardIt] = FVector2f(0, 0);
+			if (OutGuides.StrandsCurves.CurvesRootUV.Num()) OutGuides.StrandsCurves.CurvesRootUV[CardIt] = FVector2f(0, 0);
 			OutGuides.StrandsCurves.MaxLength = FMath::Max(OutGuides.StrandsCurves.MaxLength, TotalLength);
 			OutGuides.StrandsCurves.MaxRadius = GuideRadius;
 
@@ -3273,7 +3270,7 @@ namespace HairCards
 					const uint32 PointIndex = PointOffset + PointIt;
 
 					const float StrandsU = InStrands.StrandsPoints.PointsCoordU[PointIndex];
-					const FVector2f RootUV = InStrands.StrandsCurves.CurvesRootUV[CurveIndex];
+					const FVector2f RootUV = InStrands.StrandsCurves.CurvesRootUV.Num() ? InStrands.StrandsCurves.CurvesRootUV[CurveIndex] : FVector2f(0,0);
 					const float Seed = 0.f; // TODO float(InStrands.RenderData.Attributes[PointIndex].Seed) / 255.f;
 
 					const FVector3f& P0_Rest = InStrands.StrandsPoints.PointsPosition[PointIndex];
@@ -3949,6 +3946,7 @@ bool InternalImportGeometry(
 		};
 		TArray<FStrandsRootData> StrandsRoots;
 		{
+			const uint32 InAttribute = InStrandsData.GetAttributes();
 			const uint32 CurveCount = InStrandsData.StrandsCurves.Num();
 			StrandsRoots.Reserve(CurveCount);
 			for (uint32 CurveIt = 0; CurveIt < CurveCount; ++CurveIt)
@@ -3956,7 +3954,7 @@ bool InternalImportGeometry(
 				const uint32 Offset = InStrandsData.StrandsCurves.CurvesOffset[CurveIt];
 				FStrandsRootData& RootData = StrandsRoots.AddDefaulted_GetRef();
 				RootData.Position = InStrandsData.StrandsPoints.PointsPosition[Offset];
-				RootData.RootUV = InStrandsData.StrandsCurves.CurvesRootUV[CurveIt];
+				RootData.RootUV = HasHairAttribute(InAttribute, EHairAttribute::RootUV) ? InStrandsData.StrandsCurves.CurvesRootUV[CurveIt] : FVector2f(0,0);
 				RootData.CurveIndex = CurveIt;
 			}
 		}
diff --git a/Engine/Plugins/Runtime/HairStrands/Source/HairStrandsCore/Private/HairStrandsDatas.cpp b/Engine/Plugins/Runtime/HairStrands/Source/HairStrandsCore/Private/HairStrandsDatas.cpp
--- a/Engine/Plugins/Runtime/HairStrands/Source/HairStrandsCore/Private/HairStrandsDatas.cpp
+++ b/Engine/Plugins/Runtime/HairStrands/Source/HairStrandsCore/Private/HairStrandsDatas.cpp
@@ -18,9 +18,9 @@ void FHairStrandsCurves::SetNum(const uint32 NumCurves, uint32 InAttributes)
 	CurvesOffset.SetNum(NumCurves + 1);
 	CurvesCount.SetNum(NumCurves);
 	CurvesLength.SetNum(NumCurves);
-	CurvesRootUV.SetNum(NumCurves);
 
 	// Not initialized to track if the data are available
+	if (HasHairAttribute(InAttributes, EHairAttribute::RootUV))		{ CurvesRootUV.SetNum(NumCurves); }
 	if (HasHairAttribute(InAttributes, EHairAttribute::StrandID))	{ StrandIDs.SetNum(NumCurves); }
 	if (HasHairAttribute(InAttributes, EHairAttribute::ClumpID))	{ ClumpIDs.SetNum(NumCurves); }
 	if (HasHairAttribute(InAttributes, EHairAttribute::PrecomputedGuideWeights))
diff --git a/Engine/Plugins/Runtime/HairStrands/Source/HairStrandsCore/Public/HairStrandsDatas.h b/Engine/Plugins/Runtime/HairStrands/Source/HairStrandsCore/Public/HairStrandsDatas.h
--- a/Engine/Plugins/Runtime/HairStrands/Source/HairStrandsCore/Public/HairStrandsDatas.h
+++ b/Engine/Plugins/Runtime/HairStrands/Source/HairStrandsCore/Public/HairStrandsDatas.h
@@ -398,10 +398,10 @@ struct HAIRSTRANDSCORE_API FHairStrandsCurves
 	/** Normalized length relative to the max one */
 	TArray<float> CurvesLength; // [0..1]
 
-	/** Roots UV. Support UDIM coordinate up to 256x256 */
+	/** Roots UV. Support UDIM coordinate up to 256x256 (optional) */
 	TArray<FVector2f> CurvesRootUV; // [0..256]
 
-	/** Strand ID associated with each curve */
+	/** Strand ID associated with each curve (optional) */
 	TArray<int> StrandIDs;
 
 	/** Clump ID associated with each curve (optional) */
diff --git a/Engine/Shaders/Private/HairStrands/HairStrandsAttributeCommon.ush b/Engine/Shaders/Private/HairStrands/HairStrandsAttributeCommon.ush
--- a/Engine/Shaders/Private/HairStrands/HairStrandsAttributeCommon.ush
+++ b/Engine/Shaders/Private/HairStrands/HairStrandsAttributeCommon.ush
@@ -109,8 +109,15 @@ float2 GetHairStrandsDimensions(uint HairPrimitiveId, float2 HairPrimitiveUV)
 
 float2 GetHairStrandsRootUV(uint HairPrimitiveId, float2 HairPrimitiveUV)
 {
-	// Encoding need to match GroomBuilder.cpp - 
-	return UnpackHairRootUV(InternalGetHairStrandsAttributesPerCurve(HairPrimitiveId, HAIR_ATTRIBUTE_OFFSET_ROOTUV(HairStrandsVF.AttributeOffsets), HAIR_ATTRIBUTE_STRIDE_ROOTUV, HairStrandsVF.AttributeBuffer, HairStrandsVF.PointToCurveBuffer));
+	const uint RootUVOffset = HAIR_ATTRIBUTE_OFFSET_ROOTUV(HairStrandsVF.AttributeOffsets);
+	if (IsAttributeValid(RootUVOffset))
+	{
+		return UnpackHairRootUV(InternalGetHairStrandsAttributesPerCurve(HairPrimitiveId, RootUVOffset, HAIR_ATTRIBUTE_STRIDE_ROOTUV, HairStrandsVF.AttributeBuffer, HairStrandsVF.PointToCurveBuffer));
+	}
+	else
+	{
+		return 0;
+	}
 }
 
 float GetHairStrandsSeed(uint HairPrimitiveId, float2 HairPrimitiveUV)
