* Remove MaxRadius/MaxLength from FHairStrandsDatas, and keep that logic only for compressing bulk data.
* Fix imported groom's width not being respected.

diff --git a/Engine/Plugins/Runtime/HairStrands/Source/HairStrandsCore/Private/GroomAsset.cpp b/Engine/Plugins/Runtime/HairStrands/Source/HairStrandsCore/Private/GroomAsset.cpp
--- a/Engine/Plugins/Runtime/HairStrands/Source/HairStrandsCore/Private/GroomAsset.cpp
+++ b/Engine/Plugins/Runtime/HairStrands/Source/HairStrandsCore/Private/GroomAsset.cpp
@@ -719,6 +719,15 @@ void UGroomAsset::Serialize(FArchive& Ar)
 	Ar.UsingCustomVersion(FUE5MainStreamObjectVersion::GUID);	// Needed to support MeshDescription AttributesSet serialization
 	Ar.UsingCustomVersion(FAnimObjectVersion::GUID);    // Needed to support Cards and Cluster culling serialization
 
+	// For older assets, force width override to preserve prior to fixing the imported groom's width.
+	if (Ar.CustomVer(FUE5MainStreamObjectVersion::GUID) < FUE5MainStreamObjectVersion::GroomAssetWidthOverride && Ar.IsLoading())
+	{
+		for (FHairGroupsRendering& Group : HairGroupsRendering)
+		{
+			Group.GeometrySettings.HairWidth_Override = true;
+		}
+	}
+
 	if (Ar.CustomVer(FPhysicsObjectVersion::GUID) >= FPhysicsObjectVersion::GroomWithDescription)
 	{
 		FStripDataFlags StripFlags(Ar, ClassDataStripFlags);
diff --git a/Engine/Plugins/Runtime/HairStrands/Source/HairStrandsCore/Private/GroomBuilder.cpp b/Engine/Plugins/Runtime/HairStrands/Source/HairStrandsCore/Private/GroomBuilder.cpp
--- a/Engine/Plugins/Runtime/HairStrands/Source/HairStrandsCore/Private/GroomBuilder.cpp
+++ b/Engine/Plugins/Runtime/HairStrands/Source/HairStrandsCore/Private/GroomBuilder.cpp
@@ -40,7 +40,7 @@ static FAutoConsoleVariableRef CVarHairGroupIndexBuilder_MaxVoxelResolution(TEXT
 
 FString FGroomBuilder::GetVersion()
 {
-	return TEXT("v8");
+	return TEXT("v8a");
 }
 
 namespace FHairStrandsDecimation
@@ -83,9 +83,6 @@ namespace HairStrandsBuilder
 			TArray<uint32>::TIterator OffsetIterator = Curves.CurvesOffset.CreateIterator();
 			TArray<float>::TIterator LengthIterator = Curves.CurvesLength.CreateIterator();
 
-			Curves.MaxRadius = 0.0;
-			Curves.MaxLength = 0.0;
-
 			uint32 StrandOffset = 0;
 			*OffsetIterator = StrandOffset; ++OffsetIterator;
 
@@ -108,11 +105,8 @@ namespace HairStrandsBuilder
 					}
 					*CoordUIterator = StrandLength;
 					PreviousPosition = (FVector)*PositionIterator;
-
-					Curves.MaxRadius = FMath::Max(Curves.MaxRadius, *RadiusIterator);
 				}
 				*LengthIterator = StrandLength;
-				Curves.MaxLength = FMath::Max(Curves.MaxLength, StrandLength);
 			}
 
 			CountIterator.Reset();
@@ -130,14 +124,6 @@ namespace HairStrandsBuilder
 					{
 						*CoordUIterator /= *LengthIterator;
 					}
-					if (Curves.MaxRadius > 0.0f)
-					{
-						*RadiusIterator /= Curves.MaxRadius;
-					}
-				}
-				if (Curves.MaxLength > 0.0f)
-				{
-					*LengthIterator /= Curves.MaxLength;
 				}
 			}
 		}
@@ -215,6 +201,9 @@ namespace HairStrandsBuilder
 		const FHairStrandsCurves& Curves = HairStrands.StrandsCurves;
 		const FHairStrandsPoints& Points = HairStrands.StrandsPoints;
 
+		const float MaxLength = GetHairStrandsMaxLength(HairStrands);
+		const float MaxRadius = GetHairStrandsMaxRadius(HairStrands);
+
 		struct FPackedRadiusAndType
 		{
 			union
@@ -242,7 +231,7 @@ namespace HairStrandsBuilder
 				// Packed Position|CoordU|NormalizedRadius
 				{
 					const float CoordU = Points.PointsCoordU[PointIndex + IndexOffset];
-					const float NormalizedRadius = Points.PointsRadius[PointIndex + IndexOffset];
+					const float NormalizedRadius = Points.PointsRadius[PointIndex + IndexOffset] / MaxRadius;
 
 					FHairStrandsPositionFormat::Type& PackedPosition = OutPackedPositions[PointIndex + IndexOffset];
 					CopyVectorToPosition(PointPosition - HairBoxCenter, PackedPosition);
@@ -289,9 +278,8 @@ namespace HairStrandsBuilder
 
 			// Per-Curve Length
 			{
-				const float Length = Curves.CurvesLength[CurveIndex] * Curves.MaxLength;
 				FFloat16* Packed = (FFloat16*)AttributeLength.GetData();
-				Packed[CurveIndex] = Length;
+				Packed[CurveIndex] = Curves.CurvesLength[CurveIndex];
 			}
 
 			// Per-Curve Seed
@@ -347,8 +335,8 @@ namespace HairStrandsBuilder
 		OutBulkData.BoundingBox = HairStrands.BoundingBox;
 		OutBulkData.CurveCount = HairStrands.GetNumCurves();
 		OutBulkData.PointCount = HairStrands.GetNumPoints();
-		OutBulkData.MaxLength = HairStrands.StrandsCurves.MaxLength;
-		OutBulkData.MaxRadius = HairStrands.StrandsCurves.MaxRadius;
+		OutBulkData.MaxLength = MaxLength;
+		OutBulkData.MaxRadius = MaxRadius;
 		OutBulkData.Flags = FHairStrandsBulkData::DataFlags_HasData;
 
 		// Concatenate all attributes
@@ -503,8 +491,7 @@ namespace HairInterpolationBuilder
 		const FHairStrandsDatas& GuideCurvesDatas, const uint32 GuideCurveIndex, const float RootImportance, 
 		const float ShapeImportance, const float ProximityImportance)
 	{
-		const float AverageLength = FMath::Max(0.5f * (RenderCurvesDatas.StrandsCurves.CurvesLength[RenderCurveIndex] * RenderCurvesDatas.StrandsCurves.MaxLength +
-			GuideCurvesDatas.StrandsCurves.CurvesLength[GuideCurveIndex] * GuideCurvesDatas.StrandsCurves.MaxLength), SMALL_NUMBER);
+		const float AverageLength = FMath::Max(0.5f * (RenderCurvesDatas.StrandsCurves.CurvesLength[RenderCurveIndex] + GuideCurvesDatas.StrandsCurves.CurvesLength[GuideCurveIndex]), SMALL_NUMBER);
 
 		static const float DeltaCoord = 1.0f / static_cast<float>(NumSamples-1);
 
@@ -1027,7 +1014,7 @@ namespace HairInterpolationBuilder
 		{
 			const uint32 PointOffset = InData.StrandsCurves.CurvesOffset[CurveIndex];
 			const uint32 PointCount = InData.StrandsCurves.CurvesCount[CurveIndex];
-			const float  CurveLength = InData.StrandsCurves.CurvesLength[CurveIndex] * InData.StrandsCurves.MaxLength;
+			const float  CurveLength = InData.StrandsCurves.CurvesLength[CurveIndex];
 			check(PointCount > 1);
 			const FVector& P0 = (FVector)InData.StrandsPoints.PointsPosition[PointOffset];
 			const FVector& P1 = (FVector)InData.StrandsPoints.PointsPosition[PointOffset + 1];
@@ -1057,7 +1044,7 @@ namespace HairInterpolationBuilder
 	{
 		const uint32 SimOffset = SimStrandsData.StrandsCurves.CurvesOffset[SimCurveIndex];
 
-		const float CurveLength = SimStrandsData.StrandsCurves.CurvesLength[SimCurveIndex] * SimStrandsData.StrandsCurves.MaxLength;
+		const float CurveLength = SimStrandsData.StrandsCurves.CurvesLength[SimCurveIndex];
 
 		// Find with with vertex the vertex should be paired
 		const uint32 SimPointCount = SimStrandsData.StrandsCurves.CurvesCount[SimCurveIndex];
@@ -1223,7 +1210,7 @@ namespace HairInterpolationBuilder
 			{
 				const uint32 PointGlobalIndex = RenPointIndex + RenOffset;
 				const FVector& RenPointPosition = (FVector)RenStrandsData.StrandsPoints.PointsPosition[PointGlobalIndex];
-				const float RenPointDistance = RenStrandsData.StrandsPoints.PointsCoordU[PointGlobalIndex] * RenStrandsData.StrandsCurves.CurvesLength[RenCurveIndex] * RenStrandsData.StrandsCurves.MaxLength;
+				const float RenPointDistance = RenStrandsData.StrandsPoints.PointsCoordU[PointGlobalIndex] * RenStrandsData.StrandsCurves.CurvesLength[RenCurveIndex];
 
 				float TotalWeight = 0;
 				for (uint32 KIndex = 0; KIndex < FClosestGuides::Count; ++KIndex)
@@ -1414,7 +1401,7 @@ namespace HairInterpolationBuilder
 			{
 				const uint32 PointGlobalIndex = VertexIndex + CurveOffset;
 				const FVector& RenPointPosition = (FVector)RenData.StrandsPoints.PointsPosition[PointGlobalIndex];
-				const float RenPointDistance = RenData.StrandsPoints.PointsCoordU[PointGlobalIndex] * RenData.StrandsCurves.CurvesLength[CurveIndex] * RenData.StrandsCurves.MaxLength;
+				const float RenPointDistance = RenData.StrandsPoints.PointsCoordU[PointGlobalIndex] * RenData.StrandsCurves.CurvesLength[CurveIndex];
 
 				bool bHasValidGuide = false;
 				float TotalWeight = 0;
@@ -1570,7 +1557,6 @@ bool FGroomBuilder::BuildHairDescriptionGroups(const FHairDescription& HairDescr
 		FHairDescriptionGroup& Group = Out.HairGroups.AddDefaulted_GetRef();
 		Group.Info.GroupID = GroupID;
 		Group.Info.GroupName = (GroupName != NAME_None) ? GroupName : FName(FString::Printf(TEXT("Group_%d"), GroupID));
-		Group.Info.MaxImportedWidth = -1.0f; // Invalid width
 		return Group;
 	};
 
@@ -1675,20 +1661,12 @@ bool FGroomBuilder::BuildHairDescriptionGroups(const FHairDescription& HairDescr
 		if (GroomHairWidth) 
 		{
 			StrandWidth = GroomHairWidth.GetValue();
-			if (!bIsGuide)
-			{
-				Group.Info.MaxImportedWidth = FMath::Max(Group.Info.MaxImportedWidth, StrandWidth);
-			}
 		}
 
 		// Curve
 		if (StrandWidths.IsValid())
 		{
 			StrandWidth = StrandWidths[StrandID];
-			if (!bIsGuide)
-			{
-				Group.Info.MaxImportedWidth = FMath::Max(Group.Info.MaxImportedWidth, StrandWidth);
-			}
 		}
 
 		for (int32 VertexIndex = 0; VertexIndex < CurveNumVertices; ++VertexIndex, ++GlobalVertexIndex)
@@ -1715,10 +1693,6 @@ bool FGroomBuilder::BuildHairDescriptionGroups(const FHairDescription& HairDescr
 			if (VertexWidths.IsValid())
 			{
 				VertexWidth = VertexWidths[VertexID];
-				if (!bIsGuide)
-				{
-					Group.Info.MaxImportedWidth = FMath::Max(Group.Info.MaxImportedWidth, VertexWidth);
-				}
 			}
 			else if (StrandWidth != 0.f)
 			{
@@ -1824,7 +1798,6 @@ void FGroomBuilder::BuildData(
 	{
 		OutGroupInfo.GroupID = InHairDescriptionGroup.Info.GroupID;
 		OutGroupInfo.GroupName = InHairDescriptionGroup.Info.GroupName;
-		OutGroupInfo.MaxImportedWidth = InHairDescriptionGroup.Info.MaxImportedWidth;
 
 		// Rendering data
 		if (InHairDescriptionGroup.Strands.IsValid())
@@ -1842,7 +1815,7 @@ void FGroomBuilder::BuildData(
 			}
 			OutGroupInfo.NumCurves			= OutRen.GetNumCurves();
 			OutGroupInfo.NumCurveVertices	= OutRen.GetNumPoints();
-			OutGroupInfo.MaxCurveLength		= OutRen.StrandsCurves.MaxLength;
+			OutGroupInfo.MaxCurveLength		= GetHairStrandsMaxLength(OutRen);
 
 			// Sanity check
 			check(OutGroupInfo.NumCurves		<= InHairDescriptionGroup.Info.NumCurves);
@@ -2039,8 +2012,6 @@ void Decimate(
 	OutData.StrandsPoints.SetNum(PointCount, InAttributes);
 	OutData.HairDensity = InData.HairDensity;
 
-	check(InData.StrandsCurves.MaxRadius > 0);
-
 	PointCount = 0;
 	for(uint32 CurveIndex = 0; CurveIndex < CurveCount; ++CurveIndex)
 	{
@@ -2052,8 +2023,6 @@ void Decimate(
 		const float BoneSpace = FMath::Max(1.0,float(GuideCount-1) / (MaxPoints-1));
 
 		OutData.StrandsCurves.CurvesCount[CurveIndex]  = MaxPoints;
-		OutData.StrandsCurves.MaxLength = InData.StrandsCurves.MaxLength;
-		OutData.StrandsCurves.MaxRadius = InData.StrandsCurves.MaxRadius;
 		FHairStrandsDatas::CopyCurve(InData, OutData, InAttributes, GuideIndex, CurveIndex);
 	
 		for(int32 PointIndex = 0; PointIndex < MaxPoints; ++PointIndex, ++PointCount)
@@ -2085,7 +2054,6 @@ void Decimate(
 	const uint32 InCurveCount = InData.StrandsCurves.Num();
 	const uint32 OutCurveCount = FMath::Clamp(uint32(InCurveCount * CurveDecimationPercentage), 1u, InCurveCount);
 	const uint32 InAttributes = InData.GetAttributes();
-	check(InData.StrandsCurves.MaxRadius > 0);
 
 	FRandomStream Random;
 	Random.Initialize(0xdeedbeed);
@@ -2184,8 +2152,6 @@ void Decimate(
 
 		OutData.StrandsCurves.CurvesCount[OutCurveIndex]  = OutPointCount;
 		OutData.StrandsCurves.CurvesOffset[OutCurveIndex] = OutPointOffset;
-		OutData.StrandsCurves.MaxLength = InData.StrandsCurves.MaxLength;
-		OutData.StrandsCurves.MaxRadius = InData.StrandsCurves.MaxRadius;
 		FHairStrandsDatas::CopyCurve(InData, OutData, InAttributes, EffectiveCurveIndex, OutCurveIndex);
 
 		OutPointOffset += OutPointCount;
@@ -2755,7 +2721,7 @@ static void BuildClusterData(
 				RCurve.Area += InRenStrandsData.StrandsPoints.PointsRadius[PointGlobalIndex] * OutLength;
 			}
 
-			const float PointRadius = InRenStrandsData.StrandsPoints.PointsRadius[PointGlobalIndex] * InRenStrandsData.StrandsCurves.MaxRadius;
+			const float PointRadius = InRenStrandsData.StrandsPoints.PointsRadius[PointGlobalIndex];
 			RCurve.AvgRadius += PointRadius;
 			RCurve.MaxRadius = FMath::Max(RCurve.MaxRadius, PointRadius);
 		}
diff --git a/Engine/Plugins/Runtime/HairStrands/Source/HairStrandsCore/Private/GroomCacheData.cpp b/Engine/Plugins/Runtime/HairStrands/Source/HairStrandsCore/Private/GroomCacheData.cpp
--- a/Engine/Plugins/Runtime/HairStrands/Source/HairStrandsCore/Private/GroomCacheData.cpp
+++ b/Engine/Plugins/Runtime/HairStrands/Source/HairStrandsCore/Private/GroomCacheData.cpp
@@ -55,10 +55,10 @@ void FGroomCacheVertexData::Serialize(FArchive& Ar, int32 Version, EGroomCacheAt
 	}
 }
 
-FGroomCacheStrandData::FGroomCacheStrandData(FHairStrandsCurves&& CurvesData)
+FGroomCacheStrandData::FGroomCacheStrandData(FHairStrandsCurves&& CurvesData, float InMaxLength, float InMaxRadius)
 : CurvesLength(MoveTemp(CurvesData.CurvesLength))
-, MaxLength(CurvesData.MaxLength)
-, MaxRadius(CurvesData.MaxRadius)
+, MaxLength(InMaxLength)
+, MaxRadius(InMaxRadius)
 {
 }
 
@@ -71,7 +71,7 @@ void FGroomCacheStrandData::Serialize(FArchive& Ar, int32 Version, EGroomCacheAt
 
 FGroomCacheGroupData::FGroomCacheGroupData(FHairStrandsDatas&& GroupData)
 : VertexData(MoveTemp(GroupData.StrandsPoints))
-, StrandData(MoveTemp(GroupData.StrandsCurves))
+, StrandData(MoveTemp(GroupData.StrandsCurves), GetHairStrandsMaxLength(GroupData), GetHairStrandsMaxRadius(GroupData))
 , BoundingBox(GroupData.BoundingBox)
 {
 }
diff --git a/Engine/Plugins/Runtime/HairStrands/Source/HairStrandsCore/Private/GroomComponent.cpp b/Engine/Plugins/Runtime/HairStrands/Source/HairStrandsCore/Private/GroomComponent.cpp
--- a/Engine/Plugins/Runtime/HairStrands/Source/HairStrandsCore/Private/GroomComponent.cpp
+++ b/Engine/Plugins/Runtime/HairStrands/Source/HairStrandsCore/Private/GroomComponent.cpp
@@ -135,7 +135,18 @@ static FHairGroupDesc GetGroomGroupsDesc(const UGroomAsset* Asset, UGroomCompone
 	O.HairLength = Asset->HairGroupsData[GroupIndex].Strands.BulkData.MaxLength;
 	O.LODBias 	 = Asset->EffectiveLODBias[GroupIndex] > 0 ? FMath::Max(O.LODBias, Asset->EffectiveLODBias[GroupIndex]) : O.LODBias;
 
-	if (!O.HairWidth_Override)					{ O.HairWidth					= Asset->HairGroupsRendering[GroupIndex].GeometrySettings.HairWidth;					}
+	// Width can be overridden at 3 levels:
+	// * Instance width override
+	// * Asset width override
+	// * Imported width
+	if (!O.HairWidth_Override)
+	{ 
+		O.HairWidth = 
+			Asset->HairGroupsRendering[GroupIndex].GeometrySettings.HairWidth_Override ? 
+			Asset->HairGroupsRendering[GroupIndex].GeometrySettings.HairWidth : 
+			Asset->HairGroupsData[GroupIndex].Strands.BulkData.MaxRadius * 2.0f;
+	}
+
 	if (!O.HairRootScale_Override)				{ O.HairRootScale				= Asset->HairGroupsRendering[GroupIndex].GeometrySettings.HairRootScale;				}
 	if (!O.HairTipScale_Override)				{ O.HairTipScale				= Asset->HairGroupsRendering[GroupIndex].GeometrySettings.HairTipScale;					}
 	if (!O.bSupportVoxelization_Override)		{ O.bSupportVoxelization		= Asset->HairGroupsRendering[GroupIndex].ShadowSettings.bVoxelize;						}
diff --git a/Engine/Plugins/Runtime/HairStrands/Source/HairStrandsCore/Private/GroomEdit.cpp b/Engine/Plugins/Runtime/HairStrands/Source/HairStrandsCore/Private/GroomEdit.cpp
--- a/Engine/Plugins/Runtime/HairStrands/Source/HairStrandsCore/Private/GroomEdit.cpp
+++ b/Engine/Plugins/Runtime/HairStrands/Source/HairStrandsCore/Private/GroomEdit.cpp
@@ -65,7 +65,7 @@ static void HairStrandsDataToEditableHairStrands(const FHairStrandsDatas& In, TA
 			const uint32 EffectiveIndex = PointIt + PointOffset;
 			Out[StrandsIt].ControlPoints[PointIt].Position	= In.StrandsPoints.PointsPosition[EffectiveIndex];
 			Out[StrandsIt].ControlPoints[PointIt].U			= In.StrandsPoints.PointsCoordU[EffectiveIndex];
-			Out[StrandsIt].ControlPoints[PointIt].Radius	= In.StrandsPoints.PointsRadius[EffectiveIndex] * In.StrandsCurves.MaxRadius;
+			Out[StrandsIt].ControlPoints[PointIt].Radius	= In.StrandsPoints.PointsRadius[EffectiveIndex];
 			Out[StrandsIt].ControlPoints[PointIt].BaseColor	= In.StrandsPoints.PointsBaseColor[EffectiveIndex];
 			Out[StrandsIt].ControlPoints[PointIt].Roughness	= In.StrandsPoints.PointsRoughness[EffectiveIndex];
 			Out[StrandsIt].ControlPoints[PointIt].AO		= In.StrandsPoints.PointsAO[EffectiveIndex];
diff --git a/Engine/Plugins/Runtime/HairStrands/Source/HairStrandsCore/Private/HairAttributes.cpp b/Engine/Plugins/Runtime/HairStrands/Source/HairStrandsCore/Private/HairAttributes.cpp
--- a/Engine/Plugins/Runtime/HairStrands/Source/HairStrandsCore/Private/HairAttributes.cpp
+++ b/Engine/Plugins/Runtime/HairStrands/Source/HairStrandsCore/Private/HairAttributes.cpp
@@ -99,6 +99,11 @@ bool FHairStrandsCurves::HasAttribute(EHairAttribute In) const
 	return false;
 }
 
+bool FHairStrandsCurves::HasPrecomputedWeights() const
+{ 
+	return HasAttribute(EHairAttribute::PrecomputedGuideWeights); 
+}
+
 uint32 FHairStrandsDatas::GetAttributes() const
 {
 	uint32 Out = 0;
@@ -118,7 +123,7 @@ void FHairStrandsDatas::CopyPoint(const FHairStrandsDatas& In, FHairStrandsDatas
 	// Multiply by MaxRadius as HairStrandsBuilder::BuildInternalData will recompute MaxRadius based on the actual Radius not the not the normalized radius.
 	Out.StrandsPoints.PointsPosition[OutIndex]	= In.StrandsPoints.PointsPosition[InIndex];
 	Out.StrandsPoints.PointsCoordU[OutIndex]	= In.StrandsPoints.PointsCoordU[InIndex];
-	Out.StrandsPoints.PointsRadius[OutIndex]	= In.StrandsPoints.PointsRadius[InIndex] * In.StrandsCurves.MaxRadius;
+	Out.StrandsPoints.PointsRadius[OutIndex]	= In.StrandsPoints.PointsRadius[InIndex];
 
 	if (HasHairAttribute(InAttributes, EHairAttribute::Color))		{ Out.StrandsPoints.PointsBaseColor[OutIndex] = In.StrandsPoints.PointsBaseColor[InIndex]; }
 	if (HasHairAttribute(InAttributes, EHairAttribute::Roughness))	{ Out.StrandsPoints.PointsRoughness[OutIndex] = In.StrandsPoints.PointsRoughness[InIndex]; }
@@ -131,7 +136,7 @@ void FHairStrandsDatas::CopyPointLerp(const FHairStrandsDatas& In, FHairStrandsD
 	// Multiply by MaxRadius as HairStrandsBuilder::BuildInternalData will recompute MaxRadius based on the actual Radius not the not the normalized radius.
 	Out.StrandsPoints.PointsPosition[OutIndex]	= FMath::Lerp(In.StrandsPoints.PointsPosition[InIndex0], In.StrandsPoints.PointsPosition[InIndex1], InAlpha);
 	Out.StrandsPoints.PointsCoordU[OutIndex]	= FMath::Lerp(In.StrandsPoints.PointsCoordU[InIndex0], In.StrandsPoints.PointsCoordU[InIndex1], InAlpha);
-	Out.StrandsPoints.PointsRadius[OutIndex]	= FMath::Lerp(In.StrandsPoints.PointsRadius[InIndex0] , In.StrandsPoints.PointsRadius[InIndex1], InAlpha) * In.StrandsCurves.MaxRadius;
+	Out.StrandsPoints.PointsRadius[OutIndex]	= FMath::Lerp(In.StrandsPoints.PointsRadius[InIndex0], In.StrandsPoints.PointsRadius[InIndex1], InAlpha);
 
 	if (HasHairAttribute(InAttributes, EHairAttribute::Color))		{ Out.StrandsPoints.PointsBaseColor[OutIndex]	= FMath::Lerp(In.StrandsPoints.PointsBaseColor[InIndex0], In.StrandsPoints.PointsBaseColor[InIndex1], InAlpha); }
 	if (HasHairAttribute(InAttributes, EHairAttribute::Roughness))	{ Out.StrandsPoints.PointsRoughness[OutIndex]	= FMath::Lerp(In.StrandsPoints.PointsRoughness[InIndex0], In.StrandsPoints.PointsRoughness[InIndex1], InAlpha); }
diff --git a/Engine/Plugins/Runtime/HairStrands/Source/HairStrandsCore/Private/HairCardsBuilder.cpp b/Engine/Plugins/Runtime/HairStrands/Source/HairStrandsCore/Private/HairCardsBuilder.cpp
--- a/Engine/Plugins/Runtime/HairStrands/Source/HairStrandsCore/Private/HairCardsBuilder.cpp
+++ b/Engine/Plugins/Runtime/HairStrands/Source/HairStrandsCore/Private/HairCardsBuilder.cpp
@@ -773,8 +773,7 @@ namespace HairCards
 		const float ShapeImportance, 
 		const float ProximityImportance)
 	{
-		const float AverageLength = FMath::Max(0.5f * (RenderCurvesDatas.StrandsCurves.CurvesLength[RenderCurveIndex] * RenderCurvesDatas.StrandsCurves.MaxLength +
-			GuideCurvesDatas.StrandsCurves.CurvesLength[GuideCurveIndex] * GuideCurvesDatas.StrandsCurves.MaxLength), SMALL_NUMBER);
+		const float AverageLength = FMath::Max(0.5f * (RenderCurvesDatas.StrandsCurves.CurvesLength[RenderCurveIndex] + GuideCurvesDatas.StrandsCurves.CurvesLength[GuideCurveIndex]), SMALL_NUMBER);
 
 		static const float DeltaCoord = 1.0f / static_cast<float>(NumSamples-1);
 
@@ -1206,7 +1205,7 @@ namespace HairCards
 		{
 			const uint32 PointOffset = InData.StrandsCurves.CurvesOffset[CurveIndex];
 			const uint32 PointCount = InData.StrandsCurves.CurvesCount[CurveIndex];
-			const float  CurveLength = InData.StrandsCurves.CurvesLength[CurveIndex] * InData.StrandsCurves.MaxLength;
+			const float  CurveLength = InData.StrandsCurves.CurvesLength[CurveIndex];
 			check(PointCount > 1);
 			const FVector3f& P0 = InData.StrandsPoints.PointsPosition[PointOffset];
 			const FVector3f& P1 = InData.StrandsPoints.PointsPosition[PointOffset + 1];
@@ -1236,7 +1235,7 @@ namespace HairCards
 	{
 		const uint32 SimOffset = SimStrandsData.StrandsCurves.CurvesOffset[SimCurveIndex];
 
-		const float CurveLength = SimStrandsData.StrandsCurves.CurvesLength[SimCurveIndex] * SimStrandsData.StrandsCurves.MaxLength;
+		const float CurveLength = SimStrandsData.StrandsCurves.CurvesLength[SimCurveIndex];
 
 		// Find with with vertex the vertex should be paired
 		const uint32 SimPointCount = SimStrandsData.StrandsCurves.CurvesCount[SimCurveIndex];
@@ -1282,9 +1281,6 @@ namespace HairCards
 			TArray<uint32>::TIterator OffsetIterator = Curves.CurvesOffset.CreateIterator();
 			TArray<float>::TIterator LengthIterator = Curves.CurvesLength.CreateIterator();
 
-			Curves.MaxRadius = 0.0;
-			Curves.MaxLength = 0.0;
-
 			uint32 StrandOffset = 0;
 			*OffsetIterator = StrandOffset; ++OffsetIterator;
 
@@ -1308,10 +1304,8 @@ namespace HairCards
 					*CoordUIterator = StrandLength;
 					PreviousPosition = *PositionIterator;
 
-					Curves.MaxRadius = FMath::Max(Curves.MaxRadius, *RadiusIterator);
 				}
 				*LengthIterator = StrandLength;
-				Curves.MaxLength = FMath::Max(Curves.MaxLength, StrandLength);
 			}
 
 			CountIterator.Reset();
@@ -1326,9 +1320,7 @@ namespace HairCards
 				for (uint32 PointIndex = 0; PointIndex < StrandCount; ++PointIndex, ++RadiusIterator, ++CoordUIterator)
 				{
 					*CoordUIterator /= *LengthIterator;
-					*RadiusIterator /= Curves.MaxRadius;
 				}
-				*LengthIterator /= Curves.MaxLength;
 			}
 		}
 	}
@@ -1370,8 +1362,6 @@ namespace HairCards
 			const uint32 PointCount = InData.StrandsCurves.CurvesCount[InCurveIndex];
 			OutData.StrandsCurves.CurvesCount[OutCurveIndex] = PointCount;
 			OutData.StrandsCurves.CurvesOffset[OutCurveIndex] = OutPointOffset;
-			OutData.StrandsCurves.MaxLength = InData.StrandsCurves.MaxLength;
-			OutData.StrandsCurves.MaxRadius = InData.StrandsCurves.MaxRadius;
 			FHairStrandsDatas::CopyCurve(InData, OutData, InAttribute, InCurveIndex, OutCurveIndex);
 
 			for (uint32 PointIndex = 0; PointIndex < PointCount; ++PointIndex)
@@ -1500,6 +1490,8 @@ namespace HairCards
 		ClusterCurveIndexCount.SetNumZeroed(ClusterCount);
 		TAtomic<uint32> TotalCurveIndices(0);
 
+		const float SimStrandsMaxLength = GetHairStrandsMaxLength(SimStrandsData);
+
 		TArray<FClosestGuides> ClosestGuides;
 		ClosestGuides.SetNum(RenCurveCount);
 		const FDateTime StartTime = FDateTime::UtcNow();
@@ -1509,6 +1501,7 @@ namespace HairCards
 			InSettings,
 			RenCurveCount, &RenRoots, &RenStrandsData,
 			SimCurveCount, &SimRoots, &SimStrandsData, 
+			SimStrandsMaxLength,
 			&RootsGrid,
 			&CompletedTasks,
 			&SlowTask,
@@ -1556,8 +1549,8 @@ namespace HairCards
 
 		// Insure the min length threshold is lower than the max curves length, to generate at least one cluster. 
 		// This is arbitrarely set to 95% of the longest strands
-		float MinLengthTreshold = FMath::Min(InSettings.MinCardsLength, SimStrandsData.StrandsCurves.MaxLength * 0.95f);
-		float MaxLengthTreshold = InSettings.MaxCardsLength > 0 ? FMath::Min(InSettings.MaxCardsLength, SimStrandsData.StrandsCurves.MaxLength) : SimStrandsData.StrandsCurves.MaxLength;
+		float MinLengthTreshold = FMath::Min(InSettings.MinCardsLength, SimStrandsMaxLength * 0.95f);
+		float MaxLengthTreshold = InSettings.MaxCardsLength > 0 ? FMath::Min(InSettings.MaxCardsLength, SimStrandsMaxLength) : SimStrandsMaxLength;
 
 		// Compute the number of clusters
 		// * Some cluster might have no curve this is why we remove them
@@ -1571,7 +1564,7 @@ namespace HairCards
 			{
 				for (uint32 GuideIndex = 0; GuideIndex < GuideCount; ++GuideIndex)
 				{				
-					const float GuideLength = SimStrandsData.StrandsCurves.CurvesLength[GuideIndex] * SimStrandsData.StrandsCurves.MaxLength;
+					const float GuideLength = SimStrandsData.StrandsCurves.CurvesLength[GuideIndex];
 					const bool bIsCurveValid = GuideLength >= MinLengthTreshold && GuideLength <= MaxLengthTreshold;
 					if (bIsCurveValid && ClusterCurveIndexCount[GuideIndex] > 0)
 					{
@@ -1588,7 +1581,7 @@ namespace HairCards
 				if (ValidClusterCount == 0)
 				{
 					MinLengthTreshold = MinLengthTreshold > 0 ? MinLengthTreshold * 0.1f : 0;
-					MaxLengthTreshold = MaxLengthTreshold < SimStrandsData.StrandsCurves.MaxLength ? FMath::Min(SimStrandsData.StrandsCurves.MaxLength, MaxLengthTreshold * 1.1f) : SimStrandsData.StrandsCurves.MaxLength;
+					MaxLengthTreshold = MaxLengthTreshold < SimStrandsMaxLength ? FMath::Min(SimStrandsMaxLength, MaxLengthTreshold * 1.1f) : SimStrandsMaxLength;
 				}				
 			}
 			ClusterCount = FMath::Min(ValidClusterCount, uint32(FMath::Max(1, GHairCardsMaxClusterCount)));
@@ -1602,7 +1595,7 @@ namespace HairCards
 			uint32 Offset = 0;
 			for (uint32 GuideIndex = 0; GuideIndex < GuideCount; ++GuideIndex)
 			{
-				const float GuideLength = SimStrandsData.StrandsCurves.CurvesLength[GuideIndex] * SimStrandsData.StrandsCurves.MaxLength;
+				const float GuideLength = SimStrandsData.StrandsCurves.CurvesLength[GuideIndex];
 				const bool bIsCurveValid = GuideLength >= MinLengthTreshold && GuideLength <= MaxLengthTreshold;
 				if (bIsCurveValid && ClusterCurveIndexCount[GuideIndex] > 0)
 				{
@@ -2026,7 +2019,7 @@ namespace HairCards
 				const uint32 StrandIndex = Out.CurveIndices[CurveOffset + StrandIt];
 				const uint32 StrandOffset = In.StrandsCurves.CurvesOffset[StrandIndex];
 				const uint32 PointCount = In.StrandsCurves.CurvesCount[StrandIndex];
-				const float CurveLength = In.StrandsCurves.CurvesLength[StrandIndex] * In.StrandsCurves.MaxLength;
+				const float CurveLength = In.StrandsCurves.CurvesLength[StrandIndex];
 
 				if (CurveLength > LongestLength)
 				{
@@ -2438,7 +2431,7 @@ namespace HairCards
 				const uint32 StrandIndex = InClusters.CurveIndices[CurveOffset + StrandIt];
 				const uint32 StrandOffset = InDatas.StrandsCurves.CurvesOffset[StrandIndex];
 				const uint32 PointCount = InDatas.StrandsCurves.CurvesCount[StrandIndex];
-				const float CurveLength = InDatas.StrandsCurves.CurvesLength[StrandIndex] * InDatas.StrandsCurves.MaxLength;
+				const float CurveLength = InDatas.StrandsCurves.CurvesLength[StrandIndex];
 
 				Desc.MinLength = FMath::Min(CurveLength, Desc.MinLength);
 				Desc.MaxLength = FMath::Max(CurveLength, Desc.MaxLength);
@@ -2773,7 +2766,6 @@ namespace HairCards
 
 		uint32 IndexIt = 0;
 		uint32 VertexIt = 0;
-		const float GuideRadius = 0.01f;
 		for (uint32 ClusterIt = 0; ClusterIt < ClusterCount; ++ClusterIt)
 		{
 			const uint32 PointCount = InClusters.BoundCount[ClusterIt];
@@ -2786,8 +2778,6 @@ namespace HairCards
 			//OutGuides.StrandsCurves.CurvesRootUV[ClusterIt] = FVector2f(0, 0);
 			//OutGuides.StrandsCurves.StrandIDs[ClusterIt] = ClusterIt;
 			//OutGuides.StrandsCurves.GroomIDToIndex[ClusterIt] = ;
-			OutGuides.StrandsCurves.MaxLength = FMath::Max(OutGuides.StrandsCurves.MaxLength, TotalLength);
-			OutGuides.StrandsCurves.MaxRadius = GuideRadius;
 
 			float CurrentLength = 0;
 			for (uint32 PointIt = 0; PointIt < PointCount; ++PointIt)
@@ -2812,12 +2802,6 @@ namespace HairCards
 				}
 			}
 		}
-
-		const float MaxLength = OutGuides.StrandsCurves.MaxLength > 0 ? OutGuides.StrandsCurves.MaxLength : 1;
-		for (uint32 ClusterIt = 0; ClusterIt < ClusterCount; ++ClusterIt)
-		{
-			OutGuides.StrandsCurves.CurvesLength[ClusterIt] /= MaxLength;
-		}
 	}
 
 	static bool CreateCardsGuides(
@@ -2842,7 +2826,6 @@ namespace HairCards
 		InCards.CoordU.SetNum(InCards.Positions.Num());
 		InCards.LocalUVs.SetNum(InCards.Positions.Num());
 
-		const float GuideRadius = 0.01f;
 
 		// Find the principal direction of the atlas by comparing the number of segment along U and along V
 		bool bIsMainDirectionU = true;
@@ -3077,8 +3060,6 @@ namespace HairCards
 			OutGuides.StrandsCurves.CurvesOffset[CardIt] = PointOffset;
 			OutGuides.StrandsCurves.CurvesLength[CardIt] = TotalLength;
 			if (OutGuides.StrandsCurves.CurvesRootUV.Num()) OutGuides.StrandsCurves.CurvesRootUV[CardIt] = FVector2f(0, 0);
-			OutGuides.StrandsCurves.MaxLength = FMath::Max(OutGuides.StrandsCurves.MaxLength, TotalLength);
-			OutGuides.StrandsCurves.MaxRadius = GuideRadius;
 
 			float CurrentLength = 0;
 			for (uint32 PointIt = 0; PointIt < PointCount; ++PointIt)
@@ -3101,12 +3082,6 @@ namespace HairCards
 			}
 		}
 
-		const float MaxLength = OutGuides.StrandsCurves.MaxLength > 0 ? OutGuides.StrandsCurves.MaxLength : 1;
-		for (uint32 CardIt = 0; CardIt < NumCards; ++CardIt)
-		{
-			OutGuides.StrandsCurves.CurvesLength[CardIt] /= MaxLength;
-		}
-
 		return true;
 	}
 
@@ -3274,7 +3249,7 @@ namespace HairCards
 					const float Seed = 0.f; // TODO float(InStrands.RenderData.Attributes[PointIndex].Seed) / 255.f;
 
 					const FVector3f& P0_Rest = InStrands.StrandsPoints.PointsPosition[PointIndex];
-					const float Radius = InStrands.StrandsPoints.PointsRadius[PointIndex] * InStrands.StrandsCurves.MaxRadius;
+					const float Radius = InStrands.StrandsPoints.PointsRadius[PointIndex];
 
 					// Find the closest guide point using parametric length at a correspondance/matching metric
 					uint32 ClosestGuidePointIt0 = 0;
diff --git a/Engine/Plugins/Runtime/HairStrands/Source/HairStrandsCore/Private/HairStrandsDatas.cpp b/Engine/Plugins/Runtime/HairStrands/Source/HairStrandsCore/Private/HairStrandsDatas.cpp
--- a/Engine/Plugins/Runtime/HairStrands/Source/HairStrandsCore/Private/HairStrandsDatas.cpp
+++ b/Engine/Plugins/Runtime/HairStrands/Source/HairStrandsCore/Private/HairStrandsDatas.cpp
@@ -59,8 +59,6 @@ void FHairStrandsCurves::Reset()
 	ClumpIDs.Reset();
 	CurvesClosestGuideIDs.Reset();
 	CurvesClosestGuideWeights.Reset();
-	MaxLength = 0;
-	MaxRadius = 0;
 }
 
 void FHairStrandsPoints::Reset()
@@ -73,6 +71,26 @@ void FHairStrandsPoints::Reset()
 	PointsAO.Reset();
 }
 
+float GetHairStrandsMaxLength(const FHairStrandsDatas& In)
+{
+	float MaxLength = 0;
+	for (float CurveLength : In.StrandsCurves.CurvesLength)
+	{
+		MaxLength = FMath::Max(MaxLength, CurveLength);
+	}
+	return MaxLength;
+}
+
+float GetHairStrandsMaxRadius(const FHairStrandsDatas& In)
+{
+	float MaxRadius = 0;
+	for (float PointRadius : In.StrandsPoints.PointsRadius)
+	{
+		MaxRadius = FMath::Max(MaxRadius, PointRadius);
+	}
+	return MaxRadius;
+}
+
 FArchive& operator<<(FArchive& Ar, FVector4_16& Vertex)
 {
 	Ar << Vertex.X;
diff --git a/Engine/Plugins/Runtime/HairStrands/Source/HairStrandsCore/Public/GroomAsset.h b/Engine/Plugins/Runtime/HairStrands/Source/HairStrandsCore/Public/GroomAsset.h
--- a/Engine/Plugins/Runtime/HairStrands/Source/HairStrandsCore/Public/GroomAsset.h
+++ b/Engine/Plugins/Runtime/HairStrands/Source/HairStrandsCore/Public/GroomAsset.h
@@ -58,9 +58,6 @@ struct HAIRSTRANDSCORE_API FHairGroupInfo
 	UPROPERTY(VisibleAnywhere, Category = "Info", meta = (DisplayName = "Length of the longest hair strands"))
 	float MaxCurveLength = 0;
 
-	UPROPERTY(VisibleAnywhere, Category = "Info", meta = (DisplayName = "Imported width (DCC) of the larget hair strands"))
-	float MaxImportedWidth = -1.f;
-
 	UPROPERTY(VisibleAnywhere, Category = "Info", meta = (DisplayName = "Card Group ID for split hair card generation"))
 	int32 GroupCardsID = 0;
 };
diff --git a/Engine/Plugins/Runtime/HairStrands/Source/HairStrandsCore/Public/GroomAssetRendering.h b/Engine/Plugins/Runtime/HairStrands/Source/HairStrandsCore/Public/GroomAssetRendering.h
--- a/Engine/Plugins/Runtime/HairStrands/Source/HairStrandsCore/Public/GroomAssetRendering.h
+++ b/Engine/Plugins/Runtime/HairStrands/Source/HairStrandsCore/Public/GroomAssetRendering.h
@@ -19,8 +19,10 @@ struct HAIRSTRANDSCORE_API FHairGeometrySettings
 	FHairGeometrySettings();
 
 	/** Hair width (in centimeters) */
-	UPROPERTY(EditAnywhere, BlueprintReadWrite, Category = "GeometrySettings", meta = (ClampMin = "0.0001", UIMin = "0.001", UIMax = "1.0", SliderExponent = 6))
+	UPROPERTY(EditAnywhere, BlueprintReadWrite, Category = "GeometrySettings", meta = (editcondition = "HairWidth_Override", ClampMin = "0.0001", UIMin = "0.001", UIMax = "1.0", SliderExponent = 6))
 	float HairWidth;
+	UPROPERTY(EditAnywhere, BlueprintReadWrite, Category = Overrides, meta = (PinHiddenByDefault, InlineEditConditionToggle))
+	bool HairWidth_Override = false;
 
 	/** Scale the hair width at the root */
 	UPROPERTY(EditAnywhere, BlueprintReadWrite, Category = "GeometrySettings", AdvancedDisplay, meta = (ClampMin = "0.0001", UIMin = "0.001", UIMax = "2.0", SliderExponent = 6))
diff --git a/Engine/Plugins/Runtime/HairStrands/Source/HairStrandsCore/Public/GroomCacheData.h b/Engine/Plugins/Runtime/HairStrands/Source/HairStrandsCore/Public/GroomCacheData.h
--- a/Engine/Plugins/Runtime/HairStrands/Source/HairStrandsCore/Public/GroomCacheData.h
+++ b/Engine/Plugins/Runtime/HairStrands/Source/HairStrandsCore/Public/GroomCacheData.h
@@ -127,7 +127,7 @@ struct FGroomCacheVertexData
 struct FGroomCacheStrandData
 {
 	FGroomCacheStrandData() = default;
-	FGroomCacheStrandData(struct FHairStrandsCurves&& CurvesData);
+	FGroomCacheStrandData(struct FHairStrandsCurves&& CurvesData, float InMaxLength, float InMaxRadius);
 
 	TArray<float> CurvesLength;
 	float MaxLength = 0.0f;
diff --git a/Engine/Plugins/Runtime/HairStrands/Source/HairStrandsCore/Public/HairStrandsDatas.h b/Engine/Plugins/Runtime/HairStrands/Source/HairStrandsCore/Public/HairStrandsDatas.h
--- a/Engine/Plugins/Runtime/HairStrands/Source/HairStrandsCore/Public/HairStrandsDatas.h
+++ b/Engine/Plugins/Runtime/HairStrands/Source/HairStrandsCore/Public/HairStrandsDatas.h
@@ -385,7 +385,7 @@ struct HAIRSTRANDSCORE_API FHairStrandsCurves
 	/** Get the number of Curves */
 	uint32 Num() const { return CurvesCount.Num(); }
 
-	bool HasPrecomputedWeights() const { return CurvesClosestGuideIDs.Num() > 0 && CurvesClosestGuideWeights.Num() > 0; }
+	bool HasPrecomputedWeights() const;
 
 	bool HasAttribute(EHairAttribute In) const;
 
@@ -415,12 +415,6 @@ struct HAIRSTRANDSCORE_API FHairStrandsCurves
 
 	/** Custom guid weights (indexed with StrandID) (optional) */
 	TArray<FVector> CurvesClosestGuideWeights;
-
-	/** Max strands Curves length */
-	float MaxLength = 0;
-
-	/** Max strands Curves radius */
-	float MaxRadius = 0;
 };
 
 struct FHairStrandsBulkData;
@@ -458,6 +452,9 @@ struct HAIRSTRANDSCORE_API FHairStrandsDatas
 	FBox BoundingBox = FBox(EForceInit::ForceInit);
 };
 
+float GetHairStrandsMaxLength(const FHairStrandsDatas& In);
+float GetHairStrandsMaxRadius(const FHairStrandsDatas& In);
+
 struct HAIRSTRANDSCORE_API FHairStrandsBulkData
 {
 	enum EDataFlags
diff --git a/Engine/Plugins/Runtime/HairStrands/Source/HairStrandsEditor/Private/GroomAssetDetails.cpp b/Engine/Plugins/Runtime/HairStrands/Source/HairStrandsEditor/Private/GroomAssetDetails.cpp
--- a/Engine/Plugins/Runtime/HairStrands/Source/HairStrandsEditor/Private/GroomAssetDetails.cpp
+++ b/Engine/Plugins/Runtime/HairStrands/Source/HairStrandsEditor/Private/GroomAssetDetails.cpp
@@ -204,28 +204,6 @@ TSharedRef<SUniformGridPanel> MakeHairStrandsInfoGrid(const FSlateFontInfo& Deta
 		.Text(FText::AsNumber(MaxRadius*2.0f))
 	];
 
-	// Imported Width (mm)
-	Grid->AddSlot(0, 4) // x, y
-	.HAlign(HAlign_Left)
-	[
-		SNew(STextBlock)
-		.Font(DetailFontInfo)
-	.Text(LOCTEXT("HairInfo_ImportedWidth", "Max. Imported Width"))
-	];
-	Grid->AddSlot(1, 4) // x, y
-	.HAlign(HAlign_Right)
-	[
-		SNew(STextBlock)
-		.Font(DetailFontInfo)
-	.Text(LOCTEXT("HairInfo_GuideImportedWidth", ""))
-	];
-	Grid->AddSlot(2, 4) // x, y
-	.HAlign(HAlign_Right)
-	[
-		SNew(STextBlock)
-		.Font(DetailFontInfo)
-		.Text(CurrentAsset.MaxImportedWidth > 0.f ? FText::AsNumber(CurrentAsset.MaxImportedWidth) : LOCTEXT("HairInfo_ImportedWidthDefault", "Not exported"))
-	];
 	return Grid;
 }
 
diff --git a/Engine/Source/Runtime/Core/Public/UObject/UE5MainStreamObjectVersion.h b/Engine/Source/Runtime/Core/Public/UObject/UE5MainStreamObjectVersion.h
--- a/Engine/Source/Runtime/Core/Public/UObject/UE5MainStreamObjectVersion.h
+++ b/Engine/Source/Runtime/Core/Public/UObject/UE5MainStreamObjectVersion.h
@@ -334,6 +334,9 @@ struct CORE_API FUE5MainStreamObjectVersion
 		// Allow custom event nodes to use access specifiers
 		AccessSpecifiersForCustomEvents,
 
+		// Explicit override of Groom's hair width
+		GroomAssetWidthOverride,
+
 		// -----<new versions can be added above this line>-------------------------------------------------
 		VersionPlusOne,
 		LatestVersion = VersionPlusOne - 1
