Fix groom binding failing to project hair on coarser skel. mesh.

This fix ensures the hair bound intersects the skel. mesh triangle bound so that traiangles can correctly be interseted.

diff --git a/Engine/Plugins/Runtime/HairStrands/Source/HairStrandsCore/Private/GroomBindingBuilder.cpp b/Engine/Plugins/Runtime/HairStrands/Source/HairStrandsCore/Private/GroomBindingBuilder.cpp
--- a/Engine/Plugins/Runtime/HairStrands/Source/HairStrandsCore/Private/GroomBindingBuilder.cpp
+++ b/Engine/Plugins/Runtime/HairStrands/Source/HairStrandsCore/Private/GroomBindingBuilder.cpp
@@ -1104,7 +1104,8 @@ namespace GroomBinding_RootProjection
 				UE_LOG(LogHairStrands, Error, TEXT("[Groom] Binding asset could not be built. MeshLODData has 0 sections."));
 				return false;
 			}
-			
+
+			float ClosestTrianglePoint = FLT_MAX;
 			check(SectionCount > 0);
 			for (uint32 SectionIt = 0; SectionIt < SectionCount; ++SectionIt)
 			{
@@ -1148,6 +1149,11 @@ namespace GroomBinding_RootProjection
 					MeshBound += T.P0;
 					MeshBound += T.P1;
 					MeshBound += T.P2;
+
+					// Track closest point to the groom bound
+					ClosestTrianglePoint = FMath::Min(ClosestTrianglePoint, (T.P0 - InStrandsData.BoundingBox.GetCenter()).Length());
+					ClosestTrianglePoint = FMath::Min(ClosestTrianglePoint, (T.P1 - InStrandsData.BoundingBox.GetCenter()).Length());
+					ClosestTrianglePoint = FMath::Min(ClosestTrianglePoint, (T.P2 - InStrandsData.BoundingBox.GetCenter()).Length());
 				}
 			}
 
@@ -1165,6 +1171,15 @@ namespace GroomBinding_RootProjection
 			{
 				GridMin = InStrandsData.BoundingBox.Min;
 				GridMax = InStrandsData.BoundingBox.Max;
+
+				// By nature, it is possible that coarser LOD have positions which ressemble only very coarsly to 
+				// LOD0. In this case we increase the hair bound to ensure that skel. mesh triangles will be intersect 
+				// the groom bound to be correctly inserted.
+				if (bHasTransferredPosition && ClosestTrianglePoint < FLT_MAX)
+				{
+					GridMin -= FVector(ClosestTrianglePoint * 1.25f);
+					GridMax += FVector(ClosestTrianglePoint * 1.25f);
+				}
 			}
 
 			FTriangleGrid Grid(GridMin, GridMax, VoxelWorldSize);
