Fix invalid assert in groom component causing crash when deleting groom LOD while raytracing is enabled.

This is also safe as the raytracing resources will never be used since the LOD can't be selected since its index is higher than the selectable ones.

#perflight 6401bdefa20ddf1bf49d4a4b

diff --git a/Engine/Plugins/Runtime/HairStrands/Source/HairStrandsCore/Private/GroomComponent.cpp b/Engine/Plugins/Runtime/HairStrands/Source/HairStrandsCore/Private/GroomComponent.cpp
--- a/Engine/Plugins/Runtime/HairStrands/Source/HairStrandsCore/Private/GroomComponent.cpp
+++ b/Engine/Plugins/Runtime/HairStrands/Source/HairStrandsCore/Private/GroomComponent.cpp
@@ -2956,7 +2956,6 @@ void UGroomComponent::InitResources(bool bIsBindingReloading)
 						// (Lazy) Allocate static raytracing resources (owned by the grooom asset)
 						InstanceLOD.RaytracingResourceOwned = false;
 						InstanceLOD.RaytracingResource = GroomAsset->AllocateMeshesRaytracingResources(GroupIt, MeshLODIndex);
-						check(InstanceLOD.RaytracingResource);
 					}
 				}
 				#endif
