* Change groom mesh deformer to use dedicated 'ControlPoint' and 'Curve' execution domain.
* Fix groom source component not inserting all its groups as running invocation.

diff --git a/Engine/Plugins/Runtime/HairStrands/Shaders/Private/DeformerDataInterfaceGroom.ush b/Engine/Plugins/Runtime/HairStrands/Shaders/Private/DeformerDataInterfaceGroom.ush
--- a/Engine/Plugins/Runtime/HairStrands/Shaders/Private/DeformerDataInterfaceGroom.ush
+++ b/Engine/Plugins/Runtime/HairStrands/Shaders/Private/DeformerDataInterfaceGroom.ush
@@ -23,17 +23,17 @@ Buffer<uint>		{DataInterfaceName}_CurveBuffer;
 ///////////////////////////////////////////////////////////////////////////////////////////////////
 // Helpers
 
-FHairControlPoint ReadControlPointData(uint VertexIndex)
+FHairControlPoint ReadControlPointData(uint ControlPointIndex)
 {
-	const uint MaxVertices = {DataInterfaceName}_NumControlPoints;
-	if (VertexIndex >= MaxVertices)
+	const uint MaxControlPoints = {DataInterfaceName}_NumControlPoints;
+	if (ControlPointIndex >= MaxControlPoints)
 	{
 		return (FHairControlPoint)0;
 	}
 
 	// Hair deformer work in local space so no need for position offset
 	return UnpackHairControlPoint(
-		{DataInterfaceName}_PositionBuffer[VertexIndex], 
+		{DataInterfaceName}_PositionBuffer[ControlPointIndex], 
 		{DataInterfaceName}_PositionOffsetBuffer[0],
 		{DataInterfaceName}_VF_HairRadius, 
 		1.0f /*RootScale*/, 
@@ -64,36 +64,36 @@ uint ReadNumCurves_{DataInterfaceName}()
 	return {DataInterfaceName}_NumCurves;
 }
 
-float3 ReadPosition_{DataInterfaceName}(uint VertexIndex)
+float3 ReadPosition_{DataInterfaceName}(uint ControlPointIndex)
 {
-	return ReadControlPointData(VertexIndex).Position;
+	return ReadControlPointData(ControlPointIndex).Position;
 }
 
-float ReadRadius_{DataInterfaceName}(uint VertexIndex)
+float ReadRadius_{DataInterfaceName}(uint ControlPointIndex)
 {
-	return ReadControlPointData(VertexIndex).WorldRadius;
+	return ReadControlPointData(ControlPointIndex).WorldRadius;
 }
 
-float ReadCoordU_{DataInterfaceName}(uint VertexIndex)
+float ReadCoordU_{DataInterfaceName}(uint ControlPointIndex)
 {
-	return ReadControlPointData(VertexIndex).UCoord;
+	return ReadControlPointData(ControlPointIndex).UCoord;
 }
 
-float ReadLength_{DataInterfaceName}(uint VertexIndex)
+float ReadLength_{DataInterfaceName}(uint ControlPointIndex)
 {
-	const float CoordU      = ReadControlPointData(VertexIndex).UCoord;
-	const float CurveLength = UnpackHairLength(InternalGetHairStrandsAttributesPerCurve(VertexIndex, HAIR_ATTRIBUTE_OFFSET_LENGTH({DataInterfaceName}_VF_AttributeOffsets), HAIR_ATTRIBUTE_STRIDE_LENGTH, {DataInterfaceName}_AttributeBuffer, {DataInterfaceName}_VertexToCurveBuffer));
+	const float CoordU      = ReadControlPointData(ControlPointIndex).UCoord;
+	const float CurveLength = UnpackHairLength(InternalGetHairStrandsAttributesPerCurve(ControlPointIndex, HAIR_ATTRIBUTE_OFFSET_LENGTH({DataInterfaceName}_VF_AttributeOffsets), HAIR_ATTRIBUTE_STRIDE_LENGTH, {DataInterfaceName}_AttributeBuffer, {DataInterfaceName}_VertexToCurveBuffer));
 	return CurveLength * CoordU;
 }
 
-float2 ReadRootUV_{DataInterfaceName}(uint VertexIndex)
+float2 ReadRootUV_{DataInterfaceName}(uint ControlPointIndex)
 {
-	return UnpackHairRootUV(InternalGetHairStrandsAttributesPerCurve(VertexIndex, HAIR_ATTRIBUTE_OFFSET_ROOTUV({DataInterfaceName}_VF_AttributeOffsets), HAIR_ATTRIBUTE_STRIDE_ROOTUV, {DataInterfaceName}_AttributeBuffer, {DataInterfaceName}_VertexToCurveBuffer));
+	return UnpackHairRootUV(InternalGetHairStrandsAttributesPerCurve(ControlPointIndex, HAIR_ATTRIBUTE_OFFSET_ROOTUV({DataInterfaceName}_VF_AttributeOffsets), HAIR_ATTRIBUTE_STRIDE_ROOTUV, {DataInterfaceName}_AttributeBuffer, {DataInterfaceName}_VertexToCurveBuffer));
 }
 
-float ReadSeed_{DataInterfaceName}(uint VertexIndex)
+float ReadSeed_{DataInterfaceName}(uint ControlPointIndex)
 {
-	return UnpackHairSeed(InternalGetHairStrandsAttributesPerCurve(VertexIndex, HAIR_ATTRIBUTE_OFFSET_SEED({DataInterfaceName}_VF_AttributeOffsets), HAIR_ATTRIBUTE_STRIDE_SEED, {DataInterfaceName}_AttributeBuffer, {DataInterfaceName}_VertexToCurveBuffer));
+	return UnpackHairSeed(InternalGetHairStrandsAttributesPerCurve(ControlPointIndex, HAIR_ATTRIBUTE_OFFSET_SEED({DataInterfaceName}_VF_AttributeOffsets), HAIR_ATTRIBUTE_STRIDE_SEED, {DataInterfaceName}_AttributeBuffer, {DataInterfaceName}_VertexToCurveBuffer));
 }
 
 float ReadCurveOffsetPoint_{DataInterfaceName}(uint CurveIndex)
diff --git a/Engine/Plugins/Runtime/HairStrands/Source/HairStrandsCore/Private/Deformer/DeformerDataInterfaceGroom.cpp b/Engine/Plugins/Runtime/HairStrands/Source/HairStrandsCore/Private/Deformer/DeformerDataInterfaceGroom.cpp
--- a/Engine/Plugins/Runtime/HairStrands/Source/HairStrandsCore/Private/Deformer/DeformerDataInterfaceGroom.cpp
+++ b/Engine/Plugins/Runtime/HairStrands/Source/HairStrandsCore/Private/Deformer/DeformerDataInterfaceGroom.cpp
@@ -12,6 +12,7 @@
 #include "RenderGraphBuilder.h"
 #include "GlobalRenderResources.h"
 #include "HairStrandsDefinitions.h"
+#include "DeformerGroomComponentSource.h"
 
 FString UOptimusGroomDataInterface::GetDisplayName() const
 {
@@ -20,17 +21,20 @@ FString UOptimusGroomDataInterface::GetDisplayName() const
 
 TArray<FOptimusCDIPinDefinition> UOptimusGroomDataInterface::GetPinDefinitions() const
 {
+	FName ControlPoint(UOptimusGroomComponentSource::Domains::ControlPoint);
+	FName Curve(UOptimusGroomComponentSource::Domains::Curve);
+
 	TArray<FOptimusCDIPinDefinition> Defs;
-	Defs.Add({"NumControlPoints", "ReadNumControlPoints"});
-	Defs.Add({"NumCurves",        "ReadNumCurves" });
-	Defs.Add({"Position",         "ReadPosition",          Optimus::DomainName::Vertex,   "ReadPosition"});
-	Defs.Add({"Radius",           "ReadRadius",            Optimus::DomainName::Vertex,   "ReadRadius" });
-	Defs.Add({"CoordU",           "ReadCoordU",            Optimus::DomainName::Vertex,   "ReadCoordU" });
-	Defs.Add({"Length",           "ReadLength",            Optimus::DomainName::Vertex,   "ReadLength" });
-	Defs.Add({"RootUV",           "ReadRootUV",            Optimus::DomainName::Vertex,   "ReadRootUV" });
-	Defs.Add({"Seed",             "ReadSeed",              Optimus::DomainName::Vertex,   "ReadSeed"   });
-	Defs.Add({"CurveOffsetPoint", "ReadCurveOffsetPoint",  Optimus::DomainName::Triangle, "ReadCurveOffsetPoint" });
-	Defs.Add({"CurveNumPoint",    "ReadCurveNumPoint",     Optimus::DomainName::Triangle, "ReadCurveNumPoint" });
+	Defs.Add({ "NumControlPoints", "ReadNumControlPoints" });
+	Defs.Add({ "NumCurves",        "ReadNumCurves" });
+	Defs.Add({ "Position",         "ReadPosition",          ControlPoint,   "ReadPosition" });
+	Defs.Add({ "Radius",           "ReadRadius",            ControlPoint,   "ReadRadius" });
+	Defs.Add({ "CoordU",           "ReadCoordU",            ControlPoint,   "ReadCoordU" });
+	Defs.Add({ "Length",           "ReadLength",            ControlPoint,   "ReadLength" });
+	Defs.Add({ "RootUV",           "ReadRootUV",            ControlPoint,   "ReadRootUV" });
+	Defs.Add({ "Seed",             "ReadSeed",              ControlPoint,   "ReadSeed" });
+	Defs.Add({ "CurveOffsetPoint", "ReadCurveOffsetPoint",  Curve,          "ReadCurveOffsetPoint" });
+	Defs.Add({ "CurveNumPoint",    "ReadCurveNumPoint",     Curve,          "ReadCurveNumPoint" });
 	return Defs;
 }
  
diff --git a/Engine/Plugins/Runtime/HairStrands/Source/HairStrandsCore/Private/Deformer/DeformerDataInterfaceGroomWrite.cpp b/Engine/Plugins/Runtime/HairStrands/Source/HairStrandsCore/Private/Deformer/DeformerDataInterfaceGroomWrite.cpp
--- a/Engine/Plugins/Runtime/HairStrands/Source/HairStrandsCore/Private/Deformer/DeformerDataInterfaceGroomWrite.cpp
+++ b/Engine/Plugins/Runtime/HairStrands/Source/HairStrandsCore/Private/Deformer/DeformerDataInterfaceGroomWrite.cpp
@@ -12,6 +12,7 @@
 #include "SkeletalRenderPublic.h"
 #include "GroomComponent.h"
 #include "GroomInstance.h"
+#include "DeformerGroomComponentSource.h"
 
 FString UOptimusGroomWriteDataInterface::GetDisplayName() const
 {
@@ -25,8 +26,11 @@ FName UOptimusGroomWriteDataInterface::GetCategory() const
 
 TArray<FOptimusCDIPinDefinition> UOptimusGroomWriteDataInterface::GetPinDefinitions() const
 {
+	FName ControlPoint(UOptimusGroomComponentSource::Domains::ControlPoint);
+	FName Curve(UOptimusGroomComponentSource::Domains::Curve);
+
 	TArray<FOptimusCDIPinDefinition> Defs;
-	Defs.Add({ "Position", "WritePosition", Optimus::DomainName::Vertex, "ReadNumControlPoints" }); // ControlPoint domain?
+	Defs.Add({ "Position", "WritePosition", ControlPoint, "ReadNumControlPoints" });
 	return Defs;
 }
 
diff --git a/Engine/Plugins/Runtime/HairStrands/Source/HairStrandsCore/Private/Deformer/DeformerGroomComponentSource.cpp b/Engine/Plugins/Runtime/HairStrands/Source/HairStrandsCore/Private/Deformer/DeformerGroomComponentSource.cpp
--- a/Engine/Plugins/Runtime/HairStrands/Source/HairStrandsCore/Private/Deformer/DeformerGroomComponentSource.cpp
+++ b/Engine/Plugins/Runtime/HairStrands/Source/HairStrandsCore/Private/Deformer/DeformerGroomComponentSource.cpp
@@ -6,16 +6,6 @@
 
 #define LOCTEXT_NAMESPACE "DeformersGroomComponentSource"
 
-static const FHairGroupInstance* GetGroomInstance(const UGroomComponent* In)
-{
-	// HAIR_TODO: apply to all hair group, or only selected one
-	if (In && In->GetGroupCount() > 0)
-	{
-		return In->GetGroupInstance(0);
-	}
-	return nullptr;
-}
-
 FName UOptimusGroomComponentSource::Domains::ControlPoint("ControlPoint");
 FName UOptimusGroomComponentSource::Domains::Curve("Curve");
 
@@ -57,27 +47,31 @@ bool UOptimusGroomComponentSource::GetComponentElementCountsForExecutionDomain(
 		return false;
 	}
 
-	const FHairGroupInstance* Instance = GetGroomInstance(GroomComponent);
-	if (!Instance)
+	// Ensure all groups are valid
+	const uint32 GroupCount = GroomComponent->GetGroupCount();
+	for (uint32 GroupIt=0;GroupIt<GroupCount;++GroupIt)
 	{
-		return 0;
+		if (!GroomComponent->GetGroupInstance(GroupIt))
+		{
+			return 0;
+		}
 	}
 
-	const int32 NumVertices = Instance->Strands.Data->PointCount;
-	const int32 NumCurves = Instance->Strands.Data->CurveCount;
-
 	OutInvocationElementCounts.Reset();
 
 	if (InDomainName == Domains::ControlPoint || InDomainName == Domains::Curve)
 	{
-		const int32 NumInvocations = 1; //TODO: num hair group?
-
-		OutInvocationElementCounts.Reset();
+		const int32 NumInvocations = GroupCount;
 		OutInvocationElementCounts.Reserve(NumInvocations);
 		for (int32 InvocationIndex = 0; InvocationIndex < NumInvocations; ++InvocationIndex)
 		{
-			const int32 NumThreads = InDomainName == Domains::ControlPoint ? NumVertices : NumCurves;
-			OutInvocationElementCounts.Add(NumThreads);
+			if (const FHairGroupInstance* Instance = GroomComponent->GetGroupInstance(InvocationIndex))
+			{
+				const int32 NumControlPoints = Instance->Strands.Data->PointCount;
+				const int32 NumCurves = Instance->Strands.Data->CurveCount;
+				const int32 NumThreads = InDomainName == Domains::ControlPoint ? NumControlPoints : NumCurves;
+				OutInvocationElementCounts.Add(NumThreads);
+			}
 		}
 
 		return true;
