Add groom curve/point limit explicitely through the import and building process.

diff --git a/Engine/Plugins/Runtime/HairStrands/Source/HairStrandsCore/Private/GroomBuilder.cpp b/Engine/Plugins/Runtime/HairStrands/Source/HairStrandsCore/Private/GroomBuilder.cpp
--- a/Engine/Plugins/Runtime/HairStrands/Source/HairStrandsCore/Private/GroomBuilder.cpp
+++ b/Engine/Plugins/Runtime/HairStrands/Source/HairStrandsCore/Private/GroomBuilder.cpp
@@ -324,11 +324,14 @@ namespace HairStrandsBuilder
 
 			// Curves
 			{
-				UE_CLOG(PointCount > 0xFF, LogGroomBuilder, Warning, TEXT("[Groom] Curve point count is > 256, which is not supported."));
-				UE_CLOG(IndexOffset > 0xFFFFFF, LogGroomBuilder, Warning, TEXT("[Groom] Curve point offset is > 24M, which is not supported."));
+				// If the limit change, update this encoding
+				static_assert(HAIR_MAX_NUM_POINT_PER_CURVE == 0xFF);
+				static_assert(HAIR_MAX_NUM_POINT_PER_GROUP == 0xFFFFFF);
+				UE_CLOG(PointCount > HAIR_MAX_NUM_POINT_PER_CURVE, LogGroomBuilder, Warning, TEXT("[Groom] Curve point count is greater than the allowed limit (%d/%d)"), PointCount, uint32(HAIR_MAX_NUM_POINT_PER_CURVE));
+				UE_CLOG(IndexOffset > HAIR_MAX_NUM_POINT_PER_GROUP, LogGroomBuilder, Warning, TEXT("[Groom] Curve point offset is greater than the allowed limit (%d/%d)"), IndexOffset, uint32(HAIR_MAX_NUM_POINT_PER_GROUP));
 				FHairStrandsCurveFormat::Type& Curve = OutPackedCurves[CurveIndex];
-				Curve.PointCount  = FMath::Min(uint32(PointCount), 0xFFu);
-				Curve.PointOffset = FMath::Min(uint32(IndexOffset), 0xFFFFFFu);
+				Curve.PointCount  = FMath::Min(uint32(PointCount), HAIR_MAX_NUM_POINT_PER_CURVE);
+				Curve.PointOffset = FMath::Min(uint32(IndexOffset), HAIR_MAX_NUM_POINT_PER_GROUP);
 			}
 		}
 
@@ -355,7 +358,7 @@ namespace HairStrandsBuilder
 
 			for (uint32 AttributeIt=0; AttributeIt< HAIR_ATTRIBUTE_COUNT; ++AttributeIt)
 			{
-				OutBulkData.AttributeOffsets[AttributeIt] = 0xFFFFFFFF;
+				OutBulkData.AttributeOffsets[AttributeIt] = HAIR_ATRIBUTE_INVALID_OFFSET;
 			}
 
 			if (HasHairAttribute(Attributes, EHairAttribute::RootUV))
@@ -1563,9 +1566,10 @@ bool FGroomBuilder::BuildHairDescriptionGroups(const FHairDescription& HairDescr
 	// Track the imported max hair width among all imported CVs. This information is displayed to artsits to set/tune groom asset later.
 	float DDCMaxHairWidth = -1.f;
 
+	bool bCurveCountWarningIssued = false;
+	bool bPointCountWarningIssued = false;
+
 	int32 GlobalVertexIndex = 0;
-	int32 NumHairPoints = 0;
-	int32 NumGuidePoints = 0;
 	for (int32 CurveIndex = 0; CurveIndex < NumCurves; ++CurveIndex)
 	{
 		FStrandID StrandID(CurveIndex);
@@ -1595,19 +1599,24 @@ bool FGroomBuilder::BuildHairDescriptionGroups(const FHairDescription& HairDescr
 		FHairStrandsDatas* CurrentHairStrandsDatas = nullptr;
 		FHairDescriptionGroup& Group = FindOrAdd(GroupID, GroupName);
 		check(Group.Info.GroupID == GroupID);
+		bool bNumCurveValid = false;
 		if (!bIsGuide)
 		{
-			NumHairPoints += CurveNumVertices;
 			CurrentHairStrandsDatas = &Group.Strands;
-
-			++Group.Info.NumCurves;
+			if (Group.Info.NumCurves < HAIR_MAX_NUM_CURVE_PER_GROUP)
+			{
+				bNumCurveValid = true;
+				++Group.Info.NumCurves;
+			}
 		}
 		else if (bImportGuides)
 		{
-			NumGuidePoints += CurveNumVertices;
 			CurrentHairStrandsDatas = &Group.Guides;
-
-			++Group.Info.NumGuides;
+			if (Group.Info.NumGuides < HAIR_MAX_NUM_CURVE_PER_GROUP)
+			{
+				bNumCurveValid = true;
+				++Group.Info.NumGuides;
+			}
 		}
 		else
 		{
@@ -1616,10 +1625,35 @@ bool FGroomBuilder::BuildHairDescriptionGroups(const FHairDescription& HairDescr
 			continue;
 		}
 
+		// If the current curve has more control point than allows, issue a warning
+		if (CurveNumVertices > HAIR_MAX_NUM_POINT_PER_CURVE && !bPointCountWarningIssued)
+		{
+			UE_LOG(LogGroomBuilder, Warning, TEXT("[Groom] Groom contains strands with more than %d control points. Control points beyond that limit will be trimmed"), uint32(HAIR_MAX_NUM_POINT_PER_CURVE));
+			Group.Info.Flags |= uint32(EHairGroupInfoFlags::HasTrimmedPoint);
+		}
+
+		// If the current group has reached the max number of curve, skip them
+		if (!bNumCurveValid)
+		{
+			if (!bCurveCountWarningIssued)
+			{
+				UE_LOG(LogGroomBuilder, Warning, TEXT("[Groom] Group has more than %d curves per group. Curve beyong that limit won't be part of the groom"), uint32(HAIR_MAX_NUM_CURVE_PER_GROUP));
+				bCurveCountWarningIssued = true;
+			}
+			Group.Info.Flags |= uint32(EHairGroupInfoFlags::HasTrimmedCurve);
+
+
+			// Skip it
+			GlobalVertexIndex += CurveNumVertices;
+			continue;
+		}
+
 		if (!CurrentHairStrandsDatas)
+		{
 			continue;
+		}
 
-		CurrentHairStrandsDatas->StrandsCurves.CurvesCount.Add(CurveNumVertices);
+		CurrentHairStrandsDatas->StrandsCurves.CurvesCount.Add(FMath::Min(uint32(CurveNumVertices), HAIR_MAX_NUM_POINT_PER_CURVE));
 
 		if (bCanUseClosestGuidesAndWeights)
 		{
@@ -1673,6 +1707,12 @@ bool FGroomBuilder::BuildHairDescriptionGroups(const FHairDescription& HairDescr
 		{
 			FVertexID VertexID(GlobalVertexIndex);
 
+			// Skip vertex beyond the supported limit. Continue to loop for incrementing GlobalVertexIndex
+			if (VertexIndex >= HAIR_MAX_NUM_POINT_PER_CURVE)
+			{
+				continue;
+			}
+
 			CurrentHairStrandsDatas->StrandsPoints.PointsPosition.Add(VertexPositions[VertexID]);
 
 			if (bHasBaseColorAttribute)
diff --git a/Engine/Plugins/Runtime/HairStrands/Source/HairStrandsCore/Public/GroomAsset.h b/Engine/Plugins/Runtime/HairStrands/Source/HairStrandsCore/Public/GroomAsset.h
--- a/Engine/Plugins/Runtime/HairStrands/Source/HairStrandsCore/Public/GroomAsset.h
+++ b/Engine/Plugins/Runtime/HairStrands/Source/HairStrandsCore/Public/GroomAsset.h
@@ -32,6 +32,12 @@ struct FHairStrandsRestResource;
 struct FHairStrandsInterpolationResource;
 struct FHairStrandsRaytracingResource;
 
+enum class EHairGroupInfoFlags : uint8
+{
+	HasTrimmedPoint = 1,
+	HasTrimmedCurve = 2
+};
+
 USTRUCT(BlueprintType)
 struct HAIRSTRANDSCORE_API FHairGroupInfo
 {
@@ -60,6 +66,9 @@ struct HAIRSTRANDSCORE_API FHairGroupInfo
 
 	UPROPERTY(VisibleAnywhere, Category = "Info", meta = (DisplayName = "Card Group ID for split hair card generation"))
 	int32 GroupCardsID = 0;
+
+	UPROPERTY()
+	uint32 Flags = 0;
 };
 
 USTRUCT(BlueprintType)
diff --git a/Engine/Plugins/Runtime/HairStrands/Source/HairStrandsCore/Public/GroomImportOptions.h b/Engine/Plugins/Runtime/HairStrands/Source/HairStrandsCore/Public/GroomImportOptions.h
--- a/Engine/Plugins/Runtime/HairStrands/Source/HairStrandsCore/Public/GroomImportOptions.h
+++ b/Engine/Plugins/Runtime/HairStrands/Source/HairStrandsCore/Public/GroomImportOptions.h
@@ -27,28 +27,23 @@ struct HAIRSTRANDSCORE_API FGroomHairGroupPreview
 {
 	GENERATED_USTRUCT_BODY()
 
-	FGroomHairGroupPreview()
-	: GroupID(0)
-	, CurveCount(0)
-	, GuideCount(0)
-	, Attributes(0)
-	, InterpolationSettings()
-	{}
-
 	UPROPERTY(VisibleAnywhere, BlueprintReadWrite, Category = Preview)
 	FName GroupName;
 
 	UPROPERTY(VisibleAnywhere, BlueprintReadWrite, Category = Preview)
-	int32 GroupID;
+	int32 GroupID = 0;
 
 	UPROPERTY(VisibleAnywhere, BlueprintReadWrite, Category = Preview)
-	int32 CurveCount;
+	int32 CurveCount = 0;
 
 	UPROPERTY(VisibleAnywhere, BlueprintReadWrite, Category = Preview)
-	int32 GuideCount;
+	int32 GuideCount = 0;
+
+	UPROPERTY()
+	uint32 Attributes = 0;
 
 	UPROPERTY()
-	uint32 Attributes;
+	uint32 Flags = 0;
 
 	UPROPERTY(EditAnywhere, BlueprintReadWrite, Category = Preview)
 	FHairGroupsInterpolation InterpolationSettings;
@@ -62,4 +57,4 @@ class HAIRSTRANDSCORE_API UGroomHairGroupsPreview : public UObject
 public:
 	UPROPERTY(config, EditAnywhere, BlueprintReadWrite, meta = (ShowOnlyInnerProperties), Category = Preview)
 	TArray<FGroomHairGroupPreview> Groups;
-};
\ No newline at end of file
+};
diff --git a/Engine/Plugins/Runtime/HairStrands/Source/HairStrandsEditor/Private/GroomImportOptionsWindow.cpp b/Engine/Plugins/Runtime/HairStrands/Source/HairStrandsEditor/Private/GroomImportOptionsWindow.cpp
--- a/Engine/Plugins/Runtime/HairStrands/Source/HairStrandsEditor/Private/GroomImportOptionsWindow.cpp
+++ b/Engine/Plugins/Runtime/HairStrands/Source/HairStrandsEditor/Private/GroomImportOptionsWindow.cpp
@@ -31,6 +31,9 @@ enum class EHairDescriptionStatus
 	GuidesOnly, // guides-only with unspecified groom asset
 	GuidesOnlyCompatible,
 	GuidesOnlyIncompatible,
+	ValidPointLimit,
+	ValidCurveLimit,
+	ValidCurveAndPointLimit,
 	Unknown
 };
 
@@ -73,6 +76,23 @@ void SGroomImportOptionsWindow::UpdateStatus(UGroomHairGroupsPreview* Descriptio
 		}
 	}
 
+	// Check if any curve or point have been trimmed
+	for (const FGroomHairGroupPreview& Group : Description->Groups)
+	{
+		if ((Group.Flags & uint32(EHairGroupInfoFlags::HasTrimmedCurve)) && (Group.Flags & uint32(EHairGroupInfoFlags::HasTrimmedPoint)))
+		{
+			CurrentStatus = EHairDescriptionStatus::ValidCurveAndPointLimit;
+		}
+		else if (Group.Flags & uint32(EHairGroupInfoFlags::HasTrimmedCurve))
+		{
+			CurrentStatus = EHairDescriptionStatus::ValidCurveLimit;
+		}
+		else if (Group.Flags & uint32(EHairGroupInfoFlags::HasTrimmedPoint))
+		{
+			CurrentStatus = EHairDescriptionStatus::ValidPointLimit;
+		}
+	}
+
 	if (!bImportGroomCache)
 	{
 		return;
@@ -152,6 +172,12 @@ FText SGroomImportOptionsWindow::GetStatusText() const
 			return LOCTEXT("GroomOptionsWindow_ValidationText8", "Only guides were detected. The groom asset provided is compatible.");
 		case EHairDescriptionStatus::GuidesOnlyIncompatible:
 			return LOCTEXT("GroomOptionsWindow_ValidationText9", "Only guides were detected. The groom asset provided is incompatible.");
+		case EHairDescriptionStatus::ValidCurveLimit:
+			return LOCTEXT("GroomOptionsWindow_ValidationText10", "Valid. At least one group contains more curves than allowed limit (Max:4M). Curves beyond that limit will be trimmed."); static_assert(HAIR_MAX_NUM_CURVE_PER_GROUP == 4194303);
+		case EHairDescriptionStatus::ValidPointLimit:
+			return LOCTEXT("GroomOptionsWindow_ValidationText11", "Valid. At least one group contains more control points per curve than the allowed limit (Max:255). Control points beyond that limit will be trimmed."); static_assert(HAIR_MAX_NUM_POINT_PER_CURVE == 255);
+		case EHairDescriptionStatus::ValidCurveAndPointLimit:
+			return LOCTEXT("GroomOptionsWindow_ValidationText12", "Valid. At least one group contains more control points per curve and more curves than the allowed limit (curve limit:4M, point limit:255). Curves and control points beyond that limit will be trimmed.");
 		case EHairDescriptionStatus::Unset:
 		case EHairDescriptionStatus::Unknown:
 		default:
@@ -176,6 +202,10 @@ FSlateColor SGroomImportOptionsWindow::GetStatusColor() const
 		case EHairDescriptionStatus::GroomCache:
 		case EHairDescriptionStatus::GuidesOnly:
 			return FLinearColor(0.80f, 0.80f, 0, 1);
+		case EHairDescriptionStatus::ValidCurveLimit:
+		case EHairDescriptionStatus::ValidPointLimit:
+		case EHairDescriptionStatus::ValidCurveAndPointLimit:
+			return FLinearColor(0.80f, 0.80f, 0, 1);
 		case EHairDescriptionStatus::Unset:
 		case EHairDescriptionStatus::Unknown:
 		default:
@@ -443,6 +473,9 @@ bool SGroomImportOptionsWindow::CanImport() const
 		case EHairDescriptionStatus::Valid:
 		case EHairDescriptionStatus::GroomCacheCompatible:
 		case EHairDescriptionStatus::GuidesOnlyCompatible:
+		case EHairDescriptionStatus::ValidPointLimit:
+		case EHairDescriptionStatus::ValidCurveLimit:
+		case EHairDescriptionStatus::ValidCurveAndPointLimit:
 			return true;
 		case EHairDescriptionStatus::Unset:
 		case EHairDescriptionStatus::NoGroup:
diff --git a/Engine/Plugins/Runtime/HairStrands/Source/HairStrandsEditor/Private/HairStrandsFactory.cpp b/Engine/Plugins/Runtime/HairStrands/Source/HairStrandsEditor/Private/HairStrandsFactory.cpp
--- a/Engine/Plugins/Runtime/HairStrands/Source/HairStrandsEditor/Private/HairStrandsFactory.cpp
+++ b/Engine/Plugins/Runtime/HairStrands/Source/HairStrandsEditor/Private/HairStrandsFactory.cpp
@@ -143,6 +143,7 @@ UObject* UHairStrandsFactory::FactoryCreateFile(UClass* InClass, UObject* InPare
 				OutGroup.CurveCount = Group.Info.NumCurves;
 				OutGroup.GuideCount = Group.Info.NumGuides;
 				OutGroup.Attributes = Group.Attributes;
+				OutGroup.Flags      = Group.Info.Flags;
 				bGuidesOnly |= (OutGroup.CurveCount == 0 && OutGroup.GuideCount > 0);
 
 				if (OutGroup.GroupID < OutDescription.HairGroups.Num())
diff --git a/Engine/Plugins/Runtime/HairStrands/Source/HairStrandsEditor/Private/ReimportHairStrandsFactory.cpp b/Engine/Plugins/Runtime/HairStrands/Source/HairStrandsEditor/Private/ReimportHairStrandsFactory.cpp
--- a/Engine/Plugins/Runtime/HairStrands/Source/HairStrandsEditor/Private/ReimportHairStrandsFactory.cpp
+++ b/Engine/Plugins/Runtime/HairStrands/Source/HairStrandsEditor/Private/ReimportHairStrandsFactory.cpp
@@ -385,7 +385,7 @@ EReimportResult::Type UReimportHairStrandsFactory::Reimport(UObject* Obj)
 				OutGroup.CurveCount = Group.Info.NumCurves;
 				OutGroup.GuideCount = Group.Info.NumGuides;
 				OutGroup.Attributes = Group.Attributes;
-
+				OutGroup.Flags		= Group.Info.Flags;
 				if (OutGroup.GroupID < OutDescription.HairGroups.Num())
 				{
 					OutGroup.InterpolationSettings = GroomReimportOptions->InterpolationSettings[OutGroup.GroupID];
diff --git a/Engine/Shaders/Private/HairStrands/HairStrandsAttributeCommon.ush b/Engine/Shaders/Private/HairStrands/HairStrandsAttributeCommon.ush
--- a/Engine/Shaders/Private/HairStrands/HairStrandsAttributeCommon.ush
+++ b/Engine/Shaders/Private/HairStrands/HairStrandsAttributeCommon.ush
@@ -60,7 +60,7 @@ float InternalGetHairStrandsWorldRadius(Buffer<uint4> InPositionBuffer, uint Ver
 
 bool IsAttributeValid(uint InOffset)
 {
-	return InOffset != 0xFFFFFFFF;
+	return InOffset != HAIR_ATRIBUTE_INVALID_OFFSET;
 }
 
 ////////////////////////////////////////////////////////////////////////////////
diff --git a/Engine/Shaders/Private/HairStrands/HairStrandsVertexFactoryCommon.ush b/Engine/Shaders/Private/HairStrands/HairStrandsVertexFactoryCommon.ush
--- a/Engine/Shaders/Private/HairStrands/HairStrandsVertexFactoryCommon.ush
+++ b/Engine/Shaders/Private/HairStrands/HairStrandsVertexFactoryCommon.ush
@@ -124,8 +124,12 @@ struct FHairCurve
 
 FHairCurve UnpackHairCurve(uint In)
 {
+	#if HAIR_MAX_NUM_POINT_PER_CURVE != 0xFF || HAIR_MAX_NUM_POINT_PER_GROUP != 0xFFFFFFFF
+	#error Update hair curve ENCODING_TYPE
+	#endif
+
 	FHairCurve Out = (FHairCurve)0;
-	Out.PointIndex = In & 0xFFFFFF;
+	Out.PointIndex = In & HAIR_MAX_NUM_POINT_PER_GROUP;
 	Out.PointCount = In >> 24;
 	return Out;
 }
diff --git a/Engine/Shaders/Shared/HairStrandsDefinitions.h b/Engine/Shaders/Shared/HairStrandsDefinitions.h
--- a/Engine/Shaders/Shared/HairStrandsDefinitions.h
+++ b/Engine/Shaders/Shared/HairStrandsDefinitions.h
@@ -20,6 +20,13 @@
 #define HAIR_ATTRIBUTE_AO 6
 #define HAIR_ATTRIBUTE_COUNT 7
 
+// Groom limits (based on encoding)
+#define HAIR_MAX_NUM_POINT_PER_CURVE ((1u<<8)-1u)
+#define HAIR_MAX_NUM_POINT_PER_GROUP ((1u<<24)-1u)
+#define HAIR_MAX_NUM_CURVE_PER_GROUP ((1u<<22)-1u)
+
+#define HAIR_ATRIBUTE_INVALID_OFFSET 0xFFFFFFFF
+
 // Max number of discrete LOD that a hair group can have
 #define MAX_HAIR_LOD 8
 
