Change how groom cache update its data.

* Interpolation now run on the GPU
* GPU resources are now cache to avoid constant reupload of the data every frame.

diff --git a/Engine/Plugins/Runtime/HairStrands/Source/HairStrandsCore/Private/GroomComponent.cpp b/Engine/Plugins/Runtime/HairStrands/Source/HairStrandsCore/Private/GroomComponent.cpp
--- a/Engine/Plugins/Runtime/HairStrands/Source/HairStrandsCore/Private/GroomComponent.cpp
+++ b/Engine/Plugins/Runtime/HairStrands/Source/HairStrandsCore/Private/GroomComponent.cpp
@@ -992,11 +992,6 @@ public:
 		return NextFrame;
 	}
 
-	virtual const FGroomCacheAnimationData& GetInterpolatedFrameBuffer() override
-	{
-		return InterpolatedFrame;
-	}
-
 	virtual int32 GetCurrentFrameIndex() const override
 	{
 		return CurrentFrameIndex;
@@ -1019,14 +1014,11 @@ public:
 		// Find the frame indices and interpolation factor to interpolate between
 		int32 FrameIndexA = 0;
 		int32 FrameIndexB = 0;
-		float OutInterpolationFactor = 0.0f;
-		GroomCache->GetFrameIndicesAtTime(Time, bIsLooping, false, FrameIndexA, FrameIndexB, OutInterpolationFactor);
+		GroomCache->GetFrameIndicesAtTime(Time, bIsLooping, false, FrameIndexA, FrameIndexB, InterpolationFactor);
 
 		// Update and cache the frame data as needed
-		bool bComputeInterpolation = false;
 		if (FrameIndexA != CurrentFrameIndex)
 		{
-			bComputeInterpolation = true;
 			if (FrameIndexA == NextFrameIndex)
 			{
 				Swap(CurrentFrame, NextFrame);
@@ -1041,97 +1033,9 @@ public:
 
 		if (FrameIndexB != NextFrameIndex)
 		{
-			bComputeInterpolation = true;
 			GroomCache->GetGroomDataAtFrameIndex(FrameIndexB, NextFrame);
 			NextFrameIndex = FrameIndexB;
 		}
-
-		// Make sure the initial interpolated frame is populated with valid data
-		if (InterpolatedFrame.GroupsData.Num() != CurrentFrame.GroupsData.Num())
-		{
-			InterpolatedFrame = CurrentFrame;
-		}
-
-		bComputeInterpolation = bComputeInterpolation || !FMath::IsNearlyEqual(InterpolationFactor, OutInterpolationFactor, KINDA_SMALL_NUMBER);
-
-		// Do interpolation of vertex positions if needed
-		if (bComputeInterpolation)
-		{
-			Interpolate(CurrentFrame, NextFrame, OutInterpolationFactor);
-		}
-	}
-
-	void Interpolate(const FGroomCacheAnimationData& FrameA, const FGroomCacheAnimationData& FrameB, float InInterpolationFactor)
-	{
-		TRACE_CPUPROFILER_EVENT_SCOPE(FGroomCacheBuffers::InterpolateCPU);
-
-		InterpolationFactor = InInterpolationFactor;
-
-		FScopeLock Lock(GetCriticalSection());
-
-		const int32 NumGroups = FrameA.GroupsData.Num();
-		for (int32 GroupIndex = 0; GroupIndex < NumGroups; ++GroupIndex)
-		{
-			TRACE_CPUPROFILER_EVENT_SCOPE(FGroomCacheBuffers::InterpolateCPU_Group);
-
-			const FGroomCacheGroupData& CurrentGroupData = FrameA.GroupsData[GroupIndex];
-			const FGroomCacheGroupData& NextGroupData = FrameB.GroupsData[GroupIndex];
-			FGroomCacheGroupData& InterpolatedGroupData = InterpolatedFrame.GroupsData[GroupIndex];
-			const int32 NumVertices = CurrentGroupData.VertexData.PointsPosition.Num();
-			const int32 NextNumVertices = NextGroupData.VertexData.PointsPosition.Num();
-
-			// Update the bounding box used for hair strands rendering computation
-			FVector InterpolatedCenter = FMath::Lerp(CurrentGroupData.BoundingBox.GetCenter(), NextGroupData.BoundingBox.GetCenter(), InterpolationFactor);
-			InterpolatedGroupData.BoundingBox = CurrentGroupData.BoundingBox.MoveTo(InterpolatedCenter) + NextGroupData.BoundingBox.MoveTo(InterpolatedCenter);
-
-			const bool bHasRadiusData = InterpolatedGroupData.VertexData.PointsRadius.Num() > 0;
-			if (bHasRadiusData)
-			{
-				InterpolatedGroupData.StrandData.MaxRadius = FMath::Lerp(CurrentGroupData.StrandData.MaxRadius, NextGroupData.StrandData.MaxRadius, InterpolationFactor);
-			}
-
-			if (NumVertices == NextNumVertices)
-			{
-				// In case the topology is varying, make sure the interpolated group data can hold the required number of vertices
-				InterpolatedGroupData.VertexData.PointsPosition.SetNum(NumVertices);
-
-			// Parallel batched interpolation
-			const int32 BatchSize = 1024;
-			const int32 BatchCount = (NumVertices + BatchSize - 1) / BatchSize;
-
-			ParallelFor(BatchCount, [&](int32 BatchIndex)
-			{
-				const int32 Start = BatchIndex * BatchSize;
-				const int32 End = FMath::Min(Start + BatchSize, NumVertices); // one-past end index
-
-				for (int32 VertexIndex = Start; VertexIndex < End; ++VertexIndex)
-				{
-					const FVector3f& CurrentPosition = CurrentGroupData.VertexData.PointsPosition[VertexIndex];
-					const FVector3f& NextPosition = NextGroupData.VertexData.PointsPosition[VertexIndex];
-
-					InterpolatedGroupData.VertexData.PointsPosition[VertexIndex] = FMath::Lerp(CurrentPosition, NextPosition, InterpolationFactor);
-
-					if (bHasRadiusData)
-					{
-						const float CurrentRadius = CurrentGroupData.VertexData.PointsRadius[VertexIndex];
-						const float NextRadius = NextGroupData.VertexData.PointsRadius[VertexIndex];
-
-						InterpolatedGroupData.VertexData.PointsRadius[VertexIndex] = FMath::Lerp(CurrentRadius, NextRadius, InterpolationFactor);
-					}
-				}
-			});
-		}
-			else
-			{
-				// Cannot interpolate, use the closest frame
-				InterpolatedGroupData.VertexData.PointsPosition = InterpolationFactor < 0.5f ? CurrentGroupData.VertexData.PointsPosition : NextGroupData.VertexData.PointsPosition;
-
-				if (bHasRadiusData)
-				{
-					InterpolatedGroupData.VertexData.PointsRadius = InterpolationFactor < 0.5f ? CurrentGroupData.VertexData.PointsRadius : NextGroupData.VertexData.PointsRadius;
-				}
-	}
-		}
 	}
 
 	FBox GetBoundingBox()
@@ -1161,9 +1065,6 @@ protected:
 	FGroomCacheAnimationData CurrentFrame;
 	FGroomCacheAnimationData NextFrame;
 
-	/** Used for CPU interpolation */
-	FGroomCacheAnimationData InterpolatedFrame;
-
 	int32 CurrentFrameIndex = -1;
 	int32 NextFrameIndex = -1;
 	float InterpolationFactor = 0.0f;
@@ -1238,14 +1139,11 @@ public:
 		// Find the frame indices and interpolation factor to interpolate between
 		int32 FrameIndexA = 0;
 		int32 FrameIndexB = 0;
-		float OutInterpolationFactor = 0.0f;
-		GroomCache->GetFrameIndicesAtTime(Time, bIsLooping, false, FrameIndexA, FrameIndexB, OutInterpolationFactor);
+		GroomCache->GetFrameIndicesAtTime(Time, bIsLooping, false, FrameIndexA, FrameIndexB, InterpolationFactor);
 
 		// Update and cache the frame data as needed
-		bool bComputeInterpolation = false;
 		if (FrameIndexA != CurrentFrameIndex)
 		{
-			bComputeInterpolation = true;
 			if (FrameIndexA == NextFrameIndex)
 			{
 				// NextFrame is mapped to increment its ref count, but this could fail if the frame was evicted
@@ -1272,8 +1170,6 @@ public:
 
 		if (FrameIndexB != NextFrameIndex)
 		{
-			bComputeInterpolation = true;
-
 			const FGroomCacheAnimationData* DataPtr = IGroomCacheStreamingManager::Get().MapAnimationData(GroomCache, FrameIndexB);
 			if (DataPtr)
 			{
@@ -1287,20 +1183,6 @@ public:
 		{
 			return;
 		}
-
-		// Make sure the initial interpolated frame is populated with valid data
-		if (InterpolatedFrame.GroupsData.Num() != CurrentFramePtr->GroupsData.Num())
-		{
-			InterpolatedFrame = *CurrentFramePtr;
-		}
-
-		bComputeInterpolation = bComputeInterpolation || !FMath::IsNearlyEqual(InterpolationFactor, OutInterpolationFactor, KINDA_SMALL_NUMBER);
-
-		// Do interpolation of vertex positions if needed
-		if (bComputeInterpolation)
-		{
-			Interpolate(*CurrentFramePtr, *NextFramePtr, OutInterpolationFactor);
-		}
 	}
 
 private:
diff --git a/Engine/Plugins/Runtime/HairStrands/Source/HairStrandsCore/Private/HairStrandsRendering.cpp b/Engine/Plugins/Runtime/HairStrands/Source/HairStrandsCore/Private/HairStrandsRendering.cpp
--- a/Engine/Plugins/Runtime/HairStrands/Source/HairStrandsCore/Private/HairStrandsRendering.cpp
+++ b/Engine/Plugins/Runtime/HairStrands/Source/HairStrandsCore/Private/HairStrandsRendering.cpp
@@ -92,6 +92,76 @@ enum class EDeformationType : uint8
 	OffsetGuide		// Offset the guides
 };
 
+
+struct FGroomCacheResources
+{
+	FRDGBufferSRVRef PositionBuffer = nullptr;
+	FRDGBufferSRVRef RadiusBuffer = nullptr;
+	bool bHasRadiusData = false;
+};
+
+static FGroomCacheResources CreateGroomCacheBuffer(FRDGBuilder& GraphBuilder, FGroomCacheVertexData& InVertexData)
+{
+	FGroomCacheResources Out;
+	
+	if (InVertexData.PointsPosition.Num() > 0)
+	{
+		FRDGBufferRef PositionBuffer = nullptr;
+		if (InVertexData.PositionBuffer == nullptr)
+		{
+			const uint32 DataCount = InVertexData.PointsPosition.Num();
+			const uint32 DataSizeInBytes = sizeof(FVector3f) * DataCount;
+			// Deformation are upload into a Buffer<float> as the original position are float3 which is is both
+			// 1) incompatible with structure buffer alignment (128bits), and 2) incompatible with vertex buffer 
+			// as R32G32B32_FLOAT format is not well supported for SRV across HW.
+			// So instead the positions are uploaded into vertex buffer Buffer<float>
+			PositionBuffer = CreateVertexBuffer(
+				GraphBuilder,
+				TEXT("GroomCache_PositionBuffer"),
+				FRDGBufferDesc::CreateBufferDesc(sizeof(float), InVertexData.PointsPosition.Num() * 3),
+				InVertexData.PointsPosition.GetData(),
+				DataSizeInBytes,
+				ERDGInitialDataFlags::None);
+			GraphBuilder.QueueBufferExtraction(PositionBuffer, &InVertexData.PositionBuffer);
+		}
+		else
+		{
+			PositionBuffer = GraphBuilder.RegisterExternalBuffer(InVertexData.PositionBuffer);
+		}
+		Out.PositionBuffer = GraphBuilder.CreateSRV(PositionBuffer, PF_R32_FLOAT);
+	}
+
+	const uint32 RadiusDataCount = InVertexData.PointsRadius.Num();
+	if (RadiusDataCount > 0)
+	{
+		FRDGBufferRef RadiusBuffer = nullptr;
+		if (InVertexData.RadiusBuffer == nullptr)
+		{
+			const uint32 RadiusDataSizeInBytes = sizeof(float) * RadiusDataCount;
+			RadiusBuffer = CreateVertexBuffer(
+				GraphBuilder,
+				TEXT("GroomCache_RadiusBuffer"),
+				FRDGBufferDesc::CreateBufferDesc(sizeof(float), InVertexData.PointsRadius.Num()),
+				InVertexData.PointsRadius.GetData(),
+				RadiusDataSizeInBytes,
+				ERDGInitialDataFlags::None);
+			GraphBuilder.QueueBufferExtraction(RadiusBuffer, &InVertexData.PositionBuffer);
+		}
+		else
+		{
+			RadiusBuffer = GraphBuilder.RegisterExternalBuffer(InVertexData.RadiusBuffer);
+		}
+		Out.bHasRadiusData = true;
+		Out.RadiusBuffer = GraphBuilder.CreateSRV(RadiusBuffer, PF_R32_FLOAT);
+	}
+	else
+	{
+		Out.RadiusBuffer = GraphBuilder.CreateSRV(GSystemTextures.GetDefaultBuffer(GraphBuilder, 4), PF_R32_FLOAT);
+	}
+
+	return Out;
+}
+
 ////////////////////////////////////////////////////////////////////////////////////////////////////////////////////////
 class FTransferVelocityPassCS : public FGlobalShader
 {
@@ -150,9 +220,12 @@ class FGroomCacheUpdatePassCS : public FGlobalShader
 
 	BEGIN_SHADER_PARAMETER_STRUCT(FParameters, )
 		SHADER_PARAMETER(uint32, ElementCount)
-		SHADER_PARAMETER(int, bHasRadiusData)
-		SHADER_PARAMETER_RDG_BUFFER_SRV(Buffer, InAnimatedBuffer)
-		SHADER_PARAMETER_RDG_BUFFER_SRV(Buffer, InRadiusBuffer)
+		SHADER_PARAMETER(uint32, bHasRadiusData)
+		SHADER_PARAMETER(float, InterpolationFactor)
+		SHADER_PARAMETER_RDG_BUFFER_SRV(Buffer, InPosition0Buffer)
+		SHADER_PARAMETER_RDG_BUFFER_SRV(Buffer, InPosition1Buffer)
+		SHADER_PARAMETER_RDG_BUFFER_SRV(Buffer, InRadius0Buffer)
+		SHADER_PARAMETER_RDG_BUFFER_SRV(Buffer, InRadius1Buffer)
 		SHADER_PARAMETER_RDG_BUFFER_SRV(Buffer, InRestPoseBuffer)
 		SHADER_PARAMETER_RDG_BUFFER_SRV(Buffer, InDeformedOffsetBuffer)
 		SHADER_PARAMETER_RDG_BUFFER_UAV(RWBuffer, OutDeformedBuffer)
@@ -174,70 +247,31 @@ IMPLEMENT_GLOBAL_SHADER(FGroomCacheUpdatePassCS, "/Engine/Private/HairStrands/Ha
 static void AddGroomCacheUpdatePass(
 	FRDGBuilder& GraphBuilder,
 	FGlobalShaderMap* ShaderMap,
-	uint32 ElementCount,
-	const FGroomCacheGroupData& GroomCacheData,
+	uint32 PointCount,
+	float InterpolationFactor,
+	FGroomCacheResources CacheResources0,
+	FGroomCacheResources CacheResources1,
 	FRDGBufferSRVRef InBuffer,
 	FRDGBufferSRVRef InDeformedOffsetBuffer,
-	FRDGBufferUAVRef OutBuffer
-	)
+	FRDGBufferUAVRef OutBuffer)
 {
 	TRACE_CPUPROFILER_EVENT_SCOPE(AddGroomCacheUpdatePass);
 
-	if (ElementCount == 0) return;
-
-	FRDGBufferRef VertexBuffer = nullptr;
-
-	const uint32 DataCount = GroomCacheData.VertexData.PointsPosition.Num();
-	const uint32 DataSizeInBytes = sizeof(FVector3f) * DataCount;
-	if (DataSizeInBytes != 0)
-	{
-		// Deformation are upload into a Buffer<float> as the original position are float3 which is is both
-		// 1) incompatible with structure buffer alignment (128bits), and 2) incompatible with vertex buffer 
-		// as R32G32B32_FLOAT format is not well supported for SRV across HW.
-		// So instead the positions are uploaded into vertex buffer Buffer<float>
-		VertexBuffer = CreateVertexBuffer(
-			GraphBuilder,
-			TEXT("GroomCache_PositionBuffer"),
-			FRDGBufferDesc::CreateBufferDesc(sizeof(float), GroomCacheData.VertexData.PointsPosition.Num() * 3),
-			GroomCacheData.VertexData.PointsPosition.GetData(),
-			DataSizeInBytes,
-			ERDGInitialDataFlags::None);
-	}
-	else
-	{
-		return;
-	}
+	if (PointCount == 0 || !CacheResources0.PositionBuffer || !CacheResources1.PositionBuffer) return;
 
 	FGroomCacheUpdatePassCS::FParameters* Parameters = GraphBuilder.AllocParameters<FGroomCacheUpdatePassCS::FParameters>();
-	Parameters->ElementCount = ElementCount;
-	Parameters->InAnimatedBuffer = GraphBuilder.CreateSRV(VertexBuffer, PF_R32_FLOAT);
+	Parameters->ElementCount = PointCount;
+	Parameters->InPosition0Buffer = CacheResources0.PositionBuffer;
+	Parameters->InPosition1Buffer = CacheResources1.PositionBuffer;
+	Parameters->InRadius0Buffer = CacheResources0.RadiusBuffer;
+	Parameters->InRadius1Buffer = CacheResources1.RadiusBuffer;
 	Parameters->InRestPoseBuffer = InBuffer;
 	Parameters->InDeformedOffsetBuffer = InDeformedOffsetBuffer;
 	Parameters->OutDeformedBuffer = OutBuffer;
+	Parameters->InterpolationFactor = InterpolationFactor;
+	Parameters->bHasRadiusData = CacheResources0.bHasRadiusData ? 1u : 0u;
 
-	const uint32 RadiusDataCount = GroomCacheData.VertexData.PointsRadius.Num();
-	const bool bHasRadiusData = RadiusDataCount > 0;
-	FRDGBufferRef RadiusBuffer = nullptr;
-	if (bHasRadiusData)
-	{
-		const uint32 RadiusDataSizeInBytes = sizeof(float) * RadiusDataCount;
-		RadiusBuffer = CreateVertexBuffer(
-			GraphBuilder,
-			TEXT("GroomCache_RadiusBuffer"),
-			FRDGBufferDesc::CreateBufferDesc(sizeof(float), GroomCacheData.VertexData.PointsRadius.Num()),
-			GroomCacheData.VertexData.PointsRadius.GetData(),
-			RadiusDataSizeInBytes,
-			ERDGInitialDataFlags::None);
-
-		Parameters->InRadiusBuffer = GraphBuilder.CreateSRV(RadiusBuffer, PF_R32_FLOAT);
-	}
-	else
-	{
-		Parameters->InRadiusBuffer = GraphBuilder.CreateSRV(GSystemTextures.GetDefaultBuffer(GraphBuilder, 4), PF_R32_FLOAT);
-	}
-	Parameters->bHasRadiusData = bHasRadiusData ? 1 : 0;
-
-	const FIntVector DispatchCount = FIntVector(FMath::DivideAndRoundUp(ElementCount, FGroomCacheUpdatePassCS::GetGroupSize()), 1, 1);
+	const FIntVector DispatchCount = FIntVector(FMath::DivideAndRoundUp(PointCount, FGroomCacheUpdatePassCS::GetGroupSize()), 1, 1);
 	TShaderMapRef<FGroomCacheUpdatePassCS> ComputeShader(ShaderMap);
 	FComputeShaderUtils::AddPass(
 		GraphBuilder,
@@ -1710,7 +1744,9 @@ void CreateHairStrandsDebugAttributeBuffer(FRDGBuilder& GraphBuilder, FRDGExtern
 
 EGroomCacheType GetHairInstanceCacheType(const FHairGroupInstance* Instance)
 {
-	const bool bHasValidGroomCacheBuffers = Instance->Debug.GroomCacheBuffers.IsValid() && Instance->Debug.GroomCacheBuffers->GetInterpolatedFrameBuffer().GroupsData.IsValidIndex(Instance->Debug.GroupIndex);
+	const bool bHasValidGroomCacheBuffers = Instance->Debug.GroomCacheBuffers.IsValid() && 
+		Instance->Debug.GroomCacheBuffers->GetCurrentFrameBuffer().GroupsData.IsValidIndex(Instance->Debug.GroupIndex) && 
+		Instance->Debug.GroomCacheBuffers->GetNextFrameBuffer().GroupsData.IsValidIndex(Instance->Debug.GroupIndex);
 	return bHasValidGroomCacheBuffers ? Instance->Debug.GroomCacheType : EGroomCacheType::None;
 }
 
@@ -1760,9 +1796,10 @@ void ComputeHairStrandsInterpolation(
 			if (ActiveGroomCacheType == EGroomCacheType::Guides)
 			{
 				FScopeLock Lock(Instance->Debug.GroomCacheBuffers->GetCriticalSection());
-				const FGroomCacheGroupData& GroomCacheGroupData = Instance->Debug.GroomCacheBuffers->GetInterpolatedFrameBuffer().GroupsData[Instance->Debug.GroupIndex];
-
-				Instance->Guides.DeformedResource->SetPositionOffset(FHairStrandsDeformedResource::EFrameType::Current, GroomCacheGroupData.BoundingBox.GetCenter());
+				const FGroomCacheGroupData& GroomCacheGroupData0 = Instance->Debug.GroomCacheBuffers->GetCurrentFrameBuffer().GroupsData[Instance->Debug.GroupIndex];
+				const FGroomCacheGroupData& GroomCacheGroupData1 = Instance->Debug.GroomCacheBuffers->GetNextFrameBuffer().GroupsData[Instance->Debug.GroupIndex];
+				const float InterpolationFactor = Instance->Debug.GroomCacheBuffers->GetInterpolationFactor();
+				Instance->Guides.DeformedResource->SetPositionOffset(FHairStrandsDeformedResource::EFrameType::Current, FMath::Lerp(GroomCacheGroupData0.BoundingBox.GetCenter(), GroomCacheGroupData1.BoundingBox.GetCenter(), InterpolationFactor));
 
 				AddHairStrandUpdatePositionOffsetPass(
 					GraphBuilder,
@@ -1835,10 +1872,17 @@ void ComputeHairStrandsInterpolation(
 
 				if (ActiveGroomCacheType == EGroomCacheType::Guides)
 				{
-					FScopeLock Lock(Instance->Debug.GroomCacheBuffers->GetCriticalSection());
-					const FGroomCacheGroupData& GroomCacheGroupData = Instance->Debug.GroomCacheBuffers->GetInterpolatedFrameBuffer().GroupsData[Instance->Debug.GroupIndex];
+					FScopeLock Lock(Instance->Debug.GroomCacheBuffers->GetCriticalSection()); // This is not ideally it will block the rendering thread / game thread
 
-					Instance->Guides.DeformedResource->SetPositionOffset(FHairStrandsDeformedResource::EFrameType::Current, GroomCacheGroupData.BoundingBox.GetCenter());
+					FGroomCacheGroupData* GroomCacheData0 = const_cast<FGroomCacheGroupData*>(&Instance->Debug.GroomCacheBuffers->GetCurrentFrameBuffer().GroupsData[Instance->Debug.GroupIndex]);
+					FGroomCacheGroupData* GroomCacheData1 = const_cast<FGroomCacheGroupData*>(&Instance->Debug.GroomCacheBuffers->GetNextFrameBuffer().GroupsData[Instance->Debug.GroupIndex]);
+					const float InterpolationFactor = Instance->Debug.GroomCacheBuffers->GetInterpolationFactor();
+					const FVector OffsetPosition = FMath::Lerp(GroomCacheData0->BoundingBox.GetCenter(), GroomCacheData1->BoundingBox.GetCenter(), InterpolationFactor);
+
+					FGroomCacheResources CacheResources0 = CreateGroomCacheBuffer(GraphBuilder, GroomCacheData0->VertexData);
+					FGroomCacheResources CacheResources1 = CreateGroomCacheBuffer(GraphBuilder, GroomCacheData1->VertexData);
+
+					Instance->Guides.DeformedResource->SetPositionOffset(FHairStrandsDeformedResource::EFrameType::Current, OffsetPosition);
 
 					AddHairStrandUpdatePositionOffsetPass(
 						GraphBuilder,
@@ -1848,7 +1892,7 @@ void ComputeHairStrandsInterpolation(
 						Instance->Guides.DeformedResource);
 
 					// Apply the same offset to the render strands for proper rendering
-					Instance->Strands.DeformedResource->SetPositionOffset(FHairStrandsDeformedResource::EFrameType::Current, GroomCacheGroupData.BoundingBox.GetCenter());
+					Instance->Strands.DeformedResource->SetPositionOffset(FHairStrandsDeformedResource::EFrameType::Current, OffsetPosition);
 
 					AddHairStrandUpdatePositionOffsetPass(
 						GraphBuilder,
@@ -1862,7 +1906,9 @@ void ComputeHairStrandsInterpolation(
 						GraphBuilder,
 						ShaderMap,
 						Instance->Guides.RestResource->GetPointCount(),
-						GroomCacheGroupData,
+						InterpolationFactor,
+						CacheResources0,
+						CacheResources1,
 						RegisterAsSRV(GraphBuilder, Instance->Guides.RestResource->PositionBuffer),
 						RegisterAsSRV(GraphBuilder, Instance->Guides.DeformedResource->GetPositionOffsetBuffer(FHairStrandsDeformedResource::EFrameType::Current)),
 						RegisterAsUAV(GraphBuilder, Instance->Guides.DeformedResource->GetBuffer(FHairStrandsDeformedResource::Current)));
@@ -1992,9 +2038,25 @@ void ComputeHairStrandsInterpolation(
 				else if (ActiveGroomCacheType == EGroomCacheType::Strands)
 				{
 					FScopeLock Lock(Instance->Debug.GroomCacheBuffers->GetCriticalSection());
-					const FGroomCacheGroupData& GroomCacheGroupData = Instance->Debug.GroomCacheBuffers->GetInterpolatedFrameBuffer().GroupsData[Instance->Debug.GroupIndex];
 
-					Instance->Strands.DeformedResource->SetPositionOffset(FHairStrandsDeformedResource::EFrameType::Current, GroomCacheGroupData.BoundingBox.GetCenter());
+					FGroomCacheGroupData* GroomCacheData0 = const_cast<FGroomCacheGroupData*>(&Instance->Debug.GroomCacheBuffers->GetCurrentFrameBuffer().GroupsData[Instance->Debug.GroupIndex]);
+					FGroomCacheGroupData* GroomCacheData1 = const_cast<FGroomCacheGroupData*>(&Instance->Debug.GroomCacheBuffers->GetNextFrameBuffer().GroupsData[Instance->Debug.GroupIndex]);
+					const float InterpolationFactor = Instance->Debug.GroomCacheBuffers->GetInterpolationFactor();
+
+					FGroomCacheResources CacheResources0 = CreateGroomCacheBuffer(GraphBuilder, GroomCacheData0->VertexData);
+					FGroomCacheResources CacheResources1 = CreateGroomCacheBuffer(GraphBuilder, GroomCacheData1->VertexData);
+
+					// Update position offset
+					{
+						const FVector OffsetPosition = FMath::Lerp(GroomCacheData0->BoundingBox.GetCenter(), GroomCacheData1->BoundingBox.GetCenter(), InterpolationFactor);
+						Instance->Strands.DeformedResource->SetPositionOffset(FHairStrandsDeformedResource::EFrameType::Current, OffsetPosition);
+					}
+
+					// Update max radius
+					if (GroomCacheData0->VertexData.PointsRadius.Num() > 0)
+					{
+						Instance->Strands.Modifier.HairWidth = FMath::Lerp(GroomCacheData0->StrandData.MaxRadius, GroomCacheData1->StrandData.MaxRadius, InterpolationFactor) * 2.0f;
+					}
 
 					AddHairStrandUpdatePositionOffsetPass(
 						GraphBuilder,
@@ -2008,16 +2070,12 @@ void ComputeHairStrandsInterpolation(
 						GraphBuilder,
 						ShaderMap,
 						ActivePointCount,
-						GroomCacheGroupData,
+						InterpolationFactor,
+						CacheResources0,
+						CacheResources1,
 						RegisterAsSRV(GraphBuilder, Instance->Strands.RestResource->PositionBuffer),
 						RegisterAsSRV(GraphBuilder, Instance->Strands.DeformedResource->GetPositionOffsetBuffer(FHairStrandsDeformedResource::EFrameType::Current)),
-						Strands_DeformedPosition.UAV
-						);
-
-					if (GroomCacheGroupData.VertexData.PointsRadius.Num() > 0)
-					{
-						Instance->Strands.Modifier.HairWidth = GroomCacheGroupData.StrandData.MaxRadius * 2.f;
-					}
+						Strands_DeformedPosition.UAV);
 				}
 
 				// 2.2 Update tangent data based on the deformed positions
diff --git a/Engine/Plugins/Runtime/HairStrands/Source/HairStrandsCore/Public/GroomCacheData.h b/Engine/Plugins/Runtime/HairStrands/Source/HairStrandsCore/Public/GroomCacheData.h
--- a/Engine/Plugins/Runtime/HairStrands/Source/HairStrandsCore/Public/GroomCacheData.h
+++ b/Engine/Plugins/Runtime/HairStrands/Source/HairStrandsCore/Public/GroomCacheData.h
@@ -120,6 +120,10 @@ struct FGroomCacheVertexData
 	TArray<float> PointsCoordU;
 	TArray<FLinearColor> PointsBaseColor;
 
+	// GPU resources. Lazily allocated
+	TRefCountPtr<FRDGPooledBuffer> PositionBuffer;
+	TRefCountPtr<FRDGPooledBuffer> RadiusBuffer;
+
 	void Serialize(FArchive& Ar, int32 Version, EGroomCacheAttributes InAttributes);
 };
 
@@ -171,7 +175,6 @@ public:
 
 	virtual const FGroomCacheAnimationData& GetCurrentFrameBuffer() = 0;
 	virtual const FGroomCacheAnimationData& GetNextFrameBuffer() = 0;
-	virtual const FGroomCacheAnimationData& GetInterpolatedFrameBuffer() = 0;
 
 	virtual int32 GetCurrentFrameIndex() const = 0;
 	virtual int32 GetNextFrameIndex() const = 0;
diff --git a/Engine/Shaders/Private/HairStrands/HairStrandsInterpolation.usf b/Engine/Shaders/Private/HairStrands/HairStrandsInterpolation.usf
--- a/Engine/Shaders/Private/HairStrands/HairStrandsInterpolation.usf
+++ b/Engine/Shaders/Private/HairStrands/HairStrandsInterpolation.usf
@@ -843,10 +843,14 @@ void MainCS(uint DispatchThreadId : SV_DispatchThreadID)
 
 uint ElementCount;
 uint bHasRadiusData;
+float InterpolationFactor;
 
-Buffer<float> InRadiusBuffer;
-Buffer<float> InAnimatedBuffer;
-Buffer<uint4> InRestPoseBuffer;
+Buffer<float> InPosition0Buffer;
+Buffer<float> InPosition1Buffer;
+Buffer<float> InRadius0Buffer;
+Buffer<float> InRadius1Buffer;
+
+Buffer<uint4>  InRestPoseBuffer;
 Buffer<float4> InDeformedOffsetBuffer;
 RWBuffer<uint4> OutDeformedBuffer;
 
@@ -860,24 +864,27 @@ void MainCS(uint DispatchThreadId : SV_DispatchThreadID)
 	}
 
 	const uint Index3 = Index * 3;
-	const float3 InAnimatedPosition = float3(
-		InAnimatedBuffer[Index3],
-		InAnimatedBuffer[Index3+1],
-		InAnimatedBuffer[Index3+2]);
-		
-	uint4 PackedRestPosition = InRestPoseBuffer[Index];
-	float3 HairPositionOffset = InDeformedOffsetBuffer[0].xyz;
-	float3 WorldVertexPosition = InAnimatedPosition - HairPositionOffset;
-
-	FHairControlPoint CP = UnpackHairControlPoint(PackedRestPosition, float3(0, 0, 0), 1, 1, 1);
+	const float3 InPosition0 = float3(
+		InPosition0Buffer[Index3],
+		InPosition0Buffer[Index3+1],
+		InPosition0Buffer[Index3+2]);
+	
+	const float3 InPosition1 = float3(
+		InPosition1Buffer[Index3],
+		InPosition1Buffer[Index3+1],
+		InPosition1Buffer[Index3+2]);
+	
+	const float3 WorldVertexPosition = lerp(InPosition0, InPosition1, InterpolationFactor);
+
+	FHairControlPoint CP = UnpackHairControlPoint(InRestPoseBuffer[Index], float3(0, 0, 0), 1, 1, 1);
 	CP.Position = WorldVertexPosition;
 
 	if (bHasRadiusData == 1)
 	{
-		CP.WorldRadius = InRadiusBuffer[Index];
+		CP.WorldRadius = lerp(InRadius0Buffer[Index], InRadius1Buffer[Index], InterpolationFactor);
 	}
 
-	OutDeformedBuffer[Index] = PackHairControlPoint(CP, float3(0, 0, 0), 1);
+	OutDeformedBuffer[Index] = PackHairControlPoint(CP, InDeformedOffsetBuffer[0].xyz, 1);
 }
 
 #endif // SHADER_GROOMCACHE_UPDATE
