Fix CIS error.

diff --git a/Engine/Plugins/Runtime/HairStrands/Source/HairStrandsCore/Public/GroomCacheData.h b/Engine/Plugins/Runtime/HairStrands/Source/HairStrandsCore/Public/GroomCacheData.h
--- a/Engine/Plugins/Runtime/HairStrands/Source/HairStrandsCore/Public/GroomCacheData.h
+++ b/Engine/Plugins/Runtime/HairStrands/Source/HairStrandsCore/Public/GroomCacheData.h
@@ -5,6 +5,8 @@
 #include "UObject/Object.h"
 #include "GroomCacheData.generated.h"
 
+class FRDGPooledBuffer;
+
 /** Attributes in groom that can be animated */
 UENUM()
 enum class EGroomCacheAttributes : uint8
