Fix invalid binding warning when a skel. mesh does not have skin cache or mesh deformer enabled, while groom requires skinning binding.

diff --git a/Engine/Plugins/Runtime/HairStrands/Source/HairStrandsCore/Private/GroomComponent.cpp b/Engine/Plugins/Runtime/HairStrands/Source/HairStrandsCore/Private/GroomComponent.cpp
--- a/Engine/Plugins/Runtime/HairStrands/Source/HairStrandsCore/Private/GroomComponent.cpp
+++ b/Engine/Plugins/Runtime/HairStrands/Source/HairStrandsCore/Private/GroomComponent.cpp
@@ -2386,13 +2386,21 @@ void UGroomComponent::InitResources(bool bIsBindingReloading)
 		{
 			if (BindingAsset)
 			{
-				const uint32 SkelLODCount = ParentSkelMeshComponent->GetNumLODs();
+				// Extract if skin cache or mesh. deformer is enabled on at least on LOD.
+				// Since there is 1:1 mapping between groom LOD and mesh LOD, only use this has a hint.
+				bool bSupportSkinCache = ParentSkelMeshComponent->HasMeshDeformer();
+				if (!bSupportSkinCache)
+				{
+					for (uint32 SkelLODIt = 0, SkelLODCount = ParentSkelMeshComponent->GetNumLODs(); SkelLODIt < SkelLODCount; ++SkelLODIt)
+					{
+						bSupportSkinCache = bSupportSkinCache || ParentSkelMeshComponent->IsSkinCacheAllowed(SkelLODIt);
+					}
+				}
+
 				for (int32 GroupIt = 0, GroupCount = GroomAsset->HairGroupsData.Num(); GroupIt < GroupCount; ++GroupIt)
 				{
 					for (uint32 LODIt = 0, LODCount = GroomAsset->GetLODCount(); LODIt < LODCount; ++LODIt)
 					{
-						const uint32 EffectiveLODIt = FMath::Clamp<uint32>(LODIt, 0, SkelLODCount - 1);
-						const bool bSupportSkinCache = ParentSkelMeshComponent->IsSkinCacheAllowed(EffectiveLODIt);
 						const EGroomBindingType BindingType = GroomAsset->GetBindingType(GroupIt, LODIt);
 						const bool bIsVisible = GroomAsset->IsVisible(GroupIt, LODIt);
 
