Groom: Fixed missing reimport actions on groom and groom cache assets after asset action API conversion.

diff --git a/Engine/Plugins/Runtime/HairStrands/Source/HairStrandsCore/Private/GroomCache.cpp b/Engine/Plugins/Runtime/HairStrands/Source/HairStrandsCore/Private/GroomCache.cpp
--- a/Engine/Plugins/Runtime/HairStrands/Source/HairStrandsCore/Private/GroomCache.cpp
+++ b/Engine/Plugins/Runtime/HairStrands/Source/HairStrandsCore/Private/GroomCache.cpp
@@ -1,6 +1,8 @@
 // Copyright Epic Games, Inc. All Rights Reserved.
 
 #include "GroomCache.h"
+
+#include "EditorFramework/AssetImportData.h"
 #include "GroomAsset.h"
 #include "Serialization/MemoryReader.h"
 #include "Serialization/MemoryWriter.h"
@@ -30,6 +32,18 @@ void UGroomCache::Serialize(FArchive& Ar)
 	}
 }
 
+void UGroomCache::GetAssetRegistryTags(TArray<FAssetRegistryTag>& OutTags) const
+{
+#if WITH_EDITORONLY_DATA
+	if (AssetImportData)
+	{
+		OutTags.Add(FAssetRegistryTag(SourceFileTagName(), AssetImportData->GetSourceData().ToJson(), FAssetRegistryTag::TT_Hidden));
+	}
+#endif
+
+	Super::GetAssetRegistryTags(OutTags);
+}
+
 void UGroomCache::Initialize(EGroomCacheType Type)
 {
 	GroomCacheInfo.Type = Type;
diff --git a/Engine/Plugins/Runtime/HairStrands/Source/HairStrandsCore/Public/GroomCache.h b/Engine/Plugins/Runtime/HairStrands/Source/HairStrandsCore/Public/GroomCache.h
--- a/Engine/Plugins/Runtime/HairStrands/Source/HairStrandsCore/Public/GroomCache.h
+++ b/Engine/Plugins/Runtime/HairStrands/Source/HairStrandsCore/Public/GroomCache.h
@@ -21,6 +21,7 @@ public:
 
 	//~ Begin UObject Interface.
 	virtual void Serialize(FArchive& Ar) override;
+	virtual void GetAssetRegistryTags(TArray<FAssetRegistryTag>& OutTags) const;
 	//~ End UObject Interface.
 
 	void Initialize(EGroomCacheType Type);
diff --git a/Engine/Plugins/Runtime/HairStrands/Source/HairStrandsEditor/Public/GroomActions.h b/Engine/Plugins/Runtime/HairStrands/Source/HairStrandsEditor/Public/GroomActions.h
--- a/Engine/Plugins/Runtime/HairStrands/Source/HairStrandsEditor/Public/GroomActions.h
+++ b/Engine/Plugins/Runtime/HairStrands/Source/HairStrandsEditor/Public/GroomActions.h
@@ -18,6 +18,7 @@ public:
 	virtual FText GetAssetDisplayName() const override { return NSLOCTEXT("AssetTypeActions", "AssetTypeActions_Groom", "Groom"); }
 	virtual FLinearColor GetAssetColor() const override;
 	virtual TSoftClassPtr<UObject> GetAssetClass() const override { return UGroomAsset::StaticClass(); }
+	virtual bool CanImport() const override { return true; }
 	virtual TConstArrayView<FAssetCategoryPath> GetAssetCategories() const override
 	{
 		static const auto Categories = { EAssetCategoryPaths::Misc };
diff --git a/Engine/Plugins/Runtime/HairStrands/Source/HairStrandsEditor/Public/GroomCacheActions.h b/Engine/Plugins/Runtime/HairStrands/Source/HairStrandsEditor/Public/GroomCacheActions.h
--- a/Engine/Plugins/Runtime/HairStrands/Source/HairStrandsEditor/Public/GroomCacheActions.h
+++ b/Engine/Plugins/Runtime/HairStrands/Source/HairStrandsEditor/Public/GroomCacheActions.h
@@ -18,6 +18,7 @@ public:
 	virtual FText GetAssetDisplayName() const override { return NSLOCTEXT("AssetTypeActions", "AssetTypeActions_GroomCache", "Groom Cache"); }
 	virtual FLinearColor GetAssetColor() const override;
 	virtual TSoftClassPtr<UObject> GetAssetClass() const override { return UGroomCache::StaticClass(); }
+	virtual bool CanImport() const override { return true; }
 	virtual TConstArrayView<FAssetCategoryPath> GetAssetCategories() const override
 	{
 		static const auto Categories = { EAssetCategoryPaths::Misc };
