Change how hair binding skel. mesh scene proxy look up is done.

Fix fix case where the wrong skel. mesh proxy would be used when the skel. mesh recreate its proxy while being used by the groom instance.

diff --git a/Engine/Plugins/Runtime/HairStrands/Source/HairStrandsCore/Private/GroomComponent.cpp b/Engine/Plugins/Runtime/HairStrands/Source/HairStrandsCore/Private/GroomComponent.cpp
--- a/Engine/Plugins/Runtime/HairStrands/Source/HairStrandsCore/Private/GroomComponent.cpp
+++ b/Engine/Plugins/Runtime/HairStrands/Source/HairStrandsCore/Private/GroomComponent.cpp
@@ -2481,6 +2481,8 @@ void UGroomComponent::InitResources(bool bIsBindingReloading)
 		HairGroupInstance->Debug.LODSelectionTypeForDebug = LODSelectionType;
 		HairGroupInstance->DeformedComponent = DeformedMeshComponent;
 		HairGroupInstance->DeformedSection = GroomAsset->DeformedGroupSections.IsValidIndex(GroupIt) ? GroomAsset->DeformedGroupSections[GroupIt] : INDEX_NONE; 
+		HairGroupInstance->Debug.MeshComponentId = ValidatedMeshComponent ? ValidatedMeshComponent->ComponentId : FPrimitiveComponentId();
+		HairGroupInstance->Debug.CachedMeshPersistentPrimitiveIndex = FPersistentPrimitiveIndex();
 
 		if (RegisteredMeshComponent)
 		{
diff --git a/Engine/Plugins/Runtime/HairStrands/Source/HairStrandsCore/Private/GroomDebug.cpp b/Engine/Plugins/Runtime/HairStrands/Source/HairStrandsCore/Private/GroomDebug.cpp
--- a/Engine/Plugins/Runtime/HairStrands/Source/HairStrandsCore/Private/GroomDebug.cpp
+++ b/Engine/Plugins/Runtime/HairStrands/Source/HairStrandsCore/Private/GroomDebug.cpp
@@ -65,6 +65,7 @@ static FAutoConsoleVariableRef CVarHairCardsVoxelDebug(TEXT("r.HairStrands.Cards
 
 FCachedGeometry GetCacheGeometryForHair(
 	FRDGBuilder& GraphBuilder,
+	FSceneInterface* Scene,
 	FHairGroupInstance* Instance,
 	FGlobalShaderMap* ShaderMap,
 	const bool bOutputTriangleData);
@@ -72,6 +73,7 @@ FCachedGeometry GetCacheGeometryForHair(
 static void GetGroomInterpolationData(
 	FRDGBuilder& GraphBuilder,
 	FGlobalShaderMap* ShaderMap,
+	FSceneInterface* Scene,
 	const FHairStrandsInstances& Instances,
 	const EHairStrandsProjectionMeshType MeshType,
 	FHairStrandsProjectionMeshData::LOD& OutGeometries)
@@ -83,7 +85,7 @@ static void GetGroomInterpolationData(
 		if (!Instance || !Instance->Debug.MeshComponent)
 			continue;
 
-		const FCachedGeometry CachedGeometry = GetCacheGeometryForHair(GraphBuilder, Instance, ShaderMap, true);
+		const FCachedGeometry CachedGeometry = GetCacheGeometryForHair(GraphBuilder, Scene, Instance, ShaderMap, true);
 		if (CachedGeometry.Sections.Num() == 0)
 			continue;
 
@@ -1128,6 +1130,7 @@ static void AddHairDebugPrintInstancePass(
 void RunHairStrandsDebug(
 	FRDGBuilder& GraphBuilder,
 	FGlobalShaderMap* ShaderMap,
+	FSceneInterface* Scene,
 	const FSceneView& View,
 	const FHairStrandsInstances& Instances,
 	const FUintVector4& InstanceCountPerType,
@@ -1149,10 +1152,10 @@ void RunHairStrandsDebug(
 		{
 			if (GHairDebugMeshProjection_SkinCacheMesh > 0)
 			{
-				auto RenderMeshProjection = [ShaderMap, ShaderPrintData, Viewport, &ViewUniformBuffer, Instances, &GraphBuilder](FRDGBuilder& LocalGraphBuilder, EHairStrandsProjectionMeshType MeshType)
+				auto RenderMeshProjection = [ShaderMap, Scene, ShaderPrintData, Viewport, &ViewUniformBuffer, Instances, &GraphBuilder](FRDGBuilder& LocalGraphBuilder, EHairStrandsProjectionMeshType MeshType)
 				{
 					FHairStrandsProjectionMeshData::LOD MeshProjectionLODData;
-					GetGroomInterpolationData(GraphBuilder, ShaderMap, Instances, MeshType, MeshProjectionLODData);
+					GetGroomInterpolationData(GraphBuilder, ShaderMap, Scene, Instances, MeshType, MeshProjectionLODData);
 					for (FHairStrandsProjectionMeshData::Section& Section : MeshProjectionLODData.Sections)
 					{
 						AddDebugProjectionMeshPass(LocalGraphBuilder, ShaderMap, ShaderPrintData, Viewport, ViewUniformBuffer, MeshType, Section);
diff --git a/Engine/Plugins/Runtime/HairStrands/Source/HairStrandsCore/Private/GroomManager.cpp b/Engine/Plugins/Runtime/HairStrands/Source/HairStrandsCore/Private/GroomManager.cpp
--- a/Engine/Plugins/Runtime/HairStrands/Source/HairStrandsCore/Private/GroomManager.cpp
+++ b/Engine/Plugins/Runtime/HairStrands/Source/HairStrandsCore/Private/GroomManager.cpp
@@ -22,6 +22,8 @@
 #include "RenderGraphUtils.h"
 #include "CachedGeometry.h"
 #include "GroomCacheData.h"
+#include "SceneInterface.h"
+#include "PrimitiveSceneInfo.h"
 
 static int32 GHairStrandsMinLOD = 0;
 static FAutoConsoleVariableRef CVarGHairStrandsMinLOD(TEXT("r.HairStrands.MinLOD"), GHairStrandsMinLOD, TEXT("Clamp the min hair LOD to this value, preventing to reach lower/high-quality LOD."), ECVF_Scalability);
@@ -41,6 +43,9 @@ static FAutoConsoleVariableRef CVarHairStrands_InterpolationFrustumCullingEnable
 static int32 GHairStrands_ContinousLOD = 0;
 static FAutoConsoleVariableRef CVarHairStrands_ContinousLOD(TEXT("r.HairStrands.ContinuousLOD"), GHairStrands_ContinousLOD, TEXT("Continuous LODing toggle."));
 
+static int32 GHairStrands_UseSkelMeshPersistentPrimitiveIndex = 1;
+static FAutoConsoleVariableRef CVarHairStrands_UseSkelMeshPersistentPrimitiveIndex(TEXT("r.HairStrands.UseSkelMeshPersistentPrimitiveIndex"), GHairStrands_UseSkelMeshPersistentPrimitiveIndex, TEXT("Use persistent primitive index for skel. mesh lookup (experimental)."), ECVF_RenderThreadSafe);
+
 EHairBufferSwapType GetHairSwapBufferType()
 {
 	switch (GHairStrands_SwapBufferType)
@@ -82,9 +87,37 @@ static bool IsSkeletalMeshEvaluationEnabled()
 	return CVarSkelMeshGDME ? CVarSkelMeshGDME->GetValueOnAnyThread() == 0 : true;
 }
 
+// Retrive the skel. mesh scene info
+static FPrimitiveSceneInfo* GetMeshSceneInfo(FSceneInterface* Scene, FHairGroupInstance* Instance)
+{
+	if (!Instance->Debug.MeshComponentId.IsValid())
+	{
+		return nullptr;
+	}
+
+	// Try to use the cached persistent primitive index (fast)
+	FPrimitiveSceneInfo* PrimitiveSceneInfo = nullptr;
+	if (Instance->Debug.CachedMeshPersistentPrimitiveIndex.IsValid())
+	{
+		PrimitiveSceneInfo = Scene->GetPrimitiveSceneInfo(Instance->Debug.CachedMeshPersistentPrimitiveIndex);
+		PrimitiveSceneInfo = (PrimitiveSceneInfo && PrimitiveSceneInfo->PrimitiveComponentId == Instance->Debug.MeshComponentId) ? PrimitiveSceneInfo : nullptr;
+	}
+
+	// If this fails, find the primitive index, and then the primitive scene info (slow)
+	if (PrimitiveSceneInfo == nullptr)
+	{
+		const int32 PrimitiveIndex = Scene->GetScenePrimitiveComponentIds().Find(Instance->Debug.MeshComponentId);
+		PrimitiveSceneInfo = Scene->GetPrimitiveSceneInfo(PrimitiveIndex);
+		Instance->Debug.CachedMeshPersistentPrimitiveIndex = PrimitiveSceneInfo ? PrimitiveSceneInfo->GetPersistentIndex() : FPersistentPrimitiveIndex();
+	}
+
+	return PrimitiveSceneInfo;
+}
+
 // Returns the cached geometry of the underlying geometry on which a hair instance is attached to
 FCachedGeometry GetCacheGeometryForHair(
 	FRDGBuilder& GraphBuilder, 
+	FSceneInterface* Scene,
 	FHairGroupInstance* Instance, 
 	FGlobalShaderMap* ShaderMap,
 	const bool bOutputTriangleData)
@@ -94,19 +127,33 @@ FCachedGeometry GetCacheGeometryForHair(
 	{
 		if (IsSkeletalMeshEvaluationEnabled())
 		{
-			//todo: It's unsafe to be accessing the component directly here. This can be solved when we have persistent ids available on the render thread.
-			if (const USkeletalMeshComponent* SkeletalMeshComponent = Cast<USkeletalMeshComponent>(Instance->Debug.MeshComponent))
+			if (GHairStrands_UseSkelMeshPersistentPrimitiveIndex)
 			{
-				if (FSkeletalMeshSceneProxy* SceneProxy = static_cast<FSkeletalMeshSceneProxy*>(SkeletalMeshComponent->SceneProxy))
+				if (const FPrimitiveSceneInfo* PrimitiveSceneInfo = GetMeshSceneInfo(Scene, Instance))
 				{
-					SceneProxy->GetCachedGeometry(Out);
+					if (const FSkeletalMeshSceneProxy* SceneProxy = static_cast<const FSkeletalMeshSceneProxy*>(PrimitiveSceneInfo->Proxy))
+					{
+					
+						SceneProxy->GetCachedGeometry(Out);
+					}
 				}
-
-				if (GHairStrands_ManualSkinCache > 0 && Out.Sections.Num() == 0)
+			}
+			else
+			{
+				//todo: It's unsafe to be accessing the component directly here. This can be solved when we have persistent ids available on the render thread.
+				if (const USkeletalMeshComponent* SkeletalMeshComponent = Cast<USkeletalMeshComponent>(Instance->Debug.MeshComponent))
 				{
-					//#hair_todo: Need to have a (frame) cache to insure that we don't recompute the same projection several time
-					// Actual populate the cache with only the needed part based on the groom projection data. At the moment it recompute everything ...
-					BuildCacheGeometry(GraphBuilder, ShaderMap, SkeletalMeshComponent, bOutputTriangleData, Out);
+					if (FSkeletalMeshSceneProxy* SceneProxy = static_cast<FSkeletalMeshSceneProxy*>(SkeletalMeshComponent->SceneProxy))
+					{
+						SceneProxy->GetCachedGeometry(Out);
+					}
+
+					if (GHairStrands_ManualSkinCache > 0 && Out.Sections.Num() == 0)
+					{
+						//#hair_todo: Need to have a (frame) cache to insure that we don't recompute the same projection several time
+						// Actual populate the cache with only the needed part based on the groom projection data. At the moment it recompute everything ...
+						BuildCacheGeometry(GraphBuilder, ShaderMap, SkeletalMeshComponent, bOutputTriangleData, Out);
+					}
 				}
 			}
 		}
@@ -138,6 +185,7 @@ static int32 GetCacheGeometryLODIndex(const FCachedGeometry& In)
 
 static void RunInternalHairStrandsInterpolation(
 	FRDGBuilder& GraphBuilder,
+	FSceneInterface* Scene,
 	const FSceneView* View,
 	const uint32 ViewUniqueID,
 	const FHairStrandsInstances& Instances,
@@ -168,7 +216,7 @@ static void RunInternalHairStrandsInterpolation(
 		check(Instance->HairGroupPublicData);
 
 		FHairStrandsProjectionMeshData::LOD MeshDataLOD;
-		const FCachedGeometry CachedGeometry = GetCacheGeometryForHair(GraphBuilder, Instance, ShaderMap, true);
+		const FCachedGeometry CachedGeometry = GetCacheGeometryForHair(GraphBuilder, Scene, Instance, ShaderMap, true);
 		for (int32 SectionIndex = 0; SectionIndex < CachedGeometry.Sections.Num(); ++SectionIndex)
 		{
 			// Ensure all mesh's sections have the same LOD index
@@ -399,6 +447,7 @@ static void RunInternalHairStrandsInterpolation(
 
 static void RunHairStrandsInterpolation_Guide(
 	FRDGBuilder& GraphBuilder,
+	FSceneInterface* Scene,
 	const FSceneView* View,
 	const uint32 ViewUniqueID,
 	const FHairStrandsInstances& Instances,
@@ -414,6 +463,7 @@ static void RunHairStrandsInterpolation_Guide(
 
 	RunInternalHairStrandsInterpolation(
 		GraphBuilder,
+		Scene,
 		View,
 		ViewUniqueID,
 		Instances,
@@ -425,6 +475,7 @@ static void RunHairStrandsInterpolation_Guide(
 
 static void RunHairStrandsInterpolation_Strands(
 	FRDGBuilder& GraphBuilder,
+	FSceneInterface* Scene,
 	const FSceneView* View,
 	const uint32 ViewUniqueID,
 	const FHairStrandsInstances& Instances,
@@ -440,6 +491,7 @@ static void RunHairStrandsInterpolation_Strands(
 
 	RunInternalHairStrandsInterpolation(
 		GraphBuilder,
+		Scene,
 		View,
 		ViewUniqueID,
 		Instances,
@@ -745,6 +797,7 @@ static uint32 ComputeActiveCurveCount(float InScreenSize, uint32 InCurveCount)
 
 static void RunHairLODSelection(
 	FRDGBuilder& GraphBuilder, 
+	FSceneInterface* Scene,
 	const FHairStrandsInstances& Instances, 
 	const TArray<const FSceneView*>& Views, 
 	FGlobalShaderMap* ShaderMap)
@@ -774,7 +827,7 @@ static void RunHairLODSelection(
 
 		check(Instance);
 		check(Instance->HairGroupPublicData);
-		const FCachedGeometry CachedGeometry = GetCacheGeometryForHair(GraphBuilder, Instance, ShaderMap, false);
+		const FCachedGeometry CachedGeometry = GetCacheGeometryForHair(GraphBuilder, Scene, Instance, ShaderMap, false);
 		const int32 MeshLODIndex = GetCacheGeometryLODIndex(CachedGeometry);
 
 		// Perform LOD selection based on all the views	
@@ -1117,6 +1170,7 @@ static void RunHairStrandsProcess(FRDGBuilder& GraphBuilder, FGlobalShaderMap* S
 void RunHairStrandsDebug(
 	FRDGBuilder& GraphBuilder,
 	FGlobalShaderMap* ShaderMap,
+	FSceneInterface* Scene,
 	const FSceneView& View,
 	const FHairStrandsInstances& Instances,
 	const FUintVector4& InstanceCountPerType,
@@ -1170,6 +1224,7 @@ void ProcessHairStrandsBookmark(
 		check(GraphBuilder);
 		RunHairLODSelection(
 			*GraphBuilder,
+			Parameters.Scene,
 			*Parameters.Instances,
 			Parameters.AllViews,
 			Parameters.ShaderMap);
@@ -1188,6 +1243,7 @@ void ProcessHairStrandsBookmark(
 		check(GraphBuilder);
 		RunHairStrandsInterpolation_Guide(
 			*GraphBuilder,
+			Parameters.Scene,
 			Parameters.View,
 			Parameters.ViewUniqueID,
 			Instances,
@@ -1206,6 +1262,7 @@ void ProcessHairStrandsBookmark(
 		check(GraphBuilder);
 		RunHairStrandsInterpolation_Strands(
 			*GraphBuilder,
+			Parameters.Scene,
 			Parameters.View,
 			Parameters.ViewUniqueID,
 			Instances,
@@ -1219,6 +1276,7 @@ void ProcessHairStrandsBookmark(
 		RunHairStrandsDebug(
 			*GraphBuilder,
 			Parameters.ShaderMap,
+			Parameters.Scene,
 			*Parameters.View,
 			Instances,
 			Parameters.InstanceCountPerType,
@@ -1228,4 +1286,4 @@ void ProcessHairStrandsBookmark(
 			Parameters.View->UnscaledViewRect,
 			Parameters.View->ViewUniformBuffer);
 	}
-}
+}
\ No newline at end of file
diff --git a/Engine/Plugins/Runtime/HairStrands/Source/HairStrandsCore/Public/GroomInstance.h b/Engine/Plugins/Runtime/HairStrands/Source/HairStrandsCore/Public/GroomInstance.h
--- a/Engine/Plugins/Runtime/HairStrands/Source/HairStrandsCore/Public/GroomInstance.h
+++ b/Engine/Plugins/Runtime/HairStrands/Source/HairStrandsCore/Public/GroomInstance.h
@@ -10,6 +10,7 @@
 #include "GroomDesc.h"
 #include "HairStrandsInterface.h"
 #include "HairStrandsDefinitions.h"
+#include "PrimitiveSceneInfo.h"
 
 class UMeshComponent;
 class UGroomComponent;
@@ -256,6 +257,8 @@ struct HAIRSTRANDSCORE_API FHairGroupInstance : public FHairStrandsInstance
 		FPrimitiveSceneProxy*	Proxy = nullptr;
 		UMeshComponent*			MeshComponent = nullptr;
 		FString					MeshComponentName;
+		FPrimitiveComponentId	MeshComponentId;
+		FPersistentPrimitiveIndex CachedMeshPersistentPrimitiveIndex;
 		const UGroomComponent*	GroomComponentForDebug = nullptr; // For debug only, shouldn't be deferred on the rendering thread
 		FTransform				RigidCurrentLocalToWorld = FTransform::Identity;
 		FTransform				SkinningCurrentLocalToWorld = FTransform::Identity;
diff --git a/Engine/Shaders/Private/HairStrands/HairStrandsDebug.usf b/Engine/Shaders/Private/HairStrands/HairStrandsDebug.usf
--- a/Engine/Shaders/Private/HairStrands/HairStrandsDebug.usf
+++ b/Engine/Shaders/Private/HairStrands/HairStrandsDebug.usf
@@ -1437,9 +1437,9 @@ void MainCS(uint GroupIndex : SV_GroupIndex, uint3 DispatchThreadId : SV_Dispatc
 		Print(Context, TEXT(" "));
 
 		// Binding
-		if (Info.BindingType == 0)		{ Print(Context, TEXT("Skinned "), FontTurquoise); }
+		if (Info.BindingType == 0)		{ Print(Context, TEXT("None    "), FontTurquoise); }
 		else if (Info.BindingType == 1)	{ Print(Context, TEXT("Rigid   "), FontTurquoise); }
-		else if (Info.BindingType == 2)	{ Print(Context, TEXT("None    "), FontTurquoise); }
+		else if (Info.BindingType == 2)	{ Print(Context, TEXT("Skinned "), FontTurquoise); }
 
 		// Cache
 		if (Info.CacheRen)      { Print(Context, TEXT("Ren  "), FontOrange); }
diff --git a/Engine/Source/Runtime/Engine/Public/SceneInterface.h b/Engine/Source/Runtime/Engine/Public/SceneInterface.h
--- a/Engine/Source/Runtime/Engine/Public/SceneInterface.h
+++ b/Engine/Source/Runtime/Engine/Public/SceneInterface.h
@@ -63,6 +63,7 @@ namespace ERHIFeatureLevel { enum Type : int; }
 struct FDeferredDecalUpdateParams;
 struct FHairStrandsInstance;
 struct FLightRenderParameters;
+struct FPersistentPrimitiveIndex;
 template<int32 MaxSHOrder> class TSHVectorRGB;
 using FSHVectorRGB3 = TSHVectorRGB<3>;
 
@@ -157,6 +158,7 @@ public:
 
 	/** Finds the  primitive with the associated component id. */
 	virtual FPrimitiveSceneInfo* GetPrimitiveSceneInfo(FPrimitiveComponentId PrimitiveId) const = 0;
+	virtual FPrimitiveSceneInfo* GetPrimitiveSceneInfo(const FPersistentPrimitiveIndex& PersistentPrimitiveIndex) = 0;
 
 	/** Get the primitive previous local to world (used for motion blur). Returns true if the matrix was set. */
 	virtual bool GetPreviousLocalToWorld(const FPrimitiveSceneInfo* PrimitiveSceneInfo, FMatrix& OutPreviousLocalToWorld) const { return false; }
diff --git a/Engine/Source/Runtime/Renderer/Private/HairStrands/HairStrandsInterface.cpp b/Engine/Source/Runtime/Renderer/Private/HairStrands/HairStrandsInterface.cpp
--- a/Engine/Source/Runtime/Renderer/Private/HairStrands/HairStrandsInterface.cpp
+++ b/Engine/Source/Runtime/Renderer/Private/HairStrands/HairStrandsInterface.cpp
@@ -622,6 +622,7 @@ FHairStrandsBookmarkParameters CreateHairStrandsBookmarkParameters(FScene* Scene
 	Out.ViewUniqueID			= View.ViewState ? View.ViewState->UniqueID : ~0;
 	Out.SceneColorTexture		= nullptr;
 	Out.bHzbRequest				= false; // Out.HasInstances() && IsHairStrandsEnabled(EHairStrandsShaderType::Strands, View.GetShaderPlatform());
+	Out.Scene					= Scene;
 
 	// Sanity check
 	check(Out.Instances->Num() >= Out.VisibleInstances.Num());
diff --git a/Engine/Source/Runtime/Renderer/Private/RendererScene.cpp b/Engine/Source/Runtime/Renderer/Private/RendererScene.cpp
--- a/Engine/Source/Runtime/Renderer/Private/RendererScene.cpp
+++ b/Engine/Source/Runtime/Renderer/Private/RendererScene.cpp
@@ -2527,6 +2527,12 @@ FPrimitiveSceneInfo* FScene::GetPrimitiveSceneInfo(FPrimitiveComponentId Primiti
 	return GetPrimitiveSceneInfo(PrimitiveComponentIds.Find(PrimitiveId));
 }
 
+FPrimitiveSceneInfo* FScene::GetPrimitiveSceneInfo(const FPersistentPrimitiveIndex& PersistentPrimitiveIndex)
+{
+	int32 PrimitiveIndex = GetPrimitiveIndex(PersistentPrimitiveIndex);
+	return GetPrimitiveSceneInfo(PrimitiveIndex);
+}
+
 void FScene::RemovePrimitiveSceneInfo_RenderThread(FPrimitiveSceneInfo* PrimitiveSceneInfo)
 {
 	check(IsInRenderingThread());
@@ -6673,6 +6679,7 @@ public:
 	virtual void UpdateAllPrimitiveSceneInfos(FRDGBuilder& GraphBuilder, EUpdateAllPrimitiveSceneInfosAsyncOps = EUpdateAllPrimitiveSceneInfosAsyncOps::None) override {}
 	virtual FPrimitiveSceneInfo* GetPrimitiveSceneInfo(int32 PrimiteIndex) const final { return nullptr; }
 	virtual FPrimitiveSceneInfo* GetPrimitiveSceneInfo(FPrimitiveComponentId PrimitiveId) const final { return nullptr; }
+	virtual FPrimitiveSceneInfo* GetPrimitiveSceneInfo(const FPersistentPrimitiveIndex& PersistentPrimitiveIndex) override { return nullptr; }
 
 	/** Updates the transform of a primitive which has already been added to the scene. */
 	virtual void UpdatePrimitiveTransform(UPrimitiveComponent* Primitive) override {}
@@ -6965,7 +6972,6 @@ void FScene::IncrementFrameNumber()
 		SceneFrameNumberRenderThread = NewNumber;
 	});
 }
-
 #if !(UE_BUILD_SHIPPING || UE_BUILD_TEST)
 void FScene::DebugRender(TArrayView<FViewInfo> Views)
 {
diff --git a/Engine/Source/Runtime/Renderer/Private/ScenePrivate.h b/Engine/Source/Runtime/Renderer/Private/ScenePrivate.h
--- a/Engine/Source/Runtime/Renderer/Private/ScenePrivate.h
+++ b/Engine/Source/Runtime/Renderer/Private/ScenePrivate.h
@@ -3452,7 +3452,7 @@ public:
 
 	FORCEINLINE int32 GetPrimitiveIndex(const FPersistentPrimitiveIndex& PersistentPrimitiveIndex) const
 	{ 
-		if (PersistentPrimitiveIndex.Index < PersistentPrimitiveIdToIndexMap.Num())
+		if (PersistentPrimitiveIndex.IsValid() && PersistentPrimitiveIndex.Index < PersistentPrimitiveIdToIndexMap.Num())
 		{
 			return PersistentPrimitiveIdToIndexMap[PersistentPrimitiveIndex.Index];
 		}
diff --git a/Engine/Source/Runtime/Renderer/Private/VirtualShadowMaps/VirtualShadowMapClipmap.cpp b/Engine/Source/Runtime/Renderer/Private/VirtualShadowMaps/VirtualShadowMapClipmap.cpp
--- a/Engine/Source/Runtime/Renderer/Private/VirtualShadowMaps/VirtualShadowMapClipmap.cpp
+++ b/Engine/Source/Runtime/Renderer/Private/VirtualShadowMaps/VirtualShadowMapClipmap.cpp
@@ -365,7 +365,7 @@ void FVirtualShadowMapClipmap::OnPrimitiveRendered(const FPrimitiveSceneInfo* Pr
 	if (PerLightCacheEntry.IsValid())
 	{
 		FPersistentPrimitiveIndex PersistentPrimitiveId = PrimitiveSceneInfo->GetPersistentIndex();
-		check(PersistentPrimitiveId.Index >= 0);
+		check(PersistentPrimitiveId.IsValid());
 		check(PersistentPrimitiveId.Index < PerLightCacheEntry->RenderedPrimitives.Num());
 
 		// Check previous-frame state to detect transition from hidden->visible
diff --git a/Engine/Source/Runtime/Renderer/Public/HairStrandsInterface.h b/Engine/Source/Runtime/Renderer/Public/HairStrandsInterface.h
--- a/Engine/Source/Runtime/Renderer/Public/HairStrandsInterface.h
+++ b/Engine/Source/Runtime/Renderer/Public/HairStrandsInterface.h
@@ -18,6 +18,7 @@
 #include "HairStrandsDefinitions.h"
 
 class UTexture2D;
+class FSceneInterface;
 
 ////////////////////////////////////////////////////////////////////////////////////////////////////
 // Utils buffers for importing/exporting hair resources
@@ -439,6 +440,7 @@ struct FHairStrandsBookmarkParameters
 	FHairStrandsInstances VisibleInstances;
 	FHairStrandsInstances* Instances = nullptr;
 	const FSceneView* View = nullptr;// // View 0
+	FSceneInterface* Scene = nullptr;
 	TArray<const FSceneView*> AllViews;
 	FRDGTextureRef SceneColorTexture = nullptr;
 	FRDGTextureRef SceneDepthTexture = nullptr; 
diff --git a/Engine/Source/Runtime/Renderer/Public/PrimitiveSceneInfo.h b/Engine/Source/Runtime/Renderer/Public/PrimitiveSceneInfo.h
--- a/Engine/Source/Runtime/Renderer/Public/PrimitiveSceneInfo.h
+++ b/Engine/Source/Runtime/Renderer/Public/PrimitiveSceneInfo.h
@@ -266,7 +266,8 @@ ENUM_CLASS_FLAGS(EUpdateStaticMeshFlags);
  */
 struct FPersistentPrimitiveIndex
 {
-	int32 Index;
+	bool IsValid() const { return Index != INDEX_NONE; }
+	int32 Index = INDEX_NONE;
 };
 
 /**
