Clean up component access for hair binding.

diff --git a/Engine/Plugins/Runtime/GeometryCache/Source/GeometryCache/Public/GeometryCacheSceneProxy.h b/Engine/Plugins/Runtime/GeometryCache/Source/GeometryCache/Public/GeometryCacheSceneProxy.h
--- a/Engine/Plugins/Runtime/GeometryCache/Source/GeometryCache/Public/GeometryCacheSceneProxy.h
+++ b/Engine/Plugins/Runtime/GeometryCache/Source/GeometryCache/Public/GeometryCacheSceneProxy.h
@@ -318,7 +318,7 @@ public:
 	virtual bool HasRayTracingRepresentation() const override { return true; }
 #endif
 
-	const TArray<FGeomCacheTrackProxy*>& GetTracks() { return Tracks; }
+	const TArray<FGeomCacheTrackProxy*>& GetTracks() const { return Tracks; }
 
 private:
 	void FrameUpdate() const;
diff --git a/Engine/Plugins/Runtime/HairStrands/Source/HairStrandsCore/Private/GroomComponent.cpp b/Engine/Plugins/Runtime/HairStrands/Source/HairStrandsCore/Private/GroomComponent.cpp
--- a/Engine/Plugins/Runtime/HairStrands/Source/HairStrandsCore/Private/GroomComponent.cpp
+++ b/Engine/Plugins/Runtime/HairStrands/Source/HairStrandsCore/Private/GroomComponent.cpp
@@ -2467,11 +2467,10 @@ void UGroomComponent::InitResources(bool bIsBindingReloading)
 		FHairGroupInstance* HairGroupInstance = new FHairGroupInstance();
 		HairGroupInstance->AddRef();
 		HairGroupInstances.Add(HairGroupInstance);
-		HairGroupInstance->Debug.ComponentId = ComponentId.PrimIDValue;
 		HairGroupInstance->Debug.GroupIndex = GroupIt;
 		HairGroupInstance->Debug.GroupCount = GroupCount;
 		HairGroupInstance->Debug.GroomAssetName = GroomAsset->GetName();
-		HairGroupInstance->Debug.MeshComponent = RegisteredMeshComponent;
+		HairGroupInstance->Debug.MeshComponentForDebug = RegisteredMeshComponent;
 		HairGroupInstance->Debug.GroomComponentForDebug = this;
 		HairGroupInstance->Debug.GroomBindingType = BindingAsset ? BindingAsset->GroomBindingType : EGroomBindingMeshType::SkeletalMesh;
 		HairGroupInstance->Debug.GroomCacheType = GetEffectiveGroomCacheType(GroomCache, GroomAsset);
diff --git a/Engine/Plugins/Runtime/HairStrands/Source/HairStrandsCore/Private/GroomDebug.cpp b/Engine/Plugins/Runtime/HairStrands/Source/HairStrandsCore/Private/GroomDebug.cpp
--- a/Engine/Plugins/Runtime/HairStrands/Source/HairStrandsCore/Private/GroomDebug.cpp
+++ b/Engine/Plugins/Runtime/HairStrands/Source/HairStrandsCore/Private/GroomDebug.cpp
@@ -82,7 +82,7 @@ static void GetGroomInterpolationData(
 	{
 		FHairGroupInstance* Instance = static_cast<FHairGroupInstance*>(AbstractInstance);
 
-		if (!Instance || !Instance->Debug.MeshComponent)
+		if (!Instance)
 			continue;
 
 		const FCachedGeometry CachedGeometry = GetCacheGeometryForHair(GraphBuilder, Scene, Instance, ShaderMap, true);
diff --git a/Engine/Plugins/Runtime/HairStrands/Source/HairStrandsCore/Private/GroomGeometryCache.cpp b/Engine/Plugins/Runtime/HairStrands/Source/HairStrandsCore/Private/GroomGeometryCache.cpp
--- a/Engine/Plugins/Runtime/HairStrands/Source/HairStrandsCore/Private/GroomGeometryCache.cpp
+++ b/Engine/Plugins/Runtime/HairStrands/Source/HairStrandsCore/Private/GroomGeometryCache.cpp
@@ -51,24 +51,18 @@ ENGINE_API void UpdatePreviousRefToLocalMatrices(TArray<FMatrix44f>& ReferenceTo
  void BuildCacheGeometry(
 	FRDGBuilder& GraphBuilder,
 	FGlobalShaderMap* ShaderMap, 
-	const USkeletalMeshComponent* SkeletalMeshComponent, 
+	const FSkeletalMeshSceneProxy* Proxy,
 	const bool bOutputTriangleData,
 	FCachedGeometry& Out)
 {
-	if (SkeletalMeshComponent == nullptr)
+	if (Proxy == nullptr)
 	{
 		return;
 	}
+	
+	Out.LocalToWorld = FTransform(Proxy->GetLocalToWorld());
 
-	FSkeletalMeshSceneProxy* SceneProxy = static_cast<FSkeletalMeshSceneProxy*>(SkeletalMeshComponent->SceneProxy);
-	if (SceneProxy == nullptr)
-	{
-		return;
-	}
-
-	Out.LocalToWorld = FTransform(SceneProxy->GetLocalToWorld());
-
-	FSkeletalMeshObject* SkeletalMeshObject = SkeletalMeshComponent->MeshObject;
+	const FSkeletalMeshObject* SkeletalMeshObject = Proxy->GetMeshObject();
 	if (SkeletalMeshObject == nullptr)
 	{
 		return;
@@ -108,7 +102,7 @@ ENGINE_API void UpdatePreviousRefToLocalMatrices(TArray<FMatrix44f>& ReferenceTo
 		{
 			FRHIShaderResourceView * BoneBuffer		= FSkeletalMeshDeformerHelpers::GetBoneBufferForReading(SkeletalMeshObject, LODIndex, SectionIdx, false);
 			FRHIShaderResourceView * BonePrevBuffer	= FSkeletalMeshDeformerHelpers::GetBoneBufferForReading(SkeletalMeshObject, LODIndex, SectionIdx, true);
-			AddSkinUpdatePass(GraphBuilder, ShaderMap, SectionIdx, BonesOffset, SkeletalMeshComponent->GetSkinWeightBuffer(LODIndex), LODData, BoneBuffer, BonePrevBuffer, DeformedPositionsBuffer, DeformedPreviousPositionsBuffer);
+			AddSkinUpdatePass(GraphBuilder, ShaderMap, SectionIdx, BonesOffset, LODData, BoneBuffer, BonePrevBuffer, DeformedPositionsBuffer, DeformedPreviousPositionsBuffer);
 
 			FCachedGeometry::Section& OutSection = Out.Sections.AddDefaulted_GetRef();
 			OutSection.RDGPositionBuffer = DeformedPositionSRV;
@@ -135,56 +129,53 @@ ENGINE_API void UpdatePreviousRefToLocalMatrices(TArray<FMatrix44f>& ReferenceTo
  void BuildCacheGeometry(
 	 FRDGBuilder& GraphBuilder,
 	 FGlobalShaderMap* ShaderMap, 
-	 const UGeometryCacheComponent* GeometryCacheComponent, 
+	 const FGeometryCacheSceneProxy* SceneProxy,
 	 const bool bOutputTriangleData,
 	 FCachedGeometry& Out)
  {
-	 if (GeometryCacheComponent)
-	 {
-		 if (FGeometryCacheSceneProxy* SceneProxy = static_cast<FGeometryCacheSceneProxy*>(GeometryCacheComponent->SceneProxy))
-		 {
-			 Out.LocalToWorld = FTransform(SceneProxy->GetLocalToWorld());
-			 Out.LODIndex = 0;
-			 if (bOutputTriangleData)
-			 {
-				 // Prior to getting here, the GeometryCache has been validated to be flattened (ie. has only one track)
-				 check(SceneProxy->GetTracks().Num() > 0);
-				 const FGeomCacheTrackProxy* TrackProxy = SceneProxy->GetTracks()[0];
-
-				 check(TrackProxy->MeshData && TrackProxy->NextFrameMeshData);
-				 const bool bHasMotionVectors = (
-					 TrackProxy->MeshData->VertexInfo.bHasMotionVectors &&
-					 TrackProxy->NextFrameMeshData->VertexInfo.bHasMotionVectors &&
-					 TrackProxy->MeshData->Positions.Num() == TrackProxy->MeshData->MotionVectors.Num())
-					 && (TrackProxy->NextFrameMeshData->Positions.Num() == TrackProxy->NextFrameMeshData->MotionVectors.Num());
-
-				 // PositionBuffer depends on CurrentPositionBufferIndex and on if the cache has motion vectors
-				 const uint32 PositionIndex = (TrackProxy->CurrentPositionBufferIndex == -1 || bHasMotionVectors) ? 0 : TrackProxy->CurrentPositionBufferIndex % 2;
-				 for (int32 SectionIdx = 0; SectionIdx < TrackProxy->MeshData->BatchesInfo.Num(); ++SectionIdx)
-				 {
-					const FGeometryCacheMeshBatchInfo& BatchInfo = TrackProxy->MeshData->BatchesInfo[SectionIdx];
-				
-					FCachedGeometry::Section OutSection;
-					OutSection.PositionBuffer = TrackProxy->PositionBuffers[PositionIndex].GetBufferSRV();
-					OutSection.UVsBuffer = TrackProxy->TextureCoordinatesBuffer.GetBufferSRV();
-					OutSection.TotalVertexCount = TrackProxy->MeshData->Positions.Num();
-					OutSection.IndexBuffer = TrackProxy->IndexBuffer.GetBufferSRV();
-					OutSection.TotalIndexCount = TrackProxy->IndexBuffer.NumValidIndices;
-					OutSection.UVsChannelCount = 1;
-					OutSection.NumPrimitives = BatchInfo.NumTriangles;
-					OutSection.NumVertices = TrackProxy->MeshData->Positions.Num();
-					OutSection.IndexBaseIndex = BatchInfo.StartIndex;
-					OutSection.VertexBaseIndex = 0;
-					OutSection.SectionIndex = SectionIdx;
-					OutSection.LODIndex = 0;
-					OutSection.UVsChannelOffset = 0;
-				
-					if (OutSection.PositionBuffer && OutSection.IndexBuffer)
-					{
-						Out.Sections.Add(OutSection);
-					}
-				 }
-			 }
-		 }
-	 }
+	if (SceneProxy)
+	{
+		Out.LocalToWorld = FTransform(SceneProxy->GetLocalToWorld());
+		Out.LODIndex = 0;
+		if (bOutputTriangleData)
+		{
+			// Prior to getting here, the GeometryCache has been validated to be flattened (ie. has only one track)
+			check(SceneProxy->GetTracks().Num() > 0);
+			const FGeomCacheTrackProxy* TrackProxy = SceneProxy->GetTracks()[0];
+
+			check(TrackProxy->MeshData && TrackProxy->NextFrameMeshData);
+			const bool bHasMotionVectors = (
+				TrackProxy->MeshData->VertexInfo.bHasMotionVectors &&
+				TrackProxy->NextFrameMeshData->VertexInfo.bHasMotionVectors &&
+				TrackProxy->MeshData->Positions.Num() == TrackProxy->MeshData->MotionVectors.Num())
+				&& (TrackProxy->NextFrameMeshData->Positions.Num() == TrackProxy->NextFrameMeshData->MotionVectors.Num());
+
+			// PositionBuffer depends on CurrentPositionBufferIndex and on if the cache has motion vectors
+			const uint32 PositionIndex = (TrackProxy->CurrentPositionBufferIndex == -1 || bHasMotionVectors) ? 0 : TrackProxy->CurrentPositionBufferIndex % 2;
+			for (int32 SectionIdx = 0; SectionIdx < TrackProxy->MeshData->BatchesInfo.Num(); ++SectionIdx)
+			{
+				const FGeometryCacheMeshBatchInfo& BatchInfo = TrackProxy->MeshData->BatchesInfo[SectionIdx];
+			
+				FCachedGeometry::Section OutSection;
+				OutSection.PositionBuffer = TrackProxy->PositionBuffers[PositionIndex].GetBufferSRV();
+				OutSection.UVsBuffer = TrackProxy->TextureCoordinatesBuffer.GetBufferSRV();
+				OutSection.TotalVertexCount = TrackProxy->MeshData->Positions.Num();
+				OutSection.IndexBuffer = TrackProxy->IndexBuffer.GetBufferSRV();
+				OutSection.TotalIndexCount = TrackProxy->IndexBuffer.NumValidIndices;
+				OutSection.UVsChannelCount = 1;
+				OutSection.NumPrimitives = BatchInfo.NumTriangles;
+				OutSection.NumVertices = TrackProxy->MeshData->Positions.Num();
+				OutSection.IndexBaseIndex = BatchInfo.StartIndex;
+				OutSection.VertexBaseIndex = 0;
+				OutSection.SectionIndex = SectionIdx;
+				OutSection.LODIndex = 0;
+				OutSection.UVsChannelOffset = 0;
+			
+				if (OutSection.PositionBuffer && OutSection.IndexBuffer)
+				{
+					Out.Sections.Add(OutSection);
+				}
+			}
+		}
+	}
  }
diff --git a/Engine/Plugins/Runtime/HairStrands/Source/HairStrandsCore/Private/GroomGeometryCache.h b/Engine/Plugins/Runtime/HairStrands/Source/HairStrandsCore/Private/GroomGeometryCache.h
--- a/Engine/Plugins/Runtime/HairStrands/Source/HairStrandsCore/Private/GroomGeometryCache.h
+++ b/Engine/Plugins/Runtime/HairStrands/Source/HairStrandsCore/Private/GroomGeometryCache.h
@@ -6,23 +6,23 @@
 #include "HairStrandsMeshProjection.h"
 
 struct FCachedGeometry;
-class UGeometryCacheComponent;
+class FSkeletalMeshSceneProxy;
 class FGlobalShaderMap;
 class FRDGBuilder;
-class USkeletalMeshComponent;
+class FGeometryCacheSceneProxy;
 
 FHairStrandsProjectionMeshData::Section ConvertMeshSection(FCachedGeometry const& InCachedGeometry, int32 InSectionIndex);
 
 void BuildCacheGeometry(
 	FRDGBuilder& GraphBuilder,
 	FGlobalShaderMap* ShaderMap, 
-	const USkeletalMeshComponent* SkeletalMeshComponent, 
+	const FSkeletalMeshSceneProxy* Proxy,
 	const bool bOutputTriangleData,
 	FCachedGeometry& OutCachedGeometry);
 
 void BuildCacheGeometry(
 	FRDGBuilder& GraphBuilder,
 	FGlobalShaderMap* ShaderMap, 
-	const UGeometryCacheComponent* GeometryCacheComponent,
+	const FGeometryCacheSceneProxy* Proxy,
 	const bool bOutputTriangleData,
 	FCachedGeometry& OutCachedGeometry);
\ No newline at end of file
diff --git a/Engine/Plugins/Runtime/HairStrands/Source/HairStrandsCore/Private/GroomManager.cpp b/Engine/Plugins/Runtime/HairStrands/Source/HairStrandsCore/Private/GroomManager.cpp
--- a/Engine/Plugins/Runtime/HairStrands/Source/HairStrandsCore/Private/GroomManager.cpp
+++ b/Engine/Plugins/Runtime/HairStrands/Source/HairStrandsCore/Private/GroomManager.cpp
@@ -43,9 +43,6 @@ static FAutoConsoleVariableRef CVarHairStrands_InterpolationFrustumCullingEnable
 static int32 GHairStrands_ContinousLOD = 0;
 static FAutoConsoleVariableRef CVarHairStrands_ContinousLOD(TEXT("r.HairStrands.ContinuousLOD"), GHairStrands_ContinousLOD, TEXT("Continuous LODing toggle."));
 
-static int32 GHairStrands_UseSkelMeshPersistentPrimitiveIndex = 1;
-static FAutoConsoleVariableRef CVarHairStrands_UseSkelMeshPersistentPrimitiveIndex(TEXT("r.HairStrands.UseSkelMeshPersistentPrimitiveIndex"), GHairStrands_UseSkelMeshPersistentPrimitiveIndex, TEXT("Use persistent primitive index for skel. mesh lookup (experimental)."), ECVF_RenderThreadSafe);
-
 EHairBufferSwapType GetHairSwapBufferType()
 {
 	switch (GHairStrands_SwapBufferType)
@@ -88,7 +85,7 @@ static bool IsSkeletalMeshEvaluationEnabled()
 }
 
 // Retrive the skel. mesh scene info
-static FPrimitiveSceneInfo* GetMeshSceneInfo(FSceneInterface* Scene, FHairGroupInstance* Instance)
+static const FPrimitiveSceneInfo* GetMeshSceneInfo(FSceneInterface* Scene, FHairGroupInstance* Instance)
 {
 	if (!Instance->Debug.MeshComponentId.IsValid())
 	{
@@ -96,7 +93,7 @@ static FPrimitiveSceneInfo* GetMeshSceneInfo(FSceneInterface* Scene, FHairGroupI
 	}
 
 	// Try to use the cached persistent primitive index (fast)
-	FPrimitiveSceneInfo* PrimitiveSceneInfo = nullptr;
+	const FPrimitiveSceneInfo* PrimitiveSceneInfo = nullptr;
 	if (Instance->Debug.CachedMeshPersistentPrimitiveIndex.IsValid())
 	{
 		PrimitiveSceneInfo = Scene->GetPrimitiveSceneInfo(Instance->Debug.CachedMeshPersistentPrimitiveIndex);
@@ -127,32 +124,17 @@ FCachedGeometry GetCacheGeometryForHair(
 	{
 		if (IsSkeletalMeshEvaluationEnabled())
 		{
-			if (GHairStrands_UseSkelMeshPersistentPrimitiveIndex)
+			if (const FPrimitiveSceneInfo* PrimitiveSceneInfo = GetMeshSceneInfo(Scene, Instance))
 			{
-				if (const FPrimitiveSceneInfo* PrimitiveSceneInfo = GetMeshSceneInfo(Scene, Instance))
-				{
-					if (const FSkeletalMeshSceneProxy* SceneProxy = static_cast<const FSkeletalMeshSceneProxy*>(PrimitiveSceneInfo->Proxy))
-					{
-					
-						SceneProxy->GetCachedGeometry(Out);
-					}
-				}
-			}
-			else
-			{
-				//todo: It's unsafe to be accessing the component directly here. This can be solved when we have persistent ids available on the render thread.
-				if (const USkeletalMeshComponent* SkeletalMeshComponent = Cast<USkeletalMeshComponent>(Instance->Debug.MeshComponent))
-				{
-					if (FSkeletalMeshSceneProxy* SceneProxy = static_cast<FSkeletalMeshSceneProxy*>(SkeletalMeshComponent->SceneProxy))
-					{
-						SceneProxy->GetCachedGeometry(Out);
-					}
+				if (const FSkeletalMeshSceneProxy* SceneProxy = reinterpret_cast<const FSkeletalMeshSceneProxy*>(PrimitiveSceneInfo->Proxy))
+				{				
+					SceneProxy->GetCachedGeometry(Out);
 
 					if (GHairStrands_ManualSkinCache > 0 && Out.Sections.Num() == 0)
 					{
 						//#hair_todo: Need to have a (frame) cache to insure that we don't recompute the same projection several time
 						// Actual populate the cache with only the needed part based on the groom projection data. At the moment it recompute everything ...
-						BuildCacheGeometry(GraphBuilder, ShaderMap, SkeletalMeshComponent, bOutputTriangleData, Out);
+						BuildCacheGeometry(GraphBuilder, ShaderMap, SceneProxy, bOutputTriangleData, Out);
 					}
 				}
 			}
@@ -160,9 +142,12 @@ FCachedGeometry GetCacheGeometryForHair(
 	}
 	else if (Instance->Debug.GroomBindingType == EGroomBindingMeshType::GeometryCache)
 	{
-		if (const UGeometryCacheComponent* GeometryCacheComponent = Cast<const UGeometryCacheComponent>(Instance->Debug.MeshComponent))
+		if (const FPrimitiveSceneInfo* PrimitiveSceneInfo = GetMeshSceneInfo(Scene, Instance))
 		{
-			BuildCacheGeometry(GraphBuilder, ShaderMap, GeometryCacheComponent, bOutputTriangleData, Out);
+			if (const FGeometryCacheSceneProxy* SceneProxy = reinterpret_cast<const FGeometryCacheSceneProxy*>(PrimitiveSceneInfo->Proxy))
+			{
+				BuildCacheGeometry(GraphBuilder, ShaderMap, SceneProxy, bOutputTriangleData, Out);
+			}
 		}
 	}
 	return Out;
diff --git a/Engine/Plugins/Runtime/HairStrands/Source/HairStrandsCore/Private/HairStrandsMeshProjection.cpp b/Engine/Plugins/Runtime/HairStrands/Source/HairStrandsCore/Private/HairStrandsMeshProjection.cpp
--- a/Engine/Plugins/Runtime/HairStrands/Source/HairStrandsCore/Private/HairStrandsMeshProjection.cpp
+++ b/Engine/Plugins/Runtime/HairStrands/Source/HairStrandsCore/Private/HairStrandsMeshProjection.cpp
@@ -108,8 +108,7 @@ void AddSkinUpdatePass(
 	FRDGBuilder& GraphBuilder,
 	FGlobalShaderMap* ShaderMap,
 	uint32 SectionIndex,
-	uint32 BonesOffset,
-	FSkinWeightVertexBuffer* SkinWeight,
+	uint32 BonesOffset,	
 	FSkeletalMeshLODRenderData& RenderData,
 	FRHIShaderResourceView* BoneMatrices,
 	FRHIShaderResourceView* PrevBoneMatrices,
@@ -118,6 +117,7 @@ void AddSkinUpdatePass(
 {
 	check(BoneMatrices);
 
+	FSkinWeightVertexBuffer* SkinWeight = &RenderData.SkinWeightVertexBuffer;
 	const FSkelMeshRenderSection& Section = RenderData.RenderSections[SectionIndex];
 	const uint32 NumVertexToProcess = Section.NumVertices;
 	const uint32 SectionVertexBaseIndex = Section.BaseVertexIndex;
diff --git a/Engine/Plugins/Runtime/HairStrands/Source/HairStrandsCore/Public/GroomInstance.h b/Engine/Plugins/Runtime/HairStrands/Source/HairStrandsCore/Public/GroomInstance.h
--- a/Engine/Plugins/Runtime/HairStrands/Source/HairStrandsCore/Public/GroomInstance.h
+++ b/Engine/Plugins/Runtime/HairStrands/Source/HairStrandsCore/Public/GroomInstance.h
@@ -241,7 +241,6 @@ struct HAIRSTRANDSCORE_API FHairGroupInstance : public FHairStrandsInstance
 	struct FDebug
 	{
 		// Data
-		uint32					ComponentId = ~0;
 		uint32					GroupIndex = ~0;
 		uint32					GroupCount = 0;
 		FString					GroomAssetName;
@@ -255,10 +254,10 @@ struct HAIRSTRANDSCORE_API FHairGroupInstance : public FHairStrandsInstance
 		EGroomBindingMeshType	GroomBindingType;
 		EGroomCacheType			GroomCacheType;
 		FPrimitiveSceneProxy*	Proxy = nullptr;
-		UMeshComponent*			MeshComponent = nullptr;
 		FString					MeshComponentName;
 		FPrimitiveComponentId	MeshComponentId;
 		FPersistentPrimitiveIndex CachedMeshPersistentPrimitiveIndex;
+		const UMeshComponent*	MeshComponentForDebug = nullptr;
 		const UGroomComponent*	GroomComponentForDebug = nullptr; // For debug only, shouldn't be deferred on the rendering thread
 		FTransform				RigidCurrentLocalToWorld = FTransform::Identity;
 		FTransform				SkinningCurrentLocalToWorld = FTransform::Identity;
diff --git a/Engine/Plugins/Runtime/HairStrands/Source/HairStrandsCore/Public/HairStrandsMeshProjection.h b/Engine/Plugins/Runtime/HairStrands/Source/HairStrandsCore/Public/HairStrandsMeshProjection.h
--- a/Engine/Plugins/Runtime/HairStrands/Source/HairStrandsCore/Public/HairStrandsMeshProjection.h
+++ b/Engine/Plugins/Runtime/HairStrands/Source/HairStrandsCore/Public/HairStrandsMeshProjection.h
@@ -120,7 +120,6 @@ void AddSkinUpdatePass(
 	FGlobalShaderMap* ShaderMap,
 	uint32 SectionIndex,
 	uint32 BonesOffset, 
-	class FSkinWeightVertexBuffer* SkinWeight,
 	class FSkeletalMeshLODRenderData& RenderData,
 	FRHIShaderResourceView* BoneMatrices,
 	FRHIShaderResourceView* PrevBoneMatrices,
diff --git a/Engine/Source/Runtime/Engine/Private/SkeletalRenderCPUSkin.h b/Engine/Source/Runtime/Engine/Private/SkeletalRenderCPUSkin.h
--- a/Engine/Source/Runtime/Engine/Private/SkeletalRenderCPUSkin.h
+++ b/Engine/Source/Runtime/Engine/Private/SkeletalRenderCPUSkin.h
@@ -119,7 +119,7 @@ public:
 		}
 	}
 
-	virtual bool HaveValidDynamicData() override
+	virtual bool HaveValidDynamicData() const override
 	{ 
 		return ( DynamicData!=NULL ); 
 	}
diff --git a/Engine/Source/Runtime/Engine/Private/SkeletalRenderGPUSkin.h b/Engine/Source/Runtime/Engine/Private/SkeletalRenderGPUSkin.h
--- a/Engine/Source/Runtime/Engine/Private/SkeletalRenderGPUSkin.h
+++ b/Engine/Source/Runtime/Engine/Private/SkeletalRenderGPUSkin.h
@@ -411,7 +411,7 @@ public:
 		}
 	}
 
-	virtual bool HaveValidDynamicData() override
+	virtual bool HaveValidDynamicData() const override
 	{ 
 		return ( DynamicData!=NULL ); 
 	}
diff --git a/Engine/Source/Runtime/Engine/Private/SkeletalRenderStatic.h b/Engine/Source/Runtime/Engine/Private/SkeletalRenderStatic.h
--- a/Engine/Source/Runtime/Engine/Private/SkeletalRenderStatic.h
+++ b/Engine/Source/Runtime/Engine/Private/SkeletalRenderStatic.h
@@ -42,7 +42,7 @@ public:
 	}
 	//virtual const FTwoVectors& GetCustomLeftRightVectors(int32 SectionIndex) const override;
 	virtual void DrawVertexElements(FPrimitiveDrawInterface* PDI, const FMatrix& ToWorldSpace, bool bDrawNormals, bool bDrawTangents, bool bDrawBinormals) const override {};
-	virtual bool HaveValidDynamicData() override
+	virtual bool HaveValidDynamicData() const override
 	{ 
 		return false; 
 	}
diff --git a/Engine/Source/Runtime/Engine/Public/SkeletalMeshSceneProxy.h b/Engine/Source/Runtime/Engine/Public/SkeletalMeshSceneProxy.h
--- a/Engine/Source/Runtime/Engine/Public/SkeletalMeshSceneProxy.h
+++ b/Engine/Source/Runtime/Engine/Public/SkeletalMeshSceneProxy.h
@@ -110,6 +110,8 @@ public:
 
 	bool GetCachedGeometry(struct FCachedGeometry& OutCachedGeometry) const;
 
+	const FSkeletalMeshObject* GetMeshObject() const { return MeshObject; }
+
 protected:
 	AActor* Owner;
 	class FSkeletalMeshObject* MeshObject;
diff --git a/Engine/Source/Runtime/Engine/Public/SkeletalRenderPublic.h b/Engine/Source/Runtime/Engine/Public/SkeletalRenderPublic.h
--- a/Engine/Source/Runtime/Engine/Public/SkeletalRenderPublic.h
+++ b/Engine/Source/Runtime/Engine/Public/SkeletalRenderPublic.h
@@ -179,7 +179,7 @@ public:
 	/**
 	 *	Return true if this does have valid dynamic data to render
 	 */
-	virtual bool HaveValidDynamicData() = 0;
+	virtual bool HaveValidDynamicData() const = 0;
 
 	// allow access to mesh component
 	friend class FDynamicSkelMeshObjectDataCPUSkin;
