Fix CIS issue.

diff --git a/Engine/Plugins/Runtime/HairStrands/Source/HairStrandsCore/Private/GroomManager.cpp b/Engine/Plugins/Runtime/HairStrands/Source/HairStrandsCore/Private/GroomManager.cpp
--- a/Engine/Plugins/Runtime/HairStrands/Source/HairStrandsCore/Private/GroomManager.cpp
+++ b/Engine/Plugins/Runtime/HairStrands/Source/HairStrandsCore/Private/GroomManager.cpp
@@ -14,6 +14,7 @@
 #include "GroomResources.h"
 #include "GroomInstance.h"
 #include "GroomGeometryCache.h"
+#include "GeometryCacheSceneProxy.h"
 #include "HAL/IConsoleManager.h"
 #include "SceneView.h"
 #include "HairCardsVertexFactory.h"
@@ -126,7 +127,7 @@ FCachedGeometry GetCacheGeometryForHair(
 		{
 			if (const FPrimitiveSceneInfo* PrimitiveSceneInfo = GetMeshSceneInfo(Scene, Instance))
 			{
-				if (const FSkeletalMeshSceneProxy* SceneProxy = reinterpret_cast<const FSkeletalMeshSceneProxy*>(PrimitiveSceneInfo->Proxy))
+				if (const FSkeletalMeshSceneProxy* SceneProxy = static_cast<const FSkeletalMeshSceneProxy*>(PrimitiveSceneInfo->Proxy))
 				{				
 					SceneProxy->GetCachedGeometry(Out);
 
@@ -144,7 +145,7 @@ FCachedGeometry GetCacheGeometryForHair(
 	{
 		if (const FPrimitiveSceneInfo* PrimitiveSceneInfo = GetMeshSceneInfo(Scene, Instance))
 		{
-			if (const FGeometryCacheSceneProxy* SceneProxy = reinterpret_cast<const FGeometryCacheSceneProxy*>(PrimitiveSceneInfo->Proxy))
+			if (const FGeometryCacheSceneProxy* SceneProxy = static_cast<const FGeometryCacheSceneProxy*>(PrimitiveSceneInfo->Proxy))
 			{
 				BuildCacheGeometry(GraphBuilder, ShaderMap, SceneProxy, bOutputTriangleData, Out);
 			}
