Fix guide visualization for groom guide cache.

diff --git a/Engine/Plugins/Runtime/HairStrands/Source/HairStrandsCore/Private/HairStrandsRendering.cpp b/Engine/Plugins/Runtime/HairStrands/Source/HairStrandsCore/Private/HairStrandsRendering.cpp
--- a/Engine/Plugins/Runtime/HairStrands/Source/HairStrandsCore/Private/HairStrandsRendering.cpp
+++ b/Engine/Plugins/Runtime/HairStrands/Source/HairStrandsCore/Private/HairStrandsRendering.cpp
@@ -1796,10 +1796,14 @@ void ComputeHairStrandsInterpolation(
 			if (ActiveGroomCacheType == EGroomCacheType::Guides)
 			{
 				FScopeLock Lock(Instance->Debug.GroomCacheBuffers->GetCriticalSection());
-				const FGroomCacheGroupData& GroomCacheGroupData0 = Instance->Debug.GroomCacheBuffers->GetCurrentFrameBuffer().GroupsData[Instance->Debug.GroupIndex];
-				const FGroomCacheGroupData& GroomCacheGroupData1 = Instance->Debug.GroomCacheBuffers->GetNextFrameBuffer().GroupsData[Instance->Debug.GroupIndex];
+				FGroomCacheGroupData* GroomCacheData0 = const_cast<FGroomCacheGroupData*>(&Instance->Debug.GroomCacheBuffers->GetCurrentFrameBuffer().GroupsData[Instance->Debug.GroupIndex]);
+				FGroomCacheGroupData* GroomCacheData1 = const_cast<FGroomCacheGroupData*>(&Instance->Debug.GroomCacheBuffers->GetNextFrameBuffer().GroupsData[Instance->Debug.GroupIndex]);
+
 				const float InterpolationFactor = Instance->Debug.GroomCacheBuffers->GetInterpolationFactor();
-				Instance->Guides.DeformedResource->SetPositionOffset(FHairStrandsDeformedResource::EFrameType::Current, FMath::Lerp(GroomCacheGroupData0.BoundingBox.GetCenter(), GroomCacheGroupData1.BoundingBox.GetCenter(), InterpolationFactor));
+				Instance->Guides.DeformedResource->SetPositionOffset(FHairStrandsDeformedResource::EFrameType::Current, FMath::Lerp(GroomCacheData0->BoundingBox.GetCenter(), GroomCacheData1->BoundingBox.GetCenter(), InterpolationFactor));
+
+				FGroomCacheResources CacheResources0 = CreateGroomCacheBuffer(GraphBuilder, GroomCacheData0->VertexData);
+				FGroomCacheResources CacheResources1 = CreateGroomCacheBuffer(GraphBuilder, GroomCacheData1->VertexData);
 
 				AddHairStrandUpdatePositionOffsetPass(
 					GraphBuilder,
@@ -1807,6 +1811,18 @@ void ComputeHairStrandsInterpolation(
 					MeshLODIndex,
 					Instance->Guides.DeformedRootResource,
 					Instance->Guides.DeformedResource);
+
+				// Pass to upload GroomCache guide positions
+				AddGroomCacheUpdatePass(
+					GraphBuilder,
+					ShaderMap,
+					Instance->Guides.RestResource->GetPointCount(),
+					InterpolationFactor,
+					CacheResources0,
+					CacheResources1,
+					RegisterAsSRV(GraphBuilder, Instance->Guides.RestResource->PositionBuffer),
+					RegisterAsSRV(GraphBuilder, Instance->Guides.DeformedResource->GetPositionOffsetBuffer(FHairStrandsDeformedResource::EFrameType::Current)),
+					RegisterAsUAV(GraphBuilder, Instance->Guides.DeformedResource->GetBuffer(FHairStrandsDeformedResource::Current)));
 			}
 
 			AddHairTangentPass(
