Add LOD curve/point count into the LOD panels.

diff --git a/Engine/Plugins/Runtime/HairStrands/Source/HairStrandsCore/Private/GroomBuilder.cpp b/Engine/Plugins/Runtime/HairStrands/Source/HairStrandsCore/Private/GroomBuilder.cpp
--- a/Engine/Plugins/Runtime/HairStrands/Source/HairStrandsCore/Private/GroomBuilder.cpp
+++ b/Engine/Plugins/Runtime/HairStrands/Source/HairStrandsCore/Private/GroomBuilder.cpp
@@ -40,7 +40,7 @@ static FAutoConsoleVariableRef CVarHairGroupIndexBuilder_MaxVoxelResolution(TEXT
 
 FString FGroomBuilder::GetVersion()
 {
-	return TEXT("v8a");
+	return TEXT("v8e");
 }
 
 namespace FHairStrandsDecimation
@@ -2788,6 +2788,7 @@ static void BuildClusterData(
 	Out.ClusterCount = ValidClusterIndices.Num();
 	Out.ClusterInfos.Init(FHairClusterInfo(), Out.ClusterCount);
 	Out.VertexToClusterIds.SetNum(Out.VertexCount);
+	Out.LODInfos.SetNum(LODCount);
 
 	// Conservative allocation for inserting vertex indices for the various curves LOD
 	uint32* RawClusterVertexIds = new uint32[LODCount * InRenStrandsData.GetNumPoints()];
@@ -2902,6 +2903,9 @@ static void BuildClusterData(
 		{
 			const float CurveDecimation = FMath::Clamp(InSettings.LODs[LODIt].CurveDecimation, 0.0f, 1.0f);
 			LODCurveCount[LODIt] = FMath::Clamp(FMath::CeilToInt(Cluster.ClusterCurves.Num() * CurveDecimation), 1, Cluster.ClusterCurves.Num());
+
+			// For stats
+			Out.LODInfos[LODIt].CurveCount += LODCurveCount[LODIt];
 		}
 
 		// 4.4 Decimate each curve for all LODs
@@ -2920,6 +2924,12 @@ static void BuildClusterData(
 					InSettings.LODs,
 					ClusterCurve.CountPerLOD,
 					VertexLODMasks);
+
+				// For stats
+				for (uint32 LODIt = 0; LODIt < LODCount; ++LODIt)
+				{
+					Out.LODInfos[LODIt].PointCount += ClusterCurve.CountPerLOD[LODIt];
+				}
 			}
 		}
 
@@ -3085,6 +3095,7 @@ static void BuildClusterBulkData(
 
 	Out.CPULODScreenSize= In.CPULODScreenSize;
 	Out.LODVisibility	= In.LODVisibility;
+	Out.LODInfos		= In.LODInfos;
 
 	// Sanity check
 	check(Out.ClusterCount			== uint32(In.ClusterInfos.Num()));
diff --git a/Engine/Plugins/Runtime/HairStrands/Source/HairStrandsCore/Public/GroomAsset.h b/Engine/Plugins/Runtime/HairStrands/Source/HairStrandsCore/Public/GroomAsset.h
--- a/Engine/Plugins/Runtime/HairStrands/Source/HairStrandsCore/Public/GroomAsset.h
+++ b/Engine/Plugins/Runtime/HairStrands/Source/HairStrandsCore/Public/GroomAsset.h
@@ -38,6 +38,18 @@ enum class EHairGroupInfoFlags : uint8
 	HasTrimmedCurve = 2
 };
 
+USTRUCT(BlueprintType)
+struct HAIRSTRANDSCORE_API FHairGroupLODInfo
+{
+	GENERATED_BODY()
+
+	UPROPERTY(VisibleAnywhere, Category = "Info", meta = (DisplayName = "Curve Count"))
+	int32 NumPoints = 0;
+
+	UPROPERTY(VisibleAnywhere, Category = "Info", meta = (DisplayName = "Curve Count"))
+	int32 NumCurves = 0;
+};
+
 USTRUCT(BlueprintType)
 struct HAIRSTRANDSCORE_API FHairGroupInfo
 {
@@ -69,6 +81,9 @@ struct HAIRSTRANDSCORE_API FHairGroupInfo
 
 	UPROPERTY()
 	uint32 Flags = 0;
+
+	UPROPERTY()
+	TArray<FHairGroupLODInfo> LODInfos;
 };
 
 USTRUCT(BlueprintType)
diff --git a/Engine/Plugins/Runtime/HairStrands/Source/HairStrandsCore/Public/HairStrandsDatas.h b/Engine/Plugins/Runtime/HairStrands/Source/HairStrandsCore/Public/HairStrandsDatas.h
--- a/Engine/Plugins/Runtime/HairStrands/Source/HairStrandsCore/Public/HairStrandsDatas.h
+++ b/Engine/Plugins/Runtime/HairStrands/Source/HairStrandsCore/Public/HairStrandsDatas.h
@@ -729,6 +729,12 @@ struct FHairClusterIndexFormat
 	static const EPixelFormat Format = PF_R32_UINT;
 };
 
+struct FHairLODInfo
+{
+	uint32 CurveCount = 0;
+	uint32 PointCount = 0;
+};
+
 struct HAIRSTRANDSCORE_API FHairStrandsClusterCullingData
 {
 	FHairStrandsClusterCullingData();
@@ -746,6 +752,7 @@ struct HAIRSTRANDSCORE_API FHairStrandsClusterCullingData
 	TArray<FHairClusterLODInfo> ClusterLODInfos;
 	TArray<uint32>				VertexToClusterIds;
 	TArray<uint32>				ClusterVertexIds;
+	TArray<FHairLODInfo>		LODInfos;
 
 	/* Number of cluster  */
 	uint32 ClusterCount = 0;
@@ -768,6 +775,9 @@ struct HAIRSTRANDSCORE_API FHairStrandsClusterCullingBulkData
 	/* Screen size at which LOD should switches on CPU */
 	TArray<float>	CPULODScreenSize;
 
+	/* Curve count and Point count per LOD */
+	TArray<FHairLODInfo> LODInfos;
+
 	/* LOD info for the various clusters for LOD management on GPU */
 	FByteBulkData	PackedClusterInfos;		// Size - ClusterCount
 	FByteBulkData	ClusterLODInfos;		// Size - ClusterLODCount
diff --git a/Engine/Plugins/Runtime/HairStrands/Source/HairStrandsEditor/Private/GroomAssetDetails.cpp b/Engine/Plugins/Runtime/HairStrands/Source/HairStrandsEditor/Private/GroomAssetDetails.cpp
--- a/Engine/Plugins/Runtime/HairStrands/Source/HairStrandsEditor/Private/GroomAssetDetails.cpp
+++ b/Engine/Plugins/Runtime/HairStrands/Source/HairStrandsEditor/Private/GroomAssetDetails.cpp
@@ -64,9 +64,6 @@
 
 #include "Styling/AppStyle.h"
 
-static FLinearColor HairGroupColor(1.0f, 0.5f, 0.0f);
-static FLinearColor HairLODColor(1.0f, 0.5f, 0.0f);
-
 #define LOCTEXT_NAMESPACE "GroomRenderingDetails"
 
 DEFINE_LOG_CATEGORY_STATIC(LogGroomAssetDetails, Log, All);
@@ -112,6 +109,49 @@ static uint32 GetHairCardsAtlasResolution(int32 InLODIndex, int32 PrevResolution
 	return FMath::Clamp(OutResolution, MinResolution, MaxResolution);
 }
 
+
+///////////////////////////////////////////////////////////////////////////////////////////////////////////
+// Array panel for hair strands infos
+TSharedRef<SUniformGridPanel> MakeHairStrandsLODInfoGrid(const FSlateFontInfo& DetailFontInfo, const FHairLODInfo& LODInfo)
+{
+	TSharedRef<SUniformGridPanel> Grid = SNew(SUniformGridPanel).SlotPadding(2.0f);
+
+	// Header
+	Grid->AddSlot(1, 0) // x, y
+	.HAlign(HAlign_Right)
+	[
+		SNew(STextBlock)
+		.Font(DetailFontInfo)
+		.Text(LOCTEXT("HairInfo_CurveLOD", "Curves"))
+	];
+
+	Grid->AddSlot(2, 0) // x, y
+	.HAlign(HAlign_Right)
+	[
+		SNew(STextBlock)
+		.Font(DetailFontInfo)
+		.Text(LOCTEXT("HairInfo_PointLOD", "Points"))
+	];
+
+	// Strands
+	Grid->AddSlot(1, 1) // x, y
+	.HAlign(HAlign_Right)
+	[
+		SNew(STextBlock)
+		.Font(DetailFontInfo)
+		.Text(FText::AsNumber(LODInfo.CurveCount))
+	];
+	Grid->AddSlot(2, 1) // x, y
+	.HAlign(HAlign_Right)
+	[
+		SNew(STextBlock)
+		.Font(DetailFontInfo)
+		.Text(FText::AsNumber(LODInfo.PointCount))
+	];
+	   
+	return Grid;
+}
+
 ///////////////////////////////////////////////////////////////////////////////////////////////////////////
 // Array panel for hair strands infos
 TSharedRef<SUniformGridPanel> MakeHairStrandsInfoGrid(const FSlateFontInfo& DetailFontInfo, FHairGroupInfoWithVisibility& CurrentAsset, float MaxRadius)
@@ -1176,11 +1216,6 @@ FReply FGroomRenderingDetails::OnGenerateCardDataUsingPlugin(int32 GroupIndex)
 	return FReply::Handled();
 }
 
-void FGroomRenderingDetails::AddLODSlot(TSharedRef<IPropertyHandle>& LODHandle, IDetailChildrenBuilder& ChildrenBuilder, int32 GroupIndex, int32 LODIndex)
-{	
-	ExpandStruct(LODHandle, ChildrenBuilder, GroupIndex, LODIndex, true);
-}
-
 void FGroomRenderingDetails::OnGenerateElementForLODs(TSharedRef<IPropertyHandle> StructProperty, int32 LODIndex, IDetailChildrenBuilder& ChildrenBuilder, IDetailLayoutBuilder* DetailLayout, int32 GroupIndex)
 {
 	const FSlateFontInfo DetailFontInfo = IDetailLayoutBuilder::GetDetailFont();
@@ -1191,6 +1226,7 @@ void FGroomRenderingDetails::OnGenerateElementForLODs(TSharedRef<IPropertyHandle
 	static const FSlateBrush* GenericBrush = FCoreStyle::Get().GetBrush("GenericWhiteBox");
 	float OtherMargin = 2.0f;
 
+	// LOD Bar with add button
 	ChildrenBuilder.AddCustomRow(LOCTEXT("HairInfo_Separator", "Separator"))
 	.WholeRowContent()
 	.VAlign(VAlign_Fill)
@@ -1234,6 +1270,20 @@ void FGroomRenderingDetails::OnGenerateElementForLODs(TSharedRef<IPropertyHandle
 		]
 	];
 
+	// LOD Stats
+	const FHairStrandsClusterCullingBulkData& ClusterCullingBulkData = GroomAsset->HairGroupsData[GroupIndex].Strands.ClusterCullingBulkData;
+	if (ClusterCullingBulkData.IsValid() && LODIndex < ClusterCullingBulkData.LODInfos.Num())
+	{
+		const FHairLODInfo& LODInfo = ClusterCullingBulkData.LODInfos[LODIndex];
+
+		ChildrenBuilder.AddCustomRow(LOCTEXT("HairStrandsLODInfo_Array", "HairStrandsLODInfo"))
+		.ValueContent()
+		.HAlign(HAlign_Fill)
+		[
+			MakeHairStrandsLODInfoGrid(DetailFontInfo, LODInfo)
+		];
+	}
+	
 	// Rename the array entry name by its group name and adds all its existing properties
 	StructProperty->SetPropertyDisplayName(LOCTEXT("LODProperties", "LOD Properties"));
 	ExpandStructForLOD(StructProperty, ChildrenBuilder, GroupIndex, LODIndex, true); ///
diff --git a/Engine/Plugins/Runtime/HairStrands/Source/HairStrandsEditor/Private/GroomAssetDetails.h b/Engine/Plugins/Runtime/HairStrands/Source/HairStrandsEditor/Private/GroomAssetDetails.h
--- a/Engine/Plugins/Runtime/HairStrands/Source/HairStrandsEditor/Private/GroomAssetDetails.h
+++ b/Engine/Plugins/Runtime/HairStrands/Source/HairStrandsEditor/Private/GroomAssetDetails.h
@@ -67,7 +67,6 @@ private:
 	void CustomizeStrandsGroupProperties(IDetailLayoutBuilder& DetailLayout, IDetailCategoryBuilder& StrandsGroupFilesCategory);
 	void OnGenerateElementForHairGroup(TSharedRef<IPropertyHandle> StructProperty, int32 GroupIndex, IDetailChildrenBuilder& ChildrenBuilder, IDetailLayoutBuilder* DetailLayout);
 	void OnGenerateElementForLODs(TSharedRef<IPropertyHandle> StructProperty, int32 LODIndex, IDetailChildrenBuilder& ChildrenBuilder, IDetailLayoutBuilder* DetailLayout, int32 GroupIndex);
-	void AddLODSlot(TSharedRef<IPropertyHandle>& LODHandle, IDetailChildrenBuilder& ChildrenBuilder, int32 GroupIndex, int32 LODIndex);
 	void OnGenerateElementForBindingAsset(TSharedRef<IPropertyHandle> StructProperty, int32 BindingIndex, IDetailChildrenBuilder& ChildrenBuilder, IDetailLayoutBuilder* DetailLayout);
 
 	// Display custom thumbnail for material
