Add imported attributes visualization in groom asset.

diff --git a/Engine/Plugins/Runtime/HairStrands/Source/HairStrandsCore/Private/GroomBuilder.cpp b/Engine/Plugins/Runtime/HairStrands/Source/HairStrandsCore/Private/GroomBuilder.cpp
--- a/Engine/Plugins/Runtime/HairStrands/Source/HairStrandsCore/Private/GroomBuilder.cpp
+++ b/Engine/Plugins/Runtime/HairStrands/Source/HairStrandsCore/Private/GroomBuilder.cpp
@@ -40,7 +40,7 @@ static FAutoConsoleVariableRef CVarHairGroupIndexBuilder_MaxVoxelResolution(TEXT
 
 FString FGroomBuilder::GetVersion()
 {
-	return TEXT("v8e");
+	return TEXT("v8g");
 }
 
 namespace FHairStrandsDecimation
@@ -188,9 +188,9 @@ namespace HairStrandsBuilder
 		TArray<uint32> AttributeRoughness;
 		TArray<uint32> AttributeAO;
 
-		AttributeRootUV.SetNum(NumCurves);									// 4 bytes encoding - Per-Curve
 		AttributeSeed.SetNum(FMath::DivideAndRoundUp(NumCurves, 4u));		// 1 bytes encoding - Per-Curve
 		AttributeLength.SetNum(FMath::DivideAndRoundUp(NumCurves, 2u));		// 2 bytes encoding - Per-Curve
+		if (HasHairAttribute(Attributes, EHairAttribute::RootUV)) 	{ AttributeRootUV.SetNum(NumCurves); }									// 4 bytes encoding - Per-Curve
 		if (HasHairAttribute(Attributes, EHairAttribute::ClumpID))	{ AttributeClumpIDs.SetNum(FMath::DivideAndRoundUp(NumCurves, 2u)); }	// 2 bytes encoding - Per-Curve
 		if (HasHairAttribute(Attributes, EHairAttribute::Color))	{ AttributeBaseColor.SetNum(NumPoints); }								// 4 bytes encoding - Per-Vertex
 		if (HasHairAttribute(Attributes, EHairAttribute::Roughness)){ AttributeRoughness.SetNum(FMath::DivideAndRoundUp(NumPoints, 2u)); }	// 1 bytes encoding - Per-Vertex
@@ -341,6 +341,8 @@ namespace HairStrandsBuilder
 		OutBulkData.MaxLength = MaxLength;
 		OutBulkData.MaxRadius = MaxRadius;
 		OutBulkData.Flags = FHairStrandsBulkData::DataFlags_HasData;
+		OutBulkData.ImportedAttributes = HairStrands.GetAttributes();
+		OutBulkData.ImportedAttributeFlags = HairStrands.GetAttributeFlags();
 
 		// Concatenate all attributes
 		TArray<FHairStrandsAttributeFormat::Type> OutPackedAttributes;
@@ -1683,11 +1685,19 @@ bool FGroomBuilder::BuildHairDescriptionGroups(const FHairDescription& HairDescr
 		if (bHasClumpIDs)
 		{
 			CurrentHairStrandsDatas->StrandsCurves.ClumpIDs.Add(ClumpIDs[StrandID]);
+			if (false)
+			{
+				SetHairAttributeFlags(CurrentHairStrandsDatas->StrandsCurves.AttributeFlags, EHairAttributeFlags::HasMultipleClumpIDs);
+			}
 		}
 
 		if (bHasUVData)
 		{
 			CurrentHairStrandsDatas->StrandsCurves.CurvesRootUV.Add(StrandRootUV[StrandID]);
+			if (StrandRootUV[StrandID].X > 1.f || StrandRootUV[StrandID].Y > 1.f)
+			{
+				SetHairAttributeFlags(CurrentHairStrandsDatas->StrandsCurves.AttributeFlags, EHairAttributeFlags::HasRootUDIM);
+			}
 		}
 
 		// Groom 
@@ -1757,16 +1767,6 @@ bool FGroomBuilder::BuildHairDescriptionGroups(const FHairDescription& HairDescr
 		for (FHairDescriptionGroup& Group : Out.HairGroups)
 		{
 			Group.Info.GroupID = GroupIndex++;
-
-			// Propagate valid optional attributes
-			for (uint32 AttributeIt=0; AttributeIt< uint32(EHairAttribute::Count); ++AttributeIt)
-			{
-				const EHairAttribute HairAttribute = (EHairAttribute)AttributeIt;
-				if (HairDescription.HasAttribute(HairAttribute)) 
-				{ 
-					SetHairAttribute(Group.Attributes, HairAttribute); 
-				}
-			}
 		}
 	}
 
diff --git a/Engine/Plugins/Runtime/HairStrands/Source/HairStrandsCore/Private/GroomDebug.cpp b/Engine/Plugins/Runtime/HairStrands/Source/HairStrandsCore/Private/GroomDebug.cpp
--- a/Engine/Plugins/Runtime/HairStrands/Source/HairStrandsCore/Private/GroomDebug.cpp
+++ b/Engine/Plugins/Runtime/HairStrands/Source/HairStrandsCore/Private/GroomDebug.cpp
@@ -853,7 +853,7 @@ struct FHairDebugNameInfo
 	uint8  Pad0;
 };
 
-const TCHAR* GetHairAttributeText(EHairAttribute In);
+const TCHAR* GetHairAttributeText(EHairAttribute In, uint32 InFlags);
 uint32 GetHairAttributeIndex(EHairAttribute In);
 EGroomCacheType GetHairInstanceCacheType(const FHairGroupInstance* Instance);
 
@@ -1036,7 +1036,7 @@ static void AddHairDebugPrintInstancePass(
 		// HAIR_ATTRIBUTE_XXX and EHairAttribute don't have a 1:1 mapping
 		const EHairAttribute Attribute = (EHairAttribute)AttributeIt;
 		{
-			const FString Name = GetHairAttributeText(Attribute);
+			const FString Name = GetHairAttributeText(Attribute, 0u);
 			FHairDebugNameInfo& NameInfo = AttributeNameInfos.AddDefaulted_GetRef();
 			NameInfo.PrimitiveID = ~0;
 			NameInfo.Length = Name.Len();
diff --git a/Engine/Plugins/Runtime/HairStrands/Source/HairStrandsCore/Private/HairAttributes.cpp b/Engine/Plugins/Runtime/HairStrands/Source/HairStrandsCore/Private/HairAttributes.cpp
--- a/Engine/Plugins/Runtime/HairStrands/Source/HairStrandsCore/Private/HairAttributes.cpp
+++ b/Engine/Plugins/Runtime/HairStrands/Source/HairStrandsCore/Private/HairAttributes.cpp
@@ -49,15 +49,25 @@ void SetHairAttribute(uint32& Out, EHairAttribute InAttribute)
 	Out |= (1u << uint32(InAttribute)); 
 }
 
-const TCHAR* GetHairAttributeText(EHairAttribute In)
+bool HasHairAttributeFlags(uint32 In, EHairAttributeFlags InFlag)
+{ 
+	return (In & (1u << uint32(InFlag))) != 0; 
+}
+
+void SetHairAttributeFlags(uint32& Out, EHairAttributeFlags InFlag) 
+{ 
+	Out |= (1u << uint32(InFlag));
+}
+
+const TCHAR* GetHairAttributeText(EHairAttribute In, uint32 InFlags)
 {
 	// If a new optional attribute is added, please add its UI/text description here
 	static_assert(uint32(EHairAttribute::Count) == 7);
 
 	switch (In)
 	{
-	case EHairAttribute::RootUV:					return TEXT("RootUV");
-	case EHairAttribute::ClumpID:					return TEXT("ClumpID");
+	case EHairAttribute::RootUV:					return HasHairAttributeFlags(InFlags, EHairAttributeFlags::HasRootUDIM) ? TEXT("RootUV(UDIM)") : TEXT("RootUV");
+	case EHairAttribute::ClumpID:					return HasHairAttributeFlags(InFlags, EHairAttributeFlags::HasMultipleClumpIDs) ? TEXT("ClumpID(3)") : TEXT("ClumpID");
 	case EHairAttribute::StrandID:					return TEXT("StrandID");
 	case EHairAttribute::PrecomputedGuideWeights:	return TEXT("ImportedGuideWeights");
 	case EHairAttribute::Color:						return TEXT("Color");
@@ -118,6 +128,11 @@ uint32 FHairStrandsDatas::GetAttributes() const
 	return Out;
 }
 
+uint32 FHairStrandsDatas::GetAttributeFlags() const
+{
+	return StrandsCurves.AttributeFlags;
+}
+
 void FHairStrandsDatas::CopyPoint(const FHairStrandsDatas& In, FHairStrandsDatas& Out, uint32 InAttributes, uint32 InIndex, uint32 OutIndex)
 {
 	// Multiply by MaxRadius as HairStrandsBuilder::BuildInternalData will recompute MaxRadius based on the actual Radius not the not the normalized radius.
diff --git a/Engine/Plugins/Runtime/HairStrands/Source/HairStrandsCore/Private/HairStrandsDatas.cpp b/Engine/Plugins/Runtime/HairStrands/Source/HairStrandsCore/Private/HairStrandsDatas.cpp
--- a/Engine/Plugins/Runtime/HairStrands/Source/HairStrandsCore/Private/HairStrandsDatas.cpp
+++ b/Engine/Plugins/Runtime/HairStrands/Source/HairStrandsCore/Private/HairStrandsDatas.cpp
@@ -246,6 +246,8 @@ void FHairStrandsBulkData::Serialize(FArchive& Ar, UObject* Owner)
 	{
 		Ar << AttributeOffsets[AttributeIt];
 	}
+	Ar << ImportedAttributes;
+	Ar << ImportedAttributeFlags;
 
 	// Forced not inline means the bulk data won't automatically be loaded when we deserialize
 	// but only when we explicitly take action to load it
diff --git a/Engine/Plugins/Runtime/HairStrands/Source/HairStrandsCore/Public/GroomAsset.h b/Engine/Plugins/Runtime/HairStrands/Source/HairStrandsCore/Public/GroomAsset.h
--- a/Engine/Plugins/Runtime/HairStrands/Source/HairStrandsCore/Public/GroomAsset.h
+++ b/Engine/Plugins/Runtime/HairStrands/Source/HairStrandsCore/Public/GroomAsset.h
@@ -390,7 +390,8 @@ struct FHairDescriptionGroup
 	FHairGroupInfo		Info;
 	FHairStrandsDatas	Strands;
 	FHairStrandsDatas	Guides;
-	uint32				Attributes = 0;
+	uint32 GetHairAttributes() const { return Strands.GetAttributes() | Guides.GetAttributes(); }
+	uint32 GetHairAttributeFlags() const { return Strands.GetAttributeFlags() | Guides.GetAttributeFlags(); }
 };
 
 struct FHairDescriptionGroups
diff --git a/Engine/Plugins/Runtime/HairStrands/Source/HairStrandsCore/Public/GroomImportOptions.h b/Engine/Plugins/Runtime/HairStrands/Source/HairStrandsCore/Public/GroomImportOptions.h
--- a/Engine/Plugins/Runtime/HairStrands/Source/HairStrandsCore/Public/GroomImportOptions.h
+++ b/Engine/Plugins/Runtime/HairStrands/Source/HairStrandsCore/Public/GroomImportOptions.h
@@ -42,6 +42,9 @@ struct HAIRSTRANDSCORE_API FGroomHairGroupPreview
 	UPROPERTY()
 	uint32 Attributes = 0;
 
+	UPROPERTY()
+	uint32 AttributeFlags = 0;
+
 	UPROPERTY()
 	uint32 Flags = 0;
 
diff --git a/Engine/Plugins/Runtime/HairStrands/Source/HairStrandsCore/Public/HairAttributes.h b/Engine/Plugins/Runtime/HairStrands/Source/HairStrandsCore/Public/HairAttributes.h
--- a/Engine/Plugins/Runtime/HairStrands/Source/HairStrandsCore/Public/HairAttributes.h
+++ b/Engine/Plugins/Runtime/HairStrands/Source/HairStrandsCore/Public/HairAttributes.h
@@ -58,6 +58,13 @@ enum class EHairAttribute : uint8
 	AO,
 	Count
 };
-
 HAIRSTRANDSCORE_API bool HasHairAttribute(uint32 In, EHairAttribute InAttribute);
-HAIRSTRANDSCORE_API void SetHairAttribute(uint32& Out, EHairAttribute InAttribute);
\ No newline at end of file
+HAIRSTRANDSCORE_API void SetHairAttribute(uint32& Out, EHairAttribute InAttribute);
+
+enum class EHairAttributeFlags : uint8
+{
+	HasRootUDIM,
+	HasMultipleClumpIDs
+};
+HAIRSTRANDSCORE_API bool HasHairAttributeFlags(uint32 In, EHairAttributeFlags InFlag);
+HAIRSTRANDSCORE_API void SetHairAttributeFlags(uint32& Out, EHairAttributeFlags InFlag);
diff --git a/Engine/Plugins/Runtime/HairStrands/Source/HairStrandsCore/Public/HairStrandsDatas.h b/Engine/Plugins/Runtime/HairStrands/Source/HairStrandsCore/Public/HairStrandsDatas.h
--- a/Engine/Plugins/Runtime/HairStrands/Source/HairStrandsCore/Public/HairStrandsDatas.h
+++ b/Engine/Plugins/Runtime/HairStrands/Source/HairStrandsCore/Public/HairStrandsDatas.h
@@ -415,9 +415,10 @@ struct HAIRSTRANDSCORE_API FHairStrandsCurves
 
 	/** Custom guid weights (indexed with StrandID) (optional) */
 	TArray<FVector> CurvesClosestGuideWeights;
-};
 
-struct FHairStrandsBulkData;
+	/** Flags for attributes */
+	uint32 AttributeFlags = 0;
+};
 
 /** Hair strands datas that are stored on CPU */
 struct HAIRSTRANDSCORE_API FHairStrandsDatas
@@ -429,6 +430,7 @@ struct HAIRSTRANDSCORE_API FHairStrandsDatas
 	uint32 GetNumCurves() const { return StrandsCurves.Num(); }
 
 	uint32 GetAttributes() const;
+	uint32 GetAttributeFlags() const;
 
 	void Reset();
 
@@ -483,6 +485,10 @@ struct HAIRSTRANDSCORE_API FHairStrandsBulkData
 	uint32 Flags = 0;
 	uint32 AttributeOffsets[HAIR_ATTRIBUTE_COUNT] = {0};
 
+	/** Imported attribute info */
+	uint32 ImportedAttributes = 0;
+	uint32 ImportedAttributeFlags = 0;
+
 	FByteBulkData Positions;	// Size = PointCount
 	FByteBulkData Attributes;	// Size = x*PointCount + y*CurveCount (depends on the stored attributes, which could be per-vertex(x) or per-curve(y))
 	FByteBulkData VertexToCurve;// Size = PointCount
diff --git a/Engine/Plugins/Runtime/HairStrands/Source/HairStrandsEditor/Private/GroomActions.cpp b/Engine/Plugins/Runtime/HairStrands/Source/HairStrandsEditor/Private/GroomActions.cpp
--- a/Engine/Plugins/Runtime/HairStrands/Source/HairStrandsEditor/Private/GroomActions.cpp
+++ b/Engine/Plugins/Runtime/HairStrands/Source/HairStrandsEditor/Private/GroomActions.cpp
@@ -115,11 +115,12 @@ void ExecuteRebuild(const FToolMenuContext& InContext)
 					for (uint32 GroupIndex = 0; GroupIndex < GroupCount; GroupIndex++)
 					{
 						FGroomHairGroupPreview& OutGroup = GroupsPreview->Groups.AddDefaulted_GetRef();
-						OutGroup.GroupID    = GroupIndex;
-						OutGroup.GroupName	= GroomAsset->HairGroupsInfo[GroupIndex].GroupName;
-						OutGroup.CurveCount = GroomAsset->HairGroupsData[GroupIndex].Strands.BulkData.GetNumCurves();
-						OutGroup.GuideCount = GroomAsset->HairGroupsData[GroupIndex].Guides.BulkData.GetNumCurves();
-						OutGroup.Attributes = HairDescriptionGroups.HairGroups[GroupIndex].Attributes;
+						OutGroup.GroupID    	= GroupIndex;
+						OutGroup.GroupName		= GroomAsset->HairGroupsInfo[GroupIndex].GroupName;
+						OutGroup.CurveCount 	= GroomAsset->HairGroupsData[GroupIndex].Strands.BulkData.GetNumCurves();
+						OutGroup.GuideCount 	= GroomAsset->HairGroupsData[GroupIndex].Guides.BulkData.GetNumCurves();
+						OutGroup.Attributes 	= HairDescriptionGroups.HairGroups[GroupIndex].GetHairAttributes();
+						OutGroup.AttributeFlags = HairDescriptionGroups.HairGroups[GroupIndex].GetHairAttributeFlags();
 						OutGroup.InterpolationSettings = GroomAsset->HairGroupsInterpolation[GroupIndex];
 						OutGroup.InterpolationSettings.RiggingSettings.bCanEditRigging = true;
 					}
diff --git a/Engine/Plugins/Runtime/HairStrands/Source/HairStrandsEditor/Private/GroomAssetDetails.cpp b/Engine/Plugins/Runtime/HairStrands/Source/HairStrandsEditor/Private/GroomAssetDetails.cpp
--- a/Engine/Plugins/Runtime/HairStrands/Source/HairStrandsEditor/Private/GroomAssetDetails.cpp
+++ b/Engine/Plugins/Runtime/HairStrands/Source/HairStrandsEditor/Private/GroomAssetDetails.cpp
@@ -109,6 +109,35 @@ static uint32 GetHairCardsAtlasResolution(int32 InLODIndex, int32 PrevResolution
 	return FMath::Clamp(OutResolution, MinResolution, MaxResolution);
 }
 
+FText GetHairAttributeLocText(EHairAttribute In, uint32 InFlags);
+
+///////////////////////////////////////////////////////////////////////////////////////////////////////////
+// Array panel for hair strands infos
+TSharedRef<SUniformGridPanel> MakeHairStrandsAttributeInfoGrid(const FSlateFontInfo& DetailFontInfo, const uint32 InAttributes, const uint32 InAttributeFlags)
+{
+	TSharedRef<SUniformGridPanel> Grid = SNew(SUniformGridPanel).SlotPadding(2.0f);
+
+	const FLinearColor AttributeColor = FLinearColor(FColor::Yellow);
+
+	uint32 SlotIndex = 0;
+	for (uint32 AttributeIt = 0, AttributeCount = uint32(EHairAttribute::Count); AttributeIt<AttributeCount; ++AttributeIt)
+	{
+		const EHairAttribute Attribute = (EHairAttribute)AttributeIt;
+		if (HasHairAttribute(InAttributes, Attribute))
+		{
+			Grid->AddSlot(0, SlotIndex++) // x, y
+			.HAlign(HAlign_Right)
+			[
+				SNew(STextBlock)
+				.Font(DetailFontInfo)
+				.ColorAndOpacity(AttributeColor)
+				.Text(GetHairAttributeLocText(Attribute, InAttributeFlags))
+			];
+		}
+	}
+
+	return Grid;
+}
 
 ///////////////////////////////////////////////////////////////////////////////////////////////////////////
 // Array panel for hair strands infos
@@ -1453,6 +1482,12 @@ void FGroomRenderingDetails::OnGenerateElementForHairGroup(TSharedRef<IPropertyH
 		[
 			MakeHairStrandsInfoGrid(DetailFontInfo, GroomAsset->HairGroupsInfo[GroupIndex], GroomAsset->HairGroupsData[GroupIndex].Strands.BulkData.MaxRadius)
 		];
+		ChildrenBuilder.AddCustomRow(LOCTEXT("HairStrandsAttributeInfo_Array", "HairStrandsAttributeInfo"))
+		.ValueContent()
+		.HAlign(HAlign_Fill)
+		[
+			MakeHairStrandsAttributeInfoGrid(DetailFontInfo, GroomAsset->HairGroupsData[GroupIndex].Strands.BulkData.ImportedAttributes, GroomAsset->HairGroupsData[GroupIndex].Strands.BulkData.ImportedAttributeFlags)
+		];
 	}
 
 	if (GroomAsset != nullptr && GroupIndex >= 0 && GroupIndex < GroomAsset->HairGroupsCards.Num() && (PanelType == EMaterialPanelType::Cards))
diff --git a/Engine/Plugins/Runtime/HairStrands/Source/HairStrandsEditor/Private/GroomImportOptionsWindow.cpp b/Engine/Plugins/Runtime/HairStrands/Source/HairStrandsEditor/Private/GroomImportOptionsWindow.cpp
--- a/Engine/Plugins/Runtime/HairStrands/Source/HairStrandsEditor/Private/GroomImportOptionsWindow.cpp
+++ b/Engine/Plugins/Runtime/HairStrands/Source/HairStrandsEditor/Private/GroomImportOptionsWindow.cpp
@@ -241,15 +241,15 @@ static void AddAttribute(SVerticalBox::FScopedWidgetSlotArguments& Slot, FText A
 	];
 }
 
-static FText GetHairAttributeLocText(EHairAttribute In)
+FText GetHairAttributeLocText(EHairAttribute In, uint32 InFlags)
 {
 	// If a new optional attribute is added, please add its UI/text description here
 	static_assert(uint32(EHairAttribute::Count) == 7);
 
 	switch (In)
 	{
-	case EHairAttribute::RootUV:					return LOCTEXT("GroomOptionsWindow_HasRootUV", "Root UV");
-	case EHairAttribute::ClumpID:					return LOCTEXT("GroomOptionsWindow_HasClumpID", "Clump ID");
+	case EHairAttribute::RootUV:					return HasHairAttributeFlags(InFlags, EHairAttributeFlags::HasRootUDIM) ? LOCTEXT("GroomOptionsWindow_HasRootUDIM", "Root UV (UDIM)") : LOCTEXT("GroomOptionsWindow_HasRootUV", "Root UV");
+	case EHairAttribute::ClumpID:					return HasHairAttributeFlags(InFlags, EHairAttributeFlags::HasMultipleClumpIDs) ? LOCTEXT("GroomOptionsWindow_HasClumpIDs", "Clump IDs (3)") : LOCTEXT("GroomOptionsWindow_HasClumpID", "Clump ID");
 	case EHairAttribute::StrandID:					return LOCTEXT("GroomOptionsWindow_HasStrandID", "Strand ID");
 	case EHairAttribute::PrecomputedGuideWeights:	return LOCTEXT("GroomOptionsWindow_HasPercomputedGuideWeights", "Pre-Computed Guide Weights");
 	case EHairAttribute::Color:						return LOCTEXT("GroomOptionsWindow_HasColor", "Color");
@@ -285,9 +285,11 @@ void SGroomImportOptionsWindow::Construct(const FArguments& InArgs)
 
 	// Aggregate attributes from all groups (ideally we should display each group attribute separately, to check if one groom is not missing data)
 	uint32 Attributes = 0;
+	uint32 AttributeFlags = 0;
 	for (const FGroomHairGroupPreview& Group : GroupsPreview->Groups)
 	{
 		Attributes |= Group.Attributes;
+		AttributeFlags |= Group.AttributeFlags;
 	}
 
 	FText bHasAttributeText = LOCTEXT("GroomOptionsWindow_HasAttributeNone", "None");
@@ -443,7 +445,7 @@ void SGroomImportOptionsWindow::Construct(const FArguments& InArgs)
 		if (HasHairAttribute(Attributes, AttributeType))
 		{
 			SVerticalBox::FScopedWidgetSlotArguments SlotArg = VerticalSlot->InsertSlot(AttributeSlotIndex++);
-			AddAttribute(SlotArg, GetHairAttributeLocText(AttributeType));
+			AddAttribute(SlotArg, GetHairAttributeLocText(AttributeType, AttributeFlags));
 		}
 	}
 
diff --git a/Engine/Plugins/Runtime/HairStrands/Source/HairStrandsEditor/Private/HairStrandsFactory.cpp b/Engine/Plugins/Runtime/HairStrands/Source/HairStrandsEditor/Private/HairStrandsFactory.cpp
--- a/Engine/Plugins/Runtime/HairStrands/Source/HairStrandsEditor/Private/HairStrandsFactory.cpp
+++ b/Engine/Plugins/Runtime/HairStrands/Source/HairStrandsEditor/Private/HairStrandsFactory.cpp
@@ -138,12 +138,13 @@ UObject* UHairStrandsFactory::FactoryCreateFile(UClass* InClass, UObject* InPare
 			for (const FHairDescriptionGroup& Group : OutDescription.HairGroups)
 			{
 				FGroomHairGroupPreview& OutGroup = GroupsPreview->Groups.AddDefaulted_GetRef();
-				OutGroup.GroupName  = Group.Info.GroupName;
-				OutGroup.GroupID	= Group.Info.GroupID;
-				OutGroup.CurveCount = Group.Info.NumCurves;
-				OutGroup.GuideCount = Group.Info.NumGuides;
-				OutGroup.Attributes = Group.Attributes;
-				OutGroup.Flags      = Group.Info.Flags;
+				OutGroup.GroupName  		= Group.Info.GroupName;
+				OutGroup.GroupID			= Group.Info.GroupID;
+				OutGroup.CurveCount 		= Group.Info.NumCurves;
+				OutGroup.GuideCount 		= Group.Info.NumGuides;
+				OutGroup.Attributes 		= Group.GetHairAttributes();
+				OutGroup.AttributeFlags 	= Group.GetHairAttributeFlags();
+				OutGroup.Flags      		= Group.Info.Flags;
 				bGuidesOnly |= (OutGroup.CurveCount == 0 && OutGroup.GuideCount > 0);
 
 				if (OutGroup.GroupID < OutDescription.HairGroups.Num())
diff --git a/Engine/Plugins/Runtime/HairStrands/Source/HairStrandsEditor/Private/ReimportHairStrandsFactory.cpp b/Engine/Plugins/Runtime/HairStrands/Source/HairStrandsEditor/Private/ReimportHairStrandsFactory.cpp
--- a/Engine/Plugins/Runtime/HairStrands/Source/HairStrandsEditor/Private/ReimportHairStrandsFactory.cpp
+++ b/Engine/Plugins/Runtime/HairStrands/Source/HairStrandsEditor/Private/ReimportHairStrandsFactory.cpp
@@ -380,12 +380,13 @@ EReimportResult::Type UReimportHairStrandsFactory::Reimport(UObject* Obj)
 			for (const FHairDescriptionGroup& Group : OutDescription.HairGroups)
 			{
 				FGroomHairGroupPreview& OutGroup = GroupsPreview->Groups.AddDefaulted_GetRef();
-				OutGroup.GroupName  = Group.Info.GroupName;
-				OutGroup.GroupID	= Group.Info.GroupID;
-				OutGroup.CurveCount = Group.Info.NumCurves;
-				OutGroup.GuideCount = Group.Info.NumGuides;
-				OutGroup.Attributes = Group.Attributes;
-				OutGroup.Flags		= Group.Info.Flags;
+				OutGroup.GroupName  	= Group.Info.GroupName;
+				OutGroup.GroupID		= Group.Info.GroupID;
+				OutGroup.CurveCount 	= Group.Info.NumCurves;
+				OutGroup.GuideCount 	= Group.Info.NumGuides;
+				OutGroup.Attributes 	= Group.GetHairAttributes();
+				OutGroup.AttributeFlags = Group.GetHairAttributeFlags();
+				OutGroup.Flags			= Group.Info.Flags;
 				if (OutGroup.GroupID < OutDescription.HairGroups.Num())
 				{
 					OutGroup.InterpolationSettings = GroomReimportOptions->InterpolationSettings[OutGroup.GroupID];
