Fix missing groom update when changing/tweaking decimation parameters.

diff --git a/Engine/Plugins/Runtime/HairStrands/Source/HairStrandsCore/Private/GroomAsset.cpp b/Engine/Plugins/Runtime/HairStrands/Source/HairStrandsCore/Private/GroomAsset.cpp
--- a/Engine/Plugins/Runtime/HairStrands/Source/HairStrandsCore/Private/GroomAsset.cpp
+++ b/Engine/Plugins/Runtime/HairStrands/Source/HairStrandsCore/Private/GroomAsset.cpp
@@ -1448,6 +1448,14 @@ static bool IsStrandsInterpolationAttributes(const FName PropertyName)
 		|| PropertyName == GET_MEMBER_NAME_CHECKED(UGroomAsset, EnableGlobalInterpolation)
 		|| PropertyName == GET_MEMBER_NAME_CHECKED(UGroomAsset, EnableSimulationCache)
 		|| PropertyName == GET_MEMBER_NAME_CHECKED(FHairSolverSettings, EnableSimulation)
+
+		// Decimation control needs to force a rebuild
+		|| PropertyName == GET_MEMBER_NAME_CHECKED(FHairLODSettings, CurveDecimation)
+		|| PropertyName == GET_MEMBER_NAME_CHECKED(FHairLODSettings, VertexDecimation)
+		|| PropertyName == GET_MEMBER_NAME_CHECKED(FHairLODSettings, AngularThreshold)
+		|| PropertyName == GET_MEMBER_NAME_CHECKED(FHairLODSettings, ScreenSize)
+		|| PropertyName == GET_MEMBER_NAME_CHECKED(FHairLODSettings, ThicknessScale)
+		|| PropertyName == GET_MEMBER_NAME_CHECKED(FHairLODSettings, bVisible)
 		|| PropertyName == GET_MEMBER_NAME_CHECKED(FHairLODSettings, Simulation)
 		|| PropertyName == GET_MEMBER_NAME_CHECKED(FHairLODSettings, GlobalInterpolation)
 
