Replace point/curve count acces with function to ease future change.

diff --git a/Engine/Plugins/Runtime/HairStrands/Source/HairStrandsCore/Private/Deformer/DeformerDataInterfaceGroom.cpp b/Engine/Plugins/Runtime/HairStrands/Source/HairStrandsCore/Private/Deformer/DeformerDataInterfaceGroom.cpp
--- a/Engine/Plugins/Runtime/HairStrands/Source/HairStrandsCore/Private/Deformer/DeformerDataInterfaceGroom.cpp
+++ b/Engine/Plugins/Runtime/HairStrands/Source/HairStrandsCore/Private/Deformer/DeformerDataInterfaceGroom.cpp
@@ -208,8 +208,8 @@ void FOptimusGroomDataProviderProxy::GatherDispatchData(FDispatchData const& InD
 		if (FHairGroupInstance* Instance = GroomComponent->GetGroupInstance(InvocationIndex))
 		{
 			const bool bIsSRVValid = Resources[InvocationIndex].PositionSRV != nullptr;
-			const int32 NumControlPoints = bIsSRVValid ? Instance->Strands.Data->PointCount : 0;
-			const int32 NumCurves = bIsSRVValid ? Instance->Strands.Data->CurveCount : 0;
+			const int32 NumControlPoints = bIsSRVValid ? Instance->Strands.Data->GetNumPoints() : 0;
+			const int32 NumCurves = bIsSRVValid ? Instance->Strands.Data->GetNumCurves() : 0;
 
 			const FHairGroupPublicData::FVertexFactoryInput VFInput = ComputeHairStrandsVertexInputData(Instance, EGroomViewMode::None);
 			
diff --git a/Engine/Plugins/Runtime/HairStrands/Source/HairStrandsCore/Private/Deformer/DeformerDataInterfaceGroomExec.cpp b/Engine/Plugins/Runtime/HairStrands/Source/HairStrandsCore/Private/Deformer/DeformerDataInterfaceGroomExec.cpp
--- a/Engine/Plugins/Runtime/HairStrands/Source/HairStrandsCore/Private/Deformer/DeformerDataInterfaceGroomExec.cpp
+++ b/Engine/Plugins/Runtime/HairStrands/Source/HairStrandsCore/Private/Deformer/DeformerDataInterfaceGroomExec.cpp
@@ -104,8 +104,8 @@ int32 FOptimusGroomExecDataProviderProxy::GetDispatchThreadCount(TArray<FIntVect
 	{		
 		if (FHairGroupInstance* Instance = GroomComponent->GetGroupInstance(InvocationIndex))
 		{
-			const int32 NumControlPoints = Instance->Strands.Data->PointCount;
-			const int32 NumCurves = Instance->Strands.Data->CurveCount;
+			const int32 NumControlPoints = Instance->Strands.Data->GetNumPoints();
+			const int32 NumCurves = Instance->Strands.Data->GetNumCurves();
 			const int32 NumThreads = Domain == EOptimusGroomExecDomain::ControlPoint ? NumControlPoints : NumCurves;
 			ThreadCounts.Add(FIntVector(NumThreads, 1, 1));
 		}
@@ -133,8 +133,8 @@ void FOptimusGroomExecDataProviderProxy::GatherDispatchData(FDispatchData const&
 	for (int32 InvocationIndex = 0; InvocationIndex < ParameterArray.Num(); ++InvocationIndex)
 	{
 		FHairGroupInstance* Instance = GroomComponent->GetGroupInstance(InvocationIndex);
-		const int32 NumControlPoints = Instance->Strands.Data->PointCount;
-		const int32 NumCurves = Instance->Strands.Data->CurveCount;
+		const int32 NumControlPoints = Instance->Strands.Data->GetNumPoints();
+		const int32 NumCurves = Instance->Strands.Data->GetNumCurves();
 		const int32 NumThreads = Domain == EOptimusGroomExecDomain::ControlPoint ? NumControlPoints : NumCurves;
 
 		FParameters& Parameters = ParameterArray[InvocationIndex];
diff --git a/Engine/Plugins/Runtime/HairStrands/Source/HairStrandsCore/Private/Deformer/DeformerDataInterfaceGroomWrite.cpp b/Engine/Plugins/Runtime/HairStrands/Source/HairStrandsCore/Private/Deformer/DeformerDataInterfaceGroomWrite.cpp
--- a/Engine/Plugins/Runtime/HairStrands/Source/HairStrandsCore/Private/Deformer/DeformerDataInterfaceGroomWrite.cpp
+++ b/Engine/Plugins/Runtime/HairStrands/Source/HairStrandsCore/Private/Deformer/DeformerDataInterfaceGroomWrite.cpp
@@ -139,8 +139,8 @@ void FOptimusGroomWriteDataProviderProxy::AllocateResources(FRDGBuilder& GraphBu
 		if (GroomInstance)
 		{
 			// Allocate required buffers
-			const int32 NumControlPoints = GroomInstance->Strands.Data->PointCount;
-			const int32 NumCurves = GroomInstance->Strands.Data->CurveCount;
+			const int32 NumControlPoints = GroomInstance->Strands.Data->GetNumPoints();
+			const int32 NumCurves = GroomInstance->Strands.Data->GetNumCurves();
 
 			FResources& R = Resources.AddDefaulted_GetRef();
 			if (OutputMask & 1)
@@ -173,8 +173,8 @@ void FOptimusGroomWriteDataProviderProxy::GatherDispatchData(FDispatchData const
 			FResources& Resource = Resources[InvocationIndex];
 
 			const bool bValid = Resource.PositionBufferUAV != nullptr && Resource.PositionBufferSRV != nullptr && Resource.PositionOffsetBufferSRV != nullptr;
-			const int32 NumCurves = bValid ? GroomInstance->Strands.Data->CurveCount : 0;
-			const int32 NumControlPoints = bValid ? GroomInstance->Strands.Data->PointCount : 0;
+			const int32 NumCurves = bValid ? GroomInstance->Strands.Data->GetNumCurves() : 0;
+			const int32 NumControlPoints = bValid ? GroomInstance->Strands.Data->GetNumPoints() : 0;
 
 			FParameters& Parameters = ParameterArray[InvocationIndex];
 			Parameters.NumControlPoints = NumControlPoints;
diff --git a/Engine/Plugins/Runtime/HairStrands/Source/HairStrandsCore/Private/Deformer/DeformerGroomComponentSource.cpp b/Engine/Plugins/Runtime/HairStrands/Source/HairStrandsCore/Private/Deformer/DeformerGroomComponentSource.cpp
--- a/Engine/Plugins/Runtime/HairStrands/Source/HairStrandsCore/Private/Deformer/DeformerGroomComponentSource.cpp
+++ b/Engine/Plugins/Runtime/HairStrands/Source/HairStrandsCore/Private/Deformer/DeformerGroomComponentSource.cpp
@@ -67,8 +67,8 @@ bool UOptimusGroomComponentSource::GetComponentElementCountsForExecutionDomain(
 		{
 			if (const FHairGroupInstance* Instance = GroomComponent->GetGroupInstance(InvocationIndex))
 			{
-				const int32 NumControlPoints = Instance->Strands.Data->PointCount;
-				const int32 NumCurves = Instance->Strands.Data->CurveCount;
+				const int32 NumControlPoints = Instance->Strands.Data->GetNumPoints();
+				const int32 NumCurves = Instance->Strands.Data->GetNumCurves();
 				const int32 NumThreads = InDomainName == Domains::ControlPoint ? NumControlPoints : NumCurves;
 				OutInvocationElementCounts.Add(NumThreads);
 			}
diff --git a/Engine/Plugins/Runtime/HairStrands/Source/HairStrandsCore/Private/GroomAsset.cpp b/Engine/Plugins/Runtime/HairStrands/Source/HairStrandsCore/Private/GroomAsset.cpp
--- a/Engine/Plugins/Runtime/HairStrands/Source/HairStrandsCore/Private/GroomAsset.cpp
+++ b/Engine/Plugins/Runtime/HairStrands/Source/HairStrandsCore/Private/GroomAsset.cpp
@@ -2853,7 +2853,7 @@ bool UGroomAsset::BuildCardsGeometry(uint32 GroupIndex)
 
 				// Update card stats to display
 				Desc->CardsInfo.NumCardVertices = LOD.BulkData.GetNumVertices();
-				Desc->CardsInfo.NumCards = LOD.Guides.BulkData.CurveCount;// LOD.Data.Cards.IndexOffsets.Num();
+				Desc->CardsInfo.NumCards = LOD.Guides.BulkData.GetNumCurves();// LOD.Data.Cards.IndexOffsets.Num();
 			}
 		}
 	}
@@ -3292,7 +3292,7 @@ void UGroomAsset::InitCardsResources()
 			{
 				// Update card stats to display
 				Desc->CardsInfo.NumCardVertices = LOD.BulkData.GetNumVertices();
-				Desc->CardsInfo.NumCards = LOD.Guides.BulkData.CurveCount;// LOD.Data.Cards.GetNumCards();
+				Desc->CardsInfo.NumCards = LOD.Guides.BulkData.GetNumCurves();// LOD.Data.Cards.GetNumCards();
 			}
 		}
 	}
diff --git a/Engine/Plugins/Runtime/HairStrands/Source/HairStrandsCore/Private/GroomBuilder.cpp b/Engine/Plugins/Runtime/HairStrands/Source/HairStrandsCore/Private/GroomBuilder.cpp
--- a/Engine/Plugins/Runtime/HairStrands/Source/HairStrandsCore/Private/GroomBuilder.cpp
+++ b/Engine/Plugins/Runtime/HairStrands/Source/HairStrandsCore/Private/GroomBuilder.cpp
@@ -3198,7 +3198,7 @@ static FRDGBufferRef ReadPositions(
 	FRDGBufferRef OutPositions = nullptr;
 	if (const FHairGroupInstance* Instance = Component->GetGroupInstance(GroupIt))
 	{
-		const uint32 NumVertices = Instance->Strands.Data->PointCount;
+		const uint32 NumVertices = Instance->Strands.Data->GetNumPoints();
 
 		FRDGBufferSRVRef InPositions = nullptr;
 		FRDGBufferSRVRef InPositionOffset = nullptr;
@@ -3266,8 +3266,8 @@ void RunHairStrandsPositionQueries(FRDGBuilder& GraphBuilder, FGlobalShaderMap*
 			if (Q->GPUPositions->IsReady() || Q->ReadbackDelay == 0)
 			{
 				const FHairGroupInstance* Instance = Q->Component->GetGroupInstance(Q->GroupIt);
-				const uint32 CurveCount = Instance->Strands.RestResource->BulkData.CurveCount;
-				const uint32 PointCount = Instance->Strands.RestResource->BulkData.PointCount;
+				const uint32 CurveCount = Instance->Strands.RestResource->BulkData.GetNumCurves();
+				const uint32 PointCount = Instance->Strands.RestResource->BulkData.GetNumPoints();
 
 				check(Q->GroupIt < uint32(Q->Output->Groups.Num()));
 				FStrandsPositionOutput::FGroup& Group = Q->Output->Groups[Q->GroupIt];
diff --git a/Engine/Plugins/Runtime/HairStrands/Source/HairStrandsCore/Private/GroomDebug.cpp b/Engine/Plugins/Runtime/HairStrands/Source/HairStrandsCore/Private/GroomDebug.cpp
--- a/Engine/Plugins/Runtime/HairStrands/Source/HairStrandsCore/Private/GroomDebug.cpp
+++ b/Engine/Plugins/Runtime/HairStrands/Source/HairStrandsCore/Private/GroomDebug.cpp
@@ -598,7 +598,7 @@ static void AddDrawDebugStrandsCVsPass(
 	Parameters->ViewUniformBuffer = View.ViewUniformBuffer;
 	Parameters->HairStrandsVF = Instance->Strands.UniformBuffer;
 	Parameters->LocalToWorld = FMatrix44f(Instance->LocalToWorld.ToMatrixWithScale());		// LWC_TODO: Precision loss
-	Parameters->MaxVertexCount = Instance->Strands.Data->PointCount;
+	Parameters->MaxVertexCount = Instance->Strands.Data->GetNumPoints();
 	Parameters->ColorTexture = GraphBuilder.CreateUAV(ColorTexture);
 	Parameters->DepthTexture = DepthTexture;
 	Parameters->LinearSampler = TStaticSamplerState<SF_Bilinear, AM_Clamp, AM_Clamp, AM_Clamp>::GetRHI();
diff --git a/Engine/Plugins/Runtime/HairStrands/Source/HairStrandsCore/Private/GroomResources.cpp b/Engine/Plugins/Runtime/HairStrands/Source/HairStrandsCore/Private/GroomResources.cpp
--- a/Engine/Plugins/Runtime/HairStrands/Source/HairStrandsCore/Private/GroomResources.cpp
+++ b/Engine/Plugins/Runtime/HairStrands/Source/HairStrandsCore/Private/GroomResources.cpp
@@ -889,8 +889,8 @@ void FHairStrandsRestResource::InternalAllocate(FRDGBuilder& GraphBuilder)
 {
 	BulkDataRequest = FBulkDataBatchRequest();
 
-	const uint32 PointCount = BulkData.PointCount;
-	const uint32 CurveCount = BulkData.CurveCount;
+	const uint32 PointCount = BulkData.GetNumPoints();
+	const uint32 CurveCount = BulkData.GetNumCurves();
 
 	// 1. Lock data, which force the loading data from files (on non-editor build/cooked data). These data are then uploaded to the GPU
 	// 2. A local copy is done by the buffer uploader. This copy is discarded once the uploading is done.
@@ -936,12 +936,12 @@ FRDGExternalBuffer FHairStrandsRestResource::GetTangentBuffer(FRDGBuilder& Graph
 	// Lazy allocation and update
 	if (TangentBuffer.Buffer == nullptr)
 	{
-		InternalCreateVertexBufferRDG<FHairStrandsTangentFormat>(GraphBuilder, BulkData.PointCount * FHairStrandsTangentFormat::ComponentCount, TangentBuffer, ToHairResourceDebugName(TEXT("Hair.StrandsRest_TangentBuffer"), ResourceName), OwnerName, EHairResourceUsageType::Dynamic);
+		InternalCreateVertexBufferRDG<FHairStrandsTangentFormat>(GraphBuilder, BulkData.GetNumPoints() * FHairStrandsTangentFormat::ComponentCount, TangentBuffer, ToHairResourceDebugName(TEXT("Hair.StrandsRest_TangentBuffer"), ResourceName), OwnerName, EHairResourceUsageType::Dynamic);
 
 		AddHairTangentPass(
 			GraphBuilder,
 			ShaderMap,
-			BulkData.PointCount,
+			BulkData.GetNumPoints(),
 			nullptr,
 			RegisterAsSRV(GraphBuilder, PositionBuffer),
 			Register(GraphBuilder, TangentBuffer, ERDGImportedBufferFlags::CreateUAV));
@@ -972,11 +972,11 @@ FHairStrandsDeformedResource::FHairStrandsDeformedResource(FHairStrandsBulkData&
 
 void FHairStrandsDeformedResource::InternalAllocate(FRDGBuilder& GraphBuilder)
 {
-	const uint32 PointCount = BulkData.PointCount;
+	const uint32 PointCount = BulkData.GetNumPoints();
 
-	InternalCreateVertexBufferRDG<FHairStrandsPositionFormat>(GraphBuilder, BulkData.PointCount, DeformedPositionBuffer[0], ToHairResourceDebugName(HAIRSTRANDS_RESOUCE_NAME(CurveType, Hair.StrandsDeformed_DeformedPositionBuffer0), ResourceName), OwnerName, EHairResourceUsageType::Dynamic);
-	InternalCreateVertexBufferRDG<FHairStrandsPositionFormat>(GraphBuilder, BulkData.PointCount, DeformedPositionBuffer[1], ToHairResourceDebugName(HAIRSTRANDS_RESOUCE_NAME(CurveType, Hair.StrandsDeformed_DeformedPositionBuffer1), ResourceName), OwnerName, EHairResourceUsageType::Dynamic);
-	InternalCreateVertexBufferRDG<FHairStrandsTangentFormat>(GraphBuilder, BulkData.PointCount * FHairStrandsTangentFormat::ComponentCount, TangentBuffer, ToHairResourceDebugName(HAIRSTRANDS_RESOUCE_NAME(CurveType, Hair.StrandsDeformed_TangentBuffer), ResourceName), OwnerName, EHairResourceUsageType::Dynamic);
+	InternalCreateVertexBufferRDG<FHairStrandsPositionFormat>(GraphBuilder, PointCount, DeformedPositionBuffer[0], ToHairResourceDebugName(HAIRSTRANDS_RESOUCE_NAME(CurveType, Hair.StrandsDeformed_DeformedPositionBuffer0), ResourceName), OwnerName, EHairResourceUsageType::Dynamic);
+	InternalCreateVertexBufferRDG<FHairStrandsPositionFormat>(GraphBuilder, PointCount, DeformedPositionBuffer[1], ToHairResourceDebugName(HAIRSTRANDS_RESOUCE_NAME(CurveType, Hair.StrandsDeformed_DeformedPositionBuffer1), ResourceName), OwnerName, EHairResourceUsageType::Dynamic);
+	InternalCreateVertexBufferRDG<FHairStrandsTangentFormat>(GraphBuilder,  PointCount * FHairStrandsTangentFormat::ComponentCount, TangentBuffer, ToHairResourceDebugName(HAIRSTRANDS_RESOUCE_NAME(CurveType, Hair.StrandsDeformed_TangentBuffer), ResourceName), OwnerName, EHairResourceUsageType::Dynamic);
 
 	TArray<FVector4f> DefaultOffsets;
 	DefaultOffsets.Add((FVector3f)BulkData.GetPositionOffset()); // LWC_TODO: precision loss
@@ -991,7 +991,7 @@ FRDGExternalBuffer& FHairStrandsDeformedResource::GetDeformerBuffer(FRDGBuilder&
 	// Lazy allocation and update
 	if (DeformerBuffer.Buffer == nullptr)
 	{
-		InternalCreateVertexBufferRDG<FHairStrandsPositionFormat>(GraphBuilder, BulkData.PointCount, DeformerBuffer, ToHairResourceDebugName(HAIRSTRANDS_RESOUCE_NAME(CurveType, Hair.StrandsDeformed_DeformerBuffer), ResourceName), OwnerName, EHairResourceUsageType::Dynamic);
+		InternalCreateVertexBufferRDG<FHairStrandsPositionFormat>(GraphBuilder, BulkData.GetNumPoints(), DeformerBuffer, ToHairResourceDebugName(HAIRSTRANDS_RESOUCE_NAME(CurveType, Hair.StrandsDeformed_DeformerBuffer), ResourceName), OwnerName, EHairResourceUsageType::Dynamic);
 	}
 	return DeformerBuffer;
 }
diff --git a/Engine/Plugins/Runtime/HairStrands/Source/HairStrandsCore/Private/GroomTextureBuilder.cpp b/Engine/Plugins/Runtime/HairStrands/Source/HairStrandsCore/Private/GroomTextureBuilder.cpp
--- a/Engine/Plugins/Runtime/HairStrands/Source/HairStrandsCore/Private/GroomTextureBuilder.cpp
+++ b/Engine/Plugins/Runtime/HairStrands/Source/HairStrandsCore/Private/GroomTextureBuilder.cpp
@@ -1099,7 +1099,7 @@ static bool TraceTextures(
 				GroupData.Strands.RestResource->GetPositionOffset(),
 				RenderingData.GeometrySettings.HairWidth * 0.5f,
 				GroupData.Strands.RestResource->BulkData.MaxLength,
-				GroupData.Strands.RestResource->BulkData.PointCount,
+				GroupData.Strands.RestResource->BulkData.GetNumPoints(),
 				GroupIndex,
 				GroupData.Strands.RestResource->BulkData.AttributeOffsets,
 
