Add missing file from previous CL.

diff --git a/Engine/Plugins/Runtime/HairStrands/Source/HairStrandsCore/Public/GroomResources.h b/Engine/Plugins/Runtime/HairStrands/Source/HairStrandsCore/Public/GroomResources.h
--- a/Engine/Plugins/Runtime/HairStrands/Source/HairStrandsCore/Public/GroomResources.h
+++ b/Engine/Plugins/Runtime/HairStrands/Source/HairStrandsCore/Public/GroomResources.h
@@ -349,8 +349,8 @@ struct FHairStrandsRestResource : public FHairCommonResource
 	/* Curve: Point offset & count */
 	TArray<FPackedHairCurve> CurveData;
 
-	inline uint32 GetPointCount() const { return BulkData.PointCount; }
-	inline uint32 GetCurveCount() const { return BulkData.CurveCount; }
+	inline uint32 GetPointCount() const { return BulkData.GetNumPoints(); }
+	inline uint32 GetCurveCount() const { return BulkData.GetNumCurves(); }
 };
 
 struct FHairStrandsDeformedResource : public FHairCommonResource
